Added PhysicalMaterialMasks.

diff --git a/Engine/Source/Developer/AssetTools/Private/AssetTools.cpp b/Engine/Source/Developer/AssetTools/Private/AssetTools.cpp
--- a/Engine/Source/Developer/AssetTools/Private/AssetTools.cpp
+++ b/Engine/Source/Developer/AssetTools/Private/AssetTools.cpp
@@ -81,6 +81,7 @@
 #include "AssetTypeActions/AssetTypeActions_ObjectLibrary.h"
 #include "AssetTypeActions/AssetTypeActions_ParticleSystem.h"
 #include "AssetTypeActions/AssetTypeActions_PhysicalMaterial.h"
+#include "AssetTypeActions/AssetTypeActions_PhysicalMaterialMask.h"
 #include "AssetTypeActions/AssetTypeActions_PhysicsAsset.h"
 #include "AssetTypeActions/AssetTypeActions_PoseAsset.h"
 #include "AssetTypeActions/AssetTypeActions_PreviewMeshCollection.h"
@@ -239,6 +240,7 @@ UAssetToolsImpl::UAssetToolsImpl(const FObjectInitializer& ObjectInitializer)
 	RegisterAssetTypeActions(MakeShareable(new FAssetTypeActions_ObjectLibrary));
 	RegisterAssetTypeActions(MakeShareable(new FAssetTypeActions_ParticleSystem));
 	RegisterAssetTypeActions(MakeShareable(new FAssetTypeActions_PhysicalMaterial));
+	RegisterAssetTypeActions(MakeShareable(new FAssetTypeActions_PhysicalMaterialMask));
 	RegisterAssetTypeActions(MakeShareable(new FAssetTypeActions_PhysicsAsset));
 	RegisterAssetTypeActions(MakeShareable(new FAssetTypeActions_PreviewMeshCollection));
 	RegisterAssetTypeActions(MakeShareable(new FAssetTypeActions_ProceduralFoliageSpawner(FoliageCategoryBit)));
diff --git a/Engine/Source/Developer/AssetTools/Private/AssetTypeActions/AssetTypeActions_PhysicalMaterialMask.cpp b/Engine/Source/Developer/AssetTools/Private/AssetTypeActions/AssetTypeActions_PhysicalMaterialMask.cpp
new file mode 100644
--- /dev/null
+++ b/Engine/Source/Developer/AssetTools/Private/AssetTypeActions/AssetTypeActions_PhysicalMaterialMask.cpp
@@ -0,0 +1,155 @@
+// Copyright Epic Games, Inc. All Rights Reserved.
+
+#include "AssetTypeActions/AssetTypeActions_PhysicalMaterialMask.h"
+#include "PhysicalMaterials/PhysicalMaterialMask.h"
+#include "PhysicalMaterialMaskImport.h"
+#include "EditorFramework/AssetImportData.h"
+#include "Editor.h"
+#include "Modules/ModuleManager.h"
+#include "Engine/Texture.h"
+#include "PhysicalMaterialMaskImport.h"
+#include "HAL/PlatformApplicationMisc.h"
+#include "HAL/FileManager.h"
+#include "ToolMenus.h"
+#include "Misc/MessageDialog.h"
+#include "Misc/FeedbackContext.h"
+#include "Misc/FileHelper.h"
+#include "Misc/Paths.h"
+#include "Framework/Application/SlateApplication.h"
+
+#define LOCTEXT_NAMESPACE "AssetTypeActions"
+
+DEFINE_LOG_CATEGORY_STATIC(LogAssetTypeActionsPhysicalMaterialMask, Log, All);
+
+UClass* FAssetTypeActions_PhysicalMaterialMask::GetSupportedClass() const
+{
+	return UPhysicalMaterialMask::StaticClass();
+}
+
+void FAssetTypeActions_PhysicalMaterialMask::GetActions(const TArray<UObject*>& InObjects, FToolMenuSection& Section)
+{
+	TArray<TWeakObjectPtr<UPhysicalMaterialMask>> SelectedMasks = GetTypedWeakObjectPtrs<UPhysicalMaterialMask>(InObjects);
+
+	TArray<FString> ResolvedFilePaths;
+	GetResolvedSourceFilePaths(InObjects, ResolvedFilePaths);
+
+	if (SelectedMasks.Num() == 1 && ResolvedFilePaths.Num() == 1 && ResolvedFilePaths[0].IsEmpty())
+	{
+		Section.AddMenuEntry(
+			"PhysicalMaterialMask_Import",
+			LOCTEXT("PhysicalMaterialMask_ImportMaskTexture", "Import Mask Texture"),
+			LOCTEXT("PhysicalMaterialMask_ImportMaskTextureTooltip", "Imports a texture to use as a physical material mask."),
+			FSlateIcon(FEditorStyle::GetStyleSetName(), "ClassIcon.Material"),
+			FUIAction(
+				FExecuteAction::CreateSP(this, &FAssetTypeActions_PhysicalMaterialMask::ExecuteImport, SelectedMasks[0]),
+				FCanExecuteAction()
+			)
+		);
+	}
+}
+
+void FAssetTypeActions_PhysicalMaterialMask::GetResolvedSourceFilePaths(const TArray<UObject*>& TypeAssets, TArray<FString>& OutSourceFilePaths) const
+{
+	for (auto& Asset : TypeAssets)
+	{
+		TArray<FString> ExtractedFilenames;
+		UPhysicalMaterialMask* PhysMatMask = CastChecked<UPhysicalMaterialMask>(Asset);
+		if (UAssetImportData* AssetImportData = PhysMatMask->AssetImportData)
+		{
+			if (AssetImportData->GetSourceFileCount() > 0)
+			{
+				if (!AssetImportData->GetSourceData().SourceFiles[0].RelativeFilename.IsEmpty())
+				{
+					PhysMatMask->AssetImportData->ExtractFilenames(ExtractedFilenames);
+				}
+			}
+		}
+
+		if (ExtractedFilenames.Num() > 0)
+		{
+			OutSourceFilePaths.Emplace(ExtractedFilenames[0]);
+		}
+		else
+		{
+			OutSourceFilePaths.Emplace(FString());
+		}
+	}
+}
+
+
+class UThumbnailInfo* FAssetTypeActions_PhysicalMaterialMask::GetThumbnailInfo(UObject* Asset) const 
+{
+	return nullptr;
+}
+
+EThumbnailPrimType FAssetTypeActions_PhysicalMaterialMask::GetDefaultThumbnailPrimitiveType(UObject* Asset) const 
+{
+	return TPT_Sphere;
+}
+
+bool FAssetTypeActions_PhysicalMaterialMask::CanExecuteImportedAssetActions(const TArray<FString> ResolvedFilePaths) const
+{
+	// Verify that all the file paths are legitimate
+	for (const auto& SourceFilePath : ResolvedFilePaths)
+	{
+		if (!SourceFilePath.Len() || IFileManager::Get().FileSize(*SourceFilePath) == INDEX_NONE)
+		{
+			return false;
+		}
+	}
+
+	return true;
+}
+
+void FAssetTypeActions_PhysicalMaterialMask::ExecuteImport(TWeakObjectPtr<UPhysicalMaterialMask> SelectedMask)
+{
+	if (UPhysicalMaterialMask* PhysMatMask = SelectedMask.Get())
+	{
+		FPhysicalMaterialMaskImport::ImportMaskTexture(PhysMatMask);
+	}
+}
+
+void FAssetTypeActions_PhysicalMaterialMask::ExecuteReimport(TArray<TWeakObjectPtr<UPhysicalMaterialMask>> SelectedMasks)
+{
+	for (TWeakObjectPtr<UPhysicalMaterialMask> SelectedMask : SelectedMasks)
+	{
+		if (UPhysicalMaterialMask* PhysMatMask = SelectedMask.Get())
+		{
+			FPhysicalMaterialMaskImport::ReimportMaskTexture(PhysMatMask);
+		}
+	}
+}
+
+void FAssetTypeActions_PhysicalMaterialMask::ExecuteReimportWithNewFile(TWeakObjectPtr<UPhysicalMaterialMask> SelectedMask)
+{
+	if (UPhysicalMaterialMask* PhysMatMask = SelectedMask.Get())
+	{
+		FPhysicalMaterialMaskImport::ReimportMaskTextureWithNewFile(PhysMatMask);
+	}
+}
+
+void FAssetTypeActions_PhysicalMaterialMask::ExecuteOpenSourceLocation(const TArray<FString> ResolvedFilePaths)
+{
+	// Open all files in the explorer
+	for (const auto& SourceFilePath : ResolvedFilePaths)
+	{
+		FPlatformProcess::ExploreFolder(*FPaths::GetPath(SourceFilePath));
+	}
+}
+
+void FAssetTypeActions_PhysicalMaterialMask::ExecuteOpenInExternalEditor(const TArray<FString> ResolvedFilePaths)
+{
+	// Open all files in their respective editor
+	for (const auto& SourceFilePath : ResolvedFilePaths)
+	{
+		FPlatformProcess::LaunchFileInDefaultExternalApplication(*SourceFilePath, NULL, ELaunchVerb::Edit);
+	}
+}
+
+void FAssetTypeActions_PhysicalMaterialMask::ExecuteDebug(TWeakObjectPtr<UPhysicalMaterialMask> SelectedMask)
+{
+	if (UPhysicalMaterialMask* PhysMatMask = SelectedMask.Get())
+	{
+		PhysMatMask->DumpMaskData();
+	}
+}
diff --git a/Engine/Source/Developer/AssetTools/Private/AssetTypeActions/AssetTypeActions_PhysicalMaterialMask.h b/Engine/Source/Developer/AssetTools/Private/AssetTypeActions/AssetTypeActions_PhysicalMaterialMask.h
new file mode 100644
--- /dev/null
+++ b/Engine/Source/Developer/AssetTools/Private/AssetTypeActions/AssetTypeActions_PhysicalMaterialMask.h
@@ -0,0 +1,48 @@
+// Copyright Epic Games, Inc. All Rights Reserved.
+
+#pragma once
+
+#include "CoreMinimal.h"
+#include "AssetTypeActions_Base.h"
+
+class FAssetTypeActions_PhysicalMaterialMask : public FAssetTypeActions_Base
+{
+public:
+	// IAssetTypeActions Implementation
+	virtual FText GetName() const override { return NSLOCTEXT("AssetTypeActions", "AssetTypeActions_PhysicalMaterialMask", "Physical Material Mask"); }
+	virtual FColor GetTypeColor() const override { return FColor(200,162,108); }
+	virtual UClass* GetSupportedClass() const override;
+	virtual bool HasActions(const TArray<UObject*>& InObjects) const override { return true; }
+	virtual void GetActions(const TArray<UObject*>& InObjects, struct FToolMenuSection& Section) override;
+	virtual uint32 GetCategories() override { return EAssetTypeCategories::Physics; }
+	virtual bool IsImportedAsset() const override { return true; }
+	virtual void GetResolvedSourceFilePaths(const TArray<UObject*>& TypeAssets, TArray<FString>& OutSourceFilePaths) const override;
+
+	/** Returns the thumbnail info for the specified asset, if it has one. */
+	virtual class UThumbnailInfo* GetThumbnailInfo(UObject* Asset) const override;
+
+	/** Returns the default thumbnail type that should be rendered when rendering primitive shapes.  This does not need to be implemented if the asset does not render a primitive shape */
+	virtual EThumbnailPrimType GetDefaultThumbnailPrimitiveType(UObject* Asset) const override;
+
+private:
+	/** Handler to check to see if imported asset actions are allowed */
+	bool CanExecuteImportedAssetActions(const TArray<FString> ResolvedFilePaths) const;
+
+	/** Handler for Import */
+	void ExecuteImport(TWeakObjectPtr<UPhysicalMaterialMask> InSelectedMask);
+
+	/** Handler for Reimport */
+	void ExecuteReimport(TArray<TWeakObjectPtr<UPhysicalMaterialMask>> InSelectedMasks);
+
+	/** Handler for ReimportWithNewFile */
+	void ExecuteReimportWithNewFile(TWeakObjectPtr<UPhysicalMaterialMask> InSelectedMask);
+
+	/** Handler for OpenSourceLocation */
+	void ExecuteOpenSourceLocation(const TArray<FString> ResolvedFilePaths);
+
+	/** Handler for OpenInExternalEditor */
+	void ExecuteOpenInExternalEditor(const TArray<FString> ResolvedFilePaths);
+
+	/** Handler for Debug */
+	void ExecuteDebug(TWeakObjectPtr<UPhysicalMaterialMask> InSelectedMask);
+};
diff --git a/Engine/Source/Editor/ContentBrowser/Private/AssetContextMenu.cpp b/Engine/Source/Editor/ContentBrowser/Private/AssetContextMenu.cpp
--- a/Engine/Source/Editor/ContentBrowser/Private/AssetContextMenu.cpp
+++ b/Engine/Source/Editor/ContentBrowser/Private/AssetContextMenu.cpp
@@ -443,7 +443,7 @@ bool FAssetContextMenu::AddImportedAssetMenuOptions(UToolMenu* Menu)
 					FSlateIcon(FEditorStyle::GetStyleSetName(), "ContentBrowser.AssetActions.ReimportAsset"),
 					FUIAction(
 						FExecuteAction::CreateSP(this, &FAssetContextMenu::ExecuteReimport, (int32)INDEX_NONE),
-						FCanExecuteAction()
+						FCanExecuteAction::CreateSP(this, &FAssetContextMenu::CanExecuteReimportAssetActions, ResolvedFilePaths)
 					)
 				);
 				if (ValidSelectedAssetCount == 1)
@@ -456,7 +456,7 @@ bool FAssetContextMenu::AddImportedAssetMenuOptions(UToolMenu* Menu)
 						FSlateIcon(FEditorStyle::GetStyleSetName(), "ContentBrowser.AssetActions.ReimportAsset"),
 						FUIAction(
 							FExecuteAction::CreateSP(this, &FAssetContextMenu::ExecuteReimportWithNewFile, (int32)INDEX_NONE),
-							FCanExecuteAction()
+							FCanExecuteAction::CreateSP(this, &FAssetContextMenu::CanExecuteReimportAssetActions, ResolvedFilePaths)
 						)
 					);
 				}
@@ -1874,6 +1874,11 @@ bool FAssetContextMenu::AreImportedAssetActionsVisible() const
 
 bool FAssetContextMenu::CanExecuteImportedAssetActions(const TArray<FString> ResolvedFilePaths) const
 {
+	if (ResolvedFilePaths.Num() == 0)
+	{
+		return false;
+	}
+
 	// Verify that all the file paths are legitimate
 	for (const auto& SourceFilePath : ResolvedFilePaths)
 	{
@@ -1886,6 +1891,25 @@ bool FAssetContextMenu::CanExecuteImportedAssetActions(const TArray<FString> Res
 	return true;
 }
 
+bool FAssetContextMenu::CanExecuteReimportAssetActions(const TArray<FString> ResolvedFilePaths) const
+{
+	if (ResolvedFilePaths.Num() == 0)
+	{
+		return false;
+	}
+
+	// Verify that all the file paths are non-empty
+	for (const auto& SourceFilePath : ResolvedFilePaths)
+	{
+		if (!SourceFilePath.Len())
+		{
+			return false;
+		}
+	}
+
+	return true;
+}
+
 void FAssetContextMenu::ExecuteReimport(int32 SourceFileIndex /*= INDEX_NONE*/)
 {
 	// Reimport all selected assets
diff --git a/Engine/Source/Editor/ContentBrowser/Private/AssetContextMenu.h b/Engine/Source/Editor/ContentBrowser/Private/AssetContextMenu.h
--- a/Engine/Source/Editor/ContentBrowser/Private/AssetContextMenu.h
+++ b/Engine/Source/Editor/ContentBrowser/Private/AssetContextMenu.h
@@ -99,6 +99,9 @@ private:
 	/** Handler to check to see if imported asset actions are allowed */
 	bool CanExecuteImportedAssetActions(const TArray<FString> ResolvedFilePaths) const;
 
+	/** Handler to check to see if reimport asset actions are allowed */
+	bool CanExecuteReimportAssetActions(const TArray<FString> ResolvedFilePaths) const;
+
 	/** Handler for Reimport */
 	void ExecuteReimport(int32 SourceFileIndex = INDEX_NONE);
 
diff --git a/Engine/Source/Editor/UnrealEd/Classes/Factories/PhysicalMaterialMaskFactory.h b/Engine/Source/Editor/UnrealEd/Classes/Factories/PhysicalMaterialMaskFactory.h
new file mode 100644
--- /dev/null
+++ b/Engine/Source/Editor/UnrealEd/Classes/Factories/PhysicalMaterialMaskFactory.h
@@ -0,0 +1,37 @@
+// Copyright Epic Games, Inc. All Rights Reserved.
+
+//~=============================================================================
+// PhysicalMaterialFactory
+//~=============================================================================
+
+#pragma once
+
+#include "CoreMinimal.h"
+#include "UObject/ObjectMacros.h"
+#include "Templates/SubclassOf.h"
+#include "Factories/Factory.h"
+#include "EditorReimportHandler.h"
+#include "PhysicalMaterialMaskFactory.generated.h"
+
+class UPhysicalMaterialMask;
+
+UCLASS(MinimalAPI, HideCategories=Object)
+class UPhysicalMaterialMaskFactory : public UFactory, public FReimportHandler
+{
+	GENERATED_UCLASS_BODY()
+
+	UPROPERTY(EditAnywhere, Category = PhysicalMaterialMaskFactory)
+	TSubclassOf<UPhysicalMaterialMask> PhysicalMaterialMaskClass;
+
+	//~ Begin UFactory Interface
+	virtual bool ConfigureProperties() override;
+	virtual UObject* FactoryCreateNew(UClass* Class,UObject* InParent,FName Name,EObjectFlags Flags,UObject* Context,FFeedbackContext* Warn) override;
+	//~ Begin UFactory Interface	
+
+	//~ Begin FReimportHandler Interface
+	virtual bool CanReimport(UObject* Obj, TArray<FString>& OutFilenames) override;
+	virtual void SetReimportPaths(UObject* Obj, const TArray<FString>& NewReimportPaths) override;
+	virtual EReimportResult::Type Reimport(UObject* Obj) override;
+	virtual int32 GetPriority() const override;
+	//~ End FReimportHandler Interface
+};
diff --git a/Engine/Source/Editor/UnrealEd/Classes/ThumbnailRendering/PhysicalMaterialMaskThumbnailRenderer.h b/Engine/Source/Editor/UnrealEd/Classes/ThumbnailRendering/PhysicalMaterialMaskThumbnailRenderer.h
new file mode 100644
--- /dev/null
+++ b/Engine/Source/Editor/UnrealEd/Classes/ThumbnailRendering/PhysicalMaterialMaskThumbnailRenderer.h
@@ -0,0 +1,28 @@
+// Copyright Epic Games, Inc. All Rights Reserved.
+
+/**
+ *
+ * This thumbnail renderer displays the mask texture for the object in question
+ */
+
+#pragma once
+
+#include "CoreMinimal.h"
+#include "UObject/ObjectMacros.h"
+#include "ThumbnailRendering/TextureThumbnailRenderer.h"
+#include "PhysicalMaterialMaskThumbnailRenderer.generated.h"
+
+class FCanvas;
+class FRenderTarget;
+
+UCLASS(MinimalAPI)
+class UNREALED_VTABLE UPhysicalMaterialMaskThumbnailRenderer : public UTextureThumbnailRenderer
+{
+	GENERATED_UCLASS_BODY()
+
+	// Begin UThumbnailRenderer Object
+	UNREALED_API virtual void GetThumbnailSize(UObject* Object, float Zoom, uint32& OutWidth, uint32& OutHeight) const override;
+	UNREALED_API virtual void Draw(UObject* Object, int32 X, int32 Y, uint32 Width, uint32 Height, FRenderTarget*, FCanvas* Canvas) override;
+	// End UThumbnailRenderer Object
+};
+
diff --git a/Engine/Source/Editor/UnrealEd/Private/Editor.cpp b/Engine/Source/Editor/UnrealEd/Private/Editor.cpp
--- a/Engine/Source/Editor/UnrealEd/Private/Editor.cpp
+++ b/Engine/Source/Editor/UnrealEd/Private/Editor.cpp
@@ -18,6 +18,7 @@
 #include "Factories/ReimportSoundFactory.h"
 #include "Factories/ReimportSoundSurroundFactory.h"
 #include "Factories/ReimportTextureFactory.h"
+#include "Factories/PhysicalMaterialMaskFactory.h"
 #include "Engine/BlueprintGeneratedClass.h"
 #include "Misc/ScopedSlowTask.h"
 #include "UObject/UObjectIterator.h"
@@ -775,6 +776,9 @@ FReimportManager::FReimportManager()
 
 	// Create reimport handler for FBX scene
 	UReimportFbxSceneFactory::StaticClass();
+
+	// Create reimport handler for PhysicalMaterialMasks
+	UPhysicalMaterialMaskFactory::StaticClass();
 }
 
 FReimportManager::~FReimportManager()
diff --git a/Engine/Source/Editor/UnrealEd/Private/Factories/EditorFactories.cpp b/Engine/Source/Editor/UnrealEd/Private/Factories/EditorFactories.cpp
--- a/Engine/Source/Editor/UnrealEd/Private/Factories/EditorFactories.cpp
+++ b/Engine/Source/Editor/UnrealEd/Private/Factories/EditorFactories.cpp
@@ -98,6 +98,7 @@
 #include "Factories/PackageFactory.h"
 #include "Factories/ParticleSystemFactoryNew.h"
 #include "Factories/PhysicalMaterialFactoryNew.h"
+#include "Factories/PhysicalMaterialMaskFactory.h"
 #include "Factories/PolysFactory.h"
 #include "Factories/ReverbEffectFactory.h"
 #include "Factories/SoundAttenuationFactory.h"
@@ -155,6 +156,8 @@
 #include "Materials/MaterialParameterCollection.h"
 #include "Engine/ObjectLibrary.h"
 #include "PhysicalMaterials/PhysicalMaterial.h"
+#include "PhysicalMaterials/PhysicalMaterialMask.h"
+#include "PhysicalMaterialMaskImport.h"
 #include "Engine/Polys.h"
 #include "Sound/ReverbEffect.h"
 #include "Sound/SoundCue.h"
@@ -1988,6 +1991,118 @@ UObject* UPhysicalMaterialFactoryNew::FactoryCreateNew(UClass* Class, UObject* I
 	}
 }
 
+/*------------------------------------------------------------------------------
+	UPhysicalMaterialMaskFactory.
+------------------------------------------------------------------------------*/
+UPhysicalMaterialMaskFactory::UPhysicalMaterialMaskFactory(const FObjectInitializer& ObjectInitializer)
+	: Super(ObjectInitializer)
+{
+	SupportedClass = UPhysicalMaterialMask::StaticClass();
+
+	bCreateNew = true;
+	bText = false;
+	bEditAfterNew = false;
+	bEditorImport = true;
+
+	// Required to allow texture factory to take priority when importing new image files
+	ImportPriority = DefaultImportPriority - 1;
+}
+
+bool UPhysicalMaterialMaskFactory::ConfigureProperties()
+{
+	// nullptr the DataAssetClass so we can check for selection
+	PhysicalMaterialMaskClass = nullptr;
+
+	// Load the classviewer module to display a class picker
+	FClassViewerModule& ClassViewerModule = FModuleManager::LoadModuleChecked<FClassViewerModule>("ClassViewer");
+
+	// Fill in options
+	FClassViewerInitializationOptions Options;
+	Options.Mode = EClassViewerMode::ClassPicker;
+
+	TSharedPtr<FAssetClassParentFilter> Filter = MakeShareable(new FAssetClassParentFilter);
+	Options.ClassFilter = Filter;
+
+	Filter->DisallowedClassFlags = CLASS_Abstract | CLASS_Deprecated | CLASS_NewerVersionExists;
+	Filter->AllowedChildrenOfClasses.Add(UPhysicalMaterialMask::StaticClass());
+
+	const FText TitleText = LOCTEXT("CreatePhysicalMaterialMask", "Pick Physical Material Mask Class");
+	UClass* ChosenClass = nullptr;
+	const bool bPressedOk = SClassPickerDialog::PickClass(TitleText, Options, ChosenClass, UPhysicalMaterialMask::StaticClass());
+
+	if (bPressedOk)
+	{
+		PhysicalMaterialMaskClass = ChosenClass;
+	}
+
+	return bPressedOk;
+}
+UObject* UPhysicalMaterialMaskFactory::FactoryCreateNew(UClass* Class, UObject* InParent, FName Name, EObjectFlags Flags, UObject* Context, FFeedbackContext* Warn)
+{
+	if (PhysicalMaterialMaskClass != nullptr)
+	{
+		return NewObject<UPhysicalMaterialMask>(InParent, PhysicalMaterialMaskClass, Name, Flags | RF_Transactional);
+	}
+	else
+	{
+		// if we have no data asset class, use the passed-in class instead
+		check(Class->IsChildOf(UPhysicalMaterialMask::StaticClass()));
+		return NewObject<UPhysicalMaterialMask>(InParent, Class, Name, Flags);
+	}
+}
+
+bool UPhysicalMaterialMaskFactory::CanReimport( UObject* Obj, TArray<FString>& OutFilenames )
+{	
+	UPhysicalMaterialMask* PhysMatMask = Cast<UPhysicalMaterialMask>(Obj);
+	if (PhysMatMask)
+	{
+		if (PhysMatMask->AssetImportData)
+		{
+			FString FileExtension = FPaths::GetExtension(PhysMatMask->AssetImportData->GetFirstFilename());
+			if (FileExtension.Equals(TEXT("png"), ESearchCase::IgnoreCase) || FileExtension.Equals("jpg", ESearchCase::IgnoreCase))
+			{
+				OutFilenames.Add(PhysMatMask->AssetImportData->GetFirstFilename());
+			}
+		}
+
+		return true;
+	}
+	return false;
+}
+
+void UPhysicalMaterialMaskFactory::SetReimportPaths( UObject* Obj, const TArray<FString>& NewReimportPaths )
+{	
+	UPhysicalMaterialMask* PhysMatMask = Cast<UPhysicalMaterialMask>(Obj);
+	if (PhysMatMask && ensure(NewReimportPaths.Num() == 1))
+	{
+		PhysMatMask->Modify();
+
+		if (!PhysMatMask->AssetImportData)
+		{
+			PhysMatMask->AssetImportData = NewObject<UAssetImportData>(PhysMatMask, TEXT("AssetImportData"), RF_NoFlags);
+		}
+
+		PhysMatMask->AssetImportData->UpdateFilenameOnly(NewReimportPaths[0]);
+	}
+}
+
+EReimportResult::Type UPhysicalMaterialMaskFactory::Reimport( UObject* Obj )
+{
+	if (!Obj || !Obj->IsA(UPhysicalMaterialMask::StaticClass()))
+	{
+		return EReimportResult::Failed;
+	}
+
+	UPhysicalMaterialMask* PhysMatMask = Cast<UPhysicalMaterialMask>(Obj);
+
+	return FPhysicalMaterialMaskImport::ReimportMaskTexture(PhysMatMask);
+}
+
+int32 UPhysicalMaterialMaskFactory::GetPriority() const
+{
+	return ImportPriority;
+}
+
 /*------------------------------------------------------------------------------
 	UInterpDataFactoryNew.
 ------------------------------------------------------------------------------*/
diff --git a/Engine/Source/Editor/UnrealEd/Private/PhysicalMaterialMaskImport.cpp b/Engine/Source/Editor/UnrealEd/Private/PhysicalMaterialMaskImport.cpp
new file mode 100644
--- /dev/null
+++ b/Engine/Source/Editor/UnrealEd/Private/PhysicalMaterialMaskImport.cpp
@@ -0,0 +1,231 @@
+// Copyright Epic Games, Inc. All Rights Reserved.
+
+#include "PhysicalMaterialMaskImport.h"
+#include "PhysicalMaterials/PhysicalMaterialMask.h"
+#include "CoreMinimal.h"
+#include "EngineDefines.h"
+#include "Misc/MessageDialog.h"
+#include "HAL/PlatformApplicationMisc.h"
+#include "HAL/FileManager.h"
+#include "Misc/CoreMisc.h"
+#include "Misc/Paths.h"
+#include "Misc/FileHelper.h"
+#include "Misc/PackageName.h"
+#include "Misc/Attribute.h"
+#include "Modules/ModuleManager.h"
+#include "UObject/ObjectMacros.h"
+#include "UObject/UObjectGlobals.h"
+#include "UObject/Object.h"
+#include "UObject/Class.h"
+#include "UObject/UObjectIterator.h"
+#include "UObject/Package.h"
+#include "UObject/Interface.h"
+#include "Misc/PackageName.h"
+#include "Fonts/FontBulkData.h"
+#include "Fonts/CompositeFont.h"
+#include "Input/Reply.h"
+#include "Engine/Texture.h"
+#include "Editor.h"
+#include "EditorDirectories.h"
+#include "EditorReimportHandler.h"
+#include "EditorFramework/AssetImportData.h"
+#include "ObjectTools.h"
+#include "PackageTools.h"
+#include "Modules/ModuleManager.h"
+#include "IContentBrowserSingleton.h"
+#include "ContentBrowserModule.h"
+#include "ClassViewerModule.h"
+#include "ClassViewerFilter.h"
+#include "AssetRegistryModule.h"
+#include "AssetToolsModule.h"
+#include "IAssetTools.h"
+#include "DesktopPlatformModule.h"
+#include "IDesktopPlatform.h"
+#include "Engine/Texture.h"
+#include "Factories/TextureFactory.h"
+
+DEFINE_LOG_CATEGORY_STATIC(LogPhysicalMaterialMaskImport, Log, All);
+
+void FPhysicalMaterialMaskImport::ImportMaskTexture(UPhysicalMaterialMask* PhysMatMask)
+{
+	check(PhysMatMask);
+
+	const FString TextureFilename = OpenMaskFileDialog();
+
+	if (TextureFilename.IsEmpty())
+	{
+		UE_LOG(LogPhysicalMaterialMaskImport, Error, TEXT("No texture mask selected."));
+	}
+
+	if (UTexture* MaskTexture = ImportMaskTextureFile(PhysMatMask, TextureFilename))
+	{
+		PhysMatMask->SetMaskTexture(MaskTexture, TextureFilename);
+	}
+}
+
+EReimportResult::Type FPhysicalMaterialMaskImport::ReimportMaskTexture(UPhysicalMaterialMask* PhysMatMask)
+{
+	check(PhysMatMask);
+	check(PhysMatMask->AssetImportData);
+	
+	const FString TextureFilename = PhysMatMask->AssetImportData->GetFirstFilename();
+
+	if (TextureFilename.IsEmpty())
+	{
+		UE_LOG(LogPhysicalMaterialMaskImport, Error, TEXT("PhysicalMaterialMask does not have asset import data filename for reimport."));
+		return EReimportResult::Failed;
+	}
+
+	if (UTexture* MaskTexture = ImportMaskTextureFile(PhysMatMask, TextureFilename))
+	{
+		PhysMatMask->SetMaskTexture(MaskTexture, TextureFilename);
+		return EReimportResult::Succeeded;
+	}
+
+	return EReimportResult::Failed;
+}
+
+EReimportResult::Type FPhysicalMaterialMaskImport::ReimportMaskTextureWithNewFile(UPhysicalMaterialMask* PhysMatMask)
+{
+	check(PhysMatMask);
+
+	const FString TextureFilename = OpenMaskFileDialog();
+
+	if (TextureFilename.IsEmpty())
+	{
+		UE_LOG(LogPhysicalMaterialMaskImport, Error, TEXT("No texture mask selected."));
+		return EReimportResult::Failed;
+	}
+
+	if (UTexture* MaskTexture = ImportMaskTextureFile(PhysMatMask, TextureFilename))
+	{
+		PhysMatMask->SetMaskTexture(MaskTexture, TextureFilename);
+		return EReimportResult::Succeeded;
+	}
+
+	return EReimportResult::Failed;
+}
+
+FString FPhysicalMaterialMaskImport::OpenMaskFileDialog()
+{
+	TArray<FString> OutFiles;
+
+	bool bOpen = false;
+	IDesktopPlatform* DesktopPlatform = FDesktopPlatformModule::Get();
+	if (DesktopPlatform)
+	{
+		FString Filter;
+		GetFileTypeFilterString(Filter);
+		const FString DefaultPath = FEditorDirectories::Get().GetLastDirectory(ELastDirectory::GENERIC_IMPORT);
+
+		bOpen = DesktopPlatform->OpenFileDialog(
+			FSlateApplication::Get().FindBestParentWindowHandleForDialogs(nullptr),
+			NSLOCTEXT("Import", "Import", "Import from...").ToString(),
+			FEditorDirectories::Get().GetLastDirectory(ELastDirectory::GENERIC_IMPORT),
+			TEXT(""),
+			*Filter,
+			EFileDialogFlags::None,
+			OutFiles
+		);
+	}
+	if (!bOpen || OutFiles.Num() == 0)
+	{
+		return TEXT("");
+	}
+
+	return OutFiles[0];
+}
+
+UTexture* FPhysicalMaterialMaskImport::ImportMaskTextureFile(UPhysicalMaterialMask* PhysMatMask, const FString& TextureFilename)
+{
+	if (TextureFilename.IsEmpty())
+	{
+		return nullptr;
+	}
+
+	FString Extension = FPaths::GetExtension(TextureFilename).ToLower();
+	FString TextureName = FPaths::GetBaseFilename(TextureFilename) + TEXT("_Mask");
+	TextureName = ObjectTools::SanitizeObjectName(TextureName);
+
+	UTexture* Texture = nullptr;
+
+	// try opening from absolute path
+	TArray<uint8> TextureData;
+	if (!(FFileHelper::LoadFileToArray(TextureData, *TextureFilename) && TextureData.Num() > 0))
+	{
+		UE_LOG(LogPhysicalMaterialMaskImport, Error, TEXT("Unable to find texture file %s."), *TextureFilename);
+	}
+	else
+	{
+		UTextureFactory* TextureFactory = NewObject<UTextureFactory>();
+		TextureFactory->AddToRoot();
+		TextureFactory->SuppressImportOverwriteDialog();
+
+		TextureFactory->CompressionSettings = TC_Masks;
+
+		const uint8* PtrTexture = TextureData.GetData();
+		Texture = (UTexture*)TextureFactory->FactoryCreateBinary(UTexture2D::StaticClass(), PhysMatMask->GetOuter(), *TextureName, RF_NoFlags, NULL, *Extension, PtrTexture, PtrTexture + TextureData.Num(), GWarn);
+		if (Texture != NULL)
+		{
+			Texture->SRGB = false;
+			Texture->CompressionNone = true;
+			Texture->CompressionSettings = TC_Masks;
+			Texture->MipGenSettings = TMGS_NoMipmaps;
+			Texture->Filter = TF_Nearest;
+			Texture->AssetImportData->Update(TextureFilename);
+		}
+
+		TextureFactory->RemoveFromRoot();
+	}
+
+	return Texture;
+}
+
+void FPhysicalMaterialMaskImport::GetSupportedTextureFileTypes(TArray<FString>& OutFileTypes)
+{
+	OutFileTypes.Empty();
+	OutFileTypes.Emplace(TEXT("png"));
+	OutFileTypes.Emplace(TEXT("jpg"));
+}
+
+void FPhysicalMaterialMaskImport::GetSupportedTextureSourceFormats(TArray<ETextureSourceFormat>& OutSourceFormats)
+{
+	OutSourceFormats.Empty();
+	OutSourceFormats.Emplace(ETextureSourceFormat::TSF_RGBA8);
+	OutSourceFormats.Emplace(ETextureSourceFormat::TSF_BGRA8);
+	OutSourceFormats.Emplace(ETextureSourceFormat::TSF_RGBA16);
+}
+
+void FPhysicalMaterialMaskImport::GetFileTypeFilterString(FString& OutFileTypeFilter)
+{
+	TArray<FString> SupportedTypes;
+	GetSupportedTextureFileTypes(SupportedTypes);
+
+	bool bJoin = false;
+	OutFileTypeFilter = FString(TEXT("Image file ("));
+	for (const FString& FileType : SupportedTypes)
+	{
+		if (bJoin)
+		{
+			OutFileTypeFilter += TEXT(',');
+		}
+		const FString Type = FString::Printf(TEXT("*.%s"), *FileType);
+		OutFileTypeFilter += Type;
+		bJoin = true;
+	}
+	OutFileTypeFilter += FString(TEXT(")|"));
+	bJoin = false;
+	for (const FString& FileType : SupportedTypes)
+	{
+		if (bJoin)
+		{
+			OutFileTypeFilter += TEXT(';');
+		}
+		const FString Type = FString::Printf(TEXT("*.%s"), *FileType);
+		OutFileTypeFilter += Type;
+		bJoin = true;
+	}
+}
+
+#undef LOCTEXT_NAMESPACE
+
diff --git a/Engine/Source/Editor/UnrealEd/Private/ThumbnailRendering/PhysicalMaterialMaskThumbnailRenderer.cpp b/Engine/Source/Editor/UnrealEd/Private/ThumbnailRendering/PhysicalMaterialMaskThumbnailRenderer.cpp
new file mode 100644
--- /dev/null
+++ b/Engine/Source/Editor/UnrealEd/Private/ThumbnailRendering/PhysicalMaterialMaskThumbnailRenderer.cpp
@@ -0,0 +1,31 @@
+// Copyright Epic Games, Inc. All Rights Reserved.
+
+#include "ThumbnailRendering/PhysicalMaterialMaskThumbnailRenderer.h"
+#include "CanvasItem.h"
+#include "Engine/Engine.h"
+#include "EngineGlobals.h"
+#include "PhysicalMaterials/PhysicalMaterialMask.h"
+#include "CanvasTypes.h"
+
+UPhysicalMaterialMaskThumbnailRenderer::UPhysicalMaterialMaskThumbnailRenderer(const FObjectInitializer& ObjectInitializer)
+	: Super(ObjectInitializer)
+{
+}
+
+void UPhysicalMaterialMaskThumbnailRenderer::GetThumbnailSize(UObject* Object, float Zoom, uint32& OutWidth, uint32& OutHeight) const
+{
+	UPhysicalMaterialMask* PhysMatMask = Cast<UPhysicalMaterialMask>(Object);
+	if (PhysMatMask->MaskTexture != nullptr)
+	{
+		Super::GetThumbnailSize(PhysMatMask->MaskTexture, Zoom, OutWidth, OutHeight);
+	}
+}
+
+void UPhysicalMaterialMaskThumbnailRenderer::Draw(UObject* Object, int32 X, int32 Y, uint32 Width, uint32 Height, FRenderTarget* RenderTarget, FCanvas* Canvas)
+{
+	UPhysicalMaterialMask* PhysMatMask = Cast<UPhysicalMaterialMask>(Object);
+	if (PhysMatMask->MaskTexture != nullptr)
+	{
+		Super::Draw(PhysMatMask->MaskTexture, X, Y, Width, Height, RenderTarget, Canvas);
+	}
+}
diff --git a/Engine/Source/Editor/UnrealEd/Public/PhysicalMaterialMaskImport.h b/Engine/Source/Editor/UnrealEd/Public/PhysicalMaterialMaskImport.h
new file mode 100644
--- /dev/null
+++ b/Engine/Source/Editor/UnrealEd/Public/PhysicalMaterialMaskImport.h
@@ -0,0 +1,43 @@
+// Copyright Epic Games, Inc. All Rights Reserved.
+
+//~=============================================================================
+// PhysicalMaterialMaskImport
+//~=============================================================================
+
+#pragma once
+
+#include "CoreMinimal.h"
+#include "PhysicalMaterials/PhysicalMaterialMask.h"
+#include "UObject/ObjectMacros.h"
+#include "Templates/SubclassOf.h"
+#include "Factories/Factory.h"
+#include "EditorReimportHandler.h"
+
+class FPhysicalMaterialMaskImport 
+{
+public:
+	// Create physical material mask texture
+	static UNREALED_API void ImportMaskTexture(UPhysicalMaterialMask* PhysMatMask);
+
+	// Reimport physical material mask texture
+	static UNREALED_API EReimportResult::Type ReimportMaskTexture(UPhysicalMaterialMask* PhysMatMask);
+
+	// Reimport physical material mask texture with new file
+	static UNREALED_API EReimportResult::Type ReimportMaskTextureWithNewFile(UPhysicalMaterialMask* PhysMatMask);
+
+private:
+	// Opens file dialog for texture selection
+	static FString OpenMaskFileDialog();
+
+	// Imports physical material mask texture
+	static UTexture* ImportMaskTextureFile(UPhysicalMaterialMask* PhysMatMask, const FString& TextureFilename);
+
+	// Returns file formats supported for mask texture. 
+	static void GetSupportedTextureFileTypes(TArray<FString>& OutFileFormats);
+
+	// Returns texture source formats supported for mask texture. 
+	static void GetSupportedTextureSourceFormats(TArray<ETextureSourceFormat>& OutSourceFormats);
+
+	// Helper to get dialog file type filter string 
+	static void GetFileTypeFilterString(FString& OutFileTypeFilter);
+};
diff --git a/Engine/Source/Runtime/Core/Private/UObject/DevObjectVersion.cpp b/Engine/Source/Runtime/Core/Private/UObject/DevObjectVersion.cpp
--- a/Engine/Source/Runtime/Core/Private/UObject/DevObjectVersion.cpp
+++ b/Engine/Source/Runtime/Core/Private/UObject/DevObjectVersion.cpp
@@ -25,6 +25,7 @@
 #include "UObject/NiagaraObjectVersion.h"
 #include "UObject/DestructionObjectVersion.h"
 #include "UObject/ExternalPhysicsCustomObjectVersion.h"
+#include "UObject/ExternalPhysicsMaterialCustomObjectVersion.h"
 #include "UObject/CineCameraObjectVersion.h"
 
 DEFINE_LOG_CATEGORY_STATIC(LogDevObjectVersion, Log, All);
@@ -169,6 +170,13 @@ FDevVersionRegistration GRegisterDestructionObjectVersion(FDestructionObjectVers
 const FGuid FExternalPhysicsCustomObjectVersion::GUID(0x35F94A83, 0xE258406C, 0xA31809F5, 0x9610247C);
 // Register Physics custom version with Core
 FDevVersionRegistration GRegisterExternalPhysicsCustomVersion(FExternalPhysicsCustomObjectVersion::GUID, FExternalPhysicsCustomObjectVersion::LatestVersion, TEXT("Dev-Physics-Ext"));
+// Unique Physics Object version id
+
+// Unique Physics Material Object version id
+const FGuid FExternalPhysicsMaterialCustomObjectVersion::GUID(0xB68FC16E, 0x8B1B42E2, 0xB453215C, 0x058844FE);
+// Register Physics custom version with Core
+FDevVersionRegistration GRegisterExternalPhysicsMaterialCustomVersion(FExternalPhysicsMaterialCustomObjectVersion::GUID, FExternalPhysicsMaterialCustomObjectVersion::LatestVersion, TEXT("Dev-PhysicsMaterial-Chaos"));
+// Unique PhysicsMaterial  Object version id
 
 // Unique CineCamera Object version id
 const FGuid FCineCameraObjectVersion::GUID(0xB2E18506, 0x4273CFC2, 0xA54EF4BB, 0x758BBA07);
diff --git a/Engine/Source/Runtime/Core/Public/UObject/ExternalPhysicsMaterialCustomObjectVersion.h b/Engine/Source/Runtime/Core/Public/UObject/ExternalPhysicsMaterialCustomObjectVersion.h
new file mode 100644
--- /dev/null
+++ b/Engine/Source/Runtime/Core/Public/UObject/ExternalPhysicsMaterialCustomObjectVersion.h
@@ -0,0 +1,30 @@
+// Copyright Epic Games, Inc. All Rights Reserved.
+#pragma once
+
+#include "CoreTypes.h"
+#include "Misc/Guid.h"
+
+// Custom serialization version for changes made to physics materials
+struct CORE_API FExternalPhysicsMaterialCustomObjectVersion
+{
+	enum Type
+	{
+		// Before any version changes were made
+		BeforeCustomVersionWasAdded = 0,
+
+		// Added material masks to Chaos
+		AddedMaterialMasks,
+
+		// -----<new versions can be added above this line>-------------------------------------------------
+		VersionPlusOne,
+		LatestVersion = VersionPlusOne - 1
+	};
+
+	// The GUID for this custom version number
+	const static FGuid GUID;
+
+private:
+	FExternalPhysicsMaterialCustomObjectVersion() {}
+};
+
+
diff --git a/Engine/Source/Runtime/Engine/Classes/Engine/EngineTypes.h b/Engine/Source/Runtime/Engine/Classes/Engine/EngineTypes.h
--- a/Engine/Source/Runtime/Engine/Classes/Engine/EngineTypes.h
+++ b/Engine/Source/Runtime/Engine/Classes/Engine/EngineTypes.h
@@ -18,6 +18,7 @@
 class AActor;
 class UDecalComponent;
 class UPhysicalMaterial;
+class UPhysicalMaterialMask;
 class UPrimitiveComponent;
 class USceneComponent;
 
@@ -2590,6 +2591,9 @@ struct FMeshBuildSettings
 	UPROPERTY(EditAnywhere, Category=BuildSettings, meta=(DisplayName="Two-Sided Distance Field Generation"))
 	uint8 bGenerateDistanceFieldAsIfTwoSided:1;
 
+	UPROPERTY(EditAnywhere, Category=BuildSettings, meta=(DisplayName="Enable Physical Material Mask"))
+	uint8 bSupportFaceRemap : 1;
+
 	UPROPERTY(EditAnywhere, Category=BuildSettings)
 	int32 MinLightmapResolution;
 
@@ -3648,6 +3652,24 @@ enum EPhysicalSurface
 	SurfaceType_Max UMETA(Hidden)
 };
 
+/** Types of valid physical material mask colors which may be associated with a physical material */
+UENUM(BlueprintType)
+namespace EPhysicalMaterialMaskColor
+{
+	enum Type
+	{
+		Red,
+		Green,
+		Blue,
+		Cyan,
+		Magenta,
+		Yellow,
+		White,
+		Black,
+		MAX
+	};
+}
+
 /** Describes how often this component is allowed to move. */
 UENUM(BlueprintType)
 namespace EComponentMobility
diff --git a/Engine/Source/Runtime/Engine/Classes/Engine/StaticMesh.h b/Engine/Source/Runtime/Engine/Classes/Engine/StaticMesh.h
--- a/Engine/Source/Runtime/Engine/Classes/Engine/StaticMesh.h
+++ b/Engine/Source/Runtime/Engine/Classes/Engine/StaticMesh.h
@@ -664,6 +664,14 @@ class UStaticMesh : public UStreamableRenderAsset, public IInterface_CollisionDa
 	UPROPERTY(EditAnywhere, AdvancedDisplay, Category = StaticMesh)
 	uint8 bSupportUniformlyDistributedSampling : 1;
 
+	/** 
+		If true, complex collision data will store UVs and face remap table for use when performing
+	    PhysicalMaterialMask lookups in cooked builds. Note the increased memory cost for this
+		functionality.
+	*/
+	UPROPERTY(EditAnywhere, AdvancedDisplay, Category = StaticMesh)
+	uint8 bSupportPhysicalMaterialMasks : 1;
+
 	/**
 	 * If true, StaticMesh has been built at runtime
 	 */
diff --git a/Engine/Source/Runtime/Engine/Classes/Materials/Material.h b/Engine/Source/Runtime/Engine/Classes/Materials/Material.h
--- a/Engine/Source/Runtime/Engine/Classes/Materials/Material.h
+++ b/Engine/Source/Runtime/Engine/Classes/Materials/Material.h
@@ -15,12 +15,15 @@
 #include "Materials/MaterialFunction.h"
 #include "Materials/MaterialLayersFunctions.h"
 #include "Templates/UniquePtr.h"
-
+#if WITH_CHAOS
+#include "Physics/PhysicsInterfaceCore.h"
+#endif
 #include "Material.generated.h"
 
 class ITargetPlatform;
 class UMaterialExpressionComment;
 class UPhysicalMaterial;
+class UPhysicalMaterialMask;
 class USubsurfaceProfile;
 class UTexture;
 
@@ -344,6 +347,14 @@ class ENGINE_VTABLE UMaterial : public UMaterialInterface
 	UPROPERTY(EditAnywhere, Category=PhysicalMaterial)
 	class UPhysicalMaterial* PhysMaterial;
 
+	/** Physical material mask to use for this graphics material. Used for sounds, effects etc.*/
+	UPROPERTY(EditAnywhere, Category = PhysicalMaterial)
+	class UPhysicalMaterialMask* PhysMaterialMask;
+
+	/** Physical material mask map to use for this graphics material. Used for sounds, effects etc.*/
+	UPROPERTY(EditAnywhere, Category = PhysicalMaterialMask)
+	class UPhysicalMaterial* PhysicalMaterialMap[EPhysicalMaterialMaskColor::MAX];
+
 	// Reflection.
 #if WITH_EDITORONLY_DATA
 	UPROPERTY()
@@ -1016,6 +1027,8 @@ public:
 	ENGINE_API virtual bool GetRefractionSettings(float& OutBiasValue) const override;
 	ENGINE_API virtual FMaterialRenderProxy* GetRenderProxy() const override;
 	ENGINE_API virtual UPhysicalMaterial* GetPhysicalMaterial() const override;
+	ENGINE_API virtual UPhysicalMaterialMask* GetPhysicalMaterialMask() const override;
+	ENGINE_API virtual UPhysicalMaterial* GetPhysicalMaterialFromMap(int32 Index) const override;
 	ENGINE_API virtual void GetUsedTextures(TArray<UTexture*>& OutTextures, EMaterialQualityLevel::Type QualityLevel, bool bAllQualityLevels, ERHIFeatureLevel::Type FeatureLevel, bool bAllFeatureLevels) const override;
 	ENGINE_API virtual void GetUsedTexturesAndIndices(TArray<UTexture*>& OutTextures, TArray< TArray<int32> >& OutIndices, EMaterialQualityLevel::Type QualityLevel, ERHIFeatureLevel::Type FeatureLevel) const override;
 	ENGINE_API virtual void OverrideTexture(const UTexture* InTextureToOverride, UTexture* OverrideTexture, ERHIFeatureLevel::Type InFeatureLevel) override;
diff --git a/Engine/Source/Runtime/Engine/Classes/Materials/MaterialInstance.h b/Engine/Source/Runtime/Engine/Classes/Materials/MaterialInstance.h
--- a/Engine/Source/Runtime/Engine/Classes/Materials/MaterialInstance.h
+++ b/Engine/Source/Runtime/Engine/Classes/Materials/MaterialInstance.h
@@ -13,6 +13,9 @@
 #include "Materials/Material.h"
 #include "Materials/MaterialInstanceBasePropertyOverrides.h"
 #include "Misc/App.h"
+#if WITH_CHAOS
+#include "Physics/PhysicsInterfaceCore.h"
+#endif
 #include "MaterialInstance.generated.h"
 
 class ITargetPlatform;
@@ -296,6 +299,10 @@ class ENGINE_VTABLE UMaterialInstance : public UMaterialInterface
 	UPROPERTY(EditAnywhere, BlueprintReadWrite, Category=MaterialInstance)
 	class UPhysicalMaterial* PhysMaterial;
 
+	/** Physical material map used with physical material mask, when it exists.*/
+	UPROPERTY(EditAnywhere, Category = PhysicalMaterialMask)
+	class UPhysicalMaterial* PhysicalMaterialMap[EPhysicalMaterialMaskColor::MAX];
+
 	/** Parent material. */
 	UPROPERTY(EditAnywhere, BlueprintReadOnly, Category=MaterialInstance, AssetRegistrySearchable)
 	class UMaterialInterface* Parent;
@@ -491,6 +498,8 @@ public:
 	virtual ENGINE_API bool IsDependent(UMaterialInterface* TestDependency) override;
 	virtual ENGINE_API FMaterialRenderProxy* GetRenderProxy() const override;
 	virtual ENGINE_API UPhysicalMaterial* GetPhysicalMaterial() const override;
+	virtual ENGINE_API UPhysicalMaterialMask* GetPhysicalMaterialMask() const override;
+	virtual ENGINE_API UPhysicalMaterial* GetPhysicalMaterialFromMap(int32 Index) const override;
 	virtual ENGINE_API bool UpdateLightmassTextureTracking() override;
 	virtual ENGINE_API bool GetCastShadowAsMasked() const override;
 	virtual ENGINE_API float GetEmissiveBoost() const override;
diff --git a/Engine/Source/Runtime/Engine/Classes/Materials/MaterialInstanceConstant.h b/Engine/Source/Runtime/Engine/Classes/Materials/MaterialInstanceConstant.h
--- a/Engine/Source/Runtime/Engine/Classes/Materials/MaterialInstanceConstant.h
+++ b/Engine/Source/Runtime/Engine/Classes/Materials/MaterialInstanceConstant.h
@@ -9,6 +9,8 @@
 
 #include "MaterialInstanceConstant.generated.h"
 
+class UPhysicalMaterialMask;
+
 /**
  * Material Instances may be used to change the appearance of a material without incurring an expensive recompilation of the material.
  * General modification of the material cannot be supported without recompilation, so the instances are limited to changing the values of
@@ -35,6 +37,14 @@ class ENGINE_VTABLE UMaterialInstanceConstant : public UMaterialInstance
 	virtual ENGINE_API void PostEditChangeProperty(FPropertyChangedEvent& PropertyChangedEvent) override;
 #endif
 
+	/** Physical material mask to use for this graphics material. Used for sounds, effects etc.*/
+	UPROPERTY(EditAnywhere, Category = PhysicalMaterial)
+	class UPhysicalMaterialMask* PhysMaterialMask;
+
+	// Begin UMaterialInterface interface.
+	ENGINE_API virtual UPhysicalMaterialMask* GetPhysicalMaterialMask() const override;
+	// End UMaterialInterface interface.
+
 	/** Get the scalar (float) parameter value from an MIC */
 	UFUNCTION(BlueprintCallable, meta=(DisplayName = "GetScalarParameterValue", ScriptName = "GetScalarParameterValue", Keywords = "GetFloatParameterValue"), Category="Rendering|Material")
 	float K2_GetScalarParameterValue(FName ParameterName);
diff --git a/Engine/Source/Runtime/Engine/Classes/Materials/MaterialInterface.h b/Engine/Source/Runtime/Engine/Classes/Materials/MaterialInterface.h
--- a/Engine/Source/Runtime/Engine/Classes/Materials/MaterialInterface.h
+++ b/Engine/Source/Runtime/Engine/Classes/Materials/MaterialInterface.h
@@ -17,6 +17,9 @@
 #include "Materials/MaterialLayersFunctions.h"
 #include "Interfaces/Interface_AssetUserData.h"
 #include "MaterialSceneTextureId.h"
+#if WITH_CHAOS
+#include "Physics/PhysicsInterfaceCore.h"
+#endif
 #include "MaterialInterface.generated.h"
 
 class FMaterialCompiler;
@@ -24,6 +27,7 @@ class FMaterialRenderProxy;
 class FMaterialResource;
 class UMaterial;
 class UPhysicalMaterial;
+class UPhysicalMaterialMask;
 class USubsurfaceProfile;
 class UTexture;
 struct FPrimitiveViewRelevance;
@@ -364,6 +368,20 @@ public:
 	UFUNCTION(BlueprintCallable, Category = "Physics|Material")
 	virtual UPhysicalMaterial* GetPhysicalMaterial() const PURE_VIRTUAL(UMaterialInterface::GetPhysicalMaterial,return NULL;);
 
+	/**
+	 * Return a pointer to the physical material mask used by this material instance.
+	 * @return The physical material.
+	 */
+	UFUNCTION(BlueprintCallable, Category = "Physics|Material")
+	virtual UPhysicalMaterialMask* GetPhysicalMaterialMask() const PURE_VIRTUAL(UMaterialInterface::GetPhysicalMaterialMask, return nullptr;);
+
+	/**
+	 * Return a pointer to the physical material from mask map at given index.
+	 * @return The physical material.
+	 */
+	UFUNCTION(BlueprintCallable, Category = "Physics|Material")
+	virtual UPhysicalMaterial* GetPhysicalMaterialFromMap(int32 Index) const PURE_VIRTUAL(UMaterialInterface::GetPhysicalMaterialFromMap, return nullptr;);
+
 	/** Return the textures used to render this material. */
 	virtual void GetUsedTextures(TArray<UTexture*>& OutTextures, EMaterialQualityLevel::Type QualityLevel, bool bAllQualityLevels, ERHIFeatureLevel::Type FeatureLevel, bool bAllFeatureLevels) const
 		PURE_VIRTUAL(UMaterialInterface::GetUsedTextures,);
diff --git a/Engine/Source/Runtime/Engine/Classes/PhysicalMaterials/PhysicalMaterialMask.h b/Engine/Source/Runtime/Engine/Classes/PhysicalMaterials/PhysicalMaterialMask.h
new file mode 100644
--- /dev/null
+++ b/Engine/Source/Runtime/Engine/Classes/PhysicalMaterials/PhysicalMaterialMask.h
@@ -0,0 +1,88 @@
+// Copyright Epic Games, Inc. All Rights Reserved.
+
+#pragma once
+
+#include "CoreMinimal.h"
+#include "UObject/ObjectMacros.h"
+#include "UObject/Object.h"
+#include "Engine/EngineTypes.h"
+#include "Engine/Texture.h"
+#include "EngineDefines.h"
+#include "PhysicsEngine/PhysicsSettingsEnums.h"
+#if WITH_CHAOS
+#include "Physics/PhysicsInterfaceCore.h"
+#endif
+#include "PhysicalMaterialMask.generated.h"
+
+struct FPropertyChangedEvent;
+
+/**
+ * Physical material masks are used to map multiple physical materials to a single rendering material
+ */
+UCLASS(BlueprintType, Blueprintable, CollapseCategories, HideCategories = Object)
+class ENGINE_API UPhysicalMaterialMask : public UObject
+{
+	GENERATED_UCLASS_BODY()
+
+	//
+	// Object properties.
+	//
+#if WITH_EDITORONLY_DATA
+
+	UPROPERTY(VisibleAnywhere, Instanced, Category = ImportSettings)
+	class UAssetImportData* AssetImportData;
+
+	/** Mask input texture, square aspect ratio recommended. Recognized mask colors include: white, black, red, green, yellow, cyan, turquoise, and magenta. */
+	UPROPERTY(VisibleAnywhere, BlueprintReadOnly, Category = TextureSource)
+	UTexture* MaskTexture;
+
+#endif
+
+	/** StaticMesh UV channel index to use when performing lookups with this mask. */
+	UPROPERTY(EditAnywhere, BlueprintReadWrite, Category = TextureSource)
+	int32 UVChannelIndex;
+
+	/** The addressing mode to use for the X axis. */
+	UPROPERTY(EditAnywhere, BlueprintReadWrite, meta = (DisplayName = "X-axis Tiling Method"), Category = TextureSource)
+	TEnumAsByte<enum TextureAddress> AddressX;
+
+	/** The addressing mode to use for the Y axis. */
+	UPROPERTY(EditAnywhere, BlueprintReadWrite, meta = (DisplayName = "Y-axis Tiling Method"), Category = TextureSource)
+	TEnumAsByte<enum TextureAddress> AddressY;
+
+	static uint32 INVALID_MASK_INDEX;
+
+public:
+
+
+	//~ Begin UObject Interface
+#if WITH_EDITOR
+	virtual void PostEditChangeProperty(FPropertyChangedEvent& PropertyChangedEvent) override;
+#endif // WITH_EDITOR
+	virtual void PostInitProperties() override;
+	virtual void PostLoad() override;
+	virtual void FinishDestroy() override;
+	//~ End UObject Interface
+
+#if WITH_CHAOS
+	FPhysicsMaterialMaskHandle MaterialMaskHandle;
+
+	/** Get the physics-interface derived version of this material */
+	FPhysicsMaterialMaskHandle& GetPhysicsMaterialMask();
+#endif
+
+	// Helper method to query phys mat index at a given UV position in mask data
+	static uint32 GetPhysMatIndex(const TArray<uint32>& MaskData, int32 SizeX, int32 SizeY, int32 AddressX, int32 AddressY, float U, float V);
+
+#if WITH_EDITOR
+	// Helper method to set mask texture
+	void SetMaskTexture(UTexture* InMaskTexture, const FString& InTextureFilename);
+
+	// Dump mask data to log display.
+	void DumpMaskData();
+#endif // WITH_EDITOR
+
+	// Helper method to generate mask data used at runtime based on mask texture
+	void GenerateMaskData(TArray<uint32>& OutMaskData, int32& OutSizeX, int32& OutSizeY) const;
+};
+
diff --git a/Engine/Source/Runtime/Engine/Classes/PhysicsEngine/BodyInstance.h b/Engine/Source/Runtime/Engine/Classes/PhysicsEngine/BodyInstance.h
--- a/Engine/Source/Runtime/Engine/Classes/PhysicsEngine/BodyInstance.h
+++ b/Engine/Source/Runtime/Engine/Classes/PhysicsEngine/BodyInstance.h
@@ -16,6 +16,7 @@
 
 class UBodySetup;
 class UPhysicalMaterial;
+class UPhysicalMaterialMask;
 class UPrimitiveComponent;
 struct FBodyInstance;
 struct FCollisionNotifyInfo;
@@ -588,14 +589,20 @@ public:
 	/** Find the correct PhysicalMaterial for simple geometry on a given body and owner. This is really for internal use during serialization */
 	static UPhysicalMaterial* GetSimplePhysicalMaterial(const FBodyInstance* BodyInstance, TWeakObjectPtr<UPrimitiveComponent> Owner, TWeakObjectPtr<UBodySetup> BodySetupPtr);
 
-	/** Get the complex PhysicalMaterial array for this body */
+	/** Get the complex PhysicalMaterials array for this body */
 	TArray<UPhysicalMaterial*> GetComplexPhysicalMaterials() const;
 
-	/** Find the correct PhysicalMaterial for simple geometry on a given body and owner. This is really for internal use during serialization */
-	static void GetComplexPhysicalMaterials(const FBodyInstance* BodyInstance, TWeakObjectPtr<UPrimitiveComponent> Owner, TArray<UPhysicalMaterial*>& OutPhysicalMaterials);
+	/** Get the complex PhysicalMaterials and PhysicalMaterialMasks array for this body */
+	TArray<UPhysicalMaterial*> GetComplexPhysicalMaterials(TArray<FPhysicalMaterialMaskParams>& OutPhysMaterialMasks) const;
 
 	/** Get the complex PhysicalMaterials for this body */
-	void GetComplexPhysicalMaterials(TArray<UPhysicalMaterial*> &PhysMaterials) const;
+	void GetComplexPhysicalMaterials(TArray<UPhysicalMaterial*> &OutPhysMaterials) const;
+
+	/** Get the complex PhysicalMaterials and PhysicalMaterialMasks for this body */
+	void GetComplexPhysicalMaterials(TArray<UPhysicalMaterial*> &OutPhysMaterials, TArray<FPhysicalMaterialMaskParams>& OutPhysMaterialMasks) const;
+
+	/** Find the correct PhysicalMaterial and PhysicalMaterialMasks for complex geometry on a given body and owner. This is really for internal use during serialization */
+	static void GetComplexPhysicalMaterials(const FBodyInstance* BodyInstance, TWeakObjectPtr<UPrimitiveComponent> Owner, TArray<UPhysicalMaterial*>& OutPhysMaterials, TArray<FPhysicalMaterialMaskParams>* OutPhysMaterialMasks = nullptr);
 
 	/** Returns the slope override struct for this instance. If we don't have our own custom setting, it will return the setting from the body setup. */
 	const struct FWalkableSlopeOverride& GetWalkableSlopeOverride() const;
@@ -862,10 +869,10 @@ public:
 	 *  @param  SimplePhysMat			The material to use if a simple shape is provided (or complex materials are empty)
 	 *  @param  ComplexPhysMats			The array of materials to apply if a complex shape is provided
 	 */
-	static void ApplyMaterialToShape_AssumesLocked(const FPhysicsShapeHandle& InShape, UPhysicalMaterial* SimplePhysMat, const TArrayView<UPhysicalMaterial*>& ComplexPhysMats);
+	static void ApplyMaterialToShape_AssumesLocked(const FPhysicsShapeHandle& InShape, UPhysicalMaterial* SimplePhysMat, const TArrayView<UPhysicalMaterial*>& ComplexPhysMats, const TArrayView<FPhysicalMaterialMaskParams>* ComplexPhysMatMasks = nullptr);
 
 	/** Note: This function is not thread safe. Make sure you obtain the appropriate physics scene lock before calling it*/
-	void ApplyMaterialToInstanceShapes_AssumesLocked(UPhysicalMaterial* SimplePhysMat, TArray<UPhysicalMaterial*>& ComplexPhysMats);
+	void ApplyMaterialToInstanceShapes_AssumesLocked(UPhysicalMaterial* SimplePhysMat, TArray<UPhysicalMaterial*>& ComplexPhysMats, const TArrayView<FPhysicalMaterialMaskParams>& ComplexPhysMatMasks);
 
 	/** Update the instances collision filtering data */ 
 	void UpdatePhysicsFilterData();
diff --git a/Engine/Source/Runtime/Engine/Classes/PhysicsEngine/BodySetup.h b/Engine/Source/Runtime/Engine/Classes/PhysicsEngine/BodySetup.h
--- a/Engine/Source/Runtime/Engine/Classes/PhysicsEngine/BodySetup.h
+++ b/Engine/Source/Runtime/Engine/Classes/PhysicsEngine/BodySetup.h
@@ -19,6 +19,7 @@
 
 class ITargetPlatform;
 class UPhysicalMaterial;
+class UPhysicalMaterialMask;
 class UPrimitiveComponent;
 struct FShapeData;
 enum class EPhysXMeshCookFlags : uint8;
@@ -204,6 +205,13 @@ class UBodySetup : public UObject
 	UPROPERTY()
 	uint8 bGenerateMirroredCollision:1;
 
+	/** 
+	 * If true, the physics triangle mesh will store UVs and the face remap table. This is needed
+	 * to support physical material masks in scene queries. 
+	 */
+	UPROPERTY()
+	uint8 bSupportUVsAndFaceRemap:1;
+
 	/** Flag used to know if we have created the physics convex and tri meshes from the cooked data yet */
 	uint8 bCreatedPhysicsMeshes:1;
 
@@ -458,7 +466,8 @@ public:
 		FBodyInstance* OwningInstance, 
 		FVector& Scale3D, 
 		UPhysicalMaterial* SimpleMaterial,
-		TArray<UPhysicalMaterial*>& ComplexMaterials, 
+		TArray<UPhysicalMaterial*>& ComplexMaterials,
+		TArray<FPhysicalMaterialMaskParams>& ComplexMaterialMasks,
 		const FBodyCollisionData& BodyCollisionData,
 		const FTransform& RelativeTM = FTransform::Identity, 
 		TArray<FPhysicsShapeHandle>* NewShapes = NULL);
diff --git a/Engine/Source/Runtime/Engine/Private/Collision/CollisionConversions.cpp b/Engine/Source/Runtime/Engine/Private/Collision/CollisionConversions.cpp
--- a/Engine/Source/Runtime/Engine/Private/Collision/CollisionConversions.cpp
+++ b/Engine/Source/Runtime/Engine/Private/Collision/CollisionConversions.cpp
@@ -7,6 +7,7 @@
 #include "Collision/CollisionDebugDrawing.h"
 #include "Components/LineBatchComponent.h"
 #include "PhysicalMaterials/PhysicalMaterial.h"
+#include "PhysicalMaterials/PhysicalMaterialMask.h"
 #include "PhysicsEngine/PhysicsSettings.h"
 #include "PhysicsEngine/BodySetup.h"
 #include "Physics/PhysicsInterfaceUtils.h"
@@ -130,7 +131,7 @@ static FVector FindGeomOpposingNormal(ECollisionShapeType QueryGeomType, const F
 }
 
 /** Set info in the HitResult (Actor, Component, PhysMaterial, BoneName, Item) based on the supplied shape and face index */
-static void SetHitResultFromShapeAndFaceIndex(const FPhysicsShape& Shape,  const FPhysicsActor& Actor, const uint32 FaceIndex, FHitResult& OutResult, bool bReturnPhysMat)
+static void SetHitResultFromShapeAndFaceIndex(const FPhysicsShape& Shape,  const FPhysicsActor& Actor, const uint32 FaceIndex, const FVector& HitLocation, FHitResult& OutResult, bool bReturnPhysMat)
 {
 	SCOPE_CYCLE_COUNTER(STAT_CollisionSetHitResultFromShapeAndFaceIndex);
 	
@@ -183,11 +184,18 @@ static void SetHitResultFromShapeAndFaceIndex(const FPhysicsShape& Shape,  const
 
 		if (bReturnPhysMat)
 		{
-			// This function returns the single material in all cases other than trimesh or heightfield
-			if(const FPhysicsMaterial* PhysicsMaterial = GetMaterialFromInternalFaceIndex(Shape, Actor, FaceIndex))
+#if WITH_CHAOS
+			if (const FPhysicsMaterial* PhysicsMaterial = GetMaterialFromInternalFaceIndexAndHitLocation(Shape, Actor, FaceIndex, HitLocation))
 			{
 				OutResult.PhysMaterial = GetUserData(*PhysicsMaterial);
 			}
+#else
+			if (const FPhysicsMaterial* PhysicsMaterial = GetMaterialFromInternalFaceIndex(Shape, Actor, FaceIndex))
+			{
+				OutResult.PhysMaterial = GetUserData(*PhysicsMaterial);
+			}
+
+#endif
 		}
 	}
 
@@ -255,7 +263,7 @@ EConvertQueryResult ConvertQueryImpactHit(const UWorld* World, const FHitLocatio
 		if (Position.ContainsNaN())
 		{
 #if ENABLE_NAN_DIAGNOSTIC
-			SetHitResultFromShapeAndFaceIndex(HitShape, HitActor, InternalFaceIndex, OutResult, bReturnPhysMat);
+			SetHitResultFromShapeAndFaceIndex(HitShape, HitActor, InternalFaceIndex, OutResult.ImpactPoint, OutResult, bReturnPhysMat);
 			UE_LOG(LogCore, Error, TEXT("ConvertQueryImpactHit() NaN details:\n>> Actor:%s (%s)\n>> Component:%s\n>> Item:%d\n>> BoneName:%s\n>> Time:%f\n>> Distance:%f\n>> Location:%s\n>> bIsBlocking:%d\n>> bStartPenetrating:%d"),
 				*GetNameSafe(OutResult.GetActor()), OutResult.Actor.IsValid() ? *OutResult.GetActor()->GetPathName() : TEXT("no path"),
 				*GetNameSafe(OutResult.GetComponent()), OutResult.Item, *OutResult.BoneName.ToString(),
@@ -275,7 +283,7 @@ EConvertQueryResult ConvertQueryImpactHit(const UWorld* World, const FHitLocatio
 	if (bUseReturnedNormal && HitNormal.ContainsNaN())
 	{
 #if ENABLE_NAN_DIAGNOSTIC
-		SetHitResultFromShapeAndFaceIndex(HitShape, HitActor, InternalFaceIndex, OutResult, bReturnPhysMat);
+		SetHitResultFromShapeAndFaceIndex(HitShape, HitActor, InternalFaceIndex, OutResult.ImpactPoint, OutResult, bReturnPhysMat);
 		UE_LOG(LogCore, Error, TEXT("ConvertQueryImpactHit() NaN details:\n>> Actor:%s (%s)\n>> Component:%s\n>> Item:%d\n>> BoneName:%s\n>> Time:%f\n>> Distance:%f\n>> Location:%s\n>> bIsBlocking:%d\n>> bStartPenetrating:%d"),
 			*GetNameSafe(OutResult.GetActor()), OutResult.Actor.IsValid() ? *OutResult.GetActor()->GetPathName() : TEXT("no path"),
 			*GetNameSafe(OutResult.GetComponent()), OutResult.Item, *OutResult.BoneName.ToString(),
@@ -311,7 +319,7 @@ EConvertQueryResult ConvertQueryImpactHit(const UWorld* World, const FHitLocatio
 	OutResult.ImpactNormal = FindGeomOpposingNormal(SweptGeometryType, Hit, TraceStartToEnd, Normal);
 
 	// Fill in Actor, Component, material, etc.
-	SetHitResultFromShapeAndFaceIndex(HitShape, HitActor, InternalFaceIndex, OutResult, bReturnPhysMat);
+	SetHitResultFromShapeAndFaceIndex(HitShape, HitActor, InternalFaceIndex, OutResult.ImpactPoint, OutResult, bReturnPhysMat);
 
 	ECollisionShapeType GeomType = GetGeometryType(HitShape);
 
@@ -498,7 +506,7 @@ static bool ConvertOverlappedShapeToImpactHit(const UWorld* World, const FHitLoc
 
 	OutResult.Normal = OutResult.ImpactNormal;
 	
-	SetHitResultFromShapeAndFaceIndex(HitShape, HitActor, GetInternalFaceIndex(Hit), OutResult, bReturnPhysMat);
+	SetHitResultFromShapeAndFaceIndex(HitShape, HitActor, GetInternalFaceIndex(Hit), OutResult.ImpactPoint, OutResult, bReturnPhysMat);
 
 	return bBlockingHit;
 }
diff --git a/Engine/Source/Runtime/Engine/Private/Materials/Material.cpp b/Engine/Source/Runtime/Engine/Private/Materials/Material.cpp
--- a/Engine/Source/Runtime/Engine/Private/Materials/Material.cpp
+++ b/Engine/Source/Runtime/Engine/Private/Materials/Material.cpp
@@ -948,6 +948,9 @@ UMaterial::UMaterial(const FObjectInitializer& ObjectInitializer)
 
 	bIsPreviewMaterial = false;
 	bIsFunctionPreviewMaterial = false;
+
+	PhysMaterial = nullptr;
+	PhysMaterialMask = nullptr;
 }
 
 void UMaterial::PreSave(const class ITargetPlatform* TargetPlatform)
@@ -2963,6 +2966,20 @@ UPhysicalMaterial* UMaterial::GetPhysicalMaterial() const
 	return nullptr;
 }
 
+UPhysicalMaterialMask* UMaterial::GetPhysicalMaterialMask() const
+{
+	return PhysMaterialMask;
+}
+
+UPhysicalMaterial* UMaterial::GetPhysicalMaterialFromMap(int32 Index) const
+{
+	if (Index >= 0 && Index < EPhysicalMaterialMaskColor::MAX)
+	{
+		return PhysicalMaterialMap[Index];
+	}
+	return nullptr;
+}
+
 /** Helper functions for text output of properties... */
 #ifndef CASE_ENUM_TO_TEXT
 #define CASE_ENUM_TO_TEXT(txt) case txt: return TEXT(#txt);
@@ -4236,7 +4253,7 @@ bool UMaterial::CanEditChange(const FProperty* InProperty) const
 	{
 		FString PropertyName = InProperty->GetName();
 
-		if (PropertyName == GET_MEMBER_NAME_STRING_CHECKED(UMaterial, PhysMaterial))
+		if (PropertyName == GET_MEMBER_NAME_STRING_CHECKED(UMaterial, PhysMaterial) || PropertyName == GET_MEMBER_NAME_STRING_CHECKED(UMaterial, PhysMaterialMask))
 		{
 			return MaterialDomain == MD_Surface;
 		}
@@ -4428,7 +4445,7 @@ void UMaterial::PostEditChangeProperty(FPropertyChangedEvent& PropertyChangedEve
 	if( PropertyThatChanged ) 
 	{
 		// Don't recompile the material if we only changed the PhysMaterial property.
-		if (PropertyThatChanged->GetName() == TEXT("PhysMaterial"))
+		if (PropertyThatChanged->GetName() == TEXT("PhysMaterial") || PropertyThatChanged->GetName() == TEXT("PhysMaterialMask") || PropertyThatChanged->GetName() == TEXT("PhysicalMaterialMap"))
 		{
 			bRequiresCompilation = false;
 		}
diff --git a/Engine/Source/Runtime/Engine/Private/Materials/MaterialInstance.cpp b/Engine/Source/Runtime/Engine/Private/Materials/MaterialInstance.cpp
--- a/Engine/Source/Runtime/Engine/Private/Materials/MaterialInstance.cpp
+++ b/Engine/Source/Runtime/Engine/Private/Materials/MaterialInstance.cpp
@@ -430,6 +430,12 @@ UMaterialInstance::UMaterialInstance(const FObjectInitializer& ObjectInitializer
 	ReentrantFlag[0] = false;
 	ReentrantFlag[1] = false;
 	ShadingModels = MSM_Unlit;
+
+	PhysMaterial = nullptr;
+	for (UPhysicalMaterial*& PhysMat : PhysicalMaterialMap)
+	{
+		PhysMat = nullptr;
+	}
 }
 
 void UMaterialInstance::PostInitProperties()	
@@ -1955,6 +1961,20 @@ UPhysicalMaterial* UMaterialInstance::GetPhysicalMaterial() const
 	}
 }
 
+UPhysicalMaterialMask* UMaterialInstance::GetPhysicalMaterialMask() const
+{
+	return nullptr;
+}
+
+UPhysicalMaterial* UMaterialInstance::GetPhysicalMaterialFromMap(int32 Index) const
+{
+	if (Index < 0 || Index >= EPhysicalMaterialMaskColor::MAX)
+	{
+		return nullptr;
+	}
+	return PhysicalMaterialMap[Index];
+}
+
 void UMaterialInstance::GetStaticParameterValues(FStaticParameterSet& OutStaticParameters)
 {
 	check(IsInGameThread());
diff --git a/Engine/Source/Runtime/Engine/Private/Materials/MaterialInstanceConstant.cpp b/Engine/Source/Runtime/Engine/Private/Materials/MaterialInstanceConstant.cpp
--- a/Engine/Source/Runtime/Engine/Private/Materials/MaterialInstanceConstant.cpp
+++ b/Engine/Source/Runtime/Engine/Private/Materials/MaterialInstanceConstant.cpp
@@ -12,6 +12,7 @@
 UMaterialInstanceConstant::UMaterialInstanceConstant(const FObjectInitializer& ObjectInitializer)
 	: Super(ObjectInitializer)
 {
+	PhysMaterialMask = nullptr;
 }
 
 void UMaterialInstanceConstant::PostLoad()
@@ -41,6 +42,11 @@ UTexture* UMaterialInstanceConstant::K2_GetTextureParameterValue(FName Parameter
 	return Result;
 }
 
+UPhysicalMaterialMask* UMaterialInstanceConstant::GetPhysicalMaterialMask() const
+{
+	return PhysMaterialMask;
+}
+
 #if WITH_EDITOR
 void UMaterialInstanceConstant::PostEditChangeProperty(FPropertyChangedEvent& PropertyChangedEvent)
 {
diff --git a/Engine/Source/Runtime/Engine/Private/PhysicsEngine/BodyInstance.cpp b/Engine/Source/Runtime/Engine/Private/PhysicsEngine/BodyInstance.cpp
--- a/Engine/Source/Runtime/Engine/Private/PhysicsEngine/BodyInstance.cpp
+++ b/Engine/Source/Runtime/Engine/Private/PhysicsEngine/BodyInstance.cpp
@@ -469,11 +469,12 @@ void FBodyInstance::UpdatePhysicalMaterials()
 {
 	SCOPE_CYCLE_COUNTER(STAT_UpdatePhysMats);
 	UPhysicalMaterial* SimplePhysMat = GetSimplePhysicalMaterial();
-	TArray<UPhysicalMaterial*> ComplexPhysMats = GetComplexPhysicalMaterials();
+	TArray<FPhysicalMaterialMaskParams> ComplexPhysMatMasks;
+	TArray<UPhysicalMaterial*> ComplexPhysMats = GetComplexPhysicalMaterials(ComplexPhysMatMasks);
 
 	FPhysicsCommand::ExecuteWrite(GetActorReferenceWithWelding(), [&](const FPhysicsActorHandle& Actor)
 	{
-		ApplyMaterialToInstanceShapes_AssumesLocked(SimplePhysMat, ComplexPhysMats);
+		ApplyMaterialToInstanceShapes_AssumesLocked(SimplePhysMat, ComplexPhysMats, ComplexPhysMatMasks);
 	});
 }
 
@@ -1059,7 +1060,15 @@ struct FInitBodiesHelper
 	bool CreateShapes_AssumesLocked(FBodyInstance* Instance) const
 	{
 		UPhysicalMaterial* SimplePhysMat = Instance->GetSimplePhysicalMaterial();
-		TArray<UPhysicalMaterial*> ComplexPhysMats = Instance->GetComplexPhysicalMaterials();
+
+		TArray<UPhysicalMaterial*> ComplexPhysMats;
+		TArray<FPhysicalMaterialMaskParams> ComplexPhysMatMasks;
+
+#if WITH_CHAOS
+		ComplexPhysMats = Instance->GetComplexPhysicalMaterials(ComplexPhysMatMasks);
+#else
+		ComplexPhysMats = Instance->GetComplexPhysicalMaterials();
+#endif
 
 		FBodyCollisionData BodyCollisionData;
 		Instance->BuildBodyFilterData(BodyCollisionData.CollisionFilterData);
@@ -1068,7 +1077,7 @@ struct FInitBodiesHelper
 		bool bInitFail = false;
 
 		// #PHYS2 Call interface AddGeometry
-		BodySetup->AddShapesToRigidActor_AssumesLocked(Instance, Instance->Scale3D, SimplePhysMat, ComplexPhysMats, BodyCollisionData, FTransform::Identity);
+		BodySetup->AddShapesToRigidActor_AssumesLocked(Instance, Instance->Scale3D, SimplePhysMat, ComplexPhysMats, ComplexPhysMatMasks, BodyCollisionData, FTransform::Identity);
 
 		const int32 NumShapes = FPhysicsInterface::GetNumShapes(Instance->ActorHandle);
 		bInitFail |= NumShapes == 0;
@@ -1471,14 +1480,21 @@ bool FBodyInstance::Weld(FBodyInstance* TheirBody, const FTransform& TheirTM)
 		TheirBody->WeldParent = this;
 
 		UPhysicalMaterial* SimplePhysMat = TheirBody->GetSimplePhysicalMaterial();
-		TArray<UPhysicalMaterial*> ComplexPhysMats = TheirBody->GetComplexPhysicalMaterials();
 
+		TArray<UPhysicalMaterial*> ComplexPhysMats;
+		TArray<FPhysicalMaterialMaskParams> ComplexPhysMatMasks;
+	
+#if WITH_CHAOS
+		ComplexPhysMats = TheirBody->GetComplexPhysicalMaterials(ComplexPhysMatMasks);
+#else
+		ComplexPhysMats = TheirBody->GetComplexPhysicalMaterials();
+#endif
 		// This builds collision data based on this (parent) body, not their body. This gets fixed  up later though when PostShapeChange() calls UpdatePhysicsFilterData().
 		FBodyCollisionData BodyCollisionData;
 		BuildBodyFilterData(BodyCollisionData.CollisionFilterData);
 		BuildBodyCollisionFlags(BodyCollisionData.CollisionFlags, GetCollisionEnabled(), BodySetup->GetCollisionTraceFlag() == CTF_UseComplexAsSimple);
 
-		TheirBody->BodySetup->AddShapesToRigidActor_AssumesLocked(this, Scale3D, SimplePhysMat, ComplexPhysMats, BodyCollisionData, RelativeTM, &PNewShapes);
+		TheirBody->BodySetup->AddShapesToRigidActor_AssumesLocked(this, Scale3D, SimplePhysMat, ComplexPhysMats, ComplexPhysMatMasks, BodyCollisionData, RelativeTM, &PNewShapes);
 
 		FPhysicsInterface::SetSendsSleepNotifies_AssumesLocked(Actor, TheirBody->bGenerateWakeEvents);
 
@@ -2776,13 +2792,24 @@ TArray<UPhysicalMaterial*> FBodyInstance::GetComplexPhysicalMaterials() const
 	return PhysMaterials;
 }
 
-void FBodyInstance::GetComplexPhysicalMaterials(TArray<UPhysicalMaterial*>& PhysMaterials) const
+TArray<UPhysicalMaterial*> FBodyInstance::GetComplexPhysicalMaterials(TArray<FPhysicalMaterialMaskParams>& OutPhysMaterialMasks) const
+{
+	TArray<UPhysicalMaterial*> PhysMaterials;
+	GetComplexPhysicalMaterials(PhysMaterials, OutPhysMaterialMasks);
+	return PhysMaterials;
+}
+
+void FBodyInstance::GetComplexPhysicalMaterials(TArray<UPhysicalMaterial*>& OutPhysMaterials) const
 {
-	GetComplexPhysicalMaterials(this, OwnerComponent, PhysMaterials);
+	GetComplexPhysicalMaterials(this, OwnerComponent, OutPhysMaterials);
 }
 
+void FBodyInstance::GetComplexPhysicalMaterials(TArray<UPhysicalMaterial*>& OutPhysMaterials, TArray<FPhysicalMaterialMaskParams>& OutPhysMaterialMasks) const
+{
+	GetComplexPhysicalMaterials(this, OwnerComponent, OutPhysMaterials, &OutPhysMaterialMasks);
+}
 
-void FBodyInstance::GetComplexPhysicalMaterials(const FBodyInstance*, TWeakObjectPtr<UPrimitiveComponent> OwnerComp, TArray<UPhysicalMaterial*>& OutPhysicalMaterials)
+void FBodyInstance::GetComplexPhysicalMaterials(const FBodyInstance*, TWeakObjectPtr<UPrimitiveComponent> OwnerComp, TArray<UPhysicalMaterial*>& OutPhysMaterials, TArray<FPhysicalMaterialMaskParams>* OutPhysMaterialMasks)
 {
 	if(!GEngine || !GEngine->DefaultPhysMaterial)
 	{
@@ -2796,7 +2823,12 @@ void FBodyInstance::GetComplexPhysicalMaterials(const FBodyInstance*, TWeakObjec
 	if (PrimComp)
 	{
 		const int32 NumMaterials = PrimComp->GetNumMaterials();
-		OutPhysicalMaterials.SetNum(NumMaterials);
+		OutPhysMaterials.SetNum(NumMaterials);
+
+		if (OutPhysMaterialMasks)
+		{
+			OutPhysMaterialMasks->SetNum(NumMaterials);
+		}
 
 		for (int32 MatIdx = 0; MatIdx < NumMaterials; MatIdx++)
 		{
@@ -2806,9 +2838,26 @@ void FBodyInstance::GetComplexPhysicalMaterials(const FBodyInstance*, TWeakObjec
 			{
 				PhysMat = Material->GetPhysicalMaterial();
 			}
+			
+			OutPhysMaterials[MatIdx] = PhysMat;
 
-			check(PhysMat != NULL);
-			OutPhysicalMaterials[MatIdx] = PhysMat;
+			if (OutPhysMaterialMasks)
+			{
+				UPhysicalMaterialMask* PhysMatMask = nullptr;
+				UMaterialInterface* PhysMatMap = nullptr;
+
+				if (Material)
+				{
+					PhysMatMask = Material->GetPhysicalMaterialMask();
+					if (PhysMatMask)
+					{
+						PhysMatMap = Material;
+					}
+				}
+
+				(*OutPhysMaterialMasks)[MatIdx].PhysicalMaterialMask = PhysMatMask;
+				(*OutPhysMaterialMasks)[MatIdx].PhysicalMaterialMap = PhysMatMap;
+			}
 		}
 	}
 }
@@ -3719,7 +3768,7 @@ void FBodyInstance::FixupData(class UObject* Loader)
 	}
 }
 
-void FBodyInstance::ApplyMaterialToShape_AssumesLocked(const FPhysicsShapeHandle& InShape, UPhysicalMaterial* SimplePhysMat, const TArrayView<UPhysicalMaterial*>& ComplexPhysMats)
+void FBodyInstance::ApplyMaterialToShape_AssumesLocked(const FPhysicsShapeHandle& InShape, UPhysicalMaterial* SimplePhysMat, const TArrayView<UPhysicalMaterial*>& ComplexPhysMats, const TArrayView<FPhysicalMaterialMaskParams>* ComplexPhysMatMasks)
 {
 	// If a triangle mesh, need to get array of materials...
 	ECollisionShapeType GeomType = FPhysicsInterface::GetShapeType(InShape);
@@ -3727,7 +3776,18 @@ void FBodyInstance::ApplyMaterialToShape_AssumesLocked(const FPhysicsShapeHandle
 	{
 		if(ComplexPhysMats.Num())
 		{
+#if WITH_CHAOS
+			if (ensure(ComplexPhysMatMasks))
+			{
+				FPhysicsInterface::SetMaterials(InShape, ComplexPhysMats, *ComplexPhysMatMasks);
+			}
+			else
+			{
+				FPhysicsInterface::SetMaterials(InShape, ComplexPhysMats);
+			}
+#else
 			FPhysicsInterface::SetMaterials(InShape, ComplexPhysMats);
+#endif
 		}
 		else
 		{
@@ -3754,7 +3814,7 @@ void FBodyInstance::ApplyMaterialToShape_AssumesLocked(const FPhysicsShapeHandle
 	}
 }
 
-void FBodyInstance::ApplyMaterialToInstanceShapes_AssumesLocked(UPhysicalMaterial* SimplePhysMat, TArray<UPhysicalMaterial*>& ComplexPhysMats)
+void FBodyInstance::ApplyMaterialToInstanceShapes_AssumesLocked(UPhysicalMaterial* SimplePhysMat, TArray<UPhysicalMaterial*>& ComplexPhysMats, const TArrayView<FPhysicalMaterialMaskParams>& ComplexPhysMatMasks)
 {
 	FBodyInstance* TheirBI = this;
 	FBodyInstance* BIWithActor = TheirBI->WeldParent ? TheirBI->WeldParent : TheirBI;
@@ -3767,12 +3827,12 @@ void FBodyInstance::ApplyMaterialToInstanceShapes_AssumesLocked(UPhysicalMateria
 		if(TheirBI->IsShapeBoundToBody(Shape))
 		{
 			FPhysicsCommand::ExecuteShapeWrite(BIWithActor, Shape, [&](const FPhysicsShapeHandle& InnerShape)
-		{
-				ApplyMaterialToShape_AssumesLocked(InnerShape, SimplePhysMat, ComplexPhysMats);
-		});		
+			{
+				ApplyMaterialToShape_AssumesLocked(InnerShape, SimplePhysMat, ComplexPhysMats, &ComplexPhysMatMasks);
+			});		
+		}
 	}
 }
-}
 
 bool FBodyInstance::ValidateTransform(const FTransform &Transform, const FString& DebugName, const UBodySetup* Setup)
 {
diff --git a/Engine/Source/Runtime/Engine/Private/PhysicsEngine/BodySetup.cpp b/Engine/Source/Runtime/Engine/Private/PhysicsEngine/BodySetup.cpp
--- a/Engine/Source/Runtime/Engine/Private/PhysicsEngine/BodySetup.cpp
+++ b/Engine/Source/Runtime/Engine/Private/PhysicsEngine/BodySetup.cpp
@@ -179,6 +179,7 @@ UBodySetup::UBodySetup(const FObjectInitializer& ObjectInitializer)
 	bNeverNeedsCookedCollisionData = false;
 	bGenerateMirroredCollision = true;
 	bGenerateNonMirroredCollision = true;
+	bSupportUVsAndFaceRemap = false;
 	DefaultInstance.SetObjectType(ECC_PhysicsBody);
 #if WITH_EDITORONLY_DATA
 	BuildScale_DEPRECATED = 1.0f;
@@ -353,13 +354,7 @@ void UBodySetup::GetCookInfo(FCookBodySetupInfo& OutCookInfo, EPhysXMeshCookFlag
 
 			OutCookInfo.TriMeshCookFlags = CookFlags;
 
-			// If outer is a static mesh with physical material mask enabled, set retain UV and face remap table to true
-			/*
-			if (UStaticMesh* StaticMesh = Cast<UStaticMesh>(CDPObj))
-			{
-				OutCookInfo.bSupportFaceRemap = StaticMesh->GetEnablePhysicalMaterialMask();
-			}
-			*/
+			OutCookInfo.bSupportFaceRemap = bSupportUVsAndFaceRemap;
 		}
 		else
 		{
@@ -367,7 +362,7 @@ void UBodySetup::GetCookInfo(FCookBodySetupInfo& OutCookInfo, EPhysXMeshCookFlag
 		}
 	}
 
-	OutCookInfo.bSupportUVFromHitResults = UPhysicsSettings::Get()->bSupportUVFromHitResults || OutCookInfo.bSupportFaceRemap;
+	OutCookInfo.bSupportUVFromHitResults = UPhysicsSettings::Get()->bSupportUVFromHitResults || bSupportUVsAndFaceRemap;
 
 #endif // WITH_PHYSX
 }
@@ -822,7 +817,8 @@ void UBodySetup::AddShapesToRigidActor_AssumesLocked(
 	FBodyInstance* OwningInstance, 
 	FVector& Scale3D, 
 	UPhysicalMaterial* SimpleMaterial,
-	TArray<UPhysicalMaterial*>& ComplexMaterials, 
+	TArray<UPhysicalMaterial*>& ComplexMaterials,
+	TArray<FPhysicalMaterialMaskParams>& ComplexMaterialMasks,
 	const FBodyCollisionData& BodyCollisionData,
 	const FTransform& RelativeTM, 
 	TArray<FPhysicsShapeHandle>* NewShapes)
@@ -848,6 +844,9 @@ void UBodySetup::AddShapesToRigidActor_AssumesLocked(
 	AddParams.Scale = Scale3D;
 	AddParams.SimpleMaterial = SimpleMaterial;
 	AddParams.ComplexMaterials = TArrayView<UPhysicalMaterial*>(ComplexMaterials);
+#if WITH_CHAOS
+	AddParams.ComplexMaterialMasks = TArrayView<FPhysicalMaterialMaskParams>(ComplexMaterialMasks);
+#endif
 	AddParams.LocalTransform = RelativeTM;
 	AddParams.WorldTransform = OwningInstance->GetUnrealWorldTransform();
 	AddParams.Geometry = &AggGeom;
@@ -1400,6 +1399,11 @@ void UBodySetup::GetGeometryDDCKey(FString& OutString) const
 		(int32)UPhysicsSettings::Get()->bSupportUVFromHitResults,
 		(int32)GetCollisionTraceFlag(),
 		(int32)BODY_SETUP_GEOMETRY_KEY_VER);
+
+	if (bSupportUVsAndFaceRemap)
+	{
+		OutString += FString(TEXT("_1"));
+	}
 }
 
 void UBodySetup::PostInitProperties()
@@ -1502,6 +1506,8 @@ void UBodySetup::GetResourceSizeEx(FResourceSizeEx& CumulativeResourceSize)
 		const FByteBulkData& FmtData = CookedFormatData.GetFormat(FPlatformProperties::GetPhysicsFormat());
 		CumulativeResourceSize.AddDedicatedSystemMemoryBytes(FmtData.GetBulkDataSize());
 	}
+
+	CumulativeResourceSize.AddDedicatedSystemMemoryBytes(FaceRemap.GetAllocatedSize());
 	
 	// Count any UV info
 	UVInfo.GetResourceSizeEx(CumulativeResourceSize);
diff --git a/Engine/Source/Runtime/Engine/Private/PhysicsEngine/Experimental/ChaosDerivedData.cpp b/Engine/Source/Runtime/Engine/Private/PhysicsEngine/Experimental/ChaosDerivedData.cpp
--- a/Engine/Source/Runtime/Engine/Private/PhysicsEngine/Experimental/ChaosDerivedData.cpp
+++ b/Engine/Source/Runtime/Engine/Private/PhysicsEngine/Experimental/ChaosDerivedData.cpp
@@ -24,7 +24,7 @@ const TCHAR* FChaosDerivedDataCooker::GetPluginName() const
 
 const TCHAR* FChaosDerivedDataCooker::GetVersionString() const
 {
-	return TEXT("A9DFE2536B324A68B507931A0D2C5496");
+	return TEXT("CD1FE02F2AB9453FB6E173E28A076DC4");
 }
 
 FString FChaosDerivedDataCooker::GetPluginSpecificCacheKeySuffix() const
diff --git a/Engine/Source/Runtime/Engine/Private/PhysicsEngine/Experimental/PhysInterface_Chaos.cpp b/Engine/Source/Runtime/Engine/Private/PhysicsEngine/Experimental/PhysInterface_Chaos.cpp
--- a/Engine/Source/Runtime/Engine/Private/PhysicsEngine/Experimental/PhysInterface_Chaos.cpp
+++ b/Engine/Source/Runtime/Engine/Private/PhysicsEngine/Experimental/PhysInterface_Chaos.cpp
@@ -63,6 +63,7 @@ DECLARE_CYCLE_STAT(TEXT("Update Kinematics On Deferred SkelMeshes"), STAT_Update
 DECLARE_CYCLE_STAT(TEXT("Phys Events Time"), STAT_PhysicsEventTime, STATGROUP_Physics);
 DECLARE_CYCLE_STAT(TEXT("SyncComponentsToBodies (sync)"), STAT_SyncComponentsToBodies, STATGROUP_Physics);
 DECLARE_CYCLE_STAT(TEXT("SyncComponentsToBodies (async)"), STAT_SyncComponentsToBodies_Async, STATGROUP_Physics);
+DECLARE_CYCLE_STAT(TEXT("Query PhysicalMaterialMask Hit"), STAT_QueryPhysicalMaterialMaskHit, STATGROUP_Physics);
 
 DECLARE_DWORD_COUNTER_STAT(TEXT("Broadphase Adds"), STAT_NumBroadphaseAdds, STATGROUP_Physics);
 DECLARE_DWORD_COUNTER_STAT(TEXT("Broadphase Removes"), STAT_NumBroadphaseRemoves, STATGROUP_Physics);
@@ -114,6 +115,80 @@ Chaos::FChaosPhysicsMaterial* GetMaterialFromInternalFaceIndex(const FPhysicsSha
 	return nullptr;
 }
 
+Chaos::FChaosPhysicsMaterial* GetMaterialFromInternalFaceIndexAndHitLocation(const FPhysicsShape& Shape, const FPhysicsActor& Actor, uint32 InternalFaceIndex, const FVector& HitLocation)
+{
+	{
+		SCOPE_CYCLE_COUNTER(STAT_QueryPhysicalMaterialMaskHit);
+
+		if (Shape.Materials.Num() > 0 && Actor.Proxy)
+		{
+			Chaos::FPBDRigidsSolver* Solver = Actor.Proxy->GetSolver();
+
+			if (ensure(Solver))
+			{
+				if (Shape.MaterialMasks.Num() > 0)
+				{
+					UBodySetup* BodySetup = nullptr;
+
+					if (const FBodyInstance* BodyInst = GetUserData(Actor))
+					{
+						BodyInst = FPhysicsInterface::ShapeToOriginalBodyInstance(BodyInst, &Shape);
+						BodySetup = BodyInst->BodySetup.Get();	//this data should be immutable at runtime so ok to check from worker thread.
+						ECollisionShapeType GeomType = GetGeometryType(Shape);
+
+						if (BodySetup->bSupportUVsAndFaceRemap && GetGeometryType(Shape) == ECollisionShapeType::Trimesh)
+						{
+							FVector Scale(1.0f, 1.0f, 1.0f);
+							const Chaos::FImplicitObject* Geometry = Shape.Geometry.Get();
+							if (const Chaos::TImplicitObjectScaled<Chaos::FTriangleMeshImplicitObject>* ScaledTrimesh = Chaos::TImplicitObjectScaled<Chaos::FTriangleMeshImplicitObject>::AsScaled(*Geometry))
+							{
+								Scale = ScaledTrimesh->GetScale();
+							}
+
+							// Convert hit location to local
+							Chaos::FRigidTransform3 ActorToWorld(Actor.X(), Actor.R(), Scale);
+							const FVector LocalHitPos = ActorToWorld.InverseTransformPosition(HitLocation);
+
+							uint8 Index = Shape.Geometry->GetMaterialIndex(InternalFaceIndex);
+							if (Shape.MaterialMasks.IsValidIndex(Index))
+							{
+								Chaos::FChaosPhysicsMaterialMask* Mask = nullptr;
+								{
+									Chaos::TSolverQueryMaterialScope<Chaos::ELockType::Read> Scope(Solver);
+									Mask = Solver->GetQueryMaterialMasks().Get(Shape.MaterialMasks[Index].InnerHandle);
+								}
+
+								if (Mask && InternalFaceIndex < (uint32)BodySetup->FaceRemap.Num())
+								{
+									int32 RemappedFaceIndex = BodySetup->FaceRemap[InternalFaceIndex];
+									FVector2D UV;
+
+
+									if (BodySetup->CalcUVAtLocation(LocalHitPos, RemappedFaceIndex, Mask->UVChannelIndex, UV))
+									{
+										uint32 MapIdx = UPhysicalMaterialMask::GetPhysMatIndex(Mask->MaskData, Mask->SizeX, Mask->SizeY, Mask->AddressX, Mask->AddressY, UV.X, UV.Y);
+										uint32 AdjustedMapIdx = Index * EPhysicalMaterialMaskColor::MAX + MapIdx;
+										if (Shape.MaterialMaskMaps.IsValidIndex(AdjustedMapIdx))
+										{
+											uint32 MaterialIdx = Shape.MaterialMaskMaps[AdjustedMapIdx];
+											if (Shape.MaterialMaskMapMaterials.IsValidIndex(MaterialIdx))
+											{
+												Chaos::TSolverQueryMaterialScope<Chaos::ELockType::Read> Scope(Solver);
+												return Solver->GetQueryMaterials().Get(Shape.MaterialMaskMapMaterials[MaterialIdx].InnerHandle);
+											}
+										}
+									}
+								}
+							}
+						}
+					}
+				}
+			}
+		}
+	}
+
+	return GetMaterialFromInternalFaceIndex(Shape, Actor, InternalFaceIndex);
+}
 
 const Chaos::FImplicitObject& FPhysicsShapeReference_Chaos::GetGeometry() const
 {
@@ -255,6 +330,31 @@ void FPhysInterface_Chaos::SetUserData(FPhysicsMaterialHandle& InHandle, void* I
 	Chaos::FPhysicalMaterialManager::Get().UpdateMaterial(InHandle);
 }
 
+FPhysicsMaterialMaskHandle FPhysInterface_Chaos::CreateMaterialMask(const UPhysicalMaterialMask* InMaterialMask)
+{
+	Chaos::FMaterialMaskHandle NewHandle = Chaos::FPhysicalMaterialManager::Get().CreateMask();
+	FPhysInterface_Chaos::UpdateMaterialMask(NewHandle, InMaterialMask);
+	return NewHandle;
+}
+
+void FPhysInterface_Chaos::ReleaseMaterialMask(FPhysicsMaterialMaskHandle& InHandle)
+{
+	Chaos::FPhysicalMaterialManager::Get().Destroy(InHandle);
+}
+
+void FPhysInterface_Chaos::UpdateMaterialMask(FPhysicsMaterialMaskHandle& InHandle, const UPhysicalMaterialMask* InMaterialMask)
+{
+	if (Chaos::FChaosPhysicsMaterialMask* MaterialMask = InHandle.Get())
+	{
+		InMaterialMask->GenerateMaskData(MaterialMask->MaskData, MaterialMask->SizeX, MaterialMask->SizeY);
+		MaterialMask->UVChannelIndex = InMaterialMask->UVChannelIndex;
+		MaterialMask->AddressX = static_cast<int32>(InMaterialMask->AddressX);
+		MaterialMask->AddressY = static_cast<int32>(InMaterialMask->AddressY);
+	}
+
+	Chaos::FPhysicalMaterialManager::Get().UpdateMaterialMask(InHandle);
+}
+
 void FPhysInterface_Chaos::SetUserData(const FPhysicsShapeHandle& InShape, void* InUserData)
 {
 	if (CHAOS_ENSURE(InShape.Shape))
@@ -1373,7 +1473,7 @@ void FPhysInterface_Chaos::AddGeometry(FPhysicsActorHandle& InActor, const FGeom
 				OutOptShapes->Add(NewHandle);
 			}
 
-			FBodyInstance::ApplyMaterialToShape_AssumesLocked(NewHandle, InParams.SimpleMaterial, InParams.ComplexMaterials);
+			FBodyInstance::ApplyMaterialToShape_AssumesLocked(NewHandle, InParams.SimpleMaterial, InParams.ComplexMaterials, &InParams.ComplexMaterialMasks);
 
 			//TArrayView<UPhysicalMaterial*> SimpleView = MakeArrayView(&(const_cast<UPhysicalMaterial*>(InParams.SimpleMaterial)), 1);
 			//FPhysInterface_Chaos::SetMaterials(NewHandle, InParams.ComplexMaterials.Num() > 0 ? InParams.ComplexMaterials : SimpleView);
@@ -1541,6 +1641,75 @@ void FPhysInterface_Chaos::SetMaterials(const FPhysicsShapeHandle& InShape, cons
 	InShape.Shape->Materials = NewMaterialHandles;
 }
 
+void FPhysInterface_Chaos::SetMaterials(const FPhysicsShapeHandle& InShape, const TArrayView<UPhysicalMaterial*> InMaterials, const TArrayView<FPhysicalMaterialMaskParams>& InMaterialMasks)
+{
+	SetMaterials(InShape, InMaterials);
+
+	if (InMaterialMasks.Num() > 0)
+	{
+		// Build a list of handles to store on the shape
+		TArray<Chaos::FMaterialMaskHandle> NewMaterialMaskHandles;
+		TArray<uint32> NewMaterialMaskMaps;
+		TArray<Chaos::FMaterialHandle> NewMaterialMaskMaterialHandles;
+
+		NewMaterialMaskHandles.Reserve(InMaterialMasks.Num());
+
+		int MaskMapMatIdx = 0;
+
+		for (FPhysicalMaterialMaskParams& MaterialMaskData : InMaterialMasks)
+		{
+			if (MaterialMaskData.PhysicalMaterialMask && ensure(MaterialMaskData.PhysicalMaterialMap))
+			{
+				NewMaterialMaskHandles.Add(MaterialMaskData.PhysicalMaterialMask->GetPhysicsMaterialMask());
+				for (int i = 0; i < EPhysicalMaterialMaskColor::MAX; i++)
+				{
+					if (UPhysicalMaterial* MapMat = MaterialMaskData.PhysicalMaterialMap->GetPhysicalMaterialFromMap(i))
+					{
+						InShape.Shape->MaterialMaskMaps.Emplace(MaskMapMatIdx);
+						MaskMapMatIdx++;
+					}
+					else
+					{
+						InShape.Shape->MaterialMaskMaps.Emplace(INDEX_NONE);
+					}
+				}
+			}
+			else
+			{
+				NewMaterialMaskHandles.Add(Chaos::FMaterialMaskHandle());
+				for (int i = 0; i < EPhysicalMaterialMaskColor::MAX; i++)
+				{
+					InShape.Shape->MaterialMaskMaps.Emplace(INDEX_NONE);
+				}
+			}
+		}
+
+		if (MaskMapMatIdx > 0)
+		{
+			NewMaterialMaskMaterialHandles.Reserve(MaskMapMatIdx);
+
+			uint32 Offset = 0;
+
+			for (FPhysicalMaterialMaskParams& MaterialMaskData : InMaterialMasks)
+			{
+				if (MaterialMaskData.PhysicalMaterialMask)
+				{
+					for (int i = 0; i < EPhysicalMaterialMaskColor::MAX; i++)
+					{
+						if (UPhysicalMaterial* MapMat = MaterialMaskData.PhysicalMaterialMap->GetPhysicalMaterialFromMap(i))
+						{
+							NewMaterialMaskMaterialHandles.Add(MapMat->GetPhysicsMaterial());
+						}
+					}
+				}
+			}
+		}
+
+		InShape.Shape->MaterialMasks = NewMaterialMaskHandles;
+		InShape.Shape->MaterialMaskMapMaterials = NewMaterialMaskMaterialHandles;
+	}
+}
+
 void FinishSceneStat()
 {
 }
diff --git a/Engine/Source/Runtime/Engine/Private/PhysicsEngine/PhysicalMaterialMask.cpp b/Engine/Source/Runtime/Engine/Private/PhysicsEngine/PhysicalMaterialMask.cpp
new file mode 100644
--- /dev/null
+++ b/Engine/Source/Runtime/Engine/Private/PhysicsEngine/PhysicalMaterialMask.cpp
@@ -0,0 +1,411 @@
+// copyright epic games, inc. all rights reserved.
+
+/*=============================================================================
+	PhysicalMaterial.cpp
+=============================================================================*/ 
+
+#include "PhysicalMaterials/PhysicalMaterialMask.h"
+
+#if WITH_EDITOR
+
+#include "EditorFramework/AssetImportData.h"
+#include "Editor.h"
+#include "EditorDirectories.h"
+#include "ObjectTools.h"
+#include "PackageTools.h"
+#include "Modules/ModuleManager.h"
+#include "AssetRegistryModule.h"
+#include "AssetToolsModule.h"
+#include "IAssetTools.h"
+#include "DesktopPlatformModule.h"
+#include "IDesktopPlatform.h"
+#include "Engine/Texture.h"
+#include "Factories/TextureFactory.h"
+#include "HAL/PlatformApplicationMisc.h"
+#include "HAL/FileManager.h"
+#include "Misc/FileHelper.h"
+#include "Misc/Paths.h"
+#include "Misc/PackageName.h"
+
+#endif
+
+DEFINE_LOG_CATEGORY_STATIC(LogPhysicalMaterialMask, Log, All);
+
+#if WITH_EDITOR
+
+#include "HAL/IConsoleManager.h"
+
+static void OnDumpPhysicalMaterialMaskData(const TArray< FString >& Arguments)
+{
+	bool bPhysMaterialMaskFound = false;
+	if (Arguments.Num() > 0)
+	{
+		const FAssetRegistryModule& AssetRegistryModule = FModuleManager::LoadModuleChecked<FAssetRegistryModule>(TEXT("AssetRegistry"));
+		FAssetData PhysMatMaskAsset = AssetRegistryModule.Get().GetAssetByObjectPath(*Arguments[0]);
+		if (PhysMatMaskAsset.IsValid() == false)
+		{
+			TArray<FAssetData> AssetsInPackage;
+			AssetRegistryModule.Get().GetAssetsByPackageName(*Arguments[0], AssetsInPackage);
+			if (AssetsInPackage.Num() == 1)
+			{
+				PhysMatMaskAsset = AssetsInPackage[0];
+			}
+		}
+		if (PhysMatMaskAsset.IsValid())
+		{
+			if (UPhysicalMaterialMask* PhysMatMask = Cast<UPhysicalMaterialMask>(PhysMatMaskAsset.GetAsset()))
+			{
+				PhysMatMask->DumpMaskData();
+			}
+			else
+			{
+				UE_LOG(LogPhysicalMaterialMask, Warning, TEXT("Could not load PhysicalMaterialMask asset for argument: %s"), *Arguments[0]);
+			}
+		}
+		else
+		{
+			UE_LOG(LogPhysicalMaterialMask, Warning, TEXT("Could not find PhysicalMaterialMask asset for argument: %s"), *Arguments[0]);
+		}
+	}
+	else
+	{
+		UE_LOG(LogPhysicalMaterialMask, Warning, TEXT("Command requires a PhysicalMaterialMask asset reference to be specified."));
+	}
+}
+
+static FAutoConsoleCommand CmdDumpPhysicalMaterialMaskData(
+	TEXT("p.DumpPhysicalMaterialMaskData"),
+	TEXT("Outputs the current mask data for the specified physical material mask asset to the log."),
+	FConsoleCommandWithArgsDelegate::CreateStatic(&OnDumpPhysicalMaterialMaskData)
+);
+
+#endif // #if WITH_EDITOR
+
+uint32 UPhysicalMaterialMask::INVALID_MASK_INDEX = 0xF;
+
+UPhysicalMaterialMask::UPhysicalMaterialMask(const FObjectInitializer& ObjectInitializer)
+	: Super(ObjectInitializer)
+#if WITH_EDITORONLY_DATA	
+	,AssetImportData(nullptr)
+	,MaskTexture(nullptr)
+#endif
+	,UVChannelIndex(0)
+{
+}
+
+#if WITH_EDITOR
+
+void UPhysicalMaterialMask::PostEditChangeProperty(FPropertyChangedEvent& PropertyChangedEvent)
+{
+#if WITH_CHAOS
+	FPhysicsInterface::UpdateMaterialMask(MaterialMaskHandle, this);
+#endif
+
+	Super::PostEditChangeProperty(PropertyChangedEvent);
+}
+
+#endif // WITH_EDITOR
+
+void UPhysicalMaterialMask::PostInitProperties()
+{
+#if WITH_EDITORONLY_DATA
+	if (!HasAnyFlags(RF_ClassDefaultObject | RF_NeedLoad))
+	{
+		AssetImportData = NewObject<UAssetImportData>(this, TEXT("AssetImportData"));
+		AssetImportData->AddFileName(FString(), 0); // add empty filename for now
+	}
+#endif
+	Super::PostInitProperties();
+}
+
+void UPhysicalMaterialMask::PostLoad()
+{
+	Super::PostLoad();
+
+#if WITH_EDITORONLY_DATA
+	if (AssetImportData == nullptr)
+	{
+		AssetImportData = NewObject<UAssetImportData>(this, TEXT("AssetImportData"));
+	}
+#endif
+}
+
+void UPhysicalMaterialMask::FinishDestroy()
+{
+#if WITH_CHAOS
+	FPhysicsInterface::ReleaseMaterialMask(MaterialMaskHandle);
+#endif
+	Super::FinishDestroy();
+}
+
+
+#if WITH_EDITOR
+
+void UPhysicalMaterialMask::SetMaskTexture(UTexture* InMaskTexture, const FString& InTextureFilename)
+{
+	if (!AssetImportData)
+	{
+		AssetImportData = NewObject<UAssetImportData>(this, TEXT("AssetImportData"));
+	}
+	MaskTexture = InMaskTexture;
+	AssetImportData->AddFileName(InTextureFilename, 0);
+
+#if WITH_CHAOS
+	FPhysicsInterface::UpdateMaterialMask(MaterialMaskHandle, this);
+#endif
+}
+
+// This template generates mask data from the texture mask, converting colors to mask ids.
+template<typename PixelDataType, int32 RIdx, int32 GIdx, int32 BIdx, int32 AIdx> class MaskDataGenerator
+{
+public:
+
+	MaskDataGenerator(int32 SizeX, int32 SizeY, uint8* SourceTextureData)
+		: SourceData(reinterpret_cast<PixelDataType*>(SourceTextureData))
+		, TextureWidth(SizeX)
+		, TextureHeight(SizeY)
+	{
+	}
+
+	void GenerateMask(TArray<uint32>& OutMaskData)
+	{
+		check(TextureWidth % 8 == 0);
+
+		OutMaskData.Empty();
+
+		PixelDataType* PixelData = SourceData;
+
+		for (int32 Y = 0; Y < TextureHeight; ++Y)
+		{
+			for (int32 X = 0; X < TextureWidth; )
+			{
+				// Pack 8 color ids into one uint32.
+				uint32 NewMaskEntry = 0x0;
+				for (int32 Idx = 0; Idx < 8 && X < TextureWidth; X++, Idx++, PixelData += 4)
+				{
+					NewMaskEntry |= GetColorId(PixelData[RIdx], PixelData[GIdx], PixelData[BIdx]) << (Idx * 4);
+				}
+				OutMaskData.Emplace(NewMaskEntry);
+			}
+		}
+	}
+
+	bool IsEqual(FColor ColorA, FColor ColorB)
+	{
+		static int32 Tolerance = 10;
+		return FMath::Abs(ColorA.R - ColorB.R) < Tolerance && FMath::Abs(ColorA.G - ColorB.G) < Tolerance && FMath::Abs(ColorA.B - ColorB.B) < Tolerance && FMath::Abs(ColorA.A - ColorB.A) < Tolerance;
+	}
+
+	uint32 GetColorId(PixelDataType R, PixelDataType G, PixelDataType B)
+	{
+		FColor Color(R, G, B);
+		if (IsEqual(Color, FColor::White))
+		{
+			return EPhysicalMaterialMaskColor::White;
+		}
+		if (IsEqual(Color, FColor::Black))
+		{
+			return EPhysicalMaterialMaskColor::Black;
+		}
+		if (IsEqual(Color, FColor::Red))
+		{
+			return EPhysicalMaterialMaskColor::Red;
+		}
+		if (IsEqual(Color, FColor::Green))
+		{
+			return EPhysicalMaterialMaskColor::Green;
+		}
+		if (IsEqual(Color, FColor::Blue))
+		{
+			return EPhysicalMaterialMaskColor::Blue;
+		}
+		if (IsEqual(Color, FColor::Cyan))
+		{
+			return EPhysicalMaterialMaskColor::Cyan;
+		}
+		if (IsEqual(Color, FColor::Yellow))
+		{
+			return EPhysicalMaterialMaskColor::Yellow;
+		}
+		if (IsEqual(Color, FColor::Magenta))
+		{
+			return EPhysicalMaterialMaskColor::Magenta;
+		}
+		return UPhysicalMaterialMask::INVALID_MASK_INDEX; 
+	}
+
+	PixelDataType* SourceData;
+	int32 TextureWidth;
+	int32 TextureHeight;
+};
+
+uint32 UPhysicalMaterialMask::GetPhysMatIndex(const TArray<uint32>& MaskData, int32 SizeX, int32 SizeY, int32 AddressX, int32 AddressY, float U, float V)
+{
+	auto AdjustCoord = [](float& Val, TextureAddress Address)
+	{
+		if (Address == TextureAddress::TA_Clamp)
+		{
+			Val = FMath::Clamp(Val, 0.0f, 1.0f);
+		}
+		else if (Address == TextureAddress::TA_Wrap)
+		{
+			Val = FMath::Frac(Val);
+		}
+		else // Mirror
+		{
+			float IntPart;
+			Val = FMath::Modf(Val, &IntPart);
+			if (static_cast<int32>(IntPart) % 2 > 0)
+			{
+				Val = 1.0f - Val;
+			}
+		}
+	};
+
+	AdjustCoord(U, static_cast<TextureAddress>(AddressX));
+	AdjustCoord(V, static_cast<TextureAddress>(AddressX));
+
+	uint32 X = FMath::FloorToInt((float)SizeX * U);
+	uint32 Y = FMath::FloorToInt((float)SizeY * V);
+	uint32 Index = (SizeX / 8) * Y + X / 8;
+
+	if (Index < (uint32)MaskData.Num())
+	{
+		return (MaskData[Index] >> (X % 8 * 4)) & 0xf;
+	}
+
+	return UPhysicalMaterialMask::INVALID_MASK_INDEX;
+}
+
+void UPhysicalMaterialMask::GenerateMaskData(TArray<uint32>& OutMaskData, int32& OutSizeX, int32& OutSizeY) const
+{
+	OutMaskData.Empty();
+
+#if WITH_EDITOR
+
+	if (MaskTexture)
+	{
+		uint8* TextureData = MaskTexture->Source.LockMip(0);
+		if (TextureData)
+		{
+			const int32 TextureDataSize = MaskTexture->Source.CalcMipSize(0);
+
+			if (TextureDataSize > 0)
+			{
+				OutSizeX = MaskTexture->Source.GetSizeX();
+				OutSizeY = MaskTexture->Source.GetSizeY();
+
+				ETextureSourceFormat TextureDataSourceFormat = MaskTexture->Source.GetFormat();
+				switch (TextureDataSourceFormat)
+				{
+				case TSF_BGRA8:
+				{
+					MaskDataGenerator<uint8, 2, 1, 0, 3> MaskDataGen(OutSizeX, OutSizeY, TextureData);
+					MaskDataGen.GenerateMask(OutMaskData);
+					break;
+				}
+
+				case TSF_RGBA8:
+				{
+					MaskDataGenerator<uint8, 0, 1, 2, 3> MaskDataGen(OutSizeX, OutSizeY, TextureData);
+					MaskDataGen.GenerateMask(OutMaskData);
+					break;
+				}
+
+				case TSF_RGBA16:
+				{
+					MaskDataGenerator<uint16, 0, 1, 2, 3> MaskDataGen(OutSizeX, OutSizeY, TextureData);
+					MaskDataGen.GenerateMask(OutMaskData);
+					break;
+				}
+				}
+			}
+		}
+
+		MaskTexture->Source.UnlockMip(0);
+	}
+#endif // WITH_EDITOR
+
+	if (OutMaskData.Num() == 0)
+	{
+		OutSizeX = 0;
+		OutSizeY = 0;
+	}
+}
+
+void UPhysicalMaterialMask::DumpMaskData()
+{
+	int32 SizeX = 0;
+	int32 SizeY = 0;
+	TArray<uint32> MaskData;
+	GenerateMaskData(MaskData, SizeX, SizeY);
+
+	ensure(MaskData.Num() == SizeY * SizeX / 8);
+
+	if (MaskData.Num() == 0)
+	{
+		UE_LOG(LogPhysicalMaterialMask, Display, TEXT("PhysicalMaterialMask data is empty."));
+	}
+
+	auto GetColorString = [](uint32 ColorIdx) -> FString
+	{
+		switch (ColorIdx)
+		{
+		case EPhysicalMaterialMaskColor::White:
+			return FString(TEXT("White"));
+		case EPhysicalMaterialMaskColor::Black:
+			return FString(TEXT("Black"));
+		case EPhysicalMaterialMaskColor::Red:
+			return FString(TEXT("Red"));
+		case EPhysicalMaterialMaskColor::Green:
+			return FString(TEXT("Green"));
+		case EPhysicalMaterialMaskColor::Blue:
+			return FString(TEXT("Blue"));
+		case EPhysicalMaterialMaskColor::Yellow:
+			return FString(TEXT("Yellow"));
+		case EPhysicalMaterialMaskColor::Cyan:
+			return FString(TEXT("Cyan"));
+		case EPhysicalMaterialMaskColor::Magenta:
+			return FString(TEXT("Magenta"));
+		}
+		return FString(TEXT("INVALID_MASK_INDEX"));
+	};
+
+	for (int32 Y = 0, ArrayIdx = 0; Y < SizeY; Y++)
+	{
+		UE_LOG(LogPhysicalMaterialMask, Display, TEXT("Mask ROW %d:"), Y);
+
+		for (int32 X = 0; X < SizeX; ArrayIdx++)
+		{
+			// Unpack 8 color ids from one uint32.
+			uint32 MaskEntry = MaskData[ArrayIdx];
+			UE_LOG(LogPhysicalMaterialMask, Display, TEXT("  Mask PackedEntry[%d][%d] = 0x%x"), Y, X / 8, MaskEntry);
+
+			for (int32 Idx = 0; Idx < 8 && X < SizeX; X++, Idx++)
+			{
+				uint32 ColorIdx = (MaskEntry >> (Idx * 4)) & 0xf;
+				UE_LOG(LogPhysicalMaterialMask, Display, TEXT("    Mask[%d][%d] = %s (%d)"), Y, X, *GetColorString(ColorIdx), ColorIdx);
+			}
+		}
+	}
+}
+
+#endif // WITH_EDITOR
+
+#if	WITH_CHAOS
+
+FPhysicsMaterialMaskHandle& UPhysicalMaterialMask::GetPhysicsMaterialMask()
+{
+	if (!MaterialMaskHandle.IsValid())
+	{
+		MaterialMaskHandle = FPhysicsInterface::CreateMaterialMask(this);
+		check(MaterialMaskHandle.IsValid());
+
+		FPhysicsInterface::UpdateMaterialMask(MaterialMaskHandle, this);
+	}
+
+	return MaterialMaskHandle;
+}
+
+#endif
+
diff --git a/Engine/Source/Runtime/Engine/Private/StaticMesh.cpp b/Engine/Source/Runtime/Engine/Private/StaticMesh.cpp
--- a/Engine/Source/Runtime/Engine/Private/StaticMesh.cpp
+++ b/Engine/Source/Runtime/Engine/Private/StaticMesh.cpp
@@ -2974,6 +2974,16 @@ void UStaticMesh::PostEditChangeProperty(FPropertyChangedEvent& PropertyChangedE
 			BodySetup->CreatePhysicsMeshes();
 		}
 	}
+
+	if (PropertyName == GET_MEMBER_NAME_CHECKED(UStaticMesh, bSupportPhysicalMaterialMasks))
+	{
+		if (BodySetup)
+		{
+			BodySetup->bSupportUVsAndFaceRemap = bSupportPhysicalMaterialMasks;
+			BodySetup->InvalidatePhysicsData();
+			BodySetup->CreatePhysicsMeshes();
+		}
+	}
 #endif
 
 	LightMapResolution = FMath::Max(LightMapResolution, 0);
@@ -5656,8 +5666,7 @@ bool UStaticMesh::GetPhysicsTriMeshData(struct FTriMeshCollisionData* CollisionD
 	TMap<int32, int32> MeshToCollisionVertMap; // map of static mesh verts to collision verts
 
 	// If the mesh enables physical material masks, override the physics setting since UVs are always required 
-	// bool bCopyUVs = GetEnablePhysicalMaterialMask() || UPhysicsSettings::Get()->bSupportUVFromHitResults; // See if we should copy UVs
-	bool bCopyUVs = UPhysicsSettings::Get()->bSupportUVFromHitResults; // See if we should copy UVs
+	bool bCopyUVs = bSupportPhysicalMaterialMasks || UPhysicsSettings::Get()->bSupportUVFromHitResults; // See if we should copy UVs
 
 	// If copying UVs, allocate array for storing them
 	if (bCopyUVs)
@@ -5802,6 +5811,7 @@ void UStaticMesh::CreateBodySetup()
 	{
 		BodySetup = NewObject<UBodySetup>(this);
 		BodySetup->DefaultInstance.SetCollisionProfileName(UCollisionProfile::BlockAll_ProfileName);
+		BodySetup->bSupportUVsAndFaceRemap = bSupportPhysicalMaterialMasks;
 	}
 }
 
diff --git a/Engine/Source/Runtime/Engine/Public/Physics/Experimental/PhysInterface_Chaos.h b/Engine/Source/Runtime/Engine/Public/Physics/Experimental/PhysInterface_Chaos.h
--- a/Engine/Source/Runtime/Engine/Public/Physics/Experimental/PhysInterface_Chaos.h
+++ b/Engine/Source/Runtime/Engine/Public/Physics/Experimental/PhysInterface_Chaos.h
@@ -14,6 +14,7 @@
 #include "Physics/GenericPhysicsInterface.h"
 #include "Physics/Experimental/PhysicsUserData_Chaos.h"
 #include "Chaos/PhysicalMaterials.h"
+#include "PhysicalMaterials/PhysicalMaterialMask.h"
 
 //NOTE: Do not include Chaos headers directly as it means recompiling all of engine. This should be reworked to avoid allocations
 
@@ -202,6 +203,11 @@ public:
     static void UpdateMaterial(FPhysicsMaterialHandle& InHandle, UPhysicalMaterial* InMaterial);
     static void SetUserData(FPhysicsMaterialHandle& InHandle, void* InUserData);
 
+	// Material mask functions 
+	static FPhysicsMaterialMaskHandle CreateMaterialMask(const UPhysicalMaterialMask* InMaterialMask);
+	static void ReleaseMaterialMask(FPhysicsMaterialMaskHandle& InHandle);
+	static void UpdateMaterialMask(FPhysicsMaterialMaskHandle& InHandle, const UPhysicalMaterialMask* InMaterialMask);
+
 	// Actor interface functions
 	template<typename AllocatorType>
 	static int32 GetAllShapes_AssumedLocked(const FPhysicsActorHandle& InActorReference, TArray<FPhysicsShapeHandle, AllocatorType>& OutShapes);
@@ -400,7 +406,8 @@ public:
     static void SetUserData(const FPhysicsShapeHandle& InShape, void* InUserData);
     static void SetGeometry(const FPhysicsShapeHandle& InShape, physx::PxGeometry& InGeom) {}
 	static void SetLocalTransform(const FPhysicsShapeHandle& InShape, const FTransform& NewLocalTransform);
-    static void SetMaterials(const FPhysicsShapeHandle& InShape, const TArrayView<UPhysicalMaterial*>InMaterials);
+	static void SetMaterials(const FPhysicsShapeHandle& InShape, const TArrayView<UPhysicalMaterial*>InMaterials);
+	static void SetMaterials(const FPhysicsShapeHandle& InShape, const TArrayView<UPhysicalMaterial*> InMaterials, const TArrayView<FPhysicalMaterialMaskParams>& InMaterialMasks);
 };
 
 /*
@@ -447,6 +454,8 @@ FORCEINLINE void ComputeZeroDistanceImpactNormalAndPenetration(const UWorld* Wor
 
 Chaos::FChaosPhysicsMaterial* GetMaterialFromInternalFaceIndex(const FPhysicsShape& Shape, const FPhysicsActor& Actor, uint32 InternalFaceIndex);
 
+Chaos::FChaosPhysicsMaterial* GetMaterialFromInternalFaceIndexAndHitLocation(const FPhysicsShape& Shape, const FPhysicsActor& Actor, uint32 InternalFaceIndex, const FVector& HitLocation);
+
 uint32 GetTriangleMeshExternalFaceIndex(const FPhysicsShape& Shape, uint32 InternalFaceIndex);
 
 inline void GetShapes(const FPhysActorDummy& RigidActor, FPhysTypeDummy** ShapesBuffer, uint32 NumShapes)
diff --git a/Engine/Source/Runtime/Engine/Public/Physics/PhysicsInterfaceTypes.h b/Engine/Source/Runtime/Engine/Public/Physics/PhysicsInterfaceTypes.h
--- a/Engine/Source/Runtime/Engine/Public/Physics/PhysicsInterfaceTypes.h
+++ b/Engine/Source/Runtime/Engine/Public/Physics/PhysicsInterfaceTypes.h
@@ -38,6 +38,15 @@ namespace physx
 	class PxTriangleMesh;
 }
 
+struct FPhysicalMaterialMaskParams
+{
+	/** Physical materials mask */
+	UPhysicalMaterialMask* PhysicalMaterialMask;
+
+	/** Pointer to material which contains the physical material map */
+	UMaterialInterface* PhysicalMaterialMap;
+};
+
 struct FGeometryAddParams
 {
 	bool bDoubleSided;
@@ -46,6 +55,9 @@ struct FGeometryAddParams
 	FVector Scale;
 	UPhysicalMaterial* SimpleMaterial;
 	TArrayView<UPhysicalMaterial*> ComplexMaterials;
+#if WITH_CHAOS
+	TArrayView<FPhysicalMaterialMaskParams> ComplexMaterialMasks;
+#endif
 	FTransform LocalTransform;
 	FTransform WorldTransform;
 	FKAggregateGeom* Geometry;
diff --git a/Engine/Source/Runtime/Experimental/Chaos/Private/Chaos/GeometryParticles.cpp b/Engine/Source/Runtime/Experimental/Chaos/Private/Chaos/GeometryParticles.cpp
--- a/Engine/Source/Runtime/Experimental/Chaos/Private/Chaos/GeometryParticles.cpp
+++ b/Engine/Source/Runtime/Experimental/Chaos/Private/Chaos/GeometryParticles.cpp
@@ -60,6 +60,9 @@ namespace Chaos
 		, Geometry()
 		, WorldSpaceInflatedShapeBounds(TAABB<FReal, 3>(FVec3(0), FVec3(0)))
 		, Materials()
+		, MaterialMasks()
+		, MaterialMaskMaps()
+		, MaterialMaskMapMaterials()
 		, bDisable(false)
 		, bSimulate(true)
 		, CollisionTraceType(EChaosCollisionTraceFlag::Chaos_CTF_UseDefault)
@@ -96,6 +99,7 @@ namespace Chaos
 	void TPerShapeData<T, d>::Serialize(FChaosArchive& Ar)
 	{
 		Ar.UsingCustomVersion(FExternalPhysicsCustomObjectVersion::GUID);
+		Ar.UsingCustomVersion(FExternalPhysicsMaterialCustomObjectVersion::GUID);
 
 		Ar << Geometry;
 		Ar << QueryData;
@@ -132,6 +136,12 @@ namespace Chaos
 			Ar << Data;
 			CollisionTraceType = (EChaosCollisionTraceFlag)Data;
 		}
+
+		if (Ar.CustomVer(FExternalPhysicsMaterialCustomObjectVersion::GUID) >= FExternalPhysicsMaterialCustomObjectVersion::AddedMaterialMasks)
+		{
+			Ar << MaterialMasks << MaterialMaskMaps << MaterialMaskMapMaterials;
+		}
+
 	}
 
 
diff --git a/Engine/Source/Runtime/Experimental/Chaos/Private/Chaos/PhysicalMaterials.cpp b/Engine/Source/Runtime/Experimental/Chaos/Private/Chaos/PhysicalMaterials.cpp
--- a/Engine/Source/Runtime/Experimental/Chaos/Private/Chaos/PhysicalMaterials.cpp
+++ b/Engine/Source/Runtime/Experimental/Chaos/Private/Chaos/PhysicalMaterials.cpp
@@ -23,8 +23,27 @@ namespace Chaos
 		return nullptr;
 	}
 
+	FChaosPhysicsMaterialMask* FMaterialMaskHandle::Get() const
+	{
+		if (InnerHandle.IsValid())
+		{
+			return FPhysicalMaterialManager::Get().Resolve(InnerHandle);
+		}
+		return nullptr;
+	}
+
+	const FChaosPhysicsMaterialMask* FConstMaterialMaskHandle::Get() const
+	{
+		if (InnerHandle.IsValid())
+		{
+			return FPhysicalMaterialManager::Get().Resolve(InnerHandle);
+		}
+		return nullptr;
+	}
+
 	FPhysicalMaterialManager::FPhysicalMaterialManager()
 		: Materials(InitialCapacity)
+		, MaterialMasks()
 	{
 
 	}
@@ -45,6 +64,16 @@ namespace Chaos
 		return Materials.Get(InHandle);
 	}
 
+	FChaosPhysicsMaterialMask* FPhysicalMaterialManager::Resolve(FChaosMaterialMaskHandle InHandle) const
+	{
+		return MaterialMasks.Get(InHandle);
+	}
+
+	const FChaosPhysicsMaterialMask* FPhysicalMaterialManager::Resolve(FChaosConstMaterialMaskHandle InHandle) const
+	{
+		return MaterialMasks.Get(InHandle);
+	}
+	
 	void FPhysicalMaterialManager::UpdateMaterial(FMaterialHandle InHandle)
 	{
 		check(IsInGameThread());
@@ -52,11 +81,23 @@ namespace Chaos
 		OnMaterialUpdated.Broadcast(InHandle);
 	}
 
+	void FPhysicalMaterialManager::UpdateMaterialMask(FMaterialMaskHandle InHandle)
+	{
+		check(IsInGameThread());
+
+		OnMaterialMaskUpdated.Broadcast(InHandle);
+	}
+
 	const Chaos::THandleArray<FChaosPhysicsMaterial>& FPhysicalMaterialManager::GetMasterMaterials() const
 	{
 		return Materials;
 	}
 
+	const Chaos::THandleArray<FChaosPhysicsMaterialMask>& FPhysicalMaterialManager::GetMasterMaterialMasks() const
+	{
+		return MaterialMasks;
+	}
+
 	FMaterialHandle FPhysicalMaterialManager::Create()
 	{
 		LLM_SCOPE(ELLMTag::Chaos);
@@ -70,6 +111,16 @@ namespace Chaos
 		return OutHandle;
 	}
 
+	FMaterialMaskHandle FPhysicalMaterialManager::CreateMask()
+	{
+		check(IsInGameThread());
+		FMaterialMaskHandle OutHandle;
+		OutHandle.InnerHandle = MaterialMasks.Create();
+
+		OnMaterialMaskCreated.Broadcast(OutHandle);
+
+		return OutHandle;
+	}
 
 	void FPhysicalMaterialManager::Destroy(FMaterialHandle InHandle)
 	{
@@ -84,4 +135,14 @@ namespace Chaos
 		}
 	}
 
+	void FPhysicalMaterialManager::Destroy(FMaterialMaskHandle InHandle)
+	{
+		check(IsInGameThread());
+		if (InHandle.InnerHandle.IsValid())
+		{
+			OnMaterialMaskDestroyed.Broadcast(InHandle);
+
+			MaterialMasks.Destroy(InHandle.InnerHandle);
+		}
+	}
 }
\ No newline at end of file
diff --git a/Engine/Source/Runtime/Experimental/Chaos/Public/Chaos/Defines.h b/Engine/Source/Runtime/Experimental/Chaos/Public/Chaos/Defines.h
--- a/Engine/Source/Runtime/Experimental/Chaos/Public/Chaos/Defines.h
+++ b/Engine/Source/Runtime/Experimental/Chaos/Public/Chaos/Defines.h
@@ -64,11 +64,46 @@ namespace Chaos
 		}
 	};
 
-
-	template <typename T>
-	FORCEINLINE FArchive& operator<<(FArchive& Ar, FChaosPhysicsMaterial& Value)
+	class FChaosPhysicsMaterialMask
 	{
-		Value.Serialize(Ar);
-		return Ar;
-	}
+	public:
+
+		int32 SizeX;
+		int32 SizeY;
+		TArray<uint32> MaskData;
+
+		int32 UVChannelIndex;
+		int32 AddressX;
+		int32 AddressY;
+
+		FChaosPhysicsMaterialMask()
+			: SizeX(0)
+			, SizeY(0)
+			, MaskData()
+			, UVChannelIndex(0)
+			, AddressX(0)
+			, AddressY(0)
+		{
+		}
+
+		static constexpr bool IsSerializablePtr = true;
+
+		static void StaticSerialize(FArchive& Ar, TSerializablePtr<FChaosPhysicsMaterialMask>& Serializable)
+		{
+			FChaosPhysicsMaterialMask* MaterialMask = const_cast<FChaosPhysicsMaterialMask*>(Serializable.Get());
+
+			if (Ar.IsLoading())
+			{
+				MaterialMask = new FChaosPhysicsMaterialMask();
+				Serializable.SetFromRawLowLevel(MaterialMask);
+			}
+
+			MaterialMask->Serialize(Ar);
+		}
+
+		void Serialize(FArchive& Ar)
+		{
+			Ar << SizeX << SizeY << UVChannelIndex << AddressX << AddressY << MaskData;
+		}
+	};
 }
\ No newline at end of file
diff --git a/Engine/Source/Runtime/Experimental/Chaos/Public/Chaos/Framework/Handles.h b/Engine/Source/Runtime/Experimental/Chaos/Public/Chaos/Framework/Handles.h
--- a/Engine/Source/Runtime/Experimental/Chaos/Public/Chaos/Framework/Handles.h
+++ b/Engine/Source/Runtime/Experimental/Chaos/Public/Chaos/Framework/Handles.h
@@ -147,8 +147,6 @@ namespace Chaos
 	{
 	public:
 
-		static_assert(TIsTrivial<ElementType>::Value, "Handle managed data must be trivial");
-
 		using FHandle = THandle<ElementType, IndexWidth, GenerationWidth>;
 		using FConstHandle = TConstHandle<ElementType, IndexWidth, GenerationWidth>;
 		
diff --git a/Engine/Source/Runtime/Experimental/Chaos/Public/Chaos/GeometryParticles.h b/Engine/Source/Runtime/Experimental/Chaos/Public/Chaos/GeometryParticles.h
--- a/Engine/Source/Runtime/Experimental/Chaos/Public/Chaos/GeometryParticles.h
+++ b/Engine/Source/Runtime/Experimental/Chaos/Public/Chaos/GeometryParticles.h
@@ -12,6 +12,7 @@
 #include "Chaos/PhysicalMaterials.h"
 #include "UObject/PhysicsObjectVersion.h"
 #include "UObject/ExternalPhysicsCustomObjectVersion.h"
+#include "UObject/ExternalPhysicsMaterialCustomObjectVersion.h"
 
 #ifndef CHAOS_DETERMINISTIC
 #define CHAOS_DETERMINISTIC 1
@@ -52,6 +53,9 @@ namespace Chaos
 		TSerializablePtr<FImplicitObject> Geometry;
 		TAABB<FReal, 3> WorldSpaceInflatedShapeBounds;
 		TArray<FMaterialHandle> Materials;
+		TArray<FMaterialMaskHandle> MaterialMasks;
+		TArray<uint32> MaterialMaskMaps;
+		TArray<FMaterialHandle> MaterialMaskMapMaterials;
 
 		// TODO: Bitfields?
 		bool bDisable;
diff --git a/Engine/Source/Runtime/Experimental/Chaos/Public/Chaos/PhysicalMaterials.h b/Engine/Source/Runtime/Experimental/Chaos/Public/Chaos/PhysicalMaterials.h
--- a/Engine/Source/Runtime/Experimental/Chaos/Public/Chaos/PhysicalMaterials.h
+++ b/Engine/Source/Runtime/Experimental/Chaos/Public/Chaos/PhysicalMaterials.h
@@ -11,6 +11,9 @@ namespace Chaos
 	using FMaterialArray = THandleArray<FChaosPhysicsMaterial>;
 	using FChaosMaterialHandle = FMaterialArray::FHandle;
 	using FChaosConstMaterialHandle = FMaterialArray::FConstHandle;
+	using FMaterialMaskArray = THandleArray<FChaosPhysicsMaterialMask>;
+	using FChaosMaterialMaskHandle = FMaterialMaskArray::FHandle;
+	using FChaosConstMaterialMaskHandle = FMaterialMaskArray::FConstHandle;
 
 	/** 
 	 * Helper wrapper to encapsulate the access through the material manager
@@ -46,13 +49,56 @@ namespace Chaos
 		friend FArchive& operator <<(FArchive& Ar, FConstMaterialHandle& Value) { Ar << Value.InnerHandle; return Ar; }
 	};
 
+	/**
+	 * Helper wrapper to encapsulate the access through the material manager
+	 * Handles returned from the material manager are only designed to be used (resolved) on the game thread
+	 * by the physics interface.
+	 */
+	struct CHAOS_API FMaterialMaskHandle
+	{
+		FChaosPhysicsMaterialMask* Get() const;
+		FChaosMaterialMaskHandle InnerHandle;
+
+		bool IsValid() const { return Get() != nullptr; }
+
+		friend bool operator ==(const FMaterialMaskHandle& A, const FMaterialMaskHandle& B) { return A.InnerHandle == B.InnerHandle; }
+
+		friend FArchive& operator <<(FArchive& Ar, FMaterialMaskHandle& Value) { Ar << Value.InnerHandle; return Ar; }
+	};
+
+	/**
+	 * Helper wrapper to encapsulate the access through the material manager
+	 * Handles returned from the material manager are only designed to be used (resolved) on the game thread
+	 * by the physics interface.
+	 */
+	struct CHAOS_API FConstMaterialMaskHandle
+	{
+		const FChaosPhysicsMaterialMask* Get() const;
+		FChaosConstMaterialMaskHandle InnerHandle;
+
+		bool IsValid() const { return Get() != nullptr; }
+
+		friend bool operator ==(const FConstMaterialMaskHandle& A, const FConstMaterialMaskHandle& B) { return A.InnerHandle == B.InnerHandle; }
+
+		friend FArchive& operator <<(FArchive& Ar, FConstMaterialMaskHandle& Value) { Ar << Value.InnerHandle; return Ar; }
+	};
+
 	DECLARE_MULTICAST_DELEGATE_OneParam(FOnMaterialCreated, FMaterialHandle);
 	DECLARE_MULTICAST_DELEGATE_OneParam(FOnMaterialDestroyed, FMaterialHandle);
 	DECLARE_MULTICAST_DELEGATE_OneParam(FOnMaterialUpdated, FMaterialHandle);
+
+	DECLARE_MULTICAST_DELEGATE_OneParam(FOnMaterialMaskCreated, FMaterialMaskHandle);
+	DECLARE_MULTICAST_DELEGATE_OneParam(FOnMaterialMaskDestroyed, FMaterialMaskHandle);
+	DECLARE_MULTICAST_DELEGATE_OneParam(FOnMaterialMaskUpdated, FMaterialMaskHandle);
+
 	using FMaterialCreatedDelegate = FOnMaterialCreated::FDelegate;
 	using FMaterialDestroyedDelegate = FOnMaterialDestroyed::FDelegate;
 	using FMaterialUpdatedDelegate = FOnMaterialUpdated::FDelegate;
 
+	using FMaterialMaskCreatedDelegate = FOnMaterialMaskCreated::FDelegate;
+	using FMaterialMaskDestroyedDelegate = FOnMaterialMaskDestroyed::FDelegate;
+	using FMaterialMaskUpdatedDelegate = FOnMaterialMaskUpdated::FDelegate;
+
 	/** 
 	 * Global manager for physical materials.
 	 * Materials are created, updated and destroyed only on the game thread and an immutable
@@ -71,25 +117,36 @@ namespace Chaos
 
 		/** Create a new material, returning a stable handle to it - this should be stored and not the actual material pointer */
 		FMaterialHandle Create();
+		FMaterialMaskHandle CreateMask();
 
 		/** Destroy the material referenced by the provided handle */
 		void Destroy(FMaterialHandle InHandle);
+		void Destroy(FMaterialMaskHandle InHandle);
 
 		/** Get the actual material from a handle */
 		FChaosPhysicsMaterial* Resolve(FChaosMaterialHandle InHandle) const;
 		const FChaosPhysicsMaterial* Resolve(FChaosConstMaterialHandle InHandle) const;
 
+		FChaosPhysicsMaterialMask* Resolve(FChaosMaterialMaskHandle InHandle) const;
+		const FChaosPhysicsMaterialMask* Resolve(FChaosConstMaterialMaskHandle InHandle) const;
+
 		/** Signals stakeholders that the stored material for the provided handle has changed */
 		void UpdateMaterial(FMaterialHandle InHandle);
+		void UpdateMaterialMask(FMaterialMaskHandle InHandle);
 
 		/** Gets the internal list of master materials representing the current user state of the material data */
 		const THandleArray<FChaosPhysicsMaterial>& GetMasterMaterials() const;
+		const THandleArray<FChaosPhysicsMaterialMask>& GetMasterMaterialMasks() const;
 
 		/** Events */
 		FOnMaterialUpdated OnMaterialUpdated;
 		FOnMaterialCreated OnMaterialCreated;
 		FOnMaterialDestroyed OnMaterialDestroyed;
 
+		FOnMaterialMaskUpdated OnMaterialMaskUpdated;
+		FOnMaterialMaskCreated OnMaterialMaskCreated;
+		FOnMaterialMaskDestroyed OnMaterialMaskDestroyed;
+
 	private:
 
 		/** Initial size for the handle-managed array */
@@ -100,6 +157,9 @@ namespace Chaos
 
 		/** Handle-managed array of global materials. This is pushed to all solvers who all maintain a copy */
 		THandleArray<FChaosPhysicsMaterial> Materials;
+
+		/** Handle-managed array of global material masks. This is pushed to all solvers who all maintain a copy */
+		THandleArray<FChaosPhysicsMaterialMask> MaterialMasks;
 	};
 }
 
diff --git a/Engine/Source/Runtime/Experimental/ChaosSolvers/Private/ChaosSolversModule.cpp b/Engine/Source/Runtime/Experimental/ChaosSolvers/Private/ChaosSolversModule.cpp
--- a/Engine/Source/Runtime/Experimental/ChaosSolvers/Private/ChaosSolversModule.cpp
+++ b/Engine/Source/Runtime/Experimental/ChaosSolvers/Private/ChaosSolversModule.cpp
@@ -233,6 +233,10 @@ void FChaosSolversModule::Initialize()
 		OnDestroyMaterialHandle = MaterialManager.OnMaterialDestroyed.Add(Chaos::FMaterialDestroyedDelegate::CreateRaw(this, &FChaosSolversModule::OnDestroyMaterial));
 		OnUpdateMaterialHandle = MaterialManager.OnMaterialUpdated.Add(Chaos::FMaterialUpdatedDelegate::CreateRaw(this, &FChaosSolversModule::OnUpdateMaterial));
 
+		OnCreateMaterialMaskHandle = MaterialManager.OnMaterialMaskCreated.Add(Chaos::FMaterialMaskCreatedDelegate::CreateRaw(this, &FChaosSolversModule::OnCreateMaterialMask));
+		OnDestroyMaterialMaskHandle = MaterialManager.OnMaterialMaskDestroyed.Add(Chaos::FMaterialMaskDestroyedDelegate::CreateRaw(this, &FChaosSolversModule::OnDestroyMaterialMask));
+		OnUpdateMaterialMaskHandle = MaterialManager.OnMaterialMaskUpdated.Add(Chaos::FMaterialMaskUpdatedDelegate::CreateRaw(this, &FChaosSolversModule::OnUpdateMaterialMask));
+
 		bModuleInitialized = true;
 	}
 }
@@ -488,7 +492,9 @@ Chaos::FPhysicsSolver* FChaosSolversModule::CreateSolver(bool bStandalone /*= fa
 		Chaos::FPhysicalMaterialManager& Manager =	Chaos::FPhysicalMaterialManager::Get();
 		NewSolver->QueryMaterialLock.WriteLock();
 		NewSolver->QueryMaterials = Manager.GetMasterMaterials();
+		NewSolver->QueryMaterialMasks = Manager.GetMasterMaterialMasks();
 		NewSolver->SimMaterials = Manager.GetMasterMaterials();
+		NewSolver->SimMaterialMasks = Manager.GetMasterMaterialMasks();
 		NewSolver->QueryMaterialLock.WriteUnlock();
 	}
 
@@ -959,6 +965,71 @@ void FChaosSolversModule::OnDestroyMaterial(Chaos::FMaterialHandle InHandle)
 	}
 }
 
+void FChaosSolversModule::OnUpdateMaterialMask(Chaos::FMaterialMaskHandle InHandle)
+{
+	check(Dispatcher);
+
+	// Grab the material
+	Chaos::FChaosPhysicsMaterialMask* MaterialMask = InHandle.Get();
+
+	if (ensure(MaterialMask))
+	{
+		for (Chaos::FPhysicsSolver* Solver : Solvers)
+		{
+			// Send a copy of the material to each solver
+			Dispatcher->EnqueueCommandImmediate(Solver, [InHandle, MaterialMaskCopy = *MaterialMask](Chaos::FPhysicsSolver* InSolver)
+			{
+				InSolver->UpdateMaterialMask(InHandle, MaterialMaskCopy);
+			});
+		}
+	}
+}
+
+void FChaosSolversModule::OnCreateMaterialMask(Chaos::FMaterialMaskHandle InHandle)
+{
+	check(Dispatcher);
+
+	// Grab the material
+	Chaos::FChaosPhysicsMaterialMask* MaterialMask = InHandle.Get();
+
+	if (ensure(MaterialMask))
+	{
+		for (Chaos::FPhysicsSolver* Solver : Solvers)
+		{
+			// Send a copy of the material to each solver
+			Dispatcher->EnqueueCommandImmediate(Solver, [InHandle, MaterialMaskCopy = *MaterialMask](Chaos::FPhysicsSolver* InSolver)
+			{
+				InSolver->CreateMaterialMask(InHandle, MaterialMaskCopy);
+			});
+		}
+	}
+}
+
+void FChaosSolversModule::OnDestroyMaterialMask(Chaos::FMaterialMaskHandle InHandle)
+{
+	// @todo(bgallagher)
+	//check(Dispatcher);
+	if (!ensure(Dispatcher))
+	{
+		return;
+	}
+
+	// Grab the material
+	Chaos::FChaosPhysicsMaterialMask* MaterialMask = InHandle.Get();
+
+	if (ensure(MaterialMask))
+	{
+		for (Chaos::FPhysicsSolver* Solver : Solvers)
+		{
+			// Notify each solver
+			Dispatcher->EnqueueCommandImmediate(Solver, [InHandle](Chaos::FPhysicsSolver* InSolver)
+			{
+				InSolver->DestroyMaterialMask(InHandle);
+			});
+		}
+	}
+}
+
 const IChaosSettingsProvider& FChaosSolversModule::GetSettingsProvider() const
 {
 	return SettingsProvider ? *SettingsProvider : Chaos::GDefaultChaosSettings;
diff --git a/Engine/Source/Runtime/Experimental/ChaosSolvers/Private/PBDRigidsSolver.cpp b/Engine/Source/Runtime/Experimental/ChaosSolvers/Private/PBDRigidsSolver.cpp
--- a/Engine/Source/Runtime/Experimental/ChaosSolvers/Private/PBDRigidsSolver.cpp
+++ b/Engine/Source/Runtime/Experimental/ChaosSolvers/Private/PBDRigidsSolver.cpp
@@ -715,10 +715,26 @@ namespace Chaos
 		SimMaterials.Destroy(InHandle.InnerHandle);
 	}
 
+	void FPBDRigidsSolver::UpdateMaterialMask(Chaos::FMaterialMaskHandle InHandle, const Chaos::FChaosPhysicsMaterialMask& InNewData)
+	{
+		*SimMaterialMasks.Get(InHandle.InnerHandle) = InNewData;
+	}
+
+	void FPBDRigidsSolver::CreateMaterialMask(Chaos::FMaterialMaskHandle InHandle, const Chaos::FChaosPhysicsMaterialMask& InNewData)
+	{
+		ensure(SimMaterialMasks.Create(InNewData) == InHandle.InnerHandle);
+	}
+
+	void FPBDRigidsSolver::DestroyMaterialMask(Chaos::FMaterialMaskHandle InHandle)
+	{
+		SimMaterialMasks.Destroy(InHandle.InnerHandle);
+	}
+
 	void FPBDRigidsSolver::SyncQueryMaterials()
 	{
 		TSolverQueryMaterialScope<ELockType::Write> Scope(this);
 		QueryMaterials = SimMaterials;
+		QueryMaterialMasks = SimMaterialMasks;
 	}
 
 }; // namespace Chaos
diff --git a/Engine/Source/Runtime/Experimental/ChaosSolvers/Public/ChaosSolversModule.h b/Engine/Source/Runtime/Experimental/ChaosSolvers/Public/ChaosSolversModule.h
--- a/Engine/Source/Runtime/Experimental/ChaosSolvers/Public/ChaosSolversModule.h
+++ b/Engine/Source/Runtime/Experimental/ChaosSolvers/Public/ChaosSolversModule.h
@@ -294,6 +294,10 @@ public:
 	void OnCreateMaterial(Chaos::FMaterialHandle InHandle);
 	void OnDestroyMaterial(Chaos::FMaterialHandle InHandle);
 
+	void OnUpdateMaterialMask(Chaos::FMaterialMaskHandle InHandle);
+	void OnCreateMaterialMask(Chaos::FMaterialMaskHandle InHandle);
+	void OnDestroyMaterialMask(Chaos::FMaterialMaskHandle InHandle);
+
 private:
 
 	/** Safe method for always getting a settings provider (from the external caller or an internal default) */
@@ -363,6 +367,10 @@ private:
 	FDelegateHandle OnCreateMaterialHandle;
 	FDelegateHandle OnDestroyMaterialHandle;
 	FDelegateHandle OnUpdateMaterialHandle;
+
+	FDelegateHandle OnCreateMaterialMaskHandle;
+	FDelegateHandle OnDestroyMaterialMaskHandle;
+	FDelegateHandle OnUpdateMaterialMaskHandle;
 };
 
 /**
diff --git a/Engine/Source/Runtime/Experimental/ChaosSolvers/Public/PBDRigidsSolver.h b/Engine/Source/Runtime/Experimental/ChaosSolvers/Public/PBDRigidsSolver.h
--- a/Engine/Source/Runtime/Experimental/ChaosSolvers/Public/PBDRigidsSolver.h
+++ b/Engine/Source/Runtime/Experimental/ChaosSolvers/Public/PBDRigidsSolver.h
@@ -285,10 +285,15 @@ namespace Chaos
 		void UpdateMaterial(Chaos::FMaterialHandle InHandle, const Chaos::FChaosPhysicsMaterial& InNewData);
 		void CreateMaterial(Chaos::FMaterialHandle InHandle, const Chaos::FChaosPhysicsMaterial& InNewData);
 		void DestroyMaterial(Chaos::FMaterialHandle InHandle);
+		void UpdateMaterialMask(Chaos::FMaterialMaskHandle InHandle, const Chaos::FChaosPhysicsMaterialMask& InNewData);
+		void CreateMaterialMask(Chaos::FMaterialMaskHandle InHandle, const Chaos::FChaosPhysicsMaterialMask& InNewData);
+		void DestroyMaterialMask(Chaos::FMaterialMaskHandle InHandle);
 
-		/** Access to the intenal material mirrors */
+		/** Access to the internal material mirrors */
 		const THandleArray<FChaosPhysicsMaterial>& GetQueryMaterials() const { return QueryMaterials; }
+		const THandleArray<FChaosPhysicsMaterialMask>& GetQueryMaterialMasks() const { return QueryMaterialMasks; }
 		const THandleArray<FChaosPhysicsMaterial>& GetSimMaterials() const { return SimMaterials; }
+		const THandleArray<FChaosPhysicsMaterialMask>& GetSimMaterialMasks() const { return SimMaterialMasks; }
 
 		/** Copy the simulation material list to the query material list, to be done when the SQ commits an update */
 		void SyncQueryMaterials();
@@ -366,8 +371,10 @@ namespace Chaos
 		// instead of having to lock the entire physics state of the runtime.
 		FRWLock QueryMaterialLock;
 		THandleArray<FChaosPhysicsMaterial> QueryMaterials;
+		THandleArray<FChaosPhysicsMaterialMask> QueryMaterialMasks;
 		THandleArray<FChaosPhysicsMaterial> SimMaterials;
-				
+		THandleArray<FChaosPhysicsMaterialMask> SimMaterialMasks;
+
 		template<ELockType>
 		friend struct TSolverQueryMaterialScope;
 	};
diff --git a/Engine/Source/Runtime/PhysicsCore/Public/PhysicsInterfaceDeclaresCore.h b/Engine/Source/Runtime/PhysicsCore/Public/PhysicsInterfaceDeclaresCore.h
--- a/Engine/Source/Runtime/PhysicsCore/Public/PhysicsInterfaceDeclaresCore.h
+++ b/Engine/Source/Runtime/PhysicsCore/Public/PhysicsInterfaceDeclaresCore.h
@@ -32,13 +32,16 @@ namespace Chaos
 	class TPerShapeData;
 
 	class FPhysicalMaterial;
+	class FPhysicalMaterialMask;
 
 	template<typename, uint32, uint32>
 	class THandle;
 
 	struct FMaterialHandle;
+	struct FMaterialMaskHandle;
 
 	class FChaosPhysicsMaterial;
+	class FChaosPhysicsMaterialMask;
 }
 
 // Temporary dummy types until SQ implemented
@@ -68,6 +71,7 @@ using FPhysicsShape = Chaos::TPerShapeData<float, 3>;
 using FPhysicsGeometry = Chaos::FImplicitObject;
 using FPhysicsCapsuleGeometry = Chaos::TCapsule<float>;
 using FPhysicsMaterial = Chaos::FChaosPhysicsMaterial;
+using FPhysicsMaterialMask = Chaos::FChaosPhysicsMaterialMask; 
 using FPhysicsActor = Chaos::TGeometryParticle<float,3>;
 
 template <typename T>
@@ -216,6 +220,7 @@ typedef FPhysInterface_Chaos				FPhysicsCommand;
 typedef FPhysicsShapeReference_Chaos		FPhysicsShapeHandle;
 typedef FPhysicsGeometryCollection_Chaos	FPhysicsGeometryCollection;
 typedef Chaos::FMaterialHandle				FPhysicsMaterialHandle;
+typedef Chaos::FMaterialMaskHandle			FPhysicsMaterialMaskHandle;
 typedef FPhysicsShapeAdapter_Chaos			FPhysicsShapeAdapter;
 typedef FPhysicsUserData_Chaos				FPhysicsUserData;
 
