Fix issue where soundclass loading behavior would only get applied if the sound class was set on the SoundWave, not on the SoundCue.

diff --git a/Engine/Source/Runtime/Engine/Classes/Sound/SoundCue.h b/Engine/Source/Runtime/Engine/Classes/Sound/SoundCue.h
--- a/Engine/Source/Runtime/Engine/Classes/Sound/SoundCue.h
+++ b/Engine/Source/Runtime/Engine/Classes/Sound/SoundCue.h
@@ -131,6 +131,7 @@ private:
 	uint8 bHasAttenuationNodeInitialized : 1;
 	uint8 bShouldApplyInteriorVolumes : 1;
 	uint8 bShouldApplyInteriorVolumesCached : 1;
+	uint8 bIsRetainingAudio : 1;
 
 public:
 
@@ -145,6 +146,7 @@ public:
 	virtual void Serialize(FStructuredArchive::FRecord Record) override;
 	virtual bool CanBeClusterRoot() const override;
 	virtual bool CanBeInCluster() const override;
+	virtual void BeginDestroy() override;
 	//~ End UObject Interface.
 
 	//~ Begin USoundBase Interface.
@@ -246,6 +248,10 @@ public:
 	/** Call this when stream caching is enabled to prime all SoundWave assets referenced by this Sound Cue. */
 	void PrimeSoundCue();
 
+	/** Call this when stream caching is enabled to retain all soundwave assets referenced by this sound cue. */
+	void RetainSoundCue();
+	void ReleaseRetainedAudio();
+
 protected:
 	bool RecursiveFindPathToNode(USoundNode* CurrentNode, const UPTRINT CurrentHash, const UPTRINT NodeHashToFind, TArray<USoundNode*>& OutPath) const;
 
@@ -258,6 +264,8 @@ private:
 	FDelegateHandle OnPostEngineInitHandle;
 	static int32 CachedQualityLevel;
 
+	void ApplySoundClassLoadingBehaviorToSoundWaves();
+
 public:
 
 	/**
diff --git a/Engine/Source/Runtime/Engine/Classes/Sound/SoundNode.h b/Engine/Source/Runtime/Engine/Classes/Sound/SoundNode.h
--- a/Engine/Source/Runtime/Engine/Classes/Sound/SoundNode.h
+++ b/Engine/Source/Runtime/Engine/Classes/Sound/SoundNode.h
@@ -73,6 +73,8 @@ class ENGINE_API USoundNode : public UObject
 	/** Stream of random numbers to be used by this instance of USoundNode */
 	FRandomStream RandomStream;
 
+	bool bIsRetainingAudio;
+
 public:
 	//~ Begin UObject Interface
 #if WITH_EDITOR
@@ -83,6 +85,7 @@ public:
 	virtual void Serialize(FArchive& Ar) override;
 	virtual bool CanBeClusterRoot() const override;
 	virtual bool CanBeInCluster() const override;
+	virtual void BeginDestroy() override;
 	//~ End UObject Interface
 
 	//
@@ -199,5 +202,17 @@ public:
 	 * @param bRecurse when true, this will cause all children of child nodes to be primed as well.
 	 */
 	virtual void PrimeChildWavePlayers(bool bRecurse);
+
+	/**
+	 * When this is called and stream caching is enabled,
+	 * any wave player sound nodes childed off of this node
+	 * will have their audio retained into the cache.
+	 *
+	 * @param bRecurse when true, this will cause all children of child nodes to be primed as well.
+	 */
+	virtual void RetainChildWavePlayers(bool bRecurse);
+
+	virtual void ReleaseRetainerOnChildWavePlayers(bool bRecurse);
+
 };
 
diff --git a/Engine/Source/Runtime/Engine/Private/SoundCue.cpp b/Engine/Source/Runtime/Engine/Private/SoundCue.cpp
--- a/Engine/Source/Runtime/Engine/Private/SoundCue.cpp
+++ b/Engine/Source/Runtime/Engine/Private/SoundCue.cpp
@@ -49,6 +49,8 @@ USoundCue::USoundCue(const FObjectInitializer& ObjectInitializer)
 	VolumeMultiplier = 0.75f;
 	PitchMultiplier = 1.0f;
 	SubtitlePriority = DEFAULT_SUBTITLE_PRIORITY;
+
+	bIsRetainingAudio = false;
 }
 
 #if WITH_EDITOR
@@ -96,6 +98,25 @@ void USoundCue::PrimeSoundCue()
 	}
 }
 
+void USoundCue::RetainSoundCue()
+{
+	if (FirstNode)
+	{
+		FirstNode->RetainChildWavePlayers(true);
+	}
+	bIsRetainingAudio = true;
+}
+
+void USoundCue::ReleaseRetainedAudio()
+{
+	if (FirstNode)
+	{
+		FirstNode->ReleaseRetainerOnChildWavePlayers(true);
+	}
+
+	bIsRetainingAudio = false;
+}
+
 void USoundCue::Serialize(FStructuredArchive::FRecord Record)
 {
 	FArchive& UnderlyingArchive = Record.GetUnderlyingArchive();
@@ -163,9 +184,24 @@ void USoundCue::PostLoad()
 
 	CacheAggregateValues();
 	
-	if (bPrimeOnLoad && FirstNode != nullptr)
+	ESoundWaveLoadingBehavior SoundClassLoadingBehavior = ESoundWaveLoadingBehavior::Inherited;
+
+	USoundClass* CurrentSoundClass = SoundClassObject;
+
+	// Recurse through this sound class's parents until we find an override.
+	while (SoundClassLoadingBehavior == ESoundWaveLoadingBehavior::Inherited && CurrentSoundClass != nullptr)
 	{
-		FirstNode->PrimeChildWavePlayers(true);
+		SoundClassLoadingBehavior = CurrentSoundClass->Properties.LoadingBehavior;
+		CurrentSoundClass = CurrentSoundClass->ParentClass;
+	}
+
+	if (SoundClassLoadingBehavior == ESoundWaveLoadingBehavior::RetainOnLoad)
+	{
+		RetainSoundCue();
+	}
+	else if (bPrimeOnLoad || SoundClassLoadingBehavior == ESoundWaveLoadingBehavior::PrimeOnLoad)
+	{
+		PrimeSoundCue();
 	}
 }
 
@@ -397,6 +433,16 @@ void USoundCue::AudioQualityChanged()
 	EvaluateNodes(false);
 }
 
+void USoundCue::BeginDestroy()
+{
+	Super::BeginDestroy();
+	
+	if (bIsRetainingAudio)
+	{
+		ReleaseRetainedAudio();
+	}
+}
+
 FString USoundCue::GetDesc()
 {
 	FString Description = TEXT( "" );
diff --git a/Engine/Source/Runtime/Engine/Private/SoundNode.cpp b/Engine/Source/Runtime/Engine/Private/SoundNode.cpp
--- a/Engine/Source/Runtime/Engine/Private/SoundNode.cpp
+++ b/Engine/Source/Runtime/Engine/Private/SoundNode.cpp
@@ -17,6 +17,7 @@ USoundNode::USoundNode(const FObjectInitializer& ObjectInitializer)
 	: Super(ObjectInitializer)
 {
 	RandomStream.Initialize(FApp::bUseFixedSeed ? GetFName() : NAME_None);
+	bIsRetainingAudio = false;
 }
 
 
@@ -119,6 +120,78 @@ void USoundNode::PrimeChildWavePlayers(bool bRecurse)
 	}
 }
 
+void USoundNode::RetainChildWavePlayers(bool bRecurse)
+{
+	if (FPlatformCompressionUtilities::IsCurrentPlatformUsingStreamCaching())
+	{
+		// Search child nodes for wave players, then prime their waves.
+		for (USoundNode* ChildNode : ChildNodes)
+		{
+			if (ChildNode)
+			{
+				ChildNode->ConditionalPostLoad();
+				if (bRecurse)
+				{
+					ChildNode->RetainChildWavePlayers(true);
+				}
+
+				USoundNodeWavePlayer* WavePlayer = Cast<USoundNodeWavePlayer>(ChildNode);
+				if (WavePlayer != nullptr)
+				{
+					USoundWave* SoundWave = WavePlayer->GetSoundWave();
+					if (SoundWave)
+					{
+						SoundWave->ConditionalPostLoad();
+						SoundWave->RetainCompressedAudio();
+					}
+				}
+			}
+		}
+	}
+
+	bIsRetainingAudio = true;
+}
+
+void USoundNode::ReleaseRetainerOnChildWavePlayers(bool bRecurse)
+{
+	if (bIsRetainingAudio && FPlatformCompressionUtilities::IsCurrentPlatformUsingStreamCaching())
+	{
+		// Search child nodes for wave players, then prime their waves.
+		for (USoundNode* ChildNode : ChildNodes)
+		{
+			if (ChildNode)
+			{
+				ChildNode->ConditionalPostLoad();
+				if (bRecurse)
+				{
+					ChildNode->ReleaseRetainerOnChildWavePlayers(true);
+				}
+
+				USoundNodeWavePlayer* WavePlayer = Cast<USoundNodeWavePlayer>(ChildNode);
+				if (WavePlayer != nullptr)
+				{
+					USoundWave* SoundWave = WavePlayer->GetSoundWave();
+					if (SoundWave)
+					{
+						SoundWave->ReleaseCompressedAudio();
+					}
+				}
+			}
+		}
+	}
+
+	bIsRetainingAudio = false;
+}
+
+void USoundNode::BeginDestroy()
+{
+	Super::BeginDestroy();
+	if (bIsRetainingAudio)
+	{
+		ReleaseRetainerOnChildWavePlayers(true);
+	}
+}
+
 void USoundNode::ParseNodes( FAudioDevice* AudioDevice, const UPTRINT NodeWaveInstanceHash, FActiveSound& ActiveSound, const FSoundParseParameters& ParseParams, TArray<FWaveInstance*>& WaveInstances )
 {
 	for( int32 i = 0; i < ChildNodes.Num() && i < GetMaxChildNodes(); ++i )
diff --git a/Engine/Source/Runtime/Engine/Private/SoundWave.cpp b/Engine/Source/Runtime/Engine/Private/SoundWave.cpp
--- a/Engine/Source/Runtime/Engine/Private/SoundWave.cpp
+++ b/Engine/Source/Runtime/Engine/Private/SoundWave.cpp
@@ -33,7 +33,7 @@
 #include "DSP/EnvelopeFollower.h"
 #include "DSP/BufferVectorOperations.h"
 #include "Misc/OutputDeviceArchiveWrapper.h"
-
+#include "Async/Async.h"
 #include "Misc/CommandLine.h"
 
 static int32 SoundWaveDefaultLoadingBehaviorCVar = 0;
@@ -2432,13 +2432,19 @@ void USoundWave::RetainCompressedAudio(bool bForceSync /*= false*/)
 	{
 		GetHandleForChunkOfAudio([WeakThis = MakeWeakObjectPtr(this)](FAudioChunkHandle OutHandle)
 		{
-			check(IsInGameThread());
-			
-			if (WeakThis.IsValid())
+			if (OutHandle.IsValid())
 			{
-				WeakThis->FirstChunk = OutHandle;
+				AsyncTask(ENamedThreads::GameThread, [WeakThis, OutHandle]() {
+					check(IsInGameThread());
+
+					if (WeakThis.IsValid())
+					{
+						WeakThis->FirstChunk = OutHandle;
+					}
+				});
+				
 			}
-		}, false, 1, ENamedThreads::GameThread);
+		}, false, 1);
 	}
 }
 
