Back out changelist 11112514

diff --git a/Engine/Source/Developer/AssetTools/Private/AssetTools.cpp b/Engine/Source/Developer/AssetTools/Private/AssetTools.cpp
--- a/Engine/Source/Developer/AssetTools/Private/AssetTools.cpp
+++ b/Engine/Source/Developer/AssetTools/Private/AssetTools.cpp
@@ -81,7 +81,6 @@
 #include "AssetTypeActions/AssetTypeActions_ObjectLibrary.h"
 #include "AssetTypeActions/AssetTypeActions_ParticleSystem.h"
 #include "AssetTypeActions/AssetTypeActions_PhysicalMaterial.h"
-#include "AssetTypeActions/AssetTypeActions_PhysicalMaterialMask.h"
 #include "AssetTypeActions/AssetTypeActions_PhysicsAsset.h"
 #include "AssetTypeActions/AssetTypeActions_PoseAsset.h"
 #include "AssetTypeActions/AssetTypeActions_PreviewMeshCollection.h"
@@ -240,7 +239,6 @@ UAssetToolsImpl::UAssetToolsImpl(const FObjectInitializer& ObjectInitializer)
 	RegisterAssetTypeActions(MakeShareable(new FAssetTypeActions_ObjectLibrary));
 	RegisterAssetTypeActions(MakeShareable(new FAssetTypeActions_ParticleSystem));
 	RegisterAssetTypeActions(MakeShareable(new FAssetTypeActions_PhysicalMaterial));
-	RegisterAssetTypeActions(MakeShareable(new FAssetTypeActions_PhysicalMaterialMask));
 	RegisterAssetTypeActions(MakeShareable(new FAssetTypeActions_PhysicsAsset));
 	RegisterAssetTypeActions(MakeShareable(new FAssetTypeActions_PreviewMeshCollection));
 	RegisterAssetTypeActions(MakeShareable(new FAssetTypeActions_ProceduralFoliageSpawner(FoliageCategoryBit)));
diff --git a/Engine/Source/Developer/AssetTools/Private/AssetTypeActions/AssetTypeActions_PhysicalMaterialMask.cpp b/Engine/Source/Developer/AssetTools/Private/AssetTypeActions/AssetTypeActions_PhysicalMaterialMask.cpp
deleted file mode 100644
diff --git a/Engine/Source/Developer/AssetTools/Private/AssetTypeActions/AssetTypeActions_PhysicalMaterialMask.h b/Engine/Source/Developer/AssetTools/Private/AssetTypeActions/AssetTypeActions_PhysicalMaterialMask.h
deleted file mode 100644
diff --git a/Engine/Source/Editor/ContentBrowser/Private/AssetContextMenu.cpp b/Engine/Source/Editor/ContentBrowser/Private/AssetContextMenu.cpp
--- a/Engine/Source/Editor/ContentBrowser/Private/AssetContextMenu.cpp
+++ b/Engine/Source/Editor/ContentBrowser/Private/AssetContextMenu.cpp
@@ -443,7 +443,7 @@ bool FAssetContextMenu::AddImportedAssetMenuOptions(UToolMenu* Menu)
 					FSlateIcon(FEditorStyle::GetStyleSetName(), "ContentBrowser.AssetActions.ReimportAsset"),
 					FUIAction(
 						FExecuteAction::CreateSP(this, &FAssetContextMenu::ExecuteReimport, (int32)INDEX_NONE),
-						FCanExecuteAction::CreateSP(this, &FAssetContextMenu::CanExecuteReimportAssetActions, ResolvedFilePaths)
+						FCanExecuteAction()
 					)
 				);
 				if (ValidSelectedAssetCount == 1)
@@ -456,7 +456,7 @@ bool FAssetContextMenu::AddImportedAssetMenuOptions(UToolMenu* Menu)
 						FSlateIcon(FEditorStyle::GetStyleSetName(), "ContentBrowser.AssetActions.ReimportAsset"),
 						FUIAction(
 							FExecuteAction::CreateSP(this, &FAssetContextMenu::ExecuteReimportWithNewFile, (int32)INDEX_NONE),
-							FCanExecuteAction::CreateSP(this, &FAssetContextMenu::CanExecuteReimportAssetActions, ResolvedFilePaths)
+							FCanExecuteAction()
 						)
 					);
 				}
@@ -1874,11 +1874,6 @@ bool FAssetContextMenu::AreImportedAssetActionsVisible() const
 
 bool FAssetContextMenu::CanExecuteImportedAssetActions(const TArray<FString> ResolvedFilePaths) const
 {
-	if (ResolvedFilePaths.Num() == 0)
-	{
-		return false;
-	}
-
 	// Verify that all the file paths are legitimate
 	for (const auto& SourceFilePath : ResolvedFilePaths)
 	{
@@ -1891,25 +1886,6 @@ bool FAssetContextMenu::CanExecuteImportedAssetActions(const TArray<FString> Res
 	return true;
 }
 
-bool FAssetContextMenu::CanExecuteReimportAssetActions(const TArray<FString> ResolvedFilePaths) const
-{
-	if (ResolvedFilePaths.Num() == 0)
-	{
-		return false;
-	}
-
-	// Verify that all the file paths are non-empty
-	for (const auto& SourceFilePath : ResolvedFilePaths)
-	{
-		if (!SourceFilePath.Len())
-		{
-			return false;
-		}
-	}
-
-	return true;
-}
-
 void FAssetContextMenu::ExecuteReimport(int32 SourceFileIndex /*= INDEX_NONE*/)
 {
 	// Reimport all selected assets
diff --git a/Engine/Source/Editor/ContentBrowser/Private/AssetContextMenu.h b/Engine/Source/Editor/ContentBrowser/Private/AssetContextMenu.h
--- a/Engine/Source/Editor/ContentBrowser/Private/AssetContextMenu.h
+++ b/Engine/Source/Editor/ContentBrowser/Private/AssetContextMenu.h
@@ -99,9 +99,6 @@ private:
 	/** Handler to check to see if imported asset actions are allowed */
 	bool CanExecuteImportedAssetActions(const TArray<FString> ResolvedFilePaths) const;
 
-	/** Handler to check to see if reimport asset actions are allowed */
-	bool CanExecuteReimportAssetActions(const TArray<FString> ResolvedFilePaths) const;
-
 	/** Handler for Reimport */
 	void ExecuteReimport(int32 SourceFileIndex = INDEX_NONE);
 
diff --git a/Engine/Source/Editor/UnrealEd/Classes/Factories/PhysicalMaterialMaskFactory.h b/Engine/Source/Editor/UnrealEd/Classes/Factories/PhysicalMaterialMaskFactory.h
deleted file mode 100644
diff --git a/Engine/Source/Editor/UnrealEd/Classes/ThumbnailRendering/PhysicalMaterialMaskThumbnailRenderer.h b/Engine/Source/Editor/UnrealEd/Classes/ThumbnailRendering/PhysicalMaterialMaskThumbnailRenderer.h
deleted file mode 100644
diff --git a/Engine/Source/Editor/UnrealEd/Private/Editor.cpp b/Engine/Source/Editor/UnrealEd/Private/Editor.cpp
--- a/Engine/Source/Editor/UnrealEd/Private/Editor.cpp
+++ b/Engine/Source/Editor/UnrealEd/Private/Editor.cpp
@@ -18,7 +18,6 @@
 #include "Factories/ReimportSoundFactory.h"
 #include "Factories/ReimportSoundSurroundFactory.h"
 #include "Factories/ReimportTextureFactory.h"
-#include "Factories/PhysicalMaterialMaskFactory.h"
 #include "Engine/BlueprintGeneratedClass.h"
 #include "Misc/ScopedSlowTask.h"
 #include "UObject/UObjectIterator.h"
@@ -776,9 +775,6 @@ FReimportManager::FReimportManager()
 
 	// Create reimport handler for FBX scene
 	UReimportFbxSceneFactory::StaticClass();
-
-	// Create reimport handler for PhysicalMaterialMasks
-	UPhysicalMaterialMaskFactory::StaticClass();
 }
 
 FReimportManager::~FReimportManager()
diff --git a/Engine/Source/Editor/UnrealEd/Private/Factories/EditorFactories.cpp b/Engine/Source/Editor/UnrealEd/Private/Factories/EditorFactories.cpp
--- a/Engine/Source/Editor/UnrealEd/Private/Factories/EditorFactories.cpp
+++ b/Engine/Source/Editor/UnrealEd/Private/Factories/EditorFactories.cpp
@@ -98,7 +98,6 @@
 #include "Factories/PackageFactory.h"
 #include "Factories/ParticleSystemFactoryNew.h"
 #include "Factories/PhysicalMaterialFactoryNew.h"
-#include "Factories/PhysicalMaterialMaskFactory.h"
 #include "Factories/PolysFactory.h"
 #include "Factories/ReverbEffectFactory.h"
 #include "Factories/SoundAttenuationFactory.h"
@@ -156,8 +155,6 @@
 #include "Materials/MaterialParameterCollection.h"
 #include "Engine/ObjectLibrary.h"
 #include "PhysicalMaterials/PhysicalMaterial.h"
-#include "PhysicalMaterials/PhysicalMaterialMask.h"
-#include "PhysicalMaterialMaskImport.h"
 #include "Engine/Polys.h"
 #include "Sound/ReverbEffect.h"
 #include "Sound/SoundCue.h"
@@ -1991,118 +1988,6 @@ UObject* UPhysicalMaterialFactoryNew::FactoryCreateNew(UClass* Class, UObject* I
 	}
 }
 
-/*------------------------------------------------------------------------------
-	UPhysicalMaterialMaskFactory.
-------------------------------------------------------------------------------*/
-UPhysicalMaterialMaskFactory::UPhysicalMaterialMaskFactory(const FObjectInitializer& ObjectInitializer)
-	: Super(ObjectInitializer)
-{
-	SupportedClass = UPhysicalMaterialMask::StaticClass();
-
-	bCreateNew = true;
-	bText = false;
-	bEditAfterNew = false;
-	bEditorImport = true;
-
-	// Required to allow texture factory to take priority when importing new image files
-	ImportPriority = DefaultImportPriority - 1;
-}
-
-bool UPhysicalMaterialMaskFactory::ConfigureProperties()
-{
-	// nullptr the DataAssetClass so we can check for selection
-	PhysicalMaterialMaskClass = nullptr;
-
-	// Load the classviewer module to display a class picker
-	FClassViewerModule& ClassViewerModule = FModuleManager::LoadModuleChecked<FClassViewerModule>("ClassViewer");
-
-	// Fill in options
-	FClassViewerInitializationOptions Options;
-	Options.Mode = EClassViewerMode::ClassPicker;
-
-	TSharedPtr<FAssetClassParentFilter> Filter = MakeShareable(new FAssetClassParentFilter);
-	Options.ClassFilter = Filter;
-
-	Filter->DisallowedClassFlags = CLASS_Abstract | CLASS_Deprecated | CLASS_NewerVersionExists;
-	Filter->AllowedChildrenOfClasses.Add(UPhysicalMaterialMask::StaticClass());
-
-	const FText TitleText = LOCTEXT("CreatePhysicalMaterialMask", "Pick Physical Material Mask Class");
-	UClass* ChosenClass = nullptr;
-	const bool bPressedOk = SClassPickerDialog::PickClass(TitleText, Options, ChosenClass, UPhysicalMaterialMask::StaticClass());
-
-	if (bPressedOk)
-	{
-		PhysicalMaterialMaskClass = ChosenClass;
-	}
-
-	return bPressedOk;
-}
-UObject* UPhysicalMaterialMaskFactory::FactoryCreateNew(UClass* Class, UObject* InParent, FName Name, EObjectFlags Flags, UObject* Context, FFeedbackContext* Warn)
-{
-	if (PhysicalMaterialMaskClass != nullptr)
-	{
-		return NewObject<UPhysicalMaterialMask>(InParent, PhysicalMaterialMaskClass, Name, Flags | RF_Transactional);
-	}
-	else
-	{
-		// if we have no data asset class, use the passed-in class instead
-		check(Class->IsChildOf(UPhysicalMaterialMask::StaticClass()));
-		return NewObject<UPhysicalMaterialMask>(InParent, Class, Name, Flags);
-	}
-}
-
-bool UPhysicalMaterialMaskFactory::CanReimport( UObject* Obj, TArray<FString>& OutFilenames )
-{	
-	UPhysicalMaterialMask* PhysMatMask = Cast<UPhysicalMaterialMask>(Obj);
-	if (PhysMatMask)
-	{
-		if (PhysMatMask->AssetImportData)
-		{
-			FString FileExtension = FPaths::GetExtension(PhysMatMask->AssetImportData->GetFirstFilename());
-			if (FileExtension.Equals(TEXT("png"), ESearchCase::IgnoreCase) || FileExtension.Equals("jpg", ESearchCase::IgnoreCase))
-			{
-				OutFilenames.Add(PhysMatMask->AssetImportData->GetFirstFilename());
-			}
-		}
-
-		return true;
-	}
-	return false;
-}
-
-void UPhysicalMaterialMaskFactory::SetReimportPaths( UObject* Obj, const TArray<FString>& NewReimportPaths )
-{	
-	UPhysicalMaterialMask* PhysMatMask = Cast<UPhysicalMaterialMask>(Obj);
-	if (PhysMatMask && ensure(NewReimportPaths.Num() == 1))
-	{
-		PhysMatMask->Modify();
-
-		if (!PhysMatMask->AssetImportData)
-		{
-			PhysMatMask->AssetImportData = NewObject<UAssetImportData>(PhysMatMask, TEXT("AssetImportData"), RF_NoFlags);
-		}
-
-		PhysMatMask->AssetImportData->UpdateFilenameOnly(NewReimportPaths[0]);
-	}
-}
-
-EReimportResult::Type UPhysicalMaterialMaskFactory::Reimport( UObject* Obj )
-{
-	if (!Obj || !Obj->IsA(UPhysicalMaterialMask::StaticClass()))
-	{
-		return EReimportResult::Failed;
-	}
-
-	UPhysicalMaterialMask* PhysMatMask = Cast<UPhysicalMaterialMask>(Obj);
-
-	return FPhysicalMaterialMaskImport::ReimportMaskTexture(PhysMatMask);
-}
-
-int32 UPhysicalMaterialMaskFactory::GetPriority() const
-{
-	return ImportPriority;
-}
-
 /*------------------------------------------------------------------------------
 	UInterpDataFactoryNew.
 ------------------------------------------------------------------------------*/
diff --git a/Engine/Source/Editor/UnrealEd/Private/PhysicalMaterialMaskImport.cpp b/Engine/Source/Editor/UnrealEd/Private/PhysicalMaterialMaskImport.cpp
deleted file mode 100644
diff --git a/Engine/Source/Editor/UnrealEd/Private/ThumbnailRendering/PhysicalMaterialMaskThumbnailRenderer.cpp b/Engine/Source/Editor/UnrealEd/Private/ThumbnailRendering/PhysicalMaterialMaskThumbnailRenderer.cpp
deleted file mode 100644
diff --git a/Engine/Source/Editor/UnrealEd/Public/PhysicalMaterialMaskImport.h b/Engine/Source/Editor/UnrealEd/Public/PhysicalMaterialMaskImport.h
deleted file mode 100644
diff --git a/Engine/Source/Runtime/Core/Private/UObject/DevObjectVersion.cpp b/Engine/Source/Runtime/Core/Private/UObject/DevObjectVersion.cpp
--- a/Engine/Source/Runtime/Core/Private/UObject/DevObjectVersion.cpp
+++ b/Engine/Source/Runtime/Core/Private/UObject/DevObjectVersion.cpp
@@ -25,7 +25,6 @@
 #include "UObject/NiagaraObjectVersion.h"
 #include "UObject/DestructionObjectVersion.h"
 #include "UObject/ExternalPhysicsCustomObjectVersion.h"
-#include "UObject/ExternalPhysicsMaterialCustomObjectVersion.h"
 #include "UObject/CineCameraObjectVersion.h"
 
 DEFINE_LOG_CATEGORY_STATIC(LogDevObjectVersion, Log, All);
@@ -170,13 +169,6 @@ FDevVersionRegistration GRegisterDestructionObjectVersion(FDestructionObjectVers
 const FGuid FExternalPhysicsCustomObjectVersion::GUID(0x35F94A83, 0xE258406C, 0xA31809F5, 0x9610247C);
 // Register Physics custom version with Core
 FDevVersionRegistration GRegisterExternalPhysicsCustomVersion(FExternalPhysicsCustomObjectVersion::GUID, FExternalPhysicsCustomObjectVersion::LatestVersion, TEXT("Dev-Physics-Ext"));
-// Unique Physics Object version id
-
-// Unique Physics Material Object version id
-const FGuid FExternalPhysicsMaterialCustomObjectVersion::GUID(0xB68FC16E, 0x8B1B42E2, 0xB453215C, 0x058844FE);
-// Register Physics custom version with Core
-FDevVersionRegistration GRegisterExternalPhysicsMaterialCustomVersion(FExternalPhysicsMaterialCustomObjectVersion::GUID, FExternalPhysicsMaterialCustomObjectVersion::LatestVersion, TEXT("Dev-PhysicsMaterial-Chaos"));
-// Unique PhysicsMaterial  Object version id
 
 // Unique CineCamera Object version id
 const FGuid FCineCameraObjectVersion::GUID(0xB2E18506, 0x4273CFC2, 0xA54EF4BB, 0x758BBA07);
diff --git a/Engine/Source/Runtime/Core/Public/UObject/ExternalPhysicsMaterialCustomObjectVersion.h b/Engine/Source/Runtime/Core/Public/UObject/ExternalPhysicsMaterialCustomObjectVersion.h
deleted file mode 100644
diff --git a/Engine/Source/Runtime/Engine/Classes/Engine/EngineTypes.h b/Engine/Source/Runtime/Engine/Classes/Engine/EngineTypes.h
--- a/Engine/Source/Runtime/Engine/Classes/Engine/EngineTypes.h
+++ b/Engine/Source/Runtime/Engine/Classes/Engine/EngineTypes.h
@@ -18,7 +18,6 @@
 class AActor;
 class UDecalComponent;
 class UPhysicalMaterial;
-class UPhysicalMaterialMask;
 class UPrimitiveComponent;
 class USceneComponent;
 
@@ -2591,9 +2590,6 @@ struct FMeshBuildSettings
 	UPROPERTY(EditAnywhere, Category=BuildSettings, meta=(DisplayName="Two-Sided Distance Field Generation"))
 	uint8 bGenerateDistanceFieldAsIfTwoSided:1;
 
-	UPROPERTY(EditAnywhere, Category=BuildSettings, meta=(DisplayName="Enable Physical Material Mask"))
-	uint8 bSupportFaceRemap : 1;
-
 	UPROPERTY(EditAnywhere, Category=BuildSettings)
 	int32 MinLightmapResolution;
 
@@ -3652,24 +3648,6 @@ enum EPhysicalSurface
 	SurfaceType_Max UMETA(Hidden)
 };
 
-/** Types of valid physical material mask colors which may be associated with a physical material */
-UENUM(BlueprintType)
-namespace EPhysicalMaterialMaskColor
-{
-	enum Type
-	{
-		Red,
-		Green,
-		Blue,
-		Cyan,
-		Magenta,
-		Yellow,
-		White,
-		Black,
-		MAX
-	};
-}
-
 /** Describes how often this component is allowed to move. */
 UENUM(BlueprintType)
 namespace EComponentMobility
diff --git a/Engine/Source/Runtime/Engine/Classes/Engine/StaticMesh.h b/Engine/Source/Runtime/Engine/Classes/Engine/StaticMesh.h
--- a/Engine/Source/Runtime/Engine/Classes/Engine/StaticMesh.h
+++ b/Engine/Source/Runtime/Engine/Classes/Engine/StaticMesh.h
@@ -664,14 +664,6 @@ class UStaticMesh : public UStreamableRenderAsset, public IInterface_CollisionDa
 	UPROPERTY(EditAnywhere, AdvancedDisplay, Category = StaticMesh)
 	uint8 bSupportUniformlyDistributedSampling : 1;
 
-	/** 
-		If true, complex collision data will store UVs and face remap table for use when performing
-	    PhysicalMaterialMask lookups in cooked builds. Note the increased memory cost for this
-		functionality.
-	*/
-	UPROPERTY(EditAnywhere, AdvancedDisplay, Category = StaticMesh)
-	uint8 bSupportPhysicalMaterialMasks : 1;
-
 	/**
 	 * If true, StaticMesh has been built at runtime
 	 */
diff --git a/Engine/Source/Runtime/Engine/Classes/Materials/Material.h b/Engine/Source/Runtime/Engine/Classes/Materials/Material.h
--- a/Engine/Source/Runtime/Engine/Classes/Materials/Material.h
+++ b/Engine/Source/Runtime/Engine/Classes/Materials/Material.h
@@ -15,15 +15,12 @@
 #include "Materials/MaterialFunction.h"
 #include "Materials/MaterialLayersFunctions.h"
 #include "Templates/UniquePtr.h"
-#if WITH_CHAOS
-#include "Physics/PhysicsInterfaceCore.h"
-#endif
+
 #include "Material.generated.h"
 
 class ITargetPlatform;
 class UMaterialExpressionComment;
 class UPhysicalMaterial;
-class UPhysicalMaterialMask;
 class USubsurfaceProfile;
 class UTexture;
 
@@ -347,14 +344,6 @@ class ENGINE_VTABLE UMaterial : public UMaterialInterface
 	UPROPERTY(EditAnywhere, Category=PhysicalMaterial)
 	class UPhysicalMaterial* PhysMaterial;
 
-	/** Physical material mask to use for this graphics material. Used for sounds, effects etc.*/
-	UPROPERTY(EditAnywhere, Category = PhysicalMaterial)
-	class UPhysicalMaterialMask* PhysMaterialMask;
-
-	/** Physical material mask map to use for this graphics material. Used for sounds, effects etc.*/
-	UPROPERTY(EditAnywhere, Category = PhysicalMaterialMask)
-	class UPhysicalMaterial* PhysicalMaterialMap[EPhysicalMaterialMaskColor::MAX];
-
 	// Reflection.
 #if WITH_EDITORONLY_DATA
 	UPROPERTY()
@@ -1027,8 +1016,6 @@ public:
 	ENGINE_API virtual bool GetRefractionSettings(float& OutBiasValue) const override;
 	ENGINE_API virtual FMaterialRenderProxy* GetRenderProxy() const override;
 	ENGINE_API virtual UPhysicalMaterial* GetPhysicalMaterial() const override;
-	ENGINE_API virtual UPhysicalMaterialMask* GetPhysicalMaterialMask() const override;
-	ENGINE_API virtual UPhysicalMaterial* GetPhysicalMaterialFromMap(int32 Index) const override;
 	ENGINE_API virtual void GetUsedTextures(TArray<UTexture*>& OutTextures, EMaterialQualityLevel::Type QualityLevel, bool bAllQualityLevels, ERHIFeatureLevel::Type FeatureLevel, bool bAllFeatureLevels) const override;
 	ENGINE_API virtual void GetUsedTexturesAndIndices(TArray<UTexture*>& OutTextures, TArray< TArray<int32> >& OutIndices, EMaterialQualityLevel::Type QualityLevel, ERHIFeatureLevel::Type FeatureLevel) const override;
 	ENGINE_API virtual void OverrideTexture(const UTexture* InTextureToOverride, UTexture* OverrideTexture, ERHIFeatureLevel::Type InFeatureLevel) override;
diff --git a/Engine/Source/Runtime/Engine/Classes/Materials/MaterialInstance.h b/Engine/Source/Runtime/Engine/Classes/Materials/MaterialInstance.h
--- a/Engine/Source/Runtime/Engine/Classes/Materials/MaterialInstance.h
+++ b/Engine/Source/Runtime/Engine/Classes/Materials/MaterialInstance.h
@@ -13,9 +13,6 @@
 #include "Materials/Material.h"
 #include "Materials/MaterialInstanceBasePropertyOverrides.h"
 #include "Misc/App.h"
-#if WITH_CHAOS
-#include "Physics/PhysicsInterfaceCore.h"
-#endif
 #include "MaterialInstance.generated.h"
 
 class ITargetPlatform;
@@ -299,10 +296,6 @@ class ENGINE_VTABLE UMaterialInstance : public UMaterialInterface
 	UPROPERTY(EditAnywhere, BlueprintReadWrite, Category=MaterialInstance)
 	class UPhysicalMaterial* PhysMaterial;
 
-	/** Physical material map used with physical material mask, when it exists.*/
-	UPROPERTY(EditAnywhere, Category = PhysicalMaterialMask)
-	class UPhysicalMaterial* PhysicalMaterialMap[EPhysicalMaterialMaskColor::MAX];
-
 	/** Parent material. */
 	UPROPERTY(EditAnywhere, BlueprintReadOnly, Category=MaterialInstance, AssetRegistrySearchable)
 	class UMaterialInterface* Parent;
@@ -498,8 +491,6 @@ public:
 	virtual ENGINE_API bool IsDependent(UMaterialInterface* TestDependency) override;
 	virtual ENGINE_API FMaterialRenderProxy* GetRenderProxy() const override;
 	virtual ENGINE_API UPhysicalMaterial* GetPhysicalMaterial() const override;
-	virtual ENGINE_API UPhysicalMaterialMask* GetPhysicalMaterialMask() const override;
-	virtual ENGINE_API UPhysicalMaterial* GetPhysicalMaterialFromMap(int32 Index) const override;
 	virtual ENGINE_API bool UpdateLightmassTextureTracking() override;
 	virtual ENGINE_API bool GetCastShadowAsMasked() const override;
 	virtual ENGINE_API float GetEmissiveBoost() const override;
diff --git a/Engine/Source/Runtime/Engine/Classes/Materials/MaterialInstanceConstant.h b/Engine/Source/Runtime/Engine/Classes/Materials/MaterialInstanceConstant.h
--- a/Engine/Source/Runtime/Engine/Classes/Materials/MaterialInstanceConstant.h
+++ b/Engine/Source/Runtime/Engine/Classes/Materials/MaterialInstanceConstant.h
@@ -9,8 +9,6 @@
 
 #include "MaterialInstanceConstant.generated.h"
 
-class UPhysicalMaterialMask;
-
 /**
  * Material Instances may be used to change the appearance of a material without incurring an expensive recompilation of the material.
  * General modification of the material cannot be supported without recompilation, so the instances are limited to changing the values of
@@ -37,14 +35,6 @@ class ENGINE_VTABLE UMaterialInstanceConstant : public UMaterialInstance
 	virtual ENGINE_API void PostEditChangeProperty(FPropertyChangedEvent& PropertyChangedEvent) override;
 #endif
 
-	/** Physical material mask to use for this graphics material. Used for sounds, effects etc.*/
-	UPROPERTY(EditAnywhere, Category = PhysicalMaterial)
-	class UPhysicalMaterialMask* PhysMaterialMask;
-
-	// Begin UMaterialInterface interface.
-	ENGINE_API virtual UPhysicalMaterialMask* GetPhysicalMaterialMask() const override;
-	// End UMaterialInterface interface.
-
 	/** Get the scalar (float) parameter value from an MIC */
 	UFUNCTION(BlueprintCallable, meta=(DisplayName = "GetScalarParameterValue", ScriptName = "GetScalarParameterValue", Keywords = "GetFloatParameterValue"), Category="Rendering|Material")
 	float K2_GetScalarParameterValue(FName ParameterName);
diff --git a/Engine/Source/Runtime/Engine/Classes/Materials/MaterialInterface.h b/Engine/Source/Runtime/Engine/Classes/Materials/MaterialInterface.h
--- a/Engine/Source/Runtime/Engine/Classes/Materials/MaterialInterface.h
+++ b/Engine/Source/Runtime/Engine/Classes/Materials/MaterialInterface.h
@@ -17,9 +17,6 @@
 #include "Materials/MaterialLayersFunctions.h"
 #include "Interfaces/Interface_AssetUserData.h"
 #include "MaterialSceneTextureId.h"
-#if WITH_CHAOS
-#include "Physics/PhysicsInterfaceCore.h"
-#endif
 #include "MaterialInterface.generated.h"
 
 class FMaterialCompiler;
@@ -27,7 +24,6 @@ class FMaterialRenderProxy;
 class FMaterialResource;
 class UMaterial;
 class UPhysicalMaterial;
-class UPhysicalMaterialMask;
 class USubsurfaceProfile;
 class UTexture;
 struct FPrimitiveViewRelevance;
@@ -368,20 +364,6 @@ public:
 	UFUNCTION(BlueprintCallable, Category = "Physics|Material")
 	virtual UPhysicalMaterial* GetPhysicalMaterial() const PURE_VIRTUAL(UMaterialInterface::GetPhysicalMaterial,return NULL;);
 
-	/**
-	 * Return a pointer to the physical material mask used by this material instance.
-	 * @return The physical material.
-	 */
-	UFUNCTION(BlueprintCallable, Category = "Physics|Material")
-	virtual UPhysicalMaterialMask* GetPhysicalMaterialMask() const PURE_VIRTUAL(UMaterialInterface::GetPhysicalMaterialMask, return nullptr;);
-
-	/**
-	 * Return a pointer to the physical material from mask map at given index.
-	 * @return The physical material.
-	 */
-	UFUNCTION(BlueprintCallable, Category = "Physics|Material")
-	virtual UPhysicalMaterial* GetPhysicalMaterialFromMap(int32 Index) const PURE_VIRTUAL(UMaterialInterface::GetPhysicalMaterialFromMap, return nullptr;);
-
 	/** Return the textures used to render this material. */
 	virtual void GetUsedTextures(TArray<UTexture*>& OutTextures, EMaterialQualityLevel::Type QualityLevel, bool bAllQualityLevels, ERHIFeatureLevel::Type FeatureLevel, bool bAllFeatureLevels) const
 		PURE_VIRTUAL(UMaterialInterface::GetUsedTextures,);
diff --git a/Engine/Source/Runtime/Engine/Classes/PhysicalMaterials/PhysicalMaterialMask.h b/Engine/Source/Runtime/Engine/Classes/PhysicalMaterials/PhysicalMaterialMask.h
deleted file mode 100644
diff --git a/Engine/Source/Runtime/Engine/Classes/PhysicsEngine/BodyInstance.h b/Engine/Source/Runtime/Engine/Classes/PhysicsEngine/BodyInstance.h
--- a/Engine/Source/Runtime/Engine/Classes/PhysicsEngine/BodyInstance.h
+++ b/Engine/Source/Runtime/Engine/Classes/PhysicsEngine/BodyInstance.h
@@ -16,7 +16,6 @@
 
 class UBodySetup;
 class UPhysicalMaterial;
-class UPhysicalMaterialMask;
 class UPrimitiveComponent;
 struct FBodyInstance;
 struct FCollisionNotifyInfo;
@@ -589,20 +588,14 @@ public:
 	/** Find the correct PhysicalMaterial for simple geometry on a given body and owner. This is really for internal use during serialization */
 	static UPhysicalMaterial* GetSimplePhysicalMaterial(const FBodyInstance* BodyInstance, TWeakObjectPtr<UPrimitiveComponent> Owner, TWeakObjectPtr<UBodySetup> BodySetupPtr);
 
-	/** Get the complex PhysicalMaterials array for this body */
+	/** Get the complex PhysicalMaterial array for this body */
 	TArray<UPhysicalMaterial*> GetComplexPhysicalMaterials() const;
 
-	/** Get the complex PhysicalMaterials and PhysicalMaterialMasks array for this body */
-	TArray<UPhysicalMaterial*> GetComplexPhysicalMaterials(TArray<FPhysicalMaterialMaskParams>& OutPhysMaterialMasks) const;
+	/** Find the correct PhysicalMaterial for simple geometry on a given body and owner. This is really for internal use during serialization */
+	static void GetComplexPhysicalMaterials(const FBodyInstance* BodyInstance, TWeakObjectPtr<UPrimitiveComponent> Owner, TArray<UPhysicalMaterial*>& OutPhysicalMaterials);
 
 	/** Get the complex PhysicalMaterials for this body */
-	void GetComplexPhysicalMaterials(TArray<UPhysicalMaterial*> &OutPhysMaterials) const;
-
-	/** Get the complex PhysicalMaterials and PhysicalMaterialMasks for this body */
-	void GetComplexPhysicalMaterials(TArray<UPhysicalMaterial*> &OutPhysMaterials, TArray<FPhysicalMaterialMaskParams>& OutPhysMaterialMasks) const;
-
-	/** Find the correct PhysicalMaterial and PhysicalMaterialMasks for complex geometry on a given body and owner. This is really for internal use during serialization */
-	static void GetComplexPhysicalMaterials(const FBodyInstance* BodyInstance, TWeakObjectPtr<UPrimitiveComponent> Owner, TArray<UPhysicalMaterial*>& OutPhysMaterials, TArray<FPhysicalMaterialMaskParams>* OutPhysMaterialMasks = nullptr);
+	void GetComplexPhysicalMaterials(TArray<UPhysicalMaterial*> &PhysMaterials) const;
 
 	/** Returns the slope override struct for this instance. If we don't have our own custom setting, it will return the setting from the body setup. */
 	const struct FWalkableSlopeOverride& GetWalkableSlopeOverride() const;
@@ -869,10 +862,10 @@ public:
 	 *  @param  SimplePhysMat			The material to use if a simple shape is provided (or complex materials are empty)
 	 *  @param  ComplexPhysMats			The array of materials to apply if a complex shape is provided
 	 */
-	static void ApplyMaterialToShape_AssumesLocked(const FPhysicsShapeHandle& InShape, UPhysicalMaterial* SimplePhysMat, const TArrayView<UPhysicalMaterial*>& ComplexPhysMats, const TArrayView<FPhysicalMaterialMaskParams>* ComplexPhysMatMasks = nullptr);
+	static void ApplyMaterialToShape_AssumesLocked(const FPhysicsShapeHandle& InShape, UPhysicalMaterial* SimplePhysMat, const TArrayView<UPhysicalMaterial*>& ComplexPhysMats);
 
 	/** Note: This function is not thread safe. Make sure you obtain the appropriate physics scene lock before calling it*/
-	void ApplyMaterialToInstanceShapes_AssumesLocked(UPhysicalMaterial* SimplePhysMat, TArray<UPhysicalMaterial*>& ComplexPhysMats, const TArrayView<FPhysicalMaterialMaskParams>& ComplexPhysMatMasks);
+	void ApplyMaterialToInstanceShapes_AssumesLocked(UPhysicalMaterial* SimplePhysMat, TArray<UPhysicalMaterial*>& ComplexPhysMats);
 
 	/** Update the instances collision filtering data */ 
 	void UpdatePhysicsFilterData();
diff --git a/Engine/Source/Runtime/Engine/Classes/PhysicsEngine/BodySetup.h b/Engine/Source/Runtime/Engine/Classes/PhysicsEngine/BodySetup.h
--- a/Engine/Source/Runtime/Engine/Classes/PhysicsEngine/BodySetup.h
+++ b/Engine/Source/Runtime/Engine/Classes/PhysicsEngine/BodySetup.h
@@ -19,7 +19,6 @@
 
 class ITargetPlatform;
 class UPhysicalMaterial;
-class UPhysicalMaterialMask;
 class UPrimitiveComponent;
 struct FShapeData;
 enum class EPhysXMeshCookFlags : uint8;
@@ -205,13 +204,6 @@ class UBodySetup : public UObject
 	UPROPERTY()
 	uint8 bGenerateMirroredCollision:1;
 
-	/** 
-	 * If true, the physics triangle mesh will store UVs and the face remap table. This is needed
-	 * to support physical material masks in scene queries. 
-	 */
-	UPROPERTY()
-	uint8 bSupportUVsAndFaceRemap:1;
-
 	/** Flag used to know if we have created the physics convex and tri meshes from the cooked data yet */
 	uint8 bCreatedPhysicsMeshes:1;
 
@@ -466,8 +458,7 @@ public:
 		FBodyInstance* OwningInstance, 
 		FVector& Scale3D, 
 		UPhysicalMaterial* SimpleMaterial,
-		TArray<UPhysicalMaterial*>& ComplexMaterials,
-		TArray<FPhysicalMaterialMaskParams>& ComplexMaterialMasks,
+		TArray<UPhysicalMaterial*>& ComplexMaterials, 
 		const FBodyCollisionData& BodyCollisionData,
 		const FTransform& RelativeTM = FTransform::Identity, 
 		TArray<FPhysicsShapeHandle>* NewShapes = NULL);
diff --git a/Engine/Source/Runtime/Engine/Private/Collision/CollisionConversions.cpp b/Engine/Source/Runtime/Engine/Private/Collision/CollisionConversions.cpp
--- a/Engine/Source/Runtime/Engine/Private/Collision/CollisionConversions.cpp
+++ b/Engine/Source/Runtime/Engine/Private/Collision/CollisionConversions.cpp
@@ -7,7 +7,6 @@
 #include "Collision/CollisionDebugDrawing.h"
 #include "Components/LineBatchComponent.h"
 #include "PhysicalMaterials/PhysicalMaterial.h"
-#include "PhysicalMaterials/PhysicalMaterialMask.h"
 #include "PhysicsEngine/PhysicsSettings.h"
 #include "PhysicsEngine/BodySetup.h"
 #include "Physics/PhysicsInterfaceUtils.h"
@@ -131,7 +130,7 @@ static FVector FindGeomOpposingNormal(ECollisionShapeType QueryGeomType, const F
 }
 
 /** Set info in the HitResult (Actor, Component, PhysMaterial, BoneName, Item) based on the supplied shape and face index */
-static void SetHitResultFromShapeAndFaceIndex(const FPhysicsShape& Shape,  const FPhysicsActor& Actor, const uint32 FaceIndex, const FVector& HitLocation, FHitResult& OutResult, bool bReturnPhysMat)
+static void SetHitResultFromShapeAndFaceIndex(const FPhysicsShape& Shape,  const FPhysicsActor& Actor, const uint32 FaceIndex, FHitResult& OutResult, bool bReturnPhysMat)
 {
 	SCOPE_CYCLE_COUNTER(STAT_CollisionSetHitResultFromShapeAndFaceIndex);
 	
@@ -184,18 +183,11 @@ static void SetHitResultFromShapeAndFaceIndex(const FPhysicsShape& Shape,  const
 
 		if (bReturnPhysMat)
 		{
-#if WITH_CHAOS
-			if (const FPhysicsMaterial* PhysicsMaterial = GetMaterialFromInternalFaceIndexAndHitLocation(Shape, Actor, FaceIndex, HitLocation))
+			// This function returns the single material in all cases other than trimesh or heightfield
+			if(const FPhysicsMaterial* PhysicsMaterial = GetMaterialFromInternalFaceIndex(Shape, Actor, FaceIndex))
 			{
 				OutResult.PhysMaterial = GetUserData(*PhysicsMaterial);
 			}
-#else
-			if (const FPhysicsMaterial* PhysicsMaterial = GetMaterialFromInternalFaceIndex(Shape, Actor, FaceIndex))
-			{
-				OutResult.PhysMaterial = GetUserData(*PhysicsMaterial);
-			}
-
-#endif
 		}
 	}
 
@@ -263,7 +255,7 @@ EConvertQueryResult ConvertQueryImpactHit(const UWorld* World, const FHitLocatio
 		if (Position.ContainsNaN())
 		{
 #if ENABLE_NAN_DIAGNOSTIC
-			SetHitResultFromShapeAndFaceIndex(HitShape, HitActor, InternalFaceIndex, OutResult.ImpactPoint, OutResult, bReturnPhysMat);
+			SetHitResultFromShapeAndFaceIndex(HitShape, HitActor, InternalFaceIndex, OutResult, bReturnPhysMat);
 			UE_LOG(LogCore, Error, TEXT("ConvertQueryImpactHit() NaN details:\n>> Actor:%s (%s)\n>> Component:%s\n>> Item:%d\n>> BoneName:%s\n>> Time:%f\n>> Distance:%f\n>> Location:%s\n>> bIsBlocking:%d\n>> bStartPenetrating:%d"),
 				*GetNameSafe(OutResult.GetActor()), OutResult.Actor.IsValid() ? *OutResult.GetActor()->GetPathName() : TEXT("no path"),
 				*GetNameSafe(OutResult.GetComponent()), OutResult.Item, *OutResult.BoneName.ToString(),
@@ -283,7 +275,7 @@ EConvertQueryResult ConvertQueryImpactHit(const UWorld* World, const FHitLocatio
 	if (bUseReturnedNormal && HitNormal.ContainsNaN())
 	{
 #if ENABLE_NAN_DIAGNOSTIC
-		SetHitResultFromShapeAndFaceIndex(HitShape, HitActor, InternalFaceIndex, OutResult.ImpactPoint, OutResult, bReturnPhysMat);
+		SetHitResultFromShapeAndFaceIndex(HitShape, HitActor, InternalFaceIndex, OutResult, bReturnPhysMat);
 		UE_LOG(LogCore, Error, TEXT("ConvertQueryImpactHit() NaN details:\n>> Actor:%s (%s)\n>> Component:%s\n>> Item:%d\n>> BoneName:%s\n>> Time:%f\n>> Distance:%f\n>> Location:%s\n>> bIsBlocking:%d\n>> bStartPenetrating:%d"),
 			*GetNameSafe(OutResult.GetActor()), OutResult.Actor.IsValid() ? *OutResult.GetActor()->GetPathName() : TEXT("no path"),
 			*GetNameSafe(OutResult.GetComponent()), OutResult.Item, *OutResult.BoneName.ToString(),
@@ -319,7 +311,7 @@ EConvertQueryResult ConvertQueryImpactHit(const UWorld* World, const FHitLocatio
 	OutResult.ImpactNormal = FindGeomOpposingNormal(SweptGeometryType, Hit, TraceStartToEnd, Normal);
 
 	// Fill in Actor, Component, material, etc.
-	SetHitResultFromShapeAndFaceIndex(HitShape, HitActor, InternalFaceIndex, OutResult.ImpactPoint, OutResult, bReturnPhysMat);
+	SetHitResultFromShapeAndFaceIndex(HitShape, HitActor, InternalFaceIndex, OutResult, bReturnPhysMat);
 
 	ECollisionShapeType GeomType = GetGeometryType(HitShape);
 
@@ -506,7 +498,7 @@ static bool ConvertOverlappedShapeToImpactHit(const UWorld* World, const FHitLoc
 
 	OutResult.Normal = OutResult.ImpactNormal;
 	
-	SetHitResultFromShapeAndFaceIndex(HitShape, HitActor, GetInternalFaceIndex(Hit), OutResult.ImpactPoint, OutResult, bReturnPhysMat);
+	SetHitResultFromShapeAndFaceIndex(HitShape, HitActor, GetInternalFaceIndex(Hit), OutResult, bReturnPhysMat);
 
 	return bBlockingHit;
 }
diff --git a/Engine/Source/Runtime/Engine/Private/Materials/Material.cpp b/Engine/Source/Runtime/Engine/Private/Materials/Material.cpp
--- a/Engine/Source/Runtime/Engine/Private/Materials/Material.cpp
+++ b/Engine/Source/Runtime/Engine/Private/Materials/Material.cpp
@@ -948,9 +948,6 @@ UMaterial::UMaterial(const FObjectInitializer& ObjectInitializer)
 
 	bIsPreviewMaterial = false;
 	bIsFunctionPreviewMaterial = false;
-
-	PhysMaterial = nullptr;
-	PhysMaterialMask = nullptr;
 }
 
 void UMaterial::PreSave(const class ITargetPlatform* TargetPlatform)
@@ -2966,20 +2963,6 @@ UPhysicalMaterial* UMaterial::GetPhysicalMaterial() const
 	return nullptr;
 }
 
-UPhysicalMaterialMask* UMaterial::GetPhysicalMaterialMask() const
-{
-	return PhysMaterialMask;
-}
-
-UPhysicalMaterial* UMaterial::GetPhysicalMaterialFromMap(int32 Index) const
-{
-	if (Index >= 0 && Index < EPhysicalMaterialMaskColor::MAX)
-	{
-		return PhysicalMaterialMap[Index];
-	}
-	return nullptr;
-}
-
 /** Helper functions for text output of properties... */
 #ifndef CASE_ENUM_TO_TEXT
 #define CASE_ENUM_TO_TEXT(txt) case txt: return TEXT(#txt);
@@ -4253,7 +4236,7 @@ bool UMaterial::CanEditChange(const FProperty* InProperty) const
 	{
 		FString PropertyName = InProperty->GetName();
 
-		if (PropertyName == GET_MEMBER_NAME_STRING_CHECKED(UMaterial, PhysMaterial) || PropertyName == GET_MEMBER_NAME_STRING_CHECKED(UMaterial, PhysMaterialMask))
+		if (PropertyName == GET_MEMBER_NAME_STRING_CHECKED(UMaterial, PhysMaterial))
 		{
 			return MaterialDomain == MD_Surface;
 		}
@@ -4445,7 +4428,7 @@ void UMaterial::PostEditChangeProperty(FPropertyChangedEvent& PropertyChangedEve
 	if( PropertyThatChanged ) 
 	{
 		// Don't recompile the material if we only changed the PhysMaterial property.
-		if (PropertyThatChanged->GetName() == TEXT("PhysMaterial") || PropertyThatChanged->GetName() == TEXT("PhysMaterialMask") || PropertyThatChanged->GetName() == TEXT("PhysicalMaterialMap"))
+		if (PropertyThatChanged->GetName() == TEXT("PhysMaterial"))
 		{
 			bRequiresCompilation = false;
 		}
diff --git a/Engine/Source/Runtime/Engine/Private/Materials/MaterialInstance.cpp b/Engine/Source/Runtime/Engine/Private/Materials/MaterialInstance.cpp
--- a/Engine/Source/Runtime/Engine/Private/Materials/MaterialInstance.cpp
+++ b/Engine/Source/Runtime/Engine/Private/Materials/MaterialInstance.cpp
@@ -430,12 +430,6 @@ UMaterialInstance::UMaterialInstance(const FObjectInitializer& ObjectInitializer
 	ReentrantFlag[0] = false;
 	ReentrantFlag[1] = false;
 	ShadingModels = MSM_Unlit;
-
-	PhysMaterial = nullptr;
-	for (UPhysicalMaterial*& PhysMat : PhysicalMaterialMap)
-	{
-		PhysMat = nullptr;
-	}
 }
 
 void UMaterialInstance::PostInitProperties()	
@@ -1961,20 +1955,6 @@ UPhysicalMaterial* UMaterialInstance::GetPhysicalMaterial() const
 	}
 }
 
-UPhysicalMaterialMask* UMaterialInstance::GetPhysicalMaterialMask() const
-{
-	return nullptr;
-}
-
-UPhysicalMaterial* UMaterialInstance::GetPhysicalMaterialFromMap(int32 Index) const
-{
-	if (Index < 0 || Index >= EPhysicalMaterialMaskColor::MAX)
-	{
-		return nullptr;
-	}
-	return PhysicalMaterialMap[Index];
-}
-
 void UMaterialInstance::GetStaticParameterValues(FStaticParameterSet& OutStaticParameters)
 {
 	check(IsInGameThread());
diff --git a/Engine/Source/Runtime/Engine/Private/Materials/MaterialInstanceConstant.cpp b/Engine/Source/Runtime/Engine/Private/Materials/MaterialInstanceConstant.cpp
--- a/Engine/Source/Runtime/Engine/Private/Materials/MaterialInstanceConstant.cpp
+++ b/Engine/Source/Runtime/Engine/Private/Materials/MaterialInstanceConstant.cpp
@@ -12,7 +12,6 @@
 UMaterialInstanceConstant::UMaterialInstanceConstant(const FObjectInitializer& ObjectInitializer)
 	: Super(ObjectInitializer)
 {
-	PhysMaterialMask = nullptr;
 }
 
 void UMaterialInstanceConstant::PostLoad()
@@ -42,11 +41,6 @@ UTexture* UMaterialInstanceConstant::K2_GetTextureParameterValue(FName Parameter
 	return Result;
 }
 
-UPhysicalMaterialMask* UMaterialInstanceConstant::GetPhysicalMaterialMask() const
-{
-	return PhysMaterialMask;
-}
-
 #if WITH_EDITOR
 void UMaterialInstanceConstant::PostEditChangeProperty(FPropertyChangedEvent& PropertyChangedEvent)
 {
diff --git a/Engine/Source/Runtime/Engine/Private/PhysicsEngine/BodyInstance.cpp b/Engine/Source/Runtime/Engine/Private/PhysicsEngine/BodyInstance.cpp
--- a/Engine/Source/Runtime/Engine/Private/PhysicsEngine/BodyInstance.cpp
+++ b/Engine/Source/Runtime/Engine/Private/PhysicsEngine/BodyInstance.cpp
@@ -469,12 +469,11 @@ void FBodyInstance::UpdatePhysicalMaterials()
 {
 	SCOPE_CYCLE_COUNTER(STAT_UpdatePhysMats);
 	UPhysicalMaterial* SimplePhysMat = GetSimplePhysicalMaterial();
-	TArray<FPhysicalMaterialMaskParams> ComplexPhysMatMasks;
-	TArray<UPhysicalMaterial*> ComplexPhysMats = GetComplexPhysicalMaterials(ComplexPhysMatMasks);
+	TArray<UPhysicalMaterial*> ComplexPhysMats = GetComplexPhysicalMaterials();
 
 	FPhysicsCommand::ExecuteWrite(GetActorReferenceWithWelding(), [&](const FPhysicsActorHandle& Actor)
 	{
-		ApplyMaterialToInstanceShapes_AssumesLocked(SimplePhysMat, ComplexPhysMats, ComplexPhysMatMasks);
+		ApplyMaterialToInstanceShapes_AssumesLocked(SimplePhysMat, ComplexPhysMats);
 	});
 }
 
@@ -1060,15 +1059,7 @@ struct FInitBodiesHelper
 	bool CreateShapes_AssumesLocked(FBodyInstance* Instance) const
 	{
 		UPhysicalMaterial* SimplePhysMat = Instance->GetSimplePhysicalMaterial();
-
-		TArray<UPhysicalMaterial*> ComplexPhysMats;
-		TArray<FPhysicalMaterialMaskParams> ComplexPhysMatMasks;
-
-#if WITH_CHAOS
-		ComplexPhysMats = Instance->GetComplexPhysicalMaterials(ComplexPhysMatMasks);
-#else
-		ComplexPhysMats = Instance->GetComplexPhysicalMaterials();
-#endif
+		TArray<UPhysicalMaterial*> ComplexPhysMats = Instance->GetComplexPhysicalMaterials();
 
 		FBodyCollisionData BodyCollisionData;
 		Instance->BuildBodyFilterData(BodyCollisionData.CollisionFilterData);
@@ -1077,7 +1068,7 @@ struct FInitBodiesHelper
 		bool bInitFail = false;
 
 		// #PHYS2 Call interface AddGeometry
-		BodySetup->AddShapesToRigidActor_AssumesLocked(Instance, Instance->Scale3D, SimplePhysMat, ComplexPhysMats, ComplexPhysMatMasks, BodyCollisionData, FTransform::Identity);
+		BodySetup->AddShapesToRigidActor_AssumesLocked(Instance, Instance->Scale3D, SimplePhysMat, ComplexPhysMats, BodyCollisionData, FTransform::Identity);
 
 		const int32 NumShapes = FPhysicsInterface::GetNumShapes(Instance->ActorHandle);
 		bInitFail |= NumShapes == 0;
@@ -1480,21 +1471,14 @@ bool FBodyInstance::Weld(FBodyInstance* TheirBody, const FTransform& TheirTM)
 		TheirBody->WeldParent = this;
 
 		UPhysicalMaterial* SimplePhysMat = TheirBody->GetSimplePhysicalMaterial();
+		TArray<UPhysicalMaterial*> ComplexPhysMats = TheirBody->GetComplexPhysicalMaterials();
 
-		TArray<UPhysicalMaterial*> ComplexPhysMats;
-		TArray<FPhysicalMaterialMaskParams> ComplexPhysMatMasks;
-	
-#if WITH_CHAOS
-		ComplexPhysMats = TheirBody->GetComplexPhysicalMaterials(ComplexPhysMatMasks);
-#else
-		ComplexPhysMats = TheirBody->GetComplexPhysicalMaterials();
-#endif
 		// This builds collision data based on this (parent) body, not their body. This gets fixed  up later though when PostShapeChange() calls UpdatePhysicsFilterData().
 		FBodyCollisionData BodyCollisionData;
 		BuildBodyFilterData(BodyCollisionData.CollisionFilterData);
 		BuildBodyCollisionFlags(BodyCollisionData.CollisionFlags, GetCollisionEnabled(), BodySetup->GetCollisionTraceFlag() == CTF_UseComplexAsSimple);
 
-		TheirBody->BodySetup->AddShapesToRigidActor_AssumesLocked(this, Scale3D, SimplePhysMat, ComplexPhysMats, ComplexPhysMatMasks, BodyCollisionData, RelativeTM, &PNewShapes);
+		TheirBody->BodySetup->AddShapesToRigidActor_AssumesLocked(this, Scale3D, SimplePhysMat, ComplexPhysMats, BodyCollisionData, RelativeTM, &PNewShapes);
 
 		FPhysicsInterface::SetSendsSleepNotifies_AssumesLocked(Actor, TheirBody->bGenerateWakeEvents);
 
@@ -2792,24 +2776,13 @@ TArray<UPhysicalMaterial*> FBodyInstance::GetComplexPhysicalMaterials() const
 	return PhysMaterials;
 }
 
-TArray<UPhysicalMaterial*> FBodyInstance::GetComplexPhysicalMaterials(TArray<FPhysicalMaterialMaskParams>& OutPhysMaterialMasks) const
-{
-	TArray<UPhysicalMaterial*> PhysMaterials;
-	GetComplexPhysicalMaterials(PhysMaterials, OutPhysMaterialMasks);
-	return PhysMaterials;
-}
-
-void FBodyInstance::GetComplexPhysicalMaterials(TArray<UPhysicalMaterial*>& OutPhysMaterials) const
+void FBodyInstance::GetComplexPhysicalMaterials(TArray<UPhysicalMaterial*>& PhysMaterials) const
 {
-	GetComplexPhysicalMaterials(this, OwnerComponent, OutPhysMaterials);
+	GetComplexPhysicalMaterials(this, OwnerComponent, PhysMaterials);
 }
 
-void FBodyInstance::GetComplexPhysicalMaterials(TArray<UPhysicalMaterial*>& OutPhysMaterials, TArray<FPhysicalMaterialMaskParams>& OutPhysMaterialMasks) const
-{
-	GetComplexPhysicalMaterials(this, OwnerComponent, OutPhysMaterials, &OutPhysMaterialMasks);
-}
 
-void FBodyInstance::GetComplexPhysicalMaterials(const FBodyInstance*, TWeakObjectPtr<UPrimitiveComponent> OwnerComp, TArray<UPhysicalMaterial*>& OutPhysMaterials, TArray<FPhysicalMaterialMaskParams>* OutPhysMaterialMasks)
+void FBodyInstance::GetComplexPhysicalMaterials(const FBodyInstance*, TWeakObjectPtr<UPrimitiveComponent> OwnerComp, TArray<UPhysicalMaterial*>& OutPhysicalMaterials)
 {
 	if(!GEngine || !GEngine->DefaultPhysMaterial)
 	{
@@ -2823,12 +2796,7 @@ void FBodyInstance::GetComplexPhysicalMaterials(const FBodyInstance*, TWeakObjec
 	if (PrimComp)
 	{
 		const int32 NumMaterials = PrimComp->GetNumMaterials();
-		OutPhysMaterials.SetNum(NumMaterials);
-
-		if (OutPhysMaterialMasks)
-		{
-			OutPhysMaterialMasks->SetNum(NumMaterials);
-		}
+		OutPhysicalMaterials.SetNum(NumMaterials);
 
 		for (int32 MatIdx = 0; MatIdx < NumMaterials; MatIdx++)
 		{
@@ -2838,26 +2806,9 @@ void FBodyInstance::GetComplexPhysicalMaterials(const FBodyInstance*, TWeakObjec
 			{
 				PhysMat = Material->GetPhysicalMaterial();
 			}
-			
-			OutPhysMaterials[MatIdx] = PhysMat;
 
-			if (OutPhysMaterialMasks)
-			{
-				UPhysicalMaterialMask* PhysMatMask = nullptr;
-				UMaterialInterface* PhysMatMap = nullptr;
-
-				if (Material)
-				{
-					PhysMatMask = Material->GetPhysicalMaterialMask();
-					if (PhysMatMask)
-					{
-						PhysMatMap = Material;
-					}
-				}
-
-				(*OutPhysMaterialMasks)[MatIdx].PhysicalMaterialMask = PhysMatMask;
-				(*OutPhysMaterialMasks)[MatIdx].PhysicalMaterialMap = PhysMatMap;
-			}
+			check(PhysMat != NULL);
+			OutPhysicalMaterials[MatIdx] = PhysMat;
 		}
 	}
 }
@@ -3768,7 +3719,7 @@ void FBodyInstance::FixupData(class UObject* Loader)
 	}
 }
 
-void FBodyInstance::ApplyMaterialToShape_AssumesLocked(const FPhysicsShapeHandle& InShape, UPhysicalMaterial* SimplePhysMat, const TArrayView<UPhysicalMaterial*>& ComplexPhysMats, const TArrayView<FPhysicalMaterialMaskParams>* ComplexPhysMatMasks)
+void FBodyInstance::ApplyMaterialToShape_AssumesLocked(const FPhysicsShapeHandle& InShape, UPhysicalMaterial* SimplePhysMat, const TArrayView<UPhysicalMaterial*>& ComplexPhysMats)
 {
 	// If a triangle mesh, need to get array of materials...
 	ECollisionShapeType GeomType = FPhysicsInterface::GetShapeType(InShape);
@@ -3776,18 +3727,7 @@ void FBodyInstance::ApplyMaterialToShape_AssumesLocked(const FPhysicsShapeHandle
 	{
 		if(ComplexPhysMats.Num())
 		{
-#if WITH_CHAOS
-			if (ensure(ComplexPhysMatMasks))
-			{
-				FPhysicsInterface::SetMaterials(InShape, ComplexPhysMats, *ComplexPhysMatMasks);
-			}
-			else
-			{
-				FPhysicsInterface::SetMaterials(InShape, ComplexPhysMats);
-			}
-#else
 			FPhysicsInterface::SetMaterials(InShape, ComplexPhysMats);
-#endif
 		}
 		else
 		{
@@ -3814,7 +3754,7 @@ void FBodyInstance::ApplyMaterialToShape_AssumesLocked(const FPhysicsShapeHandle
 	}
 }
 
-void FBodyInstance::ApplyMaterialToInstanceShapes_AssumesLocked(UPhysicalMaterial* SimplePhysMat, TArray<UPhysicalMaterial*>& ComplexPhysMats, const TArrayView<FPhysicalMaterialMaskParams>& ComplexPhysMatMasks)
+void FBodyInstance::ApplyMaterialToInstanceShapes_AssumesLocked(UPhysicalMaterial* SimplePhysMat, TArray<UPhysicalMaterial*>& ComplexPhysMats)
 {
 	FBodyInstance* TheirBI = this;
 	FBodyInstance* BIWithActor = TheirBI->WeldParent ? TheirBI->WeldParent : TheirBI;
@@ -3827,12 +3767,12 @@ void FBodyInstance::ApplyMaterialToInstanceShapes_AssumesLocked(UPhysicalMateria
 		if(TheirBI->IsShapeBoundToBody(Shape))
 		{
 			FPhysicsCommand::ExecuteShapeWrite(BIWithActor, Shape, [&](const FPhysicsShapeHandle& InnerShape)
-			{
-				ApplyMaterialToShape_AssumesLocked(InnerShape, SimplePhysMat, ComplexPhysMats, &ComplexPhysMatMasks);
-			});		
-		}
+		{
+				ApplyMaterialToShape_AssumesLocked(InnerShape, SimplePhysMat, ComplexPhysMats);
+		});		
 	}
 }
+}
 
 bool FBodyInstance::ValidateTransform(const FTransform &Transform, const FString& DebugName, const UBodySetup* Setup)
 {
diff --git a/Engine/Source/Runtime/Engine/Private/PhysicsEngine/BodySetup.cpp b/Engine/Source/Runtime/Engine/Private/PhysicsEngine/BodySetup.cpp
--- a/Engine/Source/Runtime/Engine/Private/PhysicsEngine/BodySetup.cpp
+++ b/Engine/Source/Runtime/Engine/Private/PhysicsEngine/BodySetup.cpp
@@ -179,7 +179,6 @@ UBodySetup::UBodySetup(const FObjectInitializer& ObjectInitializer)
 	bNeverNeedsCookedCollisionData = false;
 	bGenerateMirroredCollision = true;
 	bGenerateNonMirroredCollision = true;
-	bSupportUVsAndFaceRemap = false;
 	DefaultInstance.SetObjectType(ECC_PhysicsBody);
 #if WITH_EDITORONLY_DATA
 	BuildScale_DEPRECATED = 1.0f;
@@ -354,7 +353,13 @@ void UBodySetup::GetCookInfo(FCookBodySetupInfo& OutCookInfo, EPhysXMeshCookFlag
 
 			OutCookInfo.TriMeshCookFlags = CookFlags;
 
-			OutCookInfo.bSupportFaceRemap = bSupportUVsAndFaceRemap;
+			// If outer is a static mesh with physical material mask enabled, set retain UV and face remap table to true
+			/*
+			if (UStaticMesh* StaticMesh = Cast<UStaticMesh>(CDPObj))
+			{
+				OutCookInfo.bSupportFaceRemap = StaticMesh->GetEnablePhysicalMaterialMask();
+			}
+			*/
 		}
 		else
 		{
@@ -362,7 +367,7 @@ void UBodySetup::GetCookInfo(FCookBodySetupInfo& OutCookInfo, EPhysXMeshCookFlag
 		}
 	}
 
-	OutCookInfo.bSupportUVFromHitResults = UPhysicsSettings::Get()->bSupportUVFromHitResults || bSupportUVsAndFaceRemap;
+	OutCookInfo.bSupportUVFromHitResults = UPhysicsSettings::Get()->bSupportUVFromHitResults || OutCookInfo.bSupportFaceRemap;
 
 #endif // WITH_PHYSX
 }
@@ -817,8 +822,7 @@ void UBodySetup::AddShapesToRigidActor_AssumesLocked(
 	FBodyInstance* OwningInstance, 
 	FVector& Scale3D, 
 	UPhysicalMaterial* SimpleMaterial,
-	TArray<UPhysicalMaterial*>& ComplexMaterials,
-	TArray<FPhysicalMaterialMaskParams>& ComplexMaterialMasks,
+	TArray<UPhysicalMaterial*>& ComplexMaterials, 
 	const FBodyCollisionData& BodyCollisionData,
 	const FTransform& RelativeTM, 
 	TArray<FPhysicsShapeHandle>* NewShapes)
@@ -844,9 +848,6 @@ void UBodySetup::AddShapesToRigidActor_AssumesLocked(
 	AddParams.Scale = Scale3D;
 	AddParams.SimpleMaterial = SimpleMaterial;
 	AddParams.ComplexMaterials = TArrayView<UPhysicalMaterial*>(ComplexMaterials);
-#if WITH_CHAOS
-	AddParams.ComplexMaterialMasks = TArrayView<FPhysicalMaterialMaskParams>(ComplexMaterialMasks);
-#endif
 	AddParams.LocalTransform = RelativeTM;
 	AddParams.WorldTransform = OwningInstance->GetUnrealWorldTransform();
 	AddParams.Geometry = &AggGeom;
@@ -1399,11 +1400,6 @@ void UBodySetup::GetGeometryDDCKey(FString& OutString) const
 		(int32)UPhysicsSettings::Get()->bSupportUVFromHitResults,
 		(int32)GetCollisionTraceFlag(),
 		(int32)BODY_SETUP_GEOMETRY_KEY_VER);
-
-	if (bSupportUVsAndFaceRemap)
-	{
-		OutString += FString(TEXT("_1"));
-	}
 }
 
 void UBodySetup::PostInitProperties()
@@ -1506,8 +1502,6 @@ void UBodySetup::GetResourceSizeEx(FResourceSizeEx& CumulativeResourceSize)
 		const FByteBulkData& FmtData = CookedFormatData.GetFormat(FPlatformProperties::GetPhysicsFormat());
 		CumulativeResourceSize.AddDedicatedSystemMemoryBytes(FmtData.GetBulkDataSize());
 	}
-
-	CumulativeResourceSize.AddDedicatedSystemMemoryBytes(FaceRemap.GetAllocatedSize());
 	
 	// Count any UV info
 	UVInfo.GetResourceSizeEx(CumulativeResourceSize);
diff --git a/Engine/Source/Runtime/Engine/Private/PhysicsEngine/Experimental/ChaosDerivedData.cpp b/Engine/Source/Runtime/Engine/Private/PhysicsEngine/Experimental/ChaosDerivedData.cpp
--- a/Engine/Source/Runtime/Engine/Private/PhysicsEngine/Experimental/ChaosDerivedData.cpp
+++ b/Engine/Source/Runtime/Engine/Private/PhysicsEngine/Experimental/ChaosDerivedData.cpp
@@ -24,7 +24,7 @@ const TCHAR* FChaosDerivedDataCooker::GetPluginName() const
 
 const TCHAR* FChaosDerivedDataCooker::GetVersionString() const
 {
-	return TEXT("CD1FE02F2AB9453FB6E173E28A076DC4");
+	return TEXT("A9DFE2536B324A68B507931A0D2C5496");
 }
 
 FString FChaosDerivedDataCooker::GetPluginSpecificCacheKeySuffix() const
diff --git a/Engine/Source/Runtime/Engine/Private/PhysicsEngine/Experimental/PhysInterface_Chaos.cpp b/Engine/Source/Runtime/Engine/Private/PhysicsEngine/Experimental/PhysInterface_Chaos.cpp
--- a/Engine/Source/Runtime/Engine/Private/PhysicsEngine/Experimental/PhysInterface_Chaos.cpp
+++ b/Engine/Source/Runtime/Engine/Private/PhysicsEngine/Experimental/PhysInterface_Chaos.cpp
@@ -63,7 +63,6 @@ DECLARE_CYCLE_STAT(TEXT("Update Kinematics On Deferred SkelMeshes"), STAT_Update
 DECLARE_CYCLE_STAT(TEXT("Phys Events Time"), STAT_PhysicsEventTime, STATGROUP_Physics);
 DECLARE_CYCLE_STAT(TEXT("SyncComponentsToBodies (sync)"), STAT_SyncComponentsToBodies, STATGROUP_Physics);
 DECLARE_CYCLE_STAT(TEXT("SyncComponentsToBodies (async)"), STAT_SyncComponentsToBodies_Async, STATGROUP_Physics);
-DECLARE_CYCLE_STAT(TEXT("Query PhysicalMaterialMask Hit"), STAT_QueryPhysicalMaterialMaskHit, STATGROUP_Physics);
 
 DECLARE_DWORD_COUNTER_STAT(TEXT("Broadphase Adds"), STAT_NumBroadphaseAdds, STATGROUP_Physics);
 DECLARE_DWORD_COUNTER_STAT(TEXT("Broadphase Removes"), STAT_NumBroadphaseRemoves, STATGROUP_Physics);
@@ -115,80 +114,6 @@ Chaos::FChaosPhysicsMaterial* GetMaterialFromInternalFaceIndex(const FPhysicsSha
 	return nullptr;
 }
 
-Chaos::FChaosPhysicsMaterial* GetMaterialFromInternalFaceIndexAndHitLocation(const FPhysicsShape& Shape, const FPhysicsActor& Actor, uint32 InternalFaceIndex, const FVector& HitLocation)
-{
-	{
-		SCOPE_CYCLE_COUNTER(STAT_QueryPhysicalMaterialMaskHit);
-
-		if (Shape.Materials.Num() > 0 && Actor.Proxy)
-		{
-			Chaos::FPBDRigidsSolver* Solver = Actor.Proxy->GetSolver();
-
-			if (ensure(Solver))
-			{
-				if (Shape.MaterialMasks.Num() > 0)
-				{
-					UBodySetup* BodySetup = nullptr;
-
-					if (const FBodyInstance* BodyInst = GetUserData(Actor))
-					{
-						BodyInst = FPhysicsInterface::ShapeToOriginalBodyInstance(BodyInst, &Shape);
-						BodySetup = BodyInst->BodySetup.Get();	//this data should be immutable at runtime so ok to check from worker thread.
-						ECollisionShapeType GeomType = GetGeometryType(Shape);
-
-						if (BodySetup->bSupportUVsAndFaceRemap && GetGeometryType(Shape) == ECollisionShapeType::Trimesh)
-						{
-							FVector Scale(1.0f, 1.0f, 1.0f);
-							const Chaos::FImplicitObject* Geometry = Shape.Geometry.Get();
-							if (const Chaos::TImplicitObjectScaled<Chaos::FTriangleMeshImplicitObject>* ScaledTrimesh = Chaos::TImplicitObjectScaled<Chaos::FTriangleMeshImplicitObject>::AsScaled(*Geometry))
-							{
-								Scale = ScaledTrimesh->GetScale();
-							}
-
-							// Convert hit location to local
-							Chaos::FRigidTransform3 ActorToWorld(Actor.X(), Actor.R(), Scale);
-							const FVector LocalHitPos = ActorToWorld.InverseTransformPosition(HitLocation);
-
-							uint8 Index = Shape.Geometry->GetMaterialIndex(InternalFaceIndex);
-							if (Shape.MaterialMasks.IsValidIndex(Index))
-							{
-								Chaos::FChaosPhysicsMaterialMask* Mask = nullptr;
-								{
-									Chaos::TSolverQueryMaterialScope<Chaos::ELockType::Read> Scope(Solver);
-									Mask = Solver->GetQueryMaterialMasks().Get(Shape.MaterialMasks[Index].InnerHandle);
-								}
-
-								if (Mask && InternalFaceIndex < (uint32)BodySetup->FaceRemap.Num())
-								{
-									int32 RemappedFaceIndex = BodySetup->FaceRemap[InternalFaceIndex];
-									FVector2D UV;
-
-
-									if (BodySetup->CalcUVAtLocation(LocalHitPos, RemappedFaceIndex, Mask->UVChannelIndex, UV))
-									{
-										uint32 MapIdx = UPhysicalMaterialMask::GetPhysMatIndex(Mask->MaskData, Mask->SizeX, Mask->SizeY, Mask->AddressX, Mask->AddressY, UV.X, UV.Y);
-										uint32 AdjustedMapIdx = Index * EPhysicalMaterialMaskColor::MAX + MapIdx;
-										if (Shape.MaterialMaskMaps.IsValidIndex(AdjustedMapIdx))
-										{
-											uint32 MaterialIdx = Shape.MaterialMaskMaps[AdjustedMapIdx];
-											if (Shape.MaterialMaskMapMaterials.IsValidIndex(MaterialIdx))
-											{
-												Chaos::TSolverQueryMaterialScope<Chaos::ELockType::Read> Scope(Solver);
-												return Solver->GetQueryMaterials().Get(Shape.MaterialMaskMapMaterials[MaterialIdx].InnerHandle);
-											}
-										}
-									}
-								}
-							}
-						}
-					}
-				}
-			}
-		}
-	}
-
-	return GetMaterialFromInternalFaceIndex(Shape, Actor, InternalFaceIndex);
-}
 
 const Chaos::FImplicitObject& FPhysicsShapeReference_Chaos::GetGeometry() const
 {
@@ -330,31 +255,6 @@ void FPhysInterface_Chaos::SetUserData(FPhysicsMaterialHandle& InHandle, void* I
 	Chaos::FPhysicalMaterialManager::Get().UpdateMaterial(InHandle);
 }
 
-FPhysicsMaterialMaskHandle FPhysInterface_Chaos::CreateMaterialMask(const UPhysicalMaterialMask* InMaterialMask)
-{
-	Chaos::FMaterialMaskHandle NewHandle = Chaos::FPhysicalMaterialManager::Get().CreateMask();
-	FPhysInterface_Chaos::UpdateMaterialMask(NewHandle, InMaterialMask);
-	return NewHandle;
-}
-
-void FPhysInterface_Chaos::ReleaseMaterialMask(FPhysicsMaterialMaskHandle& InHandle)
-{
-	Chaos::FPhysicalMaterialManager::Get().Destroy(InHandle);
-}
-
-void FPhysInterface_Chaos::UpdateMaterialMask(FPhysicsMaterialMaskHandle& InHandle, const UPhysicalMaterialMask* InMaterialMask)
-{
-	if (Chaos::FChaosPhysicsMaterialMask* MaterialMask = InHandle.Get())
-	{
-		InMaterialMask->GenerateMaskData(MaterialMask->MaskData, MaterialMask->SizeX, MaterialMask->SizeY);
-		MaterialMask->UVChannelIndex = InMaterialMask->UVChannelIndex;
-		MaterialMask->AddressX = static_cast<int32>(InMaterialMask->AddressX);
-		MaterialMask->AddressY = static_cast<int32>(InMaterialMask->AddressY);
-	}
-
-	Chaos::FPhysicalMaterialManager::Get().UpdateMaterialMask(InHandle);
-}
-
 void FPhysInterface_Chaos::SetUserData(const FPhysicsShapeHandle& InShape, void* InUserData)
 {
 	if (CHAOS_ENSURE(InShape.Shape))
@@ -1473,7 +1373,7 @@ void FPhysInterface_Chaos::AddGeometry(FPhysicsActorHandle& InActor, const FGeom
 				OutOptShapes->Add(NewHandle);
 			}
 
-			FBodyInstance::ApplyMaterialToShape_AssumesLocked(NewHandle, InParams.SimpleMaterial, InParams.ComplexMaterials, &InParams.ComplexMaterialMasks);
+			FBodyInstance::ApplyMaterialToShape_AssumesLocked(NewHandle, InParams.SimpleMaterial, InParams.ComplexMaterials);
 
 			//TArrayView<UPhysicalMaterial*> SimpleView = MakeArrayView(&(const_cast<UPhysicalMaterial*>(InParams.SimpleMaterial)), 1);
 			//FPhysInterface_Chaos::SetMaterials(NewHandle, InParams.ComplexMaterials.Num() > 0 ? InParams.ComplexMaterials : SimpleView);
@@ -1641,75 +1541,6 @@ void FPhysInterface_Chaos::SetMaterials(const FPhysicsShapeHandle& InShape, cons
 	InShape.Shape->Materials = NewMaterialHandles;
 }
 
-void FPhysInterface_Chaos::SetMaterials(const FPhysicsShapeHandle& InShape, const TArrayView<UPhysicalMaterial*> InMaterials, const TArrayView<FPhysicalMaterialMaskParams>& InMaterialMasks)
-{
-	SetMaterials(InShape, InMaterials);
-
-	if (InMaterialMasks.Num() > 0)
-	{
-		// Build a list of handles to store on the shape
-		TArray<Chaos::FMaterialMaskHandle> NewMaterialMaskHandles;
-		TArray<uint32> NewMaterialMaskMaps;
-		TArray<Chaos::FMaterialHandle> NewMaterialMaskMaterialHandles;
-
-		NewMaterialMaskHandles.Reserve(InMaterialMasks.Num());
-
-		int MaskMapMatIdx = 0;
-
-		for (FPhysicalMaterialMaskParams& MaterialMaskData : InMaterialMasks)
-		{
-			if (MaterialMaskData.PhysicalMaterialMask && ensure(MaterialMaskData.PhysicalMaterialMap))
-			{
-				NewMaterialMaskHandles.Add(MaterialMaskData.PhysicalMaterialMask->GetPhysicsMaterialMask());
-				for (int i = 0; i < EPhysicalMaterialMaskColor::MAX; i++)
-				{
-					if (UPhysicalMaterial* MapMat = MaterialMaskData.PhysicalMaterialMap->GetPhysicalMaterialFromMap(i))
-					{
-						InShape.Shape->MaterialMaskMaps.Emplace(MaskMapMatIdx);
-						MaskMapMatIdx++;
-					}
-					else
-					{
-						InShape.Shape->MaterialMaskMaps.Emplace(INDEX_NONE);
-					}
-				}
-			}
-			else
-			{
-				NewMaterialMaskHandles.Add(Chaos::FMaterialMaskHandle());
-				for (int i = 0; i < EPhysicalMaterialMaskColor::MAX; i++)
-				{
-					InShape.Shape->MaterialMaskMaps.Emplace(INDEX_NONE);
-				}
-			}
-		}
-
-		if (MaskMapMatIdx > 0)
-		{
-			NewMaterialMaskMaterialHandles.Reserve(MaskMapMatIdx);
-
-			uint32 Offset = 0;
-
-			for (FPhysicalMaterialMaskParams& MaterialMaskData : InMaterialMasks)
-			{
-				if (MaterialMaskData.PhysicalMaterialMask)
-				{
-					for (int i = 0; i < EPhysicalMaterialMaskColor::MAX; i++)
-					{
-						if (UPhysicalMaterial* MapMat = MaterialMaskData.PhysicalMaterialMap->GetPhysicalMaterialFromMap(i))
-						{
-							NewMaterialMaskMaterialHandles.Add(MapMat->GetPhysicsMaterial());
-						}
-					}
-				}
-			}
-		}
-
-		InShape.Shape->MaterialMasks = NewMaterialMaskHandles;
-		InShape.Shape->MaterialMaskMapMaterials = NewMaterialMaskMaterialHandles;
-	}
-}
-
 void FinishSceneStat()
 {
 }
diff --git a/Engine/Source/Runtime/Engine/Private/PhysicsEngine/PhysicalMaterialMask.cpp b/Engine/Source/Runtime/Engine/Private/PhysicsEngine/PhysicalMaterialMask.cpp
deleted file mode 100644
diff --git a/Engine/Source/Runtime/Engine/Private/StaticMesh.cpp b/Engine/Source/Runtime/Engine/Private/StaticMesh.cpp
--- a/Engine/Source/Runtime/Engine/Private/StaticMesh.cpp
+++ b/Engine/Source/Runtime/Engine/Private/StaticMesh.cpp
@@ -2974,16 +2974,6 @@ void UStaticMesh::PostEditChangeProperty(FPropertyChangedEvent& PropertyChangedE
 			BodySetup->CreatePhysicsMeshes();
 		}
 	}
-
-	if (PropertyName == GET_MEMBER_NAME_CHECKED(UStaticMesh, bSupportPhysicalMaterialMasks))
-	{
-		if (BodySetup)
-		{
-			BodySetup->bSupportUVsAndFaceRemap = bSupportPhysicalMaterialMasks;
-			BodySetup->InvalidatePhysicsData();
-			BodySetup->CreatePhysicsMeshes();
-		}
-	}
 #endif
 
 	LightMapResolution = FMath::Max(LightMapResolution, 0);
@@ -5666,7 +5656,8 @@ bool UStaticMesh::GetPhysicsTriMeshData(struct FTriMeshCollisionData* CollisionD
 	TMap<int32, int32> MeshToCollisionVertMap; // map of static mesh verts to collision verts
 
 	// If the mesh enables physical material masks, override the physics setting since UVs are always required 
-	bool bCopyUVs = bSupportPhysicalMaterialMasks || UPhysicsSettings::Get()->bSupportUVFromHitResults; // See if we should copy UVs
+	// bool bCopyUVs = GetEnablePhysicalMaterialMask() || UPhysicsSettings::Get()->bSupportUVFromHitResults; // See if we should copy UVs
+	bool bCopyUVs = UPhysicsSettings::Get()->bSupportUVFromHitResults; // See if we should copy UVs
 
 	// If copying UVs, allocate array for storing them
 	if (bCopyUVs)
@@ -5811,7 +5802,6 @@ void UStaticMesh::CreateBodySetup()
 	{
 		BodySetup = NewObject<UBodySetup>(this);
 		BodySetup->DefaultInstance.SetCollisionProfileName(UCollisionProfile::BlockAll_ProfileName);
-		BodySetup->bSupportUVsAndFaceRemap = bSupportPhysicalMaterialMasks;
 	}
 }
 
diff --git a/Engine/Source/Runtime/Engine/Public/Physics/Experimental/PhysInterface_Chaos.h b/Engine/Source/Runtime/Engine/Public/Physics/Experimental/PhysInterface_Chaos.h
--- a/Engine/Source/Runtime/Engine/Public/Physics/Experimental/PhysInterface_Chaos.h
+++ b/Engine/Source/Runtime/Engine/Public/Physics/Experimental/PhysInterface_Chaos.h
@@ -14,7 +14,6 @@
 #include "Physics/GenericPhysicsInterface.h"
 #include "Physics/Experimental/PhysicsUserData_Chaos.h"
 #include "Chaos/PhysicalMaterials.h"
-#include "PhysicalMaterials/PhysicalMaterialMask.h"
 
 //NOTE: Do not include Chaos headers directly as it means recompiling all of engine. This should be reworked to avoid allocations
 
@@ -203,11 +202,6 @@ public:
     static void UpdateMaterial(FPhysicsMaterialHandle& InHandle, UPhysicalMaterial* InMaterial);
     static void SetUserData(FPhysicsMaterialHandle& InHandle, void* InUserData);
 
-	// Material mask functions 
-	static FPhysicsMaterialMaskHandle CreateMaterialMask(const UPhysicalMaterialMask* InMaterialMask);
-	static void ReleaseMaterialMask(FPhysicsMaterialMaskHandle& InHandle);
-	static void UpdateMaterialMask(FPhysicsMaterialMaskHandle& InHandle, const UPhysicalMaterialMask* InMaterialMask);
-
 	// Actor interface functions
 	template<typename AllocatorType>
 	static int32 GetAllShapes_AssumedLocked(const FPhysicsActorHandle& InActorReference, TArray<FPhysicsShapeHandle, AllocatorType>& OutShapes);
@@ -406,8 +400,7 @@ public:
     static void SetUserData(const FPhysicsShapeHandle& InShape, void* InUserData);
     static void SetGeometry(const FPhysicsShapeHandle& InShape, physx::PxGeometry& InGeom) {}
 	static void SetLocalTransform(const FPhysicsShapeHandle& InShape, const FTransform& NewLocalTransform);
-	static void SetMaterials(const FPhysicsShapeHandle& InShape, const TArrayView<UPhysicalMaterial*>InMaterials);
-	static void SetMaterials(const FPhysicsShapeHandle& InShape, const TArrayView<UPhysicalMaterial*> InMaterials, const TArrayView<FPhysicalMaterialMaskParams>& InMaterialMasks);
+    static void SetMaterials(const FPhysicsShapeHandle& InShape, const TArrayView<UPhysicalMaterial*>InMaterials);
 };
 
 /*
@@ -454,8 +447,6 @@ FORCEINLINE void ComputeZeroDistanceImpactNormalAndPenetration(const UWorld* Wor
 
 Chaos::FChaosPhysicsMaterial* GetMaterialFromInternalFaceIndex(const FPhysicsShape& Shape, const FPhysicsActor& Actor, uint32 InternalFaceIndex);
 
-Chaos::FChaosPhysicsMaterial* GetMaterialFromInternalFaceIndexAndHitLocation(const FPhysicsShape& Shape, const FPhysicsActor& Actor, uint32 InternalFaceIndex, const FVector& HitLocation);
-
 uint32 GetTriangleMeshExternalFaceIndex(const FPhysicsShape& Shape, uint32 InternalFaceIndex);
 
 inline void GetShapes(const FPhysActorDummy& RigidActor, FPhysTypeDummy** ShapesBuffer, uint32 NumShapes)
diff --git a/Engine/Source/Runtime/Engine/Public/Physics/PhysicsInterfaceTypes.h b/Engine/Source/Runtime/Engine/Public/Physics/PhysicsInterfaceTypes.h
--- a/Engine/Source/Runtime/Engine/Public/Physics/PhysicsInterfaceTypes.h
+++ b/Engine/Source/Runtime/Engine/Public/Physics/PhysicsInterfaceTypes.h
@@ -38,15 +38,6 @@ namespace physx
 	class PxTriangleMesh;
 }
 
-struct FPhysicalMaterialMaskParams
-{
-	/** Physical materials mask */
-	UPhysicalMaterialMask* PhysicalMaterialMask;
-
-	/** Pointer to material which contains the physical material map */
-	UMaterialInterface* PhysicalMaterialMap;
-};
-
 struct FGeometryAddParams
 {
 	bool bDoubleSided;
@@ -55,9 +46,6 @@ struct FGeometryAddParams
 	FVector Scale;
 	UPhysicalMaterial* SimpleMaterial;
 	TArrayView<UPhysicalMaterial*> ComplexMaterials;
-#if WITH_CHAOS
-	TArrayView<FPhysicalMaterialMaskParams> ComplexMaterialMasks;
-#endif
 	FTransform LocalTransform;
 	FTransform WorldTransform;
 	FKAggregateGeom* Geometry;
diff --git a/Engine/Source/Runtime/Experimental/Chaos/Private/Chaos/GeometryParticles.cpp b/Engine/Source/Runtime/Experimental/Chaos/Private/Chaos/GeometryParticles.cpp
--- a/Engine/Source/Runtime/Experimental/Chaos/Private/Chaos/GeometryParticles.cpp
+++ b/Engine/Source/Runtime/Experimental/Chaos/Private/Chaos/GeometryParticles.cpp
@@ -60,9 +60,6 @@ namespace Chaos
 		, Geometry()
 		, WorldSpaceInflatedShapeBounds(TAABB<FReal, 3>(FVec3(0), FVec3(0)))
 		, Materials()
-		, MaterialMasks()
-		, MaterialMaskMaps()
-		, MaterialMaskMapMaterials()
 		, bDisable(false)
 		, bSimulate(true)
 		, CollisionTraceType(EChaosCollisionTraceFlag::Chaos_CTF_UseDefault)
@@ -99,7 +96,6 @@ namespace Chaos
 	void TPerShapeData<T, d>::Serialize(FChaosArchive& Ar)
 	{
 		Ar.UsingCustomVersion(FExternalPhysicsCustomObjectVersion::GUID);
-		Ar.UsingCustomVersion(FExternalPhysicsMaterialCustomObjectVersion::GUID);
 
 		Ar << Geometry;
 		Ar << QueryData;
@@ -136,12 +132,6 @@ namespace Chaos
 			Ar << Data;
 			CollisionTraceType = (EChaosCollisionTraceFlag)Data;
 		}
-
-		if (Ar.CustomVer(FExternalPhysicsMaterialCustomObjectVersion::GUID) >= FExternalPhysicsMaterialCustomObjectVersion::AddedMaterialMasks)
-		{
-			Ar << MaterialMasks << MaterialMaskMaps << MaterialMaskMapMaterials;
-		}
-
 	}
 
 
diff --git a/Engine/Source/Runtime/Experimental/Chaos/Private/Chaos/PhysicalMaterials.cpp b/Engine/Source/Runtime/Experimental/Chaos/Private/Chaos/PhysicalMaterials.cpp
--- a/Engine/Source/Runtime/Experimental/Chaos/Private/Chaos/PhysicalMaterials.cpp
+++ b/Engine/Source/Runtime/Experimental/Chaos/Private/Chaos/PhysicalMaterials.cpp
@@ -23,27 +23,8 @@ namespace Chaos
 		return nullptr;
 	}
 
-	FChaosPhysicsMaterialMask* FMaterialMaskHandle::Get() const
-	{
-		if (InnerHandle.IsValid())
-		{
-			return FPhysicalMaterialManager::Get().Resolve(InnerHandle);
-		}
-		return nullptr;
-	}
-
-	const FChaosPhysicsMaterialMask* FConstMaterialMaskHandle::Get() const
-	{
-		if (InnerHandle.IsValid())
-		{
-			return FPhysicalMaterialManager::Get().Resolve(InnerHandle);
-		}
-		return nullptr;
-	}
-
 	FPhysicalMaterialManager::FPhysicalMaterialManager()
 		: Materials(InitialCapacity)
-		, MaterialMasks()
 	{
 
 	}
@@ -64,16 +45,6 @@ namespace Chaos
 		return Materials.Get(InHandle);
 	}
 
-	FChaosPhysicsMaterialMask* FPhysicalMaterialManager::Resolve(FChaosMaterialMaskHandle InHandle) const
-	{
-		return MaterialMasks.Get(InHandle);
-	}
-
-	const FChaosPhysicsMaterialMask* FPhysicalMaterialManager::Resolve(FChaosConstMaterialMaskHandle InHandle) const
-	{
-		return MaterialMasks.Get(InHandle);
-	}
-	
 	void FPhysicalMaterialManager::UpdateMaterial(FMaterialHandle InHandle)
 	{
 		check(IsInGameThread());
@@ -81,23 +52,11 @@ namespace Chaos
 		OnMaterialUpdated.Broadcast(InHandle);
 	}
 
-	void FPhysicalMaterialManager::UpdateMaterialMask(FMaterialMaskHandle InHandle)
-	{
-		check(IsInGameThread());
-
-		OnMaterialMaskUpdated.Broadcast(InHandle);
-	}
-
 	const Chaos::THandleArray<FChaosPhysicsMaterial>& FPhysicalMaterialManager::GetMasterMaterials() const
 	{
 		return Materials;
 	}
 
-	const Chaos::THandleArray<FChaosPhysicsMaterialMask>& FPhysicalMaterialManager::GetMasterMaterialMasks() const
-	{
-		return MaterialMasks;
-	}
-
 	FMaterialHandle FPhysicalMaterialManager::Create()
 	{
 		LLM_SCOPE(ELLMTag::Chaos);
@@ -111,16 +70,6 @@ namespace Chaos
 		return OutHandle;
 	}
 
-	FMaterialMaskHandle FPhysicalMaterialManager::CreateMask()
-	{
-		check(IsInGameThread());
-		FMaterialMaskHandle OutHandle;
-		OutHandle.InnerHandle = MaterialMasks.Create();
-
-		OnMaterialMaskCreated.Broadcast(OutHandle);
-
-		return OutHandle;
-	}
 
 	void FPhysicalMaterialManager::Destroy(FMaterialHandle InHandle)
 	{
@@ -135,14 +84,4 @@ namespace Chaos
 		}
 	}
 
-	void FPhysicalMaterialManager::Destroy(FMaterialMaskHandle InHandle)
-	{
-		check(IsInGameThread());
-		if (InHandle.InnerHandle.IsValid())
-		{
-			OnMaterialMaskDestroyed.Broadcast(InHandle);
-
-			MaterialMasks.Destroy(InHandle.InnerHandle);
-		}
-	}
 }
\ No newline at end of file
diff --git a/Engine/Source/Runtime/Experimental/Chaos/Public/Chaos/Defines.h b/Engine/Source/Runtime/Experimental/Chaos/Public/Chaos/Defines.h
--- a/Engine/Source/Runtime/Experimental/Chaos/Public/Chaos/Defines.h
+++ b/Engine/Source/Runtime/Experimental/Chaos/Public/Chaos/Defines.h
@@ -64,46 +64,11 @@ namespace Chaos
 		}
 	};
 
-	class FChaosPhysicsMaterialMask
-	{
-	public:
-
-		int32 SizeX;
-		int32 SizeY;
-		TArray<uint32> MaskData;
-
-		int32 UVChannelIndex;
-		int32 AddressX;
-		int32 AddressY;
-
-		FChaosPhysicsMaterialMask()
-			: SizeX(0)
-			, SizeY(0)
-			, MaskData()
-			, UVChannelIndex(0)
-			, AddressX(0)
-			, AddressY(0)
-		{
-		}
-
-		static constexpr bool IsSerializablePtr = true;
-
-		static void StaticSerialize(FArchive& Ar, TSerializablePtr<FChaosPhysicsMaterialMask>& Serializable)
-		{
-			FChaosPhysicsMaterialMask* MaterialMask = const_cast<FChaosPhysicsMaterialMask*>(Serializable.Get());
 
-			if (Ar.IsLoading())
-			{
-				MaterialMask = new FChaosPhysicsMaterialMask();
-				Serializable.SetFromRawLowLevel(MaterialMask);
-			}
-
-			MaterialMask->Serialize(Ar);
-		}
-
-		void Serialize(FArchive& Ar)
-		{
-			Ar << SizeX << SizeY << UVChannelIndex << AddressX << AddressY << MaskData;
-		}
-	};
+	template <typename T>
+	FORCEINLINE FArchive& operator<<(FArchive& Ar, FChaosPhysicsMaterial& Value)
+	{
+		Value.Serialize(Ar);
+		return Ar;
+	}
 }
\ No newline at end of file
diff --git a/Engine/Source/Runtime/Experimental/Chaos/Public/Chaos/Framework/Handles.h b/Engine/Source/Runtime/Experimental/Chaos/Public/Chaos/Framework/Handles.h
--- a/Engine/Source/Runtime/Experimental/Chaos/Public/Chaos/Framework/Handles.h
+++ b/Engine/Source/Runtime/Experimental/Chaos/Public/Chaos/Framework/Handles.h
@@ -147,6 +147,8 @@ namespace Chaos
 	{
 	public:
 
+		static_assert(TIsTrivial<ElementType>::Value, "Handle managed data must be trivial");
+
 		using FHandle = THandle<ElementType, IndexWidth, GenerationWidth>;
 		using FConstHandle = TConstHandle<ElementType, IndexWidth, GenerationWidth>;
 		
diff --git a/Engine/Source/Runtime/Experimental/Chaos/Public/Chaos/GeometryParticles.h b/Engine/Source/Runtime/Experimental/Chaos/Public/Chaos/GeometryParticles.h
--- a/Engine/Source/Runtime/Experimental/Chaos/Public/Chaos/GeometryParticles.h
+++ b/Engine/Source/Runtime/Experimental/Chaos/Public/Chaos/GeometryParticles.h
@@ -12,7 +12,6 @@
 #include "Chaos/PhysicalMaterials.h"
 #include "UObject/PhysicsObjectVersion.h"
 #include "UObject/ExternalPhysicsCustomObjectVersion.h"
-#include "UObject/ExternalPhysicsMaterialCustomObjectVersion.h"
 
 #ifndef CHAOS_DETERMINISTIC
 #define CHAOS_DETERMINISTIC 1
@@ -53,9 +52,6 @@ namespace Chaos
 		TSerializablePtr<FImplicitObject> Geometry;
 		TAABB<FReal, 3> WorldSpaceInflatedShapeBounds;
 		TArray<FMaterialHandle> Materials;
-		TArray<FMaterialMaskHandle> MaterialMasks;
-		TArray<uint32> MaterialMaskMaps;
-		TArray<FMaterialHandle> MaterialMaskMapMaterials;
 
 		// TODO: Bitfields?
 		bool bDisable;
diff --git a/Engine/Source/Runtime/Experimental/Chaos/Public/Chaos/PhysicalMaterials.h b/Engine/Source/Runtime/Experimental/Chaos/Public/Chaos/PhysicalMaterials.h
--- a/Engine/Source/Runtime/Experimental/Chaos/Public/Chaos/PhysicalMaterials.h
+++ b/Engine/Source/Runtime/Experimental/Chaos/Public/Chaos/PhysicalMaterials.h
@@ -11,9 +11,6 @@ namespace Chaos
 	using FMaterialArray = THandleArray<FChaosPhysicsMaterial>;
 	using FChaosMaterialHandle = FMaterialArray::FHandle;
 	using FChaosConstMaterialHandle = FMaterialArray::FConstHandle;
-	using FMaterialMaskArray = THandleArray<FChaosPhysicsMaterialMask>;
-	using FChaosMaterialMaskHandle = FMaterialMaskArray::FHandle;
-	using FChaosConstMaterialMaskHandle = FMaterialMaskArray::FConstHandle;
 
 	/** 
 	 * Helper wrapper to encapsulate the access through the material manager
@@ -49,56 +46,13 @@ namespace Chaos
 		friend FArchive& operator <<(FArchive& Ar, FConstMaterialHandle& Value) { Ar << Value.InnerHandle; return Ar; }
 	};
 
-	/**
-	 * Helper wrapper to encapsulate the access through the material manager
-	 * Handles returned from the material manager are only designed to be used (resolved) on the game thread
-	 * by the physics interface.
-	 */
-	struct CHAOS_API FMaterialMaskHandle
-	{
-		FChaosPhysicsMaterialMask* Get() const;
-		FChaosMaterialMaskHandle InnerHandle;
-
-		bool IsValid() const { return Get() != nullptr; }
-
-		friend bool operator ==(const FMaterialMaskHandle& A, const FMaterialMaskHandle& B) { return A.InnerHandle == B.InnerHandle; }
-
-		friend FArchive& operator <<(FArchive& Ar, FMaterialMaskHandle& Value) { Ar << Value.InnerHandle; return Ar; }
-	};
-
-	/**
-	 * Helper wrapper to encapsulate the access through the material manager
-	 * Handles returned from the material manager are only designed to be used (resolved) on the game thread
-	 * by the physics interface.
-	 */
-	struct CHAOS_API FConstMaterialMaskHandle
-	{
-		const FChaosPhysicsMaterialMask* Get() const;
-		FChaosConstMaterialMaskHandle InnerHandle;
-
-		bool IsValid() const { return Get() != nullptr; }
-
-		friend bool operator ==(const FConstMaterialMaskHandle& A, const FConstMaterialMaskHandle& B) { return A.InnerHandle == B.InnerHandle; }
-
-		friend FArchive& operator <<(FArchive& Ar, FConstMaterialMaskHandle& Value) { Ar << Value.InnerHandle; return Ar; }
-	};
-
 	DECLARE_MULTICAST_DELEGATE_OneParam(FOnMaterialCreated, FMaterialHandle);
 	DECLARE_MULTICAST_DELEGATE_OneParam(FOnMaterialDestroyed, FMaterialHandle);
 	DECLARE_MULTICAST_DELEGATE_OneParam(FOnMaterialUpdated, FMaterialHandle);
-
-	DECLARE_MULTICAST_DELEGATE_OneParam(FOnMaterialMaskCreated, FMaterialMaskHandle);
-	DECLARE_MULTICAST_DELEGATE_OneParam(FOnMaterialMaskDestroyed, FMaterialMaskHandle);
-	DECLARE_MULTICAST_DELEGATE_OneParam(FOnMaterialMaskUpdated, FMaterialMaskHandle);
-
 	using FMaterialCreatedDelegate = FOnMaterialCreated::FDelegate;
 	using FMaterialDestroyedDelegate = FOnMaterialDestroyed::FDelegate;
 	using FMaterialUpdatedDelegate = FOnMaterialUpdated::FDelegate;
 
-	using FMaterialMaskCreatedDelegate = FOnMaterialMaskCreated::FDelegate;
-	using FMaterialMaskDestroyedDelegate = FOnMaterialMaskDestroyed::FDelegate;
-	using FMaterialMaskUpdatedDelegate = FOnMaterialMaskUpdated::FDelegate;
-
 	/** 
 	 * Global manager for physical materials.
 	 * Materials are created, updated and destroyed only on the game thread and an immutable
@@ -117,36 +71,25 @@ namespace Chaos
 
 		/** Create a new material, returning a stable handle to it - this should be stored and not the actual material pointer */
 		FMaterialHandle Create();
-		FMaterialMaskHandle CreateMask();
 
 		/** Destroy the material referenced by the provided handle */
 		void Destroy(FMaterialHandle InHandle);
-		void Destroy(FMaterialMaskHandle InHandle);
 
 		/** Get the actual material from a handle */
 		FChaosPhysicsMaterial* Resolve(FChaosMaterialHandle InHandle) const;
 		const FChaosPhysicsMaterial* Resolve(FChaosConstMaterialHandle InHandle) const;
 
-		FChaosPhysicsMaterialMask* Resolve(FChaosMaterialMaskHandle InHandle) const;
-		const FChaosPhysicsMaterialMask* Resolve(FChaosConstMaterialMaskHandle InHandle) const;
-
 		/** Signals stakeholders that the stored material for the provided handle has changed */
 		void UpdateMaterial(FMaterialHandle InHandle);
-		void UpdateMaterialMask(FMaterialMaskHandle InHandle);
 
 		/** Gets the internal list of master materials representing the current user state of the material data */
 		const THandleArray<FChaosPhysicsMaterial>& GetMasterMaterials() const;
-		const THandleArray<FChaosPhysicsMaterialMask>& GetMasterMaterialMasks() const;
 
 		/** Events */
 		FOnMaterialUpdated OnMaterialUpdated;
 		FOnMaterialCreated OnMaterialCreated;
 		FOnMaterialDestroyed OnMaterialDestroyed;
 
-		FOnMaterialMaskUpdated OnMaterialMaskUpdated;
-		FOnMaterialMaskCreated OnMaterialMaskCreated;
-		FOnMaterialMaskDestroyed OnMaterialMaskDestroyed;
-
 	private:
 
 		/** Initial size for the handle-managed array */
@@ -157,9 +100,6 @@ namespace Chaos
 
 		/** Handle-managed array of global materials. This is pushed to all solvers who all maintain a copy */
 		THandleArray<FChaosPhysicsMaterial> Materials;
-
-		/** Handle-managed array of global material masks. This is pushed to all solvers who all maintain a copy */
-		THandleArray<FChaosPhysicsMaterialMask> MaterialMasks;
 	};
 }
 
diff --git a/Engine/Source/Runtime/Experimental/ChaosSolvers/Private/ChaosSolversModule.cpp b/Engine/Source/Runtime/Experimental/ChaosSolvers/Private/ChaosSolversModule.cpp
--- a/Engine/Source/Runtime/Experimental/ChaosSolvers/Private/ChaosSolversModule.cpp
+++ b/Engine/Source/Runtime/Experimental/ChaosSolvers/Private/ChaosSolversModule.cpp
@@ -233,10 +233,6 @@ void FChaosSolversModule::Initialize()
 		OnDestroyMaterialHandle = MaterialManager.OnMaterialDestroyed.Add(Chaos::FMaterialDestroyedDelegate::CreateRaw(this, &FChaosSolversModule::OnDestroyMaterial));
 		OnUpdateMaterialHandle = MaterialManager.OnMaterialUpdated.Add(Chaos::FMaterialUpdatedDelegate::CreateRaw(this, &FChaosSolversModule::OnUpdateMaterial));
 
-		OnCreateMaterialMaskHandle = MaterialManager.OnMaterialMaskCreated.Add(Chaos::FMaterialMaskCreatedDelegate::CreateRaw(this, &FChaosSolversModule::OnCreateMaterialMask));
-		OnDestroyMaterialMaskHandle = MaterialManager.OnMaterialMaskDestroyed.Add(Chaos::FMaterialMaskDestroyedDelegate::CreateRaw(this, &FChaosSolversModule::OnDestroyMaterialMask));
-		OnUpdateMaterialMaskHandle = MaterialManager.OnMaterialMaskUpdated.Add(Chaos::FMaterialMaskUpdatedDelegate::CreateRaw(this, &FChaosSolversModule::OnUpdateMaterialMask));
-
 		bModuleInitialized = true;
 	}
 }
@@ -492,9 +488,7 @@ Chaos::FPhysicsSolver* FChaosSolversModule::CreateSolver(bool bStandalone /*= fa
 		Chaos::FPhysicalMaterialManager& Manager =	Chaos::FPhysicalMaterialManager::Get();
 		NewSolver->QueryMaterialLock.WriteLock();
 		NewSolver->QueryMaterials = Manager.GetMasterMaterials();
-		NewSolver->QueryMaterialMasks = Manager.GetMasterMaterialMasks();
 		NewSolver->SimMaterials = Manager.GetMasterMaterials();
-		NewSolver->SimMaterialMasks = Manager.GetMasterMaterialMasks();
 		NewSolver->QueryMaterialLock.WriteUnlock();
 	}
 
@@ -965,71 +959,6 @@ void FChaosSolversModule::OnDestroyMaterial(Chaos::FMaterialHandle InHandle)
 	}
 }
 
-void FChaosSolversModule::OnUpdateMaterialMask(Chaos::FMaterialMaskHandle InHandle)
-{
-	check(Dispatcher);
-
-	// Grab the material
-	Chaos::FChaosPhysicsMaterialMask* MaterialMask = InHandle.Get();
-
-	if (ensure(MaterialMask))
-	{
-		for (Chaos::FPhysicsSolver* Solver : Solvers)
-		{
-			// Send a copy of the material to each solver
-			Dispatcher->EnqueueCommandImmediate(Solver, [InHandle, MaterialMaskCopy = *MaterialMask](Chaos::FPhysicsSolver* InSolver)
-			{
-				InSolver->UpdateMaterialMask(InHandle, MaterialMaskCopy);
-			});
-		}
-	}
-}
-
-void FChaosSolversModule::OnCreateMaterialMask(Chaos::FMaterialMaskHandle InHandle)
-{
-	check(Dispatcher);
-
-	// Grab the material
-	Chaos::FChaosPhysicsMaterialMask* MaterialMask = InHandle.Get();
-
-	if (ensure(MaterialMask))
-	{
-		for (Chaos::FPhysicsSolver* Solver : Solvers)
-		{
-			// Send a copy of the material to each solver
-			Dispatcher->EnqueueCommandImmediate(Solver, [InHandle, MaterialMaskCopy = *MaterialMask](Chaos::FPhysicsSolver* InSolver)
-			{
-				InSolver->CreateMaterialMask(InHandle, MaterialMaskCopy);
-			});
-		}
-	}
-}
-
-void FChaosSolversModule::OnDestroyMaterialMask(Chaos::FMaterialMaskHandle InHandle)
-{
-	// @todo(bgallagher)
-	//check(Dispatcher);
-	if (!ensure(Dispatcher))
-	{
-		return;
-	}
-
-	// Grab the material
-	Chaos::FChaosPhysicsMaterialMask* MaterialMask = InHandle.Get();
-
-	if (ensure(MaterialMask))
-	{
-		for (Chaos::FPhysicsSolver* Solver : Solvers)
-		{
-			// Notify each solver
-			Dispatcher->EnqueueCommandImmediate(Solver, [InHandle](Chaos::FPhysicsSolver* InSolver)
-			{
-				InSolver->DestroyMaterialMask(InHandle);
-			});
-		}
-	}
-}
-
 const IChaosSettingsProvider& FChaosSolversModule::GetSettingsProvider() const
 {
 	return SettingsProvider ? *SettingsProvider : Chaos::GDefaultChaosSettings;
diff --git a/Engine/Source/Runtime/Experimental/ChaosSolvers/Private/PBDRigidsSolver.cpp b/Engine/Source/Runtime/Experimental/ChaosSolvers/Private/PBDRigidsSolver.cpp
--- a/Engine/Source/Runtime/Experimental/ChaosSolvers/Private/PBDRigidsSolver.cpp
+++ b/Engine/Source/Runtime/Experimental/ChaosSolvers/Private/PBDRigidsSolver.cpp
@@ -715,26 +715,10 @@ namespace Chaos
 		SimMaterials.Destroy(InHandle.InnerHandle);
 	}
 
-	void FPBDRigidsSolver::UpdateMaterialMask(Chaos::FMaterialMaskHandle InHandle, const Chaos::FChaosPhysicsMaterialMask& InNewData)
-	{
-		*SimMaterialMasks.Get(InHandle.InnerHandle) = InNewData;
-	}
-
-	void FPBDRigidsSolver::CreateMaterialMask(Chaos::FMaterialMaskHandle InHandle, const Chaos::FChaosPhysicsMaterialMask& InNewData)
-	{
-		ensure(SimMaterialMasks.Create(InNewData) == InHandle.InnerHandle);
-	}
-
-	void FPBDRigidsSolver::DestroyMaterialMask(Chaos::FMaterialMaskHandle InHandle)
-	{
-		SimMaterialMasks.Destroy(InHandle.InnerHandle);
-	}
-
 	void FPBDRigidsSolver::SyncQueryMaterials()
 	{
 		TSolverQueryMaterialScope<ELockType::Write> Scope(this);
 		QueryMaterials = SimMaterials;
-		QueryMaterialMasks = SimMaterialMasks;
 	}
 
 }; // namespace Chaos
diff --git a/Engine/Source/Runtime/Experimental/ChaosSolvers/Public/ChaosSolversModule.h b/Engine/Source/Runtime/Experimental/ChaosSolvers/Public/ChaosSolversModule.h
--- a/Engine/Source/Runtime/Experimental/ChaosSolvers/Public/ChaosSolversModule.h
+++ b/Engine/Source/Runtime/Experimental/ChaosSolvers/Public/ChaosSolversModule.h
@@ -294,10 +294,6 @@ public:
 	void OnCreateMaterial(Chaos::FMaterialHandle InHandle);
 	void OnDestroyMaterial(Chaos::FMaterialHandle InHandle);
 
-	void OnUpdateMaterialMask(Chaos::FMaterialMaskHandle InHandle);
-	void OnCreateMaterialMask(Chaos::FMaterialMaskHandle InHandle);
-	void OnDestroyMaterialMask(Chaos::FMaterialMaskHandle InHandle);
-
 private:
 
 	/** Safe method for always getting a settings provider (from the external caller or an internal default) */
@@ -367,10 +363,6 @@ private:
 	FDelegateHandle OnCreateMaterialHandle;
 	FDelegateHandle OnDestroyMaterialHandle;
 	FDelegateHandle OnUpdateMaterialHandle;
-
-	FDelegateHandle OnCreateMaterialMaskHandle;
-	FDelegateHandle OnDestroyMaterialMaskHandle;
-	FDelegateHandle OnUpdateMaterialMaskHandle;
 };
 
 /**
diff --git a/Engine/Source/Runtime/Experimental/ChaosSolvers/Public/PBDRigidsSolver.h b/Engine/Source/Runtime/Experimental/ChaosSolvers/Public/PBDRigidsSolver.h
--- a/Engine/Source/Runtime/Experimental/ChaosSolvers/Public/PBDRigidsSolver.h
+++ b/Engine/Source/Runtime/Experimental/ChaosSolvers/Public/PBDRigidsSolver.h
@@ -285,15 +285,10 @@ namespace Chaos
 		void UpdateMaterial(Chaos::FMaterialHandle InHandle, const Chaos::FChaosPhysicsMaterial& InNewData);
 		void CreateMaterial(Chaos::FMaterialHandle InHandle, const Chaos::FChaosPhysicsMaterial& InNewData);
 		void DestroyMaterial(Chaos::FMaterialHandle InHandle);
-		void UpdateMaterialMask(Chaos::FMaterialMaskHandle InHandle, const Chaos::FChaosPhysicsMaterialMask& InNewData);
-		void CreateMaterialMask(Chaos::FMaterialMaskHandle InHandle, const Chaos::FChaosPhysicsMaterialMask& InNewData);
-		void DestroyMaterialMask(Chaos::FMaterialMaskHandle InHandle);
 
-		/** Access to the internal material mirrors */
+		/** Access to the intenal material mirrors */
 		const THandleArray<FChaosPhysicsMaterial>& GetQueryMaterials() const { return QueryMaterials; }
-		const THandleArray<FChaosPhysicsMaterialMask>& GetQueryMaterialMasks() const { return QueryMaterialMasks; }
 		const THandleArray<FChaosPhysicsMaterial>& GetSimMaterials() const { return SimMaterials; }
-		const THandleArray<FChaosPhysicsMaterialMask>& GetSimMaterialMasks() const { return SimMaterialMasks; }
 
 		/** Copy the simulation material list to the query material list, to be done when the SQ commits an update */
 		void SyncQueryMaterials();
@@ -371,10 +366,8 @@ namespace Chaos
 		// instead of having to lock the entire physics state of the runtime.
 		FRWLock QueryMaterialLock;
 		THandleArray<FChaosPhysicsMaterial> QueryMaterials;
-		THandleArray<FChaosPhysicsMaterialMask> QueryMaterialMasks;
 		THandleArray<FChaosPhysicsMaterial> SimMaterials;
-		THandleArray<FChaosPhysicsMaterialMask> SimMaterialMasks;
-
+				
 		template<ELockType>
 		friend struct TSolverQueryMaterialScope;
 	};
diff --git a/Engine/Source/Runtime/PhysicsCore/Public/PhysicsInterfaceDeclaresCore.h b/Engine/Source/Runtime/PhysicsCore/Public/PhysicsInterfaceDeclaresCore.h
--- a/Engine/Source/Runtime/PhysicsCore/Public/PhysicsInterfaceDeclaresCore.h
+++ b/Engine/Source/Runtime/PhysicsCore/Public/PhysicsInterfaceDeclaresCore.h
@@ -32,16 +32,13 @@ namespace Chaos
 	class TPerShapeData;
 
 	class FPhysicalMaterial;
-	class FPhysicalMaterialMask;
 
 	template<typename, uint32, uint32>
 	class THandle;
 
 	struct FMaterialHandle;
-	struct FMaterialMaskHandle;
 
 	class FChaosPhysicsMaterial;
-	class FChaosPhysicsMaterialMask;
 }
 
 // Temporary dummy types until SQ implemented
@@ -71,7 +68,6 @@ using FPhysicsShape = Chaos::TPerShapeData<float, 3>;
 using FPhysicsGeometry = Chaos::FImplicitObject;
 using FPhysicsCapsuleGeometry = Chaos::TCapsule<float>;
 using FPhysicsMaterial = Chaos::FChaosPhysicsMaterial;
-using FPhysicsMaterialMask = Chaos::FChaosPhysicsMaterialMask; 
 using FPhysicsActor = Chaos::TGeometryParticle<float,3>;
 
 template <typename T>
@@ -220,7 +216,6 @@ typedef FPhysInterface_Chaos				FPhysicsCommand;
 typedef FPhysicsShapeReference_Chaos		FPhysicsShapeHandle;
 typedef FPhysicsGeometryCollection_Chaos	FPhysicsGeometryCollection;
 typedef Chaos::FMaterialHandle				FPhysicsMaterialHandle;
-typedef Chaos::FMaterialMaskHandle			FPhysicsMaterialMaskHandle;
 typedef FPhysicsShapeAdapter_Chaos			FPhysicsShapeAdapter;
 typedef FPhysicsUserData_Chaos				FPhysicsUserData;
 
