Make sure that we are handling both read failiures and cache overflows gracefully in stream caching.

diff --git a/Engine/Source/Runtime/Engine/Classes/Sound/SoundWave.h b/Engine/Source/Runtime/Engine/Classes/Sound/SoundWave.h
--- a/Engine/Source/Runtime/Engine/Classes/Sound/SoundWave.h
+++ b/Engine/Source/Runtime/Engine/Classes/Sound/SoundWave.h
@@ -482,6 +482,8 @@ public:
 	/** If stream caching is enabled and au.streamcache.KeepFirstChunkInMemory is 1, this will release this USoundWave's first chunk, allowing it to be deleted. */
 	void ReleaseCompressedAudio();
 
+	bool IsRetainingAudio();
+
 	/** Returns the loading behavior we should use for this sound wave.
 	 *  If this is called within Serialize(), this should be called with bCheckSoundClasses = false,
 	 *  Since there is no guarantee that the deserialized USoundClasses have been resolved yet.
diff --git a/Engine/Source/Runtime/Engine/Private/ADPCMAudioInfo.cpp b/Engine/Source/Runtime/Engine/Private/ADPCMAudioInfo.cpp
--- a/Engine/Source/Runtime/Engine/Private/ADPCMAudioInfo.cpp
+++ b/Engine/Source/Runtime/Engine/Private/ADPCMAudioInfo.cpp
@@ -23,6 +23,13 @@ FAutoConsoleVariableRef CVarDisableADPCMSeeking(
 	TEXT("Disables seeking with ADPCM.\n"),
 	ECVF_Default);
 
+static int32 ADPCMReadFailiureTimeoutCVar = 64;
+FAutoConsoleVariableRef CVarADPCMReadFailiureTimeout(
+	TEXT("au.adpcm.ADPCMReadFailiureTimeout"),
+	ADPCMReadFailiureTimeoutCVar,
+	TEXT("Sets the number of ADPCM decode attempts we'll try before stopping the sound wave altogether.\n"),
+	ECVF_Default);
+
 #define WAVE_FORMAT_LPCM  1
 #define WAVE_FORMAT_ADPCM 2
 
@@ -33,7 +40,8 @@ namespace ADPCM
 }
 
 FADPCMAudioInfo::FADPCMAudioInfo(void)
-	: UncompressedBlockSize(0)
+	: NumConsecutiveReadFailiures(0)
+	, UncompressedBlockSize(0)
 	, CompressedBlockSize(0)
 	, BlockSize(0)
 	, StreamBufferSize(0)
@@ -721,9 +729,14 @@ bool FADPCMAudioInfo::StreamCompressedData(uint8* Destination, bool bLooping, ui
 					}
 
 					FMemory::Memset(OutData, 0, BufferSize);
-					return false;
+					NumConsecutiveReadFailiures++;
+					const bool bReadAttemptTimedOut = NumConsecutiveReadFailiures > ADPCMReadFailiureTimeoutCVar;
+					UE_CLOG(bReadAttemptTimedOut, LogAudio, Warning, TEXT("ADPCM Audio Decode timed out."), bReadAttemptTimedOut);
+					return NumConsecutiveReadFailiures > ADPCMReadFailiureTimeoutCVar;
 				}
 
+				NumConsecutiveReadFailiures = 0;
+
 				// Set the current buffer offset accounting for the header in the first chunk
 				if (!bSeekPending)
 				{
diff --git a/Engine/Source/Runtime/Engine/Private/AudioStreamingCache.cpp b/Engine/Source/Runtime/Engine/Private/AudioStreamingCache.cpp
--- a/Engine/Source/Runtime/Engine/Private/AudioStreamingCache.cpp
+++ b/Engine/Source/Runtime/Engine/Private/AudioStreamingCache.cpp
@@ -340,6 +340,24 @@ FAudioChunkHandle FCachedAudioStreamingManager::GetLoadedChunk(const USoundWave*
 			if (!bIsValidChunk)
 			{
 				UE_LOG(LogAudio, Warning, TEXT("Cache overflow!!! couldn't load chunk %d for sound %s!"), ChunkIndex, *SoundWave->GetName());
+
+
+				AsyncTask(ENamedThreads::GameThread, []()
+				{
+					int32 NumChunksReleased = 0;
+
+					for (TObjectIterator<USoundWave> It; It; ++It)
+					{
+						USoundWave* Wave = *It;
+						if (Wave && Wave->IsRetainingAudio())
+						{
+							Wave->ReleaseCompressedAudio();
+							NumChunksReleased++;
+						}
+					}
+
+					UE_LOG(LogAudio, Warning, TEXT("Removed %d retained sounds from the stream cache."), NumChunksReleased);
+				});
 			}
 		}
 
diff --git a/Engine/Source/Runtime/Engine/Private/SoundWave.cpp b/Engine/Source/Runtime/Engine/Private/SoundWave.cpp
--- a/Engine/Source/Runtime/Engine/Private/SoundWave.cpp
+++ b/Engine/Source/Runtime/Engine/Private/SoundWave.cpp
@@ -2454,6 +2454,11 @@ void USoundWave::ReleaseCompressedAudio()
 	FirstChunk = FAudioChunkHandle();
 }
 
+bool USoundWave::IsRetainingAudio()
+{
+	return FirstChunk.IsValid();
+}
+
 void USoundWave::CacheInheritedLoadingBehavior()
 {
 	check(IsInGameThread());
diff --git a/Engine/Source/Runtime/Engine/Public/ADPCMAudioInfo.h b/Engine/Source/Runtime/Engine/Public/ADPCMAudioInfo.h
--- a/Engine/Source/Runtime/Engine/Public/ADPCMAudioInfo.h
+++ b/Engine/Source/Runtime/Engine/Public/ADPCMAudioInfo.h
@@ -140,6 +140,8 @@ private:
 	// until we move on to a different chunk.
 	const uint8* GetLoadedChunk(USoundWave* InSoundWave, uint32 ChunkIndex, uint32& OutChunkSize);
 
+	int32 NumConsecutiveReadFailiures;
+
 	FWaveModInfo WaveInfo;
 	const uint8*	SrcBufferData;
 	uint32			SrcBufferDataSize;
