Fix blueprint editor promotion automation tests
Fix duplicated default scene root when creating a new blueprint subclass

diff --git a/Engine/Source/Editor/Kismet/Private/BlueprintCompilationManager.cpp b/Engine/Source/Editor/Kismet/Private/BlueprintCompilationManager.cpp
--- a/Engine/Source/Editor/Kismet/Private/BlueprintCompilationManager.cpp
+++ b/Engine/Source/Editor/Kismet/Private/BlueprintCompilationManager.cpp
@@ -617,7 +617,7 @@ void FBlueprintCompilationManagerImpl::FlushCompilationQueueImpl(bool bSuppressB
 			if(UBlueprintGeneratedClass* BPGC = Cast<UBlueprintGeneratedClass>(QueuedBP->GeneratedClass))
 			{
 				if( BPGC->SimpleConstructionScript &&
-					BPGC->SimpleConstructionScript->GetSceneRootComponentTemplate() == nullptr)
+					BPGC->SimpleConstructionScript->GetSceneRootComponentTemplate(true) == nullptr)
 				{
 					bDefaultComponentMustBeAdded = true;
 				}
diff --git a/Engine/Source/Editor/UnrealEd/Private/Kismet2/BlueprintEditorUtils.cpp b/Engine/Source/Editor/UnrealEd/Private/Kismet2/BlueprintEditorUtils.cpp
--- a/Engine/Source/Editor/UnrealEd/Private/Kismet2/BlueprintEditorUtils.cpp
+++ b/Engine/Source/Editor/UnrealEd/Private/Kismet2/BlueprintEditorUtils.cpp
@@ -1385,7 +1385,7 @@ UClass* FBlueprintEditorUtils::RegenerateBlueprintClass(UBlueprint* Blueprint, U
 			&& BPGClassToRegenerate
 			&& SupportsConstructionScript(Blueprint) 
 			&& BPGClassToRegenerate->SimpleConstructionScript
-			&& (nullptr == BPGClassToRegenerate->SimpleConstructionScript->GetSceneRootComponentTemplate());
+			&& (nullptr == BPGClassToRegenerate->SimpleConstructionScript->GetSceneRootComponentTemplate(true));
 		const bool bShouldBeRecompiled = bHasCode || bDataOnlyClassThatMustBeRecompiled || bHasPendingUberGraphFrame || bDefaultComponentMustBeAdded;
 
 		if (bShouldBeRecompiled)
diff --git a/Engine/Source/Editor/UnrealEd/Private/Kismet2/Kismet2.cpp b/Engine/Source/Editor/UnrealEd/Private/Kismet2/Kismet2.cpp
--- a/Engine/Source/Editor/UnrealEd/Private/Kismet2/Kismet2.cpp
+++ b/Engine/Source/Editor/UnrealEd/Private/Kismet2/Kismet2.cpp
@@ -1503,8 +1503,6 @@ public:
 				Pair.Key->SetRelativeTransform(Pair.Value);
 			}
 
-			SCS->RemoveNode(SCS->GetDefaultSceneRootNode());
-
 			FTransform NewActorTransform = FTransform::Identity;
 			if (RootActors.Num() == 1)
 			{
diff --git a/Engine/Source/Runtime/Engine/Classes/Engine/SimpleConstructionScript.h b/Engine/Source/Runtime/Engine/Classes/Engine/SimpleConstructionScript.h
--- a/Engine/Source/Runtime/Engine/Classes/Engine/SimpleConstructionScript.h
+++ b/Engine/Source/Runtime/Engine/Classes/Engine/SimpleConstructionScript.h
@@ -54,8 +54,18 @@ class USimpleConstructionScript : public UObject
 	/** Return the Blueprint associated with this SCS instance */
 	ENGINE_API class UBlueprint* GetBlueprint() const;
 
-	/** Helper function to find the current scene root component template and/or owning SCS node */
-	ENGINE_API USceneComponent* GetSceneRootComponentTemplate(USCS_Node** OutSCSNode = nullptr) const;
+	/** 
+	 * Helper function to find the current scene root component template and optionally the owning SCS node 
+	 * @param bShouldUseDefaultRoot  Whether the default scene root, if already in the root node set, should be used (note this will return null in that case because of the expectations of other calling code)
+	 * @param OutSCSNode             If not-null, the SCSNode for the scene root component will be returned in this variable
+	 */
+	ENGINE_API USceneComponent* GetSceneRootComponentTemplate(bool bShouldUseDefaultRoot = false, USCS_Node** OutSCSNode = nullptr) const;
+
+	UE_DEPRECATED(4.25, "Use version that specifies whether the default scene root should be used")
+	ENGINE_API USceneComponent* GetSceneRootComponentTemplate(USCS_Node** OutSCSNode) const
+	{
+		return GetSceneRootComponentTemplate(false, OutSCSNode);
+	}
 #endif
 
 	/** Return the Blueprint associated with this SCS instance */
diff --git a/Engine/Source/Runtime/Engine/Private/SimpleConstructionScript.cpp b/Engine/Source/Runtime/Engine/Private/SimpleConstructionScript.cpp
--- a/Engine/Source/Runtime/Engine/Private/SimpleConstructionScript.cpp
+++ b/Engine/Source/Runtime/Engine/Private/SimpleConstructionScript.cpp
@@ -296,7 +296,7 @@ void USimpleConstructionScript::FixupSceneNodeHierarchy()
 	// by this SCS; it could be from a super SCS, or (if SceneRootNode and 
 	// SceneRootComponentTemplate is not) it could be a native component
 	USCS_Node* SceneRootNode = nullptr;
-	USceneComponent* SceneRootComponentTemplate = GetSceneRootComponentTemplate(&SceneRootNode);
+	USceneComponent* SceneRootComponentTemplate = GetSceneRootComponentTemplate(true, &SceneRootNode);
 
 	if (SceneRootComponentTemplate == nullptr)
 	{
@@ -985,7 +985,7 @@ USCS_Node* USimpleConstructionScript::FindSCSNodeByGuid(const FGuid Guid) const
 }
 
 #if WITH_EDITOR
-USceneComponent* USimpleConstructionScript::GetSceneRootComponentTemplate(USCS_Node** OutSCSNode) const
+USceneComponent* USimpleConstructionScript::GetSceneRootComponentTemplate(bool bShouldUseDefautRoot, USCS_Node** OutSCSNode) const
 {
 	UClass* GeneratedClass = GetOwnerClass();
 	UClass* ParentClass = GetParentClass();
@@ -1059,7 +1059,7 @@ USceneComponent* USimpleConstructionScript::GetSceneRootComponentTemplate(USCS_N
 		{
 			const TArray<USCS_Node*>& SCSRootNodes = SCSStack[StackIndex]->GetRootNodes();
 
-			const bool bCanUseDefaultSceneRoot = DefaultSceneRootNode && DefaultSceneRootNode->ComponentTemplate && SCSRootNodes.Contains(DefaultSceneRootNode);
+			const bool bCanUseDefaultSceneRoot = bShouldUseDefautRoot && DefaultSceneRootNode && DefaultSceneRootNode->ComponentTemplate && SCSRootNodes.Contains(DefaultSceneRootNode);
 			// Check for any scene component nodes in the root set that are not the default scene root
 			for (int32 RootNodeIndex = 0; RootNodeIndex < SCSRootNodes.Num() && RootComponentTemplate == nullptr; ++RootNodeIndex)
 			{
