- Deprecate AudioDevice implementation Engine/UnrealEngine monolythics & move to AudioDeviceManager
- Remove redundant static calls on AudioDevice in favor of AudioDeviceManager for finding Main/Active Devices (state managed by the manager, not individual devices)

diff --git a/Engine/Plugins/Media/PixelStreaming/Source/PixelStreaming/Private/AudioCapturer.cpp b/Engine/Plugins/Media/PixelStreaming/Source/PixelStreaming/Private/AudioCapturer.cpp
--- a/Engine/Plugins/Media/PixelStreaming/Source/PixelStreaming/Private/AudioCapturer.cpp
+++ b/Engine/Plugins/Media/PixelStreaming/Source/PixelStreaming/Private/AudioCapturer.cpp
@@ -114,7 +114,7 @@ int32 FAudioCapturer::Init()
 		return -1;
 	}
 
-	FAudioDevice* AudioDevice = GEngine->GetMainAudioDevice();
+	FAudioDevice* AudioDevice = FAudioDeviceManager::GetMainDevice();
 	if (!AudioDevice)
 	{
 		return -1;
@@ -139,7 +139,7 @@ int32 FAudioCapturer::Terminate()
 		return -1;
 	}
 
-	FAudioDevice* AudioDevice = GEngine->GetMainAudioDevice();
+	FAudioDevice* AudioDevice = FAudioDeviceManager::GetMainDevice();
 	if (!AudioDevice)
 	{
 		return -1;
diff --git a/Engine/Plugins/Online/OnlineSubsystemUtils/Source/OnlineSubsystemUtils/Private/OnlineSubsystemUtils.cpp b/Engine/Plugins/Online/OnlineSubsystemUtils/Source/OnlineSubsystemUtils/Private/OnlineSubsystemUtils.cpp
--- a/Engine/Plugins/Online/OnlineSubsystemUtils/Source/OnlineSubsystemUtils/Private/OnlineSubsystemUtils.cpp
+++ b/Engine/Plugins/Online/OnlineSubsystemUtils/Source/OnlineSubsystemUtils/Private/OnlineSubsystemUtils.cpp
@@ -110,7 +110,7 @@ UAudioComponent* CreateVoiceAudioComponent(uint32 SampleRate, int32 NumChannels)
 	UAudioComponent* AudioComponent = nullptr;
 	if (GEngine != nullptr)
 	{
-		if (FAudioDevice* AudioDevice = GEngine->GetMainAudioDevice())
+		if (FAudioDevice* AudioDevice = FAudioDeviceManager::GetMainDevice())
 		{
 			USoundWaveProcedural* SoundStreaming = NewObject<USoundWaveProcedural>();
 			SoundStreaming->SetSampleRate(SampleRate);
@@ -157,7 +157,7 @@ UAudioComponent* CreateVoiceAudioComponent(uint32 SampleRate, int32 NumChannels)
 UVoipListenerSynthComponent* CreateVoiceSynthComponent(uint32 SampleRate)
 {
 	UVoipListenerSynthComponent* SynthComponentPtr = nullptr;
-	if (FAudioDevice* AudioDevice = GEngine->GetMainAudioDevice())
+	if (FAudioDevice* AudioDevice = FAudioDeviceManager::GetMainDevice())
 	{
 		SynthComponentPtr = NewObject<UVoipListenerSynthComponent>();
 		if (SynthComponentPtr)
diff --git a/Engine/Plugins/Online/OnlineSubsystemUtils/Source/OnlineSubsystemUtils/Private/VoiceEngineImpl.cpp b/Engine/Plugins/Online/OnlineSubsystemUtils/Source/OnlineSubsystemUtils/Private/VoiceEngineImpl.cpp
--- a/Engine/Plugins/Online/OnlineSubsystemUtils/Source/OnlineSubsystemUtils/Private/VoiceEngineImpl.cpp
+++ b/Engine/Plugins/Online/OnlineSubsystemUtils/Source/OnlineSubsystemUtils/Private/VoiceEngineImpl.cpp
@@ -1021,9 +1021,7 @@ IOnlineSubsystem* FVoiceEngineImpl::GetOnlineSubSystem()
 FVoiceEndpoint::FVoiceEndpoint(const FString& InEndpointName, float InSampleRate, int32 InNumChannels)
 	: NumChannelsComingIn(InNumChannels)
 {
-	check(GEngine && GEngine->GetAudioDeviceManager());
-
-	IAudioDeviceModule* AudioModule = GEngine->GetAudioDeviceManager()->GetAudioDeviceModule();
+	IAudioDeviceModule* AudioModule = FAudioDeviceManager::GetChecked().GetAudioDeviceModule();
 	check(AudioModule);
 
 	PlatformEndpoint.Reset(AudioModule->CreateAudioMixerPlatformInterface());
diff --git a/Engine/Plugins/Runtime/AudioCapture/Source/AudioCaptureEditor/Private/AudioRecordingManager.cpp b/Engine/Plugins/Runtime/AudioCapture/Source/AudioCaptureEditor/Private/AudioRecordingManager.cpp
--- a/Engine/Plugins/Runtime/AudioCapture/Source/AudioCaptureEditor/Private/AudioRecordingManager.cpp
+++ b/Engine/Plugins/Runtime/AudioCapture/Source/AudioCaptureEditor/Private/AudioRecordingManager.cpp
@@ -387,7 +387,7 @@ void FAudioRecordingManager::StopRecording(TArray<USoundWave*>& OutSoundWaves)
 					ExistingSoundWave = FindObject<USoundWave>(Package, *AssetName);
 
 					// See if it's currently being played right now
-					FAudioDeviceManager* AudioDeviceManager = GEngine->GetAudioDeviceManager();
+					FAudioDeviceManager* AudioDeviceManager = FAudioDeviceManager::Get();
 					if (AudioDeviceManager && ExistingSoundWave)
 					{
 						AudioDeviceManager->StopSoundsUsingResource(ExistingSoundWave, &ComponentsToRestart);
diff --git a/Engine/Plugins/Runtime/AudioModulation/Source/AudioModulation/Private/AudioModulationStatics.cpp b/Engine/Plugins/Runtime/AudioModulation/Source/AudioModulation/Private/AudioModulationStatics.cpp
--- a/Engine/Plugins/Runtime/AudioModulation/Source/AudioModulation/Private/AudioModulationStatics.cpp
+++ b/Engine/Plugins/Runtime/AudioModulation/Source/AudioModulation/Private/AudioModulationStatics.cpp
@@ -111,10 +111,7 @@ AudioModulation::FAudioModulationImpl* UAudioModulationStatics::GetModulationImp
 	}
 	else
 	{
-		if (GEngine)
-		{
-			AudioDevice = GEngine->GetMainAudioDevice();
-		}
+		AudioDevice = FAudioDeviceManager::GetMainDevice();
 	}
 
 	if (AudioDevice && AudioDevice->IsModulationPluginEnabled())
diff --git a/Engine/Plugins/Runtime/AudioModulation/Source/AudioModulation/Private/SoundModulatorBase.cpp b/Engine/Plugins/Runtime/AudioModulation/Source/AudioModulation/Private/SoundModulatorBase.cpp
--- a/Engine/Plugins/Runtime/AudioModulation/Source/AudioModulation/Private/SoundModulatorBase.cpp
+++ b/Engine/Plugins/Runtime/AudioModulation/Source/AudioModulation/Private/SoundModulatorBase.cpp
@@ -24,7 +24,7 @@ void USoundModulatorBase::PostEditChangeProperty(struct FPropertyChangedEvent& P
 		return;
 	}
 
-	FAudioDeviceManager* DeviceManager = GEngine->GetAudioDeviceManager();
+	FAudioDeviceManager* DeviceManager = FAudioDeviceManager::Get();
 	if (!DeviceManager)
 	{
 		return;
diff --git a/Engine/Plugins/Runtime/Oculus/OculusVR/Source/OculusMR/Private/OculusMR_CastingCameraActor.cpp b/Engine/Plugins/Runtime/Oculus/OculusVR/Source/OculusMR/Private/OculusMR_CastingCameraActor.cpp
--- a/Engine/Plugins/Runtime/Oculus/OculusVR/Source/OculusMR/Private/OculusMR_CastingCameraActor.cpp
+++ b/Engine/Plugins/Runtime/Oculus/OculusVR/Source/OculusMR/Private/OculusMR_CastingCameraActor.cpp
@@ -274,17 +274,15 @@ void AOculusMR_CastingCameraActor::BeginPlay()
 	VRNotificationComponent->HMDRecenteredDelegate.Add(Delegate);
 
 #if PLATFORM_ANDROID
-	FAudioDevice* AudioDevice = FAudioDevice::GetMainAudioDevice();
-	AudioDevice->StartRecording(nullptr, 0.1);
+	FAudioDeviceManager::GetMainDevice()->StartRecording(nullptr, 0.1);
 #endif
 }
 
 void AOculusMR_CastingCameraActor::EndPlay(EEndPlayReason::Type Reason)
 {
 #if PLATFORM_ANDROID
-	FAudioDevice* AudioDevice = FAudioDevice::GetMainAudioDevice();
 	float NumChannels = 2;
-	float SampleRate = AudioDevice->GetSampleRate();
+	float SampleRate = FAudioDeviceManager::GetMainDevice()->GetSampleRate();
 	AudioDevice->StopRecording(nullptr, NumChannels, SampleRate);
 #endif
 
@@ -486,7 +484,7 @@ void AOculusMR_CastingCameraActor::Tick(float DeltaTime)
 		ForegroundCaptureActor->GetCaptureComponent2D()->TextureTarget = ForegroundRenderTargets[CaptureIndex];
 		GetCaptureComponent2D()->SetVisibility(true);
 
-		FAudioDevice* AudioDevice = FAudioDevice::GetMainAudioDevice();
+		FAudioDevice* AudioDevice = FAudioDeviceManager::GetMainDevice();
 		float NumChannels, SampleRate;
 		NumChannels = 2;
 		SampleRate = AudioDevice->GetSampleRate();
diff --git a/Engine/Plugins/Runtime/SoundMod/Source/SoundModImporter/Private/SoundModImporterFactory.cpp b/Engine/Plugins/Runtime/SoundMod/Source/SoundModImporter/Private/SoundModImporterFactory.cpp
--- a/Engine/Plugins/Runtime/SoundMod/Source/SoundModImporter/Private/SoundModImporterFactory.cpp
+++ b/Engine/Plugins/Runtime/SoundMod/Source/SoundModImporter/Private/SoundModImporterFactory.cpp
@@ -44,7 +44,7 @@ UObject* USoundModImporterFactory::FactoryCreateBinary
 
 	// TODO - Audio Threading. This needs to be sent to the audio device and wait on stopping the sounds
 	TArray<UAudioComponent*> ComponentsToRestart;
-	FAudioDeviceManager* AudioDeviceManager = GEngine->GetAudioDeviceManager();
+	FAudioDeviceManager* AudioDeviceManager = FAudioDeviceManager::Get();
 	if (AudioDeviceManager && ExistingSound)
 	{
 		// TODO: Generalize the stop sounds function
diff --git a/Engine/Plugins/Runtime/nDisplay/Source/DisplayCluster/Private/Game/Classes/Basics/DisplayClusterViewportClient.cpp b/Engine/Plugins/Runtime/nDisplay/Source/DisplayCluster/Private/Game/Classes/Basics/DisplayClusterViewportClient.cpp
--- a/Engine/Plugins/Runtime/nDisplay/Source/DisplayCluster/Private/Game/Classes/Basics/DisplayClusterViewportClient.cpp
+++ b/Engine/Plugins/Runtime/nDisplay/Source/DisplayCluster/Private/Game/Classes/Basics/DisplayClusterViewportClient.cpp
@@ -315,19 +315,13 @@ void UDisplayClusterViewportClient::Draw(FViewport* InViewport, FCanvas* SceneCa
 								bool bUpdateListenerPosition = true;
 
 								// If the main audio device is used for multiple PIE viewport clients, we only
-								// want to update the main audio device listener position if it is in focus
-								if (GEngine)
+								// want to update the main audio device listener position if it is in focus.
+								// If there is more than one world referencing the main audio device
+								if (FAudioDeviceManager::GetChecked().GetNumMainAudioDeviceWorlds() > 1)
 								{
-									FAudioDeviceManager* AudioDeviceManager = GEngine->GetAudioDeviceManager();
-
-									// If there is more than one world referencing the main audio device
-									if (AudioDeviceManager->GetNumMainAudioDeviceWorlds() > 1)
+									if (AudioDevice == FAudioDeviceManager::GetMainDevice() && !HasAudioFocus())
 									{
-										uint32 MainAudioDeviceHandle = GEngine->GetAudioDeviceHandle();
-										if (AudioDevice->DeviceHandle == MainAudioDeviceHandle && !HasAudioFocus())
-										{
-											bUpdateListenerPosition = false;
-										}
+										bUpdateListenerPosition = false;
 									}
 								}
 
diff --git a/Engine/Source/Editor/AudioEditor/Classes/Factories/SoundSurroundFactory.h b/Engine/Source/Editor/AudioEditor/Classes/Factories/SoundSurroundFactory.h
--- a/Engine/Source/Editor/AudioEditor/Classes/Factories/SoundSurroundFactory.h
+++ b/Engine/Source/Editor/AudioEditor/Classes/Factories/SoundSurroundFactory.h
@@ -9,7 +9,6 @@
 #include "CoreMinimal.h"
 #include "UObject/ObjectMacros.h"
 #include "Factories/Factory.h"
-#include "Audio.h"
 #include "SoundSurroundFactory.generated.h"
 
 const FString SurroundSpeakerLocations[SPEAKER_Count] =
diff --git a/Engine/Source/Editor/AudioEditor/Private/AssetTypeActions/AssetTypeActions_SoundBase.cpp b/Engine/Source/Editor/AudioEditor/Private/AssetTypeActions/AssetTypeActions_SoundBase.cpp
--- a/Engine/Source/Editor/AudioEditor/Private/AssetTypeActions/AssetTypeActions_SoundBase.cpp
+++ b/Engine/Source/Editor/AudioEditor/Private/AssetTypeActions/AssetTypeActions_SoundBase.cpp
@@ -294,7 +294,7 @@ bool FAssetTypeActions_SoundBase::CanExecutePlayCommand(TArray<TWeakObjectPtr<US
 
 void FAssetTypeActions_SoundBase::ExecuteMuteSound(TArray<TWeakObjectPtr<USoundBase>> Objects) const
 {	
-	if (FAudioDeviceManager* ADM = GEditor->GetAudioDeviceManager())
+	if (FAudioDeviceManager* ADM = FAudioDeviceManager::Get())
 	{
 		FAudioDebugger& Debugger = ADM->GetDebugger();
 
@@ -318,7 +318,7 @@ void FAssetTypeActions_SoundBase::ExecuteMuteSound(TArray<TWeakObjectPtr<USoundB
 
 void FAssetTypeActions_SoundBase::ExecuteSoloSound(TArray<TWeakObjectPtr<USoundBase>> Objects) const
 {
-	if (FAudioDeviceManager* ADM = GEditor->GetAudioDeviceManager())
+	if (FAudioDeviceManager* ADM = FAudioDeviceManager::Get())
 	{
 		FAudioDebugger& Debugger = ADM->GetDebugger();
 
@@ -342,7 +342,7 @@ void FAssetTypeActions_SoundBase::ExecuteSoloSound(TArray<TWeakObjectPtr<USoundB
 
 bool FAssetTypeActions_SoundBase::IsActionCheckedMute(TArray<TWeakObjectPtr<USoundBase>> Objects) const
 {
-	if (FAudioDeviceManager* ADM = GEditor->GetAudioDeviceManager())
+	if (FAudioDeviceManager* ADM = FAudioDeviceManager::Get())
 	{
 		// If *any* of the selection are muted, show the tick box as ticked.
 		FAudioDebugger& Debugger = ADM->GetDebugger();
@@ -370,7 +370,7 @@ bool FAssetTypeActions_SoundBase::IsActionCheckedMute(TArray<TWeakObjectPtr<USou
 bool FAssetTypeActions_SoundBase::IsActionCheckedSolo(TArray<TWeakObjectPtr<USoundBase>> Objects) const
 {
 	// If *any* of the selection are solod, show the tick box as ticked.
-	if (FAudioDeviceManager* ADM = GEditor->GetAudioDeviceManager())
+	if (FAudioDeviceManager* ADM = FAudioDeviceManager::Get())
 	{
 		FAudioDebugger& Debugger = ADM->GetDebugger();
 		for (TWeakObjectPtr<USoundBase> SoundBase : Objects)
@@ -396,7 +396,7 @@ bool FAssetTypeActions_SoundBase::IsActionCheckedSolo(TArray<TWeakObjectPtr<USou
 
 bool FAssetTypeActions_SoundBase::CanExecuteMuteCommand(TArray<TWeakObjectPtr<USoundBase>> Objects) const
 {
-	if (FAudioDeviceManager* ADM = GEditor->GetAudioDeviceManager())
+	if (FAudioDeviceManager* ADM = FAudioDeviceManager::Get())
 	{
 		// Allow muting if we're not Soloing.
 		FAudioDebugger& Debugger = ADM->GetDebugger();
@@ -426,7 +426,7 @@ bool FAssetTypeActions_SoundBase::CanExecuteMuteCommand(TArray<TWeakObjectPtr<US
 
 bool FAssetTypeActions_SoundBase::CanExecuteSoloCommand(TArray<TWeakObjectPtr<USoundBase>> Objects) const
 {	
-	if (FAudioDeviceManager* ADM = GEditor->GetAudioDeviceManager())
+	if (FAudioDeviceManager* ADM = FAudioDeviceManager::Get())
 	{
 		// Allow Soloing if we're not Muting.
 		FAudioDebugger& Debugger = ADM->GetDebugger();
diff --git a/Engine/Source/Editor/AudioEditor/Private/AssetTypeActions/AssetTypeActions_SoundClass.cpp b/Engine/Source/Editor/AudioEditor/Private/AssetTypeActions/AssetTypeActions_SoundClass.cpp
--- a/Engine/Source/Editor/AudioEditor/Private/AssetTypeActions/AssetTypeActions_SoundClass.cpp
+++ b/Engine/Source/Editor/AudioEditor/Private/AssetTypeActions/AssetTypeActions_SoundClass.cpp
@@ -73,7 +73,7 @@ void FAssetTypeActions_SoundClass::GetActions(const TArray<UObject*>& InObjects,
 
 void FAssetTypeActions_SoundClass::ExecuteMute(TArray<TWeakObjectPtr<USoundClass>> Objects) const
 {
-	if (FAudioDeviceManager* ADM = GEditor->GetAudioDeviceManager())
+	if (FAudioDeviceManager* ADM = FAudioDeviceManager::Get())
 	{
 		for (TWeakObjectPtr<USoundClass> i : Objects)
 		{
@@ -87,7 +87,7 @@ void FAssetTypeActions_SoundClass::ExecuteMute(TArray<TWeakObjectPtr<USoundClass
 
 void FAssetTypeActions_SoundClass::ExecuteSolo(TArray<TWeakObjectPtr<USoundClass>> Objects) const
 {
-	if (FAudioDeviceManager* ADM = GEditor->GetAudioDeviceManager())
+	if (FAudioDeviceManager* ADM = FAudioDeviceManager::Get())
 	{
 		for (TWeakObjectPtr<USoundClass> i : Objects)
 		{
@@ -101,7 +101,7 @@ void FAssetTypeActions_SoundClass::ExecuteSolo(TArray<TWeakObjectPtr<USoundClass
 
 bool FAssetTypeActions_SoundClass::IsActionCheckedMute(TArray<TWeakObjectPtr<USoundClass>> Objects) const
 {
-	if (FAudioDeviceManager* ADM = GEditor->GetAudioDeviceManager())
+	if (FAudioDeviceManager* ADM = FAudioDeviceManager::Get())
 	{
 		for (TWeakObjectPtr<USoundClass> i : Objects)
 		{
@@ -119,7 +119,7 @@ bool FAssetTypeActions_SoundClass::IsActionCheckedMute(TArray<TWeakObjectPtr<USo
 
 bool FAssetTypeActions_SoundClass::IsActionCheckedSolo(TArray<TWeakObjectPtr<USoundClass>> Objects) const
 {
-	if (FAudioDeviceManager* ADM = GEditor->GetAudioDeviceManager())
+	if (FAudioDeviceManager* ADM = FAudioDeviceManager::Get())
 	{		
 		for (TWeakObjectPtr<USoundClass> i : Objects)
 		{
@@ -137,7 +137,7 @@ bool FAssetTypeActions_SoundClass::IsActionCheckedSolo(TArray<TWeakObjectPtr<USo
 
 bool FAssetTypeActions_SoundClass::CanExecuteMuteCommand(TArray<TWeakObjectPtr<USoundClass>> Objects) const
 {
-	if (FAudioDeviceManager* ADM = GEditor->GetAudioDeviceManager())
+	if (FAudioDeviceManager* ADM = FAudioDeviceManager::Get())
 	{
 		// Allow muting if we're not Soloing.
 		for (TWeakObjectPtr<USoundClass> i : Objects)
@@ -158,7 +158,7 @@ bool FAssetTypeActions_SoundClass::CanExecuteMuteCommand(TArray<TWeakObjectPtr<U
 
 bool FAssetTypeActions_SoundClass::CanExecuteSoloCommand(TArray<TWeakObjectPtr<USoundClass>> Objects) const
 {
-	if (FAudioDeviceManager* ADM = GEditor->GetAudioDeviceManager())
+	if (FAudioDeviceManager* ADM = FAudioDeviceManager::Get())
 	{
 		// Allow Soloing if we're not Muting.
 		for (TWeakObjectPtr<USoundClass> i : Objects)
diff --git a/Engine/Source/Editor/AudioEditor/Private/Factories/SoundClassFactory.cpp b/Engine/Source/Editor/AudioEditor/Private/Factories/SoundClassFactory.cpp
--- a/Engine/Source/Editor/AudioEditor/Private/Factories/SoundClassFactory.cpp
+++ b/Engine/Source/Editor/AudioEditor/Private/Factories/SoundClassFactory.cpp
@@ -17,12 +17,9 @@ USoundClassFactory::USoundClassFactory(const FObjectInitializer& ObjectInitializ
 UObject* USoundClassFactory::FactoryCreateNew(UClass* InClass, UObject* InParent, FName InName, EObjectFlags Flags, UObject* Context, FFeedbackContext* Warn)
 {
 	USoundClass* SoundClass = NewObject<USoundClass>(InParent, InName, Flags);
-
-	class FAudioDeviceManager* AudioDeviceManager = GEngine ? GEngine->GetAudioDeviceManager() : nullptr;
-	if (AudioDeviceManager)
+	if (FAudioDeviceManager* DeviceManager = FAudioDeviceManager::Get())
 	{
-		AudioDeviceManager->InitSoundClasses();
+		DeviceManager->InitSoundClasses();
 	}
-
-	return(SoundClass);
+	return SoundClass;
 }
diff --git a/Engine/Source/Editor/AudioEditor/Private/Factories/SoundFactory.cpp b/Engine/Source/Editor/AudioEditor/Private/Factories/SoundFactory.cpp
--- a/Engine/Source/Editor/AudioEditor/Private/Factories/SoundFactory.cpp
+++ b/Engine/Source/Editor/AudioEditor/Private/Factories/SoundFactory.cpp
@@ -221,7 +221,7 @@ UObject* USoundFactory::CreateObject
 		USoundWave* ExistingSound = FindObject<USoundWave>(InParent, *Name.ToString());
 
 		TArray<UAudioComponent*> ComponentsToRestart;
-		FAudioDeviceManager* AudioDeviceManager = GEngine->GetAudioDeviceManager();
+		FAudioDeviceManager* AudioDeviceManager = FAudioDeviceManager::Get();
 		if (AudioDeviceManager && ExistingSound)
 		{
 			// Will block internally on audio thread completing outstanding commands
@@ -237,7 +237,7 @@ UObject* USoundFactory::CreateObject
 			// to be auditioned in the editor properly.
 			if (!ExistingSound->ResourceData)
 			{
-				if (FAudioDevice* AudioDevice = GEngine->GetMainAudioDevice())
+				if (FAudioDevice* AudioDevice = FAudioDeviceManager::GetMainDevice())
 				{
 					FName RuntimeFormat = AudioDevice->GetRuntimeFormat(ExistingSound);
 					ExistingSound->InitAudioResource(RuntimeFormat);
diff --git a/Engine/Source/Editor/AudioEditor/Private/Factories/SoundSubmixFactory.cpp b/Engine/Source/Editor/AudioEditor/Private/Factories/SoundSubmixFactory.cpp
--- a/Engine/Source/Editor/AudioEditor/Private/Factories/SoundSubmixFactory.cpp
+++ b/Engine/Source/Editor/AudioEditor/Private/Factories/SoundSubmixFactory.cpp
@@ -19,8 +19,7 @@ UObject* USoundSubmixFactory::FactoryCreateNew(UClass* InClass, UObject* InParen
 {
 	USoundSubmix* SoundSubmix = NewObject<USoundSubmix>(InParent, InName, Flags);
 
-	class FAudioDeviceManager* AudioDeviceManager = GEngine ? GEngine->GetAudioDeviceManager() : nullptr;
-	if (AudioDeviceManager)
+	if (FAudioDeviceManager* AudioDeviceManager = FAudioDeviceManager::Get())
 	{
 		AudioDeviceManager->InitSoundSubmixes();
 	}
diff --git a/Engine/Source/Editor/AudioEditor/Private/SoundClassEditor.cpp b/Engine/Source/Editor/AudioEditor/Private/SoundClassEditor.cpp
--- a/Engine/Source/Editor/AudioEditor/Private/SoundClassEditor.cpp
+++ b/Engine/Source/Editor/AudioEditor/Private/SoundClassEditor.cpp
@@ -161,11 +161,6 @@ void FSoundClassEditor::InitSoundClassEditor( const EToolkitMode::Type Mode, con
 
 	ExtendToolbar();
 	RegenerateMenusAndToolbars();
-			
-	if (GEditor->GetAudioDeviceManager())
-	{
-		Debugger = &GEditor->GetAudioDeviceManager()->GetDebugger();
-	}	
 
 	GraphEditor->SelectAllNodes();
 	for (UObject* SelectedNode : GraphEditor->GetSelectedNodes())
@@ -208,7 +203,6 @@ void FSoundClassEditor::BindCommands()
 
 FSoundClassEditor::FSoundClassEditor()
 	: SoundClass(nullptr)
-	, Debugger(nullptr)
 {
 }
 
@@ -397,40 +391,57 @@ void FSoundClassEditor::RedoGraphAction()
 
 void FSoundClassEditor::ToggleSolo()
 {
-	if (Debugger)
+	if (FAudioDeviceManager* DeviceManager = FAudioDeviceManager::Get())
 	{
-		Debugger->ToggleSoloSoundClass(SoundClass->GetFName());
+		DeviceManager->GetDebugger().ToggleSoloSoundClass(SoundClass->GetFName());
 	}
 }
 
 bool FSoundClassEditor::CanExcuteToggleSolo() const
 {
-	// Enable solo if we are not Muted	
-	return Debugger ? !Debugger->IsMuteSoundClass(SoundClass->GetFName()) : false;
+	// Enable solo if we are not Muted
+	if (FAudioDeviceManager* DeviceManager = FAudioDeviceManager::Get())
+	{
+		return !DeviceManager->GetDebugger().IsMuteSoundClass(SoundClass->GetFName());
+	}
+	return false;
 }
 
 bool FSoundClassEditor::IsSoloToggled() const
 {
-	return Debugger ? Debugger->IsSoloSoundClass(SoundClass->GetFName()) : false;
+	if (FAudioDeviceManager* DeviceManager = FAudioDeviceManager::Get())
+	{
+		return DeviceManager->GetDebugger().IsSoloSoundClass(SoundClass->GetFName());
+	}
+	return false;
 }
 
 void FSoundClassEditor::ToggleMute()
 {
-	if (Debugger)
+	if (FAudioDeviceManager* DeviceManager = FAudioDeviceManager::Get())
 	{
-		Debugger->ToggleMuteSoundClass(SoundClass->GetFName());
+		DeviceManager->GetDebugger().ToggleMuteSoundClass(SoundClass->GetFName());
 	}
 }
 
 bool FSoundClassEditor::CanExcuteToggleMute() const
 {
-	// Enable mute if we are not Soloed
-	return Debugger ? !Debugger->IsSoloSoundClass(SoundClass->GetFName()) : false;
+	if (FAudioDeviceManager* DeviceManager = FAudioDeviceManager::Get())
+	{
+		return !DeviceManager->GetDebugger().IsSoloSoundClass(SoundClass->GetFName());
+	}
+
+	return false;
 }
 
 bool FSoundClassEditor::IsMuteToggled() const
-{	
-	return Debugger ? Debugger->IsMuteSoundClass(SoundClass->GetFName()) : false;
+{
+	if (FAudioDeviceManager* DeviceManager = FAudioDeviceManager::Get())
+	{
+		return DeviceManager->GetDebugger().IsMuteSoundClass(SoundClass->GetFName());
+	}
+
+	return false;
 }
 
 void FSoundClassEditor::CreateSoundClass(UEdGraphPin* FromPin, const FVector2D& Location, const FString& Name)
diff --git a/Engine/Source/Editor/AudioEditor/Private/SoundClassEditor.h b/Engine/Source/Editor/AudioEditor/Private/SoundClassEditor.h
--- a/Engine/Source/Editor/AudioEditor/Private/SoundClassEditor.h
+++ b/Engine/Source/Editor/AudioEditor/Private/SoundClassEditor.h
@@ -118,9 +118,6 @@ private:
 	/** Command list for this editor */
 	TSharedPtr<FUICommandList> GraphEditorCommands;
 
-	/** Cache the audio debugger instance */
-	class FAudioDebugger* Debugger;
-
 	/**	The tab ids for all the tabs used */
 	static const FName GraphCanvasTabId;
 	static const FName PropertiesTabId;
diff --git a/Engine/Source/Editor/AudioEditor/Private/SoundCueEditor.cpp b/Engine/Source/Editor/AudioEditor/Private/SoundCueEditor.cpp
--- a/Engine/Source/Editor/AudioEditor/Private/SoundCueEditor.cpp
+++ b/Engine/Source/Editor/AudioEditor/Private/SoundCueEditor.cpp
@@ -41,7 +41,6 @@ const FName FSoundCueEditor::PaletteTabId( TEXT( "SoundCueEditor_Palette" ) );
 
 FSoundCueEditor::FSoundCueEditor()
 	: SoundCue(nullptr)
-	, Debugger(nullptr)
 {
 }
 
@@ -149,11 +148,6 @@ void FSoundCueEditor::InitSoundCueEditor(const EToolkitMode::Type Mode, const TS
 	ExtendToolbar();
 	RegenerateMenusAndToolbars();
 
-	if (GEditor->GetAudioDeviceManager())
-	{
-		Debugger = &GEditor->GetAudioDeviceManager()->GetDebugger();
-	}	
-	
 	// @todo toolkit world centric editing
 	/*if(IsWorldCentricAssetEditor())
 	{
@@ -422,40 +416,60 @@ void FSoundCueEditor::TogglePlayback()
 
 void FSoundCueEditor::ToggleSolo()
 {
-	if (Debugger)
+	if (FAudioDeviceManager* DeviceManager = FAudioDeviceManager::Get())
 	{
-		Debugger->ToggleSoloSoundCue(SoundCue->GetFName());
+		DeviceManager->GetDebugger().ToggleSoloSoundCue(SoundCue->GetFName());
 	}
 }
 
 bool FSoundCueEditor::CanExcuteToggleSolo() const
 {
 	// Allow Solo if Mute is not Toggle on
-	return Debugger ? !Debugger->IsMuteSoundCue(SoundCue->GetFName()) : false;
+	if (FAudioDeviceManager* DeviceManager = FAudioDeviceManager::Get())
+	{
+		return !DeviceManager->GetDebugger().IsMuteSoundCue(SoundCue->GetFName());
+	}
+
+	return false;
 }
 
 bool FSoundCueEditor::IsSoloToggled() const
 {
-	return Debugger ? Debugger->IsSoloSoundCue(SoundCue->GetFName()) : false;
+	if (FAudioDeviceManager* DeviceManager = FAudioDeviceManager::Get())
+	{
+		return DeviceManager->GetDebugger().IsSoloSoundCue(SoundCue->GetFName());
+	}
+
+	return false;
 }
 
 void FSoundCueEditor::ToggleMute()
 {
-	if (Debugger)
+	if (FAudioDeviceManager* DeviceManager = FAudioDeviceManager::Get())
 	{
-		Debugger->ToggleMuteSoundCue(SoundCue->GetFName());
+		DeviceManager->GetDebugger().ToggleMuteSoundCue(SoundCue->GetFName());
 	}
 }
 
 bool FSoundCueEditor::CanExcuteToggleMute() const
 {
 	// Allow Mute if Solo is not Toggle on
-	return Debugger ? !Debugger->IsSoloSoundCue(SoundCue->GetFName()) : false;
+	if (FAudioDeviceManager* DeviceManager = FAudioDeviceManager::Get())
+	{
+		return !DeviceManager->GetDebugger().IsSoloSoundCue(SoundCue->GetFName());
+	}
+
+	return false;
 }
 
 bool FSoundCueEditor::IsMuteToggled() const
 {
-	return Debugger ? Debugger->IsMuteSoundCue(SoundCue->GetFName()) : false;
+	if (FAudioDeviceManager* DeviceManager = FAudioDeviceManager::Get())
+	{
+		return DeviceManager->GetDebugger().IsMuteSoundCue(SoundCue->GetFName());
+	}
+
+	return false;
 }
 
 void FSoundCueEditor::PlaySingleNode(UEdGraphNode* Node)
diff --git a/Engine/Source/Editor/AudioEditor/Private/SoundCueEditor.h b/Engine/Source/Editor/AudioEditor/Private/SoundCueEditor.h
--- a/Engine/Source/Editor/AudioEditor/Private/SoundCueEditor.h
+++ b/Engine/Source/Editor/AudioEditor/Private/SoundCueEditor.h
@@ -210,9 +210,6 @@ private:
 	/** Command list for this editor */
 	TSharedPtr<FUICommandList> GraphEditorCommands;
 
-	/** Cache of the Audio debugger instance */
-	class FAudioDebugger* Debugger;
-
 	/**	The tab ids for all the tabs used */
 	static const FName GraphCanvasTabId;
 	static const FName PropertiesTabId;
diff --git a/Engine/Source/Editor/DetailCustomizations/Private/LinuxTargetSettingsDetails.cpp b/Engine/Source/Editor/DetailCustomizations/Private/LinuxTargetSettingsDetails.cpp
--- a/Engine/Source/Editor/DetailCustomizations/Private/LinuxTargetSettingsDetails.cpp
+++ b/Engine/Source/Editor/DetailCustomizations/Private/LinuxTargetSettingsDetails.cpp
@@ -340,8 +340,7 @@ bool FLinuxTargetSettingsDetails::IsValidAudioDeviceName(const FString& InDevice
 	bool bIsValid = false;
 
 #if WITH_ENGINE
-	FAudioDevice* AudioDevice = GEngine->GetMainAudioDevice();
-	if (AudioDevice)
+	if (FAudioDevice* AudioDevice = FAudioDeviceManager::GetMainDevice())
 	{
 		TArray<FString> DeviceNames;
 		AudioDevice->GetAudioDeviceList(DeviceNames);
@@ -365,7 +364,7 @@ TSharedRef<SWidget> FLinuxTargetSettingsDetails::MakeAudioDeviceMenu(const TShar
 	FMenuBuilder MenuBuilder(true, nullptr);
 
 #if WITH_ENGINE
-	FAudioDevice* AudioDevice = GEngine->GetMainAudioDevice();
+	FAudioDevice* AudioDevice = FAudioDeviceManager::GetMainDevice();
 	if (AudioDevice)
 	{
 		TArray<FString> AudioDeviceNames;
diff --git a/Engine/Source/Editor/Kismet/Private/BlueprintEditor.cpp b/Engine/Source/Editor/Kismet/Private/BlueprintEditor.cpp
--- a/Engine/Source/Editor/Kismet/Private/BlueprintEditor.cpp
+++ b/Engine/Source/Editor/Kismet/Private/BlueprintEditor.cpp
@@ -8986,11 +8986,11 @@ void FBlueprintEditor::UpdatePreviewActor(UBlueprint* InBlueprint, bool bInForce
 			if ( PreviewActor->IsHidden() )
 			{
 				PreviewActor->SetHidden(false);
-				PreviewActor->MarkComponentsRenderStateDirty();				
+				PreviewActor->MarkComponentsRenderStateDirty();
 			}
 
 			// Prevent any audio from playing as a result of spawning
-			if (FAudioDevice* AudioDevice = GEngine->GetMainAudioDevice())
+			if (FAudioDevice* AudioDevice = FAudioDeviceManager::GetMainDevice())
 			{
 				AudioDevice->Flush(PreviewScene.GetWorld());
 			}
diff --git a/Engine/Source/Editor/UnrealEd/Private/Commandlets/AudioMixerCommandlet.cpp b/Engine/Source/Editor/UnrealEd/Private/Commandlets/AudioMixerCommandlet.cpp
--- a/Engine/Source/Editor/UnrealEd/Private/Commandlets/AudioMixerCommandlet.cpp
+++ b/Engine/Source/Editor/UnrealEd/Private/Commandlets/AudioMixerCommandlet.cpp
@@ -268,7 +268,7 @@ public:
 		}
 
 		// Get the audio main device
-		FAudioDevice* AudioDevice = GEngine->GetMainAudioDevice();
+		FAudioDevice* AudioDevice = FAudioDeviceManager::GetMainDevice();
 
 		bool bSuccess = false;
 
@@ -335,7 +335,7 @@ public:
 
 		LoadEditorAndEngineObjects<USoundWave>(SearchPaths, SoundWaves, &IgnoreList);
 
-		FAudioDevice* AudioDevice = GEngine->GetMainAudioDevice();
+		FAudioDevice* AudioDevice = FAudioDeviceManager::GetMainDevice();
 
 		bool bSuccess = false;
 
@@ -420,7 +420,7 @@ public:
 
 		LoadEditorAndEngineObjects<USoundWave>(SearchPaths, SoundWaves, &IgnoreList);
 
-		FAudioDevice* AudioDevice = GEngine->GetMainAudioDevice();
+		FAudioDevice* AudioDevice = FAudioDeviceManager::GetMainDevice();
 
 		bool bSuccess = false;
 		float TimeCount = 0.0f;
@@ -501,7 +501,7 @@ public:
 			SoundWave->SoundGroup = ESoundGroup::SOUNDGROUP_Music;
 		}
 
-		FAudioDevice* AudioDevice = GEngine->GetMainAudioDevice();
+		FAudioDevice* AudioDevice = FAudioDeviceManager::GetMainDevice();
 
 		bool bSuccess = false;
 		float TimeCount = 0.0f;
@@ -586,7 +586,7 @@ public:
 
 		LoadEditorAndEngineObjects<USoundWave>(SearchPaths, SoundWaves, &IgnoreList);
 
-		FAudioDevice* AudioDevice = GEngine->GetMainAudioDevice();
+		FAudioDevice* AudioDevice = FAudioDeviceManager::GetMainDevice();
 		if (!AudioDevice)
 		{
 			return false;
@@ -711,7 +711,7 @@ public:
 
 		LoadEditorAndEngineObjects<USoundAttenuation>(SearchPaths, SoundAttenuations);
 
-		FAudioDevice* AudioDevice = GEngine->GetMainAudioDevice();
+		FAudioDevice* AudioDevice = FAudioDeviceManager::GetMainDevice();
 		if (!AudioDevice)
 		{
 			return false;
diff --git a/Engine/Source/Editor/UnrealEd/Private/EditorEngine.cpp b/Engine/Source/Editor/UnrealEd/Private/EditorEngine.cpp
--- a/Engine/Source/Editor/UnrealEd/Private/EditorEngine.cpp
+++ b/Engine/Source/Editor/UnrealEd/Private/EditorEngine.cpp
@@ -94,6 +94,7 @@
 
 // needed for the RemotePropagator
 #include "AudioDevice.h"
+#include "AudioDeviceManager.h"
 #include "SurfaceIterators.h"
 #include "ScopedTransaction.h"
 
@@ -1887,7 +1888,7 @@ void UEditorEngine::Tick( float DeltaSeconds, bool bIdleMode )
 	IStreamingManager::Get().Tick(DeltaSeconds);
 
 	// Update Audio. This needs to occur after rendering as the rendering code updates the listener position.
-	if (AudioDeviceManager)
+	if (FAudioDeviceManager* AudioDeviceManager = FAudioDeviceManager::Get())
 	{
 		UWorld* OldGWorld = NULL;
 		if (PlayWorld)
@@ -2287,7 +2288,7 @@ UAudioComponent* UEditorEngine::GetPreviewAudioComponent()
 
 UAudioComponent* UEditorEngine::ResetPreviewAudioComponent( USoundBase* Sound, USoundNode* SoundNode )
 {
-	if (FAudioDevice* AudioDevice = GetMainAudioDevice())
+	if (FAudioDevice* AudioDevice = FAudioDeviceManager::GetMainDevice())
 	{
 		if (PreviewAudioComponent)
 		{
diff --git a/Engine/Source/Editor/UnrealEd/Private/EditorViewportClient.cpp b/Engine/Source/Editor/UnrealEd/Private/EditorViewportClient.cpp
--- a/Engine/Source/Editor/UnrealEd/Private/EditorViewportClient.cpp
+++ b/Engine/Source/Editor/UnrealEd/Private/EditorViewportClient.cpp
@@ -5864,12 +5864,9 @@ void FEditorViewportClient::SetEnabledStats(const TArray<FString>& InEnabledStat
 	EnabledStats = InEnabledStats;
 
 #if ENABLE_AUDIO_DEBUG
-	if (GEngine)
+	if (FAudioDeviceManager* DeviceManager = FAudioDeviceManager::Get())
 	{
-		if (FAudioDeviceManager* DeviceManager = GEngine->GetAudioDeviceManager())
-		{
-			FAudioDebugger::ResolveDesiredStats(this);
-		}
+		FAudioDebugger::ResolveDesiredStats(this);
 	}
 #endif // ENABLE_AUDIO_DEBUG
 }
diff --git a/Engine/Source/Editor/UnrealEd/Private/ObjectTools.cpp b/Engine/Source/Editor/UnrealEd/Private/ObjectTools.cpp
--- a/Engine/Source/Editor/UnrealEd/Private/ObjectTools.cpp
+++ b/Engine/Source/Editor/UnrealEd/Private/ObjectTools.cpp
@@ -2108,8 +2108,7 @@ namespace ObjectTools
 		GEditor->ClearPreviewComponents();
 
 		// Ensure the audio manager is not holding on to any sounds
-		FAudioDeviceManager* AudioDeviceManager = GEditor->GetAudioDeviceManager();
-		if (AudioDeviceManager != nullptr)
+		if (FAudioDeviceManager* AudioDeviceManager = FAudioDeviceManager::Get())
 		{
 			AudioDeviceManager->UpdateActiveAudioDevices(false);
 
diff --git a/Engine/Source/Runtime/ALAudio/Private/ALAudioBuffer.cpp b/Engine/Source/Runtime/ALAudio/Private/ALAudioBuffer.cpp
--- a/Engine/Source/Runtime/ALAudio/Private/ALAudioBuffer.cpp
+++ b/Engine/Source/Runtime/ALAudio/Private/ALAudioBuffer.cpp
@@ -48,7 +48,7 @@ FALSoundBuffer* FALSoundBuffer::Init(FALAudioDevice* AudioDevice, USoundWave* In
 		return nullptr;
 	}
 
-	FAudioDeviceManager* AudioDeviceManager = GEngine->GetAudioDeviceManager();
+	FAudioDeviceManager* AudioDeviceManager = FAudioDeviceManager::Get();
 	check(AudioDeviceManager != nullptr);
 
 	FALSoundBuffer *Buffer = nullptr;
@@ -119,7 +119,7 @@ void FALSoundBuffer::CreateNativeBuffer(FALAudioDevice* AudioDevice, USoundWave*
 		Buffer->NumChannels = Wave->NumChannels;
 		Buffer->SampleRate = Wave->GetSampleRateForCurrentPlatform();
 
-		FAudioDeviceManager* AudioDeviceManager = GEngine->GetAudioDeviceManager();
+		FAudioDeviceManager* AudioDeviceManager = FAudioDeviceManager::Get();
 		check(AudioDeviceManager != nullptr);
 		AudioDeviceManager->TrackResource(Wave, Buffer);
 	}
diff --git a/Engine/Source/Runtime/ALAudio/Private/ALAudioDevice.cpp b/Engine/Source/Runtime/ALAudio/Private/ALAudioDevice.cpp
--- a/Engine/Source/Runtime/ALAudio/Private/ALAudioDevice.cpp
+++ b/Engine/Source/Runtime/ALAudio/Private/ALAudioDevice.cpp
@@ -55,7 +55,7 @@ void FALAudioDevice::TeardownHardware( void )
 	FreeSources.Empty();
 
 	// Destroy OpenAL buffers associated with this audio device
-	FAudioDeviceManager* AudioDeviceManager = GEngine->GetAudioDeviceManager();
+	FAudioDeviceManager* AudioDeviceManager = FAudioDeviceManager::Get();
 	check(AudioDeviceManager != nullptr);
 	for( int32 i = 0; i < AudioDeviceManager->Buffers.Num(); i++ )
 	{
diff --git a/Engine/Source/Runtime/Android/AndroidAudio/Private/AndroidAudioBuffer.cpp b/Engine/Source/Runtime/Android/AndroidAudio/Private/AndroidAudioBuffer.cpp
--- a/Engine/Source/Runtime/Android/AndroidAudio/Private/AndroidAudioBuffer.cpp
+++ b/Engine/Source/Runtime/Android/AndroidAudio/Private/AndroidAudioBuffer.cpp
@@ -155,7 +155,7 @@ FSLESSoundBuffer* FSLESSoundBuffer::CreateNativeBuffer( FSLESAudioDevice* AudioD
 
 	Buffer->DecompressionState = AudioDevice->CreateCompressedAudioInfo(InWave);
 
-	FAudioDeviceManager* AudioDeviceManager = GEngine->GetAudioDeviceManager();
+	FAudioDeviceManager* AudioDeviceManager = FAudioDeviceManager::Get();
 	check(AudioDeviceManager != nullptr);
 
 	AudioDeviceManager->TrackResource(InWave, Buffer);
@@ -226,7 +226,7 @@ FSLESSoundBuffer* FSLESSoundBuffer::Init(  FSLESAudioDevice* AudioDevice ,USound
 		return( NULL );
 	}
 	
-	FAudioDeviceManager* AudioDeviceManager = GEngine->GetAudioDeviceManager();
+	FAudioDeviceManager* AudioDeviceManager = FAudioDeviceManager::Get();
 	FSLESSoundBuffer* Buffer = NULL;
 	
 	EDecompressionType DecompressionType = InWave->DecompressionType;
diff --git a/Engine/Source/Runtime/AudioMixer/Private/AudioMixerBuffer.cpp b/Engine/Source/Runtime/AudioMixer/Private/AudioMixerBuffer.cpp
--- a/Engine/Source/Runtime/AudioMixer/Private/AudioMixerBuffer.cpp
+++ b/Engine/Source/Runtime/AudioMixer/Private/AudioMixerBuffer.cpp
@@ -181,8 +181,6 @@ namespace Audio
 		InWave->InvalidateSoundWaveIfNeccessary();
 #endif // WITH_EDITOR
 
-		FAudioDeviceManager* AudioDeviceManager = FAudioDevice::GetAudioDeviceManager();
-
 		FMixerBuffer* Buffer = nullptr;
 
 		EDecompressionType DecompressionType = InWave->DecompressionType;
@@ -209,7 +207,7 @@ namespace Audio
 				// Find any existing buffers
 				if (InWave->ResourceID)
 				{
-					Buffer = (FMixerBuffer*)AudioDeviceManager->GetSoundBufferForResourceID(InWave->ResourceID);
+					Buffer = (FMixerBuffer*)FAudioDeviceManager::GetChecked().GetSoundBufferForResourceID(InWave->ResourceID);
 				}
 
 				// Override with any new PCM data even if the buffer already exists
@@ -218,14 +216,14 @@ namespace Audio
 					// If we already have a buffer for this wave resource, free it
 					if (Buffer)
 					{
-						AudioDeviceManager->FreeBufferResource(Buffer);
+						FAudioDeviceManager::GetChecked().FreeBufferResource(Buffer);
 					}
 
 					// Create a new preview buffer
 					Buffer = FMixerBuffer::CreatePreviewBuffer(InAudioDevice, InWave);
 
 					// Track the new created buffer
-					AudioDeviceManager->TrackResource(InWave, Buffer);
+					FAudioDeviceManager::GetChecked().TrackResource(InWave, Buffer);
 				}
 			}
 			break;
@@ -248,7 +246,7 @@ namespace Audio
 			{
 				if (InWave->ResourceID)
 				{
-					Buffer = (FMixerBuffer*)AudioDeviceManager->GetSoundBufferForResourceID(InWave->ResourceID);
+					Buffer = (FMixerBuffer*)FAudioDeviceManager::GetChecked().GetSoundBufferForResourceID(InWave->ResourceID);
 				}
 
 				if (Buffer == nullptr)
@@ -256,7 +254,7 @@ namespace Audio
 					Buffer = FMixerBuffer::CreateNativeBuffer(InAudioDevice, InWave);
 
 					// Track the resource with the audio device manager
-					AudioDeviceManager->TrackResource(InWave, Buffer);
+					FAudioDeviceManager::GetChecked().TrackResource(InWave, Buffer);
 					InWave->RemoveAudioResource();
 				}
 			}
diff --git a/Engine/Source/Runtime/AudioMixer/Private/AudioMixerDevice.cpp b/Engine/Source/Runtime/AudioMixer/Private/AudioMixerDevice.cpp
--- a/Engine/Source/Runtime/AudioMixer/Private/AudioMixerDevice.cpp
+++ b/Engine/Source/Runtime/AudioMixer/Private/AudioMixerDevice.cpp
@@ -1594,24 +1594,22 @@ namespace Audio
 
 	void FMixerDevice::RegisterSubmixBufferListener(ISubmixBufferListener* InSubmixBufferListener, USoundSubmix* InSubmix)
 	{
-		if (!IsInAudioThread())
+		DECLARE_CYCLE_STAT(TEXT("FAudioThreadTask.RegisterSubmixBufferListener"), STAT_RegisterSubmixBufferListener, STATGROUP_AudioThreadCommands);
+		auto AudioThreadCommand = [this, InSubmixBufferListener, InSubmix]()
 		{
-			DECLARE_CYCLE_STAT(TEXT("FAudioThreadTask.RegisterSubmixBufferListener"), STAT_RegisterSubmixBufferListener, STATGROUP_AudioThreadCommands);
-			auto AudioThreadCommand = [this, InSubmixBufferListener, InSubmix]()
-			{
-				CSV_SCOPED_TIMING_STAT(Audio, StopSpectrumAnalysis);
-				RegisterSubmixBufferListener(InSubmixBufferListener, InSubmix);
-			};
+			CSV_SCOPED_TIMING_STAT(Audio, RegisterSubmixBufferListener);
+			RegisterSubmixBufferListener(InSubmixBufferListener, InSubmix);
+		};
 
+		if (!IsInAudioThread())
+		{
 			if (IsInGameThread())
 			{
 				FAudioThread::RunCommandOnAudioThread(AudioThreadCommand, GET_STATID(STAT_RegisterSubmixBufferListener));
 			}
 			else
 			{
-				AsyncTask(ENamedThreads::GameThread, [AudioThreadCommand]() {
-					FAudioThread::RunCommandOnAudioThread(AudioThreadCommand, GET_STATID(STAT_RegisterSubmixBufferListener));
-				});
+				AsyncTask(ENamedThreads::AudioThread, TFunction<void()>(AudioThreadCommand));
 			}
 			return;
 		}
@@ -1633,25 +1631,23 @@ namespace Audio
 
 	void FMixerDevice::UnregisterSubmixBufferListener(ISubmixBufferListener* InSubmixBufferListener, USoundSubmix* InSubmix)
 	{
-		if (!IsInAudioThread())
-		{
-			DECLARE_CYCLE_STAT(TEXT("FAudioThreadTask.UnregisterSubmixBufferListener"), STAT_UnregisterSubmixBufferListener, STATGROUP_AudioThreadCommands);
+		DECLARE_CYCLE_STAT(TEXT("FAudioThreadTask.UnregisterSubmixBufferListener"), STAT_UnregisterSubmixBufferListener, STATGROUP_AudioThreadCommands);
 
-			auto AudioThreadCommand = [this, InSubmixBufferListener, InSubmix]()
-			{
-				CSV_SCOPED_TIMING_STAT(Audio, UnregisterSubmixBufferListener);
-				UnregisterSubmixBufferListener(InSubmixBufferListener, InSubmix);
-			};
+		auto AudioThreadCommand = [this, InSubmixBufferListener, InSubmix]()
+		{
+			CSV_SCOPED_TIMING_STAT(Audio, UnregisterSubmixBufferListener);
+			UnregisterSubmixBufferListener(InSubmixBufferListener, InSubmix);
+		};
 
+		if (!IsInAudioThread())
+		{
 			if (IsInGameThread())
 			{
 				FAudioThread::RunCommandOnAudioThread(AudioThreadCommand, GET_STATID(STAT_UnregisterSubmixBufferListener));
 			}
 			else
 			{
-				AsyncTask(ENamedThreads::GameThread, [AudioThreadCommand]() {
-					FAudioThread::RunCommandOnAudioThread(AudioThreadCommand, GET_STATID(STAT_UnregisterSubmixBufferListener));
-				});
+				AsyncTask(ENamedThreads::AudioThread, TFunction<void()>(AudioThreadCommand));
 			}
 
 			return;
@@ -1688,8 +1684,7 @@ namespace Audio
 
 	bool FMixerDevice::IsMainAudioDevice() const
 	{
-		bool bIsMain = (this == GEngine->GetMainAudioDevice());
-		return bIsMain;
+		return this == FAudioDeviceManager::GetMainDevice();
 	}
 
 	void FMixerDevice::WhiteNoiseTest(AlignedFloatBuffer& Output)
diff --git a/Engine/Source/Runtime/AudioMixer/Private/AudioMixerSource.cpp b/Engine/Source/Runtime/AudioMixer/Private/AudioMixerSource.cpp
--- a/Engine/Source/Runtime/AudioMixer/Private/AudioMixerSource.cpp
+++ b/Engine/Source/Runtime/AudioMixer/Private/AudioMixerSource.cpp
@@ -254,7 +254,7 @@ namespace Audio
 
 			bool bIsDebug = false;
 			FString WaveInstanceName = WaveInstance->GetName(); //-V595
-			FString TestName = GEngine->GetAudioDeviceManager()->GetDebugger().GetAudioMixerDebugSoundName();
+			FString TestName = FAudioDeviceManager::GetChecked().GetDebugger().GetAudioMixerDebugSoundName();
 			if (WaveInstanceName.Contains(TestName))
 			{
 				bDebugMode = true;
diff --git a/Engine/Source/Runtime/AudioMixer/Private/Components/SynthComponent.cpp b/Engine/Source/Runtime/AudioMixer/Private/Components/SynthComponent.cpp
--- a/Engine/Source/Runtime/AudioMixer/Private/Components/SynthComponent.cpp
+++ b/Engine/Source/Runtime/AudioMixer/Private/Components/SynthComponent.cpp
@@ -411,9 +411,7 @@ FAudioDevice* USynthComponent::GetAudioDevice()
 		// Make sure it has a proper audio device handle and retrieve it
 		if (World->AudioDeviceHandle != INDEX_NONE)
 		{
-			FAudioDeviceManager* AudioDeviceManager = GEngine->GetAudioDeviceManager();
-			check(AudioDeviceManager);
-			return AudioDeviceManager->GetAudioDevice(World->AudioDeviceHandle);
+			return FAudioDeviceManager::GetChecked().GetAudioDevice(World->AudioDeviceHandle);
 		}
 	}
 
diff --git a/Engine/Source/Runtime/Engine/Classes/Engine/Engine.h b/Engine/Source/Runtime/Engine/Classes/Engine/Engine.h
--- a/Engine/Source/Runtime/Engine/Classes/Engine/Engine.h
+++ b/Engine/Source/Runtime/Engine/Classes/Engine/Engine.h
@@ -1553,7 +1553,7 @@ public:
 
 protected:
 
-	/** Whether the engine should be playing sounds.  If false at initialization time the AudioDevice will not be created */
+	/** Whether the engine should be playing sounds.  If false at initialization time, the AudioDeviceManager will not be created */
 	uint32 bUseSound:1;
 
 private:
@@ -1766,14 +1766,6 @@ private:
 		void UpdateDynamicResolutionStatus();
 	#endif
 
-protected:
-
-	/** The audio device manager */
-	FAudioDeviceManager* AudioDeviceManager;
-
-	/** Audio device handle to the main audio device. */
-	uint32 MainAudioDeviceHandle;
-
 public:
 
 	/** A collection of messages to display on-screen. */
@@ -2436,15 +2428,19 @@ public:
 	bool IsEditor();
 
 	/** @return the audio device manager of the UEngine, this allows the creation and management of multiple audio devices. */
+	UE_DEPRECATED(4.25, "FAudioDeviceManager::Get and FAudioDeviceManager::GetChecked have replaced UEngine::GetAudioDeviceManager.")
 	FAudioDeviceManager* GetAudioDeviceManager();
 
 	/** @return the main audio device handle used by the engine. */
+	UE_DEPRECATED(4.25, "FAudioDeviceManager::GetMainDevice has replaced UEngine::GetAudioDeviceHandle.")
 	uint32 GetAudioDeviceHandle() const;
 
 	/** @return the main audio device. */
+	UE_DEPRECATED(4.25, "FAudioDeviceManager::GetMainDevice has replaced UEngine::GetMainAudioDevice.")
 	class FAudioDevice* GetMainAudioDevice();
 
 	/** @return the currently active audio device */
+	UE_DEPRECATED(4.25, "FAudioDeviceManager::GetActiveDevice has replaced UEngine::GetActiveAudioDevice.")
 	class FAudioDevice* GetActiveAudioDevice();
 
 	/** @return whether we're currently running in split screen (more than one local player) */
diff --git a/Engine/Source/Runtime/Engine/Classes/Sound/SoundBase.h b/Engine/Source/Runtime/Engine/Classes/Sound/SoundBase.h
--- a/Engine/Source/Runtime/Engine/Classes/Sound/SoundBase.h
+++ b/Engine/Source/Runtime/Engine/Classes/Sound/SoundBase.h
@@ -5,7 +5,7 @@
  * The base class for a playable sound object
  */
 
-#include "Audio.h"
+#include "AudioDefines.h"
 #include "CoreMinimal.h"
 #include "IAudioExtensionPlugin.h"
 #include "SoundConcurrency.h"
@@ -13,17 +13,22 @@
 #include "SoundSubmixSend.h"
 #include "UObject/Object.h"
 #include "UObject/ObjectMacros.h"
+#include "Audio.h"
 
 #include "SoundBase.generated.h"
 
 
+class USoundAttenuation;
+class USoundClass;
 class USoundEffectSourcePreset;
 class USoundSourceBus;
 class USoundSubmix;
 class USoundEffectSourcePresetChain;
 
 struct FActiveSound;
+struct FSoundAttenuationSettings;
 struct FSoundParseParameters;
+struct FWaveInstance;
 
 /**
  * Method of virtualization when a sound is stopped due to playback constraints
diff --git a/Engine/Source/Runtime/Engine/Classes/Sound/SoundEffectSubmix.h b/Engine/Source/Runtime/Engine/Classes/Sound/SoundEffectSubmix.h
--- a/Engine/Source/Runtime/Engine/Classes/Sound/SoundEffectSubmix.h
+++ b/Engine/Source/Runtime/Engine/Classes/Sound/SoundEffectSubmix.h
@@ -15,6 +15,8 @@ class FAudioDevice;
 class FSoundEffectSubmix;
 
 struct FAudioEffectParameters;
+class FAudioDevice;
+
 
 /** Preset of a submix effect that can be shared between sounds. */
 UCLASS(config = Engine, hidecategories = Object, abstract, editinlinenew, BlueprintType)
@@ -29,12 +31,12 @@ class ENGINE_API USoundEffectSubmixPreset : public USoundEffectPreset
 /** Struct which has data needed to initialize the submix effect. */
 struct FSoundEffectSubmixInitData
 {
-	uint32 DeviceHandle;
+	Audio::FDeviceId DeviceHandle;
 	void* PresetSettings;
 	float SampleRate;
 
 	FSoundEffectSubmixInitData()
-		: DeviceHandle(INDEX_NONE)
+		: DeviceHandle(static_cast<Audio::FDeviceId>(INDEX_NONE))
 		, PresetSettings(nullptr)
 		, SampleRate(0)
 	{
diff --git a/Engine/Source/Runtime/Engine/Classes/Sound/SoundSubmixSend.h b/Engine/Source/Runtime/Engine/Classes/Sound/SoundSubmixSend.h
--- a/Engine/Source/Runtime/Engine/Classes/Sound/SoundSubmixSend.h
+++ b/Engine/Source/Runtime/Engine/Classes/Sound/SoundSubmixSend.h
@@ -39,6 +39,22 @@ enum class ESendLevelControlMethod : uint8
 	Manual,
 };
 
+/** Abstract interface for receiving audio data from a given submix. */
+class ENGINE_API ISubmixBufferListener
+{
+public:
+	/**
+	Called when a new buffer has been rendered for a given submix
+	@param OwningSubmix	The submix object which has rendered a new buffer
+	@param AudioData		Ptr to the audio buffer
+	@param NumSamples		The number of audio samples in the audio buffer
+	@param NumChannels		The number of channels of audio in the buffer (e.g. 2 for stereo, 6 for 5.1, etc)
+	@param SampleRate		The sample rate of the audio buffer
+	@param AudioClock		Double audio clock value, from start of audio rendering.
+	*/
+	virtual void OnNewSubmixBuffer(const USoundSubmix* OwningSubmix, float* AudioData, int32 NumSamples, int32 NumChannels, const int32 SampleRate, double AudioClock) = 0;
+};
+
 // Class used to send audio to submixes from USoundBase
 USTRUCT(BlueprintType)
 struct ENGINE_API FSoundSubmixSendInfo
diff --git a/Engine/Source/Runtime/Engine/Private/ActiveSound.cpp b/Engine/Source/Runtime/Engine/Private/ActiveSound.cpp
--- a/Engine/Source/Runtime/Engine/Private/ActiveSound.cpp
+++ b/Engine/Source/Runtime/Engine/Private/ActiveSound.cpp
@@ -5,6 +5,7 @@
 #include "EngineDefines.h"
 #include "Misc/App.h"
 #include "AudioThread.h"
+#include "AudioDefines.h"
 #include "AudioDevice.h"
 #include "IAudioExtensionPlugin.h"
 #include "Sound/AudioSettings.h"
@@ -951,9 +952,9 @@ void FActiveSound::OcclusionTraceDone(const FTraceHandle& TraceHandle, FTraceDat
 	FAsyncTraceDetails TraceDetails;
 	if (TraceToActiveSoundMap.RemoveAndCopyValue(TraceHandle, TraceDetails))
 	{
-		if (FAudioDeviceManager* AudioDeviceManager = GEngine->GetAudioDeviceManager())
+		if (FAudioDeviceManager* DeviceManager = FAudioDeviceManager::Get())
 		{
-			if (FAudioDevice* AudioDevice = AudioDeviceManager->GetAudioDevice(TraceDetails.AudioDeviceID))
+			if (FAudioDevice* AudioDevice = DeviceManager->GetAudioDevice(TraceDetails.AudioDeviceID))
 			{
 				FActiveSound* ActiveSound = TraceDetails.ActiveSound;
 
diff --git a/Engine/Source/Runtime/Engine/Private/Audio.cpp b/Engine/Source/Runtime/Engine/Private/Audio.cpp
--- a/Engine/Source/Runtime/Engine/Private/Audio.cpp
+++ b/Engine/Source/Runtime/Engine/Private/Audio.cpp
@@ -178,9 +178,13 @@ UClass* GetAudioPluginCustomSettingsClass(EAudioPlugin PluginType)
 FSoundBuffer::~FSoundBuffer()
 {
 	// remove ourselves from the set of waves that are tracked by the audio device
-	if (ResourceID && GEngine && GEngine->GetAudioDeviceManager())
+	if (!ResourceID)
 	{
-		GEngine->GetAudioDeviceManager()->RemoveSoundBufferForResourceID(ResourceID);
+		return;
+	}
+	if (FAudioDeviceManager* AudioDeviceManager = FAudioDeviceManager::Get())
+	{
+		AudioDeviceManager->RemoveSoundBufferForResourceID(ResourceID);
 	}
 }
 
@@ -467,13 +471,14 @@ float FSoundSource::GetDebugVolume(const float InVolume)
 #if ENABLE_AUDIO_DEBUG
 
 	// Bail if we don't have a device manager.
-	if (!GEngine || !GEngine->GetAudioDeviceManager() || !WaveInstance || !DebugInfo.IsValid() )
+	FAudioDeviceManager* DeviceManager = FAudioDeviceManager::Get();
+	if (!DeviceManager || !WaveInstance || !DebugInfo.IsValid() )
 	{
 		return OutVolume;
 	}
 
 	// Solos/Mutes (dev only).
-	FAudioDebugger& Debugger = GEngine->GetAudioDeviceManager()->GetDebugger();	
+	FAudioDebugger& Debugger = DeviceManager->GetDebugger();
 	FDebugInfo Info;
 				
 	// SoundWave Solo/Mutes.
@@ -987,27 +992,24 @@ float FWaveInstance::GetDynamicVolume() const
 {
 	float OutVolume = 1.0f;
 
-	if (GEngine)
+	if (FAudioDeviceManager* DeviceManager = FAudioDeviceManager::Get())
 	{
-		if (FAudioDeviceManager* DeviceManager = GEngine->GetAudioDeviceManager())
+		if (WaveData)
 		{
-			if (WaveData)
-			{
-				OutVolume *= DeviceManager->GetDynamicSoundVolume(ESoundType::Wave, WaveData->GetFName());
-			}
+			OutVolume *= DeviceManager->GetDynamicSoundVolume(ESoundType::Wave, WaveData->GetFName());
+		}
 
-			if (ActiveSound)
+		if (ActiveSound)
+		{
+			if (const USoundCue* Sound = Cast<USoundCue>(ActiveSound->GetSound()))
 			{
-				if (const USoundCue* Sound = Cast<USoundCue>(ActiveSound->GetSound()))
-				{
-					OutVolume *= DeviceManager->GetDynamicSoundVolume(ESoundType::Cue, Sound->GetFName());
-				}
+				OutVolume *= DeviceManager->GetDynamicSoundVolume(ESoundType::Cue, Sound->GetFName());
 			}
+		}
 
-			if (SoundClass)
-			{
-				OutVolume *= DeviceManager->GetDynamicSoundVolume(ESoundType::Class, SoundClass->GetFName());
-			}
+		if (SoundClass)
+		{
+			OutVolume *= DeviceManager->GetDynamicSoundVolume(ESoundType::Class, SoundClass->GetFName());
 		}
 	}
 
diff --git a/Engine/Source/Runtime/Engine/Private/AudioDebug.cpp b/Engine/Source/Runtime/Engine/Private/AudioDebug.cpp
--- a/Engine/Source/Runtime/Engine/Private/AudioDebug.cpp
+++ b/Engine/Source/Runtime/Engine/Private/AudioDebug.cpp
@@ -327,17 +327,17 @@ namespace
 
 	void HandleDumpActiveSounds(UWorld* World)
 	{
-		if (GEngine && GEngine->GetAudioDeviceManager())
+		if (GEngine && FAudioDeviceManager::Get())
 		{
-			GEngine->GetAudioDeviceManager()->GetDebugger().DumpActiveSounds();
+			FAudioDeviceManager::Get()->GetDebugger().DumpActiveSounds();
 		}
 	}
 	
 	void HandleClearMutesAndSolos(UWorld* World)
 	{
-		if (GEngine && GEngine->GetAudioDeviceManager())
+		if (GEngine && FAudioDeviceManager::Get())
 		{
-			GEngine->GetAudioDeviceManager()->GetDebugger().ClearMutesAndSolos();
+			FAudioDeviceManager::Get()->GetDebugger().ClearMutesAndSolos();
 		}
 	}
 
@@ -368,7 +368,7 @@ void FAudioDebugger::ToggleVisualizeDebug3dEnabled()
 #if WITH_EDITOR
 void FAudioDebugger::OnBeginPIE()
 {
-	FAudioDeviceManager* DeviceManager = GEngine->GetAudioDeviceManager();
+	FAudioDeviceManager* DeviceManager = FAudioDeviceManager::Get();
 	if (DeviceManager)
 	{
 		DeviceManager->GetDebugger().ClearMutesAndSolos();
@@ -377,7 +377,7 @@ void FAudioDebugger::OnBeginPIE()
 
 void FAudioDebugger::OnEndPIE()
 {
-	FAudioDeviceManager* DeviceManager = GEngine->GetAudioDeviceManager();
+	FAudioDeviceManager* DeviceManager = FAudioDeviceManager::Get();
 	if (DeviceManager)
 	{
 		DeviceManager->GetDebugger().ClearMutesAndSolos();
@@ -544,7 +544,7 @@ void FAudioDebugger::DrawDebugInfo(const FSoundSource& SoundSource)
 		return;
 	}
 
-	FAudioDeviceManager* DeviceManager = GEngine->GetAudioDeviceManager();
+	FAudioDeviceManager* DeviceManager = FAudioDeviceManager::Get();
 	if (DeviceManager && DeviceManager->IsVisualizeDebug3dEnabled())
 	{
 		DECLARE_CYCLE_STAT(TEXT("FAudioThreadTask.DrawSourceDebugInfo"), STAT_AudioDrawSourceDebugInfo, STATGROUP_TaskGraphTasks);
@@ -611,7 +611,7 @@ void FAudioDebugger::DrawDebugInfo(const FActiveSound& ActiveSound, const TArray
 		}
 	}
 
-	FAudioDeviceManager* DeviceManager = GEngine->GetAudioDeviceManager();
+	FAudioDeviceManager* DeviceManager = FAudioDeviceManager::Get();
 	if (DeviceManager && DeviceManager->IsVisualizeDebug3dEnabled())
 	{
 		DECLARE_CYCLE_STAT(TEXT("FAudioThreadTask.DrawActiveSoundDebugInfo"), STAT_AudioDrawActiveSoundDebugInfo, STATGROUP_TaskGraphTasks);
@@ -688,15 +688,14 @@ void FAudioDebugger::DrawDebugInfo(const FActiveSound& ActiveSound, const TArray
 void FAudioDebugger::DrawDebugInfo(UWorld& World, const TArray<FListener>& Listeners)
 {
 #if ENABLE_DRAW_DEBUG
-	if (!ActiveSoundVisualizeListenersCVar)
+	if (!GEngine || !ActiveSoundVisualizeListenersCVar)
 	{
 		return;
 	}
 
 	check(IsInAudioThread());
 
-	FAudioDeviceManager* DeviceManager = GEngine->GetAudioDeviceManager();
-	if (DeviceManager && DeviceManager->IsVisualizeDebug3dEnabled())
+	if (FAudioDeviceManager::GetChecked().IsVisualizeDebug3dEnabled())
 	{
 		DECLARE_CYCLE_STAT(TEXT("FAudioThreadTask.DrawListenerDebugInfo"), STAT_AudioDrawListenerDebugInfo, STATGROUP_TaskGraphTasks);
 		const TWeakObjectPtr<UWorld> WorldPtr = &World;
@@ -742,8 +741,7 @@ void FAudioDebugger::DrawDebugInfo(const FAudioVirtualLoop& VirtualLoop)
 		return;
 	}
 
-	FAudioDeviceManager* DeviceManager = GEngine->GetAudioDeviceManager();
-	if (DeviceManager && DeviceManager->IsVisualizeDebug3dEnabled())
+	if (FAudioDeviceManager::GetChecked().IsVisualizeDebug3dEnabled())
 	{
 		DECLARE_CYCLE_STAT(TEXT("FAudioThreadTask.DrawVirtualLoopDebugInfo"), STAT_AudioDrawVirtualLoopDebugInfo, STATGROUP_TaskGraphTasks);
 
@@ -788,7 +786,7 @@ void FAudioDebugger::DumpActiveSounds() const
 		return;
 	}
 
-	FAudioDevice* AudioDevice = GEngine->GetAudioDeviceManager()->GetActiveAudioDevice();
+	FAudioDevice* AudioDevice = FAudioDeviceManager::GetActiveDevice();
 	if (!AudioDevice)
 	{
 		return;
@@ -1568,7 +1566,7 @@ bool FAudioDebugger::ToggleStats(UWorld* World, const uint8 StatToToggle)
 		return false;
 	}
 
-	if (FAudioDeviceManager* DeviceManager = GEngine->GetAudioDeviceManager())
+	if (FAudioDeviceManager* DeviceManager = FAudioDeviceManager::Get())
 	{
 		DeviceManager->GetDebugger().ToggleStats(AudioDevice->DeviceHandle, StatToToggle);
 	}
diff --git a/Engine/Source/Runtime/Engine/Private/AudioDecompress.cpp b/Engine/Source/Runtime/Engine/Private/AudioDecompress.cpp
--- a/Engine/Source/Runtime/Engine/Private/AudioDecompress.cpp
+++ b/Engine/Source/Runtime/Engine/Private/AudioDecompress.cpp
@@ -645,9 +645,9 @@ FAsyncAudioDecompressWorker::FAsyncAudioDecompressWorker(USoundWave* InWave, int
 	, NumPrecacheFrames(InPrecacheBufferNumFrames)
 {
 	check(NumPrecacheFrames > 0);
-	if (GEngine && GEngine->GetMainAudioDevice())
+	if (FAudioDevice* AudioDevice = FAudioDeviceManager::GetMainDevice())
 	{
-		AudioInfo = GEngine->GetMainAudioDevice()->CreateCompressedAudioInfo(Wave);
+		AudioInfo = AudioDevice->CreateCompressedAudioInfo(Wave);
 	}
 }
 
diff --git a/Engine/Source/Runtime/Engine/Private/AudioDevice.cpp b/Engine/Source/Runtime/Engine/Private/AudioDevice.cpp
--- a/Engine/Source/Runtime/Engine/Private/AudioDevice.cpp
+++ b/Engine/Source/Runtime/Engine/Private/AudioDevice.cpp
@@ -724,7 +724,7 @@ void FAudioDevice::CountBytes(FArchive& Ar)
 	Sources.CountBytes(Ar);
 	// The buffers are stored on the audio device since they are shared amongst all audio devices
 	// Though we are going to count them when querying an individual audio device object about its bytes
-	GEngine->GetAudioDeviceManager()->Buffers.CountBytes(Ar);
+	FAudioDeviceManager::GetChecked().Buffers.CountBytes(Ar);
 	FreeSources.CountBytes(Ar);
 	WaveInstanceSourceMap.CountBytes(Ar);
 	Ar.CountBytes(sizeof(FWaveInstance) * WaveInstanceSourceMap.Num(), sizeof(FWaveInstance) * WaveInstanceSourceMap.Num());
@@ -1327,7 +1327,7 @@ bool FAudioDevice::HandleEnableHRTFForAllCommand(const TCHAR* Cmd, FOutputDevice
 bool FAudioDevice::HandleSoloCommand(const TCHAR* Cmd, FOutputDevice& Ar)
 {
 	// Apply the solo to the given device
-	FAudioDeviceManager* DeviceManager = GEngine->GetAudioDeviceManager();
+	FAudioDeviceManager* DeviceManager = FAudioDeviceManager::Get();
 	if (DeviceManager)
 	{
 		DeviceManager->SetSoloDevice(DeviceHandle);
@@ -1337,7 +1337,7 @@ bool FAudioDevice::HandleSoloCommand(const TCHAR* Cmd, FOutputDevice& Ar)
 
 bool FAudioDevice::HandleClearSoloCommand(const TCHAR* Cmd, FOutputDevice& Ar)
 {
-	FAudioDeviceManager* DeviceManager = GEngine->GetAudioDeviceManager();
+	FAudioDeviceManager* DeviceManager = FAudioDeviceManager::Get();
 	if (DeviceManager)
 	{
 		DeviceManager->SetSoloDevice(INDEX_NONE);
@@ -1347,7 +1347,7 @@ bool FAudioDevice::HandleClearSoloCommand(const TCHAR* Cmd, FOutputDevice& Ar)
 
 bool FAudioDevice::HandlePlayAllPIEAudioCommand(const TCHAR* Cmd, FOutputDevice& Ar)
 {
-	FAudioDeviceManager* DeviceManager = GEngine->GetAudioDeviceManager();
+	FAudioDeviceManager* DeviceManager = FAudioDeviceManager::Get();
 	if (DeviceManager)
 	{
 		DeviceManager->TogglePlayAllDeviceAudio();
@@ -1357,7 +1357,7 @@ bool FAudioDevice::HandlePlayAllPIEAudioCommand(const TCHAR* Cmd, FOutputDevice&
 
 bool FAudioDevice::HandleAudio3dVisualizeCommand(const TCHAR* Cmd, FOutputDevice& Ar)
 {
-	FAudioDeviceManager* DeviceManager = GEngine->GetAudioDeviceManager();
+	FAudioDeviceManager* DeviceManager = FAudioDeviceManager::Get();
 	if (DeviceManager)
 	{
 		DeviceManager->ToggleVisualize3dDebug();
@@ -1367,7 +1367,7 @@ bool FAudioDevice::HandleAudio3dVisualizeCommand(const TCHAR* Cmd, FOutputDevice
 
 void FAudioDevice::HandleAudioSoloCommon(const TCHAR* Cmd, FOutputDevice& Ar, FToggleSoloPtr FPtr)
 {
-	FAudioDeviceManager* DeviceManager = GEngine->GetAudioDeviceManager();
+	FAudioDeviceManager* DeviceManager = FAudioDeviceManager::Get();
 	if (DeviceManager)
 	{
 		bool bExclusive = true;
@@ -1404,7 +1404,7 @@ bool FAudioDevice::HandleAudioSoloSoundCue(const TCHAR* Cmd, FOutputDevice& Ar)
 
 bool FAudioDevice::HandleAudioMixerDebugSound(const TCHAR* Cmd, FOutputDevice& Ar)
 {
-	FAudioDeviceManager* DeviceManager = GEngine->GetAudioDeviceManager();
+	FAudioDeviceManager* DeviceManager = FAudioDeviceManager::Get();
 	if (DeviceManager)
 	{
 		DeviceManager->GetDebugger().SetAudioMixerDebugSound(Cmd);
@@ -1414,7 +1414,7 @@ bool FAudioDevice::HandleAudioMixerDebugSound(const TCHAR* Cmd, FOutputDevice& A
 
 bool FAudioDevice::HandleAudioDebugSound(const TCHAR* Cmd, FOutputDevice& Ar)
 {
-	FAudioDeviceManager* DeviceManager = GEngine->GetAudioDeviceManager();
+	FAudioDeviceManager* DeviceManager = FAudioDeviceManager::Get();
 	if (DeviceManager)
 	{
 		DeviceManager->GetDebugger().SetAudioDebugSound(Cmd);
@@ -1726,7 +1726,7 @@ bool FAudioDevice::HandleAudioMemoryInfo(const TCHAR* Cmd, FOutputDevice& Ar)
 
 bool FAudioDevice::HandleResetAllDynamicSoundVolumesCommand(const TCHAR* Cmd, FOutputDevice& Ar)
 {
-	if (FAudioDeviceManager* DeviceManager = GEngine->GetAudioDeviceManager())
+	if (FAudioDeviceManager* DeviceManager = FAudioDeviceManager::Get())
 	{
 		DeviceManager->ResetAllDynamicSoundVolumes();
 	}
@@ -1735,7 +1735,7 @@ bool FAudioDevice::HandleResetAllDynamicSoundVolumesCommand(const TCHAR* Cmd, FO
 
 bool FAudioDevice::HandleResetDynamicSoundVolumeCommand(const TCHAR* Cmd, FOutputDevice& Ar)
 {
-	if (FAudioDeviceManager* DeviceManager = GEngine->GetAudioDeviceManager())
+	if (FAudioDeviceManager* DeviceManager = FAudioDeviceManager::Get())
 	{
 		FName SoundName;
 		if (!FParse::Value(Cmd, TEXT("Name="), SoundName))
@@ -1770,61 +1770,55 @@ bool FAudioDevice::HandleGetDynamicSoundVolumeCommand(const TCHAR* Cmd, FOutputD
 		return false;
 	}
 
-	if (const FAudioDeviceManager* DeviceManager = GEngine->GetAudioDeviceManager())
+	FName SoundName;
+	if (!FParse::Value(Cmd, TEXT("Name="), SoundName))
 	{
-		FName SoundName;
-		if (!FParse::Value(Cmd, TEXT("Name="), SoundName))
+		return false;
+	}
+
+	// Optional: Defaults to Cue
+	FString SoundTypeStr;
+	ESoundType SoundType = ESoundType::Cue;
+	if (FParse::Value(Cmd, TEXT("Type="), SoundTypeStr))
+	{
+		if (SoundTypeStr == TEXT("Wave"))
 		{
-			return false;
+			SoundType = ESoundType::Wave;
 		}
-
-		// Optional: Defaults to Cue
-		FString SoundTypeStr;
-		ESoundType SoundType = ESoundType::Cue;
-		if (FParse::Value(Cmd, TEXT("Type="), SoundTypeStr))
+		else if (SoundTypeStr == TEXT("Class"))
 		{
-			if (SoundTypeStr == TEXT("Wave"))
-			{
-				SoundType = ESoundType::Wave;
-			}
-			else if (SoundTypeStr == TEXT("Class"))
-			{
-				SoundType = ESoundType::Class;
-			}
+			SoundType = ESoundType::Class;
 		}
+	}
 
-		if (!IsInAudioThread())
-		{
-			DECLARE_CYCLE_STAT(TEXT("FAudioThreadTask.GetDynamicSoundVolume"), STAT_AudioGetDynamicSoundVolume, STATGROUP_AudioThreadCommands);
+	if (!IsInAudioThread())
+	{
+		DECLARE_CYCLE_STAT(TEXT("FAudioThreadTask.GetDynamicSoundVolume"), STAT_AudioGetDynamicSoundVolume, STATGROUP_AudioThreadCommands);
 
-			const ESoundType InSoundType = SoundType;
-			const FName InSoundName = SoundName;
-			FAudioThread::RunCommandOnAudioThread([InSoundType, InSoundName]()
-			{
-				if (!GEngine)
-				{
-					return;
-				}
-				if (const FAudioDeviceManager* InDeviceManager = GEngine->GetAudioDeviceManager())
-				{
-					const float Volume = InDeviceManager->GetDynamicSoundVolume(InSoundType, InSoundName);
-					UE_LOG(LogAudio, Display, TEXT("'%s' Dynamic Volume: %.4f"), *InSoundName.GetPlainNameString(), Volume);
-				}
-			}, GET_STATID(STAT_AudioGetDynamicSoundVolume));
-		}
-		else
+		const ESoundType InSoundType = SoundType;
+		const FName InSoundName = SoundName;
+		FAudioThread::RunCommandOnAudioThread([InSoundType, InSoundName]()
 		{
-			const float Volume = DeviceManager->GetDynamicSoundVolume(SoundType, SoundName);
-			FString Msg = FString::Printf(TEXT("'%s' Dynamic Volume: %.4f"), *SoundName.GetPlainNameString(), Volume);
-			Ar.Logf(TEXT("%s"), *Msg);
-		}
+			if (FAudioDeviceManager* DeviceManager = FAudioDeviceManager::Get())
+			{
+				const float Volume = DeviceManager->GetDynamicSoundVolume(InSoundType, InSoundName);
+				UE_LOG(LogAudio, Display, TEXT("'%s' Dynamic Volume: %.4f"), *InSoundName.GetPlainNameString(), Volume);
+			}
+		}, GET_STATID(STAT_AudioGetDynamicSoundVolume));
+	}
+	else if (const FAudioDeviceManager* DeviceManager = FAudioDeviceManager::Get())
+	{
+		const float Volume = DeviceManager->GetDynamicSoundVolume(SoundType, SoundName);
+		FString Msg = FString::Printf(TEXT("'%s' Dynamic Volume: %.4f"), *SoundName.GetPlainNameString(), Volume);
+		Ar.Logf(TEXT("%s"), *Msg);
 	}
+
 	return true;
 }
 
 bool FAudioDevice::HandleSetDynamicSoundCommand(const TCHAR* Cmd, FOutputDevice& Ar)
 {
-	if (FAudioDeviceManager* DeviceManager = GEngine->GetAudioDeviceManager())
+	if (FAudioDeviceManager* DeviceManager = FAudioDeviceManager::Get())
 	{
 		FName SoundName;
 		if (!FParse::Value(Cmd, TEXT("Name="), SoundName))
@@ -4355,7 +4349,7 @@ void FAudioDevice::SendUpdateResultsToGameThread(const int32 FirstActiveIndex)
 		// At shutdown, GEngine may already be null
 		if (GEngine)
 		{
-			if (FAudioDeviceManager* AudioDeviceManager = GEngine->GetAudioDeviceManager())
+			if (FAudioDeviceManager* AudioDeviceManager = FAudioDeviceManager::Get())
 			{
 				if (FAudioDevice* AudioDevice = AudioDeviceManager->GetAudioDevice(AudioDeviceID))
 				{
@@ -4496,7 +4490,7 @@ void FAudioDevice::AddNewActiveSoundInternal(const FActiveSound& NewActiveSound,
 	// Cull one-shot active sounds if we've reached our max limit of one shot active sounds before we attempt to evaluate concurrency
 	// Check for debug sound name
 #if !UE_BUILD_SHIPPING
-	FAudioDeviceManager* AudioDeviceManager = GEngine->GetAudioDeviceManager();
+	FAudioDeviceManager* AudioDeviceManager = FAudioDeviceManager::Get();
 	FString DebugSound;
 	if (AudioDeviceManager->GetDebugger().GetAudioDebugSound(DebugSound))
 	{
@@ -5342,8 +5336,8 @@ float FAudioDevice::GetFocusFactor(const float Azimuth, const FSoundAttenuationS
 FAudioDevice::FCreateComponentParams::FCreateComponentParams()
 	: World(nullptr)
 	, Actor(nullptr)
+	, AudioDevice(FAudioDeviceManager::GetMainDevice())
 {
-	AudioDevice = (GEngine ? GEngine->GetMainAudioDevice() : nullptr);
 	CommonInit();
 }
 
@@ -5418,7 +5412,7 @@ UAudioComponent* FAudioDevice::CreateComponent(USoundBase* Sound, UWorld* World,
 	}
 	else
 	{
-		Params = MakeUnique<FCreateComponentParams>(GEngine->GetMainAudioDevice());
+		Params = MakeUnique<FCreateComponentParams>(FAudioDeviceManager::GetMainDevice());
 	}
 
 	Params->bPlay = bPlay;
@@ -6029,7 +6023,7 @@ bool FAudioDevice::HandleListSoundsCommand(const TCHAR* Cmd, FOutputDevice& Ar)
 	Ar.Logf(TEXT("Listing all sounds:"));
 
 	// Get audio device manager since thats where sound buffers are stored
-	FAudioDeviceManager* AudioDeviceManager = GEngine->GetAudioDeviceManager();
+	FAudioDeviceManager* AudioDeviceManager = FAudioDeviceManager::Get();
 	check(AudioDeviceManager != nullptr);
 
 	TArray<FSoundBuffer*> AllSounds;
@@ -6143,7 +6137,7 @@ bool FAudioDevice::IsAudioDeviceMuted() const
 	check(IsInAudioThread());
 
 	// First check to see if the device manager has "bPlayAllPIEAudio" enabled
-	FAudioDeviceManager* DeviceManager = GEngine->GetAudioDeviceManager();
+	FAudioDeviceManager* DeviceManager = FAudioDeviceManager::Get();
 	if (DeviceManager && DeviceManager->IsPlayAllDeviceAudio())
 	{
 		return false;
diff --git a/Engine/Source/Runtime/Engine/Private/AudioDeviceManager.cpp b/Engine/Source/Runtime/Engine/Private/AudioDeviceManager.cpp
--- a/Engine/Source/Runtime/Engine/Private/AudioDeviceManager.cpp
+++ b/Engine/Source/Runtime/Engine/Private/AudioDeviceManager.cpp
@@ -18,6 +18,11 @@
 #include "Settings/LevelEditorMiscSettings.h"
 #endif
 
+namespace
+{
+	FAudioDeviceManager* AudioDeviceManager = nullptr;
+} // namespace <>
+
 // Private consts for helping with index/generation determination in audio device manager
 static const uint32 AUDIO_DEVICE_HANDLE_INDEX_BITS		= 24;
 static const uint32 AUDIO_DEVICE_HANDLE_INDEX_MASK		= (1 << AUDIO_DEVICE_HANDLE_INDEX_BITS) - 1;
@@ -77,6 +82,7 @@ FAudioDeviceManager::FAudioDeviceManager()
 	, NextResourceID(1)
 	, SoloDeviceHandle(INDEX_NONE)
 	, ActiveAudioDeviceHandle(INDEX_NONE)
+	, MainAudioDeviceHandle(INDEX_NONE)
 	, bUsingAudioMixer(false)
 	, bPlayAllDeviceAudio(false)
 	, bOnlyToggleAudioMixerOnce(false)
@@ -260,29 +266,61 @@ IAudioDeviceModule* FAudioDeviceManager::GetAudioDeviceModule()
 	return AudioDeviceModule;
 }
 
-bool FAudioDeviceManager::Initialize()
+FAudioDevice* FAudioDeviceManager::GetActiveDevice()
 {
-	if (LoadDefaultAudioDeviceModule())
+	if (AudioDeviceManager)
 	{
-		check(AudioDeviceModule);
+		if (AudioDeviceManager->ActiveAudioDeviceHandle != INDEX_NONE)
+		{
+			return AudioDeviceManager->GetAudioDevice(AudioDeviceManager->ActiveAudioDeviceHandle);
+		}
 
-		const bool bIsAudioMixerEnabled = AudioDeviceModule->IsAudioMixerModule();
-		GetMutableDefault<UAudioSettings>()->SetAudioMixerEnabled(bIsAudioMixerEnabled);
+		return GetMainDevice();
+	}
 
-#if WITH_EDITOR
-		if (bIsAudioMixerEnabled)
+	return nullptr;
+}
+
+FAudioDevice* FAudioDeviceManager::GetMainDevice()
+{
+	if (FAudioDeviceManager* DeviceManager = Get())
+	{
+		return DeviceManager->GetAudioDevice(DeviceManager->GetMainAudioDeviceHandle());
+	}
+
+	return nullptr;
+}
+
+void FAudioDeviceManager::Initialize()
+{
+	if (!AudioDeviceManager)
+	{
+		AudioDeviceManager = new FAudioDeviceManager();
+		if (AudioDeviceManager->LoadDefaultAudioDeviceModule())
 		{
-			IAudioEditorModule* AudioEditorModule = &FModuleManager::LoadModuleChecked<IAudioEditorModule>("AudioEditor");
-			AudioEditorModule->RegisterAudioMixerAssetActions();
-			AudioEditorModule->RegisterEffectPresetAssetActions();
-		}
-#endif
+			check(AudioDeviceManager->AudioDeviceModule);
+
+			const bool bIsAudioMixerEnabled = AudioDeviceManager->AudioDeviceModule->IsAudioMixerModule();
+			GetMutableDefault<UAudioSettings>()->SetAudioMixerEnabled(bIsAudioMixerEnabled);
 
-		return CreateMainAudioDevice();
+#if WITH_EDITOR
+			if (bIsAudioMixerEnabled)
+			{
+				IAudioEditorModule* AudioEditorModule = &FModuleManager::LoadModuleChecked<IAudioEditorModule>("AudioEditor");
+				AudioEditorModule->RegisterAudioMixerAssetActions();
+				AudioEditorModule->RegisterEffectPresetAssetActions();
+			}
+#endif // WITH_EDITOR
+
+			if (AudioDeviceManager->CreateMainAudioDevice())
+			{
+				return;
+			}
+		}
 	}
 
 	// Failed to initialize
-	return false;
+	Shutdown();
 }
 
 bool FAudioDeviceManager::LoadDefaultAudioDeviceModule()
@@ -385,11 +423,10 @@ bool FAudioDeviceManager::CreateAudioDevice(bool bCreateNewDevice, FCreateAudioD
 	{
 		if (NumActiveAudioDevices == 1)
 		{
-			FAudioDevice* MainAudioDevice = GEngine->GetMainAudioDevice();
-			if (MainAudioDevice)
+			if (FAudioDevice* MainDevice = FAudioDeviceManager::GetMainDevice())
 			{
-				OutResults.Handle = MainAudioDevice->DeviceHandle;
-				OutResults.AudioDevice = MainAudioDevice;
+				OutResults.Handle = MainDevice->DeviceHandle;
+				OutResults.AudioDevice = MainDevice;
 				OutResults.AudioDevice->FadeIn();
 				return true;
 			}
@@ -408,12 +445,11 @@ bool FAudioDeviceManager::CreateAudioDevice(bool bCreateNewDevice, FCreateAudioD
 
 	if (NumActiveAudioDevices == 1 && !bCreateNewAudioDeviceForPlayInEditor)
 	{
-		FAudioDevice* MainAudioDevice = GEngine->GetMainAudioDevice();
-		if (MainAudioDevice)
+		if (FAudioDevice* MainDevice = FAudioDeviceManager::GetMainDevice())
 		{
 			++NumWorldsUsingMainAudioDevice;
-			OutResults.Handle = MainAudioDevice->DeviceHandle;
-			OutResults.AudioDevice = MainAudioDevice;
+			OutResults.Handle = MainDevice->DeviceHandle;
+			OutResults.AudioDevice = MainDevice;
 			bRequiresInit = false;
 		}
 		else
@@ -475,8 +511,7 @@ bool FAudioDeviceManager::CreateAudioDevice(bool bCreateNewDevice, FCreateAudioD
 		else
 		{
 			++NumWorldsUsingMainAudioDevice;
-			FAudioDevice* MainAudioDevice = GEngine->GetMainAudioDevice();
-			if (MainAudioDevice)
+			if (FAudioDevice* MainAudioDevice = FAudioDeviceManager::GetMainDevice())
 			{
 				OutResults.Handle = MainAudioDevice->DeviceHandle;
 				OutResults.AudioDevice = MainAudioDevice;
@@ -544,7 +579,7 @@ bool FAudioDeviceManager::ShutdownAudioDevice(Audio::FDeviceId Handle)
 	// If there are more than 1 device active, check to see if this handle is the main audio device handle
 	if (NumActiveAudioDevices >= 1)
 	{
-		uint32 MainDeviceHandle = GEngine->GetAudioDeviceHandle();
+		Audio::FDeviceId MainDeviceHandle = FAudioDeviceManager::GetChecked().GetMainAudioDeviceHandle();
 
 		if (NumActiveAudioDevices == 1)
 		{
@@ -557,7 +592,7 @@ bool FAudioDeviceManager::ShutdownAudioDevice(Audio::FDeviceId Handle)
 		// don't shut it down until it's the very last handle to get shut down
 		// this is because it's possible for some PIE sessions to be using the main audio device as a fallback to
 		// preserve CPU performance on low-performance machines
-		if (NumWorldsUsingMainAudioDevice > 0 && MainDeviceHandle == Handle)
+		if (NumWorldsUsingMainAudioDevice > 0 && MainDeviceHandle == static_cast<uint32>(Handle))
 		{
 			--NumWorldsUsingMainAudioDevice;
 
@@ -625,21 +660,13 @@ FAudioDevice* FAudioDeviceManager::GetAudioDevice(Audio::FDeviceId Handle)
 
 FAudioDeviceManager* FAudioDeviceManager::Get()
 {
-	if (GEngine)
-	{
-		return GEngine->GetAudioDeviceManager();
-	}
-
-	return nullptr;
+	return AudioDeviceManager;
 }
 
-FAudioDevice* FAudioDeviceManager::GetActiveAudioDevice()
+FAudioDeviceManager& FAudioDeviceManager::GetChecked()
 {
-	if (ActiveAudioDeviceHandle != INDEX_NONE)
-	{
-		return GetAudioDevice(ActiveAudioDeviceHandle);
-	}
-	return GEngine->GetMainAudioDevice();
+	check(AudioDeviceManager);
+	return *AudioDeviceManager;
 }
 
 void FAudioDeviceManager::UpdateActiveAudioDevices(bool bGameTicking)
@@ -849,6 +876,23 @@ void FAudioDeviceManager::SetSoloDevice(Audio::FDeviceId InAudioDeviceHandle)
 	}
 }
 
+void FAudioDeviceManager::Shutdown()
+{
+	// Shutdown the main audio device in the UEEngine
+	if (AudioDeviceManager)
+	{
+		FAudioCommandFence Fence;
+		Fence.BeginFence();
+		Fence.Wait();
+
+		FAudioThread::StopAudioThread();
+
+		AudioDeviceManager->ShutdownAllAudioDevices();
+
+		delete AudioDeviceManager;
+		AudioDeviceManager = nullptr;
+	}
+}
 
 uint8 FAudioDeviceManager::GetNumActiveAudioDevices() const
 {
@@ -950,8 +994,7 @@ void FAudioDeviceManager::RemoveSoundMix(USoundMix* SoundMix)
 	{
 		DECLARE_CYCLE_STAT(TEXT("FAudioThreadTask.RemoveSoundMix"), STAT_AudioRemoveSoundMix, STATGROUP_AudioThreadCommands);
 
-		FAudioDeviceManager* AudioDeviceManager = this;
-		FAudioThread::RunCommandOnAudioThread([AudioDeviceManager, SoundMix]()
+		FAudioThread::RunCommandOnAudioThread([this, SoundMix]()
 		{
 			AudioDeviceManager->RemoveSoundMix(SoundMix);
 
@@ -975,10 +1018,9 @@ void FAudioDeviceManager::TogglePlayAllDeviceAudio()
 	{
 		DECLARE_CYCLE_STAT(TEXT("FAudioThreadTask.TogglePlayAllDeviceAudio"), STAT_TogglePlayAllDeviceAudio, STATGROUP_AudioThreadCommands);
 
-		FAudioDeviceManager* AudioDeviceManager = this;
-		FAudioThread::RunCommandOnAudioThread([AudioDeviceManager]()
+		FAudioThread::RunCommandOnAudioThread([this]()
 		{
-			AudioDeviceManager->TogglePlayAllDeviceAudio();
+			TogglePlayAllDeviceAudio();
 
 		}, GET_STATID(STAT_TogglePlayAllDeviceAudio));
 
@@ -1004,11 +1046,9 @@ void FAudioDeviceManager::ToggleVisualize3dDebug()
 	{
 		DECLARE_CYCLE_STAT(TEXT("FAudioThreadTask.ToggleVisualize3dDebug"), STAT_ToggleVisualize3dDebug, STATGROUP_AudioThreadCommands);
 
-		FAudioDeviceManager* AudioDeviceManager = this;
-		FAudioThread::RunCommandOnAudioThread([AudioDeviceManager]()
+		FAudioThread::RunCommandOnAudioThread([this]()
 		{
-			AudioDeviceManager->ToggleVisualize3dDebug();
-
+			ToggleVisualize3dDebug();
 		}, GET_STATID(STAT_ToggleVisualize3dDebug));
 
 		return;
@@ -1037,10 +1077,9 @@ void FAudioDeviceManager::ResetAllDynamicSoundVolumes()
 	{
 		DECLARE_CYCLE_STAT(TEXT("FAudioThreadTask.ResetAllDynamicSoundVolumes"), STAT_ResetAllDynamicSoundVolumes, STATGROUP_AudioThreadCommands);
 
-		FAudioDeviceManager* AudioDeviceManager = this;
-		FAudioThread::RunCommandOnAudioThread([AudioDeviceManager]()
+		FAudioThread::RunCommandOnAudioThread([this]()
 		{
-			AudioDeviceManager->ResetAllDynamicSoundVolumes();
+			ResetAllDynamicSoundVolumes();
 
 		}, GET_STATID(STAT_ResetAllDynamicSoundVolumes));
 		return;
@@ -1056,11 +1095,9 @@ void FAudioDeviceManager::ResetDynamicSoundVolume(ESoundType SoundType, const FN
 	{
 		DECLARE_CYCLE_STAT(TEXT("FAudioThreadTask.ResetSoundCueTrimVolume"), STAT_ResetSoundCueTrimVolume, STATGROUP_AudioThreadCommands);
 
-		FAudioDeviceManager* AudioDeviceManager = this;
-		FAudioThread::RunCommandOnAudioThread([AudioDeviceManager, SoundType, SoundName]()
+		FAudioThread::RunCommandOnAudioThread([this, SoundType, SoundName]()
 		{
-			AudioDeviceManager->ResetDynamicSoundVolume(SoundType, SoundName);
-
+			ResetDynamicSoundVolume(SoundType, SoundName);
 		}, GET_STATID(STAT_ResetSoundCueTrimVolume));
 		return;
 	}
@@ -1075,11 +1112,9 @@ void FAudioDeviceManager::SetDynamicSoundVolume(ESoundType SoundType, const FNam
 	{
 		DECLARE_CYCLE_STAT(TEXT("FAudioThreadTask.SetDynamicSoundVolume"), STAT_SetDynamicSoundVolume, STATGROUP_AudioThreadCommands);
 
-		FAudioDeviceManager* AudioDeviceManager = this;
-		FAudioThread::RunCommandOnAudioThread([AudioDeviceManager, SoundType, SoundName, Volume]()
+		FAudioThread::RunCommandOnAudioThread([this, SoundType, SoundName, Volume]()
 		{
-			AudioDeviceManager->SetDynamicSoundVolume(SoundType, SoundName, Volume);
-
+			SetDynamicSoundVolume(SoundType, SoundName, Volume);
 		}, GET_STATID(STAT_SetDynamicSoundVolume));
 		return;
 	}
diff --git a/Engine/Source/Runtime/Engine/Private/AudioEffect.cpp b/Engine/Source/Runtime/Engine/Private/AudioEffect.cpp
--- a/Engine/Source/Runtime/Engine/Private/AudioEffect.cpp
+++ b/Engine/Source/Runtime/Engine/Private/AudioEffect.cpp
@@ -510,8 +510,7 @@ void FAudioEffectsManager::Update()
 	if (CurrentReverbAsset && CurrentReverbAsset->bChanged)
 	{
 		CurrentReverbAsset->bChanged = false;
-		class FAudioDeviceManager* DeviceManager = GEngine->GetAudioDeviceManager();
-		TArray<FAudioDevice*>& AudioDevices = DeviceManager->GetAudioDevices();
+		TArray<FAudioDevice*>& AudioDevices = FAudioDeviceManager::GetChecked().GetAudioDevices();
 		for (int32 i = 0; i < AudioDevices.Num(); ++i)
 		{
 			if (AudioDevices[i])
@@ -522,8 +521,7 @@ void FAudioEffectsManager::Update()
 		}
 		
 	}
-
-#endif
+#endif // WITH_EDITORONLY_DATA
 
 	const bool bIsReverbDone = CurrentReverbEffect.Interpolate(SourceReverbEffect, DestinationReverbEffect);
 	if (!bIsReverbDone || bReverbActive || bReverbChanged)
diff --git a/Engine/Source/Runtime/Engine/Private/Components/AudioComponent.cpp b/Engine/Source/Runtime/Engine/Private/Components/AudioComponent.cpp
--- a/Engine/Source/Runtime/Engine/Private/Components/AudioComponent.cpp
+++ b/Engine/Source/Runtime/Engine/Private/Components/AudioComponent.cpp
@@ -556,8 +556,7 @@ FAudioDevice* UAudioComponent::GetAudioDevice() const
 	{
 		if (AudioDeviceHandle != INDEX_NONE)
 		{
-			FAudioDeviceManager* AudioDeviceManager = GEngine->GetAudioDeviceManager();
-			AudioDevice = (AudioDeviceManager ? AudioDeviceManager->GetAudioDevice(AudioDeviceHandle) : nullptr);
+			AudioDevice = FAudioDeviceManager::GetChecked().GetAudioDevice(AudioDeviceHandle);
 		}
 		else if (UWorld* World = GetWorld())
 		{
@@ -565,7 +564,7 @@ FAudioDevice* UAudioComponent::GetAudioDevice() const
 		}
 		else
 		{
-			AudioDevice = GEngine->GetMainAudioDevice();
+			AudioDevice = FAudioDeviceManager::GetMainDevice();
 		}
 	}
 	return AudioDevice;
diff --git a/Engine/Source/Runtime/Engine/Private/ContentStreaming.cpp b/Engine/Source/Runtime/Engine/Private/ContentStreaming.cpp
--- a/Engine/Source/Runtime/Engine/Private/ContentStreaming.cpp
+++ b/Engine/Source/Runtime/Engine/Private/ContentStreaming.cpp
@@ -1179,13 +1179,15 @@ void FStreamingManagerCollection::PropagateLightingScenarioChange()
 void FStreamingManagerCollection::OnAudioStreamingParamsChanged()
 {
 	// Before we swap out the audio streaming manager, we'll need to stop all sounds running on all audio devices:
-	FAudioDeviceManager* DeviceManager = GEngine->GetAudioDeviceManager();
-	TArray<FAudioDevice*>& AudioDevices = DeviceManager->GetAudioDevices();
-	for (FAudioDevice* AudioDevice : AudioDevices)
+	if (FAudioDeviceManager* DeviceManager = FAudioDeviceManager::Get())
 	{
-		if (AudioDevice)
+		TArray<FAudioDevice*>& AudioDevices = DeviceManager->GetAudioDevices();
+		for (FAudioDevice* AudioDevice : AudioDevices)
 		{
-			AudioDevice->StopAllSounds(true);
+			if (AudioDevice)
+			{
+				AudioDevice->StopAllSounds(true);
+			}
 		}
 	}
 
diff --git a/Engine/Source/Runtime/Engine/Private/GameEngine.cpp b/Engine/Source/Runtime/Engine/Private/GameEngine.cpp
--- a/Engine/Source/Runtime/Engine/Private/GameEngine.cpp
+++ b/Engine/Source/Runtime/Engine/Private/GameEngine.cpp
@@ -1801,7 +1801,7 @@ void UGameEngine::Tick( float DeltaSeconds, bool bIdleMode )
 	}
 
 	// Update Audio. This needs to occur after rendering as the rendering code updates the listener position.
-	FAudioDeviceManager* GameAudioDeviceManager = GEngine->GetAudioDeviceManager();
+	FAudioDeviceManager* GameAudioDeviceManager = FAudioDeviceManager::Get();
 	if (GameAudioDeviceManager)
 	{
 		SCOPE_TIME_GUARD(TEXT("UGameEngine::Tick - Update Audio"));
diff --git a/Engine/Source/Runtime/Engine/Private/GameFramework/HapticFeedbackEffect.cpp b/Engine/Source/Runtime/Engine/Private/GameFramework/HapticFeedbackEffect.cpp
--- a/Engine/Source/Runtime/Engine/Private/GameFramework/HapticFeedbackEffect.cpp
+++ b/Engine/Source/Runtime/Engine/Private/GameFramework/HapticFeedbackEffect.cpp
@@ -168,7 +168,7 @@ float UHapticFeedbackEffect_SoundWave::GetDuration() const
 
 void UHapticFeedbackEffect_SoundWave::PrepareSoundWaveBuffer()
 {
-	FAudioDevice *AD = GEngine->GetMainAudioDevice();
+	FAudioDevice *AD = FAudioDeviceManager::GetMainDevice();
 	if (!AD || !SoundWave)
 	{
 		return;
diff --git a/Engine/Source/Runtime/Engine/Private/GameUserSettings.cpp b/Engine/Source/Runtime/Engine/Private/GameUserSettings.cpp
--- a/Engine/Source/Runtime/Engine/Private/GameUserSettings.cpp
+++ b/Engine/Source/Runtime/Engine/Private/GameUserSettings.cpp
@@ -447,8 +447,7 @@ void UGameUserSettings::ApplyNonResolutionSettings()
 		Scalability::SetQualityLevels(ScalabilityQuality);
 	}
 
-	FAudioDevice* AudioDevice = GEngine->GetMainAudioDevice();
-	if (AudioDevice)
+	if (FAudioDevice* AudioDevice = FAudioDeviceManager::GetMainDevice())
 	{
 		FAudioQualitySettings AudioSettings = AudioDevice->GetQualityLevelSettings();
 		AudioDevice->SetMaxChannels(AudioSettings.MaxChannels);
diff --git a/Engine/Source/Runtime/Engine/Private/GameViewportClient.cpp b/Engine/Source/Runtime/Engine/Private/GameViewportClient.cpp
--- a/Engine/Source/Runtime/Engine/Private/GameViewportClient.cpp
+++ b/Engine/Source/Runtime/Engine/Private/GameViewportClient.cpp
@@ -303,13 +303,10 @@ void UGameViewportClient::PostInitProperties()
 
 void UGameViewportClient::BeginDestroy()
 {
-	if (GEngine)
+	class FAudioDeviceManager* AudioDeviceManager = FAudioDeviceManager::Get();
+	if (AudioDeviceManager)
 	{
-		class FAudioDeviceManager* AudioDeviceManager = GEngine->GetAudioDeviceManager();
-		if (AudioDeviceManager)
-		{
-			AudioDeviceManager->ShutdownAudioDevice(AudioDeviceHandle);
-		}
+		AudioDeviceManager->ShutdownAudioDevice(AudioDeviceHandle);
 	}
 
 	RemoveAllViewportWidgets();
@@ -382,7 +379,7 @@ void UGameViewportClient::SetEnabledStats(const TArray<FString>& InEnabledStats)
 #if ENABLE_AUDIO_DEBUG
 	if (UWorld* MyWorld = GetWorld())
 	{
-		if (FAudioDeviceManager* DeviceManager = GEngine->GetAudioDeviceManager())
+		if (FAudioDeviceManager* DeviceManager = FAudioDeviceManager::Get())
 		{
 			FAudioDebugger::ResolveDesiredStats(this);
 		}
@@ -436,33 +433,30 @@ void UGameViewportClient::Init(struct FWorldContext& WorldContext, UGameInstance
 	// Create the cursor Widgets
 	UUserInterfaceSettings* UISettings = GetMutableDefault<UUserInterfaceSettings>(UUserInterfaceSettings::StaticClass());
 
-	if (GEngine)
+	FAudioDeviceManager* AudioDeviceManager = FAudioDeviceManager::Get();
+	if (AudioDeviceManager)
 	{
-		FAudioDeviceManager* AudioDeviceManager = GEngine->GetAudioDeviceManager();
-		if (AudioDeviceManager)
+		FAudioDeviceManager::FCreateAudioDeviceResults NewDeviceResults;
+		if (AudioDeviceManager->CreateAudioDevice(bCreateNewAudioDevice, NewDeviceResults))
 		{
-			FAudioDeviceManager::FCreateAudioDeviceResults NewDeviceResults;
-			if (AudioDeviceManager->CreateAudioDevice(bCreateNewAudioDevice, NewDeviceResults))
-			{
-				AudioDeviceHandle = NewDeviceResults.Handle;
+			AudioDeviceHandle = NewDeviceResults.Handle;
 
 #if ENABLE_AUDIO_DEBUG
-				FAudioDebugger::ResolveDesiredStats(this);
+			FAudioDebugger::ResolveDesiredStats(this);
 #endif // ENABLE_AUDIO_DEBUG
 
-				// Set the base mix of the new device based on the world settings of the world
-				if (World)
-				{
-					NewDeviceResults.AudioDevice->SetDefaultBaseSoundMix(World->GetWorldSettings()->DefaultBaseSoundMix);
-
-					// Set the world's audio device handle to use so that sounds which play in that world will use the correct audio device
-					World->SetAudioDeviceHandle(AudioDeviceHandle);
-				}
+			// Set the base mix of the new device based on the world settings of the world
+			if (World)
+			{
+				NewDeviceResults.AudioDevice->SetDefaultBaseSoundMix(World->GetWorldSettings()->DefaultBaseSoundMix);
 
-				// Set this audio device handle on the world context so future world's set onto the world context
-				// will pass the audio device handle to them and audio will play on the correct audio device
-				WorldContext.AudioDeviceHandle = AudioDeviceHandle;
+				// Set the world's audio device handle to use so that sounds which play in that world will use the correct audio device
+				World->SetAudioDeviceHandle(AudioDeviceHandle);
 			}
+
+			// Set this audio device handle on the world context so future world's set onto the world context
+			// will pass the audio device handle to them and audio will play on the correct audio device
+			WorldContext.AudioDeviceHandle = AudioDeviceHandle;
 		}
 	}
 
@@ -1309,18 +1303,15 @@ void UGameViewportClient::Draw(FViewport* InViewport, FCanvas* SceneCanvas)
 
 							// If the main audio device is used for multiple PIE viewport clients, we only
 							// want to update the main audio device listener position if it is in focus
-							if (GEngine)
-							{
-								FAudioDeviceManager* AudioDeviceManager = GEngine->GetAudioDeviceManager();
+							FAudioDeviceManager* DeviceManager = FAudioDeviceManager::Get();
 
-								// If there is more than one world referencing the main audio device
-								if (AudioDeviceManager->GetNumMainAudioDeviceWorlds() > 1)
+							// If there is more than one world referencing the main audio device
+							if (DeviceManager->GetNumMainAudioDeviceWorlds() > 1)
+							{
+								uint32 MainAudioDeviceHandle = DeviceManager->GetMainAudioDeviceHandle();
+								if (AudioDevice->DeviceHandle == MainAudioDeviceHandle && !bHasAudioFocus)
 								{
-									uint32 MainAudioDeviceHandle = GEngine->GetAudioDeviceHandle();
-									if (AudioDevice->DeviceHandle == MainAudioDeviceHandle && !bHasAudioFocus)
-									{
-										bUpdateListenerPosition = false;
-									}
+									bUpdateListenerPosition = false;
 								}
 							}
 
@@ -1876,7 +1867,7 @@ void UGameViewportClient::LostFocus(FViewport* InViewport)
 		}
 	}
 
-	if (GEngine && GEngine->GetAudioDeviceManager())
+	if (FAudioDeviceManager::Get())
 	{
 		bHasAudioFocus = false;
 	}
@@ -1889,9 +1880,9 @@ void UGameViewportClient::ReceivedFocus(FViewport* InViewport)
 		FSlateApplication::Get().SetGameIsFakingTouchEvents(true);
 	}
 
-	if (GEngine && GEngine->GetAudioDeviceManager())
+	if (FAudioDeviceManager* DeviceManager = FAudioDeviceManager::Get())
 	{
-		GEngine->GetAudioDeviceManager()->SetActiveDevice(AudioDeviceHandle);
+		DeviceManager->SetActiveDevice(AudioDeviceHandle);
 		bHasAudioFocus = true;
 	}
 }
diff --git a/Engine/Source/Runtime/Engine/Private/Interpolation.cpp b/Engine/Source/Runtime/Engine/Private/Interpolation.cpp
--- a/Engine/Source/Runtime/Engine/Private/Interpolation.cpp
+++ b/Engine/Source/Runtime/Engine/Private/Interpolation.cpp
@@ -1271,10 +1271,9 @@ void AMatineeActor::StepInterp( float DeltaTime, bool bPreview )
 
 void AMatineeActor::DisableRadioFilterIfNeeded()
 {
-	FAudioDevice* AudioDevice = GEngine->GetMainAudioDevice();
-	if( AudioDevice )
+	if(FAudioDevice* AudioDevice = FAudioDeviceManager::GetMainDevice())
 	{
-		AudioDevice->EnableRadioEffect( !bDisableRadioFilter );
+		AudioDevice->EnableRadioEffect(!bDisableRadioFilter);
 	}
 }
 
@@ -1299,10 +1298,9 @@ void AMatineeActor::EnableCinematicMode(bool bEnable)
 
 void AMatineeActor::EnableRadioFilter()
 {
-	FAudioDevice* AudioDevice = GEngine->GetMainAudioDevice();
-	if( AudioDevice )
+	if (FAudioDevice* AudioDevice = FAudioDeviceManager::GetMainDevice())
 	{
-		AudioDevice->EnableRadioEffect( true );
+		AudioDevice->EnableRadioEffect(true);
 	}
 }
 
diff --git a/Engine/Source/Runtime/Engine/Private/LevelActor.cpp b/Engine/Source/Runtime/Engine/Private/LevelActor.cpp
--- a/Engine/Source/Runtime/Engine/Private/LevelActor.cpp
+++ b/Engine/Source/Runtime/Engine/Private/LevelActor.cpp
@@ -1512,20 +1512,17 @@ void UWorld::SetAudioDeviceHandle(const uint32 InAudioDeviceHandle)
 
 FAudioDevice* UWorld::GetAudioDevice()
 {
-	FAudioDevice* AudioDevice = nullptr;
-	if (GEngine)
+	if (FAudioDeviceManager* AudioDeviceManager = FAudioDeviceManager::Get())
 	{
-		class FAudioDeviceManager* AudioDeviceManager = GEngine->GetAudioDeviceManager();
-		if (AudioDeviceManager != nullptr)
+		if (FAudioDevice* AudioDevice = AudioDeviceManager->GetAudioDevice(AudioDeviceHandle))
 		{
-			AudioDevice = AudioDeviceManager->GetAudioDevice(AudioDeviceHandle);
-			if (AudioDevice == nullptr)
-			{
-				AudioDevice = GEngine->GetMainAudioDevice();
-			}
+			return AudioDevice;
 		}
+
+		return FAudioDeviceManager::GetMainDevice();
 	}
-	return AudioDevice;
+
+	return nullptr;
 }
 
 /**
diff --git a/Engine/Source/Runtime/Engine/Private/SampleBufferIO.cpp b/Engine/Source/Runtime/Engine/Private/SampleBufferIO.cpp
--- a/Engine/Source/Runtime/Engine/Private/SampleBufferIO.cpp
+++ b/Engine/Source/Runtime/Engine/Private/SampleBufferIO.cpp
@@ -17,8 +17,7 @@ namespace Audio
 
 	void FSoundWavePCMLoader::LoadSoundWave(USoundWave* InSoundWave, TFunction<void(const USoundWave* SoundWave, const Audio::FSampleBuffer& OutSampleBuffer)> OnLoaded)
 	{
-		FAudioDevice* AudioDevice = FAudioDevice::GetMainAudioDevice();
-
+		FAudioDevice* AudioDevice = FAudioDeviceManager::GetMainDevice();
 		if (!AudioDevice || !InSoundWave)
 		{
 			return;
@@ -122,9 +121,9 @@ namespace Audio
 			CurrentSoundWave = SoundWaveToSaveTo;
 			bWasPreviouslyAddedToRoot = CurrentSoundWave->IsRooted();
 			CurrentSoundWave->AddToRoot();
+
 			// Ensure this sound wave is not currently in use:
-			FAudioDeviceManager* AudioDeviceManager = GEngine->GetAudioDeviceManager();
-			if (AudioDeviceManager)
+			if (FAudioDeviceManager* AudioDeviceManager = FAudioDeviceManager::Get())
 			{
 				AudioDeviceManager->StopSoundsUsingResource(CurrentSoundWave);
 			}
diff --git a/Engine/Source/Runtime/Engine/Private/Slate/SlateSoundDevice.cpp b/Engine/Source/Runtime/Engine/Private/Slate/SlateSoundDevice.cpp
--- a/Engine/Source/Runtime/Engine/Private/Slate/SlateSoundDevice.cpp
+++ b/Engine/Source/Runtime/Engine/Private/Slate/SlateSoundDevice.cpp
@@ -10,14 +10,13 @@ DEFINE_LOG_CATEGORY_STATIC(LogSlateSoundDevice, Log, All);
 
 void FSlateSoundDevice::PlaySound(const FSlateSound& Sound, int32 UserIndex) const
 {
-	if( GEngine && Sound.GetResourceObject() != nullptr )
+	if (Sound.GetResourceObject())
 	{
-		FAudioDevice* const AudioDevice = GEngine->GetActiveAudioDevice();
-		if(AudioDevice)
+		if (FAudioDevice* AudioDevice = FAudioDeviceManager::GetActiveDevice())
 		{
 			UObject* const Object = Sound.GetResourceObject();
 			USoundBase* const SoundResource = Cast<USoundBase>(Object);
-			if(SoundResource)
+			if (SoundResource)
 			{
 				FActiveSound NewActiveSound;
 				NewActiveSound.SetSound(SoundResource);
@@ -27,7 +26,7 @@ void FSlateSoundDevice::PlaySound(const FSlateSound& Sound, int32 UserIndex) con
 
 				AudioDevice->AddNewActiveSound(NewActiveSound);
 			}
-			else if(Object)
+			else if (Object)
 			{
 				// An Object but no SoundResource means that the FSlateSound is holding an invalid object; report that as an error
 				UE_LOG(LogSlateSoundDevice, Error, TEXT("A sound contains a non-sound resource '%s'"), *Object->GetName());
diff --git a/Engine/Source/Runtime/Engine/Private/SoundBase.cpp b/Engine/Source/Runtime/Engine/Private/SoundBase.cpp
--- a/Engine/Source/Runtime/Engine/Private/SoundBase.cpp
+++ b/Engine/Source/Runtime/Engine/Private/SoundBase.cpp
@@ -1,9 +1,11 @@
 // Copyright Epic Games, Inc. All Rights Reserved.
 #include "Sound/SoundBase.h"
 
+#include "Audio.h"
 #include "EngineDefines.h"
 #include "IAudioExtensionPlugin.h"
 #include "Sound/AudioSettings.h"
+#include "Sound/SoundAttenuation.h"
 #include "Sound/SoundClass.h"
 #include "Sound/SoundSubmix.h"
 
diff --git a/Engine/Source/Runtime/Engine/Private/SoundClass.cpp b/Engine/Source/Runtime/Engine/Private/SoundClass.cpp
--- a/Engine/Source/Runtime/Engine/Private/SoundClass.cpp
+++ b/Engine/Source/Runtime/Engine/Private/SoundClass.cpp
@@ -103,11 +103,8 @@ void USoundClass::PostLoad()
 	}
 #endif
 
-	// Use the main/default audio device for storing and retrieving sound class properties
-	FAudioDeviceManager* AudioDeviceManager = (GEngine ? GEngine->GetAudioDeviceManager() : nullptr);
-
 	// Force the properties to be initialized for this SoundClass on all active audio devices
-	if (AudioDeviceManager)
+	if (FAudioDeviceManager* AudioDeviceManager = FAudioDeviceManager::Get())
 	{
 		AudioDeviceManager->RegisterSoundClass(this);
 	}
@@ -247,11 +244,8 @@ void USoundClass::PostEditChangeProperty(struct FPropertyChangedEvent& PropertyC
 		}
 	}
 
-	// Use the main/default audio device for storing and retrieving sound class properties
-	FAudioDeviceManager* AudioDeviceManager = (GEngine ? GEngine->GetAudioDeviceManager() : nullptr);
-
 	// Force the properties to be initialized for this SoundClass on all active audio devices
-	if (AudioDeviceManager)
+	if (FAudioDeviceManager* AudioDeviceManager = FAudioDeviceManager::Get())
 	{
 		AudioDeviceManager->RegisterSoundClass(this);
 	}
@@ -298,9 +292,12 @@ void USoundClass::BeginDestroy()
 {
 	Super::BeginDestroy();
 
-	if (!GExitPurge && GEngine && GEngine->GetAudioDeviceManager())
+	if (!GExitPurge)
 	{
-		GEngine->GetAudioDeviceManager()->UnregisterSoundClass(this);
+		if (FAudioDeviceManager* DeviceManager = FAudioDeviceManager::Get())
+		{
+			DeviceManager->UnregisterSoundClass(this);
+		}
 	}
 }
 
diff --git a/Engine/Source/Runtime/Engine/Private/SoundEffectPreset.cpp b/Engine/Source/Runtime/Engine/Private/SoundEffectPreset.cpp
--- a/Engine/Source/Runtime/Engine/Private/SoundEffectPreset.cpp
+++ b/Engine/Source/Runtime/Engine/Private/SoundEffectPreset.cpp
@@ -93,11 +93,7 @@ void USoundEffectPreset::PostEditChangeProperty(FPropertyChangedEvent& PropertyC
 
 void USoundEffectSourcePresetChain::PostEditChangeProperty(FPropertyChangedEvent& PropertyChangedEvent)
 {
-	if (GEngine)
-	{
-		FAudioDeviceManager* AudioDeviceManager = GEngine->GetAudioDeviceManager();
-		AudioDeviceManager->UpdateSourceEffectChain(GetUniqueID(), Chain, bPlayEffectChainTails);
-	}
+	FAudioDeviceManager::GetChecked().UpdateSourceEffectChain(GetUniqueID(), Chain, bPlayEffectChainTails);
 }
 #endif // WITH_EDITORONLY_DATA
 
diff --git a/Engine/Source/Runtime/Engine/Private/SoundMix.cpp b/Engine/Source/Runtime/Engine/Private/SoundMix.cpp
--- a/Engine/Source/Runtime/Engine/Private/SoundMix.cpp
+++ b/Engine/Source/Runtime/Engine/Private/SoundMix.cpp
@@ -32,12 +32,11 @@ USoundMix::USoundMix(const FObjectInitializer& ObjectInitializer)
 
 void USoundMix::BeginDestroy()
 {
-	if (!GExitPurge && GEngine)
+	if (!GExitPurge)
 	{
-		FAudioDeviceManager* AudioDeviceManager = GEngine->GetAudioDeviceManager();
-		if (AudioDeviceManager)
+		if (FAudioDeviceManager* DeviceManager = FAudioDeviceManager::Get())
 		{
-			AudioDeviceManager->RemoveSoundMix(this);
+			DeviceManager->RemoveSoundMix(this);
 		}
 	}
 
diff --git a/Engine/Source/Runtime/Engine/Private/SoundSubmix.cpp b/Engine/Source/Runtime/Engine/Private/SoundSubmix.cpp
--- a/Engine/Source/Runtime/Engine/Private/SoundSubmix.cpp
+++ b/Engine/Source/Runtime/Engine/Private/SoundSubmix.cpp
@@ -214,13 +214,10 @@ void USoundSubmix::BeginDestroy()
 {
 	Super::BeginDestroy();
 
-	// Use the main/default audio device for storing and retrieving sound class properties
-	FAudioDeviceManager* AudioDeviceManager = (GEngine ? GEngine->GetAudioDeviceManager() : nullptr);
-
 	// Force the properties to be initialized for this SoundClass on all active audio devices
-	if (AudioDeviceManager)
+	if (FAudioDeviceManager* DeviceManager = FAudioDeviceManager::Get())
 	{
-		AudioDeviceManager->UnregisterSoundSubmix(this);
+		DeviceManager->UnregisterSoundSubmix(this);
 	}
 }
 
@@ -228,13 +225,10 @@ void USoundSubmix::PostLoad()
 {
 	Super::PostLoad();
 
-	// Use the main/default audio device for storing and retrieving sound class properties
-	FAudioDeviceManager* AudioDeviceManager = (GEngine ? GEngine->GetAudioDeviceManager() : nullptr);
-
 	// Force the properties to be initialized for this SoundClass on all active audio devices
-	if (AudioDeviceManager)
+	if (FAudioDeviceManager* DeviceManager = FAudioDeviceManager::Get())
 	{
-		AudioDeviceManager->RegisterSoundSubmix(this);
+		DeviceManager->RegisterSoundSubmix(this);
 	}
 }
 
@@ -325,10 +319,9 @@ void USoundSubmix::PostEditChangeProperty(FPropertyChangedEvent& PropertyChanged
 		}
 		else if (PropertyChangedEvent.Property->GetFName() == NAME_OutputVolume)
 		{
-			FAudioDeviceManager* AudioDeviceManager = (GEngine ? GEngine->GetAudioDeviceManager() : nullptr);
-			if (AudioDeviceManager)
+			if (FAudioDeviceManager* DeviceManager = FAudioDeviceManager::Get())
 			{
-				AudioDeviceManager->UpdateSubmix(this);
+				DeviceManager->UpdateSubmix(this);
 			}
 			bReinitSubmix = false;
 		}
@@ -337,9 +330,9 @@ void USoundSubmix::PostEditChangeProperty(FPropertyChangedEvent& PropertyChanged
 	if (GEngine && bReinitSubmix)
 	{
 		// Force the properties to be initialized for this SoundSubmix on all active audio devices
-		if (FAudioDeviceManager* AudioDeviceManager = GEngine->GetAudioDeviceManager())
+		if (FAudioDeviceManager* DeviceManager = FAudioDeviceManager::Get())
 		{
-			AudioDeviceManager->RegisterSoundSubmix(this);
+			DeviceManager->RegisterSoundSubmix(this);
 		}
 	}
 
diff --git a/Engine/Source/Runtime/Engine/Private/SoundWave.cpp b/Engine/Source/Runtime/Engine/Private/SoundWave.cpp
--- a/Engine/Source/Runtime/Engine/Private/SoundWave.cpp
+++ b/Engine/Source/Runtime/Engine/Private/SoundWave.cpp
@@ -219,9 +219,9 @@ void USoundWave::GetResourceSizeEx(FResourceSizeEx& CumulativeResourceSize)
 		return;
 	}
 
-	if (FAudioDevice* LocalAudioDevice = GEngine->GetMainAudioDevice())
+	if (FAudioDevice* AudioDevice = FAudioDeviceManager::GetMainDevice())
 	{
-		if (LocalAudioDevice->HasCompressedAudioInfoClass(this) && DecompressionType == DTYPE_Native)
+		if (AudioDevice->HasCompressedAudioInfoClass(this) && DecompressionType == DTYPE_Native)
 		{
 			check(!RawPCMData || RawPCMDataSize);
 			CumulativeResourceSize.AddDedicatedSystemMemoryBytes(RawPCMDataSize);
@@ -235,7 +235,7 @@ void USoundWave::GetResourceSizeEx(FResourceSizeEx& CumulativeResourceSize)
 
 			if (!FPlatformProperties::SupportsAudioStreaming() || !IsStreaming())
 			{
-				CumulativeResourceSize.AddDedicatedSystemMemoryBytes(GetCompressedDataSize(LocalAudioDevice->GetRuntimeFormat(this)));
+				CumulativeResourceSize.AddDedicatedSystemMemoryBytes(GetCompressedDataSize(AudioDevice->GetRuntimeFormat(this)));
 			}
 		}
 	}
@@ -835,8 +835,7 @@ void USoundWave::PostLoad()
 	// most likely cause us to run out of memory.
 	if (!GIsEditor && !IsTemplate( RF_ClassDefaultObject ) && GEngine)
 	{
-		FAudioDevice* AudioDevice = GEngine->GetMainAudioDevice();
-		if (AudioDevice)
+		if (FAudioDevice* AudioDevice = FAudioDeviceManager::GetMainDevice())
 		{
 			// Upload the data to the hardware, but only if we've precached startup sounds already
 			AudioDevice->Precache(this);
@@ -1511,14 +1510,16 @@ void USoundWave::FreeResources(bool bStopSoundsUsingThisResource)
 
 	// GEngine is NULL during script compilation and GEngine->Client and its audio device might be
 	// destroyed first during the exit purge.
-	if( GEngine && !GExitPurge )
+	if(GEngine && !GExitPurge)
 	{
 		// Notify the audio device to free the bulk data associated with this wave.
-		FAudioDeviceManager* AudioDeviceManager = GEngine->GetAudioDeviceManager();
-		if (bStopSoundsUsingThisResource && AudioDeviceManager)
+		if (bStopSoundsUsingThisResource)
 		{
-			AudioDeviceManager->StopSoundsUsingResource(this);
-			AudioDeviceManager->FreeResource(this);
+			if (FAudioDeviceManager* AudioDeviceManager = FAudioDeviceManager::Get())
+			{
+				AudioDeviceManager->StopSoundsUsingResource(this);
+				AudioDeviceManager->FreeResource(this);
+			}
 		}
 	}
 
diff --git a/Engine/Source/Runtime/Engine/Private/SoundWaveProcedural.cpp b/Engine/Source/Runtime/Engine/Private/SoundWaveProcedural.cpp
--- a/Engine/Source/Runtime/Engine/Private/SoundWaveProcedural.cpp
+++ b/Engine/Source/Runtime/Engine/Private/SoundWaveProcedural.cpp
@@ -20,7 +20,7 @@ USoundWaveProcedural::USoundWaveProcedural(const FObjectInitializer& ObjectIniti
 	
 	if (GEngine)
 	{
-		FAudioDevice* MainAudioDevice = GEngine->GetMainAudioDevice();
+		FAudioDevice* MainAudioDevice = FAudioDeviceManager::GetMainDevice();
 		if (MainAudioDevice && !MainAudioDevice->IsAudioMixerEnabled())
 		{
 #if PLATFORM_MAC
diff --git a/Engine/Source/Runtime/Engine/Private/UnrealEngine.cpp b/Engine/Source/Runtime/Engine/Private/UnrealEngine.cpp
--- a/Engine/Source/Runtime/Engine/Private/UnrealEngine.cpp
+++ b/Engine/Source/Runtime/Engine/Private/UnrealEngine.cpp
@@ -1745,20 +1745,7 @@ void UEngine::OnExternalUIChange(bool bInIsOpening)
 
 void UEngine::ShutdownAudioDeviceManager()
 {
-	// Shutdown the main audio device in the UEEngine
-	if (AudioDeviceManager)
-	{
-		FAudioCommandFence Fence;
-		Fence.BeginFence();
-		Fence.Wait();
-
-		FAudioThread::StopAudioThread();
-
-		AudioDeviceManager->ShutdownAllAudioDevices();
-
-		delete AudioDeviceManager;
-		AudioDeviceManager = NULL;
-	}
+	FAudioDeviceManager::Shutdown();
 }
 
 void UEngine::PreExit()
@@ -2729,7 +2716,7 @@ void UEngine::FinishDestroy()
 	{
 		// shut down all subsystems.
 		GEngine = NULL;
-		ShutdownAudioDeviceManager();
+		FAudioDeviceManager::Shutdown();
 
 		FURL::StaticExit();
 	}
@@ -2745,7 +2732,7 @@ void UEngine::Serialize(FArchive& Ar)
 	if (Ar.IsCountingMemory())
 	{
 		// Only use the main audio device when counting memory
-		if (FAudioDevice* AudioDevice = GetMainAudioDevice())
+		if (FAudioDevice* AudioDevice = FAudioDeviceManager::GetMainDevice())
 		{
 			AudioDevice->CountBytes(Ar);
 		}
@@ -2757,9 +2744,9 @@ void UEngine::AddReferencedObjects(UObject* InThis, FReferenceCollector& Collect
 	UEngine* This = CastChecked<UEngine>(InThis);
 
 	// track objects in all the active audio devices
-	if (This->AudioDeviceManager)
+	if (FAudioDeviceManager* AudioDeviceManager = FAudioDeviceManager::Get())
 	{
-		This->AudioDeviceManager->AddReferencedObjects(Collector);
+		AudioDeviceManager->AddReferencedObjects(Collector);
 	}
 
 	// TODO: This is quite dangerous as FWorldContext::AddReferencedObjects could fail to be updated when something it
@@ -2847,53 +2834,38 @@ UFont* UEngine::GetAdditionalFont(int32 AdditionalFontIndex)
 	return GEngine->AdditionalFonts.IsValidIndex(AdditionalFontIndex) ? GEngine->AdditionalFonts[AdditionalFontIndex] : NULL;
 }
 
-class FAudioDeviceManager* UEngine::GetAudioDeviceManager()
+FAudioDeviceManager* UEngine::GetAudioDeviceManager()
 {
-	return AudioDeviceManager;
+	return FAudioDeviceManager::Get();
 }
 
 uint32 UEngine::GetAudioDeviceHandle() const
 {
-	return MainAudioDeviceHandle;
+	FAudioDeviceManager* AudioDeviceManager = FAudioDeviceManager::Get();
+	return AudioDeviceManager ? AudioDeviceManager->GetMainAudioDeviceHandle() : INDEX_NONE;
 }
 
 FAudioDevice* UEngine::GetMainAudioDevice()
 {
-	if (AudioDeviceManager != nullptr)
-	{
-		return AudioDeviceManager->GetAudioDevice(MainAudioDeviceHandle);
-	}
-	return nullptr;
+	return FAudioDeviceManager::GetMainDevice();
 }
 
 FAudioDevice* UEngine::GetActiveAudioDevice()
 {
-	if (AudioDeviceManager != nullptr)
-	{
-		return AudioDeviceManager->GetActiveAudioDevice();
-	}
-	return nullptr;
+	return FAudioDeviceManager::GetActiveDevice();
 }
 
 void UEngine::InitializeAudioDeviceManager()
 {
-	if (AudioDeviceManager == nullptr && bUseSound)
+	if (bUseSound)
 	{
-		AudioDeviceManager = new FAudioDeviceManager();
-		if (AudioDeviceManager->Initialize())
-		{
-			MainAudioDeviceHandle = AudioDeviceManager->GetMainAudioDeviceHandle();
-		}
-		else
-		{
-			ShutdownAudioDeviceManager();
-		}
+		FAudioDeviceManager::Initialize();
 	}
 }
 
 bool UEngine::UseSound() const
 {
-	return (bUseSound && AudioDeviceManager != nullptr);
+	return bUseSound && FAudioDeviceManager::Get() != nullptr;
 }
 /**
 * A fake stereo rendering device used to test stereo rendering without an attached device.
@@ -3833,15 +3805,9 @@ bool UEngine::Exec( UWorld* InWorld, const TCHAR* Cmd, FOutputDevice& Ar )
 		return true;
 	}
 
-	FAudioDevice* AudioDevice = nullptr;
-	if (InWorld)
-	{
-		AudioDevice = InWorld->GetAudioDevice();
-	}
-	else
-	{
-		AudioDevice = GetMainAudioDevice();
-	}
+	FAudioDevice* AudioDevice = InWorld
+		? InWorld->GetAudioDevice()
+		: FAudioDeviceManager::GetMainDevice();
 
 	if (AudioDevice && AudioDevice->Exec(InWorld, Cmd, Ar) == true)
 	{
diff --git a/Engine/Source/Runtime/Engine/Public/ActiveSound.h b/Engine/Source/Runtime/Engine/Public/ActiveSound.h
--- a/Engine/Source/Runtime/Engine/Public/ActiveSound.h
+++ b/Engine/Source/Runtime/Engine/Public/ActiveSound.h
@@ -6,8 +6,8 @@
 #include "WorldCollision.h"
 #include "Sound/SoundAttenuation.h"
 #include "HAL/ThreadSafeBool.h"
-#include "Audio.h"
 #include "Audio/AudioDebug.h"
+#include "AudioDefines.h"
 #include "AudioDynamicParameter.h"
 #include "Components/AudioComponent.h"
 #include "DSP/VolumeFader.h"
diff --git a/Engine/Source/Runtime/Engine/Public/Audio.h b/Engine/Source/Runtime/Engine/Public/Audio.h
--- a/Engine/Source/Runtime/Engine/Public/Audio.h
+++ b/Engine/Source/Runtime/Engine/Public/Audio.h
@@ -92,14 +92,6 @@ struct FWaveInstance;
 struct FReverbSettings;
 struct FSampleLoop;
 
-namespace Audio
-{
-	/**
-	 * Typed identifier for Audio Device Id
-	 */
-	using FDeviceId = uint32;
-}
-
 enum ELoopingMode
 {
 	/** One shot sound */
diff --git a/Engine/Source/Runtime/Engine/Public/Audio/AudioDebug.h b/Engine/Source/Runtime/Engine/Public/Audio/AudioDebug.h
--- a/Engine/Source/Runtime/Engine/Public/Audio/AudioDebug.h
+++ b/Engine/Source/Runtime/Engine/Public/Audio/AudioDebug.h
@@ -5,7 +5,6 @@
 #include "CoreMinimal.h"
 
 #include "AudioDefines.h"		// For ENABLE_AUDIO_DEBUG
-#include "Audio.h"
 
 #if ENABLE_AUDIO_DEBUG
 
diff --git a/Engine/Source/Runtime/Engine/Public/AudioDefines.h b/Engine/Source/Runtime/Engine/Public/AudioDefines.h
--- a/Engine/Source/Runtime/Engine/Public/AudioDefines.h
+++ b/Engine/Source/Runtime/Engine/Public/AudioDefines.h
@@ -52,4 +52,12 @@
 /**
  * Debugger is Available on non-shipping builds
  */
-#define ENABLE_AUDIO_DEBUG !UE_BUILD_SHIPPING
\ No newline at end of file
+#define ENABLE_AUDIO_DEBUG !UE_BUILD_SHIPPING
+
+namespace Audio
+{
+	/**
+	 * Typed identifier for Audio Device Id
+	 */
+	using FDeviceId = uint32;
+}
\ No newline at end of file
diff --git a/Engine/Source/Runtime/Engine/Public/AudioDevice.h b/Engine/Source/Runtime/Engine/Public/AudioDevice.h
--- a/Engine/Source/Runtime/Engine/Public/AudioDevice.h
+++ b/Engine/Source/Runtime/Engine/Public/AudioDevice.h
@@ -1,15 +1,17 @@
 // Copyright Epic Games, Inc. All Rights Reserved.
 #pragma once
 
-#include "Audio.h"
+#include "AudioDefines.h"
 #include "AudioDeviceManager.h"
+#include "AudioDynamicParameter.h"
+#include "AudioMixer.h"
+#include "AudioVirtualLoop.h"
 #include "Components/AudioComponent.h"
 #include "CoreMinimal.h"
 #include "DSP/SpectrumAnalyzer.h"
 #include "Engine/Engine.h"
 #include "EngineGlobals.h"
 #include "IAudioExtensionPlugin.h"
-#include "AudioDynamicParameter.h"
 #include "Sound/AudioSettings.h"
 #include "Sound/AudioVolume.h"
 #include "Sound/SoundAttenuation.h"
@@ -18,8 +20,6 @@
 #include "Sound/SoundMix.h"
 #include "Sound/SoundSubmixSend.h"
 #include "Sound/SoundSourceBus.h"
-#include "AudioVirtualLoop.h"
-#include "AudioMixer.h"
 
 /**
  * Forward declares
@@ -382,23 +382,6 @@ public:
 	virtual void OnDefaultDeviceChanged() = 0;
 };
 
-/** Abstract interface for receiving audio data from a given submix. */
-class ENGINE_API ISubmixBufferListener
-{
-public:
-	/**
-	Called when a new buffer has been rendered for a given submix
-	@param OwningSubmix	The submix object which has rendered a new buffer
-	@param AudioData		Ptr to the audio buffer
-	@param NumSamples		The number of audio samples in the audio buffer
-	@param NumChannels		The number of channels of audio in the buffer (e.g. 2 for stereo, 6 for 5.1, etc)
-	@param SampleRate		The sample rate of the audio buffer
-	@param AudioClock		Double audio clock value, from start of audio rendering.
-	*/
-	virtual void OnNewSubmixBuffer(const USoundSubmix* OwningSubmix, float* AudioData, int32 NumSamples, int32 NumChannels, const int32 SampleRate, double AudioClock) = 0;
-};
-
-
 class ENGINE_API FAudioDevice : public FExec
 {
 public:
@@ -1131,28 +1114,6 @@ public:
 	{
 	}
 
-	/** Returns the main audio device of the engine */
-	static FAudioDevice* GetMainAudioDevice()
-	{
-		// Try to get GEngine's main audio device
-		FAudioDevice* AudioDevice = GEngine->GetMainAudioDevice();
-
-		// If we don't have a main audio device (maybe we're running in a non-standard mode like a commandlet)
-		if (!AudioDevice)
-		{
-			// We should have an active device for device manager
-			FAudioDeviceManager* DeviceManager = GEngine->GetAudioDeviceManager();
-			return DeviceManager->GetActiveAudioDevice();
-		}
-		return AudioDevice;
-	}
-
-	/** Returns the audio device manager */
-	static FAudioDeviceManager* GetAudioDeviceManager()
-	{
-		return GEngine->GetAudioDeviceManager();
-	}
-
 	/** Low pass filter OneOverQ value */
 	float GetLowPassFilterResonance() const;
 
@@ -1194,7 +1155,7 @@ public:
 
 	static bool IsOcclusionPluginLoaded()
 	{
-		if (FAudioDevice* MainAudioDevice = GEngine->GetMainAudioDevice())
+		if (FAudioDevice* MainAudioDevice = FAudioDeviceManager::GetMainDevice())
 		{
 			return MainAudioDevice->bOcclusionInterfaceEnabled;
 		}
@@ -1209,7 +1170,7 @@ public:
 
 	static bool IsReverbPluginLoaded()
 	{
-		if (FAudioDevice* MainAudioDevice = GEngine->GetMainAudioDevice())
+		if (FAudioDevice* MainAudioDevice = FAudioDeviceManager::GetMainDevice())
 		{
 			return MainAudioDevice->bReverbInterfaceEnabled;
 		}
@@ -1678,8 +1639,8 @@ public:
 
 	bool IsMainAudioDevice()
 	{
-		FAudioDevice* MainAudioDevice = GEngine->GetMainAudioDevice();
-		return (MainAudioDevice == nullptr || MainAudioDevice == this);
+		FAudioDevice* MainAudioDevice = FAudioDeviceManager::GetMainDevice();
+		return MainAudioDevice == nullptr || MainAudioDevice == this;
 	}
 
 	/** Set whether or not we force the use of attenuation for non-game worlds (as by default we only care about game worlds) */
diff --git a/Engine/Source/Runtime/Engine/Public/AudioDeviceManager.h b/Engine/Source/Runtime/Engine/Public/AudioDeviceManager.h
--- a/Engine/Source/Runtime/Engine/Public/AudioDeviceManager.h
+++ b/Engine/Source/Runtime/Engine/Public/AudioDeviceManager.h
@@ -51,11 +51,17 @@ public:
 	* Initialize the audio device manager.
 	* Return true if successfully initialized.
 	**/
-	bool Initialize();
+	static void Initialize();
 
 	/** Returns the handle to the main audio device. */
 	Audio::FDeviceId GetMainAudioDeviceHandle() const { return MainAudioDeviceHandle; }
 
+	/** Returns the active audio device. */
+	static FAudioDevice* GetActiveDevice();
+
+	/** Returns the main audio device. */
+	static FAudioDevice* GetMainDevice();
+
 	/** Returns true if we're currently using the audio mixer. */
 	bool IsUsingAudioMixer() const;
 
@@ -97,19 +103,22 @@ public:
 	*/
 	bool ShutdownAllAudioDevices();
 
+	/**
+	* Shuts down the AudioDeviceManager
+	*/
+	static void Shutdown();
+
 	/**
 	* Returns a ptr to the audio device associated with the handle. If the
 	* handle is invalid then a NULL device ptr will be returned.
 	*/
 	FAudioDevice* GetAudioDevice(Audio::FDeviceId Handle);
 
+	/** Returns a pointer to the AudioDeviceManager if it has been initialized */
 	static FAudioDeviceManager* Get();
 
-	/**
-	* Returns a ptr to the active audio device. If there is no active
-	* device then it will return the main audio device.
-	*/
-	FAudioDevice* GetActiveAudioDevice();
+	/** Returns reference to the AudioDeviceManager.  Asserts if it has not been initialized */
+	static FAudioDeviceManager& GetChecked();
 
 	/** Returns the current number of active audio devices. */
 	uint8 GetNumActiveAudioDevices() const;
@@ -253,9 +262,6 @@ private:
 	/** The audio mixer module name. This is the audio mixer module name to use. E.g. AudioMixerXAudio2 */
 	FString AudioMixerModuleName;
 
-	/** Handle to the main audio device. */
-	Audio::FDeviceId MainAudioDeviceHandle;
-
 	/** Count of the number of free slots in the audio device array. */
 	uint32 FreeIndicesSize;
 
@@ -283,6 +289,9 @@ private:
 	/** Which audio device is currently active */
 	Audio::FDeviceId ActiveAudioDeviceHandle;
 
+	/** The main audio device's handle */
+	Audio::FDeviceId MainAudioDeviceHandle;
+
 	/** Dynamic volume map */
 	TMap<TTuple<ESoundType, FName>, float> DynamicSoundVolumes;
 
@@ -298,6 +307,6 @@ private:
 	/** Whether or not we've toggled the audio mixer. */
 	bool bToggledAudioMixer;
 
-	/** Audio Fence to ensure that we don't allow the audio thread to drift never endingly behind. */
+	/** Audio Fence to ensure that we don't allow the audio thread to drift endlessly behind. */
 	FAudioCommandFence SyncFence;
 };
diff --git a/Engine/Source/Runtime/GameplayMediaEncoder/Private/GameplayMediaEncoder.cpp b/Engine/Source/Runtime/GameplayMediaEncoder/Private/GameplayMediaEncoder.cpp
--- a/Engine/Source/Runtime/GameplayMediaEncoder/Private/GameplayMediaEncoder.cpp
+++ b/Engine/Source/Runtime/GameplayMediaEncoder/Private/GameplayMediaEncoder.cpp
@@ -274,8 +274,7 @@ bool FGameplayMediaEncoder::Start()
 	// subscribe to engine delegates for audio output and back buffer
 	//
 
-	FAudioDevice* AudioDevice = GEngine->GetMainAudioDevice();
-	if (AudioDevice)
+	if (FAudioDevice* AudioDevice = FAudioDeviceManager::GetMainDevice())
 	{
 		bAudioFormatChecked = false;
 		AudioDevice->RegisterSubmixBufferListener(this);
@@ -298,8 +297,7 @@ void FGameplayMediaEncoder::Stop()
 
 	if (UGameEngine* GameEngine = Cast<UGameEngine>(GEngine))
 	{
-		FAudioDevice* AudioDevice = GameEngine->GetMainAudioDevice();
-		if (AudioDevice)
+		if (FAudioDevice* AudioDevice = FAudioDeviceManager::GetMainDevice())
 		{
 			AudioDevice->UnregisterSubmixBufferListener(this);
 		}
diff --git a/Engine/Source/Runtime/IOS/IOSAudio/Private/IOSAudioBuffer.cpp b/Engine/Source/Runtime/IOS/IOSAudio/Private/IOSAudioBuffer.cpp
--- a/Engine/Source/Runtime/IOS/IOSAudio/Private/IOSAudioBuffer.cpp
+++ b/Engine/Source/Runtime/IOS/IOSAudio/Private/IOSAudioBuffer.cpp
@@ -54,7 +54,7 @@ FIOSAudioSoundBuffer::FIOSAudioSoundBuffer(FIOSAudioDevice* InAudioDevice, USoun
     FMemory::Memzero(SampleData, BufferSize);
     
 	
-	FAudioDeviceManager* AudioDeviceManager = GEngine->GetAudioDeviceManager();
+	FAudioDeviceManager* AudioDeviceManager = FAudioDeviceManager::Get();
 	check(AudioDeviceManager != nullptr);
 
 	// There is no need to track this resource with the AudioDeviceManager since there is a one to one mapping between FIOSAudioSoundBuffer objects and FIOSAudioSoundSource objects
@@ -156,7 +156,7 @@ FIOSAudioSoundBuffer* FIOSAudioSoundBuffer::Init(FIOSAudioDevice* IOSAudioDevice
 		return NULL;
 	}
 
-	FAudioDeviceManager* AudioDeviceManager = GEngine->GetAudioDeviceManager();
+	FAudioDeviceManager* AudioDeviceManager = FAudioDeviceManager::Get();
 
 	FIOSAudioSoundBuffer *Buffer = NULL;
 	bool	bIsStreaming = false;
diff --git a/Engine/Source/Runtime/Launch/Private/Android/AndroidEventManager.cpp b/Engine/Source/Runtime/Launch/Private/Android/AndroidEventManager.cpp
--- a/Engine/Source/Runtime/Launch/Private/Android/AndroidEventManager.cpp
+++ b/Engine/Source/Runtime/Launch/Private/Android/AndroidEventManager.cpp
@@ -319,8 +319,7 @@ void FAppEventManager::PauseAudio()
 	bAudioPaused = true;
 	UE_LOG(LogTemp, Log, TEXT("Android pause audio"));
 
-	FAudioDevice* AudioDevice = GEngine->GetMainAudioDevice();
-	if (AudioDevice)
+	if (FAudioDevice* AudioDevice = FAudioDeviceManager::GetMainDevice())
 	{
 		if (AudioDevice->IsAudioMixerEnabled())
 		{
@@ -332,7 +331,7 @@ void FAppEventManager::PauseAudio()
 		}
 		else
 		{
-			GEngine->GetMainAudioDevice()->Suspend(false);
+			AudioDevice->Suspend(false);
 
 			// make sure the audio thread runs the pause request
 			FAudioCommandFence Fence;
@@ -354,8 +353,7 @@ void FAppEventManager::ResumeAudio()
 	bAudioPaused = false;
 	UE_LOG(LogTemp, Log, TEXT("Android resume audio"));
 
-	FAudioDevice* AudioDevice = GEngine->GetMainAudioDevice();
-	if (AudioDevice)
+	if (FAudioDevice* AudioDevice = FAudioDeviceManager::GetMainDevice())
 	{
 		if (AudioDevice->IsAudioMixerEnabled())
 		{
@@ -363,7 +361,7 @@ void FAppEventManager::ResumeAudio()
 		}
 		else
 		{
-			GEngine->GetMainAudioDevice()->Suspend(true);
+			AudioDevice->Suspend(true);
 		}
 	}
 }
diff --git a/Engine/Source/Runtime/Launch/Private/IOS/LaunchIOS.cpp b/Engine/Source/Runtime/Launch/Private/IOS/LaunchIOS.cpp
--- a/Engine/Source/Runtime/Launch/Private/IOS/LaunchIOS.cpp
+++ b/Engine/Source/Runtime/Launch/Private/IOS/LaunchIOS.cpp
@@ -60,9 +60,9 @@ void FAppEntry::Suspend(bool bIsInterrupt)
 	// if background audio is active, then we don't want to do suspend any audio
 	if ([[IOSAppDelegate GetDelegate] IsFeatureActive:EAudioFeature::BackgroundAudio] == false)
 	{
-		if (GEngine && GEngine->GetMainAudioDevice() && !IsEngineExitRequested())
+		FAudioDevice* AudioDevice = FAudioDeviceManager::GetMainDevice();
+		if (AudioDevice && !IsEngineExitRequested())
 		{
-			FAudioDevice* AudioDevice = GEngine->GetMainAudioDevice();
 			if (bIsInterrupt && DisableAudioSuspendOnAudioInterruptCvar)
 			{
 				if (FTaskGraphInterface::IsRunning() && !IsEngineExitRequested())
@@ -71,14 +71,14 @@ void FAppEntry::Suspend(bool bIsInterrupt)
 					{
 						FAudioThread::RunCommandOnAudioThread([]()
 						{
-							if (GEngine && GEngine->GetMainAudioDevice())
+							if (FAudioDevice* MainDevice = FAudioDeviceManager::GetMainDevice())
 							{
-								GEngine->GetMainAudioDevice()->SetTransientMasterVolume(0.0f);
+								MainDevice->SetTransientMasterVolume(0.0f);
 							}
 						}, TStatId());
 					}, TStatId(), NULL, ENamedThreads::GameThread);
 				}
-				else
+				else if (AudioDevice)
 				{
 					AudioDevice->SetTransientMasterVolume(0.0f);
 				}
@@ -103,12 +103,12 @@ void FAppEntry::Suspend(bool bIsInterrupt)
 					{
 						FAudioThread::RunCommandOnAudioThread([]()
 						{
-							if (GEngine && GEngine->GetMainAudioDevice())
+							if (FAudioDevice* MainDevice = FAudioDeviceManager::GetMainDevice())
 							{
-								GEngine->GetMainAudioDevice()->SuspendContext();
+								MainDevice->SuspendContext();
 							}
 						}, TStatId());
-                
+
 						FAudioCommandFence AudioCommandFence;
 						AudioCommandFence.BeginFence();
 						AudioCommandFence.Wait();
@@ -129,7 +129,7 @@ void FAppEntry::Suspend(bool bIsInterrupt)
 					}
 					FEmbeddedCommunication::AllowSleep(TEXT("Background"));
 				}
-				else
+				else if (AudioDevice)
 				{
 					AudioDevice->SuspendContext();
 				}
@@ -137,8 +137,8 @@ void FAppEntry::Suspend(bool bIsInterrupt)
 		}
 		else
 		{
-            // Increment
-            IncrementAudioSuspendCounters();
+			// Increment
+			IncrementAudioSuspendCounters();
 		}
 	}
 }
@@ -154,10 +154,8 @@ void FAppEntry::Resume(bool bIsInterrupt)
 	// @todo: should this check if we were suspended, in case this changes while in the background? (suspend with background off, but resume with background audio on? is that a thing?)
 	if ([[IOSAppDelegate GetDelegate] IsFeatureActive:EAudioFeature::BackgroundAudio] == false)
 	{
-		if (GEngine && GEngine->GetMainAudioDevice())
+		if (FAudioDevice* AudioDevice = FAudioDeviceManager::GetMainDevice())
 		{
-			FAudioDevice* AudioDevice = GEngine->GetMainAudioDevice();
-        
 			if (bIsInterrupt && DisableAudioSuspendOnAudioInterruptCvar)
 			{
 				if (FTaskGraphInterface::IsRunning())
@@ -166,9 +164,9 @@ void FAppEntry::Resume(bool bIsInterrupt)
 					{
 						FAudioThread::RunCommandOnAudioThread([]()
 						{
-							if (GEngine && GEngine->GetMainAudioDevice())
+							if (FAudioDevice* AudioDevice = FAudioDeviceManager::GetMainDevice())
 							{
-								GEngine->GetMainAudioDevice()->SetTransientMasterVolume(1.0f);
+								AudioDevice->SetTransientMasterVolume(1.0f);
 							}
 						}, TStatId());
 					}, TStatId(), NULL, ENamedThreads::GameThread);
@@ -194,8 +192,8 @@ void FAppEntry::Resume(bool bIsInterrupt)
 		}
 		else
 		{
-            // Decrement
-            DecrementAudioSuspendCounters();
+			// Decrement
+			DecrementAudioSuspendCounters();
 		}
 	}
 }
@@ -203,50 +201,44 @@ void FAppEntry::Resume(bool bIsInterrupt)
 
 void FAppEntry::ResumeAudioContext()
 {
-	if (GEngine && GEngine->GetMainAudioDevice())
+	if (FAudioDevice* AudioDevice = FAudioDeviceManager::GetMainDevice())
 	{
-		FAudioDevice* AudioDevice = GEngine->GetMainAudioDevice();
-		if (AudioDevice)
+		if (FTaskGraphInterface::IsRunning())
 		{
-			if (FTaskGraphInterface::IsRunning())
+			FFunctionGraphTask::CreateAndDispatchWhenReady([]()
 			{
-				FFunctionGraphTask::CreateAndDispatchWhenReady([]()
+				FAudioThread::RunCommandOnAudioThread([]()
 				{
-					FAudioThread::RunCommandOnAudioThread([]()
+					if (FAudioDevice* MainDevice = FAudioDeviceManager::GetMainDevice())
 					{
-						if (GEngine && GEngine->GetMainAudioDevice())
-						{
-							GEngine->GetMainAudioDevice()->ResumeContext();
-						}
-					}, TStatId());
-				}, TStatId(), NULL, ENamedThreads::GameThread);
-			}
-			else
-			{
-				AudioDevice->ResumeContext();
-			}
+						MainDevice->ResumeContext();
+					}
+				}, TStatId());
+			}, TStatId(), NULL, ENamedThreads::GameThread);
+		}
+		else
+		{
+			AudioDevice->ResumeContext();
 		}
 	}
 }
 
 void FAppEntry::RestartAudio()
 {
-	if (GEngine && GEngine->GetMainAudioDevice())
+	if (FAudioDevice* AudioDevice = FAudioDeviceManager::GetMainDevice())
 	{
-		FAudioDevice* AudioDevice = GEngine->GetMainAudioDevice();
-
 		if (FTaskGraphInterface::IsRunning())
 		{
-            //increment the counter, otherwise ResumeContext won't work
-            IncrementAudioSuspendCounters();
+			//increment the counter, otherwise ResumeContext won't work
+			IncrementAudioSuspendCounters();
 
 			FFunctionGraphTask::CreateAndDispatchWhenReady([]()
 			{
 				FAudioThread::RunCommandOnAudioThread([]()
 				{
-					if (GEngine && GEngine->GetMainAudioDevice())
+					if (FAudioDevice* MainDevice = FAudioDeviceManager::GetMainDevice())
 					{
-						GEngine->GetMainAudioDevice()->ResumeContext();
+						MainDevice->ResumeContext();
 					}
 				}, TStatId());
 			}, TStatId(), NULL, ENamedThreads::GameThread);
diff --git a/Engine/Source/Runtime/Mac/CoreAudio/Private/CoreAudioBuffer.cpp b/Engine/Source/Runtime/Mac/CoreAudio/Private/CoreAudioBuffer.cpp
--- a/Engine/Source/Runtime/Mac/CoreAudio/Private/CoreAudioBuffer.cpp
+++ b/Engine/Source/Runtime/Mac/CoreAudio/Private/CoreAudioBuffer.cpp
@@ -222,7 +222,7 @@ FCoreAudioSoundBuffer* FCoreAudioSoundBuffer::CreateProceduralBuffer( FCoreAudio
  */
 FCoreAudioSoundBuffer* FCoreAudioSoundBuffer::CreatePreviewBuffer( FCoreAudioDevice* CoreAudioDevice, USoundWave* Wave, FCoreAudioSoundBuffer* Buffer )
 {
-    FAudioDeviceManager* AudioDeviceManager = GEngine->GetAudioDeviceManager();
+    FAudioDeviceManager* AudioDeviceManager = FAudioDeviceManager::Get();
     check(AudioDeviceManager != nullptr);
     
 	if (Buffer)
@@ -272,7 +272,7 @@ FCoreAudioSoundBuffer* FCoreAudioSoundBuffer::CreateNativeBuffer( FCoreAudioDevi
 	// Keep track of associated resource name.
 	Buffer->InitAudioStreamBasicDescription( kAudioFormatLinearPCM, Wave, true );
 	
-    FAudioDeviceManager* AudioDeviceManager = GEngine->GetAudioDeviceManager();
+    FAudioDeviceManager* AudioDeviceManager = FAudioDeviceManager::Get();
     check(AudioDeviceManager != nullptr);
     
 	if (UseResourceTrackingCVar)
@@ -335,7 +335,7 @@ FCoreAudioSoundBuffer* FCoreAudioSoundBuffer::Init( FAudioDevice* AudioDevice, U
 		return NULL;
 	}
 
-    FAudioDeviceManager* AudioDeviceManager = GEngine->GetAudioDeviceManager();
+    FAudioDeviceManager* AudioDeviceManager = FAudioDeviceManager::Get();
     
 	FCoreAudioDevice *CoreAudioDevice = ( FCoreAudioDevice *)AudioDevice;
 	FCoreAudioSoundBuffer *Buffer = NULL;
diff --git a/Engine/Source/Runtime/MovieScene/Private/MovieSceneTimeController.cpp b/Engine/Source/Runtime/MovieScene/Private/MovieSceneTimeController.cpp
--- a/Engine/Source/Runtime/MovieScene/Private/MovieSceneTimeController.cpp
+++ b/Engine/Source/Runtime/MovieScene/Private/MovieSceneTimeController.cpp
@@ -104,7 +104,7 @@ double FMovieSceneTimeController_PlatformClock::GetCurrentTime() const
 
 double FMovieSceneTimeController_AudioClock::GetCurrentTime() const
 {
-	FAudioDevice* AudioDevice = GEngine ? GEngine->GetMainAudioDevice() : nullptr;
+	FAudioDevice* AudioDevice = FAudioDeviceManager::GetMainDevice();
 	return AudioDevice ? AudioDevice->GetAudioClock() : FPlatformTime::Seconds();
 }
 
diff --git a/Engine/Source/Runtime/MovieSceneCapture/Public/Protocols/AudioCaptureProtocol.h b/Engine/Source/Runtime/MovieSceneCapture/Public/Protocols/AudioCaptureProtocol.h
--- a/Engine/Source/Runtime/MovieSceneCapture/Public/Protocols/AudioCaptureProtocol.h
+++ b/Engine/Source/Runtime/MovieSceneCapture/Public/Protocols/AudioCaptureProtocol.h
@@ -84,7 +84,7 @@ public:
 
 		// Stop all sounds currently playing on the Audio Device. This helps kill looping or long audio clips. When the sequence evaluates again,
 		// these clips will resume play at the correct location.
-		if (FAudioDevice* AudioDevice = GEngine->GetActiveAudioDevice())
+		if (FAudioDevice* AudioDevice = FAudioDeviceManager::GetActiveDevice())
 		{
 			AudioDevice->StopAllSounds(true);
 		}
diff --git a/Engine/Source/Runtime/Windows/XAudio2/Private/XAudio2Buffer.cpp b/Engine/Source/Runtime/Windows/XAudio2/Private/XAudio2Buffer.cpp
--- a/Engine/Source/Runtime/Windows/XAudio2/Private/XAudio2Buffer.cpp
+++ b/Engine/Source/Runtime/Windows/XAudio2/Private/XAudio2Buffer.cpp
@@ -416,18 +416,15 @@ FXAudio2SoundBuffer* FXAudio2SoundBuffer::CreateProceduralBuffer( FXAudio2Device
 	return( Buffer );
 }
 
-FXAudio2SoundBuffer* FXAudio2SoundBuffer::CreatePreviewBuffer( FXAudio2Device* XAudio2Device, USoundWave* Wave, FXAudio2SoundBuffer* Buffer )
+FXAudio2SoundBuffer* FXAudio2SoundBuffer::CreatePreviewBuffer(FXAudio2Device* XAudio2Device, USoundWave* Wave, FXAudio2SoundBuffer* Buffer)
 {
-	FAudioDeviceManager* AudioDeviceManager = GEngine->GetAudioDeviceManager();
-	check(AudioDeviceManager != nullptr);
-
 	if (Buffer)
 	{
-		AudioDeviceManager->FreeBufferResource(Buffer);
+		FAudioDeviceManager::GetChecked().FreeBufferResource(static_cast<FSoundBuffer*>(Buffer));
 	}
 
 	// Create new buffer.
-	Buffer = new FXAudio2SoundBuffer( XAudio2Device, SoundFormat_PCMPreview );
+	Buffer = new FXAudio2SoundBuffer(XAudio2Device, SoundFormat_PCMPreview);
 
 	check(!Wave->RawPCMData || Wave->RawPCMDataSize);
 
@@ -440,11 +437,11 @@ FXAudio2SoundBuffer* FXAudio2SoundBuffer::CreatePreviewBuffer( FXAudio2Device* X
 	// Copy over whether this data should be freed on delete
 	Buffer->bDynamicResource = Wave->bDynamicResource;
 
-	Buffer->InitWaveFormatEx( WAVE_FORMAT_PCM, Wave, true );
+	Buffer->InitWaveFormatEx(WAVE_FORMAT_PCM, Wave, true);
 
-	AudioDeviceManager->TrackResource(Wave, Buffer);
+	FAudioDeviceManager::GetChecked().TrackResource(Wave, static_cast<FSoundBuffer*>(Buffer));
 
-	return( Buffer );
+	return(Buffer);
 }
 
 FXAudio2SoundBuffer* FXAudio2SoundBuffer::CreateNativeBuffer( FXAudio2Device* XAudio2Device, USoundWave* Wave )
@@ -465,7 +462,7 @@ FXAudio2SoundBuffer* FXAudio2SoundBuffer::CreateNativeBuffer( FXAudio2Device* XA
 	// Keep track of associated resource name.
 	Buffer->InitWaveFormatEx( WAVE_FORMAT_PCM, Wave, true );
 
-	FAudioDeviceManager* AudioDeviceManager = GEngine->GetAudioDeviceManager();
+	FAudioDeviceManager* AudioDeviceManager = FAudioDeviceManager::Get();
 	check(AudioDeviceManager != nullptr);
 	AudioDeviceManager->TrackResource(Wave, Buffer);
 
@@ -530,7 +527,7 @@ FXAudio2SoundBuffer* FXAudio2SoundBuffer::Init( FAudioDevice* AudioDevice, USoun
 	Wave->InvalidateSoundWaveIfNeccessary();
 #endif // WITH_EDITOR
 
-	FAudioDeviceManager* AudioDeviceManager = GEngine->GetAudioDeviceManager();
+	FAudioDeviceManager* AudioDeviceManager = FAudioDeviceManager::Get();
 
 	FXAudio2Device* XAudio2Device = ( FXAudio2Device* )AudioDevice;
 	FXAudio2SoundBuffer* Buffer = NULL;
