#ue4 - Fix FPhysInterface_Chaos::GetSquaredDistanceToBody() not working when query point is inside convex shapes.

diff --git a/Engine/Source/Runtime/Engine/Private/PhysicsEngine/Experimental/PhysInterface_Chaos.cpp b/Engine/Source/Runtime/Engine/Private/PhysicsEngine/Experimental/PhysInterface_Chaos.cpp
--- a/Engine/Source/Runtime/Engine/Private/PhysicsEngine/Experimental/PhysInterface_Chaos.cpp
+++ b/Engine/Source/Runtime/Engine/Private/PhysicsEngine/Experimental/PhysInterface_Chaos.cpp
@@ -1958,6 +1958,11 @@ bool FPhysInterface_Chaos::GetSquaredDistanceToBody(const FBodyInstance* InInsta
 			const float Phi = Shape.Shape->Geometry->PhiWithNormal(LocalPoint, Normal);
 			if (Phi <= 0)
 			{
+				OutDistanceSquared = 0;
+				if (OutOptPointOnBody)
+				{
+					*OutOptPointOnBody = InPoint;
+				}
 				break;
 			}
 			else if (Phi < MinPhi)
