Fixing memory problems with the shader compiler jobs being either leaked or prematurely deleted by using TShadredRefs instead of raw pointers

diff --git a/Engine/Plugins/Compositing/OpenColorIO/Source/OpenColorIO/Private/OpenColorIOShaderCompilationManager.cpp b/Engine/Plugins/Compositing/OpenColorIO/Source/OpenColorIO/Private/OpenColorIOShaderCompilationManager.cpp
--- a/Engine/Plugins/Compositing/OpenColorIO/Source/OpenColorIO/Private/OpenColorIOShaderCompilationManager.cpp
+++ b/Engine/Plugins/Compositing/OpenColorIO/Source/OpenColorIO/Private/OpenColorIOShaderCompilationManager.cpp
@@ -97,7 +97,7 @@ void FOpenColorIOShaderCompilationManager::RunCompileJobs()
 		{
 			for (int32 JobIndex = 0; JobIndex < CurrentWorkerInfo.QueuedJobs.Num(); JobIndex++)
 			{
-				FShaderCompileJob& CurrentJob = *((FShaderCompileJob*)(CurrentWorkerInfo.QueuedJobs[JobIndex]));
+				FShaderCompileJob& CurrentJob = StaticCastSharedRef<FShaderCompileJob>(CurrentWorkerInfo.QueuedJobs[JobIndex]).Get();
 
 				check(!CurrentJob.bFinalized);
 				CurrentJob.bFinalized = true;
@@ -186,10 +186,10 @@ void FOpenColorIOShaderCompilationManager::InitWorkerInfo()
 	}	
 }
 
-OPENCOLORIO_API void FOpenColorIOShaderCompilationManager::AddJobs(TArray<FShaderCommonCompileJob*> InNewJobs)
+OPENCOLORIO_API void FOpenColorIOShaderCompilationManager::AddJobs(TArray<TSharedRef<FShaderCommonCompileJob, ESPMode::ThreadSafe>> InNewJobs)
 {
 #if WITH_EDITOR
-	for (FShaderCommonCompileJob *Job : InNewJobs)
+	for (TSharedRef<FShaderCommonCompileJob, ESPMode::ThreadSafe>& Job : InNewJobs)
 	{
 		FOpenColorIOShaderMapCompileResults& ShaderMapInfo = OpenColorIOShaderMapJobs.FindOrAdd(Job->Id);
 		//@todo : Apply shader map isn't used for now with this compile manager. Should be merged to have a generic shader compiler
@@ -265,7 +265,7 @@ void FOpenColorIOShaderCompilationManager::ProcessCompiledOpenColorIOShaderMaps(
 		{
 			TArray<FString> Errors;
 			FOpenColorIOShaderMapFinalizeResults& CompileResults = ProcessIt.Value();
-			const TArray<FShaderCommonCompileJob*>& ResultArray = CompileResults.FinishedJobs;
+			TArray<TSharedRef<FShaderCommonCompileJob, ESPMode::ThreadSafe>>& ResultArray = CompileResults.FinishedJobs;
 
 			// Make a copy of the array as this entry of FOpenColorIOShaderMap::ShaderMapsBeingCompiled will be removed below
 			TArray<FOpenColorIOTransformResource*> ColorTransformArray = *ColorTransforms;
@@ -273,7 +273,7 @@ void FOpenColorIOShaderCompilationManager::ProcessCompiledOpenColorIOShaderMaps(
 
 			for (int32 JobIndex = 0; JobIndex < ResultArray.Num(); JobIndex++)
 			{
-				FShaderCompileJob& CurrentJob = *((FShaderCompileJob*)(ResultArray[JobIndex]));
+				FShaderCompileJob& CurrentJob = StaticCastSharedRef<FShaderCompileJob>(ResultArray[JobIndex]).Get();
 				bSuccess = bSuccess && CurrentJob.bSucceeded;
 
 				if (bSuccess)
@@ -390,11 +390,7 @@ void FOpenColorIOShaderCompilationManager::ProcessCompiledOpenColorIOShaderMaps(
 				}
 
 				// Cleanup shader jobs and compile tracking structures
-				for (int32 JobIndex = 0; JobIndex < ResultArray.Num(); JobIndex++)
-				{
-					delete ResultArray[JobIndex];
-				}
-
+				ResultArray.Empty();
 				CompiledShaderMaps.Remove(ShaderMap->GetCompilingId());
 			}
 
diff --git a/Engine/Plugins/Compositing/OpenColorIO/Source/OpenColorIO/Private/OpenColorIOShaderMap.cpp b/Engine/Plugins/Compositing/OpenColorIO/Source/OpenColorIO/Private/OpenColorIOShaderMap.cpp
--- a/Engine/Plugins/Compositing/OpenColorIO/Source/OpenColorIO/Private/OpenColorIOShaderMap.cpp
+++ b/Engine/Plugins/Compositing/OpenColorIO/Source/OpenColorIO/Private/OpenColorIOShaderMap.cpp
@@ -177,7 +177,7 @@ FShaderCompileJob* FOpenColorIOShaderType::BeginCompileShader(
 	const FOpenColorIOTransformResource* InColorTransform,
 	FShaderCompilerEnvironment* InCompilationEnvironment,
 	EShaderPlatform InPlatform,
-	TArray<FShaderCommonCompileJob*>& OutNewJobs,
+	TArray<TSharedRef<FShaderCommonCompileJob, ESPMode::ThreadSafe>>& OutNewJobs,
 	FShaderTarget InTarget
 	)
 {
@@ -210,7 +210,7 @@ FShaderCompileJob* FOpenColorIOShaderType::BeginCompileShader(
 		TEXT("/Plugin/OpenColorIO/Private/OpenColorIOShader.usf"),
 		TEXT("MainPS"),
 		FShaderTarget(GetFrequency(), InPlatform),
-		NewJob,
+		TSharedRef<FShaderCommonCompileJob, ESPMode::ThreadSafe>(NewJob),
 		OutNewJobs
 	);
 
@@ -403,7 +403,7 @@ void FOpenColorIOShaderMap::Compile(FOpenColorIOTransformResource* InColorTransf
 			Platform = InPlatform;
 
 			uint32 NumShaders = 0;
-			TArray<FShaderCommonCompileJob*> NewJobs;
+			TArray<TSharedRef<FShaderCommonCompileJob, ESPMode::ThreadSafe>> NewJobs;
 	
 			// Iterate over all shader types.
 			TMap<FShaderType*, FShaderCompileJob*> SharedShaderJobs;
@@ -469,10 +469,8 @@ void FOpenColorIOShaderMap::Compile(FOpenColorIOTransformResource* InColorTransf
 	}
 }
 
-FShader* FOpenColorIOShaderMap::ProcessCompilationResultsForSingleJob(FShaderCommonCompileJob* InSingleJob, const FSHAHash& InShaderMapHash)
+FShader* FOpenColorIOShaderMap::ProcessCompilationResultsForSingleJob(FShaderCompileJob& CurrentJob, const FSHAHash& InShaderMapHash)
 {
-	check(InSingleJob);
-	const FShaderCompileJob& CurrentJob = *((FShaderCompileJob*)InSingleJob);
 	check(CurrentJob.Id == CompilingId);
 
 	FShader* Shader = nullptr;
@@ -490,7 +488,7 @@ FShader* FOpenColorIOShaderMap::ProcessCompilationResultsForSingleJob(FShaderCom
 	return Shader;
 }
 
-bool FOpenColorIOShaderMap::ProcessCompilationResults(const TArray<FShaderCommonCompileJob*>& InCompilationResults, int32& InOutJobIndex, float& InOutTimeBudget)
+bool FOpenColorIOShaderMap::ProcessCompilationResults(const TArray<TSharedRef<FShaderCommonCompileJob, ESPMode::ThreadSafe>>& InCompilationResults, int32& InOutJobIndex, float& InOutTimeBudget)
 {
 	check(InOutJobIndex < InCompilationResults.Num());
 
@@ -501,12 +499,7 @@ bool FOpenColorIOShaderMap::ProcessCompilationResults(const TArray<FShaderCommon
 
 	do
 	{
-		FShaderCommonCompileJob* SingleJob = InCompilationResults[InOutJobIndex];
-		ensure(SingleJob);
-
-		{
-			ProcessCompilationResultsForSingleJob(SingleJob, ShaderMapHash);
-		}
+		ProcessCompilationResultsForSingleJob(StaticCastSharedRef<FShaderCompileJob>(InCompilationResults[InOutJobIndex]).Get(), ShaderMapHash);
 
 		InOutJobIndex++;
 		
diff --git a/Engine/Plugins/Compositing/OpenColorIO/Source/OpenColorIO/Public/OpenColorIOShaderCompilationManager.h b/Engine/Plugins/Compositing/OpenColorIO/Source/OpenColorIO/Public/OpenColorIOShaderCompilationManager.h
--- a/Engine/Plugins/Compositing/OpenColorIO/Source/OpenColorIO/Public/OpenColorIOShaderCompilationManager.h
+++ b/Engine/Plugins/Compositing/OpenColorIO/Source/OpenColorIO/Public/OpenColorIOShaderCompilationManager.h
@@ -3,9 +3,7 @@
 #pragma once
 
 #include "OpenColorIOShared.h"
-
-
-class FShaderCommonCompileJob;
+#include "ShaderCompiler.h"
 
 
 /** Information tracked for each shader compile worker process instance. */
@@ -27,7 +25,7 @@ struct FOpenColorIOShaderCompileWorkerInfo
 	double StartTime;
 
 	/** Jobs that this worker is responsible for compiling. */
-	TArray<FShaderCommonCompileJob*> QueuedJobs;
+	TArray<TSharedRef<FShaderCommonCompileJob, ESPMode::ThreadSafe>> QueuedJobs;
 
 	FOpenColorIOShaderCompileWorkerInfo() :
 		bIssuedTasksToWorker(false),
@@ -61,7 +59,7 @@ struct FOpenColorIOShaderMapCompileResults
 	int32 NumJobsQueued;
 	bool bAllJobsSucceeded;
 	bool bRecreateComponentRenderStateOnCompletion;
-	TArray<FShaderCommonCompileJob*> FinishedJobs;
+	TArray<TSharedRef<FShaderCommonCompileJob, ESPMode::ThreadSafe>> FinishedJobs;
 };
 
 
@@ -87,7 +85,7 @@ public:
 	~FOpenColorIOShaderCompilationManager();
 
 	OPENCOLORIO_API void Tick(float DeltaSeconds = 0.0f);
-	OPENCOLORIO_API void AddJobs(TArray<FShaderCommonCompileJob*> InNewJobs);
+	OPENCOLORIO_API void AddJobs(TArray<TSharedRef<FShaderCommonCompileJob, ESPMode::ThreadSafe>> InNewJobs);
 	OPENCOLORIO_API void ProcessAsyncResults();
 
 	void FinishCompilation(const TCHAR* InTransformName, const TArray<int32>& ShaderMapIdsToFinishCompiling);
@@ -98,7 +96,7 @@ private:
 
 	void InitWorkerInfo();
 
-	TArray<FShaderCommonCompileJob*> JobQueue;
+	TArray<TSharedRef<FShaderCommonCompileJob, ESPMode::ThreadSafe>> JobQueue;
 
 	/** Map from shader map Id to the compile results for that map, used to gather compiled results. */
 	TMap<int32, FOpenColorIOShaderMapCompileResults> OpenColorIOShaderMapJobs;
diff --git a/Engine/Plugins/Compositing/OpenColorIO/Source/OpenColorIO/Public/OpenColorIOShaderType.h b/Engine/Plugins/Compositing/OpenColorIO/Source/OpenColorIO/Public/OpenColorIOShaderType.h
--- a/Engine/Plugins/Compositing/OpenColorIO/Source/OpenColorIO/Public/OpenColorIOShaderType.h
+++ b/Engine/Plugins/Compositing/OpenColorIO/Source/OpenColorIO/Public/OpenColorIOShaderType.h
@@ -88,7 +88,7 @@ public:
 			const FOpenColorIOTransformResource* InColorTransform,
 			FShaderCompilerEnvironment* CompilationEnvironment,
 			EShaderPlatform Platform,
-			TArray<FShaderCommonCompileJob*>& NewJobs,
+			TArray<TSharedRef<FShaderCommonCompileJob, ESPMode::ThreadSafe>>& NewJobs,
 			FShaderTarget Target
 		);
 
diff --git a/Engine/Plugins/Compositing/OpenColorIO/Source/OpenColorIO/Public/OpenColorIOShared.h b/Engine/Plugins/Compositing/OpenColorIO/Source/OpenColorIO/Public/OpenColorIOShared.h
--- a/Engine/Plugins/Compositing/OpenColorIO/Source/OpenColorIO/Public/OpenColorIOShared.h
+++ b/Engine/Plugins/Compositing/OpenColorIO/Source/OpenColorIO/Public/OpenColorIOShared.h
@@ -18,6 +18,7 @@
 #include "Shader.h"
 #include "StaticParameterSet.h"
 #include "Templates/RefCounting.h"
+#include "ShaderCompiler.h"
 
 
 
@@ -143,7 +144,7 @@ public:
 		);
 
 	/** Sorts the incoming compiled jobs into the appropriate OCIO shader maps, and finalizes this shader map so that it can be used for rendering. */
-	bool ProcessCompilationResults(const TArray<class FShaderCommonCompileJob*>& InCompilationResults, int32& InOutResultIndex, float& InOutTimeBudget);
+	bool ProcessCompilationResults(const TArray<TSharedRef<FShaderCommonCompileJob, ESPMode::ThreadSafe>>& InCompilationResults, int32& InOutResultIndex, float& InOutTimeBudget);
 
 	/**
 	 * Checks whether the shader map is missing any shader types necessary for the given color transform.
@@ -274,7 +275,7 @@ private:
 	/** Indicates whether the shader map should be stored in the shader cache. */
 	uint32 bIsPersistent : 1;
 
-	FShader* ProcessCompilationResultsForSingleJob(class FShaderCommonCompileJob* InSingleJob, const FSHAHash& InShaderMapHash);
+	FShader* ProcessCompilationResultsForSingleJob(FShaderCompileJob& InSingleJob, const FSHAHash& InShaderMapHash);
 
 	bool IsOpenColorIOShaderComplete(const FOpenColorIOTransformResource* InColorTransform, const FOpenColorIOShaderType* InShaderType, bool bSilent);
 
diff --git a/Engine/Plugins/FX/Niagara/Source/NiagaraEditor/Private/NiagaraCompiler.cpp b/Engine/Plugins/FX/Niagara/Source/NiagaraEditor/Private/NiagaraCompiler.cpp
--- a/Engine/Plugins/FX/Niagara/Source/NiagaraEditor/Private/NiagaraCompiler.cpp
+++ b/Engine/Plugins/FX/Niagara/Source/NiagaraEditor/Private/NiagaraCompiler.cpp
@@ -878,15 +878,15 @@ int32 FHlslNiagaraCompiler::CompileScript(const FNiagaraCompileRequestData* InCo
 		}
 		if (NiagaraShaderType)
 		{
-			TArray<FShaderCommonCompileJob*> NewJobs;
-			CompilationJob->ShaderCompileJob = TRefCountPtr<FShaderCompileJob>(new FShaderCompileJob(JobID, nullptr, NiagaraShaderType, 0));
+			TArray<TSharedRef<FShaderCommonCompileJob, ESPMode::ThreadSafe>> NewJobs;
+			CompilationJob->ShaderCompileJob = MakeShared<FShaderCompileJob, ESPMode::ThreadSafe>(JobID, nullptr, NiagaraShaderType, 0);
 			Input.ShaderFormat = FName(TEXT("VVM_1_0"));
 			if (GNiagaraSkipVectorVMBackendOptimizations != 0)
 			{
 				Input.Environment.CompilerFlags.Add(CFLAG_SkipOptimizations);
 			}
 			CompilationJob->ShaderCompileJob->Input = Input;
-			NewJobs.Add(CompilationJob->ShaderCompileJob);
+			NewJobs.Add(StaticCastSharedPtr<FShaderCommonCompileJob>(CompilationJob->ShaderCompileJob).ToSharedRef());
 
 			GShaderCompilingManager->AddJobs(NewJobs, true, false, FString(), FString(), true);
 			bJobScheduled = true;
diff --git a/Engine/Plugins/FX/Niagara/Source/NiagaraEditor/Private/NiagaraCompiler.h b/Engine/Plugins/FX/Niagara/Source/NiagaraEditor/Private/NiagaraCompiler.h
--- a/Engine/Plugins/FX/Niagara/Source/NiagaraEditor/Private/NiagaraCompiler.h
+++ b/Engine/Plugins/FX/Niagara/Source/NiagaraEditor/Private/NiagaraCompiler.h
@@ -19,7 +19,7 @@ struct FNiagaraTranslatorOutput;
 
 struct FNiagaraCompilerJob
 {
-	TRefCountPtr<FShaderCompileJob> ShaderCompileJob;
+	TSharedPtr<FShaderCompileJob, ESPMode::ThreadSafe> ShaderCompileJob;
 	FNiagaraCompileResults CompileResults;
 	double StartTime;
 	FNiagaraTranslatorOutput TranslatorOutput;
diff --git a/Engine/Plugins/FX/Niagara/Source/NiagaraShader/Private/NiagaraShader.cpp b/Engine/Plugins/FX/Niagara/Source/NiagaraShader/Private/NiagaraShader.cpp
--- a/Engine/Plugins/FX/Niagara/Source/NiagaraShader/Private/NiagaraShader.cpp
+++ b/Engine/Plugins/FX/Niagara/Source/NiagaraShader/Private/NiagaraShader.cpp
@@ -315,12 +315,12 @@ TMap<const FShaderCompileJob*, TArray<FNiagaraDataInterfaceGPUParamInfo> > FNiag
  * Enqueues a compilation for a new shader of this type.
  * @param script - The script to link the shader with.
  */
-FShaderCompileJob* FNiagaraShaderType::BeginCompileShader(
+TSharedRef<FShaderCommonCompileJob, ESPMode::ThreadSafe> FNiagaraShaderType::BeginCompileShader(
 	uint32 ShaderMapId,
 	const FNiagaraShaderScript* Script,
 	FShaderCompilerEnvironment* CompilationEnvironment,
 	EShaderPlatform Platform,
-	TArray<FShaderCommonCompileJob*>& NewJobs,
+	TArray<TSharedRef<FShaderCommonCompileJob, ESPMode::ThreadSafe>>& NewJobs,
 	FShaderTarget Target,
 	TArray<FNiagaraDataInterfaceGPUParamInfo>& InDIParamInfo
 	)
@@ -354,6 +354,7 @@ FShaderCompileJob* FNiagaraShaderType::BeginCompileShader(
 	// Allow the shader type to modify the compile environment.
 	SetupCompileEnvironment(Platform, Script, ShaderEnvironment);
 
+	TSharedRef<FShaderCommonCompileJob, ESPMode::ThreadSafe> ShaderJob(NewJob);
 	FString DebugGroupName = Script->SourceName;
 	::GlobalBeginCompileShader(
 		DebugGroupName,
@@ -363,15 +364,12 @@ FShaderCompileJob* FNiagaraShaderType::BeginCompileShader(
 		TEXT("/Plugin/FX/Niagara/Private/NiagaraEmitterInstanceShader.usf"),
 		TEXT("SimulateMainComputeCS"),
 		FShaderTarget(GetFrequency(), Platform),
-		NewJob,
+		ShaderJob,
 		NewJobs
 	);
 
 	ExtraParamInfo.Add(NewJob, InDIParamInfo);
-
-	//NewJobs.Add(NewJob);
-
-	return NewJob;
+	return ShaderJob;
 }
 
 void FNiagaraShaderType::CacheUniformBufferIncludes(TMap<const TCHAR*, FCachedUniformBufferDeclaration>& Cache, EShaderPlatform Platform)
@@ -856,10 +854,10 @@ void FNiagaraShaderMap::Compile(
 			Platform = InPlatform;
 
 			uint32 NumShaders = 0;
-			TArray<FShaderCommonCompileJob*> NewJobs;
+			TArray<TSharedRef<FShaderCommonCompileJob, ESPMode::ThreadSafe>> NewJobs;
 	
 			// Iterate over all shader types.
-			TMap<FShaderType*, FShaderCompileJob*> SharedShaderJobs;
+			TMap<FShaderType*, TSharedRef<FShaderCommonCompileJob, ESPMode::ThreadSafe>> SharedShaderJobs;
 			for(TLinkedList<FShaderType*>::TIterator ShaderTypeIt(FShaderType::GetTypeList());ShaderTypeIt;ShaderTypeIt.Next())
 			{
 				FNiagaraShaderType* ShaderType = ShaderTypeIt->GetNiagaraShaderType();
@@ -871,7 +869,7 @@ void FNiagaraShaderMap::Compile(
 					// Only compile the shader if we don't already have it
 					if (!HasShader(ShaderType, /* PermutationId = */ 0))
 					{
-						auto* Job = ShaderType->BeginCompileShader(
+						TSharedRef<FShaderCommonCompileJob, ESPMode::ThreadSafe> Job = ShaderType->BeginCompileShader(
 							CompilingId,
 							Script,
 							CompilationEnvironment,
@@ -920,18 +918,17 @@ void FNiagaraShaderMap::Compile(
 	}
 }
 
-FShader* FNiagaraShaderMap::ProcessCompilationResultsForSingleJob(FShaderCommonCompileJob* SingleJob, const FSHAHash& ShaderMapHash)
+FShader* FNiagaraShaderMap::ProcessCompilationResultsForSingleJob(TSharedRef<FShaderCommonCompileJob, ESPMode::ThreadSafe> SingleJob, const FSHAHash& ShaderMapHash)
 {
-	check(SingleJob);
-	const FShaderCompileJob& CurrentJob = *((FShaderCompileJob*)SingleJob);
-	check(CurrentJob.Id == CompilingId);
+	TSharedRef<FShaderCompileJob, ESPMode::ThreadSafe> CurrentJob = StaticCastSharedRef<FShaderCompileJob>(SingleJob);
+	check(CurrentJob->Id == CompilingId);
 
 	FShader* Shader = nullptr;
 
-	FNiagaraShaderType* NiagaraShaderType = CurrentJob.ShaderType->GetNiagaraShaderType();
+	FNiagaraShaderType* NiagaraShaderType = CurrentJob->ShaderType->GetNiagaraShaderType();
 	check(NiagaraShaderType);
-	Shader = NiagaraShaderType->FinishCompileShader(ShaderMapHash, CurrentJob, FriendlyName);
-	bCompiledSuccessfully = CurrentJob.bSucceeded;
+	Shader = NiagaraShaderType->FinishCompileShader(ShaderMapHash, *CurrentJob, FriendlyName);
+	bCompiledSuccessfully = CurrentJob->bSucceeded;
 
 	FNiagaraShader *NiagaraShader = static_cast<FNiagaraShader*>(Shader);
 	// UE-67395 - we had a case where we polluted the DDC with a shader containing no bytecode.
@@ -942,7 +939,7 @@ FShader* FNiagaraShaderMap::ProcessCompilationResultsForSingleJob(FShaderCommonC
 	return Shader;
 }
 
-bool FNiagaraShaderMap::ProcessCompilationResults(const TArray<FShaderCommonCompileJob*>& InCompilationResults, int32& InOutJobIndex, float& TimeBudget)
+bool FNiagaraShaderMap::ProcessCompilationResults(const TArray<TSharedRef<FShaderCommonCompileJob, ESPMode::ThreadSafe>>& InCompilationResults, int32& InOutJobIndex, float& TimeBudget)
 {
 	check(InOutJobIndex < InCompilationResults.Num());
 
@@ -953,11 +950,8 @@ bool FNiagaraShaderMap::ProcessCompilationResults(const TArray<FShaderCommonComp
 
 	do
 	{
-		FShaderCommonCompileJob* SingleJob = InCompilationResults[InOutJobIndex];
-		ensure(SingleJob);
-
 		{
-			ProcessCompilationResultsForSingleJob(SingleJob, ShaderMapHash);
+			ProcessCompilationResultsForSingleJob(InCompilationResults[InOutJobIndex], ShaderMapHash);
 		}
 
 		InOutJobIndex++;
diff --git a/Engine/Plugins/FX/Niagara/Source/NiagaraShader/Private/NiagaraShaderCompilationManager.cpp b/Engine/Plugins/FX/Niagara/Source/NiagaraShader/Private/NiagaraShaderCompilationManager.cpp
--- a/Engine/Plugins/FX/Niagara/Source/NiagaraShader/Private/NiagaraShaderCompilationManager.cpp
+++ b/Engine/Plugins/FX/Niagara/Source/NiagaraShader/Private/NiagaraShaderCompilationManager.cpp
@@ -25,36 +25,36 @@ static FAutoConsoleVariableRef CVarShowNiagaraShaderWarnings(
 
 NIAGARASHADER_API FNiagaraShaderCompilationManager GNiagaraShaderCompilationManager;
 
-NIAGARASHADER_API void FNiagaraShaderCompilationManager::AddJobs(TArray<FShaderCommonCompileJob*> InNewJobs)
+NIAGARASHADER_API void FNiagaraShaderCompilationManager::AddJobs(TArray<TSharedRef<FShaderCommonCompileJob, ESPMode::ThreadSafe>> InNewJobs)
 {
 	check(IsInGameThread());
 	JobQueue.Append(InNewJobs);
 	
-	for (FShaderCommonCompileJob *Job : InNewJobs)
+	for (TSharedRef<FShaderCommonCompileJob, ESPMode::ThreadSafe> Job : InNewJobs)
 	{
 		FNiagaraShaderMapCompileResults& ShaderMapInfo = NiagaraShaderMapJobs.FindOrAdd(Job->Id);
 //		ShaderMapInfo.bApplyCompletedShaderMapForRendering = bApplyCompletedShaderMapForRendering;
 //		ShaderMapInfo.bRecreateComponentRenderStateOnCompletion = bRecreateComponentRenderStateOnCompletion;
 		ShaderMapInfo.NumJobsQueued++;
 
-		FShaderCompileJob& CurrentJob = *((FShaderCompileJob*) Job);
+		TSharedRef<FShaderCompileJob, ESPMode::ThreadSafe> CurrentJob = StaticCastSharedRef<FShaderCompileJob>(Job);
 
 		// Fast math breaks The ExecGrid layout script because floor(x/y) returns a bad value if x == y. Yay.
-		if (IsMetalPlatform((EShaderPlatform)CurrentJob.Input.Target.Platform))
+		if (IsMetalPlatform((EShaderPlatform)CurrentJob->Input.Target.Platform))
 		{
-			CurrentJob.Input.Environment.CompilerFlags.Add(CFLAG_NoFastMath);
+			CurrentJob->Input.Environment.CompilerFlags.Add(CFLAG_NoFastMath);
 		}
 
-		UE_LOG(LogNiagaraShaderCompiler, Verbose, TEXT("Adding niagara gpu shader compile job... %s"), *CurrentJob.Input.DebugGroupName);
+		UE_LOG(LogNiagaraShaderCompiler, Verbose, TEXT("Adding niagara gpu shader compile job... %s"), *CurrentJob->Input.DebugGroupName);
 
 		static ITargetPlatformManagerModule& TPM = GetTargetPlatformManagerRef();
-		const FName Format = LegacyShaderPlatformToShaderFormat(EShaderPlatform(CurrentJob.Input.Target.Platform));
+		const FName Format = LegacyShaderPlatformToShaderFormat(EShaderPlatform(CurrentJob->Input.Target.Platform));
 		FString AbsoluteDebugInfoDirectory = IFileManager::Get().ConvertToAbsolutePathForExternalAppForWrite(*(FPaths::ProjectSavedDir() / TEXT("ShaderDebugInfo")));
 		FPaths::NormalizeDirectoryName(AbsoluteDebugInfoDirectory);
-		CurrentJob.Input.DumpDebugInfoPath = AbsoluteDebugInfoDirectory / Format.ToString() / CurrentJob.Input.DebugGroupName;
-		if (!IFileManager::Get().DirectoryExists(*CurrentJob.Input.DumpDebugInfoPath))
+		CurrentJob->Input.DumpDebugInfoPath = AbsoluteDebugInfoDirectory / Format.ToString() / CurrentJob->Input.DebugGroupName;
+		if (!IFileManager::Get().DirectoryExists(*CurrentJob->Input.DumpDebugInfoPath))
 		{
-			verifyf(IFileManager::Get().MakeDirectory(*CurrentJob.Input.DumpDebugInfoPath, true), TEXT("Failed to create directory for shader debug info '%s'"), *CurrentJob.Input.DumpDebugInfoPath);
+			verifyf(IFileManager::Get().MakeDirectory(*CurrentJob->Input.DumpDebugInfoPath, true), TEXT("Failed to create directory for shader debug info '%s'"), *CurrentJob->Input.DumpDebugInfoPath);
 		}
 	}
 	GShaderCompilingManager->AddJobs(InNewJobs, true, false, FString(), FString(), true);
@@ -67,7 +67,7 @@ void FNiagaraShaderCompilationManager::ProcessAsyncResults()
 	TArray<int32> FinalizedShaderMapIDs;
 	for (int32 JobIndex = JobQueue.Num() - 1; JobIndex >= 0; JobIndex--)
 	{
-		FShaderCompileJob* Job = (FShaderCompileJob*)(JobQueue[JobIndex]);
+		TSharedRef<FShaderCompileJob, ESPMode::ThreadSafe> Job = StaticCastSharedRef<FShaderCompileJob>(JobQueue[JobIndex]);
 		if (Job->bFinalized)
 		{
 			FinalizedShaderMapIDs.Add(Job->Id);
@@ -80,26 +80,26 @@ void FNiagaraShaderCompilationManager::ProcessAsyncResults()
 	// Process the results from the shader compile worker
 	for (int32 JobIndex = JobQueue.Num() - 1; JobIndex >= 0; JobIndex--)
 	{
-		FShaderCompileJob& CurrentJob = *((FShaderCompileJob*)(JobQueue[JobIndex]));
+		TSharedRef<FShaderCompileJob, ESPMode::ThreadSafe> CurrentJob = StaticCastSharedRef<FShaderCompileJob>(JobQueue[JobIndex]);
 
-		if (!CurrentJob.bFinalized)
+		if (!CurrentJob->bFinalized)
 		{
 			continue;
 		}
 
-		CurrentJob.bSucceeded = CurrentJob.Output.bSucceeded;
-		if (CurrentJob.Output.bSucceeded)
+		CurrentJob->bSucceeded = CurrentJob->Output.bSucceeded;
+		if (CurrentJob->Output.bSucceeded)
 		{
-			UE_LOG(LogNiagaraShaderCompiler, Verbose, TEXT("GPU shader compile succeeded. Id %d"), CurrentJob.Id);
+			UE_LOG(LogNiagaraShaderCompiler, Verbose, TEXT("GPU shader compile succeeded. Id %d"), CurrentJob->Id);
 		}
 		else
 		{
-			UE_LOG(LogNiagaraShaderCompiler, Warning, TEXT("GPU shader compile failed! Id %d"), CurrentJob.Id);
+			UE_LOG(LogNiagaraShaderCompiler, Warning, TEXT("GPU shader compile failed! Id %d"), CurrentJob->Id);
 		}
 
-		FNiagaraShaderMapCompileResults& ShaderMapResults = NiagaraShaderMapJobs.FindChecked(CurrentJob.Id);
-		ShaderMapResults.FinishedJobs.Add(&CurrentJob);
-		ShaderMapResults.bAllJobsSucceeded = ShaderMapResults.bAllJobsSucceeded && CurrentJob.bSucceeded;
+		FNiagaraShaderMapCompileResults& ShaderMapResults = NiagaraShaderMapJobs.FindChecked(CurrentJob->Id);
+		ShaderMapResults.FinishedJobs.Add(JobQueue[JobIndex]);
+		ShaderMapResults.bAllJobsSucceeded = ShaderMapResults.bAllJobsSucceeded && CurrentJob->bSucceeded;
 		JobQueue.RemoveAt(JobIndex);
 	}
 
@@ -152,7 +152,7 @@ void FNiagaraShaderCompilationManager::ProcessCompiledNiagaraShaderMaps(TMap<int
 		{
 			TArray<FString> Errors;
 			FNiagaraShaderMapFinalizeResults& CompileResults = ProcessIt.Value();
-			const TArray<FShaderCommonCompileJob*>& ResultArray = CompileResults.FinishedJobs;
+			TArray<TSharedRef<FShaderCommonCompileJob, ESPMode::ThreadSafe>>& ResultArray = CompileResults.FinishedJobs;
 
 			// Make a copy of the array as this entry of FNiagaraShaderMap::ShaderMapsBeingCompiled will be removed below
 			TArray<FNiagaraShaderScript*> ScriptArray = *Scripts;
@@ -160,25 +160,25 @@ void FNiagaraShaderCompilationManager::ProcessCompiledNiagaraShaderMaps(TMap<int
 
 			for (int32 JobIndex = 0; JobIndex < ResultArray.Num(); JobIndex++)
 			{
-				FShaderCompileJob& CurrentJob = *((FShaderCompileJob*)(ResultArray[JobIndex]));
-				bSuccess = bSuccess && CurrentJob.bSucceeded;
+				TSharedRef<FShaderCompileJob, ESPMode::ThreadSafe> CurrentJob = StaticCastSharedRef<FShaderCompileJob>(ResultArray[JobIndex]);
+				bSuccess = bSuccess && CurrentJob->bSucceeded;
 
 				if (bSuccess)
 				{
-					check(CurrentJob.Output.ShaderCode.GetShaderCodeSize() > 0);
+					check(CurrentJob->Output.ShaderCode.GetShaderCodeSize() > 0);
 				}
 
-				if (GShowNiagaraShaderWarnings || !CurrentJob.bSucceeded)
+				if (GShowNiagaraShaderWarnings || !CurrentJob->bSucceeded)
 				{
-					for (int32 ErrorIndex = 0; ErrorIndex < CurrentJob.Output.Errors.Num(); ErrorIndex++)
+					for (int32 ErrorIndex = 0; ErrorIndex < CurrentJob->Output.Errors.Num(); ErrorIndex++)
 					{
-						Errors.AddUnique(CurrentJob.Output.Errors[ErrorIndex].GetErrorString().Replace(TEXT("Error"), TEXT("Err0r")));
+						Errors.AddUnique(CurrentJob->Output.Errors[ErrorIndex].GetErrorString().Replace(TEXT("Error"), TEXT("Err0r")));
 					}
 
-					if (CurrentJob.Output.Errors.Num())
+					if (CurrentJob->Output.Errors.Num())
 					{
-						UE_LOG(LogNiagaraShaderCompiler, Display, TEXT("There were issues for job \"%s\""), *CurrentJob.Input.DebugGroupName);
-						for (const FShaderCompilerError& Error : CurrentJob.Output.Errors)
+						UE_LOG(LogNiagaraShaderCompiler, Display, TEXT("There were issues for job \"%s\""), *CurrentJob->Input.DebugGroupName);
+						for (const FShaderCompilerError& Error : CurrentJob->Output.Errors)
 						{
 							UE_LOG(LogNiagaraShaderCompiler, Warning, TEXT("%s"), *Error.GetErrorString())
 						}
@@ -186,7 +186,7 @@ void FNiagaraShaderCompilationManager::ProcessCompiledNiagaraShaderMaps(TMap<int
 				}
 				else
 				{
-					UE_LOG(LogNiagaraShaderCompiler, Verbose, TEXT("Shader compile job \"%s\" completed."), *CurrentJob.Input.DebugGroupName);
+					UE_LOG(LogNiagaraShaderCompiler, Verbose, TEXT("Shader compile job \"%s\" completed."), *CurrentJob->Input.DebugGroupName);
 				}
 			}
 
@@ -277,11 +277,7 @@ void FNiagaraShaderCompilationManager::ProcessCompiledNiagaraShaderMaps(TMap<int
 				}
 
 				// Cleanup shader jobs and compile tracking structures
-				for (int32 JobIndex = 0; JobIndex < ResultArray.Num(); JobIndex++)
-				{
-					delete ResultArray[JobIndex];
-				}
-
+				ResultArray.Empty();
 				CompiledShaderMaps.Remove(ShaderMap->GetCompilingId());
 			}
 
diff --git a/Engine/Plugins/FX/Niagara/Source/NiagaraShader/Public/NiagaraShaderCompilationManager.h b/Engine/Plugins/FX/Niagara/Source/NiagaraShader/Public/NiagaraShaderCompilationManager.h
--- a/Engine/Plugins/FX/Niagara/Source/NiagaraShader/Public/NiagaraShaderCompilationManager.h
+++ b/Engine/Plugins/FX/Niagara/Source/NiagaraShader/Public/NiagaraShaderCompilationManager.h
@@ -3,6 +3,7 @@
 #pragma once
 
 #include "NiagaraShared.h"
+#include "ShaderCompiler.h"
 
 /** Results for a single compiled shader map. */
 struct FNiagaraShaderMapCompileResults
@@ -16,7 +17,7 @@ struct FNiagaraShaderMapCompileResults
 	int32 NumJobsQueued;
 	bool bAllJobsSucceeded;
 	bool bRecreateComponentRenderStateOnCompletion;
-	TArray<FShaderCommonCompileJob*> FinishedJobs;
+	TArray<TSharedRef<FShaderCommonCompileJob, ESPMode::ThreadSafe>> FinishedJobs;
 };
 
 
@@ -40,14 +41,14 @@ class FNiagaraShaderCompilationManager
 public:
 	FNiagaraShaderCompilationManager() = default;
 
-	NIAGARASHADER_API void AddJobs(TArray<class FShaderCommonCompileJob*> InNewJobs);
+	NIAGARASHADER_API void AddJobs(TArray<TSharedRef<FShaderCommonCompileJob, ESPMode::ThreadSafe>> InNewJobs);
 	NIAGARASHADER_API void ProcessAsyncResults();
 
 	void FinishCompilation(const TCHAR* ScriptName, const TArray<int32>& ShaderMapIdsToFinishCompiling);
 private:
 	void ProcessCompiledNiagaraShaderMaps(TMap<int32, FNiagaraShaderMapFinalizeResults>& CompiledShaderMaps, float TimeBudget);
 
-	TArray<class FShaderCommonCompileJob*> JobQueue;
+	TArray<TSharedRef<FShaderCommonCompileJob, ESPMode::ThreadSafe>> JobQueue;
 
 	/** Map from shader map Id to the compile results for that map, used to gather compiled results. */
 	TMap<int32, FNiagaraShaderMapCompileResults> NiagaraShaderMapJobs;
diff --git a/Engine/Plugins/FX/Niagara/Source/NiagaraShader/Public/NiagaraShaderType.h b/Engine/Plugins/FX/Niagara/Source/NiagaraShader/Public/NiagaraShaderType.h
--- a/Engine/Plugins/FX/Niagara/Source/NiagaraShader/Public/NiagaraShaderType.h
+++ b/Engine/Plugins/FX/Niagara/Source/NiagaraShader/Public/NiagaraShaderType.h
@@ -96,12 +96,12 @@ public:
 	 * Enqueues a compilation for a new shader of this type.
 	 * @param Script - The script to link the shader with.
 	 */
-	class FShaderCompileJob* BeginCompileShader(
+	TSharedRef<class FShaderCommonCompileJob, ESPMode::ThreadSafe> BeginCompileShader(
 			uint32 ShaderMapId,
 			const FNiagaraShaderScript* Script,
 			FShaderCompilerEnvironment* CompilationEnvironment,
 			EShaderPlatform Platform,
-			TArray<FShaderCommonCompileJob*>& NewJobs,
+			TArray<TSharedRef<class FShaderCommonCompileJob, ESPMode::ThreadSafe>>& NewJobs,
 			FShaderTarget Target,
 			TArray<FNiagaraDataInterfaceGPUParamInfo>& InDIParamInfo
 		);
diff --git a/Engine/Plugins/FX/Niagara/Source/NiagaraShader/Public/NiagaraShared.h b/Engine/Plugins/FX/Niagara/Source/NiagaraShader/Public/NiagaraShared.h
--- a/Engine/Plugins/FX/Niagara/Source/NiagaraShader/Public/NiagaraShared.h
+++ b/Engine/Plugins/FX/Niagara/Source/NiagaraShader/Public/NiagaraShared.h
@@ -388,7 +388,7 @@ public:
 		);
 
 	/** Sorts the incoming compiled jobs into the appropriate mesh shader maps, and finalizes this shader map so that it can be used for rendering. */
-	bool ProcessCompilationResults(const TArray<class FShaderCommonCompileJob*>& InCompilationResults, int32& ResultIndex, float& TimeBudget);
+	bool ProcessCompilationResults(const TArray<TSharedRef<class FShaderCommonCompileJob, ESPMode::ThreadSafe>>& InCompilationResults, int32& ResultIndex, float& TimeBudget);
 
 	/**
 	* Checks whether the shader map is missing any shader types necessary for the given script.
@@ -545,7 +545,7 @@ private:
 	/** Debug information about how the shader map was compiled. */
 	FString DebugDescription;
 
-	FShader* ProcessCompilationResultsForSingleJob(class FShaderCommonCompileJob* SingleJob, const FSHAHash& ShaderMapHash);
+	FShader* ProcessCompilationResultsForSingleJob(TSharedRef<class FShaderCommonCompileJob, ESPMode::ThreadSafe> SingleJob, const FSHAHash& ShaderMapHash);
 
 	bool IsNiagaraShaderComplete(const FNiagaraShaderScript* Script, const FNiagaraShaderType* ShaderType, bool bSilent);
 
diff --git a/Engine/Source/Runtime/Engine/Private/Materials/MaterialShader.cpp b/Engine/Source/Runtime/Engine/Private/Materials/MaterialShader.cpp
--- a/Engine/Source/Runtime/Engine/Private/Materials/MaterialShader.cpp
+++ b/Engine/Source/Runtime/Engine/Private/Materials/MaterialShader.cpp
@@ -846,7 +846,7 @@ FShaderCompileJob* FMaterialShaderType::BeginCompileShader(
 	FShaderCompilerEnvironment* MaterialEnvironment,
 	const FShaderPipelineType* ShaderPipeline,
 	EShaderPlatform Platform,
-	TArray<FShaderCommonCompileJob*>& NewJobs,
+	TArray<TSharedRef<FShaderCommonCompileJob, ESPMode::ThreadSafe>>& NewJobs,
 	const FString& DebugDescription,
 	const FString& DebugExtension
 	)
@@ -876,7 +876,7 @@ FShaderCompileJob* FMaterialShaderType::BeginCompileShader(
 		GetShaderFilename(),
 		GetFunctionName(),
 		FShaderTarget(GetFrequency(),Platform),
-		NewJob,
+		TSharedRef<FShaderCommonCompileJob, ESPMode::ThreadSafe>(NewJob),
 		NewJobs,
 		true,
 		DebugDescription,
@@ -892,7 +892,7 @@ void FMaterialShaderType::BeginCompileShaderPipeline(
 	FShaderCompilerEnvironment* MaterialEnvironment,
 	const FShaderPipelineType* ShaderPipeline,
 	const TArray<FMaterialShaderType*>& ShaderStages,
-	TArray<FShaderCommonCompileJob*>& NewJobs,
+	TArray<TSharedRef<FShaderCommonCompileJob, ESPMode::ThreadSafe>>& NewJobs,
 	const FString& DebugDescription,
 	const FString& DebugExtension)
 {
@@ -908,7 +908,7 @@ void FMaterialShaderType::BeginCompileShaderPipeline(
 		ShaderStage->BeginCompileShader(ShaderMapId, kUniqueShaderPermutationId, Material, MaterialEnvironment, ShaderPipeline, Platform, NewPipelineJob->StageJobs, DebugDescription, DebugExtension);
 	}
 
-	NewJobs.Add(NewPipelineJob);
+	NewJobs.Add(TSharedRef<FShaderCommonCompileJob, ESPMode::ThreadSafe>(NewPipelineJob));
 }
 
 /**
@@ -1453,7 +1453,7 @@ void FMaterialShaderMap::Compile(
 
 			uint32 NumShaders = 0;
 			uint32 NumVertexFactories = 0;
-			TArray<FShaderCommonCompileJob*> NewJobs;
+			TArray<TSharedRef<FShaderCommonCompileJob, ESPMode::ThreadSafe>> NewJobs;
   
 			// Iterate over all vertex factory types.
 			for(TLinkedList<FVertexFactoryType*>::TIterator VertexFactoryTypeIt(FVertexFactoryType::GetTypeList());VertexFactoryTypeIt;VertexFactoryTypeIt.Next())
@@ -1524,7 +1524,7 @@ void FMaterialShaderMap::Compile(
 						// Only compile the shader if we don't already have it
 						if (!HasShader(ShaderType, PermutationId))
 						{
-							auto* Job = ShaderType->BeginCompileShader(
+							FShaderCompileJob* Job = ShaderType->BeginCompileShader(
 								CompilingId,
 								PermutationId,
 								Material,
@@ -1704,7 +1704,7 @@ FShader* FMaterialShaderMap::ProcessCompilationResultsForSingleJob(FShaderCompil
 }
 
 #if WITH_EDITOR
-bool FMaterialShaderMap::ProcessCompilationResults(const TArray<FShaderCommonCompileJob*>& InCompilationResults, int32& InOutJobIndex, float& TimeBudget, TMap<const FVertexFactoryType*, TArray<const FShaderPipelineType*> >& SharedPipelines)
+bool FMaterialShaderMap::ProcessCompilationResults(const TArray<TSharedRef<FShaderCommonCompileJob, ESPMode::ThreadSafe>>& InCompilationResults, int32& InOutJobIndex, float& TimeBudget, TMap<const FVertexFactoryType*, TArray<const FShaderPipelineType*> >& SharedPipelines)
 {
 	check(InOutJobIndex < InCompilationResults.Num());
 	check(!bCompilationFinalized);
diff --git a/Engine/Source/Runtime/Engine/Private/Materials/MeshMaterialShader.cpp b/Engine/Source/Runtime/Engine/Private/Materials/MeshMaterialShader.cpp
--- a/Engine/Source/Runtime/Engine/Private/Materials/MeshMaterialShader.cpp
+++ b/Engine/Source/Runtime/Engine/Private/Materials/MeshMaterialShader.cpp
@@ -42,7 +42,7 @@ FShaderCompileJob* FMeshMaterialShaderType::BeginCompileShader(
 	FShaderCompilerEnvironment* MaterialEnvironment,
 	FVertexFactoryType* VertexFactoryType,
 	const FShaderPipelineType* ShaderPipeline,
-	TArray<FShaderCommonCompileJob*>& NewJobs,
+	TArray<TSharedRef<FShaderCommonCompileJob, ESPMode::ThreadSafe>>& NewJobs,
 	FString DebugDescription,
 	FString DebugExtension)
 {
@@ -70,6 +70,7 @@ FShaderCompileJob* FMeshMaterialShaderType::BeginCompileShader(
 	bool bAllowDevelopmentShaderCompile = Material->GetAllowDevelopmentShaderCompile();
 
 	// Compile the shader environment passed in with the shader type's source code.
+	TSharedRef<FShaderCommonCompileJob, ESPMode::ThreadSafe> ShaderJob(NewJob);
 	::GlobalBeginCompileShader(
 		Material->GetFriendlyName(),
 		VertexFactoryType,
@@ -78,7 +79,7 @@ FShaderCompileJob* FMeshMaterialShaderType::BeginCompileShader(
 		GetShaderFilename(),
 		GetFunctionName(),
 		FShaderTarget(GetFrequency(),Platform),
-		NewJob,
+		ShaderJob,
 		NewJobs,
 		bAllowDevelopmentShaderCompile,
 		DebugDescription,
@@ -96,7 +97,7 @@ void FMeshMaterialShaderType::BeginCompileShaderPipeline(
 	FVertexFactoryType* VertexFactoryType,
 	const FShaderPipelineType* ShaderPipeline,
 	const TArray<FMeshMaterialShaderType*>& ShaderStages,
-	TArray<FShaderCommonCompileJob*>& NewJobs,
+	TArray<TSharedRef<FShaderCommonCompileJob, ESPMode::ThreadSafe>>& NewJobs,
 	FString DebugDescription,
 	FString DebugExtension)
 {
@@ -112,7 +113,7 @@ void FMeshMaterialShaderType::BeginCompileShaderPipeline(
 		ShaderStage->BeginCompileShader(ShaderMapId, PermutationId, Platform, Material, MaterialEnvironment, VertexFactoryType, ShaderPipeline, NewPipelineJob->StageJobs, DebugDescription, DebugExtension);
 	}
 
-	NewJobs.Add(NewPipelineJob);
+	NewJobs.Add(TSharedRef<FShaderCommonCompileJob, ESPMode::ThreadSafe>(NewPipelineJob));
 }
 
 
@@ -183,7 +184,7 @@ uint32 FMeshMaterialShaderMap::BeginCompile(
 	const FMaterial* Material,
 	FShaderCompilerEnvironment* MaterialEnvironment,
 	EShaderPlatform InPlatform,
-	TArray<FShaderCommonCompileJob*>& NewJobs,
+	TArray<TSharedRef<FShaderCommonCompileJob, ESPMode::ThreadSafe>>& NewJobs,
 	FString DebugDescription,
 	FString DebugExtension
 	)
@@ -221,7 +222,7 @@ uint32 FMeshMaterialShaderMap::BeginCompile(
 				if (!HasShader(ShaderType, PermutationId))
 				{
 					// Compile this mesh material shader for this material and vertex factory type.
-					auto* Job = ShaderType->BeginCompileShader(
+					FShaderCompileJob* Job = ShaderType->BeginCompileShader(
 						ShaderMapId,
 						PermutationId,
 						InPlatform,
diff --git a/Engine/Source/Runtime/Engine/Private/ShaderCompiler/ShaderCompiler.cpp b/Engine/Source/Runtime/Engine/Private/ShaderCompiler/ShaderCompiler.cpp
--- a/Engine/Source/Runtime/Engine/Private/ShaderCompiler/ShaderCompiler.cpp
+++ b/Engine/Source/Runtime/Engine/Private/ShaderCompiler/ShaderCompiler.cpp
@@ -399,7 +399,7 @@ static inline void GetFormatVersionMap(TMap<FString, uint32>& OutFormatVersionMa
 	}
 }
 
-static int32 GetNumTotalJobs(const TArray<FShaderCommonCompileJob*>& Jobs)
+static int32 GetNumTotalJobs(const TArray<TSharedRef<FShaderCommonCompileJob, ESPMode::ThreadSafe>>& Jobs)
 {
 	int32 NumJobs = 0;
 	for (int32 Index = 0; Index < Jobs.Num(); ++Index)
@@ -411,19 +411,19 @@ static int32 GetNumTotalJobs(const TArray<FShaderCommonCompileJob*>& Jobs)
 	return NumJobs;
 }
 
-static void SplitJobsByType(const TArray<FShaderCommonCompileJob*>& QueuedJobs, TArray<FShaderCompileJob*>& OutQueuedSingleJobs, TArray<FShaderPipelineCompileJob*>& OutQueuedPipelineJobs)
+static void SplitJobsByType(const TArray<TSharedRef<FShaderCommonCompileJob, ESPMode::ThreadSafe>>& QueuedJobs, TArray<FShaderCompileJob*>& OutQueuedSingleJobs, TArray<FShaderPipelineCompileJob*>& OutQueuedPipelineJobs)
 {
 	for (int32 Index = 0; Index < QueuedJobs.Num(); ++Index)
 	{
-		auto* CommonJob = QueuedJobs[Index];
-		auto* PipelineJob = CommonJob->GetShaderPipelineJob();
+		TSharedRef<FShaderCommonCompileJob, ESPMode::ThreadSafe> CommonJob = QueuedJobs[Index];
+		FShaderPipelineCompileJob* PipelineJob = CommonJob->GetShaderPipelineJob();
 		if (PipelineJob)
 		{
 			OutQueuedPipelineJobs.Add(PipelineJob);
 		}
 		else
 		{
-			auto* SingleJob = CommonJob->GetSingleShaderJob();
+			FShaderCompileJob* SingleJob = CommonJob->GetSingleShaderJob();
 			check(SingleJob);
 			OutQueuedSingleJobs.Add(SingleJob);
 		}
@@ -431,7 +431,7 @@ static void SplitJobsByType(const TArray<FShaderCommonCompileJob*>& QueuedJobs,
 }
 
 // Serialize Queued Job information
-bool FShaderCompileUtilities::DoWriteTasks(const TArray<FShaderCommonCompileJob*>& QueuedJobs, FArchive& TransferFile)
+bool FShaderCompileUtilities::DoWriteTasks(const TArray<TSharedRef<FShaderCommonCompileJob, ESPMode::ThreadSafe>>& QueuedJobs, FArchive& TransferFile)
 {
 	int32 InputVersion = ShaderCompileWorkerInputVersion;
 	TransferFile << InputVersion;
@@ -641,7 +641,7 @@ static void ReadSingleJob(FShaderCompileJob* CurrentJob, FArchive& OutputFile)
 
 // Disable optimization for this crash handler to get full access to the entire stack frame when debugging a crash dump
 PRAGMA_DISABLE_OPTIMIZATION;
-static void HandleWorkerCrash(const TArray<FShaderCommonCompileJob*>& QueuedJobs, FArchive& OutputFile, int32 OutputVersion, int64 FileSize, int32 ErrorCode, int32 NumProcessedJobs, int32 CallstackLength, int32 ExceptionInfoLength)
+static void HandleWorkerCrash(const TArray<TSharedRef<FShaderCommonCompileJob, ESPMode::ThreadSafe>>& QueuedJobs, FArchive& OutputFile, int32 OutputVersion, int64 FileSize, int32 ErrorCode, int32 NumProcessedJobs, int32 CallstackLength, int32 ExceptionInfoLength)
 {
 	TArray<TCHAR> Callstack;
 	Callstack.AddUninitialized(CallstackLength + 1);
@@ -681,7 +681,7 @@ static void HandleWorkerCrash(const TArray<FShaderCommonCompileJob*>& QueuedJobs
 		WriteInputSourcePathOntoStack(TCHAR_TO_UTF8(*DebugInfo));
 	};
 
-	for (FShaderCommonCompileJob* CommonJob : QueuedJobs)
+	for (auto CommonJob : QueuedJobs)
 	{
 		if (FShaderCompileJob* SingleJob = CommonJob->GetSingleShaderJob())
 		{
@@ -725,7 +725,7 @@ static void HandleWorkerCrash(const TArray<FShaderCommonCompileJob*>& QueuedJobs
 			UE_LOG(LogShaderCompilers, Error, TEXT("SCW %d Queued Jobs, Finished %d single jobs"), QueuedJobs.Num(), NumProcessedJobs);
 			for (int32 Index = 0; Index < QueuedJobs.Num(); ++Index)
 			{
-				FShaderCommonCompileJob* CommonJob = QueuedJobs[Index];
+				auto CommonJob = QueuedJobs[Index];
 				FShaderCompileJob* SingleJob = CommonJob->GetSingleShaderJob();
 				if (SingleJob)
 				{
@@ -777,7 +777,7 @@ static void HandleWorkerCrash(const TArray<FShaderCommonCompileJob*>& QueuedJobs
 PRAGMA_ENABLE_OPTIMIZATION;
 
 // Process results from Worker Process
-void FShaderCompileUtilities::DoReadTaskResults(const TArray<FShaderCommonCompileJob*>& QueuedJobs, FArchive& OutputFile)
+void FShaderCompileUtilities::DoReadTaskResults(const TArray<TSharedRef<FShaderCommonCompileJob, ESPMode::ThreadSafe>>& QueuedJobs, FArchive& OutputFile)
 {
 	if (OutputFile.TotalSize() == 0)
 	{
@@ -1007,7 +1007,7 @@ struct FShaderCompileWorkerInfo
 	double StartTime;
 
 	/** Jobs that this worker is responsible for compiling. */
-	TArray<FShaderCommonCompileJob*> QueuedJobs;
+	TArray<TSharedRef<FShaderCommonCompileJob, ESPMode::ThreadSafe>> QueuedJobs;
 
 	FShaderCompileWorkerInfo() :
 		bIssuedTasksToWorker(false),		
@@ -1972,7 +1972,7 @@ bool FShaderCompilingManager::GetDumpShaderDebugInfo() const
 	return GDumpShaderDebugInfo != 0;
 }
 
-void FShaderCompilingManager::AddJobs(TArray<FShaderCommonCompileJob*>& NewJobs, bool bOptimizeForLowLatency, bool bRecreateComponentRenderStateOnCompletion, const FString MaterialBasePath, const FString PermutationString, bool bSkipResultProcessing)
+void FShaderCompilingManager::AddJobs(TArray<TSharedRef<FShaderCommonCompileJob, ESPMode::ThreadSafe>>& NewJobs, bool bOptimizeForLowLatency, bool bRecreateComponentRenderStateOnCompletion, const FString MaterialBasePath, const FString PermutationString, bool bSkipResultProcessing)
 {
 	check(!FPlatformProperties::RequiresCookedData());
 	// Lock CompileQueueSection so we can access the input and output queues
@@ -2302,7 +2302,7 @@ void FShaderCompilingManager::ProcessCompiledShaderMaps(
 		if (ShaderMap && Materials)
 		{
 			TArray<FString> Errors;
-			const TArray<FShaderCommonCompileJob*>& ResultArray = CompileResults.FinishedJobs;
+			TArray<TSharedRef<FShaderCommonCompileJob, ESPMode::ThreadSafe>>& ResultArray = CompileResults.FinishedJobs;
 
 			// Make a copy of the array as this entry of FMaterialShaderMap::ShaderMapsBeingCompiled will be removed below
 			TArray<FMaterial*> MaterialsArray = *Materials;
@@ -2417,11 +2417,7 @@ void FShaderCompilingManager::ProcessCompiledShaderMaps(
 				}
 
 				// Cleanup shader jobs and compile tracking structures
-				for (int32 JobIndex = 0; JobIndex < ResultArray.Num(); JobIndex++)
-				{
-					delete ResultArray[JobIndex];
-				}
-
+				ResultArray.Empty();
 				ProcessIt.RemoveCurrent();
 			}
 
@@ -2436,15 +2432,11 @@ void FShaderCompilingManager::ProcessCompiledShaderMaps(
 
 			if (GlobalShaderResults)
 			{
-				const TArray<FShaderCommonCompileJob*>& CompilationResults = GlobalShaderResults->FinishedJobs;
+				TArray<TSharedRef<FShaderCommonCompileJob, ESPMode::ThreadSafe>>& CompilationResults = GlobalShaderResults->FinishedJobs;
 
 				ProcessCompiledGlobalShaders(CompilationResults);
 
-				for (int32 ResultIndex = 0; ResultIndex < CompilationResults.Num(); ResultIndex++)
-				{
-					delete CompilationResults[ResultIndex];
-				}
-
+				CompilationResults.Empty();
 				ProcessIt.RemoveCurrent();
 			}
 		}
@@ -2585,7 +2577,7 @@ bool FShaderCompilingManager::HandlePotentialRetryOnError(TMap<int32, FShaderMap
 				|| bSpecialEngineMaterial 
 				|| It.Key() == GlobalShaderMapId)
 			{
-				const TArray<FShaderCommonCompileJob*>& CompleteJobs = Results.FinishedJobs;
+				const TArray<TSharedRef<FShaderCommonCompileJob, ESPMode::ThreadSafe>>& CompleteJobs = Results.FinishedJobs;
 				TArray<const FShaderCommonCompileJob*> ErrorJobs;
 				TArray<FString> UniqueErrors;
 				TArray<EShaderPlatform> ErrorPlatforms;
@@ -2605,7 +2597,7 @@ bool FShaderCompilingManager::HandlePotentialRetryOnError(TMap<int32, FShaderMap
 						{
 							const auto* PipelineJob = CurrentJob.GetShaderPipelineJob();
 							check(PipelineJob);
-							for (const FShaderCommonCompileJob* CommonJob : PipelineJob->StageJobs)
+							for (auto CommonJob : PipelineJob->StageJobs)
 							{
 								AddErrorsForFailedJob(*CommonJob->GetSingleShaderJob(), ErrorPlatforms, UniqueErrors, ErrorJobs);
 							}
@@ -2646,7 +2638,7 @@ bool FShaderCompilingManager::HandlePotentialRetryOnError(TMap<int32, FShaderMap
 						{
 							const auto* PipelineJob = CurrentJob.GetShaderPipelineJob();
 							check(PipelineJob);
-							for (const FShaderCommonCompileJob* CommonJob : PipelineJob->StageJobs)
+							for (auto CommonJob : PipelineJob->StageJobs)
 							{
 								ProcessErrors(*CommonJob->GetSingleShaderJob(), UniqueErrors, ErrorString);
 							}
@@ -2715,7 +2707,7 @@ bool FShaderCompilingManager::HandlePotentialRetryOnError(TMap<int32, FShaderMap
 					else
 					{
 						auto* PipelineJob = CurrentJob.GetShaderPipelineJob();
-						for (FShaderCommonCompileJob* CommonJob : PipelineJob->StageJobs)
+						for (auto CommonJob : PipelineJob->StageJobs)
 						{
 							CommonJob->GetSingleShaderJob()->Output = FShaderCompilerOutput();
 							CommonJob->bFinalized = false;
@@ -2768,23 +2760,21 @@ void FShaderCompilingManager::CancelCompilation(const TCHAR* MaterialName, const
 			int32 JobIndex = CompileQueue.Num();
 			while ( --JobIndex >= 0 )
 			{
-				if (FShaderCommonCompileJob* Job = CompileQueue[JobIndex])
+				auto Job = CompileQueue[JobIndex];
+				if (Job->Id == MapIdx)
 				{
-					if (Job->Id == MapIdx)
+					FShaderPipelineCompileJob* PipelineJob = Job->GetShaderPipelineJob();
+					if (PipelineJob)
 					{
-						FShaderPipelineCompileJob* PipelineJob = Job->GetShaderPipelineJob();
-						if (PipelineJob)
-						{
-							TotalNumJobsRemoved += PipelineJob->StageJobs.Num();
-							NumJobsRemoved += PipelineJob->StageJobs.Num();
-						}
-						else
-						{
-							++TotalNumJobsRemoved;
-							++NumJobsRemoved;
-						}
-						CompileQueue.RemoveAt(JobIndex, 1, false);
+						TotalNumJobsRemoved += PipelineJob->StageJobs.Num();
+						NumJobsRemoved += PipelineJob->StageJobs.Num();
+					}
+					else
+					{
+						++TotalNumJobsRemoved;
+						++NumJobsRemoved;
 					}
+					CompileQueue.RemoveAt(JobIndex, 1, false);
 				}
 			}
 
@@ -3184,18 +3174,18 @@ void GlobalBeginCompileShader(
 	const TCHAR* SourceFilename,
 	const TCHAR* FunctionName,
 	FShaderTarget Target,
-	FShaderCompileJob* NewJob,
-	TArray<FShaderCommonCompileJob*>& NewJobs,
+	TSharedRef<FShaderCommonCompileJob, ESPMode::ThreadSafe> NewJob,
+	TArray<TSharedRef<FShaderCommonCompileJob, ESPMode::ThreadSafe>>& NewJobs,
 	bool bAllowDevelopmentShaderCompile,
 	const FString& DebugDescription,
 	const FString& DebugExtension
-
 	)
 {
 	COOK_STAT(ShaderCompilerCookStats::GlobalBeginCompileShaderCalls++);
 	COOK_STAT(FScopedDurationTimer DurationTimer(ShaderCompilerCookStats::GlobalBeginCompileShaderTimeSec));
 
-	FShaderCompilerInput& Input = NewJob->Input;
+	TSharedRef<FShaderCompileJob, ESPMode::ThreadSafe> Job = StaticCastSharedRef<FShaderCompileJob>(NewJob);
+	FShaderCompilerInput& Input = Job->Input;
 	Input.Target = Target;
 	Input.ShaderFormat = LegacyShaderPlatformToShaderFormat(EShaderPlatform(Target.Platform));
 	Input.VirtualSourceFilePath = SourceFilename;
@@ -3266,7 +3256,7 @@ void GlobalBeginCompileShader(
 				ShaderTypeName.RemoveAt(0);
 			}
 		}
-		Input.DebugGroupName = Input.DebugGroupName / ShaderTypeName / FString::Printf(TEXT("%i"), NewJob->PermutationId);
+		Input.DebugGroupName = Input.DebugGroupName / ShaderTypeName / FString::Printf(TEXT("%i"), Job->PermutationId);
 		
 		if (GDumpShaderDebugInfoShort)
 		{
@@ -4053,11 +4043,11 @@ bool RecompileShaders(const TCHAR* Cmd, FOutputDevice& Ar)
 	return 1;
 }
 
-FShaderCompileJob* FGlobalShaderTypeCompiler::BeginCompileShader(FGlobalShaderType* ShaderType, int32 PermutationId, EShaderPlatform Platform, const FShaderPipelineType* ShaderPipeline, TArray<FShaderCommonCompileJob*>& NewJobs)
+FShaderCompileJob* FGlobalShaderTypeCompiler::BeginCompileShader(FGlobalShaderType* ShaderType, int32 PermutationId, EShaderPlatform Platform, const FShaderPipelineType* ShaderPipeline, TArray<TSharedRef<FShaderCommonCompileJob, ESPMode::ThreadSafe>>& NewJobs)
 {
 	FShaderCompileJob* NewJob = new FShaderCompileJob(GlobalShaderMapId, nullptr, ShaderType, PermutationId);
 	FShaderCompilerEnvironment& ShaderEnvironment = NewJob->Input.Environment;
-
+	
 	UE_LOG(LogShaders, Verbose, TEXT("	%s"), ShaderType->GetName());
 	COOK_STAT(GlobalShaderCookStats::ShadersCompiled++);
 
@@ -4075,28 +4065,28 @@ FShaderCompileJob* FGlobalShaderTypeCompiler::BeginCompileShader(FGlobalShaderTy
 		ShaderType->GetShaderFilename(),
 		ShaderType->GetFunctionName(),
 		FShaderTarget(ShaderType->GetFrequency(), Platform),
-		NewJob,
+		TSharedRef<FShaderCommonCompileJob, ESPMode::ThreadSafe>(NewJob),
 		NewJobs
 	);
 
 	return NewJob;
 }
 
-void FGlobalShaderTypeCompiler::BeginCompileShaderPipeline(EShaderPlatform Platform, const FShaderPipelineType* ShaderPipeline, const TArray<FGlobalShaderType*>& ShaderStages, TArray<FShaderCommonCompileJob*>& NewJobs)
+void FGlobalShaderTypeCompiler::BeginCompileShaderPipeline(EShaderPlatform Platform, const FShaderPipelineType* ShaderPipeline, const TArray<FGlobalShaderType*>& ShaderStages, TArray<TSharedRef<FShaderCommonCompileJob, ESPMode::ThreadSafe>>& NewJobs)
 {
 	check(ShaderStages.Num() > 0);
 	check(ShaderPipeline);
 	UE_LOG(LogShaders, Verbose, TEXT("	Pipeline: %s"), ShaderPipeline->GetName());
 
 	// Add all the jobs as individual first, then add the dependencies into a pipeline job
-	auto* NewPipelineJob = new FShaderPipelineCompileJob(GlobalShaderMapId, ShaderPipeline, ShaderStages.Num());
+	FShaderPipelineCompileJob* NewPipelineJob = new FShaderPipelineCompileJob(GlobalShaderMapId, ShaderPipeline, ShaderStages.Num());
 	for (int32 Index = 0; Index < ShaderStages.Num(); ++Index)
 	{
 		auto* ShaderStage = ShaderStages[Index];
 		BeginCompileShader(ShaderStage, kUniqueShaderPermutationId, Platform, ShaderPipeline, NewPipelineJob->StageJobs);
 	}
 
-	NewJobs.Add(NewPipelineJob);
+	NewJobs.Add(TSharedRef<FShaderCommonCompileJob, ESPMode::ThreadSafe>(NewPipelineJob));
 }
 
 FShader* FGlobalShaderTypeCompiler::FinishCompileShader(FGlobalShaderType* ShaderType, const FShaderCompileJob& CurrentJob, const FShaderPipelineType* ShaderPipelineType)
@@ -4190,7 +4180,7 @@ void VerifyGlobalShaders(EShaderPlatform Platform, bool bLoadedFromCacheFile)
 	}
 
 	// All jobs, single & pipeline
-	TArray<FShaderCommonCompileJob*> GlobalShaderJobs;
+	TArray<TSharedRef<FShaderCommonCompileJob, ESPMode::ThreadSafe>> GlobalShaderJobs;
 
 	// Add the single jobs first
 	TMap<TShaderTypePermutation<const FShaderType>, FShaderCompileJob*> SharedShaderJobs;
@@ -4943,7 +4933,7 @@ static inline FShader* ProcessCompiledJob(FShaderCompileJob* SingleJob, const FS
 	return Shader;
 };
 
-void ProcessCompiledGlobalShaders(const TArray<FShaderCommonCompileJob*>& CompilationResults)
+void ProcessCompiledGlobalShaders(const TArray<TSharedRef<FShaderCommonCompileJob, ESPMode::ThreadSafe>>& CompilationResults)
 {
 	UE_LOG(LogShaders, Warning, TEXT("Compiled %u global shaders"), CompilationResults.Num());
 
diff --git a/Engine/Source/Runtime/Engine/Private/ShaderCompiler/ShaderCompilerXGE.cpp b/Engine/Source/Runtime/Engine/Private/ShaderCompiler/ShaderCompilerXGE.cpp
--- a/Engine/Source/Runtime/Engine/Private/ShaderCompiler/ShaderCompilerXGE.cpp
+++ b/Engine/Source/Runtime/Engine/Private/ShaderCompiler/ShaderCompilerXGE.cpp
@@ -266,7 +266,7 @@ void FShaderCompileXGEThreadRunnable_XmlInterface::PostCompletedJobsForBatch(FSh
 {
 	// Enter the critical section so we can access the input and output queues
 	FScopeLock Lock(&Manager->CompileQueueSection);
-	for (FShaderCommonCompileJob* Job : Batch->GetJobs())
+	for (auto Job : Batch->GetJobs())
 	{
 		FShaderMapCompileResults& ShaderMapResults = Manager->ShaderMapJobs.FindChecked(Job->Id);
 		ShaderMapResults.FinishedJobs.Add(Job);
@@ -277,7 +277,7 @@ void FShaderCompileXGEThreadRunnable_XmlInterface::PostCompletedJobsForBatch(FSh
 	FPlatformAtomics::InterlockedAdd(&Manager->NumOutstandingJobs, -Batch->NumJobs());
 }
 
-void FShaderCompileXGEThreadRunnable_XmlInterface::FShaderBatch::AddJob(FShaderCommonCompileJob* Job)
+void FShaderCompileXGEThreadRunnable_XmlInterface::FShaderBatch::AddJob(TSharedRef<FShaderCommonCompileJob, ESPMode::ThreadSafe> Job)
 {
 	// We can only add jobs to a batch which hasn't been written out yet.
 	if (bTransferFileWritten)
@@ -578,7 +578,7 @@ int32 FShaderCompileXGEThreadRunnable_XmlInterface::CompilingLoop()
 	}
 
 	// Try to prepare more shader jobs (even if a build is in flight).
-	TArray<FShaderCommonCompileJob*> JobQueue;
+	TArray<TSharedRef<FShaderCommonCompileJob, ESPMode::ThreadSafe>> JobQueue;
 	{
 		// Enter the critical section so we can access the input and output queues
 		FScopeLock Lock(&Manager->CompileQueueSection);
@@ -680,11 +680,11 @@ class FXGEShaderCompilerTask
 {
 public:
 	TFuture<FXGETaskResult> Future;
-	TArray<FShaderCommonCompileJob*> ShaderJobs;
+	TArray<TSharedRef<FShaderCommonCompileJob, ESPMode::ThreadSafe>> ShaderJobs;
 	FString InputFilePath;
 	FString OutputFilePath;
 
-	FXGEShaderCompilerTask(TFuture<FXGETaskResult>&& Future, TArray<FShaderCommonCompileJob*>&& ShaderJobs, FString&& InputFilePath, FString&& OutputFilePath)
+	FXGEShaderCompilerTask(TFuture<FXGETaskResult>&& Future, TArray<TSharedRef<FShaderCommonCompileJob, ESPMode::ThreadSafe>>&& ShaderJobs, FString&& InputFilePath, FString&& OutputFilePath)
 		: Future(MoveTemp(Future))
 		, ShaderJobs(MoveTemp(ShaderJobs))
 		, InputFilePath(MoveTemp(InputFilePath))
@@ -704,7 +704,7 @@ FShaderCompileXGEThreadRunnable_InterceptionInterface::~FShaderCompileXGEThreadR
 {
 }
 
-void FShaderCompileXGEThreadRunnable_InterceptionInterface::DispatchShaderCompileJobsBatch(TArray<FShaderCommonCompileJob*>& JobsToSerialize)
+void FShaderCompileXGEThreadRunnable_InterceptionInterface::DispatchShaderCompileJobsBatch(TArray<TSharedRef<FShaderCommonCompileJob, ESPMode::ThreadSafe>>& JobsToSerialize)
 {
 #if WITH_XGE_CONTROLLER
 	FString InputFilePath = IXGEController::Get().CreateUniqueFilePath();
@@ -743,7 +743,7 @@ void FShaderCompileXGEThreadRunnable_InterceptionInterface::DispatchShaderCompil
 int32 FShaderCompileXGEThreadRunnable_InterceptionInterface::CompilingLoop()
 {
 #if WITH_XGE_CONTROLLER
-	TArray<FShaderCommonCompileJob*> PendingJobs;
+	TArray<TSharedRef<FShaderCommonCompileJob, ESPMode::ThreadSafe>> PendingJobs;
 
 	// Try to prepare more shader jobs.
 	{
@@ -767,7 +767,7 @@ int32 FShaderCompileXGEThreadRunnable_InterceptionInterface::CompilingLoop()
 
 		struct FJobBatch
 		{
-			TArray<FShaderCommonCompileJob*> Jobs;
+			TArray<TSharedRef<FShaderCommonCompileJob, ESPMode::ThreadSafe>> Jobs;
 			TSet<const FShaderType*> UniquePointers;
 
 			bool operator == (const FJobBatch& B) const
@@ -903,7 +903,7 @@ int32 FShaderCompileXGEThreadRunnable_InterceptionInterface::CompilingLoop()
 
 				UE_LOG(LogShaderCompilers, Warning, TEXT("Rescheduling shader compilation to run locally after XGE job failed: %s"), *Task->OutputFilePath);
 
-				for (FShaderCommonCompileJob* Job : Task->ShaderJobs)
+				for (TSharedRef<FShaderCommonCompileJob, ESPMode::ThreadSafe> Job : Task->ShaderJobs)
 				{
 					FShaderCompileUtilities::ExecuteShaderCompileJob(*Job);
 				}
@@ -912,7 +912,7 @@ int32 FShaderCompileXGEThreadRunnable_InterceptionInterface::CompilingLoop()
 			// Enter the critical section so we can access the input and output queues
 			{
 				FScopeLock Lock(&Manager->CompileQueueSection);
-				for (FShaderCommonCompileJob* Job : Task->ShaderJobs)
+				for (TSharedRef<FShaderCommonCompileJob, ESPMode::ThreadSafe> Job : Task->ShaderJobs)
 				{
 					FShaderMapCompileResults& ShaderMapResults = Manager->ShaderMapJobs.FindChecked(Job->Id);
 					ShaderMapResults.FinishedJobs.Add(Job);
diff --git a/Engine/Source/Runtime/Engine/Public/MaterialShaderType.h b/Engine/Source/Runtime/Engine/Public/MaterialShaderType.h
--- a/Engine/Source/Runtime/Engine/Public/MaterialShaderType.h
+++ b/Engine/Source/Runtime/Engine/Public/MaterialShaderType.h
@@ -137,7 +137,7 @@ public:
 		FShaderCompilerEnvironment* MaterialEnvironment,
 		const FShaderPipelineType* ShaderPipeline,
 		EShaderPlatform Platform,
-		TArray<FShaderCommonCompileJob*>& NewJobs,
+		TArray<TSharedRef<FShaderCommonCompileJob, ESPMode::ThreadSafe>>& NewJobs,
 		const FString& DebugDescription,
 		const FString& DebugExtension
 		);
@@ -149,7 +149,7 @@ public:
 		FShaderCompilerEnvironment* MaterialEnvironment,
 		const FShaderPipelineType* ShaderPipeline,
 		const TArray<FMaterialShaderType*>& ShaderStages,
-		TArray<FShaderCommonCompileJob*>& NewJobs,
+		TArray<TSharedRef<FShaderCommonCompileJob, ESPMode::ThreadSafe>>& NewJobs,
 		const FString& DebugDescription,
 		const FString& DebugExtension
 	);
diff --git a/Engine/Source/Runtime/Engine/Public/MaterialShared.h b/Engine/Source/Runtime/Engine/Public/MaterialShared.h
--- a/Engine/Source/Runtime/Engine/Public/MaterialShared.h
+++ b/Engine/Source/Runtime/Engine/Public/MaterialShared.h
@@ -863,7 +863,7 @@ public:
 		const FMaterial* Material,
 		FShaderCompilerEnvironment* MaterialEnvironment,
 		EShaderPlatform Platform,
-		TArray<FShaderCommonCompileJob*>& NewJobs,
+		TArray<TSharedRef<FShaderCommonCompileJob, ESPMode::ThreadSafe>>& NewJobs,
 		FString DebugDescription,
 		FString DebugExtension
 		);
@@ -975,7 +975,7 @@ public:
 
 #if WITH_EDITOR
 	/** Sorts the incoming compiled jobs into the appropriate mesh shader maps, and finalizes this shader map so that it can be used for rendering. */
-	bool ProcessCompilationResults(const TArray<FShaderCommonCompileJob*>& InCompilationResults, int32& ResultIndex, float& TimeBudget, TMap<const FVertexFactoryType*, TArray<const FShaderPipelineType*> >& SharedPipelines);
+	bool ProcessCompilationResults(const TArray<TSharedRef<FShaderCommonCompileJob, ESPMode::ThreadSafe>>& InCompilationResults, int32& ResultIndex, float& TimeBudget, TMap<const FVertexFactoryType*, TArray<const FShaderPipelineType*> >& SharedPipelines);
 #endif
 
 	/**
diff --git a/Engine/Source/Runtime/Engine/Public/MeshMaterialShaderType.h b/Engine/Source/Runtime/Engine/Public/MeshMaterialShaderType.h
--- a/Engine/Source/Runtime/Engine/Public/MeshMaterialShaderType.h
+++ b/Engine/Source/Runtime/Engine/Public/MeshMaterialShaderType.h
@@ -105,7 +105,7 @@ public:
 		FShaderCompilerEnvironment* MaterialEnvironment,
 		FVertexFactoryType* VertexFactoryType,
 		const FShaderPipelineType* ShaderPipeline,
-		TArray<FShaderCommonCompileJob*>& NewJobs,
+		TArray<TSharedRef<FShaderCommonCompileJob, ESPMode::ThreadSafe>>& NewJobs,
 		FString DebugDescription,
 		FString DebugExtension
 		);
@@ -119,7 +119,7 @@ public:
 		FVertexFactoryType* VertexFactoryType,
 		const FShaderPipelineType* ShaderPipeline,
 		const TArray<FMeshMaterialShaderType*>& ShaderStages,
-		TArray<FShaderCommonCompileJob*>& NewJobs,
+		TArray<TSharedRef<FShaderCommonCompileJob, ESPMode::ThreadSafe>>& NewJobs,
 		FString DebugDescription,
 		FString DebugExtension
 		);
diff --git a/Engine/Source/Runtime/Engine/Public/ShaderCompiler.h b/Engine/Source/Runtime/Engine/Public/ShaderCompiler.h
--- a/Engine/Source/Runtime/Engine/Public/ShaderCompiler.h
+++ b/Engine/Source/Runtime/Engine/Public/ShaderCompiler.h
@@ -30,7 +30,7 @@ class FShaderPipelineCompileJob;
 
 
 /** Stores all of the common information used to compile a shader or pipeline. */
-class FShaderCommonCompileJob : public FRefCountedObject
+class FShaderCommonCompileJob
 {
 public:
 	/** Id of the shader map this shader belongs to. */
@@ -98,7 +98,7 @@ public:
 class FShaderPipelineCompileJob : public FShaderCommonCompileJob
 {
 public:
-	TArray<FShaderCommonCompileJob*> StageJobs;
+	TArray<TSharedRef<FShaderCommonCompileJob, ESPMode::ThreadSafe>> StageJobs;
 	bool bFailedRemovingUnused;
 
 	/** Shader pipeline that this shader belongs to, may (currently) be NULL */
@@ -114,15 +114,6 @@ public:
 		StageJobs.Empty(NumStages);
 	}
 
-	~FShaderPipelineCompileJob()
-	{
-		for (int32 Index = 0; Index < StageJobs.Num(); ++Index)
-		{
-			delete StageJobs[Index];
-		}
-		StageJobs.Reset();
-	}
-
 	virtual FShaderPipelineCompileJob* GetShaderPipelineJob() override { return this; }
 	virtual const FShaderPipelineCompileJob* GetShaderPipelineJob() const override { return this; }
 };
@@ -133,12 +124,12 @@ public:
 	/**
 	* Enqueues compilation of a shader of this type.
 	*/
-	ENGINE_API static class FShaderCompileJob* BeginCompileShader(FGlobalShaderType* ShaderType, int32 PermutationId, EShaderPlatform Platform, const FShaderPipelineType* ShaderPipeline, TArray<FShaderCommonCompileJob*>& NewJobs);
+	ENGINE_API static class FShaderCompileJob* BeginCompileShader(FGlobalShaderType* ShaderType, int32 PermutationId, EShaderPlatform Platform, const FShaderPipelineType* ShaderPipeline, TArray<TSharedRef<FShaderCommonCompileJob, ESPMode::ThreadSafe>>& NewJobs);
 
 	/**
 	* Enqueues compilation of a shader pipeline of this type.
 	*/
-	ENGINE_API static void BeginCompileShaderPipeline(EShaderPlatform Platform, const FShaderPipelineType* ShaderPipeline, const TArray<FGlobalShaderType*>& ShaderStages, TArray<FShaderCommonCompileJob*>& NewJobs);
+	ENGINE_API static void BeginCompileShaderPipeline(EShaderPlatform Platform, const FShaderPipelineType* ShaderPipeline, const TArray<FGlobalShaderType*>& ShaderStages, TArray<TSharedRef<FShaderCommonCompileJob, ESPMode::ThreadSafe>>& NewJobs);
 
 	/** Either returns an equivalent existing shader of this type, or constructs a new instance. */
 	static FShader* FinishCompileShader(FGlobalShaderType* ShaderType, const FShaderCompileJob& CompileJob, const FShaderPipelineType* ShaderPipelineType);
@@ -231,8 +222,8 @@ private:
 
 namespace FShaderCompileUtilities
 {
-	bool DoWriteTasks(const TArray<FShaderCommonCompileJob*>& QueuedJobs, FArchive& TransferFile);
-	void DoReadTaskResults(const TArray<FShaderCommonCompileJob*>& QueuedJobs, FArchive& OutputFile);
+	bool DoWriteTasks(const TArray<TSharedRef<FShaderCommonCompileJob, ESPMode::ThreadSafe>>& QueuedJobs, FArchive& TransferFile);
+	void DoReadTaskResults(const TArray<TSharedRef<FShaderCommonCompileJob, ESPMode::ThreadSafe>>& QueuedJobs, FArchive& OutputFile);
 
 	/** Execute the specified (single or pipeline) shader compile job. */
 	void ExecuteShaderCompileJob(FShaderCommonCompileJob& Job);
@@ -255,7 +246,7 @@ private:
 	 */
 	class FShaderBatch
 	{
-		TArray<FShaderCommonCompileJob*> Jobs;
+		TArray<TSharedRef<FShaderCommonCompileJob, ESPMode::ThreadSafe>> Jobs;
 		bool bTransferFileWritten;
 
 	public:
@@ -290,12 +281,12 @@ private:
 		{
 			return Jobs.Num();
 		}
-		inline const TArray<FShaderCommonCompileJob*>& GetJobs() const
+		inline const TArray<TSharedRef<FShaderCommonCompileJob, ESPMode::ThreadSafe>>& GetJobs() const
 		{
 			return Jobs;
 		}
 
-		void AddJob(FShaderCommonCompileJob* Job);
+		void AddJob(TSharedRef<FShaderCommonCompileJob, ESPMode::ThreadSafe> Job);
 		
 		void WriteTransferFile();
 	};
@@ -346,7 +337,7 @@ public:
 	static bool IsSupported();
 
 private:
-	void DispatchShaderCompileJobsBatch(TArray<FShaderCommonCompileJob*>& JobsToSerialize);
+	void DispatchShaderCompileJobsBatch(TArray<TSharedRef<FShaderCommonCompileJob, ESPMode::ThreadSafe>>& JobsToSerialize);
 };
 
 #endif // PLATFORM_WINDOWS
@@ -364,7 +355,7 @@ struct FShaderMapCompileResults
 	int32 NumJobsQueued;
 	bool bAllJobsSucceeded;
 	bool bRecreateComponentRenderStateOnCompletion;
-	TArray<FShaderCommonCompileJob*> FinishedJobs;
+	TArray<TSharedRef<FShaderCommonCompileJob, ESPMode::ThreadSafe>> FinishedJobs;
 	bool bSkipResultProcessing;
 };
 
@@ -449,7 +440,7 @@ private:
 	/** Tracks whether we are compiling while the game is running.  If true, we need to throttle down shader compiling CPU usage to avoid starving the runtime threads. */
 	bool bCompilingDuringGame;
 	/** Queue of tasks that haven't been assigned to a worker yet. */
-	TArray<FShaderCommonCompileJob*> CompileQueue;
+	TArray<TSharedRef<FShaderCommonCompileJob, ESPMode::ThreadSafe>> CompileQueue;
 	/** Map from shader map Id to the compile results for that map, used to gather compiled results. */
 	TMap<int32, FShaderMapCompileResults> ShaderMapJobs;
 
@@ -606,7 +597,7 @@ public:
 	 * Adds shader jobs to be asynchronously compiled. 
 	 * FinishCompilation or ProcessAsyncResults must be used to get the results.
 	 */
-	ENGINE_API void AddJobs(TArray<FShaderCommonCompileJob*>& NewJobs, bool bOptimizeForLowLatency, bool bRecreateComponentRenderStateOnCompletion, const FString MaterialBasePath, FString PermutationString = FString(""), bool bSkipResultProcessing = false);
+	ENGINE_API void AddJobs(TArray<TSharedRef<FShaderCommonCompileJob, ESPMode::ThreadSafe>>& NewJobs, bool bOptimizeForLowLatency, bool bRecreateComponentRenderStateOnCompletion, const FString MaterialBasePath, FString PermutationString = FString(""), bool bSkipResultProcessing = false);
 
 	/**
 	* Removes all outstanding compile jobs for the passed shader maps.
@@ -664,8 +655,8 @@ extern ENGINE_API void GlobalBeginCompileShader(
 	const TCHAR* SourceFilename,
 	const TCHAR* FunctionName,
 	FShaderTarget Target,
-	FShaderCompileJob* NewJob,
-	TArray<FShaderCommonCompileJob*>& NewJobs,
+	TSharedRef<FShaderCommonCompileJob, ESPMode::ThreadSafe> NewJob,
+	TArray<TSharedRef<FShaderCommonCompileJob, ESPMode::ThreadSafe>>& NewJobs,
 	bool bAllowDevelopmentShaderCompile = true,
 	const FString& DebugDescription = "",
 	const FString& DebugExtension = ""
@@ -708,7 +699,7 @@ extern ENGINE_API void BeginRecompileGlobalShaders(const TArray<FShaderType*>& O
 extern ENGINE_API void FinishRecompileGlobalShaders();
 
 /** Called by the shader compiler to process completed global shader jobs. */
-extern ENGINE_API void ProcessCompiledGlobalShaders(const TArray<FShaderCommonCompileJob*>& CompilationResults);
+extern ENGINE_API void ProcessCompiledGlobalShaders(const TArray<TSharedRef<FShaderCommonCompileJob, ESPMode::ThreadSafe>>& CompilationResults);
 
 /**
 * Saves the global shader map as a file for the target platform.
