Mark sim data dirty flag when it is changed by the interface

diff --git a/Engine/Source/Runtime/Engine/Private/PhysicsEngine/Experimental/PhysInterface_Chaos.cpp b/Engine/Source/Runtime/Engine/Private/PhysicsEngine/Experimental/PhysInterface_Chaos.cpp
--- a/Engine/Source/Runtime/Engine/Private/PhysicsEngine/Experimental/PhysInterface_Chaos.cpp
+++ b/Engine/Source/Runtime/Engine/Private/PhysicsEngine/Experimental/PhysInterface_Chaos.cpp
@@ -48,6 +48,8 @@
 #include "Containers/ArrayView.h"
 #endif
 
+#pragma optimize("", off)
+
 DEFINE_STAT(STAT_TotalPhysicsTime);
 DEFINE_STAT(STAT_NumCloths);
 DEFINE_STAT(STAT_NumClothVerts);
@@ -1467,6 +1469,7 @@ void FPhysInterface_Chaos::SetQueryFilter(const FPhysicsShapeReference_Chaos& In
 void FPhysInterface_Chaos::SetSimulationFilter(const FPhysicsShapeReference_Chaos& InShapeRef, const FCollisionFilterData& InFilter)
 {
 	InShapeRef.Shape->SimData = InFilter;
+	InShapeRef.ActorRef->MarkShapeSimDataDirty();
 }
 
 bool FPhysInterface_Chaos::IsSimulationShape(const FPhysicsShapeHandle& InShape)
diff --git a/Engine/Source/Runtime/Experimental/Chaos/Public/Chaos/ParticleHandle.h b/Engine/Source/Runtime/Experimental/Chaos/Public/Chaos/ParticleHandle.h
--- a/Engine/Source/Runtime/Experimental/Chaos/Public/Chaos/ParticleHandle.h
+++ b/Engine/Source/Runtime/Experimental/Chaos/Public/Chaos/ParticleHandle.h
@@ -1503,6 +1503,10 @@ public:
 	// Right now it's exposed to lubricate the creation of the whole proxy system.
 	class IPhysicsProxyBase* Proxy;
 
+	// TODO: This is an awful side effect of housing the dirty flag for shape data
+	//       inside the particle, but not setting the shape data through it.
+	void MarkShapeSimDataDirty() { MarkDirty(EParticleFlags::ShapeSimData); }
+
 private:
 	TVector<T, d> MX;
 	FUniqueIdx MUniqueIdx;
