Rm pragma optimize :|

diff --git a/Engine/Source/Runtime/Engine/Private/PhysicsEngine/Experimental/PhysInterface_Chaos.cpp b/Engine/Source/Runtime/Engine/Private/PhysicsEngine/Experimental/PhysInterface_Chaos.cpp
--- a/Engine/Source/Runtime/Engine/Private/PhysicsEngine/Experimental/PhysInterface_Chaos.cpp
+++ b/Engine/Source/Runtime/Engine/Private/PhysicsEngine/Experimental/PhysInterface_Chaos.cpp
@@ -48,8 +48,6 @@
 #include "Containers/ArrayView.h"
 #endif
 
-#pragma optimize("", off)
-
 DEFINE_STAT(STAT_TotalPhysicsTime);
 DEFINE_STAT(STAT_NumCloths);
 DEFINE_STAT(STAT_NumClothVerts);
