Back out changelist 11119764, 11119772, 11119809, 11119849, 11121802, 11125895, 11124508. AutoTest assert

diff --git a/Engine/Plugins/Media/PixelStreaming/Source/PixelStreaming/Private/AudioCapturer.cpp b/Engine/Plugins/Media/PixelStreaming/Source/PixelStreaming/Private/AudioCapturer.cpp
--- a/Engine/Plugins/Media/PixelStreaming/Source/PixelStreaming/Private/AudioCapturer.cpp
+++ b/Engine/Plugins/Media/PixelStreaming/Source/PixelStreaming/Private/AudioCapturer.cpp
@@ -114,7 +114,7 @@ int32 FAudioCapturer::Init()
 		return -1;
 	}
 
-	FAudioDevice* AudioDevice = FAudioDeviceManager::GetMainDevice();
+	FAudioDevice* AudioDevice = GEngine->GetMainAudioDevice();
 	if (!AudioDevice)
 	{
 		return -1;
@@ -139,7 +139,7 @@ int32 FAudioCapturer::Terminate()
 		return -1;
 	}
 
-	FAudioDevice* AudioDevice = FAudioDeviceManager::GetMainDevice();
+	FAudioDevice* AudioDevice = GEngine->GetMainAudioDevice();
 	if (!AudioDevice)
 	{
 		return -1;
diff --git a/Engine/Plugins/Online/OnlineSubsystemUtils/Source/OnlineSubsystemUtils/Private/OnlineSubsystemUtils.cpp b/Engine/Plugins/Online/OnlineSubsystemUtils/Source/OnlineSubsystemUtils/Private/OnlineSubsystemUtils.cpp
--- a/Engine/Plugins/Online/OnlineSubsystemUtils/Source/OnlineSubsystemUtils/Private/OnlineSubsystemUtils.cpp
+++ b/Engine/Plugins/Online/OnlineSubsystemUtils/Source/OnlineSubsystemUtils/Private/OnlineSubsystemUtils.cpp
@@ -110,7 +110,7 @@ UAudioComponent* CreateVoiceAudioComponent(uint32 SampleRate, int32 NumChannels)
 	UAudioComponent* AudioComponent = nullptr;
 	if (GEngine != nullptr)
 	{
-		if (FAudioDevice* AudioDevice = FAudioDeviceManager::GetMainDevice())
+		if (FAudioDevice* AudioDevice = GEngine->GetMainAudioDevice())
 		{
 			USoundWaveProcedural* SoundStreaming = NewObject<USoundWaveProcedural>();
 			SoundStreaming->SetSampleRate(SampleRate);
@@ -157,7 +157,7 @@ UAudioComponent* CreateVoiceAudioComponent(uint32 SampleRate, int32 NumChannels)
 UVoipListenerSynthComponent* CreateVoiceSynthComponent(uint32 SampleRate)
 {
 	UVoipListenerSynthComponent* SynthComponentPtr = nullptr;
-	if (FAudioDevice* AudioDevice = FAudioDeviceManager::GetMainDevice())
+	if (FAudioDevice* AudioDevice = GEngine->GetMainAudioDevice())
 	{
 		SynthComponentPtr = NewObject<UVoipListenerSynthComponent>();
 		if (SynthComponentPtr)
diff --git a/Engine/Plugins/Online/OnlineSubsystemUtils/Source/OnlineSubsystemUtils/Private/VoiceEngineImpl.cpp b/Engine/Plugins/Online/OnlineSubsystemUtils/Source/OnlineSubsystemUtils/Private/VoiceEngineImpl.cpp
--- a/Engine/Plugins/Online/OnlineSubsystemUtils/Source/OnlineSubsystemUtils/Private/VoiceEngineImpl.cpp
+++ b/Engine/Plugins/Online/OnlineSubsystemUtils/Source/OnlineSubsystemUtils/Private/VoiceEngineImpl.cpp
@@ -1021,7 +1021,9 @@ IOnlineSubsystem* FVoiceEngineImpl::GetOnlineSubSystem()
 FVoiceEndpoint::FVoiceEndpoint(const FString& InEndpointName, float InSampleRate, int32 InNumChannels)
 	: NumChannelsComingIn(InNumChannels)
 {
-	IAudioDeviceModule* AudioModule = FAudioDeviceManager::GetChecked().GetAudioDeviceModule();
+	check(GEngine && GEngine->GetAudioDeviceManager());
+
+	IAudioDeviceModule* AudioModule = GEngine->GetAudioDeviceManager()->GetAudioDeviceModule();
 	check(AudioModule);
 
 	PlatformEndpoint.Reset(AudioModule->CreateAudioMixerPlatformInterface());
diff --git a/Engine/Plugins/Runtime/AudioCapture/Source/AudioCaptureEditor/Private/AudioRecordingManager.cpp b/Engine/Plugins/Runtime/AudioCapture/Source/AudioCaptureEditor/Private/AudioRecordingManager.cpp
--- a/Engine/Plugins/Runtime/AudioCapture/Source/AudioCaptureEditor/Private/AudioRecordingManager.cpp
+++ b/Engine/Plugins/Runtime/AudioCapture/Source/AudioCaptureEditor/Private/AudioRecordingManager.cpp
@@ -387,7 +387,7 @@ void FAudioRecordingManager::StopRecording(TArray<USoundWave*>& OutSoundWaves)
 					ExistingSoundWave = FindObject<USoundWave>(Package, *AssetName);
 
 					// See if it's currently being played right now
-					FAudioDeviceManager* AudioDeviceManager = FAudioDeviceManager::Get();
+					FAudioDeviceManager* AudioDeviceManager = GEngine->GetAudioDeviceManager();
 					if (AudioDeviceManager && ExistingSoundWave)
 					{
 						AudioDeviceManager->StopSoundsUsingResource(ExistingSoundWave, &ComponentsToRestart);
diff --git a/Engine/Plugins/Runtime/AudioModulation/Source/AudioModulation/Private/AudioModulationStatics.cpp b/Engine/Plugins/Runtime/AudioModulation/Source/AudioModulation/Private/AudioModulationStatics.cpp
--- a/Engine/Plugins/Runtime/AudioModulation/Source/AudioModulation/Private/AudioModulationStatics.cpp
+++ b/Engine/Plugins/Runtime/AudioModulation/Source/AudioModulation/Private/AudioModulationStatics.cpp
@@ -111,7 +111,10 @@ AudioModulation::FAudioModulationImpl* UAudioModulationStatics::GetModulationImp
 	}
 	else
 	{
-		AudioDevice = FAudioDeviceManager::GetMainDevice();
+		if (GEngine)
+		{
+			AudioDevice = GEngine->GetMainAudioDevice();
+		}
 	}
 
 	if (AudioDevice && AudioDevice->IsModulationPluginEnabled())
diff --git a/Engine/Plugins/Runtime/AudioModulation/Source/AudioModulation/Private/SoundModulatorBase.cpp b/Engine/Plugins/Runtime/AudioModulation/Source/AudioModulation/Private/SoundModulatorBase.cpp
--- a/Engine/Plugins/Runtime/AudioModulation/Source/AudioModulation/Private/SoundModulatorBase.cpp
+++ b/Engine/Plugins/Runtime/AudioModulation/Source/AudioModulation/Private/SoundModulatorBase.cpp
@@ -24,7 +24,7 @@ void USoundModulatorBase::PostEditChangeProperty(struct FPropertyChangedEvent& P
 		return;
 	}
 
-	FAudioDeviceManager* DeviceManager = FAudioDeviceManager::Get();
+	FAudioDeviceManager* DeviceManager = GEngine->GetAudioDeviceManager();
 	if (!DeviceManager)
 	{
 		return;
diff --git a/Engine/Plugins/Runtime/Oculus/OculusVR/Source/OculusMR/Private/OculusMRModule.cpp b/Engine/Plugins/Runtime/Oculus/OculusVR/Source/OculusMR/Private/OculusMRModule.cpp
--- a/Engine/Plugins/Runtime/Oculus/OculusVR/Source/OculusMR/Private/OculusMRModule.cpp
+++ b/Engine/Plugins/Runtime/Oculus/OculusVR/Source/OculusMR/Private/OculusMRModule.cpp
@@ -447,7 +447,8 @@ void FOculusMRModule::OnInitialWorldCreated(UWorld* NewWorld)
 				ovrp_Media_SetMrcInputVideoBufferType(ovrpMediaInputVideoBufferType_TextureHandle);
 				ovrp_Media_SetMrcFrameImageFlipped(IsAndroidOpenGLESPlatform(GMaxRHIShaderPlatform));
 				ovrp_Media_SetMrcFrameInverseAlpha(ovrpBool_True);
-				float SampleRate = FAudioDeviceManager::GetMainDevice()->GetSampleRate();
+				FAudioDevice* AudioDevice = FAudioDevice::GetMainAudioDevice();
+				float SampleRate = AudioDevice->GetSampleRate();
 				ovrp_Media_SetMrcAudioSampleRate((int)SampleRate);
 				InitMixedRealityCapture();
 				OnWorldCreated(NewWorld);
diff --git a/Engine/Plugins/Runtime/Oculus/OculusVR/Source/OculusMR/Private/OculusMR_CastingCameraActor.cpp b/Engine/Plugins/Runtime/Oculus/OculusVR/Source/OculusMR/Private/OculusMR_CastingCameraActor.cpp
--- a/Engine/Plugins/Runtime/Oculus/OculusVR/Source/OculusMR/Private/OculusMR_CastingCameraActor.cpp
+++ b/Engine/Plugins/Runtime/Oculus/OculusVR/Source/OculusMR/Private/OculusMR_CastingCameraActor.cpp
@@ -274,14 +274,15 @@ void AOculusMR_CastingCameraActor::BeginPlay()
 	VRNotificationComponent->HMDRecenteredDelegate.Add(Delegate);
 
 #if PLATFORM_ANDROID
-	FAudioDeviceManager::GetMainDevice()->StartRecording(nullptr, 0.1);
+	FAudioDevice* AudioDevice = FAudioDevice::GetMainAudioDevice();
+	AudioDevice->StartRecording(nullptr, 0.1);
 #endif
 }
 
 void AOculusMR_CastingCameraActor::EndPlay(EEndPlayReason::Type Reason)
 {
 #if PLATFORM_ANDROID
-	FAudioDevice* AudioDevice = FAudioDeviceManager::GetMainDevice();
+	FAudioDevice* AudioDevice = FAudioDevice::GetMainAudioDevice();
 	float NumChannels = 2;
 	float SampleRate = AudioDevice->GetSampleRate();
 	AudioDevice->StopRecording(nullptr, NumChannels, SampleRate);
@@ -485,7 +486,7 @@ void AOculusMR_CastingCameraActor::Tick(float DeltaTime)
 		ForegroundCaptureActor->GetCaptureComponent2D()->TextureTarget = ForegroundRenderTargets[CaptureIndex];
 		GetCaptureComponent2D()->SetVisibility(true);
 
-		FAudioDevice* AudioDevice = FAudioDeviceManager::GetMainDevice();
+		FAudioDevice* AudioDevice = FAudioDevice::GetMainAudioDevice();
 		float NumChannels, SampleRate;
 		NumChannels = 2;
 		SampleRate = AudioDevice->GetSampleRate();
diff --git a/Engine/Plugins/Runtime/SoundMod/Source/SoundModImporter/Private/SoundModImporterFactory.cpp b/Engine/Plugins/Runtime/SoundMod/Source/SoundModImporter/Private/SoundModImporterFactory.cpp
--- a/Engine/Plugins/Runtime/SoundMod/Source/SoundModImporter/Private/SoundModImporterFactory.cpp
+++ b/Engine/Plugins/Runtime/SoundMod/Source/SoundModImporter/Private/SoundModImporterFactory.cpp
@@ -1,13 +1,11 @@
 // Copyright Epic Games, Inc. All Rights Reserved.
 
 #include "SoundModImporterFactory.h"
-
-#include "AudioDeviceManager.h"
-#include "Components/AudioComponent.h"
-#include "Editor.h"
 #include "EngineGlobals.h"
+#include "Editor.h"
 #include "SoundMod.h"
 #include "xmp.h"
+#include "Components/AudioComponent.h"
 
 //////////////////////////////////////////////////////////////////////////
 // USoundModImporterFactory
@@ -46,15 +44,12 @@ UObject* USoundModImporterFactory::FactoryCreateBinary
 
 	// TODO - Audio Threading. This needs to be sent to the audio device and wait on stopping the sounds
 	TArray<UAudioComponent*> ComponentsToRestart;
-
-// 	if (FAudioDeviceManager* AudioDeviceManager = FAudioDeviceManager::Get())
-// 	{
-// 		if (ExistingSound)
-// 		{
-// 			// TODO: Generalize the stop sounds function
-// 			AudioDeviceManager->StopSoundsForReimport(ExistingSound, ComponentsToRestart);
-// 		}
-// 	}
+	FAudioDeviceManager* AudioDeviceManager = GEngine->GetAudioDeviceManager();
+	if (AudioDeviceManager && ExistingSound)
+	{
+		// TODO: Generalize the stop sounds function
+		//AudioDeviceManager->StopSoundsForReimport(ExistingSound, ComponentsToRestart);
+	}
 
 	bool bUseExistingSettings = bSoundModFactorySuppressImportOverwriteDialog;
 
diff --git a/Engine/Plugins/Runtime/nDisplay/Source/DisplayCluster/Private/Game/Classes/Basics/DisplayClusterViewportClient.cpp b/Engine/Plugins/Runtime/nDisplay/Source/DisplayCluster/Private/Game/Classes/Basics/DisplayClusterViewportClient.cpp
--- a/Engine/Plugins/Runtime/nDisplay/Source/DisplayCluster/Private/Game/Classes/Basics/DisplayClusterViewportClient.cpp
+++ b/Engine/Plugins/Runtime/nDisplay/Source/DisplayCluster/Private/Game/Classes/Basics/DisplayClusterViewportClient.cpp
@@ -315,13 +315,19 @@ void UDisplayClusterViewportClient::Draw(FViewport* InViewport, FCanvas* SceneCa
 								bool bUpdateListenerPosition = true;
 
 								// If the main audio device is used for multiple PIE viewport clients, we only
-								// want to update the main audio device listener position if it is in focus.
-								// If there is more than one world referencing the main audio device
-								if (FAudioDeviceManager::GetChecked().GetNumMainAudioDeviceWorlds() > 1)
+								// want to update the main audio device listener position if it is in focus
+								if (GEngine)
 								{
-									if (AudioDevice == FAudioDeviceManager::GetMainDevice() && !HasAudioFocus())
+									FAudioDeviceManager* AudioDeviceManager = GEngine->GetAudioDeviceManager();
+
+									// If there is more than one world referencing the main audio device
+									if (AudioDeviceManager->GetNumMainAudioDeviceWorlds() > 1)
 									{
-										bUpdateListenerPosition = false;
+										uint32 MainAudioDeviceHandle = GEngine->GetAudioDeviceHandle();
+										if (AudioDevice->DeviceHandle == MainAudioDeviceHandle && !HasAudioFocus())
+										{
+											bUpdateListenerPosition = false;
+										}
 									}
 								}
 
diff --git a/Engine/Source/Editor/AudioEditor/Classes/Factories/SoundSurroundFactory.h b/Engine/Source/Editor/AudioEditor/Classes/Factories/SoundSurroundFactory.h
--- a/Engine/Source/Editor/AudioEditor/Classes/Factories/SoundSurroundFactory.h
+++ b/Engine/Source/Editor/AudioEditor/Classes/Factories/SoundSurroundFactory.h
@@ -1,13 +1,28 @@
 // Copyright Epic Games, Inc. All Rights Reserved.
+
+//~=============================================================================
+// SoundSurroundFactory
+//~=============================================================================
+
 #pragma once
 
 #include "CoreMinimal.h"
-
-#include "Factories/Factory.h"
 #include "UObject/ObjectMacros.h"
-
+#include "Factories/Factory.h"
+#include "Audio.h"
 #include "SoundSurroundFactory.generated.h"
 
+const FString SurroundSpeakerLocations[SPEAKER_Count] =
+{
+	TEXT("_fl"),		// SPEAKER_FrontLeft
+	TEXT("_fr"),		// SPEAKER_FrontRight
+	TEXT("_fc"),		// SPEAKER_FrontCenter
+	TEXT("_lf"),		// SPEAKER_LowFrequency
+	TEXT("_sl"),		// SPEAKER_SideLeft
+	TEXT("_sr"),		// SPEAKER_SideRight
+	TEXT("_bl"),		// SPEAKER_BackLeft
+	TEXT("_br")			// SPEAKER_BackRight
+};
 
 UCLASS(MinimalAPI, hidecategories=Object)
 class USoundSurroundFactory : public UFactory
@@ -17,6 +32,7 @@ class USoundSurroundFactory : public UFactory
 	UPROPERTY()
 	float CueVolume;
 
+
 	//~ Begin UFactory Interface
 	virtual UObject* FactoryCreateBinary( UClass* Class, UObject* InParent, FName Name, EObjectFlags Flags, UObject* Context, const TCHAR* Type, const uint8*& Buffer, const uint8* BufferEnd, FFeedbackContext* Warn ) override;
 	virtual bool FactoryCanImport( const FString& Filename ) override;
diff --git a/Engine/Source/Editor/AudioEditor/Private/AssetTypeActions/AssetTypeActions_SoundBase.cpp b/Engine/Source/Editor/AudioEditor/Private/AssetTypeActions/AssetTypeActions_SoundBase.cpp
--- a/Engine/Source/Editor/AudioEditor/Private/AssetTypeActions/AssetTypeActions_SoundBase.cpp
+++ b/Engine/Source/Editor/AudioEditor/Private/AssetTypeActions/AssetTypeActions_SoundBase.cpp
@@ -294,7 +294,7 @@ bool FAssetTypeActions_SoundBase::CanExecutePlayCommand(TArray<TWeakObjectPtr<US
 
 void FAssetTypeActions_SoundBase::ExecuteMuteSound(TArray<TWeakObjectPtr<USoundBase>> Objects) const
 {	
-	if (FAudioDeviceManager* ADM = FAudioDeviceManager::Get())
+	if (FAudioDeviceManager* ADM = GEditor->GetAudioDeviceManager())
 	{
 		FAudioDebugger& Debugger = ADM->GetDebugger();
 
@@ -318,7 +318,7 @@ void FAssetTypeActions_SoundBase::ExecuteMuteSound(TArray<TWeakObjectPtr<USoundB
 
 void FAssetTypeActions_SoundBase::ExecuteSoloSound(TArray<TWeakObjectPtr<USoundBase>> Objects) const
 {
-	if (FAudioDeviceManager* ADM = FAudioDeviceManager::Get())
+	if (FAudioDeviceManager* ADM = GEditor->GetAudioDeviceManager())
 	{
 		FAudioDebugger& Debugger = ADM->GetDebugger();
 
@@ -342,7 +342,7 @@ void FAssetTypeActions_SoundBase::ExecuteSoloSound(TArray<TWeakObjectPtr<USoundB
 
 bool FAssetTypeActions_SoundBase::IsActionCheckedMute(TArray<TWeakObjectPtr<USoundBase>> Objects) const
 {
-	if (FAudioDeviceManager* ADM = FAudioDeviceManager::Get())
+	if (FAudioDeviceManager* ADM = GEditor->GetAudioDeviceManager())
 	{
 		// If *any* of the selection are muted, show the tick box as ticked.
 		FAudioDebugger& Debugger = ADM->GetDebugger();
@@ -370,7 +370,7 @@ bool FAssetTypeActions_SoundBase::IsActionCheckedMute(TArray<TWeakObjectPtr<USou
 bool FAssetTypeActions_SoundBase::IsActionCheckedSolo(TArray<TWeakObjectPtr<USoundBase>> Objects) const
 {
 	// If *any* of the selection are solod, show the tick box as ticked.
-	if (FAudioDeviceManager* ADM = FAudioDeviceManager::Get())
+	if (FAudioDeviceManager* ADM = GEditor->GetAudioDeviceManager())
 	{
 		FAudioDebugger& Debugger = ADM->GetDebugger();
 		for (TWeakObjectPtr<USoundBase> SoundBase : Objects)
@@ -396,7 +396,7 @@ bool FAssetTypeActions_SoundBase::IsActionCheckedSolo(TArray<TWeakObjectPtr<USou
 
 bool FAssetTypeActions_SoundBase::CanExecuteMuteCommand(TArray<TWeakObjectPtr<USoundBase>> Objects) const
 {
-	if (FAudioDeviceManager* ADM = FAudioDeviceManager::Get())
+	if (FAudioDeviceManager* ADM = GEditor->GetAudioDeviceManager())
 	{
 		// Allow muting if we're not Soloing.
 		FAudioDebugger& Debugger = ADM->GetDebugger();
@@ -426,7 +426,7 @@ bool FAssetTypeActions_SoundBase::CanExecuteMuteCommand(TArray<TWeakObjectPtr<US
 
 bool FAssetTypeActions_SoundBase::CanExecuteSoloCommand(TArray<TWeakObjectPtr<USoundBase>> Objects) const
 {	
-	if (FAudioDeviceManager* ADM = FAudioDeviceManager::Get())
+	if (FAudioDeviceManager* ADM = GEditor->GetAudioDeviceManager())
 	{
 		// Allow Soloing if we're not Muting.
 		FAudioDebugger& Debugger = ADM->GetDebugger();
diff --git a/Engine/Source/Editor/AudioEditor/Private/AssetTypeActions/AssetTypeActions_SoundClass.cpp b/Engine/Source/Editor/AudioEditor/Private/AssetTypeActions/AssetTypeActions_SoundClass.cpp
--- a/Engine/Source/Editor/AudioEditor/Private/AssetTypeActions/AssetTypeActions_SoundClass.cpp
+++ b/Engine/Source/Editor/AudioEditor/Private/AssetTypeActions/AssetTypeActions_SoundClass.cpp
@@ -73,7 +73,7 @@ void FAssetTypeActions_SoundClass::GetActions(const TArray<UObject*>& InObjects,
 
 void FAssetTypeActions_SoundClass::ExecuteMute(TArray<TWeakObjectPtr<USoundClass>> Objects) const
 {
-	if (FAudioDeviceManager* ADM = FAudioDeviceManager::Get())
+	if (FAudioDeviceManager* ADM = GEditor->GetAudioDeviceManager())
 	{
 		for (TWeakObjectPtr<USoundClass> i : Objects)
 		{
@@ -87,7 +87,7 @@ void FAssetTypeActions_SoundClass::ExecuteMute(TArray<TWeakObjectPtr<USoundClass
 
 void FAssetTypeActions_SoundClass::ExecuteSolo(TArray<TWeakObjectPtr<USoundClass>> Objects) const
 {
-	if (FAudioDeviceManager* ADM = FAudioDeviceManager::Get())
+	if (FAudioDeviceManager* ADM = GEditor->GetAudioDeviceManager())
 	{
 		for (TWeakObjectPtr<USoundClass> i : Objects)
 		{
@@ -101,7 +101,7 @@ void FAssetTypeActions_SoundClass::ExecuteSolo(TArray<TWeakObjectPtr<USoundClass
 
 bool FAssetTypeActions_SoundClass::IsActionCheckedMute(TArray<TWeakObjectPtr<USoundClass>> Objects) const
 {
-	if (FAudioDeviceManager* ADM = FAudioDeviceManager::Get())
+	if (FAudioDeviceManager* ADM = GEditor->GetAudioDeviceManager())
 	{
 		for (TWeakObjectPtr<USoundClass> i : Objects)
 		{
@@ -119,7 +119,7 @@ bool FAssetTypeActions_SoundClass::IsActionCheckedMute(TArray<TWeakObjectPtr<USo
 
 bool FAssetTypeActions_SoundClass::IsActionCheckedSolo(TArray<TWeakObjectPtr<USoundClass>> Objects) const
 {
-	if (FAudioDeviceManager* ADM = FAudioDeviceManager::Get())
+	if (FAudioDeviceManager* ADM = GEditor->GetAudioDeviceManager())
 	{		
 		for (TWeakObjectPtr<USoundClass> i : Objects)
 		{
@@ -137,7 +137,7 @@ bool FAssetTypeActions_SoundClass::IsActionCheckedSolo(TArray<TWeakObjectPtr<USo
 
 bool FAssetTypeActions_SoundClass::CanExecuteMuteCommand(TArray<TWeakObjectPtr<USoundClass>> Objects) const
 {
-	if (FAudioDeviceManager* ADM = FAudioDeviceManager::Get())
+	if (FAudioDeviceManager* ADM = GEditor->GetAudioDeviceManager())
 	{
 		// Allow muting if we're not Soloing.
 		for (TWeakObjectPtr<USoundClass> i : Objects)
@@ -158,7 +158,7 @@ bool FAssetTypeActions_SoundClass::CanExecuteMuteCommand(TArray<TWeakObjectPtr<U
 
 bool FAssetTypeActions_SoundClass::CanExecuteSoloCommand(TArray<TWeakObjectPtr<USoundClass>> Objects) const
 {
-	if (FAudioDeviceManager* ADM = FAudioDeviceManager::Get())
+	if (FAudioDeviceManager* ADM = GEditor->GetAudioDeviceManager())
 	{
 		// Allow Soloing if we're not Muting.
 		for (TWeakObjectPtr<USoundClass> i : Objects)
diff --git a/Engine/Source/Editor/AudioEditor/Private/Factories/ReimportSoundSurroundFactory.cpp b/Engine/Source/Editor/AudioEditor/Private/Factories/ReimportSoundSurroundFactory.cpp
--- a/Engine/Source/Editor/AudioEditor/Private/Factories/ReimportSoundSurroundFactory.cpp
+++ b/Engine/Source/Editor/AudioEditor/Private/Factories/ReimportSoundSurroundFactory.cpp
@@ -1,16 +1,14 @@
 // Copyright Epic Games, Inc. All Rights Reserved.
-#include "Factories/ReimportSoundSurroundFactory.h"
 
-#include "Audio.h"
-#include "EditorFramework/AssetImportData.h"
-#include "HAL/FileManager.h"
+#include "Factories/ReimportSoundSurroundFactory.h"
 #include "Logging/MessageLog.h"
-#include "Misc/Paths.h"
 #include "Sound/SoundWave.h"
+#include "Misc/Paths.h"
+#include "EditorFramework/AssetImportData.h"
+#include "HAL/FileManager.h"
 
 #define LOCTEXT_NAMESPACE "AudioEditorFactories"
 
-
 UReimportSoundSurroundFactory::UReimportSoundSurroundFactory(const FObjectInitializer& ObjectInitializer)
 	: Super(ObjectInitializer)
 {
diff --git a/Engine/Source/Editor/AudioEditor/Private/Factories/SoundClassFactory.cpp b/Engine/Source/Editor/AudioEditor/Private/Factories/SoundClassFactory.cpp
--- a/Engine/Source/Editor/AudioEditor/Private/Factories/SoundClassFactory.cpp
+++ b/Engine/Source/Editor/AudioEditor/Private/Factories/SoundClassFactory.cpp
@@ -17,9 +17,12 @@ USoundClassFactory::USoundClassFactory(const FObjectInitializer& ObjectInitializ
 UObject* USoundClassFactory::FactoryCreateNew(UClass* InClass, UObject* InParent, FName InName, EObjectFlags Flags, UObject* Context, FFeedbackContext* Warn)
 {
 	USoundClass* SoundClass = NewObject<USoundClass>(InParent, InName, Flags);
-	if (FAudioDeviceManager* DeviceManager = FAudioDeviceManager::Get())
+
+	class FAudioDeviceManager* AudioDeviceManager = GEngine ? GEngine->GetAudioDeviceManager() : nullptr;
+	if (AudioDeviceManager)
 	{
-		DeviceManager->InitSoundClasses();
+		AudioDeviceManager->InitSoundClasses();
 	}
-	return SoundClass;
+
+	return(SoundClass);
 }
diff --git a/Engine/Source/Editor/AudioEditor/Private/Factories/SoundFactory.cpp b/Engine/Source/Editor/AudioEditor/Private/Factories/SoundFactory.cpp
--- a/Engine/Source/Editor/AudioEditor/Private/Factories/SoundFactory.cpp
+++ b/Engine/Source/Editor/AudioEditor/Private/Factories/SoundFactory.cpp
@@ -221,7 +221,7 @@ UObject* USoundFactory::CreateObject
 		USoundWave* ExistingSound = FindObject<USoundWave>(InParent, *Name.ToString());
 
 		TArray<UAudioComponent*> ComponentsToRestart;
-		FAudioDeviceManager* AudioDeviceManager = FAudioDeviceManager::Get();
+		FAudioDeviceManager* AudioDeviceManager = GEngine->GetAudioDeviceManager();
 		if (AudioDeviceManager && ExistingSound)
 		{
 			// Will block internally on audio thread completing outstanding commands
@@ -237,7 +237,7 @@ UObject* USoundFactory::CreateObject
 			// to be auditioned in the editor properly.
 			if (!ExistingSound->ResourceData)
 			{
-				if (FAudioDevice* AudioDevice = FAudioDeviceManager::GetMainDevice())
+				if (FAudioDevice* AudioDevice = GEngine->GetMainAudioDevice())
 				{
 					FName RuntimeFormat = AudioDevice->GetRuntimeFormat(ExistingSound);
 					ExistingSound->InitAudioResource(RuntimeFormat);
diff --git a/Engine/Source/Editor/AudioEditor/Private/Factories/SoundSubmixFactory.cpp b/Engine/Source/Editor/AudioEditor/Private/Factories/SoundSubmixFactory.cpp
--- a/Engine/Source/Editor/AudioEditor/Private/Factories/SoundSubmixFactory.cpp
+++ b/Engine/Source/Editor/AudioEditor/Private/Factories/SoundSubmixFactory.cpp
@@ -19,7 +19,8 @@ UObject* USoundSubmixFactory::FactoryCreateNew(UClass* InClass, UObject* InParen
 {
 	USoundSubmix* SoundSubmix = NewObject<USoundSubmix>(InParent, InName, Flags);
 
-	if (FAudioDeviceManager* AudioDeviceManager = FAudioDeviceManager::Get())
+	class FAudioDeviceManager* AudioDeviceManager = GEngine ? GEngine->GetAudioDeviceManager() : nullptr;
+	if (AudioDeviceManager)
 	{
 		AudioDeviceManager->InitSoundSubmixes();
 	}
diff --git a/Engine/Source/Editor/AudioEditor/Private/Factories/SoundSurroundFactory.cpp b/Engine/Source/Editor/AudioEditor/Private/Factories/SoundSurroundFactory.cpp
--- a/Engine/Source/Editor/AudioEditor/Private/Factories/SoundSurroundFactory.cpp
+++ b/Engine/Source/Editor/AudioEditor/Private/Factories/SoundSurroundFactory.cpp
@@ -1,7 +1,6 @@
 // Copyright Epic Games, Inc. All Rights Reserved.
-#include "Factories/SoundSurroundFactory.h"
 
-#include "Audio.h"
+#include "Factories/SoundSurroundFactory.h"
 #include "PackageTools.h"
 #include "Sound/SoundWave.h"
 #include "Misc/Paths.h"
@@ -10,7 +9,6 @@
 #include "Editor.h"
 #include "EditorFramework/AssetImportData.h"
 
-
 USoundSurroundFactory::USoundSurroundFactory(const FObjectInitializer& ObjectInitializer)
 	: Super(ObjectInitializer)
 {
diff --git a/Engine/Source/Editor/AudioEditor/Private/SoundClassEditor.cpp b/Engine/Source/Editor/AudioEditor/Private/SoundClassEditor.cpp
--- a/Engine/Source/Editor/AudioEditor/Private/SoundClassEditor.cpp
+++ b/Engine/Source/Editor/AudioEditor/Private/SoundClassEditor.cpp
@@ -161,6 +161,11 @@ void FSoundClassEditor::InitSoundClassEditor( const EToolkitMode::Type Mode, con
 
 	ExtendToolbar();
 	RegenerateMenusAndToolbars();
+			
+	if (GEditor->GetAudioDeviceManager())
+	{
+		Debugger = &GEditor->GetAudioDeviceManager()->GetDebugger();
+	}	
 
 	GraphEditor->SelectAllNodes();
 	for (UObject* SelectedNode : GraphEditor->GetSelectedNodes())
@@ -203,6 +208,7 @@ void FSoundClassEditor::BindCommands()
 
 FSoundClassEditor::FSoundClassEditor()
 	: SoundClass(nullptr)
+	, Debugger(nullptr)
 {
 }
 
@@ -391,57 +397,40 @@ void FSoundClassEditor::RedoGraphAction()
 
 void FSoundClassEditor::ToggleSolo()
 {
-	if (FAudioDeviceManager* DeviceManager = FAudioDeviceManager::Get())
+	if (Debugger)
 	{
-		DeviceManager->GetDebugger().ToggleSoloSoundClass(SoundClass->GetFName());
+		Debugger->ToggleSoloSoundClass(SoundClass->GetFName());
 	}
 }
 
 bool FSoundClassEditor::CanExcuteToggleSolo() const
 {
-	// Enable solo if we are not Muted
-	if (FAudioDeviceManager* DeviceManager = FAudioDeviceManager::Get())
-	{
-		return !DeviceManager->GetDebugger().IsMuteSoundClass(SoundClass->GetFName());
-	}
-	return false;
+	// Enable solo if we are not Muted	
+	return Debugger ? !Debugger->IsMuteSoundClass(SoundClass->GetFName()) : false;
 }
 
 bool FSoundClassEditor::IsSoloToggled() const
 {
-	if (FAudioDeviceManager* DeviceManager = FAudioDeviceManager::Get())
-	{
-		return DeviceManager->GetDebugger().IsSoloSoundClass(SoundClass->GetFName());
-	}
-	return false;
+	return Debugger ? Debugger->IsSoloSoundClass(SoundClass->GetFName()) : false;
 }
 
 void FSoundClassEditor::ToggleMute()
 {
-	if (FAudioDeviceManager* DeviceManager = FAudioDeviceManager::Get())
+	if (Debugger)
 	{
-		DeviceManager->GetDebugger().ToggleMuteSoundClass(SoundClass->GetFName());
+		Debugger->ToggleMuteSoundClass(SoundClass->GetFName());
 	}
 }
 
 bool FSoundClassEditor::CanExcuteToggleMute() const
 {
-	if (FAudioDeviceManager* DeviceManager = FAudioDeviceManager::Get())
-	{
-		return !DeviceManager->GetDebugger().IsSoloSoundClass(SoundClass->GetFName());
-	}
-
-	return false;
+	// Enable mute if we are not Soloed
+	return Debugger ? !Debugger->IsSoloSoundClass(SoundClass->GetFName()) : false;
 }
 
 bool FSoundClassEditor::IsMuteToggled() const
-{
-	if (FAudioDeviceManager* DeviceManager = FAudioDeviceManager::Get())
-	{
-		return DeviceManager->GetDebugger().IsMuteSoundClass(SoundClass->GetFName());
-	}
-
-	return false;
+{	
+	return Debugger ? Debugger->IsMuteSoundClass(SoundClass->GetFName()) : false;
 }
 
 void FSoundClassEditor::CreateSoundClass(UEdGraphPin* FromPin, const FVector2D& Location, const FString& Name)
diff --git a/Engine/Source/Editor/AudioEditor/Private/SoundClassEditor.h b/Engine/Source/Editor/AudioEditor/Private/SoundClassEditor.h
--- a/Engine/Source/Editor/AudioEditor/Private/SoundClassEditor.h
+++ b/Engine/Source/Editor/AudioEditor/Private/SoundClassEditor.h
@@ -118,6 +118,9 @@ private:
 	/** Command list for this editor */
 	TSharedPtr<FUICommandList> GraphEditorCommands;
 
+	/** Cache the audio debugger instance */
+	class FAudioDebugger* Debugger;
+
 	/**	The tab ids for all the tabs used */
 	static const FName GraphCanvasTabId;
 	static const FName PropertiesTabId;
diff --git a/Engine/Source/Editor/AudioEditor/Private/SoundCueEditor.cpp b/Engine/Source/Editor/AudioEditor/Private/SoundCueEditor.cpp
--- a/Engine/Source/Editor/AudioEditor/Private/SoundCueEditor.cpp
+++ b/Engine/Source/Editor/AudioEditor/Private/SoundCueEditor.cpp
@@ -41,6 +41,7 @@ const FName FSoundCueEditor::PaletteTabId( TEXT( "SoundCueEditor_Palette" ) );
 
 FSoundCueEditor::FSoundCueEditor()
 	: SoundCue(nullptr)
+	, Debugger(nullptr)
 {
 }
 
@@ -148,6 +149,11 @@ void FSoundCueEditor::InitSoundCueEditor(const EToolkitMode::Type Mode, const TS
 	ExtendToolbar();
 	RegenerateMenusAndToolbars();
 
+	if (GEditor->GetAudioDeviceManager())
+	{
+		Debugger = &GEditor->GetAudioDeviceManager()->GetDebugger();
+	}	
+	
 	// @todo toolkit world centric editing
 	/*if(IsWorldCentricAssetEditor())
 	{
@@ -416,60 +422,40 @@ void FSoundCueEditor::TogglePlayback()
 
 void FSoundCueEditor::ToggleSolo()
 {
-	if (FAudioDeviceManager* DeviceManager = FAudioDeviceManager::Get())
+	if (Debugger)
 	{
-		DeviceManager->GetDebugger().ToggleSoloSoundCue(SoundCue->GetFName());
+		Debugger->ToggleSoloSoundCue(SoundCue->GetFName());
 	}
 }
 
 bool FSoundCueEditor::CanExcuteToggleSolo() const
 {
 	// Allow Solo if Mute is not Toggle on
-	if (FAudioDeviceManager* DeviceManager = FAudioDeviceManager::Get())
-	{
-		return !DeviceManager->GetDebugger().IsMuteSoundCue(SoundCue->GetFName());
-	}
-
-	return false;
+	return Debugger ? !Debugger->IsMuteSoundCue(SoundCue->GetFName()) : false;
 }
 
 bool FSoundCueEditor::IsSoloToggled() const
 {
-	if (FAudioDeviceManager* DeviceManager = FAudioDeviceManager::Get())
-	{
-		return DeviceManager->GetDebugger().IsSoloSoundCue(SoundCue->GetFName());
-	}
-
-	return false;
+	return Debugger ? Debugger->IsSoloSoundCue(SoundCue->GetFName()) : false;
 }
 
 void FSoundCueEditor::ToggleMute()
 {
-	if (FAudioDeviceManager* DeviceManager = FAudioDeviceManager::Get())
+	if (Debugger)
 	{
-		DeviceManager->GetDebugger().ToggleMuteSoundCue(SoundCue->GetFName());
+		Debugger->ToggleMuteSoundCue(SoundCue->GetFName());
 	}
 }
 
 bool FSoundCueEditor::CanExcuteToggleMute() const
 {
 	// Allow Mute if Solo is not Toggle on
-	if (FAudioDeviceManager* DeviceManager = FAudioDeviceManager::Get())
-	{
-		return !DeviceManager->GetDebugger().IsSoloSoundCue(SoundCue->GetFName());
-	}
-
-	return false;
+	return Debugger ? !Debugger->IsSoloSoundCue(SoundCue->GetFName()) : false;
 }
 
 bool FSoundCueEditor::IsMuteToggled() const
 {
-	if (FAudioDeviceManager* DeviceManager = FAudioDeviceManager::Get())
-	{
-		return DeviceManager->GetDebugger().IsMuteSoundCue(SoundCue->GetFName());
-	}
-
-	return false;
+	return Debugger ? Debugger->IsMuteSoundCue(SoundCue->GetFName()) : false;
 }
 
 void FSoundCueEditor::PlaySingleNode(UEdGraphNode* Node)
diff --git a/Engine/Source/Editor/AudioEditor/Private/SoundCueEditor.h b/Engine/Source/Editor/AudioEditor/Private/SoundCueEditor.h
--- a/Engine/Source/Editor/AudioEditor/Private/SoundCueEditor.h
+++ b/Engine/Source/Editor/AudioEditor/Private/SoundCueEditor.h
@@ -210,6 +210,9 @@ private:
 	/** Command list for this editor */
 	TSharedPtr<FUICommandList> GraphEditorCommands;
 
+	/** Cache of the Audio debugger instance */
+	class FAudioDebugger* Debugger;
+
 	/**	The tab ids for all the tabs used */
 	static const FName GraphCanvasTabId;
 	static const FName PropertiesTabId;
diff --git a/Engine/Source/Editor/DetailCustomizations/Private/LinuxTargetSettingsDetails.cpp b/Engine/Source/Editor/DetailCustomizations/Private/LinuxTargetSettingsDetails.cpp
--- a/Engine/Source/Editor/DetailCustomizations/Private/LinuxTargetSettingsDetails.cpp
+++ b/Engine/Source/Editor/DetailCustomizations/Private/LinuxTargetSettingsDetails.cpp
@@ -49,7 +49,15 @@ namespace LinuxTargetSettingsDetailsConstants
 static FText GetFriendlyNameFromLinuxRHIName(const FString& InRHIName)
 {
 	FText FriendlyRHIName = LOCTEXT("UnknownRHI", "UnknownRHI");
-	if (InRHIName == TEXT("GLSL_150_ES31"))
+	if (InRHIName == TEXT("GLSL_150"))
+	{
+		FriendlyRHIName = LOCTEXT("OpenGL3", "OpenGL 3 (SM4)");
+	}
+	else if (InRHIName == TEXT("GLSL_150_ES2"))
+	{
+		FriendlyRHIName = LOCTEXT("OpenGL3ES2", "OpenGL 3 (ES2)");
+	}
+	else if (InRHIName == TEXT("GLSL_150_ES31"))
 	{
 		FriendlyRHIName = LOCTEXT("OpenGL3ES31", "OpenGL 3 (ES3.1, Experimental)");
 	}
@@ -340,7 +348,8 @@ bool FLinuxTargetSettingsDetails::IsValidAudioDeviceName(const FString& InDevice
 	bool bIsValid = false;
 
 #if WITH_ENGINE
-	if (FAudioDevice* AudioDevice = FAudioDeviceManager::GetMainDevice())
+	FAudioDevice* AudioDevice = GEngine->GetMainAudioDevice();
+	if (AudioDevice)
 	{
 		TArray<FString> DeviceNames;
 		AudioDevice->GetAudioDeviceList(DeviceNames);
@@ -364,7 +373,7 @@ TSharedRef<SWidget> FLinuxTargetSettingsDetails::MakeAudioDeviceMenu(const TShar
 	FMenuBuilder MenuBuilder(true, nullptr);
 
 #if WITH_ENGINE
-	FAudioDevice* AudioDevice = FAudioDeviceManager::GetMainDevice();
+	FAudioDevice* AudioDevice = GEngine->GetMainAudioDevice();
 	if (AudioDevice)
 	{
 		TArray<FString> AudioDeviceNames;
diff --git a/Engine/Source/Editor/Kismet/Private/BlueprintEditor.cpp b/Engine/Source/Editor/Kismet/Private/BlueprintEditor.cpp
--- a/Engine/Source/Editor/Kismet/Private/BlueprintEditor.cpp
+++ b/Engine/Source/Editor/Kismet/Private/BlueprintEditor.cpp
@@ -8986,11 +8986,11 @@ void FBlueprintEditor::UpdatePreviewActor(UBlueprint* InBlueprint, bool bInForce
 			if ( PreviewActor->IsHidden() )
 			{
 				PreviewActor->SetHidden(false);
-				PreviewActor->MarkComponentsRenderStateDirty();
+				PreviewActor->MarkComponentsRenderStateDirty();				
 			}
 
 			// Prevent any audio from playing as a result of spawning
-			if (FAudioDevice* AudioDevice = FAudioDeviceManager::GetMainDevice())
+			if (FAudioDevice* AudioDevice = GEngine->GetMainAudioDevice())
 			{
 				AudioDevice->Flush(PreviewScene.GetWorld());
 			}
diff --git a/Engine/Source/Editor/MovieSceneCaptureDialog/Private/MovieSceneCaptureDialogModule.cpp b/Engine/Source/Editor/MovieSceneCaptureDialog/Private/MovieSceneCaptureDialogModule.cpp
--- a/Engine/Source/Editor/MovieSceneCaptureDialog/Private/MovieSceneCaptureDialogModule.cpp
+++ b/Engine/Source/Editor/MovieSceneCaptureDialog/Private/MovieSceneCaptureDialogModule.cpp
@@ -433,8 +433,9 @@ void FInEditorCapture::Start()
 
 	UGameViewportClient::OnViewportCreated().AddRaw(this, &FInEditorCapture::OnPIEViewportStarted);
 	FEditorDelegates::EndPIE.AddRaw(this, &FInEditorCapture::OnEndPIE);
-
-	if (FAudioDevice* AudioDevice = FAudioDeviceManager::GetMainDevice())
+		
+	FAudioDevice* AudioDevice = GEngine->GetMainAudioDevice();
+	if (AudioDevice != nullptr)
 	{
 		TransientMasterVolume = AudioDevice->GetTransientMasterVolume();
 		AudioDevice->SetTransientMasterVolume(0.0f);
@@ -595,7 +596,8 @@ void FInEditorCapture::Shutdown()
 
 	FObjectReader(GetMutableDefault<ULevelEditorPlaySettings>(), BackedUpPlaySettings);
 
-	if (FAudioDevice* AudioDevice = FAudioDeviceManager::GetMainDevice())
+	FAudioDevice* AudioDevice = GEngine->GetMainAudioDevice();
+	if (AudioDevice != nullptr)
 	{
 		AudioDevice->SetTransientMasterVolume(TransientMasterVolume);
 	}
diff --git a/Engine/Source/Editor/MovieSceneTools/Private/TrackEditors/AudioTrackEditor.cpp b/Engine/Source/Editor/MovieSceneTools/Private/TrackEditors/AudioTrackEditor.cpp
--- a/Engine/Source/Editor/MovieSceneTools/Private/TrackEditors/AudioTrackEditor.cpp
+++ b/Engine/Source/Editor/MovieSceneTools/Private/TrackEditors/AudioTrackEditor.cpp
@@ -296,7 +296,8 @@ void FAudioThumbnail::GenerateWaveformPreview(TArray<uint8>& OutData, TRange<flo
 	{
 		// @todo Sequencer optimize - We might want to generate the data when we generate the texture
 		// and then discard the data afterwards, though that might be a perf hit traded for better memory usage
-		if (FAudioDevice* AudioDevice = FAudioDeviceManager::GetMainDevice())
+		FAudioDevice* AudioDevice = GEngine->GetMainAudioDevice();
+		if (AudioDevice)
 		{
 			AudioDevice->StopAllSounds(true);
 
diff --git a/Engine/Source/Editor/UnrealEd/Private/Commandlets/AudioMixerCommandlet.cpp b/Engine/Source/Editor/UnrealEd/Private/Commandlets/AudioMixerCommandlet.cpp
--- a/Engine/Source/Editor/UnrealEd/Private/Commandlets/AudioMixerCommandlet.cpp
+++ b/Engine/Source/Editor/UnrealEd/Private/Commandlets/AudioMixerCommandlet.cpp
@@ -268,7 +268,7 @@ public:
 		}
 
 		// Get the audio main device
-		FAudioDevice* AudioDevice = FAudioDeviceManager::GetMainDevice();
+		FAudioDevice* AudioDevice = GEngine->GetMainAudioDevice();
 
 		bool bSuccess = false;
 
@@ -335,7 +335,7 @@ public:
 
 		LoadEditorAndEngineObjects<USoundWave>(SearchPaths, SoundWaves, &IgnoreList);
 
-		FAudioDevice* AudioDevice = FAudioDeviceManager::GetMainDevice();
+		FAudioDevice* AudioDevice = GEngine->GetMainAudioDevice();
 
 		bool bSuccess = false;
 
@@ -420,7 +420,7 @@ public:
 
 		LoadEditorAndEngineObjects<USoundWave>(SearchPaths, SoundWaves, &IgnoreList);
 
-		FAudioDevice* AudioDevice = FAudioDeviceManager::GetMainDevice();
+		FAudioDevice* AudioDevice = GEngine->GetMainAudioDevice();
 
 		bool bSuccess = false;
 		float TimeCount = 0.0f;
@@ -501,7 +501,7 @@ public:
 			SoundWave->SoundGroup = ESoundGroup::SOUNDGROUP_Music;
 		}
 
-		FAudioDevice* AudioDevice = FAudioDeviceManager::GetMainDevice();
+		FAudioDevice* AudioDevice = GEngine->GetMainAudioDevice();
 
 		bool bSuccess = false;
 		float TimeCount = 0.0f;
@@ -586,7 +586,7 @@ public:
 
 		LoadEditorAndEngineObjects<USoundWave>(SearchPaths, SoundWaves, &IgnoreList);
 
-		FAudioDevice* AudioDevice = FAudioDeviceManager::GetMainDevice();
+		FAudioDevice* AudioDevice = GEngine->GetMainAudioDevice();
 		if (!AudioDevice)
 		{
 			return false;
@@ -711,7 +711,7 @@ public:
 
 		LoadEditorAndEngineObjects<USoundAttenuation>(SearchPaths, SoundAttenuations);
 
-		FAudioDevice* AudioDevice = FAudioDeviceManager::GetMainDevice();
+		FAudioDevice* AudioDevice = GEngine->GetMainAudioDevice();
 		if (!AudioDevice)
 		{
 			return false;
diff --git a/Engine/Source/Editor/UnrealEd/Private/EditorEngine.cpp b/Engine/Source/Editor/UnrealEd/Private/EditorEngine.cpp
--- a/Engine/Source/Editor/UnrealEd/Private/EditorEngine.cpp
+++ b/Engine/Source/Editor/UnrealEd/Private/EditorEngine.cpp
@@ -94,7 +94,6 @@
 
 // needed for the RemotePropagator
 #include "AudioDevice.h"
-#include "AudioDeviceManager.h"
 #include "SurfaceIterators.h"
 #include "ScopedTransaction.h"
 
@@ -1888,7 +1887,7 @@ void UEditorEngine::Tick( float DeltaSeconds, bool bIdleMode )
 	IStreamingManager::Get().Tick(DeltaSeconds);
 
 	// Update Audio. This needs to occur after rendering as the rendering code updates the listener position.
-	if (FAudioDeviceManager* AudioDeviceManager = FAudioDeviceManager::Get())
+	if (AudioDeviceManager)
 	{
 		UWorld* OldGWorld = NULL;
 		if (PlayWorld)
@@ -2288,7 +2287,7 @@ UAudioComponent* UEditorEngine::GetPreviewAudioComponent()
 
 UAudioComponent* UEditorEngine::ResetPreviewAudioComponent( USoundBase* Sound, USoundNode* SoundNode )
 {
-	if (FAudioDevice* AudioDevice = FAudioDeviceManager::GetMainDevice())
+	if (FAudioDevice* AudioDevice = GetMainAudioDevice())
 	{
 		if (PreviewAudioComponent)
 		{
diff --git a/Engine/Source/Editor/UnrealEd/Private/EditorViewportClient.cpp b/Engine/Source/Editor/UnrealEd/Private/EditorViewportClient.cpp
--- a/Engine/Source/Editor/UnrealEd/Private/EditorViewportClient.cpp
+++ b/Engine/Source/Editor/UnrealEd/Private/EditorViewportClient.cpp
@@ -5864,9 +5864,12 @@ void FEditorViewportClient::SetEnabledStats(const TArray<FString>& InEnabledStat
 	EnabledStats = InEnabledStats;
 
 #if ENABLE_AUDIO_DEBUG
-	if (FAudioDeviceManager* DeviceManager = FAudioDeviceManager::Get())
+	if (GEngine)
 	{
-		FAudioDebugger::ResolveDesiredStats(this);
+		if (FAudioDeviceManager* DeviceManager = GEngine->GetAudioDeviceManager())
+		{
+			FAudioDebugger::ResolveDesiredStats(this);
+		}
 	}
 #endif // ENABLE_AUDIO_DEBUG
 }
diff --git a/Engine/Source/Editor/UnrealEd/Private/ObjectTools.cpp b/Engine/Source/Editor/UnrealEd/Private/ObjectTools.cpp
--- a/Engine/Source/Editor/UnrealEd/Private/ObjectTools.cpp
+++ b/Engine/Source/Editor/UnrealEd/Private/ObjectTools.cpp
@@ -2108,7 +2108,8 @@ namespace ObjectTools
 		GEditor->ClearPreviewComponents();
 
 		// Ensure the audio manager is not holding on to any sounds
-		if (FAudioDeviceManager* AudioDeviceManager = FAudioDeviceManager::Get())
+		FAudioDeviceManager* AudioDeviceManager = GEditor->GetAudioDeviceManager();
+		if (AudioDeviceManager != nullptr)
 		{
 			AudioDeviceManager->UpdateActiveAudioDevices(false);
 
diff --git a/Engine/Source/Runtime/ALAudio/Private/ALAudioBuffer.cpp b/Engine/Source/Runtime/ALAudio/Private/ALAudioBuffer.cpp
--- a/Engine/Source/Runtime/ALAudio/Private/ALAudioBuffer.cpp
+++ b/Engine/Source/Runtime/ALAudio/Private/ALAudioBuffer.cpp
@@ -48,7 +48,7 @@ FALSoundBuffer* FALSoundBuffer::Init(FALAudioDevice* AudioDevice, USoundWave* In
 		return nullptr;
 	}
 
-	FAudioDeviceManager* AudioDeviceManager = FAudioDeviceManager::Get();
+	FAudioDeviceManager* AudioDeviceManager = GEngine->GetAudioDeviceManager();
 	check(AudioDeviceManager != nullptr);
 
 	FALSoundBuffer *Buffer = nullptr;
@@ -119,7 +119,7 @@ void FALSoundBuffer::CreateNativeBuffer(FALAudioDevice* AudioDevice, USoundWave*
 		Buffer->NumChannels = Wave->NumChannels;
 		Buffer->SampleRate = Wave->GetSampleRateForCurrentPlatform();
 
-		FAudioDeviceManager* AudioDeviceManager = FAudioDeviceManager::Get();
+		FAudioDeviceManager* AudioDeviceManager = GEngine->GetAudioDeviceManager();
 		check(AudioDeviceManager != nullptr);
 		AudioDeviceManager->TrackResource(Wave, Buffer);
 	}
diff --git a/Engine/Source/Runtime/ALAudio/Private/ALAudioDevice.cpp b/Engine/Source/Runtime/ALAudio/Private/ALAudioDevice.cpp
--- a/Engine/Source/Runtime/ALAudio/Private/ALAudioDevice.cpp
+++ b/Engine/Source/Runtime/ALAudio/Private/ALAudioDevice.cpp
@@ -55,7 +55,7 @@ void FALAudioDevice::TeardownHardware( void )
 	FreeSources.Empty();
 
 	// Destroy OpenAL buffers associated with this audio device
-	FAudioDeviceManager* AudioDeviceManager = FAudioDeviceManager::Get();
+	FAudioDeviceManager* AudioDeviceManager = GEngine->GetAudioDeviceManager();
 	check(AudioDeviceManager != nullptr);
 	for( int32 i = 0; i < AudioDeviceManager->Buffers.Num(); i++ )
 	{
diff --git a/Engine/Source/Runtime/Android/AndroidAudio/Private/AndroidAudioBuffer.cpp b/Engine/Source/Runtime/Android/AndroidAudio/Private/AndroidAudioBuffer.cpp
--- a/Engine/Source/Runtime/Android/AndroidAudio/Private/AndroidAudioBuffer.cpp
+++ b/Engine/Source/Runtime/Android/AndroidAudio/Private/AndroidAudioBuffer.cpp
@@ -155,7 +155,7 @@ FSLESSoundBuffer* FSLESSoundBuffer::CreateNativeBuffer( FSLESAudioDevice* AudioD
 
 	Buffer->DecompressionState = AudioDevice->CreateCompressedAudioInfo(InWave);
 
-	FAudioDeviceManager* AudioDeviceManager = FAudioDeviceManager::Get();
+	FAudioDeviceManager* AudioDeviceManager = GEngine->GetAudioDeviceManager();
 	check(AudioDeviceManager != nullptr);
 
 	AudioDeviceManager->TrackResource(InWave, Buffer);
@@ -226,7 +226,7 @@ FSLESSoundBuffer* FSLESSoundBuffer::Init(  FSLESAudioDevice* AudioDevice ,USound
 		return( NULL );
 	}
 	
-	FAudioDeviceManager* AudioDeviceManager = FAudioDeviceManager::Get();
+	FAudioDeviceManager* AudioDeviceManager = GEngine->GetAudioDeviceManager();
 	FSLESSoundBuffer* Buffer = NULL;
 	
 	EDecompressionType DecompressionType = InWave->DecompressionType;
diff --git a/Engine/Source/Runtime/AudioMixer/Private/AudioMixerBuffer.cpp b/Engine/Source/Runtime/AudioMixer/Private/AudioMixerBuffer.cpp
--- a/Engine/Source/Runtime/AudioMixer/Private/AudioMixerBuffer.cpp
+++ b/Engine/Source/Runtime/AudioMixer/Private/AudioMixerBuffer.cpp
@@ -181,6 +181,8 @@ namespace Audio
 		InWave->InvalidateSoundWaveIfNeccessary();
 #endif // WITH_EDITOR
 
+		FAudioDeviceManager* AudioDeviceManager = FAudioDevice::GetAudioDeviceManager();
+
 		FMixerBuffer* Buffer = nullptr;
 
 		EDecompressionType DecompressionType = InWave->DecompressionType;
@@ -207,7 +209,7 @@ namespace Audio
 				// Find any existing buffers
 				if (InWave->ResourceID)
 				{
-					Buffer = (FMixerBuffer*)FAudioDeviceManager::GetChecked().GetSoundBufferForResourceID(InWave->ResourceID);
+					Buffer = (FMixerBuffer*)AudioDeviceManager->GetSoundBufferForResourceID(InWave->ResourceID);
 				}
 
 				// Override with any new PCM data even if the buffer already exists
@@ -216,14 +218,14 @@ namespace Audio
 					// If we already have a buffer for this wave resource, free it
 					if (Buffer)
 					{
-						FAudioDeviceManager::GetChecked().FreeBufferResource(Buffer);
+						AudioDeviceManager->FreeBufferResource(Buffer);
 					}
 
 					// Create a new preview buffer
 					Buffer = FMixerBuffer::CreatePreviewBuffer(InAudioDevice, InWave);
 
 					// Track the new created buffer
-					FAudioDeviceManager::GetChecked().TrackResource(InWave, Buffer);
+					AudioDeviceManager->TrackResource(InWave, Buffer);
 				}
 			}
 			break;
@@ -246,7 +248,7 @@ namespace Audio
 			{
 				if (InWave->ResourceID)
 				{
-					Buffer = (FMixerBuffer*)FAudioDeviceManager::GetChecked().GetSoundBufferForResourceID(InWave->ResourceID);
+					Buffer = (FMixerBuffer*)AudioDeviceManager->GetSoundBufferForResourceID(InWave->ResourceID);
 				}
 
 				if (Buffer == nullptr)
@@ -254,7 +256,7 @@ namespace Audio
 					Buffer = FMixerBuffer::CreateNativeBuffer(InAudioDevice, InWave);
 
 					// Track the resource with the audio device manager
-					FAudioDeviceManager::GetChecked().TrackResource(InWave, Buffer);
+					AudioDeviceManager->TrackResource(InWave, Buffer);
 					InWave->RemoveAudioResource();
 				}
 			}
diff --git a/Engine/Source/Runtime/AudioMixer/Private/AudioMixerDevice.cpp b/Engine/Source/Runtime/AudioMixer/Private/AudioMixerDevice.cpp
--- a/Engine/Source/Runtime/AudioMixer/Private/AudioMixerDevice.cpp
+++ b/Engine/Source/Runtime/AudioMixer/Private/AudioMixerDevice.cpp
@@ -1594,22 +1594,24 @@ namespace Audio
 
 	void FMixerDevice::RegisterSubmixBufferListener(ISubmixBufferListener* InSubmixBufferListener, USoundSubmix* InSubmix)
 	{
-		DECLARE_CYCLE_STAT(TEXT("FAudioThreadTask.RegisterSubmixBufferListener"), STAT_RegisterSubmixBufferListener, STATGROUP_AudioThreadCommands);
-		auto AudioThreadCommand = [this, InSubmixBufferListener, InSubmix]()
-		{
-			CSV_SCOPED_TIMING_STAT(Audio, RegisterSubmixBufferListener);
-			RegisterSubmixBufferListener(InSubmixBufferListener, InSubmix);
-		};
-
 		if (!IsInAudioThread())
 		{
+			DECLARE_CYCLE_STAT(TEXT("FAudioThreadTask.RegisterSubmixBufferListener"), STAT_RegisterSubmixBufferListener, STATGROUP_AudioThreadCommands);
+			auto AudioThreadCommand = [this, InSubmixBufferListener, InSubmix]()
+			{
+				CSV_SCOPED_TIMING_STAT(Audio, StopSpectrumAnalysis);
+				RegisterSubmixBufferListener(InSubmixBufferListener, InSubmix);
+			};
+
 			if (IsInGameThread())
 			{
 				FAudioThread::RunCommandOnAudioThread(AudioThreadCommand, GET_STATID(STAT_RegisterSubmixBufferListener));
 			}
 			else
 			{
-				AsyncTask(ENamedThreads::AudioThread, TFunction<void()>(AudioThreadCommand));
+				AsyncTask(ENamedThreads::GameThread, [AudioThreadCommand]() {
+					FAudioThread::RunCommandOnAudioThread(AudioThreadCommand, GET_STATID(STAT_RegisterSubmixBufferListener));
+				});
 			}
 			return;
 		}
@@ -1631,23 +1633,25 @@ namespace Audio
 
 	void FMixerDevice::UnregisterSubmixBufferListener(ISubmixBufferListener* InSubmixBufferListener, USoundSubmix* InSubmix)
 	{
-		DECLARE_CYCLE_STAT(TEXT("FAudioThreadTask.UnregisterSubmixBufferListener"), STAT_UnregisterSubmixBufferListener, STATGROUP_AudioThreadCommands);
-
-		auto AudioThreadCommand = [this, InSubmixBufferListener, InSubmix]()
-		{
-			CSV_SCOPED_TIMING_STAT(Audio, UnregisterSubmixBufferListener);
-			UnregisterSubmixBufferListener(InSubmixBufferListener, InSubmix);
-		};
-
 		if (!IsInAudioThread())
 		{
+			DECLARE_CYCLE_STAT(TEXT("FAudioThreadTask.UnregisterSubmixBufferListener"), STAT_UnregisterSubmixBufferListener, STATGROUP_AudioThreadCommands);
+
+			auto AudioThreadCommand = [this, InSubmixBufferListener, InSubmix]()
+			{
+				CSV_SCOPED_TIMING_STAT(Audio, UnregisterSubmixBufferListener);
+				UnregisterSubmixBufferListener(InSubmixBufferListener, InSubmix);
+			};
+
 			if (IsInGameThread())
 			{
 				FAudioThread::RunCommandOnAudioThread(AudioThreadCommand, GET_STATID(STAT_UnregisterSubmixBufferListener));
 			}
 			else
 			{
-				AsyncTask(ENamedThreads::AudioThread, TFunction<void()>(AudioThreadCommand));
+				AsyncTask(ENamedThreads::GameThread, [AudioThreadCommand]() {
+					FAudioThread::RunCommandOnAudioThread(AudioThreadCommand, GET_STATID(STAT_UnregisterSubmixBufferListener));
+				});
 			}
 
 			return;
@@ -1684,7 +1688,8 @@ namespace Audio
 
 	bool FMixerDevice::IsMainAudioDevice() const
 	{
-		return this == FAudioDeviceManager::GetMainDevice();
+		bool bIsMain = (this == GEngine->GetMainAudioDevice());
+		return bIsMain;
 	}
 
 	void FMixerDevice::WhiteNoiseTest(AlignedFloatBuffer& Output)
diff --git a/Engine/Source/Runtime/AudioMixer/Private/AudioMixerSource.cpp b/Engine/Source/Runtime/AudioMixer/Private/AudioMixerSource.cpp
--- a/Engine/Source/Runtime/AudioMixer/Private/AudioMixerSource.cpp
+++ b/Engine/Source/Runtime/AudioMixer/Private/AudioMixerSource.cpp
@@ -254,7 +254,7 @@ namespace Audio
 
 			bool bIsDebug = false;
 			FString WaveInstanceName = WaveInstance->GetName(); //-V595
-			FString TestName = FAudioDeviceManager::GetChecked().GetDebugger().GetAudioMixerDebugSoundName();
+			FString TestName = GEngine->GetAudioDeviceManager()->GetDebugger().GetAudioMixerDebugSoundName();
 			if (WaveInstanceName.Contains(TestName))
 			{
 				bDebugMode = true;
diff --git a/Engine/Source/Runtime/AudioMixer/Private/Components/SynthComponent.cpp b/Engine/Source/Runtime/AudioMixer/Private/Components/SynthComponent.cpp
--- a/Engine/Source/Runtime/AudioMixer/Private/Components/SynthComponent.cpp
+++ b/Engine/Source/Runtime/AudioMixer/Private/Components/SynthComponent.cpp
@@ -411,7 +411,9 @@ FAudioDevice* USynthComponent::GetAudioDevice()
 		// Make sure it has a proper audio device handle and retrieve it
 		if (World->AudioDeviceHandle != INDEX_NONE)
 		{
-			return FAudioDeviceManager::GetChecked().GetAudioDevice(World->AudioDeviceHandle);
+			FAudioDeviceManager* AudioDeviceManager = GEngine->GetAudioDeviceManager();
+			check(AudioDeviceManager);
+			return AudioDeviceManager->GetAudioDevice(World->AudioDeviceHandle);
 		}
 	}
 
diff --git a/Engine/Source/Runtime/Engine/Classes/Engine/Engine.h b/Engine/Source/Runtime/Engine/Classes/Engine/Engine.h
--- a/Engine/Source/Runtime/Engine/Classes/Engine/Engine.h
+++ b/Engine/Source/Runtime/Engine/Classes/Engine/Engine.h
@@ -1553,7 +1553,7 @@ public:
 
 protected:
 
-	/** Whether the engine should be playing sounds.  If false at initialization time, the AudioDeviceManager will not be created */
+	/** Whether the engine should be playing sounds.  If false at initialization time the AudioDevice will not be created */
 	uint32 bUseSound:1;
 
 private:
@@ -1766,6 +1766,14 @@ private:
 		void UpdateDynamicResolutionStatus();
 	#endif
 
+protected:
+
+	/** The audio device manager */
+	FAudioDeviceManager* AudioDeviceManager;
+
+	/** Audio device handle to the main audio device. */
+	uint32 MainAudioDeviceHandle;
+
 public:
 
 	/** A collection of messages to display on-screen. */
@@ -2428,19 +2436,15 @@ public:
 	bool IsEditor();
 
 	/** @return the audio device manager of the UEngine, this allows the creation and management of multiple audio devices. */
-	UE_DEPRECATED(4.25, "FAudioDeviceManager::Get and FAudioDeviceManager::GetChecked have replaced UEngine::GetAudioDeviceManager.")
 	FAudioDeviceManager* GetAudioDeviceManager();
 
 	/** @return the main audio device handle used by the engine. */
-	UE_DEPRECATED(4.25, "FAudioDeviceManager::GetMainDevice has replaced UEngine::GetAudioDeviceHandle.")
 	uint32 GetAudioDeviceHandle() const;
 
 	/** @return the main audio device. */
-	UE_DEPRECATED(4.25, "FAudioDeviceManager::GetMainDevice has replaced UEngine::GetMainAudioDevice.")
 	class FAudioDevice* GetMainAudioDevice();
 
 	/** @return the currently active audio device */
-	UE_DEPRECATED(4.25, "FAudioDeviceManager::GetActiveDevice has replaced UEngine::GetActiveAudioDevice.")
 	class FAudioDevice* GetActiveAudioDevice();
 
 	/** @return whether we're currently running in split screen (more than one local player) */
diff --git a/Engine/Source/Runtime/Engine/Classes/Sound/SoundBase.h b/Engine/Source/Runtime/Engine/Classes/Sound/SoundBase.h
--- a/Engine/Source/Runtime/Engine/Classes/Sound/SoundBase.h
+++ b/Engine/Source/Runtime/Engine/Classes/Sound/SoundBase.h
@@ -5,7 +5,7 @@
  * The base class for a playable sound object
  */
 
-#include "AudioDefines.h"
+#include "Audio.h"
 #include "CoreMinimal.h"
 #include "IAudioExtensionPlugin.h"
 #include "SoundConcurrency.h"
@@ -13,22 +13,17 @@
 #include "SoundSubmixSend.h"
 #include "UObject/Object.h"
 #include "UObject/ObjectMacros.h"
-#include "Audio.h"
 
 #include "SoundBase.generated.h"
 
 
-class USoundAttenuation;
-class USoundClass;
 class USoundEffectSourcePreset;
 class USoundSourceBus;
 class USoundSubmix;
 class USoundEffectSourcePresetChain;
 
 struct FActiveSound;
-struct FSoundAttenuationSettings;
 struct FSoundParseParameters;
-struct FWaveInstance;
 
 /**
  * Method of virtualization when a sound is stopped due to playback constraints
diff --git a/Engine/Source/Runtime/Engine/Classes/Sound/SoundEffectSubmix.h b/Engine/Source/Runtime/Engine/Classes/Sound/SoundEffectSubmix.h
--- a/Engine/Source/Runtime/Engine/Classes/Sound/SoundEffectSubmix.h
+++ b/Engine/Source/Runtime/Engine/Classes/Sound/SoundEffectSubmix.h
@@ -15,8 +15,6 @@ class FAudioDevice;
 class FSoundEffectSubmix;
 
 struct FAudioEffectParameters;
-class FAudioDevice;
-
 
 /** Preset of a submix effect that can be shared between sounds. */
 UCLASS(config = Engine, hidecategories = Object, abstract, editinlinenew, BlueprintType)
@@ -31,12 +29,12 @@ class ENGINE_API USoundEffectSubmixPreset : public USoundEffectPreset
 /** Struct which has data needed to initialize the submix effect. */
 struct FSoundEffectSubmixInitData
 {
-	Audio::FDeviceId DeviceHandle;
+	uint32 DeviceHandle;
 	void* PresetSettings;
 	float SampleRate;
 
 	FSoundEffectSubmixInitData()
-		: DeviceHandle(static_cast<Audio::FDeviceId>(INDEX_NONE))
+		: DeviceHandle(INDEX_NONE)
 		, PresetSettings(nullptr)
 		, SampleRate(0)
 	{
diff --git a/Engine/Source/Runtime/Engine/Classes/Sound/SoundSubmixSend.h b/Engine/Source/Runtime/Engine/Classes/Sound/SoundSubmixSend.h
--- a/Engine/Source/Runtime/Engine/Classes/Sound/SoundSubmixSend.h
+++ b/Engine/Source/Runtime/Engine/Classes/Sound/SoundSubmixSend.h
@@ -39,22 +39,6 @@ enum class ESendLevelControlMethod : uint8
 	Manual,
 };
 
-/** Abstract interface for receiving audio data from a given submix. */
-class ENGINE_API ISubmixBufferListener
-{
-public:
-	/**
-	Called when a new buffer has been rendered for a given submix
-	@param OwningSubmix	The submix object which has rendered a new buffer
-	@param AudioData		Ptr to the audio buffer
-	@param NumSamples		The number of audio samples in the audio buffer
-	@param NumChannels		The number of channels of audio in the buffer (e.g. 2 for stereo, 6 for 5.1, etc)
-	@param SampleRate		The sample rate of the audio buffer
-	@param AudioClock		Double audio clock value, from start of audio rendering.
-	*/
-	virtual void OnNewSubmixBuffer(const USoundSubmix* OwningSubmix, float* AudioData, int32 NumSamples, int32 NumChannels, const int32 SampleRate, double AudioClock) = 0;
-};
-
 // Class used to send audio to submixes from USoundBase
 USTRUCT(BlueprintType)
 struct ENGINE_API FSoundSubmixSendInfo
diff --git a/Engine/Source/Runtime/Engine/Private/ActiveSound.cpp b/Engine/Source/Runtime/Engine/Private/ActiveSound.cpp
--- a/Engine/Source/Runtime/Engine/Private/ActiveSound.cpp
+++ b/Engine/Source/Runtime/Engine/Private/ActiveSound.cpp
@@ -5,7 +5,6 @@
 #include "EngineDefines.h"
 #include "Misc/App.h"
 #include "AudioThread.h"
-#include "AudioDefines.h"
 #include "AudioDevice.h"
 #include "IAudioExtensionPlugin.h"
 #include "Sound/AudioSettings.h"
@@ -952,9 +951,9 @@ void FActiveSound::OcclusionTraceDone(const FTraceHandle& TraceHandle, FTraceDat
 	FAsyncTraceDetails TraceDetails;
 	if (TraceToActiveSoundMap.RemoveAndCopyValue(TraceHandle, TraceDetails))
 	{
-		if (FAudioDeviceManager* DeviceManager = FAudioDeviceManager::Get())
+		if (FAudioDeviceManager* AudioDeviceManager = GEngine->GetAudioDeviceManager())
 		{
-			if (FAudioDevice* AudioDevice = DeviceManager->GetAudioDevice(TraceDetails.AudioDeviceID))
+			if (FAudioDevice* AudioDevice = AudioDeviceManager->GetAudioDevice(TraceDetails.AudioDeviceID))
 			{
 				FActiveSound* ActiveSound = TraceDetails.ActiveSound;
 
diff --git a/Engine/Source/Runtime/Engine/Private/Audio.cpp b/Engine/Source/Runtime/Engine/Private/Audio.cpp
--- a/Engine/Source/Runtime/Engine/Private/Audio.cpp
+++ b/Engine/Source/Runtime/Engine/Private/Audio.cpp
@@ -178,13 +178,9 @@ UClass* GetAudioPluginCustomSettingsClass(EAudioPlugin PluginType)
 FSoundBuffer::~FSoundBuffer()
 {
 	// remove ourselves from the set of waves that are tracked by the audio device
-	if (!ResourceID)
+	if (ResourceID && GEngine && GEngine->GetAudioDeviceManager())
 	{
-		return;
-	}
-	if (FAudioDeviceManager* AudioDeviceManager = FAudioDeviceManager::Get())
-	{
-		AudioDeviceManager->RemoveSoundBufferForResourceID(ResourceID);
+		GEngine->GetAudioDeviceManager()->RemoveSoundBufferForResourceID(ResourceID);
 	}
 }
 
@@ -471,14 +467,13 @@ float FSoundSource::GetDebugVolume(const float InVolume)
 #if ENABLE_AUDIO_DEBUG
 
 	// Bail if we don't have a device manager.
-	FAudioDeviceManager* DeviceManager = FAudioDeviceManager::Get();
-	if (!DeviceManager || !WaveInstance || !DebugInfo.IsValid() )
+	if (!GEngine || !GEngine->GetAudioDeviceManager() || !WaveInstance || !DebugInfo.IsValid() )
 	{
 		return OutVolume;
 	}
 
 	// Solos/Mutes (dev only).
-	FAudioDebugger& Debugger = DeviceManager->GetDebugger();
+	FAudioDebugger& Debugger = GEngine->GetAudioDeviceManager()->GetDebugger();	
 	FDebugInfo Info;
 				
 	// SoundWave Solo/Mutes.
@@ -992,24 +987,27 @@ float FWaveInstance::GetDynamicVolume() const
 {
 	float OutVolume = 1.0f;
 
-	if (FAudioDeviceManager* DeviceManager = FAudioDeviceManager::Get())
+	if (GEngine)
 	{
-		if (WaveData)
+		if (FAudioDeviceManager* DeviceManager = GEngine->GetAudioDeviceManager())
 		{
-			OutVolume *= DeviceManager->GetDynamicSoundVolume(ESoundType::Wave, WaveData->GetFName());
-		}
+			if (WaveData)
+			{
+				OutVolume *= DeviceManager->GetDynamicSoundVolume(ESoundType::Wave, WaveData->GetFName());
+			}
 
-		if (ActiveSound)
-		{
-			if (const USoundCue* Sound = Cast<USoundCue>(ActiveSound->GetSound()))
+			if (ActiveSound)
 			{
-				OutVolume *= DeviceManager->GetDynamicSoundVolume(ESoundType::Cue, Sound->GetFName());
+				if (const USoundCue* Sound = Cast<USoundCue>(ActiveSound->GetSound()))
+				{
+					OutVolume *= DeviceManager->GetDynamicSoundVolume(ESoundType::Cue, Sound->GetFName());
+				}
 			}
-		}
 
-		if (SoundClass)
-		{
-			OutVolume *= DeviceManager->GetDynamicSoundVolume(ESoundType::Class, SoundClass->GetFName());
+			if (SoundClass)
+			{
+				OutVolume *= DeviceManager->GetDynamicSoundVolume(ESoundType::Class, SoundClass->GetFName());
+			}
 		}
 	}
 
diff --git a/Engine/Source/Runtime/Engine/Private/AudioDebug.cpp b/Engine/Source/Runtime/Engine/Private/AudioDebug.cpp
--- a/Engine/Source/Runtime/Engine/Private/AudioDebug.cpp
+++ b/Engine/Source/Runtime/Engine/Private/AudioDebug.cpp
@@ -327,17 +327,17 @@ namespace
 
 	void HandleDumpActiveSounds(UWorld* World)
 	{
-		if (GEngine && FAudioDeviceManager::Get())
+		if (GEngine && GEngine->GetAudioDeviceManager())
 		{
-			FAudioDeviceManager::Get()->GetDebugger().DumpActiveSounds();
+			GEngine->GetAudioDeviceManager()->GetDebugger().DumpActiveSounds();
 		}
 	}
 	
 	void HandleClearMutesAndSolos(UWorld* World)
 	{
-		if (GEngine && FAudioDeviceManager::Get())
+		if (GEngine && GEngine->GetAudioDeviceManager())
 		{
-			FAudioDeviceManager::Get()->GetDebugger().ClearMutesAndSolos();
+			GEngine->GetAudioDeviceManager()->GetDebugger().ClearMutesAndSolos();
 		}
 	}
 
@@ -368,7 +368,7 @@ void FAudioDebugger::ToggleVisualizeDebug3dEnabled()
 #if WITH_EDITOR
 void FAudioDebugger::OnBeginPIE()
 {
-	FAudioDeviceManager* DeviceManager = FAudioDeviceManager::Get();
+	FAudioDeviceManager* DeviceManager = GEngine->GetAudioDeviceManager();
 	if (DeviceManager)
 	{
 		DeviceManager->GetDebugger().ClearMutesAndSolos();
@@ -377,7 +377,7 @@ void FAudioDebugger::OnBeginPIE()
 
 void FAudioDebugger::OnEndPIE()
 {
-	FAudioDeviceManager* DeviceManager = FAudioDeviceManager::Get();
+	FAudioDeviceManager* DeviceManager = GEngine->GetAudioDeviceManager();
 	if (DeviceManager)
 	{
 		DeviceManager->GetDebugger().ClearMutesAndSolos();
@@ -544,7 +544,7 @@ void FAudioDebugger::DrawDebugInfo(const FSoundSource& SoundSource)
 		return;
 	}
 
-	FAudioDeviceManager* DeviceManager = FAudioDeviceManager::Get();
+	FAudioDeviceManager* DeviceManager = GEngine->GetAudioDeviceManager();
 	if (DeviceManager && DeviceManager->IsVisualizeDebug3dEnabled())
 	{
 		DECLARE_CYCLE_STAT(TEXT("FAudioThreadTask.DrawSourceDebugInfo"), STAT_AudioDrawSourceDebugInfo, STATGROUP_TaskGraphTasks);
@@ -611,7 +611,7 @@ void FAudioDebugger::DrawDebugInfo(const FActiveSound& ActiveSound, const TArray
 		}
 	}
 
-	FAudioDeviceManager* DeviceManager = FAudioDeviceManager::Get();
+	FAudioDeviceManager* DeviceManager = GEngine->GetAudioDeviceManager();
 	if (DeviceManager && DeviceManager->IsVisualizeDebug3dEnabled())
 	{
 		DECLARE_CYCLE_STAT(TEXT("FAudioThreadTask.DrawActiveSoundDebugInfo"), STAT_AudioDrawActiveSoundDebugInfo, STATGROUP_TaskGraphTasks);
@@ -688,14 +688,15 @@ void FAudioDebugger::DrawDebugInfo(const FActiveSound& ActiveSound, const TArray
 void FAudioDebugger::DrawDebugInfo(UWorld& World, const TArray<FListener>& Listeners)
 {
 #if ENABLE_DRAW_DEBUG
-	if (!GEngine || !ActiveSoundVisualizeListenersCVar)
+	if (!ActiveSoundVisualizeListenersCVar)
 	{
 		return;
 	}
 
 	check(IsInAudioThread());
 
-	if (FAudioDeviceManager::GetChecked().IsVisualizeDebug3dEnabled())
+	FAudioDeviceManager* DeviceManager = GEngine->GetAudioDeviceManager();
+	if (DeviceManager && DeviceManager->IsVisualizeDebug3dEnabled())
 	{
 		DECLARE_CYCLE_STAT(TEXT("FAudioThreadTask.DrawListenerDebugInfo"), STAT_AudioDrawListenerDebugInfo, STATGROUP_TaskGraphTasks);
 		const TWeakObjectPtr<UWorld> WorldPtr = &World;
@@ -741,7 +742,8 @@ void FAudioDebugger::DrawDebugInfo(const FAudioVirtualLoop& VirtualLoop)
 		return;
 	}
 
-	if (FAudioDeviceManager::GetChecked().IsVisualizeDebug3dEnabled())
+	FAudioDeviceManager* DeviceManager = GEngine->GetAudioDeviceManager();
+	if (DeviceManager && DeviceManager->IsVisualizeDebug3dEnabled())
 	{
 		DECLARE_CYCLE_STAT(TEXT("FAudioThreadTask.DrawVirtualLoopDebugInfo"), STAT_AudioDrawVirtualLoopDebugInfo, STATGROUP_TaskGraphTasks);
 
@@ -786,7 +788,7 @@ void FAudioDebugger::DumpActiveSounds() const
 		return;
 	}
 
-	FAudioDevice* AudioDevice = FAudioDeviceManager::GetActiveDevice();
+	FAudioDevice* AudioDevice = GEngine->GetAudioDeviceManager()->GetActiveAudioDevice();
 	if (!AudioDevice)
 	{
 		return;
@@ -1566,7 +1568,7 @@ bool FAudioDebugger::ToggleStats(UWorld* World, const uint8 StatToToggle)
 		return false;
 	}
 
-	if (FAudioDeviceManager* DeviceManager = FAudioDeviceManager::Get())
+	if (FAudioDeviceManager* DeviceManager = GEngine->GetAudioDeviceManager())
 	{
 		DeviceManager->GetDebugger().ToggleStats(AudioDevice->DeviceHandle, StatToToggle);
 	}
diff --git a/Engine/Source/Runtime/Engine/Private/AudioDecompress.cpp b/Engine/Source/Runtime/Engine/Private/AudioDecompress.cpp
--- a/Engine/Source/Runtime/Engine/Private/AudioDecompress.cpp
+++ b/Engine/Source/Runtime/Engine/Private/AudioDecompress.cpp
@@ -645,9 +645,9 @@ FAsyncAudioDecompressWorker::FAsyncAudioDecompressWorker(USoundWave* InWave, int
 	, NumPrecacheFrames(InPrecacheBufferNumFrames)
 {
 	check(NumPrecacheFrames > 0);
-	if (FAudioDevice* AudioDevice = FAudioDeviceManager::GetMainDevice())
+	if (GEngine && GEngine->GetMainAudioDevice())
 	{
-		AudioInfo = AudioDevice->CreateCompressedAudioInfo(Wave);
+		AudioInfo = GEngine->GetMainAudioDevice()->CreateCompressedAudioInfo(Wave);
 	}
 }
 
diff --git a/Engine/Source/Runtime/Engine/Private/AudioDevice.cpp b/Engine/Source/Runtime/Engine/Private/AudioDevice.cpp
--- a/Engine/Source/Runtime/Engine/Private/AudioDevice.cpp
+++ b/Engine/Source/Runtime/Engine/Private/AudioDevice.cpp
@@ -724,7 +724,7 @@ void FAudioDevice::CountBytes(FArchive& Ar)
 	Sources.CountBytes(Ar);
 	// The buffers are stored on the audio device since they are shared amongst all audio devices
 	// Though we are going to count them when querying an individual audio device object about its bytes
-	FAudioDeviceManager::GetChecked().Buffers.CountBytes(Ar);
+	GEngine->GetAudioDeviceManager()->Buffers.CountBytes(Ar);
 	FreeSources.CountBytes(Ar);
 	WaveInstanceSourceMap.CountBytes(Ar);
 	Ar.CountBytes(sizeof(FWaveInstance) * WaveInstanceSourceMap.Num(), sizeof(FWaveInstance) * WaveInstanceSourceMap.Num());
@@ -1327,7 +1327,7 @@ bool FAudioDevice::HandleEnableHRTFForAllCommand(const TCHAR* Cmd, FOutputDevice
 bool FAudioDevice::HandleSoloCommand(const TCHAR* Cmd, FOutputDevice& Ar)
 {
 	// Apply the solo to the given device
-	FAudioDeviceManager* DeviceManager = FAudioDeviceManager::Get();
+	FAudioDeviceManager* DeviceManager = GEngine->GetAudioDeviceManager();
 	if (DeviceManager)
 	{
 		DeviceManager->SetSoloDevice(DeviceHandle);
@@ -1337,7 +1337,7 @@ bool FAudioDevice::HandleSoloCommand(const TCHAR* Cmd, FOutputDevice& Ar)
 
 bool FAudioDevice::HandleClearSoloCommand(const TCHAR* Cmd, FOutputDevice& Ar)
 {
-	FAudioDeviceManager* DeviceManager = FAudioDeviceManager::Get();
+	FAudioDeviceManager* DeviceManager = GEngine->GetAudioDeviceManager();
 	if (DeviceManager)
 	{
 		DeviceManager->SetSoloDevice(INDEX_NONE);
@@ -1347,7 +1347,7 @@ bool FAudioDevice::HandleClearSoloCommand(const TCHAR* Cmd, FOutputDevice& Ar)
 
 bool FAudioDevice::HandlePlayAllPIEAudioCommand(const TCHAR* Cmd, FOutputDevice& Ar)
 {
-	FAudioDeviceManager* DeviceManager = FAudioDeviceManager::Get();
+	FAudioDeviceManager* DeviceManager = GEngine->GetAudioDeviceManager();
 	if (DeviceManager)
 	{
 		DeviceManager->TogglePlayAllDeviceAudio();
@@ -1357,7 +1357,7 @@ bool FAudioDevice::HandlePlayAllPIEAudioCommand(const TCHAR* Cmd, FOutputDevice&
 
 bool FAudioDevice::HandleAudio3dVisualizeCommand(const TCHAR* Cmd, FOutputDevice& Ar)
 {
-	FAudioDeviceManager* DeviceManager = FAudioDeviceManager::Get();
+	FAudioDeviceManager* DeviceManager = GEngine->GetAudioDeviceManager();
 	if (DeviceManager)
 	{
 		DeviceManager->ToggleVisualize3dDebug();
@@ -1367,7 +1367,7 @@ bool FAudioDevice::HandleAudio3dVisualizeCommand(const TCHAR* Cmd, FOutputDevice
 
 void FAudioDevice::HandleAudioSoloCommon(const TCHAR* Cmd, FOutputDevice& Ar, FToggleSoloPtr FPtr)
 {
-	FAudioDeviceManager* DeviceManager = FAudioDeviceManager::Get();
+	FAudioDeviceManager* DeviceManager = GEngine->GetAudioDeviceManager();
 	if (DeviceManager)
 	{
 		bool bExclusive = true;
@@ -1404,7 +1404,7 @@ bool FAudioDevice::HandleAudioSoloSoundCue(const TCHAR* Cmd, FOutputDevice& Ar)
 
 bool FAudioDevice::HandleAudioMixerDebugSound(const TCHAR* Cmd, FOutputDevice& Ar)
 {
-	FAudioDeviceManager* DeviceManager = FAudioDeviceManager::Get();
+	FAudioDeviceManager* DeviceManager = GEngine->GetAudioDeviceManager();
 	if (DeviceManager)
 	{
 		DeviceManager->GetDebugger().SetAudioMixerDebugSound(Cmd);
@@ -1414,7 +1414,7 @@ bool FAudioDevice::HandleAudioMixerDebugSound(const TCHAR* Cmd, FOutputDevice& A
 
 bool FAudioDevice::HandleAudioDebugSound(const TCHAR* Cmd, FOutputDevice& Ar)
 {
-	FAudioDeviceManager* DeviceManager = FAudioDeviceManager::Get();
+	FAudioDeviceManager* DeviceManager = GEngine->GetAudioDeviceManager();
 	if (DeviceManager)
 	{
 		DeviceManager->GetDebugger().SetAudioDebugSound(Cmd);
@@ -1726,7 +1726,7 @@ bool FAudioDevice::HandleAudioMemoryInfo(const TCHAR* Cmd, FOutputDevice& Ar)
 
 bool FAudioDevice::HandleResetAllDynamicSoundVolumesCommand(const TCHAR* Cmd, FOutputDevice& Ar)
 {
-	if (FAudioDeviceManager* DeviceManager = FAudioDeviceManager::Get())
+	if (FAudioDeviceManager* DeviceManager = GEngine->GetAudioDeviceManager())
 	{
 		DeviceManager->ResetAllDynamicSoundVolumes();
 	}
@@ -1735,7 +1735,7 @@ bool FAudioDevice::HandleResetAllDynamicSoundVolumesCommand(const TCHAR* Cmd, FO
 
 bool FAudioDevice::HandleResetDynamicSoundVolumeCommand(const TCHAR* Cmd, FOutputDevice& Ar)
 {
-	if (FAudioDeviceManager* DeviceManager = FAudioDeviceManager::Get())
+	if (FAudioDeviceManager* DeviceManager = GEngine->GetAudioDeviceManager())
 	{
 		FName SoundName;
 		if (!FParse::Value(Cmd, TEXT("Name="), SoundName))
@@ -1770,55 +1770,61 @@ bool FAudioDevice::HandleGetDynamicSoundVolumeCommand(const TCHAR* Cmd, FOutputD
 		return false;
 	}
 
-	FName SoundName;
-	if (!FParse::Value(Cmd, TEXT("Name="), SoundName))
+	if (const FAudioDeviceManager* DeviceManager = GEngine->GetAudioDeviceManager())
 	{
-		return false;
-	}
-
-	// Optional: Defaults to Cue
-	FString SoundTypeStr;
-	ESoundType SoundType = ESoundType::Cue;
-	if (FParse::Value(Cmd, TEXT("Type="), SoundTypeStr))
-	{
-		if (SoundTypeStr == TEXT("Wave"))
+		FName SoundName;
+		if (!FParse::Value(Cmd, TEXT("Name="), SoundName))
 		{
-			SoundType = ESoundType::Wave;
+			return false;
 		}
-		else if (SoundTypeStr == TEXT("Class"))
+
+		// Optional: Defaults to Cue
+		FString SoundTypeStr;
+		ESoundType SoundType = ESoundType::Cue;
+		if (FParse::Value(Cmd, TEXT("Type="), SoundTypeStr))
 		{
-			SoundType = ESoundType::Class;
+			if (SoundTypeStr == TEXT("Wave"))
+			{
+				SoundType = ESoundType::Wave;
+			}
+			else if (SoundTypeStr == TEXT("Class"))
+			{
+				SoundType = ESoundType::Class;
+			}
 		}
-	}
-
-	if (!IsInAudioThread())
-	{
-		DECLARE_CYCLE_STAT(TEXT("FAudioThreadTask.GetDynamicSoundVolume"), STAT_AudioGetDynamicSoundVolume, STATGROUP_AudioThreadCommands);
 
-		const ESoundType InSoundType = SoundType;
-		const FName InSoundName = SoundName;
-		FAudioThread::RunCommandOnAudioThread([InSoundType, InSoundName]()
+		if (!IsInAudioThread())
 		{
-			if (FAudioDeviceManager* DeviceManager = FAudioDeviceManager::Get())
+			DECLARE_CYCLE_STAT(TEXT("FAudioThreadTask.GetDynamicSoundVolume"), STAT_AudioGetDynamicSoundVolume, STATGROUP_AudioThreadCommands);
+
+			const ESoundType InSoundType = SoundType;
+			const FName InSoundName = SoundName;
+			FAudioThread::RunCommandOnAudioThread([InSoundType, InSoundName]()
 			{
-				const float Volume = DeviceManager->GetDynamicSoundVolume(InSoundType, InSoundName);
-				UE_LOG(LogAudio, Display, TEXT("'%s' Dynamic Volume: %.4f"), *InSoundName.GetPlainNameString(), Volume);
-			}
-		}, GET_STATID(STAT_AudioGetDynamicSoundVolume));
-	}
-	else if (const FAudioDeviceManager* DeviceManager = FAudioDeviceManager::Get())
-	{
-		const float Volume = DeviceManager->GetDynamicSoundVolume(SoundType, SoundName);
-		FString Msg = FString::Printf(TEXT("'%s' Dynamic Volume: %.4f"), *SoundName.GetPlainNameString(), Volume);
-		Ar.Logf(TEXT("%s"), *Msg);
+				if (!GEngine)
+				{
+					return;
+				}
+				if (const FAudioDeviceManager* InDeviceManager = GEngine->GetAudioDeviceManager())
+				{
+					const float Volume = InDeviceManager->GetDynamicSoundVolume(InSoundType, InSoundName);
+					UE_LOG(LogAudio, Display, TEXT("'%s' Dynamic Volume: %.4f"), *InSoundName.GetPlainNameString(), Volume);
+				}
+			}, GET_STATID(STAT_AudioGetDynamicSoundVolume));
+		}
+		else
+		{
+			const float Volume = DeviceManager->GetDynamicSoundVolume(SoundType, SoundName);
+			FString Msg = FString::Printf(TEXT("'%s' Dynamic Volume: %.4f"), *SoundName.GetPlainNameString(), Volume);
+			Ar.Logf(TEXT("%s"), *Msg);
+		}
 	}
-
 	return true;
 }
 
 bool FAudioDevice::HandleSetDynamicSoundCommand(const TCHAR* Cmd, FOutputDevice& Ar)
 {
-	if (FAudioDeviceManager* DeviceManager = FAudioDeviceManager::Get())
+	if (FAudioDeviceManager* DeviceManager = GEngine->GetAudioDeviceManager())
 	{
 		FName SoundName;
 		if (!FParse::Value(Cmd, TEXT("Name="), SoundName))
@@ -4349,7 +4355,7 @@ void FAudioDevice::SendUpdateResultsToGameThread(const int32 FirstActiveIndex)
 		// At shutdown, GEngine may already be null
 		if (GEngine)
 		{
-			if (FAudioDeviceManager* AudioDeviceManager = FAudioDeviceManager::Get())
+			if (FAudioDeviceManager* AudioDeviceManager = GEngine->GetAudioDeviceManager())
 			{
 				if (FAudioDevice* AudioDevice = AudioDeviceManager->GetAudioDevice(AudioDeviceID))
 				{
@@ -4490,7 +4496,7 @@ void FAudioDevice::AddNewActiveSoundInternal(const FActiveSound& NewActiveSound,
 	// Cull one-shot active sounds if we've reached our max limit of one shot active sounds before we attempt to evaluate concurrency
 	// Check for debug sound name
 #if !UE_BUILD_SHIPPING
-	FAudioDeviceManager* AudioDeviceManager = FAudioDeviceManager::Get();
+	FAudioDeviceManager* AudioDeviceManager = GEngine->GetAudioDeviceManager();
 	FString DebugSound;
 	if (AudioDeviceManager->GetDebugger().GetAudioDebugSound(DebugSound))
 	{
@@ -5336,8 +5342,8 @@ float FAudioDevice::GetFocusFactor(const float Azimuth, const FSoundAttenuationS
 FAudioDevice::FCreateComponentParams::FCreateComponentParams()
 	: World(nullptr)
 	, Actor(nullptr)
-	, AudioDevice(FAudioDeviceManager::GetMainDevice())
 {
+	AudioDevice = (GEngine ? GEngine->GetMainAudioDevice() : nullptr);
 	CommonInit();
 }
 
@@ -5412,7 +5418,7 @@ UAudioComponent* FAudioDevice::CreateComponent(USoundBase* Sound, UWorld* World,
 	}
 	else
 	{
-		Params = MakeUnique<FCreateComponentParams>(FAudioDeviceManager::GetMainDevice());
+		Params = MakeUnique<FCreateComponentParams>(GEngine->GetMainAudioDevice());
 	}
 
 	Params->bPlay = bPlay;
@@ -6023,7 +6029,7 @@ bool FAudioDevice::HandleListSoundsCommand(const TCHAR* Cmd, FOutputDevice& Ar)
 	Ar.Logf(TEXT("Listing all sounds:"));
 
 	// Get audio device manager since thats where sound buffers are stored
-	FAudioDeviceManager* AudioDeviceManager = FAudioDeviceManager::Get();
+	FAudioDeviceManager* AudioDeviceManager = GEngine->GetAudioDeviceManager();
 	check(AudioDeviceManager != nullptr);
 
 	TArray<FSoundBuffer*> AllSounds;
@@ -6137,7 +6143,7 @@ bool FAudioDevice::IsAudioDeviceMuted() const
 	check(IsInAudioThread());
 
 	// First check to see if the device manager has "bPlayAllPIEAudio" enabled
-	FAudioDeviceManager* DeviceManager = FAudioDeviceManager::Get();
+	FAudioDeviceManager* DeviceManager = GEngine->GetAudioDeviceManager();
 	if (DeviceManager && DeviceManager->IsPlayAllDeviceAudio())
 	{
 		return false;
diff --git a/Engine/Source/Runtime/Engine/Private/AudioDeviceManager.cpp b/Engine/Source/Runtime/Engine/Private/AudioDeviceManager.cpp
--- a/Engine/Source/Runtime/Engine/Private/AudioDeviceManager.cpp
+++ b/Engine/Source/Runtime/Engine/Private/AudioDeviceManager.cpp
@@ -18,11 +18,6 @@
 #include "Settings/LevelEditorMiscSettings.h"
 #endif
 
-namespace
-{
-	FAudioDeviceManager* AudioDeviceManager = nullptr;
-} // namespace <>
-
 // Private consts for helping with index/generation determination in audio device manager
 static const uint32 AUDIO_DEVICE_HANDLE_INDEX_BITS		= 24;
 static const uint32 AUDIO_DEVICE_HANDLE_INDEX_MASK		= (1 << AUDIO_DEVICE_HANDLE_INDEX_BITS) - 1;
@@ -82,7 +77,6 @@ FAudioDeviceManager::FAudioDeviceManager()
 	, NextResourceID(1)
 	, SoloDeviceHandle(INDEX_NONE)
 	, ActiveAudioDeviceHandle(INDEX_NONE)
-	, MainAudioDeviceHandle(INDEX_NONE)
 	, bUsingAudioMixer(false)
 	, bPlayAllDeviceAudio(false)
 	, bOnlyToggleAudioMixerOnce(false)
@@ -266,61 +260,29 @@ IAudioDeviceModule* FAudioDeviceManager::GetAudioDeviceModule()
 	return AudioDeviceModule;
 }
 
-FAudioDevice* FAudioDeviceManager::GetActiveDevice()
-{
-	if (AudioDeviceManager)
-	{
-		if (AudioDeviceManager->ActiveAudioDeviceHandle != INDEX_NONE)
-		{
-			return AudioDeviceManager->GetAudioDevice(AudioDeviceManager->ActiveAudioDeviceHandle);
-		}
-
-		return GetMainDevice();
-	}
-
-	return nullptr;
-}
-
-FAudioDevice* FAudioDeviceManager::GetMainDevice()
+bool FAudioDeviceManager::Initialize()
 {
-	if (FAudioDeviceManager* DeviceManager = Get())
+	if (LoadDefaultAudioDeviceModule())
 	{
-		return DeviceManager->GetAudioDevice(DeviceManager->GetMainAudioDeviceHandle());
-	}
-
-	return nullptr;
-}
-
-void FAudioDeviceManager::Initialize()
-{
-	if (!AudioDeviceManager)
-	{
-		AudioDeviceManager = new FAudioDeviceManager();
-		if (AudioDeviceManager->LoadDefaultAudioDeviceModule())
-		{
-			check(AudioDeviceManager->AudioDeviceModule);
+		check(AudioDeviceModule);
 
-			const bool bIsAudioMixerEnabled = AudioDeviceManager->AudioDeviceModule->IsAudioMixerModule();
-			GetMutableDefault<UAudioSettings>()->SetAudioMixerEnabled(bIsAudioMixerEnabled);
+		const bool bIsAudioMixerEnabled = AudioDeviceModule->IsAudioMixerModule();
+		GetMutableDefault<UAudioSettings>()->SetAudioMixerEnabled(bIsAudioMixerEnabled);
 
 #if WITH_EDITOR
-			if (bIsAudioMixerEnabled)
-			{
-				IAudioEditorModule* AudioEditorModule = &FModuleManager::LoadModuleChecked<IAudioEditorModule>("AudioEditor");
-				AudioEditorModule->RegisterAudioMixerAssetActions();
-				AudioEditorModule->RegisterEffectPresetAssetActions();
-			}
-#endif // WITH_EDITOR
-
-			if (AudioDeviceManager->CreateMainAudioDevice())
-			{
-				return;
-			}
+		if (bIsAudioMixerEnabled)
+		{
+			IAudioEditorModule* AudioEditorModule = &FModuleManager::LoadModuleChecked<IAudioEditorModule>("AudioEditor");
+			AudioEditorModule->RegisterAudioMixerAssetActions();
+			AudioEditorModule->RegisterEffectPresetAssetActions();
 		}
+#endif
+
+		return CreateMainAudioDevice();
 	}
 
 	// Failed to initialize
-	Shutdown();
+	return false;
 }
 
 bool FAudioDeviceManager::LoadDefaultAudioDeviceModule()
@@ -423,10 +385,11 @@ bool FAudioDeviceManager::CreateAudioDevice(bool bCreateNewDevice, FCreateAudioD
 	{
 		if (NumActiveAudioDevices == 1)
 		{
-			if (FAudioDevice* MainDevice = FAudioDeviceManager::GetMainDevice())
+			FAudioDevice* MainAudioDevice = GEngine->GetMainAudioDevice();
+			if (MainAudioDevice)
 			{
-				OutResults.Handle = MainDevice->DeviceHandle;
-				OutResults.AudioDevice = MainDevice;
+				OutResults.Handle = MainAudioDevice->DeviceHandle;
+				OutResults.AudioDevice = MainAudioDevice;
 				OutResults.AudioDevice->FadeIn();
 				return true;
 			}
@@ -445,11 +408,12 @@ bool FAudioDeviceManager::CreateAudioDevice(bool bCreateNewDevice, FCreateAudioD
 
 	if (NumActiveAudioDevices == 1 && !bCreateNewAudioDeviceForPlayInEditor)
 	{
-		if (FAudioDevice* MainDevice = FAudioDeviceManager::GetMainDevice())
+		FAudioDevice* MainAudioDevice = GEngine->GetMainAudioDevice();
+		if (MainAudioDevice)
 		{
 			++NumWorldsUsingMainAudioDevice;
-			OutResults.Handle = MainDevice->DeviceHandle;
-			OutResults.AudioDevice = MainDevice;
+			OutResults.Handle = MainAudioDevice->DeviceHandle;
+			OutResults.AudioDevice = MainAudioDevice;
 			bRequiresInit = false;
 		}
 		else
@@ -511,7 +475,8 @@ bool FAudioDeviceManager::CreateAudioDevice(bool bCreateNewDevice, FCreateAudioD
 		else
 		{
 			++NumWorldsUsingMainAudioDevice;
-			if (FAudioDevice* MainAudioDevice = FAudioDeviceManager::GetMainDevice())
+			FAudioDevice* MainAudioDevice = GEngine->GetMainAudioDevice();
+			if (MainAudioDevice)
 			{
 				OutResults.Handle = MainAudioDevice->DeviceHandle;
 				OutResults.AudioDevice = MainAudioDevice;
@@ -579,7 +544,7 @@ bool FAudioDeviceManager::ShutdownAudioDevice(Audio::FDeviceId Handle)
 	// If there are more than 1 device active, check to see if this handle is the main audio device handle
 	if (NumActiveAudioDevices >= 1)
 	{
-		Audio::FDeviceId MainDeviceHandle = FAudioDeviceManager::GetChecked().GetMainAudioDeviceHandle();
+		uint32 MainDeviceHandle = GEngine->GetAudioDeviceHandle();
 
 		if (NumActiveAudioDevices == 1)
 		{
@@ -592,7 +557,7 @@ bool FAudioDeviceManager::ShutdownAudioDevice(Audio::FDeviceId Handle)
 		// don't shut it down until it's the very last handle to get shut down
 		// this is because it's possible for some PIE sessions to be using the main audio device as a fallback to
 		// preserve CPU performance on low-performance machines
-		if (NumWorldsUsingMainAudioDevice > 0 && MainDeviceHandle == static_cast<uint32>(Handle))
+		if (NumWorldsUsingMainAudioDevice > 0 && MainDeviceHandle == Handle)
 		{
 			--NumWorldsUsingMainAudioDevice;
 
@@ -660,13 +625,21 @@ FAudioDevice* FAudioDeviceManager::GetAudioDevice(Audio::FDeviceId Handle)
 
 FAudioDeviceManager* FAudioDeviceManager::Get()
 {
-	return AudioDeviceManager;
+	if (GEngine)
+	{
+		return GEngine->GetAudioDeviceManager();
+	}
+
+	return nullptr;
 }
 
-FAudioDeviceManager& FAudioDeviceManager::GetChecked()
+FAudioDevice* FAudioDeviceManager::GetActiveAudioDevice()
 {
-	check(AudioDeviceManager);
-	return *AudioDeviceManager;
+	if (ActiveAudioDeviceHandle != INDEX_NONE)
+	{
+		return GetAudioDevice(ActiveAudioDeviceHandle);
+	}
+	return GEngine->GetMainAudioDevice();
 }
 
 void FAudioDeviceManager::UpdateActiveAudioDevices(bool bGameTicking)
@@ -876,23 +849,6 @@ void FAudioDeviceManager::SetSoloDevice(Audio::FDeviceId InAudioDeviceHandle)
 	}
 }
 
-void FAudioDeviceManager::Shutdown()
-{
-	// Shutdown the main audio device in the UEEngine
-	if (AudioDeviceManager)
-	{
-		FAudioCommandFence Fence;
-		Fence.BeginFence();
-		Fence.Wait();
-
-		FAudioThread::StopAudioThread();
-
-		AudioDeviceManager->ShutdownAllAudioDevices();
-
-		delete AudioDeviceManager;
-		AudioDeviceManager = nullptr;
-	}
-}
 
 uint8 FAudioDeviceManager::GetNumActiveAudioDevices() const
 {
@@ -994,7 +950,8 @@ void FAudioDeviceManager::RemoveSoundMix(USoundMix* SoundMix)
 	{
 		DECLARE_CYCLE_STAT(TEXT("FAudioThreadTask.RemoveSoundMix"), STAT_AudioRemoveSoundMix, STATGROUP_AudioThreadCommands);
 
-		FAudioThread::RunCommandOnAudioThread([this, SoundMix]()
+		FAudioDeviceManager* AudioDeviceManager = this;
+		FAudioThread::RunCommandOnAudioThread([AudioDeviceManager, SoundMix]()
 		{
 			AudioDeviceManager->RemoveSoundMix(SoundMix);
 
@@ -1018,9 +975,10 @@ void FAudioDeviceManager::TogglePlayAllDeviceAudio()
 	{
 		DECLARE_CYCLE_STAT(TEXT("FAudioThreadTask.TogglePlayAllDeviceAudio"), STAT_TogglePlayAllDeviceAudio, STATGROUP_AudioThreadCommands);
 
-		FAudioThread::RunCommandOnAudioThread([this]()
+		FAudioDeviceManager* AudioDeviceManager = this;
+		FAudioThread::RunCommandOnAudioThread([AudioDeviceManager]()
 		{
-			TogglePlayAllDeviceAudio();
+			AudioDeviceManager->TogglePlayAllDeviceAudio();
 
 		}, GET_STATID(STAT_TogglePlayAllDeviceAudio));
 
@@ -1046,9 +1004,11 @@ void FAudioDeviceManager::ToggleVisualize3dDebug()
 	{
 		DECLARE_CYCLE_STAT(TEXT("FAudioThreadTask.ToggleVisualize3dDebug"), STAT_ToggleVisualize3dDebug, STATGROUP_AudioThreadCommands);
 
-		FAudioThread::RunCommandOnAudioThread([this]()
+		FAudioDeviceManager* AudioDeviceManager = this;
+		FAudioThread::RunCommandOnAudioThread([AudioDeviceManager]()
 		{
-			ToggleVisualize3dDebug();
+			AudioDeviceManager->ToggleVisualize3dDebug();
+
 		}, GET_STATID(STAT_ToggleVisualize3dDebug));
 
 		return;
@@ -1077,9 +1037,10 @@ void FAudioDeviceManager::ResetAllDynamicSoundVolumes()
 	{
 		DECLARE_CYCLE_STAT(TEXT("FAudioThreadTask.ResetAllDynamicSoundVolumes"), STAT_ResetAllDynamicSoundVolumes, STATGROUP_AudioThreadCommands);
 
-		FAudioThread::RunCommandOnAudioThread([this]()
+		FAudioDeviceManager* AudioDeviceManager = this;
+		FAudioThread::RunCommandOnAudioThread([AudioDeviceManager]()
 		{
-			ResetAllDynamicSoundVolumes();
+			AudioDeviceManager->ResetAllDynamicSoundVolumes();
 
 		}, GET_STATID(STAT_ResetAllDynamicSoundVolumes));
 		return;
@@ -1095,9 +1056,11 @@ void FAudioDeviceManager::ResetDynamicSoundVolume(ESoundType SoundType, const FN
 	{
 		DECLARE_CYCLE_STAT(TEXT("FAudioThreadTask.ResetSoundCueTrimVolume"), STAT_ResetSoundCueTrimVolume, STATGROUP_AudioThreadCommands);
 
-		FAudioThread::RunCommandOnAudioThread([this, SoundType, SoundName]()
+		FAudioDeviceManager* AudioDeviceManager = this;
+		FAudioThread::RunCommandOnAudioThread([AudioDeviceManager, SoundType, SoundName]()
 		{
-			ResetDynamicSoundVolume(SoundType, SoundName);
+			AudioDeviceManager->ResetDynamicSoundVolume(SoundType, SoundName);
+
 		}, GET_STATID(STAT_ResetSoundCueTrimVolume));
 		return;
 	}
@@ -1112,9 +1075,11 @@ void FAudioDeviceManager::SetDynamicSoundVolume(ESoundType SoundType, const FNam
 	{
 		DECLARE_CYCLE_STAT(TEXT("FAudioThreadTask.SetDynamicSoundVolume"), STAT_SetDynamicSoundVolume, STATGROUP_AudioThreadCommands);
 
-		FAudioThread::RunCommandOnAudioThread([this, SoundType, SoundName, Volume]()
+		FAudioDeviceManager* AudioDeviceManager = this;
+		FAudioThread::RunCommandOnAudioThread([AudioDeviceManager, SoundType, SoundName, Volume]()
 		{
-			SetDynamicSoundVolume(SoundType, SoundName, Volume);
+			AudioDeviceManager->SetDynamicSoundVolume(SoundType, SoundName, Volume);
+
 		}, GET_STATID(STAT_SetDynamicSoundVolume));
 		return;
 	}
diff --git a/Engine/Source/Runtime/Engine/Private/AudioEffect.cpp b/Engine/Source/Runtime/Engine/Private/AudioEffect.cpp
--- a/Engine/Source/Runtime/Engine/Private/AudioEffect.cpp
+++ b/Engine/Source/Runtime/Engine/Private/AudioEffect.cpp
@@ -510,7 +510,8 @@ void FAudioEffectsManager::Update()
 	if (CurrentReverbAsset && CurrentReverbAsset->bChanged)
 	{
 		CurrentReverbAsset->bChanged = false;
-		TArray<FAudioDevice*>& AudioDevices = FAudioDeviceManager::GetChecked().GetAudioDevices();
+		class FAudioDeviceManager* DeviceManager = GEngine->GetAudioDeviceManager();
+		TArray<FAudioDevice*>& AudioDevices = DeviceManager->GetAudioDevices();
 		for (int32 i = 0; i < AudioDevices.Num(); ++i)
 		{
 			if (AudioDevices[i])
@@ -521,7 +522,8 @@ void FAudioEffectsManager::Update()
 		}
 		
 	}
-#endif // WITH_EDITORONLY_DATA
+
+#endif
 
 	const bool bIsReverbDone = CurrentReverbEffect.Interpolate(SourceReverbEffect, DestinationReverbEffect);
 	if (!bIsReverbDone || bReverbActive || bReverbChanged)
diff --git a/Engine/Source/Runtime/Engine/Private/Components/AudioComponent.cpp b/Engine/Source/Runtime/Engine/Private/Components/AudioComponent.cpp
--- a/Engine/Source/Runtime/Engine/Private/Components/AudioComponent.cpp
+++ b/Engine/Source/Runtime/Engine/Private/Components/AudioComponent.cpp
@@ -556,7 +556,8 @@ FAudioDevice* UAudioComponent::GetAudioDevice() const
 	{
 		if (AudioDeviceHandle != INDEX_NONE)
 		{
-			AudioDevice = FAudioDeviceManager::GetChecked().GetAudioDevice(AudioDeviceHandle);
+			FAudioDeviceManager* AudioDeviceManager = GEngine->GetAudioDeviceManager();
+			AudioDevice = (AudioDeviceManager ? AudioDeviceManager->GetAudioDevice(AudioDeviceHandle) : nullptr);
 		}
 		else if (UWorld* World = GetWorld())
 		{
@@ -564,7 +565,7 @@ FAudioDevice* UAudioComponent::GetAudioDevice() const
 		}
 		else
 		{
-			AudioDevice = FAudioDeviceManager::GetMainDevice();
+			AudioDevice = GEngine->GetMainAudioDevice();
 		}
 	}
 	return AudioDevice;
diff --git a/Engine/Source/Runtime/Engine/Private/ContentStreaming.cpp b/Engine/Source/Runtime/Engine/Private/ContentStreaming.cpp
--- a/Engine/Source/Runtime/Engine/Private/ContentStreaming.cpp
+++ b/Engine/Source/Runtime/Engine/Private/ContentStreaming.cpp
@@ -1179,15 +1179,13 @@ void FStreamingManagerCollection::PropagateLightingScenarioChange()
 void FStreamingManagerCollection::OnAudioStreamingParamsChanged()
 {
 	// Before we swap out the audio streaming manager, we'll need to stop all sounds running on all audio devices:
-	if (FAudioDeviceManager* DeviceManager = FAudioDeviceManager::Get())
+	FAudioDeviceManager* DeviceManager = GEngine->GetAudioDeviceManager();
+	TArray<FAudioDevice*>& AudioDevices = DeviceManager->GetAudioDevices();
+	for (FAudioDevice* AudioDevice : AudioDevices)
 	{
-		TArray<FAudioDevice*>& AudioDevices = DeviceManager->GetAudioDevices();
-		for (FAudioDevice* AudioDevice : AudioDevices)
+		if (AudioDevice)
 		{
-			if (AudioDevice)
-			{
-				AudioDevice->StopAllSounds(true);
-			}
+			AudioDevice->StopAllSounds(true);
 		}
 	}
 
diff --git a/Engine/Source/Runtime/Engine/Private/GameEngine.cpp b/Engine/Source/Runtime/Engine/Private/GameEngine.cpp
--- a/Engine/Source/Runtime/Engine/Private/GameEngine.cpp
+++ b/Engine/Source/Runtime/Engine/Private/GameEngine.cpp
@@ -1801,7 +1801,7 @@ void UGameEngine::Tick( float DeltaSeconds, bool bIdleMode )
 	}
 
 	// Update Audio. This needs to occur after rendering as the rendering code updates the listener position.
-	FAudioDeviceManager* GameAudioDeviceManager = FAudioDeviceManager::Get();
+	FAudioDeviceManager* GameAudioDeviceManager = GEngine->GetAudioDeviceManager();
 	if (GameAudioDeviceManager)
 	{
 		SCOPE_TIME_GUARD(TEXT("UGameEngine::Tick - Update Audio"));
diff --git a/Engine/Source/Runtime/Engine/Private/GameFramework/HapticFeedbackEffect.cpp b/Engine/Source/Runtime/Engine/Private/GameFramework/HapticFeedbackEffect.cpp
--- a/Engine/Source/Runtime/Engine/Private/GameFramework/HapticFeedbackEffect.cpp
+++ b/Engine/Source/Runtime/Engine/Private/GameFramework/HapticFeedbackEffect.cpp
@@ -168,7 +168,7 @@ float UHapticFeedbackEffect_SoundWave::GetDuration() const
 
 void UHapticFeedbackEffect_SoundWave::PrepareSoundWaveBuffer()
 {
-	FAudioDevice *AD = FAudioDeviceManager::GetMainDevice();
+	FAudioDevice *AD = GEngine->GetMainAudioDevice();
 	if (!AD || !SoundWave)
 	{
 		return;
diff --git a/Engine/Source/Runtime/Engine/Private/GameUserSettings.cpp b/Engine/Source/Runtime/Engine/Private/GameUserSettings.cpp
--- a/Engine/Source/Runtime/Engine/Private/GameUserSettings.cpp
+++ b/Engine/Source/Runtime/Engine/Private/GameUserSettings.cpp
@@ -447,7 +447,8 @@ void UGameUserSettings::ApplyNonResolutionSettings()
 		Scalability::SetQualityLevels(ScalabilityQuality);
 	}
 
-	if (FAudioDevice* AudioDevice = FAudioDeviceManager::GetMainDevice())
+	FAudioDevice* AudioDevice = GEngine->GetMainAudioDevice();
+	if (AudioDevice)
 	{
 		FAudioQualitySettings AudioSettings = AudioDevice->GetQualityLevelSettings();
 		AudioDevice->SetMaxChannels(AudioSettings.MaxChannels);
diff --git a/Engine/Source/Runtime/Engine/Private/GameViewportClient.cpp b/Engine/Source/Runtime/Engine/Private/GameViewportClient.cpp
--- a/Engine/Source/Runtime/Engine/Private/GameViewportClient.cpp
+++ b/Engine/Source/Runtime/Engine/Private/GameViewportClient.cpp
@@ -303,10 +303,13 @@ void UGameViewportClient::PostInitProperties()
 
 void UGameViewportClient::BeginDestroy()
 {
-	class FAudioDeviceManager* AudioDeviceManager = FAudioDeviceManager::Get();
-	if (AudioDeviceManager)
+	if (GEngine)
 	{
-		AudioDeviceManager->ShutdownAudioDevice(AudioDeviceHandle);
+		class FAudioDeviceManager* AudioDeviceManager = GEngine->GetAudioDeviceManager();
+		if (AudioDeviceManager)
+		{
+			AudioDeviceManager->ShutdownAudioDevice(AudioDeviceHandle);
+		}
 	}
 
 	RemoveAllViewportWidgets();
@@ -379,7 +382,7 @@ void UGameViewportClient::SetEnabledStats(const TArray<FString>& InEnabledStats)
 #if ENABLE_AUDIO_DEBUG
 	if (UWorld* MyWorld = GetWorld())
 	{
-		if (FAudioDeviceManager* DeviceManager = FAudioDeviceManager::Get())
+		if (FAudioDeviceManager* DeviceManager = GEngine->GetAudioDeviceManager())
 		{
 			FAudioDebugger::ResolveDesiredStats(this);
 		}
@@ -433,30 +436,33 @@ void UGameViewportClient::Init(struct FWorldContext& WorldContext, UGameInstance
 	// Create the cursor Widgets
 	UUserInterfaceSettings* UISettings = GetMutableDefault<UUserInterfaceSettings>(UUserInterfaceSettings::StaticClass());
 
-	FAudioDeviceManager* AudioDeviceManager = FAudioDeviceManager::Get();
-	if (AudioDeviceManager)
+	if (GEngine)
 	{
-		FAudioDeviceManager::FCreateAudioDeviceResults NewDeviceResults;
-		if (AudioDeviceManager->CreateAudioDevice(bCreateNewAudioDevice, NewDeviceResults))
+		FAudioDeviceManager* AudioDeviceManager = GEngine->GetAudioDeviceManager();
+		if (AudioDeviceManager)
 		{
-			AudioDeviceHandle = NewDeviceResults.Handle;
+			FAudioDeviceManager::FCreateAudioDeviceResults NewDeviceResults;
+			if (AudioDeviceManager->CreateAudioDevice(bCreateNewAudioDevice, NewDeviceResults))
+			{
+				AudioDeviceHandle = NewDeviceResults.Handle;
 
 #if ENABLE_AUDIO_DEBUG
-			FAudioDebugger::ResolveDesiredStats(this);
+				FAudioDebugger::ResolveDesiredStats(this);
 #endif // ENABLE_AUDIO_DEBUG
 
-			// Set the base mix of the new device based on the world settings of the world
-			if (World)
-			{
-				NewDeviceResults.AudioDevice->SetDefaultBaseSoundMix(World->GetWorldSettings()->DefaultBaseSoundMix);
+				// Set the base mix of the new device based on the world settings of the world
+				if (World)
+				{
+					NewDeviceResults.AudioDevice->SetDefaultBaseSoundMix(World->GetWorldSettings()->DefaultBaseSoundMix);
 
-				// Set the world's audio device handle to use so that sounds which play in that world will use the correct audio device
-				World->SetAudioDeviceHandle(AudioDeviceHandle);
-			}
+					// Set the world's audio device handle to use so that sounds which play in that world will use the correct audio device
+					World->SetAudioDeviceHandle(AudioDeviceHandle);
+				}
 
-			// Set this audio device handle on the world context so future world's set onto the world context
-			// will pass the audio device handle to them and audio will play on the correct audio device
-			WorldContext.AudioDeviceHandle = AudioDeviceHandle;
+				// Set this audio device handle on the world context so future world's set onto the world context
+				// will pass the audio device handle to them and audio will play on the correct audio device
+				WorldContext.AudioDeviceHandle = AudioDeviceHandle;
+			}
 		}
 	}
 
@@ -1303,15 +1309,18 @@ void UGameViewportClient::Draw(FViewport* InViewport, FCanvas* SceneCanvas)
 
 							// If the main audio device is used for multiple PIE viewport clients, we only
 							// want to update the main audio device listener position if it is in focus
-							FAudioDeviceManager* DeviceManager = FAudioDeviceManager::Get();
-
-							// If there is more than one world referencing the main audio device
-							if (DeviceManager->GetNumMainAudioDeviceWorlds() > 1)
+							if (GEngine)
 							{
-								uint32 MainAudioDeviceHandle = DeviceManager->GetMainAudioDeviceHandle();
-								if (AudioDevice->DeviceHandle == MainAudioDeviceHandle && !bHasAudioFocus)
+								FAudioDeviceManager* AudioDeviceManager = GEngine->GetAudioDeviceManager();
+
+								// If there is more than one world referencing the main audio device
+								if (AudioDeviceManager->GetNumMainAudioDeviceWorlds() > 1)
 								{
-									bUpdateListenerPosition = false;
+									uint32 MainAudioDeviceHandle = GEngine->GetAudioDeviceHandle();
+									if (AudioDevice->DeviceHandle == MainAudioDeviceHandle && !bHasAudioFocus)
+									{
+										bUpdateListenerPosition = false;
+									}
 								}
 							}
 
@@ -1867,7 +1876,7 @@ void UGameViewportClient::LostFocus(FViewport* InViewport)
 		}
 	}
 
-	if (FAudioDeviceManager::Get())
+	if (GEngine && GEngine->GetAudioDeviceManager())
 	{
 		bHasAudioFocus = false;
 	}
@@ -1880,9 +1889,9 @@ void UGameViewportClient::ReceivedFocus(FViewport* InViewport)
 		FSlateApplication::Get().SetGameIsFakingTouchEvents(true);
 	}
 
-	if (FAudioDeviceManager* DeviceManager = FAudioDeviceManager::Get())
+	if (GEngine && GEngine->GetAudioDeviceManager())
 	{
-		DeviceManager->SetActiveDevice(AudioDeviceHandle);
+		GEngine->GetAudioDeviceManager()->SetActiveDevice(AudioDeviceHandle);
 		bHasAudioFocus = true;
 	}
 }
diff --git a/Engine/Source/Runtime/Engine/Private/Interpolation.cpp b/Engine/Source/Runtime/Engine/Private/Interpolation.cpp
--- a/Engine/Source/Runtime/Engine/Private/Interpolation.cpp
+++ b/Engine/Source/Runtime/Engine/Private/Interpolation.cpp
@@ -1271,9 +1271,10 @@ void AMatineeActor::StepInterp( float DeltaTime, bool bPreview )
 
 void AMatineeActor::DisableRadioFilterIfNeeded()
 {
-	if(FAudioDevice* AudioDevice = FAudioDeviceManager::GetMainDevice())
+	FAudioDevice* AudioDevice = GEngine->GetMainAudioDevice();
+	if( AudioDevice )
 	{
-		AudioDevice->EnableRadioEffect(!bDisableRadioFilter);
+		AudioDevice->EnableRadioEffect( !bDisableRadioFilter );
 	}
 }
 
@@ -1298,9 +1299,10 @@ void AMatineeActor::EnableCinematicMode(bool bEnable)
 
 void AMatineeActor::EnableRadioFilter()
 {
-	if (FAudioDevice* AudioDevice = FAudioDeviceManager::GetMainDevice())
+	FAudioDevice* AudioDevice = GEngine->GetMainAudioDevice();
+	if( AudioDevice )
 	{
-		AudioDevice->EnableRadioEffect(true);
+		AudioDevice->EnableRadioEffect( true );
 	}
 }
 
diff --git a/Engine/Source/Runtime/Engine/Private/LevelActor.cpp b/Engine/Source/Runtime/Engine/Private/LevelActor.cpp
--- a/Engine/Source/Runtime/Engine/Private/LevelActor.cpp
+++ b/Engine/Source/Runtime/Engine/Private/LevelActor.cpp
@@ -1512,17 +1512,20 @@ void UWorld::SetAudioDeviceHandle(const uint32 InAudioDeviceHandle)
 
 FAudioDevice* UWorld::GetAudioDevice()
 {
-	if (FAudioDeviceManager* AudioDeviceManager = FAudioDeviceManager::Get())
+	FAudioDevice* AudioDevice = nullptr;
+	if (GEngine)
 	{
-		if (FAudioDevice* AudioDevice = AudioDeviceManager->GetAudioDevice(AudioDeviceHandle))
+		class FAudioDeviceManager* AudioDeviceManager = GEngine->GetAudioDeviceManager();
+		if (AudioDeviceManager != nullptr)
 		{
-			return AudioDevice;
+			AudioDevice = AudioDeviceManager->GetAudioDevice(AudioDeviceHandle);
+			if (AudioDevice == nullptr)
+			{
+				AudioDevice = GEngine->GetMainAudioDevice();
+			}
 		}
-
-		return FAudioDeviceManager::GetMainDevice();
 	}
-
-	return nullptr;
+	return AudioDevice;
 }
 
 /**
diff --git a/Engine/Source/Runtime/Engine/Private/SampleBufferIO.cpp b/Engine/Source/Runtime/Engine/Private/SampleBufferIO.cpp
--- a/Engine/Source/Runtime/Engine/Private/SampleBufferIO.cpp
+++ b/Engine/Source/Runtime/Engine/Private/SampleBufferIO.cpp
@@ -17,7 +17,8 @@ namespace Audio
 
 	void FSoundWavePCMLoader::LoadSoundWave(USoundWave* InSoundWave, TFunction<void(const USoundWave* SoundWave, const Audio::FSampleBuffer& OutSampleBuffer)> OnLoaded)
 	{
-		FAudioDevice* AudioDevice = FAudioDeviceManager::GetMainDevice();
+		FAudioDevice* AudioDevice = FAudioDevice::GetMainAudioDevice();
+
 		if (!AudioDevice || !InSoundWave)
 		{
 			return;
@@ -121,9 +122,9 @@ namespace Audio
 			CurrentSoundWave = SoundWaveToSaveTo;
 			bWasPreviouslyAddedToRoot = CurrentSoundWave->IsRooted();
 			CurrentSoundWave->AddToRoot();
-
 			// Ensure this sound wave is not currently in use:
-			if (FAudioDeviceManager* AudioDeviceManager = FAudioDeviceManager::Get())
+			FAudioDeviceManager* AudioDeviceManager = GEngine->GetAudioDeviceManager();
+			if (AudioDeviceManager)
 			{
 				AudioDeviceManager->StopSoundsUsingResource(CurrentSoundWave);
 			}
diff --git a/Engine/Source/Runtime/Engine/Private/Slate/SlateSoundDevice.cpp b/Engine/Source/Runtime/Engine/Private/Slate/SlateSoundDevice.cpp
--- a/Engine/Source/Runtime/Engine/Private/Slate/SlateSoundDevice.cpp
+++ b/Engine/Source/Runtime/Engine/Private/Slate/SlateSoundDevice.cpp
@@ -10,13 +10,14 @@ DEFINE_LOG_CATEGORY_STATIC(LogSlateSoundDevice, Log, All);
 
 void FSlateSoundDevice::PlaySound(const FSlateSound& Sound, int32 UserIndex) const
 {
-	if (Sound.GetResourceObject())
+	if( GEngine && Sound.GetResourceObject() != nullptr )
 	{
-		if (FAudioDevice* AudioDevice = FAudioDeviceManager::GetActiveDevice())
+		FAudioDevice* const AudioDevice = GEngine->GetActiveAudioDevice();
+		if(AudioDevice)
 		{
 			UObject* const Object = Sound.GetResourceObject();
 			USoundBase* const SoundResource = Cast<USoundBase>(Object);
-			if (SoundResource)
+			if(SoundResource)
 			{
 				FActiveSound NewActiveSound;
 				NewActiveSound.SetSound(SoundResource);
@@ -26,7 +27,7 @@ void FSlateSoundDevice::PlaySound(const FSlateSound& Sound, int32 UserIndex) con
 
 				AudioDevice->AddNewActiveSound(NewActiveSound);
 			}
-			else if (Object)
+			else if(Object)
 			{
 				// An Object but no SoundResource means that the FSlateSound is holding an invalid object; report that as an error
 				UE_LOG(LogSlateSoundDevice, Error, TEXT("A sound contains a non-sound resource '%s'"), *Object->GetName());
diff --git a/Engine/Source/Runtime/Engine/Private/SoundBase.cpp b/Engine/Source/Runtime/Engine/Private/SoundBase.cpp
--- a/Engine/Source/Runtime/Engine/Private/SoundBase.cpp
+++ b/Engine/Source/Runtime/Engine/Private/SoundBase.cpp
@@ -1,11 +1,9 @@
 // Copyright Epic Games, Inc. All Rights Reserved.
 #include "Sound/SoundBase.h"
 
-#include "Audio.h"
 #include "EngineDefines.h"
 #include "IAudioExtensionPlugin.h"
 #include "Sound/AudioSettings.h"
-#include "Sound/SoundAttenuation.h"
 #include "Sound/SoundClass.h"
 #include "Sound/SoundSubmix.h"
 
diff --git a/Engine/Source/Runtime/Engine/Private/SoundClass.cpp b/Engine/Source/Runtime/Engine/Private/SoundClass.cpp
--- a/Engine/Source/Runtime/Engine/Private/SoundClass.cpp
+++ b/Engine/Source/Runtime/Engine/Private/SoundClass.cpp
@@ -103,8 +103,11 @@ void USoundClass::PostLoad()
 	}
 #endif
 
+	// Use the main/default audio device for storing and retrieving sound class properties
+	FAudioDeviceManager* AudioDeviceManager = (GEngine ? GEngine->GetAudioDeviceManager() : nullptr);
+
 	// Force the properties to be initialized for this SoundClass on all active audio devices
-	if (FAudioDeviceManager* AudioDeviceManager = FAudioDeviceManager::Get())
+	if (AudioDeviceManager)
 	{
 		AudioDeviceManager->RegisterSoundClass(this);
 	}
@@ -244,8 +247,11 @@ void USoundClass::PostEditChangeProperty(struct FPropertyChangedEvent& PropertyC
 		}
 	}
 
+	// Use the main/default audio device for storing and retrieving sound class properties
+	FAudioDeviceManager* AudioDeviceManager = (GEngine ? GEngine->GetAudioDeviceManager() : nullptr);
+
 	// Force the properties to be initialized for this SoundClass on all active audio devices
-	if (FAudioDeviceManager* AudioDeviceManager = FAudioDeviceManager::Get())
+	if (AudioDeviceManager)
 	{
 		AudioDeviceManager->RegisterSoundClass(this);
 	}
@@ -292,12 +298,9 @@ void USoundClass::BeginDestroy()
 {
 	Super::BeginDestroy();
 
-	if (!GExitPurge)
+	if (!GExitPurge && GEngine && GEngine->GetAudioDeviceManager())
 	{
-		if (FAudioDeviceManager* DeviceManager = FAudioDeviceManager::Get())
-		{
-			DeviceManager->UnregisterSoundClass(this);
-		}
+		GEngine->GetAudioDeviceManager()->UnregisterSoundClass(this);
 	}
 }
 
diff --git a/Engine/Source/Runtime/Engine/Private/SoundEffectPreset.cpp b/Engine/Source/Runtime/Engine/Private/SoundEffectPreset.cpp
--- a/Engine/Source/Runtime/Engine/Private/SoundEffectPreset.cpp
+++ b/Engine/Source/Runtime/Engine/Private/SoundEffectPreset.cpp
@@ -93,7 +93,11 @@ void USoundEffectPreset::PostEditChangeProperty(FPropertyChangedEvent& PropertyC
 
 void USoundEffectSourcePresetChain::PostEditChangeProperty(FPropertyChangedEvent& PropertyChangedEvent)
 {
-	FAudioDeviceManager::GetChecked().UpdateSourceEffectChain(GetUniqueID(), Chain, bPlayEffectChainTails);
+	if (GEngine)
+	{
+		FAudioDeviceManager* AudioDeviceManager = GEngine->GetAudioDeviceManager();
+		AudioDeviceManager->UpdateSourceEffectChain(GetUniqueID(), Chain, bPlayEffectChainTails);
+	}
 }
 #endif // WITH_EDITORONLY_DATA
 
diff --git a/Engine/Source/Runtime/Engine/Private/SoundMix.cpp b/Engine/Source/Runtime/Engine/Private/SoundMix.cpp
--- a/Engine/Source/Runtime/Engine/Private/SoundMix.cpp
+++ b/Engine/Source/Runtime/Engine/Private/SoundMix.cpp
@@ -32,11 +32,12 @@ USoundMix::USoundMix(const FObjectInitializer& ObjectInitializer)
 
 void USoundMix::BeginDestroy()
 {
-	if (!GExitPurge)
+	if (!GExitPurge && GEngine)
 	{
-		if (FAudioDeviceManager* DeviceManager = FAudioDeviceManager::Get())
+		FAudioDeviceManager* AudioDeviceManager = GEngine->GetAudioDeviceManager();
+		if (AudioDeviceManager)
 		{
-			DeviceManager->RemoveSoundMix(this);
+			AudioDeviceManager->RemoveSoundMix(this);
 		}
 	}
 
diff --git a/Engine/Source/Runtime/Engine/Private/SoundSubmix.cpp b/Engine/Source/Runtime/Engine/Private/SoundSubmix.cpp
--- a/Engine/Source/Runtime/Engine/Private/SoundSubmix.cpp
+++ b/Engine/Source/Runtime/Engine/Private/SoundSubmix.cpp
@@ -214,10 +214,13 @@ void USoundSubmix::BeginDestroy()
 {
 	Super::BeginDestroy();
 
+	// Use the main/default audio device for storing and retrieving sound class properties
+	FAudioDeviceManager* AudioDeviceManager = (GEngine ? GEngine->GetAudioDeviceManager() : nullptr);
+
 	// Force the properties to be initialized for this SoundClass on all active audio devices
-	if (FAudioDeviceManager* DeviceManager = FAudioDeviceManager::Get())
+	if (AudioDeviceManager)
 	{
-		DeviceManager->UnregisterSoundSubmix(this);
+		AudioDeviceManager->UnregisterSoundSubmix(this);
 	}
 }
 
@@ -225,10 +228,13 @@ void USoundSubmix::PostLoad()
 {
 	Super::PostLoad();
 
+	// Use the main/default audio device for storing and retrieving sound class properties
+	FAudioDeviceManager* AudioDeviceManager = (GEngine ? GEngine->GetAudioDeviceManager() : nullptr);
+
 	// Force the properties to be initialized for this SoundClass on all active audio devices
-	if (FAudioDeviceManager* DeviceManager = FAudioDeviceManager::Get())
+	if (AudioDeviceManager)
 	{
-		DeviceManager->RegisterSoundSubmix(this);
+		AudioDeviceManager->RegisterSoundSubmix(this);
 	}
 }
 
@@ -319,9 +325,10 @@ void USoundSubmix::PostEditChangeProperty(FPropertyChangedEvent& PropertyChanged
 		}
 		else if (PropertyChangedEvent.Property->GetFName() == NAME_OutputVolume)
 		{
-			if (FAudioDeviceManager* DeviceManager = FAudioDeviceManager::Get())
+			FAudioDeviceManager* AudioDeviceManager = (GEngine ? GEngine->GetAudioDeviceManager() : nullptr);
+			if (AudioDeviceManager)
 			{
-				DeviceManager->UpdateSubmix(this);
+				AudioDeviceManager->UpdateSubmix(this);
 			}
 			bReinitSubmix = false;
 		}
@@ -330,9 +337,9 @@ void USoundSubmix::PostEditChangeProperty(FPropertyChangedEvent& PropertyChanged
 	if (GEngine && bReinitSubmix)
 	{
 		// Force the properties to be initialized for this SoundSubmix on all active audio devices
-		if (FAudioDeviceManager* DeviceManager = FAudioDeviceManager::Get())
+		if (FAudioDeviceManager* AudioDeviceManager = GEngine->GetAudioDeviceManager())
 		{
-			DeviceManager->RegisterSoundSubmix(this);
+			AudioDeviceManager->RegisterSoundSubmix(this);
 		}
 	}
 
diff --git a/Engine/Source/Runtime/Engine/Private/SoundWave.cpp b/Engine/Source/Runtime/Engine/Private/SoundWave.cpp
--- a/Engine/Source/Runtime/Engine/Private/SoundWave.cpp
+++ b/Engine/Source/Runtime/Engine/Private/SoundWave.cpp
@@ -219,9 +219,9 @@ void USoundWave::GetResourceSizeEx(FResourceSizeEx& CumulativeResourceSize)
 		return;
 	}
 
-	if (FAudioDevice* AudioDevice = FAudioDeviceManager::GetMainDevice())
+	if (FAudioDevice* LocalAudioDevice = GEngine->GetMainAudioDevice())
 	{
-		if (AudioDevice->HasCompressedAudioInfoClass(this) && DecompressionType == DTYPE_Native)
+		if (LocalAudioDevice->HasCompressedAudioInfoClass(this) && DecompressionType == DTYPE_Native)
 		{
 			check(!RawPCMData || RawPCMDataSize);
 			CumulativeResourceSize.AddDedicatedSystemMemoryBytes(RawPCMDataSize);
@@ -235,7 +235,7 @@ void USoundWave::GetResourceSizeEx(FResourceSizeEx& CumulativeResourceSize)
 
 			if (!FPlatformProperties::SupportsAudioStreaming() || !IsStreaming())
 			{
-				CumulativeResourceSize.AddDedicatedSystemMemoryBytes(GetCompressedDataSize(AudioDevice->GetRuntimeFormat(this)));
+				CumulativeResourceSize.AddDedicatedSystemMemoryBytes(GetCompressedDataSize(LocalAudioDevice->GetRuntimeFormat(this)));
 			}
 		}
 	}
@@ -835,7 +835,8 @@ void USoundWave::PostLoad()
 	// most likely cause us to run out of memory.
 	if (!GIsEditor && !IsTemplate( RF_ClassDefaultObject ) && GEngine)
 	{
-		if (FAudioDevice* AudioDevice = FAudioDeviceManager::GetMainDevice())
+		FAudioDevice* AudioDevice = GEngine->GetMainAudioDevice();
+		if (AudioDevice)
 		{
 			// Upload the data to the hardware, but only if we've precached startup sounds already
 			AudioDevice->Precache(this);
@@ -1510,16 +1511,14 @@ void USoundWave::FreeResources(bool bStopSoundsUsingThisResource)
 
 	// GEngine is NULL during script compilation and GEngine->Client and its audio device might be
 	// destroyed first during the exit purge.
-	if(GEngine && !GExitPurge)
+	if( GEngine && !GExitPurge )
 	{
 		// Notify the audio device to free the bulk data associated with this wave.
-		if (bStopSoundsUsingThisResource)
+		FAudioDeviceManager* AudioDeviceManager = GEngine->GetAudioDeviceManager();
+		if (bStopSoundsUsingThisResource && AudioDeviceManager)
 		{
-			if (FAudioDeviceManager* AudioDeviceManager = FAudioDeviceManager::Get())
-			{
-				AudioDeviceManager->StopSoundsUsingResource(this);
-				AudioDeviceManager->FreeResource(this);
-			}
+			AudioDeviceManager->StopSoundsUsingResource(this);
+			AudioDeviceManager->FreeResource(this);
 		}
 	}
 
diff --git a/Engine/Source/Runtime/Engine/Private/SoundWaveProcedural.cpp b/Engine/Source/Runtime/Engine/Private/SoundWaveProcedural.cpp
--- a/Engine/Source/Runtime/Engine/Private/SoundWaveProcedural.cpp
+++ b/Engine/Source/Runtime/Engine/Private/SoundWaveProcedural.cpp
@@ -20,7 +20,7 @@ USoundWaveProcedural::USoundWaveProcedural(const FObjectInitializer& ObjectIniti
 	
 	if (GEngine)
 	{
-		FAudioDevice* MainAudioDevice = FAudioDeviceManager::GetMainDevice();
+		FAudioDevice* MainAudioDevice = GEngine->GetMainAudioDevice();
 		if (MainAudioDevice && !MainAudioDevice->IsAudioMixerEnabled())
 		{
 #if PLATFORM_MAC
diff --git a/Engine/Source/Runtime/Engine/Private/UnrealEngine.cpp b/Engine/Source/Runtime/Engine/Private/UnrealEngine.cpp
--- a/Engine/Source/Runtime/Engine/Private/UnrealEngine.cpp
+++ b/Engine/Source/Runtime/Engine/Private/UnrealEngine.cpp
@@ -1745,7 +1745,20 @@ void UEngine::OnExternalUIChange(bool bInIsOpening)
 
 void UEngine::ShutdownAudioDeviceManager()
 {
-	FAudioDeviceManager::Shutdown();
+	// Shutdown the main audio device in the UEEngine
+	if (AudioDeviceManager)
+	{
+		FAudioCommandFence Fence;
+		Fence.BeginFence();
+		Fence.Wait();
+
+		FAudioThread::StopAudioThread();
+
+		AudioDeviceManager->ShutdownAllAudioDevices();
+
+		delete AudioDeviceManager;
+		AudioDeviceManager = NULL;
+	}
 }
 
 void UEngine::PreExit()
@@ -2716,7 +2729,7 @@ void UEngine::FinishDestroy()
 	{
 		// shut down all subsystems.
 		GEngine = NULL;
-		FAudioDeviceManager::Shutdown();
+		ShutdownAudioDeviceManager();
 
 		FURL::StaticExit();
 	}
@@ -2732,7 +2745,7 @@ void UEngine::Serialize(FArchive& Ar)
 	if (Ar.IsCountingMemory())
 	{
 		// Only use the main audio device when counting memory
-		if (FAudioDevice* AudioDevice = FAudioDeviceManager::GetMainDevice())
+		if (FAudioDevice* AudioDevice = GetMainAudioDevice())
 		{
 			AudioDevice->CountBytes(Ar);
 		}
@@ -2744,9 +2757,9 @@ void UEngine::AddReferencedObjects(UObject* InThis, FReferenceCollector& Collect
 	UEngine* This = CastChecked<UEngine>(InThis);
 
 	// track objects in all the active audio devices
-	if (FAudioDeviceManager* AudioDeviceManager = FAudioDeviceManager::Get())
+	if (This->AudioDeviceManager)
 	{
-		AudioDeviceManager->AddReferencedObjects(Collector);
+		This->AudioDeviceManager->AddReferencedObjects(Collector);
 	}
 
 	// TODO: This is quite dangerous as FWorldContext::AddReferencedObjects could fail to be updated when something it
@@ -2834,38 +2847,53 @@ UFont* UEngine::GetAdditionalFont(int32 AdditionalFontIndex)
 	return GEngine->AdditionalFonts.IsValidIndex(AdditionalFontIndex) ? GEngine->AdditionalFonts[AdditionalFontIndex] : NULL;
 }
 
-FAudioDeviceManager* UEngine::GetAudioDeviceManager()
+class FAudioDeviceManager* UEngine::GetAudioDeviceManager()
 {
-	return FAudioDeviceManager::Get();
+	return AudioDeviceManager;
 }
 
 uint32 UEngine::GetAudioDeviceHandle() const
 {
-	FAudioDeviceManager* AudioDeviceManager = FAudioDeviceManager::Get();
-	return AudioDeviceManager ? AudioDeviceManager->GetMainAudioDeviceHandle() : INDEX_NONE;
+	return MainAudioDeviceHandle;
 }
 
 FAudioDevice* UEngine::GetMainAudioDevice()
 {
-	return FAudioDeviceManager::GetMainDevice();
+	if (AudioDeviceManager != nullptr)
+	{
+		return AudioDeviceManager->GetAudioDevice(MainAudioDeviceHandle);
+	}
+	return nullptr;
 }
 
 FAudioDevice* UEngine::GetActiveAudioDevice()
 {
-	return FAudioDeviceManager::GetActiveDevice();
+	if (AudioDeviceManager != nullptr)
+	{
+		return AudioDeviceManager->GetActiveAudioDevice();
+	}
+	return nullptr;
 }
 
 void UEngine::InitializeAudioDeviceManager()
 {
-	if (bUseSound)
+	if (AudioDeviceManager == nullptr && bUseSound)
 	{
-		FAudioDeviceManager::Initialize();
+		AudioDeviceManager = new FAudioDeviceManager();
+		if (AudioDeviceManager->Initialize())
+		{
+			MainAudioDeviceHandle = AudioDeviceManager->GetMainAudioDeviceHandle();
+		}
+		else
+		{
+			ShutdownAudioDeviceManager();
+		}
 	}
 }
 
 bool UEngine::UseSound() const
 {
-	return bUseSound && FAudioDeviceManager::Get() != nullptr;
+	return (bUseSound && AudioDeviceManager != nullptr);
 }
 /**
 * A fake stereo rendering device used to test stereo rendering without an attached device.
@@ -3805,9 +3833,15 @@ bool UEngine::Exec( UWorld* InWorld, const TCHAR* Cmd, FOutputDevice& Ar )
 		return true;
 	}
 
-	FAudioDevice* AudioDevice = InWorld
-		? InWorld->GetAudioDevice()
-		: FAudioDeviceManager::GetMainDevice();
+	FAudioDevice* AudioDevice = nullptr;
+	if (InWorld)
+	{
+		AudioDevice = InWorld->GetAudioDevice();
+	}
+	else
+	{
+		AudioDevice = GetMainAudioDevice();
+	}
 
 	if (AudioDevice && AudioDevice->Exec(InWorld, Cmd, Ar) == true)
 	{
diff --git a/Engine/Source/Runtime/Engine/Public/ActiveSound.h b/Engine/Source/Runtime/Engine/Public/ActiveSound.h
--- a/Engine/Source/Runtime/Engine/Public/ActiveSound.h
+++ b/Engine/Source/Runtime/Engine/Public/ActiveSound.h
@@ -6,8 +6,8 @@
 #include "WorldCollision.h"
 #include "Sound/SoundAttenuation.h"
 #include "HAL/ThreadSafeBool.h"
+#include "Audio.h"
 #include "Audio/AudioDebug.h"
-#include "AudioDefines.h"
 #include "AudioDynamicParameter.h"
 #include "Components/AudioComponent.h"
 #include "DSP/VolumeFader.h"
diff --git a/Engine/Source/Runtime/Engine/Public/Audio.h b/Engine/Source/Runtime/Engine/Public/Audio.h
--- a/Engine/Source/Runtime/Engine/Public/Audio.h
+++ b/Engine/Source/Runtime/Engine/Public/Audio.h
@@ -55,23 +55,16 @@ DECLARE_CYCLE_STAT_EXTERN(TEXT("Decompress Streamed"), STAT_AudioStreamedDecompr
 DECLARE_CYCLE_STAT_EXTERN(TEXT("Buffer Creation"), STAT_AudioResourceCreationTime, STATGROUP_Audio, );
 
 
-// Forward declarations.
-class USoundClass;
-class UAudioComponent;
+class FAudioDevice;
 class USoundNode;
-class USoundSourceBus;
-class USoundSubmix;
+struct FSoundModulationControls;
 class USoundWave;
-
-class FAudioDevice;
-
+class USoundClass;
+class USoundSubmix;
+class USoundSourceBus;
 struct FActiveSound;
-struct FReverbSettings;
-struct FSampleLoop;
-struct FSoundModulationControls;
-struct FSoundSourceBusSendInfo;
 struct FWaveInstance;
-
+struct FSoundSourceBusSendInfo;
 
 /**
  * Channel definitions for multistream waves
@@ -92,17 +85,20 @@ enum EAudioSpeakers
 	SPEAKER_Count
 };
 
-const FString SurroundSpeakerLocations[SPEAKER_Count] =
+// Forward declarations.
+class UAudioComponent;
+class USoundNode;
+struct FWaveInstance;
+struct FReverbSettings;
+struct FSampleLoop;
+
+namespace Audio
 {
-	TEXT("_fl"),	// SPEAKER_FrontLeft
-	TEXT("_fr"),	// SPEAKER_FrontRight
-	TEXT("_fc"),	// SPEAKER_FrontCenter
-	TEXT("_lf"),	// SPEAKER_LowFrequency
-	TEXT("_sl"),	// SPEAKER_SideLeft
-	TEXT("_sr"),	// SPEAKER_SideRight
-	TEXT("_bl"),	// SPEAKER_BackLeft
-	TEXT("_br")		// SPEAKER_BackRight
-};
+	/**
+	 * Typed identifier for Audio Device Id
+	 */
+	using FDeviceId = uint32;
+}
 
 enum ELoopingMode
 {
diff --git a/Engine/Source/Runtime/Engine/Public/Audio/AudioDebug.h b/Engine/Source/Runtime/Engine/Public/Audio/AudioDebug.h
--- a/Engine/Source/Runtime/Engine/Public/Audio/AudioDebug.h
+++ b/Engine/Source/Runtime/Engine/Public/Audio/AudioDebug.h
@@ -5,6 +5,7 @@
 #include "CoreMinimal.h"
 
 #include "AudioDefines.h"		// For ENABLE_AUDIO_DEBUG
+#include "Audio.h"
 
 #if ENABLE_AUDIO_DEBUG
 
diff --git a/Engine/Source/Runtime/Engine/Public/AudioDefines.h b/Engine/Source/Runtime/Engine/Public/AudioDefines.h
--- a/Engine/Source/Runtime/Engine/Public/AudioDefines.h
+++ b/Engine/Source/Runtime/Engine/Public/AudioDefines.h
@@ -52,12 +52,4 @@
 /**
  * Debugger is Available on non-shipping builds
  */
-#define ENABLE_AUDIO_DEBUG !UE_BUILD_SHIPPING
-
-namespace Audio
-{
-	/**
-	 * Typed identifier for Audio Device Id
-	 */
-	using FDeviceId = uint32;
-}
\ No newline at end of file
+#define ENABLE_AUDIO_DEBUG !UE_BUILD_SHIPPING
\ No newline at end of file
diff --git a/Engine/Source/Runtime/Engine/Public/AudioDevice.h b/Engine/Source/Runtime/Engine/Public/AudioDevice.h
--- a/Engine/Source/Runtime/Engine/Public/AudioDevice.h
+++ b/Engine/Source/Runtime/Engine/Public/AudioDevice.h
@@ -1,17 +1,15 @@
 // Copyright Epic Games, Inc. All Rights Reserved.
 #pragma once
 
-#include "AudioDefines.h"
+#include "Audio.h"
 #include "AudioDeviceManager.h"
-#include "AudioDynamicParameter.h"
-#include "AudioMixer.h"
-#include "AudioVirtualLoop.h"
 #include "Components/AudioComponent.h"
 #include "CoreMinimal.h"
 #include "DSP/SpectrumAnalyzer.h"
 #include "Engine/Engine.h"
 #include "EngineGlobals.h"
 #include "IAudioExtensionPlugin.h"
+#include "AudioDynamicParameter.h"
 #include "Sound/AudioSettings.h"
 #include "Sound/AudioVolume.h"
 #include "Sound/SoundAttenuation.h"
@@ -20,6 +18,8 @@
 #include "Sound/SoundMix.h"
 #include "Sound/SoundSubmixSend.h"
 #include "Sound/SoundSourceBus.h"
+#include "AudioVirtualLoop.h"
+#include "AudioMixer.h"
 
 /**
  * Forward declares
@@ -382,6 +382,23 @@ public:
 	virtual void OnDefaultDeviceChanged() = 0;
 };
 
+/** Abstract interface for receiving audio data from a given submix. */
+class ENGINE_API ISubmixBufferListener
+{
+public:
+	/**
+	Called when a new buffer has been rendered for a given submix
+	@param OwningSubmix	The submix object which has rendered a new buffer
+	@param AudioData		Ptr to the audio buffer
+	@param NumSamples		The number of audio samples in the audio buffer
+	@param NumChannels		The number of channels of audio in the buffer (e.g. 2 for stereo, 6 for 5.1, etc)
+	@param SampleRate		The sample rate of the audio buffer
+	@param AudioClock		Double audio clock value, from start of audio rendering.
+	*/
+	virtual void OnNewSubmixBuffer(const USoundSubmix* OwningSubmix, float* AudioData, int32 NumSamples, int32 NumChannels, const int32 SampleRate, double AudioClock) = 0;
+};
+
+
 class ENGINE_API FAudioDevice : public FExec
 {
 public:
@@ -1114,6 +1131,28 @@ public:
 	{
 	}
 
+	/** Returns the main audio device of the engine */
+	static FAudioDevice* GetMainAudioDevice()
+	{
+		// Try to get GEngine's main audio device
+		FAudioDevice* AudioDevice = GEngine->GetMainAudioDevice();
+
+		// If we don't have a main audio device (maybe we're running in a non-standard mode like a commandlet)
+		if (!AudioDevice)
+		{
+			// We should have an active device for device manager
+			FAudioDeviceManager* DeviceManager = GEngine->GetAudioDeviceManager();
+			return DeviceManager->GetActiveAudioDevice();
+		}
+		return AudioDevice;
+	}
+
+	/** Returns the audio device manager */
+	static FAudioDeviceManager* GetAudioDeviceManager()
+	{
+		return GEngine->GetAudioDeviceManager();
+	}
+
 	/** Low pass filter OneOverQ value */
 	float GetLowPassFilterResonance() const;
 
@@ -1155,7 +1194,7 @@ public:
 
 	static bool IsOcclusionPluginLoaded()
 	{
-		if (FAudioDevice* MainAudioDevice = FAudioDeviceManager::GetMainDevice())
+		if (FAudioDevice* MainAudioDevice = GEngine->GetMainAudioDevice())
 		{
 			return MainAudioDevice->bOcclusionInterfaceEnabled;
 		}
@@ -1170,7 +1209,7 @@ public:
 
 	static bool IsReverbPluginLoaded()
 	{
-		if (FAudioDevice* MainAudioDevice = FAudioDeviceManager::GetMainDevice())
+		if (FAudioDevice* MainAudioDevice = GEngine->GetMainAudioDevice())
 		{
 			return MainAudioDevice->bReverbInterfaceEnabled;
 		}
@@ -1639,8 +1678,8 @@ public:
 
 	bool IsMainAudioDevice()
 	{
-		FAudioDevice* MainAudioDevice = FAudioDeviceManager::GetMainDevice();
-		return MainAudioDevice == nullptr || MainAudioDevice == this;
+		FAudioDevice* MainAudioDevice = GEngine->GetMainAudioDevice();
+		return (MainAudioDevice == nullptr || MainAudioDevice == this);
 	}
 
 	/** Set whether or not we force the use of attenuation for non-game worlds (as by default we only care about game worlds) */
diff --git a/Engine/Source/Runtime/Engine/Public/AudioDeviceManager.h b/Engine/Source/Runtime/Engine/Public/AudioDeviceManager.h
--- a/Engine/Source/Runtime/Engine/Public/AudioDeviceManager.h
+++ b/Engine/Source/Runtime/Engine/Public/AudioDeviceManager.h
@@ -51,17 +51,11 @@ public:
 	* Initialize the audio device manager.
 	* Return true if successfully initialized.
 	**/
-	static void Initialize();
+	bool Initialize();
 
 	/** Returns the handle to the main audio device. */
 	Audio::FDeviceId GetMainAudioDeviceHandle() const { return MainAudioDeviceHandle; }
 
-	/** Returns the active audio device. */
-	static FAudioDevice* GetActiveDevice();
-
-	/** Returns the main audio device. */
-	static FAudioDevice* GetMainDevice();
-
 	/** Returns true if we're currently using the audio mixer. */
 	bool IsUsingAudioMixer() const;
 
@@ -103,22 +97,19 @@ public:
 	*/
 	bool ShutdownAllAudioDevices();
 
-	/**
-	* Shuts down the AudioDeviceManager
-	*/
-	static void Shutdown();
-
 	/**
 	* Returns a ptr to the audio device associated with the handle. If the
 	* handle is invalid then a NULL device ptr will be returned.
 	*/
 	FAudioDevice* GetAudioDevice(Audio::FDeviceId Handle);
 
-	/** Returns a pointer to the AudioDeviceManager if it has been initialized */
 	static FAudioDeviceManager* Get();
 
-	/** Returns reference to the AudioDeviceManager.  Asserts if it has not been initialized */
-	static FAudioDeviceManager& GetChecked();
+	/**
+	* Returns a ptr to the active audio device. If there is no active
+	* device then it will return the main audio device.
+	*/
+	FAudioDevice* GetActiveAudioDevice();
 
 	/** Returns the current number of active audio devices. */
 	uint8 GetNumActiveAudioDevices() const;
@@ -262,6 +253,9 @@ private:
 	/** The audio mixer module name. This is the audio mixer module name to use. E.g. AudioMixerXAudio2 */
 	FString AudioMixerModuleName;
 
+	/** Handle to the main audio device. */
+	Audio::FDeviceId MainAudioDeviceHandle;
+
 	/** Count of the number of free slots in the audio device array. */
 	uint32 FreeIndicesSize;
 
@@ -289,9 +283,6 @@ private:
 	/** Which audio device is currently active */
 	Audio::FDeviceId ActiveAudioDeviceHandle;
 
-	/** The main audio device's handle */
-	Audio::FDeviceId MainAudioDeviceHandle;
-
 	/** Dynamic volume map */
 	TMap<TTuple<ESoundType, FName>, float> DynamicSoundVolumes;
 
@@ -307,6 +298,6 @@ private:
 	/** Whether or not we've toggled the audio mixer. */
 	bool bToggledAudioMixer;
 
-	/** Audio Fence to ensure that we don't allow the audio thread to drift endlessly behind. */
+	/** Audio Fence to ensure that we don't allow the audio thread to drift never endingly behind. */
 	FAudioCommandFence SyncFence;
 };
diff --git a/Engine/Source/Runtime/GameplayMediaEncoder/Private/GameplayMediaEncoder.cpp b/Engine/Source/Runtime/GameplayMediaEncoder/Private/GameplayMediaEncoder.cpp
--- a/Engine/Source/Runtime/GameplayMediaEncoder/Private/GameplayMediaEncoder.cpp
+++ b/Engine/Source/Runtime/GameplayMediaEncoder/Private/GameplayMediaEncoder.cpp
@@ -274,7 +274,8 @@ bool FGameplayMediaEncoder::Start()
 	// subscribe to engine delegates for audio output and back buffer
 	//
 
-	if (FAudioDevice* AudioDevice = FAudioDeviceManager::GetMainDevice())
+	FAudioDevice* AudioDevice = GEngine->GetMainAudioDevice();
+	if (AudioDevice)
 	{
 		bAudioFormatChecked = false;
 		AudioDevice->RegisterSubmixBufferListener(this);
@@ -297,7 +298,8 @@ void FGameplayMediaEncoder::Stop()
 
 	if (UGameEngine* GameEngine = Cast<UGameEngine>(GEngine))
 	{
-		if (FAudioDevice* AudioDevice = FAudioDeviceManager::GetMainDevice())
+		FAudioDevice* AudioDevice = GameEngine->GetMainAudioDevice();
+		if (AudioDevice)
 		{
 			AudioDevice->UnregisterSubmixBufferListener(this);
 		}
diff --git a/Engine/Source/Runtime/IOS/IOSAudio/Private/IOSAudioBuffer.cpp b/Engine/Source/Runtime/IOS/IOSAudio/Private/IOSAudioBuffer.cpp
--- a/Engine/Source/Runtime/IOS/IOSAudio/Private/IOSAudioBuffer.cpp
+++ b/Engine/Source/Runtime/IOS/IOSAudio/Private/IOSAudioBuffer.cpp
@@ -54,7 +54,7 @@ FIOSAudioSoundBuffer::FIOSAudioSoundBuffer(FIOSAudioDevice* InAudioDevice, USoun
     FMemory::Memzero(SampleData, BufferSize);
     
 	
-	FAudioDeviceManager* AudioDeviceManager = FAudioDeviceManager::Get();
+	FAudioDeviceManager* AudioDeviceManager = GEngine->GetAudioDeviceManager();
 	check(AudioDeviceManager != nullptr);
 
 	// There is no need to track this resource with the AudioDeviceManager since there is a one to one mapping between FIOSAudioSoundBuffer objects and FIOSAudioSoundSource objects
@@ -156,7 +156,7 @@ FIOSAudioSoundBuffer* FIOSAudioSoundBuffer::Init(FIOSAudioDevice* IOSAudioDevice
 		return NULL;
 	}
 
-	FAudioDeviceManager* AudioDeviceManager = FAudioDeviceManager::Get();
+	FAudioDeviceManager* AudioDeviceManager = GEngine->GetAudioDeviceManager();
 
 	FIOSAudioSoundBuffer *Buffer = NULL;
 	bool	bIsStreaming = false;
diff --git a/Engine/Source/Runtime/Launch/Private/Android/AndroidEventManager.cpp b/Engine/Source/Runtime/Launch/Private/Android/AndroidEventManager.cpp
--- a/Engine/Source/Runtime/Launch/Private/Android/AndroidEventManager.cpp
+++ b/Engine/Source/Runtime/Launch/Private/Android/AndroidEventManager.cpp
@@ -319,7 +319,8 @@ void FAppEventManager::PauseAudio()
 	bAudioPaused = true;
 	UE_LOG(LogTemp, Log, TEXT("Android pause audio"));
 
-	if (FAudioDevice* AudioDevice = FAudioDeviceManager::GetMainDevice())
+	FAudioDevice* AudioDevice = GEngine->GetMainAudioDevice();
+	if (AudioDevice)
 	{
 		if (AudioDevice->IsAudioMixerEnabled())
 		{
@@ -331,7 +332,7 @@ void FAppEventManager::PauseAudio()
 		}
 		else
 		{
-			AudioDevice->Suspend(false);
+			GEngine->GetMainAudioDevice()->Suspend(false);
 
 			// make sure the audio thread runs the pause request
 			FAudioCommandFence Fence;
@@ -353,7 +354,8 @@ void FAppEventManager::ResumeAudio()
 	bAudioPaused = false;
 	UE_LOG(LogTemp, Log, TEXT("Android resume audio"));
 
-	if (FAudioDevice* AudioDevice = FAudioDeviceManager::GetMainDevice())
+	FAudioDevice* AudioDevice = GEngine->GetMainAudioDevice();
+	if (AudioDevice)
 	{
 		if (AudioDevice->IsAudioMixerEnabled())
 		{
@@ -361,7 +363,7 @@ void FAppEventManager::ResumeAudio()
 		}
 		else
 		{
-			AudioDevice->Suspend(true);
+			GEngine->GetMainAudioDevice()->Suspend(true);
 		}
 	}
 }
diff --git a/Engine/Source/Runtime/Launch/Private/IOS/LaunchIOS.cpp b/Engine/Source/Runtime/Launch/Private/IOS/LaunchIOS.cpp
--- a/Engine/Source/Runtime/Launch/Private/IOS/LaunchIOS.cpp
+++ b/Engine/Source/Runtime/Launch/Private/IOS/LaunchIOS.cpp
@@ -60,9 +60,9 @@ void FAppEntry::Suspend(bool bIsInterrupt)
 	// if background audio is active, then we don't want to do suspend any audio
 	if ([[IOSAppDelegate GetDelegate] IsFeatureActive:EAudioFeature::BackgroundAudio] == false)
 	{
-		FAudioDevice* AudioDevice = FAudioDeviceManager::GetMainDevice();
-		if (AudioDevice && !IsEngineExitRequested())
+		if (GEngine && GEngine->GetMainAudioDevice() && !IsEngineExitRequested())
 		{
+			FAudioDevice* AudioDevice = GEngine->GetMainAudioDevice();
 			if (bIsInterrupt && DisableAudioSuspendOnAudioInterruptCvar)
 			{
 				if (FTaskGraphInterface::IsRunning() && !IsEngineExitRequested())
@@ -71,14 +71,14 @@ void FAppEntry::Suspend(bool bIsInterrupt)
 					{
 						FAudioThread::RunCommandOnAudioThread([]()
 						{
-							if (FAudioDevice* MainDevice = FAudioDeviceManager::GetMainDevice())
+							if (GEngine && GEngine->GetMainAudioDevice())
 							{
-								MainDevice->SetTransientMasterVolume(0.0f);
+								GEngine->GetMainAudioDevice()->SetTransientMasterVolume(0.0f);
 							}
 						}, TStatId());
 					}, TStatId(), NULL, ENamedThreads::GameThread);
 				}
-				else if (AudioDevice)
+				else
 				{
 					AudioDevice->SetTransientMasterVolume(0.0f);
 				}
@@ -103,12 +103,12 @@ void FAppEntry::Suspend(bool bIsInterrupt)
 					{
 						FAudioThread::RunCommandOnAudioThread([]()
 						{
-							if (FAudioDevice* MainDevice = FAudioDeviceManager::GetMainDevice())
+							if (GEngine && GEngine->GetMainAudioDevice())
 							{
-								MainDevice->SuspendContext();
+								GEngine->GetMainAudioDevice()->SuspendContext();
 							}
 						}, TStatId());
-
+                
 						FAudioCommandFence AudioCommandFence;
 						AudioCommandFence.BeginFence();
 						AudioCommandFence.Wait();
@@ -129,7 +129,7 @@ void FAppEntry::Suspend(bool bIsInterrupt)
 					}
 					FEmbeddedCommunication::AllowSleep(TEXT("Background"));
 				}
-				else if (AudioDevice)
+				else
 				{
 					AudioDevice->SuspendContext();
 				}
@@ -137,8 +137,8 @@ void FAppEntry::Suspend(bool bIsInterrupt)
 		}
 		else
 		{
-			// Increment
-			IncrementAudioSuspendCounters();
+            // Increment
+            IncrementAudioSuspendCounters();
 		}
 	}
 }
@@ -154,8 +154,10 @@ void FAppEntry::Resume(bool bIsInterrupt)
 	// @todo: should this check if we were suspended, in case this changes while in the background? (suspend with background off, but resume with background audio on? is that a thing?)
 	if ([[IOSAppDelegate GetDelegate] IsFeatureActive:EAudioFeature::BackgroundAudio] == false)
 	{
-		if (FAudioDevice* AudioDevice = FAudioDeviceManager::GetMainDevice())
+		if (GEngine && GEngine->GetMainAudioDevice())
 		{
+			FAudioDevice* AudioDevice = GEngine->GetMainAudioDevice();
+        
 			if (bIsInterrupt && DisableAudioSuspendOnAudioInterruptCvar)
 			{
 				if (FTaskGraphInterface::IsRunning())
@@ -164,9 +166,9 @@ void FAppEntry::Resume(bool bIsInterrupt)
 					{
 						FAudioThread::RunCommandOnAudioThread([]()
 						{
-							if (FAudioDevice* AudioDevice = FAudioDeviceManager::GetMainDevice())
+							if (GEngine && GEngine->GetMainAudioDevice())
 							{
-								AudioDevice->SetTransientMasterVolume(1.0f);
+								GEngine->GetMainAudioDevice()->SetTransientMasterVolume(1.0f);
 							}
 						}, TStatId());
 					}, TStatId(), NULL, ENamedThreads::GameThread);
@@ -192,8 +194,8 @@ void FAppEntry::Resume(bool bIsInterrupt)
 		}
 		else
 		{
-			// Decrement
-			DecrementAudioSuspendCounters();
+            // Decrement
+            DecrementAudioSuspendCounters();
 		}
 	}
 }
@@ -201,44 +203,50 @@ void FAppEntry::Resume(bool bIsInterrupt)
 
 void FAppEntry::ResumeAudioContext()
 {
-	if (FAudioDevice* AudioDevice = FAudioDeviceManager::GetMainDevice())
+	if (GEngine && GEngine->GetMainAudioDevice())
 	{
-		if (FTaskGraphInterface::IsRunning())
+		FAudioDevice* AudioDevice = GEngine->GetMainAudioDevice();
+		if (AudioDevice)
 		{
-			FFunctionGraphTask::CreateAndDispatchWhenReady([]()
+			if (FTaskGraphInterface::IsRunning())
 			{
-				FAudioThread::RunCommandOnAudioThread([]()
+				FFunctionGraphTask::CreateAndDispatchWhenReady([]()
 				{
-					if (FAudioDevice* MainDevice = FAudioDeviceManager::GetMainDevice())
+					FAudioThread::RunCommandOnAudioThread([]()
 					{
-						MainDevice->ResumeContext();
-					}
-				}, TStatId());
-			}, TStatId(), NULL, ENamedThreads::GameThread);
-		}
-		else
-		{
-			AudioDevice->ResumeContext();
+						if (GEngine && GEngine->GetMainAudioDevice())
+						{
+							GEngine->GetMainAudioDevice()->ResumeContext();
+						}
+					}, TStatId());
+				}, TStatId(), NULL, ENamedThreads::GameThread);
+			}
+			else
+			{
+				AudioDevice->ResumeContext();
+			}
 		}
 	}
 }
 
 void FAppEntry::RestartAudio()
 {
-	if (FAudioDevice* AudioDevice = FAudioDeviceManager::GetMainDevice())
+	if (GEngine && GEngine->GetMainAudioDevice())
 	{
+		FAudioDevice* AudioDevice = GEngine->GetMainAudioDevice();
+
 		if (FTaskGraphInterface::IsRunning())
 		{
-			//increment the counter, otherwise ResumeContext won't work
-			IncrementAudioSuspendCounters();
+            //increment the counter, otherwise ResumeContext won't work
+            IncrementAudioSuspendCounters();
 
 			FFunctionGraphTask::CreateAndDispatchWhenReady([]()
 			{
 				FAudioThread::RunCommandOnAudioThread([]()
 				{
-					if (FAudioDevice* MainDevice = FAudioDeviceManager::GetMainDevice())
+					if (GEngine && GEngine->GetMainAudioDevice())
 					{
-						MainDevice->ResumeContext();
+						GEngine->GetMainAudioDevice()->ResumeContext();
 					}
 				}, TStatId());
 			}, TStatId(), NULL, ENamedThreads::GameThread);
diff --git a/Engine/Source/Runtime/Mac/CoreAudio/Private/CoreAudioBuffer.cpp b/Engine/Source/Runtime/Mac/CoreAudio/Private/CoreAudioBuffer.cpp
--- a/Engine/Source/Runtime/Mac/CoreAudio/Private/CoreAudioBuffer.cpp
+++ b/Engine/Source/Runtime/Mac/CoreAudio/Private/CoreAudioBuffer.cpp
@@ -222,7 +222,7 @@ FCoreAudioSoundBuffer* FCoreAudioSoundBuffer::CreateProceduralBuffer( FCoreAudio
  */
 FCoreAudioSoundBuffer* FCoreAudioSoundBuffer::CreatePreviewBuffer( FCoreAudioDevice* CoreAudioDevice, USoundWave* Wave, FCoreAudioSoundBuffer* Buffer )
 {
-    FAudioDeviceManager* AudioDeviceManager = FAudioDeviceManager::Get();
+    FAudioDeviceManager* AudioDeviceManager = GEngine->GetAudioDeviceManager();
     check(AudioDeviceManager != nullptr);
     
 	if (Buffer)
@@ -272,7 +272,7 @@ FCoreAudioSoundBuffer* FCoreAudioSoundBuffer::CreateNativeBuffer( FCoreAudioDevi
 	// Keep track of associated resource name.
 	Buffer->InitAudioStreamBasicDescription( kAudioFormatLinearPCM, Wave, true );
 	
-    FAudioDeviceManager* AudioDeviceManager = FAudioDeviceManager::Get();
+    FAudioDeviceManager* AudioDeviceManager = GEngine->GetAudioDeviceManager();
     check(AudioDeviceManager != nullptr);
     
 	if (UseResourceTrackingCVar)
@@ -335,7 +335,7 @@ FCoreAudioSoundBuffer* FCoreAudioSoundBuffer::Init( FAudioDevice* AudioDevice, U
 		return NULL;
 	}
 
-    FAudioDeviceManager* AudioDeviceManager = FAudioDeviceManager::Get();
+    FAudioDeviceManager* AudioDeviceManager = GEngine->GetAudioDeviceManager();
     
 	FCoreAudioDevice *CoreAudioDevice = ( FCoreAudioDevice *)AudioDevice;
 	FCoreAudioSoundBuffer *Buffer = NULL;
diff --git a/Engine/Source/Runtime/MovieScene/Private/MovieSceneTimeController.cpp b/Engine/Source/Runtime/MovieScene/Private/MovieSceneTimeController.cpp
--- a/Engine/Source/Runtime/MovieScene/Private/MovieSceneTimeController.cpp
+++ b/Engine/Source/Runtime/MovieScene/Private/MovieSceneTimeController.cpp
@@ -104,7 +104,7 @@ double FMovieSceneTimeController_PlatformClock::GetCurrentTime() const
 
 double FMovieSceneTimeController_AudioClock::GetCurrentTime() const
 {
-	FAudioDevice* AudioDevice = FAudioDeviceManager::GetMainDevice();
+	FAudioDevice* AudioDevice = GEngine ? GEngine->GetMainAudioDevice() : nullptr;
 	return AudioDevice ? AudioDevice->GetAudioClock() : FPlatformTime::Seconds();
 }
 
diff --git a/Engine/Source/Runtime/MovieSceneCapture/Public/Protocols/AudioCaptureProtocol.h b/Engine/Source/Runtime/MovieSceneCapture/Public/Protocols/AudioCaptureProtocol.h
--- a/Engine/Source/Runtime/MovieSceneCapture/Public/Protocols/AudioCaptureProtocol.h
+++ b/Engine/Source/Runtime/MovieSceneCapture/Public/Protocols/AudioCaptureProtocol.h
@@ -84,7 +84,7 @@ public:
 
 		// Stop all sounds currently playing on the Audio Device. This helps kill looping or long audio clips. When the sequence evaluates again,
 		// these clips will resume play at the correct location.
-		if (FAudioDevice* AudioDevice = FAudioDeviceManager::GetActiveDevice())
+		if (FAudioDevice* AudioDevice = GEngine->GetActiveAudioDevice())
 		{
 			AudioDevice->StopAllSounds(true);
 		}
diff --git a/Engine/Source/Runtime/Windows/XAudio2/Private/XAudio2Buffer.cpp b/Engine/Source/Runtime/Windows/XAudio2/Private/XAudio2Buffer.cpp
--- a/Engine/Source/Runtime/Windows/XAudio2/Private/XAudio2Buffer.cpp
+++ b/Engine/Source/Runtime/Windows/XAudio2/Private/XAudio2Buffer.cpp
@@ -416,15 +416,18 @@ FXAudio2SoundBuffer* FXAudio2SoundBuffer::CreateProceduralBuffer( FXAudio2Device
 	return( Buffer );
 }
 
-FXAudio2SoundBuffer* FXAudio2SoundBuffer::CreatePreviewBuffer(FXAudio2Device* XAudio2Device, USoundWave* Wave, FXAudio2SoundBuffer* Buffer)
+FXAudio2SoundBuffer* FXAudio2SoundBuffer::CreatePreviewBuffer( FXAudio2Device* XAudio2Device, USoundWave* Wave, FXAudio2SoundBuffer* Buffer )
 {
+	FAudioDeviceManager* AudioDeviceManager = GEngine->GetAudioDeviceManager();
+	check(AudioDeviceManager != nullptr);
+
 	if (Buffer)
 	{
-		FAudioDeviceManager::GetChecked().FreeBufferResource(static_cast<FSoundBuffer*>(Buffer));
+		AudioDeviceManager->FreeBufferResource(Buffer);
 	}
 
 	// Create new buffer.
-	Buffer = new FXAudio2SoundBuffer(XAudio2Device, SoundFormat_PCMPreview);
+	Buffer = new FXAudio2SoundBuffer( XAudio2Device, SoundFormat_PCMPreview );
 
 	check(!Wave->RawPCMData || Wave->RawPCMDataSize);
 
@@ -437,11 +440,11 @@ FXAudio2SoundBuffer* FXAudio2SoundBuffer::CreatePreviewBuffer(FXAudio2Device* XA
 	// Copy over whether this data should be freed on delete
 	Buffer->bDynamicResource = Wave->bDynamicResource;
 
-	Buffer->InitWaveFormatEx(WAVE_FORMAT_PCM, Wave, true);
+	Buffer->InitWaveFormatEx( WAVE_FORMAT_PCM, Wave, true );
 
-	FAudioDeviceManager::GetChecked().TrackResource(Wave, static_cast<FSoundBuffer*>(Buffer));
+	AudioDeviceManager->TrackResource(Wave, Buffer);
 
-	return(Buffer);
+	return( Buffer );
 }
 
 FXAudio2SoundBuffer* FXAudio2SoundBuffer::CreateNativeBuffer( FXAudio2Device* XAudio2Device, USoundWave* Wave )
@@ -462,7 +465,7 @@ FXAudio2SoundBuffer* FXAudio2SoundBuffer::CreateNativeBuffer( FXAudio2Device* XA
 	// Keep track of associated resource name.
 	Buffer->InitWaveFormatEx( WAVE_FORMAT_PCM, Wave, true );
 
-	FAudioDeviceManager* AudioDeviceManager = FAudioDeviceManager::Get();
+	FAudioDeviceManager* AudioDeviceManager = GEngine->GetAudioDeviceManager();
 	check(AudioDeviceManager != nullptr);
 	AudioDeviceManager->TrackResource(Wave, Buffer);
 
@@ -527,7 +530,7 @@ FXAudio2SoundBuffer* FXAudio2SoundBuffer::Init( FAudioDevice* AudioDevice, USoun
 	Wave->InvalidateSoundWaveIfNeccessary();
 #endif // WITH_EDITOR
 
-	FAudioDeviceManager* AudioDeviceManager = FAudioDeviceManager::Get();
+	FAudioDeviceManager* AudioDeviceManager = GEngine->GetAudioDeviceManager();
 
 	FXAudio2Device* XAudio2Device = ( FXAudio2Device* )AudioDevice;
 	FXAudio2SoundBuffer* Buffer = NULL;
