Fix CI NonUnity build.

diff --git a/Engine/Source/Runtime/Engine/Public/Physics/PhysicsInterfaceTypes.h b/Engine/Source/Runtime/Engine/Public/Physics/PhysicsInterfaceTypes.h
--- a/Engine/Source/Runtime/Engine/Public/Physics/PhysicsInterfaceTypes.h
+++ b/Engine/Source/Runtime/Engine/Public/Physics/PhysicsInterfaceTypes.h
@@ -7,6 +7,7 @@
 #include "PhysicsInterfaceDeclares.h"
 #include "PhysicsEngine/BodySetupEnums.h"
 #include "PhysicsInterfaceTypesCore.h"
+#include "Materials/MaterialInterface.h"
 
 #if WITH_CHAOS
 #include "Chaos/Serializable.h"
@@ -22,6 +23,7 @@ namespace Chaos
 #endif
 
 class UPhysicalMaterialMask;
+class UMaterialInterface;
 
 // Defines for enabling hitch repeating (see ScopedSQHitchRepeater.h)
 #if !UE_BUILD_SHIPPING
