Fix CIS NonUnity build - take 2.

diff --git a/Engine/Source/Runtime/Engine/Public/Physics/PhysicsInterfaceTypes.h b/Engine/Source/Runtime/Engine/Public/Physics/PhysicsInterfaceTypes.h
--- a/Engine/Source/Runtime/Engine/Public/Physics/PhysicsInterfaceTypes.h
+++ b/Engine/Source/Runtime/Engine/Public/Physics/PhysicsInterfaceTypes.h
@@ -7,7 +7,6 @@
 #include "PhysicsInterfaceDeclares.h"
 #include "PhysicsEngine/BodySetupEnums.h"
 #include "PhysicsInterfaceTypesCore.h"
-#include "Materials/MaterialInterface.h"
 
 #if WITH_CHAOS
 #include "Chaos/Serializable.h"
