Fix deprecation warning
#rnx

diff --git a/Engine/Source/Runtime/Engine/Private/PhysicsEngine/PhysicalMaterialMask.cpp b/Engine/Source/Runtime/Engine/Private/PhysicsEngine/PhysicalMaterialMask.cpp
--- a/Engine/Source/Runtime/Engine/Private/PhysicsEngine/PhysicalMaterialMask.cpp
+++ b/Engine/Source/Runtime/Engine/Private/PhysicsEngine/PhysicalMaterialMask.cpp
@@ -155,11 +155,11 @@ void UPhysicalMaterialMask::SetMaskTexture(UTexture* InMaskTexture, const FStrin
 	MaskTexture = InMaskTexture;
 	AssetImportData->AddFileName(InTextureFilename, 0);
 
-	UProperty* ParamProperty = FindField<UProperty>(UPhysicalMaterialMask::StaticClass(), GET_MEMBER_NAME_STRING_CHECKED(UPhysicalMaterialMask, MaskTexture));
+	FProperty* ParamProperty = FindField<FProperty>(UPhysicalMaterialMask::StaticClass(), GET_MEMBER_NAME_STRING_CHECKED(UPhysicalMaterialMask, MaskTexture));
 	FPropertyChangedEvent TextureChangedEvent(ParamProperty);
 	PostEditChangeProperty(TextureChangedEvent);
 
-	ParamProperty = FindField<UProperty>(UPhysicalMaterialMask::StaticClass(), GET_MEMBER_NAME_STRING_CHECKED(UPhysicalMaterialMask, AssetImportData));
+	ParamProperty = FindField<FProperty>(UPhysicalMaterialMask::StaticClass(), GET_MEMBER_NAME_STRING_CHECKED(UPhysicalMaterialMask, AssetImportData));
 	FPropertyChangedEvent AssetImportDataChangedEvent(ParamProperty);
 	PostEditChangeProperty(AssetImportDataChangedEvent);
 }
