Don't start the studio analytics thread timer if the platform doesn't support threading.

diff --git a/Engine/Source/Runtime/Engine/Private/Analytics/StudioAnalytics.cpp b/Engine/Source/Runtime/Engine/Private/Analytics/StudioAnalytics.cpp
--- a/Engine/Source/Runtime/Engine/Private/Analytics/StudioAnalytics.cpp
+++ b/Engine/Source/Runtime/Engine/Private/Analytics/StudioAnalytics.cpp
@@ -34,7 +34,10 @@ void FStudioAnalytics::SetProvider(TSharedRef<IAnalyticsProvider> InAnalytics)
 
 	TimeEstimation = FPlatformTime::Seconds();
 
-	TimerThread = FThread(TEXT("Studio Analytics Timer Thread"), []() { RunTimer_Concurrent(); });
+	if (FPlatformProcess::SupportsMultithreading())
+	{
+		TimerThread = FThread(TEXT("Studio Analytics Timer Thread"), []() { RunTimer_Concurrent(); });
+	}
 }
 
 IAnalyticsProvider& FStudioAnalytics::GetProvider()
