Add new data to the NetworkProfiler to help track how much impact Push Model is having and what classes should be prioritized.

diff --git a/Engine/Source/Programs/NetworkProfiler/NetworkProfiler/MainWindow.cs b/Engine/Source/Programs/NetworkProfiler/NetworkProfiler/MainWindow.cs
--- a/Engine/Source/Programs/NetworkProfiler/NetworkProfiler/MainWindow.cs
+++ b/Engine/Source/Programs/NetworkProfiler/NetworkProfiler/MainWindow.cs
@@ -2,6 +2,7 @@
 
 using System;
 using System.Collections.Generic;
+using System.Collections.ObjectModel;
 using System.ComponentModel;
 using System.Data;
 using System.Drawing;
@@ -101,6 +102,9 @@ namespace NetworkProfiler
 			SetupColumns( ActorPerfPropsListView, new String[] { "Actor", "MS", "KB/s", "Bytes", "Count", "Update HZ", "Rep HZ", "Waste" } );
 			SetupColumns( ActorPerfPropsDetailsListView, new String[] { "Property", "Bytes", "Count" } );
 
+			SetupColumns( ObjectReplicationListView, new String[] { "Object Class", "# Comparisons", "# Replications", "Comparison Time (ms)", "Avg. Time Per Compare (ms)" } );
+			SetupColumns( ObjectPropertyReplicationListView, new String[] { "Property", "# Comparisons", "# Times Changed", "# Replications" } );
+
 			ActorPerfPropsDetailsListView.Columns[0].Width = 170;
 			ActorPerfPropsDetailsListView.Columns[1].Width = 50;
 			ActorPerfPropsDetailsListView.Columns[2].Width = 50;
@@ -339,8 +343,9 @@ namespace NetworkProfiler
 				{
 
 				}
-				catch ( System.Exception )
+				catch ( System.Exception se)
 				{
+					Console.Out.WriteLine(se.StackTrace);
 					ClearStreamAndChart();
 				}
 			}
@@ -392,6 +397,7 @@ namespace NetworkProfiler
 			StreamParser.ParseStreamIntoListView( CurrentNetworkStream, CurrentNetworkStream.ActorNameToSummary, ActorListView );
 			StreamParser.ParseStreamIntoListView( CurrentNetworkStream, CurrentNetworkStream.PropertyNameToSummary, PropertyListView );
 			StreamParser.ParseStreamIntoListView( CurrentNetworkStream, CurrentNetworkStream.RPCNameToSummary, RPCListView );
+			StreamParser.ParseStreamIntoReplicationListView( CurrentNetworkStream, CurrentNetworkStream.ObjectNameToReplicationSummary, ObjectReplicationListView );
 
 			ActorFilterBox.Items.Clear();
 			ActorFilterBox.Items.Add( "" );
@@ -815,9 +821,37 @@ namespace NetworkProfiler
 
             checkAllConnectionsCheckBox.Checked = !bAreAllUnchecked;
         }
-    }
 
-    public class FilterValues
+		private void ObjectReplicationListView_SelectedIndexChanged(object sender, EventArgs e)
+		{
+			ReadOnlyCollection<PropertyReplicationSummary> Summaries = null;
+
+			if (ObjectReplicationListView.SelectedItems.Count > 0)
+			{
+				string ObjectName = ObjectReplicationListView.SelectedItems[0].Text;
+				int ObjectNameIndex = CurrentNetworkStream.NameArray.FindIndex(MaybeName => (ObjectName == MaybeName));
+
+				if (ObjectNameIndex != -1)
+				{
+					Summaries = CurrentNetworkStream.ObjectNameToReplicationSummary[ObjectNameIndex].Properties;
+				}
+			}
+
+			StreamParser.ParseStreamIntoPropertyReplicationListView(CurrentNetworkStream, Summaries, ObjectPropertyReplicationListView);
+		}
+
+		private void ObjectReplicationListView_ColumnClick(object sender, ColumnClickEventArgs e)
+		{
+			HandleListViewSorting(e.Column, ObjectReplicationListView);
+		}
+
+		private void ObjectPropertyReplicationListView_ColumnClick(object sender, ColumnClickEventArgs e)
+		{
+			HandleListViewSorting(e.Column, ObjectPropertyReplicationListView);
+		}
+	}
+
+	public class FilterValues
 	{
 		public string ActorFilter		= "";
 		public string PropertyFilter	= "";
diff --git a/Engine/Source/Programs/NetworkProfiler/NetworkProfiler/NetworkStream.cs b/Engine/Source/Programs/NetworkProfiler/NetworkProfiler/NetworkStream.cs
--- a/Engine/Source/Programs/NetworkProfiler/NetworkProfiler/NetworkStream.cs
+++ b/Engine/Source/Programs/NetworkProfiler/NetworkProfiler/NetworkStream.cs
@@ -3,6 +3,7 @@
 
 using System;
 using System.Collections.Generic;
+using System.Collections.ObjectModel;
 using System.Linq;
 using System.Text;
 using System.Windows.Forms;
@@ -46,6 +47,8 @@ namespace NetworkProfiler
 		/** Mapping from RPC name to summary */
 		public Dictionary<int,TypeSummary> RPCNameToSummary = new Dictionary<int,TypeSummary>();
 
+		public Dictionary<int, ObjectReplicationSummary> ObjectNameToReplicationSummary = new Dictionary<int, ObjectReplicationSummary>();
+
 		public StreamHeader Header;
 
 		/**
@@ -209,4 +212,58 @@ namespace NetworkProfiler
             TimeInMS = InTimeInMS;
 		}
 	}
+
+	public class ObjectReplicationSummary
+	{
+		public ObjectReplicationSummary(int InObjectNameIndex, List<int> PropertyNameIndices)
+		{
+			ObjectNameIndex = InObjectNameIndex;
+			PropertiesPrivate = new List<PropertyReplicationSummary>();
+			for (int i = 0; i < PropertyNameIndices.Count; i++)
+			{
+				PropertiesPrivate.Add(new PropertyReplicationSummary(PropertyNameIndices[i]));
+			}
+		}
+
+		public readonly int ObjectNameIndex;
+		public int LastReplicatedFrame = -1;
+		public int LastReplicateFrameWithData = -1;
+
+		public int NumberOfComparisons = 0;
+		public int NumberOfReplications = 0;
+		public int NumberOfReplicationsWithData = 0;
+
+		public int NumberOfFramesReplicated = 0;
+		public int NumberOfFramesReplicatedWithData = 0;
+
+		public float TimeSpentComparingProperties = 0;
+
+		private List<PropertyReplicationSummary> PropertiesPrivate;
+		public ReadOnlyCollection<PropertyReplicationSummary> Properties
+		{
+			get { return PropertiesPrivate.AsReadOnly(); }
+		}
+		
+	}
+
+	public class PropertyReplicationSummary
+	{
+		public PropertyReplicationSummary(int InPropertyNameIndex)
+		{
+			PropertyNameIndex = InPropertyNameIndex;
+		}
+
+		public readonly int PropertyNameIndex;
+		public int LastReplicatedFrame = -1;
+		public int LastComparedFrame = -1;
+		public int LastChangedFrame = -1;
+
+		public int NumberOfComparisons = 0;
+		public int NumberOfChanges = 0;
+		public int NumberOfReplications = 0;
+
+		public int NumberOfFramesReplicated = 0;
+		public int NumberOfFramesCompared = 0;
+		public int NumberOfFramesChanged = 0;
+	}
 }
diff --git a/Engine/Source/Programs/NetworkProfiler/NetworkProfiler/StreamParser.cs b/Engine/Source/Programs/NetworkProfiler/NetworkProfiler/StreamParser.cs
--- a/Engine/Source/Programs/NetworkProfiler/NetworkProfiler/StreamParser.cs
+++ b/Engine/Source/Programs/NetworkProfiler/NetworkProfiler/StreamParser.cs
@@ -2,7 +2,9 @@
 
 
 using System;
+using System.Collections;
 using System.Collections.Generic;
+using System.Collections.ObjectModel;
 using System.Linq;
 using System.Text;
 using System.IO;
@@ -13,7 +15,8 @@ namespace NetworkProfiler
 {
 	class StreamParser
 	{
-        public static NetworkStream NetworkStream = new NetworkStream();
+		public static NetworkStream NetworkStream = new NetworkStream();
+
 		/**
 		 * Helper function for handling updating actor summaries as they require a bit more work.
 		 * 
@@ -77,6 +80,8 @@ namespace NetworkProfiler
             List<TokenReplicateProperty> LastProperties = new List<TokenReplicateProperty>();
 			List<TokenWritePropertyHeader> LastPropertyHeaders = new List<TokenWritePropertyHeader>();
 
+			Dictionary<int, TokenPropertyComparison> ObjectNamesToPropertyComparisons = new Dictionary<int, TokenPropertyComparison>();
+
 			TokenFrameMarker LastFrameMarker = null;
 
 			InMainWindow.ShowProgress( true );
@@ -221,6 +226,9 @@ namespace NetworkProfiler
 				if( Token.TokenType == ETokenTypes.FrameMarker )
 				{
 					LastFrameMarker = (TokenFrameMarker) Token;
+					ObjectNamesToPropertyComparisons = new Dictionary<int, TokenPropertyComparison>();
+
+					//Console.Out.WriteLine("EndOfFrame: " + NetworkStream.Frames.Count.ToString("0"));
 				}
 
 				// Bail out if we hit the end. We already flushed tokens above.
@@ -239,7 +247,7 @@ namespace NetworkProfiler
 					if( Token.TokenType == ETokenTypes.ReplicateActor )
 					{
 						// Encountered a new actor so we can finish up existing one for summary.
-						FinishActorProperties(Token as TokenReplicateActor, LastProperties, LastPropertyHeaders );
+						FinishActorProperties(Token as TokenReplicateActor, LastProperties, LastPropertyHeaders);
 						Debug.Assert(LastProperties.Count == 0);		// We shouldn't have any properties now
 						Debug.Assert(LastPropertyHeaders.Count == 0);	// We shouldn't have any property headers now either
 						HandleActorSummary(NetworkStream, LastActorToken);
@@ -265,6 +273,25 @@ namespace NetworkProfiler
 						var TokenWritePropertyHeader = Token as TokenWritePropertyHeader;
                         LastPropertyHeaders.Add(TokenWritePropertyHeader);
 					}
+					else if ( Token.TokenType == ETokenTypes.PropertyComparison )
+					{
+						var TokenPropertyComparison = Token as TokenPropertyComparison;
+						ObjectNamesToPropertyComparisons[TokenPropertyComparison.ObjectNameIndex] = TokenPropertyComparison;
+						HandleObjectComparison(NetworkStream, TokenPropertyComparison);
+					}
+					else if ( Token.TokenType == ETokenTypes.ReplicatePropertiesMetaData )
+					{
+						var TokenReplicatePropertiesMetaData = Token as TokenReplicatePropertiesMetaData;
+						TokenPropertyComparison Comparison = null;
+						if (ObjectNamesToPropertyComparisons.TryGetValue(TokenReplicatePropertiesMetaData.ObjectNameIndex, out Comparison))
+						{
+							HandleObjectReplication(NetworkStream, Comparison, TokenReplicatePropertiesMetaData);
+						}
+						else
+						{
+							Console.Out.WriteLine(NetworkStream.GetName(TokenReplicatePropertiesMetaData.ObjectNameIndex));
+						}
+					}
 					else
 					{
 						CurrentFrameTokens.Add(Token);
@@ -314,5 +341,153 @@ namespace NetworkProfiler
 			ListView.AutoResizeColumns(ColumnHeaderAutoResizeStyle.HeaderSize);
 			ListView.EndUpdate();
 		}
+
+		private static void HandleObjectComparison(NetworkStream NetworkStream, TokenPropertyComparison ObjectComparison)
+		{
+			// We're guaranteed that before any Replication takes place (at least in a single frame), a comparison will also take place.
+			// So, if this is the very first comparison for a given object class, we'll also need to set up our summary.
+
+			ObjectReplicationSummary ObjectSummary = null;
+			if (!NetworkStream.ObjectNameToReplicationSummary.TryGetValue(ObjectComparison.ObjectNameIndex, out ObjectSummary))
+			{
+				// If we don't have a comparison token the first time we see this object, something is broken in the stream.
+				Debug.Assert(ObjectComparison.ExportedPropertyNames != null);
+
+				ObjectSummary = new ObjectReplicationSummary(ObjectComparison.ObjectNameIndex, ObjectComparison.ExportedPropertyNames);
+				NetworkStream.ObjectNameToReplicationSummary.Add(ObjectSummary.ObjectNameIndex, ObjectSummary);
+			}
+
+			ObjectSummary.NumberOfComparisons++;
+			ObjectSummary.TimeSpentComparingProperties += ObjectComparison.TimeSpentComparing;
+
+			ReadOnlyCollection<PropertyReplicationSummary> ObjectProperties = ObjectSummary.Properties;
+			BitArray PropertiesCompared = ObjectComparison.ComparedProperties;
+			BitArray PropertiesChanged = ObjectComparison.ChangedProperties;
+
+			Debug.Assert(PropertiesChanged.Count == ObjectProperties.Count);
+			Debug.Assert(PropertiesCompared.Count == ObjectProperties.Count);
+
+			int FrameNum = NetworkStream.Frames.Count;
+
+			for (int i = 0; i < ObjectProperties.Count; i++)
+			{
+				if (PropertiesCompared[i])
+				{
+					PropertyReplicationSummary ObjectProperty = ObjectProperties[i];
+					ObjectProperty.NumberOfComparisons++;
+
+					if (ObjectProperty.LastComparedFrame != FrameNum)
+					{
+						ObjectProperty.LastComparedFrame++;
+						ObjectProperty.NumberOfFramesCompared++;
+					}
+
+					if (PropertiesChanged[i])
+					{
+						ObjectProperty.NumberOfChanges++;
+						if (ObjectProperty.LastChangedFrame != FrameNum)
+						{
+							ObjectProperty.LastChangedFrame = FrameNum;
+							ObjectProperty.NumberOfFramesChanged++;
+						}
+					}
+				}
+			}
+		}
+
+		private static void HandleObjectReplication(NetworkStream NetworkStream, TokenPropertyComparison ObjectComparison, TokenReplicatePropertiesMetaData ObjectReplication)
+		{
+			// TODO: We may be able to move this data into the Per Frame data so we can display it when selecting a range.
+			//			For now, we're just going to show a summary for the entire profile.
+
+			
+			// If we're replicating this object, we're guaranteed that it was compared before
+			// so this should be valid.
+			ObjectReplicationSummary ObjectSummary = NetworkStream.ObjectNameToReplicationSummary[ObjectReplication.ObjectNameIndex];
+
+			int FrameNum = NetworkStream.Frames.Count;
+
+			ReadOnlyCollection<PropertyReplicationSummary> ObjectProperties = ObjectSummary.Properties;
+			BitArray PropertiesCompared = ObjectComparison.ComparedProperties;
+			BitArray PropertiesChanged = ObjectComparison.ChangedProperties;
+			BitArray PropertiesFiltered = ObjectReplication.FilteredProperties;
+
+			// At this point, our object summary should be up to date so we can move on to property summaries.
+			// We will do that by comparing the bitfields sent 
+			Debug.Assert(PropertiesChanged.Count == ObjectProperties.Count);
+			Debug.Assert(PropertiesCompared.Count == ObjectProperties.Count);
+			Debug.Assert(PropertiesFiltered.Count == ObjectProperties.Count);
+
+			ObjectSummary.NumberOfReplications++;
+			if (ObjectSummary.LastReplicatedFrame != FrameNum)
+			{
+				// If this is the first replication for an object on a given frame, we better have compared its properties.
+				ObjectSummary.LastReplicatedFrame = FrameNum;
+				ObjectSummary.NumberOfFramesReplicated++;
+			}
+
+			for (int i = 0; i < ObjectProperties.Count; i++)
+			{
+				if (PropertiesCompared[i] && PropertiesChanged[i] && !PropertiesFiltered[i])
+				{
+					PropertyReplicationSummary ObjectProperty = ObjectProperties[i];
+
+					// The property was compared, changed, and wasn't filtered, that means it was replicated.
+					// Note, we ignore the WasNewObjectComparison here because we may legitimately replicate this
+					// property multiple times in the same frame to multiple connections and individual connections
+					// can have different filters applied to them.
+					ObjectProperty.NumberOfReplications++;
+					if (ObjectProperty.LastReplicatedFrame != FrameNum)
+					{
+						ObjectProperty.LastReplicatedFrame = FrameNum;
+						ObjectProperty.NumberOfFramesReplicated++;
+					}
+				}
+			}
+		}
+
+		public static void ParseStreamIntoReplicationListView( NetworkStream NetworkStream, Dictionary<int, ObjectReplicationSummary> Summaries, ListView ListView )
+		{
+			ListView.BeginUpdate();
+			ListView.Items.Clear();
+
+			// Columns are "Object Class", "# Comparisons", "# Replications", "Comparison Time", "Avg. Time Per Compare"
+			var Columns = new string[5];
+			foreach (var SummaryEntry in Summaries)
+			{
+				Columns[0] = NetworkStream.GetName(SummaryEntry.Key);
+				Columns[1] = (SummaryEntry.Value.NumberOfComparisons).ToString("0");
+				Columns[2] = (SummaryEntry.Value.NumberOfReplications).ToString("0");
+				Columns[3] = (SummaryEntry.Value.TimeSpentComparingProperties).ToString("0.0");
+				Columns[4] = ((float)SummaryEntry.Value.TimeSpentComparingProperties / (float)SummaryEntry.Value.NumberOfComparisons).ToString("0.000");
+				ListView.Items.Add(new ListViewItem(Columns));
+			}
+
+			ListView.AutoResizeColumns(ColumnHeaderAutoResizeStyle.HeaderSize);
+			ListView.EndUpdate();
+		}
+
+		public static void ParseStreamIntoPropertyReplicationListView(NetworkStream NetworkStream, ReadOnlyCollection<PropertyReplicationSummary> Summaries, ListView ListView )
+		{
+			ListView.BeginUpdate();
+			ListView.Items.Clear();
+
+			if (Summaries != null)
+			{
+				// Columns are "Property", "# Comparisons", "# Times Changed", "# Replications"
+				var Columns = new string[4];
+				foreach (var Summary in Summaries)
+				{
+					Columns[0] = NetworkStream.GetName(Summary.PropertyNameIndex);
+					Columns[1] = (Summary.NumberOfComparisons).ToString("0");
+					Columns[2] = (Summary.NumberOfChanges).ToString("0");
+					Columns[3] = (Summary.NumberOfReplications).ToString("0");
+					ListView.Items.Add(new ListViewItem(Columns));
+				}
+			}
+
+			ListView.AutoResizeColumns(ColumnHeaderAutoResizeStyle.HeaderSize);
+			ListView.EndUpdate();
+		}
 	}
 }
diff --git a/Engine/Source/Programs/NetworkProfiler/NetworkProfiler/Tokens.cs b/Engine/Source/Programs/NetworkProfiler/NetworkProfiler/Tokens.cs
--- a/Engine/Source/Programs/NetworkProfiler/NetworkProfiler/Tokens.cs
+++ b/Engine/Source/Programs/NetworkProfiler/NetworkProfiler/Tokens.cs
@@ -2,6 +2,7 @@
 
 
 using System;
+using System.Collections;
 using System.Collections.Generic;
 using System.Linq;
 using System.Text;
@@ -13,26 +14,28 @@ namespace NetworkProfiler
 	/** Enum values need to be in sync with UE3 */
 	public enum ETokenTypes
 	{
-		FrameMarker			= 0,	// Frame marker, signaling beginning of frame.	
-		SocketSendTo,				// FSocket::SendTo
-		SendBunch,					// UChannel::SendBunch
-		SendRPC,					// Sending RPC
-		ReplicateActor,				// Replicated object	
-		ReplicateProperty,			// Property being replicated.
-		EndOfStreamMarker,			// End of stream marker		
-		Event,						// Event
-		RawSocketData,				// Raw socket data being sent
-		SendAck,					// Ack being sent
-		WritePropertyHeader,		// Property header being written
-		ExportBunch,				// Exported GUIDs
-		MustBeMappedGuids,			// Must be mapped GUIDs
-		BeginContentBlock,			// Content block headers
-		EndContentBlock,			// Content block footers
-		WritePropertyHandle,		// Property handles
-		ConnectionChange,			// Connection changed
-		NameReference,				// Reference to name
-		ConnectionReference,		// Reference to connection
-		MaxAndInvalid,				// Invalid token, also used as the max token index
+		FrameMarker			= 0,		// Frame marker, signaling beginning of frame.	
+		SocketSendTo,					// FSocket::SendTo
+		SendBunch,						// UChannel::SendBunch
+		SendRPC,						// Sending RPC
+		ReplicateActor,					// Replicated object	
+		ReplicateProperty,				// Property being replicated.
+		EndOfStreamMarker,				// End of stream marker		
+		Event,							// Event
+		RawSocketData,					// Raw socket data being sent
+		SendAck,						// Ack being sent
+		WritePropertyHeader,			// Property header being written
+		ExportBunch,					// Exported GUIDs
+		MustBeMappedGuids,				// Must be mapped GUIDs
+		BeginContentBlock,				// Content block headers
+		EndContentBlock,				// Content block footers
+		WritePropertyHandle,			// Property handles
+		ConnectionChange,				// Connection changed
+		NameReference,					// Reference to name
+		ConnectionReference,			// Reference to connection
+		PropertyComparison,				// Data about property comparions.
+		ReplicatePropertiesMetaData,	// Data about properties that were filtered out during replication.
+		MaxAndInvalid,					// Invalid token, also used as the max token index
 	}
 
 	/** Enum values need to be in sync with UE3 */
@@ -141,6 +144,15 @@ namespace NetworkProfiler
 				case ETokenTypes.ConnectionChange:
 					SerializedToken = new TokenConnectionChanged( BinaryStream );
 					break;
+
+				case ETokenTypes.PropertyComparison:
+					SerializedToken = new TokenPropertyComparison(BinaryStream);
+					break;
+
+				case ETokenTypes.ReplicatePropertiesMetaData:
+					SerializedToken = new TokenReplicatePropertiesMetaData(BinaryStream);
+					break;
+
 				default:
 					throw new InvalidDataException();
 			}
@@ -860,8 +872,91 @@ namespace NetworkProfiler
 		}
 	}
 
+	/**
+	 * Token that tracks information about property comparisons for objects.
+	 */ 
+	class TokenPropertyComparison : TokenBase
+	{
+		/** Index to the Name of the object whose properties we were comparing. */
+		public int ObjectNameIndex;
+
+		/** The amount of time we spent comparing the properties. */
+		public float TimeSpentComparing;
+
+		/**
+		 * A BitArray describing which of the top level properties of the object were actually compared.
+		 * The number of bits will always match the number of top level properties in the class.
+		 */
+		public BitArray ComparedProperties;
+
+		/**
+		 * A BitArray describing which of the top level properties of the object were found to have changed
+		 * after we compared them.
+		 * The number of bits will always match the number of top level properties in the class.
+		 */ 
+		public BitArray ChangedProperties;
+
+		public List<int> ExportedPropertyNames;
+
+		public TokenPropertyComparison(BinaryReader BinaryStream)
+		{
+			ObjectNameIndex = TokenHelper.LoadPackedInt(BinaryStream);
+			TimeSpentComparing = BinaryStream.ReadSingle();
+			TokenHelper.ReadBitArray(BinaryStream, ref ComparedProperties);
+			TokenHelper.ReadBitArray(BinaryStream, ref ChangedProperties);
+
+			int NumExportedPropertyNames = TokenHelper.LoadPackedInt(BinaryStream);
+			if (NumExportedPropertyNames > 0)
+			{
+				ExportedPropertyNames = new List<int>(NumExportedPropertyNames);
+				for (int i = 0; i < NumExportedPropertyNames; ++i)
+				{
+					ExportedPropertyNames.Add(TokenHelper.LoadPackedInt(BinaryStream));
+				}
+			}
+		}
+	}
+
+	/**
+	 * Token that tracks basic metadata about replication for objects.
+	 */ 
+	class TokenReplicatePropertiesMetaData : TokenBase
+	{
+		/** Index to the Name of the object whose properties we were replicating. */
+		public int ObjectNameIndex;
+
+		/**
+		 * Whether or not we resent our entire history.
+		 * This is used to indicate we were resending everything for replay recording (checkpoints).
+		 * Note, properties that were filtered for the connection or that were inactive won't have
+		 * been sent, so using FilteredProperties is still required to see what was actually sent.
+		 */ 
+		public bool bSentAllChangedProperties;
+
+		/**
+		 * A BitArray describing which of the top level properties of the object were inactive (would not
+		 * be replicated) during a call to ReplicateProperties.
+		 * The number of bits will always match the number of top level properties in the class,
+		 * unless bWasAnythingSent is false (in which case it will be null).
+		 */
+		public BitArray FilteredProperties;
+
+		public TokenReplicatePropertiesMetaData(BinaryReader BinaryStream)
+		{
+			ObjectNameIndex = TokenHelper.LoadPackedInt(BinaryStream);
+
+			byte Flags = BinaryStream.ReadByte();
+			bSentAllChangedProperties = (Flags & 0x1) != 0;
+			TokenHelper.ReadBitArray(BinaryStream, ref FilteredProperties);
+		}
+	}
+
+
 	public class TokenHelper
 	{
+		public const int NumBitsPerDWord = 32;
+		public const int NumBitsPerDWordLog2 = 5;
+
 		static public TreeNode AddNode( TreeNodeCollection Tree, string Text )
 		{
 			TreeNode[] Childs = Tree.Find( Text, false );
@@ -897,5 +992,21 @@ namespace NetworkProfiler
 
 			return (int)Value;
 		}
+
+		public static void ReadBitArray(BinaryReader BinaryStream, ref BitArray OutBitArray)
+		{
+			// TODO: Verify Endianness
+			int NumBits = LoadPackedInt(BinaryStream);
+			int NumInts = ((NumBits + NumBitsPerDWord - 1) >> NumBitsPerDWordLog2);
+			int[] ReadValues = new int[NumInts];
+
+			for (Int32 Idx = 0; Idx < NumInts; Idx++)
+			{
+				ReadValues[Idx] = LoadPackedInt(BinaryStream);
+			}
+
+			OutBitArray = new BitArray(ReadValues);
+			OutBitArray.Length = NumBits;
+		}
 	}
 }
\ No newline at end of file
diff --git a/Engine/Source/Runtime/Engine/Private/DataReplication.cpp b/Engine/Source/Runtime/Engine/Private/DataReplication.cpp
--- a/Engine/Source/Runtime/Engine/Private/DataReplication.cpp
+++ b/Engine/Source/Runtime/Engine/Private/DataReplication.cpp
@@ -1541,6 +1541,13 @@ bool FObjectReplicator::ReplicateProperties( FOutBunch & Bunch, FReplicationFlag
 	// Replicate properties in the layout
 	const bool bHasRepLayout = RepLayout->ReplicateProperties(SendingRepState, ChangelistMgr->GetRepChangelistState(), (uint8*)Object, ObjectClass, OwningChannel, Writer, RepFlags);
 
+	NETWORK_PROFILER(
+		if (bHasRepLayout)
+		{
+			GNetworkProfiler.TrackReplicatePropertiesMetadata(RepLayout->GetOwner(), SendingRepState->InactiveParents, (Connection->ResendAllDataState != EResendAllDataState::None), Connection);
+		}
+	);
+
 	// Replicate all the custom delta properties (fast arrays, etc)
 	ReplicateCustomDeltaProperties(Writer, RepFlags);
 
diff --git a/Engine/Source/Runtime/Engine/Private/NetworkProfiler.cpp b/Engine/Source/Runtime/Engine/Private/NetworkProfiler.cpp
--- a/Engine/Source/Runtime/Engine/Private/NetworkProfiler.cpp
+++ b/Engine/Source/Runtime/Engine/Private/NetworkProfiler.cpp
@@ -8,10 +8,10 @@
 #include "HAL/FileManager.h"
 #include "Misc/CommandLine.h"
 #include "Misc/Paths.h"
-#include "Misc/ScopeLock.h"
 #include "Misc/App.h"
 #include "Engine/EngineBaseTypes.h"
 #include "Engine/World.h"
+#include "Serialization/MemoryWriter.h"
 
 #if USE_NETWORK_PROFILER
 
@@ -38,6 +38,7 @@ enum class ENetworkProfilerVersionHistory : uint32
 	ChannelTypesAsStrings = 11,
 	AddressesAsStrings = 12,
 	LargeRPCFix = 13,				// Changed RPC size tracking to packed 32 bits instead of plain uint16s, to handle very large RPCs.
+	PushModelTracking = 14,			// Adding some new data to help profile Push Model changes, and to help identify things that could utilize Push Model effectively.
 
 	// New history items go above here.
 
@@ -69,7 +70,9 @@ enum ENetworkProfilingPayloadType
 	NPTYPE_WritePropertyHandle,			// Property handles
 	NPTYPE_ConnectionChanged,			// Connection changed
 	NPTYPE_NameReference,				// New reference to name
-	NPTYPE_ConnectionReference			// New reference to connection
+	NPTYPE_ConnectionReference,			// New reference to connection
+	NPTYPE_PropertyComparison,			// Data about property comparions.
+	NPTYPE_ReplicatePropertiesMetadata	// Full set of top level properties for an object.
 };
 
 /*=============================================================================
@@ -496,11 +499,11 @@ void FNetworkProfiler::FlushOutgoingBunches( UNetConnection* Connection )
  */
 void FNetworkProfiler::TrackReplicateActor( const AActor* Actor, FReplicationFlags RepFlags, uint32 Cycles, UNetConnection* Connection )
 {
-	if( bIsTrackingEnabled )
+	if (bIsTrackingEnabled)
 	{
-		SCOPE_LOCK_REF( CriticalSection );
+		SCOPE_LOCK_REF(CriticalSection);
 
-		SetCurrentConnection( Connection );
+		SetCurrentConnection(Connection);
 
 		uint32 NameTableIndex = GetNameTableIndex(GetNameSafe(Actor->GetClass()));
 
@@ -508,7 +511,7 @@ void FNetworkProfiler::TrackReplicateActor( const AActor* Actor, FReplicationFla
 		(*FileWriter) << Type;
 		uint8 NetFlags = (RepFlags.bNetInitial << 1) | (RepFlags.bNetOwner << 2);
 		(*FileWriter) << NetFlags;
-		(*FileWriter).SerializeIntPacked( NameTableIndex );
+		(*FileWriter).SerializeIntPacked(NameTableIndex);
 		float TimeInMS = FPlatformTime::ToMilliseconds(Cycles);	// FIXME: We may want to just pass in cycles to profiler to we don't lose precision
 		(*FileWriter) << TimeInMS;
 		// Use actor replication as indication whether session is worth keeping or not.
@@ -516,6 +519,63 @@ void FNetworkProfiler::TrackReplicateActor( const AActor* Actor, FReplicationFla
 	}
 }
 
+template<typename Allocator>
+static void ProfilerSerializeBitArray(FArchive& Ar, TBitArray<Allocator>& ToWrite)
+{
+	uint32 NumBits = ToWrite.Num();
+	Ar.SerializeIntPacked(NumBits);
+
+	uint32* Data = ToWrite.GetData();
+	const uint32 NumDWords = ((NumBits + NumBitsPerDWORD - 1) >> NumBitsPerDWORDLogTwo);
+	for (uint32 i = 0; i < NumDWords; ++i)
+	{
+		Ar.SerializeIntPacked(Data[i]);
+	}
+}
+
+void FNetworkProfiler::TrackCompareProperties_Unsafe(const FString& ObjectName, uint32 Cycles, TBitArray<>& PropertiesCompared, TBitArray<>& PropertiesThatChanged, TArray<uint8>& PropertyNameExportData)
+{
+	uint32 ObjectNameTableIndex = GetNameTableIndex(ObjectName);
+	uint8 Type = NPTYPE_PropertyComparison;
+	float TimeInMS = FPlatformTime::ToMilliseconds(Cycles);
+
+	(*FileWriter) << Type;
+	(*FileWriter).SerializeIntPacked(ObjectNameTableIndex);
+	(*FileWriter) << TimeInMS;
+	ProfilerSerializeBitArray(*FileWriter, PropertiesCompared);
+	ProfilerSerializeBitArray(*FileWriter, PropertiesThatChanged);
+
+	if (PropertyNameExportData.Num() == 0)
+	{
+		uint32 Size = 0;
+		(*FileWriter).SerializeIntPacked(Size);
+	}
+	else
+	{
+		(*FileWriter).Serialize(PropertyNameExportData.GetData(), PropertyNameExportData.Num());
+	}
+
+}
+
+void FNetworkProfiler::TrackReplicatePropertiesMetadata(const UObject* Object, TBitArray<>& InactiveProperties, bool bSentAllProperties, UNetConnection* Connection)
+{
+	if (bIsTrackingEnabled)
+	{
+		SCOPE_LOCK_REF(CriticalSection);
+
+		SetCurrentConnection(Connection);
+
+		uint32 ObjectNameTableIndex = GetNameTableIndex(GetNameSafe(Object));
+		uint8 Type = NPTYPE_ReplicatePropertiesMetadata;
+		uint8 Flags = (bSentAllProperties ? 1 : 0);
+		
+		(*FileWriter) << Type;
+		(*FileWriter).SerializeIntPacked(ObjectNameTableIndex);
+		(*FileWriter) << Flags;
+		ProfilerSerializeBitArray(*FileWriter, InactiveProperties);
+	}
+}
+
 /**
  * Track property being replicated.
  *
diff --git a/Engine/Source/Runtime/Engine/Private/RepLayout.cpp b/Engine/Source/Runtime/Engine/Private/RepLayout.cpp
--- a/Engine/Source/Runtime/Engine/Private/RepLayout.cpp
+++ b/Engine/Source/Runtime/Engine/Private/RepLayout.cpp
@@ -138,6 +138,11 @@ FConsoleVariableSinkHandle CreateMaxArrayMemoryCVarAndRegisterSink()
 FConsoleVariableSinkHandle MaxRepArraySizeHandle = CreateMaxArraySizeCVarAndRegisterSink();
 FConsoleVariableSinkHandle MaxRepArrayMemorySink = CreateMaxArrayMemoryCVarAndRegisterSink();
 
+namespace UE4PushModelPrivate
+{
+	class FPushModelPerNetDriverState;
+}
+
 namespace UE4_RepLayout_Private
 {
 	template<typename OutputType, typename CommandType, typename BufferType>
@@ -247,12 +252,23 @@ namespace UE4_RepLayout_Private
 		}
 	}
 
-	UE4PushModelPrivate::FPushModelPerNetDriverState* GetPerNetDriverState(const UE4PushModelPrivate::FPushModelPerNetDriverHandle Handle)
+	UE4PushModelPrivate::FPushModelPerNetDriverState* GetPerNetDriverState(const FRepChangelistState* ChangelistState)
 	{
+		const UE4PushModelPrivate::FPushModelPerNetDriverHandle Handle = ChangelistState->GetPushModelObjectHandle();
 		return Handle.IsValid() ? UE4PushModelPrivate::GetPerNetDriverState(Handle) : nullptr;
 	}
-
+#else
+	UE4PushModelPrivate::FPushModelPerNetDriverState* GetPerNetDriverState(const FRepChangelistState* ChangelistState)
+	{
+		return nullptr;
+	}
 #endif
+
+	static bool IsNetworkProfilerEnabled()
+	{
+		NETWORK_PROFILER(return true;);
+		return false;
+	}
 }
 
 //~ TODO: Consider moving the FastArray members into their own sub-struct to save memory for non fast array
@@ -1110,6 +1126,13 @@ struct FComparePropertiesSharedParams
 	FSendingRepState* const RepState;
 	FRepChangelistState* const RepChangelistState;
 	FRepChangedPropertyTracker* const RepChangedPropertyTracker;
+	UE4PushModelPrivate::FPushModelPerNetDriverState* const PushModelState = nullptr;
+	const bool bValidateProperties = false;
+	const bool bIsNetworkProfilerActive = false;
+#if (WITH_PUSH_VALIDATION_SUPPORT || USE_NETWORK_PROFILER)
+	TBitArray<> PropertiesCompared;
+	TBitArray<> PropertiesChanged;
+#endif
 };
 
 struct FComparePropertiesStackParams
@@ -1132,7 +1155,7 @@ static void CompareProperties_Array_r(
 	const uint16 CmdIndex,
 	const uint16 Handle);
 
-static void CompareRoleProperty(
+static bool CompareRoleProperty(
 	const FComparePropertiesSharedParams& SharedParams,
 	FComparePropertiesStackParams& StackParams,
 	const uint16 RoleOrRemoteRoleIndex,
@@ -1146,7 +1169,10 @@ static void CompareRoleProperty(
 	{
 		SavedRoleOrRemoteRole = ActorRoleOrRemoteRole;
 		StackParams.Changed.Add(Handle);
+		return true;
 	}
+
+	return false;
 }
 
 static void CompareRoleProperties(
@@ -1160,7 +1186,9 @@ static void CompareRoleProperties(
 	}
 }
 
-static void CompareParentProperty(
+// Compare the specific FRepParentCmd.
+// Returns true if the property (or any of its nested FRepLayoutCmds) has changed.
+static bool CompareParentProperty(
 	const int32 ParentIndex,
 	const FComparePropertiesSharedParams& SharedParams,
 	FComparePropertiesStackParams& StackParams)
@@ -1173,13 +1201,19 @@ static void CompareParentProperty(
 	// Further, this will keep the last active state of the property in the shadow buffer,
 	// meaning the next time the property becomes active it will be sent to all connections.
 	const bool bIsActive = !SharedParams.RepChangedPropertyTracker || SharedParams.RepChangedPropertyTracker->Parents[ParentIndex].Active;
-
 	const bool bShouldSkip = !bIsLifetime || !bIsActive || (Parent.Condition == COND_InitialOnly && !SharedParams.bIsInitial);
 
 	if (bShouldSkip)
 	{
-		return;
+		return false;
+	}
+
+#if USE_NETWORK_PROFILER
+	if (SharedParams.bIsNetworkProfilerActive)
+	{
+		const_cast<TBitArray<>&>(SharedParams.PropertiesCompared)[ParentIndex] = true;
 	}
+#endif
 
 	const FRepLayoutCmd& Cmd = SharedParams.Cmds[Parent.CmdStart];
 
@@ -1187,20 +1221,52 @@ static void CompareParentProperty(
 	{
 		if (UNLIKELY(ParentIndex == (int32)AActor::ENetFields_Private::Role))
 		{
-			CompareRoleProperty(SharedParams, StackParams, (int32)AActor::ENetFields_Private::Role, SharedParams.RepState->SavedRole);
-			return;
+			return CompareRoleProperty(SharedParams, StackParams, (int32)AActor::ENetFields_Private::Role, SharedParams.RepState->SavedRole);
 		}
 		if (UNLIKELY(ParentIndex == (int32)AActor::ENetFields_Private::RemoteRole))
 		{
-			CompareRoleProperty(SharedParams, StackParams, (int32)AActor::ENetFields_Private::RemoteRole, SharedParams.RepState->SavedRemoteRole);
-			return;
+			return CompareRoleProperty(SharedParams, StackParams, (int32)AActor::ENetFields_Private::RemoteRole, SharedParams.RepState->SavedRemoteRole);
 		}
 	}
+	
+	const int32 NumChanges = StackParams.Changed.Num();
 
 	// Note, Handle - 1 to account for CompareProperties_r incrementing handles.
 	CompareProperties_r(SharedParams, StackParams, Parent.CmdStart, Parent.CmdEnd, Cmd.RelativeHandle - 1);
+
+	return !!(StackParams.Changed.Num() - NumChanges);
 }
 
+namespace UE4_RepLayout_Private
+{
+	static bool CompareParentPropertyHelper(
+		const int32 ParentIndex,
+		const FComparePropertiesSharedParams& SharedParams,
+		FComparePropertiesStackParams& StackParams)
+	{
+		const bool bDidPropertyChange = CompareParentProperty(ParentIndex, SharedParams, StackParams);
+
+	#if USE_NETWORK_PROFILER
+		if (SharedParams.bIsNetworkProfilerActive)
+		{
+			const_cast<TBitArray<>&>(SharedParams.PropertiesChanged)[ParentIndex] = bDidPropertyChange;
+		}
+	#endif // USING_NETWORK_PROFILER
+
+		return bDidPropertyChange;
+	}
+
+#if WITH_PUSH_MODEL
+	static bool IsPropertyDirty(
+		const int32 ParentIndex,
+		const FComparePropertiesSharedParams& SharedParams,
+		FComparePropertiesStackParams& StackParams)
+	{
+		return !EnumHasAnyFlags(SharedParams.Parents[ParentIndex].Flags, ERepParentFlags::UsePushModel) || SharedParams.PushModelState->IsPropertyDirty(ParentIndex);
+	}
+#endif // WITH_PUSH_MODEL	
+}	
+
 static void CompareParentProperties(
 	const FComparePropertiesSharedParams& SharedParams,
 	FComparePropertiesStackParams& StackParams)
@@ -1208,8 +1274,7 @@ static void CompareParentProperties(
 	check(StackParams.ShadowData);
 
 #if WITH_PUSH_MODEL
-
-	if (UE4PushModelPrivate::FPushModelPerNetDriverState* PushModelObject = UE4_RepLayout_Private::GetPerNetDriverState(SharedParams.RepChangelistState->GetPushModelObjectHandle()))
+	if (SharedParams.PushModelState != nullptr)
 	{
 		// Typically, on an initial compare all properties will be dirty anyway.
 		// However, if an Actor is awakened from Dormancy, it's SendingRepState will have been recreated, invalidating
@@ -1219,8 +1284,8 @@ static void CompareParentProperties(
 		// To get around this, and just to be sure our code paths are consistent, we forcibly mark Role and RemoteRole dirty here.
 		if (EnumHasAnyFlags(SharedParams.Flags, ERepLayoutFlags::IsActor) && SharedParams.bIsInitial)
 		{
-			PushModelObject->MarkPropertyDirty((int32)AActor::ENetFields_Private::RemoteRole);
-			PushModelObject->MarkPropertyDirty((int32)AActor::ENetFields_Private::Role);
+			SharedParams.PushModelState->MarkPropertyDirty((int32)AActor::ENetFields_Private::RemoteRole);
+			SharedParams.PushModelState->MarkPropertyDirty((int32)AActor::ENetFields_Private::Role);
 		}
 
 		// If we're forcibly comparing all properties, then don't bother checking dirty state.
@@ -1228,50 +1293,51 @@ static void CompareParentProperties(
 		{
 			for (int32 ParentIndex = 0; ParentIndex < SharedParams.Parents.Num(); ++ParentIndex)
 			{
-				CompareParentProperty(ParentIndex, SharedParams, StackParams);
+				UE4_RepLayout_Private::CompareParentPropertyHelper(ParentIndex, SharedParams, StackParams);
 			}
 		}
 
+#if (WITH_PUSH_VALIDATION_SUPPORT)		
+		// If we're running validation, then we'll check everything regardless of push model state.
+		else if (SharedParams.bValidateProperties)
+		{
+			for (int32 ParentIndex = 0; ParentIndex < SharedParams.Parents.Num(); ++ParentIndex)
+			{
+				const bool bIsPropertyDirty = UE4_RepLayout_Private::IsPropertyDirty(ParentIndex, SharedParams, StackParams);
+				const bool bDidPropertyChange = UE4_RepLayout_Private::CompareParentPropertyHelper(ParentIndex, SharedParams, StackParams);
+
+				ensureAlwaysMsgf(bDidPropertyChange && !bIsPropertyDirty, TEXT("Push Model Property changed value, but was not marked dirty! Property=%s"), *SharedParams.Parents[ParentIndex].Property->GetPathName());
+			}	
+		}
+#endif // WITH_PUSH_VALIDATION_SUPPORT
+
 		// If we have full push model support, then we only need to check properties that are actually dirty.
-		else if (!GbPushModelValidateProperties && EnumHasAnyFlags(SharedParams.Flags, ERepLayoutFlags::FullPushSupport))
+		else if (EnumHasAnyFlags(SharedParams.Flags, ERepLayoutFlags::FullPushSupport))
 		{
-			for (TConstSetBitIterator<> It = PushModelObject->GetDirtyProperties(); It; ++It)
+			for (TConstSetBitIterator<> It = SharedParams.PushModelState->GetDirtyProperties(); It; ++It)
 			{
-				CompareParentProperty(It.GetIndex(), SharedParams, StackParams);
+				UE4_RepLayout_Private::CompareParentPropertyHelper(It.GetIndex(), SharedParams, StackParams);
 			}
 		}
 		else
 		{
-#if WITH_PUSH_VALIDATION_SUPPORT
-			int32 StartChangedNum = 0;
-#endif
-
 			for (int32 ParentIndex = 0; ParentIndex < SharedParams.Parents.Num(); ++ParentIndex)
 			{
-				const bool bIsPropertyDirty = !EnumHasAnyFlags(SharedParams.Parents[ParentIndex].Flags, ERepParentFlags::UsePushModel) || PushModelObject->IsPropertyDirty(ParentIndex);
-				if (GbPushModelValidateProperties || bIsPropertyDirty)
+				if (UE4_RepLayout_Private::IsPropertyDirty(ParentIndex, SharedParams, StackParams))
 				{
-					CompareParentProperty(ParentIndex, SharedParams, StackParams);
-
-#if WITH_PUSH_VALIDATION_SUPPORT
-					const bool bWasFoundDirtyButNotMarkedDirty = GbPushModelValidateProperties && !bIsPropertyDirty && StackParams.Changed.Num() > StartChangedNum;
-
-					ensureAlwaysMsgf(!bWasFoundDirtyButNotMarkedDirty, TEXT("Push Model Property changed value, but was not marked dirty! Property=%s"), *SharedParams.Parents[ParentIndex].Property->GetPathName());
-
-					StartChangedNum = StackParams.Changed.Num();
-#endif
+					UE4_RepLayout_Private::CompareParentPropertyHelper(ParentIndex, SharedParams, StackParams);
 				}
 			}
 		}
 
-		PushModelObject->ResetDirtyStates();
+		SharedParams.PushModelState->ResetDirtyStates();
 		return;
 	}
-#endif
+#endif // WITH_PUSH_MODEL
 
 	for (int32 ParentIndex = 0; ParentIndex < SharedParams.Parents.Num(); ++ParentIndex)
 	{
-		CompareParentProperty(ParentIndex, SharedParams, StackParams);
+		UE4_RepLayout_Private::CompareParentPropertyHelper(ParentIndex, SharedParams, StackParams);
 	}
 }
 
@@ -1412,7 +1478,10 @@ bool FRepLayout::CompareProperties(
 		Cmds,
 		RepState,
 		RepChangelistState,
-		(RepState ? RepState->RepChangedPropertyTracker.Get() : nullptr)
+		(RepState ? RepState->RepChangedPropertyTracker.Get() : nullptr),
+		/*PushModelState=*/UE4_RepLayout_Private::GetPerNetDriverState(RepChangelistState),
+		/*bValidateProperties=*/GbPushModelValidateProperties,
+		/*bIsNetworkProfilerActive=*/UE4_RepLayout_Private::IsNetworkProfilerEnabled()
 	};
 
 	FComparePropertiesStackParams StackParams{
@@ -1423,11 +1492,35 @@ bool FRepLayout::CompareProperties(
 
 	if (RepFlags.bRolesOnly)
 	{
+		// Don't track anything for the network profiler for this, since we want to use this frames *main* comparison
+		// data for tracking.
+		// This may throw off the number of total comparisons / replications for role properties, but this should
+		// happen infrequently and the overall time and bandwidth spent on these properties is likely negligible
+		// in the overall profile.
 		CompareRoleProperties(SharedParams, StackParams);
 	}
 	else
 	{
+#if USE_NETWORK_PROFILER 
+		const uint32 ReplicateParentPropertiesStartTime = SharedParams.bIsNetworkProfilerActive ? FPlatformTime::Cycles() : 0;
+		if (SharedParams.bIsNetworkProfilerActive)
+		{
+			SharedParams.PropertiesChanged.Init(false, Parents.Num());
+			SharedParams.PropertiesCompared.Init(false, Parents.Num());
+		}
+
+		struct FPropertyNameHelper
+		{
+			static FString ConvertParentCmdToPropertyName(const FRepParentCmd& Parent)
+			{
+				return Parent.CachedPropertyName.ToString();
+			}
+		};
+#endif
+	
 		CompareParentProperties(SharedParams, StackParams);
+
+		NETWORK_PROFILER(GNetworkProfiler.TrackCompareProperties(Owner, FPlatformTime::Cycles() - ReplicateParentPropertiesStartTime, SharedParams.PropertiesCompared, SharedParams.PropertiesChanged, Parents, &FPropertyNameHelper::ConvertParentCmdToPropertyName););
 	}
 
 	if (Changed.Num() == 0)
diff --git a/Engine/Source/Runtime/Engine/Public/Net/NetworkProfiler.h b/Engine/Source/Runtime/Engine/Public/Net/NetworkProfiler.h
--- a/Engine/Source/Runtime/Engine/Public/Net/NetworkProfiler.h
+++ b/Engine/Source/Runtime/Engine/Public/Net/NetworkProfiler.h
@@ -9,6 +9,9 @@
 #include "CoreMinimal.h"
 #include "Engine/EngineTypes.h"
 #include "IPAddress.h"
+#include "Containers/BitArray.h"
+#include "Serialization/MemoryWriter.h"
+#include "Misc/ScopeLock.h"
 
 class AActor;
 class FOutBunch;
@@ -167,9 +170,14 @@ private:
 	*/
 	int32 GetAddressTableIndex( const FString& Address );
 
+	void TrackCompareProperties_Unsafe(const FString& ObjectName, uint32 Cycles, TBitArray<>& PropertiesCompared, TBitArray<>& PropertiesThatChanged, TArray<uint8>& PropertyNameExportData);
+
 	// Used with FScopedIgnoreReplicateProperties.
 	uint32 IgnorePropertyCount;
 
+	// Set of names that correspond to Object's whose top level property names have been exported.
+	TSet<FString> ExportedObjects;
+
 public:
 	/**
 	 * Constructor, initializing members.
@@ -315,7 +323,56 @@ public:
 	 * @param	Actor		Actor being replicated
 	 */
 	void TrackReplicateActor( const AActor* Actor, FReplicationFlags RepFlags, uint32 Cycles, UNetConnection* Connection );
-	
+
+	/**
+	 * Track a set of metadata for a ReplicateProperties call.
+	 *
+	 * @param	Object				Object being replicated
+	 * @param	InactiveProperties	Bitfield describing the properties that are inactive.
+	 * @param	bWasAnythingSent	Whether or not any properties were actually replicated.
+	 * @param	bSentAlProperties	Whether or not we're going to try sending all the properties from the beginning of time.
+	 * @param	Connection			The connection that we're replicating properties to.
+	 */
+	void TrackReplicatePropertiesMetadata(const UObject* Object, TBitArray<>& InactiveProperties, bool bSentAllProperties, UNetConnection* Connection);
+
+	/**
+	 * Track time used to compare properties for a given object.
+	 *
+	 * @param	Object					Object being replicated
+	 * @param	Cycles					The number of CPU Cycles we spent comparing the properties for this object.
+	 * @param	PropertiesCompared		The properties that were compared (only tracks top level properties).
+	 * @param	PropertiesThatChanged	The properties that actually changed (only tracks top level properties).
+	 * @param	PropertyNameContainers	Array of items that we can convert to property names if we need to export them (should only happen the first time we see a given object).
+	 * @param	PropertyNameProjection	Project that can be used to convery a PropertyNameContainer to a usable property name.
+	 */
+	template<typename T, typename ProjectionType>
+	void TrackCompareProperties(const UObject* Object, uint32 Cycles, TBitArray<>& PropertiesCompared, TBitArray<>& PropertiesThatChanged, const TArray<T>& PropertyNameContainers, ProjectionType PropertyNameProjection)
+	{
+		if (bIsTrackingEnabled)
+		{
+            FScopeLock ScopeLock(&CriticalSection);
+
+			FString ObjectName = GetNameSafe(Object);
+			TArray<uint8> PropertyNameExportData;
+
+			if (!ExportedObjects.Contains(ObjectName))
+			{
+				uint32 NumProperties = PropertyNameContainers.Num();
+				PropertyNameExportData.Reserve(2 + (NumProperties * 2));
+				FMemoryWriter PropertyNameAr(PropertyNameExportData);
+
+				PropertyNameAr.SerializeIntPacked(NumProperties);
+				for (const T& PropertyNameContainer : PropertyNameContainers)
+				{
+					uint32 PropertyNameIndex = GetNameTableIndex(PropertyNameProjection(PropertyNameContainer));
+					PropertyNameAr.SerializeIntPacked(PropertyNameIndex);
+				}
+			}
+
+			TrackCompareProperties_Unsafe(ObjectName, Cycles, PropertiesCompared, PropertiesThatChanged, PropertyNameExportData);
+		}
+	}
+
 	/**
 	 * Track property being replicated.
 	 *
