Added 'restart handshake' diagnostics - to determine why detection of clients sending packets from multiple IP's at once, sometimes still triggers the restart handshake - despite Dual IP detection.

diff --git a/Engine/Source/Runtime/Engine/Private/PacketHandlers/StatelessConnectHandlerComponent.cpp b/Engine/Source/Runtime/Engine/Private/PacketHandlers/StatelessConnectHandlerComponent.cpp
--- a/Engine/Source/Runtime/Engine/Private/PacketHandlers/StatelessConnectHandlerComponent.cpp
+++ b/Engine/Source/Runtime/Engine/Private/PacketHandlers/StatelessConnectHandlerComponent.cpp
@@ -164,6 +164,13 @@ DEFINE_LOG_CATEGORY(LogHandshake);
 #define PACKETLOSS_TEST 0
 
 
+// Whether or not clients should send diagnostics to the server with the restart handshake, detailing why the request was accepted.
+#define RESTART_HANDSHAKE_DIAGNOSTICS 0
+
+// Disable client sending of handshake diagnostics
+#define DISABLE_SEND_HANDSHAKE_DIAGNOSTICS 1
+
+
 /**
  * Defines
  */
@@ -172,6 +179,10 @@ DEFINE_LOG_CATEGORY(LogHandshake);
 #define RESTART_HANDSHAKE_PACKET_SIZE_BITS		2
 #define RESTART_RESPONSE_SIZE_BITS				355
 
+#if RESTART_HANDSHAKE_DIAGNOSTICS
+#define RESTART_RESPONSE_DIAGNOSTICS_SIZE_BITS	483
+#endif
+
 
 // The number of seconds between secret value updates, and the random variance applied to this
 #define SECRET_UPDATE_TIME			15.f
@@ -194,6 +205,49 @@ TAutoConsoleVariable<FString> CVarNetMagicHeader(
 	TEXT("String representing binary bits which are prepended to every packet sent by the game. Max length: 32 bits."));
 
 
+#if RESTART_HANDSHAKE_DIAGNOSTICS
+TAutoConsoleVariable<int32> CVarNetRestartHandshakeDiagnostics(
+	TEXT("net.RestartHandshakeDiagnostics"),
+	0,
+	TEXT("Enables or disables restart handshake diagnostics. Serverside this controls logging. Clientside this controls sending."));
+#endif
+
+
+/**
+ * Structs/enums
+ */
+
+#if RESTART_HANDSHAKE_DIAGNOSTICS
+struct FRestartHandshakeDiagnostics
+{
+	float LastRestartPacketTimeDiff = -1.f;
+	float LastNetConnPacketTimeDiff = -1.f;
+
+	uint64 ReservedUnused = 0;
+
+
+	bool IsValid()
+	{
+		return LastRestartPacketTimeDiff != -1.f || LastNetConnPacketTimeDiff != -1.f;
+	}
+};
+
+static_assert(((RESTART_RESPONSE_DIAGNOSTICS_SIZE_BITS - RESTART_RESPONSE_SIZE_BITS) - (sizeof(FRestartHandshakeDiagnostics) * 8)) == 0,
+				"FRestartHandshakeDiagnostics must be properly factored into packet size defines.");
+
+FArchive& operator << (FArchive& Ar, FRestartHandshakeDiagnostics& D)
+{
+	Ar << D.LastRestartPacketTimeDiff;
+	Ar << D.LastNetConnPacketTimeDiff;
+	Ar << D.ReservedUnused;
+
+	return Ar;
+}
+
+static FRestartHandshakeDiagnostics HandshakeDiagnostics;
+#endif
+
+
 /**
  * StatelessConnectHandlerComponent
  */
@@ -415,7 +469,15 @@ void StatelessConnectHandlerComponent::SendChallengeResponse(uint8 InSecretId, f
 
 	if (ServerConn != nullptr)
 	{
-		const int32 BaseSize = GetAdjustedSizeBits(bRestartedHandshake ? RESTART_RESPONSE_SIZE_BITS : HANDSHAKE_PACKET_SIZE_BITS);
+		int32 RestartHandshakeResponseSize = RESTART_RESPONSE_SIZE_BITS;
+
+#if RESTART_HANDSHAKE_DIAGNOSTICS && !DISABLE_SEND_HANDSHAKE_DIAGNOSTICS
+		bool bEnableDiagnostics = bRestartedHandshake && !!CVarNetRestartHandshakeDiagnostics.GetValueOnAnyThread();
+
+		RestartHandshakeResponseSize = bEnableDiagnostics ? RESTART_RESPONSE_DIAGNOSTICS_SIZE_BITS : RestartHandshakeResponseSize;
+#endif
+
+		const int32 BaseSize = GetAdjustedSizeBits(bRestartedHandshake ? RestartHandshakeResponseSize : HANDSHAKE_PACKET_SIZE_BITS);
 		FBitWriter ResponsePacket(BaseSize + 1 /* Termination bit */);
 		uint8 bHandshakePacket = 1;
 		uint8 bRestartHandshake = (bRestartedHandshake ? 1 : 0);
@@ -435,6 +497,13 @@ void StatelessConnectHandlerComponent::SendChallengeResponse(uint8 InSecretId, f
 		if (bRestartedHandshake)
 		{
 			ResponsePacket.Serialize(AuthorisedCookie, COOKIE_BYTE_SIZE);
+
+#if RESTART_HANDSHAKE_DIAGNOSTICS && !DISABLE_SEND_HANDSHAKE_DIAGNOSTICS
+			if (bEnableDiagnostics)
+			{
+				ResponsePacket << HandshakeDiagnostics;
+			}
+#endif
 		}
 
 #if !UE_BUILD_SHIPPING
@@ -623,7 +692,12 @@ void StatelessConnectHandlerComponent::CapHandshakePacket(FBitWriter& HandshakeP
 {
 	uint32 NumBits = HandshakePacket.GetNumBits() - GetAdjustedSizeBits(0);
 
+#if RESTART_HANDSHAKE_DIAGNOSTICS
+	check(NumBits == HANDSHAKE_PACKET_SIZE_BITS || NumBits == RESTART_HANDSHAKE_PACKET_SIZE_BITS || NumBits == RESTART_RESPONSE_SIZE_BITS
+			|| NumBits == RESTART_RESPONSE_DIAGNOSTICS_SIZE_BITS);
+#else
 	check(NumBits == HANDSHAKE_PACKET_SIZE_BITS || NumBits == RESTART_HANDSHAKE_PACKET_SIZE_BITS || NumBits == RESTART_RESPONSE_SIZE_BITS);
+#endif
 
 	FPacketAudit::AddStage(TEXT("PostPacketHandler"), HandshakePacket);
 
@@ -767,9 +841,20 @@ void StatelessConnectHandlerComponent::Incoming(FBitReader& Packet)
 							// Detect this by checking if any restart handshake requests have been received in roughly the last second
 							// (Dual IP situations will make the server send them constantly) - and override the checks as a failsafe,
 							// if no NetConnection packets have been received in the last second.
+							double LastRestartPacketTimeDiff = CurrentTime - LastRestartPacketTimestamp;
+							double LastNetConnPacketTimeDiff = CurrentTime - LastNetConnPacketTime;
+
 							bPassedDualIPCheck = LastRestartPacketTimestamp == 0.0 ||
-													((CurrentTime - LastRestartPacketTimestamp) - 1.1) > 0.0 ||
-													((CurrentTime - LastNetConnPacketTime) - 1.0) > 0.0;
+													((LastRestartPacketTimeDiff) - 1.1) > 0.0 ||
+													((LastNetConnPacketTimeDiff) - 1.0) > 0.0;
+
+#if RESTART_HANDSHAKE_DIAGNOSTICS
+							if (bPassedDualIPCheck)
+							{
+								HandshakeDiagnostics.LastRestartPacketTimeDiff = LastRestartPacketTimeDiff;
+								HandshakeDiagnostics.LastNetConnPacketTimeDiff = LastNetConnPacketTimeDiff;
+							}
+#endif
 						}
 
 						LastRestartPacketTimestamp = CurrentTime;
@@ -922,6 +1007,21 @@ void StatelessConnectHandlerComponent::IncomingConnectionless(const TSharedPtr<c
 							if (bRestartHandshake)
 							{
 								FMemory::Memcpy(AuthorisedCookie, OrigCookie, UE_ARRAY_COUNT(AuthorisedCookie));
+
+#if RESTART_HANDSHAKE_DIAGNOSTICS
+								if (HandshakeDiagnostics.IsValid() && !!CVarNetRestartHandshakeDiagnostics.GetValueOnAnyThread())
+								{
+									FDDoSDetection* DDoS = Handler->GetDDoS();
+
+									UE_CLOG((DDoS == nullptr || !DDoS->CheckLogRestrictions()), LogHandshake, Log,
+											TEXT("Got restart handshake diagnostics: LastRestartPacketTimeDiff: %f, ")
+											TEXT("LastNetConnPacketTimeDiff: %f"),
+											HandshakeDiagnostics.LastRestartPacketTimeDiff,
+											HandshakeDiagnostics.LastNetConnPacketTimeDiff);
+
+									HandshakeDiagnostics = FRestartHandshakeDiagnostics();
+								}
+#endif
 							}
 							else
 							{
@@ -984,9 +1084,14 @@ bool StatelessConnectHandlerComponent::ParseHandshakePacket(FBitReader& Packet,
 	uint32 BitsLeft = Packet.GetBitsLeft();
 	bool bHandshakePacketSize = BitsLeft == (HANDSHAKE_PACKET_SIZE_BITS - 1);
 	bool bRestartResponsePacketSize = BitsLeft == (RESTART_RESPONSE_SIZE_BITS - 1);
+	bool bRestartResponseDiagnosticsPacketSize = false
+#if RESTART_HANDSHAKE_DIAGNOSTICS
+			|| BitsLeft == (RESTART_RESPONSE_DIAGNOSTICS_SIZE_BITS - 1)
+#endif
+		;
 
 	// Only accept handshake packets of precisely the right size
-	if (bHandshakePacketSize || bRestartResponsePacketSize)
+	if (bHandshakePacketSize || bRestartResponsePacketSize || bRestartResponseDiagnosticsPacketSize)
 	{
 		bOutRestartHandshake = !!Packet.ReadBit();
 		OutSecretId = Packet.ReadBit();
@@ -995,9 +1100,16 @@ bool StatelessConnectHandlerComponent::ParseHandshakePacket(FBitReader& Packet,
 
 		Packet.Serialize(OutCookie, COOKIE_BYTE_SIZE);
 
-		if (bRestartResponsePacketSize)
+		if (bRestartResponsePacketSize || bRestartResponseDiagnosticsPacketSize)
 		{
 			Packet.Serialize(OutOrigCookie, COOKIE_BYTE_SIZE);
+
+#if RESTART_HANDSHAKE_DIAGNOSTICS
+			if (bRestartResponseDiagnosticsPacketSize)
+			{
+				Packet << HandshakeDiagnostics;
+			}
+#endif
 		}
 
 		bValidPacket = !Packet.IsError();
