- Add console variable to enable / disable Mesh LODs
- Currently disabled by default until I debug an issue with some particle systems

diff --git a/Engine/Source/Runtime/Engine/Private/Particles/ParticleSystemRender.cpp b/Engine/Source/Runtime/Engine/Private/Particles/ParticleSystemRender.cpp
--- a/Engine/Source/Runtime/Engine/Private/Particles/ParticleSystemRender.cpp
+++ b/Engine/Source/Runtime/Engine/Private/Particles/ParticleSystemRender.cpp
@@ -49,6 +49,14 @@ DECLARE_CYCLE_STAT(TEXT("DynamicSpriteEmitterData GetDynamicMeshElementsEmitter
 
 #include "InGamePerformanceTracker.h"
 
+static int32 GFXAllowParticleMeshLODs = 0;
+static FAutoConsoleVariableRef CVarFXAllowParticleMeshLODs(
+	TEXT("fx.FXAllowParticleMeshLODs"),
+	GFXAllowParticleMeshLODs,
+	TEXT("If we allow particle meshes to use LODs or not"),
+	ECVF_Default
+);
+
 /** 
  * Whether to track particle rendering stats.  
  * Enable with the TRACKPARTICLERENDERINGSTATS command. 
@@ -1559,7 +1567,7 @@ uint32 FDynamicMeshEmitterData::GetMeshLODIndexFromProxy(const FParticleSystemSc
 		}
 	}
 
-	if (!InOwnerProxy)
+	if (!InOwnerProxy || !GFXAllowParticleMeshLODs)
 	{
 		return FirstAvailableLOD;
 	}
