Remove definition of same stat as physx.

diff --git a/Engine/Source/Runtime/Engine/Private/PhysicsEngine/Experimental/PhysScene_Chaos.cpp b/Engine/Source/Runtime/Engine/Private/PhysicsEngine/Experimental/PhysScene_Chaos.cpp
--- a/Engine/Source/Runtime/Engine/Private/PhysicsEngine/Experimental/PhysScene_Chaos.cpp
+++ b/Engine/Source/Runtime/Engine/Private/PhysicsEngine/Experimental/PhysScene_Chaos.cpp
@@ -57,7 +57,7 @@ TAutoConsoleVariable<int32> CVar_ChaosUpdateKinematicsOnDeferredSkelMeshes(TEXT(
 
 TAutoConsoleVariable<int32> CVar_ChaosSimulationEnable(TEXT("P.Chaos.Simulation.Enable"), 1, TEXT("Enable / disable chaos simulation. If disabled, physics will not tick."));
 
-DECLARE_CYCLE_STAT(TEXT("Update Kinematics On Deferred SkelMeshes"), STAT_UpdateKinematicsOnDeferredSkelMeshes, STATGROUP_Physics);
+DECLARE_CYCLE_STAT(TEXT("Update Kinematics On Deferred SkelMeshes"), STAT_UpdateKinematicsOnDeferredSkelMeshesChaos, STATGROUP_Physics);
 
 #if WITH_EDITOR
 #include "Editor.h"
@@ -1576,7 +1576,7 @@ void FPhysScene_ChaosInterface::ClearPreSimKinematicUpdate(USkeletalMeshComponen
 
 void FPhysScene_ChaosInterface::UpdateKinematicsOnDeferredSkelMeshes()
 {
-	SCOPE_CYCLE_COUNTER(STAT_UpdateKinematicsOnDeferredSkelMeshes);
+	SCOPE_CYCLE_COUNTER(STAT_UpdateKinematicsOnDeferredSkelMeshesChaos);
 
 	for (const TPair<USkeletalMeshComponent*, FDeferredKinematicUpdateInfo>& DeferredKinematicUpdate : DeferredKinematicUpdateSkelMeshes)
 	{
