Add the DISABLE_REPLICATED_PRIVATE_PROPERTY macro that allows to stop replicating inherited properties that are private.

diff --git a/Engine/Source/Runtime/Engine/Public/Net/UnrealNetwork.h b/Engine/Source/Runtime/Engine/Public/Net/UnrealNetwork.h
--- a/Engine/Source/Runtime/Engine/Public/Net/UnrealNetwork.h
+++ b/Engine/Source/Runtime/Engine/Public/Net/UnrealNetwork.h
@@ -282,7 +282,7 @@ ENGINE_API void RegisterReplicatedLifetimeProperty(
 /*-----------------------------------------------------------------------------
 	Disable macros.
 	Use these macros to state that properties tagged replicated 
-	are voluntarely not replicated. This silences an error about missing
+	are voluntarily not replicated. This silences an error about missing
 	registered properties when class replication is started
 -----------------------------------------------------------------------------*/
 
@@ -290,6 +290,10 @@ ENGINE_API void RegisterReplicatedLifetimeProperty(
 #define DISABLE_REPLICATED_PROPERTY(c,v) \
 DisableReplicatedLifetimeProperty(StaticClass(), c::StaticClass(), GET_MEMBER_NAME_CHECKED(c,v), OutLifetimeProps);
 
+/** Use this macro to disable inherited properties that are private. Be careful since it removes a compile-time error when the variable doesn't exist */
+#define DISABLE_REPLICATED_PRIVATE_PROPERTY(c,v) \
+DisableReplicatedLifetimeProperty(StaticClass(), c::StaticClass(), FName(TEXT(#v)), OutLifetimeProps);
+
 #define DISABLE_REPLICATED_PROPERTY_FAST(c,v) \
 { \
 	static const bool bIsValid_##c_##v = ValidateReplicatedClassInheritance(StaticClass(), c::StaticClass(), TEXT(#v)); \
