Fix CIS.

diff --git a/Engine/Source/Runtime/Engine/Private/PhysicsEngine/Experimental/PhysScene_Chaos.cpp b/Engine/Source/Runtime/Engine/Private/PhysicsEngine/Experimental/PhysScene_Chaos.cpp
--- a/Engine/Source/Runtime/Engine/Private/PhysicsEngine/Experimental/PhysScene_Chaos.cpp
+++ b/Engine/Source/Runtime/Engine/Private/PhysicsEngine/Experimental/PhysScene_Chaos.cpp
@@ -22,6 +22,7 @@
 #include "PhysicsEngine/BodySetup.h"
 #include "PhysicsEngine/PhysicsSettings.h"
 #include "Components/PrimitiveComponent.h"
+#include "Components/SkeletalMeshComponent.h"
 
 #include "PhysicsProxy/FieldSystemPhysicsProxy.h"
 #include "PhysicsProxy/GeometryCollectionPhysicsProxy.h"
