Turn live coding back on in non editor builds of the engine, was mistakenly turned off.

Ben.Marsh

diff --git a/Engine/Source/Runtime/Engine/Engine.Build.cs b/Engine/Source/Runtime/Engine/Engine.Build.cs
--- a/Engine/Source/Runtime/Engine/Engine.Build.cs
+++ b/Engine/Source/Runtime/Engine/Engine.Build.cs
@@ -319,6 +319,11 @@ public class Engine : ModuleRules
 			}
 		);
 
+		if (Target.bWithLiveCoding && Target.Type != TargetType.Editor)
+		{
+			DynamicallyLoadedModuleNames.Add("LiveCoding");
+		}
+
 		if (Target.Type != TargetType.Server)
 		{
 			DynamicallyLoadedModuleNames.AddRange(
