Don't need the dependency on livecoding any more in Engine.Build.cs

diff --git a/Engine/Source/Runtime/Engine/Engine.Build.cs b/Engine/Source/Runtime/Engine/Engine.Build.cs
--- a/Engine/Source/Runtime/Engine/Engine.Build.cs
+++ b/Engine/Source/Runtime/Engine/Engine.Build.cs
@@ -255,11 +255,6 @@ public class Engine : ModuleRules
 				}
 			);
 
-			if (Target.bWithLiveCoding)
-			{
-				PrivateDependencyModuleNames.Add("LiveCoding");
-			}
-
 			PrivateDependencyModuleNames.Add("CollisionAnalyzer");
 			CircularlyReferencedDependentModules.Add("CollisionAnalyzer");
 
@@ -319,7 +314,7 @@ public class Engine : ModuleRules
 			}
 		);
 
-		if (Target.bWithLiveCoding && Target.Type != TargetType.Editor)
+		if (Target.bWithLiveCoding)
 		{
 			DynamicallyLoadedModuleNames.Add("LiveCoding");
 		}
