Fix for assert with niagara shader compiling waiting for the FShaderCompileJob::bFinalize flag to get set.  Fallback path for XGE compiles were failing to set the variable.

diff --git a/Engine/Source/Runtime/Engine/Private/ShaderCompiler/ShaderCompiler.cpp b/Engine/Source/Runtime/Engine/Private/ShaderCompiler/ShaderCompiler.cpp
--- a/Engine/Source/Runtime/Engine/Private/ShaderCompiler/ShaderCompiler.cpp
+++ b/Engine/Source/Runtime/Engine/Private/ShaderCompiler/ShaderCompiler.cpp
@@ -1456,10 +1456,6 @@ void FShaderCompileThreadRunnable::CompileDirectlyThroughDll()
 			for (int32 JobIndex = 0; JobIndex < CurrentWorkerInfo.QueuedJobs.Num(); JobIndex++)
 			{
 				FShaderCommonCompileJob& CurrentJob = *CurrentWorkerInfo.QueuedJobs[JobIndex];
-
-				check(!CurrentJob.bFinalized);
-				CurrentJob.bFinalized = true;
-
 				FShaderCompileUtilities::ExecuteShaderCompileJob(CurrentJob);
 			}
 
@@ -1470,6 +1466,8 @@ void FShaderCompileThreadRunnable::CompileDirectlyThroughDll()
 
 void FShaderCompileUtilities::ExecuteShaderCompileJob(FShaderCommonCompileJob& Job)
 {
+	check(!Job.bFinalized);
+
 	static ITargetPlatformManagerModule& TPM = GetTargetPlatformManagerRef();
 	auto* SingleJob = Job.GetSingleShaderJob();
 	if (SingleJob)
@@ -1533,6 +1531,8 @@ void FShaderCompileUtilities::ExecuteShaderCompileJob(FShaderCommonCompileJob& J
 
 		CompileShaderPipeline(Compiler, Format, PipelineJob, FString(FPlatformProcess::ShaderDir()));
 	}
+
+	Job.bFinalized = true;
 }
 
 int32 FShaderCompileThreadRunnable::CompilingLoop()
