Remove ignoreall volumes from query.

diff --git a/Engine/Source/Runtime/Engine/Private/Components/PrimitiveComponent.cpp b/Engine/Source/Runtime/Engine/Private/Components/PrimitiveComponent.cpp
--- a/Engine/Source/Runtime/Engine/Private/Components/PrimitiveComponent.cpp
+++ b/Engine/Source/Runtime/Engine/Private/Components/PrimitiveComponent.cpp
@@ -84,6 +84,13 @@ static FAutoConsoleVariableRef CVarHitDistanceTolerance(
 	TEXT("Hits that are less than this distance are ignored."),
 	ECVF_Default);
 
+static int32 AlwaysCreatePhysicsStateConversionHackCVar = 0;
+static FAutoConsoleVariableRef CVarAlwaysCreatePhysicsStateConversionHack(
+	TEXT("p.AlwaysCreatePhysicsStateConversionHack"),
+	AlwaysCreatePhysicsStateConversionHackCVar,
+	TEXT("Hack to convert actors with query and ignore all to always create physics."),
+	ECVF_Default);
+
 #if !(UE_BUILD_SHIPPING || UE_BUILD_TEST)
 int32 CVarShowInitialOverlaps = 0;
 FAutoConsoleVariableRef CVarRefShowInitialOverlaps(
@@ -730,6 +737,16 @@ void UPrimitiveComponent::OnCreatePhysicsState()
 	{
 		//UE_LOG(LogPrimitiveComponent, Warning, TEXT("Creating Physics State (%s : %s)"), *GetNameSafe(GetOuter()),  *GetName());
 
+		if (AlwaysCreatePhysicsStateConversionHackCVar > 0)
+		{
+			static FCollisionResponseContainer IgnoreAll(ECR_Ignore);
+			if (BodyInstance.GetCollisionEnabled() == ECollisionEnabled::QueryOnly && BodyInstance.GetResponseToChannels() == IgnoreAll)
+			{
+				bAlwaysCreatePhysicsState = true;
+				BodyInstance.SetCollisionEnabled(ECollisionEnabled::NoCollision);
+			}
+		}
+
 		UBodySetup* BodySetup = GetBodySetup();
 		if(BodySetup)
 		{
