CL 1 of 3: Refactor Audio Device Manager.

diff --git a/Engine/Plugins/Media/PixelStreaming/Source/PixelStreaming/Private/AudioCapturer.cpp b/Engine/Plugins/Media/PixelStreaming/Source/PixelStreaming/Private/AudioCapturer.cpp
--- a/Engine/Plugins/Media/PixelStreaming/Source/PixelStreaming/Private/AudioCapturer.cpp
+++ b/Engine/Plugins/Media/PixelStreaming/Source/PixelStreaming/Private/AudioCapturer.cpp
@@ -114,7 +114,7 @@ int32 FAudioCapturer::Init()
 		return -1;
 	}
 
-	FAudioDevice* AudioDevice = GEngine->GetMainAudioDevice();
+	FAudioDeviceHandle AudioDevice = GEngine->GetMainAudioDevice();
 	if (!AudioDevice)
 	{
 		return -1;
@@ -139,7 +139,7 @@ int32 FAudioCapturer::Terminate()
 		return -1;
 	}
 
-	FAudioDevice* AudioDevice = GEngine->GetMainAudioDevice();
+	FAudioDeviceHandle AudioDevice = GEngine->GetMainAudioDevice();
 	if (!AudioDevice)
 	{
 		return -1;
diff --git a/Engine/Plugins/Online/OnlineSubsystemUtils/Source/OnlineSubsystemUtils/Private/OnlineSubsystemUtils.cpp b/Engine/Plugins/Online/OnlineSubsystemUtils/Source/OnlineSubsystemUtils/Private/OnlineSubsystemUtils.cpp
--- a/Engine/Plugins/Online/OnlineSubsystemUtils/Source/OnlineSubsystemUtils/Private/OnlineSubsystemUtils.cpp
+++ b/Engine/Plugins/Online/OnlineSubsystemUtils/Source/OnlineSubsystemUtils/Private/OnlineSubsystemUtils.cpp
@@ -110,7 +110,7 @@ UAudioComponent* CreateVoiceAudioComponent(uint32 SampleRate, int32 NumChannels)
 	UAudioComponent* AudioComponent = nullptr;
 	if (GEngine != nullptr)
 	{
-		if (FAudioDevice* AudioDevice = GEngine->GetMainAudioDevice())
+		if (FAudioDeviceHandle AudioDevice = GEngine->GetMainAudioDevice())
 		{
 			USoundWaveProcedural* SoundStreaming = NewObject<USoundWaveProcedural>();
 			SoundStreaming->SetSampleRate(SampleRate);
@@ -157,7 +157,7 @@ UAudioComponent* CreateVoiceAudioComponent(uint32 SampleRate, int32 NumChannels)
 UVoipListenerSynthComponent* CreateVoiceSynthComponent(uint32 SampleRate)
 {
 	UVoipListenerSynthComponent* SynthComponentPtr = nullptr;
-	if (FAudioDevice* AudioDevice = GEngine->GetMainAudioDevice())
+	if (FAudioDeviceHandle AudioDevice = GEngine->GetMainAudioDevice())
 	{
 		SynthComponentPtr = NewObject<UVoipListenerSynthComponent>();
 		if (SynthComponentPtr)
@@ -187,7 +187,7 @@ UVoipListenerSynthComponent* CreateVoiceSynthComponent(UWorld* World, uint32 Sam
 
 	if (World)
 	{
-		if (FAudioDevice* AudioDevice = World->GetAudioDevice())
+		if (FAudioDeviceHandle AudioDeviceHandle = World->GetAudioDevice())
 		{
 			SynthComponentPtr = NewObject<UVoipListenerSynthComponent>();
 			if (SynthComponentPtr)
@@ -243,7 +243,7 @@ void ApplyVoiceSettings(UVoipListenerSynthComponent* InSynthComponent, const FVo
 
 			// By ensuring that this Audio Component's device handle is INDEX_NONE, we ensure that we will revert to
 			// using the audio device associated with the World we just registered this audio component on.
-			AudioComponent->AudioDeviceHandle = INDEX_NONE;
+			AudioComponent->AudioDeviceID = INDEX_NONE;
 		}
 	}
 
diff --git a/Engine/Plugins/Runtime/AudioModulation/Source/AudioModulation/Private/AudioModulationStatics.cpp b/Engine/Plugins/Runtime/AudioModulation/Source/AudioModulation/Private/AudioModulationStatics.cpp
--- a/Engine/Plugins/Runtime/AudioModulation/Source/AudioModulation/Private/AudioModulationStatics.cpp
+++ b/Engine/Plugins/Runtime/AudioModulation/Source/AudioModulation/Private/AudioModulationStatics.cpp
@@ -104,7 +104,7 @@ UWorld* UAudioModulationStatics::GetAudioWorld(const UObject* WorldContextObject
 
 AudioModulation::FAudioModulationImpl* UAudioModulationStatics::GetModulationImpl(UWorld* World)
 {
-	FAudioDevice* AudioDevice = nullptr;
+	FAudioDeviceHandle AudioDevice;
 	if (World)
 	{
 		AudioDevice = World->GetAudioDevice();
diff --git a/Engine/Plugins/Runtime/Oculus/OculusAudio/Source/OculusAudio/Private/OculusAudioGeometryComponent.cpp b/Engine/Plugins/Runtime/Oculus/OculusAudio/Source/OculusAudio/Private/OculusAudioGeometryComponent.cpp
--- a/Engine/Plugins/Runtime/Oculus/OculusAudio/Source/OculusAudio/Private/OculusAudioGeometryComponent.cpp
+++ b/Engine/Plugins/Runtime/Oculus/OculusAudio/Source/OculusAudio/Private/OculusAudioGeometryComponent.cpp
@@ -328,7 +328,7 @@ ovrAudioContext UOculusAudioGeometryComponent::GetContext(UWorld* World)
 					// couldn't get a World
 					return nullptr;
 				}
-				FAudioDevice* AudioDevice = World->GetAudioDevice();
+				FAudioDevice* AudioDevice = World->GetAudioDeviceRaw();
 				if (AudioDevice == nullptr)
 				{
 					// This happens when cooking for native UE4 AudioMixer integration
diff --git a/Engine/Plugins/Runtime/Oculus/OculusVR/Source/OculusMR/Private/OculusMRModule.cpp b/Engine/Plugins/Runtime/Oculus/OculusVR/Source/OculusMR/Private/OculusMRModule.cpp
--- a/Engine/Plugins/Runtime/Oculus/OculusVR/Source/OculusMR/Private/OculusMRModule.cpp
+++ b/Engine/Plugins/Runtime/Oculus/OculusVR/Source/OculusMR/Private/OculusMRModule.cpp
@@ -447,7 +447,7 @@ void FOculusMRModule::OnInitialWorldCreated(UWorld* NewWorld)
 				ovrp_Media_SetMrcInputVideoBufferType(ovrpMediaInputVideoBufferType_TextureHandle);
 				ovrp_Media_SetMrcFrameImageFlipped(IsAndroidOpenGLESPlatform(GMaxRHIShaderPlatform));
 				ovrp_Media_SetMrcFrameInverseAlpha(ovrpBool_True);
-				FAudioDevice* AudioDevice = FAudioDevice::GetMainAudioDevice();
+				FAudioDeviceHandle AudioDevice = FAudioDevice::GetMainAudioDevice();
 				float SampleRate = AudioDevice->GetSampleRate();
 				ovrp_Media_SetMrcAudioSampleRate((int)SampleRate);
 				InitMixedRealityCapture();
diff --git a/Engine/Plugins/Runtime/Oculus/OculusVR/Source/OculusMR/Private/OculusMR_CastingCameraActor.cpp b/Engine/Plugins/Runtime/Oculus/OculusVR/Source/OculusMR/Private/OculusMR_CastingCameraActor.cpp
--- a/Engine/Plugins/Runtime/Oculus/OculusVR/Source/OculusMR/Private/OculusMR_CastingCameraActor.cpp
+++ b/Engine/Plugins/Runtime/Oculus/OculusVR/Source/OculusMR/Private/OculusMR_CastingCameraActor.cpp
@@ -274,7 +274,7 @@ void AOculusMR_CastingCameraActor::BeginPlay()
 	VRNotificationComponent->HMDRecenteredDelegate.Add(Delegate);
 
 #if PLATFORM_ANDROID
-	FAudioDevice* AudioDevice = FAudioDevice::GetMainAudioDevice();
+	FAudioDeviceHandle AudioDevice = FAudioDevice::GetMainAudioDevice();
 	AudioDevice->StartRecording(nullptr, 0.1);
 #endif
 }
@@ -282,7 +282,7 @@ void AOculusMR_CastingCameraActor::BeginPlay()
 void AOculusMR_CastingCameraActor::EndPlay(EEndPlayReason::Type Reason)
 {
 #if PLATFORM_ANDROID
-	FAudioDevice* AudioDevice = FAudioDevice::GetMainAudioDevice();
+	FAudioDeviceHandle AudioDevice = FAudioDevice::GetMainAudioDevice();
 	float NumChannels = 2;
 	float SampleRate = AudioDevice->GetSampleRate();
 	AudioDevice->StopRecording(nullptr, NumChannels, SampleRate);
@@ -486,7 +486,7 @@ void AOculusMR_CastingCameraActor::Tick(float DeltaTime)
 		ForegroundCaptureActor->GetCaptureComponent2D()->TextureTarget = ForegroundRenderTargets[CaptureIndex];
 		GetCaptureComponent2D()->SetVisibility(true);
 
-		FAudioDevice* AudioDevice = FAudioDevice::GetMainAudioDevice();
+		FAudioDeviceHandle AudioDevice = FAudioDevice::GetMainAudioDevice();
 		float NumChannels, SampleRate;
 		NumChannels = 2;
 		SampleRate = AudioDevice->GetSampleRate();
diff --git a/Engine/Plugins/Runtime/nDisplay/Source/DisplayCluster/Private/Game/Classes/Basics/DisplayClusterViewportClient.cpp b/Engine/Plugins/Runtime/nDisplay/Source/DisplayCluster/Private/Game/Classes/Basics/DisplayClusterViewportClient.cpp
--- a/Engine/Plugins/Runtime/nDisplay/Source/DisplayCluster/Private/Game/Classes/Basics/DisplayClusterViewportClient.cpp
+++ b/Engine/Plugins/Runtime/nDisplay/Source/DisplayCluster/Private/Game/Classes/Basics/DisplayClusterViewportClient.cpp
@@ -323,8 +323,8 @@ void UDisplayClusterViewportClient::Draw(FViewport* InViewport, FCanvas* SceneCa
 									// If there is more than one world referencing the main audio device
 									if (AudioDeviceManager->GetNumMainAudioDeviceWorlds() > 1)
 									{
-										uint32 MainAudioDeviceHandle = GEngine->GetAudioDeviceHandle();
-										if (AudioDevice->DeviceHandle == MainAudioDeviceHandle && !HasAudioFocus())
+										uint32 MainAudioDeviceID = GEngine->GetMainAudioDeviceID();
+										if (AudioDevice->DeviceID == MainAudioDeviceID && !HasAudioFocus())
 										{
 											bUpdateListenerPosition = false;
 										}
diff --git a/Engine/Source/Editor/AudioEditor/Private/Factories/SoundFactory.cpp b/Engine/Source/Editor/AudioEditor/Private/Factories/SoundFactory.cpp
--- a/Engine/Source/Editor/AudioEditor/Private/Factories/SoundFactory.cpp
+++ b/Engine/Source/Editor/AudioEditor/Private/Factories/SoundFactory.cpp
@@ -237,7 +237,7 @@ UObject* USoundFactory::CreateObject
 			// to be auditioned in the editor properly.
 			if (!ExistingSound->ResourceData)
 			{
-				if (FAudioDevice* AudioDevice = GEngine->GetMainAudioDevice())
+				if (FAudioDeviceHandle AudioDevice = GEngine->GetMainAudioDevice())
 				{
 					FName RuntimeFormat = AudioDevice->GetRuntimeFormat(ExistingSound);
 					ExistingSound->InitAudioResource(RuntimeFormat);
diff --git a/Engine/Source/Editor/DetailCustomizations/Private/LinuxTargetSettingsDetails.cpp b/Engine/Source/Editor/DetailCustomizations/Private/LinuxTargetSettingsDetails.cpp
--- a/Engine/Source/Editor/DetailCustomizations/Private/LinuxTargetSettingsDetails.cpp
+++ b/Engine/Source/Editor/DetailCustomizations/Private/LinuxTargetSettingsDetails.cpp
@@ -348,7 +348,7 @@ bool FLinuxTargetSettingsDetails::IsValidAudioDeviceName(const FString& InDevice
 	bool bIsValid = false;
 
 #if WITH_ENGINE
-	FAudioDevice* AudioDevice = GEngine->GetMainAudioDevice();
+	FAudioDeviceHandle AudioDevice = GEngine->GetMainAudioDevice();
 	if (AudioDevice)
 	{
 		TArray<FString> DeviceNames;
@@ -373,7 +373,7 @@ TSharedRef<SWidget> FLinuxTargetSettingsDetails::MakeAudioDeviceMenu(const TShar
 	FMenuBuilder MenuBuilder(true, nullptr);
 
 #if WITH_ENGINE
-	FAudioDevice* AudioDevice = GEngine->GetMainAudioDevice();
+	FAudioDeviceHandle AudioDevice = GEngine->GetMainAudioDevice();
 	if (AudioDevice)
 	{
 		TArray<FString> AudioDeviceNames;
diff --git a/Engine/Source/Editor/Kismet/Private/BlueprintEditor.cpp b/Engine/Source/Editor/Kismet/Private/BlueprintEditor.cpp
--- a/Engine/Source/Editor/Kismet/Private/BlueprintEditor.cpp
+++ b/Engine/Source/Editor/Kismet/Private/BlueprintEditor.cpp
@@ -8990,7 +8990,7 @@ void FBlueprintEditor::UpdatePreviewActor(UBlueprint* InBlueprint, bool bInForce
 			}
 
 			// Prevent any audio from playing as a result of spawning
-			if (FAudioDevice* AudioDevice = GEngine->GetMainAudioDevice())
+			if (FAudioDeviceHandle AudioDevice = GEngine->GetMainAudioDevice())
 			{
 				AudioDevice->Flush(PreviewScene.GetWorld());
 			}
diff --git a/Engine/Source/Editor/MovieSceneCaptureDialog/Private/MovieSceneCaptureDialogModule.cpp b/Engine/Source/Editor/MovieSceneCaptureDialog/Private/MovieSceneCaptureDialogModule.cpp
--- a/Engine/Source/Editor/MovieSceneCaptureDialog/Private/MovieSceneCaptureDialogModule.cpp
+++ b/Engine/Source/Editor/MovieSceneCaptureDialog/Private/MovieSceneCaptureDialogModule.cpp
@@ -434,8 +434,8 @@ void FInEditorCapture::Start()
 	UGameViewportClient::OnViewportCreated().AddRaw(this, &FInEditorCapture::OnPIEViewportStarted);
 	FEditorDelegates::EndPIE.AddRaw(this, &FInEditorCapture::OnEndPIE);
 		
-	FAudioDevice* AudioDevice = GEngine->GetMainAudioDevice();
-	if (AudioDevice != nullptr)
+	FAudioDeviceHandle AudioDevice = GEngine->GetMainAudioDevice();
+	if (AudioDevice)
 	{
 		TransientMasterVolume = AudioDevice->GetTransientMasterVolume();
 		AudioDevice->SetTransientMasterVolume(0.0f);
@@ -596,8 +596,8 @@ void FInEditorCapture::Shutdown()
 
 	FObjectReader(GetMutableDefault<ULevelEditorPlaySettings>(), BackedUpPlaySettings);
 
-	FAudioDevice* AudioDevice = GEngine->GetMainAudioDevice();
-	if (AudioDevice != nullptr)
+	FAudioDeviceHandle AudioDevice = GEngine->GetMainAudioDevice();
+	if (AudioDevice)
 	{
 		AudioDevice->SetTransientMasterVolume(TransientMasterVolume);
 	}
diff --git a/Engine/Source/Editor/MovieSceneTools/Private/TrackEditors/AudioTrackEditor.cpp b/Engine/Source/Editor/MovieSceneTools/Private/TrackEditors/AudioTrackEditor.cpp
--- a/Engine/Source/Editor/MovieSceneTools/Private/TrackEditors/AudioTrackEditor.cpp
+++ b/Engine/Source/Editor/MovieSceneTools/Private/TrackEditors/AudioTrackEditor.cpp
@@ -296,7 +296,7 @@ void FAudioThumbnail::GenerateWaveformPreview(TArray<uint8>& OutData, TRange<flo
 	{
 		// @todo Sequencer optimize - We might want to generate the data when we generate the texture
 		// and then discard the data afterwards, though that might be a perf hit traded for better memory usage
-		FAudioDevice* AudioDevice = GEngine->GetMainAudioDevice();
+		FAudioDeviceHandle AudioDevice = GEngine->GetMainAudioDevice();
 		if (AudioDevice)
 		{
 			AudioDevice->StopAllSounds(true);
diff --git a/Engine/Source/Editor/Persona/Private/AnimationEditorViewportClient.cpp b/Engine/Source/Editor/Persona/Private/AnimationEditorViewportClient.cpp
--- a/Engine/Source/Editor/Persona/Private/AnimationEditorViewportClient.cpp
+++ b/Engine/Source/Editor/Persona/Private/AnimationEditorViewportClient.cpp
@@ -138,7 +138,7 @@ FAnimationViewportClient::FAnimationViewportClient(const TSharedRef<IPersonaPrev
 	{
 		World->bAllowAudioPlayback = !ConfigOption->bMuteAudio;
 
-		if(FAudioDevice* AudioDevice = World->GetAudioDevice())
+		if(FAudioDevice* AudioDevice = World->GetAudioDeviceRaw())
 		{
 			AudioDevice->SetUseAttenuationForNonGameWorlds(ConfigOption->bUseAudioAttenuation);
 		}
@@ -222,7 +222,7 @@ void FAnimationViewportClient::OnToggleUseAudioAttenuation()
 	UWorld* World = PreviewScene->GetWorld();
 	if(World)
 	{
-		if(FAudioDevice* AudioDevice = GetWorld()->GetAudioDevice())
+		if(FAudioDevice* AudioDevice = GetWorld()->GetAudioDeviceRaw())
 		{
 			AudioDevice->SetUseAttenuationForNonGameWorlds(ConfigOption->bUseAudioAttenuation);
 		}
@@ -1885,7 +1885,7 @@ void FAnimationViewportClient::UpdateAudioListener(const FSceneView& View)
 
 	if (ViewportWorld)
 	{
-		if (FAudioDevice* AudioDevice = ViewportWorld->GetAudioDevice())
+		if (FAudioDevice* AudioDevice = ViewportWorld->GetAudioDeviceRaw())
 		{
 			const FVector& ViewLocation = GetViewLocation();
 			const FRotator& ViewRotation = GetViewRotation();
diff --git a/Engine/Source/Editor/UnrealEd/Private/Commandlets/AudioMixerCommandlet.cpp b/Engine/Source/Editor/UnrealEd/Private/Commandlets/AudioMixerCommandlet.cpp
--- a/Engine/Source/Editor/UnrealEd/Private/Commandlets/AudioMixerCommandlet.cpp
+++ b/Engine/Source/Editor/UnrealEd/Private/Commandlets/AudioMixerCommandlet.cpp
@@ -268,7 +268,7 @@ public:
 		}
 
 		// Get the audio main device
-		FAudioDevice* AudioDevice = GEngine->GetMainAudioDevice();
+		FAudioDeviceHandle AudioDevice = GEngine->GetMainAudioDevice();
 
 		bool bSuccess = false;
 
@@ -335,7 +335,7 @@ public:
 
 		LoadEditorAndEngineObjects<USoundWave>(SearchPaths, SoundWaves, &IgnoreList);
 
-		FAudioDevice* AudioDevice = GEngine->GetMainAudioDevice();
+		FAudioDeviceHandle AudioDevice = GEngine->GetMainAudioDevice();
 
 		bool bSuccess = false;
 
@@ -364,7 +364,7 @@ public:
 					{
 						int32 NumActiveSounds = AudioDevice->GetNumActiveSources();
 
-						PlayOneShotSound(AudioDevice, SoundWaves);
+						PlayOneShotSound(AudioDevice.GetAudioDevice(), SoundWaves);
 					}
 
 					CurrentTime = AudioDevice->GetAudioTime();
@@ -420,7 +420,7 @@ public:
 
 		LoadEditorAndEngineObjects<USoundWave>(SearchPaths, SoundWaves, &IgnoreList);
 
-		FAudioDevice* AudioDevice = GEngine->GetMainAudioDevice();
+		FAudioDeviceHandle AudioDevice = GEngine->GetMainAudioDevice();
 
 		bool bSuccess = false;
 		float TimeCount = 0.0f;
@@ -441,7 +441,7 @@ public:
 
 			if (bSuccess)
 			{
-				UAudioComponent* LoopingSound = SpawnLoopingSound(AudioDevice, World, SoundWaves);
+				UAudioComponent* LoopingSound = SpawnLoopingSound(AudioDevice.GetAudioDevice(), World, SoundWaves);
 
 				while (true)
 				{
@@ -501,7 +501,7 @@ public:
 			SoundWave->SoundGroup = ESoundGroup::SOUNDGROUP_Music;
 		}
 
-		FAudioDevice* AudioDevice = GEngine->GetMainAudioDevice();
+		FAudioDeviceHandle AudioDevice = GEngine->GetMainAudioDevice();
 
 		bool bSuccess = false;
 		float TimeCount = 0.0f;
@@ -522,7 +522,7 @@ public:
 
 			if (bSuccess)
 			{
-				UAudioComponent* LoopingSound = SpawnLoopingSound(AudioDevice, World, SoundWaves);
+				UAudioComponent* LoopingSound = SpawnLoopingSound(AudioDevice.GetAudioDevice(), World, SoundWaves);
 				LoopingSound->Play(0.0f);
 
 				while (true)
@@ -586,7 +586,7 @@ public:
 
 		LoadEditorAndEngineObjects<USoundWave>(SearchPaths, SoundWaves, &IgnoreList);
 
-		FAudioDevice* AudioDevice = GEngine->GetMainAudioDevice();
+		FAudioDeviceHandle AudioDevice = GEngine->GetMainAudioDevice();
 		if (!AudioDevice)
 		{
 			return false;
@@ -609,7 +609,7 @@ public:
 
 		for (int32 i = 0; i < NumLoops; ++i)
 		{
-			UAudioComponent* LoopingSound = SpawnLoopingSound(AudioDevice, World, SoundWaves);
+			UAudioComponent* LoopingSound = SpawnLoopingSound(AudioDevice.GetAudioDevice(), World, SoundWaves);
 			if (!LoopingSound)
 			{
 				return false;
@@ -711,7 +711,7 @@ public:
 
 		LoadEditorAndEngineObjects<USoundAttenuation>(SearchPaths, SoundAttenuations);
 
-		FAudioDevice* AudioDevice = GEngine->GetMainAudioDevice();
+		FAudioDeviceHandle AudioDevice = GEngine->GetMainAudioDevice();
 		if (!AudioDevice)
 		{
 			return false;
@@ -737,7 +737,7 @@ public:
 			int32 AttenuationIndex = FMath::RandRange(0, SoundAttenuations.Num() - 1);
 			USoundAttenuation* Attenuation = SoundAttenuations[AttenuationIndex];
 
-			UAudioComponent* LoopingSound = SpawnLoopingSound(AudioDevice, World, SoundWaves, true, Attenuation);
+			UAudioComponent* LoopingSound = SpawnLoopingSound(AudioDevice.GetAudioDevice(), World, SoundWaves, true, Attenuation);
 			if (!LoopingSound)
 			{
 				return false;
diff --git a/Engine/Source/Editor/UnrealEd/Private/EditorEngine.cpp b/Engine/Source/Editor/UnrealEd/Private/EditorEngine.cpp
--- a/Engine/Source/Editor/UnrealEd/Private/EditorEngine.cpp
+++ b/Engine/Source/Editor/UnrealEd/Private/EditorEngine.cpp
@@ -1901,13 +1901,10 @@ void UEditorEngine::Tick( float DeltaSeconds, bool bIdleMode )
 		if (bRequestEndPlayMapQueued)
 		{
 			// Shutdown all audio devices if we've requested end playmap now to avoid issues with GC running
-			TArray<FAudioDevice*>& AudioDevices = AudioDeviceManager->GetAudioDevices();
+			TArray<FAudioDevice*> AudioDevices = AudioDeviceManager->GetAudioDevices();
 			for (FAudioDevice* AudioDevice : AudioDevices)
 			{
-				if (AudioDevice)
-				{
-					AudioDevice->Flush(nullptr);
-				}
+				AudioDevice->Flush(nullptr);
 			}
 		}
 
@@ -2287,7 +2284,7 @@ UAudioComponent* UEditorEngine::GetPreviewAudioComponent()
 
 UAudioComponent* UEditorEngine::ResetPreviewAudioComponent( USoundBase* Sound, USoundNode* SoundNode )
 {
-	if (FAudioDevice* AudioDevice = GetMainAudioDevice())
+	if (FAudioDevice* AudioDevice = GetMainAudioDeviceRaw())
 	{
 		if (PreviewAudioComponent)
 		{
diff --git a/Engine/Source/Editor/UnrealEd/Private/EditorServer.cpp b/Engine/Source/Editor/UnrealEd/Private/EditorServer.cpp
--- a/Engine/Source/Editor/UnrealEd/Private/EditorServer.cpp
+++ b/Engine/Source/Editor/UnrealEd/Private/EditorServer.cpp
@@ -2058,7 +2058,7 @@ void UEditorEngine::EditorDestroyWorld( FWorldContext & Context, const FText& Cl
 	CloseEditedWorldAssets(ContextWorld);
 
 	// Stop all audio and remove references 
-	if (FAudioDevice* AudioDevice = ContextWorld->GetAudioDevice())
+	if (FAudioDevice* AudioDevice = ContextWorld->GetAudioDeviceRaw())
 	{
 		AudioDevice->Flush(ContextWorld);
 	}
diff --git a/Engine/Source/Editor/UnrealEd/Private/LevelEditorViewport.cpp b/Engine/Source/Editor/UnrealEd/Private/LevelEditorViewport.cpp
--- a/Engine/Source/Editor/UnrealEd/Private/LevelEditorViewport.cpp
+++ b/Engine/Source/Editor/UnrealEd/Private/LevelEditorViewport.cpp
@@ -4658,7 +4658,7 @@ void FLevelEditorViewportClient::UpdateAudioListener(const FSceneView& View)
 
 	if (ViewportWorld)
 	{
-		if (FAudioDevice* AudioDevice = ViewportWorld->GetAudioDevice())
+		if (FAudioDevice* AudioDevice = ViewportWorld->GetAudioDeviceRaw())
 		{
 			FVector ViewLocation = GetViewLocation();
 
diff --git a/Engine/Source/Editor/UnrealEd/Private/ObjectTools.cpp b/Engine/Source/Editor/UnrealEd/Private/ObjectTools.cpp
--- a/Engine/Source/Editor/UnrealEd/Private/ObjectTools.cpp
+++ b/Engine/Source/Editor/UnrealEd/Private/ObjectTools.cpp
@@ -2116,7 +2116,7 @@ namespace ObjectTools
 			const int32 NumAudioDevices = AudioDeviceManager->GetNumActiveAudioDevices();
 			for (int32 DeviceIndex = 0; DeviceIndex < NumAudioDevices; DeviceIndex++)
 			{
-				FAudioDevice* AudioDevice = AudioDeviceManager->GetAudioDevice(DeviceIndex);
+				FAudioDevice* AudioDevice = AudioDeviceManager->GetAudioDeviceRaw(DeviceIndex);
 				if (AudioDevice != nullptr)
 				{
 					AudioDevice->StopAllSounds();
diff --git a/Engine/Source/Editor/UnrealEd/Private/PlayLevel.cpp b/Engine/Source/Editor/UnrealEd/Private/PlayLevel.cpp
--- a/Engine/Source/Editor/UnrealEd/Private/PlayLevel.cpp
+++ b/Engine/Source/Editor/UnrealEd/Private/PlayLevel.cpp
@@ -724,7 +724,7 @@ void UEditorEngine::TeardownPlaySession(FWorldContext& PieWorldContext)
 	bool bWasSimulatingInEditor = PlayInEditorSessionInfo->OriginalRequestParams.WorldType == EPlaySessionWorldType::SimulateInEditor;
 	
 	// Stop all audio and remove references to temp level.
-	if (FAudioDevice* AudioDevice = PlayWorld->GetAudioDevice())
+	if (FAudioDevice* AudioDevice = PlayWorld->GetAudioDeviceRaw())
 	{
 		AudioDevice->Flush(PlayWorld);
 		AudioDevice->ResetInterpolation();
@@ -2296,7 +2296,7 @@ void UEditorEngine::ResetPIEAudioSetting(UWorld *CurrentPieWorld)
 	ULevelEditorPlaySettings* PlayInSettings = GetMutableDefault<ULevelEditorPlaySettings>();
 	if (!PlayInSettings->EnableGameSound)
 	{
-		if (FAudioDevice* AudioDevice = CurrentPieWorld->GetAudioDevice())
+		if (FAudioDevice* AudioDevice = CurrentPieWorld->GetAudioDeviceRaw())
 		{
 			AudioDevice->SetTransientMasterVolume(0.0f);
 		}
@@ -2476,7 +2476,7 @@ void UEditorEngine::StartPlayInEditorSession(FRequestPlaySessionParams& InReques
 	}
 
 	// Flush all audio sources from the editor world
-	if (FAudioDevice* AudioDevice = EditorWorld->GetAudioDevice())
+	if (FAudioDevice* AudioDevice = EditorWorld->GetAudioDeviceRaw())
 	{
 		AudioDevice->Flush(EditorWorld);
 		AudioDevice->ResetInterpolation();
@@ -2728,7 +2728,7 @@ UGameInstance* UEditorEngine::CreateInnerProcessPIEGameInstance(FRequestPlaySess
 
 		if (!InParams.EditorPlaySettings->EnableGameSound)
 		{
-			if (FAudioDevice* GameInstanceAudioDevice = GameInstance->GetWorld()->GetAudioDevice())
+			if (FAudioDevice* GameInstanceAudioDevice = GameInstance->GetWorld()->GetAudioDeviceRaw())
 			{
 				GameInstanceAudioDevice->SetTransientMasterVolume(0.0f);
 			}
diff --git a/Engine/Source/Runtime/AudioMixer/Private/AudioMixerBlueprintLibrary.cpp b/Engine/Source/Runtime/AudioMixer/Private/AudioMixerBlueprintLibrary.cpp
--- a/Engine/Source/Runtime/AudioMixer/Private/AudioMixerBlueprintLibrary.cpp
+++ b/Engine/Source/Runtime/AudioMixer/Private/AudioMixerBlueprintLibrary.cpp
@@ -22,7 +22,7 @@ static FAudioDevice* GetAudioDeviceFromWorldContext(const UObject* WorldContextO
 		return nullptr;
 	}
 
-	return ThisWorld->GetAudioDevice();
+	return ThisWorld->GetAudioDevice().GetAudioDevice();
 }
 
 static Audio::FMixerDevice* GetAudioMixerDeviceFromWorldContext(const UObject* WorldContextObject)
diff --git a/Engine/Source/Runtime/AudioMixer/Private/AudioMixerDevice.cpp b/Engine/Source/Runtime/AudioMixer/Private/AudioMixerDevice.cpp
--- a/Engine/Source/Runtime/AudioMixer/Private/AudioMixerDevice.cpp
+++ b/Engine/Source/Runtime/AudioMixer/Private/AudioMixerDevice.cpp
@@ -1688,7 +1688,7 @@ namespace Audio
 
 	bool FMixerDevice::IsMainAudioDevice() const
 	{
-		bool bIsMain = (this == GEngine->GetMainAudioDevice());
+		bool bIsMain = (this == GEngine->GetMainAudioDeviceRaw());
 		return bIsMain;
 	}
 
diff --git a/Engine/Source/Runtime/AudioMixer/Private/AudioMixerSubmix.cpp b/Engine/Source/Runtime/AudioMixer/Private/AudioMixerSubmix.cpp
--- a/Engine/Source/Runtime/AudioMixer/Private/AudioMixerSubmix.cpp
+++ b/Engine/Source/Runtime/AudioMixer/Private/AudioMixerSubmix.cpp
@@ -110,7 +110,7 @@ namespace Audio
 				FSoundEffectSubmix* SubmixEffect = static_cast<FSoundEffectSubmix*>(EffectPreset->CreateNewEffect());
 				
 				FSoundEffectSubmixInitData InitData;
-				InitData.DeviceHandle = MixerDevice->DeviceHandle;
+				InitData.DeviceID = MixerDevice->DeviceID;
 				InitData.SampleRate = MixerDevice->GetSampleRate();
 				InitData.PresetSettings = nullptr;
 
diff --git a/Engine/Source/Runtime/AudioMixer/Private/Components/SynthComponent.cpp b/Engine/Source/Runtime/AudioMixer/Private/Components/SynthComponent.cpp
--- a/Engine/Source/Runtime/AudioMixer/Private/Components/SynthComponent.cpp
+++ b/Engine/Source/Runtime/AudioMixer/Private/Components/SynthComponent.cpp
@@ -408,13 +408,7 @@ FAudioDevice* USynthComponent::GetAudioDevice()
 	// If the synth component has a world, that means it was already registed with that world
 	if (UWorld* World = GetWorld())
 	{
-		// Make sure it has a proper audio device handle and retrieve it
-		if (World->AudioDeviceHandle != INDEX_NONE)
-		{
-			FAudioDeviceManager* AudioDeviceManager = GEngine->GetAudioDeviceManager();
-			check(AudioDeviceManager);
-			return AudioDeviceManager->GetAudioDevice(World->AudioDeviceHandle);
-		}
+		return World->AudioDeviceHandle.GetAudioDevice();
 	}
 
 	// Otherwise, retrieve the audio component's audio device (probably from it's owner)
diff --git a/Engine/Source/Runtime/AudioMixer/Private/Effects/AudioMixerSubmixEffectDynamicsProcessor.cpp b/Engine/Source/Runtime/AudioMixer/Private/Effects/AudioMixerSubmixEffectDynamicsProcessor.cpp
--- a/Engine/Source/Runtime/AudioMixer/Private/Effects/AudioMixerSubmixEffectDynamicsProcessor.cpp
+++ b/Engine/Source/Runtime/AudioMixer/Private/Effects/AudioMixerSubmixEffectDynamicsProcessor.cpp
@@ -27,7 +27,7 @@ FSubmixEffectDynamicsProcessor::~FSubmixEffectDynamicsProcessor()
 		Audio::FMixerDevice* MixerDevice = nullptr;
 		if (FAudioDeviceManager* DeviceManager = FAudioDeviceManager::Get())
 		{
-			MixerDevice = static_cast<Audio::FMixerDevice*>(DeviceManager->GetAudioDevice(DeviceId));
+			MixerDevice = static_cast<Audio::FMixerDevice*>(DeviceManager->GetAudioDeviceRaw(DeviceId));
 		}
 
 		if (MixerDevice)
@@ -50,7 +50,7 @@ void FSubmixEffectDynamicsProcessor::Init(const FSoundEffectSubmixInitData& Init
 	AudioOutputFrame.Reset();
 	AudioOutputFrame.AddZeroed(ProcessorScratchNumChannels);
 
-	DeviceId = InitData.DeviceHandle;
+	DeviceId = InitData.DeviceID;
 }
 
 void FSubmixEffectDynamicsProcessor::OnNewSubmixBuffer(
@@ -130,7 +130,7 @@ void FSubmixEffectDynamicsProcessor::OnPresetChanged()
 	Audio::FMixerDevice* MixerDevice = nullptr;
 	if (FAudioDeviceManager* DeviceManager = FAudioDeviceManager::Get())
 	{
-		MixerDevice = static_cast<Audio::FMixerDevice*>(DeviceManager->GetAudioDevice(DeviceId));
+		MixerDevice = static_cast<Audio::FMixerDevice*>(DeviceManager->GetAudioDeviceRaw(DeviceId));
 	}
 
 	if (MixerDevice)
diff --git a/Engine/Source/Runtime/Engine/Classes/Components/AudioComponent.h b/Engine/Source/Runtime/Engine/Classes/Components/AudioComponent.h
--- a/Engine/Source/Runtime/Engine/Classes/Components/AudioComponent.h
+++ b/Engine/Source/Runtime/Engine/Classes/Components/AudioComponent.h
@@ -251,7 +251,7 @@ private:
 
 public:
 	/** The specific audio device to play this component on */
-	uint32 AudioDeviceHandle;
+	uint32 AudioDeviceID;
 
 	/** Configurable, serialized ID for audio plugins */
 	UPROPERTY()
diff --git a/Engine/Source/Runtime/Engine/Classes/Engine/Engine.h b/Engine/Source/Runtime/Engine/Classes/Engine/Engine.h
--- a/Engine/Source/Runtime/Engine/Classes/Engine/Engine.h
+++ b/Engine/Source/Runtime/Engine/Classes/Engine/Engine.h
@@ -17,6 +17,7 @@
 #include "Subsystems/SubsystemCollection.h"
 #include "Subsystems/EngineSubsystem.h"
 #include "RHI.h"
+#include "AudioDeviceManager.h"
 #include "Engine.generated.h"
 
 #define WITH_DYNAMIC_RESOLUTION (!UE_SERVER)
@@ -24,7 +25,6 @@
 class AMatineeActor;
 class APlayerController;
 class Error;
-class FAudioDeviceManager;
 class FCanvas;
 class FCommonViewportClient;
 class FFineGrainedPerformanceTracker;
@@ -366,7 +366,7 @@ struct FWorldContext
 	bool	bWaitingOnOnlineSubsystem;
 
 	/** Handle to this world context's audio device.*/
-	uint32 AudioDeviceHandle;
+	uint32 AudioDeviceID;
 
 	/** Custom description to be display in blueprint debugger UI */
 	FString CustomDescription;
@@ -414,7 +414,7 @@ struct FWorldContext
 		, PIEWorldFeatureLevel(ERHIFeatureLevel::Num)
 		, RunAsDedicated(false)
 		, bWaitingOnOnlineSubsystem(false)
-		, AudioDeviceHandle(INDEX_NONE)
+		, AudioDeviceID(INDEX_NONE)
 		, ThisCurrentWorld(nullptr)
 	{ }
 
@@ -1772,7 +1772,7 @@ protected:
 	FAudioDeviceManager* AudioDeviceManager;
 
 	/** Audio device handle to the main audio device. */
-	uint32 MainAudioDeviceHandle;
+	FAudioDeviceHandle MainAudioDeviceHandle;
 
 public:
 
@@ -2439,13 +2439,14 @@ public:
 	FAudioDeviceManager* GetAudioDeviceManager();
 
 	/** @return the main audio device handle used by the engine. */
-	uint32 GetAudioDeviceHandle() const;
+	uint32 GetMainAudioDeviceID() const;
 
 	/** @return the main audio device. */
-	class FAudioDevice* GetMainAudioDevice();
+	FAudioDeviceHandle GetMainAudioDevice();
+	class FAudioDevice* GetMainAudioDeviceRaw();
 
 	/** @return the currently active audio device */
-	class FAudioDevice* GetActiveAudioDevice();
+	FAudioDeviceHandle GetActiveAudioDevice();
 
 	/** @return whether we're currently running in split screen (more than one local player) */
 	bool IsSplitScreen(UWorld *InWorld);
diff --git a/Engine/Source/Runtime/Engine/Classes/Engine/GameViewportClient.h b/Engine/Source/Runtime/Engine/Classes/Engine/GameViewportClient.h
--- a/Engine/Source/Runtime/Engine/Classes/Engine/GameViewportClient.h
+++ b/Engine/Source/Runtime/Engine/Classes/Engine/GameViewportClient.h
@@ -18,6 +18,7 @@
 #include "Engine/DebugDisplayProperty.h"
 #include "UObject/SoftObjectPath.h"
 #include "StereoRendering.h"
+#include "AudioDeviceManager.h"
 
 #include "GameViewportClient.generated.h"
 
@@ -100,6 +101,9 @@ protected:
 	/** If true will suppress the blue transition text messages. */
 	bool bSuppressTransitionMessage;
 
+	/** Strong handle to the audio device used by this viewport. */
+	FAudioDeviceHandle AudioDevice;
+
 public:
 
 	/** see enum EViewModeIndex */
@@ -982,9 +986,6 @@ private:
 	/** Mouse cursor locking behavior when the viewport is clicked */
 	EMouseLockMode MouseLockMode;
 
-	/** Handle to the audio device created for this viewport. Each viewport (for multiple PIE) will have its own audio device. */
-	uint32 AudioDeviceHandle;
-
 	/** Whether or not this audio device is in audio-focus */
 	bool bHasAudioFocus;
 
diff --git a/Engine/Source/Runtime/Engine/Classes/Engine/World.h b/Engine/Source/Runtime/Engine/Classes/Engine/World.h
--- a/Engine/Source/Runtime/Engine/Classes/Engine/World.h
+++ b/Engine/Source/Runtime/Engine/Classes/Engine/World.h
@@ -22,7 +22,7 @@
 #include "Physics/PhysicsInterfaceDeclares.h"
 #include "Particles/WorldPSCPool.h"
 #include "Containers/SortedMap.h"
-
+#include "AudioDeviceManager.h"
 #include "Subsystems/WorldSubsystem.h"
 #include "Subsystems/SubsystemCollection.h"
 
@@ -1188,7 +1188,7 @@ private:
 public:
 
 	/** Handle to the active audio device for this world. */
-	uint32 AudioDeviceHandle;
+	FAudioDeviceHandle AudioDeviceHandle;
 
 #if WITH_EDITOR
 	/** Hierarchical LOD System. Used when WorldSetting.bEnableHierarchicalLODSystem is true */
@@ -3306,18 +3306,20 @@ public:
 	 */
 	class AAudioVolume* GetAudioSettings( const FVector& ViewLocation, struct FReverbSettings* OutReverbSettings, struct FInteriorSettings* OutInteriorSettings );
 
-	/** Returns the audio device handle for this world.*/
-	uint32 GetAudioDeviceHandle() const { return AudioDeviceHandle; }
+	void SetAudioDevice(FAudioDeviceHandle& InHandle);
 
-	/** Sets the audio device handle to the active audio device for this world.*/
-	void SetAudioDeviceHandle(const uint32 InAudioDeviceHandle);
+	/**
+	 * Get the audio device used by this world.
+	 */
+	FAudioDeviceHandle GetAudioDevice();
 
 	/**
-	* Returns the audio device associated with this world, or returns the main audio device if there is none.
+	* Returns the audio device associated with this world.
+	* Lifecycle of the audio device is not guaranteed unless you used GetAudioDevice().
 	*
 	* @return Audio device to use with this world.
 	*/
-	class FAudioDevice* GetAudioDevice();
+	class FAudioDevice* GetAudioDeviceRaw();
 
 	/** Return the URL of this level on the local machine. */
 	virtual FString GetLocalURL() const;
diff --git a/Engine/Source/Runtime/Engine/Classes/Sound/SoundBase.h b/Engine/Source/Runtime/Engine/Classes/Sound/SoundBase.h
--- a/Engine/Source/Runtime/Engine/Classes/Sound/SoundBase.h
+++ b/Engine/Source/Runtime/Engine/Classes/Sound/SoundBase.h
@@ -13,7 +13,7 @@
 #include "SoundSubmixSend.h"
 #include "UObject/Object.h"
 #include "UObject/ObjectMacros.h"
-
+#include "AudioDeviceManager.h"
 #include "SoundBase.generated.h"
 
 
diff --git a/Engine/Source/Runtime/Engine/Classes/Sound/SoundEffectSubmix.h b/Engine/Source/Runtime/Engine/Classes/Sound/SoundEffectSubmix.h
--- a/Engine/Source/Runtime/Engine/Classes/Sound/SoundEffectSubmix.h
+++ b/Engine/Source/Runtime/Engine/Classes/Sound/SoundEffectSubmix.h
@@ -29,12 +29,12 @@ class ENGINE_API USoundEffectSubmixPreset : public USoundEffectPreset
 /** Struct which has data needed to initialize the submix effect. */
 struct FSoundEffectSubmixInitData
 {
-	uint32 DeviceHandle;
+	uint32 DeviceID;
 	void* PresetSettings;
 	float SampleRate;
 
 	FSoundEffectSubmixInitData()
-		: DeviceHandle(INDEX_NONE)
+		: DeviceID(INDEX_NONE)
 		, PresetSettings(nullptr)
 		, SampleRate(0)
 	{
diff --git a/Engine/Source/Runtime/Engine/Private/ActiveSound.cpp b/Engine/Source/Runtime/Engine/Private/ActiveSound.cpp
--- a/Engine/Source/Runtime/Engine/Private/ActiveSound.cpp
+++ b/Engine/Source/Runtime/Engine/Private/ActiveSound.cpp
@@ -198,7 +198,7 @@ int32 FActiveSound::GetPlayCount() const
 		return 0;
 	}
 
-	if (const int32* PlayCount = Sound->CurrentPlayCount.Find(AudioDevice->DeviceHandle))
+	if (const int32* PlayCount = Sound->CurrentPlayCount.Find(AudioDevice->DeviceID))
 	{
 		return *PlayCount;
 	}
@@ -764,13 +764,13 @@ void FActiveSound::MarkPendingDestroy(bool bDestroyNow)
 
 	if (Sound && !bIsStopping)
 	{
-		int32* PlayCount = AudioDevice ? Sound->CurrentPlayCount.Find(AudioDevice->DeviceHandle) : nullptr;
+		int32* PlayCount = AudioDevice ? Sound->CurrentPlayCount.Find(AudioDevice->DeviceID) : nullptr;
 		if (PlayCount)
 		{
 			*PlayCount = FMath::Max(*PlayCount - 1, 0);
 			if (*PlayCount == 0)
 			{
-				Sound->CurrentPlayCount.Remove(AudioDevice->DeviceHandle);
+				Sound->CurrentPlayCount.Remove(AudioDevice->DeviceID);
 			}
 		}
 
@@ -953,7 +953,7 @@ void FActiveSound::OcclusionTraceDone(const FTraceHandle& TraceHandle, FTraceDat
 	{
 		if (FAudioDeviceManager* AudioDeviceManager = GEngine->GetAudioDeviceManager())
 		{
-			if (FAudioDevice* AudioDevice = AudioDeviceManager->GetAudioDevice(TraceDetails.AudioDeviceID))
+			if (FAudioDevice* AudioDevice = AudioDeviceManager->GetAudioDeviceRaw(TraceDetails.AudioDeviceID))
 			{
 				FActiveSound* ActiveSound = TraceDetails.ActiveSound;
 
@@ -1011,7 +1011,7 @@ void FActiveSound::CheckOcclusion(const FVector ListenerLocation, const FVector
 				const uint32 SoundOwnerID = OwnerID;
 				TWeakObjectPtr<UWorld> SoundWorld = World;
 				FAsyncTraceDetails TraceDetails;
-				TraceDetails.AudioDeviceID = AudioDevice->DeviceHandle;
+				TraceDetails.AudioDeviceID = AudioDevice->DeviceID;
 				TraceDetails.ActiveSound = this;
 
 				FAudioThread::RunCommandOnGameThread([SoundWorld, SoundLocation, ListenerLocation, OcclusionTraceChannel, SoundOwnerID, bUseComplexCollisionForOcclusion, TraceDetails]
diff --git a/Engine/Source/Runtime/Engine/Private/AudioDebug.cpp b/Engine/Source/Runtime/Engine/Private/AudioDebug.cpp
--- a/Engine/Source/Runtime/Engine/Private/AudioDebug.cpp
+++ b/Engine/Source/Runtime/Engine/Private/AudioDebug.cpp
@@ -210,7 +210,7 @@ namespace
 			return nullptr;
 		}
 
-		return World->GetAudioDevice();
+		return World->GetAudioDeviceRaw();
 	}
 
 	struct FAudioStats
@@ -788,7 +788,7 @@ void FAudioDebugger::DumpActiveSounds() const
 		return;
 	}
 
-	FAudioDevice* AudioDevice = GEngine->GetAudioDeviceManager()->GetActiveAudioDevice();
+	FAudioDeviceHandle AudioDevice = GEngine->GetAudioDeviceManager()->GetActiveAudioDevice();
 	if (!AudioDevice)
 	{
 		return;
@@ -861,7 +861,7 @@ void FAudioDebugger::ResolveDesiredStats(FViewportClient* ViewportClient)
 
 	if (ViewportClient->IsStatEnabled(TEXT("Sounds")))
 	{
-		FAudioStats& Stats = AudioDeviceStats.FindOrAdd(AudioDevice->DeviceHandle);
+		FAudioStats& Stats = AudioDeviceStats.FindOrAdd(AudioDevice->DeviceID);
 		SetStats |= ERequestedAudioStats::Sounds;
 
 		if (Stats.DisplayFlags & static_cast<uint8>(FAudioStats::EDisplayFlags::Debug))
@@ -891,7 +891,7 @@ void FAudioDebugger::ResolveDesiredStats(FViewportClient* ViewportClient)
 
 	DECLARE_CYCLE_STAT(TEXT("FAudioThreadTask.ResolveDesiredStats"), STAT_AudioResolveDesiredStats, STATGROUP_TaskGraphTasks);
 
-	const uint32 DeviceID = AudioDevice->DeviceHandle;
+	const uint32 DeviceID = AudioDevice->DeviceID;
 	if (IsInAudioThread())
 	{
 		FAudioStats_AudioThread& Stats = AudioDeviceStats_AudioThread.FindOrAdd(DeviceID);
@@ -924,7 +924,7 @@ int32 FAudioDebugger::RenderStatCues(UWorld* World, FViewport* Viewport, FCanvas
 	Y += FontHeight;
 
 	int32 ActiveSoundCount = 0;
-	FAudioStats& AudioStats = AudioDeviceStats.FindOrAdd(AudioDevice->DeviceHandle);
+	FAudioStats& AudioStats = AudioDeviceStats.FindOrAdd(AudioDevice->DeviceID);
 
 	if (!SoundCueDebugMinimalCVar)
 	{
@@ -1064,7 +1064,7 @@ int32 FAudioDebugger::RenderStatMixes(UWorld* World, FViewport* Viewport, FCanva
 
 	bool bDisplayedSoundMixes = false;
 
-	FAudioStats& AudioStats = AudioDeviceStats.FindOrAdd(AudioDevice->DeviceHandle);
+	FAudioStats& AudioStats = AudioDeviceStats.FindOrAdd(AudioDevice->DeviceID);
 	if (AudioStats.StatSoundMixes.Num() > 0)
 	{
 		bDisplayedSoundMixes = true;
@@ -1213,7 +1213,7 @@ int32 FAudioDebugger::RenderStatSounds(UWorld* World, FViewport* Viewport, FCanv
 	const int32 FontHeight = GetStatsFont()->GetMaxCharHeight() + 2;
 	Y += FontHeight;
 
-	FAudioStats& AudioStats = AudioDeviceStats.FindOrAdd(AudioDevice->DeviceHandle);
+	FAudioStats& AudioStats = AudioDeviceStats.FindOrAdd(AudioDevice->DeviceID);
 
 	const uint8 bDebug = AudioStats.DisplayFlags & static_cast<uint8>(FAudioStats::EDisplayFlags::Debug);
 
@@ -1375,7 +1375,7 @@ int32 FAudioDebugger::RenderStatSounds(UWorld* World, FViewport* Viewport, FCanv
 		}
 	}
 
-	Canvas->DrawShadowedString(X, Y, *FString::Printf(TEXT("Audio Device ID: %u, Max Sounds Displayed: %i"), AudioDevice->DeviceHandle, AudioDebugSoundMaxNumDisplayedCVar), GetStatsFont(), HeaderColor);
+	Canvas->DrawShadowedString(X, Y, *FString::Printf(TEXT("Audio Device ID: %u, Max Sounds Displayed: %i"), AudioDevice->DeviceID, AudioDebugSoundMaxNumDisplayedCVar), GetStatsFont(), HeaderColor);
 	Y += FontHeight;
 
 	Canvas->DrawShadowedString(X, Y, *FString::Printf(TEXT("Total Sounds: %i, Sound Waves: %i"), AudioStats.StatSoundInfos.Num(), TotalSoundWavesNum), GetStatsFont(), HeaderColor);
@@ -1491,7 +1491,7 @@ int32 FAudioDebugger::RenderStatWaves(UWorld* World, FViewport* Viewport, FCanva
 
 	const int32 FontHeight = GetStatsFont()->GetMaxCharHeight() + 2;
 
-	FAudioStats& AudioStats = AudioDeviceStats.FindOrAdd(AudioDevice->DeviceHandle);
+	FAudioStats& AudioStats = AudioDeviceStats.FindOrAdd(AudioDevice->DeviceID);
 	Canvas->DrawShadowedString(X, Y, TEXT("Active Sound Waves:"), GetStatsFont(), FLinearColor::Green);
 	Y += TabWidth;
 
@@ -1551,8 +1551,8 @@ int32 FAudioDebugger::RenderStatWaves(UWorld* World, FViewport* Viewport, FCanva
 
 void FAudioDebugger::RemoveDevice(const FAudioDevice& AudioDevice)
 {
-	AudioDeviceStats.Remove(AudioDevice.DeviceHandle);
-	AudioDeviceStats_AudioThread.Remove(AudioDevice.DeviceHandle);
+	AudioDeviceStats.Remove(AudioDevice.DeviceID);
+	AudioDeviceStats_AudioThread.Remove(AudioDevice.DeviceID);
 }
 
 bool FAudioDebugger::ToggleStats(UWorld* World, const uint8 StatToToggle)
@@ -1570,7 +1570,7 @@ bool FAudioDebugger::ToggleStats(UWorld* World, const uint8 StatToToggle)
 
 	if (FAudioDeviceManager* DeviceManager = GEngine->GetAudioDeviceManager())
 	{
-		DeviceManager->GetDebugger().ToggleStats(AudioDevice->DeviceHandle, StatToToggle);
+		DeviceManager->GetDebugger().ToggleStats(AudioDevice->DeviceID, StatToToggle);
 	}
 
 	return true;
@@ -1638,7 +1638,7 @@ bool FAudioDebugger::PostStatModulatorHelp(UWorld* World, FCommonViewportClient*
 		return false;
 	}
 
-	if (FAudioDevice* AudioDevice = World->GetAudioDevice())
+	if (FAudioDevice* AudioDevice = World->GetAudioDeviceRaw())
 	{
 		if (AudioDevice->IsModulationPluginEnabled())
 		{
@@ -1715,7 +1715,7 @@ bool FAudioDebugger::ToggleStatSounds(UWorld* World, FCommonViewportClient* View
 
 	FAudioDevice* AudioDevice = GetWorldAudio(World);
 	check(AudioDevice);
-	FAudioStats& Stats = AudioDeviceStats.FindOrAdd(AudioDevice->DeviceHandle);
+	FAudioStats& Stats = AudioDeviceStats.FindOrAdd(AudioDevice->DeviceID);
 	Stats.DisplayFlags = ShowSounds;
 	Stats.DisplaySort = DisplaySort;
 
@@ -1733,7 +1733,7 @@ void FAudioDebugger::SendUpdateResultsToGameThread(const FAudioDevice& AudioDevi
 {
 	check(IsInAudioThread());
 
-	FAudioStats_AudioThread* Stats_AudioThread = AudioDeviceStats_AudioThread.Find(AudioDevice.DeviceHandle);
+	FAudioStats_AudioThread* Stats_AudioThread = AudioDeviceStats_AudioThread.Find(AudioDevice.DeviceID);
 	if (!Stats_AudioThread)
 	{
 		return;
@@ -1838,7 +1838,7 @@ void FAudioDebugger::SendUpdateResultsToGameThread(const FAudioDevice& AudioDevi
 
 	DECLARE_CYCLE_STAT(TEXT("FGameThreadAudioTask.AudioSendResults"), STAT_AudioSendResults, STATGROUP_TaskGraphTasks);
 
-	const uint32 AudioDeviceID = AudioDevice.DeviceHandle;
+	const uint32 AudioDeviceID = AudioDevice.DeviceID;
 
 	TArray<FTransform> ListenerTransforms;
 	for (const FListener& Listener : AudioDevice.GetListeners())
diff --git a/Engine/Source/Runtime/Engine/Private/AudioDevice.cpp b/Engine/Source/Runtime/Engine/Private/AudioDevice.cpp
--- a/Engine/Source/Runtime/Engine/Private/AudioDevice.cpp
+++ b/Engine/Source/Runtime/Engine/Private/AudioDevice.cpp
@@ -251,7 +251,7 @@ FAudioDevice::FAudioDevice()
 	: NumStoppingSources(32)
 	, SampleRate(0)
 	, NumPrecacheFrames(MONO_PCM_BUFFER_SAMPLES)
-	, DeviceHandle(static_cast<Audio::FDeviceId>(INDEX_NONE))
+	, DeviceID(static_cast<Audio::FDeviceId>(INDEX_NONE))
 	, SpatializationPluginInterface(nullptr)
 	, ReverbPluginInterface(nullptr)
 	, OcclusionInterface(nullptr)
@@ -1330,7 +1330,7 @@ bool FAudioDevice::HandleSoloCommand(const TCHAR* Cmd, FOutputDevice& Ar)
 	FAudioDeviceManager* DeviceManager = GEngine->GetAudioDeviceManager();
 	if (DeviceManager)
 	{
-		DeviceManager->SetSoloDevice(DeviceHandle);
+		DeviceManager->SetSoloDevice(DeviceID);
 	}
 	return true;
 }
@@ -4348,7 +4348,7 @@ void FAudioDevice::SendUpdateResultsToGameThread(const int32 FirstActiveIndex)
 {
 	DECLARE_CYCLE_STAT(TEXT("FGameThreadAudioTask.AudioSendResults"), STAT_AudioSendResults, STATGROUP_TaskGraphTasks);
 
-	const Audio::FDeviceId AudioDeviceID = DeviceHandle;
+	const Audio::FDeviceId AudioDeviceID = DeviceID;
 	UReverbEffect* ReverbEffect = Effects ? Effects->GetCurrentReverbEffect() : nullptr;
 	FAudioThread::RunCommandOnGameThread([AudioDeviceID, ReverbEffect]()
 	{
@@ -4357,7 +4357,7 @@ void FAudioDevice::SendUpdateResultsToGameThread(const int32 FirstActiveIndex)
 		{
 			if (FAudioDeviceManager* AudioDeviceManager = GEngine->GetAudioDeviceManager())
 			{
-				if (FAudioDevice* AudioDevice = AudioDeviceManager->GetAudioDevice(AudioDeviceID))
+				if (FAudioDeviceHandle AudioDevice = AudioDeviceManager->GetAudioDevice(AudioDeviceID))
 				{
 					AudioDevice->CurrentReverbEffect = ReverbEffect;
 				}
@@ -4577,10 +4577,10 @@ void FAudioDevice::AddNewActiveSoundInternal(const FActiveSound& NewActiveSound,
 		}
 	}
 
-	int32* PlayCount = Sound->CurrentPlayCount.Find(DeviceHandle);
+	int32* PlayCount = Sound->CurrentPlayCount.Find(DeviceID);
 	if (!PlayCount)
 	{
-		PlayCount = &Sound->CurrentPlayCount.Add(DeviceHandle);
+		PlayCount = &Sound->CurrentPlayCount.Add(DeviceID);
 	}
 	(*PlayCount)++;
 
@@ -5343,7 +5343,7 @@ FAudioDevice::FCreateComponentParams::FCreateComponentParams()
 	: World(nullptr)
 	, Actor(nullptr)
 {
-	AudioDevice = (GEngine ? GEngine->GetMainAudioDevice() : nullptr);
+	AudioDevice = (GEngine ? GEngine->GetMainAudioDeviceRaw() : nullptr);
 	CommonInit();
 }
 
@@ -5360,7 +5360,7 @@ FAudioDevice::FCreateComponentParams::FCreateComponentParams(UWorld* InWorld, AA
 		Actor = (World ? World->GetWorldSettings() : nullptr);
 	}
 
-	AudioDevice = (World ? World->GetAudioDevice() : nullptr);
+	AudioDevice = (World ? World->GetAudioDeviceRaw() : nullptr);
 	CommonInit();
 }
 
@@ -5368,7 +5368,7 @@ FAudioDevice::FCreateComponentParams::FCreateComponentParams(AActor* InActor)
 	: Actor(InActor)
 {
 	World = (Actor ? Actor->GetWorld() : nullptr);
-	AudioDevice = (World ? World->GetAudioDevice() : nullptr);
+	AudioDevice = (World ? World->GetAudioDeviceRaw() : nullptr);
 	CommonInit();
 }
 
@@ -5418,7 +5418,7 @@ UAudioComponent* FAudioDevice::CreateComponent(USoundBase* Sound, UWorld* World,
 	}
 	else
 	{
-		Params = MakeUnique<FCreateComponentParams>(GEngine->GetMainAudioDevice());
+		Params = MakeUnique<FCreateComponentParams>(GEngine->GetMainAudioDeviceRaw());
 	}
 
 	Params->bPlay = bPlay;
@@ -5499,7 +5499,7 @@ UAudioComponent* FAudioDevice::CreateComponent(USoundBase* Sound, const FCreateC
 				}
 				else
 				{
-					AudioComponent->AudioDeviceHandle = Params.AudioDevice->DeviceHandle;
+					AudioComponent->AudioDeviceID = Params.AudioDevice->DeviceID;
 				}
 
 				if (Params.bPlay)
diff --git a/Engine/Source/Runtime/Engine/Private/AudioDeviceManager.cpp b/Engine/Source/Runtime/Engine/Private/AudioDeviceManager.cpp
--- a/Engine/Source/Runtime/Engine/Private/AudioDeviceManager.cpp
+++ b/Engine/Source/Runtime/Engine/Private/AudioDeviceManager.cpp
@@ -12,26 +12,13 @@
 #include "Misc/CommandLine.h"
 #include "Misc/ConfigCacheIni.h"
 #include "UObject/UObjectIterator.h"
+#include "Audio/AudioDebug.h"
 
 #if WITH_EDITOR
 #include "AudioEditorModule.h"
 #include "Settings/LevelEditorMiscSettings.h"
 #endif
 
-// Private consts for helping with index/generation determination in audio device manager
-static const uint32 AUDIO_DEVICE_HANDLE_INDEX_BITS		= 24;
-static const uint32 AUDIO_DEVICE_HANDLE_INDEX_MASK		= (1 << AUDIO_DEVICE_HANDLE_INDEX_BITS) - 1;
-static const uint32 AUDIO_DEVICE_HANDLE_GENERATION_BITS = 8;
-static const uint32 AUDIO_DEVICE_HANDLE_GENERATION_MASK = (1 << AUDIO_DEVICE_HANDLE_GENERATION_BITS) - 1;
-
-static const uint16 AUDIO_DEVICE_MINIMUM_FREE_AUDIO_DEVICE_INDICES = 32;
-
-// The number of multiple audio devices allowed by default
-static const uint32 AUDIO_DEVICE_DEFAULT_ALLOWED_DEVICE_COUNT = 2;
-
-// The max number of audio devices allowed
-static const uint32 AUDIO_DEVICE_MAX_DEVICE_COUNT = 8;
-
 static int32 GCVarEnableAudioThreadWait = 1;
 TAutoConsoleVariable<int32> CVarEnableAudioThreadWait(
 	TEXT("AudioThread.EnableAudioThreadWait"),
@@ -57,26 +44,16 @@ FAutoConsoleVariableRef CVarAudioVisualizeEnabled(
 	TEXT("0: Not Enabled, 1: Enabled"),
 	ECVF_Default);
 
-
-FAudioDeviceManager::FCreateAudioDeviceResults::FCreateAudioDeviceResults()
-	: Handle(INDEX_NONE)
-	, bNewDevice(false)
-	, AudioDevice(nullptr)
-{
-}
-
 /*-----------------------------------------------------------------------------
 FAudioDeviceManager implementation.
 -----------------------------------------------------------------------------*/
 
 FAudioDeviceManager::FAudioDeviceManager()
 	: AudioDeviceModule(nullptr)
-	, FreeIndicesSize(0)
-	, NumActiveAudioDevices(0)
-	, NumWorldsUsingMainAudioDevice(0)
+	, DeviceIDCounter(0)
 	, NextResourceID(1)
 	, SoloDeviceHandle(INDEX_NONE)
-	, ActiveAudioDeviceHandle(INDEX_NONE)
+	, ActiveAudioDeviceID(INDEX_NONE)
 	, bUsingAudioMixer(false)
 	, bPlayAllDeviceAudio(false)
 	, bOnlyToggleAudioMixerOnce(false)
@@ -84,6 +61,7 @@ FAudioDeviceManager::FAudioDeviceManager()
 {
 
 #if ENABLE_AUDIO_DEBUG
+	AudioDebugger = TUniquePtr<FAudioDebugger>(new FAudioDebugger());
 
 	// Check for a command line debug sound argument.
 	FString DebugSound;
@@ -97,8 +75,7 @@ FAudioDeviceManager::FAudioDeviceManager()
 
 FAudioDeviceManager::~FAudioDeviceManager()
 {
-	// Confirm that we freed all the audio devices
-	check(NumActiveAudioDevices == 0);
+	Devices.Reset();
 
 	// Release any loaded buffers - this calls stop on any sources that need it
 	for (int32 Index = Buffers.Num() - 1; Index >= 0; Index--)
@@ -151,84 +128,81 @@ void FAudioDeviceManager::ToggleAudioMixer()
 			// We could have multiple audio devices, so loop through them and patch them up best we can to
 			// get parity. E.g. we need to pass the handle from the odl to the new, set whether or not its active
 			// and try and get the mix-states to be the same.
-			for (int32 DeviceIndex = 0; DeviceIndex < Devices.Num(); ++DeviceIndex)
+			for (auto& DeviceContainer : Devices)
 			{
-				FAudioDevice* AudioDevice = Devices[DeviceIndex];
+				FAudioDevice*& AudioDevice = DeviceContainer.Value.Device;
+
+				check(AudioDevice);
+
+				// Get the audio device handle and whether it is active
+				uint32 DeviceID = AudioDevice->DeviceID;
+				check(DeviceContainer.Key == DeviceID);
+				bool bIsActive = (DeviceID == ActiveAudioDeviceID);
 
-				if (AudioDevice)
+				// To transfer mix states, we need to re-base the absolute clocks on the mix states
+				// so the target audio device timing won't result in the mixes suddenly stopping.
+				TMap<USoundMix*, FSoundMixState> MixModifiers = AudioDevice->GetSoundMixModifiers();
+				TArray<USoundMix*> PrevPassiveSoundMixModifiers = AudioDevice->GetPrevPassiveSoundMixModifiers();
+				USoundMix* BaseSoundMix = AudioDevice->GetDefaultBaseSoundMixModifier();
+				double AudioClock = AudioDevice->GetAudioClock();
+
+				for (TPair<USoundMix*, FSoundMixState>& SoundMixPair : MixModifiers)
 				{
-					// Get the audio device handle and whether it is active
-					uint32 Handle = AudioDevice->DeviceHandle;
-					bool bIsActive = (Handle == ActiveAudioDeviceHandle);
-
-					// To transfer mix states, we need to re-base the absolute clocks on the mix states
-					// so the target audio device timing won't result in the mixes suddenly stopping.
-					TMap<USoundMix*, FSoundMixState> MixModifiers = AudioDevice->GetSoundMixModifiers();
-					TArray<USoundMix*> PrevPassiveSoundMixModifiers = AudioDevice->GetPrevPassiveSoundMixModifiers();
-					USoundMix* BaseSoundMix = AudioDevice->GetDefaultBaseSoundMixModifier();
-					double AudioClock = AudioDevice->GetAudioClock();
-
-					for (TPair<USoundMix*, FSoundMixState>& SoundMixPair : MixModifiers)
+					// Rebase so that a new clock starting from 0.0 won't cause mixes to stop.
+					SoundMixPair.Value.StartTime -= AudioClock;
+					SoundMixPair.Value.FadeInStartTime -= AudioClock;
+					SoundMixPair.Value.FadeInEndTime -= AudioClock;
+
+					if (SoundMixPair.Value.EndTime > 0.0f)
 					{
-						// Rebase so that a new clock starting from 0.0 won't cause mixes to stop.
-						SoundMixPair.Value.StartTime -= AudioClock;
-						SoundMixPair.Value.FadeInStartTime -= AudioClock;
-						SoundMixPair.Value.FadeInEndTime -= AudioClock;
-
-						if (SoundMixPair.Value.EndTime > 0.0f)
-						{
-							SoundMixPair.Value.EndTime -= AudioClock;
-						}
-
-						if (SoundMixPair.Value.FadeOutStartTime > 0.0f)
-						{
-							SoundMixPair.Value.FadeOutStartTime -= AudioClock;
-						}
+						SoundMixPair.Value.EndTime -= AudioClock;
 					}
 
-					// Tear it down and delete the old audio device. This does a bunch of cleanup.
-					AudioDevice->Teardown();
-					delete AudioDevice;
-
-					// Make a new audio device using the new audio device module
-					AudioDevice = AudioDeviceModule->CreateAudioDevice();
-
-					// Some AudioDeviceModules override CreteAudioMixerPlatformInterface, which means we can create a Audio::FMixerDevice.
-					if (AudioDevice == nullptr)
+					if (SoundMixPair.Value.FadeOutStartTime > 0.0f)
 					{
-						checkf(AudioDeviceModule->IsAudioMixerModule(), TEXT("Please override AudioDeviceModule->CreateAudioDevice()"))
-						AudioDevice = new Audio::FMixerDevice(AudioDeviceModule->CreateAudioMixerPlatformInterface());
+						SoundMixPair.Value.FadeOutStartTime -= AudioClock;
 					}
+				}
 
-					check(AudioDevice);
+				// Tear it down and delete the old audio device. This does a bunch of cleanup.
+				AudioDevice->Teardown();
+				delete AudioDevice;
 
-					// Set the new audio device into the slot of the old audio device in the manager
-					Devices[DeviceIndex] = AudioDevice;
+				// Make a new audio device using the new audio device module
+				AudioDevice = AudioDeviceModule->CreateAudioDevice();
 
-					// Set the new audio device handle to the old audio device handle
-					AudioDevice->DeviceHandle = Handle;
+				// Some AudioDeviceModules override CreteAudioMixerPlatformInterface, which means we can create a Audio::FMixerDevice.
+				if (AudioDevice == nullptr)
+				{
+					checkf(AudioDeviceModule->IsAudioMixerModule(), TEXT("Please override AudioDeviceModule->CreateAudioDevice()"))
+						AudioDevice = new Audio::FMixerDevice(AudioDeviceModule->CreateAudioMixerPlatformInterface());
+				}
 
-					// Re-init the new audio device using appropriate settings so it behaves the same
-					if (AudioDevice->Init(AudioSettings->GetHighestMaxChannels()))
-					{
-						AudioDevice->SetMaxChannels(QualityLevelMaxChannels);
-					}
+				check(AudioDevice);
 
-					// Transfer the sound mix modifiers to the new audio engine
-					AudioDevice->SetSoundMixModifiers(MixModifiers, PrevPassiveSoundMixModifiers, BaseSoundMix);
-					// Setup the mute state of the audio device to be the same that it was
-					if (bIsActive)
-					{
-						AudioDevice->SetDeviceMuted(false);
-					}
-					else
-					{
-						AudioDevice->SetDeviceMuted(true);
-					}
+				// Set the new audio device handle to the old audio device handle
+				AudioDevice->DeviceID = DeviceID;
 
-					// Fade in the new audio device (used only in audio mixer to prevent pops on startup/shutdown)
-					AudioDevice->FadeIn();
+				// Re-init the new audio device using appropriate settings so it behaves the same
+				if (AudioDevice->Init(AudioSettings->GetHighestMaxChannels()))
+				{
+					AudioDevice->SetMaxChannels(QualityLevelMaxChannels);
+				}
+
+				// Transfer the sound mix modifiers to the new audio engine
+				AudioDevice->SetSoundMixModifiers(MixModifiers, PrevPassiveSoundMixModifiers, BaseSoundMix);
+				// Setup the mute state of the audio device to be the same that it was
+				if (bIsActive)
+				{
+					AudioDevice->SetDeviceMuted(false);
+				}
+				else
+				{
+					AudioDevice->SetDeviceMuted(true);
 				}
+
+				// Fade in the new audio device (used only in audio mixer to prevent pops on startup/shutdown)
+				AudioDevice->FadeIn();
 			}
 
 			// We now must free any resources that have been cached with the old audio engine
@@ -260,6 +234,53 @@ IAudioDeviceModule* FAudioDeviceManager::GetAudioDeviceModule()
 	return AudioDeviceModule;
 }
 
+
+
+FAudioDeviceParams FAudioDeviceManager::GetDefaultParamsForNewWorld()
+{
+	bool bCreateNewAudioDeviceForPlayInEditor = false;
+
+#if WITH_EDITOR
+	bCreateNewAudioDeviceForPlayInEditor = GetDefault<ULevelEditorMiscSettings>()->bCreateNewAudioDeviceForPlayInEditor;
+#endif
+
+	FAudioDeviceParams Params;
+	Params.Scope = bCreateNewAudioDeviceForPlayInEditor ? EAudioDeviceScope::Unique : EAudioDeviceScope::Shared;
+
+	return Params;
+}
+
+FAudioDeviceHandle FAudioDeviceManager::RequestAudioDevice(const FAudioDeviceParams& InParams)
+{
+	FScopeLock ScopeLock(&DeviceMapCriticalSection);
+	if (InParams.Scope == EAudioDeviceScope::Unique)
+	{
+		return CreateNewDevice(InParams);
+	}
+	else
+	{
+		// See if we already have a device we can use.
+		for (auto& Device : Devices)
+		{
+			if (CanUseAudioDevice(InParams, Device.Value))
+			{
+				if (InParams.AssociatedWorld != nullptr)
+				{
+					
+					Device.Value.WorldsUsingThisDevice.AddUnique(InParams.AssociatedWorld);
+					FAudioDeviceManagerDelegates::OnWorldRegisteredToAudioDevice.Broadcast(InParams.AssociatedWorld, Device.Key);
+				}
+
+				return BuildNewHandle(Device.Value, Device.Key, InParams);
+			}
+		}
+
+		// If we did not find a suitable device, build one.
+		return CreateNewDevice(InParams);
+	}
+	
+}
+
 bool FAudioDeviceManager::Initialize()
 {
 	if (LoadDefaultAudioDeviceModule())
@@ -278,7 +299,15 @@ bool FAudioDeviceManager::Initialize()
 		}
 #endif
 
-		return CreateMainAudioDevice();
+		// Initialize the main audio device.
+		FAudioDeviceParams MainDeviceParams;
+		MainDeviceParams.Scope = EAudioDeviceScope::Shared;
+		MainDeviceParams.bIsNonRealtime = false;
+
+		MainAudioDeviceHandle = RequestAudioDevice(MainDeviceParams);
+		check(MainAudioDeviceHandle.IsValid());
+		FAudioThread::StartAudioThread();
+		return true;
 	}
 
 	// Failed to initialize
@@ -301,19 +330,26 @@ bool FAudioDeviceManager::LoadDefaultAudioDeviceModule()
 	if (!bForceAudioMixer && !bForceNoAudioMixer)
 	{
 		GConfig->GetBool(TEXT("Audio"), TEXT("UseAudioMixer"), bUsingAudioMixer, GEngineIni);
+		// Get the audio mixer and non-audio mixer device module names
+		GConfig->GetString(TEXT("Audio"), TEXT("AudioDeviceModuleName"), AudioDeviceModuleName, GEngineIni);
+		GConfig->GetString(TEXT("Audio"), TEXT("AudioMixerModuleName"), AudioMixerModuleName, GEngineIni);
 	}
 	else if (bForceNoAudioMixer)
 	{
+		GConfig->GetString(TEXT("Audio"), TEXT("AudioDeviceModuleName"), AudioDeviceModuleName, GEngineIni);
+
 		// Allow no audio mixer override from command line
 		bUsingAudioMixer = false;
 	}
+	else if(bForceAudioMixer)
+	{
+		GConfig->GetString(TEXT("Audio"), TEXT("AudioMixerModuleName"), AudioMixerModuleName, GEngineIni);
+	}
 
 	// Check for config bool that restricts audio mixer toggle to only once. This will allow us to patch audio mixer on or off after initial login.
 	GConfig->GetBool(TEXT("Audio"), TEXT("OnlyToggleAudioMixerOnce"), bOnlyToggleAudioMixerOnce, GEngineIni);
 
-	// Get the audio mixer and non-audio mixer device module names
-	GConfig->GetString(TEXT("Audio"), TEXT("AudioDeviceModuleName"), AudioDeviceModuleName, GEngineIni);
-	GConfig->GetString(TEXT("Audio"), TEXT("AudioMixerModuleName"), AudioMixerModuleName, GEngineIni);
+	
 
 	if (bForceNonRealtimeRenderer)
 	{
@@ -355,271 +391,113 @@ bool FAudioDeviceManager::LoadDefaultAudioDeviceModule()
 	return AudioDeviceModule != nullptr;
 }
 
-bool FAudioDeviceManager::CreateMainAudioDevice()
+FAudioDeviceHandle FAudioDeviceManager::CreateNewDevice(const FAudioDeviceParams& InParams)
 {
-	FAudioDeviceManager::FCreateAudioDeviceResults NewDeviceResults;
-
-	// Create a new audio device.
-	if (CreateAudioDevice(true, NewDeviceResults))
-	{
-		MainAudioDeviceHandle = NewDeviceResults.Handle;
-		SetActiveDevice(MainAudioDeviceHandle);
-		FAudioThread::StartAudioThread();
-		return true;
-	}
-	return false;
-}
-
-bool FAudioDeviceManager::CreateAudioDevice(bool bCreateNewDevice, FCreateAudioDeviceResults& OutResults)
-{
-	OutResults = FCreateAudioDeviceResults();
-
-	// If we don't have an audio device module, then we can't create new audio devices.
-	if (AudioDeviceModule == nullptr)
-	{
-		return false;
-	}
-
-	// If we are running without the editor, we only need one audio device.
-	if (!GIsEditor)
-	{
-		if (NumActiveAudioDevices == 1)
-		{
-			FAudioDevice* MainAudioDevice = GEngine->GetMainAudioDevice();
-			if (MainAudioDevice)
-			{
-				OutResults.Handle = MainAudioDevice->DeviceHandle;
-				OutResults.AudioDevice = MainAudioDevice;
-				OutResults.AudioDevice->FadeIn();
-				return true;
-			}
-			return false;
-		}
-	}
-
-	bool bRequiresInit = true;
-
-	// For the first PIE window, we'll just use the main audio device
-	bool bCreateNewAudioDeviceForPlayInEditor = true;
-
-#if WITH_EDITOR
-	bCreateNewAudioDeviceForPlayInEditor = GetDefault<ULevelEditorMiscSettings>()->bCreateNewAudioDeviceForPlayInEditor;
-#endif
-
-	if (NumActiveAudioDevices == 1 && !bCreateNewAudioDeviceForPlayInEditor)
+	Audio::FDeviceId DeviceID = GetNewDeviceID();
+	Devices.Emplace(DeviceID, FAudioDeviceContainer(InParams, this));
+	FAudioDeviceContainer* ContainerPtr = Devices.Find(DeviceID);
+	check(ContainerPtr);
+	if (!ensure(ContainerPtr->Device))
 	{
-		FAudioDevice* MainAudioDevice = GEngine->GetMainAudioDevice();
-		if (MainAudioDevice)
-		{
-			++NumWorldsUsingMainAudioDevice;
-			OutResults.Handle = MainAudioDevice->DeviceHandle;
-			OutResults.AudioDevice = MainAudioDevice;
-			bRequiresInit = false;
-		}
-		else
-		{
-			return false;
-		}
+		// Initializing the audio device failed. Remove the device container and return an empty handle.
+		Devices.Remove(DeviceID);
+		return FAudioDeviceHandle();
 	}
 	else
 	{
-		if (NumActiveAudioDevices < AUDIO_DEVICE_DEFAULT_ALLOWED_DEVICE_COUNT || (bCreateNewDevice && NumActiveAudioDevices < AUDIO_DEVICE_MAX_DEVICE_COUNT))
-		{
-			// Create the new audio device and make sure it succeeded
-			OutResults.AudioDevice = AudioDeviceModule->CreateAudioDevice();
-
-			// Some AudioDeviceModules override CreteAudioMixerPlatformInterface, which means we can create a Audio::FMixerDevice.
-			if (OutResults.AudioDevice == nullptr)
-			{
-				checkf(AudioDeviceModule->IsAudioMixerModule(), TEXT("Please override AudioDeviceModule->CreateAudioDevice()"))
-					OutResults.AudioDevice = new Audio::FMixerDevice(AudioDeviceModule->CreateAudioMixerPlatformInterface());
-			}
-
-			if (OutResults.AudioDevice == nullptr)
-			{
-				return false;
-			}
+		FAudioDeviceHandle Handle = BuildNewHandle(*ContainerPtr, DeviceID, InParams);
+		FAudioDeviceManagerDelegates::OnAudioDeviceCreated.Broadcast(DeviceID);
+		return Handle;
+	}
+}
 
-			// Now generation a new audio device handle for the device and store the
-			// ptr to the new device in the array of audio devices.
+bool FAudioDeviceManager::IsValidAudioDevice(Audio::FDeviceId Handle) const
+{
+	return Devices.Contains(Handle);
+}
 
-			uint32 AudioDeviceIndex(INDEX_NONE);
+bool FAudioDeviceManager::ShutdownAudioDevice(Audio::FDeviceId Handle)
+{
+	// Make sure we have a non-null device ptr in the index slot, then delete it
+	Devices.Remove(Handle);
+	return true;
+}
 
-			// First check to see if we should start recycling audio device indices, if not
-			// then we add a new entry to the Generation array and generate a new index
-			if (FreeIndicesSize > AUDIO_DEVICE_MINIMUM_FREE_AUDIO_DEVICE_INDICES)
-			{
-				FreeIndices.Dequeue(AudioDeviceIndex);
-				--FreeIndicesSize;
-				check(int32(AudioDeviceIndex) < Devices.Num());
-				check(Devices[AudioDeviceIndex] == nullptr);
-				Devices[AudioDeviceIndex] = OutResults.AudioDevice;
-			}
-			else
-			{
-				// Add a zeroth generation entry in the Generation array, get a brand new
-				// index and append the created device to the end of the Devices array
+void FAudioDeviceManager::IncrementDevice(Audio::FDeviceId DeviceID)
+{
+	// If there is an FAudioDeviceHandle out in the world
+	check(Devices.Contains(DeviceID));
 
-				Generations.Add(0);
-				AudioDeviceIndex = Generations.Num() - 1;
-				check(AudioDeviceIndex < (1 << AUDIO_DEVICE_HANDLE_INDEX_BITS));
-				Devices.Add(OutResults.AudioDevice);
-			}
+	FAudioDeviceContainer& Container = Devices[DeviceID];
+	Container.NumberOfHandlesToThisDevice++;
+}
 
-			OutResults.bNewDevice = true;
-			OutResults.Handle = CreateHandle(AudioDeviceIndex, Generations[AudioDeviceIndex]);
+void FAudioDeviceManager::DecrementDevice(Audio::FDeviceId DeviceID, UWorld* InWorld)
+{
+	FScopeLock ScopeLock(&DeviceMapCriticalSection);
 
-			// Store the handle on the audio device itself
-			OutResults.AudioDevice->DeviceHandle = OutResults.Handle;
-		}
-		else
-		{
-			++NumWorldsUsingMainAudioDevice;
-			FAudioDevice* MainAudioDevice = GEngine->GetMainAudioDevice();
-			if (MainAudioDevice)
-			{
-				OutResults.Handle = MainAudioDevice->DeviceHandle;
-				OutResults.AudioDevice = MainAudioDevice;
-			}
-		}
-	}
+	// If there is an FAudioDeviceHandle out in the world
+	check(Devices.Contains(DeviceID));
 
-	++NumActiveAudioDevices;
+	FAudioDeviceContainer& Container = Devices[DeviceID];
+	check(Container.NumberOfHandlesToThisDevice > 0);
+	Container.NumberOfHandlesToThisDevice--;
 
-	if (bRequiresInit)
+	// If there is no longer anyone using this device, shut it down.
+	if (!Container.NumberOfHandlesToThisDevice)
 	{
-		// Set to highest max channels initially provided by any quality setting, so that
-		// setting to lower quality but potentially returning to higher quality later at
-		// runtime is supported.
-		const UAudioSettings* AudioSettings = GetDefault<UAudioSettings>();
-		const int32 HighestMaxChannels = AudioSettings ? AudioSettings->GetHighestMaxChannels() : 0;
-		if (OutResults.AudioDevice && OutResults.AudioDevice->Init(HighestMaxChannels))
-		{
-			const FAudioQualitySettings& QualitySettings = OutResults.AudioDevice->GetQualityLevelSettings();
-			OutResults.AudioDevice->SetMaxChannels(QualitySettings.MaxChannels);
-		}
-		else
-		{
-			ShutdownAudioDevice(OutResults.Handle);
-			OutResults = FCreateAudioDeviceResults();
-		}
+		FAudioDeviceManagerDelegates::OnAudioDeviceDestroyed.Broadcast(DeviceID);
+		Devices.Remove(DeviceID);
 	}
-
-	// We need to call fade in, in case we're reusing audio devices
-	if (OutResults.AudioDevice)
+	else if (InWorld)
 	{
-		OutResults.AudioDevice->FadeIn();
+		Container.WorldsUsingThisDevice.Remove(InWorld);
 	}
-
-	return (OutResults.AudioDevice != nullptr);
 }
 
-bool FAudioDeviceManager::IsValidAudioDeviceHandle(Audio::FDeviceId Handle) const
+bool FAudioDeviceManager::ShutdownAllAudioDevices()
 {
-	if (AudioDeviceModule == nullptr || Handle == INDEX_NONE)
-	{
-		return false;
-	}
-
-	uint32 Index = GetIndex(Handle);
-	if (int32(Index) >= Generations.Num())
-	{
-		return false;
-	}
-
-	uint8 Generation = GetGeneration(Handle);
-	return Generations[Index] == Generation;
+	Devices.Reset();
+	return true;
 }
 
-bool FAudioDeviceManager::ShutdownAudioDevice(Audio::FDeviceId Handle)
+FAudioDeviceHandle FAudioDeviceManager::BuildNewHandle(FAudioDeviceContainer&Container, Audio::FDeviceId DeviceID, const FAudioDeviceParams &InParams)
 {
-	if (!IsValidAudioDeviceHandle(Handle))
-	{
-		return false;
-	}
-
-	check(NumActiveAudioDevices > 0);
-	--NumActiveAudioDevices;
-
-	// If there are more than 1 device active, check to see if this handle is the main audio device handle
-	if (NumActiveAudioDevices >= 1)
-	{
-		uint32 MainDeviceHandle = GEngine->GetAudioDeviceHandle();
-
-		if (NumActiveAudioDevices == 1)
-		{
-			// If we only have one audio device left, then set the active
-			// audio device to be the main audio device
-			SetActiveDevice(MainDeviceHandle);
-		}
-
-		// If this is the main device handle and there's more than one reference to the main device,
-		// don't shut it down until it's the very last handle to get shut down
-		// this is because it's possible for some PIE sessions to be using the main audio device as a fallback to
-		// preserve CPU performance on low-performance machines
-		if (NumWorldsUsingMainAudioDevice > 0 && MainDeviceHandle == Handle)
-		{
-			--NumWorldsUsingMainAudioDevice;
-
-			return true;
-		}
-	}
-
-	uint32 Index = GetIndex(Handle);
-	uint8 Generation = GetGeneration(Handle);
-
-	check(int32(Index) < Generations.Num());
-
-	// Bump up the generation at the given index. This will invalidate
-	// the handle without needing to broadcast to everybody who might be using the handle
-	Generations[Index] = ++Generation;
-
-	// Make sure we have a non-null device ptr in the index slot, then delete it
-	FAudioDevice* AudioDevice = Devices[Index];
-	check(AudioDevice != nullptr);
-
-    // Tear down the audio device
-	AudioDevice->Teardown();
-
-	delete AudioDevice;
-
-	// Nullify the audio device slot for future audio device creations
-	Devices[Index] = nullptr;
-
-	// Add this index to the list of free indices
-	++FreeIndicesSize;
-	FreeIndices.Enqueue(Index);
+	Container.NumberOfHandlesToThisDevice++;
+	return FAudioDeviceHandle(Container.Device, DeviceID, InParams.AssociatedWorld);
+}
 
-	return true;
+bool FAudioDeviceManager::CanUseAudioDevice(const FAudioDeviceParams& InParams, const FAudioDeviceContainer& InContainer)
+{
+	return InContainer.Scope == EAudioDeviceScope::Shared
+		&& InParams.AudioModule == InContainer.SpecifiedModule
+		&& InParams.bIsNonRealtime == InContainer.bIsNonRealtime;
 }
 
-bool FAudioDeviceManager::ShutdownAllAudioDevices()
+FAudioDeviceHandle FAudioDeviceManager::GetAudioDevice(Audio::FDeviceId Handle)
 {
-	for (FAudioDevice* AudioDevice : Devices)
+	FAudioDeviceContainer* Container = Devices.Find(Handle);
+	if (Container)
 	{
-		if (AudioDevice)
-		{
-			ShutdownAudioDevice(AudioDevice->DeviceHandle);
-		}
+		FAudioDeviceParams Params = FAudioDeviceParams();
+		return BuildNewHandle(*Container, Handle, Params);
+	}
+	else
+	{
+		return FAudioDeviceHandle();
 	}
-
-	check(NumActiveAudioDevices == 0);
-	check(NumWorldsUsingMainAudioDevice == 0);
-
-	return true;
 }
 
-FAudioDevice* FAudioDeviceManager::GetAudioDevice(Audio::FDeviceId Handle)
+FAudioDevice* FAudioDeviceManager::GetAudioDeviceRaw(Audio::FDeviceId Handle)
 {
-	if (!IsValidAudioDeviceHandle(Handle))
+	if (!IsValidAudioDevice(Handle))
 	{
 		return nullptr;
 	}
 
-	uint32 Index = GetIndex(Handle);
-	check(int32(Index) < Devices.Num());
-	FAudioDevice* AudioDevice = Devices[Index];
+	FAudioDevice* AudioDevice = Devices[Handle].Device;
 	check(AudioDevice != nullptr);
+
 	return AudioDevice;
 }
 
@@ -633,11 +511,11 @@ FAudioDeviceManager* FAudioDeviceManager::Get()
 	return nullptr;
 }
 
-FAudioDevice* FAudioDeviceManager::GetActiveAudioDevice()
+FAudioDeviceHandle FAudioDeviceManager::GetActiveAudioDevice()
 {
-	if (ActiveAudioDeviceHandle != INDEX_NONE)
+	if (ActiveAudioDeviceID != INDEX_NONE)
 	{
-		return GetAudioDevice(ActiveAudioDeviceHandle);
+		return GetAudioDevice(ActiveAudioDeviceID);
 	}
 	return GEngine->GetMainAudioDevice();
 }
@@ -667,12 +545,9 @@ void FAudioDeviceManager::UpdateActiveAudioDevices(bool bGameTicking)
 	}
 
 
-	for (FAudioDevice* AudioDevice : Devices)
+	for (auto& DeviceContainer : Devices)
 	{
-		if (AudioDevice)
-		{
-			AudioDevice->Update(bGameTicking);
-		}
+		DeviceContainer.Value.Device->Update(bGameTicking);
 	}
 
 	if (GCVarEnableAudioThreadWait)
@@ -681,124 +556,124 @@ void FAudioDeviceManager::UpdateActiveAudioDevices(bool bGameTicking)
 	}
 }
 
+void FAudioDeviceManager::IterateOverAllDevices(TFunction<void(Audio::FDeviceId, FAudioDevice*)> ForEachDevice)
+{
+	FScopeLock ScopeLock(&DeviceMapCriticalSection);
+	for (auto& DeviceContainer : Devices)
+	{
+		ForEachDevice(DeviceContainer.Key, DeviceContainer.Value.Device);
+	}
+}
+
+void FAudioDeviceManager::IterateOverAllDevices(TFunction<void(Audio::FDeviceId, const FAudioDevice*)> ForEachDevice) const
+{
+	// We have to cheat a little to make this safe: we cast our crit section to a mutable pointer in order to scope lock.
+	FCriticalSection* ConstCastCritSection = const_cast<FCriticalSection*>(&DeviceMapCriticalSection);
+	FScopeLock ScopeLock(ConstCastCritSection);
+
+	for (const auto& DeviceContainer : Devices)
+	{
+		ForEachDevice(DeviceContainer.Key, DeviceContainer.Value.Device);
+	}
+}
+
 void FAudioDeviceManager::AddReferencedObjects(FReferenceCollector& Collector)
 {
-	for (FAudioDevice* AudioDevice : Devices)
+	FScopeLock ScopeLock(&DeviceMapCriticalSection);
+	for (auto& DeviceContainer : Devices)
 	{
-		if (AudioDevice)
-		{
-			AudioDevice->AddReferencedObjects(Collector);
-		}
+		check(DeviceContainer.Value.Device);
+		DeviceContainer.Value.Device->AddReferencedObjects(Collector);
 	}
 }
 
 void FAudioDeviceManager::StopSoundsUsingResource(USoundWave* InSoundWave, TArray<UAudioComponent*>* StoppedComponents)
 {
-	for (FAudioDevice* AudioDevice : Devices)
+	FScopeLock ScopeLock(&DeviceMapCriticalSection);
+	for (auto& DeviceContainer : Devices)
 	{
-		if (AudioDevice)
-		{
-			AudioDevice->StopSoundsUsingResource(InSoundWave, StoppedComponents);
-		}
+		DeviceContainer.Value.Device->StopSoundsUsingResource(InSoundWave, StoppedComponents);
 	}
 }
 
 void FAudioDeviceManager::RegisterSoundClass(USoundClass* SoundClass)
 {
-	for (FAudioDevice* AudioDevice : Devices)
+	FScopeLock ScopeLock(&DeviceMapCriticalSection);
+	for (auto& DeviceContainer : Devices)
 	{
-		if (AudioDevice)
-		{
-			AudioDevice->RegisterSoundClass(SoundClass);
-		}
+		DeviceContainer.Value.Device->RegisterSoundClass(SoundClass);
 	}
 }
 
 void FAudioDeviceManager::UnregisterSoundClass(USoundClass* SoundClass)
 {
-	for (FAudioDevice* AudioDevice : Devices)
+	FScopeLock ScopeLock(&DeviceMapCriticalSection);
+	for (auto& DeviceContainer : Devices)
 	{
-		if (AudioDevice)
-		{
-			AudioDevice->UnregisterSoundClass(SoundClass);
-		}
+		DeviceContainer.Value.Device->UnregisterSoundClass(SoundClass);
 	}
 }
 
 void FAudioDeviceManager::InitSoundClasses()
 {
-	for (FAudioDevice* AudioDevice : Devices)
+	FScopeLock ScopeLock(&DeviceMapCriticalSection);
+	for (auto& DeviceContainer : Devices)
 	{
-		if (AudioDevice)
-		{
-			AudioDevice->InitSoundClasses();
-		}
+		DeviceContainer.Value.Device->InitSoundClasses();
 	}
 }
 
 void FAudioDeviceManager::RegisterSoundSubmix(USoundSubmix* SoundSubmix)
 {
-	for (FAudioDevice* AudioDevice : Devices)
+	FScopeLock ScopeLock(&DeviceMapCriticalSection);
+	for (auto& DeviceContainer : Devices)
 	{
-		if (AudioDevice)
-		{
-			AudioDevice->RegisterSoundSubmix(SoundSubmix, true);
-		}
+		DeviceContainer.Value.Device->RegisterSoundSubmix(SoundSubmix, true);
 	}
 }
 
 void FAudioDeviceManager::UnregisterSoundSubmix(USoundSubmix* SoundSubmix)
 {
-	for (FAudioDevice* AudioDevice : Devices)
+	FScopeLock ScopeLock(&DeviceMapCriticalSection);
+	for (auto& DeviceContainer : Devices)
 	{
-		if (AudioDevice)
-		{
-			AudioDevice->UnregisterSoundSubmix(SoundSubmix);
-		}
+		DeviceContainer.Value.Device->UnregisterSoundSubmix(SoundSubmix);
 	}
 }
 
 void FAudioDeviceManager::InitSoundSubmixes()
 {
-	for (FAudioDevice* AudioDevice : Devices)
+	FScopeLock ScopeLock(&DeviceMapCriticalSection);
+	for (auto& DeviceContainer : Devices)
 	{
-		if (AudioDevice)
-		{
-			AudioDevice->InitSoundSubmixes();
-		}
+		DeviceContainer.Value.Device->InitSoundSubmixes();
 	}
 }
 
 void FAudioDeviceManager::InitSoundEffectPresets()
 {
-	for (FAudioDevice* AudioDevice : Devices)
+	FScopeLock ScopeLock(&DeviceMapCriticalSection);
+	for (auto& DeviceContainer : Devices)
 	{
-		if (AudioDevice)
-		{
-			AudioDevice->InitSoundEffectPresets();
-		}
+		DeviceContainer.Value.Device->InitSoundEffectPresets();
 	}
 }
 
 void FAudioDeviceManager::UpdateSourceEffectChain(const uint32 SourceEffectChainId, const TArray<FSourceEffectChainEntry>& SourceEffectChain, const bool bPlayEffectChainTails)
 {
-	for (FAudioDevice* AudioDevice : Devices)
+	FScopeLock ScopeLock(&DeviceMapCriticalSection);
+	for (auto& DeviceContainer : Devices)
 	{
-		if (AudioDevice)
-		{
-			AudioDevice->UpdateSourceEffectChain(SourceEffectChainId, SourceEffectChain, bPlayEffectChainTails);
-		}
+		DeviceContainer.Value.Device->UpdateSourceEffectChain(SourceEffectChainId, SourceEffectChain, bPlayEffectChainTails);
 	}
 }
 
 void FAudioDeviceManager::UpdateSubmix(USoundSubmix* SoundSubmix)
 {
-	for (FAudioDevice* AudioDevice : Devices)
+	FScopeLock ScopeLock(&DeviceMapCriticalSection);
+	for (auto& DeviceContainer : Devices)
 	{
-		if (AudioDevice)
-		{
-			AudioDevice->UpdateSubmixProperties(SoundSubmix);
-		}
+		DeviceContainer.Value.Device->UpdateSubmixProperties(SoundSubmix);
 	}
 }
 
@@ -807,19 +682,21 @@ void FAudioDeviceManager::SetActiveDevice(uint32 InAudioDeviceHandle)
 	// Only change the active device if there are no solo'd audio devices
 	if (SoloDeviceHandle == INDEX_NONE)
 	{
-		for (FAudioDevice* AudioDevice : Devices)
+		FScopeLock ScopeLock(&DeviceMapCriticalSection);
+		// Iterate over all of our devices and mute every device except for InAudioDeviceHandle:
+		for (auto& DeviceContainer : Devices)
 		{
-			if (AudioDevice)
+			check(DeviceContainer.Value.Device);
+			FAudioDevice* AudioDevice = DeviceContainer.Value.Device;
+
+			if (DeviceContainer.Key == InAudioDeviceHandle)
 			{
-				if (AudioDevice->DeviceHandle == InAudioDeviceHandle)
-				{
-					ActiveAudioDeviceHandle = InAudioDeviceHandle;
-					AudioDevice->SetDeviceMuted(false);
-				}
-				else
-				{
-					AudioDevice->SetDeviceMuted(true);
-				}
+				ActiveAudioDeviceID = InAudioDeviceHandle;
+				AudioDevice->SetDeviceMuted(false);
+			}
+			else
+			{
+				AudioDevice->SetDeviceMuted(true);
 			}
 		}
 	}
@@ -830,20 +707,22 @@ void FAudioDeviceManager::SetSoloDevice(Audio::FDeviceId InAudioDeviceHandle)
 	SoloDeviceHandle = InAudioDeviceHandle;
 	if (SoloDeviceHandle != INDEX_NONE)
 	{
-		for (FAudioDevice* AudioDevice : Devices)
+		FScopeLock ScopeLock(&DeviceMapCriticalSection);
+		for (auto& DeviceContainer : Devices)
 		{
-			if (AudioDevice)
+			check(DeviceContainer.Value.Device);
+			check(DeviceContainer.Key == DeviceContainer.Value.Device->DeviceID);
+			FAudioDevice*& AudioDevice = DeviceContainer.Value.Device;
+
+			// Un-mute the active audio device and mute non-active device, as long as its not the main audio device (which is used to play UI sounds)
+			if (AudioDevice->DeviceID == InAudioDeviceHandle)
 			{
-				// Un-mute the active audio device and mute non-active device, as long as its not the main audio device (which is used to play UI sounds)
-				if (AudioDevice->DeviceHandle == InAudioDeviceHandle)
-				{
-					ActiveAudioDeviceHandle = InAudioDeviceHandle;
-					AudioDevice->SetDeviceMuted(false);
-				}
-				else
-				{
-					AudioDevice->SetDeviceMuted(true);
-				}
+				ActiveAudioDeviceID = InAudioDeviceHandle;
+				AudioDevice->SetDeviceMuted(false);
+			}
+			else
+			{
+				AudioDevice->SetDeviceMuted(true);
 			}
 		}
 	}
@@ -852,38 +731,58 @@ void FAudioDeviceManager::SetSoloDevice(Audio::FDeviceId InAudioDeviceHandle)
 
 uint8 FAudioDeviceManager::GetNumActiveAudioDevices() const
 {
-	return NumActiveAudioDevices;
+	return Devices.Num();
 }
 
 uint8 FAudioDeviceManager::GetNumMainAudioDeviceWorlds() const
 {
-	return NumWorldsUsingMainAudioDevice;
+	const Audio::FDeviceId MainDeviceID = MainAudioDeviceHandle.GetDeviceID();
+	if (Devices.Contains(MainDeviceID))
+	{
+		return Devices[MainDeviceID].WorldsUsingThisDevice.Num();
+	}
+	else
+	{
+		return 0;
+	}
 }
 
-uint32 FAudioDeviceManager::GetIndex(uint32 Handle) const
+TArray<FAudioDevice*> FAudioDeviceManager::GetAudioDevices()
 {
-	return Handle & AUDIO_DEVICE_HANDLE_INDEX_MASK;
+	TArray<FAudioDevice*> DeviceList;
+	FScopeLock ScopeLock(&DeviceMapCriticalSection);
+	for (auto& Device : Devices)
+	{
+		DeviceList.Add(Device.Value.Device);
+	}
+
+	return DeviceList;
 }
 
-uint32 FAudioDeviceManager::GetGeneration(uint32 Handle) const
+TArray<UWorld*> FAudioDeviceManager::GetWorldsUsingAudioDevice(const Audio::FDeviceId& InID)
 {
-	return (Handle >> AUDIO_DEVICE_HANDLE_INDEX_BITS) & AUDIO_DEVICE_HANDLE_GENERATION_MASK;
+	if (Devices.Contains(InID))
+	{
+		return Devices[InID].WorldsUsingThisDevice;
+	}
+	else
+	{
+		return TArray<UWorld*>();
+	}
 }
 
-uint32 FAudioDeviceManager::CreateHandle(uint32 DeviceIndex, uint8 Generation)
+
+uint32 FAudioDeviceManager::GetNewDeviceID()
 {
-	return (DeviceIndex | (Generation << AUDIO_DEVICE_HANDLE_INDEX_BITS));
+	return ++DeviceIDCounter;
 }
 
 void FAudioDeviceManager::StopSourcesUsingBuffer(FSoundBuffer* SoundBuffer)
 {
-	for (FAudioDevice* AudioDevice : Devices)
+	IterateOverAllDevices([SoundBuffer](Audio::FDeviceId Id, FAudioDevice* Device)
 	{
-		if (AudioDevice)
-		{
-			AudioDevice->StopSourcesUsingBuffer(SoundBuffer);
-		}
-	}
+		Device->StopSourcesUsingBuffer(SoundBuffer);
+	});
 }
 
 void FAudioDeviceManager::TrackResource(USoundWave* SoundWave, FSoundBuffer* Buffer)
@@ -960,13 +859,10 @@ void FAudioDeviceManager::RemoveSoundMix(USoundMix* SoundMix)
 		return;
 	}
 
-	for (FAudioDevice* AudioDevice : Devices)
+	IterateOverAllDevices([SoundMix](Audio::FDeviceId Id, FAudioDevice* Device)
 	{
-		if (AudioDevice)
-		{
-			AudioDevice->RemoveSoundMix(SoundMix);
-		}
-	}
+		Device->RemoveSoundMix(SoundMix);
+	});
 }
 
 void FAudioDeviceManager::TogglePlayAllDeviceAudio()
@@ -991,7 +887,7 @@ void FAudioDeviceManager::TogglePlayAllDeviceAudio()
 bool FAudioDeviceManager::IsVisualizeDebug3dEnabled() const
 {
 #if ENABLE_AUDIO_DEBUG
-	return AudioDebugger.IsVisualizeDebug3dEnabled() || CVarIsVisualizeEnabled;
+	return GetDebugger().IsVisualizeDebug3dEnabled() || CVarIsVisualizeEnabled;
 #else // ENABLE_AUDIO_DEBUG
 	return false;
 #endif // !ENABLE_AUDIO_DEBUG
@@ -1092,6 +988,192 @@ void FAudioDeviceManager::SetDynamicSoundVolume(ESoundType SoundType, const FNam
 #if ENABLE_AUDIO_DEBUG
 FAudioDebugger& FAudioDeviceManager::GetDebugger()
 {
-	return AudioDebugger;
+	check(AudioDebugger.IsValid());
+
+	return *AudioDebugger;
+}
+
+const FAudioDebugger& FAudioDeviceManager::GetDebugger() const
+{
+	check(AudioDebugger.IsValid());
+
+	return *AudioDebugger;
 }
+
 #endif // ENABLE_AUDIO_DEBUG
+
+FAudioDeviceHandle::FAudioDeviceHandle()
+	: World(nullptr)
+	, Device(nullptr)
+	, DeviceId(INDEX_NONE)
+{
+
+}
+
+FAudioDeviceHandle::FAudioDeviceHandle(FAudioDevice* InDevice, Audio::FDeviceId InID, UWorld* InWorld)
+	: World(InWorld)
+	, Device(InDevice)
+	, DeviceId(InID)
+{
+	// Note that the reference count is incremented before this function is called.
+}
+
+FAudioDeviceHandle::FAudioDeviceHandle(const FAudioDeviceHandle& Other)
+	: FAudioDeviceHandle()
+{
+	*this = Other;
+}
+
+FAudioDeviceHandle::FAudioDeviceHandle(FAudioDeviceHandle&& Other)
+	: FAudioDeviceHandle()
+{
+	// If this chunk was previously referencing another chunk, remove that chunk here.
+	if (IsValid())
+	{
+		check(FAudioDeviceManager::Get());
+		FAudioDeviceManager::Get()->DecrementDevice(DeviceId, World);
+	}
+
+	Device = Other.Device;
+	DeviceId = Other.DeviceId;
+
+	// If the chunk we're moving from is valid, we null it out without needing to decrement.
+	Other.Device = nullptr;
+	Other.DeviceId = INDEX_NONE;
+}
+
+FAudioDeviceHandle::~FAudioDeviceHandle()
+{
+	if (GEngine && IsValid())
+	{
+		GEngine->GetAudioDeviceManager()->DecrementDevice(DeviceId, World);
+	}
+}
+
+FAudioDevice* FAudioDeviceHandle::GetAudioDevice() const
+{
+	return Device;
+}
+
+Audio::FDeviceId FAudioDeviceHandle::GetDeviceID() const
+{
+	return DeviceId;
+}
+
+bool FAudioDeviceHandle::IsValid() const
+{
+	return Device != nullptr;
+}
+
+FAudioDeviceHandle& FAudioDeviceHandle::operator=(const FAudioDeviceHandle& Other)
+{
+	// If this chunk was previously referencing another chunk, remove that chunk here.
+	if (IsValid())
+	{
+		check(FAudioDeviceManager::Get());
+		FAudioDeviceManager::Get()->DecrementDevice(DeviceId, World);
+	}
+
+	Device = Other.Device;
+	DeviceId = Other.DeviceId;
+
+	if (IsValid())
+	{
+		// Increment the reference count for the new streaming chunk.
+		FAudioDeviceManager::Get()->IncrementDevice(DeviceId);
+	}
+
+	return *this;
+}
+
+FAudioDeviceManager::FAudioDeviceContainer::FAudioDeviceContainer(const FAudioDeviceParams& InParams, FAudioDeviceManager* DeviceManager)
+	: NumberOfHandlesToThisDevice(0)
+	, Scope(InParams.Scope)
+	, bIsNonRealtime(InParams.bIsNonRealtime)
+	, SpecifiedModule(InParams.AudioModule)
+{
+	// Here we create an entirely new audio device.
+	if (bIsNonRealtime)
+	{
+		IAudioDeviceModule* NonRealtimeModule = FModuleManager::LoadModulePtr<IAudioDeviceModule>(TEXT("NonRealtimeAudioRenderer"));
+		check(NonRealtimeModule);
+		Device = NonRealtimeModule->CreateAudioDevice();
+	}
+	else if (SpecifiedModule != nullptr)
+	{
+		Device = SpecifiedModule->CreateAudioDevice();
+	}
+	else
+	{
+		check(DeviceManager->AudioDeviceModule);
+		Device = DeviceManager->AudioDeviceModule->CreateAudioDevice();
+
+		if (!Device)
+		{
+			Device = new Audio::FMixerDevice(DeviceManager->AudioDeviceModule->CreateAudioMixerPlatformInterface());
+		}
+	}
+
+	check(Device);
+
+	// Set to highest max channels initially provided by any quality setting, so that
+		// setting to lower quality but potentially returning to higher quality later at
+		// runtime is supported.
+	const UAudioSettings* AudioSettings = GetDefault<UAudioSettings>();
+	const int32 HighestMaxChannels = AudioSettings ? AudioSettings->GetHighestMaxChannels() : 0;
+	if (Device->Init(HighestMaxChannels))
+	{
+		const FAudioQualitySettings& QualitySettings = Device->GetQualityLevelSettings();
+		Device->SetMaxChannels(QualitySettings.MaxChannels);
+		Device->FadeIn();
+	}
+	else
+	{
+		UE_LOG(LogAudio, Warning, TEXT("FAudioDevice::Init Failed!"));
+		Device->Teardown();
+		delete Device;
+		Device = nullptr;
+	}
+}
+
+FAudioDeviceManager::FAudioDeviceContainer::FAudioDeviceContainer()
+{
+	checkNoEntry();
+}
+
+FAudioDeviceManager::FAudioDeviceContainer::FAudioDeviceContainer(FAudioDeviceContainer&& Other)
+{
+	Device = Other.Device;
+	Other.Device = nullptr;
+
+	NumberOfHandlesToThisDevice = Other.NumberOfHandlesToThisDevice;
+	Other.NumberOfHandlesToThisDevice = 0;
+
+	WorldsUsingThisDevice = MoveTemp(Other.WorldsUsingThisDevice);
+
+	Scope = Other.Scope;
+	Other.Scope = EAudioDeviceScope::Default;
+
+	bIsNonRealtime = Other.bIsNonRealtime;
+	Other.bIsNonRealtime = false;
+
+	SpecifiedModule = Other.SpecifiedModule;
+	Other.SpecifiedModule = nullptr;
+}
+
+FAudioDeviceManager::FAudioDeviceContainer::~FAudioDeviceContainer()
+{
+	// Shutdown the audio device.
+	check(NumberOfHandlesToThisDevice == 0);
+	if (Device)
+	{
+		Device->FadeOut();
+		Device->Teardown();
+		delete Device;
+		Device = nullptr;
+	}
+}
+
+FAudioDeviceManagerDelegates::FOnAudioDeviceCreated FAudioDeviceManagerDelegates::OnAudioDeviceCreated;
+FAudioDeviceManagerDelegates::FOnAudioDeviceDestroyed FAudioDeviceManagerDelegates::OnAudioDeviceDestroyed;
+FAudioDeviceManagerDelegates::FOnWorldRegisteredToAudioDevice FAudioDeviceManagerDelegates::OnWorldRegisteredToAudioDevice;
diff --git a/Engine/Source/Runtime/Engine/Private/AudioEffect.cpp b/Engine/Source/Runtime/Engine/Private/AudioEffect.cpp
--- a/Engine/Source/Runtime/Engine/Private/AudioEffect.cpp
+++ b/Engine/Source/Runtime/Engine/Private/AudioEffect.cpp
@@ -511,7 +511,7 @@ void FAudioEffectsManager::Update()
 	{
 		CurrentReverbAsset->bChanged = false;
 		class FAudioDeviceManager* DeviceManager = GEngine->GetAudioDeviceManager();
-		TArray<FAudioDevice*>& AudioDevices = DeviceManager->GetAudioDevices();
+		TArray<FAudioDevice*> AudioDevices = DeviceManager->GetAudioDevices();
 		for (int32 i = 0; i < AudioDevices.Num(); ++i)
 		{
 			if (AudioDevices[i])
@@ -520,7 +520,6 @@ void FAudioEffectsManager::Update()
 				EffectsManager->SetReverbSettings(CurrentReverbSettings, true);
 			}
 		}
-		
 	}
 
 #endif
diff --git a/Engine/Source/Runtime/Engine/Private/AudioVolume.cpp b/Engine/Source/Runtime/Engine/Private/AudioVolume.cpp
--- a/Engine/Source/Runtime/Engine/Private/AudioVolume.cpp
+++ b/Engine/Source/Runtime/Engine/Private/AudioVolume.cpp
@@ -91,7 +91,7 @@ void AAudioVolume::AddProxy() const
 {
 	if (UWorld* World = GetWorld())
 	{
-		if (FAudioDevice* AudioDevice = World->GetAudioDevice())
+		if (FAudioDevice* AudioDevice = World->GetAudioDeviceRaw())
 		{
 			DECLARE_CYCLE_STAT(TEXT("FAudioThreadTask.AddAudioVolumeProxy"), STAT_AudioAddAudioVolumeProxy, STATGROUP_TaskGraphTasks);
 
@@ -122,7 +122,7 @@ void AAudioVolume::RemoveProxy() const
 	UWorld* World = GetWorld();
 	if (World)
 	{
-		if (FAudioDevice* AudioDevice = World->GetAudioDevice())
+		if (FAudioDevice* AudioDevice = World->GetAudioDeviceRaw())
 		{
 			DECLARE_CYCLE_STAT(TEXT("FAudioThreadTask.RemoveAudioVolumeProxy"), STAT_AudioRemoveAudioVolumeProxy, STATGROUP_TaskGraphTasks);
 
@@ -148,7 +148,7 @@ void AAudioVolume::UpdateProxy() const
 {
 	if (UWorld* World = GetWorld())
 	{
-		if (FAudioDevice* AudioDevice = World->GetAudioDevice())
+		if (FAudioDevice* AudioDevice = World->GetAudioDeviceRaw())
 		{
 			DECLARE_CYCLE_STAT(TEXT("FAudioThreadTask.UpdateAudioVolumeProxy"), STAT_AudioUpdateAudioVolumeProxy, STATGROUP_TaskGraphTasks);
 
diff --git a/Engine/Source/Runtime/Engine/Private/Components/AudioComponent.cpp b/Engine/Source/Runtime/Engine/Private/Components/AudioComponent.cpp
--- a/Engine/Source/Runtime/Engine/Private/Components/AudioComponent.cpp
+++ b/Engine/Source/Runtime/Engine/Private/Components/AudioComponent.cpp
@@ -76,7 +76,7 @@ UAudioComponent::UAudioComponent(const FObjectInitializer& ObjectInitializer)
 	EnvelopeFollowerAttackTime = 10;
 	EnvelopeFollowerReleaseTime = 100;
 
-	AudioDeviceHandle = INDEX_NONE;
+	AudioDeviceID = INDEX_NONE;
 	AudioComponentID = FPlatformAtomics::InterlockedIncrement(reinterpret_cast<volatile int64*>(&AudioComponentIDCounter));
 
 	RandomStream.Initialize(FApp::bUseFixedSeed ? GetFName() : NAME_None);
@@ -554,18 +554,18 @@ FAudioDevice* UAudioComponent::GetAudioDevice() const
 
 	if (GEngine)
 	{
-		if (AudioDeviceHandle != INDEX_NONE)
+		if (AudioDeviceID != INDEX_NONE)
 		{
 			FAudioDeviceManager* AudioDeviceManager = GEngine->GetAudioDeviceManager();
-			AudioDevice = (AudioDeviceManager ? AudioDeviceManager->GetAudioDevice(AudioDeviceHandle) : nullptr);
+			AudioDevice = (AudioDeviceManager ? AudioDeviceManager->GetAudioDeviceRaw(AudioDeviceID) : nullptr);
 		}
 		else if (UWorld* World = GetWorld())
 		{
-			AudioDevice = World->GetAudioDevice();
+			AudioDevice = World->AudioDeviceHandle.GetAudioDevice();
 		}
 		else
 		{
-			AudioDevice = GEngine->GetMainAudioDevice();
+			AudioDevice = GEngine->GetMainAudioDeviceRaw();
 		}
 	}
 	return AudioDevice;
diff --git a/Engine/Source/Runtime/Engine/Private/ContentStreaming.cpp b/Engine/Source/Runtime/Engine/Private/ContentStreaming.cpp
--- a/Engine/Source/Runtime/Engine/Private/ContentStreaming.cpp
+++ b/Engine/Source/Runtime/Engine/Private/ContentStreaming.cpp
@@ -1180,7 +1180,7 @@ void FStreamingManagerCollection::OnAudioStreamingParamsChanged()
 {
 	// Before we swap out the audio streaming manager, we'll need to stop all sounds running on all audio devices:
 	FAudioDeviceManager* DeviceManager = GEngine->GetAudioDeviceManager();
-	TArray<FAudioDevice*>& AudioDevices = DeviceManager->GetAudioDevices();
+	TArray<FAudioDevice*> AudioDevices = DeviceManager->GetAudioDevices();
 	for (FAudioDevice* AudioDevice : AudioDevices)
 	{
 		if (AudioDevice)
diff --git a/Engine/Source/Runtime/Engine/Private/GameFramework/HapticFeedbackEffect.cpp b/Engine/Source/Runtime/Engine/Private/GameFramework/HapticFeedbackEffect.cpp
--- a/Engine/Source/Runtime/Engine/Private/GameFramework/HapticFeedbackEffect.cpp
+++ b/Engine/Source/Runtime/Engine/Private/GameFramework/HapticFeedbackEffect.cpp
@@ -168,7 +168,7 @@ float UHapticFeedbackEffect_SoundWave::GetDuration() const
 
 void UHapticFeedbackEffect_SoundWave::PrepareSoundWaveBuffer()
 {
-	FAudioDevice *AD = GEngine->GetMainAudioDevice();
+	FAudioDeviceHandle AD = GEngine->GetMainAudioDevice();
 	if (!AD || !SoundWave)
 	{
 		return;
diff --git a/Engine/Source/Runtime/Engine/Private/GameUserSettings.cpp b/Engine/Source/Runtime/Engine/Private/GameUserSettings.cpp
--- a/Engine/Source/Runtime/Engine/Private/GameUserSettings.cpp
+++ b/Engine/Source/Runtime/Engine/Private/GameUserSettings.cpp
@@ -447,7 +447,7 @@ void UGameUserSettings::ApplyNonResolutionSettings()
 		Scalability::SetQualityLevels(ScalabilityQuality);
 	}
 
-	FAudioDevice* AudioDevice = GEngine->GetMainAudioDevice();
+	FAudioDeviceHandle AudioDevice = GEngine->GetMainAudioDevice();
 	if (AudioDevice)
 	{
 		FAudioQualitySettings AudioSettings = AudioDevice->GetQualityLevelSettings();
diff --git a/Engine/Source/Runtime/Engine/Private/GameViewportClient.cpp b/Engine/Source/Runtime/Engine/Private/GameViewportClient.cpp
--- a/Engine/Source/Runtime/Engine/Private/GameViewportClient.cpp
+++ b/Engine/Source/Runtime/Engine/Private/GameViewportClient.cpp
@@ -171,7 +171,6 @@ UGameViewportClient::UGameViewportClient(const FObjectInitializer& ObjectInitial
 	, MouseCaptureMode(EMouseCaptureMode::CapturePermanently)
 	, bHideCursorDuringCapture(false)
 	, MouseLockMode(EMouseLockMode::LockOnCapture)
-	, AudioDeviceHandle(INDEX_NONE)
 	, bHasAudioFocus(false)
 	, bIsMouseOverClient(false)
 #if WITH_EDITOR
@@ -256,7 +255,6 @@ UGameViewportClient::UGameViewportClient(FVTableHelper& Helper)
 	, MouseCaptureMode(EMouseCaptureMode::CapturePermanently)
 	, bHideCursorDuringCapture(false)
 	, MouseLockMode(EMouseLockMode::LockOnCapture)
-	, AudioDeviceHandle(INDEX_NONE)
 	, bHasAudioFocus(false)
 {
 
@@ -303,14 +301,7 @@ void UGameViewportClient::PostInitProperties()
 
 void UGameViewportClient::BeginDestroy()
 {
-	if (GEngine)
-	{
-		class FAudioDeviceManager* AudioDeviceManager = GEngine->GetAudioDeviceManager();
-		if (AudioDeviceManager)
-		{
-			AudioDeviceManager->ShutdownAudioDevice(AudioDeviceHandle);
-		}
-	}
+	AudioDevice = FAudioDeviceHandle();
 
 	RemoveAllViewportWidgets();
 	Super::BeginDestroy();
@@ -441,11 +432,14 @@ void UGameViewportClient::Init(struct FWorldContext& WorldContext, UGameInstance
 		FAudioDeviceManager* AudioDeviceManager = GEngine->GetAudioDeviceManager();
 		if (AudioDeviceManager)
 		{
-			FAudioDeviceManager::FCreateAudioDeviceResults NewDeviceResults;
-			if (AudioDeviceManager->CreateAudioDevice(bCreateNewAudioDevice, NewDeviceResults))
-			{
-				AudioDeviceHandle = NewDeviceResults.Handle;
+			// Get a new audio device for this world.
+			FAudioDeviceParams DeviceParams = AudioDeviceManager->GetDefaultParamsForNewWorld();
+			DeviceParams.AssociatedWorld = World;
+
+			AudioDevice = AudioDeviceManager->RequestAudioDevice(DeviceParams);
 
+			if (ensure(AudioDevice.IsValid()))
+			{
 #if ENABLE_AUDIO_DEBUG
 				FAudioDebugger::ResolveDesiredStats(this);
 #endif // ENABLE_AUDIO_DEBUG
@@ -453,15 +447,15 @@ void UGameViewportClient::Init(struct FWorldContext& WorldContext, UGameInstance
 				// Set the base mix of the new device based on the world settings of the world
 				if (World)
 				{
-					NewDeviceResults.AudioDevice->SetDefaultBaseSoundMix(World->GetWorldSettings()->DefaultBaseSoundMix);
+					AudioDevice.GetAudioDevice()->SetDefaultBaseSoundMix(World->GetWorldSettings()->DefaultBaseSoundMix);
 
 					// Set the world's audio device handle to use so that sounds which play in that world will use the correct audio device
-					World->SetAudioDeviceHandle(AudioDeviceHandle);
+					World->SetAudioDevice(AudioDevice);
 				}
 
 				// Set this audio device handle on the world context so future world's set onto the world context
 				// will pass the audio device handle to them and audio will play on the correct audio device
-				WorldContext.AudioDeviceHandle = AudioDeviceHandle;
+				WorldContext.AudioDeviceID = AudioDevice.GetDeviceID();
 			}
 		}
 	}
@@ -1233,8 +1227,6 @@ void UGameViewportClient::Draw(FViewport* InViewport, FCanvas* SceneCanvas)
 	}
 
 	TMap<ULocalPlayer*,FSceneView*> PlayerViewMap;
-
-	FAudioDevice* AudioDevice = MyWorld->GetAudioDevice();
 	TArray<FSceneView*> Views;
 
 	for (FLocalPlayerIterator Iterator(GEngine, MyWorld); Iterator; ++Iterator)
@@ -1303,7 +1295,7 @@ void UGameViewportClient::Draw(FViewport* InViewport, FCanvas* SceneCanvas)
 						PlayerViewMap.Add(LocalPlayer, View);
 
 						// Update the listener.
-						if (AudioDevice != NULL && PlayerController != NULL)
+						if (AudioDevice && PlayerController != NULL)
 						{
 							bool bUpdateListenerPosition = true;
 
@@ -1316,8 +1308,8 @@ void UGameViewportClient::Draw(FViewport* InViewport, FCanvas* SceneCanvas)
 								// If there is more than one world referencing the main audio device
 								if (AudioDeviceManager->GetNumMainAudioDeviceWorlds() > 1)
 								{
-									uint32 MainAudioDeviceHandle = GEngine->GetAudioDeviceHandle();
-									if (AudioDevice->DeviceHandle == MainAudioDeviceHandle && !bHasAudioFocus)
+									Audio::FDeviceId MainAudioDeviceID = GEngine->GetMainAudioDeviceID();
+									if (AudioDevice->DeviceID == MainAudioDeviceID && !bHasAudioFocus)
 									{
 										bUpdateListenerPosition = false;
 									}
@@ -1824,7 +1816,7 @@ void UGameViewportClient::Precache()
 	if(!GIsEditor)
 	{
 		// Precache sounds...
-		if (FAudioDevice* AudioDevice = GetWorld()->GetAudioDevice())
+		if (AudioDevice)
 		{
 			UE_LOG(LogPlayerManagement, Log, TEXT("Precaching sounds..."));
 			for(TObjectIterator<USoundWave> It;It;++It)
@@ -1891,7 +1883,7 @@ void UGameViewportClient::ReceivedFocus(FViewport* InViewport)
 
 	if (GEngine && GEngine->GetAudioDeviceManager())
 	{
-		GEngine->GetAudioDeviceManager()->SetActiveDevice(AudioDeviceHandle);
+		GEngine->GetAudioDeviceManager()->SetActiveDevice(AudioDevice.GetDeviceID());
 		bHasAudioFocus = true;
 	}
 }
diff --git a/Engine/Source/Runtime/Engine/Private/GameplayStatics.cpp b/Engine/Source/Runtime/Engine/Private/GameplayStatics.cpp
--- a/Engine/Source/Runtime/Engine/Private/GameplayStatics.cpp
+++ b/Engine/Source/Runtime/Engine/Private/GameplayStatics.cpp
@@ -1222,7 +1222,7 @@ bool UGameplayStatics::AreAnyListenersWithinRange(const UObject* WorldContextObj
 	}
 	
 	// If there is no valid world from the world context object then there certainly are no listeners
-	if (FAudioDevice* AudioDevice = ThisWorld->GetAudioDevice())
+	if (FAudioDeviceHandle AudioDevice = ThisWorld->GetAudioDevice())
 	{
 		return AudioDevice->LocationIsAudible(Location, MaximumRange);
 	}	
@@ -1244,7 +1244,7 @@ bool UGameplayStatics::GetClosestListenerLocation(const UObject* WorldContextObj
 	}
 
 	// If there is no valid world from the world context object then there certainly are no listeners
-	FAudioDevice* AudioDevice = ThisWorld->GetAudioDevice();
+	FAudioDeviceHandle AudioDevice = ThisWorld->GetAudioDevice();
 	if (!AudioDevice)
 	{
 		return false;
@@ -1273,7 +1273,7 @@ void UGameplayStatics::SetGlobalPitchModulation(const UObject* WorldContextObjec
 		return;
 	}
 
-	if (FAudioDevice* AudioDevice = ThisWorld->GetAudioDevice())
+	if (FAudioDeviceHandle AudioDevice = ThisWorld->GetAudioDevice())
 	{
 		AudioDevice->SetGlobalPitchModulation(PitchModulation, TimeSec);
 	}
@@ -1292,7 +1292,7 @@ void UGameplayStatics::SetSoundClassDistanceScale(const UObject* WorldContextObj
 		return;
 	}
 
-	if (FAudioDevice* AudioDevice = ThisWorld->GetAudioDevice())
+	if (FAudioDeviceHandle AudioDevice = ThisWorld->GetAudioDevice())
 	{
 		AudioDevice->SetSoundClassDistanceScale(SoundClass, DistanceAttenuationScale, TimeSec);
 	}
@@ -1311,7 +1311,7 @@ void UGameplayStatics::SetGlobalListenerFocusParameters(const UObject* WorldCont
 		return;
 	}
 
-	if (FAudioDevice* AudioDevice = ThisWorld->GetAudioDevice())
+	if (FAudioDeviceHandle AudioDevice = ThisWorld->GetAudioDevice())
 	{
 		FGlobalFocusSettings NewFocusSettings;
 		NewFocusSettings.FocusAzimuthScale = FMath::Max(FocusAzimuthScale, 0.0f);
@@ -1340,7 +1340,7 @@ void UGameplayStatics::PlaySound2D(const UObject* WorldContextObject, USoundBase
 		return;
 	}
 
-	if (FAudioDevice* AudioDevice = ThisWorld->GetAudioDevice())
+	if (FAudioDeviceHandle AudioDevice = ThisWorld->GetAudioDevice())
 	{
 		FActiveSound NewActiveSound;
 		NewActiveSound.SetSound(Sound);
@@ -1382,7 +1382,7 @@ UAudioComponent* UGameplayStatics::CreateSound2D(const UObject* WorldContextObje
 	}
 
 	FAudioDevice::FCreateComponentParams Params = bPersistAcrossLevelTransition
-		? FAudioDevice::FCreateComponentParams(ThisWorld->GetAudioDevice())
+		? FAudioDevice::FCreateComponentParams(ThisWorld->GetAudioDeviceRaw())
 		: FAudioDevice::FCreateComponentParams(ThisWorld);
 
 	if (ConcurrencySettings)
@@ -1427,7 +1427,7 @@ void UGameplayStatics::PlaySoundAtLocation(const UObject* WorldContextObject, cl
 		return;
 	}
 
-	if (FAudioDevice* AudioDevice = ThisWorld->GetAudioDevice())
+	if (FAudioDeviceHandle AudioDevice = ThisWorld->GetAudioDevice())
 	{
 		AudioDevice->PlaySoundAtLocation(Sound, ThisWorld, VolumeMultiplier, PitchMultiplier, StartTime, Location, Rotation, AttenuationSettings, ConcurrencySettings, nullptr, OwningActor);
 	}
@@ -1633,7 +1633,7 @@ void UGameplayStatics::SetBaseSoundMix(const UObject* WorldContextObject, USound
 		return;
 	}
 
-	if (FAudioDevice* AudioDevice = ThisWorld->GetAudioDevice())
+	if (FAudioDeviceHandle AudioDevice = ThisWorld->GetAudioDevice())
 	{
 		AudioDevice->SetBaseSoundMix(InSoundMix);
 	}
@@ -1652,7 +1652,7 @@ void UGameplayStatics::PushSoundMixModifier(const UObject* WorldContextObject, U
 		return;
 	}
 
-	if (FAudioDevice* AudioDevice = ThisWorld->GetAudioDevice())
+	if (FAudioDeviceHandle AudioDevice = ThisWorld->GetAudioDevice())
 	{
 		AudioDevice->PushSoundMixModifier(InSoundMixModifier);
 	}
@@ -1691,7 +1691,7 @@ void UGameplayStatics::SetSoundMixClassOverride(const UObject* WorldContextObjec
 		return;
 	}
 
-	if (FAudioDevice* AudioDevice = ThisWorld->GetAudioDevice())
+	if (FAudioDeviceHandle AudioDevice = ThisWorld->GetAudioDevice())
 	{
 		AudioDevice->SetSoundMixClassOverride(InSoundMixModifier, InSoundClass, Volume, Pitch, FadeInTime, bApplyToChildren);
 	}
@@ -1710,7 +1710,7 @@ void UGameplayStatics::ClearSoundMixClassOverride(const UObject* WorldContextObj
 		return;
 	}
 
-	if (FAudioDevice* AudioDevice = ThisWorld->GetAudioDevice())
+	if (FAudioDeviceHandle AudioDevice = ThisWorld->GetAudioDevice())
 	{
 		AudioDevice->ClearSoundMixClassOverride(InSoundMixModifier, InSoundClass, FadeOutTime);
 	}
@@ -1729,7 +1729,7 @@ void UGameplayStatics::PopSoundMixModifier(const UObject* WorldContextObject, US
 		return;
 	}
 
-	if (FAudioDevice* AudioDevice = ThisWorld->GetAudioDevice())
+	if (FAudioDeviceHandle AudioDevice = ThisWorld->GetAudioDevice())
 	{
 		AudioDevice->PopSoundMixModifier(InSoundMixModifier);
 	}
@@ -1748,7 +1748,7 @@ void UGameplayStatics::ClearSoundMixModifiers(const UObject* WorldContextObject)
 		return;
 	}
 
-	if (FAudioDevice* AudioDevice = ThisWorld->GetAudioDevice())
+	if (FAudioDeviceHandle AudioDevice = ThisWorld->GetAudioDevice())
 	{
 		AudioDevice->ClearSoundMixModifiers();
 	}
@@ -1767,7 +1767,7 @@ void UGameplayStatics::ActivateReverbEffect(const UObject* WorldContextObject, c
 		return;
 	}
 
-	if (FAudioDevice* AudioDevice = ThisWorld->GetAudioDevice())
+	if (FAudioDeviceHandle AudioDevice = ThisWorld->GetAudioDevice())
 	{
 		AudioDevice->ActivateReverbEffect(ReverbEffect, TagName, Priority, Volume, FadeTime);
 	}
@@ -1786,7 +1786,7 @@ void UGameplayStatics::DeactivateReverbEffect(const UObject* WorldContextObject,
 		return;
 	}
 
-	if (FAudioDevice* AudioDevice = ThisWorld->GetAudioDevice())
+	if (FAudioDeviceHandle AudioDevice = ThisWorld->GetAudioDevice())
 	{
 		AudioDevice->DeactivateReverbEffect(TagName);
 	}
@@ -1805,7 +1805,7 @@ class UReverbEffect* UGameplayStatics::GetCurrentReverbEffect(const UObject* Wor
 		return nullptr;
 	}
 
-	if (FAudioDevice* AudioDevice = ThisWorld->GetAudioDevice())
+	if (FAudioDeviceHandle AudioDevice = ThisWorld->GetAudioDevice())
 	{
 		return AudioDevice->GetCurrentReverbEffect();
 	}
@@ -1825,7 +1825,7 @@ void UGameplayStatics::SetMaxAudioChannelsScaled(const UObject* WorldContextObje
 		return;
 	}
 
-	if (FAudioDevice* AudioDevice = ThisWorld->GetAudioDevice())
+	if (FAudioDeviceHandle AudioDevice = ThisWorld->GetAudioDevice())
 	{
 		AudioDevice->SetMaxChannelsScaled(MaxChannelCountScale);
 	}
@@ -1845,7 +1845,7 @@ int32 UGameplayStatics::GetMaxAudioChannelCount(const UObject* WorldContextObjec
 		return 0;
 	}
 
-	if (FAudioDevice* AudioDevice = ThisWorld->GetAudioDevice())
+	if (FAudioDeviceHandle AudioDevice = ThisWorld->GetAudioDevice())
 	{
 		return AudioDevice->GetMaxChannels();
 	}
diff --git a/Engine/Source/Runtime/Engine/Private/Interpolation.cpp b/Engine/Source/Runtime/Engine/Private/Interpolation.cpp
--- a/Engine/Source/Runtime/Engine/Private/Interpolation.cpp
+++ b/Engine/Source/Runtime/Engine/Private/Interpolation.cpp
@@ -1271,7 +1271,7 @@ void AMatineeActor::StepInterp( float DeltaTime, bool bPreview )
 
 void AMatineeActor::DisableRadioFilterIfNeeded()
 {
-	FAudioDevice* AudioDevice = GEngine->GetMainAudioDevice();
+	FAudioDeviceHandle AudioDevice = GEngine->GetMainAudioDevice();
 	if( AudioDevice )
 	{
 		AudioDevice->EnableRadioEffect( !bDisableRadioFilter );
@@ -1299,7 +1299,7 @@ void AMatineeActor::EnableCinematicMode(bool bEnable)
 
 void AMatineeActor::EnableRadioFilter()
 {
-	FAudioDevice* AudioDevice = GEngine->GetMainAudioDevice();
+	FAudioDeviceHandle AudioDevice = GEngine->GetMainAudioDevice();
 	if( AudioDevice )
 	{
 		AudioDevice->EnableRadioEffect( true );
diff --git a/Engine/Source/Runtime/Engine/Private/LevelActor.cpp b/Engine/Source/Runtime/Engine/Private/LevelActor.cpp
--- a/Engine/Source/Runtime/Engine/Private/LevelActor.cpp
+++ b/Engine/Source/Runtime/Engine/Private/LevelActor.cpp
@@ -1505,28 +1505,6 @@ AAudioVolume* UWorld::GetAudioSettings( const FVector& ViewLocation, FReverbSett
 	return nullptr;
 }
 
-void UWorld::SetAudioDeviceHandle(const uint32 InAudioDeviceHandle)
-{
-	AudioDeviceHandle = InAudioDeviceHandle;
-}
-
-FAudioDevice* UWorld::GetAudioDevice()
-{
-	FAudioDevice* AudioDevice = nullptr;
-	if (GEngine)
-	{
-		class FAudioDeviceManager* AudioDeviceManager = GEngine->GetAudioDeviceManager();
-		if (AudioDeviceManager != nullptr)
-		{
-			AudioDevice = AudioDeviceManager->GetAudioDevice(AudioDeviceHandle);
-			if (AudioDevice == nullptr)
-			{
-				AudioDevice = GEngine->GetMainAudioDevice();
-			}
-		}
-	}
-	return AudioDevice;
-}
 
 /**
  * Sets bMapNeedsLightingFullyRebuild to the specified value.  Marks the worldsettings package dirty if the value changed.
diff --git a/Engine/Source/Runtime/Engine/Private/PlayerCameraManager.cpp b/Engine/Source/Runtime/Engine/Private/PlayerCameraManager.cpp
--- a/Engine/Source/Runtime/Engine/Private/PlayerCameraManager.cpp
+++ b/Engine/Source/Runtime/Engine/Private/PlayerCameraManager.cpp
@@ -596,7 +596,7 @@ void APlayerCameraManager::ApplyAudioFade()
 		UWorld* World = GetWorld();
 		if (World)
 		{
-			if (FAudioDevice* AudioDevice = World->GetAudioDevice())
+			if (FAudioDevice* AudioDevice = World->GetAudioDeviceRaw())
 			{
 				AudioDevice->SetTransientMasterVolume(1.0f - FadeAmount);
 			}
@@ -611,7 +611,7 @@ void APlayerCameraManager::StopAudioFade()
 		UWorld* World = GetWorld();
 		if (World)
 		{
-			if (FAudioDevice* AudioDevice = World->GetAudioDevice())
+			if (FAudioDevice* AudioDevice = World->GetAudioDeviceRaw())
 			{
 				AudioDevice->SetTransientMasterVolume(1.0f);
 			}
diff --git a/Engine/Source/Runtime/Engine/Private/PreviewScene.cpp b/Engine/Source/Runtime/Engine/Private/PreviewScene.cpp
--- a/Engine/Source/Runtime/Engine/Private/PreviewScene.cpp
+++ b/Engine/Source/Runtime/Engine/Private/PreviewScene.cpp
@@ -98,7 +98,7 @@ FPreviewScene::~FPreviewScene()
 		UWorld* World = GetWorld();
 		if (World)
 		{
-			if (FAudioDevice* AudioDevice = World->GetAudioDevice())
+			if (FAudioDeviceHandle AudioDevice = World->GetAudioDevice())
 			{
 				AudioDevice->Flush(GetWorld(), false);
 			}
diff --git a/Engine/Source/Runtime/Engine/Private/SampleBufferIO.cpp b/Engine/Source/Runtime/Engine/Private/SampleBufferIO.cpp
--- a/Engine/Source/Runtime/Engine/Private/SampleBufferIO.cpp
+++ b/Engine/Source/Runtime/Engine/Private/SampleBufferIO.cpp
@@ -17,7 +17,7 @@ namespace Audio
 
 	void FSoundWavePCMLoader::LoadSoundWave(USoundWave* InSoundWave, TFunction<void(const USoundWave* SoundWave, const Audio::FSampleBuffer& OutSampleBuffer)> OnLoaded)
 	{
-		FAudioDevice* AudioDevice = FAudioDevice::GetMainAudioDevice();
+		FAudioDeviceHandle AudioDevice = FAudioDevice::GetMainAudioDevice();
 
 		if (!AudioDevice || !InSoundWave)
 		{
diff --git a/Engine/Source/Runtime/Engine/Private/Slate/SlateSoundDevice.cpp b/Engine/Source/Runtime/Engine/Private/Slate/SlateSoundDevice.cpp
--- a/Engine/Source/Runtime/Engine/Private/Slate/SlateSoundDevice.cpp
+++ b/Engine/Source/Runtime/Engine/Private/Slate/SlateSoundDevice.cpp
@@ -12,7 +12,7 @@ void FSlateSoundDevice::PlaySound(const FSlateSound& Sound, int32 UserIndex) con
 {
 	if( GEngine && Sound.GetResourceObject() != nullptr )
 	{
-		FAudioDevice* const AudioDevice = GEngine->GetActiveAudioDevice();
+		FAudioDeviceHandle AudioDevice = GEngine->GetActiveAudioDevice();
 		if(AudioDevice)
 		{
 			UObject* const Object = Sound.GetResourceObject();
diff --git a/Engine/Source/Runtime/Engine/Private/SoundSubmix.cpp b/Engine/Source/Runtime/Engine/Private/SoundSubmix.cpp
--- a/Engine/Source/Runtime/Engine/Private/SoundSubmix.cpp
+++ b/Engine/Source/Runtime/Engine/Private/SoundSubmix.cpp
@@ -37,7 +37,7 @@ void USoundSubmix::StartRecordingOutput(const UObject* WorldContextObject, float
 
 	// Find device for this specific audio recording thing.
 	UWorld* ThisWorld = GEngine->GetWorldFromContextObject(WorldContextObject, EGetWorldErrorMode::LogAndReturnNull);
-	FAudioDevice* DesiredAudioDevice = ThisWorld->GetAudioDevice();
+	FAudioDevice* DesiredAudioDevice = ThisWorld->GetAudioDeviceRaw();
 
 	StartRecordingOutput(DesiredAudioDevice, ExpectedDuration);
 }
@@ -59,7 +59,7 @@ void USoundSubmix::StopRecordingOutput(const UObject* WorldContextObject, EAudio
 
 	// Find device for this specific audio recording thing.
 	UWorld* ThisWorld = GEngine->GetWorldFromContextObject(WorldContextObject, EGetWorldErrorMode::LogAndReturnNull);
-	FAudioDevice* DesiredAudioDevice = ThisWorld->GetAudioDevice();
+	FAudioDevice* DesiredAudioDevice = ThisWorld->GetAudioDeviceRaw();
 
 	StopRecordingOutput(DesiredAudioDevice, ExportType, Name, Path, ExistingSoundWaveToOverwrite);
 }
@@ -139,7 +139,7 @@ void USoundSubmix::StartEnvelopeFollowing(const UObject* WorldContextObject)
 
 	// Find device for this specific audio recording thing.
 	UWorld* ThisWorld = GEngine->GetWorldFromContextObject(WorldContextObject, EGetWorldErrorMode::LogAndReturnNull);
-	FAudioDevice* AudioDevice = ThisWorld->GetAudioDevice();
+	FAudioDevice* AudioDevice = ThisWorld->GetAudioDeviceRaw();
 
 	StartEnvelopeFollowing(AudioDevice);
 }
@@ -161,7 +161,7 @@ void USoundSubmix::StopEnvelopeFollowing(const UObject* WorldContextObject)
 
 	// Find device for this specific audio recording thing.
 	UWorld* ThisWorld = GEngine->GetWorldFromContextObject(WorldContextObject, EGetWorldErrorMode::LogAndReturnNull);
-	FAudioDevice* AudioDevice = ThisWorld->GetAudioDevice();
+	FAudioDevice* AudioDevice = ThisWorld->GetAudioDeviceRaw();
 
 	StopEnvelopeFollowing(AudioDevice);
 }
@@ -183,7 +183,7 @@ void USoundSubmix::AddEnvelopeFollowerDelegate(const UObject* WorldContextObject
 
 	// Find device for this specific audio recording thing.
 	UWorld* ThisWorld = GEngine->GetWorldFromContextObject(WorldContextObject, EGetWorldErrorMode::LogAndReturnNull);
-	FAudioDevice* AudioDevice = ThisWorld->GetAudioDevice();
+	FAudioDevice* AudioDevice = ThisWorld->GetAudioDeviceRaw();
 	if (AudioDevice)
 	{
 		AudioDevice->AddEnvelopeFollowerDelegate(this, OnSubmixEnvelopeBP);
@@ -198,7 +198,7 @@ void USoundSubmix::SetSubmixOutputVolume(const UObject* WorldContextObject, floa
 	}
 
 	UWorld* ThisWorld = GEngine->GetWorldFromContextObject(WorldContextObject, EGetWorldErrorMode::LogAndReturnNull);
-	FAudioDevice* AudioDevice = ThisWorld->GetAudioDevice();
+	FAudioDevice* AudioDevice = ThisWorld->GetAudioDeviceRaw();
 	if (AudioDevice)
 	{
 		AudioDevice->SetSubmixOutputVolume(this, InOutputVolume);
diff --git a/Engine/Source/Runtime/Engine/Private/SoundWave.cpp b/Engine/Source/Runtime/Engine/Private/SoundWave.cpp
--- a/Engine/Source/Runtime/Engine/Private/SoundWave.cpp
+++ b/Engine/Source/Runtime/Engine/Private/SoundWave.cpp
@@ -219,7 +219,7 @@ void USoundWave::GetResourceSizeEx(FResourceSizeEx& CumulativeResourceSize)
 		return;
 	}
 
-	if (FAudioDevice* LocalAudioDevice = GEngine->GetMainAudioDevice())
+	if (FAudioDevice* LocalAudioDevice = GEngine->GetMainAudioDeviceRaw())
 	{
 		if (LocalAudioDevice->HasCompressedAudioInfoClass(this) && DecompressionType == DTYPE_Native)
 		{
@@ -835,7 +835,7 @@ void USoundWave::PostLoad()
 	// most likely cause us to run out of memory.
 	if (!GIsEditor && !IsTemplate( RF_ClassDefaultObject ) && GEngine)
 	{
-		FAudioDevice* AudioDevice = GEngine->GetMainAudioDevice();
+		FAudioDeviceHandle AudioDevice = GEngine->GetMainAudioDevice();
 		if (AudioDevice)
 		{
 			// Upload the data to the hardware, but only if we've precached startup sounds already
diff --git a/Engine/Source/Runtime/Engine/Private/SoundWaveProcedural.cpp b/Engine/Source/Runtime/Engine/Private/SoundWaveProcedural.cpp
--- a/Engine/Source/Runtime/Engine/Private/SoundWaveProcedural.cpp
+++ b/Engine/Source/Runtime/Engine/Private/SoundWaveProcedural.cpp
@@ -20,7 +20,7 @@ USoundWaveProcedural::USoundWaveProcedural(const FObjectInitializer& ObjectIniti
 	
 	if (GEngine)
 	{
-		FAudioDevice* MainAudioDevice = GEngine->GetMainAudioDevice();
+		FAudioDevice* MainAudioDevice = GEngine->GetMainAudioDeviceRaw();
 		if (MainAudioDevice && !MainAudioDevice->IsAudioMixerEnabled())
 		{
 #if PLATFORM_MAC
diff --git a/Engine/Source/Runtime/Engine/Private/UnrealEngine.cpp b/Engine/Source/Runtime/Engine/Private/UnrealEngine.cpp
--- a/Engine/Source/Runtime/Engine/Private/UnrealEngine.cpp
+++ b/Engine/Source/Runtime/Engine/Private/UnrealEngine.cpp
@@ -1082,11 +1082,13 @@ FWorldContext
 
 void FWorldContext::SetCurrentWorld(UWorld* World)
 {
-	if (World != nullptr)
+	if (World != nullptr && AudioDeviceID != INDEX_NONE)
 	{
 		// Set the world's audio device handle so that audio components playing in the 
 		// world will use the correct audio device instance.
-		World->SetAudioDeviceHandle(AudioDeviceHandle);
+		FAudioDeviceHandle AudioDevice = FAudioDeviceManager::Get()->GetAudioDevice(AudioDeviceID);
+
+		World->SetAudioDevice(AudioDevice);
 	}
 
 	for (int32 idx = 0; idx < ExternalReferences.Num(); ++idx)
@@ -1754,8 +1756,6 @@ void UEngine::ShutdownAudioDeviceManager()
 
 		FAudioThread::StopAudioThread();
 
-		AudioDeviceManager->ShutdownAllAudioDevices();
-
 		delete AudioDeviceManager;
 		AudioDeviceManager = NULL;
 	}
@@ -2745,7 +2745,7 @@ void UEngine::Serialize(FArchive& Ar)
 	if (Ar.IsCountingMemory())
 	{
 		// Only use the main audio device when counting memory
-		if (FAudioDevice* AudioDevice = GetMainAudioDevice())
+		if (FAudioDeviceHandle AudioDevice = GetMainAudioDevice())
 		{
 			AudioDevice->CountBytes(Ar);
 		}
@@ -2852,27 +2852,37 @@ class FAudioDeviceManager* UEngine::GetAudioDeviceManager()
 	return AudioDeviceManager;
 }
 
-uint32 UEngine::GetAudioDeviceHandle() const
+uint32 UEngine::GetMainAudioDeviceID() const
+{
+	return MainAudioDeviceHandle.GetDeviceID();
+}
+
+FAudioDeviceHandle UEngine::GetMainAudioDevice()
 {
-	return MainAudioDeviceHandle;
+	if (AudioDeviceManager != nullptr)
+	{
+		return AudioDeviceManager->GetMainAudioDeviceHandle();
+	}
+	return FAudioDeviceHandle();
 }
 
-FAudioDevice* UEngine::GetMainAudioDevice()
+FAudioDevice* UEngine::GetMainAudioDeviceRaw()
 {
 	if (AudioDeviceManager != nullptr)
 	{
-		return AudioDeviceManager->GetAudioDevice(MainAudioDeviceHandle);
+		return AudioDeviceManager->GetMainAudioDeviceHandle().GetAudioDevice();
 	}
+
 	return nullptr;
 }
 
-FAudioDevice* UEngine::GetActiveAudioDevice()
+FAudioDeviceHandle UEngine::GetActiveAudioDevice()
 {
 	if (AudioDeviceManager != nullptr)
 	{
 		return AudioDeviceManager->GetActiveAudioDevice();
 	}
-	return nullptr;
+	return FAudioDeviceHandle();
 }
 
 void UEngine::InitializeAudioDeviceManager()
@@ -3836,11 +3846,11 @@ bool UEngine::Exec( UWorld* InWorld, const TCHAR* Cmd, FOutputDevice& Ar )
 	FAudioDevice* AudioDevice = nullptr;
 	if (InWorld)
 	{
-		AudioDevice = InWorld->GetAudioDevice();
+		AudioDevice = InWorld->GetAudioDeviceRaw();
 	}
 	else
 	{
-		AudioDevice = GetMainAudioDevice();
+		AudioDevice = GetMainAudioDeviceRaw();
 	}
 
 	if (AudioDevice && AudioDevice->Exec(InWorld, Cmd, Ar) == true)
@@ -12458,7 +12468,7 @@ bool UEngine::LoadMap( FWorldContext& WorldContext, FURL URL, class UPendingNetG
 		}
 
 		// Stop all audio to remove references to current level.
-		if (FAudioDevice* AudioDevice = WorldContext.World()->GetAudioDevice())
+		if (FAudioDevice* AudioDevice = WorldContext.World()->GetAudioDeviceRaw())
 		{
 			AudioDevice->Flush(WorldContext.World());
 			AudioDevice->SetTransientMasterVolume(1.0f);
@@ -12698,7 +12708,7 @@ bool UEngine::LoadMap( FWorldContext& WorldContext, FURL URL, class UPendingNetG
 
 	WorldContext.World()->SetGameMode(URL);
 
-	if (FAudioDevice* AudioDevice = WorldContext.World()->GetAudioDevice())
+	if (FAudioDevice* AudioDevice = WorldContext.World()->GetAudioDeviceRaw())
 	{
 		AudioDevice->SetDefaultBaseSoundMix(WorldContext.World()->GetWorldSettings()->DefaultBaseSoundMix);
 	}
@@ -13823,7 +13833,7 @@ bool UEngine::CommitMapChange( FWorldContext &Context )
 
 		// make sure any looping sounds, etc are stopped
 
-		if (FAudioDevice* AudioDevice = Context.World()->GetAudioDevice())
+		if (FAudioDevice* AudioDevice = Context.World()->GetAudioDeviceRaw())
 		{
 			AudioDevice->StopAllSounds();
 		}
diff --git a/Engine/Source/Runtime/Engine/Private/World.cpp b/Engine/Source/Runtime/Engine/Private/World.cpp
--- a/Engine/Source/Runtime/Engine/Private/World.cpp
+++ b/Engine/Source/Runtime/Engine/Private/World.cpp
@@ -345,7 +345,7 @@ UWorld::UWorld( const FObjectInitializer& ObjectInitializer )
 , bIsBuilt(false)
 , bShouldTick(true)
 , ActiveLevelCollectionIndex(INDEX_NONE)
-, AudioDeviceHandle(INDEX_NONE)
+, AudioDeviceHandle()
 #if WITH_EDITOR
 , HierarchicalLODBuilder(new FHierarchicalLODBuilder(this))
 #endif
@@ -6181,7 +6181,7 @@ UWorld* FSeamlessTravelHandler::Tick()
 			CurrentWorld->ClearFlags(RF_Standalone);
 
 			// Stop all audio to remove references to old world
-			if (FAudioDevice* AudioDevice = CurrentWorld->GetAudioDevice())
+			if (FAudioDevice* AudioDevice = CurrentWorld->GetAudioDeviceRaw())
 			{
 				AudioDevice->Flush(CurrentWorld);
 			}
@@ -6283,7 +6283,7 @@ UWorld* FSeamlessTravelHandler::Tick()
 			// if we've already switched to entry before and this is the transition to the new map, re-create the gameinfo
 			if (bSwitchedToDefaultMap && !bIsClient)
 			{
-				if (FAudioDevice* AudioDevice = LoadedWorld->GetAudioDevice())
+				if (FAudioDevice* AudioDevice = LoadedWorld->GetAudioDeviceRaw())
 				{
 					AudioDevice->SetDefaultBaseSoundMix(LoadedWorld->GetWorldSettings()->DefaultBaseSoundMix);
 				}
@@ -7474,6 +7474,20 @@ void UWorld::AddPostProcessingSettings(FVector ViewLocation, FSceneView* SceneVi
 	}
 }
 
+void UWorld::SetAudioDevice(FAudioDeviceHandle& InHandle)
+{
+	AudioDeviceHandle = InHandle;
+}
+
+FAudioDeviceHandle UWorld::GetAudioDevice()
+{
+	return AudioDeviceHandle;
+}
+
+class FAudioDevice* UWorld::GetAudioDeviceRaw()
+{
+	return AudioDeviceHandle.GetAudioDevice();
+}
 
 /**
 * Dump visible actors in current world.
diff --git a/Engine/Source/Runtime/Engine/Private/WorldSettings.cpp b/Engine/Source/Runtime/Engine/Private/WorldSettings.cpp
--- a/Engine/Source/Runtime/Engine/Private/WorldSettings.cpp
+++ b/Engine/Source/Runtime/Engine/Private/WorldSettings.cpp
@@ -194,7 +194,7 @@ void AWorldSettings::PostRegisterAllComponents()
 	Super::PostRegisterAllComponents();
 
 	UWorld* World = GetWorld();
-	if (FAudioDevice* AudioDevice = World->GetAudioDevice())
+	if (FAudioDeviceHandle AudioDevice = World->GetAudioDevice())
 	{
 		AudioDevice->SetDefaultAudioSettings(World, DefaultReverbSettings, DefaultAmbientZoneSettings);
 	}
@@ -622,7 +622,7 @@ void AWorldSettings::InternalPostPropertyChanged(FName PropertyName)
 	if (PropertyName == GET_MEMBER_NAME_CHECKED(AWorldSettings, DefaultReverbSettings) || PropertyName == GET_MEMBER_NAME_CHECKED(AWorldSettings, DefaultAmbientZoneSettings))
 {
 		UWorld* World = GetWorld();
-		if (FAudioDevice* AudioDevice = World->GetAudioDevice())
+		if (FAudioDeviceHandle AudioDevice = World->GetAudioDevice())
 	{
 			AudioDevice->SetDefaultAudioSettings(World, DefaultReverbSettings, DefaultAmbientZoneSettings);
 		}
diff --git a/Engine/Source/Runtime/Engine/Public/Audio.h b/Engine/Source/Runtime/Engine/Public/Audio.h
--- a/Engine/Source/Runtime/Engine/Public/Audio.h
+++ b/Engine/Source/Runtime/Engine/Public/Audio.h
@@ -92,14 +92,6 @@ struct FWaveInstance;
 struct FReverbSettings;
 struct FSampleLoop;
 
-namespace Audio
-{
-	/**
-	 * Typed identifier for Audio Device Id
-	 */
-	using FDeviceId = uint32;
-}
-
 enum ELoopingMode
 {
 	/** One shot sound */
diff --git a/Engine/Source/Runtime/Engine/Public/AudioDevice.h b/Engine/Source/Runtime/Engine/Public/AudioDevice.h
--- a/Engine/Source/Runtime/Engine/Public/AudioDevice.h
+++ b/Engine/Source/Runtime/Engine/Public/AudioDevice.h
@@ -1132,10 +1132,10 @@ public:
 	}
 
 	/** Returns the main audio device of the engine */
-	static FAudioDevice* GetMainAudioDevice()
+	static FAudioDeviceHandle GetMainAudioDevice()
 	{
 		// Try to get GEngine's main audio device
-		FAudioDevice* AudioDevice = GEngine->GetMainAudioDevice();
+		FAudioDeviceHandle AudioDevice = GEngine->GetMainAudioDevice();
 
 		// If we don't have a main audio device (maybe we're running in a non-standard mode like a commandlet)
 		if (!AudioDevice)
@@ -1194,7 +1194,7 @@ public:
 
 	static bool IsOcclusionPluginLoaded()
 	{
-		if (FAudioDevice* MainAudioDevice = GEngine->GetMainAudioDevice())
+		if (FAudioDeviceHandle MainAudioDevice = GEngine->GetMainAudioDevice())
 		{
 			return MainAudioDevice->bOcclusionInterfaceEnabled;
 		}
@@ -1209,7 +1209,7 @@ public:
 
 	static bool IsReverbPluginLoaded()
 	{
-		if (FAudioDevice* MainAudioDevice = GEngine->GetMainAudioDevice())
+		if (FAudioDeviceHandle MainAudioDevice = GEngine->GetMainAudioDevice())
 		{
 			return MainAudioDevice->bReverbInterfaceEnabled;
 		}
@@ -1678,8 +1678,8 @@ public:
 
 	bool IsMainAudioDevice()
 	{
-		FAudioDevice* MainAudioDevice = GEngine->GetMainAudioDevice();
-		return (MainAudioDevice == nullptr || MainAudioDevice == this);
+		FAudioDeviceHandle MainAudioDevice = GEngine->GetMainAudioDevice();
+		return (!MainAudioDevice || MainAudioDevice.GetAudioDevice() == this);
 	}
 
 	/** Set whether or not we force the use of attenuation for non-game worlds (as by default we only care about game worlds) */
@@ -1710,7 +1710,7 @@ public:
 	int32 NumPrecacheFrames;
 
 	/** The handle for this audio device used in the audio device manager. */
-	Audio::FDeviceId DeviceHandle;
+	Audio::FDeviceId DeviceID;
 
 	/** 3rd party audio spatialization interface. */
 	TAudioSpatializationPtr SpatializationPluginInterface;
diff --git a/Engine/Source/Runtime/Engine/Public/AudioDeviceManager.h b/Engine/Source/Runtime/Engine/Public/AudioDeviceManager.h
--- a/Engine/Source/Runtime/Engine/Public/AudioDeviceManager.h
+++ b/Engine/Source/Runtime/Engine/Public/AudioDeviceManager.h
@@ -2,9 +2,7 @@
 #pragma once
 
 #include "AudioDefines.h"
-#include "Audio/AudioDebug.h"
 #include "AudioThread.h"
-#include "Containers/Queue.h"
 #include "CoreMinimal.h"
 
 
@@ -20,9 +18,19 @@ class USoundClass;
 class USoundMix;
 class USoundSubmix;
 class USoundWave;
-
+class UWorld;
+class FAudioDevice;
+class FAudioDebugger;
 struct FSourceEffectChainEntry;
 
+namespace Audio
+{
+	/**
+	 * Typed identifier for Audio Device Id
+	 */
+	using FDeviceId = uint32;
+}
+
 enum class ESoundType : uint8
 {
 	Class,
@@ -30,6 +38,149 @@ enum class ESoundType : uint8
 	Wave
 };
 
+// This enum is used in FAudioDeviceManager::RequestAudioDevice to map a given UWorld to an audio device.
+enum class EAudioDeviceScope : uint8
+{
+	// Default to the behavior specified by the editor preferences.
+	Default,
+	// Use an audio device that can be shared by multiple worlds.
+	Shared,
+	// Create a new audio device specifically for this handle.
+	Unique
+};
+
+// Parameters passed into FAudioDeviceManager::RequestAudioDevice.
+struct FAudioDeviceParams
+{
+	// Optional world parameter. This allows tools to surface information about which worlds are being rendered through which audio devices.
+	UWorld* AssociatedWorld;
+	// This should be set to EAudioDeviceScope::Unique if you'd like to force a new device to be created from scratch, or use EAudioDeviceScope::Shared to use an existing device if possible.
+	EAudioDeviceScope Scope;
+	// Set this to true to get a handle to a non realtime audio renderer.
+	bool bIsNonRealtime;
+	// Use this to force this audio device to use a specific audio module. If nullptr, uses the default audio module.
+	IAudioDeviceModule* AudioModule;
+
+	FAudioDeviceParams()
+		: AssociatedWorld(nullptr)
+		, Scope(EAudioDeviceScope::Default)
+		, bIsNonRealtime(false)
+		, AudioModule(nullptr)
+	{}
+};
+
+// Strong handle to an audio device. Guarantees that the audio device it references will stay alive while it is in scope.
+class ENGINE_API FAudioDeviceHandle
+{
+public:
+	FAudioDeviceHandle();
+	FAudioDeviceHandle(const FAudioDeviceHandle& Other);
+	FAudioDeviceHandle(FAudioDeviceHandle&& Other);
+
+	FAudioDeviceHandle& operator=(const FAudioDeviceHandle& Other);
+
+	~FAudioDeviceHandle();
+
+	// gets a pointer to the compressed chunk.
+	FAudioDevice* GetAudioDevice() const;
+
+	// Returns the device ID for the audio device referenced by this handle.
+	Audio::FDeviceId GetDeviceID() const;
+
+	// Checks whether this points to a valid compressed chunk.
+	bool IsValid() const;
+
+private:
+	// This constructor should only be called by FAudioDeviceManager.
+	FAudioDeviceHandle(FAudioDevice* InDevice, Audio::FDeviceId InID, UWorld* InWorld);
+
+	// The world that this FAudioDeviceHandle was requested with.
+	// Null if this device handle wasn't generated by RequestAudioDevice.
+	UWorld* World;
+
+	FAudioDevice* Device;
+	Audio::FDeviceId DeviceId;
+
+	friend class FAudioDeviceManager;
+
+public:
+	// These are convenience operators to use an FAudioDeviceHandle like an FAudioDevice* or an Audio::FDeviceId.
+	// For safety, we still require explicit casting to an FAudioDevice* to ensure ownership isn't lost due to programmer error.
+	const FAudioDevice& operator*() const
+	{
+		return *Device;
+	}
+
+	FAudioDevice& operator*()
+	{
+		return *Device;
+	}
+
+	const FAudioDevice* operator->() const
+	{
+		return Device;
+	}
+
+	FAudioDevice* operator->()
+	{
+		return Device;
+	}
+
+	FORCEINLINE explicit operator bool() const
+	{
+		return IsValid();
+	}
+
+	inline bool operator==(const FAudioDeviceHandle& Other) const
+	{
+		return DeviceId == Other.DeviceId;
+	}
+
+	inline bool operator==(const FAudioDevice*& Other) const
+	{
+		return Device == Other;
+	}
+
+	inline bool operator==(FAudioDevice*& Other) const
+	{
+		return Device == Other;
+	}
+
+	inline bool operator==(const Audio::FDeviceId& Other) const
+	{
+		return DeviceId == Other;
+	}
+
+	explicit operator FAudioDevice*() const
+	{
+		return Device;
+	}
+
+	explicit operator Audio::FDeviceId() const
+	{
+		return DeviceId;
+	}
+};
+
+// List of delegates for the audio device manager.
+class ENGINE_API FAudioDeviceManagerDelegates
+{
+public:
+	// This delegate is called whenever an entirely new audio device is created.
+	DECLARE_MULTICAST_DELEGATE_OneParam(FOnAudioDeviceCreated, Audio::FDeviceId /* AudioDeviceId*/);
+	static FOnAudioDeviceCreated OnAudioDeviceCreated;
+
+	// This delegate is called whenever an audio device is destroyed.
+	DECLARE_MULTICAST_DELEGATE_OneParam(FOnAudioDeviceDestroyed, Audio::FDeviceId /* AudioDeviceId*/);
+	static FOnAudioDeviceDestroyed OnAudioDeviceDestroyed;
+
+	// This delegate is called whenever a world is registered to an audio device. Please note that UWorlds are not guaranteed to
+	// be registered to the same audio device throughout their lifecycle,
+	// and there is no guarantee on the lifespan of both the UWorld and the Audio Device registered in this callback.
+	DECLARE_MULTICAST_DELEGATE_TwoParams(FOnWorldRegisteredToAudioDevice, const UWorld* /*InWorld */, Audio::FDeviceId /* AudioDeviceId*/);
+	static FOnWorldRegisteredToAudioDevice OnWorldRegisteredToAudioDevice;
+};
+
 /**
 * Class for managing multiple audio devices.
 */
@@ -54,54 +205,44 @@ public:
 	bool Initialize();
 
 	/** Returns the handle to the main audio device. */
-	Audio::FDeviceId GetMainAudioDeviceHandle() const { return MainAudioDeviceHandle; }
+	FAudioDeviceHandle GetMainAudioDeviceHandle() const { return MainAudioDeviceHandle; }
+	Audio::FDeviceId GetMainAudioDeviceID() const { return MainAudioDeviceHandle.GetDeviceID(); }
 
 	/** Returns true if we're currently using the audio mixer. */
 	bool IsUsingAudioMixer() const;
 
-	struct FCreateAudioDeviceResults
-	{
-		uint32 Handle;
-		uint8  bNewDevice : 1;
-		FAudioDevice* AudioDevice;
-
-		FCreateAudioDeviceResults();
-	};
-
 	/**
 	 * returns the currently used audio device module for this platform.
 	 * returns nullptr if Initialize() has not been called yet.
 	 */
 	IAudioDeviceModule* GetAudioDeviceModule();
 
+	FAudioDeviceParams GetDefaultParamsForNewWorld();
+
 	/**
-	* Creates and audio device instance internally and returns a
-	* handle to the audio device. Returns true on success.
+	* Creates or requests an audio device instance internally and returns a
+	* handle to the audio device.
+	* This audio device is guaranteed to be alive as long as the returned handle is in scope.
 	*/
-	bool CreateAudioDevice(bool bCreateNewDevice, FCreateAudioDeviceResults& OutResults);
+	FAudioDeviceHandle RequestAudioDevice(const FAudioDeviceParams& InParams);
 
 	/**
 	* Returns whether the audio device handle is valid (i.e. points to
 	* an actual audio device instance)
 	*/
-	bool IsValidAudioDeviceHandle(Audio::FDeviceId Handle) const;
-
-	/**
-	* Shutsdown the audio device associated with the handle. The handle
-	* will become invalid after the audio device is shut down.
-	*/
-	bool ShutdownAudioDevice(Audio::FDeviceId Handle);
+	bool IsValidAudioDevice(Audio::FDeviceId DeviceID) const;
 
 	/**
-	* Shuts down all active audio devices
-	*/
-	bool ShutdownAllAudioDevices();
+	 * Returns a strong handle to the audio device associated with the given device ID.
+	 * if the device ID is invalid we return an invalid, zeroed handle.
+	 */
+	FAudioDeviceHandle GetAudioDevice(Audio::FDeviceId DeviceID);
 
 	/**
-	* Returns a ptr to the audio device associated with the handle. If the
+	* Returns a raw ptr to the audio device associated with the handle. If the
 	* handle is invalid then a NULL device ptr will be returned.
 	*/
-	FAudioDevice* GetAudioDevice(Audio::FDeviceId Handle);
+	FAudioDevice* GetAudioDeviceRaw(Audio::FDeviceId DeviceID);
 
 	static FAudioDeviceManager* Get();
 
@@ -109,7 +250,7 @@ public:
 	* Returns a ptr to the active audio device. If there is no active
 	* device then it will return the main audio device.
 	*/
-	FAudioDevice* GetActiveAudioDevice();
+	FAudioDeviceHandle GetActiveAudioDevice();
 
 	/** Returns the current number of active audio devices. */
 	uint8 GetNumActiveAudioDevices() const;
@@ -120,6 +261,9 @@ public:
 	/** Updates all active audio devices */
 	void UpdateActiveAudioDevices(bool bGameTicking);
 
+	void IterateOverAllDevices(TFunction<void(Audio::FDeviceId, FAudioDevice*)> ForEachDevice);
+	void IterateOverAllDevices(TFunction<void(Audio::FDeviceId, const FAudioDevice*)> ForEachDevice) const;
+
 	/** Tracks objects in the active audio devices. */
 	void AddReferencedObjects(FReferenceCollector& Collector);
 
@@ -203,6 +347,7 @@ public:
 #if ENABLE_AUDIO_DEBUG
 	/** Get the audio debugger instance */
 	FAudioDebugger& GetDebugger();
+	const FAudioDebugger& GetDebugger() const;
 #endif // ENABLE_AUDIO_DEBUG
 
 public:
@@ -214,23 +359,20 @@ public:
 	TMap<int32, FSoundBuffer*>	WaveBufferMap;
 
 	/** Returns all the audio devices managed by device manager. */
-	TArray<FAudioDevice*>& GetAudioDevices() { return Devices; }
-	
+	TArray<FAudioDevice*> GetAudioDevices();
+
+	TArray<UWorld*> GetWorldsUsingAudioDevice(const Audio::FDeviceId& InID);
+
 private:
 
 #if ENABLE_AUDIO_DEBUG
 	/** Instance of audio debugger shared across audio devices */
-	FAudioDebugger AudioDebugger;
+	TUniquePtr<FAudioDebugger> AudioDebugger;
 #endif // ENABLE_AUDIO_DEBUG
 
-	/** Returns index of the given handle */
-	uint32 GetIndex(uint32 Handle) const;
-
-	/** Returns the generation of the given handle */
-	uint32 GetGeneration(uint32 Handle) const;
 
 	/** Creates a handle given the index and a generation value. */
-	uint32 CreateHandle(uint32 DeviceIndex, uint8 Generation);
+	uint32 GetNewDeviceID();
 
 	/** Toggles between audio mixer and non-audio mixer audio engine. */
 	void ToggleAudioMixer();
@@ -238,11 +380,24 @@ private:
 	/** Load audio device module. */
 	bool LoadDefaultAudioDeviceModule();
 
-	/** Create a "main" audio device. */
-	bool CreateMainAudioDevice();
+	FAudioDeviceHandle CreateNewDevice(const FAudioDeviceParams& InParams);
+
+	/**
+	* Shutsdown the audio device associated with the handle. The handle
+	* will become invalid after the audio device is shut down.
+	*/
+	bool ShutdownAudioDevice(Audio::FDeviceId DeviceID);
+
+	// Called exclusively by the FAudioDeviceHandle copy constructor and assignment operators:
+	void IncrementDevice(Audio::FDeviceId DeviceID);
+
+	// Called exclusively by the FAudioDeviceHandle dtor.
+	void DecrementDevice(Audio::FDeviceId DeviceID, UWorld* InWorld);
 
-	/** Array for generation counts of audio devices in indices */
-	TArray<uint8> Generations;
+	/**
+	* Shuts down all active audio devices
+	*/
+	bool ShutdownAllAudioDevices();
 
 	/** Audio device module which creates (old backend) audio devices. */
 	IAudioDeviceModule* AudioDeviceModule;
@@ -254,25 +409,60 @@ private:
 	FString AudioMixerModuleName;
 
 	/** Handle to the main audio device. */
-	Audio::FDeviceId MainAudioDeviceHandle;
+	FAudioDeviceHandle MainAudioDeviceHandle;
+
+	struct FAudioDeviceContainer
+	{
+		// Singularly owned device.
+		// Could be a TUniquePtr if FAudioDevice was not an incomplete type here.
+		FAudioDevice* Device;
 
-	/** Count of the number of free slots in the audio device array. */
-	uint32 FreeIndicesSize;
+		// Ref count of FAudioDeviceHandles referencing this device.
+		int32 NumberOfHandlesToThisDevice;
 
-	/** Number of actively created audio device instances. */
-	uint8 NumActiveAudioDevices;
+		/** Worlds that have been registered to this device. */
+		TArray<UWorld*> WorldsUsingThisDevice;
 
-	/** Number of worlds using the main audio device instance. */
-	uint8 NumWorldsUsingMainAudioDevice;
+		/** Whether this device can be shared. */
+		EAudioDeviceScope Scope;
 
-	/** Queue for free indices */
-	TQueue<uint32> FreeIndices;
+		/** Whether this audio device is realtime or not. */
+		bool bIsNonRealtime;
+
+		/** Module this was created with. If nullptr, this device was created with the default module. */
+		IAudioDeviceModule* SpecifiedModule;
+
+		FAudioDeviceContainer(const FAudioDeviceParams& InParams, FAudioDeviceManager* DeviceManager);
+		~FAudioDeviceContainer();
+
+		FAudioDeviceContainer(const FAudioDeviceContainer& Other)
+		{
+			// We explicitly enforce that we invoke the move instructor.
+			// If this was hit, you likely need to call either Devices.Emplace(args) or Devices.Add(DeviceID, MoveTemp(Container));
+			checkNoEntry();
+		}
+
+		FAudioDeviceContainer(FAudioDeviceContainer&& Other);
+
+	private:
+		FAudioDeviceContainer();
+	};
+
+	FAudioDeviceHandle BuildNewHandle(FAudioDeviceContainer&Container, Audio::FDeviceId DeviceID, const FAudioDeviceParams &InParams);
+
+	/**
+	 * This function is used to check if we can use an existing audio device.
+	 */
+	static bool CanUseAudioDevice(const FAudioDeviceParams& InParams, const FAudioDeviceContainer& InContainer);
 
 	/**
-	* Array of audio device pointers. If the device has been free, then
-	* the device ptr at the array index will be null.
+	* Bank of audio devices. Will increase in size as we create new audio devices,
 	*/
-	TArray<FAudioDevice*> Devices;
+	TMap<Audio::FDeviceId, FAudioDeviceContainer> Devices;
+	FCriticalSection DeviceMapCriticalSection;
+
+	/* Counter used by GetNewDeviceID() to generate a unique ID for a given audio device. */
+	uint32 DeviceIDCounter;
 
 	/** Next resource ID to assign out to a wave/buffer */
 	int32 NextResourceID;
@@ -281,7 +471,7 @@ private:
 	Audio::FDeviceId SoloDeviceHandle;
 
 	/** Which audio device is currently active */
-	Audio::FDeviceId ActiveAudioDeviceHandle;
+	Audio::FDeviceId ActiveAudioDeviceID;
 
 	/** Dynamic volume map */
 	TMap<TTuple<ESoundType, FName>, float> DynamicSoundVolumes;
@@ -300,4 +490,6 @@ private:
 
 	/** Audio Fence to ensure that we don't allow the audio thread to drift never endingly behind. */
 	FAudioCommandFence SyncFence;
+
+	friend class FAudioDeviceHandle;
 };
diff --git a/Engine/Source/Runtime/GameplayMediaEncoder/Private/GameplayMediaEncoder.cpp b/Engine/Source/Runtime/GameplayMediaEncoder/Private/GameplayMediaEncoder.cpp
--- a/Engine/Source/Runtime/GameplayMediaEncoder/Private/GameplayMediaEncoder.cpp
+++ b/Engine/Source/Runtime/GameplayMediaEncoder/Private/GameplayMediaEncoder.cpp
@@ -274,7 +274,7 @@ bool FGameplayMediaEncoder::Start()
 	// subscribe to engine delegates for audio output and back buffer
 	//
 
-	FAudioDevice* AudioDevice = GEngine->GetMainAudioDevice();
+	FAudioDeviceHandle AudioDevice = GEngine->GetMainAudioDevice();
 	if (AudioDevice)
 	{
 		bAudioFormatChecked = false;
@@ -298,7 +298,7 @@ void FGameplayMediaEncoder::Stop()
 
 	if (UGameEngine* GameEngine = Cast<UGameEngine>(GEngine))
 	{
-		FAudioDevice* AudioDevice = GameEngine->GetMainAudioDevice();
+		FAudioDevice* AudioDevice = GameEngine->GetMainAudioDeviceRaw();
 		if (AudioDevice)
 		{
 			AudioDevice->UnregisterSubmixBufferListener(this);
diff --git a/Engine/Source/Runtime/Launch/Private/Android/AndroidEventManager.cpp b/Engine/Source/Runtime/Launch/Private/Android/AndroidEventManager.cpp
--- a/Engine/Source/Runtime/Launch/Private/Android/AndroidEventManager.cpp
+++ b/Engine/Source/Runtime/Launch/Private/Android/AndroidEventManager.cpp
@@ -319,7 +319,7 @@ void FAppEventManager::PauseAudio()
 	bAudioPaused = true;
 	UE_LOG(LogTemp, Log, TEXT("Android pause audio"));
 
-	FAudioDevice* AudioDevice = GEngine->GetMainAudioDevice();
+	FAudioDeviceHandle AudioDevice = GEngine->GetMainAudioDevice();
 	if (AudioDevice)
 	{
 		if (AudioDevice->IsAudioMixerEnabled())
@@ -354,7 +354,7 @@ void FAppEventManager::ResumeAudio()
 	bAudioPaused = false;
 	UE_LOG(LogTemp, Log, TEXT("Android resume audio"));
 
-	FAudioDevice* AudioDevice = GEngine->GetMainAudioDevice();
+	FAudioDeviceHandle AudioDevice = GEngine->GetMainAudioDevice();
 	if (AudioDevice)
 	{
 		if (AudioDevice->IsAudioMixerEnabled())
diff --git a/Engine/Source/Runtime/Launch/Private/IOS/LaunchIOS.cpp b/Engine/Source/Runtime/Launch/Private/IOS/LaunchIOS.cpp
--- a/Engine/Source/Runtime/Launch/Private/IOS/LaunchIOS.cpp
+++ b/Engine/Source/Runtime/Launch/Private/IOS/LaunchIOS.cpp
@@ -62,7 +62,7 @@ void FAppEntry::Suspend(bool bIsInterrupt)
 	{
 		if (GEngine && GEngine->GetMainAudioDevice() && !IsEngineExitRequested())
 		{
-			FAudioDevice* AudioDevice = GEngine->GetMainAudioDevice();
+			FAudioDeviceHandle AudioDevice = GEngine->GetMainAudioDevice();
 			if (bIsInterrupt && DisableAudioSuspendOnAudioInterruptCvar)
 			{
 				if (FTaskGraphInterface::IsRunning() && !IsEngineExitRequested())
@@ -156,7 +156,7 @@ void FAppEntry::Resume(bool bIsInterrupt)
 	{
 		if (GEngine && GEngine->GetMainAudioDevice())
 		{
-			FAudioDevice* AudioDevice = GEngine->GetMainAudioDevice();
+			FAudioDeviceHandle AudioDevice = GEngine->GetMainAudioDevice();
         
 			if (bIsInterrupt && DisableAudioSuspendOnAudioInterruptCvar)
 			{
@@ -205,7 +205,7 @@ void FAppEntry::ResumeAudioContext()
 {
 	if (GEngine && GEngine->GetMainAudioDevice())
 	{
-		FAudioDevice* AudioDevice = GEngine->GetMainAudioDevice();
+		FAudioDeviceHandle AudioDevice = GEngine->GetMainAudioDevice();
 		if (AudioDevice)
 		{
 			if (FTaskGraphInterface::IsRunning())
@@ -233,7 +233,7 @@ void FAppEntry::RestartAudio()
 {
 	if (GEngine && GEngine->GetMainAudioDevice())
 	{
-		FAudioDevice* AudioDevice = GEngine->GetMainAudioDevice();
+		FAudioDeviceHandle AudioDevice = GEngine->GetMainAudioDevice();
 
 		if (FTaskGraphInterface::IsRunning())
 		{
diff --git a/Engine/Source/Runtime/MovieScene/Private/MovieSceneTimeController.cpp b/Engine/Source/Runtime/MovieScene/Private/MovieSceneTimeController.cpp
--- a/Engine/Source/Runtime/MovieScene/Private/MovieSceneTimeController.cpp
+++ b/Engine/Source/Runtime/MovieScene/Private/MovieSceneTimeController.cpp
@@ -104,7 +104,7 @@ double FMovieSceneTimeController_PlatformClock::GetCurrentTime() const
 
 double FMovieSceneTimeController_AudioClock::GetCurrentTime() const
 {
-	FAudioDevice* AudioDevice = GEngine ? GEngine->GetMainAudioDevice() : nullptr;
+	FAudioDevice* AudioDevice = GEngine ? GEngine->GetMainAudioDeviceRaw() : nullptr;
 	return AudioDevice ? AudioDevice->GetAudioClock() : FPlatformTime::Seconds();
 }
 
diff --git a/Engine/Source/Runtime/MovieSceneCapture/Public/Protocols/AudioCaptureProtocol.h b/Engine/Source/Runtime/MovieSceneCapture/Public/Protocols/AudioCaptureProtocol.h
--- a/Engine/Source/Runtime/MovieSceneCapture/Public/Protocols/AudioCaptureProtocol.h
+++ b/Engine/Source/Runtime/MovieSceneCapture/Public/Protocols/AudioCaptureProtocol.h
@@ -84,7 +84,7 @@ public:
 
 		// Stop all sounds currently playing on the Audio Device. This helps kill looping or long audio clips. When the sequence evaluates again,
 		// these clips will resume play at the correct location.
-		if (FAudioDevice* AudioDevice = GEngine->GetActiveAudioDevice())
+		if (FAudioDeviceHandle AudioDevice = GEngine->GetActiveAudioDevice())
 		{
 			AudioDevice->StopAllSounds(true);
 		}
