Stop compression codecs being garbage collected during async compression

diff --git a/Engine/Source/Runtime/Engine/Private/Animation/AnimCompressionDerivedData.cpp b/Engine/Source/Runtime/Engine/Private/Animation/AnimCompressionDerivedData.cpp
--- a/Engine/Source/Runtime/Engine/Private/Animation/AnimCompressionDerivedData.cpp
+++ b/Engine/Source/Runtime/Engine/Private/Animation/AnimCompressionDerivedData.cpp
@@ -35,7 +35,7 @@ const TCHAR* FDerivedDataAnimationCompression::GetVersionString() const
 	// This is a version string that mimics the old versioning scheme. If you
 	// want to bump this version, generate a new guid using VS->Tools->Create GUID and
 	// return it here. Ex.
-	return TEXT("A20FA6C92A4443CDAE897F0B63C32814");
+	return TEXT("0F1CECE507FE4F89A374B4C8E7B55052");
 }
 
 bool FDerivedDataAnimationCompression::Build( TArray<uint8>& OutDataArray )
@@ -168,6 +168,19 @@ void FAsyncCompressedAnimationsManagement::Tick(float DeltaTime)
 	}
 }
 
+void FAsyncCompressedAnimationsManagement::AddReferencedObjects(FReferenceCollector& Collector)
+{
+	for (const FActiveAsyncCompressionTask& Task : ActiveAsyncCompressionTasks)
+	{
+		Task.DataToCompress->AddReferencedObjects(Collector);
+	}
+
+	for (const FQueuedAsyncCompressionWork QueuedTask : QueuedAsyncCompressionWork)
+	{
+		QueuedTask.Compressor.GetCompressibleData()->AddReferencedObjects(Collector);
+	}
+}
+
 TStatId FAsyncCompressedAnimationsManagement::GetStatId() const
 {
 	RETURN_QUICK_DECLARE_CYCLE_STAT(FAsyncCompressedAnimationsTracker, STATGROUP_Tickables);
@@ -213,8 +226,9 @@ bool FAsyncCompressedAnimationsManagement::RequestAsyncCompression(FDerivedDataA
 void FAsyncCompressedAnimationsManagement::StartAsyncWork(FDerivedDataAnimationCompression& Compressor, UAnimSequence* Anim, const uint64 NewTaskSize, const bool bPerformFrameStripping)
 {
 	const FString CacheKey = Compressor.GetPluginSpecificCacheKeySuffix();
+	FCompressibleAnimPtr SourceData = Compressor.GetCompressibleData();
 	uint32 AsyncHandle = GetDerivedDataCacheRef().GetAsynchronous(&Compressor);
-	ActiveAsyncCompressionTasks.Emplace(Anim, CacheKey, NewTaskSize, AsyncHandle, bPerformFrameStripping);
+	ActiveAsyncCompressionTasks.Emplace(Anim, SourceData, CacheKey, NewTaskSize, AsyncHandle, bPerformFrameStripping);
 }
 
 bool FAsyncCompressedAnimationsManagement::WaitOnActiveCompression(UAnimSequence* Anim)
diff --git a/Engine/Source/Runtime/Engine/Private/Animation/AnimCompressionDerivedData.h b/Engine/Source/Runtime/Engine/Private/Animation/AnimCompressionDerivedData.h
--- a/Engine/Source/Runtime/Engine/Private/Animation/AnimCompressionDerivedData.h
+++ b/Engine/Source/Runtime/Engine/Private/Animation/AnimCompressionDerivedData.h
@@ -20,7 +20,7 @@ class FDerivedDataAnimationCompression : public FDerivedDataPluginInterface
 {
 private:
 	// The anim data to compress
-	TSharedPtr<FCompressibleAnimData> DataToCompressPtr;
+	FCompressibleAnimPtr DataToCompressPtr;
 
 	// The Type of anim data to compress (makes up part of DDC key)
 	const TCHAR* TypeName;
@@ -35,12 +35,14 @@ public:
 	FDerivedDataAnimationCompression(const TCHAR* InTypeName, const FString& InAssetDDCKey, TSharedPtr<FAnimCompressContext> InCompressContext);
 	virtual ~FDerivedDataAnimationCompression();
 
-	void SetCompressibleData(TSharedRef<FCompressibleAnimData> InCompressibleAnimData)
+	void SetCompressibleData(FCompressibleAnimRef InCompressibleAnimData)
 	{
 		DataToCompressPtr = InCompressibleAnimData;
 		check(DataToCompressPtr->Skeleton != nullptr);
 	}
 
+	FCompressibleAnimPtr GetCompressibleData() const { return DataToCompressPtr; }
+
 	uint64 GetMemoryUsage() const
 	{
 		return DataToCompressPtr.IsValid() ? DataToCompressPtr->GetApproxRawSize() : 0;
diff --git a/Engine/Source/Runtime/Engine/Private/Animation/AnimCompressionTypes.cpp b/Engine/Source/Runtime/Engine/Private/Animation/AnimCompressionTypes.cpp
--- a/Engine/Source/Runtime/Engine/Private/Animation/AnimCompressionTypes.cpp
+++ b/Engine/Source/Runtime/Engine/Private/Animation/AnimCompressionTypes.cpp
@@ -495,6 +495,8 @@ void FCompressedAnimSequence::SerializeCompressedData(FArchive& Ar, bool bDDCDat
 		Ar << BoneCodecDDCHandle;
 		Ar << CurveCodecPath;
 
+		check(!BoneCodecDDCHandle.Equals(TEXT("None"), ESearchCase::IgnoreCase)); // Failed DDC data?
+
 		int32 NumCurveBytes;
 		Ar << NumCurveBytes;
 
@@ -615,6 +617,7 @@ void FCompressedAnimSequence::SerializeCompressedData(FArchive& Ar, bool bDDCDat
 		}
 
 		FString BoneCodecDDCHandle = BoneCompressionCodec != nullptr ? BoneCompressionCodec->GetCodecDDCHandle() : TEXT("");
+		check(!BoneCodecDDCHandle.Equals(TEXT("None"), ESearchCase::IgnoreCase)); // Will write broken DDC data to DDC!
 		Ar << BoneCodecDDCHandle;
 
 		FString CurveCodecPath = CurveCompressionCodec->GetPathName();
diff --git a/Engine/Source/Runtime/Engine/Private/Animation/AnimSequence.cpp b/Engine/Source/Runtime/Engine/Private/Animation/AnimSequence.cpp
--- a/Engine/Source/Runtime/Engine/Private/Animation/AnimSequence.cpp
+++ b/Engine/Source/Runtime/Engine/Private/Animation/AnimSequence.cpp
@@ -2231,7 +2231,7 @@ void UAnimSequence::RequestAnimCompression(FRequestAnimCompressionParams Params)
 			// No trivial key removal is done at this point (impossible error metrics of -1), since all of the techniques will perform it themselves
 			CompressRawAnimData(-1.0f, -1.0f);
 
-			TSharedRef<FCompressibleAnimData> CompressibleData = MakeShared<FCompressibleAnimData>(this, bPerformStripping);
+			FCompressibleAnimRef CompressibleData = MakeShared<FCompressibleAnimData, ESPMode::ThreadSafe>(this, bPerformStripping);
 			AnimCompressor->SetCompressibleData(CompressibleData);
 
 			if (bSkipDDC || (CompressCommandletVersion == INDEX_NONE))
@@ -2256,8 +2256,12 @@ void UAnimSequence::RequestAnimCompression(FRequestAnimCompressionParams Params)
 			bCompressedDataFromDDC = true;
 		}
 
-		delete AnimCompressor; // Would really like to do auto mem management but GetDerivedDataCacheRef().GetSynchronous expects a pointer it can delete
-		AnimCompressor = nullptr;
+		if (AnimCompressor)
+		{
+			// Would really like to do auto mem management but GetDerivedDataCacheRef().GetSynchronous expects a pointer it can delete
+			delete AnimCompressor;
+			AnimCompressor = nullptr;
+		}
 	}
 
 	if (OutData.Num() > 0) // Haven't async compressed
diff --git a/Engine/Source/Runtime/Engine/Private/Animation/AnimStreamable.cpp b/Engine/Source/Runtime/Engine/Private/Animation/AnimStreamable.cpp
--- a/Engine/Source/Runtime/Engine/Private/Animation/AnimStreamable.cpp
+++ b/Engine/Source/Runtime/Engine/Private/Animation/AnimStreamable.cpp
@@ -612,7 +612,7 @@ void UAnimStreamable::RequestCompressedDataForChunk(const FString& ChunkDDCKey,
 
 		if (bSkipDDC || !GetDerivedDataCacheRef().GetSynchronous(*FinalDDCKey, OutData))
 		{
-			TSharedRef<FCompressibleAnimData> CompressibleData = MakeShared<FCompressibleAnimData>(BoneCompressionSettings, CurveCompressionSettings, GetSkeleton(), Interpolation, Chunk.SequenceLength, ChunkNumFrames+1);
+			FCompressibleAnimRef CompressibleData = MakeShared<FCompressibleAnimData, ESPMode::ThreadSafe>(BoneCompressionSettings, CurveCompressionSettings, GetSkeleton(), Interpolation, Chunk.SequenceLength, ChunkNumFrames+1);
 
 			CompressibleData->RawAnimationData.AddDefaulted(RawAnimationData.Num());
 
diff --git a/Engine/Source/Runtime/Engine/Public/Animation/AnimCompressionDerivedDataPublic.h b/Engine/Source/Runtime/Engine/Public/Animation/AnimCompressionDerivedDataPublic.h
--- a/Engine/Source/Runtime/Engine/Public/Animation/AnimCompressionDerivedDataPublic.h
+++ b/Engine/Source/Runtime/Engine/Public/Animation/AnimCompressionDerivedDataPublic.h
@@ -8,20 +8,23 @@
 
 class UAnimSequence;
 class FDerivedDataAnimationCompression;
+struct FCompressibleAnimData;
 
 extern ENGINE_API class FAsyncCompressedAnimationsManagement* GAsyncCompressedAnimationsTracker;
 
 // Animation data that is currently being compressed
 struct FActiveAsyncCompressionTask
 {
-	FActiveAsyncCompressionTask(UAnimSequence* InSequence, const FString& InCacheKey, const uint64 InTaskSize, const uint32 InAsyncHandle, bool bInPerformFrameStripping)
+	FActiveAsyncCompressionTask(UAnimSequence* InSequence, FCompressibleAnimPtr InDataToCompress, const FString& InCacheKey, const uint64 InTaskSize, const uint32 InAsyncHandle, bool bInPerformFrameStripping)
 		: Sequence(InSequence)
+		, DataToCompress(InDataToCompress)
 		, TaskSize(InTaskSize)
 		, CacheKey(InCacheKey)
 		, AsyncHandle(InAsyncHandle)
 		, bPerformFrameStripping(bInPerformFrameStripping)
 	{}
 	UAnimSequence* Sequence;
+	FCompressibleAnimPtr DataToCompress;
 	uint64 TaskSize;
 	FString CacheKey;
 	uint32 AsyncHandle;
@@ -47,7 +50,7 @@ struct FQueuedAsyncCompressionWork
 //   Maintains active compressions
 //   tracks memory usage of async compression
 //   Gives API for blocking on compression
-class ENGINE_API FAsyncCompressedAnimationsManagement : public FTickableEditorObject
+class ENGINE_API FAsyncCompressedAnimationsManagement : public FTickableEditorObject, public FGCObject
 {
 public:
 	static FAsyncCompressedAnimationsManagement& Get();
@@ -81,6 +84,10 @@ private:
 	virtual ETickableTickType GetTickableTickType() const override { return ETickableTickType::Always; }
 	virtual TStatId GetStatId() const override;
 	/* End FTickableEditorObject */
+
+	// FGCObject interface
+	void AddReferencedObjects(FReferenceCollector& Collector) override;
+	// End of FGCObject interface
 };
 
 #endif
\ No newline at end of file
diff --git a/Engine/Source/Runtime/Engine/Public/Animation/AnimCompressionTypes.h b/Engine/Source/Runtime/Engine/Public/Animation/AnimCompressionTypes.h
--- a/Engine/Source/Runtime/Engine/Public/Animation/AnimCompressionTypes.h
+++ b/Engine/Source/Runtime/Engine/Public/Animation/AnimCompressionTypes.h
@@ -236,10 +236,19 @@ public:
 
 	void Update(struct FCompressedAnimSequence& CompressedData) const;
 
+	void AddReferencedObjects(FReferenceCollector& Collector)
+	{
+		Collector.AddReferencedObject(BoneCompressionSettings);
+		Collector.AddReferencedObject(CurveCompressionSettings);
+	}
+
 private:
 
 };
 
+typedef TSharedPtr<FCompressibleAnimData, ESPMode::ThreadSafe> FCompressibleAnimPtr;
+typedef TSharedRef<FCompressibleAnimData, ESPMode::ThreadSafe> FCompressibleAnimRef;
+
 // Wrapper Code
 template <typename T>
 struct TArrayMaker
