Fix crash in anim decompression when running with ISPC

diff --git a/Engine/Source/Runtime/Engine/Private/Animation/AnimEncoding_PerTrackCompression.cpp b/Engine/Source/Runtime/Engine/Private/Animation/AnimEncoding_PerTrackCompression.cpp
--- a/Engine/Source/Runtime/Engine/Private/Animation/AnimEncoding_PerTrackCompression.cpp
+++ b/Engine/Source/Runtime/Engine/Private/Animation/AnimEncoding_PerTrackCompression.cpp
@@ -771,7 +771,10 @@ void AEFPerTrackCompressionCodec::GetPoseTranslations(
 	FAnimSequenceDecompressionContext& DecompContext)
 {
 	const int32 PairCount = DesiredPairs.Num();
-
+	if (PairCount == 0)
+	{
+		return;
+	}
 	if (INTEL_ISPC)
 	{
 #if INTEL_ISPC
