#ue4 - Fix perch result setting bogus distance to floor, which caused player to alternate between rising for a while then falling because AdjustFloorHeight() would constantly think the LineTraceDist was 0 and needed to move up. Capsule adjustments should be based on the capsule sweep distance to hit, not the line trace when perching. Now ComputeFloorDist keeps the distance to the sweep hit, and perching uses that distance to collision to make the floor height adjustment. This could occur when walking over steep unwalkable geo where the perch radius would allow continued walking.

diff --git a/Engine/Source/Runtime/Engine/Private/Components/CharacterMovementComponent.cpp b/Engine/Source/Runtime/Engine/Private/Components/CharacterMovementComponent.cpp
--- a/Engine/Source/Runtime/Engine/Private/Components/CharacterMovementComponent.cpp
+++ b/Engine/Source/Runtime/Engine/Private/Components/CharacterMovementComponent.cpp
@@ -6152,7 +6152,6 @@ void UCharacterMovementComponent::ComputeFloorDist(const FVector& CapsuleLocatio
 	
 	// No hits were acceptable.
 	OutFloorResult.bWalkableFloor = false;
-	OutFloorResult.FloorDist = SweepDistance;
 }
 
 
@@ -6244,7 +6243,8 @@ void UCharacterMovementComponent::FindFloor(const FVector& CapsuleLocation, FFin
 				// If the regular capsule is on an unwalkable surface but the perched one would allow us to stand, override the normal to be one that is walkable.
 				if (!OutFloorResult.bWalkableFloor)
 				{
-					OutFloorResult.SetFromLineTrace(PerchFloorResult.HitResult, OutFloorResult.FloorDist, FMath::Min(PerchFloorResult.FloorDist, PerchFloorResult.LineDist), true);
+					// Floor distances are used as the distance of the regular capsule to the point of collision, to make sure AdjustFloorHeight() behaves correctly.
+					OutFloorResult.SetFromLineTrace(PerchFloorResult.HitResult, OutFloorResult.FloorDist, FMath::Max(OutFloorResult.FloorDist, MIN_FLOOR_DIST), true);
 				}
 			}
 			else
@@ -6439,7 +6439,7 @@ bool UCharacterMovementComponent::ComputePerchResult(const float TestRadius, con
 {
 	if (InMaxFloorDist <= 0.f)
 	{
-		return 0.f;
+		return false;
 	}
 
 	// Sweep further than actual requested distance, because a reduced capsule radius means we could miss some hits that the normal radius would contact.
