Only render RVT volume wireframe when the actor is selected

diff --git a/Engine/Source/Runtime/Engine/Private/VT/RuntimeVirtualTextureVolume.cpp b/Engine/Source/Runtime/Engine/Private/VT/RuntimeVirtualTextureVolume.cpp
--- a/Engine/Source/Runtime/Engine/Private/VT/RuntimeVirtualTextureVolume.cpp
+++ b/Engine/Source/Runtime/Engine/Private/VT/RuntimeVirtualTextureVolume.cpp
@@ -14,6 +14,7 @@ ARuntimeVirtualTextureVolume::ARuntimeVirtualTextureVolume(const FObjectInitiali
 	// Add box for visualization of bounds
 	Box = CreateDefaultSubobject<UBoxComponent>(TEXT("Bounds"));
 	Box->SetBoxExtent(FVector(0.5f, 0.5f, 1.f), false);
+	Box->bDrawOnlyIfSelected = true;
 	Box->SetIsVisualizationComponent(true);
 	Box->SetCollisionEnabled(ECollisionEnabled::NoCollision);
 	Box->SetCanEverAffectNavigation(false);
