Fixed CIS caused by a bad merge

#rnx

diff --git a/Engine/Source/Runtime/Engine/Classes/Materials/MaterialInterface.h b/Engine/Source/Runtime/Engine/Classes/Materials/MaterialInterface.h
--- a/Engine/Source/Runtime/Engine/Classes/Materials/MaterialInterface.h
+++ b/Engine/Source/Runtime/Engine/Classes/Materials/MaterialInterface.h
@@ -17,6 +17,7 @@
 #include "Materials/MaterialLayersFunctions.h"
 #include "Interfaces/Interface_AssetUserData.h"
 #include "MaterialSceneTextureId.h"
+#include "Materials/MaterialRelevance.h"
 #if WITH_CHAOS
 #include "Physics/PhysicsInterfaceCore.h"
 #endif
