Added WeldSimulatedBodies option when using bAutoManageAttachment for Cascade and Niagara

#rnx

diff --git a/Engine/Plugins/FX/Niagara/Source/Niagara/Private/NiagaraComponent.cpp b/Engine/Plugins/FX/Niagara/Source/Niagara/Private/NiagaraComponent.cpp
--- a/Engine/Plugins/FX/Niagara/Source/Niagara/Private/NiagaraComponent.cpp
+++ b/Engine/Plugins/FX/Niagara/Source/Niagara/Private/NiagaraComponent.cpp
@@ -885,7 +885,7 @@ void UNiagaraComponent::ActivateInternal(bool bReset /* = false */, bool bIsScal
 				SavedAutoAttachRelativeRotation = GetRelativeRotation();
 				SavedAutoAttachRelativeScale3D = GetRelativeScale3D();
 				//bIsChangingAutoAttachment = true;
-				AttachToComponent(NewParent, FAttachmentTransformRules(AutoAttachLocationRule, AutoAttachRotationRule, AutoAttachScaleRule, false), AutoAttachSocketName);
+				AttachToComponent(NewParent, FAttachmentTransformRules(AutoAttachLocationRule, AutoAttachRotationRule, AutoAttachScaleRule, bAutoAttachWeldSimulatedBodies), AutoAttachSocketName);
 				//bIsChangingAutoAttachment = false;
 			}
 
diff --git a/Engine/Plugins/FX/Niagara/Source/Niagara/Public/NiagaraComponent.h b/Engine/Plugins/FX/Niagara/Source/Niagara/Public/NiagaraComponent.h
--- a/Engine/Plugins/FX/Niagara/Source/Niagara/Public/NiagaraComponent.h
+++ b/Engine/Plugins/FX/Niagara/Source/Niagara/Public/NiagaraComponent.h
@@ -133,6 +133,13 @@ public:
 	UPROPERTY(EditAnywhere, BlueprintReadOnly, Category = Attachment)
 	uint32 bAutoManageAttachment : 1;
 
+	/**
+	 * Option for how we handle bWeldSimulatedBodies when we attach to the AutoAttachParent, if bAutoManageAttachment is true.
+	 * @see bAutoManageAttachment
+	 */
+	UPROPERTY(EditAnywhere, BlueprintReadWrite, Category=Attachment, meta=(EditCondition="bAutoManageAttachment"))
+	uint32 bAutoAttachWeldSimulatedBodies : 1;
+
 	/**
 	 * Time between forced UpdateTransforms for systems that use dynamically calculated bounds,
 	 * Which is effectively how often the bounds are shrunk.
diff --git a/Engine/Source/Runtime/Engine/Classes/Particles/ParticleSystemComponent.h b/Engine/Source/Runtime/Engine/Classes/Particles/ParticleSystemComponent.h
--- a/Engine/Source/Runtime/Engine/Classes/Particles/ParticleSystemComponent.h
+++ b/Engine/Source/Runtime/Engine/Classes/Particles/ParticleSystemComponent.h
@@ -487,6 +487,13 @@ public:
 	 */
 	UPROPERTY(EditAnywhere, BlueprintReadOnly, Category=Attachment)
 	uint8 bAutoManageAttachment:1;
+
+	/**
+	 * Option for how we handle bWeldSimulatedBodies when we attach to the AutoAttachParent, if bAutoManageAttachment is true.
+	 * @see bAutoManageAttachment
+	 */
+	UPROPERTY(EditAnywhere, BlueprintReadWrite, Category=Attachment, meta=(EditCondition="bAutoManageAttachment"))
+	uint8 bAutoAttachWeldSimulatedBodies:1;
 	
 	/** If this component is having it's significance managed by gameplay code. */
 	uint8 bIsManagingSignificance : 1;
diff --git a/Engine/Source/Runtime/Engine/Private/Particles/ParticleComponents.cpp b/Engine/Source/Runtime/Engine/Private/Particles/ParticleComponents.cpp
--- a/Engine/Source/Runtime/Engine/Private/Particles/ParticleComponents.cpp
+++ b/Engine/Source/Runtime/Engine/Private/Particles/ParticleComponents.cpp
@@ -5956,7 +5956,7 @@ void UParticleSystemComponent::ActivateSystem(bool bFlagAsJustAttached)
 					SavedAutoAttachRelativeLocation = GetRelativeLocation();
 					SavedAutoAttachRelativeRotation = GetRelativeRotation();
 					SavedAutoAttachRelativeScale3D = GetRelativeScale3D();
-					AttachToComponent(NewParent, FAttachmentTransformRules(AutoAttachLocationRule, AutoAttachRotationRule, AutoAttachScaleRule, false), AutoAttachSocketName);
+					AttachToComponent(NewParent, FAttachmentTransformRules(AutoAttachLocationRule, AutoAttachRotationRule, AutoAttachScaleRule, bAutoAttachWeldSimulatedBodies), AutoAttachSocketName);
 				}
 
 				bDidAutoAttach = true;
