Fix CIS

diff --git a/Engine/Source/Runtime/Engine/Private/PhysicsEngine/BodyInstance.cpp b/Engine/Source/Runtime/Engine/Private/PhysicsEngine/BodyInstance.cpp
--- a/Engine/Source/Runtime/Engine/Private/PhysicsEngine/BodyInstance.cpp
+++ b/Engine/Source/Runtime/Engine/Private/PhysicsEngine/BodyInstance.cpp
@@ -898,11 +898,13 @@ void FBodyInstance::UpdatePhysicsFilterData()
 			UpdateMassProperties();
 		}
 
+#if WITH_CHAOS
 		//If filtering changed we must update GT structure right away
 		if (FPhysScene* PhysScene = GetPhysicsScene())
 		{
 			PhysScene->GetScene().UpdateActorInAccelerationStructure(Actor);
 		}
+#endif
 	});
 
 	UpdateInterpolateWhenSubStepping();
