Chaos - reduce default joint angle tolerance

diff --git a/Engine/Source/Runtime/Engine/Private/PhysicsEngine/ImmediatePhysics/ImmediatePhysicsChaos/ImmediatePhysicsSimulation_Chaos.cpp b/Engine/Source/Runtime/Engine/Private/PhysicsEngine/ImmediatePhysics/ImmediatePhysicsChaos/ImmediatePhysicsSimulation_Chaos.cpp
--- a/Engine/Source/Runtime/Engine/Private/PhysicsEngine/ImmediatePhysics/ImmediatePhysicsChaos/ImmediatePhysicsSimulation_Chaos.cpp
+++ b/Engine/Source/Runtime/Engine/Private/PhysicsEngine/ImmediatePhysics/ImmediatePhysicsChaos/ImmediatePhysicsSimulation_Chaos.cpp
@@ -57,7 +57,7 @@ int32 ChaosImmediate_Joint_PairIterations = -1;
 int32 ChaosImmediate_Joint_PushOutPairIterations = -1;
 float ChaosImmediate_Joint_SwingTwistAngleTolerance = 1.0e-6f;
 float ChaosImmediate_Joint_PositionTolerance = 0.05f;
-float ChaosImmediate_Joint_AngleTolerance = 0.01f;
+float ChaosImmediate_Joint_AngleTolerance = 0.001f;
 int32 ChaosImmediate_Joint_EnableTwistLimits = 1;
 int32 ChaosImmediate_Joint_EnableSwingLimits = 1;
 int32 ChaosImmediate_Joint_EnableDrives = 1;
