Reduce the size of the skeletal mesh editor raw data

diff --git a/Engine/Source/Developer/SkeletalMeshUtilitiesCommon/Private/LODUtilities.cpp b/Engine/Source/Developer/SkeletalMeshUtilitiesCommon/Private/LODUtilities.cpp
--- a/Engine/Source/Developer/SkeletalMeshUtilitiesCommon/Private/LODUtilities.cpp
+++ b/Engine/Source/Developer/SkeletalMeshUtilitiesCommon/Private/LODUtilities.cpp
@@ -2130,7 +2130,7 @@ public:
 		: LODModel(InLODModel)
 		, RefSkeleton(InRefSkeleton)
 		, BaseImportData(InBaseImportData)
-		, MorphLODPoints(InMorphLODPoints)
+		, CompressMorphLODPoints(InMorphLODPoints)
 		, MorphTargetDeltas(InMorphDeltas)
 		, BaseIndexData(InBaseIndexData)
 		, BaseWedgePointIndices(InBaseWedgePointIndices)
@@ -2149,6 +2149,21 @@ public:
 		MeshUtilities = &FModuleManager::Get().LoadModuleChecked<IMeshUtilities>("MeshUtilities");
 	}
 
+	//Decompress the shape points data
+	void DecompressData()
+	{
+		const TArray<FVector>& BaseMeshPoints = BaseImportData.Points;
+		MorphLODPoints = BaseMeshPoints;
+		int32 ModifiedPointIndex = 0;
+		for (uint32 PointIndex : ModifiedPoints)
+		{
+			MorphLODPoints[PointIndex] = CompressMorphLODPoints[ModifiedPointIndex];
+			ModifiedPointIndex++;
+		}
+
+		check(MorphLODPoints.Num() == MeshDataBundle.Vertices.Num());
+	}
+
 	void PrepareTangents()
 	{
 		TArray<bool> WasProcessed;
@@ -2299,6 +2314,7 @@ public:
 
 	void DoWork()
 	{
+		DecompressData();
 		PrepareTangents();
 		ComputeTangents();
 		ComputeMorphDeltas();
@@ -2330,7 +2346,8 @@ private:
 	// @todo not thread safe
 	const FReferenceSkeleton& RefSkeleton;
 	const FSkeletalMeshImportData& BaseImportData;
-	const TArray<FVector> MorphLODPoints;
+	const TArray<FVector> CompressMorphLODPoints;
+	TArray<FVector> MorphLODPoints;
 
 	IMeshUtilities* MeshUtilities;
 
diff --git a/Engine/Source/Editor/UnrealEd/Private/Fbx/FbxSkeletalMeshImport.cpp b/Engine/Source/Editor/UnrealEd/Private/Fbx/FbxSkeletalMeshImport.cpp
--- a/Engine/Source/Editor/UnrealEd/Private/Fbx/FbxSkeletalMeshImport.cpp
+++ b/Engine/Source/Editor/UnrealEd/Private/Fbx/FbxSkeletalMeshImport.cpp
@@ -1517,6 +1517,9 @@ bool UnFbx::FFbxImporter::FillSkeletalMeshImportPoints(FSkeletalMeshImportData*
 
 bool UnFbx::FFbxImporter::GatherPointsForMorphTarget(FSkeletalMeshImportData* OutData, TArray<FbxNode*>& NodeArray, TArray< FbxShape* >* FbxShapeArray, TSet<uint32>& ModifiedPoints)
 {
+	check(OutData);
+	TArray<FVector> CompressPoints;
+	CompressPoints.Reserve(OutData->Points.Num());
 	FSkeletalMeshImportData NewImportData = *OutData;
 	NewImportData.Points.Empty();
 
@@ -1545,15 +1548,14 @@ bool UnFbx::FFbxImporter::GatherPointsForMorphTarget(FSkeletalMeshImportData* Ou
 	for ( int32 PointIdx = 0; PointIdx < OutData->Points.Num(); ++PointIdx )
 	{
 		int32 OriginalPointIdx = OutData->PointToRawMap[ PointIdx ];
-
+		//Rebuild the data with only the modified point
 		if ( ( NewImportData.Points[ OriginalPointIdx ] - OutData->Points[ PointIdx ] ).SizeSquared() > FMath::Square(THRESH_POINTS_ARE_SAME) )
 		{
 			ModifiedPoints.Add( PointIdx );
+			CompressPoints.Add(NewImportData.Points[OriginalPointIdx]);
 		}
-
-		OutData->Points[ PointIdx ] = NewImportData.Points[ OriginalPointIdx ];
 	}
-
+	OutData->Points = CompressPoints;
 	return true;
 }
 
diff --git a/Engine/Source/Runtime/Engine/Private/Rendering/SkeletalMeshLODImporterData.cpp b/Engine/Source/Runtime/Engine/Private/Rendering/SkeletalMeshLODImporterData.cpp
--- a/Engine/Source/Runtime/Engine/Private/Rendering/SkeletalMeshLODImporterData.cpp
+++ b/Engine/Source/Runtime/Engine/Private/Rendering/SkeletalMeshLODImporterData.cpp
@@ -588,6 +588,7 @@ enum
 	RAW_SKELETAL_MESH_BULKDATA_VER_INITIAL = 0,
 	RAW_SKELETAL_MESH_BULKDATA_VER_AlternateInfluence = 1,
 	RAW_SKELETAL_MESH_BULKDATA_VER_RebuildSystem = 2,
+	RAW_SKELETAL_MESH_BULKDATA_VER_CompressMorphTargetData = 3,
 	// Add new raw mesh versions here.
 
 	RAW_SKELETAL_MESH_BULKDATA_VER_PLUS_ONE,
@@ -654,6 +655,23 @@ FArchive& operator<<(FArchive& Ar, FSkeletalMeshImportData& RawMesh)
 		RawMesh.AlternateInfluenceProfileNames.Empty();
 	}
 
+	if (Ar.IsLoading() && Version < RAW_SKELETAL_MESH_BULKDATA_VER_CompressMorphTargetData)
+	{
+		//Compress the morph target data
+		for (int32 MorphTargetIndex = 0; MorphTargetIndex < RawMesh.MorphTargets.Num(); ++MorphTargetIndex)
+		{
+			const TSet<uint32>& ModifiedPoints = RawMesh.MorphTargetModifiedPoints[MorphTargetIndex];
+			FSkeletalMeshImportData& ToCompressShapeImportData = RawMesh.MorphTargets[MorphTargetIndex];
+			TArray<FVector> CompressPoints;
+			CompressPoints.Reserve(ToCompressShapeImportData.Points.Num());
+			for (uint32 PointIndex : ModifiedPoints)
+			{
+				CompressPoints.Add(ToCompressShapeImportData.Points[PointIndex]);
+			}
+			ToCompressShapeImportData.Points = CompressPoints;
+		}
+	}
+
 	return Ar;
 }
 
@@ -735,6 +753,7 @@ void FRawSkeletalMeshBulkData::SaveRawMesh(FSkeletalMeshImportData& InMesh)
 
 	//Clear the bulk data before writing it
 	BulkData.RemoveBulkData();
+	BulkData.StoreCompressedOnDisk(NAME_Zlib);
 
 	// Get a lock on the bulk data
 	{
