make sure we actually release the retained chunk on USoundWave::BeginDestroy().

diff --git a/Engine/Source/Runtime/Engine/Private/SoundWave.cpp b/Engine/Source/Runtime/Engine/Private/SoundWave.cpp
--- a/Engine/Source/Runtime/Engine/Private/SoundWave.cpp
+++ b/Engine/Source/Runtime/Engine/Private/SoundWave.cpp
@@ -969,6 +969,8 @@ void USoundWave::BeginDestroy()
 		AsyncLoadingDataFormats.Empty();
 	}
 #endif
+
+	ReleaseCompressedAudio();
 }
 
 void USoundWave::InitAudioResource( FByteBulkData& CompressedData )
