Adding SetBooleanParameter Function to UFXAsset and VFXMarshaller.

#test compile/PIE
#rnx

diff --git a/Engine/Plugins/FX/Niagara/Source/Niagara/Private/NiagaraComponent.cpp b/Engine/Plugins/FX/Niagara/Source/Niagara/Private/NiagaraComponent.cpp
--- a/Engine/Plugins/FX/Niagara/Source/Niagara/Private/NiagaraComponent.cpp
+++ b/Engine/Plugins/FX/Niagara/Source/Niagara/Private/NiagaraComponent.cpp
@@ -473,6 +473,11 @@ UNiagaraComponent::UNiagaraComponent(const FObjectInitializer& ObjectInitializer
 }
 
 /********* UFXSystemComponent *********/
+void UNiagaraComponent::SetBoolParameter(FName Parametername, bool Param)
+{
+	SetVariableBool(Parametername, Param);
+}
+
 void UNiagaraComponent::SetFloatParameter(FName ParameterName, float Param)
 {
 	SetVariableFloat(ParameterName, Param);
diff --git a/Engine/Plugins/FX/Niagara/Source/Niagara/Public/NiagaraComponent.h b/Engine/Plugins/FX/Niagara/Source/Niagara/Public/NiagaraComponent.h
--- a/Engine/Plugins/FX/Niagara/Source/Niagara/Public/NiagaraComponent.h
+++ b/Engine/Plugins/FX/Niagara/Source/Niagara/Public/NiagaraComponent.h
@@ -47,10 +47,12 @@ class NIAGARA_API UNiagaraComponent : public UFXSystemComponent
 public:
 
 	/********* UFXSystemComponent *********/
+	void SetBoolParameter(FName ParameterName, bool Param) override;
 	void SetFloatParameter(FName ParameterName, float Param) override;
 	void SetVectorParameter(FName ParameterName, FVector Param) override;
 	void SetColorParameter(FName ParameterName, FLinearColor Param) override;
 	void SetActorParameter(FName ParameterName, class AActor* Param) override;
+
 	virtual UFXSystemAsset* GetFXSystemAsset() const override;
 	void SetEmitterEnable(FName EmitterName, bool bNewEnableState) override;
 	void ReleaseToPool() override;
diff --git a/Engine/Source/Runtime/Engine/Classes/Particles/ParticleSystemComponent.h b/Engine/Source/Runtime/Engine/Classes/Particles/ParticleSystemComponent.h
--- a/Engine/Source/Runtime/Engine/Classes/Particles/ParticleSystemComponent.h
+++ b/Engine/Source/Runtime/Engine/Classes/Particles/ParticleSystemComponent.h
@@ -333,6 +333,10 @@ class ENGINE_API UFXSystemComponent : public UPrimitiveComponent
 	GENERATED_UCLASS_BODY()
 public:
 
+	/**Change a named boolean parameter, ParticleSystemComponent converts to float.*/
+	UFUNCTION(BlueprintCallable, Category = "Effects|Components|ParticleSystem")
+	virtual void SetBoolParameter(FName ParameterName, bool Param) {}
+
 	/** Change a named float parameter */
 	UFUNCTION(BlueprintCallable, Category="Effects|Components|ParticleSystem")
 	virtual void SetFloatParameter(FName ParameterName, float Param) {}
@@ -1059,6 +1063,8 @@ public:
 	 */
 	void SetEmitterEnable(FName EmitterName, bool bNewEnableState) override;
 
+	void SetBoolParameter(FName ParameterName, bool Param) override;
+
 	/** Change a named float parameter */
 	void SetFloatParameter(FName ParameterName, float Param) override;
 
diff --git a/Engine/Source/Runtime/Engine/Private/Particles/ParticleComponents.cpp b/Engine/Source/Runtime/Engine/Private/Particles/ParticleComponents.cpp
--- a/Engine/Source/Runtime/Engine/Private/Particles/ParticleComponents.cpp
+++ b/Engine/Source/Runtime/Engine/Private/Particles/ParticleComponents.cpp
@@ -7072,6 +7072,18 @@ int32 UParticleSystemComponent::GetLODLevel()
 {
 	return LODLevel;
 }
+/**
+ *	Set a named float instance parameter on this ParticleSystemComponent.
+ *	Updates the parameter if it already exists, or creates a new entry if not.
+	This maps a boolean to a float for parity as cascade doesn't have booleans.
+	This for adding functionality to the parent UFXSystemComponent to set boolean
+	variables.
+ */
+void UParticleSystemComponent::SetBoolParameter(FName Name, bool Value)
+{
+	SetFloatParameter(Name, float(Value));
+
+}
 
 /** 
  *	Set a named float instance parameter on this ParticleSystemComponent. 
