CIS fix

diff --git a/Engine/Source/Runtime/Engine/Public/Animation/AnimCompressionDerivedDataPublic.h b/Engine/Source/Runtime/Engine/Public/Animation/AnimCompressionDerivedDataPublic.h
--- a/Engine/Source/Runtime/Engine/Public/Animation/AnimCompressionDerivedDataPublic.h
+++ b/Engine/Source/Runtime/Engine/Public/Animation/AnimCompressionDerivedDataPublic.h
@@ -5,6 +5,7 @@
 #if WITH_EDITOR
 
 #include "TickableEditorObject.h"
+#include "Animation/AnimCompressionTypes.h"
 
 class UAnimSequence;
 class FDerivedDataAnimationCompression;
