Speculative fix for a UAV state warning

diff --git a/Engine/Source/Runtime/Engine/Private/DistanceFieldAtlas.cpp b/Engine/Source/Runtime/Engine/Private/DistanceFieldAtlas.cpp
--- a/Engine/Source/Runtime/Engine/Private/DistanceFieldAtlas.cpp
+++ b/Engine/Source/Runtime/Engine/Private/DistanceFieldAtlas.cpp
@@ -1579,12 +1579,17 @@ void FLandscapeTextureAtlas::UpdateAllocations(FRHICommandListImmediate& RHICmdL
 			{
 				typename FUploadHeightFieldToAtlasCS::FParameters* Parameters = GraphBuilder.AllocParameters<typename FUploadHeightFieldToAtlasCS::FParameters>();
 				const FIntPoint UpdateRegion = PendingUploads[Idx].SetShaderParameters(Parameters, *this);
+				const bool bNeedBarrier = Idx > 0;
 				GraphBuilder.AddPass(
 					RDG_EVENT_NAME("UploadHeightFieldToAtlas"),
 					Parameters,
 					ERDGPassFlags::Compute,
-					[Parameters, ComputeShaderPtr, UpdateRegion](FRHICommandList& CmdList)
+					[Parameters, ComputeShaderPtr, UpdateRegion, bNeedBarrier](FRHICommandList& CmdList)
 				{
+					if (bNeedBarrier)
+					{
+						CmdList.TransitionResource(EResourceTransitionAccess::ERWNoBarrier, EResourceTransitionPipeline::EComputeToCompute, Parameters->RWHeightFieldAtlas);
+					}
 					FComputeShaderUtils::Dispatch(CmdList, ComputeShaderPtr, *Parameters, FIntVector(UpdateRegion.X, UpdateRegion.Y, 1));
 				});
 			}
@@ -1598,12 +1603,17 @@ void FLandscapeTextureAtlas::UpdateAllocations(FRHICommandListImmediate& RHICmdL
 			{
 				typename FUploadVisibilityToAtlasCS::FParameters* Parameters = GraphBuilder.AllocParameters<typename FUploadVisibilityToAtlasCS::FParameters>();
 				const FIntPoint UpdateRegion = PendingUploads[Idx].SetShaderParameters(Parameters, *this);
+				const bool bNeedBarrier = Idx > 0;
 				GraphBuilder.AddPass(
 					RDG_EVENT_NAME("UploadVisibilityToAtlas"),
 					Parameters,
 					ERDGPassFlags::Compute,
-					[Parameters, ComputeShaderPtr, UpdateRegion](FRHICommandList& CmdList)
+					[Parameters, ComputeShaderPtr, UpdateRegion, bNeedBarrier](FRHICommandList& CmdList)
 				{
+					if (bNeedBarrier)
+					{
+						CmdList.TransitionResource(EResourceTransitionAccess::ERWNoBarrier, EResourceTransitionPipeline::EComputeToCompute, Parameters->RWVisibilityAtlas);
+					}
 					FComputeShaderUtils::Dispatch(CmdList, ComputeShaderPtr, *Parameters, FIntVector(UpdateRegion.X, UpdateRegion.Y, 1));
 				});
 			}
