Fix race condition that would destroy the replay driver before the map even loaded

diff --git a/Engine/Source/Runtime/Engine/Classes/Engine/World.h b/Engine/Source/Runtime/Engine/Classes/Engine/World.h
--- a/Engine/Source/Runtime/Engine/Classes/Engine/World.h
+++ b/Engine/Source/Runtime/Engine/Classes/Engine/World.h
@@ -2925,6 +2925,9 @@ public:
 	// Destroys the current demo net driver
 	void DestroyDemoNetDriver();
 
+	// Remove internal references to pending demo net driver when starting a replay, but do not destroy it
+	void ClearDemoNetDriver();
+
 	/** Returns true if we are currently playing a replay */
 	bool IsPlayingReplay() const;
 
diff --git a/Engine/Source/Runtime/Engine/Private/DemoNetDriver.cpp b/Engine/Source/Runtime/Engine/Private/DemoNetDriver.cpp
--- a/Engine/Source/Runtime/Engine/Private/DemoNetDriver.cpp
+++ b/Engine/Source/Runtime/Engine/Private/DemoNetDriver.cpp
@@ -1145,7 +1145,7 @@ bool UDemoNetDriver::InitConnectInternal(FString& Error)
 				return false;
 			}
 
-			GetWorld()->DemoNetDriver = nullptr;
+			GetWorld()->ClearDemoNetDriver();
 			SetWorld(nullptr);
 
 			auto NewPendingNetGame = NewObject<UDemoPendingNetGame>();
diff --git a/Engine/Source/Runtime/Engine/Private/World.cpp b/Engine/Source/Runtime/Engine/Private/World.cpp
--- a/Engine/Source/Runtime/Engine/Private/World.cpp
+++ b/Engine/Source/Runtime/Engine/Private/World.cpp
@@ -3915,6 +3915,22 @@ void UWorld::DestroyDemoNetDriver()
 	}
 }
 
+void UWorld::ClearDemoNetDriver()
+{
+	FLevelCollection* const SourceCollection = FindCollectionByType(ELevelCollectionType::DynamicSourceLevels);
+	if (SourceCollection)
+	{
+		SourceCollection->SetDemoNetDriver(nullptr);
+	}
+	FLevelCollection* const StaticCollection = FindCollectionByType(ELevelCollectionType::StaticLevels);
+	if (StaticCollection)
+	{
+		StaticCollection->SetDemoNetDriver(nullptr);
+	}
+
+	DemoNetDriver = nullptr;
+}
+
 bool UWorld::SetGameMode(const FURL& InURL)
 {
 	if( IsServer() && !AuthorityGameMode )
