Fix GetTriangleMeshExternalFaceIndex checking type wrong and doing sketchy stuff.

diff --git a/Engine/Source/Runtime/Engine/Private/PhysicsEngine/Experimental/PhysInterface_Chaos.cpp b/Engine/Source/Runtime/Engine/Private/PhysicsEngine/Experimental/PhysInterface_Chaos.cpp
--- a/Engine/Source/Runtime/Engine/Private/PhysicsEngine/Experimental/PhysInterface_Chaos.cpp
+++ b/Engine/Source/Runtime/Engine/Private/PhysicsEngine/Experimental/PhysInterface_Chaos.cpp
@@ -2169,19 +2169,24 @@ bool FPhysInterface_Chaos::GetSquaredDistanceToBody(const FBodyInstance* InInsta
 uint32 GetTriangleMeshExternalFaceIndex(const FPhysicsShape& Shape, uint32 InternalFaceIndex)
 {
 	using namespace Chaos;
-	uint8 Type = Shape.Geometry->GetType();
-	if (ensure(Type | ImplicitObjectType::TriangleMesh))
+	uint8 OuterType = Shape.Geometry->GetType();
+	uint8 InnerType = GetInnerType(OuterType);
+	if (ensure(InnerType == ImplicitObjectType::TriangleMesh))
 	{
 		const FTriangleMeshImplicitObject* TriangleMesh = nullptr;
 
-		if (Type | ImplicitObjectType::IsScaled)
+		if (IsScaled(OuterType))
 		{
-			const TImplicitObjectScaled<FTriangleMeshImplicitObject>* ScaledTriangleMesh = static_cast<const TImplicitObjectScaled<FTriangleMeshImplicitObject>*>(Shape.Geometry.Get());
-			TriangleMesh = ScaledTriangleMesh->GetUnscaledObject();
+			const TImplicitObjectScaled<FTriangleMeshImplicitObject>& ScaledTriangleMesh = Shape.Geometry->GetObjectChecked<TImplicitObjectScaled<FTriangleMeshImplicitObject>>();
+			TriangleMesh = ScaledTriangleMesh.GetUnscaledObject();
+		}
+		else if(IsInstanced(OuterType))
+		{
+			TriangleMesh = Shape.Geometry->GetObjectChecked<TImplicitObjectInstanced<FTriangleMeshImplicitObject>>().GetInstancedObject();
 		}
 		else
 		{
-			TriangleMesh = static_cast<const FTriangleMeshImplicitObject*>(Shape.Geometry.Get());
+			TriangleMesh = &Shape.Geometry->GetObjectChecked<FTriangleMeshImplicitObject>();
 		}
 
 		return TriangleMesh->GetExternalFaceIndexFromInternal(InternalFaceIndex);
