Remove PreFilter stat which was too small and inflating perf numbers

diff --git a/Engine/Source/Runtime/Engine/Private/PhysicsEngine/CollisionQueryFilterCallback.cpp b/Engine/Source/Runtime/Engine/Private/PhysicsEngine/CollisionQueryFilterCallback.cpp
--- a/Engine/Source/Runtime/Engine/Private/PhysicsEngine/CollisionQueryFilterCallback.cpp
+++ b/Engine/Source/Runtime/Engine/Private/PhysicsEngine/CollisionQueryFilterCallback.cpp
@@ -89,7 +89,7 @@ ECollisionQueryHitType FCollisionQueryFilterCallback::CalcQueryHitType(const FCo
 #if WITH_PHYSX
 ECollisionQueryHitType FCollisionQueryFilterCallback::PreFilterImp(const FCollisionFilterData& FilterData, const physx::PxShape& Shape, const physx::PxActor& Actor)
 {
-	SCOPE_CYCLE_COUNTER(STAT_Collision_PreFilter);
+	//SCOPE_CYCLE_COUNTER(STAT_Collision_PreFilter);
 	FCollisionFilterData ShapeFilter = PhysXInterface::GetQueryFilterData(Shape);
 
 	// We usually don't have ignore components so we try to avoid the virtual getSimulationFilterData() call below. 'word2' of shape sim filter data is componentID.
@@ -110,7 +110,7 @@ ECollisionQueryHitType FCollisionQueryFilterCallback::PreFilterImp(const FCollis
 
 ECollisionQueryHitType FCollisionQueryFilterCallback::PreFilterImp(const FCollisionFilterData& FilterData, const Chaos::TPerShapeData<float,3>& Shape, const Chaos::TGeometryParticle<float,3>& Actor)
 {
-	SCOPE_CYCLE_COUNTER(STAT_Collision_PreFilter);
+	//SCOPE_CYCLE_COUNTER(STAT_Collision_PreFilter);
 	FCollisionFilterData ShapeFilter = ChaosInterface::GetQueryFilterData(Shape);
 
 	// We usually don't have ignore components so we try to avoid the virtual getSimulationFilterData() call below. 'word2' of shape sim filter data is componentID.
