# Exposed montage sync leader on montage instances
#rnx

diff --git a/Engine/Source/Runtime/Engine/Classes/Animation/AnimMontage.h b/Engine/Source/Runtime/Engine/Classes/Animation/AnimMontage.h
--- a/Engine/Source/Runtime/Engine/Classes/Animation/AnimMontage.h
+++ b/Engine/Source/Runtime/Engine/Classes/Animation/AnimMontage.h
@@ -459,6 +459,7 @@ public:
 	float GetPlayRate() const { return PlayRate; }
 	float GetDeltaMoved() const { return DeltaMoved; }
 	float GetPreviousPosition() const { return PreviousPosition;  }
+	const FAnimMontageInstance* GetMontageSyncLeader() const { return MontageSyncLeader; } 
 
 	/** 
 	 * Setters
