Fix for check hit if you launch the engine with an invalid value for AudioDeviceModuleName or AudioMixerModuleName in [Platform]Engine.ini, as well as an ensure that hits when the engine is run with -nosound.

diff --git a/Engine/Source/Runtime/Engine/Private/AudioDeviceManager.cpp b/Engine/Source/Runtime/Engine/Private/AudioDeviceManager.cpp
--- a/Engine/Source/Runtime/Engine/Private/AudioDeviceManager.cpp
+++ b/Engine/Source/Runtime/Engine/Private/AudioDeviceManager.cpp
@@ -305,7 +305,8 @@ bool FAudioDeviceManager::Initialize()
 		MainDeviceParams.bIsNonRealtime = false;
 
 		MainAudioDeviceHandle = RequestAudioDevice(MainDeviceParams);
-		check(MainAudioDeviceHandle.IsValid());
+		UE_CLOG(!MainAudioDeviceHandle.IsValid(), LogAudio, Display, TEXT("Audio device could not be initialized. Please check the value for AudioDeviceModuleName and AudioMixerModuleName in [Platform]Engine.ini."));
+
 		FAudioThread::StartAudioThread();
 		return true;
 	}
@@ -397,8 +398,10 @@ FAudioDeviceHandle FAudioDeviceManager::CreateNewDevice(const FAudioDeviceParams
 	Devices.Emplace(DeviceID, FAudioDeviceContainer(InParams, this));
 	FAudioDeviceContainer* ContainerPtr = Devices.Find(DeviceID);
 	check(ContainerPtr);
-	if (!ensure(ContainerPtr->Device))
+	if (!ContainerPtr->Device)
 	{
+		UE_LOG(LogAudio, Display, TEXT("Audio device could not be initialized. Please check the value for AudioDeviceModuleName and AudioMixerModuleName in [Platform]Engine.ini."));
+
 		// Initializing the audio device failed. Remove the device container and return an empty handle.
 		Devices.Remove(DeviceID);
 		return FAudioDeviceHandle();
