Fix stomped change from CL resolve. Would cause assert to hit on shutdown

diff --git a/Engine/Source/Runtime/Engine/Private/AudioDeviceManager.cpp b/Engine/Source/Runtime/Engine/Private/AudioDeviceManager.cpp
--- a/Engine/Source/Runtime/Engine/Private/AudioDeviceManager.cpp
+++ b/Engine/Source/Runtime/Engine/Private/AudioDeviceManager.cpp
@@ -75,6 +75,7 @@ FAudioDeviceManager::FAudioDeviceManager()
 
 FAudioDeviceManager::~FAudioDeviceManager()
 {
+	MainAudioDeviceHandle.Reset();
 	Devices.Reset();
 
 	// Release any loaded buffers - this calls stop on any sources that need it
