Minifix: TWeakObjPtr.IsValid() could be called after the object table is destroyed.

diff --git a/Engine/Source/Runtime/Engine/Classes/Sound/SoundEffectPreset.h b/Engine/Source/Runtime/Engine/Classes/Sound/SoundEffectPreset.h
--- a/Engine/Source/Runtime/Engine/Classes/Sound/SoundEffectPreset.h
+++ b/Engine/Source/Runtime/Engine/Classes/Sound/SoundEffectPreset.h
@@ -37,6 +37,9 @@ public:
 
 	void AddReferencedEffects(FReferenceCollector& Collector);
 
+
+	virtual void BeginDestroy() override;
+
 protected:
 
 #if WITH_EDITORONLY_DATA
diff --git a/Engine/Source/Runtime/Engine/Private/SoundEffectBase.cpp b/Engine/Source/Runtime/Engine/Private/SoundEffectBase.cpp
--- a/Engine/Source/Runtime/Engine/Private/SoundEffectBase.cpp
+++ b/Engine/Source/Runtime/Engine/Private/SoundEffectBase.cpp
@@ -49,14 +49,12 @@ USoundEffectPreset* FSoundEffectBase::GetPreset()
 
 void FSoundEffectBase::ClearPreset(bool bRemoveFromPreset)
 {
-	if (Preset.IsValid())
+	if (bRemoveFromPreset && Preset.IsValid())
 	{
-		if (bRemoveFromPreset)
-		{
-			Preset->RemoveEffectInstance(this);
-		}
-		Preset.Reset();
+		Preset->RemoveEffectInstance(this);
 	}
+
+	Preset.Reset();
 }
 
 bool FSoundEffectBase::Update()
diff --git a/Engine/Source/Runtime/Engine/Private/SoundEffectPreset.cpp b/Engine/Source/Runtime/Engine/Private/SoundEffectPreset.cpp
--- a/Engine/Source/Runtime/Engine/Private/SoundEffectPreset.cpp
+++ b/Engine/Source/Runtime/Engine/Private/SoundEffectPreset.cpp
@@ -15,14 +15,6 @@ USoundEffectPreset::USoundEffectPreset(const FObjectInitializer& ObjectInitializ
 
 USoundEffectPreset::~USoundEffectPreset()
 {
-	for (int32 i = 0; i < Instances.Num(); ++i)
-	{
-		if (Instances[i])
-		{
-			Instances[i]->ClearPreset(false /* bRemoveFromPreset */);
-		}
-	}
-	Instances.Reset();
 }
 
 
@@ -78,6 +70,20 @@ void USoundEffectPreset::AddReferencedEffects(FReferenceCollector& Collector)
 	}
 }
 
+void USoundEffectPreset::BeginDestroy()
+{
+	for (int32 i = 0; i < Instances.Num(); ++i)
+	{
+		if (Instances[i])
+		{
+			Instances[i]->ClearPreset(false /* bRemoveFromPreset */);
+		}
+	}
+	Instances.Reset();
+
+	Super::BeginDestroy();
+}
+
 void USoundEffectPreset::RemoveEffectInstance(FSoundEffectBase* InSource)
 {
 	Instances.RemoveSwap(InSource);
