Chaos - Deferred UpdateKinematicBonesToAnim optimizations
- Batch calls to update spatial acceleration
- Remove critical section in FPhysicsSolverBase::AddDirty

diff --git a/Engine/Source/Runtime/Engine/Private/PhysicsEngine/Experimental/PhysScene_Chaos.cpp b/Engine/Source/Runtime/Engine/Private/PhysicsEngine/Experimental/PhysScene_Chaos.cpp
--- a/Engine/Source/Runtime/Engine/Private/PhysicsEngine/Experimental/PhysScene_Chaos.cpp
+++ b/Engine/Source/Runtime/Engine/Private/PhysicsEngine/Experimental/PhysScene_Chaos.cpp
@@ -680,6 +680,41 @@ void FPhysScene_Chaos::UpdateActorInAccelerationStructure(const FPhysicsActorHan
 #endif
 }
 
+void FPhysScene_Chaos::UpdateActorsInAccelerationStructure(const TArrayView<FPhysicsActorHandle>& Actors)
+{
+#if WITH_CHAOS
+	using namespace Chaos;
+
+	if (GetSpacialAcceleration())
+	{
+		ExternalDataLock.WriteLock();
+
+		auto SpatialAcceleration = GetSpacialAcceleration();
+
+		if (SpatialAcceleration)
+		{
+			int32 NumActors = Actors.Num();
+			for (int32 ActorIndex = 0; ActorIndex < NumActors; ++ActorIndex)
+			{
+				const FPhysicsActorHandle& Actor = Actors[ActorIndex];
+
+				// @todo(chaos): dedupe code in UpdateActorInAccelerationStructure
+				TAABB<FReal, 3> WorldBounds;
+				const bool bHasBounds = Actor->Geometry()->HasBoundingBox();
+				if (bHasBounds)
+				{
+					WorldBounds = Actor->Geometry()->BoundingBox().TransformedAABB(TRigidTransform<FReal, 3>(Actor->X(), Actor->R()));
+				}
+
+				Chaos::TAccelerationStructureHandle<float, 3> AccelerationHandle(Actor);
+				SpatialAcceleration->UpdateElementIn(AccelerationHandle, WorldBounds, bHasBounds, Actor->SpatialIdx());
+			}
+		}
+
+		ExternalDataLock.WriteUnlock();
+	}
+#endif
+}
 
 template<typename ObjectType>
 void RemovePhysicsProxy(ObjectType* InObject, Chaos::FPhysicsSolver* InSolver, FChaosSolversModule* InModule)
@@ -1575,10 +1610,101 @@ void FPhysScene_ChaosInterface::ClearPreSimKinematicUpdate(USkeletalMeshComponen
 	}
 }
 
+// Collect all the actors that need moving, along with their transforms
+// Extracted from USkeletalMeshComponent::UpdateKinematicBonesToAnim
+// @todo(chaos): merge this functionality back into USkeletalMeshComponent
+template<typename T_ACTORCONTAINER, typename T_TRANSFORMCONTAINER>
+void GatherActorsAndTransforms(
+	USkeletalMeshComponent* SkelComp, 
+	const TArray<FTransform>& InComponentSpaceTransforms, 
+	ETeleportType Teleport, 
+	bool bNeedsSkinning, 
+	T_ACTORCONTAINER& KinematicUpdateActors,
+	T_TRANSFORMCONTAINER& KinematicUpdateTransforms,
+	T_ACTORCONTAINER& TeleportActors,
+	T_TRANSFORMCONTAINER& TeleportTransforms)
+{
+	bool bTeleport = Teleport == ETeleportType::TeleportPhysics;
+	const UPhysicsAsset* PhysicsAsset = SkelComp->GetPhysicsAsset();
+	const FTransform& CurrentLocalToWorld = SkelComp->GetComponentTransform();
+	const int32 NumBodies = SkelComp->Bodies.Num();
+	for (int32 i = 0; i < NumBodies; i++)
+	{
+		FBodyInstance* BodyInst = SkelComp->Bodies[i];
+		FPhysicsActorHandle& ActorHandle = BodyInst->ActorHandle;
+		if (bTeleport || !BodyInst->IsInstanceSimulatingPhysics())
+		{
+			const int32 BoneIndex = BodyInst->InstanceBoneIndex;
+			if (BoneIndex != INDEX_NONE)
+			{
+				const FTransform BoneTransform = InComponentSpaceTransforms[BoneIndex] * CurrentLocalToWorld;
+				if (!bTeleport)
+				{
+					KinematicUpdateActors.Add(ActorHandle);
+					KinematicUpdateTransforms.Add(BoneTransform);
+				}
+				else
+				{
+					TeleportActors.Add(ActorHandle);
+					TeleportTransforms.Add(BoneTransform);
+				}
+				if (!PhysicsAsset->SkeletalBodySetups[i]->bSkipScaleFromAnimation)
+				{
+					const FVector& MeshScale3D = CurrentLocalToWorld.GetScale3D();
+					if (MeshScale3D.IsUniform())
+					{
+						BodyInst->UpdateBodyScale(BoneTransform.GetScale3D());
+					}
+					else
+					{
+						BodyInst->UpdateBodyScale(MeshScale3D);
+					}
+				}
+			}
+		}
+	}
+}
+
+// Move all actors that need teleporting
+void ProcessTeleportActors(FPhysScene_Chaos& Scene, const TArrayView<FPhysicsActorHandle>& ActorHandles, const TArrayView<FTransform>& Transforms)
+{
+	int32 NumActors = ActorHandles.Num();
+	if (NumActors > 0)
+	{
+		for (int32 ActorIndex = 0; ActorIndex < NumActors; ++ActorIndex)
+		{
+			const FPhysicsActorHandle& ActorHandle = ActorHandles[ActorIndex];
+			const FTransform& ActorTransform = Transforms[ActorIndex];
+			ActorHandle->SetX(ActorTransform.GetLocation(), false);	// only set dirty once in SetR
+			ActorHandle->SetR(ActorTransform.GetRotation());
+			ActorHandle->UpdateShapeBounds();
+		}
+
+		Scene.UpdateActorsInAccelerationStructure(ActorHandles);
+	}
+}
+
+// Set all actor kinematic targets
+void ProcessKinematicTargetActors(FPhysScene_Chaos& Scene, const TArrayView<FPhysicsActorHandle>& ActorHandles, const TArrayView<FTransform>& Transforms)
+{
+	// TODO - kinematic targets
+	ProcessTeleportActors(Scene, ActorHandles, Transforms);
+}
+
+// Collect the actors and transforms of all the bodies we have to move, and process them in bulk
+// to avoid locks in the Spatial Acceleration and the Solver's Dirty Proxy systems.
 void FPhysScene_ChaosInterface::UpdateKinematicsOnDeferredSkelMeshes()
 {
 	SCOPE_CYCLE_COUNTER(STAT_UpdateKinematicsOnDeferredSkelMeshesChaos);
 
+	using FActorsArray = TArray<FPhysicsActorHandle, TInlineAllocator<64>>;
+	using FTransformsArray = TArray<FTransform, TInlineAllocator<64>>;
+
+	FActorsArray KinematicUpdateActors;
+	FTransformsArray KinematicUpdateTransforms;
+	FActorsArray TeleportActors;
+	FTransformsArray TeleportTransforms;
+
 	for (const TPair<USkeletalMeshComponent*, FDeferredKinematicUpdateInfo>& DeferredKinematicUpdate : DeferredKinematicUpdateSkelMeshes)
 	{
 		USkeletalMeshComponent* SkelComp = DeferredKinematicUpdate.Key;
@@ -1586,14 +1712,23 @@ void FPhysScene_ChaosInterface::UpdateKinematicsOnDeferredSkelMeshes()
 
 		ensure(SkelComp->bDeferredKinematicUpdate); // Should be true if in map!
 
-		// Perform kinematic updates
-		SkelComp->UpdateKinematicBonesToAnim(SkelComp->GetComponentSpaceTransforms(), Info.TeleportType, Info.bNeedsSkinning, EAllowKinematicDeferral::DisallowDeferral);
+		if (!SkelComp->bEnablePerPolyCollision)
+		{
+			GatherActorsAndTransforms(SkelComp, SkelComp->GetComponentSpaceTransforms(), Info.TeleportType, Info.bNeedsSkinning, KinematicUpdateActors, KinematicUpdateTransforms, TeleportActors, TeleportTransforms);
+		}
+		else
+		{
+			// TODO: acceleration for per-poly collision
+			SkelComp->UpdateKinematicBonesToAnim(SkelComp->GetComponentSpaceTransforms(), Info.TeleportType, Info.bNeedsSkinning, EAllowKinematicDeferral::DisallowDeferral);
+		}
 
 		// Clear deferred flag
 		SkelComp->bDeferredKinematicUpdate = false;
 	}
 
-	// Empty map now all is done
+	ProcessKinematicTargetActors(Scene, KinematicUpdateActors, KinematicUpdateTransforms);
+	ProcessTeleportActors(Scene, TeleportActors, TeleportTransforms);
+
 	DeferredKinematicUpdateSkelMeshes.Reset();
 }
 
diff --git a/Engine/Source/Runtime/Engine/Private/PhysicsEngine/PhysAnim.cpp b/Engine/Source/Runtime/Engine/Private/PhysicsEngine/PhysAnim.cpp
--- a/Engine/Source/Runtime/Engine/Private/PhysicsEngine/PhysAnim.cpp
+++ b/Engine/Source/Runtime/Engine/Private/PhysicsEngine/PhysAnim.cpp
@@ -476,6 +476,8 @@ void USkeletalMeshComponent::CompleteParallelBlendPhysics()
 	ParallelBlendPhysicsCompletionTask.SafeRelease();
 }
 
+// NOTE: See GatherActorsAndTransforms in PhysScene_Chaos.cpp where this code is cloned for deferred mode.
+// @todo(chaos): merge required deferred functionality back into USkeletalMeshComponent
 void USkeletalMeshComponent::UpdateKinematicBonesToAnim(const TArray<FTransform>& InSpaceBases, ETeleportType Teleport, bool bNeedsSkinning, EAllowKinematicDeferral DeferralAllowed)
 {
 	SCOPE_CYCLE_COUNTER(STAT_UpdateRBBones);
diff --git a/Engine/Source/Runtime/Engine/Public/Physics/Experimental/PhysScene_Chaos.h b/Engine/Source/Runtime/Engine/Public/Physics/Experimental/PhysScene_Chaos.h
--- a/Engine/Source/Runtime/Engine/Public/Physics/Experimental/PhysScene_Chaos.h
+++ b/Engine/Source/Runtime/Engine/Public/Physics/Experimental/PhysScene_Chaos.h
@@ -137,6 +137,7 @@ public:
 
 	void RemoveActorFromAccelerationStructure(FPhysicsActorHandle& Actor);
 	void UpdateActorInAccelerationStructure(const FPhysicsActorHandle& Actor);
+	void UpdateActorsInAccelerationStructure(const TArrayView<FPhysicsActorHandle>& Actors);
 
 #if XGE_FIXED
 	template<typename PayloadType>
diff --git a/Engine/Source/Runtime/Experimental/Chaos/Public/Chaos/Framework/PhysicsSolverBase.h b/Engine/Source/Runtime/Experimental/Chaos/Public/Chaos/Framework/PhysicsSolverBase.h
--- a/Engine/Source/Runtime/Experimental/Chaos/Public/Chaos/Framework/PhysicsSolverBase.h
+++ b/Engine/Source/Runtime/Experimental/Chaos/Public/Chaos/Framework/PhysicsSolverBase.h
@@ -22,12 +22,10 @@ namespace Chaos
 
 		void AddDirtyProxy(IPhysicsProxyBase * ProxyBaseIn)
 		{
-			FScopeLock Lock(&DirtyProxiesCriticalSection);
 			DirtyProxiesSet.Add(ProxyBaseIn);
 		}
 		void RemoveDirtyProxy(IPhysicsProxyBase * ProxyBaseIn)
 		{
-			FScopeLock Lock(&DirtyProxiesCriticalSection);
 			DirtyProxiesSet.Remove(ProxyBaseIn);
 		}
 
@@ -47,7 +45,6 @@ namespace Chaos
 			EMultiBufferMode BufferMode;
 
 			// Input Proxy Map
-			FCriticalSection DirtyProxiesCriticalSection;
 			TSet< IPhysicsProxyBase *> DirtyProxiesSet;
 
 #if CHAOS_CHECKED
