Fix merge issue: merges seem to have accepted target for any files that had additional changes in the target to resolve against.

diff --git a/Engine/Source/Editor/Kismet/Private/BlueprintEditor.cpp b/Engine/Source/Editor/Kismet/Private/BlueprintEditor.cpp
--- a/Engine/Source/Editor/Kismet/Private/BlueprintEditor.cpp
+++ b/Engine/Source/Editor/Kismet/Private/BlueprintEditor.cpp
@@ -8990,7 +8990,7 @@ void FBlueprintEditor::UpdatePreviewActor(UBlueprint* InBlueprint, bool bInForce
 			}
 
 			// Prevent any audio from playing as a result of spawning
-			if (FAudioDevice* AudioDevice = GEngine->GetMainAudioDevice())
+			if (FAudioDeviceHandle AudioDevice = GEngine->GetMainAudioDevice())
 			{
 				AudioDevice->Flush(PreviewScene.GetWorld());
 			}
diff --git a/Engine/Source/Editor/MovieSceneCaptureDialog/Private/MovieSceneCaptureDialogModule.cpp b/Engine/Source/Editor/MovieSceneCaptureDialog/Private/MovieSceneCaptureDialogModule.cpp
--- a/Engine/Source/Editor/MovieSceneCaptureDialog/Private/MovieSceneCaptureDialogModule.cpp
+++ b/Engine/Source/Editor/MovieSceneCaptureDialog/Private/MovieSceneCaptureDialogModule.cpp
@@ -434,8 +434,8 @@ void FInEditorCapture::Start()
 	UGameViewportClient::OnViewportCreated().AddRaw(this, &FInEditorCapture::OnPIEViewportStarted);
 	FEditorDelegates::EndPIE.AddRaw(this, &FInEditorCapture::OnEndPIE);
 		
-	FAudioDevice* AudioDevice = GEngine->GetMainAudioDevice();
-	if (AudioDevice != nullptr)
+	FAudioDeviceHandle AudioDevice = GEngine->GetMainAudioDevice();
+	if (AudioDevice)
 	{
 		TransientMasterVolume = AudioDevice->GetTransientMasterVolume();
 		AudioDevice->SetTransientMasterVolume(0.0f);
@@ -596,8 +596,8 @@ void FInEditorCapture::Shutdown()
 
 	FObjectReader(GetMutableDefault<ULevelEditorPlaySettings>(), BackedUpPlaySettings);
 
-	FAudioDevice* AudioDevice = GEngine->GetMainAudioDevice();
-	if (AudioDevice != nullptr)
+	FAudioDeviceHandle AudioDevice = GEngine->GetMainAudioDevice();
+	if (AudioDevice)
 	{
 		AudioDevice->SetTransientMasterVolume(TransientMasterVolume);
 	}
diff --git a/Engine/Source/Editor/Persona/Private/AnimationEditorViewportClient.cpp b/Engine/Source/Editor/Persona/Private/AnimationEditorViewportClient.cpp
--- a/Engine/Source/Editor/Persona/Private/AnimationEditorViewportClient.cpp
+++ b/Engine/Source/Editor/Persona/Private/AnimationEditorViewportClient.cpp
@@ -138,7 +138,7 @@ FAnimationViewportClient::FAnimationViewportClient(const TSharedRef<IPersonaPrev
 	{
 		World->bAllowAudioPlayback = !ConfigOption->bMuteAudio;
 
-		if(FAudioDevice* AudioDevice = World->GetAudioDevice())
+		if(FAudioDevice* AudioDevice = World->GetAudioDeviceRaw())
 		{
 			AudioDevice->SetUseAttenuationForNonGameWorlds(ConfigOption->bUseAudioAttenuation);
 		}
@@ -222,7 +222,7 @@ void FAnimationViewportClient::OnToggleUseAudioAttenuation()
 	UWorld* World = PreviewScene->GetWorld();
 	if(World)
 	{
-		if(FAudioDevice* AudioDevice = GetWorld()->GetAudioDevice())
+		if(FAudioDevice* AudioDevice = GetWorld()->GetAudioDeviceRaw())
 		{
 			AudioDevice->SetUseAttenuationForNonGameWorlds(ConfigOption->bUseAudioAttenuation);
 		}
@@ -1885,7 +1885,7 @@ void FAnimationViewportClient::UpdateAudioListener(const FSceneView& View)
 
 	if (ViewportWorld)
 	{
-		if (FAudioDevice* AudioDevice = ViewportWorld->GetAudioDevice())
+		if (FAudioDevice* AudioDevice = ViewportWorld->GetAudioDeviceRaw())
 		{
 			const FVector& ViewLocation = GetViewLocation();
 			const FRotator& ViewRotation = GetViewRotation();
diff --git a/Engine/Source/Editor/UnrealEd/Private/EditorEngine.cpp b/Engine/Source/Editor/UnrealEd/Private/EditorEngine.cpp
--- a/Engine/Source/Editor/UnrealEd/Private/EditorEngine.cpp
+++ b/Engine/Source/Editor/UnrealEd/Private/EditorEngine.cpp
@@ -1901,13 +1901,10 @@ void UEditorEngine::Tick( float DeltaSeconds, bool bIdleMode )
 		if (bRequestEndPlayMapQueued)
 		{
 			// Shutdown all audio devices if we've requested end playmap now to avoid issues with GC running
-			TArray<FAudioDevice*>& AudioDevices = AudioDeviceManager->GetAudioDevices();
+			TArray<FAudioDevice*> AudioDevices = AudioDeviceManager->GetAudioDevices();
 			for (FAudioDevice* AudioDevice : AudioDevices)
 			{
-				if (AudioDevice)
-				{
-					AudioDevice->Flush(nullptr);
-				}
+				AudioDevice->Flush(nullptr);
 			}
 		}
 
@@ -2287,7 +2284,7 @@ UAudioComponent* UEditorEngine::GetPreviewAudioComponent()
 
 UAudioComponent* UEditorEngine::ResetPreviewAudioComponent( USoundBase* Sound, USoundNode* SoundNode )
 {
-	if (FAudioDevice* AudioDevice = GetMainAudioDevice())
+	if (FAudioDevice* AudioDevice = GetMainAudioDeviceRaw())
 	{
 		if (PreviewAudioComponent)
 		{
diff --git a/Engine/Source/Editor/UnrealEd/Private/EditorServer.cpp b/Engine/Source/Editor/UnrealEd/Private/EditorServer.cpp
--- a/Engine/Source/Editor/UnrealEd/Private/EditorServer.cpp
+++ b/Engine/Source/Editor/UnrealEd/Private/EditorServer.cpp
@@ -2058,7 +2058,7 @@ void UEditorEngine::EditorDestroyWorld( FWorldContext & Context, const FText& Cl
 	CloseEditedWorldAssets(ContextWorld);
 
 	// Stop all audio and remove references 
-	if (FAudioDevice* AudioDevice = ContextWorld->GetAudioDevice())
+	if (FAudioDevice* AudioDevice = ContextWorld->GetAudioDeviceRaw())
 	{
 		AudioDevice->Flush(ContextWorld);
 	}
diff --git a/Engine/Source/Editor/UnrealEd/Private/LevelEditorViewport.cpp b/Engine/Source/Editor/UnrealEd/Private/LevelEditorViewport.cpp
--- a/Engine/Source/Editor/UnrealEd/Private/LevelEditorViewport.cpp
+++ b/Engine/Source/Editor/UnrealEd/Private/LevelEditorViewport.cpp
@@ -4658,7 +4658,7 @@ void FLevelEditorViewportClient::UpdateAudioListener(const FSceneView& View)
 
 	if (ViewportWorld)
 	{
-		if (FAudioDevice* AudioDevice = ViewportWorld->GetAudioDevice())
+		if (FAudioDevice* AudioDevice = ViewportWorld->GetAudioDeviceRaw())
 		{
 			FVector ViewLocation = GetViewLocation();
 
diff --git a/Engine/Source/Editor/UnrealEd/Private/ObjectTools.cpp b/Engine/Source/Editor/UnrealEd/Private/ObjectTools.cpp
--- a/Engine/Source/Editor/UnrealEd/Private/ObjectTools.cpp
+++ b/Engine/Source/Editor/UnrealEd/Private/ObjectTools.cpp
@@ -2116,7 +2116,7 @@ namespace ObjectTools
 			const int32 NumAudioDevices = AudioDeviceManager->GetNumActiveAudioDevices();
 			for (int32 DeviceIndex = 0; DeviceIndex < NumAudioDevices; DeviceIndex++)
 			{
-				FAudioDevice* AudioDevice = AudioDeviceManager->GetAudioDevice(DeviceIndex);
+				FAudioDevice* AudioDevice = AudioDeviceManager->GetAudioDeviceRaw(DeviceIndex);
 				if (AudioDevice != nullptr)
 				{
 					AudioDevice->StopAllSounds();
diff --git a/Engine/Source/Editor/UnrealEd/Private/PlayLevel.cpp b/Engine/Source/Editor/UnrealEd/Private/PlayLevel.cpp
--- a/Engine/Source/Editor/UnrealEd/Private/PlayLevel.cpp
+++ b/Engine/Source/Editor/UnrealEd/Private/PlayLevel.cpp
@@ -724,7 +724,7 @@ void UEditorEngine::TeardownPlaySession(FWorldContext& PieWorldContext)
 	bool bWasSimulatingInEditor = PlayInEditorSessionInfo->OriginalRequestParams.WorldType == EPlaySessionWorldType::SimulateInEditor;
 	
 	// Stop all audio and remove references to temp level.
-	if (FAudioDevice* AudioDevice = PlayWorld->GetAudioDevice())
+	if (FAudioDevice* AudioDevice = PlayWorld->GetAudioDeviceRaw())
 	{
 		AudioDevice->Flush(PlayWorld);
 		AudioDevice->ResetInterpolation();
@@ -2296,7 +2296,7 @@ void UEditorEngine::ResetPIEAudioSetting(UWorld *CurrentPieWorld)
 	ULevelEditorPlaySettings* PlayInSettings = GetMutableDefault<ULevelEditorPlaySettings>();
 	if (!PlayInSettings->EnableGameSound)
 	{
-		if (FAudioDevice* AudioDevice = CurrentPieWorld->GetAudioDevice())
+		if (FAudioDevice* AudioDevice = CurrentPieWorld->GetAudioDeviceRaw())
 		{
 			AudioDevice->SetTransientMasterVolume(0.0f);
 		}
diff --git a/Engine/Source/Runtime/Engine/Classes/Engine/Engine.h b/Engine/Source/Runtime/Engine/Classes/Engine/Engine.h
--- a/Engine/Source/Runtime/Engine/Classes/Engine/Engine.h
+++ b/Engine/Source/Runtime/Engine/Classes/Engine/Engine.h
@@ -17,6 +17,7 @@
 #include "Subsystems/SubsystemCollection.h"
 #include "Subsystems/EngineSubsystem.h"
 #include "RHI.h"
+#include "AudioDeviceManager.h"
 #include "Engine.generated.h"
 
 #define WITH_DYNAMIC_RESOLUTION (!UE_SERVER)
@@ -24,7 +25,6 @@
 class AMatineeActor;
 class APlayerController;
 class Error;
-class FAudioDeviceManager;
 class FCanvas;
 class FCommonViewportClient;
 class FFineGrainedPerformanceTracker;
@@ -366,7 +366,7 @@ struct FWorldContext
 	bool	bWaitingOnOnlineSubsystem;
 
 	/** Handle to this world context's audio device.*/
-	uint32 AudioDeviceHandle;
+	uint32 AudioDeviceID;
 
 	/** Custom description to be display in blueprint debugger UI */
 	FString CustomDescription;
@@ -414,7 +414,7 @@ struct FWorldContext
 		, PIEWorldFeatureLevel(ERHIFeatureLevel::Num)
 		, RunAsDedicated(false)
 		, bWaitingOnOnlineSubsystem(false)
-		, AudioDeviceHandle(INDEX_NONE)
+		, AudioDeviceID(INDEX_NONE)
 		, ThisCurrentWorld(nullptr)
 	{ }
 
@@ -1772,7 +1772,7 @@ protected:
 	FAudioDeviceManager* AudioDeviceManager;
 
 	/** Audio device handle to the main audio device. */
-	uint32 MainAudioDeviceHandle;
+	FAudioDeviceHandle MainAudioDeviceHandle;
 
 public:
 
@@ -2439,13 +2439,14 @@ public:
 	FAudioDeviceManager* GetAudioDeviceManager();
 
 	/** @return the main audio device handle used by the engine. */
-	uint32 GetAudioDeviceHandle() const;
+	uint32 GetMainAudioDeviceID() const;
 
 	/** @return the main audio device. */
-	class FAudioDevice* GetMainAudioDevice();
+	FAudioDeviceHandle GetMainAudioDevice();
+	class FAudioDevice* GetMainAudioDeviceRaw();
 
 	/** @return the currently active audio device */
-	class FAudioDevice* GetActiveAudioDevice();
+	FAudioDeviceHandle GetActiveAudioDevice();
 
 	/** @return whether we're currently running in split screen (more than one local player) */
 	bool IsSplitScreen(UWorld *InWorld);
diff --git a/Engine/Source/Runtime/Engine/Classes/Engine/GameViewportClient.h b/Engine/Source/Runtime/Engine/Classes/Engine/GameViewportClient.h
--- a/Engine/Source/Runtime/Engine/Classes/Engine/GameViewportClient.h
+++ b/Engine/Source/Runtime/Engine/Classes/Engine/GameViewportClient.h
@@ -18,6 +18,7 @@
 #include "Engine/DebugDisplayProperty.h"
 #include "UObject/SoftObjectPath.h"
 #include "StereoRendering.h"
+#include "AudioDeviceManager.h"
 
 #include "GameViewportClient.generated.h"
 
@@ -100,6 +101,9 @@ protected:
 	/** If true will suppress the blue transition text messages. */
 	bool bSuppressTransitionMessage;
 
+	/** Strong handle to the audio device used by this viewport. */
+	FAudioDeviceHandle AudioDevice;
+
 public:
 
 	/** see enum EViewModeIndex */
@@ -982,9 +986,6 @@ private:
 	/** Mouse cursor locking behavior when the viewport is clicked */
 	EMouseLockMode MouseLockMode;
 
-	/** Handle to the audio device created for this viewport. Each viewport (for multiple PIE) will have its own audio device. */
-	uint32 AudioDeviceHandle;
-
 	/** Whether or not this audio device is in audio-focus */
 	bool bHasAudioFocus;
 
diff --git a/Engine/Source/Runtime/Engine/Private/AudioDevice.cpp b/Engine/Source/Runtime/Engine/Private/AudioDevice.cpp
--- a/Engine/Source/Runtime/Engine/Private/AudioDevice.cpp
+++ b/Engine/Source/Runtime/Engine/Private/AudioDevice.cpp
@@ -251,7 +251,7 @@ FAudioDevice::FAudioDevice()
 	: NumStoppingSources(32)
 	, SampleRate(0)
 	, NumPrecacheFrames(MONO_PCM_BUFFER_SAMPLES)
-	, DeviceHandle(static_cast<Audio::FDeviceId>(INDEX_NONE))
+	, DeviceID(static_cast<Audio::FDeviceId>(INDEX_NONE))
 	, SpatializationPluginInterface(nullptr)
 	, ReverbPluginInterface(nullptr)
 	, OcclusionInterface(nullptr)
@@ -1330,7 +1330,7 @@ bool FAudioDevice::HandleSoloCommand(const TCHAR* Cmd, FOutputDevice& Ar)
 	FAudioDeviceManager* DeviceManager = GEngine->GetAudioDeviceManager();
 	if (DeviceManager)
 	{
-		DeviceManager->SetSoloDevice(DeviceHandle);
+		DeviceManager->SetSoloDevice(DeviceID);
 	}
 	return true;
 }
@@ -4348,7 +4348,7 @@ void FAudioDevice::SendUpdateResultsToGameThread(const int32 FirstActiveIndex)
 {
 	DECLARE_CYCLE_STAT(TEXT("FGameThreadAudioTask.AudioSendResults"), STAT_AudioSendResults, STATGROUP_TaskGraphTasks);
 
-	const Audio::FDeviceId AudioDeviceID = DeviceHandle;
+	const Audio::FDeviceId AudioDeviceID = DeviceID;
 	UReverbEffect* ReverbEffect = Effects ? Effects->GetCurrentReverbEffect() : nullptr;
 	FAudioThread::RunCommandOnGameThread([AudioDeviceID, ReverbEffect]()
 	{
@@ -4357,7 +4357,7 @@ void FAudioDevice::SendUpdateResultsToGameThread(const int32 FirstActiveIndex)
 		{
 			if (FAudioDeviceManager* AudioDeviceManager = GEngine->GetAudioDeviceManager())
 			{
-				if (FAudioDevice* AudioDevice = AudioDeviceManager->GetAudioDevice(AudioDeviceID))
+				if (FAudioDeviceHandle AudioDevice = AudioDeviceManager->GetAudioDevice(AudioDeviceID))
 				{
 					AudioDevice->CurrentReverbEffect = ReverbEffect;
 				}
@@ -4577,10 +4577,10 @@ void FAudioDevice::AddNewActiveSoundInternal(const FActiveSound& NewActiveSound,
 		}
 	}
 
-	int32* PlayCount = Sound->CurrentPlayCount.Find(DeviceHandle);
+	int32* PlayCount = Sound->CurrentPlayCount.Find(DeviceID);
 	if (!PlayCount)
 	{
-		PlayCount = &Sound->CurrentPlayCount.Add(DeviceHandle);
+		PlayCount = &Sound->CurrentPlayCount.Add(DeviceID);
 	}
 	(*PlayCount)++;
 
@@ -5343,7 +5343,7 @@ FAudioDevice::FCreateComponentParams::FCreateComponentParams()
 	: World(nullptr)
 	, Actor(nullptr)
 {
-	AudioDevice = (GEngine ? GEngine->GetMainAudioDevice() : nullptr);
+	AudioDevice = (GEngine ? GEngine->GetMainAudioDeviceRaw() : nullptr);
 	CommonInit();
 }
 
@@ -5513,7 +5513,7 @@ UAudioComponent* FAudioDevice::CreateComponent(USoundBase* Sound, const FCreateC
 				}
 				else
 				{
-					AudioComponent->AudioDeviceHandle = Params.AudioDevice->DeviceHandle;
+					AudioComponent->AudioDeviceID = Params.AudioDevice->DeviceID;
 				}
 
 				if (Params.bPlay)
diff --git a/Engine/Source/Runtime/Engine/Private/GameViewportClient.cpp b/Engine/Source/Runtime/Engine/Private/GameViewportClient.cpp
--- a/Engine/Source/Runtime/Engine/Private/GameViewportClient.cpp
+++ b/Engine/Source/Runtime/Engine/Private/GameViewportClient.cpp
@@ -171,7 +171,6 @@ UGameViewportClient::UGameViewportClient(const FObjectInitializer& ObjectInitial
 	, MouseCaptureMode(EMouseCaptureMode::CapturePermanently)
 	, bHideCursorDuringCapture(false)
 	, MouseLockMode(EMouseLockMode::LockOnCapture)
-	, AudioDeviceHandle(INDEX_NONE)
 	, bHasAudioFocus(false)
 	, bIsMouseOverClient(false)
 #if WITH_EDITOR
@@ -256,7 +255,6 @@ UGameViewportClient::UGameViewportClient(FVTableHelper& Helper)
 	, MouseCaptureMode(EMouseCaptureMode::CapturePermanently)
 	, bHideCursorDuringCapture(false)
 	, MouseLockMode(EMouseLockMode::LockOnCapture)
-	, AudioDeviceHandle(INDEX_NONE)
 	, bHasAudioFocus(false)
 {
 
@@ -434,11 +432,14 @@ void UGameViewportClient::Init(struct FWorldContext& WorldContext, UGameInstance
 		FAudioDeviceManager* AudioDeviceManager = GEngine->GetAudioDeviceManager();
 		if (AudioDeviceManager)
 		{
-			FAudioDeviceManager::FCreateAudioDeviceResults NewDeviceResults;
-			if (AudioDeviceManager->CreateAudioDevice(bCreateNewAudioDevice, NewDeviceResults))
-			{
-				AudioDeviceHandle = NewDeviceResults.Handle;
+			// Get a new audio device for this world.
+			FAudioDeviceParams DeviceParams = AudioDeviceManager->GetDefaultParamsForNewWorld();
+			DeviceParams.AssociatedWorld = World;
+
+			AudioDevice = AudioDeviceManager->RequestAudioDevice(DeviceParams);
 
+			if (AudioDevice.IsValid())
+			{
 #if ENABLE_AUDIO_DEBUG
 				FAudioDebugger::ResolveDesiredStats(this);
 #endif // ENABLE_AUDIO_DEBUG
@@ -446,15 +447,15 @@ void UGameViewportClient::Init(struct FWorldContext& WorldContext, UGameInstance
 				// Set the base mix of the new device based on the world settings of the world
 				if (World)
 				{
-					NewDeviceResults.AudioDevice->SetDefaultBaseSoundMix(World->GetWorldSettings()->DefaultBaseSoundMix);
+					AudioDevice.GetAudioDevice()->SetDefaultBaseSoundMix(World->GetWorldSettings()->DefaultBaseSoundMix);
 
 					// Set the world's audio device handle to use so that sounds which play in that world will use the correct audio device
-					World->SetAudioDeviceHandle(AudioDeviceHandle);
+					World->SetAudioDevice(AudioDevice);
 				}
 
 				// Set this audio device handle on the world context so future world's set onto the world context
 				// will pass the audio device handle to them and audio will play on the correct audio device
-				WorldContext.AudioDeviceHandle = AudioDeviceHandle;
+				WorldContext.AudioDeviceID = AudioDevice.GetDeviceID();
 			}
 		}
 	}
@@ -1226,8 +1227,6 @@ void UGameViewportClient::Draw(FViewport* InViewport, FCanvas* SceneCanvas)
 	}
 
 	TMap<ULocalPlayer*,FSceneView*> PlayerViewMap;
-
-	FAudioDevice* AudioDevice = MyWorld->GetAudioDevice();
 	TArray<FSceneView*> Views;
 
 	for (FLocalPlayerIterator Iterator(GEngine, MyWorld); Iterator; ++Iterator)
@@ -1296,7 +1295,7 @@ void UGameViewportClient::Draw(FViewport* InViewport, FCanvas* SceneCanvas)
 						PlayerViewMap.Add(LocalPlayer, View);
 
 						// Update the listener.
-						if (AudioDevice != NULL && PlayerController != NULL)
+						if (AudioDevice && PlayerController != NULL)
 						{
 							bool bUpdateListenerPosition = true;
 
@@ -1309,8 +1308,8 @@ void UGameViewportClient::Draw(FViewport* InViewport, FCanvas* SceneCanvas)
 								// If there is more than one world referencing the main audio device
 								if (AudioDeviceManager->GetNumMainAudioDeviceWorlds() > 1)
 								{
-									uint32 MainAudioDeviceHandle = GEngine->GetAudioDeviceHandle();
-									if (AudioDevice->DeviceHandle == MainAudioDeviceHandle && !bHasAudioFocus)
+									Audio::FDeviceId MainAudioDeviceID = GEngine->GetMainAudioDeviceID();
+									if (AudioDevice->DeviceID == MainAudioDeviceID && !bHasAudioFocus)
 									{
 										bUpdateListenerPosition = false;
 									}
@@ -1817,7 +1816,7 @@ void UGameViewportClient::Precache()
 	if(!GIsEditor)
 	{
 		// Precache sounds...
-		if (FAudioDevice* AudioDevice = GetWorld()->GetAudioDevice())
+		if (AudioDevice)
 		{
 			UE_LOG(LogPlayerManagement, Log, TEXT("Precaching sounds..."));
 			for(TObjectIterator<USoundWave> It;It;++It)
@@ -1884,7 +1883,7 @@ void UGameViewportClient::ReceivedFocus(FViewport* InViewport)
 
 	if (GEngine && GEngine->GetAudioDeviceManager())
 	{
-		GEngine->GetAudioDeviceManager()->SetActiveDevice(AudioDeviceHandle);
+		GEngine->GetAudioDeviceManager()->SetActiveDevice(AudioDevice.GetDeviceID());
 		bHasAudioFocus = true;
 	}
 }
diff --git a/Engine/Source/Runtime/Engine/Private/Interpolation.cpp b/Engine/Source/Runtime/Engine/Private/Interpolation.cpp
--- a/Engine/Source/Runtime/Engine/Private/Interpolation.cpp
+++ b/Engine/Source/Runtime/Engine/Private/Interpolation.cpp
@@ -1271,7 +1271,7 @@ void AMatineeActor::StepInterp( float DeltaTime, bool bPreview )
 
 void AMatineeActor::DisableRadioFilterIfNeeded()
 {
-	FAudioDevice* AudioDevice = GEngine->GetMainAudioDevice();
+	FAudioDeviceHandle AudioDevice = GEngine->GetMainAudioDevice();
 	if( AudioDevice )
 	{
 		AudioDevice->EnableRadioEffect( !bDisableRadioFilter );
@@ -1299,7 +1299,7 @@ void AMatineeActor::EnableCinematicMode(bool bEnable)
 
 void AMatineeActor::EnableRadioFilter()
 {
-	FAudioDevice* AudioDevice = GEngine->GetMainAudioDevice();
+	FAudioDeviceHandle AudioDevice = GEngine->GetMainAudioDevice();
 	if( AudioDevice )
 	{
 		AudioDevice->EnableRadioEffect( true );
diff --git a/Engine/Source/Runtime/Engine/Private/SoundSubmix.cpp b/Engine/Source/Runtime/Engine/Private/SoundSubmix.cpp
--- a/Engine/Source/Runtime/Engine/Private/SoundSubmix.cpp
+++ b/Engine/Source/Runtime/Engine/Private/SoundSubmix.cpp
@@ -37,7 +37,7 @@ void USoundSubmix::StartRecordingOutput(const UObject* WorldContextObject, float
 
 	// Find device for this specific audio recording thing.
 	UWorld* ThisWorld = GEngine->GetWorldFromContextObject(WorldContextObject, EGetWorldErrorMode::LogAndReturnNull);
-	FAudioDevice* DesiredAudioDevice = ThisWorld->GetAudioDevice();
+	FAudioDevice* DesiredAudioDevice = ThisWorld->GetAudioDeviceRaw();
 
 	StartRecordingOutput(DesiredAudioDevice, ExpectedDuration);
 }
@@ -59,7 +59,7 @@ void USoundSubmix::StopRecordingOutput(const UObject* WorldContextObject, EAudio
 
 	// Find device for this specific audio recording thing.
 	UWorld* ThisWorld = GEngine->GetWorldFromContextObject(WorldContextObject, EGetWorldErrorMode::LogAndReturnNull);
-	FAudioDevice* DesiredAudioDevice = ThisWorld->GetAudioDevice();
+	FAudioDevice* DesiredAudioDevice = ThisWorld->GetAudioDeviceRaw();
 
 	StopRecordingOutput(DesiredAudioDevice, ExportType, Name, Path, ExistingSoundWaveToOverwrite);
 }
@@ -139,7 +139,7 @@ void USoundSubmix::StartEnvelopeFollowing(const UObject* WorldContextObject)
 
 	// Find device for this specific audio recording thing.
 	UWorld* ThisWorld = GEngine->GetWorldFromContextObject(WorldContextObject, EGetWorldErrorMode::LogAndReturnNull);
-	FAudioDevice* AudioDevice = ThisWorld->GetAudioDevice();
+	FAudioDevice* AudioDevice = ThisWorld->GetAudioDeviceRaw();
 
 	StartEnvelopeFollowing(AudioDevice);
 }
@@ -161,7 +161,7 @@ void USoundSubmix::StopEnvelopeFollowing(const UObject* WorldContextObject)
 
 	// Find device for this specific audio recording thing.
 	UWorld* ThisWorld = GEngine->GetWorldFromContextObject(WorldContextObject, EGetWorldErrorMode::LogAndReturnNull);
-	FAudioDevice* AudioDevice = ThisWorld->GetAudioDevice();
+	FAudioDevice* AudioDevice = ThisWorld->GetAudioDeviceRaw();
 
 	StopEnvelopeFollowing(AudioDevice);
 }
@@ -183,7 +183,7 @@ void USoundSubmix::AddEnvelopeFollowerDelegate(const UObject* WorldContextObject
 
 	// Find device for this specific audio recording thing.
 	UWorld* ThisWorld = GEngine->GetWorldFromContextObject(WorldContextObject, EGetWorldErrorMode::LogAndReturnNull);
-	FAudioDevice* AudioDevice = ThisWorld->GetAudioDevice();
+	FAudioDevice* AudioDevice = ThisWorld->GetAudioDeviceRaw();
 	if (AudioDevice)
 	{
 		AudioDevice->AddEnvelopeFollowerDelegate(this, OnSubmixEnvelopeBP);
@@ -198,7 +198,7 @@ void USoundSubmix::SetSubmixOutputVolume(const UObject* WorldContextObject, floa
 	}
 
 	UWorld* ThisWorld = GEngine->GetWorldFromContextObject(WorldContextObject, EGetWorldErrorMode::LogAndReturnNull);
-	FAudioDevice* AudioDevice = ThisWorld->GetAudioDevice();
+	FAudioDevice* AudioDevice = ThisWorld->GetAudioDeviceRaw();
 	if (AudioDevice)
 	{
 		AudioDevice->SetSubmixOutputVolume(this, InOutputVolume);
diff --git a/Engine/Source/Runtime/Engine/Private/UnrealEngine.cpp b/Engine/Source/Runtime/Engine/Private/UnrealEngine.cpp
--- a/Engine/Source/Runtime/Engine/Private/UnrealEngine.cpp
+++ b/Engine/Source/Runtime/Engine/Private/UnrealEngine.cpp
@@ -1082,11 +1082,13 @@ FWorldContext
 
 void FWorldContext::SetCurrentWorld(UWorld* World)
 {
-	if (World != nullptr)
+	if (World != nullptr && AudioDeviceID != INDEX_NONE)
 	{
 		// Set the world's audio device handle so that audio components playing in the 
 		// world will use the correct audio device instance.
-		World->SetAudioDeviceHandle(AudioDeviceHandle);
+		FAudioDeviceHandle AudioDevice = FAudioDeviceManager::Get()->GetAudioDevice(AudioDeviceID);
+
+		World->SetAudioDevice(AudioDevice);
 	}
 
 	for (int32 idx = 0; idx < ExternalReferences.Num(); ++idx)
@@ -2744,7 +2746,7 @@ void UEngine::Serialize(FArchive& Ar)
 	if (Ar.IsCountingMemory())
 	{
 		// Only use the main audio device when counting memory
-		if (FAudioDevice* AudioDevice = GetMainAudioDevice())
+		if (FAudioDeviceHandle AudioDevice = GetMainAudioDevice())
 		{
 			AudioDevice->CountBytes(Ar);
 		}
@@ -2851,12 +2853,12 @@ class FAudioDeviceManager* UEngine::GetAudioDeviceManager()
 	return AudioDeviceManager;
 }
 
-uint32 UEngine::GetAudioDeviceHandle() const
+uint32 UEngine::GetMainAudioDeviceID() const
 {
-	return MainAudioDeviceHandle;
+	return MainAudioDeviceHandle.GetDeviceID();
 }
 
-FAudioDevice* UEngine::GetMainAudioDevice()
+FAudioDeviceHandle UEngine::GetMainAudioDevice()
 {
 	if (AudioDeviceManager != nullptr)
 	{
@@ -2874,13 +2876,13 @@ FAudioDevice* UEngine::GetMainAudioDeviceRaw()
 	return nullptr;
 }
 
-FAudioDevice* UEngine::GetActiveAudioDevice()
+FAudioDeviceHandle UEngine::GetActiveAudioDevice()
 {
 	if (AudioDeviceManager != nullptr)
 	{
 		return AudioDeviceManager->GetActiveAudioDevice();
 	}
-	return nullptr;
+	return FAudioDeviceHandle();
 }
 
 void UEngine::InitializeAudioDeviceManager()
@@ -3844,11 +3846,11 @@ bool UEngine::Exec( UWorld* InWorld, const TCHAR* Cmd, FOutputDevice& Ar )
 	FAudioDevice* AudioDevice = nullptr;
 	if (InWorld)
 	{
-		AudioDevice = InWorld->GetAudioDevice();
+		AudioDevice = InWorld->GetAudioDeviceRaw();
 	}
 	else
 	{
-		AudioDevice = GetMainAudioDevice();
+		AudioDevice = GetMainAudioDeviceRaw();
 	}
 
 	if (AudioDevice && AudioDevice->Exec(InWorld, Cmd, Ar) == true)
@@ -12466,7 +12468,7 @@ bool UEngine::LoadMap( FWorldContext& WorldContext, FURL URL, class UPendingNetG
 		}
 
 		// Stop all audio to remove references to current level.
-		if (FAudioDevice* AudioDevice = WorldContext.World()->GetAudioDevice())
+		if (FAudioDevice* AudioDevice = WorldContext.World()->GetAudioDeviceRaw())
 		{
 			AudioDevice->Flush(WorldContext.World());
 			AudioDevice->SetTransientMasterVolume(1.0f);
@@ -12706,7 +12708,7 @@ bool UEngine::LoadMap( FWorldContext& WorldContext, FURL URL, class UPendingNetG
 
 	WorldContext.World()->SetGameMode(URL);
 
-	if (FAudioDevice* AudioDevice = WorldContext.World()->GetAudioDevice())
+	if (FAudioDevice* AudioDevice = WorldContext.World()->GetAudioDeviceRaw())
 	{
 		AudioDevice->SetDefaultBaseSoundMix(WorldContext.World()->GetWorldSettings()->DefaultBaseSoundMix);
 	}
@@ -13831,7 +13833,7 @@ bool UEngine::CommitMapChange( FWorldContext &Context )
 
 		// make sure any looping sounds, etc are stopped
 
-		if (FAudioDevice* AudioDevice = Context.World()->GetAudioDevice())
+		if (FAudioDevice* AudioDevice = Context.World()->GetAudioDeviceRaw())
 		{
 			AudioDevice->StopAllSounds();
 		}
diff --git a/Engine/Source/Runtime/Engine/Private/WorldSettings.cpp b/Engine/Source/Runtime/Engine/Private/WorldSettings.cpp
--- a/Engine/Source/Runtime/Engine/Private/WorldSettings.cpp
+++ b/Engine/Source/Runtime/Engine/Private/WorldSettings.cpp
@@ -194,7 +194,7 @@ void AWorldSettings::PostRegisterAllComponents()
 	Super::PostRegisterAllComponents();
 
 	UWorld* World = GetWorld();
-	if (FAudioDevice* AudioDevice = World->GetAudioDevice())
+	if (FAudioDeviceHandle AudioDevice = World->GetAudioDevice())
 	{
 		AudioDevice->SetDefaultAudioSettings(World, DefaultReverbSettings, DefaultAmbientZoneSettings);
 	}
@@ -622,7 +622,7 @@ void AWorldSettings::InternalPostPropertyChanged(FName PropertyName)
 	if (PropertyName == GET_MEMBER_NAME_CHECKED(AWorldSettings, DefaultReverbSettings) || PropertyName == GET_MEMBER_NAME_CHECKED(AWorldSettings, DefaultAmbientZoneSettings))
 {
 		UWorld* World = GetWorld();
-		if (FAudioDevice* AudioDevice = World->GetAudioDevice())
+		if (FAudioDeviceHandle AudioDevice = World->GetAudioDevice())
 	{
 			AudioDevice->SetDefaultAudioSettings(World, DefaultReverbSettings, DefaultAmbientZoneSettings);
 		}
