Adding an assert to CompressedAnimSequence::SerializeCompressedData to validate that the compressed animation data serialized correctly.

If it doesn't then the decompression code will eventually crash in a much less obvious way so this should make any future problem much easier to debug.

#rnx

diff --git a/Engine/Source/Runtime/Engine/Private/Animation/AnimCompressionTypes.cpp b/Engine/Source/Runtime/Engine/Private/Animation/AnimCompressionTypes.cpp
--- a/Engine/Source/Runtime/Engine/Private/Animation/AnimCompressionTypes.cpp
+++ b/Engine/Source/Runtime/Engine/Private/Animation/AnimCompressionTypes.cpp
@@ -459,6 +459,9 @@ void FCompressedAnimSequence::SerializeCompressedData(FArchive& Ar, bool bDDCDat
 
 			FOwnedBulkDataPtr* OwnedPtr = OptionalBulk.StealFileMapping();
 
+			// Decompression will crash later if the data failed to load so assert now to make it easier to debug in the future.
+			checkf(OwnedPtr->GetPointer() != nullptr || Size == 0, TEXT("Compressed animation data failed to load")); 
+
 #if WITH_EDITOR
 			check(!bUseMapping && !OwnedPtr->GetMappedHandle());
 			CompressedByteStream.Empty(Size);
