Some cook-time improvements mainly centered on material cooking brought over from UE4/Dev-Core:

Move MD5 computation on writer thread during package save

Fix bug in Async MD5 code.  Some packages won't be saved to disk  (eg: if nativization converts the blueprint contents so there's nothing left to store on disk) and thus won't produce a TFuture for the CookedHash in the FSavePackageResultStruct.  If this is the case, we shouldn't try and setup a continuation to update the CookedHash in the asset registry for that package.

Reduce cost of GetDependentFunctions on UMaterialFunction.  Specific changes are:

1) Introduce IterateDependentFunctions as a method to reduce temporary allocations when computing dependent functions (original author: PJ Kack, taken from shelved CL6968386)
2) Add transient UMaterialFunction::DependentFunctionExpressionCandidates field (computed in PostLoad and PostEditChangeProperty) wth only function expression types that we might have dependent functions on to reduce the amount of time we spend iterating on FunctionExpressions and casting them to see if they're a type we could have to get dependent functions from.

#[fyi] Patrick.Finegan

Improve cook times by avoiding re-computing the static parameter values twice during calls to FMaterial::CacheShaders (once in FMaterialResource::GetShaderMapId and again in FMaterialResource::GetStaticParameterSet) by allowing the MaterialInstance to store a cached parameter value set and re-use it in selected scopes (only applied to FMaterial::CacheShaders for now, but may be applicable elsewhere).

Testing:
1) QAGame produced same WindowsNoEditor cooked data (barring metadata that has some non-determinism issues) with/without this change.
2) InfiltratorDemo looked visually correct on Win64 with and without cooked data

Reduce iterations done on Material Functions in UMaterialInstance::GetStaticParameterValues by changing the underlying API to do lookups for multiple parameters at the same time.

Tested by a full cook with debug code that check()'d the parameter values and expression IDs output from the old code path was equal to the parameter values  and expression IDs output from the new code path.

Fix android compile error caused by my last submit.

#rnx

Fixing copyright notice to new (undated) format in a file I added previously.

#rnx

Reduce cost of FMaterialUpdateContext and RecacheMaterialInstanceUniformExpressions to save 4 minutes off of a large dataset cook time.  Changes:

1) FMaterialUpdateContext doesn't create a FGlobalComponentRecreateRenderStateContext instance unless FApp::CanEverRender() is true.
2) FGlobalComponentRecreateRenderStateContext stores FComponentRecreateRenderStateContext instances in a TArray instead of a TIndirectArray to avoid excessive numbers of individual allocations and frees.
3) FGlobalComponentRecreateRenderStateContext doesn't create FComponentRecreateRenderStateContext instances for components that aren't registered or don't have a render state created.  Analysis of the code shows that constructing those component contexts is pointlessly hitting the allocator (or bloating the TArray) because they will not do any useful work in their destructor.
4) RecacheMaterialInstanceUniformExpressions doesn't incur the cost of a TObjectIterator if FApp::CanEverRender() is false, as static analysis of the code shows no useful work done within the iteration if FApp::CanEverRender() returns false.

diff --git a/Engine/Source/Runtime/Engine/Classes/Materials/Material.h b/Engine/Source/Runtime/Engine/Classes/Materials/Material.h
--- a/Engine/Source/Runtime/Engine/Classes/Materials/Material.h
+++ b/Engine/Source/Runtime/Engine/Classes/Materials/Material.h
@@ -1039,8 +1039,8 @@ public:
 	ENGINE_API virtual FMaterialResource* AllocateResource();
 	ENGINE_API virtual FMaterialResource* GetMaterialResource(ERHIFeatureLevel::Type InFeatureLevel, EMaterialQualityLevel::Type QualityLevel = EMaterialQualityLevel::Num) override;
 	ENGINE_API virtual const FMaterialResource* GetMaterialResource(ERHIFeatureLevel::Type InFeatureLevel, EMaterialQualityLevel::Type QualityLevel = EMaterialQualityLevel::Num) const override;
-	ENGINE_API virtual bool GetStaticSwitchParameterValue(const FMaterialParameterInfo& ParameterInfo,bool& OutValue,FGuid& OutExpressionGuid, bool bOveriddenOnly = false, bool bCheckParent = true) const override;
-	ENGINE_API virtual bool GetStaticComponentMaskParameterValue(const FMaterialParameterInfo& ParameterInfo, bool& R, bool& G, bool& B, bool& A, FGuid& OutExpressionGuid, bool bOveriddenOnly = false, bool bCheckParent = true) const override;
+	ENGINE_API virtual bool GetStaticSwitchParameterValues(FStaticParamEvaluationContext& EvalContext, TBitArray<>& OutValues, FGuid* OutExpressionGuids, bool bCheckParent = true) const override;
+	ENGINE_API virtual bool GetStaticComponentMaskParameterValues(FStaticParamEvaluationContext& EvalContext, TBitArray<>& OutRGBAOrderedValues, FGuid* OutExpressionGuids, bool bCheckParent = true) const override;
 	ENGINE_API virtual bool GetTerrainLayerWeightParameterValue(const FMaterialParameterInfo& ParameterInfo, int32& OutWeightmapIndex, FGuid& OutExpressionGuid) const override;
 	ENGINE_API virtual bool GetMaterialLayersParameterValue(const FMaterialParameterInfo& ParameterInfo, FMaterialLayersFunctions& OutLayers, FGuid& OutExpressionGuid, bool bCheckParent = true) const override;
 	ENGINE_API virtual bool UpdateLightmassTextureTracking() override;
diff --git a/Engine/Source/Runtime/Engine/Classes/Materials/MaterialInstance.h b/Engine/Source/Runtime/Engine/Classes/Materials/MaterialInstance.h
--- a/Engine/Source/Runtime/Engine/Classes/Materials/MaterialInstance.h
+++ b/Engine/Source/Runtime/Engine/Classes/Materials/MaterialInstance.h
@@ -490,8 +490,8 @@ public:
 	virtual ENGINE_API void OverrideScalarParameterDefault(const FMaterialParameterInfo& ParameterInfo, float Value, bool bOverride, ERHIFeatureLevel::Type FeatureLevel) override;
 	virtual ENGINE_API bool CheckMaterialUsage(const EMaterialUsage Usage) override;
 	virtual ENGINE_API bool CheckMaterialUsage_Concurrent(const EMaterialUsage Usage) const override;
-	virtual ENGINE_API bool GetStaticSwitchParameterValue(const FMaterialParameterInfo& ParameterInfo, bool &OutValue, FGuid &OutExpressionGuid, bool bOveriddenOnly = false, bool bCheckParent = true) const override;
-	virtual ENGINE_API bool GetStaticComponentMaskParameterValue(const FMaterialParameterInfo& ParameterInfo, bool &R, bool &G, bool &B, bool &A, FGuid &OutExpressionGuid, bool bOveriddenOnly = false, bool bCheckParent = true) const override;
+	virtual ENGINE_API bool GetStaticSwitchParameterValues(FStaticParamEvaluationContext& EvalContext, TBitArray<>& OutValues, FGuid* OutExpressionGuids, bool bCheckParent = true) const override;
+	virtual ENGINE_API bool GetStaticComponentMaskParameterValues(FStaticParamEvaluationContext& EvalContext, TBitArray<>& OutRGBAOrderedValues, FGuid* OutExpressionGuids, bool bCheckParent = true) const override;
 	virtual ENGINE_API bool GetTerrainLayerWeightParameterValue(const FMaterialParameterInfo& ParameterInfo, int32& OutWeightmapIndex, FGuid &OutExpressionGuid) const override;
 	virtual ENGINE_API bool GetMaterialLayersParameterValue(const FMaterialParameterInfo& ParameterInfo, FMaterialLayersFunctions& OutLayers, FGuid &OutExpressionGuid, bool bCheckParent = true) const override;
 			ENGINE_API bool UpdateMaterialLayersParameterValue(const FMaterialParameterInfo& ParameterInfo, const FMaterialLayersFunctions& LayersValue, const bool bOverridden, const FGuid& GUID);
diff --git a/Engine/Source/Runtime/Engine/Classes/Materials/MaterialInterface.h b/Engine/Source/Runtime/Engine/Classes/Materials/MaterialInterface.h
--- a/Engine/Source/Runtime/Engine/Classes/Materials/MaterialInterface.h
+++ b/Engine/Source/Runtime/Engine/Classes/Materials/MaterialInterface.h
@@ -387,6 +387,31 @@ public:
 	 */
 	virtual const FMaterialResource* GetMaterialResource(ERHIFeatureLevel::Type InFeatureLevel, EMaterialQualityLevel::Type QualityLevel = EMaterialQualityLevel::Num) const { return NULL; }
 
+	struct FStaticParamEvaluationContext
+	{
+	public:
+		FStaticParamEvaluationContext(int32 InParameterNum, const FMaterialParameterInfo* InParameterInfos) : PendingParameterNum(InParameterNum), ParameterInfos(InParameterInfos) { PendingParameters.Add(true, PendingParameterNum); ResolvedByOverride.Add(false, PendingParameterNum); }
+
+		const FMaterialParameterInfo* GetParameterInfo(int32 ParamIndex) const { check(ParamIndex >= 0 && ParamIndex < PendingParameters.Num()); return ParameterInfos + ParamIndex; }
+		bool AllResolved() const { return PendingParameterNum == 0; }
+		bool IsResolved(int32 ParamIndex) const { return !PendingParameters[ParamIndex]; }
+		bool IsResolvedByOverride(int32 ParamIndex) const { return ResolvedByOverride[ParamIndex]; }
+		void MarkParameterResolved(int32 ParamIndex, bool bIsOverride);
+		int32 GetPendingParameterNum() const { return PendingParameterNum; }
+		int32 GetTotalParameterNum() const { return PendingParameters.Num(); }
+
+		/**
+		 * Perform an operation on the pending parameters until performed on all pending parameters, or until the operation asks to stop iterating by returning false.
+	 	 * @param Op - The operation to perform on each pending parameter.  Accepts the index for the parameter and the parameter info associated with it.  If the Op returns true, we keep iterating if there are more pending parameters, if it returns false, we stop immediately after this Op.
+		 */
+		void ForEachPendingParameter(TFunctionRef<bool(int32 ParamIndex, const FMaterialParameterInfo& ParamInfo)> Op);
+	private:
+		TBitArray<> PendingParameters;
+		TBitArray<> ResolvedByOverride;
+		int32 PendingParameterNum;
+		const FMaterialParameterInfo* ParameterInfos;
+	};
+
 	/**
 	* Get the value of the given static switch parameter
 	*
@@ -394,8 +419,18 @@ public:
 	* @param	OutValue		Will contain the value of the parameter if successful
 	* @return					True if successful
 	*/
-	virtual bool GetStaticSwitchParameterValue(const FMaterialParameterInfo& ParameterInfo,bool &OutValue,FGuid &OutExpressionGuid, bool bOveriddenOnly = false, bool bCheckParent = true) const
-		PURE_VIRTUAL(UMaterialInterface::GetStaticSwitchParameterValue,return false;);
+	ENGINE_API bool GetStaticSwitchParameterValue(const FMaterialParameterInfo& ParameterInfo, bool& OutValue, FGuid& OutExpressionGuid, bool bOveriddenOnly = false, bool bCheckParent = true) const;
+
+	/**
+	* Get the values of the given set of static switch parameters.
+	*
+	* @param	EvalContext			The evaluation context used while determining parameter values.
+	* @param	OutValues			If successful, will contain the value of the requested parameters.  Must be pre-sized to fit at least all parameter values in the evaluation.  If unsuccessful, may still be written to with intermediate values.
+	* @param	OutExpressionGuids	If successful, will contain the identifier of the owning expression of the requested parameter values.  Must be non null and pre-sized to fit at least all parameter values in the evaluation.  If unsuccessful, may still be written to with intermediate values.
+	* @return						True if successfully obtained ALL parameter values requested
+	*/
+	virtual bool GetStaticSwitchParameterValues(FStaticParamEvaluationContext& EvalContext, TBitArray<>& OutValues, FGuid* OutExpressionGuids, bool bCheckParent = true) const
+		PURE_VIRTUAL(UMaterialInterface::GetStaticSwitchParameterValues,return false;);
 
 	/**
 	* Get the value of the given static component mask parameter
@@ -404,8 +439,18 @@ public:
 	* @param	R, G, B, A		Will contain the values of the parameter if successful
 	* @return					True if successful
 	*/
-	virtual bool GetStaticComponentMaskParameterValue(const FMaterialParameterInfo& ParameterInfo, bool &R, bool &G, bool &B, bool &A, FGuid &OutExpressionGuid, bool bOveriddenOnly = false, bool bCheckParent = true) const
-		PURE_VIRTUAL(UMaterialInterface::GetStaticComponentMaskParameterValue,return false;);
+	ENGINE_API bool GetStaticComponentMaskParameterValue(const FMaterialParameterInfo& ParameterInfo, bool& R, bool& G, bool& B, bool& A, FGuid& OutExpressionGuid, bool bOveriddenOnly = false, bool bCheckParent = true) const;
+
+	/**
+	* Get the values of the given set of static component mask parameters
+	*
+	* @param	EvalContext				The evaluation context used while determining parameter values
+	* @param	OutRGBAOrderedValues	If successful, will contain the value of the requested parameters.  Must be pre-sized to fit at least four bits for all parameter values in the evaluation.  If unsuccessful, may still be written to with intermediate values.
+	* @param	OutExpressionGuids		If successful, will contain the identifier of the owning expression of the requested parameter values.  Must be non null and pre-sized to fit at least all parameter values in the evaluation.  If unsuccessful, may still be written to with intermediate values.
+	* @return							True if successfully obtained ALL parameter values requested
+	*/
+	virtual bool GetStaticComponentMaskParameterValues(FStaticParamEvaluationContext& EvalContext, TBitArray<>& OutRGBAOrderedValues, FGuid* OutExpressionGuids, bool bCheckParent = true) const
+		PURE_VIRTUAL(UMaterialInterface::GetStaticComponentMaskParameterValues,return false;);
 
 	/**
 	* Get the weightmap index of the given terrain layer weight parameter
diff --git a/Engine/Source/Runtime/Engine/Private/Components/ActorComponent.cpp b/Engine/Source/Runtime/Engine/Private/Components/ActorComponent.cpp
--- a/Engine/Source/Runtime/Engine/Private/Components/ActorComponent.cpp
+++ b/Engine/Source/Runtime/Engine/Private/Components/ActorComponent.cpp
@@ -178,7 +178,10 @@ FGlobalComponentRecreateRenderStateContext::FGlobalComponentRecreateRenderStateC
 	// recreate render state for all components.
 	for (UActorComponent* Component : TObjectRange<UActorComponent>())
 	{
-		ComponentContexts.Add(new FComponentRecreateRenderStateContext(Component, &ScenesToUpdateAllPrimitiveSceneInfos));
+		if (Component->IsRegistered() && Component->IsRenderStateCreated())
+		{
+			ComponentContexts.Emplace(Component, &ScenesToUpdateAllPrimitiveSceneInfos);
+		}
 	}
 
 	UpdateAllPrimitiveSceneInfos();
diff --git a/Engine/Source/Runtime/Engine/Private/Materials/Material.cpp b/Engine/Source/Runtime/Engine/Private/Materials/Material.cpp
--- a/Engine/Source/Runtime/Engine/Private/Materials/Material.cpp
+++ b/Engine/Source/Runtime/Engine/Private/Materials/Material.cpp
@@ -63,6 +63,7 @@
 #include "UObject/EditorObjectVersion.h"
 #include "UObject/ReleaseObjectVersion.h"
 #include "RenderCore/Public/RenderUtils.h"
+#include "Materials/MaterialStaticParameterValueResolver.h"
 
 #if WITH_EDITOR
 #include "Logging/TokenizedMessage.h"
@@ -2780,143 +2781,116 @@ bool UMaterial::GetFontParameterValue(const FMaterialParameterInfo& ParameterInf
 	return false;
 }
 
-
-bool UMaterial::GetStaticSwitchParameterValue(const FMaterialParameterInfo& ParameterInfo, bool& OutValue, FGuid& OutExpressionGuid, bool bOveriddenOnly, bool bCheckParent /*= true*/) const
+bool UMaterial::GetStaticSwitchParameterValues(FStaticParamEvaluationContext& EvalContext, TBitArray<>& OutValues, FGuid* OutExpressionGuids, bool bCheckParent /*= true*/) const
 {
-	// In the case of duplicate parameters with different values, this will return the
-	// first matching expression found, not necessarily the one that's used for rendering
-	UMaterialExpressionStaticBoolParameter* Parameter = nullptr;
+	if (EvalContext.AllResolved())
+	{
+		return true;
+	}
+
+	check(OutValues.Num() >= EvalContext.GetTotalParameterNum());
+	check(OutExpressionGuids);
 
 	for (UMaterialExpression* Expression : Expressions)
 	{
-		// Only need to check parameters that match in associated scope
-		if (ParameterInfo.Association == EMaterialParameterAssociation::GlobalParameter)
-		{
-			if (UMaterialExpressionStaticBoolParameter* ExpressionParameter = Cast<UMaterialExpressionStaticBoolParameter>(Expression))
-			{
-				if (ExpressionParameter->IsNamedParameter(ParameterInfo, OutValue, OutExpressionGuid))
-				{
-					return !bOveriddenOnly;
-				}
-			}
-			else if (UMaterialExpressionMaterialFunctionCall* FunctionCall = Cast<UMaterialExpressionMaterialFunctionCall>(Expression))
-			{
-				UMaterialFunctionInterface* Function = FunctionCall->MaterialFunction;
-				UMaterialFunctionInterface* ParameterOwner = nullptr;
-
-				if (Function && Function->OverrideNamedStaticSwitchParameter(ParameterInfo, OutValue, OutExpressionGuid))
-				{
-					return true;
-				}
+		TMaterialStaticParameterValueResolver<UMaterialExpressionStaticBoolParameter, TBitArray<>> ValueResolver(EvalContext, OutValues, OutExpressionGuids);
 
-				if (Function && Function->GetNamedParameterOfType(ParameterInfo, Parameter, &ParameterOwner))
+		EvalContext.ForEachPendingParameter([&EvalContext, Expression, &ValueResolver, &OutValues, &OutExpressionGuids](int32 ParamIndex, const FMaterialParameterInfo& ParameterInfo) -> bool
+			{
+				if (ParameterInfo.Association == EMaterialParameterAssociation::GlobalParameter)
 				{
-					if (ParameterOwner->OverrideNamedStaticSwitchParameter(ParameterInfo, OutValue, OutExpressionGuid))
+					if (UMaterialExpressionStaticBoolParameter* ExpressionParameter = Cast<UMaterialExpressionStaticBoolParameter>(Expression))
 					{
-						return true;
+						bool bTempVal;
+						if (ExpressionParameter->IsNamedParameter(ParameterInfo, bTempVal, OutExpressionGuids[ParamIndex]))
+						{
+							OutValues[ParamIndex] = bTempVal;
+							EvalContext.MarkParameterResolved(ParamIndex, true);
+						}
+					}
+					else if (UMaterialExpressionMaterialFunctionCall* FunctionCall = Cast<UMaterialExpressionMaterialFunctionCall>(Expression))
+					{
+						ValueResolver.AttemptResolve(ParamIndex, ParameterInfo, FunctionCall->MaterialFunction);
 					}
-					Parameter->IsNamedParameter(ParameterInfo, OutValue, OutExpressionGuid);
-					return !bOveriddenOnly;
-				}
-			}
-		}
-		else
-		{
-			if (UMaterialExpressionMaterialAttributeLayers* LayersExpression = Cast<UMaterialExpressionMaterialAttributeLayers>(Expression))
-			{
-				UMaterialFunctionInterface* Function = LayersExpression->GetParameterAssociatedFunction(ParameterInfo);
-				UMaterialFunctionInterface* ParameterOwner = nullptr;
-
-				if (Function && Function->OverrideNamedStaticSwitchParameter(ParameterInfo, OutValue, OutExpressionGuid))
-				{
-					return true;
 				}
-
-				if (Function && Function->GetNamedParameterOfType(ParameterInfo, Parameter, &ParameterOwner))
+				else
 				{
-					if (ParameterOwner->OverrideNamedStaticSwitchParameter(ParameterInfo, OutValue, OutExpressionGuid))
+					if (UMaterialExpressionMaterialAttributeLayers* LayersExpression = Cast<UMaterialExpressionMaterialAttributeLayers>(Expression))
 					{
-						return true;
+						ValueResolver.AttemptResolve(ParamIndex, ParameterInfo, LayersExpression->GetParameterAssociatedFunction(ParameterInfo));
 					}
-					Parameter->IsNamedParameter(ParameterInfo, OutValue, OutExpressionGuid);
-					return !bOveriddenOnly;
 				}
-			}
+				return !EvalContext.AllResolved();
+			});
+
+		ValueResolver.ResolveQueued();
+
+		if (EvalContext.AllResolved())
+		{
+			return true;
 		}
 	}
 
-	return false;
+	return EvalContext.AllResolved();
 }
 
-
-bool UMaterial::GetStaticComponentMaskParameterValue(const FMaterialParameterInfo& ParameterInfo, bool& OutR, bool& OutG, bool& OutB, bool& OutA, FGuid& OutExpressionGuid, bool bOveriddenOnly, bool bCheckParent /*= true*/) const
+bool UMaterial::GetStaticComponentMaskParameterValues(FStaticParamEvaluationContext& EvalContext, TBitArray<>& OutRGBAOrderedValues, FGuid* OutExpressionGuids, bool bCheckParent /*= true*/) const
 {
-	// In the case of duplicate parameters with different values, this will return the
-	// first matching expression found, not necessarily the one that's used for rendering
-	UMaterialExpressionStaticComponentMaskParameter* Parameter = nullptr;
+	if (EvalContext.AllResolved())
+	{
+		return true;
+	}
+
+	check(OutRGBAOrderedValues.Num() >= (EvalContext.GetTotalParameterNum()*4));
+	check(OutExpressionGuids);
 
 	for (UMaterialExpression* Expression : Expressions)
 	{
-		// Only need to check parameters that match in associated scope
-		if (ParameterInfo.Association == EMaterialParameterAssociation::GlobalParameter)
-		{
-			if (UMaterialExpressionStaticComponentMaskParameter* ExpressionParameter = Cast<UMaterialExpressionStaticComponentMaskParameter>(Expression))
-			{
-				if (ExpressionParameter->IsNamedParameter(ParameterInfo, OutR, OutG, OutB, OutA, OutExpressionGuid))
-				{
-					return true;
-				}
-			}
-			else if (UMaterialExpressionMaterialFunctionCall* FunctionCall = Cast<UMaterialExpressionMaterialFunctionCall>(Expression))
-			{
-				UMaterialFunctionInterface* Function = FunctionCall->MaterialFunction;
-				UMaterialFunctionInterface* ParameterOwner = nullptr;
-
-				if (Function && Function->OverrideNamedStaticComponentMaskParameter(ParameterInfo, OutR, OutG, OutB, OutA, OutExpressionGuid))
-				{
-					return true;
-				}
+		TMaterialStaticParameterValueResolver<UMaterialExpressionStaticComponentMaskParameter, TBitArray<>> ValueResolver(EvalContext, OutRGBAOrderedValues, OutExpressionGuids);
 
-				if (Function && Function->GetNamedParameterOfType(ParameterInfo, Parameter, &ParameterOwner))
+		EvalContext.ForEachPendingParameter([&EvalContext, Expression, &ValueResolver, &OutRGBAOrderedValues, &OutExpressionGuids](int32 ParamIndex, const FMaterialParameterInfo& ParameterInfo) -> bool
+			{
+				if (ParameterInfo.Association == EMaterialParameterAssociation::GlobalParameter)
 				{
-					if (ParameterOwner->OverrideNamedStaticComponentMaskParameter(ParameterInfo, OutR, OutG, OutB, OutA, OutExpressionGuid))
+					if (UMaterialExpressionStaticComponentMaskParameter* ExpressionParameter = Cast<UMaterialExpressionStaticComponentMaskParameter>(Expression))
 					{
-						return true;
+						bool bTempR, bTempG, bTempB, bTempA;
+						if (ExpressionParameter->IsNamedParameter(ParameterInfo, bTempR, bTempG, bTempB, bTempA, OutExpressionGuids[ParamIndex]))
+						{
+							int32 ParamRGBAIndex = ParamIndex * 4;
+							OutRGBAOrderedValues[ParamRGBAIndex + 0] = bTempR;
+							OutRGBAOrderedValues[ParamRGBAIndex + 1] = bTempG;
+							OutRGBAOrderedValues[ParamRGBAIndex + 2] = bTempB;
+							OutRGBAOrderedValues[ParamRGBAIndex + 3] = bTempA;
+							EvalContext.MarkParameterResolved(ParamIndex, true);
+						}
+					}
+					else if (UMaterialExpressionMaterialFunctionCall* FunctionCall = Cast<UMaterialExpressionMaterialFunctionCall>(Expression))
+					{
+						ValueResolver.AttemptResolve(ParamIndex, ParameterInfo, FunctionCall->MaterialFunction);
 					}
-					Parameter->IsNamedParameter(ParameterInfo, OutR, OutG, OutB, OutA, OutExpressionGuid);
-					return !bOveriddenOnly;
-				}
-			}
-		}
-		else
-		{
-			if (UMaterialExpressionMaterialAttributeLayers* LayersExpression = Cast<UMaterialExpressionMaterialAttributeLayers>(Expression))
-			{
-				UMaterialFunctionInterface* Function = LayersExpression->GetParameterAssociatedFunction(ParameterInfo);
-				UMaterialFunctionInterface* ParameterOwner = nullptr;
-
-				if (Function && Function->OverrideNamedStaticComponentMaskParameter(ParameterInfo, OutR, OutG, OutB, OutA, OutExpressionGuid))
-				{
-					return true;
 				}
-
-				if (Function && Function->GetNamedParameterOfType(ParameterInfo, Parameter, &ParameterOwner))
+				else
 				{
-					if (ParameterOwner->OverrideNamedStaticComponentMaskParameter(ParameterInfo, OutR, OutG, OutB, OutA, OutExpressionGuid))
+					if (UMaterialExpressionMaterialAttributeLayers* LayersExpression = Cast<UMaterialExpressionMaterialAttributeLayers>(Expression))
 					{
-						return true;
+						ValueResolver.AttemptResolve(ParamIndex, ParameterInfo, LayersExpression->GetParameterAssociatedFunction(ParameterInfo));
 					}
-					Parameter->IsNamedParameter(ParameterInfo, OutR, OutG, OutB, OutA, OutExpressionGuid);
-					return !bOveriddenOnly;
 				}
-			}
+				return !EvalContext.AllResolved();
+			});
+
+		ValueResolver.ResolveQueued();
+
+		if (EvalContext.AllResolved())
+		{
+			return true;
 		}
 	}
 
-	return false;
+	return EvalContext.AllResolved();
 }
 
-
 bool UMaterial::GetTerrainLayerWeightParameterValue(const FMaterialParameterInfo& ParameterInfo, int32& OutWeightmapIndex, FGuid& OutExpressionGuid) const
 {
 	bool bSuccess = false;
diff --git a/Engine/Source/Runtime/Engine/Private/Materials/MaterialInstance.cpp b/Engine/Source/Runtime/Engine/Private/Materials/MaterialInstance.cpp
--- a/Engine/Source/Runtime/Engine/Private/Materials/MaterialInstance.cpp
+++ b/Engine/Source/Runtime/Engine/Private/Materials/MaterialInstance.cpp
@@ -37,6 +37,7 @@
 #include "Curves/CurveLinearColorAtlas.h"
 #include "HAL/ThreadHeartBeat.h"
 #include "Misc/ScopedSlowTask.h"
+#include "Materials/MaterialStaticParameterValueResolver.h"
 
 DECLARE_CYCLE_STAT(TEXT("MaterialInstance CopyMatInstParams"), STAT_MaterialInstance_CopyMatInstParams, STATGROUP_Shaders);
 DECLARE_CYCLE_STAT(TEXT("MaterialInstance Serialize"), STAT_MaterialInstance_Serialize, STATGROUP_Shaders);
@@ -62,7 +63,7 @@ void CacheMaterialInstanceUniformExpressions(const UMaterialInstance* MaterialIn
  */
 void RecacheMaterialInstanceUniformExpressions(const UMaterialInterface* ParentMaterial, bool bRecreateUniformBuffer)
 {
-	if (GIsEditor)
+	if (GIsEditor && FApp::CanEverRender())
 	{
 		UE_LOG(LogMaterial,Verbose,TEXT("Recaching MI Uniform Expressions for parent %s"), *ParentMaterial->GetFullName());
 		TArray<FMICReentranceGuard> ReentranceGuards;
@@ -2024,17 +2025,19 @@ void UMaterialInstance::GetStaticParameterValues(FStaticParameterSet& OutStaticP
 		GetAllStaticSwitchParameterInfo(OutParameterInfo, Guids);
 		OutStaticParameters.StaticSwitchParameters.AddZeroed(OutParameterInfo.Num());
 
+		FStaticParamEvaluationContext SwitchEvalContext(OutParameterInfo.Num(), OutParameterInfo.GetData());
+		TBitArray<> StaticSwitchValues(false, OutParameterInfo.Num());
+		ensure(GetStaticSwitchParameterValues(SwitchEvalContext, StaticSwitchValues, Guids.GetData()));
 		for(int32 ParameterIdx=0; ParameterIdx<OutParameterInfo.Num(); ParameterIdx++)
 		{
 			FStaticSwitchParameter& ParentParameter = OutStaticParameters.StaticSwitchParameters[ParameterIdx];
 			FMaterialParameterInfo& ParameterInfo = OutParameterInfo[ParameterIdx];
-			FGuid ExpressionId = Guids[ParameterIdx];
 
 			ParentParameter.bOverride = false;
 			ParentParameter.ParameterInfo = ParameterInfo;
 
-			GetStaticSwitchParameterValue(ParameterInfo, ParentParameter.Value, ExpressionId);
-			ParentParameter.ExpressionGUID = ExpressionId;
+			ParentParameter.Value = StaticSwitchValues[ParameterIdx];
+			ParentParameter.ExpressionGUID = Guids[ParameterIdx];
 
 			// If the SourceInstance is overriding this parameter, use its settings
 			for(int32 SwitchParamIdx = 0; SwitchParamIdx < StaticParameters.StaticSwitchParameters.Num(); SwitchParamIdx++)
@@ -2055,17 +2058,24 @@ void UMaterialInstance::GetStaticParameterValues(FStaticParameterSet& OutStaticP
 		// Static Component Mask Parameters
 		GetAllStaticComponentMaskParameterInfo(OutParameterInfo, Guids);
 		OutStaticParameters.StaticComponentMaskParameters.AddZeroed(OutParameterInfo.Num());
+
+		FStaticParamEvaluationContext ComponentMaskEvalContext(OutParameterInfo.Num(), OutParameterInfo.GetData());
+		TBitArray<> StaticComponentMaskValues(false, OutParameterInfo.Num()*4);
+		ensure(GetStaticComponentMaskParameterValues(ComponentMaskEvalContext, StaticComponentMaskValues, Guids.GetData()));
 		for(int32 ParameterIdx=0; ParameterIdx<OutParameterInfo.Num(); ParameterIdx++)
 		{
 			FStaticComponentMaskParameter& ParentParameter = OutStaticParameters.StaticComponentMaskParameters[ParameterIdx];
 			FMaterialParameterInfo& ParameterInfo = OutParameterInfo[ParameterIdx];
-			FGuid ExpressionId = Guids[ParameterIdx];
 
 			ParentParameter.bOverride = false;
 			ParentParameter.ParameterInfo = ParameterInfo;
 			
-			GetStaticComponentMaskParameterValue(ParameterInfo, ParentParameter.R, ParentParameter.G, ParentParameter.B, ParentParameter.A, ExpressionId);
-			ParentParameter.ExpressionGUID = ExpressionId;
+			int32 ParamRGBAIndex = ParameterIdx * 4;
+			ParentParameter.R = StaticComponentMaskValues[ParamRGBAIndex + 0];
+			ParentParameter.G = StaticComponentMaskValues[ParamRGBAIndex + 1];
+			ParentParameter.B = StaticComponentMaskValues[ParamRGBAIndex + 2];
+			ParentParameter.A = StaticComponentMaskValues[ParamRGBAIndex + 3];
+			ParentParameter.ExpressionGUID = Guids[ParameterIdx];
 			
 			// If the SourceInstance is overriding this parameter, use its settings
 			for(int32 MaskParamIdx = 0; MaskParamIdx < StaticParameters.StaticComponentMaskParameters.Num(); MaskParamIdx++)
@@ -3083,124 +3093,164 @@ void UMaterialInstance::CacheShadersForResources(EShaderPlatform ShaderPlatform,
 	}
 }
 
-bool UMaterialInstance::GetStaticSwitchParameterValue(const FMaterialParameterInfo& ParameterInfo, bool &OutValue,FGuid &OutExpressionGuid, bool bOveriddenOnly, bool bCheckParent /*= true*/) const
+bool UMaterialInstance::GetStaticSwitchParameterValues(FStaticParamEvaluationContext& EvalContext, TBitArray<>& OutValues, FGuid* OutExpressionGuids, bool bCheckParent /*= true*/) const
 {
+	if (EvalContext.AllResolved())
+	{
+		return true;
+	}
+
 	if (GetReentrantFlag())
 	{
 		return false;
 	}
 
+	check(OutValues.Num() >= EvalContext.GetTotalParameterNum());
+	check(OutExpressionGuids);
+
 	// Instance override
-	for (int32 ValueIndex = 0;ValueIndex < StaticParameters.StaticSwitchParameters.Num();ValueIndex++)
+	for (int32 ValueIndex = 0; ValueIndex < StaticParameters.StaticSwitchParameters.Num(); ValueIndex++)
 	{
 		const FStaticSwitchParameter& Param = StaticParameters.StaticSwitchParameters[ValueIndex];
-		if (Param.bOverride && Param.ParameterInfo == ParameterInfo)
+		if (Param.bOverride)
 		{
-			OutValue = Param.Value;
-			OutExpressionGuid = Param.ExpressionGUID;
-			return true;
+			EvalContext.ForEachPendingParameter([&EvalContext, &Param, &OutValues, &OutExpressionGuids](int32 ParamIndex, const FMaterialParameterInfo& ParameterInfo) -> bool
+				{
+					if (Param.ParameterInfo == ParameterInfo)
+					{
+						OutValues[ParamIndex] = Param.Value;
+						OutExpressionGuids[ParamIndex] = Param.ExpressionGUID;
+						EvalContext.MarkParameterResolved(ParamIndex, true);
+						return !EvalContext.AllResolved();
+					}
+					return true;
+				});
+
+			if (EvalContext.AllResolved())
+			{
+				return true;
+			}
 		}
 	}
-	
+
 	// Instance-included default
-	if (ParameterInfo.Association != EMaterialParameterAssociation::GlobalParameter)
+	UMaterialExpressionStaticBoolParameter* Parameter = nullptr;
+	for (const FStaticMaterialLayersParameter& LayersParam : StaticParameters.MaterialLayersParameters)
 	{
-		UMaterialExpressionStaticBoolParameter* Parameter = nullptr;
-		for (const FStaticMaterialLayersParameter& LayersParam : StaticParameters.MaterialLayersParameters)
+		if (LayersParam.bOverride)
 		{
-			if (LayersParam.bOverride)
-			{
-				UMaterialFunctionInterface* Function = LayersParam.GetParameterAssociatedFunction(ParameterInfo);
-				UMaterialFunctionInterface* ParameterOwner = nullptr;
-
-				if (Function && Function->OverrideNamedStaticSwitchParameter(ParameterInfo, OutValue, OutExpressionGuid))
-				{
-					return true;
-				}
+			TMaterialStaticParameterValueResolver<UMaterialExpressionStaticBoolParameter, TBitArray<>> ValueResolver(EvalContext, OutValues, OutExpressionGuids);
 
-				if (Function && Function->GetNamedParameterOfType(ParameterInfo, Parameter, &ParameterOwner))
+			EvalContext.ForEachPendingParameter([&EvalContext, &LayersParam, &ValueResolver, &OutValues, &OutExpressionGuids](int32 ParamIndex, const FMaterialParameterInfo& ParameterInfo) -> bool
 				{
-					if (ParameterOwner->OverrideNamedStaticSwitchParameter(ParameterInfo, OutValue, OutExpressionGuid))
+					if (ParameterInfo.Association == EMaterialParameterAssociation::GlobalParameter)
 					{
 						return true;
-						
 					}
-					Parameter->IsNamedParameter(ParameterInfo, OutValue, OutExpressionGuid);
-					return !bOveriddenOnly;
-				}
+
+					ValueResolver.AttemptResolve(ParamIndex, ParameterInfo, LayersParam.GetParameterAssociatedFunction(ParameterInfo));
+					return !EvalContext.AllResolved();
+				});
+
+			ValueResolver.ResolveQueued();
+
+			if (EvalContext.AllResolved())
+			{
+				return true;
 			}
 		}
 	}
-	
+
 	// Next material in hierarchy
 	if (Parent && bCheckParent)
 	{
 		FMICReentranceGuard	Guard(this);
-		return Parent->GetStaticSwitchParameterValue(ParameterInfo, OutValue, OutExpressionGuid, bOveriddenOnly);
+		return Parent->GetStaticSwitchParameterValues(EvalContext, OutValues, OutExpressionGuids);
 	}
-	
-	return false;
+
+	return EvalContext.AllResolved();
 }
 
-bool UMaterialInstance::GetStaticComponentMaskParameterValue(const FMaterialParameterInfo& ParameterInfo, bool &OutR, bool &OutG, bool &OutB, bool &OutA, FGuid &OutExpressionGuid, bool bOveriddenOnly, bool bCheckParent /*= true*/) const
+bool UMaterialInstance::GetStaticComponentMaskParameterValues(FStaticParamEvaluationContext& EvalContext, TBitArray<>& OutRGBAOrderedValues, FGuid* OutExpressionGuids, bool bCheckParent /*= true*/) const
 {
+	if (EvalContext.AllResolved())
+	{
+		return true;
+	}
+
 	if (GetReentrantFlag())
 	{
 		return false;
 	}
 
+	check(OutRGBAOrderedValues.Num() >= (EvalContext.GetTotalParameterNum()*4));
+	check(OutExpressionGuids);
+
 	// Instance override
-	for (int32 ValueIndex = 0;ValueIndex < StaticParameters.StaticComponentMaskParameters.Num();ValueIndex++)
+	for (int32 ValueIndex = 0; ValueIndex < StaticParameters.StaticComponentMaskParameters.Num(); ValueIndex++)
 	{
 		const FStaticComponentMaskParameter& Param = StaticParameters.StaticComponentMaskParameters[ValueIndex];
-		if (Param.bOverride && Param.ParameterInfo == ParameterInfo)
+		if (Param.bOverride)
 		{
-			OutR = Param.R;
-			OutG = Param.G;
-			OutB = Param.B;
-			OutA = Param.A;
-			OutExpressionGuid = Param.ExpressionGUID;
-			return true;
+			EvalContext.ForEachPendingParameter([&EvalContext, &Param, &OutRGBAOrderedValues, &OutExpressionGuids](int32 ParamIndex, const FMaterialParameterInfo& ParameterInfo) -> bool
+				{
+					if (Param.ParameterInfo == ParameterInfo)
+					{
+						int32 ParamRGBAIndex = ParamIndex * 4;
+						OutRGBAOrderedValues[ParamRGBAIndex + 0] = Param.R;
+						OutRGBAOrderedValues[ParamRGBAIndex + 1] = Param.G;
+						OutRGBAOrderedValues[ParamRGBAIndex + 2] = Param.B;
+						OutRGBAOrderedValues[ParamRGBAIndex + 3] = Param.A;
+						OutExpressionGuids[ParamIndex] = Param.ExpressionGUID;
+						EvalContext.MarkParameterResolved(ParamIndex, true);
+						return !EvalContext.AllResolved();
+					}
+					return true;
+				});
+
+			if (EvalContext.AllResolved())
+			{
+				return true;
+			}
 		}
 	}
-	
+
 	// Instance-included default
-	if (ParameterInfo.Association != EMaterialParameterAssociation::GlobalParameter)
+	UMaterialExpressionStaticBoolParameter* Parameter = nullptr;
+	for (const FStaticMaterialLayersParameter& LayersParam : StaticParameters.MaterialLayersParameters)
 	{
-		UMaterialExpressionStaticComponentMaskParameter* Parameter = nullptr;
-		for (const FStaticMaterialLayersParameter& LayersParam : StaticParameters.MaterialLayersParameters)
+		if (LayersParam.bOverride)
 		{
-			if (LayersParam.bOverride)
-			{
-				UMaterialFunctionInterface* Function = LayersParam.GetParameterAssociatedFunction(ParameterInfo);
-				UMaterialFunctionInterface* ParameterOwner = nullptr;
+			TMaterialStaticParameterValueResolver<UMaterialExpressionStaticComponentMaskParameter, TBitArray<>> ValueResolver(EvalContext, OutRGBAOrderedValues, OutExpressionGuids);
 
-				if (Function && Function->OverrideNamedStaticComponentMaskParameter(ParameterInfo, OutR, OutG, OutB, OutA, OutExpressionGuid))
-				{
-					return true;
-				}
-
-				if (Function && Function->GetNamedParameterOfType(ParameterInfo, Parameter, &ParameterOwner))
+			EvalContext.ForEachPendingParameter([&EvalContext, &LayersParam, &ValueResolver, &OutRGBAOrderedValues, &OutExpressionGuids](int32 ParamIndex, const FMaterialParameterInfo& ParameterInfo) -> bool
 				{
-					if (ParameterOwner->OverrideNamedStaticComponentMaskParameter(ParameterInfo, OutR, OutG, OutB, OutA, OutExpressionGuid))
+					if (ParameterInfo.Association == EMaterialParameterAssociation::GlobalParameter)
 					{
 						return true;
 					}
-					Parameter->IsNamedParameter(ParameterInfo, OutR, OutG, OutB, OutA, OutExpressionGuid);
-					return !bOveriddenOnly;
-				}
+
+					ValueResolver.AttemptResolve(ParamIndex, ParameterInfo, LayersParam.GetParameterAssociatedFunction(ParameterInfo));
+					return !EvalContext.AllResolved();
+				});
+
+			ValueResolver.ResolveQueued();
+
+			if (EvalContext.AllResolved())
+			{
+				return true;
 			}
 		}
 	}
-	
+
 	// Next material in hierarchy
 	if (Parent && bCheckParent)
 	{
 		FMICReentranceGuard	Guard(this);
-		return Parent->GetStaticComponentMaskParameterValue(ParameterInfo, OutR, OutG, OutB, OutA, OutExpressionGuid, bOveriddenOnly);
+		return Parent->GetStaticComponentMaskParameterValues(EvalContext, OutRGBAOrderedValues, OutExpressionGuids);
 	}
-	
-	return false;
+
+	return EvalContext.AllResolved();
 }
 
 bool UMaterialInstance::GetTerrainLayerWeightParameterValue(const FMaterialParameterInfo& ParameterInfo, int32& OutWeightmapIndex, FGuid &OutExpressionGuid) const
diff --git a/Engine/Source/Runtime/Engine/Private/Materials/MaterialInterface.cpp b/Engine/Source/Runtime/Engine/Private/Materials/MaterialInterface.cpp
--- a/Engine/Source/Runtime/Engine/Private/Materials/MaterialInterface.cpp
+++ b/Engine/Source/Runtime/Engine/Private/Materials/MaterialInterface.cpp
@@ -91,6 +91,47 @@ void UMaterialInterface::GetUsedTexturesAndIndices(TArray<UTexture*>& OutTexture
 	OutIndices.AddDefaulted(OutTextures.Num());
 }
 
+bool UMaterialInterface::GetStaticSwitchParameterValue(const FMaterialParameterInfo& ParameterInfo, bool& OutValue, FGuid& OutExpressionGuid, bool bOveriddenOnly /*= false*/, bool bCheckParent /*= true*/) const
+{
+	TBitArray<> Output(false, 1); // Relying on the default allocator to be inline to avoid allocation here.
+	FStaticParamEvaluationContext EvalContext(1, &ParameterInfo);
+	if (!GetStaticSwitchParameterValues(EvalContext, Output, &OutExpressionGuid, bCheckParent))
+	{
+		return false;
+	}
+
+	if (bOveriddenOnly && !EvalContext.IsResolvedByOverride(0))
+	{
+		return false;
+	}
+
+	OutValue = Output[0];
+
+	return true;
+}
+
+bool UMaterialInterface::GetStaticComponentMaskParameterValue(const FMaterialParameterInfo& ParameterInfo, bool& R, bool& G, bool& B, bool& A, FGuid& OutExpressionGuid, bool bOveriddenOnly /*= false*/, bool bCheckParent /*= true*/) const
+{
+	TBitArray<> Output(false, 4); // Relying on the default allocator to be inline to avoid allocation here.
+	FStaticParamEvaluationContext EvalContext(1, &ParameterInfo);
+	if (!GetStaticComponentMaskParameterValues(EvalContext, Output, &OutExpressionGuid, bCheckParent))
+	{
+		return false;
+	}
+	
+	if (bOveriddenOnly && !EvalContext.IsResolvedByOverride(0))
+	{
+		return false;
+	}
+
+	R = Output[0];
+	G = Output[1];
+	B = Output[2];
+	A = Output[3];
+
+	return true;
+}
+
 FMaterialRelevance UMaterialInterface::GetRelevance_Internal(const UMaterial* Material, ERHIFeatureLevel::Type InFeatureLevel) const
 {
 	if(Material)
@@ -747,3 +788,24 @@ void UMaterialInterface::RemoveUserDataOfClass(TSubclassOf<UAssetUserData> InUse
 		}
 	}
 }
+
+void UMaterialInterface::FStaticParamEvaluationContext::MarkParameterResolved(int32 ParamIndex, bool bIsOverride)
+{
+	FBitReference BitRef = PendingParameters[ParamIndex];
+	check(BitRef);
+	BitRef = false;
+	ResolvedByOverride[ParamIndex] = bIsOverride;
+	--PendingParameterNum;
+}
+
+void UMaterialInterface::FStaticParamEvaluationContext::ForEachPendingParameter(TFunctionRef<bool(int32, const FMaterialParameterInfo&)> Op)
+{
+	for (TConstSetBitIterator<> It(PendingParameters); It; ++It)
+	{
+		int32 ParamIndex = It.GetIndex();
+		if (!Op(ParamIndex, ParameterInfos[ParamIndex]))
+		{
+			break;
+		}
+	}
+}
diff --git a/Engine/Source/Runtime/Engine/Private/Materials/MaterialShared.cpp b/Engine/Source/Runtime/Engine/Private/Materials/MaterialShared.cpp
--- a/Engine/Source/Runtime/Engine/Private/Materials/MaterialShared.cpp
+++ b/Engine/Source/Runtime/Engine/Private/Materials/MaterialShared.cpp
@@ -3145,7 +3145,7 @@ void FMaterial::SaveShaderStableKeys(EShaderPlatform TargetShaderPlatform, FStab
 FMaterialUpdateContext::FMaterialUpdateContext(uint32 Options, EShaderPlatform InShaderPlatform)
 {
 	bool bReregisterComponents = (Options & EOptions::ReregisterComponents) != 0;
-	bool bRecreateRenderStates = (Options & EOptions::RecreateRenderStates) != 0;
+	bool bRecreateRenderStates = ((Options & EOptions::RecreateRenderStates) != 0) && FApp::CanEverRender();
 
 	bSyncWithRenderingThread = (Options & EOptions::SyncWithRenderingThread) != 0;
 	if (bReregisterComponents)
diff --git a/Engine/Source/Runtime/Engine/Private/Materials/MaterialStaticParameterValueResolver.h b/Engine/Source/Runtime/Engine/Private/Materials/MaterialStaticParameterValueResolver.h
new file mode 100644
--- /dev/null
+++ b/Engine/Source/Runtime/Engine/Private/Materials/MaterialStaticParameterValueResolver.h
@@ -0,0 +1,187 @@
+// Copyright Epic Games, Inc. All Rights Reserved.
+
+#pragma once
+
+#include "CoreMinimal.h"
+#include "Materials/MaterialFunctionInterface.h"
+#include "Materials/MaterialInterface.h"
+#include "Materials/MaterialExpressionStaticBoolParameter.h"
+#include "Materials/MaterialExpressionStaticComponentMaskParameter.h"
+
+template <typename ParameterType, typename OutputType>
+class TMaterialStaticParameterValueResolver
+{
+public:
+	TMaterialStaticParameterValueResolver(UMaterialInterface::FStaticParamEvaluationContext& InEvalContext, OutputType& InValues, FGuid* InExpressionGuids)
+	: EvalContext(InEvalContext)
+	, OutValues(InValues)
+	, OutExpressionGuids(InExpressionGuids)
+	{
+	}
+
+	/**
+	 * Tries to resolve a single parameter from the set of parameters in the originally provided evaluation context by obtaining it from a Function.  The parameter may be resolved immediately or queued for resolving later.
+	 * @param ParamIndex The index of the parameter to try to resolve.
+	 * @param ParameterInfo The info associated with the parameter to try to resolve.
+	 * @param Function The function from which to try and obtain the value for the parameter.
+	 */
+	void AttemptResolve(int32 ParamIndex, const FMaterialParameterInfo& ParameterInfo, UMaterialFunctionInterface* Function)
+	{
+		if (!Function)
+		{
+			return;
+		}
+
+		if (ExtractValueFromFunction(ParamIndex, ParameterInfo, Function, nullptr))
+		{
+			EvalContext.MarkParameterResolved(ParamIndex, true);
+			return;
+		}
+
+		FunctionMap.FindOrAdd(Function).Add(ParamIndex);
+	}
+
+	/**
+	 * Resolves all parameters that we have previously attempted to resolve but weren't able to immedately resolve.
+	 */
+	void ResolveQueued()
+	{
+		for (auto& FunctionMapPair : FunctionMap)
+		{
+			UMaterialFunctionInterface* Function = FunctionMapPair.Key;
+			TArray<int32, TInlineAllocator<16>>& ParameterIndicesForFunction = FunctionMapPair.Value;
+			int32 PendingParametersForFunction = ParameterIndicesForFunction.Num();
+
+			if (PendingParametersForFunction == 0)
+				continue;
+
+			if (UMaterialFunctionInterface* ParameterFunction = Function->GetBaseFunction())
+			{
+				struct FParameterAndOwner
+				{
+					ParameterType* Parameter;
+					UMaterialFunctionInterface* Owner;
+				};
+
+				TArray<FParameterAndOwner, TInlineAllocator<16>> ParametersAndOwners;
+				ParametersAndOwners.AddZeroed(PendingParametersForFunction);
+
+				const UClass* TargetClass = ParameterType::StaticClass();
+
+				auto GetExpressionParameterByNamePredicate =
+					[this, &PendingParametersForFunction, &ParameterIndicesForFunction, &ParametersAndOwners, TargetClass](UMaterialFunctionInterface* InFunction) -> bool
+				{
+					for (UMaterialExpression* FunctionExpression : *InFunction->GetFunctionExpressions())
+					{
+						if (ParameterType* ExpressionParameter = (FunctionExpression && FunctionExpression->IsA(TargetClass)) ? (ParameterType*)FunctionExpression : nullptr)
+						{
+							for (int32 FunctionParameterIdx = 0; FunctionParameterIdx < ParameterIndicesForFunction.Num(); ++FunctionParameterIdx)
+							{
+								if (ParametersAndOwners[FunctionParameterIdx].Owner)
+									continue;
+
+								if (ExpressionParameter->ParameterName == EvalContext.GetParameterInfo(ParameterIndicesForFunction[FunctionParameterIdx])->Name)
+								{
+									ParametersAndOwners[FunctionParameterIdx] = {ExpressionParameter, InFunction};
+
+									if (--PendingParametersForFunction == 0)
+									{
+										return false;
+									}
+								}
+							}
+						}
+					}
+
+					return true; // not found, continue iterating
+				};
+
+				if (ParameterFunction->IterateDependentFunctions(GetExpressionParameterByNamePredicate))
+				{
+					GetExpressionParameterByNamePredicate(ParameterFunction);
+				}
+
+				for (int32 FunctionParameterIdx = 0; FunctionParameterIdx < ParameterIndicesForFunction.Num(); ++FunctionParameterIdx)
+				{
+					const FParameterAndOwner& ParameterAndOwner = ParametersAndOwners[FunctionParameterIdx];
+					if (!ParameterAndOwner.Owner)
+						continue;
+
+					int32 ParamIndex = ParameterIndicesForFunction[FunctionParameterIdx];
+					const FMaterialParameterInfo* ParameterInfo = EvalContext.GetParameterInfo(ParamIndex);
+					bool bIsOverride = ExtractValueFromFunction(ParamIndex, *ParameterInfo, ParameterAndOwner.Owner, ParameterAndOwner.Parameter);
+					EvalContext.MarkParameterResolved(ParamIndex, bIsOverride);
+				}
+			}
+		}
+	}
+
+
+private:
+	/**
+	 * Obtains a parameter value for a single parameter from the given Function.  Also accepts a fallback parameter for the default value to be used if the Function doesn't have any overrides for the parameter.
+	 * @param ParamIndex The index of the parameter to try to resolve.
+	 * @param ParameterInfo The info associated with the parameter to try to resolve.
+	 * @param Function The function from which to try and obtain the value for the parameter.
+	 * @param Parameter If non-null, the parameter from which to try and obtain the default value if no overrides exist for the parameter on the function.
+	 * @return true if a *non-default* or override value was obtained for the parameter, if no value was obtained, or if we fell back to obtaining the default value, then false will be returned.
+	 */
+	bool ExtractValueFromFunction(int32 ParamIndex, const FMaterialParameterInfo& ParameterInfo, UMaterialFunctionInterface* Function, ParameterType* Parameter)
+	{
+		static_assert(TIsSame<ParameterType,void>::Value, "Must provide explicit specialization for value retrieval method when using TMaterialStaticParameterValueResolver");
+		return false;
+	}
+
+	UMaterialInterface::FStaticParamEvaluationContext& EvalContext;
+	OutputType& OutValues;
+	FGuid* OutExpressionGuids;
+	TSortedMap<UMaterialFunctionInterface*, TArray<int32, TInlineAllocator<16>>, TInlineAllocator<16>> FunctionMap;
+};
+
+
+template<>
+inline bool TMaterialStaticParameterValueResolver<UMaterialExpressionStaticBoolParameter, TBitArray<>>::ExtractValueFromFunction(int32 ParamIndex, const FMaterialParameterInfo& ParameterInfo, UMaterialFunctionInterface* Function, UMaterialExpressionStaticBoolParameter* Parameter)
+{
+	bool bTempVal;
+	if (Function->OverrideNamedStaticSwitchParameter(ParameterInfo, bTempVal, OutExpressionGuids[ParamIndex]))
+	{
+		OutValues[ParamIndex] = bTempVal;
+		return true;
+	}
+
+	if (Parameter)
+	{
+		bTempVal = OutValues[ParamIndex];
+		ensure(Parameter->IsNamedParameter(ParameterInfo, bTempVal, OutExpressionGuids[ParamIndex]));
+		OutValues[ParamIndex] = bTempVal;
+	}
+
+	return false;
+}
+
+template<>
+inline bool TMaterialStaticParameterValueResolver<UMaterialExpressionStaticComponentMaskParameter, TBitArray<>>::ExtractValueFromFunction(int32 ParamIndex, const FMaterialParameterInfo& ParameterInfo, UMaterialFunctionInterface* Function, UMaterialExpressionStaticComponentMaskParameter* Parameter)
+{
+	bool bTempR, bTempG, bTempB, bTempA;
+	if (Function->OverrideNamedStaticComponentMaskParameter(ParameterInfo, bTempR, bTempG, bTempB, bTempA, OutExpressionGuids[ParamIndex]))
+	{
+		int32 ParamRGBAIndex = ParamIndex * 4;
+		OutValues[ParamRGBAIndex + 0] = bTempR;
+		OutValues[ParamRGBAIndex + 1] = bTempG;
+		OutValues[ParamRGBAIndex + 2] = bTempB;
+		OutValues[ParamRGBAIndex + 3] = bTempA;
+		return true;
+	}
+
+	if (Parameter)
+	{
+		ensure(Parameter->IsNamedParameter(ParameterInfo, bTempR, bTempG, bTempB, bTempA, OutExpressionGuids[ParamIndex]));
+		int32 ParamRGBAIndex = ParamIndex * 4;
+		OutValues[ParamRGBAIndex + 0] = bTempR;
+		OutValues[ParamRGBAIndex + 1] = bTempG;
+		OutValues[ParamRGBAIndex + 2] = bTempB;
+		OutValues[ParamRGBAIndex + 3] = bTempA;
+	}
+
+	return false;
+}
diff --git a/Engine/Source/Runtime/Engine/Public/ComponentRecreateRenderStateContext.h b/Engine/Source/Runtime/Engine/Public/ComponentRecreateRenderStateContext.h
--- a/Engine/Source/Runtime/Engine/Public/ComponentRecreateRenderStateContext.h
+++ b/Engine/Source/Runtime/Engine/Public/ComponentRecreateRenderStateContext.h
@@ -61,7 +61,7 @@ public:
 
 private:
 	/** The recreate contexts for the individual components. */
-	TIndirectArray<FComponentRecreateRenderStateContext> ComponentContexts;
+	TArray<FComponentRecreateRenderStateContext> ComponentContexts;
 
 	TSet<FSceneInterface*> ScenesToUpdateAllPrimitiveSceneInfos;
 
