Changing a cast to an explicit set in UParticleSystemComponent::SetBoolParameter
#rnx
#tested PIE. build

diff --git a/Engine/Source/Runtime/Engine/Private/Particles/ParticleComponents.cpp b/Engine/Source/Runtime/Engine/Private/Particles/ParticleComponents.cpp
--- a/Engine/Source/Runtime/Engine/Private/Particles/ParticleComponents.cpp
+++ b/Engine/Source/Runtime/Engine/Private/Particles/ParticleComponents.cpp
@@ -7081,7 +7081,7 @@ int32 UParticleSystemComponent::GetLODLevel()
  */
 void UParticleSystemComponent::SetBoolParameter(FName Name, bool Value)
 {
-	SetFloatParameter(Name, float(Value));
+	SetFloatParameter(Name, (Value ? 1.0f : 0.0f));
 
 }
 
