Fixing audio device crash.

- Making sure non-editor builds only use one audio device (shared)

diff --git a/Engine/Source/Runtime/Engine/Private/AudioDeviceManager.cpp b/Engine/Source/Runtime/Engine/Private/AudioDeviceManager.cpp
--- a/Engine/Source/Runtime/Engine/Private/AudioDeviceManager.cpp
+++ b/Engine/Source/Runtime/Engine/Private/AudioDeviceManager.cpp
@@ -246,7 +246,14 @@ FAudioDeviceParams FAudioDeviceManager::GetDefaultParamsForNewWorld()
 #endif
 
 	FAudioDeviceParams Params;
-	Params.Scope = bCreateNewAudioDeviceForPlayInEditor ? EAudioDeviceScope::Unique : EAudioDeviceScope::Shared;
+	if (GIsEditor)
+	{
+		Params.Scope = bCreateNewAudioDeviceForPlayInEditor ? EAudioDeviceScope::Unique : EAudioDeviceScope::Shared;
+	}
+	else
+	{
+		Params.Scope = EAudioDeviceScope::Shared;
+	}
 
 	return Params;
 }
