Correctly update in-memory String Table asset references when a String Table asset is renamed

diff --git a/Engine/Source/Developer/AssetTools/Private/AssetRenameManager.cpp b/Engine/Source/Developer/AssetTools/Private/AssetRenameManager.cpp
--- a/Engine/Source/Developer/AssetTools/Private/AssetRenameManager.cpp
+++ b/Engine/Source/Developer/AssetTools/Private/AssetRenameManager.cpp
@@ -979,6 +979,8 @@ struct FSoftObjectPathRenameSerializer : public FArchiveUObject
 			DirtyDelegateHandle = UPackage::PackageMarkedDirtyEvent.AddRaw(this, &FSoftObjectPathRenameSerializer::OnMarkPackageDirty);
 		}
 
+		this->ArIsObjectReferenceCollector = true;
+
 		// Mark it as saving to correctly process all references
 		this->SetIsSaving(true);
 	}
@@ -996,7 +998,7 @@ struct FSoftObjectPathRenameSerializer : public FArchiveUObject
 		}
 
 		FFieldClass* PropertyClass = InProperty->GetClass();
-		if (PropertyClass->GetCastFlags() & (CASTCLASS_FBoolProperty | CASTCLASS_FNameProperty | CASTCLASS_FStrProperty | CASTCLASS_FTextProperty | CASTCLASS_FMulticastDelegateProperty))
+		if (PropertyClass->GetCastFlags() & (CASTCLASS_FBoolProperty | CASTCLASS_FNameProperty | CASTCLASS_FStrProperty | CASTCLASS_FMulticastDelegateProperty))
 		{
 			return true;
 		}
diff --git a/Engine/Source/Runtime/Core/Private/Internationalization/TextHistory.cpp b/Engine/Source/Runtime/Core/Private/Internationalization/TextHistory.cpp
--- a/Engine/Source/Runtime/Core/Private/Internationalization/TextHistory.cpp
+++ b/Engine/Source/Runtime/Core/Private/Internationalization/TextHistory.cpp
@@ -2613,12 +2613,20 @@ void FTextHistory_StringTableEntry::FStringTableReferenceData::GetTableIdAndKey(
 	OutKey = Key;
 }
 
-void FTextHistory_StringTableEntry::FStringTableReferenceData::CollectStringTableAssetReferences(FStructuredArchive::FRecord Record) const
+void FTextHistory_StringTableEntry::FStringTableReferenceData::CollectStringTableAssetReferences(FStructuredArchive::FRecord Record)
 {
 	if (Record.GetUnderlyingArchive().IsObjectReferenceCollector())
 	{
 		FScopeLock ScopeLock(&DataCS);
+
+		const FName OldTableId = TableId;
 		IStringTableEngineBridge::CollectStringTableAssetReferences(TableId, Record.EnterField(SA_FIELD_NAME(TEXT("AssetReferences"))));
+
+		if (TableId != OldTableId)
+		{
+			// This String Table asset was redirected, so we'll need to re-resolve the String Table entry later
+			StringTableEntry.Reset();
+		}
 	}
 }
 
diff --git a/Engine/Source/Runtime/Core/Private/Internationalization/TextHistory.h b/Engine/Source/Runtime/Core/Private/Internationalization/TextHistory.h
--- a/Engine/Source/Runtime/Core/Private/Internationalization/TextHistory.h
+++ b/Engine/Source/Runtime/Core/Private/Internationalization/TextHistory.h
@@ -665,7 +665,7 @@ private:
 		void GetTableIdAndKey(FName& OutTableId, FString& OutKey) const;
 
 		/** Collect any string table asset references */
-		void CollectStringTableAssetReferences(FStructuredArchive::FRecord Record) const;
+		void CollectStringTableAssetReferences(FStructuredArchive::FRecord Record);
 
 		/** Resolve the string table pointer, potentially re-caching it if it's missing or stale */
 		FStringTableEntryConstPtr ResolveStringTableEntry();
diff --git a/Engine/Source/Runtime/Core/Public/Internationalization/StringTableCore.h b/Engine/Source/Runtime/Core/Public/Internationalization/StringTableCore.h
--- a/Engine/Source/Runtime/Core/Public/Internationalization/StringTableCore.h
+++ b/Engine/Source/Runtime/Core/Public/Internationalization/StringTableCore.h
@@ -235,11 +235,11 @@ public:
 	}
 
 	/** Collect a string table asset reference */
-	static void CollectStringTableAssetReferences(const FName InTableId, FStructuredArchive::FSlot Slot)
+	static void CollectStringTableAssetReferences(FName& InOutTableId, FStructuredArchive::FSlot Slot)
 	{
 		if (InstancePtr)
 		{
-			InstancePtr->CollectStringTableAssetReferencesImpl(InTableId, Slot);
+			InstancePtr->CollectStringTableAssetReferencesImpl(InOutTableId, Slot);
 		}
 	}
 
@@ -262,7 +262,7 @@ protected:
 	virtual int32 LoadStringTableAssetImpl(const FName InTableId, FLoadStringTableAssetCallback InLoadedCallback) = 0;
 	virtual void FullyLoadStringTableAssetImpl(FName& InOutTableId) = 0;
 	virtual void RedirectStringTableAssetImpl(FName& InOutTableId) = 0;
-	virtual void CollectStringTableAssetReferencesImpl(const FName InTableId, FStructuredArchive::FSlot Slot) = 0;
+	virtual void CollectStringTableAssetReferencesImpl(FName& InOutTableId, FStructuredArchive::FSlot Slot) = 0;
 	virtual bool IsStringTableFromAssetImpl(const FName InTableId) = 0;
 	virtual bool IsStringTableAssetBeingReplacedImpl(const UStringTable* InStringTableAsset) = 0;
 
diff --git a/Engine/Source/Runtime/Engine/Private/Internationalization/StringTable.cpp b/Engine/Source/Runtime/Engine/Private/Internationalization/StringTable.cpp
--- a/Engine/Source/Runtime/Engine/Private/Internationalization/StringTable.cpp
+++ b/Engine/Source/Runtime/Engine/Private/Internationalization/StringTable.cpp
@@ -245,12 +245,31 @@ private:
 		}
 	}
 
-	virtual void CollectStringTableAssetReferencesImpl(const FName InTableId, FStructuredArchive::FSlot Slot) override
+	virtual void CollectStringTableAssetReferencesImpl(FName& InOutTableId, FStructuredArchive::FSlot Slot) override
 	{
 		check(Slot.GetUnderlyingArchive().IsObjectReferenceCollector());
 
-		UObject* StringTableAsset = FStringTableRegistry::Get().FindStringTableAsset(InTableId);
-		Slot << StringTableAsset;
+		UObject* StringTableAsset = FStringTableRegistry::Get().FindStringTableAsset(InOutTableId);
+
+		FSoftObjectPath StringTableAssetReference = GetAssetReference(InOutTableId);
+		if (StringTableAssetReference.IsValid())
+		{
+			FStructuredArchiveRecord Record = Slot.EnterRecord();
+
+			// Save a hard-reference so that dependent tables are loaded correctly
+			Record << SA_VALUE(TEXT("Ptr"), StringTableAsset);
+
+			// Save a weak-reference so that redirector fix-up during rename works
+			Record << SA_VALUE(TEXT("Ref"), StringTableAssetReference);
+
+			// This may have redirected our weak-reference
+			InOutTableId = *StringTableAssetReference.ToString();
+		}
+		else
+		{
+			// We have to write something to the slot...
+			Slot << StringTableAsset;
+		}
 	}
 
 	virtual bool IsStringTableFromAssetImpl(const FName InTableId) override
