CSVProfiler: Breakup WorldTickMisc into major tick groups

#rnx

diff --git a/Engine/Source/Runtime/Engine/Private/LevelTick.cpp b/Engine/Source/Runtime/Engine/Private/LevelTick.cpp
--- a/Engine/Source/Runtime/Engine/Private/LevelTick.cpp
+++ b/Engine/Source/Runtime/Engine/Private/LevelTick.cpp
@@ -1549,6 +1549,7 @@ void UWorld::Tick( ELevelTick TickType, float DeltaSeconds )
 			{
 				SCOPE_TIME_GUARD_MS(TEXT("UWorld::Tick - TG_PrePhysics"), 10);
 				SCOPE_CYCLE_COUNTER(STAT_TG_PrePhysics);
+				CSV_SCOPED_TIMING_STAT_EXCLUSIVE(PrePhysicsMisc);
 				CSV_SCOPED_SET_WAIT_STAT(PrePhysics);
 				RunTickGroup(TG_PrePhysics);
 			}
@@ -1558,12 +1559,14 @@ void UWorld::Tick( ELevelTick TickType, float DeltaSeconds )
 			{
 				SCOPE_CYCLE_COUNTER(STAT_TG_StartPhysics);
 				SCOPE_TIME_GUARD_MS(TEXT("UWorld::Tick - TG_StartPhysics"), 10);
+				CSV_SCOPED_TIMING_STAT_EXCLUSIVE(StartPhysicsMisc);
 				CSV_SCOPED_SET_WAIT_STAT(StartPhysics);
 				RunTickGroup(TG_StartPhysics);
 			}
 			{
 				SCOPE_CYCLE_COUNTER(STAT_TG_DuringPhysics);
 				SCOPE_TIME_GUARD_MS(TEXT("UWorld::Tick - TG_DuringPhysics"), 10);
+				CSV_SCOPED_TIMING_STAT_EXCLUSIVE(DuringPhysicsMisc);
 				CSV_SCOPED_SET_WAIT_STAT(DuringPhysics);
 				RunTickGroup(TG_DuringPhysics, false); // No wait here, we should run until idle though. We don't care if all of the async ticks are done before we start running post-phys stuff
 			}
@@ -1571,12 +1574,14 @@ void UWorld::Tick( ELevelTick TickType, float DeltaSeconds )
 			{
 				SCOPE_CYCLE_COUNTER(STAT_TG_EndPhysics);
 				SCOPE_TIME_GUARD_MS(TEXT("UWorld::Tick - TG_EndPhysics"), 10);
+				CSV_SCOPED_TIMING_STAT_EXCLUSIVE(EndPhysicsMisc);
 				CSV_SCOPED_SET_WAIT_STAT(EndPhysics);
 				RunTickGroup(TG_EndPhysics);
 			}
 			{
 				SCOPE_CYCLE_COUNTER(STAT_TG_PostPhysics);
 				SCOPE_TIME_GUARD_MS(TEXT("UWorld::Tick - TG_PostPhysics"), 10);
+				CSV_SCOPED_TIMING_STAT_EXCLUSIVE(PostPhysicsMisc);
 				CSV_SCOPED_SET_WAIT_STAT(PostPhysics);
 				RunTickGroup(TG_PostPhysics);
 			}
