Encryption packet handler component changes to handle passing more than just a key

diff --git a/Engine/Plugins/Online/OnlineSubsystemUtils/Source/OnlineSubsystemUtils/Private/OnlineBeaconClient.cpp b/Engine/Plugins/Online/OnlineSubsystemUtils/Source/OnlineSubsystemUtils/Private/OnlineBeaconClient.cpp
--- a/Engine/Plugins/Online/OnlineSubsystemUtils/Source/OnlineSubsystemUtils/Private/OnlineBeaconClient.cpp
+++ b/Engine/Plugins/Online/OnlineSubsystemUtils/Source/OnlineSubsystemUtils/Private/OnlineBeaconClient.cpp
@@ -130,9 +130,9 @@ bool AOnlineBeaconClient::InitClient(FURL& URL)
 					}
 
 #if NETCONNECTION_HAS_SETENCRYPTIONKEY
-					if (EncryptionKey.Num() > 0)
+					if (!EncryptionData.Identifier.IsEmpty())
 					{
-						BeaconConnection->SetEncryptionKey(EncryptionKey);
+						BeaconConnection->SetEncryptionData(EncryptionData);
 					}
 #endif
 
@@ -214,15 +214,23 @@ void AOnlineBeaconClient::Tick(float DeltaTime)
 
 void AOnlineBeaconClient::SetEncryptionToken(const FString& InEncryptionToken)
 {
-	EncryptionToken = InEncryptionToken;
+	EncryptionData.Identifier = InEncryptionToken;
 }
 
 void AOnlineBeaconClient::SetEncryptionKey(TArrayView<uint8> InEncryptionKey)
 {
 	if (CVarNetAllowEncryption.GetValueOnGameThread() != 0)
 	{
-		EncryptionKey.Reset(InEncryptionKey.Num());
-		EncryptionKey.Append(InEncryptionKey.GetData(), InEncryptionKey.Num());
+		EncryptionData.Key.Reset(InEncryptionKey.Num());
+		EncryptionData.Key.Append(InEncryptionKey.GetData(), InEncryptionKey.Num());
+	}
+}
+
+void AOnlineBeaconClient::SetEncryptionData(const FEncryptionData& InEncryptionData)
+{
+	if (CVarNetAllowEncryption.GetValueOnGameThread() != 0)
+	{
+		EncryptionData = InEncryptionData;
 	}
 }
 
@@ -237,10 +245,10 @@ void AOnlineBeaconClient::SendInitialJoin()
 
 		if (CVarNetAllowEncryption.GetValueOnGameThread() == 0)
 		{
-			EncryptionToken.Reset();
+			EncryptionData.Identifier.Empty();
 		}
 
-		FNetControlMessage<NMT_Hello>::Send(NetDriver->ServerConnection, IsLittleEndian, LocalNetworkVersion, EncryptionToken);
+		FNetControlMessage<NMT_Hello>::Send(NetDriver->ServerConnection, IsLittleEndian, LocalNetworkVersion, EncryptionData.Identifier);
 
 		NetDriver->ServerConnection->FlushNet();
 	}
@@ -440,7 +448,7 @@ void AOnlineBeaconClient::FinalizeEncryptedConnection(const FEncryptionKeyRespon
 		{
 			if (Response.Response == EEncryptionResponse::Success)
 			{
-				Connection->EnableEncryptionWithKey(Response.EncryptionKey);
+				Connection->EnableEncryption(Response.EncryptionData);
 			}
 			else
 			{
diff --git a/Engine/Plugins/Online/OnlineSubsystemUtils/Source/OnlineSubsystemUtils/Private/OnlineBeaconHost.cpp b/Engine/Plugins/Online/OnlineSubsystemUtils/Source/OnlineSubsystemUtils/Private/OnlineBeaconHost.cpp
--- a/Engine/Plugins/Online/OnlineSubsystemUtils/Source/OnlineSubsystemUtils/Private/OnlineBeaconHost.cpp
+++ b/Engine/Plugins/Online/OnlineSubsystemUtils/Source/OnlineSubsystemUtils/Private/OnlineBeaconHost.cpp
@@ -364,7 +364,7 @@ void AOnlineBeaconHost::SendWelcomeControlMessage(const FEncryptionKeyResponse&
 		{
 			if (Response.Response == EEncryptionResponse::Success)
 			{
-				Connection->EnableEncryptionWithKeyServer(Response.EncryptionKey);
+				Connection->EnableEncryptionServer(Response.EncryptionData);
 				SendWelcomeControlMessage(Connection);
 			}
 			else
diff --git a/Engine/Plugins/Online/OnlineSubsystemUtils/Source/OnlineSubsystemUtils/Public/OnlineBeaconClient.h b/Engine/Plugins/Online/OnlineSubsystemUtils/Source/OnlineSubsystemUtils/Public/OnlineBeaconClient.h
--- a/Engine/Plugins/Online/OnlineSubsystemUtils/Source/OnlineSubsystemUtils/Public/OnlineBeaconClient.h
+++ b/Engine/Plugins/Online/OnlineSubsystemUtils/Source/OnlineSubsystemUtils/Public/OnlineBeaconClient.h
@@ -76,6 +76,7 @@ class ONLINESUBSYSTEMUTILS_API AOnlineBeaconClient : public AOnlineBeacon
 	 *
 	 * @param EncryptionToken the token to use
 	 */
+	UE_DEPRECATED(4.24, "Use SetEncryptionData instead")
 	void SetEncryptionToken(const FString& InEncryptionToken);
 
 	/**
@@ -83,8 +84,16 @@ class ONLINESUBSYSTEMUTILS_API AOnlineBeaconClient : public AOnlineBeacon
 	 *
 	 * @param EncryptionKey the key to use
 	 */
+	UE_DEPRECATED(4.24, "Use SetEncryptionData instead")
 	void SetEncryptionKey(TArrayView<uint8> InEncryptionKey);
 
+	/**
+	 * Sets the encryption data that will be used for server connections.
+	 *
+	 * @param InEncryptionData the encryption data to use, including key and identifier
+	 */
+	void SetEncryptionData(const FEncryptionData& InEncryptionData);
+
 	/**
 	 * Send the packet for triggering the initial join
 	 */
@@ -178,11 +187,8 @@ protected:
 
 private:
 
-	/** Token sent to servers when connecting with an NMT_Hello message. */
-	FString EncryptionToken;
-
-	/** Key used when connecting to servers. */
-	TArray<uint8> EncryptionKey;
+	/** Encryption data used when connecting to servers. */
+	FEncryptionData EncryptionData;
 
 	/**
 	 * Setup the connection for encryption with a given key
diff --git a/Engine/Plugins/Runtime/PacketHandlers/AESGCMHandlerComponent/Source/Private/AESGCMHandlerComponent.cpp b/Engine/Plugins/Runtime/PacketHandlers/AESGCMHandlerComponent/Source/Private/AESGCMHandlerComponent.cpp
--- a/Engine/Plugins/Runtime/PacketHandlers/AESGCMHandlerComponent/Source/Private/AESGCMHandlerComponent.cpp
+++ b/Engine/Plugins/Runtime/PacketHandlers/AESGCMHandlerComponent/Source/Private/AESGCMHandlerComponent.cpp
@@ -17,16 +17,16 @@ FAESGCMHandlerComponent::FAESGCMHandlerComponent()
 	EncryptionContext = IPlatformCrypto::Get().CreateContext();
 }
 
-void FAESGCMHandlerComponent::SetEncryptionKey(TArrayView<const uint8> NewKey)
+void FAESGCMHandlerComponent::SetEncryptionData(const FEncryptionData& EncryptionData)
 {
-	if (NewKey.Num() != KeySizeInBytes)
+	if (EncryptionData.Key.Num() != KeySizeInBytes)
 	{
 		UE_LOG(PacketHandlerLog, Log, TEXT("FAESGCMHandlerComponent::SetEncryptionKey. NewKey is not %d bytes long, ignoring."), KeySizeInBytes);
 		return;
 	}
 
 	Key.Reset(KeySizeInBytes);
-	Key.Append(NewKey.GetData(), NewKey.Num());
+	Key.Append(EncryptionData.Key.GetData(), EncryptionData.Key.Num());
 }
 
 void FAESGCMHandlerComponent::EnableEncryption()
diff --git a/Engine/Plugins/Runtime/PacketHandlers/AESGCMHandlerComponent/Source/Public/AESGCMHandlerComponent.h b/Engine/Plugins/Runtime/PacketHandlers/AESGCMHandlerComponent/Source/Public/AESGCMHandlerComponent.h
--- a/Engine/Plugins/Runtime/PacketHandlers/AESGCMHandlerComponent/Source/Public/AESGCMHandlerComponent.h
+++ b/Engine/Plugins/Runtime/PacketHandlers/AESGCMHandlerComponent/Source/Public/AESGCMHandlerComponent.h
@@ -30,7 +30,7 @@ public:
 	static const int32 AuthTagSizeInBytes = 16;
 
 	// Replace the key used for encryption with NewKey if NewKey is exactly KeySizeInBytes long.
-	virtual void SetEncryptionKey(TArrayView<const uint8> NewKey) override;
+	virtual void SetEncryptionData(const FEncryptionData& EncryptionData) override;
 
 	// After calling this, future outgoing packets will be encrypted (until a call to DisableEncryption).
 	virtual void EnableEncryption() override;
diff --git a/Engine/Plugins/Runtime/PacketHandlers/AESHandlerComponent/Source/Private/AESHandlerComponent.cpp b/Engine/Plugins/Runtime/PacketHandlers/AESHandlerComponent/Source/Private/AESHandlerComponent.cpp
--- a/Engine/Plugins/Runtime/PacketHandlers/AESHandlerComponent/Source/Private/AESHandlerComponent.cpp
+++ b/Engine/Plugins/Runtime/PacketHandlers/AESHandlerComponent/Source/Private/AESHandlerComponent.cpp
@@ -21,16 +21,16 @@ FAESHandlerComponent::FAESHandlerComponent()
 	EncryptionContext = IPlatformCrypto::Get().CreateContext();
 }
 
-void FAESHandlerComponent::SetEncryptionKey(TArrayView<const uint8> NewKey)
+void FAESHandlerComponent::SetEncryptionData(const FEncryptionData& EncryptionData)
 {
-	if (NewKey.Num() != KeySizeInBytes)
+	if (EncryptionData.Key.Num() != KeySizeInBytes)
 	{
 		UE_LOG(PacketHandlerLog, Log, TEXT("FAESHandlerComponent::SetEncryptionKey. NewKey is not %d bytes long, ignoring."), KeySizeInBytes);
 		return;
 	}
 
 	Key.Reset(KeySizeInBytes);
-	Key.Append(NewKey.GetData(), NewKey.Num());
+	Key.Append(EncryptionData.Key.GetData(), EncryptionData.Key.Num());
 }
 
 void FAESHandlerComponent::EnableEncryption()
diff --git a/Engine/Plugins/Runtime/PacketHandlers/AESHandlerComponent/Source/Public/AESHandlerComponent.h b/Engine/Plugins/Runtime/PacketHandlers/AESHandlerComponent/Source/Public/AESHandlerComponent.h
--- a/Engine/Plugins/Runtime/PacketHandlers/AESHandlerComponent/Source/Public/AESHandlerComponent.h
+++ b/Engine/Plugins/Runtime/PacketHandlers/AESHandlerComponent/Source/Public/AESHandlerComponent.h
@@ -26,9 +26,8 @@ public:
 	// This handler uses AES256, which has 32-byte keys.
 	static const int32 BlockSizeInBytes = 16;
 
-
 	// Replace the key used for encryption with NewKey if NewKey is exactly KeySizeInBytes long.
-	virtual void SetEncryptionKey(TArrayView<const uint8> NewKey) override;
+	virtual void SetEncryptionData(const FEncryptionData& EncryptionData) override;
 
 	// After calling this, future outgoing packets will be encrypted (until a call to DisableEncryption).
 	virtual void EnableEncryption() override;
diff --git a/Engine/Source/Runtime/CoreUObject/Public/UObject/CoreNet.h b/Engine/Source/Runtime/CoreUObject/Public/UObject/CoreNet.h
--- a/Engine/Source/Runtime/CoreUObject/Public/UObject/CoreNet.h
+++ b/Engine/Source/Runtime/CoreUObject/Public/UObject/CoreNet.h
@@ -534,6 +534,15 @@ struct FNetDeltaSerializeInfo
 	FString DebugName;
 };
 
+struct FEncryptionData
+{
+	/** Encryption key */
+	TArray<uint8> Key;
+	/** Encryption fingerprint */
+	TArray<uint8> Fingerprint;
+	/** Encryption identifier */
+	FString Identifier;
+};
 
 /**
  * Checksum macros for verifying archives stay in sync
diff --git a/Engine/Source/Runtime/Engine/Classes/Engine/NetConnection.h b/Engine/Source/Runtime/Engine/Classes/Engine/NetConnection.h
--- a/Engine/Source/Runtime/Engine/Classes/Engine/NetConnection.h
+++ b/Engine/Source/Runtime/Engine/Classes/Engine/NetConnection.h
@@ -860,18 +860,36 @@ public:
 	/**
 	 * Sets the encryption key and enables encryption.
 	 */
+	UE_DEPRECATED(4.24, "Use SetEncryptionData instead.")
 	ENGINE_API void EnableEncryptionWithKey(TArrayView<const uint8> Key);
+	
+	/**
+	 * Sets the encryption data and enables encryption.
+	 */
+	ENGINE_API void EnableEncryption(const FEncryptionData& EncryptionData);
 
 	/**
 	 * Sets the encryption key, enables encryption, and sends the encryption ack to the client.
 	 */
+	UE_DEPRECATED(4.24, "Use SetEncryptionData instead.")
 	ENGINE_API void EnableEncryptionWithKeyServer(TArrayView<const uint8> Key);
 
+	/**
+	 * Sets the encryption data, enables encryption, and sends the encryption ack to the client.
+	 */
+	ENGINE_API void EnableEncryptionServer(const FEncryptionData& EncryptionData);
+
 	/**
 	 * Sets the key for the underlying encryption packet handler component, but doesn't modify encryption enabled state.
 	 */
+	UE_DEPRECATED(4.24, "Use SetEncryptionData instead.")
 	ENGINE_API void SetEncryptionKey(TArrayView<const uint8> Key);
 
+	/**
+	 * Sets the data for the underlying encryption packet handler component, but doesn't modify encryption enabled state.
+	 */
+	ENGINE_API void SetEncryptionData(const FEncryptionData& EncryptionData);
+
 	/**
 	 * Sends an NMT_EncryptionAck message
 	 */
diff --git a/Engine/Source/Runtime/Engine/Classes/Engine/NetworkDelegates.h b/Engine/Source/Runtime/Engine/Classes/Engine/NetworkDelegates.h
--- a/Engine/Source/Runtime/Engine/Classes/Engine/NetworkDelegates.h
+++ b/Engine/Source/Runtime/Engine/Classes/Engine/NetworkDelegates.h
@@ -131,7 +131,10 @@ struct FEncryptionKeyResponse
 	/** Error message related to the response */
 	FString ErrorMsg;
 	/** Encryption key */
+	UE_DEPRECATED(4.24, "Use EncryptionData struct instead.")
 	TArray<uint8> EncryptionKey;
+	/** Encryption data */
+	FEncryptionData EncryptionData;
 
 	FEncryptionKeyResponse()
 		: Response(EEncryptionResponse::Failure)
diff --git a/Engine/Source/Runtime/Engine/Private/NetConnection.cpp b/Engine/Source/Runtime/Engine/Private/NetConnection.cpp
--- a/Engine/Source/Runtime/Engine/Private/NetConnection.cpp
+++ b/Engine/Source/Runtime/Engine/Private/NetConnection.cpp
@@ -472,34 +472,50 @@ void UNetConnection::NotifyAnalyticsProvider()
 }
 
 void UNetConnection::EnableEncryptionWithKey(TArrayView<const uint8> Key)
+{
+	FEncryptionData EncryptionData;
+	EncryptionData.Key.Append(Key.GetData(), Key.Num());
+
+	EnableEncryption(EncryptionData);
+}
+
+void UNetConnection::EnableEncryption(const FEncryptionData& EncryptionData)
 {
 	if (Handler.IsValid())
 	{
-		UE_LOG(LogNet, Verbose, TEXT("UNetConnection::EnableEncryptionWithKey, %s"), *Describe());
+		UE_LOG(LogNet, Verbose, TEXT("UNetConnection::EnableEncryption, %s"), *Describe());
 
 		TSharedPtr<FEncryptionComponent> EncryptionComponent = Handler->GetEncryptionComponent();
 		if (EncryptionComponent.IsValid())
 		{
-			EncryptionComponent->SetEncryptionKey(Key);
+			EncryptionComponent->SetEncryptionData(EncryptionData);
 			EncryptionComponent->EnableEncryption();
 		}
 		else
 		{
-			UE_LOG(LogNet, Warning, TEXT("UNetConnection::EnableEncryptionWithKey, encryption component not found!"));
+			UE_LOG(LogNet, Warning, TEXT("UNetConnection::EnableEncryption, encryption component not found!"));
 		}
 	}
 }
 
 void UNetConnection::EnableEncryptionWithKeyServer(TArrayView<const uint8> Key)
+{
+	FEncryptionData EncryptionData;
+	EncryptionData.Key.Append(Key.GetData(), Key.Num());
+
+	EnableEncryptionServer(EncryptionData);
+}
+
+void UNetConnection::EnableEncryptionServer(const FEncryptionData& EncryptionData)
 {
 	if (State != USOCK_Invalid && State != USOCK_Closed && Driver)
 	{
 		SendClientEncryptionAck();
-		EnableEncryptionWithKey(Key);
+		EnableEncryption(EncryptionData);
 	}
 	else
 	{
-		UE_LOG(LogNet, Log, TEXT("UNetConnection::EnableEncryptionWithKeyServer, connection in invalid state. %s"), *Describe());
+		UE_LOG(LogNet, Log, TEXT("UNetConnection::EnableEncryptionServer, connection in invalid state. %s"), *Describe());
 	}
 }
 
@@ -517,19 +533,27 @@ void UNetConnection::SendClientEncryptionAck()
 }
 
 void UNetConnection::SetEncryptionKey(TArrayView<const uint8> Key)
+{
+	FEncryptionData EncryptionData;
+	EncryptionData.Key.Append(Key.GetData(), Key.Num());
+
+	SetEncryptionData(EncryptionData);
+}
+
+void UNetConnection::SetEncryptionData(const FEncryptionData& EncryptionData)
 {
 	if (Handler.IsValid())
 	{
-		UE_LOG(LogNet, Verbose, TEXT("UNetConnection::SetEncryptionKey, %s"), *Describe());
+		UE_LOG(LogNet, Verbose, TEXT("UNetConnection::SetEncryptionData, %s"), *Describe());
 
 		TSharedPtr<FEncryptionComponent> EncryptionComponent = Handler->GetEncryptionComponent();
 		if (EncryptionComponent.IsValid())
 		{
-			EncryptionComponent->SetEncryptionKey(Key);
+			EncryptionComponent->SetEncryptionData(EncryptionData);
 		}
 		else
 		{
-			UE_LOG(LogNet, Warning, TEXT("UNetConnection::SetEncryptionKey, encryption component not found!"));
+			UE_LOG(LogNet, Warning, TEXT("UNetConnection::SetEncryptionData, encryption component not found!"));
 		}
 	}
 }
diff --git a/Engine/Source/Runtime/Engine/Private/PendingNetGame.cpp b/Engine/Source/Runtime/Engine/Private/PendingNetGame.cpp
--- a/Engine/Source/Runtime/Engine/Private/PendingNetGame.cpp
+++ b/Engine/Source/Runtime/Engine/Private/PendingNetGame.cpp
@@ -371,7 +371,7 @@ void UPendingNetGame::FinalizeEncryptedConnection(const FEncryptionKeyResponse&
 		{
 			if (Response.Response == EEncryptionResponse::Success)
 			{
-				Connection->EnableEncryptionWithKey(Response.EncryptionKey);
+				Connection->EnableEncryption(Response.EncryptionData);
 			}
 			else
 			{
@@ -415,7 +415,7 @@ void UPendingNetGame::SetEncryptionKey(const FEncryptionKeyResponse& Response)
 			{
 				if (Response.Response == EEncryptionResponse::Success)
 				{
-					Connection->SetEncryptionKey(Response.EncryptionKey);
+					Connection->SetEncryptionData(Response.EncryptionData);
 				}
 				else
 				{
diff --git a/Engine/Source/Runtime/Engine/Private/World.cpp b/Engine/Source/Runtime/Engine/Private/World.cpp
--- a/Engine/Source/Runtime/Engine/Private/World.cpp
+++ b/Engine/Source/Runtime/Engine/Private/World.cpp
@@ -5319,7 +5319,7 @@ void UWorld::SendChallengeControlMessage(const FEncryptionKeyResponse& Response,
 		{
 			if (Response.Response == EEncryptionResponse::Success)
 			{
-				Connection->EnableEncryptionWithKeyServer(Response.EncryptionKey);
+				Connection->EnableEncryptionServer(Response.EncryptionData);
 				SendChallengeControlMessage(Connection);
 			}
 			else
diff --git a/Engine/Source/Runtime/PacketHandlers/PacketHandler/Public/EncryptionComponent.h b/Engine/Source/Runtime/PacketHandlers/PacketHandler/Public/EncryptionComponent.h
--- a/Engine/Source/Runtime/PacketHandlers/PacketHandler/Public/EncryptionComponent.h
+++ b/Engine/Source/Runtime/PacketHandlers/PacketHandler/Public/EncryptionComponent.h
@@ -5,6 +5,7 @@
 #include "CoreMinimal.h"
 #include "PacketHandler.h"
 #include "Containers/ArrayView.h"
+#include "UObject/CoreNet.h"
 
 /**
  * IEncryptionComponentInterface
@@ -36,5 +37,16 @@ public:
 	 * Sets the encryption key to be used by this component. The key must be the correct size for the encryption algorithm
 	 * that the component implements. This should be called before EnableEncryption.
 	 */
-	virtual void SetEncryptionKey(TArrayView<const uint8> Key) = 0;
+	UE_DEPRECATED(4.24, "Use SetEncryptionData instead")
+	virtual void SetEncryptionKey(TArrayView<const uint8> Key)
+	{
+		FEncryptionData EncryptionData;
+		EncryptionData.Key.Append(Key.GetData(), Key.Num());
+		SetEncryptionData(EncryptionData);
+	}
+
+	/**
+	 * Sets the encryption data to be used by this component. This should be called before EnableEncryption.
+	 */
+	virtual void SetEncryptionData(const FEncryptionData& EncryptionData) = 0;
 };
