non-unity fix

diff --git a/Engine/Source/Runtime/Engine/Classes/Engine/NetworkDelegates.h b/Engine/Source/Runtime/Engine/Classes/Engine/NetworkDelegates.h
--- a/Engine/Source/Runtime/Engine/Classes/Engine/NetworkDelegates.h
+++ b/Engine/Source/Runtime/Engine/Classes/Engine/NetworkDelegates.h
@@ -7,6 +7,7 @@
 #include "UObject/UObjectGlobals.h"
 #include "UObject/Object.h"
 #include "Engine/EngineBaseTypes.h"
+#include "UObject/CoreNet.h"
 
 class UPendingNetGame;
 
