Fix duplicated test case in the URL serializer check command.

#rnx

diff --git a/Engine/Source/Runtime/Engine/Private/URL.cpp b/Engine/Source/Runtime/Engine/Private/URL.cpp
--- a/Engine/Source/Runtime/Engine/Private/URL.cpp
+++ b/Engine/Source/Runtime/Engine/Private/URL.cpp
@@ -674,7 +674,7 @@ static bool URLSerializationTests(UWorld* InWorld, const TCHAR* Cmd, FOutputDevi
 		TestCases.Push(FURLTestCase(TEXT("[2001:0db8:85a3:0000:0000:8a2e:0370:7334]:7778"), TEXT("2001:0db8:85a3:0000:0000:8a2e:0370:7334"), TEXT(""), 7778));
 		TestCases.Push(FURLTestCase(TEXT("[2001:db8:85a3::8a2e:370:7334]"), TEXT("2001:db8:85a3::8a2e:370:7334")));
 		TestCases.Push(FURLTestCase(TEXT("epic://2001:db8:85a3::8a2e:370:7334"), TEXT("2001:db8:85a3::8a2e:370:7334"), TEXT("epic")));
-		TestCases.Push(FURLTestCase(TEXT("192.168.0.1"), TEXT("192.168.0.1")));
+		TestCases.Push(FURLTestCase(TEXT("192.168.0.1:17777"), TEXT("192.168.0.1"), TEXT(""), 17777));
 		TestCases.Push(FURLTestCase(TEXT("test://192.168.0.1"), TEXT("192.168.0.1"), TEXT("test")));
 		TestCases.Push(FURLTestCase(TEXT("::ffff:192.168.0.1"), TEXT("::ffff:192.168.0.1")));
 		TestCases.Push(FURLTestCase(TEXT("[::ffff:192.168.0.1]"), TEXT("::ffff:192.168.0.1")));
