Fix use of non-squared value in AActor::GetReplayPriority.

diff --git a/Engine/Source/Runtime/Engine/Private/ActorReplication.cpp b/Engine/Source/Runtime/Engine/Private/ActorReplication.cpp
--- a/Engine/Source/Runtime/Engine/Private/ActorReplication.cpp
+++ b/Engine/Source/Runtime/Engine/Private/ActorReplication.cpp
@@ -106,7 +106,7 @@ float AActor::GetReplayPriority(const FVector& ViewPos, const FVector& ViewDir,
 		{
 			Time *= 2.4f;
 		}
-		else if (DistSq < FARSIGHTTHRESHOLD)
+		else if (DistSq < FARSIGHTTHRESHOLDSQUARED)
 		{
 			Time *= 0.8f;
 		}
