Add support for IsEncryptionEnabled for child connection objects.

diff --git a/Engine/Source/Runtime/Engine/Classes/Engine/ChildConnection.h b/Engine/Source/Runtime/Engine/Classes/Engine/ChildConnection.h
--- a/Engine/Source/Runtime/Engine/Classes/Engine/ChildConnection.h
+++ b/Engine/Source/Runtime/Engine/Classes/Engine/ChildConnection.h
@@ -64,6 +64,11 @@ public:
 		return Parent->IsNetReady(Saturate);
 	}
 
+	virtual bool IsEncryptionEnabled() const override
+	{
+		return Parent->IsEncryptionEnabled();
+	}
+
 	virtual void Tick() override
 	{
 		State = Parent->State;
diff --git a/Engine/Source/Runtime/Engine/Classes/Engine/NetConnection.h b/Engine/Source/Runtime/Engine/Classes/Engine/NetConnection.h
--- a/Engine/Source/Runtime/Engine/Classes/Engine/NetConnection.h
+++ b/Engine/Source/Runtime/Engine/Classes/Engine/NetConnection.h
@@ -962,7 +962,7 @@ public:
 	/**
 	 * Returns true if encryption is enabled for this connection.
 	 */
-	ENGINE_API bool IsEncryptionEnabled() const;
+	ENGINE_API virtual bool IsEncryptionEnabled() const;
 
 	/** 
 	* Gets a unique ID for the connection, this ID depends on the underlying connection
