Chaos : Sleeping

- Pass sleep state to the physics thread by making sure the forces are torques are reset properly during the Vehicle::PutToSleep

diff --git a/Engine/Source/Runtime/Engine/Private/PhysicsEngine/Experimental/PhysScene_Chaos.cpp b/Engine/Source/Runtime/Engine/Private/PhysicsEngine/Experimental/PhysScene_Chaos.cpp
--- a/Engine/Source/Runtime/Engine/Private/PhysicsEngine/Experimental/PhysScene_Chaos.cpp
+++ b/Engine/Source/Runtime/Engine/Private/PhysicsEngine/Experimental/PhysScene_Chaos.cpp
@@ -1480,6 +1480,16 @@ void FPhysScene_ChaosInterface::AddRadialForceToBody_AssumesLocked(FBodyInstance
 
 void FPhysScene_ChaosInterface::ClearForces_AssumesLocked(FBodyInstance* BodyInstance, bool bAllowSubstepping)
 {
+	FPhysicsActorHandle& Handle = BodyInstance->GetPhysicsActorHandle();
+	if (ensure(FPhysicsInterface::IsValid(Handle)))
+	{
+		Chaos::TPBDRigidParticle<float, 3>* Rigid = Handle->CastToRigidParticle();
+		if (ensure(Rigid))
+		{
+			Rigid->SetF(Chaos::TVector<float, 3>(0.f,0.f,0.f));
+			Rigid->MarkClean(Chaos::EParticleFlags::F);
+		}
+	}
 }
 
 void FPhysScene_ChaosInterface::AddTorque_AssumesLocked(FBodyInstance* BodyInstance, const FVector& Torque, bool bAllowSubstepping, bool bAccelChange)
@@ -1512,7 +1522,16 @@ void FPhysScene_ChaosInterface::AddTorque_AssumesLocked(FBodyInstance* BodyInsta
 
 void FPhysScene_ChaosInterface::ClearTorques_AssumesLocked(FBodyInstance* BodyInstance, bool bAllowSubstepping)
 {
-	// #todo : Implement
+	FPhysicsActorHandle& Handle = BodyInstance->GetPhysicsActorHandle();
+	if (ensure(FPhysicsInterface::IsValid(Handle)))
+	{
+		Chaos::TPBDRigidParticle<float, 3>* Rigid = Handle->CastToRigidParticle();
+		if (ensure(Rigid))
+		{
+			Rigid->SetTorque(Chaos::TVector<float, 3>(0.f, 0.f, 0.f));
+			Rigid->MarkClean(Chaos::EParticleFlags::Torque);
+		}
+	}
 }
 
 void FPhysScene_ChaosInterface::SetKinematicTarget_AssumesLocked(FBodyInstance* BodyInstance, const FTransform& TargetTM, bool bAllowSubstepping)
diff --git a/Engine/Source/Runtime/Experimental/Chaos/Private/Chaos/ParticleHandle.cpp b/Engine/Source/Runtime/Experimental/Chaos/Private/Chaos/ParticleHandle.cpp
--- a/Engine/Source/Runtime/Experimental/Chaos/Private/Chaos/ParticleHandle.cpp
+++ b/Engine/Source/Runtime/Experimental/Chaos/Private/Chaos/ParticleHandle.cpp
@@ -94,6 +94,12 @@ namespace Chaos
 		}
 	}
 
+	template<>
+	void Chaos::TGeometryParticle<float, 3>::MarkClean(const EParticleFlags CleanBits)
+	{
+		this->MDirtyFlags.MarkClean(CleanBits);
+	}
+
 	template <>
 	template <>
 	int32 TGeometryParticleHandleImp<float, 3, true>::GetPayload<int32>(int32 Idx)
diff --git a/Engine/Source/Runtime/Experimental/Chaos/Public/Chaos/ParticleDirtyFlags.h b/Engine/Source/Runtime/Experimental/Chaos/Public/Chaos/ParticleDirtyFlags.h
--- a/Engine/Source/Runtime/Experimental/Chaos/Public/Chaos/ParticleDirtyFlags.h
+++ b/Engine/Source/Runtime/Experimental/Chaos/Public/Chaos/ParticleDirtyFlags.h
@@ -73,7 +73,7 @@ namespace Chaos
 
 		void MarkClean(const EParticleFlags CleanBits)
 		{
-			Bits |= ~(int32)CleanBits;
+			Bits &= ~(int32)CleanBits;
 		}
 
 		void Clear()
diff --git a/Engine/Source/Runtime/Experimental/Chaos/Public/Chaos/ParticleHandle.h b/Engine/Source/Runtime/Experimental/Chaos/Public/Chaos/ParticleHandle.h
--- a/Engine/Source/Runtime/Experimental/Chaos/Public/Chaos/ParticleHandle.h
+++ b/Engine/Source/Runtime/Experimental/Chaos/Public/Chaos/ParticleHandle.h
@@ -1516,6 +1516,8 @@ public:
 		this->MDirtyFlags.Clear();
 	}
 
+	void MarkClean(const EParticleFlags DirtyBits);
+
 	TGeometryParticleHandle<T, d>* Handle() const
 	{
 		if (Proxy)
