Add optional, macro enabled instrumentation for audio device handles on shutdown, instead of making it a fatal error to not clean up all of your handles.

Also fix issue where we were not broadcasting the FOnAudioDeviceDestroyed delegate on shutdown.

diff --git a/Engine/Source/Runtime/Engine/Private/AudioDeviceManager.cpp b/Engine/Source/Runtime/Engine/Private/AudioDeviceManager.cpp
--- a/Engine/Source/Runtime/Engine/Private/AudioDeviceManager.cpp
+++ b/Engine/Source/Runtime/Engine/Private/AudioDeviceManager.cpp
@@ -14,6 +14,12 @@
 #include "UObject/UObjectIterator.h"
 #include "Audio/AudioDebug.h"
 
+#if INSTRUMENT_AUDIODEVICE_HANDLES
+#include "Containers/StringConv.h"
+#include "HAL/PlatformStackWalk.h"
+#endif
+
+
 #if WITH_EDITOR
 #include "AudioEditorModule.h"
 #include "Settings/LevelEditorMiscSettings.h"
@@ -44,6 +50,55 @@ FAutoConsoleVariableRef CVarAudioVisualizeEnabled(
 	TEXT("0: Not Enabled, 1: Enabled"),
 	ECVF_Default);
 
+static FAutoConsoleCommand GReportAudioDevicesCommand(
+	TEXT("au.ReportAudioDevices"),
+	TEXT("This will log any active audio devices (instances of the audio engine) alive right now."),
+	FConsoleCommandDelegate::CreateStatic(
+		[]()
+	{
+		FAudioDeviceManager::Get()->LogListOfAudioDevices();
+	})
+);
+
+// Some stress tests:
+#if INSTRUMENT_AUDIODEVICE_HANDLES
+static TArray<FAudioDeviceHandle> IntentionallyLeakedHandles;
+
+static FAutoConsoleCommand GLeakAudioDeviceCommand(
+	TEXT("au.stresstest.LeakAnAudioDevice"),
+	TEXT("This will intentionally leak a new audio device. Obviously, should only be used for testing."),
+	FConsoleCommandDelegate::CreateStatic(
+		[]()
+{
+	FAudioDeviceParams Params;
+	Params.Scope = EAudioDeviceScope::Unique;
+	IntentionallyLeakedHandles.Add(FAudioDeviceManager::Get()->RequestAudioDevice(Params));
+})
+);
+
+static FAutoConsoleCommand GLeakAudioDeviceHandleCommand(
+	TEXT("au.stresstest.LeakAnAudioDeviceHandle"),
+	TEXT("This will intentionally leak a new handle to an audio device. Obviously, should only be used for testing."),
+	FConsoleCommandDelegate::CreateStatic(
+		[]()
+{
+	FAudioDeviceParams Params;
+	Params.Scope = EAudioDeviceScope::Shared;
+	IntentionallyLeakedHandles.Add(FAudioDeviceManager::Get()->RequestAudioDevice(Params));
+})
+);
+
+static FAutoConsoleCommand GCleanUpAudioDeviceLeaksCommand(
+	TEXT("au.stresstest.CleanUpAudioDeviceLeaks"),
+	TEXT("Clean up any audio devices created through a leak command."),
+	FConsoleCommandDelegate::CreateStatic(
+		[]()
+{
+	IntentionallyLeakedHandles.Reset();
+})
+);
+#endif
+
 /*-----------------------------------------------------------------------------
 FAudioDeviceManager implementation.
 -----------------------------------------------------------------------------*/
@@ -76,6 +131,13 @@ FAudioDeviceManager::FAudioDeviceManager()
 FAudioDeviceManager::~FAudioDeviceManager()
 {
 	MainAudioDeviceHandle.Reset();
+
+	// Notify anyone listening to the device manager that we are about to destroy the audio device.
+	for (auto& Device : Devices)
+	{
+		FAudioDeviceManagerDelegates::OnAudioDeviceDestroyed.Broadcast(Device.Key);
+	}
+
 	Devices.Reset();
 
 	// Release any loaded buffers - this calls stop on any sources that need it
@@ -242,18 +304,15 @@ FAudioDeviceParams FAudioDeviceManager::GetDefaultParamsForNewWorld()
 	bool bCreateNewAudioDeviceForPlayInEditor = false;
 
 #if WITH_EDITOR
-	bCreateNewAudioDeviceForPlayInEditor = GetDefault<ULevelEditorMiscSettings>()->bCreateNewAudioDeviceForPlayInEditor;
-#endif
-
-	FAudioDeviceParams Params;
+	// GIsEditor is necessary here to ignore this setting for -game situations.
 	if (GIsEditor)
 	{
-		Params.Scope = bCreateNewAudioDeviceForPlayInEditor ? EAudioDeviceScope::Unique : EAudioDeviceScope::Shared;
-	}
-	else
-	{
-		Params.Scope = EAudioDeviceScope::Shared;
+		bCreateNewAudioDeviceForPlayInEditor = GetDefault<ULevelEditorMiscSettings>()->bCreateNewAudioDeviceForPlayInEditor;
 	}
+#endif
+
+	FAudioDeviceParams Params;
+	Params.Scope = bCreateNewAudioDeviceForPlayInEditor ? EAudioDeviceScope::Unique : EAudioDeviceScope::Shared;
 
 	return Params;
 }
@@ -286,7 +345,6 @@ FAudioDeviceHandle FAudioDeviceManager::RequestAudioDevice(const FAudioDevicePar
 		// If we did not find a suitable device, build one.
 		return CreateNewDevice(InParams);
 	}
-	
 }
 
 bool FAudioDeviceManager::Initialize()
@@ -358,8 +416,6 @@ bool FAudioDeviceManager::LoadDefaultAudioDeviceModule()
 	// Check for config bool that restricts audio mixer toggle to only once. This will allow us to patch audio mixer on or off after initial login.
 	GConfig->GetBool(TEXT("Audio"), TEXT("OnlyToggleAudioMixerOnce"), bOnlyToggleAudioMixerOnce, GEngineIni);
 
-	
-
 	if (bForceNonRealtimeRenderer)
 	{
 		AudioDeviceModule = FModuleManager::LoadModulePtr<IAudioDeviceModule>(TEXT("NonRealtimeAudioRenderer"));
@@ -493,6 +549,14 @@ bool FAudioDeviceManager::CanUseAudioDevice(const FAudioDeviceParams& InParams,
 		&& InParams.bIsNonRealtime == InContainer.bIsNonRealtime;
 }
 
+#if INSTRUMENT_AUDIODEVICE_HANDLES
+uint32 FAudioDeviceManager::CreateUniqueStackWalkID()
+{
+	static uint32 UniqueStackWalkID = 0;
+	return UniqueStackWalkID++;
+}
+#endif
+
 FAudioDeviceHandle FAudioDeviceManager::GetAudioDevice(Audio::FDeviceId Handle)
 {
 	FAudioDeviceContainer* Container = Devices.Find(Handle);
@@ -790,6 +854,55 @@ TArray<UWorld*> FAudioDeviceManager::GetWorldsUsingAudioDevice(const Audio::FDev
 	}
 }
 
+#if INSTRUMENT_AUDIODEVICE_HANDLES
+void FAudioDeviceManager::AddStackWalkForContainer(Audio::FDeviceId InId, uint32 StackWalkID, FString&& InStackWalk)
+{
+	check(Devices.Contains(InId));
+	check(!Devices[InId].HandleCreationStackWalks.Contains(StackWalkID));
+	FAudioDeviceContainer& Container = Devices[InId];
+	Container.HandleCreationStackWalks.Add(StackWalkID, MoveTemp(InStackWalk));
+}
+
+void FAudioDeviceManager::RemoveStackWalkForContainer(Audio::FDeviceId InId, uint32 StackWalkID)
+{
+	check(Devices.Contains(InId));
+	check(Devices[InId].HandleCreationStackWalks.Contains(StackWalkID));
+	FAudioDeviceContainer& Container = Devices[InId];
+	Container.HandleCreationStackWalks.Remove(StackWalkID);
+}
+#endif
+
+void FAudioDeviceManager::LogListOfAudioDevices()
+{
+	FString ListOfDevices;
+
+	for (auto& DeviceContainer : Devices)
+	{
+		FString DeviceInfo = FString::Printf(TEXT(R"(
+					Device %d:
+					Scope: %s 
+					Realtime: %s
+					Number Of Owners: %d 
+		)"),
+			DeviceContainer.Key,
+			DeviceContainer.Value.Scope == EAudioDeviceScope::Unique ? TEXT("Unique") : TEXT("Shared"),
+			DeviceContainer.Value.bIsNonRealtime ? TEXT("No") : TEXT("Yes"),
+			DeviceContainer.Value.NumberOfHandlesToThisDevice);
+
+#if INSTRUMENT_AUDIODEVICE_HANDLES
+		for (auto& StackWalkString : DeviceContainer.Value.HandleCreationStackWalks)
+		{
+			DeviceInfo += TEXT("Handle Created here still alive:\n");
+			DeviceInfo += StackWalkString.Value;
+			DeviceInfo += TEXT("\n\n");
+		}
+#endif
+
+		ListOfDevices += DeviceInfo;
+	}
+
+	UE_LOG(LogAudio, Display, TEXT("List of devices: \n%s"), *ListOfDevices);
+}
 
 uint32 FAudioDeviceManager::GetNewDeviceID()
 {
@@ -1026,7 +1139,9 @@ FAudioDeviceHandle::FAudioDeviceHandle()
 	, Device(nullptr)
 	, DeviceId(INDEX_NONE)
 {
-
+#if INSTRUMENT_AUDIODEVICE_HANDLES
+	StackWalkID = INDEX_NONE;
+#endif
 }
 
 FAudioDeviceHandle::FAudioDeviceHandle(FAudioDevice* InDevice, Audio::FDeviceId InID, UWorld* InWorld)
@@ -1034,6 +1149,9 @@ FAudioDeviceHandle::FAudioDeviceHandle(FAudioDevice* InDevice, Audio::FDeviceId
 	, Device(InDevice)
 	, DeviceId(InID)
 {
+#if INSTRUMENT_AUDIODEVICE_HANDLES
+	AddStackDumpToAudioDeviceContainer();
+#endif
 }
 
 FAudioDeviceHandle::FAudioDeviceHandle(const FAudioDeviceHandle& Other)
@@ -1048,11 +1166,40 @@ FAudioDeviceHandle::FAudioDeviceHandle(FAudioDeviceHandle&& Other)
 	*this = MoveTemp(Other);
 }
 
+#if INSTRUMENT_AUDIODEVICE_HANDLES
+void FAudioDeviceHandle::AddStackDumpToAudioDeviceContainer()
+{
+	static const int32 MaxPlatformWalkStringCount = 1024 * 4;
+
+	ANSICHAR PlatformDump[MaxPlatformWalkStringCount];
+	FMemory::Memzero(PlatformDump, MaxPlatformWalkStringCount * sizeof(ANSICHAR));
+
+	FPlatformStackWalk::StackWalkAndDump(PlatformDump, MaxPlatformWalkStringCount - 1, 2);
+	
+	FString FormattedDump = TEXT("New Handle Created:\n");
+
+	int32 DumpLength = FCStringAnsi::Strlen(PlatformDump);
+
+	// If this hits, increase the max character length.
+	ensure(DumpLength < MaxPlatformWalkStringCount - 1);
+
+	FormattedDump.AppendChars(ANSI_TO_TCHAR(PlatformDump), DumpLength);
+	FormattedDump += TEXT("\n");
+	StackWalkID = FAudioDeviceManager::Get()->CreateUniqueStackWalkID();
+	FAudioDeviceManager::Get()->AddStackWalkForContainer(DeviceId, StackWalkID, MoveTemp(FormattedDump));
+}
+#endif
+
 FAudioDeviceHandle::~FAudioDeviceHandle()
 {
 	if (GEngine && IsValid())
 	{
 		GEngine->GetAudioDeviceManager()->DecrementDevice(DeviceId, World);
+
+#if INSTRUMENT_AUDIODEVICE_HANDLES
+		check(StackWalkID != INDEX_NONE);
+		GEngine->GetAudioDeviceManager()->RemoveStackWalkForContainer(DeviceId, StackWalkID);
+#endif
 	}
 }
 
@@ -1092,6 +1239,10 @@ FAudioDeviceHandle& FAudioDeviceHandle::operator=(const FAudioDeviceHandle& Othe
 	{
 		// Increment the reference count for the new streaming chunk.
 		FAudioDeviceManager::Get()->IncrementDevice(DeviceId);
+
+#if INSTRUMENT_AUDIODEVICE_HANDLES
+		AddStackDumpToAudioDeviceContainer();
+#endif
 	}
 
 	return *this;
@@ -1100,9 +1251,13 @@ FAudioDeviceHandle& FAudioDeviceHandle::operator=(const FAudioDeviceHandle& Othe
 FAudioDeviceHandle& FAudioDeviceHandle::operator=(FAudioDeviceHandle&& Other)
 {
 	// If this chunk was previously referencing another chunk, remove that chunk here.
-	if (IsValid())
+	if (FAudioDeviceManager::Get() && IsValid())
 	{
-		check(FAudioDeviceManager::Get());
+#if INSTRUMENT_AUDIODEVICE_HANDLES
+		check(StackWalkID != INDEX_NONE);
+		GEngine->GetAudioDeviceManager()->RemoveStackWalkForContainer(DeviceId, StackWalkID);
+#endif
+
 		FAudioDeviceManager::Get()->DecrementDevice(DeviceId, World);
 	}
 
@@ -1113,6 +1268,14 @@ FAudioDeviceHandle& FAudioDeviceHandle::operator=(FAudioDeviceHandle&& Other)
 	Other.Device = nullptr;
 	Other.DeviceId = INDEX_NONE;
 
+#if INSTRUMENT_AUDIODEVICE_HANDLES
+	if (IsValid())
+	{
+
+		AddStackDumpToAudioDeviceContainer();
+	}
+#endif
+
 	return *this;
 }
 
@@ -1189,12 +1352,31 @@ FAudioDeviceManager::FAudioDeviceContainer::FAudioDeviceContainer(FAudioDeviceCo
 
 	SpecifiedModule = Other.SpecifiedModule;
 	Other.SpecifiedModule = nullptr;
+
+#if INSTRUMENT_AUDIODEVICE_HANDLES
+	HandleCreationStackWalks = MoveTemp(Other.HandleCreationStackWalks);
+#endif
 }
 
 FAudioDeviceManager::FAudioDeviceContainer::~FAudioDeviceContainer()
 {
 	// Shutdown the audio device.
-	check(NumberOfHandlesToThisDevice == 0);
+	if (NumberOfHandlesToThisDevice != 0)
+	{
+		UE_LOG(LogAudio, Display, TEXT("Shutting down audio device while %d references to it are still alive. For more information, compile with INSTRUMENT_AUDIODEVICE_HANDLES."), NumberOfHandlesToThisDevice);
+
+#if INSTRUMENT_AUDIODEVICE_HANDLES
+		FString ActiveDeviceHandles;
+		for (auto& StackWalkString : HandleCreationStackWalks)
+		{
+			ActiveDeviceHandles += StackWalkString.Value;
+			ActiveDeviceHandles += TEXT("\n\n");
+		}
+
+		UE_LOG(LogAudio, Warning, TEXT("List Of Active Handles: \n%s"), *ActiveDeviceHandles);
+#endif
+	}
+
 	if (Device)
 	{
 		Device->FadeOut();
diff --git a/Engine/Source/Runtime/Engine/Public/AudioDeviceManager.h b/Engine/Source/Runtime/Engine/Public/AudioDeviceManager.h
--- a/Engine/Source/Runtime/Engine/Public/AudioDeviceManager.h
+++ b/Engine/Source/Runtime/Engine/Public/AudioDeviceManager.h
@@ -10,6 +10,12 @@
 class FAudioDebugger;
 #endif // ENABLE_AUDIO_DEBUG
 
+// Set this to one if you'd like to check who owns
+// handles to an audio device.
+#ifndef INSTRUMENT_AUDIODEVICE_HANDLES
+#define INSTRUMENT_AUDIODEVICE_HANDLES 0
+#endif
+
 class FReferenceCollector;
 class FSoundBuffer;
 class IAudioDeviceModule;
@@ -79,6 +85,8 @@ public:
 
 	FAudioDeviceHandle& operator=(const FAudioDeviceHandle& Other);
 	FAudioDeviceHandle& operator=(FAudioDeviceHandle&& Other);
+
+	
 	~FAudioDeviceHandle();
 
 	// gets a pointer to the compressed chunk.
@@ -103,6 +111,11 @@ private:
 	FAudioDevice* Device;
 	Audio::FDeviceId DeviceId;
 
+#if INSTRUMENT_AUDIODEVICE_HANDLES
+	uint32 StackWalkID;
+	void AddStackDumpToAudioDeviceContainer();
+#endif
+
 	friend class FAudioDeviceManager;
 
 public:
@@ -365,6 +378,13 @@ public:
 
 	TArray<UWorld*> GetWorldsUsingAudioDevice(const Audio::FDeviceId& InID);
 
+#if INSTRUMENT_AUDIODEVICE_HANDLES
+	void AddStackWalkForContainer(Audio::FDeviceId InId, uint32 StackWalkID, FString&& InStackWalk);
+	void RemoveStackWalkForContainer(Audio::FDeviceId InId, uint32 StackWalkID);
+#endif
+
+	void LogListOfAudioDevices();
+
 private:
 
 #if ENABLE_AUDIO_DEBUG
@@ -446,6 +466,10 @@ private:
 
 		FAudioDeviceContainer(FAudioDeviceContainer&& Other);
 
+#if INSTRUMENT_AUDIODEVICE_HANDLES
+		TMap<uint32, FString> HandleCreationStackWalks;
+#endif
+
 	private:
 		FAudioDeviceContainer();
 	};
@@ -457,6 +481,10 @@ private:
 	 */
 	static bool CanUseAudioDevice(const FAudioDeviceParams& InParams, const FAudioDeviceContainer& InContainer);
 
+#if INSTRUMENT_AUDIODEVICE_HANDLES
+	static uint32 CreateUniqueStackWalkID();
+#endif
+
 	/**
 	* Bank of audio devices. Will increase in size as we create new audio devices,
 	*/
