Studio Analytics - Prevent bogus pie start time.

diff --git a/Engine/Source/Editor/UnrealEd/Private/PlayLevel.cpp b/Engine/Source/Editor/UnrealEd/Private/PlayLevel.cpp
--- a/Engine/Source/Editor/UnrealEd/Private/PlayLevel.cpp
+++ b/Engine/Source/Editor/UnrealEd/Private/PlayLevel.cpp
@@ -109,6 +109,8 @@
 #include "Widgets/Text/STextBlock.h"
 #include "Widgets/SBoxPanel.h"
 #include "Subsystems/AssetEditorSubsystem.h"
+#include "Async/Async.h"
+#include "StudioAnalytics.h"
 #include "UObject/SoftObjectPath.h"
 #include "IAssetViewport.h"
 
@@ -2353,7 +2355,7 @@ void UEditorEngine::StartPlayInEditorSession(FRequestPlaySessionParams& InReques
 
 	// Broadcast PreBeginPIE before checks that might block PIE below (BeginPIE is broadcast below after the checks)
 	FEditorDelegates::PreBeginPIE.Broadcast(InRequestParams.WorldType == EPlaySessionWorldType::SimulateInEditor);
-	double PIEStartTime = FPlatformTime::Seconds();
+	const double PIEStartTime = FStudioAnalytics::GetAnalyticSeconds()
 	const FScopedBusyCursor BusyCursor;
 
 	// Block PIE when there is a transaction recording into the undo buffer. This is generally avoided
diff --git a/Engine/Source/Runtime/Engine/Private/Analytics/StudioAnalytics.cpp b/Engine/Source/Runtime/Engine/Private/Analytics/StudioAnalytics.cpp
--- a/Engine/Source/Runtime/Engine/Private/Analytics/StudioAnalytics.cpp
+++ b/Engine/Source/Runtime/Engine/Private/Analytics/StudioAnalytics.cpp
@@ -51,16 +51,17 @@ void FStudioAnalytics::RunTimer_Concurrent()
 {
 	TimeEstimation = FPlatformTime::Seconds();
 
+	const double FixedInterval = 0.0333333333334;
+	const double BreakpointHitchTime = 1;
+
 	while (bInitialized)
 	{
-		const double FixedInterval = 0.0333333333334;
-
 		const double StartTime = FPlatformTime::Seconds();
 		FPlatformProcess::Sleep((float)FixedInterval);
 		const double EndTime = FPlatformTime::Seconds();
 		const double DeltaTime = EndTime - StartTime;
 
-		if (DeltaTime > 1.0)
+		if (DeltaTime > BreakpointHitchTime)
 		{
 			TimeEstimation += FixedInterval;
 		}
@@ -115,6 +116,12 @@ void FStudioAnalytics::FireEvent_Loading(const FString& LoadingName, double Seco
 		return;
 	}
 
+	// Throw out anything over 10 hours - 
+	if (SecondsSpentLoading > 36000)
+	{
+		return;
+	}
+
 	if (FStudioAnalytics::IsAvailable())
 	{
 		TArray<FAnalyticsEventAttribute> Attributes;
diff --git a/Engine/Source/Runtime/Engine/Private/GameInstance.cpp b/Engine/Source/Runtime/Engine/Private/GameInstance.cpp
--- a/Engine/Source/Runtime/Engine/Private/GameInstance.cpp
+++ b/Engine/Source/Runtime/Engine/Private/GameInstance.cpp
@@ -327,6 +327,11 @@ bool UGameInstance::InitializePIE(bool bAnyBlueprintErrors, int32 PIEInstance, b
 
 FGameInstancePIEResult UGameInstance::StartPlayInEditorGameInstance(ULocalPlayer* LocalPlayer, const FGameInstancePIEParameters& Params)
 {
+	if (PIEStartTime == 0)
+	{
+		PIEStartTime = Params.PIEStartTime;
+	}
+
 	OnStart();
 
 	UEditorEngine* const EditorEngine = CastChecked<UEditorEngine>(GetEngine());
