Prevent crash when callbacks remove actor.

diff --git a/Engine/Source/Runtime/Engine/Private/PhysicsEngine/Experimental/PhysScene_Chaos.cpp b/Engine/Source/Runtime/Engine/Private/PhysicsEngine/Experimental/PhysScene_Chaos.cpp
--- a/Engine/Source/Runtime/Engine/Private/PhysicsEngine/Experimental/PhysScene_Chaos.cpp
+++ b/Engine/Source/Runtime/Engine/Private/PhysicsEngine/Experimental/PhysScene_Chaos.cpp
@@ -343,12 +343,25 @@ struct FPhysScenePendingComponentTransform_Chaos
 	/** Component to move */
 	TWeakObjectPtr<UPrimitiveComponent> OwningComp;
 	/** New transform from physics engine */
-	FTransform NewTransform;
+	FVector NewTranslation;
+	FQuat NewRotation;
+	bool bHasValidTransform;
+	bool bHasWakeEvent;
+	
+	FPhysScenePendingComponentTransform_Chaos(UPrimitiveComponent* InOwningComp, const FVector& InNewTranslation, const FQuat& InNewRotation, const bool InHasWakeEvent)
+		: OwningComp(InOwningComp)
+		, NewTranslation(InNewTranslation)
+		, NewRotation(InNewRotation)
+		, bHasValidTransform(true)
+		, bHasWakeEvent(InHasWakeEvent)
+	{}
 
-	FPhysScenePendingComponentTransform_Chaos(UPrimitiveComponent* InOwningComp, const FTransform& InNewTransform)
+	FPhysScenePendingComponentTransform_Chaos(UPrimitiveComponent* InOwningComp)
 		: OwningComp(InOwningComp)
-		, NewTransform(InNewTransform)
+		, bHasValidTransform(false)
+		, bHasWakeEvent(true)
 	{}
+
 };
 
 FPhysScene_Chaos::FPhysScene_Chaos(AActor* InSolverActor
@@ -2073,29 +2086,23 @@ void FPhysScene_ChaosInterface::SyncBodies(Chaos::FPhysicsSolver* Solver)
 						UPrimitiveComponent* OwnerComponent = BodyInstance->OwnerComponent.Get();
 						if (OwnerComponent != nullptr)
 						{
-							AActor* Owner = OwnerComponent->GetOwner();
-
+							bool bPendingMove = false;
 							if (BodyInstance->InstanceBodyIndex == INDEX_NONE)
 							{
 								Chaos::TRigidTransform<float, 3> NewTransform(ActiveParticle->X(), ActiveParticle->R());
 
 								if (!NewTransform.EqualsNoScale(OwnerComponent->GetComponentTransform()))
 								{
+									bPendingMove = true;
 									const FVector MoveBy = NewTransform.GetLocation() - OwnerComponent->GetComponentTransform().GetLocation();
 									const FQuat NewRotation = NewTransform.GetRotation();
-
-									OwnerComponent->MoveComponent(MoveBy, NewRotation, false, NULL, MOVECOMP_SkipPhysicsMove);
-								}
-
-								if (Owner != NULL && !Owner->IsPendingKill())
-								{
-									Owner->CheckStillInWorld();
+									PendingTransforms.Add(FPhysScenePendingComponentTransform_Chaos(OwnerComponent, MoveBy, NewRotation, Proxy->HasAwakeEvent()));
 								}
 							}
 
-							if (Proxy->HasAwakeEvent())
+							if (Proxy->HasAwakeEvent() && !bPendingMove)
 							{
-								OwnerComponent->DispatchWakeEvents(ESleepEvent::SET_Wakeup, NAME_None);
+								PendingTransforms.Add(FPhysScenePendingComponentTransform_Chaos(OwnerComponent));
 							}
 							Proxy->ClearEvents();
 						}
@@ -2109,6 +2116,31 @@ void FPhysScene_ChaosInterface::SyncBodies(Chaos::FPhysicsSolver* Solver)
 			}
 		}
 	}
+	for (const FPhysScenePendingComponentTransform_Chaos& ComponentTransform : PendingTransforms)
+	{
+		if (ComponentTransform.OwningComp != nullptr)
+		{
+			AActor* Owner = ComponentTransform.OwningComp->GetOwner();
+
+			if (ComponentTransform.bHasValidTransform)
+			{
+				ComponentTransform.OwningComp->MoveComponent(ComponentTransform.NewTranslation, ComponentTransform.NewRotation, false, NULL, MOVECOMP_SkipPhysicsMove);
+			}
+
+			if (Owner != NULL && !Owner->IsPendingKill())
+			{
+				Owner->CheckStillInWorld();
+			}
+		}
+
+		if (ComponentTransform.OwningComp != nullptr)
+		{
+			if (ComponentTransform.bHasWakeEvent)
+			{
+				ComponentTransform.OwningComp->DispatchWakeEvents(ESleepEvent::SET_Wakeup, NAME_None);
+			}
+		}
+	}
 }
 
 FPhysicsConstraintReference_Chaos 
