Deprecate duplicated attributes in FPrimitiveViewRelevance with different names from FMaterialRelevance

#rnx

diff --git a/Engine/Source/Runtime/Engine/Public/PrimitiveViewRelevance.h b/Engine/Source/Runtime/Engine/Public/PrimitiveViewRelevance.h
--- a/Engine/Source/Runtime/Engine/Public/PrimitiveViewRelevance.h
+++ b/Engine/Source/Runtime/Engine/Public/PrimitiveViewRelevance.h
@@ -11,6 +11,26 @@
  */
 struct FPrimitiveViewRelevance : public FMaterialRelevance
 {
+#if WITH_EDITOR
+	// Deprecate common attributes with FMaterialRelevance that have different names
+	UE_DEPRECATED(4.25, "ShadingModelMaskRelevance has been renamed ShadingModelMask")
+	uint16 ShadingModelMaskRelevance;
+	UE_DEPRECATED(4.25, "bOpaqueRelevance has been renamed bOpaque")
+	uint32 bOpaqueRelevance : 1;
+	UE_DEPRECATED(4.25, "bMaskedRelevance has been renamed bMasked")
+	uint32 bMaskedRelevance : 1;
+	UE_DEPRECATED(4.25, "bTranslucentVelocityRelevance has been renamed bOutputsTranslucentVelocity")
+	uint32 bTranslucentVelocityRelevance : 1;
+	UE_DEPRECATED(4.25, "bDistortionRelevance has been renamed bDistortion")
+	uint32 bDistortionRelevance : 1;
+	UE_DEPRECATED(4.25, "bSeparateTranslucencyRelevance has been renamed bSeparateTranslucency")
+	uint32 bSeparateTranslucencyRelevance : 1;
+	UE_DEPRECATED(4.25, "bNormalTranslucencyRelevance has been renamed bNormalTranslucency")
+	uint32 bNormalTranslucencyRelevance : 1;
+	UE_DEPRECATED(4.25, "bHairStrandsRelevance has been renamed bHairStrands")
+	uint32 bHairStrandsRelevance : 1;
+#endif
+
 	// Warning: This class is memzeroed externally as 0 is assumed a
 	// valid value for all members meaning 'not relevant'. If this
 	// changes existing class usage should be re-evaluated
