fix for nightly test hitting assert that would hit if you called GEngine->GetMainAudioDevice before FAudioDeviceManager::Initialize was called.

diff --git a/Engine/Source/Runtime/Engine/Private/SoundWave.cpp b/Engine/Source/Runtime/Engine/Private/SoundWave.cpp
--- a/Engine/Source/Runtime/Engine/Private/SoundWave.cpp
+++ b/Engine/Source/Runtime/Engine/Private/SoundWave.cpp
@@ -835,7 +835,7 @@ void USoundWave::PostLoad()
 	// most likely cause us to run out of memory.
 	if (!GIsEditor && !IsTemplate( RF_ClassDefaultObject ) && GEngine)
 	{
-		FAudioDeviceHandle AudioDevice = GEngine->GetMainAudioDevice();
+		FAudioDevice* AudioDevice = GEngine->GetMainAudioDeviceRaw();
 		if (AudioDevice)
 		{
 			// Upload the data to the hardware, but only if we've precached startup sounds already
