Disable "deprecated warnings" to avoid compilation warning during copy / move constructor (clang with editor only)

#rnx

diff --git a/Engine/Source/Runtime/Engine/Public/PrimitiveViewRelevance.h b/Engine/Source/Runtime/Engine/Public/PrimitiveViewRelevance.h
--- a/Engine/Source/Runtime/Engine/Public/PrimitiveViewRelevance.h
+++ b/Engine/Source/Runtime/Engine/Public/PrimitiveViewRelevance.h
@@ -9,6 +9,7 @@
  * The different types of relevance a primitive scene proxy can declare towards a particular scene view.
  * the class is only storing bits, and has an |= operator
  */
+PRAGMA_DISABLE_DEPRECATION_WARNINGS
 struct FPrimitiveViewRelevance : public FMaterialRelevance
 {
 #if WITH_EDITOR
@@ -114,3 +115,4 @@ struct FPrimitiveViewRelevance : public FMaterialRelevance
 		return *this;
 	}
 };
+PRAGMA_ENABLE_DEPRECATION_WARNINGS
