Add debug context to DDC requests that take the key as a parameter

The existing GetSynchronous, GetAsynchronous, and Put functions are deprecated in favor of the DebugContext overloads. It is unlikely that a licensee has derived from FDerivedDataCacheInterface, but anyone who does will have a compile error that they override a non-virtual function and that they don't implement a pure virtual function, which will be straightforward to fix.

diff --git a/Engine/Plugins/FX/Niagara/Source/Niagara/Private/NiagaraSpriteRendererProperties.cpp b/Engine/Plugins/FX/Niagara/Source/Niagara/Private/NiagaraSpriteRendererProperties.cpp
--- a/Engine/Plugins/FX/Niagara/Source/Niagara/Private/NiagaraSpriteRendererProperties.cpp
+++ b/Engine/Plugins/FX/Niagara/Source/Niagara/Private/NiagaraSpriteRendererProperties.cpp
@@ -298,7 +298,7 @@ void UNiagaraSpriteRendererProperties::CacheDerivedData()
 		TArray<uint8> Data;
 
 		COOK_STAT(auto Timer = NiagaraCutoutCookStats::UsageStats.TimeSyncWork());
-		if (GetDerivedDataCacheRef().GetSynchronous(*KeyString, Data))
+		if (GetDerivedDataCacheRef().GetSynchronous(*KeyString, Data, GetPathName()))
 		{
 			COOK_STAT(Timer.AddHit(Data.Num()));
 			DerivedData.BoundingGeometry.Empty(Data.Num() / sizeof(FVector2D));
@@ -312,7 +312,7 @@ void UNiagaraSpriteRendererProperties::CacheDerivedData()
 			Data.Empty(DerivedData.BoundingGeometry.Num() * sizeof(FVector2D));
 			Data.AddUninitialized(DerivedData.BoundingGeometry.Num() * sizeof(FVector2D));
 			FPlatformMemory::Memcpy(Data.GetData(), DerivedData.BoundingGeometry.GetData(), DerivedData.BoundingGeometry.Num() * DerivedData.BoundingGeometry.GetTypeSize());
-			GetDerivedDataCacheRef().Put(*KeyString, Data);
+			GetDerivedDataCacheRef().Put(*KeyString, Data, GetPathName());
 			COOK_STAT(Timer.AddMiss(Data.Num()));
 		}
 	}
diff --git a/Engine/Plugins/FX/Niagara/Source/NiagaraShader/Private/NiagaraShader.cpp b/Engine/Plugins/FX/Niagara/Source/NiagaraShader/Private/NiagaraShader.cpp
--- a/Engine/Plugins/FX/Niagara/Source/NiagaraShader/Private/NiagaraShader.cpp
+++ b/Engine/Plugins/FX/Niagara/Source/NiagaraShader/Private/NiagaraShader.cpp
@@ -557,7 +557,7 @@ void FNiagaraShaderMap::LoadFromDerivedDataCache(const FNiagaraShaderScript* Scr
 			TArray<uint8> CachedData;
 			const FString DataKey = GetNiagaraShaderMapKeyString(ShaderMapId, Platform);
 
-			if (GetDerivedDataCacheRef().GetSynchronous(*DataKey, CachedData))
+			if (GetDerivedDataCacheRef().GetSynchronous(*DataKey, CachedData, Script->GetFriendlyName()))
 			{
 				COOK_STAT(Timer.AddHit(CachedData.Num()));
 				InOutShaderMap = new FNiagaraShaderMap();
@@ -590,7 +590,7 @@ void FNiagaraShaderMap::SaveToDerivedDataCache()
 	FMemoryWriter Ar(SaveData, true);
 	Serialize(Ar);
 
-	GetDerivedDataCacheRef().Put(*GetNiagaraShaderMapKeyString(ShaderMapId, Platform), SaveData);
+	GetDerivedDataCacheRef().Put(*GetNiagaraShaderMapKeyString(ShaderMapId, Platform), SaveData, FriendlyName);
 	COOK_STAT(Timer.AddMiss(SaveData.Num()));
 }
 
diff --git a/Engine/Source/Developer/DerivedDataCache/Private/DerivedDataCache.cpp b/Engine/Source/Developer/DerivedDataCache/Private/DerivedDataCache.cpp
--- a/Engine/Source/Developer/DerivedDataCache/Private/DerivedDataCache.cpp
+++ b/Engine/Source/Developer/DerivedDataCache/Private/DerivedDataCache.cpp
@@ -299,7 +299,7 @@ public:
 		DDC_SCOPE_CYCLE_COUNTER(DDC_GetSynchronous);
 		check(DataDeriver);
 		FString CacheKey = FDerivedDataCache::BuildCacheKey(DataDeriver);
-		UE_LOG(LogDerivedDataCache, Verbose, TEXT("GetSynchronous %s"), *CacheKey);
+		UE_LOG(LogDerivedDataCache, Verbose, TEXT("GetSynchronous %s from '%s'"), *CacheKey, *DataDeriver->GetDebugContextString());
 		FAsyncTask<FBuildAsyncWorker> PendingTask(DataDeriver, *CacheKey, true);
 		AddToAsyncCompletionCounter(1);
 		PendingTask.StartSynchronousTask();
@@ -317,7 +317,7 @@ public:
 		FScopeLock ScopeLock(&SynchronizationObject);
 		const uint32 Handle = NextHandle();
 		FString CacheKey = FDerivedDataCache::BuildCacheKey(DataDeriver);
-		UE_LOG(LogDerivedDataCache, Verbose, TEXT("GetAsynchronous %s"), *CacheKey);
+		UE_LOG(LogDerivedDataCache, Verbose, TEXT("GetAsynchronous %s from '%s'"), *CacheKey, *DataDeriver->GetDebugContextString());
 		const bool bSync = !DataDeriver->IsBuildThreadsafe();
 		FAsyncTask<FBuildAsyncWorker>* AsyncTask = new FAsyncTask<FBuildAsyncWorker>(DataDeriver, *CacheKey, bSync);
 		check(!PendingTasks.Contains(Handle));
@@ -389,10 +389,10 @@ public:
 		return true;
 	}
 
-	virtual bool GetSynchronous(const TCHAR* CacheKey, TArray<uint8>& OutData) override
+	virtual bool GetSynchronous(const TCHAR* CacheKey, TArray<uint8>& OutData, const FStringView& DataContext) override
 	{
 		DDC_SCOPE_CYCLE_COUNTER(DDC_GetSynchronous_Data);
-		UE_LOG(LogDerivedDataCache, Verbose, TEXT("GetSynchronous %s"), CacheKey);
+		UE_LOG(LogDerivedDataCache, Verbose, TEXT("GetSynchronous %s from '%.*s'"), CacheKey, DataContext.Len(), DataContext.GetData());
 		FAsyncTask<FBuildAsyncWorker> PendingTask((FDerivedDataPluginInterface*)NULL, CacheKey, true);
 		AddToAsyncCompletionCounter(1);
 		PendingTask.StartSynchronousTask();
@@ -400,11 +400,11 @@ public:
 		return PendingTask.GetTask().bSuccess;
 	}
 
-	virtual uint32 GetAsynchronous(const TCHAR* CacheKey) override
+	virtual uint32 GetAsynchronous(const TCHAR* CacheKey, const FStringView& DataContext) override
 	{
 		DDC_SCOPE_CYCLE_COUNTER(DDC_GetAsynchronous_Handle);
 		FScopeLock ScopeLock(&SynchronizationObject);
-		UE_LOG(LogDerivedDataCache, Verbose, TEXT("GetAsynchronous %s"), CacheKey);
+		UE_LOG(LogDerivedDataCache, Verbose, TEXT("GetAsynchronous %s from '%.*s'"), CacheKey, DataContext.Len(), DataContext.GetData());
 		const uint32 Handle = NextHandle();
 		FAsyncTask<FBuildAsyncWorker>* AsyncTask = new FAsyncTask<FBuildAsyncWorker>((FDerivedDataPluginInterface*)NULL, CacheKey, false);
 		check(!PendingTasks.Contains(Handle));
@@ -414,9 +414,10 @@ public:
 		return Handle;
 	}
 
-	virtual void Put(const TCHAR* CacheKey, TArray<uint8>& Data, bool bPutEvenIfExists = false) override
+	virtual void Put(const TCHAR* CacheKey, TArray<uint8>& Data, const FStringView& DataContext, bool bPutEvenIfExists = false) override
 	{
 		DDC_SCOPE_CYCLE_COUNTER(DDC_Put);
+		UE_LOG(LogDerivedDataCache, Verbose, TEXT("Put %s from '%.*s'"), CacheKey, DataContext.Len(), DataContext.GetData());
 		STAT(double ThisTime = 0);
 		{
 			SCOPE_SECONDS_COUNTER(ThisTime);
diff --git a/Engine/Source/Developer/DerivedDataCache/Public/DerivedDataCacheInterface.h b/Engine/Source/Developer/DerivedDataCache/Public/DerivedDataCacheInterface.h
--- a/Engine/Source/Developer/DerivedDataCache/Public/DerivedDataCacheInterface.h
+++ b/Engine/Source/Developer/DerivedDataCache/Public/DerivedDataCacheInterface.h
@@ -4,6 +4,7 @@
 #pragma once
 
 #include "CoreMinimal.h"
+#include "Containers/StringView.h"
 #include "Modules/ModuleInterface.h"
 
 class FDerivedDataCacheUsageStats;
@@ -129,21 +130,57 @@ public:
 	 * @param	CacheKey	Key to identify the data
 	 * @return	true if the data was retrieved from the cache
 	**/
-	virtual bool GetSynchronous(const TCHAR* CacheKey, TArray<uint8>& OutData) = 0; 
+	UE_DEPRECATED(4.25, "'GetSynchronous' without a DebugContext is deprecated. Please provide a DebugContext!")
+	bool GetSynchronous(const TCHAR* CacheKey, TArray<uint8>& OutData) { return GetSynchronous(CacheKey, OutData, TEXT("Unknown Context")); }
 
 	/** 
 	 * Starts the async process of checking the cache and if the item is present, retrieving the cached results
-	 * @param	CacheKey	Key to identify the data
+	 * @param	CacheKey		Key to identify the data
 	 * @return	a handle that can be used for PollAsynchronousCompletion, WaitAsynchronousCompletion and GetAsynchronousResults
 	**/
-	virtual uint32 GetAsynchronous(const TCHAR* CacheKey) = 0;
+	UE_DEPRECATED(4.25, "'GetAsynchronous' without a DebugContext is deprecated. Please provide a DebugContext!")
+	uint32 GetAsynchronous(const TCHAR* CacheKey) { return GetAsynchronous(CacheKey, TEXT("Unknown Context")); }
+
+	/** 
+	 * Puts data into the cache. This is fire-and-forget and typically asynchronous.
+	 * @param	CacheKey	Key to identify the data
+	 * @param	Data		Data to put in the cache under this key
+	**/
+	UE_DEPRECATED(4.25, "'Put' without a DebugContext is deprecated. Please provide a DebugContext!")
+	void Put(const TCHAR* CacheKey, TArray<uint8>& Data, bool bPutEvenIfExists = false) { Put(CacheKey, Data, TEXT("Unknown Context"), bPutEvenIfExists); }
+
+	/** 
+	 * Synchronously checks the cache and if the item is present, it returns the cached results, otherwise it returns false.
+	 *
+	 * Prefer to use a FDerivedDataPluginInterface instead of generating the key directly.
+	 *
+	 * @param	CacheKey		Key to identify the data.
+	 * @param	DebugContext	A string used to describe the data being generated. Typically the path to the object that it is generated from is sufficient.
+	 * @return	true if the data was retrieved from the cache.
+	**/
+	virtual bool GetSynchronous(const TCHAR* CacheKey, TArray<uint8>& OutData, const FStringView& DebugContext) = 0; 
+
+	/** 
+	 * Starts the async process of checking the cache and if the item is present, retrieving the cached results.
+	 *
+	 * Prefer to use a FDerivedDataPluginInterface instead of generating the key directly.
+	 *
+	 * @param	CacheKey		Key to identify the data.
+	 * @param	DebugContext	A string used to describe the data being generated. Typically the path to the object that it is generated from is sufficient.
+	 * @return	A handle that can be used for PollAsynchronousCompletion, WaitAsynchronousCompletion, and GetAsynchronousResults.
+	**/
+	virtual uint32 GetAsynchronous(const TCHAR* CacheKey, const FStringView& DebugContext) = 0;
 
 	/** 
 	 * Puts data into the cache. This is fire-and-forget and typically asynchronous.
+	 *
+	 * Prefer to use a FDerivedDataPluginInterface instead of generating the key directly.
+	 *
 	 * @param	CacheKey	Key to identify the data
 	 * @param	Data		Data to put in the cache under this key
+	 * @param	DataContext	A string used to describe the data being generated. Typically the path to the object that it is generated from is sufficient.
 	**/
-	virtual void Put(const TCHAR* CacheKey, TArray<uint8>& Data, bool bPutEvenIfExists = false) = 0;
+	virtual void Put(const TCHAR* CacheKey, TArray<uint8>& Data, const FStringView& DataContext, bool bPutEvenIfExists = false) = 0;
 
 	/**
 	 * Hint that the data associated with the key is transient and may be optionally purged from the cache.
diff --git a/Engine/Source/Developer/DerivedDataCache/Public/DerivedDataPluginInterface.h b/Engine/Source/Developer/DerivedDataCache/Public/DerivedDataPluginInterface.h
--- a/Engine/Source/Developer/DerivedDataCache/Public/DerivedDataPluginInterface.h
+++ b/Engine/Source/Developer/DerivedDataCache/Public/DerivedDataPluginInterface.h
@@ -40,7 +40,7 @@ public:
 	/** Indicated that this plugin generates deterministic data. This is used for DDC verification */
 	virtual bool IsDeterministic() const { return false; }
 
-	/** Indicated that this plugin generates deterministic data. This is used for DDC verification */
+	/** A string used to describe the data being generated. Typically the path to the object that it is generated from is sufficient. */
 	virtual FString GetDebugContextString() const { return TEXT("Unknown Context"); }
 
 	/** 
diff --git a/Engine/Source/Developer/PakFileUtilities/Private/PakFileUtilities.cpp b/Engine/Source/Developer/PakFileUtilities/Private/PakFileUtilities.cpp
--- a/Engine/Source/Developer/PakFileUtilities/Private/PakFileUtilities.cpp
+++ b/Engine/Source/Developer/PakFileUtilities/Private/PakFileUtilities.cpp
@@ -796,7 +796,7 @@ bool FCompressedFileBuffer::CompressFileToWorkingBuffer(const FPakInputPair& InF
 
 		if (GetDerivedDataCacheRef().CachedDataProbablyExists(*DDCKey))
 		{
-			int32 AsyncHandle = GetDerivedDataCacheRef().GetAsynchronous(*DDCKey);
+			int32 AsyncHandle = GetDerivedDataCacheRef().GetAsynchronous(*DDCKey, InFile.Dest);
 			GetDerivedDataCacheRef().WaitAsynchronousCompletion(AsyncHandle);
 			GetData.Empty(GetData.Max());
 			bool Result = false;
@@ -869,7 +869,7 @@ bool FCompressedFileBuffer::CompressFileToWorkingBuffer(const FPakInputPair& InF
 		GetData.Empty(GetData.Max());
 		FMemoryWriter Ar(GetData, true);
 		SerializeDDCData(Ar);
-		GetDerivedDataCacheRef().Put(*DDCKey, GetData);
+		GetDerivedDataCacheRef().Put(*DDCKey, GetData, InFile.Dest);
 	}
 
 
diff --git a/Engine/Source/Editor/UnrealEd/Private/StaticLightingSystem/StaticLightingDebug.cpp b/Engine/Source/Editor/UnrealEd/Private/StaticLightingSystem/StaticLightingDebug.cpp
--- a/Engine/Source/Editor/UnrealEd/Private/StaticLightingSystem/StaticLightingDebug.cpp
+++ b/Engine/Source/Editor/UnrealEd/Private/StaticLightingSystem/StaticLightingDebug.cpp
@@ -39,7 +39,7 @@ static void WriteTexel(UTexture2D* Texture, int32 X, int32 Y, FColor NewColor)
 		{
 			// The runtime data needs to be fully cached in memory for this to work.
 			// These changes won't (and don't need to) persist.
-			if (Texture->PlatformData->TryInlineMipData())
+			if (Texture->PlatformData->TryInlineMipData(0, Texture))
 			{
 				// Release the texture's resources and block until the rendering thread is done accessing it
 				Texture->ReleaseResource();
diff --git a/Engine/Source/Runtime/Engine/Classes/Engine/Texture.h b/Engine/Source/Runtime/Engine/Classes/Engine/Texture.h
--- a/Engine/Source/Runtime/Engine/Classes/Engine/Texture.h
+++ b/Engine/Source/Runtime/Engine/Classes/Engine/Texture.h
@@ -505,9 +505,10 @@ public:
 	 * @param OutMipData -	Must point to an array of pointers with at least
 	 *						Texture.Mips.Num() - FirstMipToLoad + 1 entries. Upon
 	 *						return those pointers will contain mip data.
+	 * @param Texture - The texture to load mips for.
 	 * @returns true if all requested mips have been loaded.
 	 */
-	bool TryLoadMips(int32 FirstMipToLoad, void** OutMipData);
+	bool TryLoadMips(int32 FirstMipToLoad, void** OutMipData, UTexture* Texture);
 
 	/** Serialization. */
 	void Serialize(FArchive& Ar, class UTexture* Owner);
@@ -572,7 +573,7 @@ public:
 		uint32 InFlags,
 		class ITextureCompressorModule* Compressor);
 	void FinishCache();
-	ENGINE_API bool TryInlineMipData(int32 FirstMipToLoad = 0);
+	ENGINE_API bool TryInlineMipData(int32 FirstMipToLoad = 0, UTexture* Texture = nullptr);
 	bool AreDerivedMipsAvailable() const;
 	bool AreDerivedVTChunksAvailable() const;
 #endif
diff --git a/Engine/Source/Runtime/Engine/Classes/Sound/SoundWave.h b/Engine/Source/Runtime/Engine/Classes/Sound/SoundWave.h
--- a/Engine/Source/Runtime/Engine/Classes/Sound/SoundWave.h
+++ b/Engine/Source/Runtime/Engine/Classes/Sound/SoundWave.h
@@ -79,7 +79,7 @@ struct FStreamedAudioChunk
 	 * Place chunk data in the derived data cache associated with the provided
 	 * key.
 	 */
-	uint32 StoreInDerivedDataCache(const FString& InDerivedDataKey);
+	uint32 StoreInDerivedDataCache(const FString& InDerivedDataKey, const FStringView& SoundWaveName);
 #endif // #if WITH_EDITORONLY_DATA
 };
 
diff --git a/Engine/Source/Runtime/Engine/Private/Animation/AnimSequence.cpp b/Engine/Source/Runtime/Engine/Private/Animation/AnimSequence.cpp
--- a/Engine/Source/Runtime/Engine/Private/Animation/AnimSequence.cpp
+++ b/Engine/Source/Runtime/Engine/Private/Animation/AnimSequence.cpp
@@ -2224,7 +2224,7 @@ void UAnimSequence::RequestAnimCompression(FRequestAnimCompressionParams Params)
 		// For debugging DDC/Compression issues		
 		const bool bSkipDDC = false;
 
-		if (bSkipDDC || !GetDerivedDataCacheRef().GetSynchronous(*FinalDDCKey, OutData))
+		if (bSkipDDC || !GetDerivedDataCacheRef().GetSynchronous(*FinalDDCKey, OutData, AnimCompressor->GetDebugContextString()))
 		{
 			// Data does not exist, need to build it.
 			// Filter RAW data to get rid of mismatched tracks (translation/rotation data with a different number of keys than there are frames)
diff --git a/Engine/Source/Runtime/Engine/Private/Animation/AnimStreamable.cpp b/Engine/Source/Runtime/Engine/Private/Animation/AnimStreamable.cpp
--- a/Engine/Source/Runtime/Engine/Private/Animation/AnimStreamable.cpp
+++ b/Engine/Source/Runtime/Engine/Private/Animation/AnimStreamable.cpp
@@ -610,7 +610,7 @@ void UAnimStreamable::RequestCompressedDataForChunk(const FString& ChunkDDCKey,
 		Chunk.StartTime = FrameStart * FrameLength;
 		Chunk.SequenceLength = ChunkNumFrames * FrameLength;
 
-		if (bSkipDDC || !GetDerivedDataCacheRef().GetSynchronous(*FinalDDCKey, OutData))
+		if (bSkipDDC || !GetDerivedDataCacheRef().GetSynchronous(*FinalDDCKey, OutData, GetPathName()))
 		{
 			FCompressibleAnimRef CompressibleData = MakeShared<FCompressibleAnimData, ESPMode::ThreadSafe>(BoneCompressionSettings, CurveCompressionSettings, GetSkeleton(), Interpolation, Chunk.SequenceLength, ChunkNumFrames+1);
 
diff --git a/Engine/Source/Runtime/Engine/Private/AudioDerivedData.cpp b/Engine/Source/Runtime/Engine/Private/AudioDerivedData.cpp
--- a/Engine/Source/Runtime/Engine/Private/AudioDerivedData.cpp
+++ b/Engine/Source/Runtime/Engine/Private/AudioDerivedData.cpp
@@ -206,7 +206,8 @@ static const FPlatformAudioCookOverrides* GetCookOverridesForRunningPlatform()
  */
 static uint32 PutDerivedDataInCache(
 	FStreamedAudioPlatformData* DerivedData,
-	const FString& DerivedDataKeySuffix
+	const FString& DerivedDataKeySuffix,
+	const FStringView& SoundWaveName
 	)
 {
 	TArray<uint8> RawDerivedData;
@@ -243,14 +244,14 @@ static uint32 PutDerivedDataInCache(
 				);
 		}
 
-		TotalBytesPut += Chunk.StoreInDerivedDataCache(ChunkDerivedDataKey);
+		TotalBytesPut += Chunk.StoreInDerivedDataCache(ChunkDerivedDataKey, SoundWaveName);
 	}
 
 	// Store derived data.
 	// At this point we've stored all the non-inline data in the DDC, so this will only serialize and store the metadata and any inline chunks
 	FMemoryWriter Ar(RawDerivedData, /*bIsPersistent=*/ true);
 	DerivedData->Serialize(Ar, NULL);
-	GetDerivedDataCacheRef().Put(*DerivedDataKey, RawDerivedData);
+	GetDerivedDataCacheRef().Put(*DerivedDataKey, RawDerivedData, SoundWaveName);
 	TotalBytesPut += RawDerivedData.Num();
 	UE_LOG(LogAudio, Verbose, TEXT("%s  Derived Data: %d bytes"), *LogString, RawDerivedData.Num());
 	return TotalBytesPut;
@@ -449,7 +450,7 @@ class FStreamedAudioCacheDerivedDataWorker : public FNonAbandonableTask
 				// @todo: This will remove the streaming bulk data, which we immediately reload below!
 				// Should ideally avoid this redundant work, but it only happens when we actually have
 				// to build the compressed audio, which should only ever be once.
-				this->BytesCached = PutDerivedDataInCache(DerivedData, KeySuffix);
+				this->BytesCached = PutDerivedDataInCache(DerivedData, KeySuffix, SoundWave.GetPathName());
 
 				check(this->BytesCached != 0);
 			}
@@ -504,7 +505,7 @@ public:
 		bool bForDDC = (CacheFlags & EStreamedAudioCacheFlags::ForDDCBuild) != 0;
 		bool bAllowAsyncBuild = (CacheFlags & EStreamedAudioCacheFlags::AllowAsyncBuild) != 0;
 
-		if (!bForceRebuild && GetDerivedDataCacheRef().GetSynchronous(*DerivedData->DerivedDataKey, RawDerivedData))
+		if (!bForceRebuild && GetDerivedDataCacheRef().GetSynchronous(*DerivedData->DerivedDataKey, RawDerivedData, SoundWave.GetPathName()))
 		{
 			BytesCached = RawDerivedData.Num();
 			FMemoryReader Ar(RawDerivedData, /*bIsPersistent=*/ true);
@@ -658,7 +659,7 @@ static void BeginLoadDerivedChunks(TIndirectArray<FStreamedAudioChunk>& Chunks,
 		const FStreamedAudioChunk& Chunk = Chunks[ChunkIndex];
 		if (!Chunk.DerivedDataKey.IsEmpty())
 		{
-			OutHandles[ChunkIndex] = DDC.GetAsynchronous(*Chunk.DerivedDataKey);
+			OutHandles[ChunkIndex] = DDC.GetAsynchronous(*Chunk.DerivedDataKey, TEXT("Unknown SoundWave"_SV));
 		}
 	}
 }
@@ -802,14 +803,14 @@ int32 FStreamedAudioPlatformData::GetChunkFromDDC(int32 ChunkIndex, uint8** OutC
 	{
 		if (bMakeSureChunkIsLoaded)
 		{
-			if (DDC.GetSynchronous(*Chunk.DerivedDataKey, TempData))
+			if (DDC.GetSynchronous(*Chunk.DerivedDataKey, TempData, TEXT("Unknown SoundWave"_SV)))
 			{
 				ChunkDataSize = DeserializeChunkFromDDC(TempData, Chunk, ChunkIndex, OutChunkData);
 			}
 		}
 		else
 		{
-			AsyncHandle = DDC.GetAsynchronous(*Chunk.DerivedDataKey);
+			AsyncHandle = DDC.GetAsynchronous(*Chunk.DerivedDataKey, TEXT("Unknown SoundWave"_SV));
 		}
 	}
 
diff --git a/Engine/Source/Runtime/Engine/Private/AudioStreaming.cpp b/Engine/Source/Runtime/Engine/Private/AudioStreaming.cpp
--- a/Engine/Source/Runtime/Engine/Private/AudioStreaming.cpp
+++ b/Engine/Source/Runtime/Engine/Private/AudioStreaming.cpp
@@ -68,7 +68,7 @@ void FAsyncStreamDerivedChunkWorker::DoWork()
 
 	TArray<uint8> DerivedChunkData;
 
-	if (GetDerivedDataCacheRef().GetSynchronous(*DerivedDataKey, DerivedChunkData))
+	if (GetDerivedDataCacheRef().GetSynchronous(*DerivedDataKey, DerivedChunkData, TEXT("Unknown Audio"_SV)))
 	{
 		FMemoryReader Ar(DerivedChunkData, true);
 		int32 ChunkSize = 0;
diff --git a/Engine/Source/Runtime/Engine/Private/DistanceFieldAtlas.cpp b/Engine/Source/Runtime/Engine/Private/DistanceFieldAtlas.cpp
--- a/Engine/Source/Runtime/Engine/Private/DistanceFieldAtlas.cpp
+++ b/Engine/Source/Runtime/Engine/Private/DistanceFieldAtlas.cpp
@@ -855,7 +855,7 @@ void FDistanceFieldVolumeData::CacheDerivedData(const FString& InDDCKey, UStatic
 	TArray<uint8> DerivedData;
 
 	COOK_STAT(auto Timer = DistanceFieldCookStats::UsageStats.TimeSyncWork());
-	if (GetDerivedDataCacheRef().GetSynchronous(*InDDCKey, DerivedData))
+	if (GetDerivedDataCacheRef().GetSynchronous(*InDDCKey, DerivedData, Mesh->GetPathName()))
 	{
 		COOK_STAT(Timer.AddHit(DerivedData.Num()));
 		FMemoryReader Ar(DerivedData, /*bIsPersistent=*/ true);
@@ -1210,7 +1210,7 @@ void FDistanceFieldAsyncQueue::ProcessAsyncTasks()
 				// Save built distance field volume to DDC
 				FMemoryWriter Ar(DerivedData, /*bIsPersistent=*/ true);
 				Ar << *(Task->StaticMesh->RenderData->LODResources[0].DistanceFieldData);
-				GetDerivedDataCacheRef().Put(*Task->DDCKey, DerivedData);
+				GetDerivedDataCacheRef().Put(*Task->DDCKey, DerivedData, Task->StaticMesh->GetPathName());
 				COOK_STAT(Timer.AddMiss(DerivedData.Num()));
 			}
 		}
diff --git a/Engine/Source/Runtime/Engine/Private/Materials/MaterialShader.cpp b/Engine/Source/Runtime/Engine/Private/Materials/MaterialShader.cpp
--- a/Engine/Source/Runtime/Engine/Private/Materials/MaterialShader.cpp
+++ b/Engine/Source/Runtime/Engine/Private/Materials/MaterialShader.cpp
@@ -1052,7 +1052,7 @@ void FMaterialShaderMap::LoadFromDerivedDataCache(const FMaterial* Material, con
 			TArray<uint8> CachedData;
 			const FString DataKey = GetMaterialShaderMapKeyString(ShaderMapId, InPlatform);
 
-			if (GetDerivedDataCacheRef().GetSynchronous(*DataKey, CachedData))
+			if (GetDerivedDataCacheRef().GetSynchronous(*DataKey, CachedData, Material->GetFriendlyName()))
 			{
 				COOK_STAT(Timer.AddHit(CachedData.Num()));
 				InOutShaderMap = new FMaterialShaderMap(InPlatform);
@@ -1086,7 +1086,7 @@ void FMaterialShaderMap::SaveToDerivedDataCache()
 	FMemoryWriter Ar(SaveData, true);
 	Serialize(Ar);
 
-	GetDerivedDataCacheRef().Put(*GetMaterialShaderMapKeyString(ShaderMapId, GetShaderPlatform()), SaveData);
+	GetDerivedDataCacheRef().Put(*GetMaterialShaderMapKeyString(ShaderMapId, GetShaderPlatform()), SaveData, FriendlyName);
 	COOK_STAT(Timer.AddMiss(SaveData.Num()));
 }
 #endif // WITH_EDITOR
diff --git a/Engine/Source/Runtime/Engine/Private/Particles/ParticleModules.cpp b/Engine/Source/Runtime/Engine/Private/Particles/ParticleModules.cpp
--- a/Engine/Source/Runtime/Engine/Private/Particles/ParticleModules.cpp
+++ b/Engine/Source/Runtime/Engine/Private/Particles/ParticleModules.cpp
@@ -1296,7 +1296,7 @@ void UParticleModuleRequired::CacheDerivedData()
 	TArray<uint8> Data;
 
 	COOK_STAT(auto Timer = SubUVAnimationCookStats::UsageStats.TimeSyncWork());
-	if (GetDerivedDataCacheRef().GetSynchronous(*KeyString, Data))
+	if (GetDerivedDataCacheRef().GetSynchronous(*KeyString, Data, GetPathName()))
 	{
 		COOK_STAT(Timer.AddHit(Data.Num()));
 		DerivedData.BoundingGeometry.Empty(Data.Num() / sizeof(FVector2D));
@@ -1310,7 +1310,7 @@ void UParticleModuleRequired::CacheDerivedData()
 		Data.Empty(DerivedData.BoundingGeometry.Num() * sizeof(FVector2D));
 		Data.AddUninitialized(DerivedData.BoundingGeometry.Num() * sizeof(FVector2D));
 		FPlatformMemory::Memcpy(Data.GetData(), DerivedData.BoundingGeometry.GetData(), DerivedData.BoundingGeometry.Num() * DerivedData.BoundingGeometry.GetTypeSize());
-		GetDerivedDataCacheRef().Put(*KeyString, Data);
+		GetDerivedDataCacheRef().Put(*KeyString, Data, GetPathName());
 		COOK_STAT(Timer.AddMiss(Data.Num()));
 	}
 #endif
diff --git a/Engine/Source/Runtime/Engine/Private/Particles/SubUVAnimation.cpp b/Engine/Source/Runtime/Engine/Private/Particles/SubUVAnimation.cpp
--- a/Engine/Source/Runtime/Engine/Private/Particles/SubUVAnimation.cpp
+++ b/Engine/Source/Runtime/Engine/Private/Particles/SubUVAnimation.cpp
@@ -103,7 +103,7 @@ void USubUVAnimation::CacheDerivedData()
 	TArray<uint8> Data;
 
 	COOK_STAT(auto Timer = SubUVAnimationCookStats::UsageStats.TimeSyncWork());
-	if (GetDerivedDataCacheRef().GetSynchronous(*KeyString, Data))
+	if (GetDerivedDataCacheRef().GetSynchronous(*KeyString, Data, GetPathName()))
 	{
 		COOK_STAT(Timer.AddHit(Data.Num()));
 		DerivedData.BoundingGeometry.Empty(Data.Num() / sizeof(FVector2D));
@@ -117,7 +117,7 @@ void USubUVAnimation::CacheDerivedData()
 		Data.Empty(DerivedData.BoundingGeometry.Num() * sizeof(FVector2D));
 		Data.AddUninitialized(DerivedData.BoundingGeometry.Num() * sizeof(FVector2D));
 		FPlatformMemory::Memcpy(Data.GetData(), DerivedData.BoundingGeometry.GetData(), DerivedData.BoundingGeometry.Num() * DerivedData.BoundingGeometry.GetTypeSize());
-		GetDerivedDataCacheRef().Put(*KeyString, Data);
+		GetDerivedDataCacheRef().Put(*KeyString, Data, GetPathName());
 		COOK_STAT(Timer.AddMiss(Data.Num()));
 	}
 #endif
diff --git a/Engine/Source/Runtime/Engine/Private/ShaderCompiler/ShaderCompiler.cpp b/Engine/Source/Runtime/Engine/Private/ShaderCompiler/ShaderCompiler.cpp
--- a/Engine/Source/Runtime/Engine/Private/ShaderCompiler/ShaderCompiler.cpp
+++ b/Engine/Source/Runtime/Engine/Private/ShaderCompiler/ShaderCompiler.cpp
@@ -4403,7 +4403,7 @@ void SaveGlobalShaderMapToDerivedDataCache(EShaderPlatform Platform)
 		FMemoryWriter Ar(SaveData, true);
 		SerializeGlobalShaders(Ar, GGlobalShaderMap[Platform], &ShaderKeysToSave);
 
-		GetDerivedDataCacheRef().Put(*GetGlobalShaderMapKeyString(ShaderMapId, Platform, ShaderFilenameDependencies.Value), SaveData);
+		GetDerivedDataCacheRef().Put(*GetGlobalShaderMapKeyString(ShaderMapId, Platform, ShaderFilenameDependencies.Value), SaveData, TEXT("GlobalShaderMap"_SV));
 		COOK_STAT(Timer.AddMiss(SaveData.Num()));
 	}
 }
@@ -4614,7 +4614,7 @@ void CompileGlobalShaderMap(EShaderPlatform Platform, bool bRefreshShaderMap)
 
 				const FString DataKey = GetGlobalShaderMapKeyString(ShaderMapId, Platform, ShaderFilenameDependencies.Value);
 
-				AsyncDDCRequestHandles[HandleIndex] = GetDerivedDataCacheRef().GetAsynchronous(*DataKey);
+				AsyncDDCRequestHandles[HandleIndex] = GetDerivedDataCacheRef().GetAsynchronous(*DataKey, TEXT("GlobalShaderMap"_SV));
 
 				++HandleIndex;
 			}
diff --git a/Engine/Source/Runtime/Engine/Private/SkeletalMeshRenderData.cpp b/Engine/Source/Runtime/Engine/Private/SkeletalMeshRenderData.cpp
--- a/Engine/Source/Runtime/Engine/Private/SkeletalMeshRenderData.cpp
+++ b/Engine/Source/Runtime/Engine/Private/SkeletalMeshRenderData.cpp
@@ -148,7 +148,7 @@ void FSkeletalMeshRenderData::Cache(USkeletalMesh* Owner)
 		FString DerivedDataKey = BuildSkeletalMeshDerivedDataKey(Owner);
 
 		TArray<uint8> DerivedData;
-		if (GetDerivedDataCacheRef().GetSynchronous(*DerivedDataKey, DerivedData))
+		if (GetDerivedDataCacheRef().GetSynchronous(*DerivedDataKey, DerivedData, Owner->GetPathName()))
 		{
 			COOK_STAT(Timer.AddHit(DerivedData.Num()));
 
@@ -341,7 +341,7 @@ void FSkeletalMeshRenderData::Cache(USkeletalMesh* Owner)
 				const bool bNeedsCPUAccess = FSkeletalMeshLODRenderData::ShouldKeepCPUResources(Owner, LODIndex, bForceKeepCPUResources);
 				LODData.SerializeStreamedData(Ar, Owner, LODIndex, LODStripFlags, bNeedsCPUAccess, bForceKeepCPUResources);
 			}
-			GetDerivedDataCacheRef().Put(*DerivedDataKey, DerivedData);
+			GetDerivedDataCacheRef().Put(*DerivedDataKey, DerivedData, Owner->GetPathName());
 
 			int32 T1 = FPlatformTime::Cycles();
 			UE_LOG(LogSkeletalMesh, Log, TEXT("Built Skeletal Mesh [%.2fs] %s"), FPlatformTime::ToMilliseconds(T1 - T0) / 1000.0f, *Owner->GetPathName());
diff --git a/Engine/Source/Runtime/Engine/Private/SoundWave.cpp b/Engine/Source/Runtime/Engine/Private/SoundWave.cpp
--- a/Engine/Source/Runtime/Engine/Private/SoundWave.cpp
+++ b/Engine/Source/Runtime/Engine/Private/SoundWave.cpp
@@ -149,7 +149,7 @@ void FStreamedAudioChunk::Serialize(FArchive& Ar, UObject* Owner, int32 ChunkInd
 }
 
 #if WITH_EDITORONLY_DATA
-uint32 FStreamedAudioChunk::StoreInDerivedDataCache(const FString& InDerivedDataKey)
+uint32 FStreamedAudioChunk::StoreInDerivedDataCache(const FString& InDerivedDataKey, const FStringView& SoundWaveName)
 {
 	int32 BulkDataSizeInBytes = BulkData.GetBulkDataSize();
 	check(BulkDataSizeInBytes > 0);
@@ -166,7 +166,7 @@ uint32 FStreamedAudioChunk::StoreInDerivedDataCache(const FString& InDerivedData
 	}
 
 	const uint32 Result = DerivedData.Num();
-	GetDerivedDataCacheRef().Put(*InDerivedDataKey, DerivedData);
+	GetDerivedDataCacheRef().Put(*InDerivedDataKey, DerivedData, SoundWaveName);
 	DerivedDataKey = InDerivedDataKey;
 	BulkData.RemoveBulkData();
 	return Result;
diff --git a/Engine/Source/Runtime/Engine/Private/StaticMesh.cpp b/Engine/Source/Runtime/Engine/Private/StaticMesh.cpp
--- a/Engine/Source/Runtime/Engine/Private/StaticMesh.cpp
+++ b/Engine/Source/Runtime/Engine/Private/StaticMesh.cpp
@@ -2364,7 +2364,7 @@ void FStaticMeshRenderData::Cache(UStaticMesh* Owner, const FStaticMeshLODSettin
 		DerivedDataKey = BuildStaticMeshDerivedDataKey(KeySuffix);
 
 		TArray<uint8> DerivedData;
-		if (GetDerivedDataCacheRef().GetSynchronous(*DerivedDataKey, DerivedData))
+		if (GetDerivedDataCacheRef().GetSynchronous(*DerivedDataKey, DerivedData, Owner->GetPathName()))
 		{
 			COOK_STAT(Timer.AddHit(DerivedData.Num()));
 			FMemoryReader Ar(DerivedData, /*bIsPersistent=*/ true);
@@ -2444,7 +2444,7 @@ void FStaticMeshRenderData::Cache(UStaticMesh* Owner, const FStaticMeshLODSettin
 #endif
 			if (bSaveDDC)
 			{
-				GetDerivedDataCacheRef().Put(*DerivedDataKey, DerivedData);
+				GetDerivedDataCacheRef().Put(*DerivedDataKey, DerivedData, Owner->GetPathName());
 			}
 
 			int32 T1 = FPlatformTime::Cycles();
@@ -3818,7 +3818,7 @@ bool UStaticMesh::LoadMeshDescription(int32 LodIndex, FMeshDescription& OutMeshD
 	if (GetMeshDataKey(LodIndex, MeshDataKey))
 	{
 		TArray<uint8> DerivedData;
-		if (GetDerivedDataCacheRef().GetSynchronous(*MeshDataKey, DerivedData))
+		if (GetDerivedDataCacheRef().GetSynchronous(*MeshDataKey, DerivedData, GetPathName()))
 		{
 			// If there was valid DDC data, we assume this is because the asset is an old one with valid RawMeshBulkData
 			check(!SourceModel.RawMeshBulkData->IsEmpty());
@@ -4136,7 +4136,7 @@ void UStaticMesh::CacheMeshData()
 						const bool bIsPersistent = true;
 						FMemoryWriter Ar(DerivedData, bIsPersistent);
 						MeshDescriptionBulkData.Serialize(Ar, this);
-						GetDerivedDataCacheRef().Put(*MeshDataKey, DerivedData);
+						GetDerivedDataCacheRef().Put(*MeshDataKey, DerivedData, GetPathName());
 					}
 				}
 			}
diff --git a/Engine/Source/Runtime/Engine/Private/Streaming/Texture2DMipDataProvider_DDC.cpp b/Engine/Source/Runtime/Engine/Private/Streaming/Texture2DMipDataProvider_DDC.cpp
--- a/Engine/Source/Runtime/Engine/Private/Streaming/Texture2DMipDataProvider_DDC.cpp
+++ b/Engine/Source/Runtime/Engine/Private/Streaming/Texture2DMipDataProvider_DDC.cpp
@@ -35,7 +35,7 @@ void FTexture2DMipDataProvider_DDC::Init(const FTextureUpdateContext& Context, c
 			const FTexture2DMipMap& OwnerMip = OwnerMips[MipIndex];
 			if (!OwnerMip.DerivedDataKey.IsEmpty())
 			{
-				DDCHandles[MipIndex] = GetDerivedDataCacheRef().GetAsynchronous(*OwnerMip.DerivedDataKey);
+				DDCHandles[MipIndex] = GetDerivedDataCacheRef().GetAsynchronous(*OwnerMip.DerivedDataKey, Texture2D->GetPathName());
 			}
 		}
 		*SyncOptions.bSnooze = true;
diff --git a/Engine/Source/Runtime/Engine/Private/Streaming/Texture2DStreamIn_DDC.cpp b/Engine/Source/Runtime/Engine/Private/Streaming/Texture2DStreamIn_DDC.cpp
--- a/Engine/Source/Runtime/Engine/Private/Streaming/Texture2DStreamIn_DDC.cpp
+++ b/Engine/Source/Runtime/Engine/Private/Streaming/Texture2DStreamIn_DDC.cpp
@@ -128,7 +128,7 @@ void FTexture2DStreamIn_DDC::DoCreateAsyncDDCRequests(const FContext& Context)
 			if (!MipMap.DerivedDataKey.IsEmpty())
 			{
 				check(!DDCHandles[MipIndex]);
-				DDCHandles[MipIndex] = GetDerivedDataCacheRef().GetAsynchronous(*MipMap.DerivedDataKey);
+				DDCHandles[MipIndex] = GetDerivedDataCacheRef().GetAsynchronous(*MipMap.DerivedDataKey, Context.Texture->GetPathName());
 
 #if !UE_BUILD_SHIPPING
 				// On some platforms the IO is too fast to test cancelation requests timing issues.
@@ -196,7 +196,7 @@ void FTexture2DStreamIn_DDC::DoLoadNewMipsFromDDC(const FContext& Context)
 				}
 				else
 				{
-					bDDCValid = GetDerivedDataCacheRef().GetSynchronous(*MipMap.DerivedDataKey, DerivedMipData);
+					bDDCValid = GetDerivedDataCacheRef().GetSynchronous(*MipMap.DerivedDataKey, DerivedMipData, Context.Texture->GetPathName());
 				}
 
 				if (bDDCValid)				
diff --git a/Engine/Source/Runtime/Engine/Private/Texture2D.cpp b/Engine/Source/Runtime/Engine/Private/Texture2D.cpp
--- a/Engine/Source/Runtime/Engine/Private/Texture2D.cpp
+++ b/Engine/Source/Runtime/Engine/Private/Texture2D.cpp
@@ -428,7 +428,7 @@ void FTexture2DMipMap::Serialize(FArchive& Ar, UObject* Owner, int32 MipIdx)
 }
 
 #if WITH_EDITORONLY_DATA
-uint32 FTexture2DMipMap::StoreInDerivedDataCache(const FString& InDerivedDataKey)
+uint32 FTexture2DMipMap::StoreInDerivedDataCache(const FString& InDerivedDataKey, const FStringView& TextureName)
 {
 	int32 BulkDataSizeInBytes = BulkData.GetBulkDataSize();
 	check(BulkDataSizeInBytes > 0);
@@ -442,7 +442,7 @@ uint32 FTexture2DMipMap::StoreInDerivedDataCache(const FString& InDerivedDataKey
 		BulkData.Unlock();
 	}
 	const uint32 Result = DerivedData.Num();
-	GetDerivedDataCacheRef().Put(*InDerivedDataKey, DerivedData);
+	GetDerivedDataCacheRef().Put(*InDerivedDataKey, DerivedData, TextureName);
 	DerivedDataKey = InDerivedDataKey;
 	BulkData.RemoveBulkData();
 	return Result;
diff --git a/Engine/Source/Runtime/Engine/Private/TextureDerivedData.cpp b/Engine/Source/Runtime/Engine/Private/TextureDerivedData.cpp
--- a/Engine/Source/Runtime/Engine/Private/TextureDerivedData.cpp
+++ b/Engine/Source/Runtime/Engine/Private/TextureDerivedData.cpp
@@ -550,7 +550,7 @@ static void GetBuildSettingsPerFormat(const UTexture& Texture, const FTextureBui
  * @param DerivedDataKeySuffix - The key suffix at which to store derived data.
  * @return number of bytes put to the DDC (total, including all mips)
  */
-uint32 PutDerivedDataInCache(FTexturePlatformData* DerivedData, const FString& DerivedDataKeySuffix)
+uint32 PutDerivedDataInCache(FTexturePlatformData* DerivedData, const FString& DerivedDataKeySuffix, const FStringView& TextureName)
 {
 	TArray<uint8> RawDerivedData;
 	FString DerivedDataKey;
@@ -596,7 +596,7 @@ uint32 PutDerivedDataInCache(FTexturePlatformData* DerivedData, const FString& D
 		if (!bInline)
 		{
 			// store in the DDC, also drop the bulk data storage.
-			TotalBytesPut += Mip.StoreInDerivedDataCache(MipDerivedDataKey);
+			TotalBytesPut += Mip.StoreInDerivedDataCache(MipDerivedDataKey, TextureName);
 		}
 	}
 
@@ -611,7 +611,7 @@ uint32 PutDerivedDataInCache(FTexturePlatformData* DerivedData, const FString& D
 				*FString::Printf(TEXT("%s_VTCHUNK%u"), *DerivedDataKeySuffix, ChunkIndex));
 
 			FVirtualTextureDataChunk& Chunk = DerivedData->VTData->Chunks[ChunkIndex];
-			TotalBytesPut += Chunk.StoreInDerivedDataCache(ChunkDerivedDataKey);
+			TotalBytesPut += Chunk.StoreInDerivedDataCache(ChunkDerivedDataKey, TextureName);
 		}
 	}
 
@@ -620,7 +620,7 @@ uint32 PutDerivedDataInCache(FTexturePlatformData* DerivedData, const FString& D
 	FMemoryWriter Ar(RawDerivedData, /*bIsPersistent=*/ true);
 	DerivedData->Serialize(Ar, NULL);
 	TotalBytesPut += RawDerivedData.Num();
-	GetDerivedDataCacheRef().Put(*DerivedDataKey, RawDerivedData, true);
+	GetDerivedDataCacheRef().Put(*DerivedDataKey, RawDerivedData, TextureName, /*bPutEvenIfExists*/ true);
 	UE_LOG(LogTexture,Verbose,TEXT("%s  Derived Data: %d bytes"),*LogString,RawDerivedData.Num());
 	return TotalBytesPut;
 }
@@ -875,7 +875,7 @@ typedef TArray<uint32> FAsyncVTChunkHandles;
  * @param FirstMipToLoad - Index of the first mip to retrieve.
  * @param OutHandles - Handles to the asynchronous DDC gets.
  */
-static void BeginLoadDerivedMips(TIndirectArray<FTexture2DMipMap>& Mips, int32 FirstMipToLoad, FAsyncMipHandles& OutHandles)
+static void BeginLoadDerivedMips(TIndirectArray<FTexture2DMipMap>& Mips, int32 FirstMipToLoad, UTexture* Texture, FAsyncMipHandles& OutHandles)
 {
 	FDerivedDataCacheInterface& DDC = GetDerivedDataCacheRef();
 	OutHandles.AddZeroed(Mips.Num());
@@ -884,12 +884,12 @@ static void BeginLoadDerivedMips(TIndirectArray<FTexture2DMipMap>& Mips, int32 F
 		const FTexture2DMipMap& Mip = Mips[MipIndex];
 		if (Mip.DerivedDataKey.IsEmpty() == false)
 		{
-			OutHandles[MipIndex] = DDC.GetAsynchronous(*Mip.DerivedDataKey);
+			OutHandles[MipIndex] = DDC.GetAsynchronous(*Mip.DerivedDataKey, Texture ? Texture->GetPathName() : TEXT("Unknown Texture"_SV));
 		}
 	}
 }
 
-static void BeginLoadDerivedVTChunks(const TArray<FVirtualTextureDataChunk>& Chunks, FAsyncVTChunkHandles& OutHandles)
+static void BeginLoadDerivedVTChunks(const TArray<FVirtualTextureDataChunk>& Chunks, UTexture* Texture, FAsyncVTChunkHandles& OutHandles)
 {
 	FDerivedDataCacheInterface& DDC = GetDerivedDataCacheRef();
 	OutHandles.AddZeroed(Chunks.Num());
@@ -898,7 +898,7 @@ static void BeginLoadDerivedVTChunks(const TArray<FVirtualTextureDataChunk>& Chu
 		const FVirtualTextureDataChunk& Chunk = Chunks[ChunkIndex];
 		if (Chunk.DerivedDataKey.IsEmpty() == false)
 		{
-			OutHandles[ChunkIndex] = DDC.GetAsynchronous(*Chunk.DerivedDataKey);
+			OutHandles[ChunkIndex] = DDC.GetAsynchronous(*Chunk.DerivedDataKey, Texture ? Texture->GetPathName() : TEXT("Unknown Texture"_SV));
 		}
 	}
 }
@@ -921,17 +921,17 @@ static void CheckMipSize(FTexture2DMipMap& Mip, EPixelFormat PixelFormat, int32
 	}
 }
 
-bool FTexturePlatformData::TryInlineMipData(int32 FirstMipToLoad)
+bool FTexturePlatformData::TryInlineMipData(int32 FirstMipToLoad, UTexture* Texture)
 {
 	FAsyncMipHandles AsyncHandles;
 	FAsyncVTChunkHandles AsyncVTHandles;
 	TArray<uint8> TempData;
 	FDerivedDataCacheInterface& DDC = GetDerivedDataCacheRef();
 
-	BeginLoadDerivedMips(Mips, FirstMipToLoad, AsyncHandles);
+	BeginLoadDerivedMips(Mips, FirstMipToLoad, Texture, AsyncHandles);
 	if (VTData != nullptr)
 	{
-		BeginLoadDerivedVTChunks(VTData->Chunks, AsyncVTHandles);
+		BeginLoadDerivedVTChunks(VTData->Chunks, Texture, AsyncVTHandles);
 	}
 
 	// Process regular mips
@@ -1046,7 +1046,7 @@ bool FTexturePlatformData::IsReadyForAsyncPostLoad() const
 	return true;
 }
 
-bool FTexturePlatformData::TryLoadMips(int32 FirstMipToLoad, void** OutMipData)
+bool FTexturePlatformData::TryLoadMips(int32 FirstMipToLoad, void** OutMipData, UTexture* Texture)
 {
 	int32 NumMipsCached = 0;
 	const int32 LoadableMips = Mips.Num() - ((GetNumMipsInTail() > 0) ? (GetNumMipsInTail() - 1) : 0);
@@ -1055,7 +1055,7 @@ bool FTexturePlatformData::TryLoadMips(int32 FirstMipToLoad, void** OutMipData)
 	TArray<uint8> TempData;
 	FAsyncMipHandles AsyncHandles;
 	FDerivedDataCacheInterface& DDC = GetDerivedDataCacheRef();
-	BeginLoadDerivedMips(Mips, FirstMipToLoad, AsyncHandles);
+	BeginLoadDerivedMips(Mips, FirstMipToLoad, Texture, AsyncHandles);
 #endif // #if WITH_EDITOR
 
 	// Handle the case where we inlined more mips than we intend to keep resident
@@ -1471,7 +1471,7 @@ void FTexturePlatformData::SerializeCooked(FArchive& Ar, UTexture* Owner, bool b
 
 void UTexture2D::GetMipData(int32 FirstMipToLoad, void** OutMipData)
 {
-	if (PlatformData->TryLoadMips(FirstMipToLoad, OutMipData) == false)
+	if (PlatformData->TryLoadMips(FirstMipToLoad, OutMipData, this) == false)
 	{
 		// Unable to load mips from the cache. Rebuild the texture and try again.
 		UE_LOG(LogTexture,Warning,TEXT("GetMipData failed for %s (%s)"),
@@ -1480,7 +1480,7 @@ void UTexture2D::GetMipData(int32 FirstMipToLoad, void** OutMipData)
 		if (!GetOutermost()->bIsCookedForEditor)
 		{
 			ForceRebuildPlatformData();
-			if (PlatformData->TryLoadMips(FirstMipToLoad, OutMipData) == false)
+			if (PlatformData->TryLoadMips(FirstMipToLoad, OutMipData, this) == false)
 			{
 				UE_LOG(LogTexture, Error, TEXT("Failed to build texture %s."), *GetPathName());
 			}
@@ -1491,7 +1491,7 @@ void UTexture2D::GetMipData(int32 FirstMipToLoad, void** OutMipData)
 
 void UTextureCube::GetMipData(int32 FirstMipToLoad, void** OutMipData)
 {
-	if (PlatformData->TryLoadMips(FirstMipToLoad, OutMipData) == false)
+	if (PlatformData->TryLoadMips(FirstMipToLoad, OutMipData, this) == false)
 	{
 		// Unable to load mips from the cache. Rebuild the texture and try again.
 		UE_LOG(LogTexture,Warning,TEXT("GetMipData failed for %s (%s)"),
@@ -1500,7 +1500,7 @@ void UTextureCube::GetMipData(int32 FirstMipToLoad, void** OutMipData)
 		if (!GetOutermost()->bIsCookedForEditor)
 		{
 			ForceRebuildPlatformData();
-			if (PlatformData->TryLoadMips(FirstMipToLoad, OutMipData) == false)
+			if (PlatformData->TryLoadMips(FirstMipToLoad, OutMipData, this) == false)
 			{
 				UE_LOG(LogTexture, Error, TEXT("Failed to build texture %s."), *GetPathName());
 			}
diff --git a/Engine/Source/Runtime/Engine/Private/TextureDerivedDataTask.cpp b/Engine/Source/Runtime/Engine/Private/TextureDerivedDataTask.cpp
--- a/Engine/Source/Runtime/Engine/Private/TextureDerivedDataTask.cpp
+++ b/Engine/Source/Runtime/Engine/Private/TextureDerivedDataTask.cpp
@@ -234,12 +234,12 @@ void FTextureCacheDerivedDataWorker::BuildTexture()
 		// @todo: This will remove the streaming bulk data, which we immediately reload below!
 		// Should ideally avoid this redundant work, but it only happens when we actually have 
 		// to build the texture, which should only ever be once.
-		this->BytesCached = PutDerivedDataInCache(DerivedData, KeySuffix);
+		this->BytesCached = PutDerivedDataInCache(DerivedData, KeySuffix, Texture.GetPathName());
 
 		if (DerivedData->VTData->Chunks.Num())
 		{
 			const bool bInlineMips = (CacheFlags & ETextureCacheFlags::InlineMips) != 0;
-			bSucceeded = !bInlineMips || DerivedData->TryInlineMipData(BuildSettingsPerLayer[0].LODBiasWithCinematicMips);
+			bSucceeded = !bInlineMips || DerivedData->TryInlineMipData(BuildSettingsPerLayer[0].LODBiasWithCinematicMips, &Texture);
 		}
 		else
 		{
@@ -309,13 +309,13 @@ void FTextureCacheDerivedDataWorker::BuildTexture()
 			// @todo: This will remove the streaming bulk data, which we immediately reload below!
 			// Should ideally avoid this redundant work, but it only happens when we actually have 
 			// to build the texture, which should only ever be once.
-			this->BytesCached = PutDerivedDataInCache(DerivedData, KeySuffix);
+			this->BytesCached = PutDerivedDataInCache(DerivedData, KeySuffix, Texture.GetPathName());
 		}
 
 		if (DerivedData->Mips.Num())
 		{
 			const bool bInlineMips = (CacheFlags & ETextureCacheFlags::InlineMips) != 0;
-			bSucceeded = !bInlineMips || DerivedData->TryInlineMipData(BuildSettingsPerLayer[0].LODBiasWithCinematicMips);
+			bSucceeded = !bInlineMips || DerivedData->TryInlineMipData(BuildSettingsPerLayer[0].LODBiasWithCinematicMips, &Texture);
 		}
 		else
 		{
@@ -409,7 +409,7 @@ void FTextureCacheDerivedDataWorker::DoWork()
 
 	TArray<uint8> RawDerivedData;
 
-	if (!bForceRebuild && GetDerivedDataCacheRef().GetSynchronous(*DerivedData->DerivedDataKey, RawDerivedData))
+	if (!bForceRebuild && GetDerivedDataCacheRef().GetSynchronous(*DerivedData->DerivedDataKey, RawDerivedData, Texture.GetPathName()))
 	{
 		const bool bInlineMips = (CacheFlags & ETextureCacheFlags::InlineMips) != 0;
 		const bool bForDDC = (CacheFlags & ETextureCacheFlags::ForDDCBuild) != 0;
@@ -421,11 +421,11 @@ void FTextureCacheDerivedDataWorker::DoWork()
 		// Load any streaming (not inline) mips that are necessary for our platform.
 		if (bForDDC)
 		{
-			bSucceeded = DerivedData->TryLoadMips(0,NULL);
+			bSucceeded = DerivedData->TryLoadMips(0, nullptr, &Texture);
 		}
 		else if (bInlineMips)
 		{
-			bSucceeded = DerivedData->TryInlineMipData(BuildSettingsPerLayer[0].LODBiasWithCinematicMips);
+			bSucceeded = DerivedData->TryInlineMipData(BuildSettingsPerLayer[0].LODBiasWithCinematicMips, &Texture);
 		}
 		else
 		{
diff --git a/Engine/Source/Runtime/Engine/Private/TextureDerivedDataTask.h b/Engine/Source/Runtime/Engine/Private/TextureDerivedDataTask.h
--- a/Engine/Source/Runtime/Engine/Private/TextureDerivedDataTask.h
+++ b/Engine/Source/Runtime/Engine/Private/TextureDerivedDataTask.h
@@ -27,7 +27,7 @@ enum
 #if WITH_EDITOR
 
 void GetTextureDerivedDataKeySuffix(const UTexture& Texture, const FTextureBuildSettings* BuildSettingsPerLayer, FString& OutKeySuffix);
-uint32 PutDerivedDataInCache(FTexturePlatformData* DerivedData, const FString& DerivedDataKeySuffix);
+uint32 PutDerivedDataInCache(FTexturePlatformData* DerivedData, const FString& DerivedDataKeySuffix, const FStringView& TextureName);
 
 namespace ETextureCacheFlags
 {
diff --git a/Engine/Source/Runtime/Engine/Private/VT/VirtualTextureBuiltData.cpp b/Engine/Source/Runtime/Engine/Private/VT/VirtualTextureBuiltData.cpp
--- a/Engine/Source/Runtime/Engine/Private/VT/VirtualTextureBuiltData.cpp
+++ b/Engine/Source/Runtime/Engine/Private/VT/VirtualTextureBuiltData.cpp
@@ -214,7 +214,7 @@ bool FVirtualTextureDataChunk::ShortenKey(const FString& CacheKey, FString& Resu
 	return true;
 }
 
-uint32 FVirtualTextureDataChunk::StoreInDerivedDataCache(const FString& InDerivedDataKey)
+uint32 FVirtualTextureDataChunk::StoreInDerivedDataCache(const FString& InDerivedDataKey, const FStringView& TextureName)
 {
 	int32 BulkDataSizeInBytes = BulkData.GetBulkDataSize();
 	check(BulkDataSizeInBytes > 0);
@@ -228,7 +228,7 @@ uint32 FVirtualTextureDataChunk::StoreInDerivedDataCache(const FString& InDerive
 		BulkData.Unlock();
 	}
 	const uint32 Result = DerivedData.Num();
-	GetDerivedDataCacheRef().Put(*InDerivedDataKey, DerivedData);
+	GetDerivedDataCacheRef().Put(*InDerivedDataKey, DerivedData, TextureName);
 	DerivedDataKey = InDerivedDataKey;
 	ShortenKey(DerivedDataKey, ShortDerivedDataKey);
 
diff --git a/Engine/Source/Runtime/Engine/Private/VT/VirtualTextureBuiltData.h b/Engine/Source/Runtime/Engine/Private/VT/VirtualTextureBuiltData.h
--- a/Engine/Source/Runtime/Engine/Private/VT/VirtualTextureBuiltData.h
+++ b/Engine/Source/Runtime/Engine/Private/VT/VirtualTextureBuiltData.h
@@ -60,7 +60,7 @@ struct FVirtualTextureDataChunk
 	bool ShortenKey(const FString& CacheKey, FString& Result);
 	FThreadSafeBool bFileAvailableInVTDDCDache;
 
-	uint32 StoreInDerivedDataCache(const FString& InDerivedDataKey);
+	uint32 StoreInDerivedDataCache(const FString& InDerivedDataKey, const FStringView& TextureName);
 #endif // WITH_EDITORONLY_DATA
 };
 
diff --git a/Engine/Source/Runtime/Engine/Private/VT/VirtualTextureChunkDDCCache.cpp b/Engine/Source/Runtime/Engine/Private/VT/VirtualTextureChunkDDCCache.cpp
--- a/Engine/Source/Runtime/Engine/Private/VT/VirtualTextureChunkDDCCache.cpp
+++ b/Engine/Source/Runtime/Engine/Private/VT/VirtualTextureChunkDDCCache.cpp
@@ -175,7 +175,7 @@ public:
 		// Fetch data from DDC
 		TArray<uint8> Results;
 		//TODO(ddebaets) this sync request seems to be blocking here while it uses the job pool. add overload to perform it on this thread?
-		const bool DDCResult = DDC.GetSynchronous(*Chunk->DerivedDataKey, Results);
+		const bool DDCResult = DDC.GetSynchronous(*Chunk->DerivedDataKey, Results, Filename);
 		if (DDCResult == false)
 		{
 			UE_LOG(LogVTDiskCache, Error, TEXT("Failed to fetch data from DDC (key: %s)"), *Chunk->DerivedDataKey);
diff --git a/Engine/Source/Runtime/Engine/Private/VolumeTexture.cpp b/Engine/Source/Runtime/Engine/Private/VolumeTexture.cpp
--- a/Engine/Source/Runtime/Engine/Private/VolumeTexture.cpp
+++ b/Engine/Source/Runtime/Engine/Private/VolumeTexture.cpp
@@ -376,7 +376,7 @@ public:
 		bGreyScaleFormat = (PixelFormat == PF_G8) || (PixelFormat == PF_BC4);
 
 		FTexturePlatformData* PlatformData = Owner->PlatformData;
-		if (PlatformData && PlatformData->TryLoadMips(MipBias, InitialData.GetMipData() + MipBias))
+		if (PlatformData && PlatformData->TryLoadMips(MipBias, InitialData.GetMipData() + MipBias, Owner))
 		{
 			for (int32 MipIndex = MipBias; MipIndex < NumMips; ++MipIndex)
 			{
diff --git a/Engine/Source/Runtime/Engine/Public/TextureResource.h b/Engine/Source/Runtime/Engine/Public/TextureResource.h
--- a/Engine/Source/Runtime/Engine/Public/TextureResource.h
+++ b/Engine/Source/Runtime/Engine/Public/TextureResource.h
@@ -138,7 +138,7 @@ struct FTexture2DMipMap
 	 * Place mip-map data in the derived data cache associated with the provided
 	 * key.
 	 */
-	uint32 StoreInDerivedDataCache(const FString& InDerivedDataKey);
+	uint32 StoreInDerivedDataCache(const FString& InDerivedDataKey, const FStringView& TextureName);
 #endif // #if WITH_EDITORONLY_DATA
 };
 
diff --git a/Engine/Source/Runtime/Landscape/Classes/LandscapeComponent.h b/Engine/Source/Runtime/Landscape/Classes/LandscapeComponent.h
--- a/Engine/Source/Runtime/Landscape/Classes/LandscapeComponent.h
+++ b/Engine/Source/Runtime/Landscape/Classes/LandscapeComponent.h
@@ -141,10 +141,10 @@ public:
 	static FString GetDDCKeyString(const FGuid& StateId);
 
 	/** Loads the platform data from DDC */
-	bool LoadFromDDC(const FGuid& StateId);
+	bool LoadFromDDC(const FGuid& StateId, UObject* Component);
 
 	/** Saves the compressed platform data to the DDC */
-	void SaveToDDC(const FGuid& StateId);
+	void SaveToDDC(const FGuid& StateId, UObject* Component);
 
 	/* Serializer */
 	friend FArchive& operator<<(FArchive& Ar, FLandscapeComponentDerivedData& Data);
diff --git a/Engine/Source/Runtime/Landscape/Private/Landscape.cpp b/Engine/Source/Runtime/Landscape/Private/Landscape.cpp
--- a/Engine/Source/Runtime/Landscape/Private/Landscape.cpp
+++ b/Engine/Source/Runtime/Landscape/Private/Landscape.cpp
@@ -285,14 +285,14 @@ void ULandscapeComponent::CheckGenerateLandscapePlatformData(bool bIsCooking, co
 			// The DDC is only useful when cooking (see else).
 
 			COOK_STAT(auto Timer = LandscapeCookStats::UsageStats.TimeSyncWork());
-			if (PlatformData.LoadFromDDC(NewSourceHash))
+			if (PlatformData.LoadFromDDC(NewSourceHash, this))
 			{
 				COOK_STAT(Timer.AddHit(PlatformData.GetPlatformDataSize()));
 			}
 			else
 			{
 				GeneratePlatformVertexData(TargetPlatform);
-				PlatformData.SaveToDDC(NewSourceHash);
+				PlatformData.SaveToDDC(NewSourceHash, this);
 				COOK_STAT(Timer.AddMiss(PlatformData.GetPlatformDataSize()));
 			}
 		}
@@ -3358,15 +3358,15 @@ FArchive& operator<<(FArchive& Ar, FLandscapeComponentDerivedData& Data)
 	return Ar << Data.CompressedLandscapeData;
 }
 
-bool FLandscapeComponentDerivedData::LoadFromDDC(const FGuid& StateId)
+bool FLandscapeComponentDerivedData::LoadFromDDC(const FGuid& StateId, UObject* Component)
 {
-	return GetDerivedDataCacheRef().GetSynchronous(*GetDDCKeyString(StateId), CompressedLandscapeData);
+	return GetDerivedDataCacheRef().GetSynchronous(*GetDDCKeyString(StateId), CompressedLandscapeData, Component->GetPathName());
 }
 
-void FLandscapeComponentDerivedData::SaveToDDC(const FGuid& StateId)
+void FLandscapeComponentDerivedData::SaveToDDC(const FGuid& StateId, UObject* Component)
 {
 	check(CompressedLandscapeData.Num() > 0);
-	GetDerivedDataCacheRef().Put(*GetDDCKeyString(StateId), CompressedLandscapeData);
+	GetDerivedDataCacheRef().Put(*GetDDCKeyString(StateId), CompressedLandscapeData, Component->GetPathName());
 }
 
 void LandscapeMaterialsParameterValuesGetter(FStaticParameterSet& OutStaticParameterSet, UMaterialInstance* Material)
diff --git a/Engine/Source/Runtime/Landscape/Private/LandscapeCollision.cpp b/Engine/Source/Runtime/Landscape/Private/LandscapeCollision.cpp
--- a/Engine/Source/Runtime/Landscape/Private/LandscapeCollision.cpp
+++ b/Engine/Source/Runtime/Landscape/Private/LandscapeCollision.cpp
@@ -724,7 +724,7 @@ void ULandscapeHeightfieldCollisionComponent::SpeculativelyLoadAsyncDDCCollsionD
 			static FName PhysicsFormatName(FPlatformProperties::GetPhysicsFormat());
 
 			FString Key = GetHFDDCKeyString(PhysicsFormatName, false, HeightfieldGuid, CookedPhysicalMaterials);
-			uint32 Handle = GetDerivedDataCacheRef().GetAsynchronous(*Key);
+			uint32 Handle = GetDerivedDataCacheRef().GetAsynchronous(*Key, GetPathName());
 			check(!SpeculativeDDCRequest.IsValid());
 			SpeculativeDDCRequest = MakeShareable(new FAsyncPreRegisterDDCRequest(Key, Handle));
 			World->AsyncPreRegisterDDCRequests.Add(SpeculativeDDCRequest);
@@ -840,7 +840,7 @@ bool ULandscapeHeightfieldCollisionComponent::CookCollisionData(const FName& For
 				}
 			}
 
-			if (GetDerivedDataCacheRef().GetSynchronous(*DDCKey, OutCookedData))
+			if (GetDerivedDataCacheRef().GetSynchronous(*DDCKey, OutCookedData, GetPathName()))
 			{
 				COOK_STAT(Timer.AddHit(OutCookedData.Num()));
 				bShouldSaveCookedDataToDDC[CookedDataIndex] = false;
@@ -989,7 +989,7 @@ bool ULandscapeHeightfieldCollisionComponent::CookCollisionData(const FName& For
 
 		if (!GLandscapeCollisionSkipDDC && bShouldSaveCookedDataToDDC[CookedDataIndex] && HeightfieldGuid.IsValid())
 		{
-			GetDerivedDataCacheRef().Put(*GetHFDDCKeyString(Format, bUseDefMaterial, HeightfieldGuid, InOutMaterials), OutCookedData);
+			GetDerivedDataCacheRef().Put(*GetHFDDCKeyString(Format, bUseDefMaterial, HeightfieldGuid, InOutMaterials), OutCookedData, GetPathName());
 			bShouldSaveCookedDataToDDC[CookedDataIndex] = false;
 		}
 	}
@@ -1041,7 +1041,7 @@ bool ULandscapeMeshCollisionComponent::CookCollisionData(const FName& Format, bo
 				}
 			}
 
-			if (GetDerivedDataCacheRef().GetSynchronous(*DDCKey, OutCookedData))
+			if (GetDerivedDataCacheRef().GetSynchronous(*DDCKey, OutCookedData, GetPathName()))
 			{
 				COOK_STAT(Timer.AddHit(OutCookedData.Num()));
 				bShouldSaveCookedDataToDDC[CookedDataIndex] = false;
@@ -1189,7 +1189,7 @@ bool ULandscapeMeshCollisionComponent::CookCollisionData(const FName& Format, bo
 
 		if (!GLandscapeCollisionSkipDDC && bShouldSaveCookedDataToDDC[CookedDataIndex] && MeshGuid.IsValid())
 		{
-			GetDerivedDataCacheRef().Put(*GetHFDDCKeyString(Format, bUseDefMaterial, MeshGuid, InOutMaterials), OutCookedData);
+			GetDerivedDataCacheRef().Put(*GetHFDDCKeyString(Format, bUseDefMaterial, MeshGuid, InOutMaterials), OutCookedData, GetPathName());
 			bShouldSaveCookedDataToDDC[CookedDataIndex] = false;
 		}
 	}
@@ -1778,7 +1778,7 @@ void ULandscapeHeightfieldCollisionComponent::Serialize(FArchive& Ar)
 			CookCollisionData(Format, false, true, CookedCollisionData, CookedPhysicalMaterials);
 			if (!GLandscapeCollisionSkipDDC && HeightfieldGuid.IsValid())
 			{
-				GetDerivedDataCacheRef().Put(*GetHFDDCKeyString(Format, false, HeightfieldGuid, CookedPhysicalMaterials), CookedCollisionData);
+				GetDerivedDataCacheRef().Put(*GetHFDDCKeyString(Format, false, HeightfieldGuid, CookedPhysicalMaterials), CookedCollisionData, GetPathName());
 			}
 		}
 	}
@@ -2117,12 +2117,12 @@ void ULandscapeHeightfieldCollisionComponent::PreSave(const class ITargetPlatfor
 			static const FName PhysicsFormatName(FPlatformProperties::GetPhysicsFormat());
 			if(CookedCollisionData.Num() && HeightfieldGuid.IsValid())
 			{
-				GetDerivedDataCacheRef().Put(*GetHFDDCKeyString(PhysicsFormatName, false, HeightfieldGuid, CookedPhysicalMaterials), CookedCollisionData);
+				GetDerivedDataCacheRef().Put(*GetHFDDCKeyString(PhysicsFormatName, false, HeightfieldGuid, CookedPhysicalMaterials), CookedCollisionData, GetPathName());
 			}
 
 			if(CookedCollisionDataEd.Num() && HeightfieldGuid.IsValid())
 			{
-				GetDerivedDataCacheRef().Put(*GetHFDDCKeyString(PhysicsFormatName, true, HeightfieldGuid, TArray<UPhysicalMaterial*>()), CookedCollisionDataEd);
+				GetDerivedDataCacheRef().Put(*GetHFDDCKeyString(PhysicsFormatName, true, HeightfieldGuid, TArray<UPhysicalMaterial*>()), CookedCollisionDataEd, GetPathName());
 			}
 		}
 #endif// WITH_EDITOR
