Fixed broken resending of stateless handshake acks, due to saving a shared FInternetAddr, instead of cloning it.

Bug report and fix provided by [REDACTED_PERSON], from Sumo Digital.

diff --git a/Engine/Source/Runtime/Engine/Private/PacketHandlers/StatelessConnectHandlerComponent.cpp b/Engine/Source/Runtime/Engine/Private/PacketHandlers/StatelessConnectHandlerComponent.cpp
--- a/Engine/Source/Runtime/Engine/Private/PacketHandlers/StatelessConnectHandlerComponent.cpp
+++ b/Engine/Source/Runtime/Engine/Private/PacketHandlers/StatelessConnectHandlerComponent.cpp
@@ -165,7 +165,7 @@ DEFINE_LOG_CATEGORY(LogHandshake);
 
 
 // Whether or not clients should send diagnostics to the server with the restart handshake, detailing why the request was accepted.
-#define RESTART_HANDSHAKE_DIAGNOSTICS 0
+#define RESTART_HANDSHAKE_DIAGNOSTICS 1
 
 // Disable client sending of handshake diagnostics
 #define DISABLE_SEND_HANDSHAKE_DIAGNOSTICS 1
@@ -1034,7 +1034,7 @@ void StatelessConnectHandlerComponent::IncomingConnectionless(const TSharedPtr<c
 							}
 
 							bRestartedHandshake = bRestartHandshake;
-							LastChallengeSuccessAddress = Address;
+							LastChallengeSuccessAddress = Address->Clone();
 
 
 							// Now ack the challenge response - the cookie is stored in AuthorisedCookie, to enable retries
