Remove unintential handshake diagnostics enable - to get it to filter-out from UE4/Main.

diff --git a/Engine/Source/Runtime/Engine/Private/PacketHandlers/StatelessConnectHandlerComponent.cpp b/Engine/Source/Runtime/Engine/Private/PacketHandlers/StatelessConnectHandlerComponent.cpp
--- a/Engine/Source/Runtime/Engine/Private/PacketHandlers/StatelessConnectHandlerComponent.cpp
+++ b/Engine/Source/Runtime/Engine/Private/PacketHandlers/StatelessConnectHandlerComponent.cpp
@@ -165,7 +165,7 @@ DEFINE_LOG_CATEGORY(LogHandshake);
 
 
 // Whether or not clients should send diagnostics to the server with the restart handshake, detailing why the request was accepted.
-#define RESTART_HANDSHAKE_DIAGNOSTICS 1
+#define RESTART_HANDSHAKE_DIAGNOSTICS 0
 
 // Disable client sending of handshake diagnostics
 #define DISABLE_SEND_HANDSHAKE_DIAGNOSTICS 1
