CIS error (signed/unsigned mistmatch)

diff --git a/Engine/Source/Runtime/Engine/Private/SoundWave.cpp b/Engine/Source/Runtime/Engine/Private/SoundWave.cpp
--- a/Engine/Source/Runtime/Engine/Private/SoundWave.cpp
+++ b/Engine/Source/Runtime/Engine/Private/SoundWave.cpp
@@ -2399,7 +2399,7 @@ void USoundWave::GetHandleForChunkOfAudio(TFunction<void(FAudioChunkHandle)> OnL
 {
 	// if we are requesting a chunk that is out of bounds,
 	// early exit.
-	if (ChunkIndex >= GetNumChunks())
+	if (ChunkIndex >= static_cast<int32>(GetNumChunks()))
 	{
 		FAudioChunkHandle EmptyChunkHandle;
 		OnLoadCompleted(EmptyChunkHandle);
