Fix for ensure being hit in editor/game tests in nightly. Issue turns out to be that we used to return false on FAudioDeviceManager::Initialize if the main audio device init failed. Situation that happens here:
(1) Main Audio Device fails to init, but audio thread is started
(2)  the game viewport client requests an audio device, which should return the main audio device but the main audio device wasn't created, so it tries to create the audio device
(3) new requested audio device also fails to init, but this time the audio thread is started

diff --git a/Engine/Source/Runtime/Engine/Private/AudioDeviceManager.cpp b/Engine/Source/Runtime/Engine/Private/AudioDeviceManager.cpp
--- a/Engine/Source/Runtime/Engine/Private/AudioDeviceManager.cpp
+++ b/Engine/Source/Runtime/Engine/Private/AudioDeviceManager.cpp
@@ -371,9 +371,15 @@ bool FAudioDeviceManager::Initialize()
 		MainDeviceParams.bIsNonRealtime = false;
 
 		MainAudioDeviceHandle = RequestAudioDevice(MainDeviceParams);
-		UE_CLOG(!MainAudioDeviceHandle.IsValid(), LogAudio, Display, TEXT("Audio device could not be initialized. Please check the value for AudioDeviceModuleName and AudioMixerModuleName in [Platform]Engine.ini."));
+		
+		if (!MainAudioDeviceHandle)
+		{
+			UE_LOG(LogAudio, Display, TEXT("Audio device could not be initialized. Please check the value for AudioDeviceModuleName and AudioMixerModuleName in [Platform]Engine.ini."));
+			return false;
+		}
 
 		FAudioThread::StartAudioThread();
+
 		return true;
 	}
 
