Improved time boxing for replay checkpoint saving: implemented time boxing for SerializeNetGuidCache step.

diff --git a/Engine/Source/Runtime/Engine/Classes/Engine/DemoNetDriver.h b/Engine/Source/Runtime/Engine/Classes/Engine/DemoNetDriver.h
--- a/Engine/Source/Runtime/Engine/Classes/Engine/DemoNetDriver.h
+++ b/Engine/Source/Runtime/Engine/Classes/Engine/DemoNetDriver.h
@@ -1081,6 +1081,7 @@ private:
 		ECheckpointSaveState_Idle,
 		ECheckpointSaveState_ProcessCheckpointActors,
 		ECheckpointSaveState_SerializeDeletedStartupActors,
+		ECheckpointSaveState_CacheNetGuids,
 		ECheckpointSaveState_SerializeGuidCache,
 		ECheckpointSaveState_SerializeNetFieldExportGroupMap,
 		ECheckpointSaveState_SerializeDemoFrameFromQueuedDemoPackets,
@@ -1231,7 +1232,8 @@ protected:
 
 	bool DemoReplicateActor(AActor* Actor, UNetConnection* Connection, bool bMustReplicate);
 
-	void SerializeGuidCache(TSharedPtr<class FNetGUIDCache> GuidCache, FArchive* CheckpointArchive);
+	void CacheNetGuids();
+	bool SerializeGuidCache(const FRepActorsCheckpointParams& Params, FArchive* CheckpointArchive);
 
 	void NotifyDemoPlaybackFailure(EDemoPlayFailure::Type FailureType);
 
@@ -1276,5 +1278,19 @@ public:
 
 private:
 	TAtomic<float> LastReplayFrameFidelity{ 0 };
+
 //////////////////////////////////////////////////////////////////////////
+// Time boxing for net guid cache serialization
+private:
+	struct FNetGuidCacheItem
+	{
+		FNetworkGUID NetGuid;
+		FNetGuidCacheObject NetGuidCacheObject;
+	};
+
+	TArray<FNetGuidCacheItem> NetGuidCacheSnapshot;
+	int32 NextNetGuidForRecording;
+	int32 NumNetGuidsForRecording;
+	FArchivePos NetGuidsCountPos;
+/////////////////////////////////////////////////////////////////////////
 };
diff --git a/Engine/Source/Runtime/Engine/Private/DemoNetDriver.cpp b/Engine/Source/Runtime/Engine/Private/DemoNetDriver.cpp
--- a/Engine/Source/Runtime/Engine/Private/DemoNetDriver.cpp
+++ b/Engine/Source/Runtime/Engine/Private/DemoNetDriver.cpp
@@ -1832,18 +1832,28 @@ bool UDemoNetDriver::DemoReplicateActor(AActor* Actor, UNetConnection* Connectio
 	return bDidReplicateActor;
 }
 
-void UDemoNetDriver::SerializeGuidCache(TSharedPtr<FNetGUIDCache> InGuidCache, FArchive* CheckpointArchive)
+class FRepActorsCheckpointParams
 {
-	int32 NumValues = 0;
-	int32 UnloadedValues = 0;
-
-	FArchivePos CountPos = CheckpointArchive->Tell();
-
-	*CheckpointArchive << NumValues;
+public:
+	const double StartCheckpointTime;
+	const double CheckpointMaxUploadTimePerFrame;
+};
 
+// Serialisation of net guids cache can take too long for a single frame (if checkpoint saving is timeboxed).
+// Make a snapshot of net guids that need to be serialised for checkpoint so serialisation can be split between multiple frames.
+// This operation is not time boxed and is expected to be fast and to fit into a single frame.
+void UDemoNetDriver::CacheNetGuids()
+{
+	int32 NumValues = 0;
 	const bool bDeltaCheckpoint = HasDeltaCheckpoints();
+	const double StartTime = FPlatformTime::Seconds();
+
+	// initialize NetGuidCache serialization
+	NetGuidCacheSnapshot.Reset();
+	NextNetGuidForRecording = 0;
+	NumNetGuidsForRecording = 0;
 
-	for (auto It = InGuidCache->ObjectLookup.CreateIterator(); It; ++It)
+	for (auto It = GuidCache->ObjectLookup.CreateIterator(); It; ++It)
 	{
 		FNetworkGUID& NetworkGUID = It.Key();
 		FNetGuidCacheObject& CacheObject = It.Value();
@@ -1855,50 +1865,83 @@ void UDemoNetDriver::SerializeGuidCache(TSharedPtr<FNetGUIDCache> InGuidCache, F
 
 		if (NetworkGUID.IsValid())
 		{
-			const UObject* Object = CacheObject.Object.Get();
+			NetGuidCacheSnapshot.Add({ NetworkGUID, CacheObject });
 
-			if (NetworkGUID.IsStatic() || (Object && Object->IsNameStableForNetworking()))
-			{		
-				// if we know the guid was specifically deleted, do not serialize it
-				if (DeletedNetStartupActorGUIDs.Contains(NetworkGUID))
-				{
-					continue;
-				}
+			CacheObject.bDirtyForReplay = false;
 
-				FString PathName = Object ? Object->GetName() : CacheObject.PathName.ToString();
+			++NumValues;
+		}
+	}
 
-				GEngine->NetworkRemapPath(this, PathName, false);
+	UE_LOG(LogDemo, Verbose, TEXT("CacheNetGuids: %d, %.1f ms"), NumValues, (FPlatformTime::Seconds() - StartTime) * 1000);
+}
 
-				*CheckpointArchive << NetworkGUID;
-				*CheckpointArchive << CacheObject.OuterGUID;
-				*CheckpointArchive << PathName;
-				*CheckpointArchive << CacheObject.NetworkChecksum;
+// Checkpoint saving step.
+// Serialise as many net guids as fit into a single frame (if time boxed) from previously made snapshot
+bool UDemoNetDriver::SerializeGuidCache(const FRepActorsCheckpointParams& Params, FArchive* CheckpointArchive)
+{
+	if (NextNetGuidForRecording == 0) // is the first iteration?
+	{
+		NetGuidsCountPos = CheckpointArchive->Tell();
+		*CheckpointArchive << NextNetGuidForRecording;
+	}
 
-				uint8 Flags = 0;
-				Flags |= CacheObject.bNoLoad ? (1 << 0) : 0;
-				Flags |= CacheObject.bIgnoreWhenMissing ? (1 << 1) : 0;
+	const double StartTime = FPlatformTime::Seconds();
+	const double Deadline = Params.StartCheckpointTime + Params.CheckpointMaxUploadTimePerFrame;
 
-				*CheckpointArchive << Flags;
+	check(NextNetGuidForRecording >= 0 && NextNetGuidForRecording < NetGuidCacheSnapshot.Num());
 
-				CacheObject.bDirtyForReplay = false;
+	for (; NextNetGuidForRecording != NetGuidCacheSnapshot.Num(); ++NextNetGuidForRecording)
+	{
+		FNetworkGUID& NetworkGUID = NetGuidCacheSnapshot[NextNetGuidForRecording].NetGuid;
+		FNetGuidCacheObject& CacheObject = NetGuidCacheSnapshot[NextNetGuidForRecording].NetGuidCacheObject;
 
-				++NumValues;
+		const UObject* Object = CacheObject.Object.Get();
 
-				const bool bUnloaded = Object == nullptr || !Object->IsNameStableForNetworking();
-				if (bUnloaded)
-				{
-					++UnloadedValues;
-				}
+		if (NetworkGUID.IsStatic() || (Object && Object->IsNameStableForNetworking()))
+		{		
+			// if we know the guid was specifically deleted, do not serialize it
+			if (DeletedNetStartupActorGUIDs.Contains(NetworkGUID))
+			{
+				continue;
+			}
+
+			FString PathName = Object ? Object->GetName() : CacheObject.PathName.ToString();
+
+			GEngine->NetworkRemapPath(this, PathName, false);
+
+			*CheckpointArchive << NetworkGUID;
+			*CheckpointArchive << CacheObject.OuterGUID;
+			*CheckpointArchive << PathName;
+			*CheckpointArchive << CacheObject.NetworkChecksum;
+
+			uint8 Flags = 0;
+			Flags |= CacheObject.bNoLoad ? (1 << 0) : 0;
+			Flags |= CacheObject.bIgnoreWhenMissing ? (1 << 1) : 0;
+
+			*CheckpointArchive << Flags;
+
+			++NumNetGuidsForRecording;
+
+			if (FPlatformTime::Seconds() >= Deadline)
+			{
+				break;
 			}
 		}
 	}
 
-	FArchivePos Pos = CheckpointArchive->Tell();
-	CheckpointArchive->Seek(CountPos);
-	*CheckpointArchive << NumValues;
-	CheckpointArchive->Seek(Pos);
+	bool bCompleted = NextNetGuidForRecording == NetGuidCacheSnapshot.Num();
+	if (bCompleted)
+	{
+		FArchivePos Pos = CheckpointArchive->Tell();
+		CheckpointArchive->Seek(NetGuidsCountPos);
+		*CheckpointArchive << NumNetGuidsForRecording;
+		CheckpointArchive->Seek(Pos);
+	}
+
+	UE_LOG(LogDemo, Log, TEXT("Checkpoint. SerializeGuidCache: %i/%i (total %i), took %.3f (%.3f)"), NextNetGuidForRecording, NetGuidCacheSnapshot.Num(), NumNetGuidsForRecording, FPlatformTime::Seconds() - Params.StartCheckpointTime, FPlatformTime::Seconds() - StartTime);
 
-	UE_LOG(LogDemo, Verbose, TEXT("Checkpoint. SerializeGuidCache: %i Unloaded: %i"), NumValues, UnloadedValues);
+	return bCompleted;
 }
 
 float UDemoNetDriver::GetCheckpointSaveMaxMSPerFrame() const
@@ -1935,6 +1978,8 @@ void UDemoNetDriver::SaveCheckpoint()
 	
 	const bool bDeltaCheckpoint = HasDeltaCheckpoints();
 
+	CSV_SCOPED_TIMING_STAT(Basic, DemoSaveCheckpointTime);
+
 	if (HasLevelStreamingFixes())
 	{
 		SCOPED_NAMED_EVENT(UDemoNetDriver_ReplayLevelSortAndAssign, FColor::Purple);
@@ -2053,14 +2098,6 @@ void UDemoNetDriver::SaveCheckpoint()
 	}
 }
 
-class FRepActorsCheckpointParams
-{
-public:
-
-	const double StartCheckpointTime;
-	const double CheckpointMaxUploadTimePerFrame;
-};
-
 // Only start execution if a certain percentage remains of the 
 static bool inline ShouldExecuteState(const FRepActorsCheckpointParams& Params, double CurrentTime, double RequiredRatioToStart)
 {
@@ -2089,6 +2126,8 @@ void UDemoNetDriver::TickCheckpoint()
 		return;
 	}
 
+	CSV_SCOPED_TIMING_STAT(Basic, DemoTickCheckpointTime);
+
 	FRepActorsCheckpointParams Params
 	{
 		FPlatformTime::Seconds(),
@@ -2225,22 +2264,33 @@ void UDemoNetDriver::TickCheckpoint()
 							*CheckpointArchive << DeletedNetStartupActors;
 						}
 
-						CheckpointSaveContext.CheckpointSaveState = ECheckpointSaveState_SerializeGuidCache;
+						CheckpointSaveContext.CheckpointSaveState = ECheckpointSaveState_CacheNetGuids;
 					}
 				}
 				break;
 
-				case ECheckpointSaveState_SerializeGuidCache:
+				case ECheckpointSaveState_CacheNetGuids:
 				{
-					// Postpone execution of this state if we have used to much of our alloted time, this value can be tweaked based on profiling
-					const double RequiredRatioFor_SerializeGuidCache = 0.8;
-					if ((bExecuteNextState = ShouldExecuteState(Params, CurrentTime, RequiredRatioFor_SerializeGuidCache)) == true)
+					// Postpone execution of this state if we have used too much of our alloted time, this value can be tweaked based on profiling
+					const double RequiredRatioFor_CacheNetGuids = 0.8;
+					if ((bExecuteNextState = ShouldExecuteState(Params, CurrentTime, RequiredRatioFor_CacheNetGuids)) == true)
 					{
-						SCOPED_NAMED_EVENT(UDemoNetDriver_SerializeGuidCache, FColor::Green);
+						SCOPED_NAMED_EVENT(UDemoNetDriver_CacheNetGuids, FColor::Green);
 
-						// Save the current guid cache
-						SerializeGuidCache( GuidCache, CheckpointArchive );
+						CacheNetGuids();
+						CheckpointSaveContext.CheckpointSaveState = ECheckpointSaveState_SerializeGuidCache;
+					}
+				}
+				break;
 
+				case ECheckpointSaveState_SerializeGuidCache:
+				{
+					SCOPED_NAMED_EVENT(UDemoNetDriver_SerializeGuidCache, FColor::Green);
+
+					// Save the current guid cache
+					bExecuteNextState = SerializeGuidCache(Params, CheckpointArchive);
+					if (bExecuteNextState)
+					{
 						CheckpointSaveContext.CheckpointSaveState = ECheckpointSaveState_SerializeNetFieldExportGroupMap;
 					}
 				}
