Support AnimRootMotion-style root motion sources, reduce memory usage, make SimulatedTick() customizable per project
- Add IgnoreZAccumulate option to RootMotionSources to match AnimRootMotion behavior (override root motion sources with this flag will control horizontal velocity but fall when in air like AnimRootMotion)
- Remove InlineAllocator setting from FRootMotionSourceGroup array properties to save memory in overwhelminingly common case of projects/movement ticks/SavedMoves without root motion sources present
- Make SimulatedTick virtual so projects can override for custom SimulatedProxy handling
- Add Rotation value to two CharacterMovementComponent tick logs for help debugging rotation-related issues

diff --git a/Engine/Source/Runtime/Engine/Classes/GameFramework/CharacterMovementComponent.h b/Engine/Source/Runtime/Engine/Classes/GameFramework/CharacterMovementComponent.h
--- a/Engine/Source/Runtime/Engine/Classes/GameFramework/CharacterMovementComponent.h
+++ b/Engine/Source/Runtime/Engine/Classes/GameFramework/CharacterMovementComponent.h
@@ -2075,7 +2075,7 @@ protected:
 	virtual void PerformMovement(float DeltaTime);
 
 	/** Special Tick for Simulated Proxies */
-	void SimulatedTick(float DeltaSeconds);
+	virtual void SimulatedTick(float DeltaSeconds);
 
 	/** Simulate movement on a non-owning client. Called by SimulatedTick(). */
 	virtual void SimulateMovement(float DeltaTime);
diff --git a/Engine/Source/Runtime/Engine/Classes/GameFramework/RootMotionSource.h b/Engine/Source/Runtime/Engine/Classes/GameFramework/RootMotionSource.h
--- a/Engine/Source/Runtime/Engine/Classes/GameFramework/RootMotionSource.h
+++ b/Engine/Source/Runtime/Engine/Classes/GameFramework/RootMotionSource.h
@@ -65,7 +65,12 @@ enum class ERootMotionSourceSettingsFlags : uint8
 	// do not reduce SimulationTime. Disabling this is useful for sources that
 	// are more about providing velocity (like jumps), vs. sources that need
 	// the precision of partial ticks for say ending up at an exact location (MoveTo)
-	DisablePartialEndTick		= 0x02
+	DisablePartialEndTick		= 0x02,
+	// Whether to ignore impact to Z when accumulating output to Velocity
+	// Setting this flag on override sources provides the same behavior as
+	// animation root motion
+	IgnoreZAccumulate			= 0x04
+
 };
 
 enum class ERootMotionSourceID : uint16 { Invalid = 0 };
@@ -733,10 +738,10 @@ struct ENGINE_API FRootMotionSourceGroup
 	virtual ~FRootMotionSourceGroup() {}
 
 	/** Root Motion Sources currently applied in this Group */
-	TArray< TSharedPtr<FRootMotionSource>, TInlineAllocator<8> > RootMotionSources;
+	TArray< TSharedPtr<FRootMotionSource> > RootMotionSources;
 
 	/** Root Motion Sources to be added next frame */
-	TArray< TSharedPtr<FRootMotionSource>, TInlineAllocator<4> > PendingAddRootMotionSources;
+	TArray< TSharedPtr<FRootMotionSource> > PendingAddRootMotionSources;
 
 	/** 
 	 *  Whether this group has additive root motion sources
@@ -750,6 +755,12 @@ struct ENGINE_API FRootMotionSourceGroup
 	UPROPERTY()
 	uint8 bHasOverrideSources:1;
 
+	/** 
+	 *  Whether this group has override root motion sources that have IgnoreZAccumulate flag
+	 **/
+	UPROPERTY()
+	uint8 bHasOverrideSourcesWithIgnoreZAccumulate:1;
+
 	/** True when we had additive velocity applied last tick, checked to know if we should restore
 	 *  LastPreAdditiveVelocity before a Velocity computation */
 	UPROPERTY()
@@ -792,6 +803,9 @@ struct ENGINE_API FRootMotionSourceGroup
 	/** @return true if Velocity will be overridden by root motion sources, meaning we can skip all normal movement-based velocity calculations */
 	bool HasOverrideVelocity() const;
 
+	/** @return true if Velocity will be overridden by root motion sources that do not affect Z velocity (matching animation root motion behavior) */
+	bool HasOverrideVelocityWithIgnoreZAccumulate() const;
+
 	/** @return true if any axis of velocity has additive velocity applied by root motion sources */
 	bool HasAdditiveVelocity() const;
 
diff --git a/Engine/Source/Runtime/Engine/Private/Components/CharacterMovementComponent.cpp b/Engine/Source/Runtime/Engine/Private/Components/CharacterMovementComponent.cpp
--- a/Engine/Source/Runtime/Engine/Private/Components/CharacterMovementComponent.cpp
+++ b/Engine/Source/Runtime/Engine/Private/Components/CharacterMovementComponent.cpp
@@ -4357,7 +4357,7 @@ void UCharacterMovementComponent::PhysFalling(float deltaTime, int32 Iterations)
 				if (subTimeTickRemaining > KINDA_SMALL_NUMBER && !bJustTeleported)
 				{
 					const FVector NewVelocity = (Delta / subTimeTickRemaining);
-					Velocity = HasAnimRootMotion() && !CurrentRootMotion.HasOverrideVelocity() ? FVector(Velocity.X, Velocity.Y, NewVelocity.Z) : NewVelocity;
+					Velocity = HasAnimRootMotion() || CurrentRootMotion.HasOverrideVelocityWithIgnoreZAccumulate() ? FVector(Velocity.X, Velocity.Y, NewVelocity.Z) : NewVelocity;
 				}
 
 				if (subTimeTickRemaining > KINDA_SMALL_NUMBER && (Delta | Adjusted) > 0.f)
@@ -4413,7 +4413,7 @@ void UCharacterMovementComponent::PhysFalling(float deltaTime, int32 Iterations)
 						if (subTimeTickRemaining > KINDA_SMALL_NUMBER && !bJustTeleported)
 						{
 							const FVector NewVelocity = (Delta / subTimeTickRemaining);
-							Velocity = HasAnimRootMotion() && !CurrentRootMotion.HasOverrideVelocity() ? FVector(Velocity.X, Velocity.Y, NewVelocity.Z) : NewVelocity;
+							Velocity = HasAnimRootMotion() || CurrentRootMotion.HasOverrideVelocityWithIgnoreZAccumulate() ? FVector(Velocity.X, Velocity.Y, NewVelocity.Z) : NewVelocity;
 						}
 
 						// bDitch=true means that pawn is straddling two slopes, neither of which he can stand on
@@ -8139,8 +8139,8 @@ void UCharacterMovementComponent::ReplicateMoveToServer(float DeltaTime, const F
 
 		ClientData->ClientUpdateTime = MyWorld->TimeSeconds;
 
-		UE_CLOG(CharacterOwner && UpdatedComponent, LogNetPlayerMovement, VeryVerbose, TEXT("ClientMove Time %f Acceleration %s Velocity %s Position %s DeltaTime %f Mode %s MovementBase %s.%s (Dynamic:%d) DualMove? %d"),
-			NewMove->TimeStamp, *NewMove->Acceleration.ToString(), *Velocity.ToString(), *UpdatedComponent->GetComponentLocation().ToString(), NewMove->DeltaTime, *GetMovementName(),
+		UE_CLOG(CharacterOwner && UpdatedComponent, LogNetPlayerMovement, VeryVerbose, TEXT("ClientMove Time %f Acceleration %s Velocity %s Position %s Rotation %s DeltaTime %f Mode %s MovementBase %s.%s (Dynamic:%d) DualMove? %d"),
+			NewMove->TimeStamp, *NewMove->Acceleration.ToString(), *Velocity.ToString(), *UpdatedComponent->GetComponentLocation().ToString(), *UpdatedComponent->GetComponentRotation().ToCompactString(), NewMove->DeltaTime, *GetMovementName(),
 			*GetNameSafe(NewMove->EndBase.Get()), *NewMove->EndBoneName.ToString(), MovementBaseUtility::IsDynamicBase(NewMove->EndBase.Get()) ? 1 : 0, ClientData->PendingMove.IsValid() ? 1 : 0);
 
 		
@@ -8745,8 +8745,8 @@ void UCharacterMovementComponent::ServerMove_Implementation(
 		MoveAutonomous(TimeStamp, DeltaTime, MoveFlags, Accel);
 	}
 
-	UE_CLOG(CharacterOwner && UpdatedComponent, LogNetPlayerMovement, VeryVerbose, TEXT("ServerMove Time %f Acceleration %s Velocity %s Position %s DeltaTime %f Mode %s MovementBase %s.%s (Dynamic:%d)"),
-			TimeStamp, *Accel.ToString(), *Velocity.ToString(), *UpdatedComponent->GetComponentLocation().ToString(), DeltaTime, *GetMovementName(),
+	UE_CLOG(CharacterOwner && UpdatedComponent, LogNetPlayerMovement, VeryVerbose, TEXT("ServerMove Time %f Acceleration %s Velocity %s Position %s Rotation %s DeltaTime %f Mode %s MovementBase %s.%s (Dynamic:%d)"),
+			TimeStamp, *Accel.ToString(), *Velocity.ToString(), *UpdatedComponent->GetComponentLocation().ToString(), *UpdatedComponent->GetComponentRotation().ToCompactString(), DeltaTime, *GetMovementName(),
 			*GetNameSafe(GetMovementBase()), *CharacterOwner->GetBasedMovement().BoneName.ToString(), MovementBaseUtility::IsDynamicBase(GetMovementBase()) ? 1 : 0);
 
 	ServerMoveHandleClientError(TimeStamp, DeltaTime, Accel, ClientLoc, ClientMovementBase, ClientBaseBoneName, ClientMovementMode);
diff --git a/Engine/Source/Runtime/Engine/Private/GameFramework/RootMotionSource.cpp b/Engine/Source/Runtime/Engine/Private/GameFramework/RootMotionSource.cpp
--- a/Engine/Source/Runtime/Engine/Private/GameFramework/RootMotionSource.cpp
+++ b/Engine/Source/Runtime/Engine/Private/GameFramework/RootMotionSource.cpp
@@ -1265,6 +1265,7 @@ void FRootMotionSource_JumpForce::AddReferencedObjects(class FReferenceCollector
 FRootMotionSourceGroup::FRootMotionSourceGroup()
 	: bHasAdditiveSources(false)
 	, bHasOverrideSources(false)
+	, bHasOverrideSourcesWithIgnoreZAccumulate(false)
 	, bIsAdditiveVelocityApplied(false)
 	, LastPreAdditiveVelocity(ForceInitToZero)
 {
@@ -1280,6 +1281,11 @@ bool FRootMotionSourceGroup::HasOverrideVelocity() const
 	return bHasOverrideSources;
 }
 
+bool FRootMotionSourceGroup::HasOverrideVelocityWithIgnoreZAccumulate() const
+{
+	return bHasOverrideSourcesWithIgnoreZAccumulate;
+}
+
 bool FRootMotionSourceGroup::HasAdditiveVelocity() const
 {
 	return bHasAdditiveSources;
@@ -1419,6 +1425,7 @@ void FRootMotionSourceGroup::PrepareRootMotion(float DeltaTime, const ACharacter
 	// Prepare active sources
 	{
 		bHasOverrideSources = false;
+		bHasOverrideSourcesWithIgnoreZAccumulate = false;
 		bHasAdditiveSources = false;
 		LastAccumulatedSettings.Clear();
 
@@ -1589,6 +1596,11 @@ void FRootMotionSourceGroup::PrepareRootMotion(float DeltaTime, const ACharacter
 				else if (RootMotionSource->AccumulateMode == ERootMotionAccumulateMode::Override)
 				{
 					bHasOverrideSources = true;
+
+					if (RootMotionSource->Settings.HasFlag(ERootMotionSourceSettingsFlags::IgnoreZAccumulate))
+					{
+						bHasOverrideSourcesWithIgnoreZAccumulate = true;
+					}
 				}
 			}
 		}
@@ -1663,6 +1675,7 @@ void FRootMotionSourceGroup::AccumulateRootMotionVelocityFromSource
 
 	const FVector RootMotionVelocity = RootMotionParams.GetRootMotionTransform().GetTranslation();
 
+	const FVector InputVelocity = InOutVelocity;
 	if (RootMotionSource.AccumulateMode == ERootMotionAccumulateMode::Override)
 	{
 		InOutVelocity = RootMotionVelocity;
@@ -1671,6 +1684,10 @@ void FRootMotionSourceGroup::AccumulateRootMotionVelocityFromSource
 	{
 		InOutVelocity += RootMotionVelocity;
 	}
+	if (RootMotionSource.Settings.HasFlag(ERootMotionSourceSettingsFlags::IgnoreZAccumulate))
+	{
+		InOutVelocity.Z = InputVelocity.Z;
+	}
 }
 
 void FRootMotionSourceGroup::SetPendingRootMotionSourceMinStartTimes(float NewStartTime)
@@ -1940,6 +1957,7 @@ bool FRootMotionSourceGroup::NetSerialize(FArchive& Ar, class UPackageMap* Map,
 
 	FArchive_Serialize_BitfieldBool(Ar, bHasAdditiveSources);
 	FArchive_Serialize_BitfieldBool(Ar, bHasOverrideSources);
+	FArchive_Serialize_BitfieldBool(Ar, bHasOverrideSourcesWithIgnoreZAccumulate);
 	LastPreAdditiveVelocity.NetSerialize(Ar, Map, bOutSuccess);
 	FArchive_Serialize_BitfieldBool(Ar, bIsAdditiveVelocityApplied);
 	Ar << LastAccumulatedSettings.Flags;
@@ -2057,6 +2075,7 @@ void FRootMotionSourceGroup::Clear()
 	bIsAdditiveVelocityApplied = false;
 	bHasAdditiveSources = false;
 	bHasOverrideSources = false;
+	bHasOverrideSourcesWithIgnoreZAccumulate = false;
 	LastAccumulatedSettings.Clear();
 }
 
@@ -2097,6 +2116,7 @@ FRootMotionSourceGroup& FRootMotionSourceGroup::operator=(const FRootMotionSourc
 
 		bHasAdditiveSources = Other.bHasAdditiveSources;
 		bHasOverrideSources = Other.bHasOverrideSources;
+		bHasOverrideSourcesWithIgnoreZAccumulate = Other.bHasOverrideSourcesWithIgnoreZAccumulate;
 		LastPreAdditiveVelocity = Other.LastPreAdditiveVelocity;
 		bIsAdditiveVelocityApplied = Other.bIsAdditiveVelocityApplied;
 		LastAccumulatedSettings = Other.LastAccumulatedSettings;
@@ -2108,6 +2128,7 @@ bool FRootMotionSourceGroup::operator==(const FRootMotionSourceGroup& Other) con
 {
 	if (bHasAdditiveSources != Other.bHasAdditiveSources || 
 		bHasOverrideSources != Other.bHasOverrideSources ||
+		bHasOverrideSourcesWithIgnoreZAccumulate != Other.bHasOverrideSourcesWithIgnoreZAccumulate ||
 		bIsAdditiveVelocityApplied != Other.bIsAdditiveVelocityApplied ||
 		!LastPreAdditiveVelocity.Equals(Other.LastPreAdditiveVelocity, 1.f))
 	{
