Moved in VP debug printing start down 10 pixels under mouse control popup.

diff --git a/Engine/Source/Runtime/Engine/Private/UnrealEngine.cpp b/Engine/Source/Runtime/Engine/Private/UnrealEngine.cpp
--- a/Engine/Source/Runtime/Engine/Private/UnrealEngine.cpp
+++ b/Engine/Source/Runtime/Engine/Private/UnrealEngine.cpp
@@ -10051,7 +10051,7 @@ void DrawStatsHUD( UWorld* World, FViewport* Viewport, FCanvas* Canvas, UCanvas*
 	const int32 FPSXOffset	= (FPlatformProperties::SupportsWindowedMode() ? 110 : 250);
 	const int32 StatsXOffset = 100;// FPlatformProperties::SupportsWindowedMode() ? 4 : 100;
 
-	static const int32 MessageStartY = GIsEditor ? 35 : 100; // Account for safe frame
+	static const int32 MessageStartY = GIsEditor ? 45 : 100; // Account for safe frame
 	int32 MessageY = MessageStartY;
 
 	const FVector2D FontScale(1, 1);
