Added extra checks to prevent crashes when building static mesh descriptions.
#rnx

diff --git a/Engine/Source/Runtime/Engine/Private/StaticMesh.cpp b/Engine/Source/Runtime/Engine/Private/StaticMesh.cpp
--- a/Engine/Source/Runtime/Engine/Private/StaticMesh.cpp
+++ b/Engine/Source/Runtime/Engine/Private/StaticMesh.cpp
@@ -5006,6 +5006,12 @@ void UStaticMesh::BuildFromMeshDescription(const FMeshDescription& MeshDescripti
 	// Fill vertex buffers
 
 	int32 NumVertexInstances = MeshDescription.VertexInstances().GetArraySize();
+	int32 NumTriangles = MeshDescription.Triangles().Num();
+
+	if (NumVertexInstances == 0 || NumTriangles == 0)
+	{
+		return;
+	}
 
 	TArray<FStaticMeshBuildVertex> StaticMeshBuildVertices;
 	StaticMeshBuildVertices.SetNum(NumVertexInstances);
@@ -5065,7 +5071,6 @@ void UStaticMesh::BuildFromMeshDescription(const FMeshDescription& MeshDescripti
 	// Fill index buffer and sections array
 
 	int32 NumPolygonGroups = MeshDescription.PolygonGroups().Num();
-	int32 NumTriangles = MeshDescription.Triangles().Num();
 
 	TPolygonGroupAttributesConstRef<FName> MaterialSlotNames = MeshDescriptionAttributes.GetPolygonGroupMaterialSlotNames();
 
