Change to CommitMeshDescription so it doesn't also convert the MeshDescription back to RawMesh and store it as RawMeshBulkData; instead the RawMeshBulkData is emptied and invalidated, so that any legacy tool which requests it will have to convert it from MeshDescription at that moment.
#rnx

diff --git a/Engine/Source/Runtime/Engine/Private/StaticMesh.cpp b/Engine/Source/Runtime/Engine/Private/StaticMesh.cpp
--- a/Engine/Source/Runtime/Engine/Private/StaticMesh.cpp
+++ b/Engine/Source/Runtime/Engine/Private/StaticMesh.cpp
@@ -3914,7 +3914,7 @@ FMeshDescription* UStaticMesh::CreateMeshDescription(int32 LodIndex, FMeshDescri
 }
 
 
-void UStaticMesh::CommitMeshDescription(int32 LodIndex, const FCommitMeshDescriptionParams & Params)
+void UStaticMesh::CommitMeshDescription(int32 LodIndex, const FCommitMeshDescriptionParams& Params)
 {
 	TRACE_CPUPROFILER_EVENT_SCOPE(UStaticMesh::CommitMeshDescription);
 
@@ -3933,43 +3933,21 @@ void UStaticMesh::CommitMeshDescription(int32 LodIndex, const FCommitMeshDescrip
 			SourceModel.MeshDescriptionBulkData = MakeUnique<FMeshDescriptionBulkData>();
 		}
 
-		// Handle ConvertToRawMesh and SaveMeshDescription in parallel
-		// Something like ParallelInvoke would be cleaner here, but we need the ParallelFor ability to join work on the current thread
-		ParallelFor(2,
-			[this, &SourceModel, &Params](int32 Num)
-			{
-				if (Num == 0)
-				{
-					TMap<FName, int32> MaterialMap;
-					for (int32 MaterialIndex = 0; MaterialIndex < StaticMaterials.Num(); ++MaterialIndex)
-					{
-						MaterialMap.Add(StaticMaterials[MaterialIndex].ImportedMaterialSlotName, MaterialIndex);
-					}
-
-					FRawMesh TempRawMesh;
-					FMeshDescriptionOperations::ConvertToRawMesh(*SourceModel.MeshDescription, TempRawMesh, MaterialMap);
-					SourceModel.RawMeshBulkData->SaveRawMesh(TempRawMesh);
-				}
-				else
-				{
-					SourceModel.MeshDescriptionBulkData->SaveMeshDescription(*SourceModel.MeshDescription);
-
-					if (Params.bUseHashAsGuid)
-					{
-						SourceModel.MeshDescriptionBulkData->UseHashAsGuid();
-					}
-				}
-			}
-		);
+		SourceModel.MeshDescriptionBulkData->SaveMeshDescription(*SourceModel.MeshDescription);
+		if (Params.bUseHashAsGuid)
+		{
+			SourceModel.MeshDescriptionBulkData->UseHashAsGuid();
+		}
 	}
 	else
 	{
 		SourceModel.MeshDescriptionBulkData.Reset();
-
-		// Mesh description is null, remove the rawmesh data
-		SourceModel.RawMeshBulkData->Empty();
 	}
 
+	// Clear RawMeshBulkData and mark as invalid.
+	// If any legacy tool needs the RawMesh at this point, it will do a conversion from MD at that moment.
+	SourceModel.RawMeshBulkData->Empty();
+
 	// This part is not thread-safe, so we give the caller the option of calling it manually from the mainthread
 	if (Params.bMarkPackageDirty)
 	{
