System time code can generate sub frame

diff --git a/Engine/Source/Runtime/Engine/Classes/Engine/SystemTimeTimecodeProvider.h b/Engine/Source/Runtime/Engine/Classes/Engine/SystemTimeTimecodeProvider.h
--- a/Engine/Source/Runtime/Engine/Classes/Engine/SystemTimeTimecodeProvider.h
+++ b/Engine/Source/Runtime/Engine/Classes/Engine/SystemTimeTimecodeProvider.h
@@ -18,19 +18,29 @@ private:
 	UPROPERTY(EditAnywhere, Category = Timecode)
 	FFrameRate FrameRate;
 
+	/** When generating frame time, should we generate full frame without subframe value.*/
+	UPROPERTY(EditAnywhere, Category = Timecode)
+	bool bGenerateFullFrame;
+
 	ETimecodeProviderSynchronizationState State;
 
 public:
 
-	USystemTimeTimecodeProvider():
-		FrameRate(60, 1),
-		State(ETimecodeProviderSynchronizationState::Closed)
-	{
-	}
+	USystemTimeTimecodeProvider();
+
+	/** Generate a frame time value, including subframe, using the system clock. */
+	static FFrameTime GenerateFrameTimeFromSystemTime(FFrameRate Rate);
 
 	/** Generate a timecode value using the system clock. */
 	static FTimecode GenerateTimecodeFromSystemTime(FFrameRate Rate);
 
+	/**
+	 * Generate a frame time value, including subframe, using the high performance clock
+	 * Using the high performance clock is faster but will make the value drift over time.
+	 * This is an optimized version. Prefer GenerateTimecodeFromSystemTime, if the value need to be accurate.
+	 **/
+	static FFrameTime GenerateFrameTimeFromHighPerformanceClock(FFrameRate Rate);
+
 	/**
 	 * Generate a timecode value using the high performance clock
 	 * Using the high performance clock is faster but will make the value drift over time.
diff --git a/Engine/Source/Runtime/Engine/Private/Engine/SystemTimeTimecodeProvider.cpp b/Engine/Source/Runtime/Engine/Private/Engine/SystemTimeTimecodeProvider.cpp
--- a/Engine/Source/Runtime/Engine/Private/Engine/SystemTimeTimecodeProvider.cpp
+++ b/Engine/Source/Runtime/Engine/Private/Engine/SystemTimeTimecodeProvider.cpp
@@ -6,6 +6,34 @@
 #include "Misc/CoreMisc.h"
 #include "Misc/DateTime.h"
 
+namespace
+{
+	static double ComputeTimeCodeOffset()
+	{
+		const FDateTime DateTime = FDateTime::Now();
+		const double HighPerformanceClock = FPlatformTime::Seconds();
+		return DateTime.GetTimeOfDay().GetTotalSeconds() - HighPerformanceClock;
+	}
+
+	static double HighPerformanceClockDelta = ComputeTimeCodeOffset();
+};
+
+
+USystemTimeTimecodeProvider::USystemTimeTimecodeProvider()
+	: FrameRate(60, 1)
+	, bGenerateFullFrame(false)
+	, State(ETimecodeProviderSynchronizationState::Closed)
+{
+}
+
+
+FFrameTime USystemTimeTimecodeProvider::GenerateFrameTimeFromSystemTime(FFrameRate FrameRate)
+{
+	const FDateTime DateTime = FDateTime::Now();
+	const FTimespan Timespan = DateTime.GetTimeOfDay();
+	return FrameRate.AsFrameTime(Timespan.GetTotalSeconds());
+}
+
 
 FTimecode USystemTimeTimecodeProvider::GenerateTimecodeFromSystemTime(FFrameRate FrameRate)
 {
@@ -15,18 +43,16 @@ FTimecode USystemTimeTimecodeProvider::GenerateTimecodeFromSystemTime(FFrameRate
 }
 
 
-static double ComputeTimeCodeOffset()
+FFrameTime USystemTimeTimecodeProvider::GenerateFrameTimeFromHighPerformanceClock(FFrameRate FrameRate)
 {
-	const FDateTime DateTime = FDateTime::Now();
-	const double HighPerformanceClock = FPlatformTime::Seconds();
-	return DateTime.GetTimeOfDay().GetTotalSeconds() - HighPerformanceClock;
+	constexpr double SecondsPerDay = 24.0 * 60.0 * 60.0;
+	return FrameRate.AsFrameTime(fmod(HighPerformanceClockDelta + FPlatformTime::Seconds(), SecondsPerDay));
 }
 
 
 FTimecode USystemTimeTimecodeProvider::GenerateTimecodeFromHighPerformanceClock(FFrameRate FrameRate)
 {
 	constexpr double SecondsPerDay = 24.0 * 60.0 * 60.0;
-	static double HighPerformanceClockDelta = ComputeTimeCodeOffset();
 	const FTimespan Timespan = FTimespan::FromSeconds(fmod(HighPerformanceClockDelta + FPlatformTime::Seconds(), SecondsPerDay));
 	return FTimecode::FromTimespan(Timespan, FrameRate, false);
 }
@@ -35,5 +61,7 @@ FTimecode USystemTimeTimecodeProvider::GenerateTimecodeFromHighPerformanceClock(
 FQualifiedFrameTime USystemTimeTimecodeProvider::GetQualifiedFrameTime() const
 {
 	// Generate a TC to prevent subframe
-	return FQualifiedFrameTime(GenerateTimecodeFromSystemTime(FrameRate), FrameRate);
+	return bGenerateFullFrame ? 
+		FQualifiedFrameTime(GenerateTimecodeFromSystemTime(FrameRate), FrameRate)
+		: FQualifiedFrameTime(GenerateFrameTimeFromSystemTime(FrameRate), FrameRate);
 }
diff --git a/Engine/Source/Runtime/Engine/Private/UnrealEngine.cpp b/Engine/Source/Runtime/Engine/Private/UnrealEngine.cpp
--- a/Engine/Source/Runtime/Engine/Private/UnrealEngine.cpp
+++ b/Engine/Source/Runtime/Engine/Private/UnrealEngine.cpp
@@ -2248,10 +2248,10 @@ void UEngine::UpdateTimecode()
 	else if(bGenerateDefaultTimecode)
 	{
 #if PLATFORM_DESKTOP
-		FQualifiedFrameTime NewFrameTime = FQualifiedFrameTime(USystemTimeTimecodeProvider::GenerateTimecodeFromSystemTime(GenerateDefaultTimecodeFrameRate), GenerateDefaultTimecodeFrameRate);
+		FQualifiedFrameTime NewFrameTime = FQualifiedFrameTime(USystemTimeTimecodeProvider::GenerateFrameTimeFromSystemTime(GenerateDefaultTimecodeFrameRate), GenerateDefaultTimecodeFrameRate);
 #else //PLATFORM_DESKTOP
 		//If a user wish to have an accurate TC value on console, he should set an accurate TC provider in his project settings.
-		FQualifiedFrameTime NewFrameTime = FQualifiedFrameTime(USystemTimeTimecodeProvider::GenerateTimecodeFromHighPerformanceClock(GenerateDefaultTimecodeFrameRate), GenerateDefaultTimecodeFrameRate);
+		FQualifiedFrameTime NewFrameTime = FQualifiedFrameTime(USystemTimeTimecodeProvider::GenerateFrameTimeFromHighPerformanceClock(GenerateDefaultTimecodeFrameRate), GenerateDefaultTimecodeFrameRate);
 #endif
 		NewFrameTime.Time.FrameNumber -= GenerateDefaultTimecodeFrameDelay;
 		FApp::SetCurrentFrameTime(NewFrameTime);
