Remove PostEditChangeProperty from UTimecodeProvider.

diff --git a/Engine/Source/Runtime/Engine/Private/Engine/TimecodeProvider.cpp b/Engine/Source/Runtime/Engine/Private/Engine/TimecodeProvider.cpp
--- a/Engine/Source/Runtime/Engine/Private/Engine/TimecodeProvider.cpp
+++ b/Engine/Source/Runtime/Engine/Private/Engine/TimecodeProvider.cpp
@@ -23,16 +23,3 @@ FTimecode UTimecodeProvider::GetDelayedTimecode() const
 	const FQualifiedFrameTime NewFrameTime = GetDelayedQualifiedFrameTime();
 	return FTimecode::FromFrameNumber(NewFrameTime.Time.GetFrame(), NewFrameTime.Rate);
 }
-
-#if WITH_EDITOR
-void UTimecodeProvider::PostEditChangeProperty(struct FPropertyChangedEvent& PropertyChangedEvent)
-{
-	if (PropertyChangedEvent.GetPropertyName() == GET_MEMBER_NAME_CHECKED(UTimecodeProvider, FrameDelay))
-	{
-		if (!SupportSubFrameDelay())
-		{
-			FrameDelay = FMath::FloorToInt(FrameDelay);
-		}
-	}
-}
-#endif
\ No newline at end of file
