Use float to generate default timecode value.

#rnx

diff --git a/Engine/Source/Runtime/Engine/Classes/Engine/Engine.h b/Engine/Source/Runtime/Engine/Classes/Engine/Engine.h
--- a/Engine/Source/Runtime/Engine/Classes/Engine/Engine.h
+++ b/Engine/Source/Runtime/Engine/Classes/Engine/Engine.h
@@ -1373,7 +1373,7 @@ public:
 
 	/** Number of frames to subtract from generated default timecode. */
 	UPROPERTY(AdvancedDisplay, config, EditAnywhere, Category=Timecode, meta=(EditCondition="bGenerateDefaultTimecode"))
-	int32 GenerateDefaultTimecodeFrameDelay;
+	float GenerateDefaultTimecodeFrameDelay;
 
 public:
 	/** 
diff --git a/Engine/Source/Runtime/Engine/Classes/Engine/SystemTimeTimecodeProvider.h b/Engine/Source/Runtime/Engine/Classes/Engine/SystemTimeTimecodeProvider.h
--- a/Engine/Source/Runtime/Engine/Classes/Engine/SystemTimeTimecodeProvider.h
+++ b/Engine/Source/Runtime/Engine/Classes/Engine/SystemTimeTimecodeProvider.h
@@ -15,6 +15,7 @@ class ENGINE_API USystemTimeTimecodeProvider : public UTimecodeProvider
 
 private:
 
+	/** The frame rate at which the timecode value will be generated. */
 	UPROPERTY(EditAnywhere, Category = Timecode)
 	FFrameRate FrameRate;
 
@@ -22,6 +23,13 @@ private:
 	UPROPERTY(EditAnywhere, Category = Timecode)
 	bool bGenerateFullFrame;
 
+	/**
+	 * Use the high performance clock instead of the system time to generate the timecode value.
+	 * Using the high performance clock is faster but will make the value drift over time.
+	 */
+	UPROPERTY(AdvancedDisplay, EditAnywhere, Category = Timecode)
+	bool bUseHighPerformanceClock;
+
 	ETimecodeProviderSynchronizationState State;
 
 public:
diff --git a/Engine/Source/Runtime/Engine/Private/Engine/SystemTimeTimecodeProvider.cpp b/Engine/Source/Runtime/Engine/Private/Engine/SystemTimeTimecodeProvider.cpp
--- a/Engine/Source/Runtime/Engine/Private/Engine/SystemTimeTimecodeProvider.cpp
+++ b/Engine/Source/Runtime/Engine/Private/Engine/SystemTimeTimecodeProvider.cpp
@@ -22,6 +22,7 @@ namespace
 USystemTimeTimecodeProvider::USystemTimeTimecodeProvider()
 	: FrameRate(60, 1)
 	, bGenerateFullFrame(true)
+	, bUseHighPerformanceClock(false)
 	, State(ETimecodeProviderSynchronizationState::Closed)
 {
 }
@@ -61,6 +62,6 @@ FTimecode USystemTimeTimecodeProvider::GenerateTimecodeFromHighPerformanceClock(
 FQualifiedFrameTime USystemTimeTimecodeProvider::GetQualifiedFrameTime() const
 {
 	return bGenerateFullFrame ? 
-		FQualifiedFrameTime(GenerateTimecodeFromSystemTime(FrameRate), FrameRate)
-		: FQualifiedFrameTime(GenerateFrameTimeFromSystemTime(FrameRate), FrameRate);
+		FQualifiedFrameTime(bUseHighPerformanceClock ? GenerateTimecodeFromHighPerformanceClock(FrameRate) : GenerateTimecodeFromSystemTime(FrameRate), FrameRate)
+		: FQualifiedFrameTime(bUseHighPerformanceClock ? GenerateFrameTimeFromHighPerformanceClock(FrameRate) : GenerateFrameTimeFromSystemTime(FrameRate), FrameRate);
 }
diff --git a/Engine/Source/Runtime/Engine/Private/UnrealEngine.cpp b/Engine/Source/Runtime/Engine/Private/UnrealEngine.cpp
--- a/Engine/Source/Runtime/Engine/Private/UnrealEngine.cpp
+++ b/Engine/Source/Runtime/Engine/Private/UnrealEngine.cpp
@@ -2253,7 +2253,7 @@ void UEngine::UpdateTimecode()
 		//If a user wish to have an accurate TC value on console, he should set an accurate TC provider in his project settings.
 		FQualifiedFrameTime NewFrameTime = FQualifiedFrameTime(USystemTimeTimecodeProvider::GenerateFrameTimeFromHighPerformanceClock(GenerateDefaultTimecodeFrameRate), GenerateDefaultTimecodeFrameRate);
 #endif
-		NewFrameTime.Time.FrameNumber -= GenerateDefaultTimecodeFrameDelay;
+		NewFrameTime.Time -= FFrameTime::FromDecimal(GenerateDefaultTimecodeFrameDelay);
 		FApp::SetCurrentFrameTime(NewFrameTime);
 	}
 }
