Fix for deterministic cooking issue with Child Montages

diff --git a/Engine/Source/Runtime/Engine/Private/Animation/AnimMontage.cpp b/Engine/Source/Runtime/Engine/Private/Animation/AnimMontage.cpp
--- a/Engine/Source/Runtime/Engine/Private/Animation/AnimMontage.cpp
+++ b/Engine/Source/Runtime/Engine/Private/Animation/AnimMontage.cpp
@@ -7,6 +7,7 @@
 #include "Animation/AnimMontage.h"
 #include "UObject/LinkerLoad.h"
 #include "UObject/Package.h"
+#include "UObject/UObjectThreadContext.h"
 #include "Animation/AssetMappingTable.h"
 #include "Animation/AnimSequence.h"
 #include "Animation/AnimInstance.h"
@@ -596,7 +597,12 @@ void UAnimMontage::RefreshCacheData()
 	// This gets called whenever notifies are modified in the editor, so refresh our branch list
 	RefreshBranchingPointMarkers();
 #if WITH_EDITOR
-	PropagateChanges();
+	if (!FUObjectThreadContext::Get().IsRoutingPostLoad)
+	{
+		// This is not needed during post load (as the child montages themselves will handle
+		// updating and calling it can cause deterministic cooking issues depending on load order
+		PropagateChanges();
+	}
 #endif // WITH_EDITOR
 }
 
diff --git a/Engine/Source/Runtime/Engine/Private/AnimationAsset.cpp b/Engine/Source/Runtime/Engine/Private/AnimationAsset.cpp
--- a/Engine/Source/Runtime/Engine/Private/AnimationAsset.cpp
+++ b/Engine/Source/Runtime/Engine/Private/AnimationAsset.cpp
@@ -200,6 +200,16 @@ void UAnimationAsset::PostLoad()
 		Skeleton->ConditionalPostLoad();
 	}
 
+	// Load Parent Asset, to make sure anything accessing from PostLoad has valid data to access
+	if (ParentAsset)
+	{
+		if (FLinkerLoad* ParentAssetLinker = ParentAsset->GetLinker())
+		{
+			ParentAssetLinker->Preload(ParentAsset);
+		}
+		ParentAsset->ConditionalPostLoad();
+	}
+
 	ValidateSkeleton();
 
 	check( Skeleton==NULL || SkeletonGuid.IsValid() );
