Chaos - RBAN disable scale (TEMP)

diff --git a/Engine/Source/Runtime/Engine/Private/PhysicsEngine/ImmediatePhysics/ImmediatePhysicsChaos/ImmediatePhysicsJointHandle_Chaos.cpp b/Engine/Source/Runtime/Engine/Private/PhysicsEngine/ImmediatePhysics/ImmediatePhysicsChaos/ImmediatePhysicsJointHandle_Chaos.cpp
--- a/Engine/Source/Runtime/Engine/Private/PhysicsEngine/ImmediatePhysics/ImmediatePhysicsChaos/ImmediatePhysicsJointHandle_Chaos.cpp
+++ b/Engine/Source/Runtime/Engine/Private/PhysicsEngine/ImmediatePhysics/ImmediatePhysicsChaos/ImmediatePhysicsJointHandle_Chaos.cpp
@@ -171,6 +171,9 @@ namespace ImmediatePhysics_Chaos
 			// BodyInstance/PhysX has the constraint locations in actor-space, but we need them in Center-of-Mass space
 			ConstraintFrames[0] = FParticleUtilities::ActorLocalToParticleLocal(TGenericParticleHandle<FReal, 3>(Actor1->GetParticle()), ConstraintInstance->GetRefFrame(EConstraintFrame::Frame1));
 			ConstraintFrames[1] = FParticleUtilities::ActorLocalToParticleLocal(TGenericParticleHandle<FReal, 3>(Actor2->GetParticle()), ConstraintInstance->GetRefFrame(EConstraintFrame::Frame2));
+			FReal JointScale = ConstraintInstance->GetLastKnownScale();
+			ConstraintFrames[0].ScaleTranslation(JointScale);
+			ConstraintFrames[1].ScaleTranslation(JointScale);
 			TransferJointSettings(ConstraintInstance, ConstraintSettings);
 		}
 		else
