CIS fix

diff --git a/Engine/Source/Runtime/Engine/Private/AnimationAsset.cpp b/Engine/Source/Runtime/Engine/Private/AnimationAsset.cpp
--- a/Engine/Source/Runtime/Engine/Private/AnimationAsset.cpp
+++ b/Engine/Source/Runtime/Engine/Private/AnimationAsset.cpp
@@ -200,6 +200,7 @@ void UAnimationAsset::PostLoad()
 		Skeleton->ConditionalPostLoad();
 	}
 
+#if WITH_EDITORONLY_DATA
 	// Load Parent Asset, to make sure anything accessing from PostLoad has valid data to access
 	if (ParentAsset)
 	{
@@ -209,6 +210,7 @@ void UAnimationAsset::PostLoad()
 		}
 		ParentAsset->ConditionalPostLoad();
 	}
+#endif
 
 	ValidateSkeleton();
 
