Skip adding DORM_Initial actors to the active object list in AddNetworkActor and AddInitialObjects

diff --git a/Engine/Source/Runtime/Engine/Classes/Engine/NetDriver.h b/Engine/Source/Runtime/Engine/Classes/Engine/NetDriver.h
--- a/Engine/Source/Runtime/Engine/Classes/Engine/NetDriver.h
+++ b/Engine/Source/Runtime/Engine/Classes/Engine/NetDriver.h
@@ -1602,6 +1602,8 @@ public:
 
 	bool DidHitchLastFrame() const;
 
+	static bool IsDormInitialStartupActor(AActor* Actor);
+
 protected:
 
 	bool bMaySendProperties;
diff --git a/Engine/Source/Runtime/Engine/Private/DemoNetDriver.cpp b/Engine/Source/Runtime/Engine/Private/DemoNetDriver.cpp
--- a/Engine/Source/Runtime/Engine/Private/DemoNetDriver.cpp
+++ b/Engine/Source/Runtime/Engine/Private/DemoNetDriver.cpp
@@ -1777,7 +1777,7 @@ Demo Recording tick.
 
 bool UDemoNetDriver::DemoReplicateActor(AActor* Actor, UNetConnection* Connection, bool bMustReplicate)
 {
-	if (Actor->NetDormancy == DORM_Initial && Actor->IsNetStartupActor())
+	if (IsDormInitialStartupActor(Actor))
 	{
 		return false;
 	}
@@ -2854,7 +2854,7 @@ void UDemoNetDriver::TickDemoRecordFrame(float DeltaSeconds)
 						continue;
 					}
 
-					if (Actor->NetDormancy == DORM_Initial && Actor->IsNetStartupActor())
+					if (IsDormInitialStartupActor(Actor))
 					{
 						ActorsToRemove.Add(Actor);
 						continue;
diff --git a/Engine/Source/Runtime/Engine/Private/NetworkDriver.cpp b/Engine/Source/Runtime/Engine/Private/NetworkDriver.cpp
--- a/Engine/Source/Runtime/Engine/Private/NetworkDriver.cpp
+++ b/Engine/Source/Runtime/Engine/Private/NetworkDriver.cpp
@@ -448,10 +448,13 @@ void UNetDriver::AddNetworkActor(AActor* Actor)
 	LLM_SCOPE(ELLMTag::Networking);
 	ensureMsgf(Actor == nullptr || !(Actor->HasAnyFlags(RF_ClassDefaultObject | RF_ArchetypeObject)), TEXT("%s is a CDO or Archetype and should not be replicated."), *GetFullNameSafe(Actor));
 
-	GetNetworkObjectList().FindOrAdd(Actor, this);
-	if (ReplicationDriver)
+	if (!IsDormInitialStartupActor(Actor))
 	{
-		ReplicationDriver->AddNetworkActor(Actor);
+		GetNetworkObjectList().FindOrAdd(Actor, this);
+		if (ReplicationDriver)
+		{
+			ReplicationDriver->AddNetworkActor(Actor);
+		}
 	}
 }
 
@@ -3868,7 +3871,7 @@ void UNetDriver::ServerReplicateActors_BuildConsiderList( TArray<FNetworkObjectI
 			continue;
 		}
 
-		if ( Actor->NetDormancy == DORM_Initial && Actor->IsNetStartupActor() )
+		if ( IsDormInitialStartupActor(Actor) )
 		{
 			// This stat isn't that useful in its current form when using NetworkActors list
 			// We'll want to track initially dormant actors some other way to track them with stats
@@ -5739,6 +5742,11 @@ bool UNetDriver::DidHitchLastFrame() const
 	return bDidHitchLastFrame;
 }
 
+bool UNetDriver::IsDormInitialStartupActor(AActor* Actor)
+{
+	return Actor && Actor->IsNetStartupActor() && (Actor->NetDormancy == DORM_Initial);
+}
+
 ECreateReplicationChangelistMgrFlags UNetDriver::GetCreateReplicationChangelistMgrFlags() const
 {
 	return ECreateReplicationChangelistMgrFlags::None;
diff --git a/Engine/Source/Runtime/Engine/Private/NetworkObjectList.cpp b/Engine/Source/Runtime/Engine/Private/NetworkObjectList.cpp
--- a/Engine/Source/Runtime/Engine/Private/NetworkObjectList.cpp
+++ b/Engine/Source/Runtime/Engine/Private/NetworkObjectList.cpp
@@ -21,7 +21,7 @@ void FNetworkObjectList::AddInitialObjects(UWorld* const World, UNetDriver* NetD
 	for (FActorIterator Iter(World); Iter; ++Iter)
 	{
 		AActor* Actor = *Iter;
-		if (Actor != nullptr && !Actor->IsPendingKill() && ULevel::IsNetActor(Actor))
+		if (Actor != nullptr && !Actor->IsPendingKill() && ULevel::IsNetActor(Actor) && !UNetDriver::IsDormInitialStartupActor(Actor))
 		{
 			FindOrAdd(Actor, NetDriver);
 		}
