Deprecating UNetConnection::InternalAck to eventually be private
Deprecating UPackageMap::bSuppressLogs
DemoNetDriver GetWorld() cleanup
Remove check being hit when playing an older replay

diff --git a/Engine/Plugins/Online/OnlineSubsystemOculus/Source/Private/OculusNetConnection.cpp b/Engine/Plugins/Online/OnlineSubsystemOculus/Source/Private/OculusNetConnection.cpp
--- a/Engine/Plugins/Online/OnlineSubsystemOculus/Source/Private/OculusNetConnection.cpp
+++ b/Engine/Plugins/Online/OnlineSubsystemOculus/Source/Private/OculusNetConnection.cpp
@@ -117,7 +117,7 @@ void UOculusNetConnection::LowLevelSend(void* Data, int32 CountBits, FOutPacketT
 	if (!bBlockSend && CountBytes > 0)
 	{
 		UE_LOG(LogNetTraffic, VeryVerbose, TEXT("Low level send to: %llu Count: %d"), PeerID, CountBytes);
-		ovr_Net_SendPacket(PeerID, static_cast<size_t>(CountBytes), DataToSend, (InternalAck) ? ovrSend_Reliable : ovrSend_Unreliable);
+		ovr_Net_SendPacket(PeerID, static_cast<size_t>(CountBytes), DataToSend, (IsInternalAck()) ? ovrSend_Reliable : ovrSend_Unreliable);
 	}
 }
 
diff --git a/Engine/Source/Runtime/CoreUObject/Public/UObject/CoreNet.h b/Engine/Source/Runtime/CoreUObject/Public/UObject/CoreNet.h
--- a/Engine/Source/Runtime/CoreUObject/Public/UObject/CoreNet.h
+++ b/Engine/Source/Runtime/CoreUObject/Public/UObject/CoreNet.h
@@ -185,6 +185,7 @@ class COREUOBJECT_API UPackageMap : public UObject
 
 protected:
 
+	UE_DEPRECATED(4.25, "bSuppressLogs will be removed in a future release.")
 	bool					bSuppressLogs;
 
 	bool					bShouldTrackUnmappedGuids;
diff --git a/Engine/Source/Runtime/Engine/Classes/Engine/NetConnection.h b/Engine/Source/Runtime/Engine/Classes/Engine/NetConnection.h
--- a/Engine/Source/Runtime/Engine/Classes/Engine/NetConnection.h
+++ b/Engine/Source/Runtime/Engine/Classes/Engine/NetConnection.h
@@ -259,8 +259,9 @@ public:
 UCLASS(customConstructor, Abstract, MinimalAPI, transient, config=Engine)
 class ENGINE_VTABLE UNetConnection : public UPlayer
 {
-	GENERATED_UCLASS_BODY()
+	GENERATED_BODY()
 
+public:
 	/** child connections for secondary viewports */
 	UPROPERTY(transient)
 	TArray<class UChildConnection*> Children;
@@ -296,9 +297,23 @@ class ENGINE_VTABLE UNetConnection : public UPlayer
 	UPROPERTY()
 	int32	MaxPacket;						// Maximum packet size.
 
+	UE_DEPRECATED(4.25, "Please use IsInternalAck/SetInternalAck instead")
 	UPROPERTY()
 	uint32 InternalAck:1;					// Internally ack all packets, for 100% reliable connections.
 
+	bool IsInternalAck() const { return bInternalAck; }
+	void SetInternalAck(bool bValue) 
+	{ 
+		bInternalAck = bValue; 
+		PRAGMA_DISABLE_DEPRECATION_WARNINGS
+		InternalAck = bValue;
+		PRAGMA_ENABLE_DEPRECATION_WARNINGS
+	}
+
+private:
+	uint32 bInternalAck : 1;	// Internally ack all packets, for 100% reliable connections.
+
+public:
 	struct FURL			URL;				// URL of the other side.
 	
 	/** The remote address of this connection, typically generated from the URL. */
@@ -1148,7 +1163,7 @@ public:
 	{
 		// The InternalAck and ServerConnection conditions, are only there to exclude demo's and clients from this check,
 		// so that the check is only performed on servers.
-		return !!InternalAck || Driver->ServerConnection != nullptr || InReliable[0] != InitInReliable;
+		return IsInternalAck() || Driver->ServerConnection != nullptr || InReliable[0] != InitInReliable;
 	}
 
 	/**
diff --git a/Engine/Source/Runtime/Engine/Private/DataChannel.cpp b/Engine/Source/Runtime/Engine/Private/DataChannel.cpp
--- a/Engine/Source/Runtime/Engine/Private/DataChannel.cpp
+++ b/Engine/Source/Runtime/Engine/Private/DataChannel.cpp
@@ -88,7 +88,7 @@ static FAutoConsoleVariableRef CVarNetMaxConstructedPartialBunchSizeBytes(
 template<typename T>
 static const bool IsBunchTooLarge(UNetConnection* Connection, T* Bunch)
 {
-	return !Connection->InternalAck && Bunch != nullptr && Bunch->GetNumBytes() > NetMaxConstructedPartialBunchSizeBytes;
+	return !Connection->IsInternalAck() && Bunch != nullptr && Bunch->GetNumBytes() > NetMaxConstructedPartialBunchSizeBytes;
 }
 
 /*-----------------------------------------------------------------------------
@@ -205,7 +205,7 @@ bool UChannel::CleanUp(const bool bForDestroy, EChannelCloseReason CloseReason)
 	}
 
 	// remember sequence number of first non-acked outgoing reliable bunch for this slot
-	if (OutRec != NULL && !Connection->InternalAck)
+	if (OutRec != NULL && !Connection->IsInternalAck())
 	{
 		Connection->PendingOutRec[ChIndex] = OutRec->ChSequence;
 		//UE_LOG(LogNetTraffic, Log, TEXT("%i save pending out bunch %i"),ChIndex,Connection->PendingOutRec[ChIndex]);
@@ -413,7 +413,7 @@ void UChannel::ReceivedRawBunch( FInBunch & Bunch, bool & bOutSkipAck )
 	// Immediately consume the NetGUID portion of this bunch, regardless if it is partial or reliable.
 	// NOTE - For replays, we do this even earlier, to try and load this as soon as possible, in case there is an issue creating the channel
 	// If a replay fails to create a channel, we want to salvage as much as possible
-	if ( Bunch.bHasPackageMapExports && !Connection->InternalAck )
+	if ( Bunch.bHasPackageMapExports && !Connection->IsInternalAck() )
 	{
 		Cast<UPackageMapClient>( Connection->PackageMap )->ReceiveNetGUIDBunch( Bunch );
 
@@ -424,7 +424,7 @@ void UChannel::ReceivedRawBunch( FInBunch & Bunch, bool & bOutSkipAck )
 		}
 	}
 
-	if ( Connection->InternalAck && Broken )
+	if ( Connection->IsInternalAck() && Broken )
 	{
 		return;
 	}
@@ -434,7 +434,7 @@ void UChannel::ReceivedRawBunch( FInBunch & Bunch, bool & bOutSkipAck )
 	if ( Bunch.bReliable && Bunch.ChSequence != Connection->InReliable[ChIndex] + 1 )
 	{
 		// We shouldn't hit this path on 100% reliable connections
-		check( !Connection->InternalAck );
+		check( !Connection->IsInternalAck() );
 		// If this bunch has a dependency on a previous unreceived bunch, buffer it.
 		checkSlow(!Bunch.bOpen);
 
@@ -492,7 +492,7 @@ void UChannel::ReceivedRawBunch( FInBunch & Bunch, bool & bOutSkipAck )
 		while( InRec )
 		{
 			// We shouldn't hit this path on 100% reliable connections
-			check( !Connection->InternalAck );
+			check( !Connection->IsInternalAck() );
 
 			if( InRec->ChSequence!=Connection->InReliable[ChIndex]+1 )
 				break;
@@ -682,7 +682,7 @@ bool UChannel::ReceivedNextBunch( FInBunch & Bunch, bool & bOutSkipAck )
 				// Merge problem - delete InPartialBunch. This is mainly so that in the unlikely chance that ChSequence wraps around, we wont merge two completely separate partial bunches.
 
 				// We shouldn't hit this path on 100% reliable connections
-				check( !Connection->InternalAck );
+				check( !Connection->IsInternalAck());
 				
 				bOutSkipAck = true;	// Don't ack the packet, since we didn't process the bunch
 
@@ -771,7 +771,7 @@ bool UChannel::ReceivedNextBunch( FInBunch & Bunch, bool & bOutSkipAck )
 					return false;
 				}
 
-				if ( !ensure( !Connection->InternalAck ) )
+				if ( !ensure( !Connection->IsInternalAck()) )
 				{
 					// Shouldn't be possible for 100% reliable connections
 					Broken = 1;
@@ -973,7 +973,7 @@ FPacketIdRange UChannel::SendBunch( FOutBunch* Bunch, bool Merge )
 
 	// Add any export bunches
 	// Replay connections will manage export bunches separately.
-	if (!Connection->InternalAck)
+	if (!Connection->IsInternalAck())
 	{
 		AppendExportBunches( OutgoingBunches );
 	}
@@ -1081,7 +1081,7 @@ FPacketIdRange UChannel::SendBunch( FOutBunch* Bunch, bool Merge )
 
 	const bool bOverflowsReliable = (NumOutRec + OutgoingBunches.Num() >= RELIABLE_BUFFER + Bunch->bClose);
 
-	if ((GCVarNetPartialBunchReliableThreshold > 0) && (OutgoingBunches.Num() >= GCVarNetPartialBunchReliableThreshold) && !Connection->InternalAck)
+	if ((GCVarNetPartialBunchReliableThreshold > 0) && (OutgoingBunches.Num() >= GCVarNetPartialBunchReliableThreshold) && !Connection->IsInternalAck())
 	{
 		if (!bOverflowsReliable)
 		{
@@ -1816,7 +1816,7 @@ int64 UActorChannel::Close(EChannelCloseReason Reason)
 
 		if (Connection)
 		{
-			if ((Reason == EChannelCloseReason::Dormancy) && !Connection->InternalAck) // Replay connections always keep dormant channels open and handle this logic elsewhere
+			if ((Reason == EChannelCloseReason::Dormancy) && !Connection->IsInternalAck()) // Replay connections always keep dormant channels open and handle this logic elsewhere
 			{
 				if (Connection->Driver)
 			{
@@ -2521,7 +2521,7 @@ void UActorChannel::ReceivedBunch( FInBunch & Bunch )
 			Mark.Pop(Bunch);
 
 			// we can now map guid to channel, even if all the bunches get queued
-			if (Connection->InternalAck)
+			if (Connection->IsInternalAck())
 			{
 				Connection->NotifyActorNetGUID(this);
 			}
@@ -2624,7 +2624,7 @@ void UActorChannel::ProcessBunch( FInBunch & Bunch )
 			UE_LOG(LogNet, Warning, TEXT("UActorChannel::ProcessBunch: SerializeNewActor failed to find/spawn actor. Actor: %s, Channel: %i"), *GetFullNameSafe(NewChannelActor), ChIndex);
 			Broken = 1;
 
-			if (!Connection->InternalAck
+			if (!Connection->IsInternalAck()
 #if !UE_BUILD_SHIPPING
 				&& !bBlockChannelFailure
 #endif
@@ -2691,9 +2691,9 @@ void UActorChannel::ProcessBunch( FInBunch & Bunch )
 		
 		if ( Bunch.IsError() )
 		{
-			if ( Connection->InternalAck )
+			if ( Connection->IsInternalAck() )
 			{
-				UE_LOG( LogNet, Warning, TEXT( "UActorChannel::ReceivedBunch: ReadContentBlockPayload FAILED. Bunch.IsError() == TRUE. (InternalAck) Breaking actor. RepObj: %s, Channel: %i" ), RepObj ? *RepObj->GetFullName() : TEXT( "NULL" ), ChIndex );
+				UE_LOG( LogNet, Warning, TEXT( "UActorChannel::ReceivedBunch: ReadContentBlockPayload FAILED. Bunch.IsError() == TRUE. (IsInternalAck) Breaking actor. RepObj: %s, Channel: %i" ), RepObj ? *RepObj->GetFullName() : TEXT( "NULL" ), ChIndex );
 				Broken = 1;
 				break;
 			}
@@ -2734,9 +2734,9 @@ void UActorChannel::ProcessBunch( FInBunch & Bunch )
 
 		if ( !Replicator->ReceivedBunch( Reader, RepFlags, bHasRepLayout, bHasUnmapped ) )
 		{
-			if ( Connection->InternalAck )
+			if ( Connection->IsInternalAck() )
 			{
-				UE_LOG( LogNet, Warning, TEXT( "UActorChannel::ProcessBunch: Replicator.ReceivedBunch failed (Ignoring because of InternalAck). RepObj: %s, Channel: %i" ), RepObj ? *RepObj->GetFullName() : TEXT( "NULL" ), ChIndex );
+				UE_LOG( LogNet, Warning, TEXT( "UActorChannel::ProcessBunch: Replicator.ReceivedBunch failed (Ignoring because of IsInternalAck). RepObj: %s, Channel: %i" ), RepObj ? *RepObj->GetFullName() : TEXT( "NULL" ), ChIndex );
 				Broken = 1;
 				continue;		// Don't consider this catastrophic in replays
 			}
@@ -3257,7 +3257,7 @@ void UActorChannel::BecomeDormant()
 	bPendingDormancy = false;
 	Dormant = true;
 	// Replays do not close dormant channels currently, so just mark the channel dormant directly.
-	if (Connection && Connection->InternalAck)
+	if (Connection && Connection->IsInternalAck())
 	{
 		Connection->Driver->NotifyActorFullyDormantForConnection(Actor, Connection);
 	}
@@ -3509,7 +3509,7 @@ UObject* UActorChannel::ReadContentBlockHeader( FInBunch & Bunch, bool& bObjectD
 		if ( SubObj == NULL )
 		{
 			// (ignore though if this is for replays)
-			if ( !Connection->InternalAck )
+			if ( !Connection->IsInternalAck() )
 			{
 				UE_LOG( LogNetTraffic, Log, TEXT( "ReadContentBlockHeader: Stably named sub-object not found. Its level may have streamed out. Component: %s, Actor: %s" ), *Connection->Driver->GuidCache->FullNetGUIDPath( NetGUID ), *Actor->GetName() );
 			}
@@ -3561,7 +3561,7 @@ UObject* UActorChannel::ReadContentBlockHeader( FInBunch & Bunch, bool& bObjectD
 		if ( SubObj == NULL )
 		{
 			// (unless we're using replays, which could be backwards compatibility kicking in)
-			if ( !Connection->InternalAck )
+			if ( !Connection->IsInternalAck() )
 			{
 				UE_LOG( LogNetTraffic, Error, TEXT( "UActorChannel::ReadContentBlockHeader: Unable to read sub-object class (SubObj == NULL). Actor: %s" ), *Actor->GetName() );
 				Bunch.SetError();
@@ -3664,7 +3664,7 @@ int32 UActorChannel::WriteFieldHeaderAndPayload( FNetBitWriter& Bunch, const FCl
 
 	NET_CHECKSUM( Bunch );
 
-	if ( Connection->InternalAck && !bIgnoreInternalAck )
+	if ( Connection->IsInternalAck() && !bIgnoreInternalAck )
 	{
 		check( NetFieldExportGroup != nullptr );
 
@@ -3672,7 +3672,7 @@ int32 UActorChannel::WriteFieldHeaderAndPayload( FNetBitWriter& Bunch, const FCl
 
 		check( NetFieldExportHandle >= 0 );
 
-		( ( UPackageMapClient* )Connection->PackageMap )->TrackNetFieldExport( NetFieldExportGroup, NetFieldExportHandle );
+		CastChecked<UPackageMapClient>(Connection->PackageMap)->TrackNetFieldExport( NetFieldExportGroup, NetFieldExportHandle );
 
 		check( NetFieldExportHandle < NetFieldExportGroup->NetFieldExports.Num() );
 
@@ -3714,7 +3714,7 @@ bool UActorChannel::ReadFieldHeaderAndPayload( UObject* Object, const FClassNetC
 	
 	NET_CHECKSUM( Bunch );
 
-	if ( Connection->InternalAck )
+	if ( Connection->IsInternalAck() )
 	{
 		if ( NetFieldExportGroup == nullptr )
 		{
@@ -3812,7 +3812,7 @@ static FORCEINLINE FString GenerateClassNetCacheNetFieldExportGroupName( const U
 
 FNetFieldExportGroup* UActorChannel::GetOrCreateNetFieldExportGroupForClassNetCache( const UObject* Object )
 {
-	if ( !Connection->InternalAck )
+	if ( !Connection->IsInternalAck() )
 	{
 		return nullptr;
 	}
@@ -3869,14 +3869,14 @@ FNetFieldExportGroup* UActorChannel::GetOrCreateNetFieldExportGroupForClassNetCa
 
 FNetFieldExportGroup* UActorChannel::GetNetFieldExportGroupForClassNetCache( UClass* ObjectClass )
 {
-	if ( !Connection->InternalAck )
+	if ( !Connection->IsInternalAck() )
 	{
 		return nullptr;
 	}	
 
 	const FString NetFieldExportGroupName = GenerateClassNetCacheNetFieldExportGroupName( ObjectClass );
 
-	UPackageMapClient* PackageMapClient = ( ( UPackageMapClient* )Connection->PackageMap );
+	UPackageMapClient* PackageMapClient = CastChecked<UPackageMapClient>(Connection->PackageMap);
 
 	TSharedPtr< FNetFieldExportGroup > NetFieldExportGroup = PackageMapClient->GetNetFieldExportGroup( NetFieldExportGroupName );
 
diff --git a/Engine/Source/Runtime/Engine/Private/DataReplication.cpp b/Engine/Source/Runtime/Engine/Private/DataReplication.cpp
--- a/Engine/Source/Runtime/Engine/Private/DataReplication.cpp
+++ b/Engine/Source/Runtime/Engine/Private/DataReplication.cpp
@@ -152,7 +152,7 @@ public:
 			UpdateCachedRepLayout();
 			UPackageMapClient* PackageMapClient = ((UPackageMapClient*)Params.Map);
 
-			if (PackageMapClient && PackageMapClient->GetConnection()->InternalAck)
+			if (PackageMapClient && PackageMapClient->GetConnection()->IsInternalAck())
 			{
 				if (Ar.IsSaving())
 				{
@@ -321,7 +321,7 @@ bool FObjectReplicator::SendCustomDeltaProperty(UObject* InObject, uint16 Custom
 	Parms.CustomDeltaIndex = CustomDeltaIndex;
 	Parms.bSupportsFastArrayDeltaStructSerialization = bSupportsFastArrayDelta;
 	Parms.Connection = Connection;
-	Parms.bInternalAck = Connection->InternalAck;
+	Parms.bInternalAck = Connection->IsInternalAck();
 
 	return FNetSerializeCB::SendCustomDeltaProperty(*RepLayout, Parms, CustomDeltaIndex);
 }
@@ -352,9 +352,9 @@ void FObjectReplicator::InitRecentProperties(uint8* Source)
 	// for replays (and the DemoNetDriver will be acts as a server during recording).
 	TSharedPtr<FRepChangedPropertyTracker> RepChangedPropertyTracker = bCreateSendingState ? ConnectionDriver->FindOrCreateRepChangedPropertyTracker(MyObject) : nullptr;
 
-	// If acting as a server and are InternalAck, that means we're recording.
+	// If acting as a server and are IsInternalAck, that means we're recording.
 	// In that case, we don't need to create any receiving state, as no one will be sending data to us.
-	ECreateRepStateFlags Flags = (Connection->InternalAck && bIsServer) ? ECreateRepStateFlags::SkipCreateReceivingState : ECreateRepStateFlags::None;
+	ECreateRepStateFlags Flags = (Connection->IsInternalAck() && bIsServer) ? ECreateRepStateFlags::SkipCreateReceivingState : ECreateRepStateFlags::None;
 	RepState = LocalRepLayout.CreateRepState(Source, RepChangedPropertyTracker, Flags);
 
 	if (!bCreateSendingState)
@@ -374,7 +374,7 @@ void FObjectReplicator::InitRecentProperties(uint8* Source)
 		const uint16 NumLifetimeCustomDeltaProperties = FNetSerializeCB::GetNumLifetimeCustomDeltaProperties(LocalRepLayout);
 		SendingRepState->RecentCustomDeltaState.SetNum(NumLifetimeCustomDeltaProperties);
 
-		const bool bIsRecordingReplay = Connection->InternalAck;
+		const bool bIsRecordingReplay = Connection->IsInternalAck();
 
 		if (bIsRecordingReplay)
 		{
@@ -577,7 +577,7 @@ void FObjectReplicator::StartReplicating(class UActorChannel * InActorChannel)
 	if (ConnectionNetDriver->IsServer() || ConnectionNetDriver->MaySendProperties())
 	{
 		// We don't need to handle retirement if our connection is reliable.
-		if (!Connection->InternalAck)
+		if (!Connection->IsInternalAck())
 		{
 			if (FSendingRepState * SendingRepState = RepState.IsValid() ? RepState->GetSendingRepState() : nullptr)
 			{
@@ -1012,7 +1012,7 @@ bool FObjectReplicator::ReceivedBunch(FNetBitReader& Bunch, const FReplicationFl
 			Parms.Reader = &Reader;
 			Parms.NetSerializeCB = &NetSerializeCB;
 			Parms.Connection = Connection;
-			Parms.bInternalAck = Connection->InternalAck;
+			Parms.bInternalAck = Connection->IsInternalAck();
 			Parms.Object = Object;
 
 			if (!FNetSerializeCB::ReceiveCustomDeltaProperty(LocalRepLayout, ReceivingRepState, Parms, ReplicatedProp))
@@ -1261,7 +1261,7 @@ void FObjectReplicator::UpdateGuidToReplicatorMap()
 		Parms.GatherGuidReferences = &LocalReferencedGuids;
 		Parms.TrackedGuidMemoryBytes = &LocalTrackedGuidMemoryBytes;
 		Parms.Object = GetObject();
-		Parms.bInternalAck = Connection->InternalAck;
+		Parms.bInternalAck = Connection->IsInternalAck();
 
 		LocalRepLayout.GatherGuidReferences(RepState->GetReceivingRepState(), Parms, LocalReferencedGuids, LocalTrackedGuidMemoryBytes);
 	}
@@ -1318,7 +1318,7 @@ bool FObjectReplicator::MoveMappedObjectToUnmapped(const FNetworkGUID& GUID)
 	FNetSerializeCB NetSerializeCB(Connection->Driver);
 	FNetDeltaSerializeInfo Parms;
 	Parms.Connection = Connection;
-	Parms.bInternalAck = Connection->InternalAck;
+	Parms.bInternalAck = Connection->IsInternalAck();
 	Parms.Map = Connection->PackageMap;
 	Parms.Object = GetObject();
 	Parms.NetSerializeCB = &NetSerializeCB;
@@ -1474,7 +1474,7 @@ void FObjectReplicator::ReplicateCustomDeltaProperties( FNetBitWriter & Bunch, F
 
 		FPropertyRetirement** LastNext = nullptr;
 
-		if (!Connection->InternalAck)
+		if (!Connection->IsInternalAck())
 		{
 			// Get info.
 			FPropertyRetirement& Retire = SendingRepState->Retirement[CustomDeltaProperty];
@@ -1499,7 +1499,7 @@ void FObjectReplicator::ReplicateCustomDeltaProperties( FNetBitWriter & Bunch, F
 			continue;
 		}
 
-		if (!Connection->InternalAck)
+		if (!Connection->IsInternalAck())
 		{
 			*LastNext = new FPropertyRetirement();
 
@@ -1849,7 +1849,7 @@ void FObjectReplicator::QueueRemoteFunctionBunch( UFunction* Func, FOutBunch &Bu
 			PackageMapClient->GetMustBeMappedGuidsInLastBunch().Reset();
 		}
 
-		if (!Connection->InternalAck)
+		if (!Connection->IsInternalAck())
 		{
 			// Copy over any exported bunches
 			PackageMapClient->AppendExportBunches(OwningChannel->QueuedExportBunches);
@@ -1994,7 +1994,7 @@ void FObjectReplicator::UpdateUnmappedObjects(bool & bOutHasMoreUnmapped)
 	FNetDeltaSerializeInfo Parms;
 	Parms.Object = Object;
 	Parms.Connection = Connection;
-	Parms.bInternalAck = Connection->InternalAck;
+	Parms.bInternalAck = Connection->IsInternalAck();
 	Parms.Map = Connection->PackageMap;
 	Parms.NetSerializeCB = &NetSerializeCB;
 
@@ -2065,7 +2065,7 @@ void FObjectReplicator::UpdateUnmappedObjects(bool & bOutHasMoreUnmapped)
 
 			if (!bSuccess)
 			{
-				if (bIsServer && !Connection->InternalAck)
+				if (bIsServer && !Connection->IsInternalAck())
 				{
 					// Close our connection and abort rpcs as things are invalid
 					PendingLocalRPCs.Empty();
diff --git a/Engine/Source/Runtime/Engine/Private/DemoNetDriver.cpp b/Engine/Source/Runtime/Engine/Private/DemoNetDriver.cpp
--- a/Engine/Source/Runtime/Engine/Private/DemoNetDriver.cpp
+++ b/Engine/Source/Runtime/Engine/Private/DemoNetDriver.cpp
@@ -125,6 +125,8 @@ static void ConditionallyDisplayBurnInTime(uint32 RecordedCL, float CurrentDemoT
 
 static void FlushNetChecked(UNetConnection& NetConnection)
 {
+	DECLARE_SCOPE_CYCLE_COUNTER(TEXT("Replay net flush"), STAT_ReplayFlushNet, STATGROUP_Net);
+
 	NetConnection.FlushNet();
 	check(NetConnection.SendBuffer.GetNumBits() == 0);
 }
@@ -260,9 +262,14 @@ private:
 };
 
 // Helps manage packets, and any associations with streaming levels or exported GUIDs / fields.
-class FScopedPacketManager : public FNoncopyable
+class FScopedPacketManager
 {
 public:
+	FScopedPacketManager(FScopedPacketManager&&) = delete;
+	FScopedPacketManager(const FScopedPacketManager&) = delete;
+	FScopedPacketManager& operator=(const FScopedPacketManager&) = delete;
+	FScopedPacketManager& operator=(FScopedPacketManager&&) = delete;
+
 	FScopedPacketManager(UDemoNetConnection& InConnection, const uint32 InSeenLevelIndex):
 		Connection(InConnection),
 		Packets((Connection.ResendAllDataState != EResendAllDataState::None) ? InConnection.QueuedCheckpointPackets : InConnection.QueuedDemoPackets),
@@ -320,9 +327,14 @@ public:
  * This relies on serialization always used a fixed number of bytes for primitive types,
  * and Sane implementations of Seek and Tell.
  */
-class FScopedStoreArchiveOffset : public FNoncopyable
+class FScopedStoreArchiveOffset
 {
 public:
+	FScopedStoreArchiveOffset(FScopedStoreArchiveOffset&&) = delete;
+	FScopedStoreArchiveOffset(const FScopedStoreArchiveOffset&) = delete;
+	FScopedStoreArchiveOffset& operator=(const FScopedStoreArchiveOffset&) = delete;
+	FScopedStoreArchiveOffset& operator=(FScopedStoreArchiveOffset&&) = delete;
+
 	typedef UDemoNetDriver::FArchivePos FArchivePos;
 
 	FScopedStoreArchiveOffset(FArchive& InAr):
@@ -349,9 +361,14 @@ private:
 	const FArchivePos StartPosition;
 };
 
-class FScopedIgnoreOpenChannels : public FNoncopyable
+class FScopedIgnoreOpenChannels
 {
 public:
+	FScopedIgnoreOpenChannels(FScopedIgnoreOpenChannels&&) = delete;
+	FScopedIgnoreOpenChannels(const FScopedIgnoreOpenChannels&) = delete;
+	FScopedIgnoreOpenChannels& operator=(const FScopedIgnoreOpenChannels&) = delete;
+	FScopedIgnoreOpenChannels& operator=(FScopedIgnoreOpenChannels&&) = delete;
+
 	FScopedIgnoreOpenChannels(UNetConnection* InConnection):
 		Connection(InConnection)
 	{
@@ -373,9 +390,14 @@ private:
 	TWeakObjectPtr<UNetConnection> Connection;
 };
 
-class FScopedIgnoreClosedChannels : public FNoncopyable
+class FScopedIgnoreClosedChannels
 {
 public:
+	FScopedIgnoreClosedChannels(FScopedIgnoreClosedChannels&&) = delete;
+	FScopedIgnoreClosedChannels(const FScopedIgnoreClosedChannels&) = delete;
+	FScopedIgnoreClosedChannels& operator=(const FScopedIgnoreClosedChannels&) = delete;
+	FScopedIgnoreClosedChannels& operator=(FScopedIgnoreClosedChannels&&) = delete;
+
 	FScopedIgnoreClosedChannels(UNetConnection* InConnection, TSet<FNetworkGUID>&& IgnoredBunchGuids)
 		: ServerConnection(InConnection)
 	{
@@ -694,9 +716,14 @@ private:
 	TOptional<FGotoResult> GotoResult;
 };
 
-class FScopedForceUnicodeInArchive : public FNoncopyable
+class FScopedForceUnicodeInArchive
 {
 public:
+	FScopedForceUnicodeInArchive(FScopedForceUnicodeInArchive&&) = delete;
+	FScopedForceUnicodeInArchive(const FScopedForceUnicodeInArchive&) = delete;
+	FScopedForceUnicodeInArchive& operator=(const FScopedForceUnicodeInArchive&) = delete;
+	FScopedForceUnicodeInArchive& operator=(FScopedForceUnicodeInArchive&&) = delete;
+
 	FScopedForceUnicodeInArchive(FArchive& InArchive)
 		: Archive(InArchive)
 		, bWasUnicode(InArchive.IsForcingUnicode())
@@ -955,22 +982,22 @@ void UDemoNetDriver::ResetDemoState()
 
 bool UDemoNetDriver::InitConnect(FNetworkNotify* InNotify, const FURL& ConnectURL, FString& Error)
 {
-	if (GetWorld() == nullptr)
+	if (World == nullptr)
 	{
-		UE_LOG(LogDemo, Error, TEXT("GetWorld() == nullptr"));
+		UE_LOG(LogDemo, Error, TEXT("World == nullptr"));
 		return false;
 	}
 
-	if (GetWorld()->GetGameInstance() == nullptr)
+	if (World->GetGameInstance() == nullptr)
 	{
-		UE_LOG(LogDemo, Error, TEXT("GetWorld()->GetGameInstance() == nullptr"));
+		UE_LOG(LogDemo, Error, TEXT("World->GetGameInstance() == nullptr"));
 		return false;
 	}
 
 	// handle default initialization
 	if (!InitBase(true, InNotify, ConnectURL, false, Error))
 	{
-		GetWorld()->GetGameInstance()->HandleDemoPlaybackFailure(EDemoPlayFailure::InitBase, FString(TEXT("InitBase FAILED")));
+		World->GetGameInstance()->HandleDemoPlaybackFailure(EDemoPlayFailure::InitBase, FString(TEXT("InitBase FAILED")));
 		return false;
 	}
 
@@ -998,7 +1025,7 @@ bool UDemoNetDriver::InitConnect(FNetworkNotify* InNotify, const FURL& ConnectUR
 	if (GetDuplicateLevelID() == -1)
 	{
 		// Set this driver as the demo net driver for the source level collection.
-		FLevelCollection* const SourceCollection = GetWorld()->FindCollectionByType(ELevelCollectionType::DynamicSourceLevels);
+		FLevelCollection* const SourceCollection = World->FindCollectionByType(ELevelCollectionType::DynamicSourceLevels);
 		if (SourceCollection)
 		{
 			SourceCollection->SetDemoNetDriver(this);
@@ -1007,7 +1034,7 @@ bool UDemoNetDriver::InitConnect(FNetworkNotify* InNotify, const FURL& ConnectUR
 	else
 	{
 		// Set this driver as the demo net driver for the duplicate level collection.
-		FLevelCollection* const DuplicateCollection = GetWorld()->FindCollectionByType(ELevelCollectionType::DynamicDuplicatedLevels);
+		FLevelCollection* const DuplicateCollection = World->FindCollectionByType(ELevelCollectionType::DynamicDuplicatedLevels);
 		if (DuplicateCollection)
 		{
 			DuplicateCollection->SetDemoNetDriver(this);
@@ -1042,7 +1069,7 @@ bool UDemoNetDriver::InitConnect(FNetworkNotify* InNotify, const FURL& ConnectUR
 
 bool UDemoNetDriver::ReadPlaybackDemoHeader(FString& Error)
 {
-	UGameInstance* GameInstance = GetWorld()->GetGameInstance();
+	UGameInstance* GameInstance = World->GetGameInstance();
 
 	PlaybackDemoHeader = FNetworkDemoHeader();
 
@@ -1121,7 +1148,7 @@ bool UDemoNetDriver::InitConnectInternal(FString& Error)
 
 	if (GetDuplicateLevelID() == -1)
 	{
-		if (bAsyncLoadWorld && GetWorld()->WorldType != EWorldType::PIE) // Editor doesn't support async map travel
+		if (bAsyncLoadWorld && World->WorldType != EWorldType::PIE) // Editor doesn't support async map travel
 		{
 			LevelNamesAndTimes = PlaybackDemoHeader.LevelNamesAndTimes;
 
@@ -1136,11 +1163,11 @@ bool UDemoNetDriver::InitConnectInternal(FString& Error)
 			FURL LocalDemoURL;
 			LocalDemoURL.Map = PlaybackDemoHeader.LevelNamesAndTimes[0].LevelName;
 
-			FWorldContext * WorldContext = GEngine->GetWorldContextFromWorld(GetWorld());
+			FWorldContext * WorldContext = GEngine->GetWorldContextFromWorld(World);
 
 			if (WorldContext == nullptr)
 			{
-				UGameInstance* GameInstance = GetWorld()->GetGameInstance();
+				UGameInstance* GameInstance = World->GetGameInstance();
 
 				Error = FString::Printf(TEXT("No world context"));
 				UE_LOG(LogDemo, Error, TEXT("UDemoNetDriver::InitConnect: %s"), *Error);
@@ -1148,7 +1175,7 @@ bool UDemoNetDriver::InitConnectInternal(FString& Error)
 				return false;
 			}
 
-			GetWorld()->ClearDemoNetDriver();
+			World->ClearDemoNetDriver();
 			SetWorld(nullptr);
 
 			auto NewPendingNetGame = NewObject<UDemoPendingNetGame>();
@@ -1180,7 +1207,7 @@ bool UDemoNetDriver::InitListen(FNetworkNotify* InNotify, FURL& ListenURL, bool
 
 	check(World != nullptr);
 
-	class AWorldSettings * WorldSettings = World->GetWorldSettings(); 
+	AWorldSettings* WorldSettings = World->GetWorldSettings(); 
 
 	if (!WorldSettings)
 	{
@@ -1234,7 +1261,7 @@ bool UDemoNetDriver::InitListen(FNetworkNotify* InNotify, FURL& ListenURL, bool
 
 	AddNewLevel(World->GetOuter()->GetName());
 
-	bool Result = WriteNetworkDemoHeader(Error);
+	const bool Result = WriteNetworkDemoHeader(Error);
 
 	// Spawn the demo recording spectator.
 	SpawnDemoRecSpectator(Connection, ListenURL);
@@ -1248,7 +1275,7 @@ void UDemoNetDriver::OnLevelAddedToWorld(ULevel* InLevel, UWorld* InWorld)
 
 	Super::OnLevelAddedToWorld(InLevel, InWorld);
 
-	if (InLevel && !InLevel->bClientOnlyVisible && GetWorld() == InWorld && HasLevelStreamingFixes() && IsPlaying())
+	if (InLevel && !InLevel->bClientOnlyVisible && (World == InWorld) && HasLevelStreamingFixes() && IsPlaying())
 	{
 		if (!NewStreamingLevelsThisFrame.Contains(InLevel) && !LevelsPendingFastForward.Contains(InLevel))
 		{
@@ -1274,7 +1301,7 @@ void UDemoNetDriver::OnLevelRemovedFromWorld(ULevel* InLevel, UWorld* InWorld)
 {
 	Super::OnLevelRemovedFromWorld(InLevel, InWorld);
 
-	if (InLevel && !InLevel->bClientOnlyVisible && GetWorld() == InWorld && HasLevelStreamingFixes() && IsPlaying())
+	if (InLevel && !InLevel->bClientOnlyVisible && (World == InWorld) && HasLevelStreamingFixes() && IsPlaying())
 	{
 		const FString LevelPackageName = GetLevelPackageName(*InLevel);
 		if (LevelStatusesByName.Contains(LevelPackageName))
@@ -1469,7 +1496,7 @@ bool UDemoNetDriver::IsServer() const
 
 bool UDemoNetDriver::ShouldTickFlushAsyncEndOfFrame() const
 {
-	return	GEngine && GEngine->ShouldDoAsyncEndOfFrameTasks() && CVarDemoClientRecordAsyncEndOfFrame.GetValueOnAnyThread() != 0 && World && World->IsRecordingClientReplay();
+	return GEngine && GEngine->ShouldDoAsyncEndOfFrameTasks() && CVarDemoClientRecordAsyncEndOfFrame.GetValueOnAnyThread() != 0 && World && World->IsRecordingClientReplay();
 }
 
 void UDemoNetDriver::TickFlush(float DeltaSeconds)
@@ -1505,7 +1532,7 @@ void UDemoNetDriver::TickFlushInternal(float DeltaSeconds)
 		return Collection.GetDemoNetDriver() == this;
 	}) : INDEX_NONE;
 
-	FScopedLevelCollectionContextSwitch LCSwitch(FoundCollectionIndex, GetWorld());
+	FScopedLevelCollectionContextSwitch LCSwitch(FoundCollectionIndex, World);
 
 	Super::TickFlush(DeltaSeconds);
 
@@ -1595,12 +1622,12 @@ void UDemoNetDriver::TickDispatch(float DeltaSeconds)
 	LLM_SCOPE(ELLMTag::Networking);
 
 	// Set the context on the world for this driver's level collection.
-		const int32 FoundCollectionIndex = World ? World->GetLevelCollections().IndexOfByPredicate([this](const FLevelCollection& Collection)
+	const int32 FoundCollectionIndex = World ? World->GetLevelCollections().IndexOfByPredicate([this](const FLevelCollection& Collection)
 	{
 		return Collection.GetDemoNetDriver() == this;
 	}) : INDEX_NONE;
 
-	FScopedLevelCollectionContextSwitch LCSwitch(FoundCollectionIndex, GetWorld());
+	FScopedLevelCollectionContextSwitch LCSwitch(FoundCollectionIndex, World);
 
 	Super::TickDispatch(DeltaSeconds);
 
@@ -1972,8 +1999,10 @@ void UDemoNetDriver::SaveCheckpoint()
 		return;
 	}
 
+	UNetConnection* Connection = ClientConnections[0];
+
 	check(CheckpointArchive->TotalSize() == 0);
-	check(ClientConnections[0]->SendBuffer.GetNumBits() == 0);
+	check(Connection->SendBuffer.GetNumBits() == 0);
 	check(CheckpointSaveContext.CheckpointSaveState == ECheckpointSaveState_Idle);
 	
 	const bool bDeltaCheckpoint = HasDeltaCheckpoints();
@@ -1990,7 +2019,6 @@ void UDemoNetDriver::SaveCheckpoint()
 			const UObject* Level;
 		};
 
-		const UNetConnection* Connection = ClientConnections[0];
 		const FActorChannelMap& ActorChannelMap = Connection->ActorChannelMap();
 		const FNetworkObjectList::FNetworkObjectSet& AllObjectsSet = GetNetworkObjectList().GetAllObjects();
 
@@ -2055,7 +2083,7 @@ void UDemoNetDriver::SaveCheckpoint()
 		{
 			AActor* Actor = ObjectInfo.Get()->Actor;
 
-			if (ClientConnections[0]->FindActorChannelRef(Actor))
+			if (Connection->FindActorChannelRef(Actor))
 			{
 				if (!bDeltaCheckpoint || ObjectInfo->bDirtyForReplay)
 				{
@@ -2072,7 +2100,7 @@ void UDemoNetDriver::SaveCheckpoint()
 		return;
 	}
 
-	UPackageMapClient* PackageMapClient = Cast<UPackageMapClient>(ClientConnections[0]->PackageMap);
+	UPackageMapClient* PackageMapClient = Cast<UPackageMapClient>(Connection->PackageMap);
 
 	PackageMapClient->SavePackageMapExportAckStatus(CheckpointSaveContext.CheckpointAckState);
 
@@ -2112,13 +2140,14 @@ static bool inline ShouldExecuteState(const FRepActorsCheckpointParams& Params,
 
 void UDemoNetDriver::TickCheckpoint()
 {
+	CSV_SCOPED_TIMING_STAT(Basic, DemoRecordCheckpointTime);
+	DECLARE_SCOPE_CYCLE_COUNTER(TEXT("SaveCheckpoint time"), STAT_ReplayCheckpointSaveTime, STATGROUP_Net);
+
 	if (CheckpointSaveContext.CheckpointSaveState == ECheckpointSaveState_Idle)
 	{
 		return;
 	}
 
-	DECLARE_SCOPE_CYCLE_COUNTER(TEXT("SaveCheckpoint time"), STAT_ReplayCheckpointSaveTime, STATGROUP_Net);
-
 	FArchive* CheckpointArchive = ReplayStreamer->GetCheckpointArchive();
 
 	if (!ensure(CheckpointArchive != nullptr))
@@ -2592,9 +2621,13 @@ public:
 	FVector		Location;
 };
 
-class FRepActorsParams : public FNoncopyable
+class FRepActorsParams
 {
 public:
+	FRepActorsParams(FRepActorsParams&&) = delete;
+	FRepActorsParams(const FRepActorsParams&) = delete;
+	FRepActorsParams& operator=(const FRepActorsParams&) = delete;
+	FRepActorsParams& operator=(FRepActorsParams&&) = delete;
 
 	FRepActorsParams(const bool bInUseAdaptiveNetFrequency, const bool bInDoFindActorChannel, const bool bInDoCheckDormancy,
 					const float InMinRecordHz, const float InMaxRecordHz, const float InServerTickTime,
@@ -2690,7 +2723,6 @@ void UDemoNetDriver::BuildSortedLevelPriorityOnLevels(const TArray<FDemoActorPri
 	OutLevelIntervals.Sort([](const FLevelnterval& A, const FLevelnterval& B) { return (B.Priority < A.Priority) || ((A.Priority == B.Priority) && (A.LevelIndex < B.LevelIndex)); });
 }
 
-
 void UDemoNetDriver::TickDemoRecordFrame(float DeltaSeconds)
 {
 	FArchive* FileAr = ReplayStreamer->GetStreamingArchive();
@@ -2771,7 +2803,7 @@ void UDemoNetDriver::TickDemoRecordFrame(float DeltaSeconds)
 		DECLARE_SCOPE_CYCLE_COUNTER(TEXT("Replay prioritize time"), STAT_ReplayPrioritizeTime, STATGROUP_Net);
 
 		const double ConsiderTimeLimit = RecordTimeLimit * RecordBuildConsiderAndPrioritizeTimeSlice;
-		auto HasConsiderTimeBeenExhausted = [ConsiderTimeLimit, RecordFrameStartTime, RecordTimeLimit]()
+		auto HasConsiderTimeBeenExhausted = [ConsiderTimeLimit, RecordFrameStartTime, RecordTimeLimit]() -> bool
 		{
 			return RecordTimeLimit > 0.f && (FPlatformTime::Seconds() - RecordFrameStartTime) > ConsiderTimeLimit;
 		};
@@ -2811,10 +2843,9 @@ void UDemoNetDriver::TickDemoRecordFrame(float DeltaSeconds)
 				for (UNetConnection* Connection : World->NetDriver->ClientConnections)
 				{
 					FReplayViewer ReplayViewer(Connection);
-
-					if (ReplayViewer.ViewTarget)
+					if (ReplayViewer.ViewTarget != nullptr)
 					{
-						ReplayViewers.Add(FReplayViewer(Connection));
+						ReplayViewers.Add(MoveTemp(ReplayViewer));
 					}
 				}
 			}
@@ -3002,22 +3033,27 @@ void UDemoNetDriver::TickDemoRecordFrame(float DeltaSeconds)
 		RecordTimeLimit
 	);
 
-	if (HasLevelStreamingFixes())
 	{
-		const FDemoActorPriority* Priorities = PrioritizedActors.GetData();
+		DECLARE_SCOPE_CYCLE_COUNTER(TEXT("Replay actor replication time"), STAT_ReplayReplicateActors, STATGROUP_Net);
 
-		// Split per level		
-		for (const FLevelnterval& Interval : LevelIntervals)
+		if (HasLevelStreamingFixes())
 		{
-			FScopedPacketManager PacketManager(*(UDemoNetConnection*)ClientConnections[0], Interval.LevelIndex);
-			bool bContinue = ReplicatePrioritizedActors(&Priorities[Interval.StartIndex], Interval.Count, Params);
-			if (!bContinue)
-				break;
+			const FDemoActorPriority* Priorities = PrioritizedActors.GetData();
+
+			// Split per level		
+			for (const FLevelnterval& Interval : LevelIntervals)
+			{
+				FScopedPacketManager PacketManager(*ClientConnection, Interval.LevelIndex);
+				if (!ReplicatePrioritizedActors(&Priorities[Interval.StartIndex], Interval.Count, Params))
+				{
+					break;
+				}
+			}
+		}
+		else
+		{
+			ReplicatePrioritizedActors(PrioritizedActors.GetData(), PrioritizedActors.Num(), Params);
 		}
-	}
-	else
-	{
-		ReplicatePrioritizedActors(PrioritizedActors.GetData(), PrioritizedActors.Num(), Params);
 	}
 
 	CSV_CUSTOM_STAT(Basic, DemoNumReplicatedActors, Params.NumActorsReplicated, ECsvCustomStatOp::Set);
@@ -3026,7 +3062,7 @@ void UDemoNetDriver::TickDemoRecordFrame(float DeltaSeconds)
 
 	WriteDemoFrameFromQueuedDemoPackets(*FileAr, ClientConnection->QueuedDemoPackets, DemoCurrentTime, EWriteDemoFrameFlags::None);
 
-	float ReplicatedPercent = NumPrioritizedActors != 0 ? (float)Params.NumActorsReplicated / (float)NumPrioritizedActors : 1.0f;
+	const float ReplicatedPercent = NumPrioritizedActors != 0 ? (float)Params.NumActorsReplicated / (float)NumPrioritizedActors : 1.0f;
 	AdjustConsiderTime(ReplicatedPercent);
 	LastReplayFrameFidelity = ReplicatedPercent;
 }
@@ -3041,10 +3077,12 @@ bool UDemoNetDriver::ReplicatePrioritizedActor(const FActorPriority& ActorPriori
 	const bool bDoFindActorChannel = Params.bDoFindActorChannel;
 	const bool bDoCheckDormancy = Params.bDoCheckDormancy;
 
+	UDemoNetConnection* Connection = CastChecked<UDemoNetConnection>(ClientConnections[0]);
+
 	// Deletion entry
 	if (ActorInfo == nullptr && DestructionInfo != nullptr)
 	{
-		UActorChannel* Channel = (UActorChannel*)ClientConnections[0]->CreateChannelByName(NAME_Actor, EChannelCreateFlags::OpenedLocally);
+		UActorChannel* Channel = (UActorChannel*)Connection->CreateChannelByName(NAME_Actor, EChannelCreateFlags::OpenedLocally);
 		if (Channel)
 		{
 			UE_LOG(LogDemo, Verbose, TEXT("TickDemoRecord creating destroy channel for NetGUID <%s,%s> Priority: %d"), *DestructionInfo->NetGUID.ToString(), *DestructionInfo->PathName, ActorPriority.Priority);
@@ -3053,7 +3091,7 @@ bool UDemoNetDriver::ReplicatePrioritizedActor(const FActorPriority& ActorPriori
 			Channel->SetChannelActorForDestroy(DestructionInfo);
 
 			// Remove from connection's to-be-destroyed list (close bunch is reliable, so it will make it there)
-			ClientConnections[0]->GetDestroyedStartupOrDormantActorGUIDs().Remove(DestructionInfo->NetGUID);
+			Connection->GetDestroyedStartupOrDormantActorGUIDs().Remove(DestructionInfo->NetGUID);
 		}
 	}
 	else if (ActorInfo != nullptr && DestructionInfo == nullptr)
@@ -3062,7 +3100,7 @@ bool UDemoNetDriver::ReplicatePrioritizedActor(const FActorPriority& ActorPriori
 		
 		if (bDoCheckDormancy)
 		{
-			UActorChannel* Channel = (bDoFindActorChannel ? ClientConnections[0]->FindActorChannelRef(Actor) : ActorPriority.Channel);
+			UActorChannel* Channel = (bDoFindActorChannel ? Connection->FindActorChannelRef(Actor) : ActorPriority.Channel);
 			if (Channel && ShouldActorGoDormantForDemo(Actor, Channel))
 			{
 				// Either shouldn't go dormant, or is already dormant
@@ -3114,7 +3152,7 @@ bool UDemoNetDriver::ReplicatePrioritizedActor(const FActorPriority& ActorPriori
 
 		Actor->CallPreReplication(this);
 
-		const bool bDidReplicateActor = DemoReplicateActor(Actor, ClientConnections[0], false);
+		const bool bDidReplicateActor = DemoReplicateActor(Actor, Connection, false);
 
 		const bool bUpdatedExternalData = UpdateExternalDataForActor(Actor);
 
@@ -3691,7 +3729,7 @@ bool UDemoNetDriver::ConditionallyProcessPlaybackPackets()
 
 	if (CurPacket.LevelIndex != CurrentLevelIndex)
 	{
-		GetWorld()->GetGameInstance()->OnSeamlessTravelDuringReplay();
+		World->GetGameInstance()->OnSeamlessTravelDuringReplay();
 		CurrentLevelIndex = CurPacket.LevelIndex;
 		ProcessSeamlessTravel(CurrentLevelIndex);
 		return false;
@@ -3746,6 +3784,8 @@ bool UDemoNetDriver::ProcessPacket(const uint8* Data, int32 Count)
 
 void UDemoNetDriver::WriteDemoFrameFromQueuedDemoPackets(FArchive& Ar, TArray<FQueuedDemoPacket>& QueuedPackets, float FrameTime, EWriteDemoFrameFlags Flags)
 {
+	DECLARE_SCOPE_CYCLE_COUNTER(TEXT("Replay write frame time"), STAT_ReplayWriteDemoFrame, STATGROUP_Net);
+
 	Ar << CurrentLevelIndex;
 	
 	// Save total absolute demo time in seconds
@@ -3803,7 +3843,7 @@ void UDemoNetDriver::WriteDemoFrameFromQueuedDemoPackets(FArchive& Ar, TArray<FQ
 		if (!EnumHasAnyFlags(Flags, EWriteDemoFrameFlags::SkipGameSpecific))
 		{
 			FDemoFrameDataMap Data;
-			FNetworkReplayDelegates::OnWriteGameSpecificFrameData.Broadcast(GetWorld(), FrameTime, Data);
+			FNetworkReplayDelegates::OnWriteGameSpecificFrameData.Broadcast(World, FrameTime, Data);
 
 			Ar << Data;
 		}
@@ -4291,7 +4331,7 @@ APlayerController* UDemoNetDriver::CreateDemoPlayerController(UNetConnection* Co
 
 	// Get the replay spectator controller class from the default game mode object,
 	// since the game mode instance isn't replicated to clients of live games.
-	AGameStateBase* GameState = GetWorld() != nullptr ? GetWorld()->GetGameState() : nullptr;
+	AGameStateBase* GameState = World != nullptr ? World->GetGameState() : nullptr;
 	TSubclassOf<AGameModeBase> DefaultGameModeClass = GameState != nullptr ? GameState->GameModeClass : nullptr;
 
 	// If we don't have a game mode class from the world, try to get it from the URL option.
@@ -4334,7 +4374,7 @@ APlayerController* UDemoNetDriver::CreateDemoPlayerController(UNetConnection* Co
 
 	// If the controller doesn't have a player state, we are probably recording on a client.
 	// Spawn one manually.
-	if (NewDemoController->PlayerState == nullptr && GetWorld() != nullptr && GetWorld()->IsRecordingClientReplay())
+	if (NewDemoController->PlayerState == nullptr && World != nullptr && World->IsRecordingClientReplay())
 	{
 		NewDemoController->InitPlayerState();
 	}
@@ -4405,8 +4445,8 @@ bool UDemoNetDriver::SpawnSplitscreenViewer(ULocalPlayer* NewPlayer, UWorld* InW
 	NewSplitscreenController->ChangeState(NAME_Spectating);
 	NewPlayer->CurrentNetSpeed = 0;
 
-	// Link the local player to the player controller as the localplayer has been marked as active
-	// but without a playercontroller, the player will never be considered "ready" by other systems.
+	// Link the local player to the player controller as the local player has been marked as active
+	// but without a PlayerController, the player will never be considered "ready" by other systems.
 	NewPlayer->PlayerController = NewSplitscreenController;
 
 	// This would typically be set via SetPlayer, but we need to call SetPlayer with the LocalPlayer
@@ -4615,7 +4655,7 @@ void UDemoNetDriver::RespawnNecessaryNetStartupActors(TArray<AActor*>& SpawnedAc
 
 		const FTransform SpawnTransform = FTransform(RollbackActor.Rotation, RollbackActor.Location, RollbackActor.Scale3D);
 
-		AActor* Actor = GetWorld()->SpawnActorAbsolute(RollbackActor.Archetype->GetClass(), SpawnTransform, SpawnInfo);
+		AActor* Actor = World->SpawnActorAbsolute(RollbackActor.Archetype->GetClass(), SpawnTransform, SpawnInfo);
 		if (Actor)
 		{
 			if (!ensure( Actor->GetFullName() == It.Key()))
@@ -4719,7 +4759,6 @@ void UDemoNetDriver::PrepFastForwardLevels()
 	check(!bIsLoadingCheckpoint);
 
 	// Do a quick pass to double check everything is still valid, and that we have data for the levels.
-	UWorld* LocalWorld = GetWorld();
 	for (TWeakObjectPtr<UObject> &WeakLevel : NewStreamingLevelsThisFrame)
 	{
 		// For playback, we should only ever see ULevels in this list.
@@ -5004,7 +5043,6 @@ bool UDemoNetDriver::FastForwardLevels(const FGotoResult& GotoResult)
 	// If we've gotten this far, it means we should have something to process.
 	check(ReadPacketsHelper.Packets.Num() > 0);
 
-	UWorld* LocalWorld = GetWorld();
 	for (ULevel* Level : LevelsPendingFastForward)
 	{
 		// Track the appropriate level, and mark it as ready.
@@ -5024,13 +5062,13 @@ bool UDemoNetDriver::FastForwardLevels(const FGotoResult& GotoResult)
 				// Put this actor on the rollback list so we can undelete it during future scrubbing,
 				// then delete it.
 				QueueNetStartupActorForRollbackViaDeletion(Actor);
-				LocalWorld->DestroyActor(Actor, true);
+				World->DestroyActor(Actor, true);
 			}
 			else
 			{
 				if (RollbackNetStartupActors.Contains(Actor->GetFullName()))
 				{
-					LocalWorld->DestroyActor(Actor, true);
+					World->DestroyActor(Actor, true);
 				}
 				else
 				{
@@ -5196,13 +5234,15 @@ bool UDemoNetDriver::LoadCheckpoint(const FGotoResult& GotoResult)
 		}
 	}
 
+	check(World);
+
 	if (LevelForCheckpoint != CurrentLevelIndex)
 	{
-		GetWorld()->GetGameInstance()->OnSeamlessTravelDuringReplay();
+		World->GetGameInstance()->OnSeamlessTravelDuringReplay();
 
-		for (FActorIterator It(GetWorld()); It; ++It)
+		for (FActorIterator It(World); It; ++It)
 		{
-			GetWorld()->DestroyActor(*It, true);
+			World->DestroyActor(*It, true);
 		}
 
 		// Clean package map to prepare to restore it to the checkpoint state
@@ -5260,7 +5300,7 @@ bool UDemoNetDriver::LoadCheckpoint(const FGotoResult& GotoResult)
 
 	PauseChannels(false);
 
-	FNetworkReplayDelegates::OnPreScrub.Broadcast(GetWorld());
+	FNetworkReplayDelegates::OnPreScrub.Broadcast(World);
 
 	bIsLoadingCheckpoint = true;
 
@@ -5301,7 +5341,7 @@ bool UDemoNetDriver::LoadCheckpoint(const FGotoResult& GotoResult)
 
 
 	// Destroy all non startup actors. They will get restored with the checkpoint
-	for (FActorIterator It(GetWorld()); It; ++It)
+	for (FActorIterator It(World); It; ++It)
 	{
 		// If there are any existing actors that are bAlwaysRelevant, don't queue their bunches.
 		// Actors that do queue their bunches might not appear immediately after the checkpoint is loaded,
@@ -5350,11 +5390,11 @@ bool UDemoNetDriver::LoadCheckpoint(const FGotoResult& GotoResult)
 			continue;
 		}
 
-		GetWorld()->DestroyActor(*It, true);
+		World->DestroyActor(*It, true);
 	}
 
 	// Destroy all particle FX attached to the WorldSettings (the WorldSettings actor persists but the particle FX spawned at runtime shouldn't)
-	GetWorld()->HandleTimelineScrubbed();
+	World->HandleTimelineScrubbed();
 
 	// Remove references to our KeepAlive actors so that cleaning up the channels won't destroy them.
 	for (int32 i = ServerConnection->OpenChannels.Num() - 1; i >= 0; i--)
@@ -5383,7 +5423,7 @@ bool UDemoNetDriver::LoadCheckpoint(const FGotoResult& GotoResult)
 			UActorChannel* ActorChannel = Cast<UActorChannel>(OpenChannel);
 			if (ActorChannel != nullptr && ActorChannel->GetActor() != nullptr && !ActorChannel->GetActor()->IsNetStartupActor())
 			{
-				GetWorld()->DestroyActor(ActorChannel->GetActor(), true);
+				World->DestroyActor(ActorChannel->GetActor(), true);
 			}
 		}
 	}
@@ -5394,11 +5434,11 @@ bool UDemoNetDriver::LoadCheckpoint(const FGotoResult& GotoResult)
 	PlaybackFrames.Empty();
 
 	// Destroy startup actors that need to rollback via being destroyed and re-created
-	for (FActorIterator It(GetWorld()); It; ++It)
+	for (FActorIterator It(World); It; ++It)
 	{
 		if (RollbackNetStartupActors.Contains(It->GetFullName()))
 		{
-			GetWorld()->DestroyActor(*It, true);
+			World->DestroyActor(*It, true);
 		}
 	}
 
@@ -6057,7 +6097,7 @@ TSharedPtr<FInternetAddr> FInternetAddrDemo::DemoInternetAddr = MakeShareable(ne
 UDemoNetConnection::UDemoNetConnection( const FObjectInitializer& ObjectInitializer ) : Super( ObjectInitializer )
 {
 	MaxPacket = MAX_DEMO_READ_WRITE_BUFFER;
-	InternalAck = true;
+	SetInternalAck(true);
 }
 
 void UDemoNetConnection::InitConnection( UNetDriver* InDriver, EConnectionState InState, const FURL& InURL, int32 InConnectionSpeed, int32 InMaxPacket)
@@ -6066,7 +6106,7 @@ void UDemoNetConnection::InitConnection( UNetDriver* InDriver, EConnectionState
 	Super::InitConnection( InDriver, InState, InURL, InConnectionSpeed );
 
 	MaxPacket = (InMaxPacket == 0 || InMaxPacket > MAX_DEMO_READ_WRITE_BUFFER) ? MAX_DEMO_READ_WRITE_BUFFER : InMaxPacket;
-	InternalAck = true;
+	SetInternalAck(true);
 
 	InitSendBuffer();
 
@@ -6099,11 +6139,11 @@ void UDemoNetConnection::LowLevelSend(void* Data, int32 CountBits, FOutPacketTra
 	if (ResendAllDataState != EResendAllDataState::None)
 	{
 		// This path is only active for a checkpoint saving out, we need to queue in separate list
-		new(QueuedCheckpointPackets)FQueuedDemoPacket((uint8*)Data, CountBits, Traits);
+		QueuedCheckpointPackets.Emplace((uint8*)Data, CountBits, Traits);
 		return;
 	}
 
-	new(QueuedDemoPackets)FQueuedDemoPacket((uint8*)Data, CountBits, Traits);
+	QueuedDemoPackets.Emplace((uint8*)Data, CountBits, Traits);
 }
 
 void UDemoNetConnection::TrackSendForProfiler(const void* Data, int32 NumBytes)
@@ -6438,7 +6478,7 @@ void UDemoNetDriver::NotifyActorDestroyed(AActor* Actor, bool IsSeamlessTravel)
 
 void UDemoNetDriver::CleanupOutstandingRewindActors()
 {
-	if (UWorld* LocalWorld = GetWorld())
+	if (World)
 	{
 		for (const FNetworkGUID& NetGUID : TrackedRewindActorsByGUID)
 		{
@@ -6447,7 +6487,7 @@ void UDemoNetDriver::CleanupOutstandingRewindActors()
 				if (AActor* Actor = Cast<AActor>(CacheObject->Object))
 				{
 					// Destroy the actor before removing entries from the GuidCache so its entries are still valid in NotifyActorDestroyed.
-					LocalWorld->DestroyActor(Actor);
+					World->DestroyActor(Actor);
 
 					ensureMsgf(GuidCache->NetGUIDLookup.Remove(CacheObject->Object) > 0, TEXT("CleanupOutstandingRewindActors: No entry found for %d in NetGUIDLookup"), NetGUID.Value);
 					GuidCache->ObjectLookup.Remove(NetGUID);
@@ -6519,6 +6559,11 @@ void UDemoNetDriver::NotifyActorLevelUnloaded(AActor* Actor)
 
 void UDemoNetDriver::QueueNetStartupActorForRollbackViaDeletion(AActor* Actor)
 {
+	if (!IsPlaying())
+	{
+		return;		// We should only be doing this at runtime while playing a replay
+	}
+
 	check(Actor != nullptr);
 
 	if (!Actor->IsNetStartupActor())
@@ -6526,9 +6571,10 @@ void UDemoNetDriver::QueueNetStartupActorForRollbackViaDeletion(AActor* Actor)
 		return;		// We only want startup actors
 	}
 
-	if (!IsPlaying())
+	if (Actor->bReplayRewindable)
 	{
-		return;		// We should only be doing this at runtime while playing a replay
+		UE_LOG(LogDemo, Warning, TEXT("Attempted to queue a Replay Rewindable Actor for Rollback Via Deletion. Replay may have artifacts (%s)"), *GetFullNameSafe(Actor));
+		return;
 	}
 
 	const FString ActorFullName = Actor->GetFullName();
@@ -6537,12 +6583,6 @@ void UDemoNetDriver::QueueNetStartupActorForRollbackViaDeletion(AActor* Actor)
 		return;		// This actor is already queued up
 	}
 
-	if (Actor->bReplayRewindable)
-	{
-		UE_LOG(LogDemo, Warning, TEXT("Attempted to queue a Replay Rewindable Actor for Rollback Via Deletion. Replay may have artifacts (%s)"), *ActorFullName);
-		return;
-	}
-
 	FRollbackNetStartupActorInfo& RollbackActor = RollbackNetStartupActors.Add(ActorFullName);
 
 	RollbackActor.Name		= Actor->GetFName();
@@ -6638,7 +6678,7 @@ void UDemoNetDriver::NotifyDemoPlaybackFailure(EDemoPlayFailure::Type FailureTyp
 
 	if (bIsPlaying)
 	{
-		if (UWorld* LocalWorld = GetWorld())
+		if (World)
 		{
 			if (UGameInstance* GameInstance = World->GetGameInstance())
 			{
diff --git a/Engine/Source/Runtime/Engine/Private/NetConnection.cpp b/Engine/Source/Runtime/Engine/Private/NetConnection.cpp
--- a/Engine/Source/Runtime/Engine/Private/NetConnection.cpp
+++ b/Engine/Source/Runtime/Engine/Private/NetConnection.cpp
@@ -174,6 +174,7 @@ UNetConnection::UNetConnection(const FObjectInitializer& ObjectInitializer)
 ,   OwningActor			( nullptr )
 ,	MaxPacket			( 0 )
 ,	InternalAck			( false )
+,	bInternalAck		( false )
 ,	RemoteAddr			( nullptr )
 ,	MaxPacketHandlerBits ( 0 )
 ,	State				( USOCK_Invalid )
@@ -1307,7 +1308,7 @@ void UNetConnection::ReceivedRawPacket( void* InData, int32 Count )
 
 void UNetConnection::PostTickDispatch()
 {
-	if (!InternalAck)
+	if (!IsInternalAck())
 	{
 #if DO_ENABLE_NET_TEST
 		ReinjectDelayedPackets();
@@ -1399,7 +1400,7 @@ void UNetConnection::FlushNet(bool bIgnoreSimulation)
 	TimeSensitive = 0;
 
 	// If there is any pending data to send, send it.
-	if (SendBuffer.GetNumBits() || HasDirtyAcks || ( Driver->GetElapsedTime() - LastSendTime > Driver->KeepAliveTime && !InternalAck && State != USOCK_Closed))
+	if (SendBuffer.GetNumBits() || HasDirtyAcks || ( Driver->GetElapsedTime() - LastSendTime > Driver->KeepAliveTime && !IsInternalAck() && State != USOCK_Closed))
 	{
 		// Due to the PacketHandler handshake code, servers must never send the client data,
 		// before first receiving a client control packet (which is taken as an indication of a complete handshake).
@@ -1436,7 +1437,7 @@ void UNetConnection::FlushNet(bool bIgnoreSimulation)
 		SendBuffer.WriteBit(1);
 
 		// Refresh outgoing header with latest data
-		if ( !InternalAck )
+		if ( !IsInternalAck() )
 		{
 			// if we update ack, we also update received ack associated with outgoing seq
 			// so we know how many ack bits we need to write (which is updated in received packet)
@@ -1476,7 +1477,7 @@ void UNetConnection::FlushNet(bool bIgnoreSimulation)
 
 		// if the connection is closing/being destroyed/etc we need to send immediately regardless of settings
 		// because we won't be around to send it delayed
-		if (State != USOCK_Closed && !IsGarbageCollecting() && !bIgnoreSimulation && !InternalAck)
+		if (State != USOCK_Closed && !IsGarbageCollecting() && !bIgnoreSimulation && !IsInternalAck())
 		{
 			bWasPacketEmulated = CheckOutgoingPacketEmulation(Traits);
 		}
@@ -1510,7 +1511,7 @@ void UNetConnection::FlushNet(bool bIgnoreSimulation)
 		OutBytesPerSecondHistory[Index]	= FMath::Min(OutBytesPerSecond / 1024, 255);
 
 		// Increase outgoing sequence number
-		if (!InternalAck)
+		if (!IsInternalAck())
 		{
 			PacketNotify.CommitAndIncrementOutSeq();
 		}
@@ -2025,7 +2026,7 @@ void UNetConnection::ReceivedPacket( FBitReader& Reader, bool bIsReinjectedPacke
 	}
 
 #if DO_ENABLE_NET_TEST
-	if (!InternalAck && !bIsReinjectingDelayedPackets)
+	if (!IsInternalAck() && !bIsReinjectingDelayedPackets)
 	{
 		if (PacketSimulationSettings.PktIncomingLoss)
 		{
@@ -2071,7 +2072,7 @@ void UNetConnection::ReceivedPacket( FBitReader& Reader, bool bIsReinjectedPacke
 		LastReceiveRealtime = CurrentReceiveTimeInS;
 	}
 
-	if (InternalAck)
+	if (IsInternalAck())
 	{
 		++InPacketId;
 	}	
@@ -2259,7 +2260,7 @@ void UNetConnection::ReceivedPacket( FBitReader& Reader, bool bIsReinjectedPacke
 	while( !Reader.AtEnd() && State!=USOCK_Closed )
 	{
 		// For demo backwards compatibility, old replays still have this bit
-		if (InternalAck && EngineNetworkProtocolVersion < EEngineNetworkVersionHistory::HISTORY_ACKS_INCLUDED_IN_HEADER)
+		if (IsInternalAck() && EngineNetworkProtocolVersion < EEngineNetworkVersionHistory::HISTORY_ACKS_INCLUDED_IN_HEADER)
 		{
 			const bool IsAckDummy = Reader.ReadBit() == 1u;
 		}
@@ -2318,7 +2319,7 @@ PRAGMA_ENABLE_DEPRECATION_WARNINGS
 
 			if ( Bunch.bReliable )
 			{
-				if ( InternalAck )
+				if ( IsInternalAck() )
 				{
 					// We can derive the sequence for 100% reliable connections
 					Bunch.ChSequence = InReliable[Bunch.ChIndex] + 1;
@@ -2439,7 +2440,7 @@ PRAGMA_ENABLE_DEPRECATION_WARNINGS
 			{
 				Driver->NetGUIDInBytes += (BunchDataBits + (HeaderPos - IncomingStartPos)) >> 3;
 
-				if ( InternalAck )
+				if ( IsInternalAck() )
 				{
 					// NOTE - For replays, we do this even earlier, to try and load this as soon as possible, in case there is an issue creating the channel
 					// If a replay fails to create a channel, we want to salvage as much as possible
@@ -2498,7 +2499,7 @@ PRAGMA_ENABLE_DEPRECATION_WARNINGS
 
 			// We're on a 100% reliable connection and we are rolling back some data.
 			// In that case, we can generally ignore these bunches.
-			if (InternalAck && Channel && bIgnoreAlreadyOpenedChannels)
+			if (IsInternalAck() && Channel && bIgnoreAlreadyOpenedChannels)
 			{
 				// This was an open bunch for a channel that's already opened.
 				// We can ignore future bunches from this channel.
@@ -2547,7 +2548,7 @@ PRAGMA_ENABLE_DEPRECATION_WARNINGS
 			if ( Bunch.bReliable && Bunch.ChSequence <= InReliable[Bunch.ChIndex] )
 			{
 				UE_LOG( LogNetTraffic, Log, TEXT( "UNetConnection::ReceivedPacket: Received outdated bunch (Channel %d Current Sequence %i)" ), Bunch.ChIndex, InReliable[Bunch.ChIndex] );
-				check( !InternalAck );		// Should be impossible with 100% reliable connections
+				check( !IsInternalAck() );		// Should be impossible with 100% reliable connections
 				continue;
 			}
 			
@@ -2562,7 +2563,7 @@ PRAGMA_ENABLE_DEPRECATION_WARNINGS
 				const bool ValidUnreliableOpen = Bunch.bOpen && (Bunch.bClose || Bunch.bPartial);
 				if (!ValidUnreliableOpen)
 				{
-					if ( InternalAck )
+					if ( IsInternalAck() )
 					{
 						// Should be impossible with 100% reliable connections
 						UE_LOG( LogNetTraffic, Error, TEXT( "      Received unreliable bunch before open with reliable connection (Channel %d Current Sequence %i)" ), Bunch.ChIndex, InReliable[Bunch.ChIndex] );
@@ -2607,7 +2608,7 @@ PRAGMA_ENABLE_DEPRECATION_WARNINGS
 				}
 
 				// peek for guid
-				if (InternalAck && bIgnoreActorBunches)
+				if (IsInternalAck() && bIgnoreActorBunches)
 				{
 					if (Bunch.bOpen && (!Bunch.bPartial || Bunch.bPartialInitial) && (Bunch.ChName == NAME_Actor))
 					{
@@ -2704,7 +2705,7 @@ PRAGMA_ENABLE_DEPRECATION_WARNINGS
 		LastGoodPacketRealtime = FPlatformTime::Seconds();
 	}
 
-	if( !InternalAck )
+	if( !IsInternalAck() )
 	{
 		// We always call AckSequence even if we are explicitly rejecting the packet as this updates the expected InSeq used to drive future acks.
 		if ( bSkipAck )
@@ -2756,14 +2757,14 @@ PRAGMA_ENABLE_DEPRECATION_WARNINGS
 
 void UNetConnection::SetIgnoreAlreadyOpenedChannels(bool bInIgnoreAlreadyOpenedChannels)
 {
-	check(InternalAck);
+	check(IsInternalAck());
 	bIgnoreAlreadyOpenedChannels = bInIgnoreAlreadyOpenedChannels;
 	IgnoringChannels.Reset();
 }
 
 void UNetConnection::SetIgnoreActorBunches(bool bInIgnoreActorBunches, TSet<FNetworkGUID>&& InIgnoredBunchGuids)
 {
-	check(InternalAck);
+	check(IsInternalAck());
 	bIgnoreActorBunches = bInIgnoreActorBunches;
 
 	IgnoredBunchChannels.Empty();
@@ -2798,7 +2799,7 @@ void UNetConnection::PrepareWriteBitsToSendBuffer(const int32 SizeInBits, const
 	}
 
 	// If this is the start of the queue, make sure to add the packet id
-	if ( SendBuffer.GetNumBits() == 0 && !InternalAck )
+	if ( SendBuffer.GetNumBits() == 0 && !IsInternalAck() )
 	{
 #if UE_NET_TRACE_ENABLED
 		// If tracing is enabled setup the NetTraceCollector for outgoing data
@@ -2953,7 +2954,7 @@ int32 UNetConnection::SendRawBunch(FOutBunch& Bunch, bool InAllowMerge, const FN
 	SendBunchHeader.WriteBit( Bunch.bHasMustBeMappedGUIDs );
 	SendBunchHeader.WriteBit( Bunch.bPartial );
 
-	if ( Bunch.bReliable && !InternalAck )
+	if ( Bunch.bReliable && !IsInternalAck() )
 	{
 		SendBunchHeader.WriteIntWrapped(Bunch.ChSequence, MAX_CHSEQUENCE);
 	}
@@ -3212,7 +3213,7 @@ void UNetConnection::Tick()
 	const float MaxNetTickRateFloat = MaxNetTickRate > 0 ? float(MaxNetTickRate) : MAX_flt;
 	const float DesiredTickRate = FMath::Clamp(EngineTickRate, 0.0f, MaxNetTickRateFloat);
 	// Apply net tick rate limiting if the desired net tick rate is strictly less than the engine tick rate.
-	if (!InternalAck && MaxNetTickRateFloat < EngineTickRate && DesiredTickRate > 0.0f)
+	if (!IsInternalAck() && MaxNetTickRateFloat < EngineTickRate && DesiredTickRate > 0.0f)
 	{
 		const float MinNetFrameTime = 1.0f/DesiredTickRate;
 		if (FrameTime < MinNetFrameTime)
@@ -3232,7 +3233,7 @@ void UNetConnection::Tick()
 	}
 
 	// Pretend everything was acked, for 100% reliable connections or demo recording.
-	if( InternalAck )
+	if(IsInternalAck())
 	{
 		const bool bIsServer = Driver->IsServer();
 		OutAckPacketId = OutPacketId;
@@ -3247,7 +3248,7 @@ void UNetConnection::Tick()
 			UChannel* Channel = Channels[ChannelIndex];
 			if (Channel)
 			{
-				for(FOutBunch* OutBunch=Channel->OutRec; OutBunch; OutBunch=OutBunch->Next)
+				for (FOutBunch* OutBunch=Channel->OutRec; OutBunch; OutBunch=OutBunch->Next)
 				{
 					OutBunch->ReceivedAck = 1;
 				}
@@ -4169,7 +4170,7 @@ uint32 UNetConnection::GetParentConnectionId() const
 
 USimulatedClientNetConnection::USimulatedClientNetConnection( const FObjectInitializer& ObjectInitializer ) : Super( ObjectInitializer )
 {
-	InternalAck = true;
+	SetInternalAck(true);
 }
 
 void USimulatedClientNetConnection::HandleClientPlayer( class APlayerController* PC, class UNetConnection* NetConnection )
diff --git a/Engine/Source/Runtime/Engine/Private/NetworkDriver.cpp b/Engine/Source/Runtime/Engine/Private/NetworkDriver.cpp
--- a/Engine/Source/Runtime/Engine/Private/NetworkDriver.cpp
+++ b/Engine/Source/Runtime/Engine/Private/NetworkDriver.cpp
@@ -575,7 +575,7 @@ void UNetDriver::TickFlush(float DeltaSeconds)
 	}
 	FSimpleScopeSecondsCounter ScopedTimer(GTickFlushGameDriverTimeSeconds, bEnableTimer);
 
-	if ( IsServer() && ClientConnections.Num() > 0 && ClientConnections[0]->InternalAck == false )
+	if ( IsServer() && ClientConnections.Num() > 0 && ClientConnections[0]->IsInternalAck() == false )
 	{
 		// Update all clients.
 #if WITH_SERVER_CODE
@@ -2238,7 +2238,7 @@ void UNetDriver::ProcessRemoteFunctionForChannel(UActorChannel* Ch, const FClass
 		}
 	}
 
-	if ( Connection->InternalAck )
+	if ( Connection->IsInternalAck() )
 	{
 		Connection->FlushNet();
 	}
diff --git a/Engine/Source/Runtime/Engine/Private/PackageMapClient.cpp b/Engine/Source/Runtime/Engine/Private/PackageMapClient.cpp
--- a/Engine/Source/Runtime/Engine/Private/PackageMapClient.cpp
+++ b/Engine/Source/Runtime/Engine/Private/PackageMapClient.cpp
@@ -223,7 +223,7 @@ bool UPackageMapClient::SerializeObject( FArchive& Ar, UClass* Class, UObject*&
 			}
 
 #if 0		// Enable this code to force any actor with missing/broken content to not load in replays
-			if (NetGUID.IsValid() && Connection->InternalAck && GuidCache->IsGUIDBroken(NetGUID, true))
+			if (NetGUID.IsValid() && Connection->IsInternalAck() && GuidCache->IsGUIDBroken(NetGUID, true))
 			{
 				Ar.SetError();
 				UE_LOG(LogNetPackageMap, Warning, TEXT("UPackageMapClient::SerializeObject: InternalAck GUID broken."));
@@ -261,7 +261,7 @@ bool UPackageMapClient::SerializeObject( FArchive& Ar, UClass* Class, UObject*&
 				}
 			}
 
-			UE_CLOG(!bSuppressLogs, LogNetPackageMap, Log, TEXT("UPackageMapClient::SerializeObject Serialized Object %s as <%s>"), Object ? *Object->GetPathName() : TEXT("NULL"), *NetGUID.ToString());
+			UE_LOG(LogNetPackageMap, VeryVerbose, TEXT("UPackageMapClient::SerializeObject Serialized Object %s as <%s>"), Object ? *Object->GetPathName() : TEXT("NULL"), *NetGUID.ToString());
 		}
 		
 #if !(UE_BUILD_SHIPPING || UE_BUILD_TEST)
@@ -894,7 +894,7 @@ FNetworkGUID UPackageMapClient::InternalLoadObject( FArchive & Ar, UObject *& Ob
 	{
 		Object = GetObjectFromNetGUID( NetGUID, GuidCache->IsExportingNetGUIDBunch );
 
-		UE_CLOG( !bSuppressLogs, LogNetPackageMap, Log, TEXT( "InternalLoadObject loaded %s from NetGUID <%s>" ), Object ? *Object->GetFullName() : TEXT( "NULL" ), *NetGUID.ToString() );
+		UE_LOG(LogNetPackageMap, VeryVerbose, TEXT( "InternalLoadObject loaded %s from NetGUID <%s>" ), Object ? *Object->GetFullName() : TEXT( "NULL" ), *NetGUID.ToString() );
 	}
 
 	// ----------------	
@@ -1126,7 +1126,7 @@ bool UPackageMapClient::ExportNetGUID( FNetworkGUID NetGUID, UObject* Object, FS
 	check( !NetGUID.IsDefault() );
 	check( Object == NULL || ShouldSendFullPath( Object, NetGUID ) );
 
-	if (Connection->InternalAck)
+	if (Connection->IsInternalAck())
 	{
 		return ExportNetGUIDForReplay(NetGUID, Object, PathName, ObjOuter);
 	}
@@ -1346,7 +1346,7 @@ void UPackageMapClient::AddNetFieldExportGroup( const FString& PathName, TShared
 
 void UPackageMapClient::TrackNetFieldExport( FNetFieldExportGroup* NetFieldExportGroup, const int32 NetFieldExportHandle )
 {
-	check(Connection->InternalAck);
+	check(Connection->IsInternalAck());
 	check(NetFieldExportGroup);
 
 	checkf(NetFieldExportGroup->NetFieldExports.IsValidIndex(NetFieldExportHandle),
@@ -1427,7 +1427,7 @@ void UPackageMapClient::SerializeNetFieldExportGroupMap( FArchive& Ar, bool bCle
 
 void UPackageMapClient::AppendExportData(FArchive& Archive)
 {
-	check(Connection->InternalAck);
+	check(Connection->IsInternalAck());
 
 	AppendNetFieldExports(Archive);
 	AppendNetExportGUIDs(Archive);
@@ -1435,7 +1435,7 @@ void UPackageMapClient::AppendExportData(FArchive& Archive)
 
 void UPackageMapClient::ReceiveExportData(FArchive& Archive)
 {
-	check(Connection->InternalAck);
+	check(Connection->IsInternalAck());
 
 	ReceiveNetFieldExports(Archive);
 	ReceiveNetExportGUIDs(Archive);
@@ -1489,7 +1489,7 @@ void UPackageMapClient::AppendNetFieldExports(FArchive& Archive)
 
 void UPackageMapClient::AppendNetFieldExportsInternal(FArchive& Archive, const TSet<uint64>& InNetFieldExports, EAppendNetExportFlags Flags)
 {
-	check(Connection->InternalAck);
+	check(Connection->IsInternalAck());
 
 	uint32 NetFieldCount = InNetFieldExports.Num();
 	Archive.SerializeIntPacked(NetFieldCount);
@@ -1547,7 +1547,7 @@ void UPackageMapClient::AppendNetFieldExportsInternal(FArchive& Archive, const T
 
 void UPackageMapClient::ReceiveNetFieldExportsCompat(FInBunch &InBunch)
 {
-	if(!Connection->InternalAck)
+	if (!Connection->IsInternalAck())
 	{
 		UE_LOG(LogNetPackageMap, Error, TEXT("ReceiveNetFieldExportsCompat: connection is not a replay connection."));
 		InBunch.SetError();
@@ -1642,7 +1642,7 @@ void UPackageMapClient::ReceiveNetFieldExports(FArchive& Archive)
 {
 	DECLARE_SCOPE_CYCLE_COUNTER(TEXT("ReceiveNetFieldExports time"), STAT_ReceiveNetFieldExportsTime, STATGROUP_Net);
 
-	check(Connection->InternalAck);
+	check(Connection->IsInternalAck());
 
 	// Read number of net field exports
 	uint32 NumLayoutCmdExports = 0;
@@ -1716,7 +1716,7 @@ void UPackageMapClient::ReceiveNetFieldExports(FArchive& Archive)
 
 void UPackageMapClient::AppendNetExportGUIDs(FArchive& Archive)
 {
-	check(Connection->InternalAck);
+	check(Connection->IsInternalAck());
 
 	uint32 NumGUIDs = ExportGUIDArchives.Num();
 	Archive.SerializeIntPacked(NumGUIDs);
@@ -1733,7 +1733,7 @@ void UPackageMapClient::ReceiveNetExportGUIDs(FArchive& Archive)
 {
 	DECLARE_SCOPE_CYCLE_COUNTER(TEXT("ReceiveNetExportGUIDs time"), STAT_ReceiveNetExportGUIDsTime, STATGROUP_Net);
 
-	check(Connection->InternalAck);
+	check(Connection->IsInternalAck());
 	TGuardValue<bool> IsExportingGuard(GuidCache->IsExportingNetGUIDBunch, true);
 
 	uint32 NumGUIDs = 0;
@@ -1769,7 +1769,7 @@ void UPackageMapClient::ReceiveNetExportGUIDs(FArchive& Archive)
 
 void UPackageMapClient::AppendExportBunches(TArray<FOutBunch *>& OutgoingBunches)
 {
-	check(!Connection->InternalAck);
+	check(!Connection->IsInternalAck());
 	check(NetFieldExports.Num() == 0);
 
 	// Finish current in progress bunch if necessary
@@ -1868,7 +1868,7 @@ void UPackageMapClient::NotifyBunchCommit( const int32 OutPacketId, const FOutBu
 		// (GUID information doesn't change, so updating to the latest expected sequence is unnecessary)
 		if ( ExpectedPacketIdRef == GUID_PACKET_NOT_ACKED )
 		{
-			if ( Connection->InternalAck )
+			if ( Connection->IsInternalAck() )
 			{
 				// Auto ack now if the connection is 100% reliable
 				ExpectedPacketIdRef = GUID_PACKET_ACKED;
@@ -2550,6 +2550,9 @@ FNetworkGUID FNetGUIDCache::GetNetGUID(const UObject* Object) const
 	return NetGUID;
 }
 
+#define COMPOSE_NET_GUID( Index, IsStatic )	( ( ( Index ) << 1 ) | ( IsStatic ) )
+#define ALLOC_NEW_NET_GUID( IsStatic )		( COMPOSE_NET_GUID( ++UniqueNetIDs[ IsStatic ], IsStatic ) )
+
 /**
  *	Generate a new NetGUID for this object and assign it.
  */
@@ -2557,9 +2560,6 @@ FNetworkGUID FNetGUIDCache::AssignNewNetGUID_Server( UObject* Object )
 {
 	check( IsNetGUIDAuthority() );
 
-#define COMPOSE_NET_GUID( Index, IsStatic )	( ( ( Index ) << 1 ) | ( IsStatic ) )
-#define ALLOC_NEW_NET_GUID( IsStatic )		( COMPOSE_NET_GUID( ++UniqueNetIDs[ IsStatic ], IsStatic ) )
-
 	// Generate new NetGUID and assign it
 	const int32 IsStatic = IsDynamicObject( Object ) ? 0 : 1;
 
@@ -2591,6 +2591,9 @@ FNetworkGUID FNetGUIDCache::AssignNewNetGUIDFromPath_Server( const FString& Path
 	return NewNetGuid;
 }
 
+#undef COMPOSE_NET_GUID
+#undef ALLOC_NEW_NET_GUID
+
 void FNetGUIDCache::RegisterNetGUID_Internal( const FNetworkGUID& NetGUID, const FNetGuidCacheObject& CacheObject )
 {
 	LLM_SCOPE(ELLMTag::Networking);
diff --git a/Engine/Source/Runtime/Engine/Private/RepLayout.cpp b/Engine/Source/Runtime/Engine/Private/RepLayout.cpp
--- a/Engine/Source/Runtime/Engine/Private/RepLayout.cpp
+++ b/Engine/Source/Runtime/Engine/Private/RepLayout.cpp
@@ -783,7 +783,7 @@ static uint32 GetRepLayoutCmdCompatibleChecksum(
 	const uint32			InChecksum)
 {
 	// Compatible checksums are only used for InternalAck connections
-	if (ServerConnection && !ServerConnection->InternalAck)
+	if (ServerConnection && !ServerConnection->IsInternalAck())
 	{
 		return 0;
 	}
@@ -1630,7 +1630,7 @@ bool FRepLayout::ReplicateProperties(
 
 	if (OwningChannel->Connection->ResendAllDataState != EResendAllDataState::None)
 	{
-		check(OwningChannel->Connection->InternalAck);
+		check(OwningChannel->Connection->IsInternalAck());
 
 		// If we are resending data since open, we don't want to affect the current state of channel/replication, so just do the minimum and send the data, and return
 		if (RepState->LifetimeChangelist.Num() > 0)
@@ -1743,7 +1743,7 @@ bool FRepLayout::ReplicateProperties(
 	check(Changed.Num() > 0);
 
 	// do not build shared state for InternalAck (demo) connections
-	if (!OwningChannel->Connection->InternalAck && (GNetSharedSerializedData != 0))
+	if (!OwningChannel->Connection->IsInternalAck() && (GNetSharedSerializedData != 0))
 	{
 		// if no shared serialization info exists, build it
 		if (!RepChangelistState->SharedSerialization.IsValid())
@@ -1768,7 +1768,7 @@ bool FRepLayout::ReplicateProperties(
 	}
 
 	// Send the final merged change list
-	if (OwningChannel->Connection->InternalAck)
+	if (OwningChannel->Connection->IsInternalAck())
 	{
 		// Remember all properties that have changed since this channel was first opened in case we need it (for bResendAllDataSinceOpen)
 		// We use UnfilteredChanged so LifetimeChangelist contains all properties, regardless of Active state.
@@ -3290,7 +3290,7 @@ bool FRepLayout::ReceiveProperties(
 
 	UE_NET_TRACE_SCOPE(Properties, InBunch, OwningChannel->Connection->GetInTraceCollector(), ENetTraceVerbosity::Trace);
 
-	if (OwningChannel->Connection->InternalAck)
+	if (OwningChannel->Connection->IsInternalAck())
 	{
 		return ReceiveProperties_BackwardsCompatible(OwningChannel->Connection, RepState, Data, InBunch, bOutHasUnmapped, bEnableRepNotifies, bOutGuidsChanged);
 	}
@@ -6111,7 +6111,7 @@ void FRepLayout::SendPropertiesForRPC(
 
 	if (!IsEmpty())
 	{
-		if (Channel->Connection->InternalAck)
+		if (Channel->Connection->IsInternalAck())
 		{
 			TArray<uint16> Changed;
 
@@ -6182,7 +6182,7 @@ void FRepLayout::ReceivePropertiesForRPC(
 			}
 		}
 
-		if (Channel->Connection->InternalAck)
+		if (Channel->Connection->IsInternalAck())
 		{
 			bool bHasUnmapped = false;
 			bool bGuidsChanged = false;
@@ -6572,7 +6572,7 @@ void FRepLayout::PreSendCustomDeltaProperties(
 {
 	using namespace UE4_RepLayout_Private;
 
-	if (!Connection->InternalAck)
+	if (!Connection->IsInternalAck())
 	{
 		const FLifetimeCustomDeltaState& LocalLifetimeCustomPropertyState = *LifetimeCustomPropertyState;
 
@@ -6700,7 +6700,7 @@ bool FRepLayout::DeltaSerializeFastArrayProperty(FFastArrayDeltaSerializeParams&
 	UPackageMapClient* PackageMap = static_cast<UPackageMapClient*>(DeltaSerializeInfo.Map);
 	UNetConnection* Connection = DeltaSerializeInfo.Connection;
 	const bool bIsWriting = !!DeltaSerializeInfo.Writer;
-	const bool bInternalAck = !!Connection->InternalAck;
+	const bool bInternalAck = Connection->IsInternalAck();
 
 	FRepObjectDataBuffer ObjectData(Object);
 	FScriptArray* ObjectArray = GetTypedProperty<FScriptArray>(ObjectData, FastArrayItemCmd);
@@ -6735,10 +6735,6 @@ bool FRepLayout::DeltaSerializeFastArrayProperty(FFastArrayDeltaSerializeParams&
 			PackageMap->AddNetFieldExportGroup(OwnerPathName, LocalNetFieldExportGroup);
 		}
 
-		checkf(LocalNetFieldExportGroup->NetFieldExports.Num() == Cmds.Num(),
-			TEXT("NetFieldExports.Num() does not match number of commands! PathName = %s, NetFieldExportGroup.PathName = %s, Cmds.Num() = %d, NetFieldExports.Num() = %d"),
-			*OwnerPathName, *(LocalNetFieldExportGroup->PathName), Cmds.Num(), LocalNetFieldExportGroup->NetFieldExports.Num());
-
 		NetFieldExportGroup = LocalNetFieldExportGroup.Get();
 	}
 
diff --git a/Engine/Source/Runtime/GameplayTags/Private/GameplayTagContainer.cpp b/Engine/Source/Runtime/GameplayTags/Private/GameplayTagContainer.cpp
--- a/Engine/Source/Runtime/GameplayTags/Private/GameplayTagContainer.cpp
+++ b/Engine/Source/Runtime/GameplayTags/Private/GameplayTagContainer.cpp
@@ -1218,7 +1218,7 @@ bool FGameplayTag::NetSerialize_Packed(FArchive& Ar, class UPackageMap* Map, boo
 		FGameplayTagNetIndex NetIndex = INVALID_TAGNETINDEX;
 
 		UPackageMapClient* PackageMapClient = Cast<UPackageMapClient>(Map);
-		const bool bIsReplay = PackageMapClient && PackageMapClient->GetConnection() && PackageMapClient->GetConnection()->InternalAck;
+		const bool bIsReplay = PackageMapClient && PackageMapClient->GetConnection() && PackageMapClient->GetConnection()->IsInternalAck();
 
 		TSharedPtr<FNetFieldExportGroup> NetFieldExportGroup;
 
