Added FGPUSortManager that handles different GPU sort tasks.
The current clients for it are Cascade and Niagara.

diff --git a/Engine/Plugins/FX/Niagara/Shaders/Private/NiagaraCopyIntBuffer.usf b/Engine/Plugins/FX/Niagara/Shaders/Private/NiagaraCopyIntBuffer.usf
deleted file mode 100644
diff --git a/Engine/Plugins/FX/Niagara/Shaders/Private/NiagaraSpriteVertexFactory.ush b/Engine/Plugins/FX/Niagara/Shaders/Private/NiagaraSpriteVertexFactory.ush
--- a/Engine/Plugins/FX/Niagara/Shaders/Private/NiagaraSpriteVertexFactory.ush
+++ b/Engine/Plugins/FX/Niagara/Shaders/Private/NiagaraSpriteVertexFactory.ush
@@ -843,7 +843,7 @@ FVertexFactoryIntermediates GetVertexFactoryIntermediates(FVertexFactoryInput In
 
 	uint ParticleID = Input.InstanceId;
 
-#if FEATURE_LEVEL >= FEATURE_LEVEL_SM4 && !SHADOW_DEPTH_SHADER && !SPHERICAL_OPACITY_FOR_SHADOW_DEPTHS
+#if !SHADOW_DEPTH_SHADER && !SPHERICAL_OPACITY_FOR_SHADOW_DEPTHS
 	// Don't use sort buffer in the shadow because they are only valid after opaque.
 	if(NiagaraSpriteVFLooseParameters.SortedIndicesOffset != 0xFFFFFFFF)
 	{
diff --git a/Engine/Plugins/FX/Niagara/Source/Niagara/Classes/NiagaraEmitterInstanceBatcher.h b/Engine/Plugins/FX/Niagara/Source/Niagara/Classes/NiagaraEmitterInstanceBatcher.h
--- a/Engine/Plugins/FX/Niagara/Source/Niagara/Classes/NiagaraEmitterInstanceBatcher.h
+++ b/Engine/Plugins/FX/Niagara/Source/Niagara/Classes/NiagaraEmitterInstanceBatcher.h
@@ -23,20 +23,7 @@ the same VectorVM byte code / compute shader code
 #include "NiagaraScriptExecutionContext.h"
 #include "NiagaraGPUInstanceCountManager.h"
 
-class FNiagaraIndicesVertexBuffer : public FParticleIndicesVertexBuffer
-{
-public:
-
-	FNiagaraIndicesVertexBuffer(int32 InIndexCount);
-
-	FUnorderedAccessViewRHIRef VertexBufferUAV;
-
-	// The allocation count.
-	const int32 IndexCount;
-
-	// Currently used count.
-	int32 UsedIndexCount = 0;
-};
+class FGPUSortManager;
 
 enum class ETickStage
 {
@@ -51,16 +38,7 @@ public:
 	using FNiagaraBufferArray = TArray<FRHIUnorderedAccessView*, TMemStackAllocator<>>;
 	using FOverlappableTicks = TArray<FNiagaraGPUSystemTick*, TMemStackAllocator<>>;
 
-	NiagaraEmitterInstanceBatcher(ERHIFeatureLevel::Type InFeatureLevel, EShaderPlatform InShaderPlatform)
-		: FeatureLevel(InFeatureLevel)
-		, ShaderPlatform(InShaderPlatform)
-		, ParticleSortBuffers(true)
-		// @todo REMOVE THIS HACK
-		, LastFrameThatDrainedData(GFrameNumberRenderThread)
-		, NumAllocatedFreeIDListSizes(0)
-		, bFreeIDListSizesBufferCleared(false)
-	{
-	}
+	NiagaraEmitterInstanceBatcher(ERHIFeatureLevel::Type InFeatureLevel, EShaderPlatform InShaderPlatform, FGPUSortManager* InGPUSortManager);
 
 	~NiagaraEmitterInstanceBatcher();
 
@@ -116,11 +94,16 @@ public:
 		FRHIUniformBuffer* SceneTexturesUniformBuffer,
 		bool bAllowGPUParticleUpdate) override;
 
-	int32 AddSortedGPUSimulation(const FNiagaraGPUSortInfo& SortInfo);
-	void SortGPUParticles(FRHICommandListImmediate& RHICmdList);
-	void ResolveParticleSortBuffers(FRHICommandListImmediate& RHICmdList, int32 ResultBufferIndex);
+	/** 
+	 * Register work for GPU sorting (using the GPUSortManager). 
+	 * The constraints of the sort request are defined in SortInfo.SortFlags.
+	 * The sort task bindings are set in SortInfo.AllocationInfo.
+	 * The initial keys and values are generated in the GenerateSortKeys() callback.
+	 *
+	 * Return true if the work was registered, or false it GPU sorting is not available or impossible.
+	 */
+	bool AddSortedGPUSimulation(FNiagaraGPUSortInfo& SortInfo);
 
-	const FParticleIndicesVertexBuffer& GetGPUSortedBuffer() const { return SortedVertexBuffers.Last(); }
 	const FGlobalDistanceFieldParameterData& GetGlobalDistanceFieldParameters() const { return GlobalDistanceFieldParams; }
 
 	void ProcessDebugInfo(FRHICommandList &RHICmdLis, const FNiagaraComputeExecutionContext* Context) const;
@@ -165,6 +148,9 @@ public:
 	/** Run the dispatch over multiple stages */
 	void DispatchMultipleStages(const FNiagaraGPUSystemTick& Tick, FNiagaraComputeInstanceData *Instance, FRHICommandList &RHICmdList, FRHIUniformBuffer* ViewUniformBuffer, FNiagaraShader* ComputeShader);
 
+	/** Get the shared SortManager, used in the rendering loop to call FGPUSortManager::OnPreRender() and FGPUSortManager::OnPostRenderOpaque() */
+	virtual FGPUSortManager* GetGPUSortManager() const override;
+
 private:
 	using FEmitterInstanceList = TArray<FNiagaraComputeInstanceData*>;
 
@@ -174,6 +160,27 @@ private:
 
 	bool ShouldTickForStage(const FNiagaraGPUSystemTick& Tick, ETickStage TickStage) const;
 
+	/**
+	 * Generate all the initial keys and values for a GPUSortManager sort batch.
+	 * Sort batches are created when GPU sort tasks are registered in AddSortedGPUSimulation().
+	 * Each sort task defines constraints about when the initial sort data can generated and
+	 * and when the sorted results are needed (see EGPUSortFlags for details).
+	 * Currently, for Niagara, all the sort tasks have the EGPUSortFlags::KeyGenAfterPreRender flag
+	 * and so the callback registered in GPUSortManager->Register() only has the EGPUSortFlags::KeyGenAfterPreRender usage.
+	 * This garanties that GenerateSortKeys() only gets called after PreRender(), which is a constraint required because
+	 * Niagara renders the current state of the GPU emitters, before the are ticked 
+	 * (Niagara GPU emitters are ticked at InitView and in PostRenderOpaque).
+	 * Note that this callback must only initialize the content for the elements that relates to the tasks it has registered in this batch.
+	 *
+	 * @param RHICmdList - The command list used to initiate the keys and values on GPU.
+	 * @param BatchId - The GPUSortManager batch id (regrouping several similar sort tasks).
+	 * @param NumElementsInBatch - The number of elements grouped in the batch (each element maps to a sort task)
+	 * @param Flags - Details about the key precision (see EGPUSortFlags::AnyKeyPrecision) and the keygen location (see EGPUSortFlags::AnyKeyGenLocation).
+	 * @param KeysUAV - The UAV that holds all the initial keys used to sort the values (being the particle indices here). 
+	 * @param ValuesUAV - The UAV that holds the initial values (particle indices) to be sorted accordingly to the keys.
+	 */
+	void GenerateSortKeys(FRHICommandListImmediate& RHICmdList, int32 BatchId, int32 NumElementsInBatch, EGPUSortFlags Flags, FRHIUnorderedAccessView* KeysUAV, FRHIUnorderedAccessView* ValuesUAV);
+
 //	void SimStepReadback(const FNiagaraComputeInstanceData& Instance, FRHICommandList& RHICmdList) const;
 
 	inline uint32 UnpackEmitterDispatchCount(uint8* PackedData)
@@ -197,29 +204,24 @@ private:
 	/** Shader platform that will be rendering this effects system */
 	EShaderPlatform ShaderPlatform;
 
-	// Number of particle to sort this frame.
-	int32 SortedParticleCount = 0;
-	int32 NumFramesRequiringShrinking = 0;
+	/** The shared GPUSortManager, used to register GPU sort tasks in order to generate sorted particle indices per emitter. */
+	TRefCountPtr<FGPUSortManager> GPUSortManager;
+	/** All sort tasks registered in AddSortedGPUSimulation(). Holds all the data required in GenerateSortKeys(). */
 	TArray<FNiagaraGPUSortInfo> SimulationsToSort;
-	FParticleSortBuffers ParticleSortBuffers;
 
 	// GPU emitter instance count buffer. Contains the actual particle / instance count generate in the GPU tick.
 	FNiagaraGPUInstanceCountManager GPUInstanceCounterManager;
 
 	// @todo REMOVE THIS HACK
 	uint32 LastFrameThatDrainedData;
-
-	// The result of the GPU sort. Each next element replace the previous.
-	// The last entry is used to transfer the result of the ParticleSortBuffers.
-	TIndirectArray<FNiagaraIndicesVertexBuffer> SortedVertexBuffers;
 	
 	TArray<FNiagaraGPUSystemTick> Ticks_RT;
 	FGlobalDistanceFieldParameterData GlobalDistanceFieldParams;
 
 	/** A buffer of list sizes used by UpdateFreeIDBuffers to allow overlapping several dispatches. */
 	FRWBuffer FreeIDListSizesBuffer;
-	uint32 NumAllocatedFreeIDListSizes;
-	bool bFreeIDListSizesBufferCleared;
+	uint32 NumAllocatedFreeIDListSizes = 0;
+	bool bFreeIDListSizesBufferCleared = false;
 
 	/** List of emitter instances which need their free ID buffers updated post render. */
 	FEmitterInstanceList DeferredIDBufferUpdates;
diff --git a/Engine/Plugins/FX/Niagara/Source/Niagara/Classes/NiagaraGPUSortInfo.h b/Engine/Plugins/FX/Niagara/Source/Niagara/Classes/NiagaraGPUSortInfo.h
--- a/Engine/Plugins/FX/Niagara/Source/Niagara/Classes/NiagaraGPUSortInfo.h
+++ b/Engine/Plugins/FX/Niagara/Source/Niagara/Classes/NiagaraGPUSortInfo.h
@@ -8,6 +8,7 @@ NiagaraGPUSortInfo.h: GPU particle sorting helper
 #include "CoreMinimal.h"
 #include "RHI.h"
 #include "RHIResources.h"
+#include "GPUSortManager.h"
 #include "NiagaraGPUSortInfo.generated.h"
 
 UENUM()
@@ -42,4 +43,17 @@ struct FNiagaraGPUSortInfo
 	// View data.
 	FVector ViewOrigin;
 	FVector ViewDirection;
+
+	// The GPUSortManager bindings for this sort task.
+	FGPUSortManager::FAllocationInfo AllocationInfo;
+	// The sort constraints for the task in GPUSortManager.
+	EGPUSortFlags SortFlags = EGPUSortFlags::None;
+
+	// Set the SortFlags based on the emitter and material constraints.
+	FORCEINLINE void SetSortFlags(bool bHighPrecisionKeys, bool bTranslucentMaterial)
+	{
+		SortFlags = EGPUSortFlags::KeyGenAfterPreRender | EGPUSortFlags::ValuesAsInt32 | 
+			(bHighPrecisionKeys ? EGPUSortFlags::HighPrecisionKeys : EGPUSortFlags::LowPrecisionKeys) |
+			(bTranslucentMaterial ? EGPUSortFlags::AnySortLocation : EGPUSortFlags::SortAfterPreRender);
+	}
 };
\ No newline at end of file
diff --git a/Engine/Plugins/FX/Niagara/Source/Niagara/Private/NiagaraEmitterInstanceBatcher.cpp b/Engine/Plugins/FX/Niagara/Source/Niagara/Private/NiagaraEmitterInstanceBatcher.cpp
--- a/Engine/Plugins/FX/Niagara/Source/Niagara/Private/NiagaraEmitterInstanceBatcher.cpp
+++ b/Engine/Plugins/FX/Niagara/Source/Niagara/Private/NiagaraEmitterInstanceBatcher.cpp
@@ -69,15 +69,6 @@ static FAutoConsoleVariableRef CVarNiagaraGpuMaxQueuedRenderFrames(
 	ECVF_Default
 );
 
-FNiagaraIndicesVertexBuffer::FNiagaraIndicesVertexBuffer(int32 InIndexCount)
-	: IndexCount(InIndexCount)
-{
-	FRHIResourceCreateInfo CreateInfo;
-	VertexBufferRHI = RHICreateVertexBuffer((uint32)IndexCount * sizeof(int32), BUF_Static | BUF_ShaderResource | BUF_UnorderedAccess, CreateInfo);
-	VertexBufferSRV = RHICreateShaderResourceView(VertexBufferRHI, sizeof(int32), PF_R32_SINT);
-	VertexBufferUAV = RHICreateUnorderedAccessView(VertexBufferRHI, PF_R32_SINT);
-}
-
 const FName NiagaraEmitterInstanceBatcher::Name(TEXT("NiagaraEmitterInstanceBatcher"));
 
 FFXSystemInterface* NiagaraEmitterInstanceBatcher::GetInterface(const FName& InName)
@@ -85,10 +76,33 @@ FFXSystemInterface* NiagaraEmitterInstanceBatcher::GetInterface(const FName& InN
 	return InName == Name ? this : nullptr;
 }
 
+NiagaraEmitterInstanceBatcher::NiagaraEmitterInstanceBatcher(ERHIFeatureLevel::Type InFeatureLevel, EShaderPlatform InShaderPlatform, FGPUSortManager* InGPUSortManager)
+	: FeatureLevel(InFeatureLevel)
+	, ShaderPlatform(InShaderPlatform)
+	, GPUSortManager(InGPUSortManager)
+	// @todo REMOVE THIS HACK
+	, LastFrameThatDrainedData(GFrameNumberRenderThread)
+	, NumAllocatedFreeIDListSizes(0)
+	, bFreeIDListSizesBufferCleared(false)
+{
+	// Register the batcher callback in the GPUSortManager. 
+	// The callback is used to generate the initial keys and values for the GPU sort tasks, 
+	// the values being the sorted particle indices used by the Niagara renderers.
+	// The registration also involves defining the list of flags possibly used in GPUSortManager::AddTask()
+	if (GPUSortManager)
+	{
+		GPUSortManager->Register(FGPUSortKeyGenDelegate::CreateLambda([this](FRHICommandListImmediate& RHICmdList, int32 BatchId, int32 NumElementsInBatch, EGPUSortFlags Flags, FRHIUnorderedAccessView* KeysUAV, FRHIUnorderedAccessView* ValuesUAV)
+		{ 
+			GenerateSortKeys(RHICmdList, BatchId, NumElementsInBatch, Flags, KeysUAV, ValuesUAV);
+		}), 
+		EGPUSortFlags::AnyKeyPrecision | EGPUSortFlags::KeyGenAfterPreRender | EGPUSortFlags::AnySortLocation | EGPUSortFlags::ValuesAsInt32,
+		Name);
+	}
+}
+
 NiagaraEmitterInstanceBatcher::~NiagaraEmitterInstanceBatcher()
 {
 	FinishDispatches();
-	ParticleSortBuffers.ReleaseRHI();
 }
 
 void NiagaraEmitterInstanceBatcher::InstanceDeallocated_RenderThread(const FNiagaraSystemInstanceID InstanceID)
@@ -796,14 +810,11 @@ void NiagaraEmitterInstanceBatcher::PreInitViews(FRHICommandListImmediate& RHICm
 {
 	LLM_SCOPE(ELLMTag::Niagara);
 
-	SortedParticleCount = 0;
+	// Reset the list of GPUSort tasks and release any resources they hold on to.
+	// It might be worth considering doing so at the end of the render to free the resources immediately.
+	// (note that currently there are no callback appropriate to do it)
 	SimulationsToSort.Reset();
 
-	for (FNiagaraIndicesVertexBuffer& SortedVertexBuffer : SortedVertexBuffers)
-	{
-		SortedVertexBuffer.UsedIndexCount = 0;
-	}
-
 	// Update draw indirect buffer to max possible size.
 	if (bAllowGPUParticleUpdate)
 	{
@@ -936,12 +947,8 @@ void NiagaraEmitterInstanceBatcher::PreRender(FRHICommandListImmediate& RHICmdLi
 
 	GlobalDistanceFieldParams = GlobalDistanceFieldParameterData ? *GlobalDistanceFieldParameterData : FGlobalDistanceFieldParameterData();
 
-	// Sort buffer after mesh batches are issued, before tick (which will change the GPU instance count).
-	SortGPUParticles(RHICmdList);
-
 	// Update draw indirect args from the simulation results.
 	GPUInstanceCounterManager.UpdateDrawIndirectBuffer(RHICmdList, FeatureLevel);
-
 }
 
 void NiagaraEmitterInstanceBatcher::OnDestroy()
@@ -950,200 +957,46 @@ void NiagaraEmitterInstanceBatcher::OnDestroy()
 	FFXSystemInterface::OnDestroy();
 }
 
-int32 NiagaraEmitterInstanceBatcher::AddSortedGPUSimulation(const FNiagaraGPUSortInfo& SortInfo)
+bool NiagaraEmitterInstanceBatcher::AddSortedGPUSimulation(FNiagaraGPUSortInfo& SortInfo)
 {
-	const int32 ResultOffset = SortedParticleCount;
-	SimulationsToSort.Add(SortInfo);
-
-	SortedParticleCount += SortInfo.ParticleCount;
-
-	if (!SortedVertexBuffers.Num())
-	{
-		SortedVertexBuffers.Add(new FNiagaraIndicesVertexBuffer(FMath::Max(GNiagaraGPUSortingMinBufferSize, (int32)(SortedParticleCount * GNiagaraGPUSortingBufferSlack))));
-	}
-	// If we don't fit anymore, reallocate to a bigger size.
-	else if (SortedParticleCount > SortedVertexBuffers.Last().IndexCount)
+	if (GPUSortManager && GPUSortManager->AddTask(SortInfo.AllocationInfo, SortInfo.ParticleCount, SortInfo.SortFlags))
 	{
-		SortedVertexBuffers.Add(new FNiagaraIndicesVertexBuffer((int32)(SortedParticleCount * GNiagaraGPUSortingBufferSlack)));
+		// It's not worth currently to have a map between SortInfo.AllocationInfo.SortBatchId and the relevant indices in SimulationsToSort
+		// because the number of batches is expect to be very small (1 or 2). If this change, it might be worth reconsidering.
+		SimulationsToSort.Add(SortInfo);
+		return true;
 	}
-
-	// Keep track of the last used index, which is also the first used index of next entry
-	// if we need to increase the size of SortedVertexBuffers. Used in FNiagaraCopyIntBufferRegionCS
-	SortedVertexBuffers.Last().UsedIndexCount = SortedParticleCount;
-
-	return ResultOffset;
-}
-
-void NiagaraEmitterInstanceBatcher::SortGPUParticles(FRHICommandListImmediate& RHICmdList)
-{
-	if (SortedParticleCount > 0 && SortedVertexBuffers.Num() > 0 && SimulationsToSort.Num() && GNiagaraGPUSortingBufferSlack > 1.f)
-	{
-		SCOPED_GPU_STAT(RHICmdList, NiagaraGPUSorting);
-
-		//UE_LOG(LogNiagara, Warning, TEXT("NiagaraEmitterInstanceBatcher::SortGPUParticles:  %0xP"), this);
-
-		ensure(SortedVertexBuffers.Last().IndexCount >= SortedParticleCount);
-
-		// The particle sort buffer must be able to hold all the particles.
-		if (SortedVertexBuffers.Last().IndexCount != ParticleSortBuffers.GetSize())
-		{
-			ParticleSortBuffers.ReleaseRHI();
-			ParticleSortBuffers.SetBufferSize(SortedVertexBuffers.Last().IndexCount);
-			ParticleSortBuffers.InitRHI();
-		}
-
-		INC_DWORD_STAT_BY(STAT_NiagaraGPUSortedParticles, SortedParticleCount);
-		INC_DWORD_STAT_BY(STAT_NiagaraGPUSortedBuffers, ParticleSortBuffers.GetSize());
-
-		// Make sure our outputs are safe to write to.
-		const int32 InitialSortBufferIndex = 0;
-		FRHIUnorderedAccessView* OutputUAVs[2] = { ParticleSortBuffers.GetKeyBufferUAV(InitialSortBufferIndex), ParticleSortBuffers.GetVertexBufferUAV(InitialSortBufferIndex) };
-		RHICmdList.TransitionResources(EResourceTransitionAccess::ERWBarrier, EResourceTransitionPipeline::EGfxToCompute, OutputUAVs, 2);
-
-		// EmitterKey = (EmitterIndex & EmitterKeyMask) << EmitterKeyShift.
-		// SortKey = (Key32 >> SortKeyShift) & SortKeyMask.
-		uint32 EmitterKeyMask = (1 << FMath::CeilLogTwo(SimulationsToSort.Num())) - 1;
-		uint32 EmitterKeyShift = 16;
-		uint32 SortKeyMask = 0xFFFF;
-		
-		{
-			SCOPED_DRAW_EVENT(RHICmdList, NiagaraSortKeyGen);
-
-			// Bind the shader
-			
-			FNiagaraSortKeyGenCS::FPermutationDomain PermutationVector;
-			PermutationVector.Set<FNiagaraSortKeyGenCS::FSortUsingMaxPrecision>(GNiagaraGPUSortingUseMaxPrecision != 0);
-			
-			TShaderMapRef<FNiagaraSortKeyGenCS> KeyGenCS(GetGlobalShaderMap(FeatureLevel), PermutationVector);
-			RHICmdList.SetComputeShader(KeyGenCS->GetComputeShader());
-			KeyGenCS->SetOutput(RHICmdList, ParticleSortBuffers.GetKeyBufferUAV(InitialSortBufferIndex), ParticleSortBuffers.GetVertexBufferUAV(InitialSortBufferIndex));
-
-			// (SortKeyMask, SortKeyShift, SortKeySignBit)
-			FUintVector4 SortKeyParams(SortKeyMask, 0, 0x8000, 0); 
-			if (GNiagaraGPUSortingUseMaxPrecision != 0)
-			{
-				EmitterKeyMask = FMath::Max<uint32>(EmitterKeyMask , 1); // Need at list 1 bit for the above logic
-				uint32 UnusedBits = FPlatformMath::CountLeadingZeros(EmitterKeyMask << EmitterKeyShift);
-				EmitterKeyShift += UnusedBits;
-				SortKeyMask = ~(EmitterKeyMask << EmitterKeyShift);
-
-				SortKeyParams.X = SortKeyMask;
-				SortKeyParams.Y = 16 - UnusedBits;
-				SortKeyParams.Z <<= UnusedBits;
-			}
-			
-			int32 OutputOffset = 0;
-			for (int32 EmitterIndex = 0; EmitterIndex < SimulationsToSort.Num(); ++EmitterIndex)
-			{
-				const FNiagaraGPUSortInfo& SortInfo = SimulationsToSort[EmitterIndex];
-				KeyGenCS->SetParameters(RHICmdList, SortInfo, (uint32)EmitterIndex << EmitterKeyShift, OutputOffset, SortKeyParams);
-				DispatchComputeShader(RHICmdList, *KeyGenCS, FMath::DivideAndRoundUp(SortInfo.ParticleCount, NIAGARA_KEY_GEN_THREAD_COUNT), 1, 1);
-
-				OutputOffset += SortInfo.ParticleCount;
-			}
-			KeyGenCS->UnbindBuffers(RHICmdList);
-		}
-
-		// We may be able to remove this transition if each step isn't dependent on the previous one.
-		RHICmdList.TransitionResources(EResourceTransitionAccess::ERWBarrier, EResourceTransitionPipeline::EComputeToCompute, OutputUAVs, 2);
-
-		// Sort buffers and copy results to index buffers.
-		{
-			const uint32 KeyMask = (EmitterKeyMask << EmitterKeyShift) | SortKeyMask;
-			const int32 ResultBufferIndex = SortGPUBuffers(RHICmdList, ParticleSortBuffers.GetSortBuffers(), InitialSortBufferIndex, KeyMask, SortedParticleCount, FeatureLevel);
-			ResolveParticleSortBuffers(RHICmdList, ResultBufferIndex);
-		}
-
-		// Only keep the last sorted index buffer, which is of the same size as ParticleSortBuffers.GetSize().
-		SortedVertexBuffers.RemoveAt(0, SortedVertexBuffers.Num() - 1);
-
-		// Resize the buffer to maximize next frame.
-		// Those ratio must take into consideration the slack ratio to be stable.
-
-		const int32 RecommandedSize = FMath::Max(GNiagaraGPUSortingMinBufferSize, (int32)(SortedParticleCount * GNiagaraGPUSortingBufferSlack));
-		const float BufferUsage = (float)SortedParticleCount /	(float)ParticleSortBuffers.GetSize();
-
-		if (RecommandedSize < ParticleSortBuffers.GetSize() / GNiagaraGPUSortingBufferSlack)
-		{
-			if (NumFramesRequiringShrinking >= GNiagaraGPUSortingFrameCountBeforeBufferShrinking)
-			{
-				NumFramesRequiringShrinking = 0;
-				ParticleSortBuffers.ReleaseRHI();
-				ParticleSortBuffers.SetBufferSize(0);
-
-				// Add an entry that should fit well for next frame.
-				SortedVertexBuffers.Empty();
-				SortedVertexBuffers.Add(new FNiagaraIndicesVertexBuffer(RecommandedSize));
-			}
-			else
-			{
-				++NumFramesRequiringShrinking;
-			}
-		}
-		else // Reset counter since we are not in a shrinking situation anymore.
-		{
-			NumFramesRequiringShrinking = 0;
-		}
-	}
-	else // If the are no sort task, we don't need any of the sort buffers.
+	else
 	{
-		if (NumFramesRequiringShrinking >= GNiagaraGPUSortingFrameCountBeforeBufferShrinking)
-		{
-			NumFramesRequiringShrinking = 0;
-			ParticleSortBuffers.ReleaseRHI();
-			ParticleSortBuffers.SetBufferSize(0);
-			SortedVertexBuffers.Empty();
-		}
-		else
-		{
-			++NumFramesRequiringShrinking;
-		}
+		return false;
 	}
 }
 
-void NiagaraEmitterInstanceBatcher::ResolveParticleSortBuffers(FRHICommandListImmediate& RHICmdList, int32 ResultBufferIndex)
+void NiagaraEmitterInstanceBatcher::GenerateSortKeys(FRHICommandListImmediate& RHICmdList, int32 BatchId, int32 NumElementsInBatch, EGPUSortFlags Flags, FRHIUnorderedAccessView* KeysUAV, FRHIUnorderedAccessView* ValuesUAV)
 {
-	SCOPED_DRAW_EVENT(RHICmdList, NiagaraResolveParticleSortBuffers);
-
-#if 0 // TODO use this once working properly!
-	if (SortedVertexBuffers.Num() == 1)
-	{
-		RHICmdList.CopyVertexBuffer(ParticleSortBuffers.GetSortedVertexBufferRHI(ResultBufferIndex), SortedVertexBuffers.Last().VertexBufferRHI);
-		return;
-	}
-#endif
+	// Currently all Niagara KeyGen must execute after PreRender() - in between PreInitViews() and PostRenderOpaque(), when the GPU simulation are possibly ticked.
+	check(EnumHasAnyFlags(Flags, EGPUSortFlags::KeyGenAfterPreRender));
 
-	TShaderMapRef<FNiagaraCopyIntBufferRegionCS> CopyBufferCS(GetGlobalShaderMap(FeatureLevel));
-	RHICmdList.SetComputeShader(CopyBufferCS->GetComputeShader());
+	const FGPUSortManager::FKeyGenInfo KeyGenInfo((uint32)NumElementsInBatch, EnumHasAnyFlags(Flags, EGPUSortFlags::HighPrecisionKeys));
 
-	int32 StartingIndex = 0;
+	FNiagaraSortKeyGenCS::FPermutationDomain PermutationVector;
+	PermutationVector.Set<FNiagaraSortKeyGenCS::FSortUsingMaxPrecision>(EnumHasAnyFlags(Flags, EGPUSortFlags::HighPrecisionKeys));
+	TShaderMapRef<FNiagaraSortKeyGenCS> KeyGenCS(GetGlobalShaderMap(FeatureLevel), PermutationVector);
+	RHICmdList.SetComputeShader(KeyGenCS->GetComputeShader());
+	KeyGenCS->SetOutput(RHICmdList, KeysUAV, ValuesUAV);
 
-	for (int32 Index = 0; Index < SortedVertexBuffers.Num(); Index += NIAGARA_COPY_BUFFER_BUFFER_COUNT)
+	FRHIUnorderedAccessView* OutputUAVs[] = { KeysUAV, ValuesUAV };
+	for (const FNiagaraGPUSortInfo& SortInfo : SimulationsToSort)
 	{
-		FRHIUnorderedAccessView* UAVs[NIAGARA_COPY_BUFFER_BUFFER_COUNT] = {};
-		int32 UsedIndexCounts[NIAGARA_COPY_BUFFER_BUFFER_COUNT] = {};
-
-		const int32 NumBuffers = FMath::Min<int32>(NIAGARA_COPY_BUFFER_BUFFER_COUNT, SortedVertexBuffers.Num() - Index);
-
-		int32 LastCount = StartingIndex;
-		for (int32 SubIndex = 0; SubIndex < NumBuffers; ++SubIndex)
+		if (SortInfo.AllocationInfo.SortBatchId == BatchId)
 		{
-			const FNiagaraIndicesVertexBuffer& SortBuffer = SortedVertexBuffers[Index + SubIndex];
-			UAVs[SubIndex] = SortBuffer.VertexBufferUAV;
-			UsedIndexCounts[SubIndex] = SortBuffer.UsedIndexCount;
-
-			LastCount = SortBuffer.UsedIndexCount;
+			KeyGenCS->SetParameters(RHICmdList, SortInfo, (uint32)SortInfo.AllocationInfo.ElementIndex << KeyGenInfo.ElementKeyShift, SortInfo.AllocationInfo.BufferOffset, KeyGenInfo.SortKeyParams);
+			DispatchComputeShader(RHICmdList, *KeyGenCS, FMath::DivideAndRoundUp(SortInfo.ParticleCount, NIAGARA_KEY_GEN_THREAD_COUNT), 1, 1);
+			// TR-KeyGen : No sync needed between tasks since they update different parts of the data (assuming it's ok if cache lines overlap).
+			RHICmdList.TransitionResources(EResourceTransitionAccess::ERWNoBarrier, EResourceTransitionPipeline::EComputeToCompute, OutputUAVs, UE_ARRAY_COUNT(OutputUAVs));
 		}
-
-		RHICmdList.TransitionResources(EResourceTransitionAccess::ERWBarrier, EResourceTransitionPipeline::EComputeToCompute, UAVs, NumBuffers);
-
-		CopyBufferCS->SetParameters(RHICmdList, ParticleSortBuffers.GetSortedVertexBufferSRV(ResultBufferIndex), UAVs, UsedIndexCounts, StartingIndex, NumBuffers);
-		DispatchComputeShader(RHICmdList, *CopyBufferCS, FMath::DivideAndRoundUp(LastCount - StartingIndex, NIAGARA_COPY_BUFFER_THREAD_COUNT), 1, 1);
-		RHICmdList.TransitionResources(EResourceTransitionAccess::ERWBarrier, EResourceTransitionPipeline::EComputeToGfx, UAVs, NumBuffers);
-
-		StartingIndex = LastCount;
 	}
-	CopyBufferCS->UnbindBuffers(RHICmdList);
+	KeyGenCS->UnbindBuffers(RHICmdList);
 }
 
 void NiagaraEmitterInstanceBatcher::ProcessDebugInfo(FRHICommandList &RHICmdList, const FNiagaraComputeExecutionContext* Context) const
@@ -1489,3 +1342,8 @@ void NiagaraEmitterInstanceBatcher::Run(const FNiagaraGPUSystemTick& Tick, const
 	DestinationData.UnsetShaderParams(Shader, RHICmdList);
 	Shader->InstanceCountsParam.UnsetUAV(RHICmdList, ComputeShader);
 }
+
+FGPUSortManager* NiagaraEmitterInstanceBatcher::GetGPUSortManager() const
+{
+	return GPUSortManager;
+}
diff --git a/Engine/Plugins/FX/Niagara/Source/Niagara/Private/NiagaraGPUInstanceCountManager.cpp b/Engine/Plugins/FX/Niagara/Source/Niagara/Private/NiagaraGPUInstanceCountManager.cpp
--- a/Engine/Plugins/FX/Niagara/Source/Niagara/Private/NiagaraGPUInstanceCountManager.cpp
+++ b/Engine/Plugins/FX/Niagara/Source/Niagara/Private/NiagaraGPUInstanceCountManager.cpp
@@ -3,7 +3,7 @@
 #include "NiagaraGPUInstanceCountManager.h"
 #include "NiagaraStats.h"
 #include "Containers/DynamicRHIResourceArray.h"
-#include "NiagaraSortingGPU.h" // FNiagaraCopyIntBufferRegionCS
+#include "GPUSortManager.h" // CopyUIntBufferToTargets
 
 int32 GNiagaraMinGPUInstanceCount = 2048;
 static FAutoConsoleVariableRef CVarNiagaraMinGPUInstanceCount(
@@ -110,28 +110,25 @@ void FNiagaraGPUInstanceCountManager::ResizeBuffers(FRHICommandListImmediate& RH
 		// If we need to increase the buffer size to RecommendedInstanceCounts because the buffer is too small.
 		else if (RequiredInstanceCounts > AllocatedInstanceCounts)
 		{
+			SCOPED_DRAW_EVENT(RHICmdList, ResizeNiagaraGPUCounts);
+
 			// Init a bigger buffer filled with 0.
 			TResourceArray<uint32> InitData;
 			InitData.AddZeroed(RecommendedInstanceCounts);
 			FRWBuffer NextCountBuffer;
 			NextCountBuffer.Initialize(sizeof(uint32), RecommendedInstanceCounts, EPixelFormat::PF_R32_UINT, BUF_Static, TEXT("NiagaraGPUInstanceCounts"), &InitData);
 
-			// Because the shader works with SINT, we need to temporarily create view for the shader to work. Those will be (deferred) deleted at end of the scope.
-			FUnorderedAccessViewRHIRef NextCountBufferUAVAsInt = RHICreateUnorderedAccessView(NextCountBuffer.Buffer, EPixelFormat::PF_R32_SINT);
-			FShaderResourceViewRHIRef CountBufferSRVAsInt = RHICreateShaderResourceView(CountBuffer.Buffer, sizeof(uint32), EPixelFormat::PF_R32_SINT);
-
+			// TR-InstanceCount : The new count where just initialized here, while the previous count needs to transit in readable state..
+			RHICmdList.TransitionResource(EResourceTransitionAccess::ERWBarrier, EResourceTransitionPipeline::EComputeToCompute, NextCountBuffer.UAV);
+			RHICmdList.TransitionResource(EResourceTransitionAccess::EReadable, EResourceTransitionPipeline::EComputeToCompute, CountBuffer.UAV);
+			
 			// Copy the current buffer in the next buffer.
-			TShaderMapRef<FNiagaraCopyIntBufferRegionCS> CopyBufferCS(GetGlobalShaderMap(FeatureLevel));
-			RHICmdList.SetComputeShader(CopyBufferCS->GetComputeShader());
-			FRHIUnorderedAccessView* UAVs[NIAGARA_COPY_BUFFER_BUFFER_COUNT] = {};
-			int32 UsedIndexCounts[NIAGARA_COPY_BUFFER_BUFFER_COUNT] = {};
-			UAVs[0] = NextCountBufferUAVAsInt;
-			UsedIndexCounts[0] = AllocatedInstanceCounts;
-			RHICmdList.TransitionResources(EResourceTransitionAccess::ERWBarrier, EResourceTransitionPipeline::EComputeToCompute, UAVs, 1);
-			CopyBufferCS->SetParameters(RHICmdList, CountBufferSRVAsInt, UAVs, UsedIndexCounts, 0, 1);
-			DispatchComputeShader(RHICmdList, *CopyBufferCS, FMath::DivideAndRoundUp(AllocatedInstanceCounts, NIAGARA_COPY_BUFFER_THREAD_COUNT), 1, 1);
-			RHICmdList.TransitionResources(EResourceTransitionAccess::ERWBarrier, EResourceTransitionPipeline::EComputeToGfx, UAVs, 1);
-			CopyBufferCS->UnbindBuffers(RHICmdList);
+			FRHIUnorderedAccessView* UAVs[] = { NextCountBuffer.UAV };
+			int32 UsedIndexCounts[] = { AllocatedInstanceCounts };
+			CopyUIntBufferToTargets(RHICmdList, FeatureLevel, CountBuffer.SRV, UAVs, UsedIndexCounts, 0, UE_ARRAY_COUNT(UAVs)); 
+
+			// TR-InstanceCount : The counts can either be used in the next simulation or as the input for the next draw indirect task. In both case, the target is for compute.
+			RHICmdList.TransitionResource(EResourceTransitionAccess::ERWBarrier, EResourceTransitionPipeline::EComputeToCompute, NextCountBuffer.UAV);
 
 			// Swap the buffers
 			AllocatedInstanceCounts = RecommendedInstanceCounts;
diff --git a/Engine/Plugins/FX/Niagara/Source/Niagara/Private/NiagaraModule.cpp b/Engine/Plugins/FX/Niagara/Source/Niagara/Private/NiagaraModule.cpp
--- a/Engine/Plugins/FX/Niagara/Source/Niagara/Private/NiagaraModule.cpp
+++ b/Engine/Plugins/FX/Niagara/Source/Niagara/Private/NiagaraModule.cpp
@@ -299,9 +299,11 @@ void INiagaraModule::StartupModule()
 		return FNiagaraTypeRegistry::GetDefaultDataInterfaceByName(DIClassName);
 	}));
 
-	FFXSystemInterface::RegisterCustomFXSystem(NiagaraEmitterInstanceBatcher::Name, FCreateCustomFXSystemDelegate::CreateLambda([](ERHIFeatureLevel::Type InFeatureLevel, EShaderPlatform InShaderPlatform) -> FFXSystemInterface*
+	FFXSystemInterface::RegisterCustomFXSystem(
+		NiagaraEmitterInstanceBatcher::Name, 
+		FCreateCustomFXSystemDelegate::CreateLambda([](ERHIFeatureLevel::Type InFeatureLevel, EShaderPlatform InShaderPlatform, FGPUSortManager* InGPUSortManager) -> FFXSystemInterface*
 	{
-		return new NiagaraEmitterInstanceBatcher(InFeatureLevel, InShaderPlatform);
+		return new NiagaraEmitterInstanceBatcher(InFeatureLevel, InShaderPlatform, InGPUSortManager);
 	}));
 }
 
diff --git a/Engine/Plugins/FX/Niagara/Source/Niagara/Private/NiagaraRendererMeshes.cpp b/Engine/Plugins/FX/Niagara/Source/Niagara/Private/NiagaraRendererMeshes.cpp
--- a/Engine/Plugins/FX/Niagara/Source/Niagara/Private/NiagaraRendererMeshes.cpp
+++ b/Engine/Plugins/FX/Niagara/Source/Niagara/Private/NiagaraRendererMeshes.cpp
@@ -238,7 +238,7 @@ void FNiagaraRendererMeshes::GetDynamicMeshElements(const TArray<const FSceneVie
 	{
 		check(MaterialProxy);
 		EBlendMode BlendMode = MaterialProxy->GetMaterial(FeatureLevel)->GetBlendMode();
-		bHasTranslucentMaterials |= BlendMode == BLEND_AlphaComposite || BlendMode == BLEND_AlphaHoldout || BlendMode == BLEND_Translucent;
+		bHasTranslucentMaterials |= IsTranslucentBlendMode(BlendMode);
 	}
 
 	bool bShouldSort = SortMode != ENiagaraSortMode::None && (bHasTranslucentMaterials || !bSortOnlyWhenTranslucent);
@@ -355,6 +355,7 @@ void FNiagaraRendererMeshes::GetDynamicMeshElements(const TArray<const FSceneVie
 				{
 					SortInfo.ParticleCount = NumInstances;
 					SortInfo.SortMode = SortMode;
+					SortInfo.SetSortFlags(GNiagaraGPUSortingUseMaxPrecision, bHasTranslucentMaterials);
 					if (bCustomSorting)
 					{
 						SortVarIdx = ENiagaraMeshVFLayout::CustomSorting;
@@ -381,18 +382,16 @@ void FNiagaraRendererMeshes::GetDynamicMeshElements(const TArray<const FSceneVie
 					check(ParticleData.IsValid());
 					if (SortInfo.SortMode != ENiagaraSortMode::None && SortInfo.SortAttributeOffset != INDEX_NONE)
 					{
-						if (GNiagaraGPUSorting &&
-							GNiagaraGPUSortingCPUToGPUThreshold >= 0 &&
+						if (GNiagaraGPUSortingCPUToGPUThreshold >= 0 &&
 							SortInfo.ParticleCount >= GNiagaraGPUSortingCPUToGPUThreshold &&
 							FNiagaraUtilities::AllowComputeShaders(Batcher->GetShaderPlatform()))
 						{
 							SortInfo.ParticleCount = NumInstances;
 							SortInfo.ParticleDataFloatSRV = ParticleData.SRV;
 							SortInfo.FloatDataStride = SourceParticleData->GetFloatStride() / sizeof(float);
-							const int32 IndexBufferOffset = Batcher->AddSortedGPUSimulation(SortInfo);
-							if (IndexBufferOffset != INDEX_NONE)
+							if (Batcher->AddSortedGPUSimulation(SortInfo))
 							{
-								CollectorResources.VertexFactory->SetSortedIndices(Batcher->GetGPUSortedBuffer().VertexBufferSRV, IndexBufferOffset);
+								CollectorResources.VertexFactory->SetSortedIndices(SortInfo.AllocationInfo.BufferSRV, SortInfo.AllocationInfo.BufferOffset);
 							}
 						}
 						else
@@ -408,7 +407,7 @@ void FNiagaraRendererMeshes::GetDynamicMeshElements(const TArray<const FSceneVie
 				}
 				else
 				{
-					if (SortInfo.SortMode != ENiagaraSortMode::None && SortInfo.SortAttributeOffset != INDEX_NONE && GNiagaraGPUSorting)
+					if (SortInfo.SortMode != ENiagaraSortMode::None && SortInfo.SortAttributeOffset != INDEX_NONE)
 					{
 						// Here we need to be conservative about the InstanceCount, since the final value is only known on the GPU after the simulation.
 						SortInfo.ParticleCount = SourceParticleData->GetNumInstances();
@@ -416,10 +415,9 @@ void FNiagaraRendererMeshes::GetDynamicMeshElements(const TArray<const FSceneVie
 						SortInfo.FloatDataStride = SourceParticleData->GetFloatStride() / sizeof(float);
 						SortInfo.GPUParticleCountSRV = Batcher->GetGPUInstanceCounterManager().GetInstanceCountBuffer().SRV;
 						SortInfo.GPUParticleCountOffset = SourceParticleData->GetGPUInstanceCountBufferOffset();
-						const int32 IndexBufferOffset = Batcher->AddSortedGPUSimulation(SortInfo);
-						if (IndexBufferOffset != INDEX_NONE && SortInfo.GPUParticleCountOffset != INDEX_NONE)
+						if (Batcher->AddSortedGPUSimulation(SortInfo))
 						{
-							CollectorResources.VertexFactory->SetSortedIndices(Batcher->GetGPUSortedBuffer().VertexBufferSRV, IndexBufferOffset);
+							CollectorResources.VertexFactory->SetSortedIndices(SortInfo.AllocationInfo.BufferSRV, SortInfo.AllocationInfo.BufferOffset);
 						}
 					}
 					if (SourceParticleData->GetGPUBufferFloat().SRV.IsValid())
diff --git a/Engine/Plugins/FX/Niagara/Source/Niagara/Private/NiagaraRendererSprites.cpp b/Engine/Plugins/FX/Niagara/Source/Niagara/Private/NiagaraRendererSprites.cpp
--- a/Engine/Plugins/FX/Niagara/Source/Niagara/Private/NiagaraRendererSprites.cpp
+++ b/Engine/Plugins/FX/Niagara/Source/Niagara/Private/NiagaraRendererSprites.cpp
@@ -334,11 +334,13 @@ void FNiagaraRendererSprites::SetVertexFactoryParticleData(
 		FNiagaraGPUSortInfo SortInfo;
 		SortInfo.SortAttributeOffset = INDEX_NONE;
 		int32 SortingVarIdx = INDEX_NONE;
-		bool bShouldSort = SortMode != ENiagaraSortMode::None && (BlendMode == BLEND_AlphaComposite || BlendMode == BLEND_AlphaHoldout || BlendMode == BLEND_Translucent || !bSortOnlyWhenTranslucent);
+		const bool bHasTranslucentMaterials = IsTranslucentBlendMode(BlendMode);
+		bool bShouldSort = SortMode != ENiagaraSortMode::None && (bHasTranslucentMaterials || !bSortOnlyWhenTranslucent);
 		if (bShouldSort)
 		{
 			SortInfo.ParticleCount = NumInstances;
 			SortInfo.SortMode = SortMode;
+			SortInfo.SetSortFlags(GNiagaraGPUSortingUseMaxPrecision, bHasTranslucentMaterials); 
 			if (SortInfo.SortMode == ENiagaraSortMode::CustomAscending || SortInfo.SortMode == ENiagaraSortMode::CustomDecending)
 			{
 				SortingVarIdx = ENiagaraSpriteVFLayout::CustomSorting;	
@@ -356,29 +358,25 @@ void FNiagaraRendererSprites::SetVertexFactoryParticleData(
 					SortInfo.ViewDirection = SceneProxy->GetLocalToWorld().GetTransposed().TransformVector(SortInfo.ViewDirection);
 				}
 			}
-
 			SortInfo.SortAttributeOffset = VFVariables[SortingVarIdx].GetGPUOffset();
 		}
 
-
-		if (SimTarget == ENiagaraSimTarget::CPUSim)//TODO: Compute shader for sorting gpu sims and larger cpu sims.
+		if (SimTarget == ENiagaraSimTarget::CPUSim)
 		{
 			int32 ParticleStrideInFloats = GbEnableMinimalGPUBuffers ? SourceParticleData->GetNumInstances() : SourceParticleData->GetFloatStride() / sizeof(float);
 			check(CPUSimParticleDataAllocation.ParticleData.IsValid());
 			if (SortInfo.SortMode != ENiagaraSortMode::None && SortInfo.SortAttributeOffset != INDEX_NONE)
 			{
-				if (GNiagaraGPUSorting &&
-					GNiagaraGPUSortingCPUToGPUThreshold >= 0 &&
-					SortInfo.ParticleCount >= GNiagaraGPUSortingCPUToGPUThreshold &&
+				if (GNiagaraGPUSortingCPUToGPUThreshold >= 0 &&
+					SortInfo.ParticleCount >= GNiagaraGPUSortingCPUToGPUThreshold && 
 					FNiagaraUtilities::AllowComputeShaders(Batcher->GetShaderPlatform()))
 				{
 					SortInfo.ParticleCount = NumInstances;
 					SortInfo.ParticleDataFloatSRV = CPUSimParticleDataAllocation.ParticleData.SRV;
 					SortInfo.FloatDataStride = ParticleStrideInFloats;
-					const int32 IndexBufferOffset = Batcher->AddSortedGPUSimulation(SortInfo);
-					if (IndexBufferOffset != INDEX_NONE)
+					if (Batcher->AddSortedGPUSimulation(SortInfo))
 					{
-						OutVertexFactory.SetSortedIndices(Batcher->GetGPUSortedBuffer().VertexBufferSRV, IndexBufferOffset);
+						OutVertexFactory.SetSortedIndices(SortInfo.AllocationInfo.BufferSRV, SortInfo.AllocationInfo.BufferOffset);
 					}
 				}
 				else
@@ -395,7 +393,7 @@ void FNiagaraRendererSprites::SetVertexFactoryParticleData(
 		}
 		else // ENiagaraSimTarget::GPUSim
 		{
-			if (SortInfo.SortMode != ENiagaraSortMode::None && SortInfo.SortAttributeOffset != INDEX_NONE && GNiagaraGPUSorting)
+			if (SortInfo.SortMode != ENiagaraSortMode::None && SortInfo.SortAttributeOffset != INDEX_NONE)
 			{
 				// Here we need to be conservative about the InstanceCount, since the final value is only known on the GPU after the simulation.
 				SortInfo.ParticleCount = SourceParticleData->GetNumInstances();
@@ -404,10 +402,9 @@ void FNiagaraRendererSprites::SetVertexFactoryParticleData(
 				SortInfo.FloatDataStride = SourceParticleData->GetFloatStride() / sizeof(float);
 				SortInfo.GPUParticleCountSRV = Batcher->GetGPUInstanceCounterManager().GetInstanceCountBuffer().SRV;
 				SortInfo.GPUParticleCountOffset = SourceParticleData->GetGPUInstanceCountBufferOffset();
-				const int32 IndexBufferOffset = Batcher->AddSortedGPUSimulation(SortInfo);
-				if (IndexBufferOffset != INDEX_NONE && SortInfo.GPUParticleCountOffset != INDEX_NONE)
+				if (Batcher->AddSortedGPUSimulation(SortInfo))
 				{
-					OutVertexFactory.SetSortedIndices(Batcher->GetGPUSortedBuffer().VertexBufferSRV, IndexBufferOffset);
+					OutVertexFactory.SetSortedIndices(SortInfo.AllocationInfo.BufferSRV, SortInfo.AllocationInfo.BufferOffset);
 				}
 			}
 
diff --git a/Engine/Plugins/FX/Niagara/Source/NiagaraVertexFactories/Private/NiagaraSortingGPU.cpp b/Engine/Plugins/FX/Niagara/Source/NiagaraVertexFactories/Private/NiagaraSortingGPU.cpp
--- a/Engine/Plugins/FX/Niagara/Source/NiagaraVertexFactories/Private/NiagaraSortingGPU.cpp
+++ b/Engine/Plugins/FX/Niagara/Source/NiagaraVertexFactories/Private/NiagaraSortingGPU.cpp
@@ -8,14 +8,6 @@ NiagaraSortingGPU.cpp: Niagara sorting shaders
 #include "NiagaraGPUSortInfo.h"
 #include "ShaderParameterUtils.h"
 
-int32 GNiagaraGPUSorting = 1;
-static FAutoConsoleVariableRef CVarNiagaraGPUSorting(
-	TEXT("Niagara.GPUSorting"),
-	GNiagaraGPUSorting,
-	TEXT("Whether to sort particles on the GPU"),
-	ECVF_Default
-);
-
 int32 GNiagaraGPUSortingUseMaxPrecision = 0;
 static FAutoConsoleVariableRef CVarNiagaraGPUSortinUseMaxPrecision(
 	TEXT("Niagara.GPUSorting.UseMaxPrecision"),
@@ -32,46 +24,6 @@ static FAutoConsoleVariableRef CVarNiagaraGPUSortinCPUToGPUThreshold(
 	ECVF_Default
 );
 
-float GNiagaraGPUSortingBufferSlack = 2.f;
-static FAutoConsoleVariableRef CVarNiagaraGPUSortingBufferSlack(
-	TEXT("Niagara.GPUSorting.BufferSlack"),
-	GNiagaraGPUSortingBufferSlack,
-	TEXT("Slack ratio when allocating GPU sort buffer (default=2)"),
-	ECVF_Default
-);
-
-int32 GNiagaraGPUSortingMinBufferSize = 8192;
-static FAutoConsoleVariableRef CVarNiagaraGPUSortingMinBufferSize(
-	TEXT("Niagara.GPUSorting.MinBufferSize"),
-	GNiagaraGPUSortingMinBufferSize,
-	TEXT("Minimum GPU sort buffer size, in particles (default=8192)"),
-	ECVF_Default
-);
-
-int32 GNiagaraGPUSortingFrameCountBeforeBufferShrinking = 100;
-static FAutoConsoleVariableRef CVarNiagaraGPUSortingFrameCountBeforeBufferShrinking(
-	TEXT("Niagara.GPUSorting.FrameCountBeforeShrinking"),
-	GNiagaraGPUSortingFrameCountBeforeBufferShrinking,
-	TEXT("Number of consecutive frames where the GPU sort buffer is considered oversized before allowing shrinking. (default=100)"),
-	ECVF_Default
-);
-
-class FNiagaraSortingDummyUAV : public FRenderResource
-{
-public:
-	FRWBuffer Buffer;
-
-	virtual void InitRHI() override
-	{
-		Buffer.Initialize(sizeof(int32), 1, EPixelFormat::PF_R32_SINT, BUF_Static, TEXT("FNiagaraSortingDummyUAV"));
-	}
-
-	virtual void ReleaseRHI() override
-	{
-		Buffer.Release();
-	}
-}; 
-
 IMPLEMENT_GLOBAL_SHADER(FNiagaraSortKeyGenCS, "/Plugin/FX/Niagara/Private/NiagaraSortKeyGen.usf", "GenerateParticleSortKeys", SF_Compute);
 
 void FNiagaraSortKeyGenCS::ModifyCompilationEnvironment(const FGlobalShaderPermutationParameters& Parameters, FShaderCompilerEnvironment& OutEnvironment)
@@ -154,70 +106,4 @@ void FNiagaraSortKeyGenCS::UnbindBuffers(FRHICommandList& RHICmdList)
 	SetUAVParameter(RHICmdList, ComputeShaderRHI, OutParticleIndices, nullptr);
 }
 
-//////////////////////////////////////////////////////////////////////////////////////////////////////////////////////
-
-IMPLEMENT_SHADER_TYPE(, FNiagaraCopyIntBufferRegionCS, TEXT("/Plugin/FX/Niagara/Private/NiagaraCopyIntBuffer.usf"), TEXT("MainCS"), SF_Compute);
-
-FNiagaraCopyIntBufferRegionCS::FNiagaraCopyIntBufferRegionCS(const ShaderMetaType::CompiledShaderInitializerType& Initializer)
-	: FGlobalShader(Initializer)
-{
-	CopyParams.Bind(Initializer.ParameterMap, TEXT("CopyParams"));
-	SourceData.Bind(Initializer.ParameterMap, TEXT("SourceData"));
-	DestData[0].Bind(Initializer.ParameterMap, TEXT("DestData0"));
-	DestData[1].Bind(Initializer.ParameterMap, TEXT("DestData1"));
-	DestData[2].Bind(Initializer.ParameterMap, TEXT("DestData2"));
-}
-
-bool FNiagaraCopyIntBufferRegionCS::Serialize(FArchive& Ar)
-{
-	bool bShaderHasOutdatedParameters = FGlobalShader::Serialize(Ar);
-	Ar << CopyParams;
-	Ar << SourceData;
-	Ar << DestData[0];
-	Ar << DestData[1];
-	Ar << DestData[2];
-	return bShaderHasOutdatedParameters;
-}
-
-void FNiagaraCopyIntBufferRegionCS::SetParameters(
-	FRHICommandList& RHICmdList,
-	FRHIShaderResourceView* InSourceData,
-	FRHIUnorderedAccessView* const* InDestDatas,
-	const int32* InUsedIndexCounts,
-	int32 StartingIndex,
-	int32 DestCount)
-{
-	static TGlobalResource<FNiagaraSortingDummyUAV> NiagaraSortingDummyUAV[NIAGARA_COPY_BUFFER_BUFFER_COUNT];
-		
-	FRHIComputeShader* ComputeShaderRHI = GetComputeShader();
-	check(DestCount > 0 && DestCount <= NIAGARA_COPY_BUFFER_BUFFER_COUNT);
-
-	RHICmdList.SetShaderResourceViewParameter(ComputeShaderRHI, SourceData.GetBaseIndex(), InSourceData);
-
-	FUintVector4 CopyParamsValue(StartingIndex, 0, 0, 0);
-	for (int32 Index = 0; Index < DestCount; ++Index)
-	{
-		SetUAVParameter(RHICmdList, ComputeShaderRHI, DestData[Index], InDestDatas[Index]);
-		CopyParamsValue[Index + 1] = InUsedIndexCounts[Index];
-	}
-
-	for (int32 Index = DestCount; Index < NIAGARA_COPY_BUFFER_BUFFER_COUNT; ++Index)
-	{
-		RHICmdList.TransitionResource(EResourceTransitionAccess::ERWNoBarrier, EResourceTransitionPipeline::EComputeToCompute, NiagaraSortingDummyUAV[Index].Buffer.UAV);
-		SetUAVParameter(RHICmdList, ComputeShaderRHI, DestData[Index], NiagaraSortingDummyUAV[Index].Buffer.UAV);
-	}
-
-	SetShaderValue(RHICmdList, ComputeShaderRHI, CopyParams, CopyParamsValue);
-}
-
-void FNiagaraCopyIntBufferRegionCS::UnbindBuffers(FRHICommandList& RHICmdList)
-{
-	FRHIComputeShader* ComputeShaderRHI = GetComputeShader();
-	SetSRVParameter(RHICmdList, ComputeShaderRHI, SourceData, nullptr);
-	for (int32 Index = 0; Index < NIAGARA_COPY_BUFFER_BUFFER_COUNT; ++Index)
-	{
-		SetUAVParameter(RHICmdList, ComputeShaderRHI, DestData[Index], nullptr);
-	}
-}
-
 
diff --git a/Engine/Plugins/FX/Niagara/Source/NiagaraVertexFactories/Public/NiagaraSortingGPU.h b/Engine/Plugins/FX/Niagara/Source/NiagaraVertexFactories/Public/NiagaraSortingGPU.h
--- a/Engine/Plugins/FX/Niagara/Source/NiagaraVertexFactories/Public/NiagaraSortingGPU.h
+++ b/Engine/Plugins/FX/Niagara/Source/NiagaraVertexFactories/Public/NiagaraSortingGPU.h
@@ -12,12 +12,8 @@ NiagaraSortingGPU.h: Niagara sorting shaders
 
 struct FNiagaraGPUSortInfo;
 
-extern NIAGARAVERTEXFACTORIES_API int32 GNiagaraGPUSorting;
 extern NIAGARAVERTEXFACTORIES_API int32 GNiagaraGPUSortingUseMaxPrecision;
 extern NIAGARAVERTEXFACTORIES_API int32 GNiagaraGPUSortingCPUToGPUThreshold;
-extern NIAGARAVERTEXFACTORIES_API float GNiagaraGPUSortingBufferSlack;
-extern NIAGARAVERTEXFACTORIES_API int32 GNiagaraGPUSortingMinBufferSize;
-extern NIAGARAVERTEXFACTORIES_API int32 GNiagaraGPUSortingFrameCountBeforeBufferShrinking;
 
 #define NIAGARA_KEY_GEN_THREAD_COUNT 64
 #define NIAGARA_COPY_BUFFER_THREAD_COUNT 64
@@ -82,58 +78,3 @@ private:
 	/** Output indices buffer. */
 	FShaderResourceParameter OutParticleIndices;
 };
-
-/**
- * Compute shader used to copy a reference buffer split in several buffers.
- * Each buffer contains a segment of the buffer, and each buffer is increasingly bigger to hold the next part.
- * Could alternatively use DMA copy, if the RHI provided a way to copy buffer regions.
- */
-class NIAGARAVERTEXFACTORIES_API FNiagaraCopyIntBufferRegionCS : public FGlobalShader
-{
-	DECLARE_SHADER_TYPE(FNiagaraCopyIntBufferRegionCS,Global);
-
-public:
-
-	static bool ShouldCompilePermutation(const FGlobalShaderPermutationParameters& Parameters)
-	{
-		return RHISupportsComputeShaders(Parameters.Platform);
-	}
-
-	static void ModifyCompilationEnvironment(const FGlobalShaderPermutationParameters& Parameters, FShaderCompilerEnvironment& OutEnvironment)
-	{
-		FGlobalShader::ModifyCompilationEnvironment(Parameters, OutEnvironment);
-		OutEnvironment.SetDefine(TEXT("THREAD_COUNT"), NIAGARA_COPY_BUFFER_THREAD_COUNT);
-		OutEnvironment.SetDefine(TEXT("BUFFER_COUNT"), NIAGARA_COPY_BUFFER_BUFFER_COUNT);
-	}
-
-	/** Default constructor. */
-	FNiagaraCopyIntBufferRegionCS() {}
-
-	/** Initialization constructor. */
-	explicit FNiagaraCopyIntBufferRegionCS(const ShaderMetaType::CompiledShaderInitializerType& Initializer);
-
-	/** Serialization. */
-	virtual bool Serialize(FArchive& Ar) override;
-
-	/**
-	 * Set parameters.
-	 */
-	void SetParameters(
-		FRHICommandList& RHICmdList,	
-		FRHIShaderResourceView* InSourceData,
-		FRHIUnorderedAccessView* const* InDestDatas,
-		const int32* InUsedIndexCounts, 
-		int32 StartingIndex,
-		int32 DestCount);
-
-	/**
-	 * Unbinds any buffers that have been bound.
-	 */
-	void UnbindBuffers(FRHICommandList& RHICmdList);
-
-private:
-
-	FShaderParameter CopyParams;
-	FShaderResourceParameter SourceData;
-	FShaderResourceParameter DestData[NIAGARA_COPY_BUFFER_BUFFER_COUNT];
-};
diff --git a/Engine/Shaders/Private/CopyUIntBuffer.usf b/Engine/Shaders/Private/CopyUIntBuffer.usf
new file mode 100644
--- /dev/null
+++ b/Engine/Shaders/Private/CopyUIntBuffer.usf
@@ -0,0 +1,57 @@
+// Copyright Epic Games, Inc. All Rights Reserved.
+
+/*==============================================================================
+	CopyUIntBuffer.usf: Shader to copy one uint buffer into several overlapping targets.
+	Each target is expected to be bigger than the previous one, so that copy resumes where it left.
+	Ex : (ABCDEFGHIJK) -> (ABC--, ---DEFGHI-, ---------JK---)
+
+	This shader is essentially used to manage growable buffers in cases the
+	binding needs to be set before the final size required is known. 
+	In this scenario, the smaller buffers are temporary and only the final (biggest) buffer
+	becomes persistent, this is why it has apparently unused space at the beginning but
+	it will be used the next frame to hold all of the data.
+==============================================================================*/
+
+/*------------------------------------------------------------------------------
+	Compile time parameters:
+		THREAD_COUNT - The number of threads to launch per workgroup.
+------------------------------------------------------------------------------*/
+
+#include "/Engine/Private/Common.ush"
+
+
+// (StartingIndex, UsedIndexCount0, UsedIndexCount1, UsedIndexCount2)
+uint4 CopyParams;
+
+
+Buffer<uint> SourceData;
+
+RWBuffer<uint> DestData0;
+RWBuffer<uint> DestData1;
+RWBuffer<uint> DestData2;
+
+[numthreads(THREAD_COUNT,1,1)]
+void MainCS(uint IndexOffset : SV_DispatchThreadID)
+{
+	const uint StartingIndex = CopyParams.x;
+	const uint UsedIndexCount0 = CopyParams.y;
+	const uint UsedIndexCount1 = CopyParams.z;
+	const uint UsedIndexCount2 = CopyParams.w;
+
+	uint Index = StartingIndex + IndexOffset;
+
+	// The source buffer copy is splitted between the destination buffers.
+	if (Index < UsedIndexCount0)
+	{
+		DestData0[Index] = SourceData[Index];
+	}
+	else if (Index < UsedIndexCount1)
+	{
+		DestData1[Index] = SourceData[Index];
+	}
+	else if (Index < UsedIndexCount2)
+	{
+		DestData2[Index] = SourceData[Index];
+	}
+}
+ 
\ No newline at end of file
diff --git a/Engine/Source/Runtime/Core/Private/HAL/LowLevelMemTracker.cpp b/Engine/Source/Runtime/Core/Private/HAL/LowLevelMemTracker.cpp
--- a/Engine/Source/Runtime/Core/Private/HAL/LowLevelMemTracker.cpp
+++ b/Engine/Source/Runtime/Core/Private/HAL/LowLevelMemTracker.cpp
@@ -83,6 +83,7 @@ DECLARE_LLM_MEMORY_STAT(TEXT("StaticMesh"), STAT_StaticMeshLLM, STATGROUP_LLMFUL
 DECLARE_LLM_MEMORY_STAT(TEXT("Materials"), STAT_MaterialsLLM, STATGROUP_LLMFULL);
 DECLARE_LLM_MEMORY_STAT(TEXT("Particles"), STAT_ParticlesLLM, STATGROUP_LLMFULL);
 DECLARE_LLM_MEMORY_STAT(TEXT("Niagara"), STAT_NiagaraLLM, STATGROUP_LLMFULL);
+DECLARE_LLM_MEMORY_STAT(TEXT("GPUSort"), STAT_GPUSortLLM, STATGROUP_LLMFULL);
 DECLARE_LLM_MEMORY_STAT(TEXT("GC"), STAT_GCLLM, STATGROUP_LLMFULL);
 DECLARE_LLM_MEMORY_STAT(TEXT("UI"), STAT_UILLM, STATGROUP_LLMFULL);
 DECLARE_LLM_MEMORY_STAT(TEXT("NavigationRecast"), STAT_NavigationRecastLLM, STATGROUP_LLMFULL);
diff --git a/Engine/Source/Runtime/Core/Public/HAL/LowLevelMemTracker.h b/Engine/Source/Runtime/Core/Public/HAL/LowLevelMemTracker.h
--- a/Engine/Source/Runtime/Core/Public/HAL/LowLevelMemTracker.h
+++ b/Engine/Source/Runtime/Core/Public/HAL/LowLevelMemTracker.h
@@ -182,6 +182,7 @@ enum class ELLMTagSet : uint8
 	macro(Materials,							"Materials",					GET_STATFNAME(STAT_MaterialsLLM),							GET_STATFNAME(STAT_MaterialsSummaryLLM),		-1)\
 	macro(Particles,							"Particles",					GET_STATFNAME(STAT_ParticlesLLM),							GET_STATFNAME(STAT_ParticlesSummaryLLM),		-1)\
 	macro(Niagara,								"Niagara",						GET_STATFNAME(STAT_NiagaraLLM),								GET_STATFNAME(STAT_NiagaraSummaryLLM),			-1)\
+	macro(GPUSort,								"GPUSort",						GET_STATFNAME(STAT_GPUSortLLM),								GET_STATFNAME(STAT_EngineSummaryLLM),			-1)\
 	macro(GC,									"GC",							GET_STATFNAME(STAT_GCLLM),									GET_STATFNAME(STAT_EngineSummaryLLM),			-1)\
 	macro(UI,									"UI",							GET_STATFNAME(STAT_UILLM),									GET_STATFNAME(STAT_UISummaryLLM),				-1)\
 	macro(NavigationRecast,						"NavigationRecast",				GET_STATFNAME(STAT_NavigationRecastLLM),					GET_STATFNAME(STAT_NavigationSummaryLLM),		-1)\
diff --git a/Engine/Source/Runtime/Engine/Private/GPUSort.cpp b/Engine/Source/Runtime/Engine/Private/GPUSort.cpp
--- a/Engine/Source/Runtime/Engine/Private/GPUSort.cpp
+++ b/Engine/Source/Runtime/Engine/Private/GPUSort.cpp
@@ -714,7 +714,7 @@ int32 SortGPUBuffers(FRHICommandListImmediate& RHICmdList, FGPUSortBuffers SortB
 
 	check(RHISupportsComputeShaders(GShaderPlatformForFeatureLevel[FeatureLevel]));
 
-	SCOPED_DRAW_EVENTF(RHICmdList, SortGPU, TEXT("SortGPU_%d"), Count);
+	SCOPED_DRAW_EVENTF(RHICmdList, SortGPU, TEXT("Sort(%d)"), Count);
 
 	// Determine how many tiles need to be sorted.
 	const int32 TileCount = Count / TILE_SIZE;
@@ -830,15 +830,24 @@ int32 SortGPUBuffers(FRHICommandListImmediate& RHICmdList, FGPUSortBuffers SortB
 			//UAV is going to SRV, so transition to Readable.
 			RHICmdList.TransitionResource(EResourceTransitionAccess::EReadable, EResourceTransitionPipeline::EComputeToCompute, GSortOffsetBuffers.BufferUAVs[1]);
 
+
 			FRHIUnorderedAccessView* PrePhase3BarrierUAVS[2];
 			PrePhase3BarrierUAVS[0] = SortBuffers.RemoteKeyUAVs[BufferIndex ^ 0x1];
 			PrePhase3BarrierUAVS[1] = SortBuffers.RemoteValueUAVs[BufferIndex ^ 0x1];
+
 			RHICmdList.TransitionResources(EResourceTransitionAccess::ERWBarrier, EResourceTransitionPipeline::EComputeToCompute, PrePhase3BarrierUAVS, 2);
 
+			const bool bIsLastPass = ((PassBits << RADIX_BITS) & KeyMask) == 0;
 			// Phase 3: Downsweep to compute final offsets and scatter keys.
 			RHICmdList.SetComputeShader(DownsweepCS->GetComputeShader());
-			DownsweepCS->SetOutput(RHICmdList, SortBuffers.RemoteKeyUAVs[BufferIndex ^ 0x1], SortBuffers.RemoteValueUAVs[BufferIndex ^ 0x1]);
-			DownsweepCS->SetParameters(RHICmdList, SortBuffers.RemoteKeySRVs[BufferIndex], SortBuffers.RemoteValueSRVs[BufferIndex], GSortOffsetBuffers.BufferSRVs[1], SortUniformBufferRef, GRadixSortParametersBuffer.SortParametersBufferSRV );
+			{
+				FRHIUnorderedAccessView* ValuesUAV = (bIsLastPass && SortBuffers.FinalValuesUAV)? SortBuffers.FinalValuesUAV : SortBuffers.RemoteValueUAVs[BufferIndex ^ 0x1];
+				DownsweepCS->SetOutput(RHICmdList, SortBuffers.RemoteKeyUAVs[BufferIndex ^ 0x1], ValuesUAV);
+			}
+			{
+				FRHIShaderResourceView* ValuesSRV = (PassIndex == 0 && SortBuffers.FirstValuesSRV) ? SortBuffers.FirstValuesSRV : SortBuffers.RemoteValueSRVs[BufferIndex];
+				DownsweepCS->SetParameters(RHICmdList, SortBuffers.RemoteKeySRVs[BufferIndex], ValuesSRV, GSortOffsetBuffers.BufferSRVs[1], SortUniformBufferRef, GRadixSortParametersBuffer.SortParametersBufferSRV );
+			}
 			DispatchComputeShader(RHICmdList, *DownsweepCS, GroupCount, 1, 1 );
 			DownsweepCS->UnbindBuffers(RHICmdList);
 
diff --git a/Engine/Source/Runtime/Engine/Private/GPUSort.h b/Engine/Source/Runtime/Engine/Private/GPUSort.h
--- a/Engine/Source/Runtime/Engine/Private/GPUSort.h
+++ b/Engine/Source/Runtime/Engine/Private/GPUSort.h
@@ -24,6 +24,11 @@ struct FGPUSortBuffers
 	/** Unordered access views for vertex buffers containing the values. */
 	FRHIUnorderedAccessView* RemoteValueUAVs[2];
 
+	/** Shader resource view holding the initial state of the values. */
+	FRHIShaderResourceView* FirstValuesSRV = nullptr;
+	/** Unordered access view holding the final state of the value. */
+	FRHIUnorderedAccessView* FinalValuesUAV = nullptr;
+
 	/** Default constructor. */
 	FGPUSortBuffers()
 	{
@@ -44,7 +49,7 @@ int32 GetGPUSortPassCount(uint32 KeyMask);
  * @param Count - How many items in the buffer need to be sorted.
  * @returns The index of the buffer containing sorted results.
  */
-ENGINE_API int32 SortGPUBuffers(FRHICommandListImmediate& RHICmdList, FGPUSortBuffers SortBuffers, int32 BufferIndex, uint32 KeyMask, int32 Count, ERHIFeatureLevel::Type FeatureLevel);
+int32 SortGPUBuffers(FRHICommandListImmediate& RHICmdList, FGPUSortBuffers SortBuffers, int32 BufferIndex, uint32 KeyMask, int32 Count, ERHIFeatureLevel::Type FeatureLevel);
 
 /**
  * GPU sorting tests.
diff --git a/Engine/Source/Runtime/Engine/Private/GPUSortManager.cpp b/Engine/Source/Runtime/Engine/Private/GPUSortManager.cpp
new file mode 100644
--- /dev/null
+++ b/Engine/Source/Runtime/Engine/Private/GPUSortManager.cpp
@@ -0,0 +1,850 @@
+// Copyright Epic Games, Inc. All Rights Reserved.
+
+/*=============================================================================
+	GPUSortManager.h: Interface manage sorting between different FXSystemInterface
+=============================================================================*/
+
+#include "GPUSortManager.h"
+#include "GPUSort.h"
+#include "GlobalShader.h"
+#include "ShaderPermutation.h"
+#include "ShaderParameterUtils.h"
+
+//*****************************************************************************
+
+int32 GGPUSortFrameCountBeforeBufferShrinking = 100;
+static FAutoConsoleVariableRef CVarGPUSortFrameCountBeforeBufferShrinking(
+	TEXT("fx.GPUSort.FrameCountBeforeShrinking"),
+	GGPUSortFrameCountBeforeBufferShrinking,
+	TEXT("Number of consecutive frames where the GPU sort buffer is considered oversized before allowing shrinking. (default=100)"),
+	ECVF_Default
+);
+
+int32 GGPUSortMinBufferSize = 8192;
+static FAutoConsoleVariableRef CVarGPUSortMinBufferSize(
+	TEXT("fx.GPUSort.MinBufferSize"),
+	GGPUSortMinBufferSize,
+	TEXT("Minimum GPU sort buffer size, in particles (default=8192)"),
+	ECVF_Default
+);
+
+float GGPUSortBufferSlack = 2.f;
+static FAutoConsoleVariableRef CVarGPUSortBufferSlack(
+	TEXT("fx.GPUSort.BufferSlack"),
+	GGPUSortBufferSlack,
+	TEXT("Slack ratio when resizing GPU sort buffers. Must be bigger than 1 (default=2)"),
+	ECVF_Default
+);
+
+int32 GGPUSortEnable = 1;
+static FAutoConsoleVariableRef CVarGPUSortEnable(
+	TEXT("fx.AllowGPUSorting"),
+	GGPUSortEnable,
+	TEXT("Whether to enable GPU sort (default 1)"),
+	ECVF_Default
+);
+
+int32 GGPUSortStressTest = 0;
+static FAutoConsoleVariableRef CVarGPUSortStressTest(
+	TEXT("fx.GPUSort.StressTest"),
+	GGPUSortStressTest,
+	TEXT("Force a stress test on the GPU sort by release persistent data every frame (default=0)"),
+	ECVF_Cheat
+);
+
+
+//*****************************************************************************
+
+DECLARE_GPU_STAT_NAMED(GPUKeyGenAndSort, TEXT("GPU KeyGen & Sort"));
+
+//*****************************************************************************
+//****************************** FGPUSortDummyUAV *****************************
+//*****************************************************************************
+
+class FGPUSortDummyUAV : public FRenderResource
+{
+public:
+	FRWBuffer Buffer;
+
+	virtual void InitRHI() override
+	{
+		Buffer.Initialize(sizeof(int32), 1, EPixelFormat::PF_R32_UINT, BUF_Static, TEXT("FGPUSortDummyUAV"));
+	}
+
+	virtual void ReleaseRHI() override
+	{
+		Buffer.Release();
+	}
+}; 
+
+//*****************************************************************************
+//****************************** FCopyUIntBufferCS *****************************
+//*****************************************************************************
+
+#define COPYUINTCS_BUFFER_COUNT 3
+#define COPYUINTCS_THREAD_COUNT 64
+
+/**
+ * Compute shader used to copy a reference buffer split in several buffers.
+ * Each buffer contains a segment of the buffer, and each buffer is increasingly bigger to hold the next part.
+ * Could alternatively use DMA copy, if the RHI provided a way to copy buffer regions.
+ */
+class FCopyUIntBufferCS : public FGlobalShader
+{
+	DECLARE_SHADER_TYPE(FCopyUIntBufferCS,Global);
+
+public:
+
+	static bool ShouldCompilePermutation(const FGlobalShaderPermutationParameters& Parameters)
+	{
+		return RHISupportsComputeShaders(Parameters.Platform);
+	}
+
+	static void ModifyCompilationEnvironment(const FGlobalShaderPermutationParameters& Parameters, FShaderCompilerEnvironment& OutEnvironment)
+	{
+		FGlobalShader::ModifyCompilationEnvironment(Parameters, OutEnvironment);
+		OutEnvironment.SetDefine(TEXT("THREAD_COUNT"), COPYUINTCS_THREAD_COUNT);
+		OutEnvironment.SetDefine(TEXT("BUFFER_COUNT"), COPYUINTCS_BUFFER_COUNT);
+	}
+
+	/** Default constructor. */
+	FCopyUIntBufferCS() {}
+
+	/** Initialization constructor. */
+	explicit FCopyUIntBufferCS(const ShaderMetaType::CompiledShaderInitializerType& Initializer);
+
+	/** Serialization. */
+	virtual bool Serialize(FArchive& Ar) override;
+
+	/**
+	 * Set parameters.
+	 */
+	void SetParameters(
+		FRHICommandList& RHICmdList,	
+		FRHIShaderResourceView* InSourceData,
+		FRHIUnorderedAccessView* const* InDestDatas,
+		const int32* InUsedIndexCounts, 
+		int32 StartingIndex,
+		int32 DestCount);
+
+	/**
+	 * Unbinds any buffers that have been bound.
+	 */
+	void UnbindBuffers(FRHICommandList& RHICmdList);
+
+private:
+
+	FShaderParameter CopyParams;
+	FShaderResourceParameter SourceData;
+	FShaderResourceParameter DestData[COPYUINTCS_BUFFER_COUNT];
+};
+
+//*****************************************************************************
+
+IMPLEMENT_SHADER_TYPE(, FCopyUIntBufferCS, TEXT("/Engine/Private/CopyUIntBuffer.usf"), TEXT("MainCS"), SF_Compute);
+
+FCopyUIntBufferCS::FCopyUIntBufferCS(const ShaderMetaType::CompiledShaderInitializerType& Initializer)
+	: FGlobalShader(Initializer)
+{
+	CopyParams.Bind(Initializer.ParameterMap, TEXT("CopyParams"));
+	SourceData.Bind(Initializer.ParameterMap, TEXT("SourceData"));
+	DestData[0].Bind(Initializer.ParameterMap, TEXT("DestData0"));
+	DestData[1].Bind(Initializer.ParameterMap, TEXT("DestData1"));
+	DestData[2].Bind(Initializer.ParameterMap, TEXT("DestData2"));
+}
+
+bool FCopyUIntBufferCS::Serialize(FArchive& Ar)
+{
+	bool bShaderHasOutdatedParameters = FGlobalShader::Serialize(Ar);
+	Ar << CopyParams;
+	Ar << SourceData;
+	Ar << DestData[0];
+	Ar << DestData[1];
+	Ar << DestData[2];
+	return bShaderHasOutdatedParameters;
+}
+
+void FCopyUIntBufferCS::SetParameters(
+	FRHICommandList& RHICmdList,
+	FRHIShaderResourceView* InSourceData,
+	FRHIUnorderedAccessView* const* InDestDatas,
+	const int32* InUsedIndexCounts,
+	int32 StartingIndex,
+	int32 DestCount)
+{
+	static TGlobalResource<FGPUSortDummyUAV> NiagaraSortingDummyUAV[COPYUINTCS_BUFFER_COUNT];
+		
+	FRHIComputeShader* ComputeShaderRHI = GetComputeShader();
+	check(DestCount > 0 && DestCount <= COPYUINTCS_BUFFER_COUNT);
+
+	RHICmdList.SetShaderResourceViewParameter(ComputeShaderRHI, SourceData.GetBaseIndex(), InSourceData);
+
+	FUintVector4 CopyParamsValue(StartingIndex, 0, 0, 0);
+	for (int32 Index = 0; Index < DestCount; ++Index)
+	{
+		SetUAVParameter(RHICmdList, ComputeShaderRHI, DestData[Index], InDestDatas[Index]);
+		CopyParamsValue[Index + 1] = InUsedIndexCounts[Index];
+	}
+
+	for (int32 Index = DestCount; Index < COPYUINTCS_BUFFER_COUNT; ++Index)
+	{
+		// TR-DummyUAVs : those buffers are only ever used here, but there content is never accessed.
+		RHICmdList.TransitionResource(EResourceTransitionAccess::ERWNoBarrier, EResourceTransitionPipeline::EComputeToCompute, NiagaraSortingDummyUAV[Index].Buffer.UAV);
+		SetUAVParameter(RHICmdList, ComputeShaderRHI, DestData[Index], NiagaraSortingDummyUAV[Index].Buffer.UAV);
+	}
+
+	SetShaderValue(RHICmdList, ComputeShaderRHI, CopyParams, CopyParamsValue);
+}
+
+void FCopyUIntBufferCS::UnbindBuffers(FRHICommandList& RHICmdList)
+{
+	FRHIComputeShader* ComputeShaderRHI = GetComputeShader();
+	SetSRVParameter(RHICmdList, ComputeShaderRHI, SourceData, nullptr);
+	for (int32 Index = 0; Index < COPYUINTCS_BUFFER_COUNT; ++Index)
+	{
+		SetUAVParameter(RHICmdList, ComputeShaderRHI, DestData[Index], nullptr);
+	}
+}
+
+//*****************************************************************************
+
+void CopyUIntBufferToTargets(FRHICommandListImmediate& RHICmdList, ERHIFeatureLevel::Type FeatureLevel,  FRHIShaderResourceView* SourceSRV, FRHIUnorderedAccessView*const* TargetUAVs, int32* TargetSizes, int32 StartingOffset, int32 NumTargets)
+{
+	// No that resource transition must be made outside this call as we don't know how the content of those have been generated, and will be used.
+
+	TShaderMapRef<FCopyUIntBufferCS> CopyBufferCS(GetGlobalShaderMap(FeatureLevel));
+	RHICmdList.SetComputeShader(CopyBufferCS->GetComputeShader());
+	
+	int32 Index0InPass = 0;
+	while (Index0InPass < NumTargets)
+	{
+		const int32 NumTargetsInPass = FMath::Min<int32>(NumTargets - Index0InPass, COPYUINTCS_BUFFER_COUNT);
+		const int32 NumElementsInPass = TargetSizes[Index0InPass + NumTargetsInPass - 1] - StartingOffset;
+
+		CopyBufferCS->SetParameters(RHICmdList, SourceSRV, TargetUAVs + Index0InPass, TargetSizes + Index0InPass, StartingOffset, NumTargetsInPass);
+		DispatchComputeShader(RHICmdList, *CopyBufferCS, FMath::DivideAndRoundUp(NumElementsInPass, COPYUINTCS_THREAD_COUNT), 1, 1);
+
+		StartingOffset += NumElementsInPass;
+		Index0InPass += COPYUINTCS_BUFFER_COUNT;
+	};
+	
+	CopyBufferCS->UnbindBuffers(RHICmdList);
+}
+
+//*****************************************************************************
+//*********************** FGPUSortManager::FKeyGenInfo ************************
+//*****************************************************************************
+
+FGPUSortManager::FKeyGenInfo::FKeyGenInfo(uint32 NumElements, bool bHighPrecisionKeys)
+{
+	ElementKeyMask = (1 << FMath::CeilLogTwo(NumElements)) - 1;
+	ElementKeyShift = 16;
+	SortKeyMask = 0xFFFF;
+
+	SortKeyParams.X = SortKeyMask;
+	SortKeyParams.Y = 0;
+	SortKeyParams.Z = 0x8000;
+	SortKeyParams.W = 0; 
+
+	if (bHighPrecisionKeys)
+	{
+		ElementKeyMask = FMath::Max<uint32>(ElementKeyMask , 1); // Need at list 1 bit for the above logic
+		uint32 UnusedBits = FPlatformMath::CountLeadingZeros(ElementKeyMask << ElementKeyShift);
+		ElementKeyShift += UnusedBits;
+		SortKeyMask = ~(ElementKeyMask << ElementKeyShift);
+
+		SortKeyParams.X = SortKeyMask;
+		SortKeyParams.Y = 16 - UnusedBits;
+		SortKeyParams.Z <<= UnusedBits;
+	}
+}
+
+//*****************************************************************************
+//*********************** FGPUSortManager::FValueBuffer ***********************
+//*****************************************************************************
+
+FGPUSortManager::FValueBuffer::FValueBuffer(int32 InAllocatedCount, int32 InUsedCount, const FGPUSortManager::FSettings& InSettings)
+	: AllocatedCount(InAllocatedCount)
+	, UsedCount(InUsedCount)
+{
+	check(InUsedCount >= 0 && InUsedCount <= InAllocatedCount);
+
+	FRHIResourceCreateInfo CreateInfo;
+	VertexBufferRHI = RHICreateVertexBuffer((uint32)InAllocatedCount * sizeof(uint32), BUF_Static | BUF_ShaderResource | BUF_UnorderedAccess, CreateInfo);
+
+	UInt32SRV = RHICreateShaderResourceView(VertexBufferRHI, sizeof(uint32), PF_R32_UINT);
+	UInt32UAV = RHICreateUnorderedAccessView(VertexBufferRHI, PF_R32_UINT);
+
+	if (EnumHasAnyFlags(InSettings.AllowedFlags, EGPUSortFlags::ValuesAsInt32))
+	{
+		Int32SRV = RHICreateShaderResourceView(VertexBufferRHI, sizeof(int32), PF_R32_SINT);
+		Int32UAV = RHICreateUnorderedAccessView(VertexBufferRHI, PF_R32_SINT);
+	}
+	if (EnumHasAnyFlags(InSettings.AllowedFlags, EGPUSortFlags::ValuesAsG16R16F))
+	{
+		G16R16SRV = RHICreateShaderResourceView(VertexBufferRHI, sizeof(FFloat16) * 2, PF_G16R16F);
+		G16R16UAV = RHICreateUnorderedAccessView(VertexBufferRHI, PF_G16R16F);
+	}
+}
+
+void FGPUSortManager::FValueBuffer::Allocate(FAllocationInfo& OutInfo, int32 ValueCount, EGPUSortFlags Flags)
+{
+	OutInfo.BufferOffset = UsedCount;
+	UsedCount += ValueCount;
+	check(UsedCount <= AllocatedCount);
+	if (EnumHasAnyFlags(Flags, EGPUSortFlags::ValuesAsInt32))
+	{
+		OutInfo.BufferSRV = Int32SRV;
+	}
+	else
+	{
+		checkSlow(EnumHasAnyFlags(Flags, EGPUSortFlags::ValuesAsG16R16F));
+		OutInfo.BufferSRV = G16R16SRV;
+	}
+}
+
+void FGPUSortManager::FValueBuffer::ReleaseRHI()
+{
+	UInt32SRV.SafeRelease();
+	UInt32UAV.SafeRelease();
+	Int32SRV.SafeRelease();
+	Int32UAV.SafeRelease();
+	G16R16SRV.SafeRelease();
+	G16R16UAV.SafeRelease();
+	FVertexBuffer::ReleaseRHI();
+
+	AllocatedCount=0;
+	UsedCount=0;
+}
+
+//*****************************************************************************
+//******************** FGPUSortManager::FDynamicValueBuffer *******************
+//*****************************************************************************
+
+void FGPUSortManager::FDynamicValueBuffer::Allocate(FAllocationInfo& OutInfo, const FGPUSortManager::FSettings& InSettings, int32 ValueCount, EGPUSortFlags Flags)
+{
+	if (ValueBuffers.Num())
+	{
+		FValueBuffer& LastValueBuffer = ValueBuffers.Last();
+		const int32 RequiredCount = LastValueBuffer.UsedCount + ValueCount;
+
+		if (RequiredCount <= LastValueBuffer.AllocatedCount)
+		{
+			LastValueBuffer.Allocate(OutInfo, ValueCount, Flags);
+		}
+		else
+		{
+			const int32 NewBufferSize = (int32)(RequiredCount * InSettings.BufferSlack);
+			FValueBuffer* NewValueBuffer = new FValueBuffer(NewBufferSize, LastValueBuffer.UsedCount, InSettings);
+			NewValueBuffer->Allocate(OutInfo, ValueCount, Flags);
+			ValueBuffers.Add(NewValueBuffer);
+		}
+	}
+	else
+	{
+		const int32 NewBufferSize = FMath::Max(InSettings.MinBufferSize, (int32)(ValueCount * InSettings.BufferSlack));
+		FValueBuffer* NewValueBuffer = new FValueBuffer(NewBufferSize, 0, InSettings);
+		NewValueBuffer->Allocate(OutInfo, ValueCount, Flags);
+		ValueBuffers.Add(NewValueBuffer);
+	}
+}
+
+void FGPUSortManager::FDynamicValueBuffer::SkrinkAndReset(const FGPUSortManager::FSettings& InSettings)
+{
+	const int32 UsedCount = GetUsedCount();
+	if (UsedCount > 0)
+	{
+		const int32 AllocatedCount = GetAllocatedCount();
+		// Only keep the last buffer which is the one of the right size.
+		ValueBuffers.RemoveAt(0, ValueBuffers.Num() - 1);
+
+		const float RecommendedSize = FMath::Max<float>(InSettings.MinBufferSize, UsedCount * InSettings.BufferSlack);
+		if (RecommendedSize * InSettings.BufferSlack < AllocatedCount)
+		{
+			if (++NumFramesRequiringShrinking > InSettings.FrameCountBeforeShrinking)
+			{
+				ValueBuffers.Empty();
+				ValueBuffers.Add(new FValueBuffer((int32)RecommendedSize, 0, InSettings));
+			}
+		}
+		else
+		{
+			NumFramesRequiringShrinking = 0;
+		}
+
+		// Reset UsedCount for next frame.
+		ValueBuffers.Last().UsedCount = 0;
+	}
+	else
+	{
+		if (++NumFramesRequiringShrinking > InSettings.FrameCountBeforeShrinking)
+		{
+			ValueBuffers.Empty();
+		}
+	}
+
+	CurrentSortBatchId = INDEX_NONE;
+}
+
+void FGPUSortManager::FDynamicValueBuffer::ReleaseRHI()
+{
+	for (FValueBuffer& ValueBuffer : ValueBuffers)
+	{
+		ValueBuffer.ReleaseRHI();
+	}
+}
+
+//*****************************************************************************
+//************************ FGPUSortManager::FSortBatch ************************
+//*****************************************************************************
+
+void FGPUSortManager::FSortBatch::UpdateProcessingOrder()
+{
+	if (EnumHasAnyFlags(Flags, EGPUSortFlags::SortAfterPreRender))
+	{
+		checkSlow(!EnumHasAnyFlags(Flags, EGPUSortFlags::KeyGenAfterPostRenderOpaque));
+		ProcessingOrder = ESortBatchProcessingOrder::KeyGenAndSortAfterPreRender;
+	}
+	else // SortAfterPostRenderOpaque
+	{
+		checkSlow(EnumHasAnyFlags(Flags, EGPUSortFlags::SortAfterPostRenderOpaque));
+		if (EnumHasAnyFlags(Flags, EGPUSortFlags::KeyGenAfterPreRender))
+		{
+			ProcessingOrder = ESortBatchProcessingOrder::KeyGenAfterPreRenderAndSortAfterPostRenderOpaque;
+		}
+		else
+		{
+			ProcessingOrder = ESortBatchProcessingOrder::KeyGenAndSortAfterPostRenderOpaque;
+		}
+	}
+}
+
+void FGPUSortManager::FSortBatch::GenerateKeys(FRHICommandListImmediate& RHICmdList, const FGPUSortManager::FCallbackArray& InCallbacks, EGPUSortFlags KeyGenLocation)
+{
+	const int32 InitialIndex = 0;
+	FRHIUnorderedAccessView* KeyValueUAVs[] = { SortBuffers->GetKeyBufferUAV(InitialIndex), DynamicValueBuffer->ValueBuffers.Last().UInt32UAV };
+
+	// TR-KeyGen : Sync the keys with the last GPU sort task.
+	RHICmdList.TransitionResource(EResourceTransitionAccess::ERWBarrier, EResourceTransitionPipeline::EComputeToCompute, KeyValueUAVs[0]);
+	// TR-KeyGen : Make sure the values are safe to write to (could be still in use in the last render).
+	RHICmdList.TransitionResource(EResourceTransitionAccess::ERWBarrier, EResourceTransitionPipeline::EGfxToCompute, KeyValueUAVs[1]);
+
+	for (const FCallbackInfo& Callback : InCallbacks)
+	{
+		if (EnumHasAnyFlags(KeyGenLocation & Callback.Flags, EGPUSortFlags::AnyKeyGenLocation))
+		{
+			SCOPED_DRAW_EVENTF(RHICmdList, GPUSortBatch, TEXT("KeyGen_%s"), *Callback.Name.ToString());
+			const bool bAsInt32 = EnumHasAnyFlags(Callback.Flags, EGPUSortFlags::ValuesAsInt32);
+			FRHIUnorderedAccessView* TypedValueUAV = bAsInt32 ? DynamicValueBuffer->ValueBuffers.Last().Int32UAV : DynamicValueBuffer->ValueBuffers.Last().G16R16UAV;
+			// TR-KeyGen : TypedValueUAV is the same as ValueUAVs[1] but with a different type. The callback needs to do an ERWNoBarrier between each dispatch updating partially the content.
+			Callback.Delegate.Execute(RHICmdList, Id, NumElements, (Flags & EGPUSortFlags::AnyKeyPrecision) | KeyGenLocation, KeyValueUAVs[0], TypedValueUAV);
+		}
+	}
+
+	// TR-KeyGen : Those buffers will now be red as input for the compute GPU sort.
+	RHICmdList.TransitionResources(EResourceTransitionAccess::ERWBarrier, EResourceTransitionPipeline::EComputeToCompute, KeyValueUAVs, UE_ARRAY_COUNT(KeyValueUAVs));
+}
+
+void FGPUSortManager::FSortBatch::SortAndResolve(FRHICommandListImmediate& RHICmdList, ERHIFeatureLevel::Type InFeatureLevel)
+{
+	check(SortBuffers && DynamicValueBuffer);
+	const FValueBufferArray& ValueBuffers = DynamicValueBuffer->ValueBuffers;
+	if (ValueBuffers.Num())
+	{
+		FGPUSortBuffers SortBuffer = SortBuffers->GetSortBuffers();
+
+		// TR-SortedValues : at this point the first sorted values buffer was updated by the GenerateKeys and are ready for compute.
+
+		SortBuffer.FirstValuesSRV = ValueBuffers.Last().UInt32SRV;
+		SortBuffer.FinalValuesUAV = ValueBuffers.Last().UInt32UAV;
+
+		const int32 InitialIndex = 0;
+		const FGPUSortManager::FKeyGenInfo KeyGenInfo((uint32)NumElements, EnumHasAnyFlags(Flags, EGPUSortFlags::HighPrecisionKeys) != 0);
+		const uint32 KeyMask = (KeyGenInfo.ElementKeyMask << KeyGenInfo.ElementKeyShift) | KeyGenInfo.SortKeyMask;
+		const int32 ResultIndex = SortGPUBuffers(RHICmdList, SortBuffer, InitialIndex, KeyMask, DynamicValueBuffer->GetUsedCount(), InFeatureLevel);
+
+		if (ValueBuffers.Num() > 1)
+		{
+			SCOPED_DRAW_EVENT(RHICmdList, CopyResults);
+
+			TArray<FRHIUnorderedAccessView*, TInlineAllocator<(COPYUINTCS_BUFFER_COUNT+1)>> TargetUAVs;
+			TArray<int32, TInlineAllocator<COPYUINTCS_BUFFER_COUNT>> TargetSizes;
+			for (FValueBuffer& ValueBuffer : DynamicValueBuffer->ValueBuffers)
+			{
+				// The biggest buffer contains 
+				if (ValueBuffer.UInt32UAV != SortBuffer.FinalValuesUAV)
+				{
+					TargetUAVs.Add(ValueBuffer.UInt32UAV);
+					TargetSizes.Add(ValueBuffer.UsedCount);
+				}
+			}
+
+			// TR-SortedValues : The sort value buffers (from the GPU sort dispatches) will now be updated by the FCopyUIntBufferCS dispatch.
+			RHICmdList.TransitionResource(EResourceTransitionAccess::EReadable, EResourceTransitionPipeline::EComputeToCompute, SortBuffer.FinalValuesUAV);
+			RHICmdList.TransitionResources(EResourceTransitionAccess::ERWBarrier, EResourceTransitionPipeline::EComputeToCompute, TargetUAVs.GetData(), TargetUAVs.Num());
+
+			CopyUIntBufferToTargets(RHICmdList, InFeatureLevel, SortBuffer.FirstValuesSRV, TargetUAVs.GetData(), TargetSizes.GetData(), 0, TargetUAVs.Num());
+
+			// TR-SortedValues : The other (smaller) buffers containing sorted values have been updated by the FCopyUIntBufferCS dispatch, and can now be red by the Gfx pipeline.
+			RHICmdList.TransitionResources(EResourceTransitionAccess::EReadable, EResourceTransitionPipeline::EComputeToGfx, TargetUAVs.GetData(), TargetUAVs.Num());
+		}
+
+		// TR-SortedValues : Make this readable by the Gfx pipeline.
+		RHICmdList.TransitionResource(EResourceTransitionAccess::EReadable, EResourceTransitionPipeline::EComputeToGfx, SortBuffer.FinalValuesUAV);
+	}
+}
+
+//*****************************************************************************
+//****************************** FGPUSortManager ******************************
+//*****************************************************************************
+
+bool FGPUSortManager::HasEither(EGPUSortFlags Flags, EGPUSortFlags A, EGPUSortFlags B)
+{
+	const bool bHasA = (Flags & A) != EGPUSortFlags::None;
+	const bool bHasB = (Flags & B) != EGPUSortFlags::None;
+	return bHasA != bHasB;
+}
+
+EGPUSortFlags FGPUSortManager::GetBatchFlags(EGPUSortFlags TaskFlags)
+{
+	// The batch only contains information about prevision and sort location. 
+	// Value formats are globally defined and all formats are available in all batches.
+	// KeyGen location callbacks (either PreRender or PostRenderOpaque) must be handled by 
+	// the system pushing the task to avoid overcomplexity in the sort manager.
+	// The batch accumulates all the location where key gen must be called.
+	return TaskFlags & (EGPUSortFlags::AnyKeyPrecision | EGPUSortFlags::AnyKeyGenLocation | EGPUSortFlags::AnySortLocation);
+}
+
+EGPUSortFlags FGPUSortManager::CombineBatchFlags(EGPUSortFlags BatchFlags, EGPUSortFlags TaskFlags)
+{
+	// Key precision and Sort location gets restricted as tasks are added to the batch.
+	return ((BatchFlags & TaskFlags) & (EGPUSortFlags::AnyKeyPrecision | EGPUSortFlags::AnySortLocation)) | 
+	// KeyGen location increases and since each task only has one KeyGen location, compatible with the sort location, this is safe.
+	   	   ((BatchFlags | TaskFlags) & EGPUSortFlags::AnyKeyGenLocation);
+}
+
+bool FGPUSortManager::TestBatchFlags(EGPUSortFlags BatchFlags, EGPUSortFlags TaskFlags)
+{
+	// Check that the batch supports this precision and this location. Note that only one of the precision and sort location needs to be in.
+	return EnumHasAnyFlags(BatchFlags, TaskFlags & EGPUSortFlags::AnyKeyPrecision) && EnumHasAnyFlags(BatchFlags, TaskFlags & EGPUSortFlags::AnySortLocation);
+}
+
+const TCHAR* FGPUSortManager::GetPrecisionString(EGPUSortFlags BatchFlags)
+{
+	return EnumHasAnyFlags(BatchFlags, EGPUSortFlags::LowPrecisionKeys) ? TEXT("LowPrecision") : TEXT("HighPrecision");
+}
+
+FGPUSortManager::FGPUSortManager(ERHIFeatureLevel::Type InFeatureLevel) 
+	: FeatureLevel(InFeatureLevel) 
+{
+	if (GGPUSortBufferSlack > 1.f)
+	{
+		Settings.BufferSlack = GGPUSortBufferSlack;
+	}
+	Settings.FrameCountBeforeShrinking = GGPUSortFrameCountBeforeBufferShrinking;
+	Settings.MinBufferSize = GGPUSortMinBufferSize;
+}
+
+FGPUSortManager::~FGPUSortManager()
+{
+	check(!SortBatches.Num());
+
+	ReleaseSortBuffers();
+}
+
+void FGPUSortManager::Register(const FGPUSortKeyGenDelegate& CallbackDelegate, EGPUSortFlags UsedFlags, const FName& InName)
+{
+	// Callbacks must only have a single value format otherwise we can't tell which UAV to pass to it.
+	// A system could use two callbacks to generate the initial values if ultimately required.
+	check(HasEither(UsedFlags, EGPUSortFlags::ValuesAsInt32, EGPUSortFlags::ValuesAsG16R16F));
+
+	Settings.AllowedFlags |= UsedFlags;
+
+	FCallbackInfo& Callback = Callbacks.AddDefaulted_GetRef();
+	Callback.Delegate = CallbackDelegate;
+	Callback.Flags = UsedFlags;
+	Callback.Name = InName;
+}
+
+bool FGPUSortManager::AddTask(FAllocationInfo& OutInfo, int32 ValueCount, EGPUSortFlags TaskFlags)
+{
+	if (!GGPUSortEnable)
+	{
+		return false;
+	}
+
+	LLM_SCOPE(ELLMTag::GPUSort);
+
+	// Can only use flags that were registered initially.
+	check(EnumHasAllFlags(Settings.AllowedFlags, TaskFlags));
+	
+	// The must be at least one key precision.
+	checkSlow(EnumHasAnyFlags(TaskFlags, EGPUSortFlags::AnyKeyPrecision));
+	// Key gen must be set at only one place, and as early as possible to maximize batch grouping.
+	checkSlow(HasEither(TaskFlags, EGPUSortFlags::KeyGenAfterPreRender, EGPUSortFlags::KeyGenAfterPostRenderOpaque));
+	// The must be at least one sort location.
+	checkSlow(EnumHasAnyFlags(TaskFlags, EGPUSortFlags::AnySortLocation));
+	// Value format must be known since we set the SRV in OutInfo.
+	checkSlow(HasEither(TaskFlags, EGPUSortFlags::ValuesAsInt32, EGPUSortFlags::ValuesAsG16R16F));
+	// Invalid if it needs to be sorted before the key are to be generated. Ex : Cascade opaque materials using depth buffer in the simulation.
+	checkSlow(!EnumHasAllFlags(TaskFlags, EGPUSortFlags::KeyGenAfterPostRenderOpaque | EGPUSortFlags::SortAfterPreRender));
+
+	// Search into the current batch for a compatible entries
+	for (FSortBatch& SortBatch : SortBatches)
+	{
+		// Check that the batch supports this precision and this location. Note that only one of the precision and sort location needs to be in.
+		if (TestBatchFlags(SortBatch.Flags, TaskFlags))
+		{
+			OutInfo.SortBatchId = SortBatch.Id;
+			OutInfo.ElementIndex = SortBatch.NumElements;
+
+			SortBatch.NumElements += 1;
+			SortBatch.Flags = CombineBatchFlags(SortBatch.Flags, TaskFlags);
+			SortBatch.DynamicValueBuffer->Allocate(OutInfo, Settings, ValueCount, TaskFlags);
+			return true;
+		}
+	}
+
+	// Otherwise, create a new sort batch.
+	const int32 BatchId = SortBatches.Num();
+	OutInfo.SortBatchId = BatchId;
+	OutInfo.ElementIndex = 0;
+
+	FSortBatch& SortBatch = SortBatches.AddDefaulted_GetRef();
+	SortBatch.Id = BatchId;
+	SortBatch.NumElements = 1;
+	SortBatch.Flags = GetBatchFlags(TaskFlags);
+	// Bind a sort buffer to the batch. Use the Flags for the task and not from the current batch
+	// since the task flags are more restrictive at this point than the batch flags that subsets as task get grouped.
+	SortBatch.DynamicValueBuffer = GetDynamicValueBufferFromPool(TaskFlags, BatchId);
+	checkSlow(SortBatch.DynamicValueBuffer);
+	SortBatch.DynamicValueBuffer->Allocate(OutInfo, Settings, ValueCount, TaskFlags);
+
+	return true;
+}
+
+void FGPUSortManager::FinalizeSortBatches()
+{
+	// At this point the batches are complete as no more tasks will be added to them.
+	// Update FDynamicValueBuffer::LastSortBatchFlags before the batch can potentially be destroyed.
+	// This will helps rebind the same DynamicValueBuffer to the same batches.
+	for (FSortBatch& SortBatch : SortBatches)
+	{
+		// Resolve the sort location if there is still ambiguity about where the batch needs to be sorted.
+		if (EnumHasAllFlags(SortBatch.Flags, EGPUSortFlags::AnySortLocation))
+		{
+			// If it needs KeyGen after PostRenderOpaque(), then remove the possibility of sorting after PreRender().
+			if (EnumHasAnyFlags(SortBatch.Flags, EGPUSortFlags::KeyGenAfterPostRenderOpaque))
+			{
+				SortBatch.Flags ^= EGPUSortFlags::SortAfterPreRender;
+			}
+			// Otherwise it must have the possibility to KeyGen after PreRender(), and so remove the possibility of sorting after PostRenderOpaque().
+			else
+			{
+				checkSlow(EnumHasAnyFlags(SortBatch.Flags, EGPUSortFlags::KeyGenAfterPreRender))
+				SortBatch.Flags ^= EGPUSortFlags::SortAfterPostRenderOpaque;
+			}
+		}
+		checkSlow(HasEither(SortBatch.Flags, EGPUSortFlags::SortAfterPreRender, EGPUSortFlags::SortAfterPostRenderOpaque));
+
+		// Update the sort batch flags now that they are final. This will help rebind the same dynamic value sort batches to the s
+		checkSlow(SortBatch.DynamicValueBuffer);
+		SortBatch.DynamicValueBuffer->LastSortBatchFlags = SortBatch.Flags;
+
+		// Update the sort batch processing order.
+		SortBatch.UpdateProcessingOrder();
+	}
+}
+
+void FGPUSortManager::UpdateSortBuffersPool()
+{
+	int32 NumSortBuffersRequired = DynamicValueBufferPool.Num() ? 1 : 0;
+	int32 MaxSortBuffersSizes = 0;
+
+	int32 NumDedicatedSortBuffersRequired = 0;
+	for (const FDynamicValueBuffer& DynamicValueBuffer : DynamicValueBufferPool)
+	{
+		MaxSortBuffersSizes = FMath::Max<int32>(MaxSortBuffersSizes, DynamicValueBuffer.GetAllocatedCount());
+
+		// If some keys need to be generated at PreRender() but only be used in PostRenderOpaque(), then this buffer won't be sharable in between that timeframe.
+		if (EnumHasAllFlags(DynamicValueBuffer.LastSortBatchFlags, EGPUSortFlags::KeyGenAfterPreRender | EGPUSortFlags::SortAfterPostRenderOpaque))
+		{
+			++NumDedicatedSortBuffersRequired;
+		}
+	}
+	// There are currently no usecase where there can be more than one of those buffers so the SortBuffers pooling strategy is simple.
+	ensure(NumDedicatedSortBuffersRequired <= 1);
+
+	// The dedicated sort buffers can be reused.
+	NumSortBuffersRequired = FMath::Max<int32>(NumSortBuffersRequired, NumDedicatedSortBuffersRequired);
+
+	// Allocate all the buffers at the right size.
+	for (int32 SortBufferIndex = 0; SortBufferIndex < NumSortBuffersRequired; ++SortBufferIndex)
+	{
+		if (SortBuffersPool.IsValidIndex(SortBufferIndex))
+		{
+			FParticleSortBuffers* SortBuffers = SortBuffersPool[SortBufferIndex];
+			check(SortBuffers);
+			if (SortBuffers->GetSize() != MaxSortBuffersSizes)
+			{
+				SortBuffers->ReleaseRHI();
+				SortBuffers->SetBufferSize(MaxSortBuffersSizes);
+				SortBuffers->InitRHI();
+			}
+		}
+		else
+		{
+			FParticleSortBuffers* SortBuffers = new FParticleSortBuffers;
+			SortBuffers->SetBufferSize(MaxSortBuffersSizes);
+			SortBuffers->InitRHI();
+			SortBuffersPool.Add(SortBuffers);
+		}
+	}
+
+	// Free unrequired buffers.
+	while (SortBuffersPool.Num() > NumSortBuffersRequired)
+	{
+		FParticleSortBuffers* SortBuffers = SortBuffersPool.Pop();
+		check(SortBuffers);
+		SortBuffers->ReleaseRHI();
+		SortBuffers->SetBufferSize(0);
+		delete SortBuffers;
+	}
+}
+
+void FGPUSortManager::OnPreRender(FRHICommandListImmediate& RHICmdList)
+{
+	LLM_SCOPE(ELLMTag::GPUSort);
+
+	FinalizeSortBatches();
+	UpdateSortBuffersPool();
+
+	// Sort batches so that the next batch to handle is at the end of the array.
+	SortBatches.Sort([](const FSortBatch& A, const FSortBatch& B) { return (uint32)A.ProcessingOrder > (uint32)B.ProcessingOrder; });
+
+	if (SortBatches.Num())
+	{
+		SCOPED_GPU_STAT(RHICmdList, GPUKeyGenAndSort);
+		while (SortBatches.Num() && SortBatches.Last().ProcessingOrder == ESortBatchProcessingOrder::KeyGenAndSortAfterPreRender)
+		{
+			// Remove the SortBatch but don't remove the SortBuffers from the pool since it can be reused immediately.
+			FSortBatch SortBatch = SortBatches.Pop();
+			FParticleSortBuffers* SortBuffers = SortBuffersPool.Last();
+			SortBatch.SortBuffers = SortBuffers;
+			checkSlow(SortBuffers && SortBatch.GetUsedValueCount() <= SortBuffers->GetSize());
+
+			SCOPED_DRAW_EVENTF(RHICmdList, GPUSortBatch, TEXT("GPUSort_Batch%d(%s)"), SortBatch.Id, GetPrecisionString(SortBatch.Flags));
+
+			SortBatch.GenerateKeys(RHICmdList, Callbacks, EGPUSortFlags::KeyGenAfterPreRender);
+			SortBatch.SortAndResolve(RHICmdList, FeatureLevel);
+
+			// Release the sort batch.
+			checkSlow(SortBatch.DynamicValueBuffer && SortBatch.DynamicValueBuffer->CurrentSortBatchId == SortBatch.Id);
+			SortBatch.DynamicValueBuffer->CurrentSortBatchId = INDEX_NONE;
+			SortBatch.SortBuffers = nullptr;
+		}
+
+		for (int32 BatchIndex = SortBatches.Num() - 1; BatchIndex >= 0 && SortBatches[BatchIndex].ProcessingOrder == ESortBatchProcessingOrder::KeyGenAfterPreRenderAndSortAfterPostRenderOpaque; --BatchIndex)
+		{
+			// Those sort batches will be processed again after PostRenderOpaque(). Because of this, the particle sort buffers can not be reused.
+			FSortBatch& SortBatch = SortBatches[BatchIndex];
+			FParticleSortBuffers* SortBuffers = SortBuffersPool.Pop();
+			SortBatch.SortBuffers = SortBuffers;
+			checkSlow(SortBuffers && SortBatch.GetUsedValueCount() <= SortBuffers->GetSize());
+
+			SCOPED_DRAW_EVENTF(RHICmdList, GPUSortBatchPreStep, TEXT("GPUSort_Batch%d(PreStep,%s)"), SortBatch.Id, GetPrecisionString(SortBatch.Flags));
+
+			SortBatch.GenerateKeys(RHICmdList, Callbacks, EGPUSortFlags::KeyGenAfterPreRender);
+		}
+	}
+}
+
+void FGPUSortManager::OnPostRenderOpaque(FRHICommandListImmediate& RHICmdList)
+{
+	LLM_SCOPE(ELLMTag::GPUSort);
+
+	if (SortBatches.Num())
+	{
+		SCOPED_GPU_STAT(RHICmdList, GPUKeyGenAndSort);
+		while (SortBatches.Num())
+		{
+			// Remove the SortBatch but don't remove the SortBuffers from the pool since it can be reused immediately.
+			FSortBatch SortBatch = SortBatches.Pop();
+			checkSlow((SortBatch.SortBuffers != nullptr) == (SortBatch.ProcessingOrder == ESortBatchProcessingOrder::KeyGenAfterPreRenderAndSortAfterPostRenderOpaque));
+
+			SCOPED_DRAW_EVENTF(RHICmdList, GPUSortBatch, TEXT("GPUSort_Batch%d(%s)"), SortBatch.Id, GetPrecisionString(SortBatch.Flags));
+
+			if (!SortBatch.SortBuffers)
+			{
+				FParticleSortBuffers* SortBuffers = SortBuffersPool.Pop();
+				SortBatch.SortBuffers = SortBuffers;
+				checkSlow(SortBuffers && SortBatch.GetUsedValueCount() <= SortBuffers->GetSize());
+			}
+
+			SortBatch.GenerateKeys(RHICmdList, Callbacks, EGPUSortFlags::KeyGenAfterPostRenderOpaque);
+			SortBatch.SortAndResolve(RHICmdList, FeatureLevel);
+
+			// Release the sort batch.
+			checkSlow(SortBatch.DynamicValueBuffer && SortBatch.DynamicValueBuffer->CurrentSortBatchId == SortBatch.Id);
+			SortBatch.DynamicValueBuffer->CurrentSortBatchId = INDEX_NONE;
+			SortBuffersPool.Push(SortBatch.SortBuffers);
+			SortBatch.SortBuffers = nullptr;
+		}
+	}
+
+	ResetDynamicValuesBuffers();
+}
+
+FGPUSortManager::FDynamicValueBuffer* FGPUSortManager::GetDynamicValueBufferFromPool(EGPUSortFlags TaskFlags, int32 SortBatchId)
+{
+	for (FDynamicValueBuffer& DynamicValueBuffer : DynamicValueBufferPool)
+	{
+		// If the buffer is unused and compatible.
+		if (DynamicValueBuffer.CurrentSortBatchId == INDEX_NONE && TestBatchFlags(DynamicValueBuffer.LastSortBatchFlags, TaskFlags))
+		{
+			DynamicValueBuffer.CurrentSortBatchId = SortBatchId;
+			return &DynamicValueBuffer; // Safe since we are using IndirectArray.
+		}
+	}
+
+	FDynamicValueBuffer* DynamicValueBuffer = new FDynamicValueBuffer;
+	DynamicValueBuffer->CurrentSortBatchId = SortBatchId;
+	DynamicValueBufferPool.Add(DynamicValueBuffer);
+	return DynamicValueBuffer;
+}
+
+void FGPUSortManager::ResetDynamicValuesBuffers()
+{
+	check(!SortBatches.Num()) ;
+
+	if (!GGPUSortStressTest)
+	{
+		for (int32 Index = 0; Index < DynamicValueBufferPool.Num(); ++Index)
+		{
+			FDynamicValueBuffer& DynamicValueBuffer = DynamicValueBufferPool[Index];
+			DynamicValueBuffer.SkrinkAndReset(Settings);
+			if (DynamicValueBuffer.GetAllocatedCount() == 0)
+			{
+				DynamicValueBufferPool.RemoveAtSwap(Index);
+				--Index;
+			}
+		}
+	}
+	else
+	{
+		DynamicValueBufferPool.Empty();
+		ReleaseSortBuffers();
+	}
+}
+
+void FGPUSortManager::ReleaseSortBuffers()
+{
+	for (FParticleSortBuffers* SortBuffers : SortBuffersPool)
+	{
+		check(SortBuffers);
+		SortBuffers->ReleaseRHI();
+		delete SortBuffers;
+	}
+
+	SortBuffersPool.Empty();
+}
diff --git a/Engine/Source/Runtime/Engine/Private/Particles/FXSystem.cpp b/Engine/Source/Runtime/Engine/Private/Particles/FXSystem.cpp
--- a/Engine/Source/Runtime/Engine/Private/Particles/FXSystem.cpp
+++ b/Engine/Source/Runtime/Engine/Private/Particles/FXSystem.cpp
@@ -15,6 +15,7 @@
 #include "Components/VectorFieldComponent.h"
 #include "SceneUtils.h"
 #include "Renderer/Private/SceneRendering.h" // needed for STATGROUP_CommandListMarkers
+#include "GPUSortManager.h"
 
 
 TMap<FName, FCreateCustomFXSystemDelegate> FFXSystemInterface::CreateCustomFXDelegates;
@@ -25,14 +26,19 @@ TMap<FName, FCreateCustomFXSystemDelegate> FFXSystemInterface::CreateCustomFXDel
 
 FFXSystemInterface* FFXSystemInterface::Create(ERHIFeatureLevel::Type InFeatureLevel, EShaderPlatform InShaderPlatform)
 {
+	// The FGPUSortManager is currently only being used by FFXSystemInterface implementations.
+	// Because of this, the lifetime management of the GPUSortManager only consists of a ref counter incremented initially 
+	// in the gamethread (in this function) and decremented on the renderthread (when the system interfaces are deleted).
+	TRefCountPtr<FGPUSortManager> GPUSortManager = new FGPUSortManager(InFeatureLevel);
+
 	if (CreateCustomFXDelegates.Num())
 	{
-		FFXSystemSet* Set = new FFXSystemSet;
-		Set->FXSystems.Add(new FFXSystem(InFeatureLevel, InShaderPlatform));
+		FFXSystemSet* Set = new FFXSystemSet(GPUSortManager);
+		Set->FXSystems.Add(new FFXSystem(InFeatureLevel, InShaderPlatform, GPUSortManager));
 
 		for (TMap<FName, FCreateCustomFXSystemDelegate>::TConstIterator Ite(CreateCustomFXDelegates); Ite; ++Ite)
 		{
-			FFXSystemInterface* CustomFX = Ite.Value().Execute(InFeatureLevel, InShaderPlatform);
+			FFXSystemInterface* CustomFX = Ite.Value().Execute(InFeatureLevel, InShaderPlatform, GPUSortManager);
 			if (CustomFX)
 			{
 				Set->FXSystems.Add(CustomFX);
@@ -42,7 +48,7 @@ FFXSystemInterface* FFXSystemInterface::Create(ERHIFeatureLevel::Type InFeatureL
 	}
 	else
 	{
-		return new FFXSystem(InFeatureLevel, InShaderPlatform);
+		return new FFXSystem(InFeatureLevel, InShaderPlatform, GPUSortManager);
 	}
 }
 
@@ -174,15 +180,30 @@ namespace FXConsoleVariables
 	FX system.
 ------------------------------------------------------------------------------*/
 
-FFXSystem::FFXSystem(ERHIFeatureLevel::Type InFeatureLevel, EShaderPlatform InShaderPlatform)
+FFXSystem::FFXSystem(ERHIFeatureLevel::Type InFeatureLevel, EShaderPlatform InShaderPlatform, FGPUSortManager* InGPUSortManager)
 	: ParticleSimulationResources(NULL)
 	, FeatureLevel(InFeatureLevel)
 	, ShaderPlatform(InShaderPlatform)
+	, GPUSortManager(InGPUSortManager)
 #if WITH_EDITOR
 	, bSuspended(false)
 #endif // #if WITH_EDITOR
 {
 	InitGPUSimulation();
+
+	// Register the callback in the GPUSortManager. 
+	// The callback is used to generate the initial keys and values for the GPU sort tasks, 
+	// the values being the sorted particle indices used by the GPU particles.
+	// The registration also involves defining the list of flags possibly used in GPUSortManager::AddTask()
+	if (GPUSortManager)
+	{
+		GPUSortManager->Register(FGPUSortKeyGenDelegate::CreateLambda([this](FRHICommandListImmediate& RHICmdList, int32 BatchId, int32 NumElementsInBatch, EGPUSortFlags Flags, FRHIUnorderedAccessView* KeysUAV, FRHIUnorderedAccessView* ValuesUAV)
+		{ 
+			GenerateSortKeys(RHICmdList, BatchId, NumElementsInBatch, Flags, KeysUAV, ValuesUAV);
+		}), 
+		EGPUSortFlags::LowPrecisionKeys | EGPUSortFlags::KeyGenAfterPostRenderOpaque | EGPUSortFlags::AnySortLocation | EGPUSortFlags::ValuesAsG16R16F,
+		Name);
+	}
 }
 
 FFXSystem::~FFXSystem()
@@ -446,14 +467,18 @@ void FFXSystem::PostRenderOpaque(
 {
 	if (RHISupportsGPUParticles() && IsParticleCollisionModeSupported(GetShaderPlatform(), PCM_DepthBuffer))
 	{
-		SCOPED_DRAW_EVENT(RHICmdList, GPUParticles_PostRenderOpaque);
 		if (bAllowGPUParticleUpdate)
 		{
+			SCOPED_DRAW_EVENT(RHICmdList, GPUParticles_PostRenderOpaque);
 			PrepareGPUSimulation(RHICmdList);
 			SimulateGPUParticles(RHICmdList, EParticleSimulatePhase::CollisionDepthBuffer, ViewUniformBuffer, NULL, SceneTexturesUniformBufferStruct, SceneTexturesUniformBuffer);
 			FinalizeGPUSimulation(RHICmdList);
 		}
-
-		SortGPUParticles(RHICmdList);		
 	}
 }
+
+FGPUSortManager* FFXSystem::GetGPUSortManager() const
+{
+	return GPUSortManager;
+}
+
diff --git a/Engine/Source/Runtime/Engine/Private/Particles/FXSystemPrivate.h b/Engine/Source/Runtime/Engine/Private/Particles/FXSystemPrivate.h
--- a/Engine/Source/Runtime/Engine/Private/Particles/FXSystemPrivate.h
+++ b/Engine/Source/Runtime/Engine/Private/Particles/FXSystemPrivate.h
@@ -10,6 +10,7 @@
 #include "FXSystem.h"
 #include "VectorField.h"
 #include "ParticleSortingGPU.h"
+#include "GPUSortManager.h"
 
 class FCanvas;
 class FGlobalDistanceFieldParameterData;
@@ -130,7 +131,7 @@ class FFXSystem : public FFXSystemInterface
 public:
 
 	/** Default constructoer. */
-	FFXSystem(ERHIFeatureLevel::Type InFeatureLevel, EShaderPlatform InShaderPlatform);
+	FFXSystem(ERHIFeatureLevel::Type InFeatureLevel, EShaderPlatform InShaderPlatform, FGPUSortManager* InGPUSortManager);
 
 	/** Destructor. */
 	virtual ~FFXSystem();
@@ -196,19 +197,46 @@ public:
 		return ParticleSimulationResources;
 	}
 
-	/**
-	 * Prepares a GPU simulation to be sorted for a particular view.
+	/** 
+	 * Register work for GPU sorting (using the GPUSortManager). 
+	 * The initial keys and values are generated in the GenerateSortKeys() callback.
+	 *
 	 * @param Simulation The simulation to be sorted.
 	 * @param ViewOrigin The origin of the view from which to sort.
-	 * @returns an offset in to the sorted buffer from which the simulation may render.
+	 * @param bIsTranslucent Whether this is for sorting translucent particles or opaque particles, affect when the data is required for the rendering.
+	 * @param OutInfo The bindings for this GPU sort task, if success. 
+	 * @returns true if the work was registered, or false it GPU sorting is not available or impossible.
 	 */
-	int32 AddSortedGPUSimulation(FParticleSimulationGPU* Simulation, const FVector& ViewOrigin);
+	bool AddSortedGPUSimulation(FParticleSimulationGPU* Simulation, const FVector& ViewOrigin, bool bIsTranslucent, FGPUSortManager::FAllocationInfo& OutInfo);
 
 	void PrepareGPUSimulation(FRHICommandListImmediate& RHICmdList);
 	void FinalizeGPUSimulation(FRHICommandListImmediate& RHICmdList);
 
+	/** Get the shared SortManager, used in the rendering loop to call FGPUSortManager::OnPreRender() and FGPUSortManager::OnPostRenderOpaque() */
+	virtual FGPUSortManager* GetGPUSortManager() const override;
+
 private:
 
+	/**
+	 * Generate all the initial keys and values for a GPUSortManager sort batch.
+	 * Sort batches are created when GPU sort tasks are registered in AddSortedGPUSimulation().
+	 * Each sort task defines constraints about when the initial sort data can generated and
+	 * and when the sorted results are needed (see EGPUSortFlags for details).
+	 * Currently, for Cascade, all the sort tasks have the EGPUSortFlags::KeyGenAfterPostRenderOpaque flag
+	 * and so the callback registered in GPUSortManager->Register() only has the EGPUSortFlags::KeyGenAfterPostRenderOpaque usage.
+	 * This garanties that GenerateSortKeys() only gets called after PostRenderOpaque(), which is a constraint required because
+	 * cascade renders the GPU emitters after they have been ticked in PostRenderOpaque.
+	 * Note that this callback must only initialize the content for the elements that relates to the tasks it has registered in this batch.
+	 *
+	 * @param RHICmdList The command list used to initiate the keys and values on GPU.
+	 * @param BatchId The GPUSortManager batch id (regrouping several similar sort tasks).
+	 * @param NumElementsInBatch The number of elements grouped in the batch (each element maps to a sort task)
+	 * @param Flags Details about the key precision (see EGPUSortFlags::AnyKeyPrecision) and the keygen location (see EGPUSortFlags::AnyKeyGenLocation).
+	 * @param KeysUAV The UAV that holds all the initial keys used to sort the values (being the particle indices here). 
+	 * @param ValuesUAV The UAV that holds the initial values (particle indices) to be sorted accordingly to the keys.
+	 */
+	void GenerateSortKeys(FRHICommandListImmediate& RHICmdList, int32 BatchId, int32 NumElementsInBatch, EGPUSortFlags Flags, FRHIUnorderedAccessView* KeysUAV, FRHIUnorderedAccessView* ValuesUAV);
+
 	/*--------------------------------------------------------------------------
 		Private interface for GPU simulations.
 	--------------------------------------------------------------------------*/
@@ -238,12 +266,6 @@ private:
 	 */
 	void AdvanceGPUParticleFrame(bool bAllowGPUParticleUpdate);
 
-	/**
-	 * Sorts all GPU particles that have called AddSortedGPUSimulation since the
-	 * last reset.
-	 */
-	void SortGPUParticles(FRHICommandListImmediate& RHICmdList);
-
 	bool UsesGlobalDistanceFieldInternal() const;
 	bool UsesDepthBufferInternal() const;
 	bool RequiresEarlyViewUniformBufferInternal() const;
@@ -300,6 +322,11 @@ private:
 	/** Shader platform that will be rendering this effects system */
 	EShaderPlatform ShaderPlatform;
 
+	/** The shared GPUSortManager, used to register GPU sort tasks in order to generate sorted particle indices per emitter. */
+	TRefCountPtr<FGPUSortManager> GPUSortManager;
+	/** All sort tasks registered in AddSortedGPUSimulation(). Holds all the data required in GenerateSortKeys(). */
+	TArray<FParticleSimulationSortInfo> SimulationsToSort;
+
 	/** Previous frame new particles for multi-gpu simulation*/
 	TArray<FNewParticle> LastFrameNewParticles;
 #if WITH_EDITOR
diff --git a/Engine/Source/Runtime/Engine/Private/Particles/FXSystemSet.cpp b/Engine/Source/Runtime/Engine/Private/Particles/FXSystemSet.cpp
--- a/Engine/Source/Runtime/Engine/Private/Particles/FXSystemSet.cpp
+++ b/Engine/Source/Runtime/Engine/Private/Particles/FXSystemSet.cpp
@@ -5,6 +5,12 @@
 =============================================================================*/
 
 #include "FXSystemSet.h"
+#include "GPUSortManager.h"
+
+FFXSystemSet::FFXSystemSet(FGPUSortManager* InGPUSortManager)
+	: GPUSortManager(InGPUSortManager)
+{
+}
 
 FFXSystemInterface* FFXSystemSet::GetInterface(const FName& InName)
 {
@@ -192,3 +198,8 @@ FFXSystemSet::~FFXSystemSet()
 		delete FXSystem;
 	}
 }
+
+FGPUSortManager* FFXSystemSet::GetGPUSortManager() const
+{
+	return GPUSortManager;
+}
diff --git a/Engine/Source/Runtime/Engine/Private/Particles/FXSystemSet.h b/Engine/Source/Runtime/Engine/Private/Particles/FXSystemSet.h
--- a/Engine/Source/Runtime/Engine/Private/Particles/FXSystemSet.h
+++ b/Engine/Source/Runtime/Engine/Private/Particles/FXSystemSet.h
@@ -8,6 +8,9 @@
 
 #include "CoreMinimal.h"
 #include "FXSystem.h"
+#include "Templates/RefCounting.h"
+
+class FGPUSortManager;
 
 /**
  * FX system.
@@ -16,6 +19,8 @@ class FFXSystemSet : public FFXSystemInterface
 {
 public:
 
+	FFXSystemSet(FGPUSortManager* InGPUSortManager);
+
 	TArray<FFXSystemInterface*> FXSystems;
 
 	virtual FFXSystemInterface* GetInterface(const FName& InName) override;
@@ -45,8 +50,14 @@ public:
 
 	virtual void OnDestroy() override;
 
+	/** Get the shared SortManager, used in the rendering loop to call FGPUSortManager::OnPreRender() and FGPUSortManager::OnPostRenderOpaque() */
+	virtual FGPUSortManager* GetGPUSortManager() const override;
+
 protected:
 
 	/** By making the destructor protected, an instance must be destroyed via FFXSystemInterface::Destroy. */
 	virtual ~FFXSystemSet();
+
+	// We keep a reference to the GPUSortManager only for GetGPUSortManager().
+	TRefCountPtr<FGPUSortManager> GPUSortManager;
 };
diff --git a/Engine/Source/Runtime/Engine/Private/Particles/ParticleGpuSimulation.cpp b/Engine/Source/Runtime/Engine/Private/Particles/ParticleGpuSimulation.cpp
--- a/Engine/Source/Runtime/Engine/Private/Particles/ParticleGpuSimulation.cpp
+++ b/Engine/Source/Runtime/Engine/Private/Particles/ParticleGpuSimulation.cpp
@@ -372,25 +372,9 @@ public:
 	FParticleAttributesTexture RenderAttributesTexture;
 	/** Texture holding simulation attributes. */
 	FParticleAttributesTexture SimulationAttributesTexture;
-	/** Vertex buffer that points to the current sorted vertex buffer. */
-	FParticleIndicesVertexBuffer SortedVertexBuffer;
-
-	FParticleSortBuffers ParticleSortBuffers;
 
 	/** Frame index used to track double buffered resources on the GPU. */
-	int32 FrameIndex;
-
-	/** List of simulations to be sorted. */
-	TArray<FParticleSimulationSortInfo> SimulationsToSort;
-	/** The total number of sorted particles. */
-	int32 SortedParticleCount;
-
-	/** Default constructor. */
-	FParticleSimulationResources()
-		: FrameIndex(0)
-		, SortedParticleCount(0)
-	{
-	}
+	int32 FrameIndex = 0;
 
 	/**
 	 * Initialize resources.
@@ -404,7 +388,6 @@ public:
 			ParticleResources->StateTextures[1].InitResource();
 			ParticleResources->RenderAttributesTexture.InitResource();
 			ParticleResources->SimulationAttributesTexture.InitResource();
-			ParticleResources->SortedVertexBuffer.InitResource();
 
 			FRHITexture* AttributeTextures[2];
 			AttributeTextures[0] = ParticleResources->RenderAttributesTexture.TextureRHI;
@@ -425,7 +408,6 @@ public:
 			ParticleResources->StateTextures[1].ReleaseResource();
 			ParticleResources->RenderAttributesTexture.ReleaseResource();
 			ParticleResources->SimulationAttributesTexture.ReleaseResource();
-			ParticleResources->SortedVertexBuffer.ReleaseResource();
 		});
 	}
 
@@ -680,7 +662,7 @@ TGlobalResource<FGPUSpriteVertexDeclaration> GGPUSpriteVertexDeclaration;
 struct FGPUSpriteMeshDataUserData  : public FOneFrameResource
 {
 	int32 SortedOffset = 0;
-	FRHIShaderResourceView* SortedParticleIndicesSRV = nullptr;
+	FShaderResourceViewRHIRef SortedParticleIndicesSRV;
 };
 
 /**
@@ -2546,10 +2528,13 @@ public:
 /**
  * Vertex buffer used to hold particle indices.
  */
-class FGPUParticleVertexBuffer : public FParticleIndicesVertexBuffer
+class FGPUParticleVertexBuffer : public FVertexBuffer
 {
 public:
 
+	/** Shader resource view of the vertex buffer. */
+	FShaderResourceViewRHIRef VertexBufferSRV;
+
 	/** The number of particles referenced by this vertex buffer. */
 	int32 ParticleCount;
 
@@ -2588,6 +2573,13 @@ public:
 			VertexBufferSRV = RHICreateShaderResourceView(VertexBufferRHI, BufferStride, PF_G16R16F);
 		}
 	}
+
+	/** Release RHI resources. */
+	virtual void ReleaseRHI() override
+	{
+		VertexBufferSRV.SafeRelease();
+		FVertexBuffer::ReleaseRHI();
+	}
 };
 
 /**
@@ -3065,6 +3057,7 @@ public:
 				const bool bTranslucent = RendersWithTranslucentMaterial();
 				const bool bAllowSorting = FXConsoleVariables::bAllowGPUSorting
 					&& FeatureLevel == ERHIFeatureLevel::SM5
+					// && RHISupportsComputeShaders(ViewFamily.GetShaderPlatform())
 					&& bTranslucent;
 
 				// Iterate over views and assign parameters for each.
@@ -3081,14 +3074,16 @@ public:
 				FGPUSpriteMeshDataUserData* MeshBatchUserData = nullptr;
 				if (bAllowSorting && SortMode == PSORTMODE_DistanceToView)
 				{
-					MeshBatchUserData = &Collector.AllocateOneFrameResource<FGPUSpriteMeshDataUserData>();
-
 					// Extensibility TODO: This call to AddSortedGPUSimulation is very awkward. When rendering a frame we need to
 					// accumulate all GPU particle emitters that need to be sorted. That is so they can be sorted in one big radix
 					// sort for efficiency. Ideally that state is per-scene renderer but the renderer doesn't know anything about particles.
-					check(SimulationResources->SortedVertexBuffer.IsInitialized());
-					MeshBatchUserData->SortedOffset = FXSystem->AddSortedGPUSimulation(Simulation, View->ViewMatrices.GetViewOrigin());
-					MeshBatchUserData->SortedParticleIndicesSRV = SimulationResources->SortedVertexBuffer.VertexBufferSRV;
+					FGPUSortManager::FAllocationInfo SortedIndicesInfo;
+					if (FXSystem->AddSortedGPUSimulation(Simulation, View->ViewMatrices.GetViewOrigin(), bTranslucent, SortedIndicesInfo))
+					{
+						MeshBatchUserData = &Collector.AllocateOneFrameResource<FGPUSpriteMeshDataUserData>();
+						MeshBatchUserData->SortedOffset = SortedIndicesInfo.BufferOffset;
+						MeshBatchUserData->SortedParticleIndicesSRV = SortedIndicesInfo.BufferSRV;
+					}
 				}
 				check(Simulation->VertexBuffer.IsInitialized());
 				VertexFactory.SetUnsortedParticleIndicesSRV(Simulation->VertexBuffer.VertexBufferSRV);
@@ -4521,18 +4516,6 @@ void FFXSystem::InitGPUResources()
 	{
 		check(ParticleSimulationResources);
 		ParticleSimulationResources->Init();
-
-		ParticleSimulationResources->ParticleSortBuffers.SetBufferSize(GParticleSimulationTextureSizeX * GParticleSimulationTextureSizeY);
-		FParticleSimulationResources* InParticleSimulationResources = ParticleSimulationResources;
-		ENQUEUE_RENDER_COMMAND(FInitParticleSortBuffersCommand)(
-			[InParticleSimulationResources](FRHICommandList& RHICmdList)
-			{
-				InParticleSimulationResources->ParticleSortBuffers.InitRHI();
-
-				// Initialize SortedVertexBuffer to a valid resource, ensuring it can be used in GetDynamicMeshElementsEmitter() 
-				InParticleSimulationResources->SortedVertexBuffer.VertexBufferRHI = InParticleSimulationResources->ParticleSortBuffers.GetSortedVertexBufferRHI(0);
-				InParticleSimulationResources->SortedVertexBuffer.VertexBufferSRV = InParticleSimulationResources->ParticleSortBuffers.GetSortedVertexBufferSRV(0);
-			});
 	}
 }
 
@@ -4542,12 +4525,6 @@ void FFXSystem::ReleaseGPUResources()
 	{
 		check(ParticleSimulationResources);
 		ParticleSimulationResources->Release();
-		FParticleSimulationResources* InParticleSimulationResources = ParticleSimulationResources;
-		ENQUEUE_RENDER_COMMAND(FReleaseParticleSortBuffersCommand)(
-			[InParticleSimulationResources](FRHICommandList& RHICmdList)
-			{
-				InParticleSimulationResources->ParticleSortBuffers.ReleaseRHI();
-			});
 	}
 }
 
@@ -4590,18 +4567,43 @@ void FFXSystem::RemoveGPUSimulation(FParticleSimulationGPU* Simulation)
 	}
 }
 
-int32 FFXSystem::AddSortedGPUSimulation(FParticleSimulationGPU* Simulation, const FVector& ViewOrigin)
+bool FFXSystem::AddSortedGPUSimulation(FParticleSimulationGPU* Simulation, const FVector& ViewOrigin, bool bIsTranslucent, FGPUSortManager::FAllocationInfo& OutInfo)
 {
 	LLM_SCOPE(ELLMTag::Particles);
+	check(RHISupportsComputeShaders(ShaderPlatform));
+
+	const EGPUSortFlags SortFlags = 
+		EGPUSortFlags::KeyGenAfterPostRenderOpaque |
+		EGPUSortFlags::ValuesAsG16R16F | 
+		EGPUSortFlags::LowPrecisionKeys | 
+		EGPUSortFlags::SortAfterPostRenderOpaque;
+
+	// Currently opaque materials would need SortAfterPreRender but this is incompatible with KeyGenAfterPostRenderOpaque
+	if (bIsTranslucent && GPUSortManager && GPUSortManager->AddTask(OutInfo, Simulation->VertexBuffer.ParticleCount, SortFlags))
+	{
+		SimulationsToSort.Emplace(Simulation->VertexBuffer.VertexBufferSRV, ViewOrigin, (uint32)Simulation->VertexBuffer.ParticleCount, OutInfo);
+		return true;
+	}
+	else
+	{
+		return false;
+	}
+}
 
-	check(FeatureLevel == ERHIFeatureLevel::SM5);
-	const int32 BufferOffset = ParticleSimulationResources->SortedParticleCount;
-	ParticleSimulationResources->SortedParticleCount += Simulation->VertexBuffer.ParticleCount;
-	FParticleSimulationSortInfo* SortInfo = new(ParticleSimulationResources->SimulationsToSort) FParticleSimulationSortInfo();
-	SortInfo->VertexBufferSRV = Simulation->VertexBuffer.VertexBufferSRV;
-	SortInfo->ViewOrigin = ViewOrigin;
-	SortInfo->ParticleCount = Simulation->VertexBuffer.ParticleCount;
-	return BufferOffset;
+void FFXSystem::GenerateSortKeys(FRHICommandListImmediate& RHICmdList, int32 BatchId, int32 NumElementsInBatch, EGPUSortFlags Flags, FRHIUnorderedAccessView* KeysUAV, FRHIUnorderedAccessView* ValuesUAV)
+{
+	check(EnumHasAnyFlags(Flags, EGPUSortFlags::KeyGenAfterPostRenderOpaque));
+	check(EnumHasAnyFlags(Flags, EGPUSortFlags::LowPrecisionKeys));
+
+	// First generate keys for each emitter to be sorted.
+	const int32 TotalParticleCount = GenerateParticleSortKeys(
+		RHICmdList,
+		KeysUAV,
+		ValuesUAV,
+		ParticleSimulationResources->GetVisualizeStateTextures().PositionTextureRHI,
+		SimulationsToSort,
+		FeatureLevel,
+		BatchId);
 }
 
 void FFXSystem::AdvanceGPUParticleFrame(bool bAllowGPUParticleUpdate)
@@ -4614,22 +4616,7 @@ void FFXSystem::AdvanceGPUParticleFrame(bool bAllowGPUParticleUpdate)
 
 	// Reset the list of sorted simulations. As PreRenderView is called on GPU simulations we'll
 	// allocate space for them in the sorted index buffer.
-	ParticleSimulationResources->SimulationsToSort.Reset();
-	ParticleSimulationResources->SortedParticleCount = 0;
-}
-
-void FFXSystem::SortGPUParticles(FRHICommandListImmediate& RHICmdList)
-{
-	if (ParticleSimulationResources->SimulationsToSort.Num() > 0)
-	{
-		SortParticlesGPU(
-			RHICmdList,
-			ParticleSimulationResources->ParticleSortBuffers,
-			ParticleSimulationResources->GetVisualizeStateTextures().PositionTextureRHI,
-			ParticleSimulationResources->SimulationsToSort,
-			GetFeatureLevel()
-			);
-	}
+	SimulationsToSort.Reset();
 }
 
 /**
diff --git a/Engine/Source/Runtime/Engine/Private/Particles/ParticleSortingGPU.cpp b/Engine/Source/Runtime/Engine/Private/Particles/ParticleSortingGPU.cpp
--- a/Engine/Source/Runtime/Engine/Private/Particles/ParticleSortingGPU.cpp
+++ b/Engine/Source/Runtime/Engine/Private/Particles/ParticleSortingGPU.cpp
@@ -50,6 +50,7 @@ public:
 	static bool ShouldCompilePermutation(const FGlobalShaderPermutationParameters& Parameters)
 	{
 		return IsFeatureLevelSupported(Parameters.Platform, ERHIFeatureLevel::SM5);
+		// return RHISupportsComputeShaders(Parameters.Platform);
 	}
 
 	static void ModifyCompilationEnvironment(const FGlobalShaderPermutationParameters& Parameters, FShaderCompilerEnvironment& OutEnvironment)
@@ -175,72 +176,59 @@ IMPLEMENT_SHADER_TYPE(,FParticleSortKeyGenCS,TEXT("/Engine/Private/ParticleSortK
  * @param SimulationsToSort - A list of simulations to generate sort keys for.
  * @returns the total number of particles being sorted.
  */
-static int32 GenerateParticleSortKeys(
+int32 GenerateParticleSortKeys(
 	FRHICommandListImmediate& RHICmdList,
 	FRHIUnorderedAccessView* KeyBufferUAV,
 	FRHIUnorderedAccessView* SortedVertexBufferUAV,
 	FRHITexture2D* PositionTextureRHI,
 	const TArray<FParticleSimulationSortInfo>& SimulationsToSort,
-	ERHIFeatureLevel::Type FeatureLevel
+	ERHIFeatureLevel::Type FeatureLevel,
+	int32 BatchId
 	)
 {
-	SCOPED_DRAW_EVENT(RHICmdList, ParticleSortKeyGen);
-	check(FeatureLevel == ERHIFeatureLevel::SM5);
+	check(FeatureLevel == ERHIFeatureLevel::SM5 || FeatureLevel == ERHIFeatureLevel::ES3_1);
 
 	FParticleKeyGenParameters KeyGenParameters;
 	FParticleKeyGenUniformBufferRef KeyGenUniformBuffer;
 	const uint32 MaxGroupCount = 128;
 	int32 TotalParticleCount = 0;
 
-	FRHIUnorderedAccessView* OutputUAVs[2];
-	OutputUAVs[0] = KeyBufferUAV;
-	OutputUAVs[1] = SortedVertexBufferUAV;
-
-	//make sure our outputs are safe to write to.
-	RHICmdList.TransitionResources(EResourceTransitionAccess::ERWBarrier, EResourceTransitionPipeline::EGfxToCompute, OutputUAVs, 2);
-
 	// Grab the shader, set output.
 	TShaderMapRef<FParticleSortKeyGenCS> KeyGenCS(GetGlobalShaderMap(FeatureLevel));
 	RHICmdList.SetComputeShader(KeyGenCS->GetComputeShader());
 	KeyGenCS->SetOutput(RHICmdList, KeyBufferUAV, SortedVertexBufferUAV);
 	KeyGenCS->SetPositionTextures(RHICmdList, PositionTextureRHI);
 
+	FRHIUnorderedAccessView* OutputUAVs[] = { KeyBufferUAV, SortedVertexBufferUAV };
 	// For each simulation, generate keys and store them in the sorting buffers.
-	const int32 SimulationCount = SimulationsToSort.Num();
-	for (int32 SimulationIndex = 0; SimulationIndex < SimulationCount; ++SimulationIndex)
+	for (const FParticleSimulationSortInfo& SortInfo : SimulationsToSort)
 	{
-		const FParticleSimulationSortInfo& SortInfo = SimulationsToSort[SimulationIndex];
-
-		// Create the uniform buffer.
-		const uint32 ParticleCount = SortInfo.ParticleCount;
-		const uint32 AlignedParticleCount = ((ParticleCount + PARTICLE_KEY_GEN_THREAD_COUNT - 1) & (~(PARTICLE_KEY_GEN_THREAD_COUNT - 1)));
-		const uint32 ChunkCount = AlignedParticleCount / PARTICLE_KEY_GEN_THREAD_COUNT;
-		const uint32 GroupCount = FMath::Clamp<uint32>( ChunkCount, 1, MaxGroupCount );
-		KeyGenParameters.ViewOrigin = SortInfo.ViewOrigin;
-		KeyGenParameters.ChunksPerGroup = ChunkCount / GroupCount;
-		KeyGenParameters.ExtraChunkCount = ChunkCount % GroupCount;
-		KeyGenParameters.OutputOffset = TotalParticleCount;
-		KeyGenParameters.EmitterKey = SimulationIndex << 16;
-		KeyGenParameters.KeyCount = ParticleCount;
-		KeyGenUniformBuffer = FParticleKeyGenUniformBufferRef::CreateUniformBufferImmediate( KeyGenParameters, UniformBuffer_SingleDraw );
-
-		// Dispatch.
-		KeyGenCS->SetParameters(RHICmdList, KeyGenUniformBuffer, SortInfo.VertexBufferSRV);
-		DispatchComputeShader(RHICmdList, *KeyGenCS, GroupCount, 1, 1);
-
-		//we may be able to remove this transition if each step isn't dependent on the previous one.
-		RHICmdList.TransitionResources(EResourceTransitionAccess::ERWBarrier, EResourceTransitionPipeline::EComputeToCompute, OutputUAVs, 2);
-
-		// Update offset in to the buffer.
-		TotalParticleCount += ParticleCount;
+		if (SortInfo.AllocationInfo.SortBatchId == BatchId)
+		{
+			// Create the uniform buffer.
+			const uint32 ParticleCount = SortInfo.ParticleCount;
+			const uint32 AlignedParticleCount = ((ParticleCount + PARTICLE_KEY_GEN_THREAD_COUNT - 1) & (~(PARTICLE_KEY_GEN_THREAD_COUNT - 1)));
+			const uint32 ChunkCount = AlignedParticleCount / PARTICLE_KEY_GEN_THREAD_COUNT;
+			const uint32 GroupCount = FMath::Clamp<uint32>( ChunkCount, 1, MaxGroupCount );
+			KeyGenParameters.ViewOrigin = SortInfo.ViewOrigin;
+			KeyGenParameters.ChunksPerGroup = ChunkCount / GroupCount;
+			KeyGenParameters.ExtraChunkCount = ChunkCount % GroupCount;
+			KeyGenParameters.OutputOffset = SortInfo.AllocationInfo.BufferOffset;
+			KeyGenParameters.EmitterKey = (uint32)SortInfo.AllocationInfo.ElementIndex << 16;
+			KeyGenParameters.KeyCount = ParticleCount;
+			KeyGenUniformBuffer = FParticleKeyGenUniformBufferRef::CreateUniformBufferImmediate( KeyGenParameters, UniformBuffer_SingleDraw );
+
+			// Dispatch.
+			KeyGenCS->SetParameters(RHICmdList, KeyGenUniformBuffer, SortInfo.VertexBufferSRV);
+			DispatchComputeShader(RHICmdList, *KeyGenCS, GroupCount, 1, 1);
+
+			// TR-KeyGen : No sync needed between tasks since they update different parts of the data (assuming it's ok if cache line overlap).
+			RHICmdList.TransitionResources(EResourceTransitionAccess::ERWNoBarrier, EResourceTransitionPipeline::EComputeToCompute, OutputUAVs, UE_ARRAY_COUNT(OutputUAVs));
+		}
 	}
 
 	// Clear the output buffer.
 	KeyGenCS->UnbindBuffers(RHICmdList);
-
-	//make sure our outputs are readable as SRVs to further gfx steps.
-	RHICmdList.TransitionResources(EResourceTransitionAccess::EReadable, EResourceTransitionPipeline::EComputeToGfx, OutputUAVs, 2);
-
 	return TotalParticleCount;
 }
 
@@ -267,18 +255,6 @@ void FParticleSortBuffers::InitRHI()
 
 			VertexBufferSortSRVs[BufferIndex] = RHICreateShaderResourceView(VertexBuffers[BufferIndex], /*Stride=*/ sizeof(uint32), PF_R32_UINT);
 			VertexBufferSortUAVs[BufferIndex] = RHICreateUnorderedAccessView(VertexBuffers[BufferIndex], PF_R32_UINT);
-
-			if (bAsInt32)
-			{
-				VertexBufferSRVs[BufferIndex] = RHICreateShaderResourceView(VertexBuffers[BufferIndex], /*Stride=*/ sizeof(int32), PF_R32_SINT);;
-				VertexBufferUAVs[BufferIndex] = RHICreateUnorderedAccessView(VertexBuffers[BufferIndex], PF_R32_SINT);
-			}
-			else
-			{
-				VertexBufferSRVs[BufferIndex] = RHICreateShaderResourceView(VertexBuffers[BufferIndex], /*Stride=*/ sizeof(FFloat16) * 2, PF_G16R16F);
-				VertexBufferUAVs[BufferIndex] = RHICreateUnorderedAccessView(VertexBuffers[BufferIndex], PF_G16R16F);
-			}
-
 		}
 	}
 }
@@ -296,8 +272,6 @@ void FParticleSortBuffers::ReleaseRHI()
 
 		VertexBufferSortUAVs[BufferIndex].SafeRelease();
 		VertexBufferSortSRVs[BufferIndex].SafeRelease();
-		VertexBufferUAVs[BufferIndex].SafeRelease();
-		VertexBufferSRVs[BufferIndex].SafeRelease();
 		VertexBuffers[BufferIndex].SafeRelease();
 	}
 }
@@ -319,60 +293,3 @@ FGPUSortBuffers FParticleSortBuffers::GetSortBuffers()
 
 	return SortBuffers;
 }
-
-/*------------------------------------------------------------------------------
-	Public interface.
-------------------------------------------------------------------------------*/
-
-/**
- * Sort particles on the GPU.
- * @param ParticleSortBuffers - Buffers to use while sorting GPU particles.
- * @param PositionTextureRHI - Texture containing world space position for all particles.
- * @param SimulationsToSort - A list of simulations that must be sorted.
- */
-void SortParticlesGPU(
-	FRHICommandListImmediate& RHICmdList,
-	FParticleSortBuffers& ParticleSortBuffers,
-	FRHITexture2D* PositionTextureRHI,
-	const TArray<FParticleSimulationSortInfo>& SimulationsToSort,
-	ERHIFeatureLevel::Type FeatureLevel
-	)
-{
-	SCOPED_DRAW_EVENTF(RHICmdList, ParticleSort, TEXT("ParticleSort_%d"), SimulationsToSort.Num());
-
-	// Ensure the sorted vertex buffers are not currently bound as input streams.
-	// They should only ever be bound to streams 0 or 1, so clear them.
-	{
-		const int32 StreamCount = 2;
-		for (int32 StreamIndex = 0; StreamIndex < StreamCount; ++StreamIndex)
-		{
-			RHICmdList.SetStreamSource(StreamIndex, nullptr, 0);
-		}
-	}
-
-	// Buffer index should start so that we always end up on buffer 0 so that callers do not need to dynmically change which buffer they are using
-	const uint32 EmitterKeyMask = (1 << FMath::CeilLogTwo(SimulationsToSort.Num())) - 1;
-	const uint32 KeyMask = (EmitterKeyMask << 16) | 0xFFFF;
-
-	const int32 BufferIndex = GetGPUSortPassCount(KeyMask) & 1;
-
-	// First generate keys for each emitter to be sorted.
-	const int32 TotalParticleCount = GenerateParticleSortKeys(
-		RHICmdList,
-		ParticleSortBuffers.GetKeyBufferUAV(BufferIndex),
-		ParticleSortBuffers.GetVertexBufferUAV(BufferIndex),
-		PositionTextureRHI,
-		SimulationsToSort,
-		FeatureLevel
-		);
-
-	// Update stats.
-	INC_DWORD_STAT_BY( STAT_SortedGPUEmitters, SimulationsToSort.Num() );
-	INC_DWORD_STAT_BY( STAT_SortedGPUParticles, TotalParticleCount );
-
-	// Now sort the particles based on the generated keys.
-	FGPUSortBuffers SortBuffers = ParticleSortBuffers.GetSortBuffers();
-
-	const int32 FinalBufferIndex = SortGPUBuffers(RHICmdList, SortBuffers, BufferIndex, KeyMask, TotalParticleCount, FeatureLevel);
-	ensure(FinalBufferIndex == 0);
-}
diff --git a/Engine/Source/Runtime/Engine/Private/Particles/ParticleSortingGPU.h b/Engine/Source/Runtime/Engine/Private/Particles/ParticleSortingGPU.h
--- a/Engine/Source/Runtime/Engine/Private/Particles/ParticleSortingGPU.h
+++ b/Engine/Source/Runtime/Engine/Private/Particles/ParticleSortingGPU.h
@@ -8,144 +8,39 @@
 
 #include "CoreMinimal.h"
 #include "RenderResource.h"
-
-struct FGPUSortBuffers;
-
-/**
- * Buffers in GPU memory used to sort particles.
- */
-class ENGINE_API FParticleSortBuffers : public FRenderResource
-{
-public:
-
-	/** Initialization constructor. */
-	explicit FParticleSortBuffers(bool InAsInt32 = false)
-		: BufferSize(0)
-		, bAsInt32(InAsInt32)
-	{
-	}
-
-	void SetBufferSize(int32 InBufferSize)
-	{
-		BufferSize = InBufferSize;
-	}
-
-	/**
-	 * Initialize RHI resources.
-	 */
-	virtual void InitRHI() override;
-
-	/**
-	 * Release RHI resources.
-	 */
-	virtual void ReleaseRHI() override;
-
-	/**
-	 * Retrieve the UAV for writing particle sort keys.
-	 */
-	FRHIUnorderedAccessView* GetKeyBufferUAV(int32 BufferIndex)
-	{
-		check((BufferIndex & 0xFFFFFFFE) == 0);
-		return KeyBufferUAVs[BufferIndex];
-	}
-
-	/**
-	 * Retrieve the UAV for writing particle vertices.
-	 * bAsUint : whether to return a G16R16 view or a Uint32 view.
-	 */
-	FORCEINLINE FRHIUnorderedAccessView* GetVertexBufferUAV(int32 BufferIndex)
-	{
-		check((BufferIndex & 0xFFFFFFFE) == 0);
-		return VertexBufferUAVs[BufferIndex];
-	}
-
-	/**
-	 * Retrieve buffers needed to sort on the GPU.
-	 */
-	FGPUSortBuffers GetSortBuffers();
-
-	/**
-	 * Retrieve the sorted vertex buffer that results will always be located at.
-	 */
-	FRHIVertexBuffer* GetSortedVertexBufferRHI(int32 BufferIndex)
-	{
-		check((BufferIndex & 0xFFFFFFFE) == 0);
-		return VertexBuffers[BufferIndex];
-	}
-
-	/**
-	 * Retrieve the SRV that sort results will always be located at.
-	 */
-	FRHIShaderResourceView* GetSortedVertexBufferSRV(int32 BufferIndex)
-	{
-		check((BufferIndex & 0xFFFFFFFE) == 0);
-		return VertexBufferSRVs[BufferIndex];
-	}
-	   
-	/**
-	 * Retrieve the UAV for the sorted vertex buffer at the given index.
-	 */
-	FRHIUnorderedAccessView* GetSortedVertexBufferUAV(int32 BufferIndex)
-	{
-		check((BufferIndex & 0xFFFFFFFE) == 0);
-		return VertexBufferUAVs[BufferIndex];
-	}
-
-
-	/**
-	 * Get the size allocated for sorted vertex buffers.
-	 */
-	int32 GetSize() { return BufferSize; }
-
-private:
-
-	/** Vertex buffer storage for particle sort keys. */
-	FVertexBufferRHIRef KeyBuffers[2];
-	/** Shader resource view for particle sort keys. */
-	FShaderResourceViewRHIRef KeyBufferSRVs[2];
-	/** Unordered access view for particle sort keys. */
-	FUnorderedAccessViewRHIRef KeyBufferUAVs[2];
-
-	/** Vertex buffer containing sorted particle vertices. */
-	FVertexBufferRHIRef VertexBuffers[2];
-	/** Shader resource view for reading particle vertices out of the sorting buffer. */
-	FShaderResourceViewRHIRef VertexBufferSRVs[2];
-	/** Unordered access view for writing particle vertices in to the sorting buffer. */
-	FUnorderedAccessViewRHIRef VertexBufferUAVs[2];
-	/** Shader resource view for sorting particle vertices. */
-	FShaderResourceViewRHIRef VertexBufferSortSRVs[2];
-	/** Unordered access view for sorting particle vertices. */
-	FUnorderedAccessViewRHIRef VertexBufferSortUAVs[2];
-
-	/** Size allocated for buffers. */
-	int32 BufferSize;
-	/** Whether to allocate UAV and SRV as Int32 instead of G16R16F. */
-	bool bAsInt32;
-};
+#include "GPUSortManager.h"
 
 /**
  * The information required to sort particles belonging to an individual simulation.
  */
 struct FParticleSimulationSortInfo
 {
+	typedef FGPUSortManager::FAllocationInfo FAllocationInfo;
+
+	FParticleSimulationSortInfo(FRHIShaderResourceView* InVertexBufferSRV, const FVector& InViewOrigin, uint32 InParticleCount, const FAllocationInfo& InAllocationInfo) 
+		: VertexBufferSRV(InVertexBufferSRV)
+		, ViewOrigin(InViewOrigin) 
+		, ParticleCount(InParticleCount) 
+		, AllocationInfo(InAllocationInfo)
+	{}
+
 	/** Vertex buffer containing indices in to the particle state texture. */
-	FRHIShaderResourceView* VertexBufferSRV;
+	FRHIShaderResourceView* VertexBufferSRV = nullptr;
 	/** World space position from which to sort. */
 	FVector ViewOrigin;
 	/** The number of particles in the simulation. */
-	uint32 ParticleCount;
+	uint32 ParticleCount = 0;
+
+	// The GPUSortManager bindings for this sort task.
+	FGPUSortManager::FAllocationInfo AllocationInfo;
 };
 
-/**
- * Sort particles on the GPU.
- * @param ParticleSortBuffers - Buffers to use while sorting GPU particles.
- * @param PositionTextureRHI - Texture containing world space position for all particles.
- * @param SimulationsToSort - A list of simulations that must be sorted.
- */
-void SortParticlesGPU(
+int32 GenerateParticleSortKeys(
 	FRHICommandListImmediate& RHICmdList,
-	FParticleSortBuffers& ParticleSortBuffers,
+	FRHIUnorderedAccessView* KeyBufferUAV,
+	FRHIUnorderedAccessView* SortedVertexBufferUAV,
 	FRHITexture2D* PositionTextureRHI,
 	const TArray<FParticleSimulationSortInfo>& SimulationsToSort,
-	ERHIFeatureLevel::Type FeatureLevel
-	);
+	ERHIFeatureLevel::Type FeatureLevel,
+	int32 BatchId = INDEX_NONE);
+
diff --git a/Engine/Source/Runtime/Engine/Public/FXSystem.h b/Engine/Source/Runtime/Engine/Public/FXSystem.h
--- a/Engine/Source/Runtime/Engine/Public/FXSystem.h
+++ b/Engine/Source/Runtime/Engine/Public/FXSystem.h
@@ -16,6 +16,7 @@ class UVectorFieldComponent;
 struct FGPUSpriteEmitterInfo;
 struct FGPUSpriteResourceData;
 struct FParticleEmitterInstance;
+class FGPUSortManager;
 
 /*-----------------------------------------------------------------------------
 	Forward declarations.
@@ -110,7 +111,8 @@ inline bool RHISupportsGPUParticles()
 }
 
 class FFXSystemInterface;
-DECLARE_DELEGATE_RetVal_TwoParams(FFXSystemInterface*, FCreateCustomFXSystemDelegate, ERHIFeatureLevel::Type, EShaderPlatform);
+class FGPUSortManager;
+DECLARE_DELEGATE_RetVal_ThreeParams(FFXSystemInterface*, FCreateCustomFXSystemDelegate, ERHIFeatureLevel::Type, EShaderPlatform, FGPUSortManager*);
 
 /*-----------------------------------------------------------------------------
 	The interface to the FX system runtime.
@@ -230,6 +232,9 @@ public:
 
 	bool IsPendingKill() const { return bIsPendingKill; }
 
+	/** Get the shared SortManager, used in the rendering loop to call FGPUSortManager::OnPreRender() and FGPUSortManager::OnPostRenderOpaque() */
+	virtual FGPUSortManager* GetGPUSortManager() const = 0;
+
 protected:
 	
 	friend class FFXSystemSet;
diff --git a/Engine/Source/Runtime/Engine/Public/GPUSortManager.h b/Engine/Source/Runtime/Engine/Public/GPUSortManager.h
new file mode 100644
--- /dev/null
+++ b/Engine/Source/Runtime/Engine/Public/GPUSortManager.h
@@ -0,0 +1,508 @@
+// Copyright Epic Games, Inc. All Rights Reserved.
+
+/*=============================================================================
+	GPUSortManager.h: Interface manage sorting between different FXSystemInterface
+=============================================================================*/
+
+#pragma once
+
+#include "CoreMinimal.h"
+#include "Templates/RefCounting.h"
+#include "RHIResources.h"
+
+/**
+ * Copy one uint buffer into several overlapping targets.
+ * Each target is expected to be bigger than the previous one, so that copy resumes where it left.
+ * Ex : (ABCDEFGHIJK) -> (ABC--, ---DEFGHI-, ---------JK---)
+ * This is essentially used to manage growable buffers in cases the
+ * binding needs to be set before the final size required is known. 
+ * In this scenario, the smaller buffers are temporary and only the final (biggest) buffer
+ * becomes persistent, this is why it has apparently unused space at the beginning but
+ * it will be used the next frame to hold all of the data.
+ * Note that required calls to RHICmdList.TransitionResource need to be handled before and after this function.
+ *
+ * @param RHICmdList - The command list used to issue the dispatch.
+ * @param FeatureLevel - The current feature level, used to access the global shadermap.
+ * @param SourceSRV - The source uint buffer to be copied into the others.
+ * @param TargetUAVs - The destination overlapping uint buffers.
+ * @param TargetSizes - The copy size of each of the buffers. Once a buffer size is reached, the copy targets the next buffer.
+ * @param StartingOffset - The starting position at which the copy starts. Applies for both the source and targets.
+ * @param NumTargets - The number of elements in TargetUAVs and TargetSizes.
+ */
+ENGINE_API void CopyUIntBufferToTargets(FRHICommandListImmediate& RHICmdList, ERHIFeatureLevel::Type FeatureLevel, FRHIShaderResourceView* SourceSRV, FRHIUnorderedAccessView*const* TargetUAVs, int32* TargetSizes, int32 StartingOffset, int32 NumTargets);
+
+struct FGPUSortBuffers;
+
+/** Buffers in GPU memory used to sort particles. */
+class FParticleSortBuffers : public FRenderResource
+{
+public:
+
+	void SetBufferSize(int32 InBufferSize)
+	{
+		BufferSize = InBufferSize;
+	}
+
+	/** Initialize RHI resources. */
+	virtual void InitRHI() override;
+
+	/** Release RHI resources. */
+	virtual void ReleaseRHI() override;
+
+	/** Retrieve the UAV for writing particle sort keys. */
+	FRHIUnorderedAccessView* GetKeyBufferUAV(int32 BufferIndex)
+	{
+		check((BufferIndex & 0xFFFFFFFE) == 0);
+		return KeyBufferUAVs[BufferIndex];
+	}
+
+	/** Retrieve buffers needed to sort on the GPU. */
+	FGPUSortBuffers GetSortBuffers();
+
+	/** Retrieve the sorted vertex buffer that results will always be located at. */
+	FRHIVertexBuffer* GetSortedVertexBufferRHI(int32 BufferIndex)
+	{
+		check((BufferIndex & 0xFFFFFFFE) == 0);
+		return VertexBuffers[BufferIndex];
+	}
+
+	/** Get the size allocated for sorted vertex buffers. */
+	int32 GetSize() { return BufferSize; }
+
+private:
+
+	/** Vertex buffer storage for particle sort keys. */
+	FVertexBufferRHIRef KeyBuffers[2];
+	/** Shader resource view for particle sort keys. */
+	FShaderResourceViewRHIRef KeyBufferSRVs[2];
+	/** Unordered access view for particle sort keys. */
+	FUnorderedAccessViewRHIRef KeyBufferUAVs[2];
+
+	/** Vertex buffer containing sorted particle vertices. */
+	FVertexBufferRHIRef VertexBuffers[2];
+	/** Shader resource view for sorting particle vertices. */
+	FShaderResourceViewRHIRef VertexBufferSortSRVs[2];
+	/** Unordered access view for sorting particle vertices. */
+	FUnorderedAccessViewRHIRef VertexBufferSortUAVs[2];
+
+	/** Size allocated for buffers. */
+	int32 BufferSize = 0;
+};
+
+/** Different sort flags used to define constraint for FGPUSortManager tasks. */
+enum class EGPUSortFlags : uint32
+{
+	None			= 0x00,
+	/**
+	 * Sorting happens on either 16 bit keys or 32 bit uint keys. Those can come from float point formats being converted to uint.
+	 * Because sort tasks are grouped into batches, the keys also have to start with the batch element index, which ends up either
+	 * increasing the 16 bit keys or taking space in the 32 bit keys (by the number of bits required to store the element indices).
+	 * The more task are created, the more bits are required to encode the element index (up to a limit of 16^2 elements).
+	 * Low precision keys have the advantage of being quicker on the GPU when performing the radix sort, which is affected by the number of bits.
+	 */
+	LowPrecisionKeys	= 0x01,		// Low precision keys ((EleIdx << 16) | (16 bits key))
+	HighPrecisionKeys	= 0x02,		// High precision keys ((EleIdx << (32 - IdxBitCount)) | (32 bits key >> IdxBitCount))
+
+	AnyKeyPrecision		= 0x03,
+
+	/**
+	 * The sort task are created in FPrimitiveSceneProxy::GetDynamicMeshElements() but have different requirements
+	 * in terms of when the keys can be generated, and when the sorted values are required within the frame.
+	 * Cascade renders the particles after GPU simulation (either in PreRender(), or in PostRenderOpaque() see ::CollisionDepthBuffer).
+	 * Niagara, on the other hand, renders the particle before they get updated by simulation (either in PreInitViews() or PostRenderOpaque())
+	 * For Niagara emitters updated in PreInitViews(), this happens before GetDynamicMeshElements() so doesn't impose any restriction in the end.
+	 * This means each system has some flexibility on when the sort keys can be generated depending on the simulation point for each emitter.
+	 */
+	KeyGenAfterPreRender			= 0x04,		// Typically Cascade emitters simulated in PreRender() or any Niagara emitters.
+	KeyGenAfterPostRenderOpaque		= 0x08,		// Typically any Cascade emitters, or Niagara emitters simulated in PreInitViews().
+
+	AnyKeyGenLocation				= 0x0C,
+
+	/**
+	 * Once keys are generated, they must be sorted and the indices only need to be ready before the primitive are rendered.
+	 * The (frame) location only depends on whether the emitter uses opaque materials or translucent materials.
+	 * It is possible to have an incorrect setup if the keys can only be generated after PostRenderOpaque() but the sorting must be ready after PreRender().
+	 * This happens with opaque Cascade emitters when simulation needs the depth buffer to do collision tests.
+	 * A sort requests can have both SortIndicesAfterPreRender and SortIndicesAfterPostRenderOpaque when it is allowed to sort keys at either point.
+	 */
+	SortAfterPreRender		 	= 0x10,		// Any Emitter since good for both opaque or translucent
+	SortAfterPostRenderOpaque	= 0x20,		// Emitter using translucent material.
+
+	AnySortLocation				= 0x30,
+
+	/**
+	 *  The sorted values can either be read as 1D or 2D, but doesn't otherwise affect the sorting implementation as long as the storage size is the same.
+	 *  Currently it is 32 bits in both cases. Note that the sorting shader reads the indices as UInt32 (see FParticleSortBuffers::InitRHI()).
+	 */
+	ValuesAsG16R16F	= 0x40,
+	ValuesAsInt32	= 0x80,
+
+	AnyValueFormat	= 0xC0,
+};
+
+ENUM_CLASS_FLAGS(EGPUSortFlags)
+
+/**
+ * Callback used by the FGPUSortManager to make each client system initialize the sort keys for a specific sort batch.
+ * Sort batches are created when GPU sort tasks are registered in FGPUSortManager::AddTask().
+ * Note that the system that is getting this callback must only initialize the data for the elements is has registered in this batch.
+ * 
+ * @param RHICmdList - The command list used to initiate the keys and values on GPU.
+ * @param BatchId - The GPUSortManager batch id (regrouping several similar sort tasks).
+ * @param NumElementsInBatch - The number of elements grouped in the batch (each element maps to a sort task)
+ * @param Flags - Details about the key precision (see EGPUSortFlags::AnyKeyPrecision) and the keygen location (see EGPUSortFlags::AnyKeyGenLocation).
+ * @param KeysUAV - The UAV that holds all the initial keys used to sort the values.
+ * @param ValuesUAV - The UAV that holds the initial values to be sorted accordingly to the keys.
+ */
+ DECLARE_DELEGATE_SixParams(FGPUSortKeyGenDelegate, FRHICommandListImmediate&, int32, int32, EGPUSortFlags, FRHIUnorderedAccessView*, FRHIUnorderedAccessView*);
+
+/**
+ * A manager that handles different GPU sort tasks. 
+ * Each task has different constraints about when the data can be generated and when the sorted results need to be available.
+ * The usecase involes registering a sort task through AddTask() and then initializing the sort data through the KeyGen callback.
+ * The sort manager can be enabled or disabled through "fx.AllowGPUSorting"
+ */
+class ENGINE_API FGPUSortManager : public FRefCountedObject
+{
+public:
+
+	/** Information about the bindings for a given sort task. */
+	struct FAllocationInfo
+	{
+		/** The buffer that will hold the final sorted values */
+		FShaderResourceViewRHIRef BufferSRV;
+		/** The offset within BufferSRV where the sort task data starts. */
+		uint32 BufferOffset = INDEX_NONE;
+		/** The batch Id in which the sort task is grouped. */
+		int32 SortBatchId = INDEX_NONE;
+		/** The element index of the sort task within the batch. */
+		int32 ElementIndex = INDEX_NONE;
+	};
+
+	/** A little helper to generate the batch element keys based on the number of elements in the batch and the sort precision. */
+	struct ENGINE_API FKeyGenInfo
+	{
+		FKeyGenInfo(uint32 NumElements, bool bHighPrecisionKeys);
+
+		/**
+		 * ElementKey = (ElementIndex & EmitterKeyMask) << ElementKeyShift.
+		 * SortKey = (Key32 >> SortKeyShift) & SortKeyMask.
+		 */
+		uint32 ElementKeyMask;
+		uint32 ElementKeyShift;
+		uint32 SortKeyMask;
+		/** (SortKeyMask, SortKeyShift, SortKeySignBit, -) */
+		FUintVector4 SortKeyParams; 
+	};
+
+private:
+
+	/** Holds data relative to client systems of the GPU sort manager. */
+	struct FCallbackInfo
+	{
+		/** A delegate used to initialize the keys and values. */
+		FGPUSortKeyGenDelegate Delegate;
+		/** The flag this client system might use when creating batches. */
+		EGPUSortFlags Flags = EGPUSortFlags::None;
+		/** The name of this client system. Used for GPU markers. */
+		FName Name;
+	};
+
+	typedef TArray<FCallbackInfo> FCallbackArray;
+
+	/** Different settings used to manage the buffer allocations. */
+	struct FSettings
+	{
+		/** The union of all clients flags. */
+		EGPUSortFlags AllowedFlags = EGPUSortFlags::None;
+		/** The buffer slack when increasing or shrinking. See "fx.GPUSort.BufferSlack" */
+		float BufferSlack = 2.f;
+		/** The number of frames before shrinking a buffer that is too big. See "fx.GPUSort.FrameCountBeforeShrinking" */
+		int32 FrameCountBeforeShrinking = 100;
+		/** The minimal size of the sort buffers. See "fx.GPUSort.MinBufferSize" */
+		int32 MinBufferSize = 8192;
+	};
+
+	/** A vertex buffer owning the final sorted values. */
+	class FValueBuffer : public FVertexBuffer
+	{
+	public:
+
+		/**
+		 * Create a value buffer of a given size. The render resource are created immediately.
+		 *
+		 * @param InAllocatedCount - The buffer size, allocated now.
+		 * @param InUsedCount - An initial used count, typically the used size of the previous buffer in FDynamicValueBuffer::ValueBuffers.
+		 * @param InSettings - The buffer allocation settings, in particular what data types are requires for SRV and UAV.
+		 */
+		FValueBuffer(int32 InAllocatedCount, int32 InUsedCount, const FSettings& InSettings);
+
+		~FValueBuffer() { ReleaseRHI(); }
+
+		/** Release the render resources. */
+		void ReleaseRHI() override;
+
+		/** Used for sorting */
+		FShaderResourceViewRHIRef UInt32SRV;
+		FUnorderedAccessViewRHIRef UInt32UAV;
+		/** EGPUSortFlags::IndicesAsInt32 */
+		FShaderResourceViewRHIRef Int32SRV;
+		FUnorderedAccessViewRHIRef Int32UAV;
+		/** EGPUSortFlags::IndicesAsG16R16F */
+		FShaderResourceViewRHIRef G16R16SRV;
+		FUnorderedAccessViewRHIRef G16R16UAV;
+
+		/** The actual buffer size. */
+		int32 AllocatedCount = 0;
+		/** The currently used size. */
+		int32 UsedCount = 0;
+
+		/**
+		 * Allocate some space into the buffer. The buffer needs to be big enough to hold the allocation.
+		 * 
+		 * @param OutInfo - The details about the allocation bindings.
+		 * @param ValueCount - The number of elements to add to the buffer.
+		 * @param Flags - Flag about what data format needs to be bound, see EGPUSortFlags::AnyValueFormat.
+		 */
+		FORCEINLINE void Allocate(FAllocationInfo& OutInfo, int32 ValueCount, EGPUSortFlags Flags);
+	};
+
+	typedef TIndirectArray<FValueBuffer, TInlineAllocator<4>> FValueBufferArray;
+
+	/**
+	 * Encapsulates the idea of a growable FValueBuffer, that dynamically change in size depending on requirements.
+	 * What is specific here is that even after being grown, previous references to the smaller version are still valid for a frame.
+	 * This requirement comes client system registering GPU sort tasks in InitViews(), before the total size required for the frame is known.
+	 */
+	struct FDynamicValueBuffer
+	{
+		~FDynamicValueBuffer() { ReleaseRHI(); }
+
+		/** The current list of increasingly bigger value buffers. There will only be more than one element the frame the buffer grows. */
+		FValueBufferArray ValueBuffers;
+
+		/** A counter used to delay shrinking by several frames. */
+		int32 NumFramesRequiringShrinking = 0;
+
+		/** The last batch flags that used this buffer. Used to help rebind the same SortBatch between frames. */
+		EGPUSortFlags LastSortBatchFlags = EGPUSortFlags::None;
+
+		/** The sort batch using this buffer for this frame. */
+		int32 CurrentSortBatchId = INDEX_NONE;
+
+		/**
+		 * Allocate some space into the buffer for this frame.
+		 * 
+		 * @param OutInfo - The details about the allocation bindings.
+		 * @param InSettings - Buffer allocation settings.
+		 * @param ValueCount - The number of elements to add to the buffer.
+		 * @param Flags - Flag about what data format needs to be bound, see EGPUSortFlags::AnyValueFormat.
+		 */
+		void Allocate(FAllocationInfo& OutInfo, const FSettings& InSettings, int32 ValueCount, EGPUSortFlags Flags);
+
+		/**
+		 * Shrinks the buffer after it has been used this frame if applicable.
+		 * Resets other data so that it can be reused on the next frame.
+		 * This includes releasing the smaller buffers that where allocated 
+		 * this frame if the buffer needed to be grown.
+		 * 
+		 * @param InSettings - Buffer allocation settings, in particular the number of frames required before actually shrinking.
+		 */
+		void SkrinkAndReset(const FSettings& InSettings);
+		
+		/** Release resources */
+		void ReleaseRHI();
+		/** Get currently used size for this frame. */
+		FORCEINLINE int32 GetUsedCount() const { return ValueBuffers.Num() ? ValueBuffers.Last().UsedCount : 0; }
+		/** Get the buffer allocated size . */
+		FORCEINLINE int32 GetAllocatedCount() const { return ValueBuffers.Num() ? ValueBuffers.Last().AllocatedCount : 0; }
+	};
+
+	typedef TIndirectArray<FDynamicValueBuffer, TInlineAllocator<4>> FDynamicValueBufferArray;
+
+	/**
+	 * Defines the order in which sort batches are processed every frame.  
+	 * KeyGenAfterPreRenderAndSortAfterPostRenderOpaque is the most complex case since
+	 * the KeyGen is called after PreRendeR() while the sort only happens after PostRenderOpaque().
+	 * This creates a situation where the FParticleSortBuffers need to be keep bound to the batch
+	 * and can not be reused by other batches with this processing order.
+	 * Note also that KeyGenAfterPreRenderAndSortAfterPostRenderOpaque could include KeyGen after PostRenderOpaque().
+	 */
+	enum class ESortBatchProcessingOrder : uint32
+	{
+		Undefined = 0,
+		KeyGenAndSortAfterPreRender = 1,
+		KeyGenAfterPreRenderAndSortAfterPostRenderOpaque = 2,
+		KeyGenAndSortAfterPostRenderOpaque = 3,
+	};
+
+	/**
+	 * Defines a sort batch that regroups several sort tasks created by FGPUSortManager::AddTask(), becoming sort the batch elements.
+	 * The condition to regroup task in the same batch rely mostly around the key precisions and the sort location (PreRender or PostRenderOpaque).
+	 * Each SortBatch maps to a single GPU sort task.
+	 */
+	struct FSortBatch
+	{
+		/** An Id defining this batch. Passed along FGPUSortKeyGenDelegate. */
+		int32 Id = INDEX_NONE;
+		/** The number of elements (tasks) grouped in this batch. */
+		int32 NumElements = 0;
+
+		/**
+		 * The common requirement flag for this batch. Compatible with each element.
+		 * As tasks are added, the flag can become more restrictive, and new registered tasks 
+		 * could require creating new batches if their requirements are too different.
+		 */
+		EGPUSortFlags Flags = EGPUSortFlags::None;
+
+		/** The buffer to hold the final sorted values. */
+		FDynamicValueBuffer* DynamicValueBuffer = nullptr;
+
+		/** The scratch buffers used to hold the initial keys and values and perform the GPU sort. */
+		FParticleSortBuffers* SortBuffers = nullptr;
+
+		/** This batch processing order in the frame. */
+		ESortBatchProcessingOrder ProcessingOrder = ESortBatchProcessingOrder::Undefined;
+
+		/** Compute the processing order from the current value of Flags. */
+		void UpdateProcessingOrder();
+
+		/**
+		 * Call each delegate registered in the FGPUSortManager so they setup their task data for each element they own in the batch. 
+		 * 
+		 * @param RHICmdList - The command list used to perform the task.
+		 * @param InCallbacks - The client system callbacks, from FGPUSortManager::Register().
+		 * @param KeyGenLocation - The location in the frame where we are generating keys (see EGPUSortFlags::AnyKeyGenLocation).
+		 */
+		void GenerateKeys(FRHICommandListImmediate& RHICmdList, const FCallbackArray& InCallbacks, EGPUSortFlags KeyGenLocation);
+
+		/** 
+		 * Perform the GPU sort and resolve to all buffers referenced in FGPUSortManager::AddTask(), if the buffer grew in size.
+		 *
+		 * @param RHICmdList - The command list used to perform the dispatches.
+		 * @param InFeatureLevel - The current feature level, used to access the global shadermap.
+		 */
+		void SortAndResolve(FRHICommandListImmediate& RHICmdList, ERHIFeatureLevel::Type InFeatureLevel);
+
+		/** Get batch currently used size for this frame. */
+		FORCEINLINE int32 GetUsedValueCount() const { return DynamicValueBuffer ? DynamicValueBuffer->GetUsedCount() : 0; }
+		/** Get the batch buffer allocated size. */
+		FORCEINLINE int32 GetAllocatedValueCount() const { return DynamicValueBuffer ? DynamicValueBuffer->GetAllocatedCount() : 0; }
+	};
+
+	/** Expected that few will be enough, since we only need severals batches when ESortBatchProcessingOrder::KeyGenAfterPreRenderAndSortAfterPostRenderOpaque. */
+	typedef TArray<FSortBatch, TInlineAllocator<4>> FSortBatchArray;
+
+public:
+
+	/** Creates the sort manager, this is when the settings are configured. */
+	FGPUSortManager(ERHIFeatureLevel::Type InFeatureLevel);
+
+	~FGPUSortManager();
+
+	/**
+	 * Register a client system into the sort manager. Client systems are systems that will call AddTask().
+	 * Those systems need to register a callback, the possible used flags for their tasks and a name.
+	 * The callback is required because the initial content to perform GPU sorts is only configured long after
+	 * AddTask() gets called in the rendering loop.
+	 *
+	 * @param CallbackDelegate - The callback that will be used to set the initial keys and values for each batch elements.
+	 * @param UsedFlags - The possibly used flags for each of the client system tasks, used to perform some optimizations.
+	 * @param InName - A name defining this client system, used for GPU markers.
+	 */
+	void Register(const FGPUSortKeyGenDelegate& CallbackDelegate, EGPUSortFlags UsedFlags, const FName& InName);
+
+	/**
+	 * Add a GPU sort task to process this frame. Tasks are expected to be created before OnPreRender() gets called.
+	 *
+	 * @param OutInfo - The details about the allocation bindings, if success.
+	 * @param IndexCount - The number of values that need to be sorted by this task.
+	 * @param SortFlags - The sort task requirements.
+	 * 
+	 * Return true if the task was added to a batch, false otherwise, could be because GPU sorting is disabled.
+	 */
+	bool AddTask(FAllocationInfo& OutInfo, int32 ValueCount, EGPUSortFlags TaskFlags);
+
+
+	/**
+	 * Callback that needs to be called in the rendering loop, after calls to FFXSystemInterface::PreRender() are issued.
+	 * At this point, no other GPU sort tasks are allowed to be created through AddTask() as the sort batches
+	 * get finalized. At this point, any task with the EGPUSortFlags::SortAfterPreRender will possibly be done,
+	 * unless it also had the EGPUSortFlags::SortAfterPostRenderOpaque flag.
+	 * Tasks with EGPUSortFlags::KeyGenAfterPreRender get their client callbacks issued to set the initial keys and values. 
+	 *
+	 * @param RHICmdList - The command list to be used.
+	 */
+	void OnPreRender(FRHICommandListImmediate& RHICmdList);
+
+	/**
+	 * Callback that needs to be called in the rendering loop, after calls to FFXSystemInterface::PostRenderOpaque() are issued.
+	 * Tasks with EGPUSortFlags::KeyGenAfterPostRenderOpaque get their client callbacks issued to set the initial keys and values. 
+	 * After this any sort task that hasn't been processed yet, in OnPreRender(), gets processed now.
+	 *
+	 * @param RHICmdList - The command list to be used.
+	 */
+	void OnPostRenderOpaque(FRHICommandListImmediate& RHICmdList);
+
+private:
+
+	/** A helper to test that a flag has one and exactly one of A or B. */
+	static FORCEINLINE bool HasEither(EGPUSortFlags Flags, EGPUSortFlags A, EGPUSortFlags B);
+	/** Create the initial batch flags from a task flags. */
+	static FORCEINLINE EGPUSortFlags GetBatchFlags(EGPUSortFlags TaskFlags);
+	/** Update a batch flags after adding a task into it. */
+	static FORCEINLINE EGPUSortFlags CombineBatchFlags(EGPUSortFlags BatchFlags, EGPUSortFlags TaskFlags);
+	/** Test whether a batch flags are compatible with a given task. Used to know if the task can be merged in the batch.  */
+	static FORCEINLINE bool TestBatchFlags(EGPUSortFlags BatchFlags, EGPUSortFlags TaskFlags);
+	/** Convert the precision flags into a string, used for GPU markers. */
+	static FORCEINLINE const TCHAR* GetPrecisionString(EGPUSortFlags BatchFlags);
+
+	/**
+	 * Find an unused buffer to assign to a newly created FSortBatch. 
+	 * Although the sort batches are recreated every frame depending on what is being visible,
+	 * the buffers used to sort the indices and store the sorted results are persistent to some extent 
+	 * (see FBufferAllocationSettings which define the allocation and lifetime strategy).
+	 *
+	 * Each persistent buffer being assigned to one frame batch, we try to rebind the same buffers to the same batches
+	 * by using the constraints being imposed on the batch through the sortflags. The reason to do so is that 
+	 * the allocation strategy for the buffer is meant to be optimal if the required buffer size is consistent between frames.
+	 * (the batch flags are linked to the emitters being grouped into it, and so is the batch sorted indices count)
+	 *
+	 * @param TaskFlags - Some constraints being applied to the batch, using to help match the same buffers to the same batches between frames.
+	 * @param SortBatchId - The newly created batch index which will use the buffer this frame.
+	 * 
+	 * Return a buffer to store the sorted particle indices.
+	 */
+	FDynamicValueBuffer* GetDynamicValueBufferFromPool(EGPUSortFlags TaskFlags, int32 SortBatchId);
+
+	/** Setup the final sort flags and the processing order of all batches. Called after no other tasks will be added this frame. */
+	void FinalizeSortBatches();
+	/** Make sure there is enough GPU sort buffers to satisfy all batches created this frame. Free one used ones. */
+	void UpdateSortBuffersPool();
+	/** Resize (shrink) the DynamicValueBuffers and free the unused ones. */
+	void ResetDynamicValuesBuffers();
+	/** Delete all GPU sort buffers. */
+	void ReleaseSortBuffers();
+
+	/**
+	 * The different batches that will be used this frame, each one ending up being sorting work on the GPU.
+	 * In many cases, there will only be one batch but some cases makes it possible to have 2 
+	 * or more (for example when different key precision are used).
+	 */
+	FSortBatchArray SortBatches;
+
+	/** Settings to handle the buffer reallocation strategy. */
+	FSettings Settings;
+
+	/** A pool of buffer to hold the final sorted values for each batch. Reused between frames. */
+	FDynamicValueBufferArray DynamicValueBufferPool;
+
+	/** A pool of GPU sort buffers. The same sort GPU buffer can normally be used in several sort batches. */
+	TArray<FParticleSortBuffers*, TInlineAllocator<2>> SortBuffersPool;
+
+	/** The client system callbacks, see Register(). */
+	FCallbackArray Callbacks;
+
+	/** The current feature level, needed to access the global shader map. */
+	ERHIFeatureLevel::Type FeatureLevel;
+};
diff --git a/Engine/Source/Runtime/Engine/Public/ParticleResources.h b/Engine/Source/Runtime/Engine/Public/ParticleResources.h
--- a/Engine/Source/Runtime/Engine/Public/ParticleResources.h
+++ b/Engine/Source/Runtime/Engine/Public/ParticleResources.h
@@ -80,22 +80,3 @@ private:
 
 /** The global scratch vertex buffer. */
 extern TGlobalResource<FParticleScratchVertexBuffer> GParticleScratchVertexBuffer;
-
-
-/**
- * Vertex buffer used to hold particle indices.
- */
-class FParticleIndicesVertexBuffer : public FVertexBuffer
-{
-public:
-
-	/** Shader resource view of the vertex buffer. */
-	FShaderResourceViewRHIRef VertexBufferSRV;
-
-	/** Release RHI resources. */
-	virtual void ReleaseRHI() override
-	{
-		VertexBufferSRV.SafeRelease();
-		FVertexBuffer::ReleaseRHI();
-	}
-};
\ No newline at end of file
diff --git a/Engine/Source/Runtime/Renderer/Private/DeferredShadingRenderer.cpp b/Engine/Source/Runtime/Renderer/Private/DeferredShadingRenderer.cpp
--- a/Engine/Source/Runtime/Renderer/Private/DeferredShadingRenderer.cpp
+++ b/Engine/Source/Runtime/Renderer/Private/DeferredShadingRenderer.cpp
@@ -42,6 +42,7 @@
 #include "ShaderPrint.h"
 #include "GpuDebugRendering.h"
 #include "HairStrands/HairStrandsRendering.h"
+#include "GPUSortManager.h"
 
 static TAutoConsoleVariable<int32> CVarStencilForLODDither(
 	TEXT("r.StencilForLODDither"),
@@ -1586,6 +1587,10 @@ void FDeferredShadingSceneRenderer::Render(FRHICommandListImmediate& RHICmdList)
 		SCOPE_CYCLE_COUNTER(STAT_FDeferredShadingSceneRenderer_FXSystem_PreRender);
 		RHICmdList.SetCurrentStat(GET_STATID(STAT_CLM_FXPreRender));
 		Scene->FXSystem->PreRender(RHICmdList, &Views[0].GlobalDistanceFieldInfo.ParameterData, Views[0].AllowGPUParticleUpdate());
+		if (FGPUSortManager* GPUSortManager = Scene->FXSystem->GetGPUSortManager())
+		{
+			GPUSortManager->OnPreRender(RHICmdList);
+		}
 	}
 
 	if (AsyncDitherLODEndFence)
@@ -2507,6 +2512,11 @@ void FDeferredShadingSceneRenderer::Render(FRHICommandListImmediate& RHICmdList)
 			SceneTextureUniformBuffer.GetReference(),
 			Views[0].AllowGPUParticleUpdate()
 		);
+
+		if (FGPUSortManager* GPUSortManager = Scene->FXSystem->GetGPUSortManager())
+		{
+			GPUSortManager->OnPostRenderOpaque(RHICmdList);
+		}
 		ServiceLocalQueue();
 	}
 
diff --git a/Engine/Source/Runtime/Renderer/Private/MobileShadingRenderer.cpp b/Engine/Source/Runtime/Renderer/Private/MobileShadingRenderer.cpp
--- a/Engine/Source/Runtime/Renderer/Private/MobileShadingRenderer.cpp
+++ b/Engine/Source/Runtime/Renderer/Private/MobileShadingRenderer.cpp
@@ -49,6 +49,7 @@
 #include "SkyAtmosphereRendering.h"
 #include "VisualizeTexture.h"
 #include "VT/VirtualTextureSystem.h"
+#include "GPUSortManager.h"
 
 uint32 GetShadowQuality();
 
@@ -460,6 +461,10 @@ void FMobileSceneRenderer::Render(FRHICommandListImmediate& RHICmdList)
 	if (Scene->FXSystem && ViewFamily.EngineShowFlags.Particles)
 	{
 		Scene->FXSystem->PreRender(RHICmdList, NULL, !Views[0].bIsPlanarReflection);
+		if (FGPUSortManager* GPUSortManager = Scene->FXSystem->GetGPUSortManager())
+		{
+			GPUSortManager->OnPreRender(RHICmdList);
+		}
 	}
 	FRHICommandListExecutor::GetImmediateCommandList().PollOcclusionQueries();
 	RHICmdList.ImmediateFlush(EImmediateFlushType::DispatchToRHIThread);
@@ -720,6 +725,10 @@ void FMobileSceneRenderer::Render(FRHICommandListImmediate& RHICmdList)
 			nullptr,
 			Views[0].AllowGPUParticleUpdate()
 		);
+		if (FGPUSortManager* GPUSortManager = Scene->FXSystem->GetGPUSortManager())
+		{
+			GPUSortManager->OnPostRenderOpaque(RHICmdList);
+		}
 		RHICmdList.ImmediateFlush(EImmediateFlushType::DispatchToRHIThread);
 	}
 
diff --git a/Engine/Source/Runtime/Renderer/Private/SceneHitProxyRendering.cpp b/Engine/Source/Runtime/Renderer/Private/SceneHitProxyRendering.cpp
--- a/Engine/Source/Runtime/Renderer/Private/SceneHitProxyRendering.cpp
+++ b/Engine/Source/Runtime/Renderer/Private/SceneHitProxyRendering.cpp
@@ -22,6 +22,7 @@
 #include "GPUScene.h"
 #include "Rendering/ColorVertexBuffer.h"
 #include "FXSystem.h"
+#include "GPUSortManager.h"
 
 class FHitProxyShaderElementData : public FMeshMaterialShaderElementData
 {
@@ -614,7 +615,19 @@ void FDeferredShadingSceneRenderer::RenderHitProxies(FRHICommandListImmediate& R
 		// Notify the FX system that the scene is about to be rendered.
 		if (Scene->FXSystem && Views.IsValidIndex(0))
 		{
+			FGPUSortManager* GPUSortManager = Scene->FXSystem->GetGPUSortManager();
 			Scene->FXSystem->PreRender(RHICmdList, &Views[0].GlobalDistanceFieldInfo.ParameterData, false);
+			if (GPUSortManager)
+			{
+				GPUSortManager->OnPreRender(RHICmdList);
+			}
+			// Call PostRenderOpaque now as this is irrelevant for when rendering hit proxies.
+			// because we don't tick the particles in the render loop (see last param being "false").
+			Scene->FXSystem->PostRenderOpaque(RHICmdList, Views[0].ViewUniformBuffer, nullptr, nullptr, false);
+			if (GPUSortManager)
+			{
+				GPUSortManager->OnPostRenderOpaque(RHICmdList);
+			}
 		}
 
 		::DoRenderHitProxies(RHICmdList, this, HitProxyRT, HitProxyDepthRT);
