Fix issue where we'd trigger a false positive for a cache overflow as soon as we saturate the cache (and have a retained chunk least recently used in the cache).
Also fix issue where we were retaining all variants of a sound on all quality levels if a sound cue is being retained.
Also fix issue where the loading behavior of a sound class set on a Sound Class sound node was not being set.

diff --git a/Engine/Source/Runtime/Engine/Classes/Sound/SoundNodeQualityLevel.h b/Engine/Source/Runtime/Engine/Classes/Sound/SoundNodeQualityLevel.h
--- a/Engine/Source/Runtime/Engine/Classes/Sound/SoundNodeQualityLevel.h
+++ b/Engine/Source/Runtime/Engine/Classes/Sound/SoundNodeQualityLevel.h
@@ -35,6 +35,8 @@ public:
 	virtual int32 GetMaxChildNodes() const override;
 	virtual int32 GetMinChildNodes() const override;
 	virtual void PrimeChildWavePlayers(bool bRecurse) override;
+	virtual void RetainChildWavePlayers(bool bRecurse) override;
+	virtual void ReleaseRetainerOnChildWavePlayers(bool bRecurse) override;
 
 #if WITH_EDITOR
 	virtual FText GetInputPinName(int32 PinIndex) const override;
@@ -44,6 +46,5 @@ public:
 #if WITH_EDITOR
 	void ReconcileNode(bool bReconstructNode);
 #endif
-
 };
 
diff --git a/Engine/Source/Runtime/Engine/Classes/Sound/SoundNodeSoundClass.h b/Engine/Source/Runtime/Engine/Classes/Sound/SoundNodeSoundClass.h
--- a/Engine/Source/Runtime/Engine/Classes/Sound/SoundNodeSoundClass.h
+++ b/Engine/Source/Runtime/Engine/Classes/Sound/SoundNodeSoundClass.h
@@ -28,4 +28,11 @@ public:
 	//~ Begin USoundNode Interface. 
 	virtual void ParseNodes( class FAudioDevice* AudioDevice, const UPTRINT NodeWaveInstanceHash, FActiveSound& ActiveSound, const FSoundParseParameters& ParseParams, TArray<FWaveInstance*>& WaveInstances ) override;
 	//~ End USoundNode Interface. 
+
+	virtual void PostLoad() override;
+	virtual void BeginDestroy() override;
+
+private:
+	// This is set to true if this node's soundclass specifically dictates that a sound should be retained.
+	bool bRetainingAudioDueToSoundClass;
 };
diff --git a/Engine/Source/Runtime/Engine/Private/AudioStreamingCache.cpp b/Engine/Source/Runtime/Engine/Private/AudioStreamingCache.cpp
--- a/Engine/Source/Runtime/Engine/Private/AudioStreamingCache.cpp
+++ b/Engine/Source/Runtime/Engine/Private/AudioStreamingCache.cpp
@@ -18,6 +18,7 @@ AudioStreaming.cpp: Implementation of audio streaming classes.
 #include "HAL/IConsoleManager.h"
 #include "HAL/LowLevelMemTracker.h"
 #include "AudioDecompress.h"
+#include "AudioDevice.h"
 #include "AudioCompressionSettingsUtils.h"
 
 static int32 DebugMaxElementsDisplayCVar = 128;
@@ -344,6 +345,11 @@ FAudioChunkHandle FCachedAudioStreamingManager::GetLoadedChunk(const USoundWave*
 
 				AsyncTask(ENamedThreads::GameThread, []()
 				{
+					if (GEngine && GEngine->GetMainAudioDevice())
+					{
+						GEngine->GetMainAudioDevice()->Exec(nullptr, TEXT("audiomemreport"));
+					}
+					
 					int32 NumChunksReleased = 0;
 
 					for (TObjectIterator<USoundWave> It; It; ++It)
@@ -976,6 +982,13 @@ FAudioChunkCache::FCacheElement* FAudioChunkCache::InsertChunk(const FChunkKey&
 		}
 		else
 		{
+			static bool bLoggedCacheSaturated = false;
+			if (!bLoggedCacheSaturated)
+			{
+				UE_LOG(LogAudio, Display, TEXT("Audio Stream Cache: Using %d of %d chunks.."), ChunksInUse, CachePool.Num());
+				bLoggedCacheSaturated = true;
+			}
+			
 			// The pools filled, so we're going to need to evict.
 			CacheElement = EvictLeastRecentChunk();
 
@@ -984,18 +997,18 @@ FAudioChunkCache::FCacheElement* FAudioChunkCache::InsertChunk(const FChunkKey&
 				return nullptr;
 			}
 		}
-	}
 
-	check(CacheElement);
-	CacheElement->bIsLoaded = false;
-	CacheElement->Key = InKey;
-	TouchElement(CacheElement);
+		check(CacheElement);
+		CacheElement->bIsLoaded = false;
+		CacheElement->Key = InKey;
+		TouchElement(CacheElement);
 
-	// If we've got multiple chunks, we can not cache the least recent chunk
-	// without worrying about a circular dependency.
-	if (LeastRecentElement == nullptr && ChunksInUse > 1)
-	{
-		SetUpLeastRecentChunk();
+		// If we've got multiple chunks, we can not cache the least recent chunk
+		// without worrying about a circular dependency.
+		if (LeastRecentElement == nullptr && ChunksInUse > 1)
+		{
+			SetUpLeastRecentChunk();
+		}
 	}
 
 	return CacheElement;
@@ -1035,7 +1048,7 @@ FAudioChunkCache::FCacheElement* FAudioChunkCache::EvictLeastRecentChunk()
 		const FCacheElement* ElementToStopAt = MostRecentElement->LessRecentElement;
 
 		// Otherwise, we need to crawl up the cache from least recent used to most to find a chunk that is not in use:
-		while (CacheElement != ElementToStopAt)
+		while (CacheElement && CacheElement != ElementToStopAt)
 		{
 			if (CacheElement->CanEvictChunk())
 			{
@@ -1054,13 +1067,13 @@ FAudioChunkCache::FCacheElement* FAudioChunkCache::EvictLeastRecentChunk()
 			{
 				CacheElement = CacheElement->MoreRecentElement;
 			}
+		}
 
-			// If we ever hit this, it means that we couldn't find any cache elements that aren't in use.
-			if (CacheElement != MostRecentElement)
-			{
-				ensureMsgf(false, TEXT("Cache blown! Please increase the cache size or load less audio."));
-				return nullptr;
-			}
+		// If we ever hit this, it means that we couldn't find any cache elements that aren't in use.
+		if (!CacheElement || CacheElement == ElementToStopAt)
+		{
+			ensureMsgf(false, TEXT("Cache blown! Please increase the cache size or load less audio."));
+			return nullptr;
 		}
 	}
 
diff --git a/Engine/Source/Runtime/Engine/Private/SoundNodeQualityLevel.cpp b/Engine/Source/Runtime/Engine/Private/SoundNodeQualityLevel.cpp
--- a/Engine/Source/Runtime/Engine/Private/SoundNodeQualityLevel.cpp
+++ b/Engine/Source/Runtime/Engine/Private/SoundNodeQualityLevel.cpp
@@ -68,6 +68,44 @@ void USoundNodeQualityLevel::PrimeChildWavePlayers(bool bRecurse)
 	}
 }
 
+void USoundNodeQualityLevel::RetainChildWavePlayers(bool bRecurse)
+{
+	// If we're able to retrieve a valid cached quality level for this sound cue,
+	// only retain that quality level.
+	int32 QualityLevel = USoundCue::GetCachedQualityLevel();
+
+#if WITH_EDITOR
+	if (GIsEditor && QualityLevel < 0)
+	{
+		QualityLevel = GetDefault<ULevelEditorPlaySettings>()->PlayInEditorSoundQualityLevel;
+	}
+#endif
+
+	if (ChildNodes.IsValidIndex(QualityLevel) && ChildNodes[QualityLevel])
+	{
+		ChildNodes[QualityLevel]->RetainChildWavePlayers(bRecurse);
+	}
+}
+
+void USoundNodeQualityLevel::ReleaseRetainerOnChildWavePlayers(bool bRecurse)
+{
+	// If we're able to retrieve a valid cached quality level for this sound cue,
+	// only release that quality level.
+	int32 QualityLevel = USoundCue::GetCachedQualityLevel();
+
+#if WITH_EDITOR
+	if (GIsEditor && QualityLevel < 0)
+	{
+		QualityLevel = GetDefault<ULevelEditorPlaySettings>()->PlayInEditorSoundQualityLevel;
+	}
+#endif
+
+	if (ChildNodes.IsValidIndex(QualityLevel) && ChildNodes[QualityLevel])
+	{
+		ChildNodes[QualityLevel]->ReleaseRetainerOnChildWavePlayers(bRecurse);
+	}
+}
+
 int32 USoundNodeQualityLevel::GetMaxChildNodes() const
 {
 	return GetDefault<UAudioSettings>()->QualityLevels.Num();
diff --git a/Engine/Source/Runtime/Engine/Private/SoundNodeSoundClass.cpp b/Engine/Source/Runtime/Engine/Private/SoundNodeSoundClass.cpp
--- a/Engine/Source/Runtime/Engine/Private/SoundNodeSoundClass.cpp
+++ b/Engine/Source/Runtime/Engine/Private/SoundNodeSoundClass.cpp
@@ -3,6 +3,7 @@
 
 #include "Sound/SoundNodeSoundClass.h"
 #include "ActiveSound.h"
+#include "Sound/SoundClass.h"
 
 /*-----------------------------------------------------------------------------
 	USoundNodeSoundClass implementation.
@@ -10,6 +11,7 @@
 
 USoundNodeSoundClass::USoundNodeSoundClass(const FObjectInitializer& ObjectInitializer)
 	: Super(ObjectInitializer)
+	, bRetainingAudioDueToSoundClass(false)
 {
 }
 
@@ -23,3 +25,38 @@ void USoundNodeSoundClass::ParseNodes( class FAudioDevice* AudioDevice, const UP
 
 	Super::ParseNodes( AudioDevice, NodeWaveInstanceHash, ActiveSound, UpdatedParseParams, WaveInstances );
 }
+
+void USoundNodeSoundClass::PostLoad()
+{
+	Super::PostLoad();
+	
+	ESoundWaveLoadingBehavior SoundClassLoadingBehavior = ESoundWaveLoadingBehavior::Inherited;
+
+	USoundClass* CurrentSoundClass = SoundClassOverride;
+
+	// Recurse through this sound class's parents until we find an override.
+	while (SoundClassLoadingBehavior == ESoundWaveLoadingBehavior::Inherited && CurrentSoundClass != nullptr)
+	{
+		SoundClassLoadingBehavior = CurrentSoundClass->Properties.LoadingBehavior;
+		CurrentSoundClass = CurrentSoundClass->ParentClass;
+	}
+
+	if (SoundClassLoadingBehavior == ESoundWaveLoadingBehavior::RetainOnLoad)
+	{
+		RetainChildWavePlayers(true);
+		bRetainingAudioDueToSoundClass = true;
+	}
+	else if (SoundClassLoadingBehavior == ESoundWaveLoadingBehavior::PrimeOnLoad)
+	{
+		PrimeChildWavePlayers(true);
+	}
+}
+
+void USoundNodeSoundClass::BeginDestroy()
+{
+	Super::BeginDestroy();
+	if (bRetainingAudioDueToSoundClass)
+	{
+		ReleaseRetainerOnChildWavePlayers(true);
+	}
+}
