Fixed non unity build issue around "RHIResources.h" which also needs "RHI.h"

diff --git a/Engine/Source/Runtime/Engine/Public/GPUSortManager.h b/Engine/Source/Runtime/Engine/Public/GPUSortManager.h
--- a/Engine/Source/Runtime/Engine/Public/GPUSortManager.h
+++ b/Engine/Source/Runtime/Engine/Public/GPUSortManager.h
@@ -8,6 +8,7 @@
 
 #include "CoreMinimal.h"
 #include "Templates/RefCounting.h"
+#include "RHI.h"
 #include "RHIResources.h"
 
 /**
