More non unity build fix around FGPUSortManager.h

diff --git a/Engine/Source/Runtime/Engine/Private/GPUSortManager.cpp b/Engine/Source/Runtime/Engine/Private/GPUSortManager.cpp
--- a/Engine/Source/Runtime/Engine/Private/GPUSortManager.cpp
+++ b/Engine/Source/Runtime/Engine/Private/GPUSortManager.cpp
@@ -9,6 +9,7 @@
 #include "GlobalShader.h"
 #include "ShaderPermutation.h"
 #include "ShaderParameterUtils.h"
+#include "ProfilingDebugging/RealtimeGPUProfiler.h"
 
 //*****************************************************************************
 
diff --git a/Engine/Source/Runtime/Engine/Public/GPUSortManager.h b/Engine/Source/Runtime/Engine/Public/GPUSortManager.h
--- a/Engine/Source/Runtime/Engine/Public/GPUSortManager.h
+++ b/Engine/Source/Runtime/Engine/Public/GPUSortManager.h
@@ -10,6 +10,7 @@
 #include "Templates/RefCounting.h"
 #include "RHI.h"
 #include "RHIResources.h"
+#include "RenderResource.h"
 
 /**
  * Copy one uint buffer into several overlapping targets.
