Improved resource transition in SortGPUBuffers when using the FinalValuesUAV override

diff --git a/Engine/Source/Runtime/Engine/Private/GPUSort.cpp b/Engine/Source/Runtime/Engine/Private/GPUSort.cpp
--- a/Engine/Source/Runtime/Engine/Private/GPUSort.cpp
+++ b/Engine/Source/Runtime/Engine/Private/GPUSort.cpp
@@ -841,7 +841,17 @@ int32 SortGPUBuffers(FRHICommandListImmediate& RHICmdList, FGPUSortBuffers SortB
 			// Phase 3: Downsweep to compute final offsets and scatter keys.
 			RHICmdList.SetComputeShader(DownsweepCS->GetComputeShader());
 			{
-				FRHIUnorderedAccessView* ValuesUAV = (bIsLastPass && SortBuffers.FinalValuesUAV)? SortBuffers.FinalValuesUAV : SortBuffers.RemoteValueUAVs[BufferIndex ^ 0x1];
+				FRHIUnorderedAccessView* ValuesUAV = nullptr;
+				if (bIsLastPass && SortBuffers.FinalValuesUAV)
+				{
+					ValuesUAV = SortBuffers.FinalValuesUAV;
+					// Transition resource since FinalValuesUAV can also be SortBuffers.FirstValuesSRV.
+					RHICmdList.TransitionResource(EResourceTransitionAccess::ERWBarrier, EResourceTransitionPipeline::EComputeToCompute, ValuesUAV);
+				}
+				else
+				{
+					ValuesUAV = SortBuffers.RemoteValueUAVs[BufferIndex ^ 0x1];
+				}
 				DownsweepCS->SetOutput(RHICmdList, SortBuffers.RemoteKeyUAVs[BufferIndex ^ 0x1], ValuesUAV);
 			}
 			{
