Fixed duplicated "FX.AllowGPUSorting" console variable.
Fixed some missing include around GPU profiling in FNiagaraGPUInstanceCountManager

diff --git a/Engine/Plugins/FX/Niagara/Source/Niagara/Private/NiagaraGPUInstanceCountManager.cpp b/Engine/Plugins/FX/Niagara/Source/Niagara/Private/NiagaraGPUInstanceCountManager.cpp
--- a/Engine/Plugins/FX/Niagara/Source/Niagara/Private/NiagaraGPUInstanceCountManager.cpp
+++ b/Engine/Plugins/FX/Niagara/Source/Niagara/Private/NiagaraGPUInstanceCountManager.cpp
@@ -4,6 +4,7 @@
 #include "NiagaraStats.h"
 #include "Containers/DynamicRHIResourceArray.h"
 #include "GPUSortManager.h" // CopyUIntBufferToTargets
+#include "ProfilingDebugging/RealtimeGPUProfiler.h"
 
 int32 GNiagaraMinGPUInstanceCount = 2048;
 static FAutoConsoleVariableRef CVarNiagaraMinGPUInstanceCount(
diff --git a/Engine/Source/Runtime/Engine/Private/GPUSortManager.cpp b/Engine/Source/Runtime/Engine/Private/GPUSortManager.cpp
--- a/Engine/Source/Runtime/Engine/Private/GPUSortManager.cpp
+++ b/Engine/Source/Runtime/Engine/Private/GPUSortManager.cpp
@@ -10,6 +10,7 @@
 #include "ShaderPermutation.h"
 #include "ShaderParameterUtils.h"
 #include "ProfilingDebugging/RealtimeGPUProfiler.h"
+#include "FXSystem.h" // FXConsoleVariables::bAllowGPUSorting
 
 //*****************************************************************************
 
@@ -37,14 +38,6 @@ static FAutoConsoleVariableRef CVarGPUSortBufferSlack(
 	ECVF_Default
 );
 
-int32 GGPUSortEnable = 1;
-static FAutoConsoleVariableRef CVarGPUSortEnable(
-	TEXT("fx.AllowGPUSorting"),
-	GGPUSortEnable,
-	TEXT("Whether to enable GPU sort (default 1)"),
-	ECVF_Default
-);
-
 int32 GGPUSortStressTest = 0;
 static FAutoConsoleVariableRef CVarGPUSortStressTest(
 	TEXT("fx.GPUSort.StressTest"),
@@ -569,7 +562,7 @@ void FGPUSortManager::Register(const FGPUSortKeyGenDelegate& CallbackDelegate, E
 
 bool FGPUSortManager::AddTask(FAllocationInfo& OutInfo, int32 ValueCount, EGPUSortFlags TaskFlags)
 {
-	if (!GGPUSortEnable)
+	if (!FXConsoleVariables::bAllowGPUSorting)
 	{
 		return false;
 	}
