Fix messed up merge.

#rnx

diff --git a/Engine/Source/Runtime/Engine/Private/Materials/MaterialExpressions.cpp b/Engine/Source/Runtime/Engine/Private/Materials/MaterialExpressions.cpp
--- a/Engine/Source/Runtime/Engine/Private/Materials/MaterialExpressions.cpp
+++ b/Engine/Source/Runtime/Engine/Private/Materials/MaterialExpressions.cpp
@@ -15135,41 +15135,6 @@ void UMaterialExpressionVertexNormalWS::GetCaption(TArray<FString>& OutCaptions)
 {
 	OutCaptions.Add(TEXT("VertexNormalWS"));
 }
-#endif // WITH_EDITOR
-
-	: Super(ObjectInitializer)
-{
-#if WITH_EDITORONLY_DATA
-	// Structure to hold one-time initialization
-	struct FConstructorStatics
-	{
-		FText NAME_Vectors;
-		FText NAME_Coordinates;
-		FConstructorStatics()
-			: NAME_Vectors(LOCTEXT("Vectors", "Vectors"))
-			, NAME_Coordinates(LOCTEXT("Coordinates", "Coordinates"))
-		{
-		}
-	};
-	static FConstructorStatics ConstructorStatics;
-
-	MenuCategories.Add(ConstructorStatics.NAME_Vectors);
-	MenuCategories.Add(ConstructorStatics.NAME_Coordinates);
-
-	bShaderInputData = true;
-#endif
-}
-
-#if WITH_EDITOR
-int32 UMaterialExpressionVertexTangentWS::Compile(class FMaterialCompiler* Compiler, int32 OutputIndex)
-{
-	return Compiler->VertexTangent();
-}
-
-void UMaterialExpressionVertexTangentWS::GetCaption(TArray<FString>& OutCaptions) const
-{
-	OutCaptions.Add(TEXT("VertexTangentWS"));
-}
 #endif // WITH_EDITOR
 
 UMaterialExpressionPixelNormalWS::UMaterialExpressionPixelNormalWS(const FObjectInitializer& ObjectInitializer)
