Update copy right notice

#rnx

diff --git a/Engine/Source/Runtime/Engine/Classes/Materials/MaterialExpressionPerInstanceCustomData.h b/Engine/Source/Runtime/Engine/Classes/Materials/MaterialExpressionPerInstanceCustomData.h
--- a/Engine/Source/Runtime/Engine/Classes/Materials/MaterialExpressionPerInstanceCustomData.h
+++ b/Engine/Source/Runtime/Engine/Classes/Materials/MaterialExpressionPerInstanceCustomData.h
@@ -1,3 +1,5 @@
+// Copyright Epic Games, Inc. All Rights Reserved.
+
 #pragma once
 
 #include "CoreMinimal.h"
