Static analysis fixes

diff --git a/Engine/Source/Runtime/Engine/Private/DataChannel.cpp b/Engine/Source/Runtime/Engine/Private/DataChannel.cpp
--- a/Engine/Source/Runtime/Engine/Private/DataChannel.cpp
+++ b/Engine/Source/Runtime/Engine/Private/DataChannel.cpp
@@ -1464,7 +1464,7 @@ void UControlChannel::ReceivedBunch( FInBunch& Bunch )
 	// Process the packet
 	while (!Bunch.AtEnd() && Connection != NULL && Connection->State != USOCK_Closed) // if the connection got closed, we don't care about the rest
 	{
-		uint8 MessageType;
+		uint8 MessageType = 0;
 		Bunch << MessageType;
 		if (Bunch.IsError())
 		{
diff --git a/Engine/Source/Runtime/Engine/Private/DemoNetDriver.cpp b/Engine/Source/Runtime/Engine/Private/DemoNetDriver.cpp
--- a/Engine/Source/Runtime/Engine/Private/DemoNetDriver.cpp
+++ b/Engine/Source/Runtime/Engine/Private/DemoNetDriver.cpp
@@ -4958,7 +4958,7 @@ bool UDemoNetDriver::FastForwardLevels(const FGotoResult& GotoResult)
 
 				if (PlaybackDemoHeader.Version >= HISTORY_MULTIPLE_LEVELS)
 				{
-					int32 LevelIndex;
+					int32 LevelIndex = INDEX_NONE;
 					*CheckpointArchive << LevelIndex;
 				}
 
@@ -5557,7 +5557,7 @@ bool UDemoNetDriver::LoadCheckpoint(const FGotoResult& GotoResult)
 
 		if (PlaybackDemoHeader.Version >= HISTORY_MULTIPLE_LEVELS)
 		{
-			int32 LevelIndex;
+			int32 LevelIndex = INDEX_NONE;
 			*GotoCheckpointArchive << LevelIndex;
 		}
 
