PR #6503: Wrong type comparison (Contributed by [REDACTED_PERSON])

diff --git a/Engine/Source/Runtime/Engine/Private/AI/Navigation/NavigationTypes.cpp b/Engine/Source/Runtime/Engine/Private/AI/Navigation/NavigationTypes.cpp
--- a/Engine/Source/Runtime/Engine/Private/AI/Navigation/NavigationTypes.cpp
+++ b/Engine/Source/Runtime/Engine/Private/AI/Navigation/NavigationTypes.cpp
@@ -266,7 +266,7 @@ void FNavAgentProperties::UpdateWithCollisionComponent(UShapeComponent* Collisio
 
 bool FNavAgentProperties::IsNavDataMatching(const FNavAgentProperties& Other) const
 {
-	return (PreferredNavData == Other.PreferredNavData || PreferredNavData == nullptr || Other.PreferredNavData == nullptr);
+	return (PreferredNavData == Other.PreferredNavData || PreferredNavData.IsNull() || Other.PreferredNavData.IsNull());
 }
 
 void FNavAgentProperties::SetPreferredNavData(TSubclassOf<AActor> NavDataClass)
