PR #6514: Blueprint support for FIntPoint (Contributed by [REDACTED_PERSON])

diff --git a/Engine/Source/Runtime/Core/Public/Math/IntPoint.h b/Engine/Source/Runtime/Core/Public/Math/IntPoint.h
--- a/Engine/Source/Runtime/Core/Public/Math/IntPoint.h
+++ b/Engine/Source/Runtime/Core/Public/Math/IntPoint.h
@@ -42,6 +42,14 @@ public:
 	 */
 	FIntPoint(int32 InX, int32 InY);
 
+	/**
+	 * Create and initialize a new instance with a single int.
+	 * Both X and Y will be initialized to this value
+	 *
+	 * @param InXY The x and y-coordinate.
+	 */
+	FIntPoint(int32 InXY);
+
 	/**
 	 * Create and initialize a new instance to zero.
 	 *
@@ -106,7 +114,15 @@ public:
 	 * @return Reference to this point after addition.
 	 */
 	FIntPoint& operator+=(const FIntPoint& Other);
-	
+
+	/**
+	 * Multiply another point component-wise from this point.
+	 *
+	 * @param Other The point to multiply with this point.
+	 * @return Reference to this point after multiplication.
+	 */
+	FIntPoint& operator*=(const FIntPoint& Other);
+
 	/**
 	 * Subtract another point component-wise from this point.
 	 *
@@ -163,6 +179,14 @@ public:
 	 */
 	FIntPoint operator-(const FIntPoint& Other) const;
 
+	/**
+	 * Get the result of multiplication on this point.
+	 *
+	 * @param Other The point to multiply with this point.
+	 * @return A new multiplied int point
+	 */
+	FIntPoint operator*(const FIntPoint& Other) const;
+
 	/**
 	 * Get the result of division on this point.
 	 *
@@ -335,6 +359,10 @@ FORCEINLINE FIntPoint::FIntPoint(int32 InX, int32 InY)
 	, Y(InY)
 { }
 
+FORCEINLINE FIntPoint::FIntPoint(int32 InXY)
+	: X(InXY)
+	, Y(InXY)
+{ }
 
 FORCEINLINE FIntPoint::FIntPoint(EForceInit)
 	: X(0)
@@ -408,6 +436,15 @@ FORCEINLINE FIntPoint& FIntPoint::operator-=(const FIntPoint& Other)
 }
 
 
+FORCEINLINE FIntPoint& FIntPoint::operator*=(const FIntPoint& Other)
+{
+	X *= Other.X;
+	Y *= Other.Y;
+
+	return *this;
+}
+
+
 FORCEINLINE FIntPoint& FIntPoint::operator/=(const FIntPoint& Other)
 {
 	X /= Other.X;
@@ -491,6 +528,12 @@ FORCEINLINE FIntPoint FIntPoint::operator-(const FIntPoint& Other) const
 }
 
 
+FORCEINLINE FIntPoint FIntPoint::operator*(const FIntPoint& Other) const
+{
+	return FIntPoint(*this) *= Other;
+}
+
+
 FORCEINLINE FIntPoint FIntPoint::operator/(const FIntPoint& Other) const
 {
 	return FIntPoint(*this) /= Other;
diff --git a/Engine/Source/Runtime/Engine/Classes/Kismet/KismetMathLibrary.h b/Engine/Source/Runtime/Engine/Classes/Kismet/KismetMathLibrary.h
--- a/Engine/Source/Runtime/Engine/Classes/Kismet/KismetMathLibrary.h
+++ b/Engine/Source/Runtime/Engine/Classes/Kismet/KismetMathLibrary.h
@@ -830,6 +830,83 @@ class ENGINE_API UKismetMathLibrary : public UBlueprintFunctionLibrary
 	static float FixedTurn(float InCurrent, float InDesired, float InDeltaRate);
 
 
+	//
+	// IntPoint constants
+	//
+	
+	/** Zero Int Point (0, 0) */
+	UFUNCTION(BlueprintPure, meta = (ScriptConstant = "Zero", ScriptConstantHost = "IntPoint"), Category = "Math|IntPoint|Constants")
+	static FIntPoint IntPoint_Zero();
+	
+	/** One Int Point (1, 1) */
+	UFUNCTION(BlueprintPure, meta = (ScriptConstant = "One", ScriptConstantHost = "IntPoint"), Category = "Math|IntPoint|Constants")
+	static FIntPoint IntPoint_One();
+	
+	/** Up Int Point (0, -1) */
+	UFUNCTION(BlueprintPure, meta = (ScriptConstant = "Up", ScriptConstantHost = "IntPoint"), Category = "Math|IntPoint|Constants")
+	static FIntPoint IntPoint_Up();
+	
+	/** Left Int Point (-1, 0) */
+	UFUNCTION(BlueprintPure, meta = (ScriptConstant = "Left", ScriptConstantHost = "IntPoint"), Category = "Math|IntPoint|Constants")
+	static FIntPoint IntPoint_Left();
+	
+	/** Right Int Point (1, 0) */
+	UFUNCTION(BlueprintPure, meta = (ScriptConstant = "Right", ScriptConstantHost = "IntPoint"), Category = "Math|IntPoint|Constants")
+	static FIntPoint IntPoint_Right();
+	
+	/** Down Int Point (0, 1) */
+	UFUNCTION(BlueprintPure, meta = (ScriptConstant = "Down", ScriptConstantHost = "IntPoint"), Category = "Math|IntPoint|Constants")
+	static FIntPoint IntPoint_Down();
+
+	//
+	// IntPoint functions
+	//
+
+	/** Convert an IntPoint to a Vector2D */
+	UFUNCTION(BlueprintPure, meta = (DisplayName = "To Vector2D (IntPoint)", CompactNodeTitle = "->", ScriptMethod = "Vector2D", Keywords = "cast convert", BlueprintAutocast), Category = "Math|Conversions")
+	static FVector2D Conv_IntPointToVector2D(FIntPoint InIntPoint);
+
+	/** Returns IntPoint A added by B */
+	UFUNCTION(BlueprintPure, meta = (DisplayName = "IntPoint + IntPoint", CompactNodeTitle = "+", ScriptMethod = "Add", ScriptOperator = "+;+=", Keywords = "+ add plus"), Category = "Math|IntPoint")
+	static FIntPoint Add_IntPointIntPoint(FIntPoint A, FIntPoint B);
+
+	/** Returns IntPoint A added by B */
+	UFUNCTION(BlueprintPure, meta = (DisplayName = "IntPoint + int", CompactNodeTitle = "+", ScriptMethod = "AddInt", ScriptOperator = "+;+=", Keywords = "+ add plus"), Category = "Math|IntPoint")
+	static FIntPoint Add_IntPointInt(FIntPoint A, int32 B);
+
+	/** Returns IntPoint A subtracted by B */
+	UFUNCTION(BlueprintPure, meta = (DisplayName = "IntPoint - IntPoint", CompactNodeTitle = "-", ScriptMethod = "Subtract", ScriptOperator = "-;-=", Keywords = "- subtract minus"), Category = "Math|IntPoint")
+	static FIntPoint Subtract_IntPointIntPoint(FIntPoint A, FIntPoint B);
+
+	/** Returns IntPoint A subtracted by B */
+	UFUNCTION(BlueprintPure, meta = (DisplayName = "IntPoint - Int", CompactNodeTitle = "-", ScriptMethod = "SubtractInt", ScriptOperator = "-;-=", Keywords = "- subtract minus"), Category = "Math|IntPoint")
+	static FIntPoint Subtract_IntPointInt(FIntPoint A, int32 B);
+
+	/** Returns IntPoint A multiplied by B */
+	UFUNCTION(BlueprintPure, meta = (DisplayName = "IntPoint * IntPoint", CompactNodeTitle = "*", ScriptMethod = "Multiply", ScriptOperator = "*;*=", Keywords = "* multiply"), Category = "Math|IntPoint")
+	static FIntPoint Multiply_IntPointIntPoint(FIntPoint A, FIntPoint B);
+
+	/** Returns IntPoint A multiplied by B */
+	UFUNCTION(BlueprintPure, meta = (DisplayName = "IntPoint * Int", CompactNodeTitle = "*", ScriptMethod = "MultiplyInt", ScriptOperator = "*;*=", Keywords = "* multiply"), Category = "Math|IntPoint")
+	static FIntPoint Multiply_IntPointInt(FIntPoint A, int32 B);
+
+	/** Returns IntPoint A divided by B */
+	UFUNCTION(BlueprintPure, meta = (DisplayName = "IntPoint / IntPoint", CompactNodeTitle = "/", ScriptMethod = "Divide", ScriptOperator = "/;/=", Keywords = "/ divide"), Category = "Math|IntPoint")
+	static FIntPoint Divide_IntPointIntPoint(FIntPoint A, FIntPoint B);
+
+	/** Returns IntPoint A divided by B */
+	UFUNCTION(BlueprintPure, meta = (DisplayName = "IntPoint / Int", CompactNodeTitle = "/", ScriptMethod = "DivideInt", ScriptOperator = "/;/=", Keywords = "/ divide"), Category = "Math|IntPoint")
+	static FIntPoint Divide_IntPointInt(FIntPoint A, int32 B);
+
+	/** Returns true if IntPoint A is equal to IntPoint B (A == B) */
+	UFUNCTION(BlueprintPure, meta = (DisplayName = "Equal (IntPoint)", CompactNodeTitle = "==", ScriptMethod = "Equals", ScriptOperator = "==", Keywords = "== equal"), Category = "Math|IntPoint")
+	static bool Equal_IntPointIntPoint(FIntPoint A, FIntPoint B);
+
+	/** Returns true if IntPoint A is NOT equal to IntPoint B (A != B) */
+	UFUNCTION(BlueprintPure, meta = (DisplayName = "Not Equal (IntPoint)", CompactNodeTitle = "!=", ScriptMethod = "NotEqual", ScriptOperator = "==", Keywords = "== not equal"), Category = "Math|IntPoint")
+	static bool NotEqual_IntPointIntPoint(FIntPoint A, FIntPoint B);
+	
+
 	//
 	// Vector2D constants - exposed for scripting
 	//
@@ -863,7 +940,7 @@ class ENGINE_API UKismetMathLibrary : public UBlueprintFunctionLibrary
 	UFUNCTION(BlueprintPure, meta = (DisplayName = "To Vector (Vector2D)", CompactNodeTitle = "->", ScriptMethod = "Vector", Keywords = "cast convert", BlueprintAutocast), Category = "Math|Conversions")
 	static FVector Conv_Vector2DToVector(FVector2D InVector2D, float Z = 0);
 
-	/** Convert a Vector2D to a Vector */
+	/** Convert a Vector2D to an IntPoint */
 	UFUNCTION(BlueprintPure, meta = (DisplayName = "To IntPoint (Vector2D)", CompactNodeTitle = "->", ScriptMethod = "IntPoint", Keywords = "cast convert", BlueprintAutocast), Category = "Math|Conversions")
 	static FIntPoint Conv_Vector2DToIntPoint(FVector2D InVector2D);
 
@@ -3947,6 +4024,8 @@ private:
 	static void ReportError_Divide_VectorInt();
 	static void ReportError_Divide_VectorVector();
 	static void ReportError_ProjectVectorOnToVector();
+	static void ReportError_Divide_IntPointOnInt();
+	static void ReportError_Divide_IntPointOnIntPoint();
 	static void ReportError_Divide_Vector2DFloat();
 	static void ReportError_Divide_Vector2DVector2D();
 	static void ReportError_DaysInMonth();
@@ -3957,4 +4036,3 @@ private:
 #if KISMET_MATH_INLINE_ENABLED
 #include "KismetMathLibrary.inl"
 #endif
-
diff --git a/Engine/Source/Runtime/Engine/Classes/Kismet/KismetMathLibrary.inl b/Engine/Source/Runtime/Engine/Classes/Kismet/KismetMathLibrary.inl
--- a/Engine/Source/Runtime/Engine/Classes/Kismet/KismetMathLibrary.inl
+++ b/Engine/Source/Runtime/Engine/Classes/Kismet/KismetMathLibrary.inl
@@ -909,6 +909,124 @@ bool UKismetMathLibrary::EqualEqual_TransformTransform(const FTransform& A, cons
 }
 
 
+/* FIntPoint
+*****************************************************************************/
+
+KISMET_MATH_FORCEINLINE
+FIntPoint UKismetMathLibrary::IntPoint_Zero()
+{
+	return FIntPoint::ZeroValue;
+}
+
+KISMET_MATH_FORCEINLINE
+FIntPoint UKismetMathLibrary::IntPoint_One()
+{
+	return FIntPoint(1);
+}
+
+KISMET_MATH_FORCEINLINE
+FIntPoint UKismetMathLibrary::IntPoint_Up()
+{
+	return FIntPoint(0, -1);
+}
+
+KISMET_MATH_FORCEINLINE
+FIntPoint UKismetMathLibrary::IntPoint_Left()
+{
+	return FIntPoint(-1, 0);
+}
+
+KISMET_MATH_FORCEINLINE
+FIntPoint UKismetMathLibrary::IntPoint_Right()
+{
+	return FIntPoint(1, 0);
+}
+
+KISMET_MATH_FORCEINLINE
+FIntPoint UKismetMathLibrary::IntPoint_Down()
+{
+	return FIntPoint(0, 1);
+}
+
+KISMET_MATH_FORCEINLINE
+FVector2D UKismetMathLibrary::Conv_IntPointToVector2D(FIntPoint InIntPoint)
+{
+	return FVector2D(InIntPoint.X, InIntPoint.Y);
+}
+
+
+KISMET_MATH_FORCEINLINE
+FIntPoint UKismetMathLibrary::Add_IntPointIntPoint(FIntPoint A, FIntPoint B)
+{
+	return A + B;
+}
+
+KISMET_MATH_FORCEINLINE
+FIntPoint UKismetMathLibrary::Add_IntPointInt(FIntPoint A, int32 B)
+{
+	return A + FIntPoint(B);
+}
+
+KISMET_MATH_FORCEINLINE
+FIntPoint UKismetMathLibrary::Subtract_IntPointIntPoint(FIntPoint A, FIntPoint B)
+{
+	return A - B;
+}
+
+KISMET_MATH_FORCEINLINE
+FIntPoint UKismetMathLibrary::Subtract_IntPointInt(FIntPoint A, int32 B)
+{
+	return A - FIntPoint(B);
+}
+
+KISMET_MATH_FORCEINLINE
+FIntPoint UKismetMathLibrary::Multiply_IntPointIntPoint(FIntPoint A, FIntPoint B)
+{
+	return A * B;
+}
+
+KISMET_MATH_FORCEINLINE
+FIntPoint UKismetMathLibrary::Multiply_IntPointInt(FIntPoint A, int32 B)
+{
+	return A * FIntPoint(B);
+}
+
+KISMET_MATH_FORCEINLINE
+FIntPoint UKismetMathLibrary::Divide_IntPointIntPoint(FIntPoint A, FIntPoint B)
+{
+	if (B.X == 0 || B.Y == 0)
+	{
+		ReportError_Divide_IntPointOnIntPoint();
+		return FIntPoint::ZeroValue;
+	}
+	return A / B;
+}
+
+KISMET_MATH_FORCEINLINE
+FIntPoint UKismetMathLibrary::Divide_IntPointInt(FIntPoint A, int32 B)
+{
+	if (B == 0)
+	{
+		ReportError_Divide_IntPointOnInt();
+		return FIntPoint::ZeroValue;
+	}
+
+	return A / FIntPoint(B);
+}
+
+KISMET_MATH_FORCEINLINE
+bool UKismetMathLibrary::Equal_IntPointIntPoint(FIntPoint A, FIntPoint B)
+{
+	return A == B;
+}
+
+KISMET_MATH_FORCEINLINE
+bool UKismetMathLibrary::NotEqual_IntPointIntPoint(FIntPoint A, FIntPoint B)
+{
+	return A != B;
+}
+
+
 /* FVector2D
 *****************************************************************************/
 
diff --git a/Engine/Source/Runtime/Engine/Classes/Kismet/KismetStringLibrary.h b/Engine/Source/Runtime/Engine/Classes/Kismet/KismetStringLibrary.h
--- a/Engine/Source/Runtime/Engine/Classes/Kismet/KismetStringLibrary.h
+++ b/Engine/Source/Runtime/Engine/Classes/Kismet/KismetStringLibrary.h
@@ -36,6 +36,10 @@ class ENGINE_API UKismetStringLibrary : public UBlueprintFunctionLibrary
 	UFUNCTION(BlueprintPure, meta=(DisplayName = "ToString (IntVector)", CompactNodeTitle = "->", BlueprintAutocast), Category="Utilities|String")
 	static FString Conv_IntVectorToString(FIntVector InIntVec);
 
+	/** Converts an IntPoint value to a string, in the form 'X= Y=' */
+	UFUNCTION(BlueprintPure, meta = (DisplayName = "ToString (IntPoint)", CompactNodeTitle = "->", BlueprintAutocast), Category = "Utilities|String")
+	static FString Conv_IntPointToString(FIntPoint InIntPoint);
+
 	/** Converts a vector2d value to a string, in the form 'X= Y=' */
 	UFUNCTION(BlueprintPure, meta=(DisplayName = "ToString (vector2d)", CompactNodeTitle = "->", BlueprintAutocast), Category="Utilities|String")
 	static FString Conv_Vector2dToString(FVector2D InVec);
diff --git a/Engine/Source/Runtime/Engine/Private/KismetMathLibrary.cpp b/Engine/Source/Runtime/Engine/Private/KismetMathLibrary.cpp
--- a/Engine/Source/Runtime/Engine/Private/KismetMathLibrary.cpp
+++ b/Engine/Source/Runtime/Engine/Private/KismetMathLibrary.cpp
@@ -128,6 +128,16 @@ void UKismetMathLibrary::ReportError_Divide_Vector2DFloat()
 	FFrame::KismetExecutionMessage(TEXT("Divide by zero: Divide_Vector2DFloat"), ELogVerbosity::Warning, DivideByZeroWarning);
 }
 
+void UKismetMathLibrary::ReportError_Divide_IntPointOnInt()
+{
+	FFrame::KismetExecutionMessage(TEXT("Divide by zero: Divide_IntPointInt"), ELogVerbosity::Warning, DivideByZeroWarning);
+}
+
+void UKismetMathLibrary::ReportError_Divide_IntPointOnIntPoint()
+{
+	FFrame::KismetExecutionMessage(TEXT("Divide by zero: Divide_IntPointIntPoint"), ELogVerbosity::Warning, DivideByZeroWarning);
+}
+
 void UKismetMathLibrary::ReportError_DaysInMonth()
 {
 	FFrame::KismetExecutionMessage(TEXT("Invalid month (must be between 1 and 12): DaysInMonth"), ELogVerbosity::Warning, InvalidDateWarning);
diff --git a/Engine/Source/Runtime/Engine/Private/KismetStringLibrary.cpp b/Engine/Source/Runtime/Engine/Private/KismetStringLibrary.cpp
--- a/Engine/Source/Runtime/Engine/Private/KismetStringLibrary.cpp
+++ b/Engine/Source/Runtime/Engine/Private/KismetStringLibrary.cpp
@@ -76,6 +76,11 @@ FString UKismetStringLibrary::Conv_IntVectorToString(FIntVector InIntVec)
 	return InIntVec.ToString();
 }
 
+FString UKismetStringLibrary::Conv_IntPointToString(FIntPoint InIntPoint)
+{
+	return InIntPoint.ToString();
+}
+
 FString UKismetStringLibrary::Conv_Vector2dToString(FVector2D InVec)
 {
 	return InVec.ToString();	
