PR #6411: Added conversions from int64 to int32 and to uint8 (byte) (Contributed by [REDACTED_PERSON])

diff --git a/Engine/Source/Runtime/Engine/Classes/Kismet/KismetMathLibrary.h b/Engine/Source/Runtime/Engine/Classes/Kismet/KismetMathLibrary.h
--- a/Engine/Source/Runtime/Engine/Classes/Kismet/KismetMathLibrary.h
+++ b/Engine/Source/Runtime/Engine/Classes/Kismet/KismetMathLibrary.h
@@ -3382,6 +3382,14 @@ class ENGINE_API UKismetMathLibrary : public UBlueprintFunctionLibrary
 	UFUNCTION(BlueprintPure, meta=(DisplayName = "ToByte (integer)", CompactNodeTitle = "->", Keywords="cast convert", BlueprintAutocast), Category="Math|Conversions")
 	static uint8 Conv_IntToByte(int32 InInt);
 
+	/** Converts a 64 bit integer to a 32 bit integer (if the integer is too large, returns the low 32 bits) */
+	UFUNCTION(BlueprintPure, meta=(DisplayName = "ToInt (Int64)", CompactNodeTitle = "->", Keywords="cast convert", BlueprintAutocast), Category="Math|Conversions")
+	static int32 Conv_Int64ToInt(int64 InInt);
+
+	/** Converts a 64 bit integer to a byte (if the integer is too large, returns the low 8 bits) */
+	UFUNCTION(BlueprintPure, meta=(DisplayName = "ToByte (Int64)", CompactNodeTitle = "->", Keywords="cast convert", BlueprintAutocast), Category="Math|Conversions")
+	static uint8 Conv_Int64ToByte(int64 InInt);
+
 	/** Converts an integer to an IntVector*/
 	UFUNCTION(BlueprintPure, meta = (DisplayName = "ToIntVector (integer)", CompactNodeTitle = "->", Keywords = "cast convert", BlueprintAutocast), Category = "Math|Conversions")
 	static FIntVector Conv_IntToIntVector(int32 InInt);
diff --git a/Engine/Source/Runtime/Engine/Classes/Kismet/KismetMathLibrary.inl b/Engine/Source/Runtime/Engine/Classes/Kismet/KismetMathLibrary.inl
--- a/Engine/Source/Runtime/Engine/Classes/Kismet/KismetMathLibrary.inl
+++ b/Engine/Source/Runtime/Engine/Classes/Kismet/KismetMathLibrary.inl
@@ -2573,6 +2573,18 @@ uint8 UKismetMathLibrary::Conv_IntToByte(int32 InInt)
 	return (uint8)InInt;
 }
 
+KISMET_MATH_FORCEINLINE
+int32 UKismetMathLibrary::Conv_Int64ToInt(int64 InInt)
+{
+	return (int32)InInt;
+}
+
+KISMET_MATH_FORCEINLINE
+uint8 UKismetMathLibrary::Conv_Int64ToByte(int64 InInt)
+{
+	return (uint8)InInt;
+}
+
 KISMET_MATH_FORCEINLINE
 FIntVector UKismetMathLibrary::Conv_IntToIntVector(int32 InInt)
 {
