- PR #6556: Add an argument to include from child actors when hiding/showing components in SceneCaptureComponent (Contributed by [REDACTED_PERSON])
- Change some fields from bool to uint8:1

diff --git a/Engine/Source/Runtime/Engine/Classes/Components/SceneCaptureComponent.h b/Engine/Source/Runtime/Engine/Classes/Components/SceneCaptureComponent.h
--- a/Engine/Source/Runtime/Engine/Classes/Components/SceneCaptureComponent.h
+++ b/Engine/Source/Runtime/Engine/Classes/Components/SceneCaptureComponent.h
@@ -68,6 +68,18 @@ class ENGINE_API USceneCaptureComponent : public USceneComponent
 	UPROPERTY(interp, Category = SceneCapture, meta = (DisplayName = "Capture Source"))
 	TEnumAsByte<enum ESceneCaptureSource> CaptureSource;
 
+	/** Whether to update the capture's contents every frame.  If disabled, the component will render once on load and then only when moved. */
+	UPROPERTY(EditAnywhere, BlueprintReadWrite, Category = SceneCapture)
+	uint8 bCaptureEveryFrame : 1;
+
+	/** Whether to update the capture's contents on movement.  Disable if you are going to capture manually from blueprint. */
+	UPROPERTY(EditAnywhere, BlueprintReadWrite, Category = SceneCapture)
+	uint8 bCaptureOnMovement : 1;
+
+	/** Whether to persist the rendering state even if bCaptureEveryFrame==false.  This allows velocities for Motion Blur and Temporal AA to be computed. */
+	UPROPERTY(EditAnywhere, BlueprintReadWrite, Category = SceneCapture, meta = (editcondition = "!bCaptureEveryFrame"))
+	bool bAlwaysPersistRenderingState;
+
 	/** The components won't rendered by current component.*/
  	UPROPERTY()
  	TArray<TWeakObjectPtr<UPrimitiveComponent> > HiddenComponents;
@@ -84,18 +96,6 @@ class ENGINE_API USceneCaptureComponent : public USceneComponent
 	UPROPERTY(EditInstanceOnly, BlueprintReadWrite, Category=SceneCapture)
 	TArray<AActor*> ShowOnlyActors;
 
-	/** Whether to update the capture's contents every frame.  If disabled, the component will render once on load and then only when moved. */
-	UPROPERTY(EditAnywhere, BlueprintReadWrite, Category=SceneCapture)
-	bool bCaptureEveryFrame;
-
-	/** Whether to update the capture's contents on movement.  Disable if you are going to capture manually from blueprint. */
-	UPROPERTY(EditAnywhere, BlueprintReadWrite, Category=SceneCapture)
-	bool bCaptureOnMovement;
-	
-	/** Whether to persist the rendering state even if bCaptureEveryFrame==false.  This allows velocities for Motion Blur and Temporal AA to be computed. */
-	UPROPERTY(EditAnywhere, BlueprintReadWrite, Category = SceneCapture, meta = (editcondition = "!bCaptureEveryFrame"))
-	bool bAlwaysPersistRenderingState;
-
 	/** Scales the distance used by LOD. Set to values greater than 1 to cause the scene capture to use lower LODs than the main view to speed up the scene capture pass. */
 	UPROPERTY(EditAnywhere, Category=PlanarReflection, meta=(UIMin = ".1", UIMax = "10"), AdvancedDisplay)
 	float LODDistanceFactor;
@@ -133,25 +133,34 @@ public:
 	UFUNCTION(BlueprintCallable, Category = "Rendering|SceneCapture")
 	void HideComponent(UPrimitiveComponent* InComponent);
 
-	/** Adds all primitive components in the actor to our list of hidden components. */
+	/**
+	 * Adds all primitive components in the actor to our list of hidden components.
+	 * @param bIncludeFromChildActors Whether to include the components from child actors
+	 */
 	UFUNCTION(BlueprintCallable, Category = "Rendering|SceneCapture")
-	void HideActorComponents(AActor* InActor);
+	void HideActorComponents(AActor* InActor, const bool bIncludeFromChildActors = false);
 
 	/** Adds the component to our list of show-only components. */
 	UFUNCTION(BlueprintCallable, Category = "Rendering|SceneCapture")
 	void ShowOnlyComponent(UPrimitiveComponent* InComponent);
 
-	/** Adds all primitive components in the actor to our list of show-only components. */
+	/**
+	 * Adds all primitive components in the actor to our list of show-only components.
+	 * @param bIncludeFromChildActors Whether to include the components from child actors
+	 */
 	UFUNCTION(BlueprintCallable, Category = "Rendering|SceneCapture")
-	void ShowOnlyActorComponents(AActor* InActor);
+	void ShowOnlyActorComponents(AActor* InActor, const bool bIncludeFromChildActors = false);
 
 	/** Removes a component from the Show Only list. */
 	UFUNCTION(BlueprintCallable, Category = "Rendering|SceneCapture")
 	void RemoveShowOnlyComponent(UPrimitiveComponent* InComponent);
 
-	/** Removes a actor's components from the Show Only list. */
+	/**
+	 * Removes an actor's components from the Show Only list.
+	 * @param bIncludeFromChildActors Whether to remove the components from child actors
+	 */
 	UFUNCTION(BlueprintCallable, Category = "Rendering|SceneCapture")
-	void RemoveShowOnlyActorComponents(AActor* InActor);
+	void RemoveShowOnlyActorComponents(AActor* InActor, const bool bIncludeFromChildActors = false);
 
 	/** Clears the Show Only list. */
 	UFUNCTION(BlueprintCallable, Category = "Rendering|SceneCapture")
diff --git a/Engine/Source/Runtime/Engine/Private/Components/SceneCaptureComponent.cpp b/Engine/Source/Runtime/Engine/Private/Components/SceneCaptureComponent.cpp
--- a/Engine/Source/Runtime/Engine/Private/Components/SceneCaptureComponent.cpp
+++ b/Engine/Source/Runtime/Engine/Private/Components/SceneCaptureComponent.cpp
@@ -210,17 +210,15 @@ void USceneCaptureComponent::HideComponent(UPrimitiveComponent* InComponent)
 	}
 }
 
-void USceneCaptureComponent::HideActorComponents(AActor* InActor)
+void USceneCaptureComponent::HideActorComponents(AActor* InActor, const bool bIncludeFromChildActors)
 {
 	if (InActor)
 	{
-		for (UActorComponent* Component : InActor->GetComponents())
+		TInlineComponentArray<UPrimitiveComponent*> PrimitiveComponents(InActor, bIncludeFromChildActors);
+		for (UPrimitiveComponent* PrimComp : PrimitiveComponents)
 		{
-			if (UPrimitiveComponent* PrimComp = Cast<UPrimitiveComponent>(Component))
-			{
-				TWeakObjectPtr<UPrimitiveComponent> WeakComponent(PrimComp);
-				HiddenComponents.AddUnique(WeakComponent);
-			}
+			TWeakObjectPtr<UPrimitiveComponent> WeakComponent(PrimComp);
+			HiddenComponents.AddUnique(WeakComponent);
 		}
 	}
 }
@@ -235,19 +233,17 @@ void USceneCaptureComponent::ShowOnlyComponent(UPrimitiveComponent* InComponent)
 	}
 }
 
-void USceneCaptureComponent::ShowOnlyActorComponents(AActor* InActor)
+void USceneCaptureComponent::ShowOnlyActorComponents(AActor* InActor, const bool bIncludeFromChildActors)
 {
 	if (InActor)
 	{
 		// Backward compatibility - set PrimitiveRenderMode to PRM_UseShowOnlyList if BP / game code tries to add a ShowOnlyComponent
 		PrimitiveRenderMode = ESceneCapturePrimitiveRenderMode::PRM_UseShowOnlyList;
 
-		for (UActorComponent* Component : InActor->GetComponents())
-			{
-			if (UPrimitiveComponent* PrimComp = Cast<UPrimitiveComponent>(Component))
-			{
-				ShowOnlyComponents.Add(PrimComp);
-			}
+		TInlineComponentArray<UPrimitiveComponent*> PrimitiveComponents(InActor, bIncludeFromChildActors);
+		for (UPrimitiveComponent* PrimComp : PrimitiveComponents)
+		{
+			ShowOnlyComponents.Add(PrimComp);
 		}
 	}
 }
@@ -258,17 +254,15 @@ void USceneCaptureComponent::RemoveShowOnlyComponent(UPrimitiveComponent* InComp
 	ShowOnlyComponents.Remove(WeakComponent);
 }
 
-void USceneCaptureComponent::RemoveShowOnlyActorComponents(AActor* InActor)
+void USceneCaptureComponent::RemoveShowOnlyActorComponents(AActor* InActor, const bool bIncludeFromChildActors)
 {
 	if (InActor)
 	{
-		for (UActorComponent* Component : InActor->GetComponents())
+		TInlineComponentArray<UPrimitiveComponent*> PrimitiveComponents(InActor, bIncludeFromChildActors);
+		for (UPrimitiveComponent* PrimComp : PrimitiveComponents)
 		{
-			if (UPrimitiveComponent* PrimComp = Cast<UPrimitiveComponent>(Component))
-			{
-				TWeakObjectPtr<UPrimitiveComponent> WeakComponent(PrimComp);
-				ShowOnlyComponents.Remove(WeakComponent);
-			}
+			TWeakObjectPtr<UPrimitiveComponent> WeakComponent(PrimComp);
+			ShowOnlyComponents.Remove(WeakComponent);
 		}
 	}
 }
