Added two broadcast delegates to GameViewportClient allowing one to listen for UGameViewportClient::InputKey and UGameViewportClient::InputAxis events.

diff --git a/Engine/Source/Runtime/Engine/Classes/Engine/GameViewportClient.h b/Engine/Source/Runtime/Engine/Classes/Engine/GameViewportClient.h
--- a/Engine/Source/Runtime/Engine/Classes/Engine/GameViewportClient.h
+++ b/Engine/Source/Runtime/Engine/Classes/Engine/GameViewportClient.h
@@ -34,6 +34,8 @@ class UNetDriver;
 
 /** Delegate for overriding the behavior when a navigation action is taken, Not to be confused with FNavigationDelegate which allows a specific widget to override behavior for itself */
 DECLARE_DELEGATE_RetVal_TwoParams(bool, FCustomNavigationHandler, const uint32, TSharedPtr<SWidget>);
+DECLARE_MULTICAST_DELEGATE_SevenParams(FOnInputAxisSignature, FViewport* /*InViewport*/, int32 /*ControllerId*/, FKey /*Key*/, float /*Delta*/, float /*DeltaTime*/, int32 /*NumSamples*/, bool /*bGamepad*/);
+DECLARE_MULTICAST_DELEGATE_OneParam(FOnInputKeySignature, const FInputKeyEventArgs& /*EventArgs*/);
 
 /**
  * A game viewport (FViewport) is a high-level abstract interface for the
@@ -572,6 +574,18 @@ public:
 		return CustomNavigationEvent;
 	}
 
+	/** fetches OnInputKeyEvent reference */
+	FOnInputKeySignature& OnInputKey()
+	{
+		return OnInputKeyEvent;
+	}
+
+	/** fetches OnInputAxisEvent reference */
+	FOnInputAxisSignature& OnInputAxis()
+	{
+		return OnInputAxisEvent;
+	}
+
 	/** Return the engine show flags for this viewport */
 	virtual FEngineShowFlags* GetEngineShowFlags() override
 	{ 
@@ -958,6 +972,12 @@ private:
 	/** Delegate for custom navigation behavior */
 	FCustomNavigationHandler CustomNavigationEvent;
 
+	/** A broadcast delegate broadcasting from UGameViewportClient::InputKey */
+	FOnInputKeySignature OnInputKeyEvent;
+
+	/** A broadcast delegate broadcasting from UGameViewportClient::InputAxis */
+	FOnInputAxisSignature OnInputAxisEvent;
+
 	/** Data needed to display perframe stat tracking when STAT UNIT is enabled */
 	FStatUnitData* StatUnitData;
 
@@ -995,4 +1015,4 @@ private:
 	/** Should the mouse send touch events. */
 	bool bUseMouseForTouchInEditor;
 #endif
-};
+};
\ No newline at end of file
diff --git a/Engine/Source/Runtime/Engine/Private/GameViewportClient.cpp b/Engine/Source/Runtime/Engine/Private/GameViewportClient.cpp
--- a/Engine/Source/Runtime/Engine/Private/GameViewportClient.cpp
+++ b/Engine/Source/Runtime/Engine/Private/GameViewportClient.cpp
@@ -534,6 +534,8 @@ bool UGameViewportClient::InputKey(const FInputKeyEventArgs& EventArgs)
 		GEngine->RemapGamepadControllerIdForPIE(this, ControllerId);
 	}
 
+	OnInputKeyEvent.Broadcast(EventArgs);
+
 #if WITH_EDITOR
 	// Give debugger commands a chance to process key binding
 	if (GameViewportInputKeyDelegate.IsBound())
@@ -600,6 +602,8 @@ bool UGameViewportClient::InputAxis(FViewport* InViewport, int32 ControllerId, F
 		GEngine->RemapGamepadControllerIdForPIE(this, ControllerId);
 	}
 
+	OnInputAxisEvent.Broadcast(InViewport, ControllerId, Key, Delta, DeltaTime, NumSamples, bGamepad);
+	
 	bool bResult = false;
 
 	// Don't allow mouse/joystick input axes while in PIE and the console has forced the cursor to be visible.  It's
