Actors that have pending latent actions are not ready to be destroyed

diff --git a/Engine/Source/Runtime/Engine/Private/Actor.cpp b/Engine/Source/Runtime/Engine/Private/Actor.cpp
--- a/Engine/Source/Runtime/Engine/Private/Actor.cpp
+++ b/Engine/Source/Runtime/Engine/Private/Actor.cpp
@@ -1008,6 +1008,8 @@ void AActor::TickActor( float DeltaSeconds, ELevelTick TickType, FActorTickFunct
 
 void AActor::Tick( float DeltaSeconds )
 {
+	bool bShouldAutoDestroy = bAutoDestroyWhenFinished;
+
 	if (GetClass()->HasAnyClassFlags(CLASS_CompiledFromBlueprint) || !GetClass()->HasAnyClassFlags(CLASS_Native))
 	{
 		// Blueprint code outside of the construction script should not run in the editor
@@ -1027,25 +1029,28 @@ void AActor::Tick( float DeltaSeconds )
 		UWorld* MyWorld = GetWorld();
 		if (MyWorld)
 		{
-			MyWorld->GetLatentActionManager().ProcessLatentActions(this, MyWorld->GetDeltaSeconds());
+			FLatentActionManager& LatentActionManager = MyWorld->GetLatentActionManager();
+			LatentActionManager.ProcessLatentActions(this, MyWorld->GetDeltaSeconds());
+			if (bShouldAutoDestroy)
+			{
+				bShouldAutoDestroy = (LatentActionManager.GetNumActionsForObject(this) == 0);
+			}
 		}
 	}
 
-	if (bAutoDestroyWhenFinished)
+	if (bShouldAutoDestroy)
 	{
-		bool bOKToDestroy = true;
-
 		for (UActorComponent* const Comp : GetComponents())
 		{
 			if (Comp && !Comp->IsReadyForOwnerToAutoDestroy())
 			{
-				bOKToDestroy = false;
+				bShouldAutoDestroy = false;
 				break;
 			}
 		}
 
 		// die!
-		if (bOKToDestroy)
+		if (bShouldAutoDestroy)
 		{
 			Destroy();
 		}
