PR #6601: Modify UCameraComponent::OnRegister so that editor-only components aren't created when running a commandlet (Contributed by [REDACTED_PERSON])

diff --git a/Engine/Source/Runtime/Engine/Private/Camera/CameraComponent.cpp b/Engine/Source/Runtime/Engine/Private/Camera/CameraComponent.cpp
--- a/Engine/Source/Runtime/Engine/Private/Camera/CameraComponent.cpp
+++ b/Engine/Source/Runtime/Engine/Private/Camera/CameraComponent.cpp
@@ -111,7 +111,8 @@ FText UCameraComponent::GetFilmbackText() const
 void UCameraComponent::OnRegister()
 {
 #if WITH_EDITORONLY_DATA
-	if (AActor* MyOwner = GetOwner())
+	AActor* MyOwner = GetOwner();
+	if ((MyOwner != nullptr) && !IsRunningCommandlet())
 	{
 		if (ProxyMeshComponent == nullptr)
 		{
