PR #6606: Added a IsEmpty function to the KismetStringLibrary (Contributed by [REDACTED_PERSON])

diff --git a/Engine/Source/Runtime/Engine/Classes/Kismet/KismetStringLibrary.h b/Engine/Source/Runtime/Engine/Classes/Kismet/KismetStringLibrary.h
--- a/Engine/Source/Runtime/Engine/Classes/Kismet/KismetStringLibrary.h
+++ b/Engine/Source/Runtime/Engine/Classes/Kismet/KismetStringLibrary.h
@@ -261,6 +261,14 @@ class ENGINE_API UKismetStringLibrary : public UBlueprintFunctionLibrary
 	 */
 	UFUNCTION(BlueprintPure, Category="Utilities|String", meta=(CompactNodeTitle = "LEN", Keywords = "length"))
 	static int32 Len(const FString& S);
+	
+	/**
+	 *	Returns true if the string is empty
+	 *	@param InString - The string to check
+	 *	@return Whether or not the string is empty
+	 */
+	UFUNCTION(BlueprintPure, Category="Utilities|String", meta=(Keywords = "empty"))
+	static bool IsEmpty(const FString& InString);
 
 	/** 
 	 * Returns a substring from the string starting at the specified position
diff --git a/Engine/Source/Runtime/Engine/Private/KismetStringLibrary.cpp b/Engine/Source/Runtime/Engine/Private/KismetStringLibrary.cpp
--- a/Engine/Source/Runtime/Engine/Private/KismetStringLibrary.cpp
+++ b/Engine/Source/Runtime/Engine/Private/KismetStringLibrary.cpp
@@ -46,6 +46,11 @@ int32 UKismetStringLibrary::Len(const FString& S)
 	return S.Len();
 }
 
+bool UKismetStringLibrary::IsEmpty(const FString& InString)
+{
+	return InString.IsEmpty();
+}
+
 FString UKismetStringLibrary::Conv_FloatToString(float InFloat)
 {
 	return FString::SanitizeFloat(InFloat);
