Mark the Default Physics volume transient in the world to avoid it causing a failure in CheckForWorldGCLeaks.

diff --git a/Engine/Source/Runtime/Engine/Classes/Engine/World.h b/Engine/Source/Runtime/Engine/Classes/Engine/World.h
--- a/Engine/Source/Runtime/Engine/Classes/Engine/World.h
+++ b/Engine/Source/Runtime/Engine/Classes/Engine/World.h
@@ -1004,7 +1004,7 @@ public:
 
 private:
 	/** DefaultPhysicsVolume used for whole game **/
-	UPROPERTY()
+	UPROPERTY(Transient)
 	APhysicsVolume*								DefaultPhysicsVolume;
 
 public:
