PR #6634: Passed through the DepthPriority parameter instead of hardcoded value (Contributed by [REDACTED_PERSON])

diff --git a/Engine/Source/Runtime/Engine/Private/PrimitiveDrawingUtils.cpp b/Engine/Source/Runtime/Engine/Private/PrimitiveDrawingUtils.cpp
--- a/Engine/Source/Runtime/Engine/Private/PrimitiveDrawingUtils.cpp
+++ b/Engine/Source/Runtime/Engine/Private/PrimitiveDrawingUtils.cpp
@@ -967,9 +967,9 @@ void DrawWireSphereAutoSides(class FPrimitiveDrawInterface* PDI, const FVector&
 
 void DrawWireSphere(class FPrimitiveDrawInterface* PDI, const FTransform& Transform, const FLinearColor& Color, float Radius, int32 NumSides, uint8 DepthPriority, float Thickness, float DepthBias, bool bScreenSpace)
 {
-	DrawCircle(PDI, Transform.GetLocation(), Transform.GetScaledAxis(EAxis::X), Transform.GetScaledAxis(EAxis::Y), Color, Radius, NumSides, SDPG_World, Thickness, DepthBias, bScreenSpace);
-	DrawCircle(PDI, Transform.GetLocation(), Transform.GetScaledAxis(EAxis::X), Transform.GetScaledAxis(EAxis::Z), Color, Radius, NumSides, SDPG_World, Thickness, DepthBias, bScreenSpace);
-	DrawCircle(PDI, Transform.GetLocation(), Transform.GetScaledAxis(EAxis::Y), Transform.GetScaledAxis(EAxis::Z), Color, Radius, NumSides, SDPG_World, Thickness, DepthBias, bScreenSpace);
+	DrawCircle(PDI, Transform.GetLocation(), Transform.GetScaledAxis(EAxis::X), Transform.GetScaledAxis(EAxis::Y), Color, Radius, NumSides, DepthPriority, Thickness, DepthBias, bScreenSpace);
+	DrawCircle(PDI, Transform.GetLocation(), Transform.GetScaledAxis(EAxis::X), Transform.GetScaledAxis(EAxis::Z), Color, Radius, NumSides, DepthPriority, Thickness, DepthBias, bScreenSpace);
+	DrawCircle(PDI, Transform.GetLocation(), Transform.GetScaledAxis(EAxis::Y), Transform.GetScaledAxis(EAxis::Z), Color, Radius, NumSides, DepthPriority, Thickness, DepthBias, bScreenSpace);
 }
 
 
