Optimize component instance data determination and application when rerunning construction scripts.
ActorComponents created from the User Construction Script now track a Serialized Index instead of reverse engineering it as part of component instance data cache actions.

diff --git a/Engine/Source/Runtime/Core/Public/UObject/FrameworkObjectVersion.h b/Engine/Source/Runtime/Core/Public/UObject/FrameworkObjectVersion.h
--- a/Engine/Source/Runtime/Core/Public/UObject/FrameworkObjectVersion.h
+++ b/Engine/Source/Runtime/Core/Public/UObject/FrameworkObjectVersion.h
@@ -125,6 +125,9 @@ struct CORE_API FFrameworkObjectVersion
 		// Enforce visibility of blueprint functions - e.g. raise an error if calling a private function from another blueprint:
 		EnforceBlueprintFunctionVisibility,
 
+		// ActorComponents now store their serialization index
+		StoringUCSSerializationIndex,
+
 		// -----<new versions can be added above this line>-------------------------------------------------
 		VersionPlusOne,
 		LatestVersion = VersionPlusOne - 1
diff --git a/Engine/Source/Runtime/Engine/Classes/Components/ActorComponent.h b/Engine/Source/Runtime/Engine/Classes/Components/ActorComponent.h
--- a/Engine/Source/Runtime/Engine/Classes/Components/ActorComponent.h
+++ b/Engine/Source/Runtime/Engine/Classes/Components/ActorComponent.h
@@ -120,6 +120,10 @@ private:
 	/** Used for fast removal of end of frame update */
 	int32 MarkedForEndOfFrameUpdateArrayIndex;
 
+	/** Populated when the component is created and tracks the often used order of creation on a per archetype/per actor basis */
+	UPROPERTY()
+	int32 UCSSerializationIndex;
+
 protected:
 	/** 
 	 *  Indicates if this ActorComponent is currently registered with a scene. 
@@ -219,6 +223,10 @@ private:
 	/** True if this component is only used for visualization, usually a sprite or text */
 	UPROPERTY()
 	uint8 bIsVisualizationComponent : 1;
+
+	/** Marks this component pending kill once PostLoad occurs. Used to clean up old native default subobjects that were removed from code */
+	UPROPERTY()
+	uint8 bNeedsUCSSerializationIndexEvaluted : 1;
 #endif
 
 private:
@@ -257,6 +265,25 @@ public:
 	EComponentCreationMethod CreationMethod;
 
 public:
+	/** Returns the UCS serialization index. This can be an expensive operation in the editor if you are dealing with a component that was saved before this information was present and it needs to be calculated. */
+	int32 GetUCSSerializationIndex() const
+	{
+#if WITH_EDITORONLY_DATA
+		if (bNeedsUCSSerializationIndexEvaluted)
+		{
+			const_cast<UActorComponent*>(this)->DetermineUCSSerializationIndexForLegacyComponent();
+		}
+#endif
+
+		return UCSSerializationIndex;
+	}
+
+private:
+	/** Calculate the UCS serialization index for a component that was saved before we started saving this data */
+	void DetermineUCSSerializationIndexForLegacyComponent();
+
+public:
+
 	/** Tracks whether the component has been added to one of the world's end of frame update lists */
 	uint32 GetMarkedForEndOfFrameUpdateState() const { return MarkedForEndOfFrameUpdateState; }
 
@@ -921,6 +948,7 @@ private:
 	void ClearNeedEndOfFrameUpdate_Internal();
 
 	friend struct FMarkComponentEndOfFrameUpdateState;
+	friend struct FSetUCSSerializationIndex;
 	friend struct FActorComponentInstanceData;
 	friend class FActorComponentDetails;
 	friend class FComponentReregisterContextBase;
diff --git a/Engine/Source/Runtime/Engine/Classes/GameFramework/Actor.h b/Engine/Source/Runtime/Engine/Classes/GameFramework/Actor.h
--- a/Engine/Source/Runtime/Engine/Classes/GameFramework/Actor.h
+++ b/Engine/Source/Runtime/Engine/Classes/GameFramework/Actor.h
@@ -753,10 +753,10 @@ private:
 	/** Whether this actor is temporarily hidden within the editor; used for show/hide/etc functionality w/o dirtying the actor. */
 	UPROPERTY(Transient)
 	uint8 bHiddenEdTemporary:1;
+#endif // WITH_EDITORONLY_DATA
 
 	/** Set while actor is being constructed. Used to ensure that construction is not re-entrant. */
-	uint8 bActorIsBeingConstructed:1;
-#endif // WITH_EDITORONLY_DATA
+	uint8 bActorIsBeingConstructed : 1;
 
 public:
 	/** Array of tags that can be used for grouping and categorizing. */
diff --git a/Engine/Source/Runtime/Engine/Private/ActorConstruction.cpp b/Engine/Source/Runtime/Engine/Private/ActorConstruction.cpp
--- a/Engine/Source/Runtime/Engine/Private/ActorConstruction.cpp
+++ b/Engine/Source/Runtime/Engine/Private/ActorConstruction.cpp
@@ -30,6 +30,7 @@
 #include "Components/ChildActorComponent.h"
 #include "ProfilingDebugging/CsvProfiler.h"
 #include "Algo/Transform.h"
+#include "Misc/ScopeExit.h"
 
 #if WITH_EDITOR
 #include "Editor.h"
@@ -486,42 +487,12 @@ void AActor::RerunConstructionScripts()
 				}
 
 				// Add a new item
-				ComponentMapping.Insert(FComponentData(), Index);
-				ComponentMapping[Index].OldComponent = Component;
-				ComponentMapping[Index].OldOuter = bIsInnerComponent ? CSAddedComponent : nullptr;
-				ComponentMapping[Index].OldArchetype = Component->GetArchetype();
-				ComponentMapping[Index].OldName = Component->GetFName();
-
-				// If it's a UCS-created component, store a serialized index which will be used to match it to the reinstanced counterpart later
-				int32 SerializedIndex = -1;
-				if (Component->CreationMethod == EComponentCreationMethod::UserConstructionScript)
-				{
-					bool bFound = false;
-					for (const UActorComponent* BlueprintCreatedComponent : BlueprintCreatedComponents)
-					{
-						if (BlueprintCreatedComponent)
-						{
-							if (BlueprintCreatedComponent == Component)
-							{
-								SerializedIndex++;
-								bFound = true;
-								break;
-							}
-							else if (BlueprintCreatedComponent->CreationMethod == EComponentCreationMethod::UserConstructionScript &&
-									 BlueprintCreatedComponent->GetArchetype() == ComponentMapping[Index].OldArchetype)
-							{
-								SerializedIndex++;
-							}
-						}
-					}
-
-					if (!bFound)
-					{
-						SerializedIndex = -1;
-					}
-				}
-
-				ComponentMapping[Index].UCSComponentIndex = SerializedIndex;
+				FComponentData& ComponentData = ComponentMapping.InsertDefaulted_GetRef(Index);
+				ComponentData.OldComponent = Component;
+				ComponentData.OldOuter = bIsInnerComponent ? CSAddedComponent : nullptr;
+				ComponentData.OldArchetype = Component->GetArchetype();
+				ComponentData.OldName = Component->GetFName();
+				ComponentData.UCSComponentIndex = Component->GetUCSSerializationIndex();
 			}
 		}
 #endif
@@ -593,12 +564,20 @@ void AActor::RerunConstructionScripts()
 		NameToNewComponent.Reserve(NewComponents.Num());
 		ComponentToArchetypeMap.Reserve(NewComponents.Num());
 
+		typedef TArray<UActorComponent*, TInlineAllocator<4>> FComponentsForSerializationIndex;
+		TMap<int32, FComponentsForSerializationIndex> SerializationIndexToNewUCSComponents;
+
 		for (UActorComponent* NewComponent : NewComponents)
 		{
 			if (GetComponentAddedByConstructionScript(NewComponent))
 			{
 				NameToNewComponent.Add(NewComponent->GetFName(), NewComponent);
 				ComponentToArchetypeMap.Add(NewComponent, NewComponent->GetArchetype());
+
+				if (NewComponent->CreationMethod == EComponentCreationMethod::UserConstructionScript)
+				{
+					SerializationIndexToNewUCSComponents.FindOrAdd(NewComponent->GetUCSSerializationIndex()).Add(NewComponent);
+				}
 			}
 		}
 
@@ -607,36 +586,22 @@ void AActor::RerunConstructionScripts()
 		{
 			if (ComponentData.OldComponent->CreationMethod == EComponentCreationMethod::UserConstructionScript)
 			{
-				// If created by the UCS, look for a component whose class, archetype and serialized index matches
-				for (UActorComponent* NewComponent : NewComponents)
+				if (ComponentData.UCSComponentIndex >= 0)
 				{
-					if (NewComponent->CreationMethod == EComponentCreationMethod::UserConstructionScript &&
-						ComponentData.OldComponent->GetClass() == NewComponent->GetClass() &&
-						ComponentData.OldArchetype == NewComponent->GetArchetype() &&
-						ComponentData.UCSComponentIndex >= 0)
+					// If created by the UCS, look for a component whose class, archetype, and serialized index matches
+					if (FComponentsForSerializationIndex* NewUCSComponentsToConsider = SerializationIndexToNewUCSComponents.Find(ComponentData.UCSComponentIndex))
 					{
-						int32 FoundSerializedIndex = -1;
-						bool bMatches = false;
-						for (UActorComponent* BlueprintCreatedComponent : BlueprintCreatedComponents)
+						for (int32 Index = 0; Index < NewUCSComponentsToConsider->Num(); ++Index)
 						{
-							if (BlueprintCreatedComponent && BlueprintCreatedComponent->CreationMethod == EComponentCreationMethod::UserConstructionScript)
+							UActorComponent* NewComponent = (*NewUCSComponentsToConsider)[Index];
+							if (   ComponentData.OldComponent->GetClass() == NewComponent->GetClass() 
+							    && ComponentData.OldArchetype == ComponentToArchetypeMap[NewComponent])
 							{
-								UObject* BlueprintComponentTemplate = ComponentToArchetypeMap.FindRef(BlueprintCreatedComponent);
-								if (BlueprintComponentTemplate &&
-									ComponentData.OldArchetype == BlueprintComponentTemplate &&
-									++FoundSerializedIndex == ComponentData.UCSComponentIndex)
-								{
-									bMatches = (BlueprintCreatedComponent == NewComponent);
-									break;
-								}
+								OldToNewComponentMapping.Add(ComponentData.OldComponent, NewComponent);
+								NewUCSComponentsToConsider->RemoveAtSwap(Index, 1, false);
+								break;
 							}
 						}
-
-						if (bMatches)
-						{
-							OldToNewComponentMapping.Add(ComponentData.OldComponent, NewComponent);
-							break;
-						}
 					}
 				}
 			}
@@ -693,6 +658,11 @@ void AActor::RerunConstructionScripts()
 	}
 }
 
+namespace
+{
+	TMap<const AActor*, TMap<const UObject*, int32>, TInlineSetAllocator<4>> UCSBlueprintComponentArchetypeCounts;
+}
+
 bool AActor::ExecuteConstruction(const FTransform& Transform, const FRotationConversionCache* TransformRotationCache, const FComponentInstanceDataCache* InstanceDataCache, bool bIsDefaultTransform)
 {
 	check(!IsPendingKill());
@@ -702,8 +672,13 @@ bool AActor::ExecuteConstruction(const FTransform& Transform, const FRotationCon
 	// Guard against reentrancy due to attribute editing at construction time.
 	// @see RerunConstructionScripts()
 	checkf(!bActorIsBeingConstructed, TEXT("Actor construction is not reentrant"));
-	FGuardValue_Bitfield(bActorIsBeingConstructed, true);
 #endif
+	bActorIsBeingConstructed = true;
+	ON_SCOPE_EXIT
+	{
+		bActorIsBeingConstructed = false;
+		UCSBlueprintComponentArchetypeCounts.Remove(this);
+	};
 
 	// ensure that any existing native root component gets this new transform
 	// we can skip this in the default case as the given transform will be the root component's transform
@@ -948,7 +923,7 @@ static FName FindFirstFreeName(UObject* Outer, FName BaseName)
 	// Binary search if we appear to have used this name a lot, else
 	// just linear search for the first free index:
 	if (FindObjectFast<UObject>(Outer, FName(BaseName, 100)))
-	{
+		{
 		// could be replaced by Algo::LowerBound if TRange or some other
 		// integer range type could be made compatible with the algo's
 		// implementation:
@@ -1223,6 +1198,17 @@ void AActor::CheckComponentInstanceName(const FName InName)
 	}
 }
 
+struct FSetUCSSerializationIndex
+{
+	friend class AActor;
+
+private:
+	FORCEINLINE static void Set(UActorComponent* Component, int32 SerializationIndex)
+	{
+		Component->UCSSerializationIndex = SerializationIndex;
+	}
+};
+
 void AActor::PostCreateBlueprintComponent(UActorComponent* NewActorComp)
 {
 	if (NewActorComp)
@@ -1232,6 +1218,14 @@ void AActor::PostCreateBlueprintComponent(UActorComponent* NewActorComp)
 		// Need to do this so component gets saved - Components array is not serialized
 		BlueprintCreatedComponents.Add(NewActorComp);
 
+		if (bActorIsBeingConstructed)
+		{
+			TMap<const UObject*, int32>& ComponentArchetypeCounts = UCSBlueprintComponentArchetypeCounts.FindOrAdd(this);
+			int32& Count = ComponentArchetypeCounts.FindOrAdd(NewActorComp->GetArchetype());
+			FSetUCSSerializationIndex::Set(NewActorComp, Count);
+			++Count;
+		}
+
 		// The component may not have been added to ReplicatedComponents if it was duplicated from
 		// a template, since ReplicatedComponents is normally only updated before the duplicated properties
 		// are copied over - in this case bReplicates would not have been set yet, but it will be now.
diff --git a/Engine/Source/Runtime/Engine/Private/ComponentInstanceDataCache.cpp b/Engine/Source/Runtime/Engine/Private/ComponentInstanceDataCache.cpp
--- a/Engine/Source/Runtime/Engine/Private/ComponentInstanceDataCache.cpp
+++ b/Engine/Source/Runtime/Engine/Private/ComponentInstanceDataCache.cpp
@@ -300,37 +300,7 @@ FActorComponentInstanceData::FActorComponentInstanceData(const UActorComponent*
 	SourceComponentTemplate = SourceComponent->GetArchetype();
 	SourceComponentCreationMethod = SourceComponent->CreationMethod;
 	SourceComponentTypeSerializedIndex = INDEX_NONE;
-
-	// UCS components can share the same template (e.g. an AddComponent node inside a loop), so we also cache their serialization index here (relative to the shared template) as a means for identification
-	if(SourceComponentCreationMethod == EComponentCreationMethod::UserConstructionScript)
-	{
-		AActor* ComponentOwner = SourceComponent->GetOwner();
-		if (ComponentOwner)
-		{
-			bool bFound = false;
-			for (const UActorComponent* BlueprintCreatedComponent : ComponentOwner->BlueprintCreatedComponents)
-			{
-				if (BlueprintCreatedComponent)
-				{
-					if (BlueprintCreatedComponent == SourceComponent)
-					{
-						++SourceComponentTypeSerializedIndex;
-						bFound = true;
-						break;
-					}
-					else if (   BlueprintCreatedComponent->CreationMethod == SourceComponentCreationMethod
-						&& BlueprintCreatedComponent->GetArchetype() == SourceComponentTemplate)
-					{
-						++SourceComponentTypeSerializedIndex;
-					}
-				}
-			}
-			if (!bFound)
-			{
-				SourceComponentTypeSerializedIndex = INDEX_NONE;
-			}
-		}
-	}
+	SourceComponentTypeSerializedIndex = SourceComponent->GetUCSSerializationIndex();
 
 	if (SourceComponent->IsEditableWhenInherited())
 	{
@@ -384,24 +354,7 @@ bool FActorComponentInstanceData::MatchesComponent(const UActorComponent* Compon
 		}
 		else if (SourceComponentTypeSerializedIndex >= 0)
 		{
-			int32 FoundSerializedComponentsOfType = -1;
-			AActor* ComponentOwner = Component->GetOwner();
-			if (ComponentOwner)
-			{
-				for (const UActorComponent* BlueprintCreatedComponent : ComponentOwner->BlueprintCreatedComponents)
-				{
-					if (BlueprintCreatedComponent != nullptr && BlueprintCreatedComponent->CreationMethod == SourceComponentCreationMethod)
-					{
-						const UObject* BlueprintComponentTemplate = ComponentToArchetypeMap.FindChecked(BlueprintCreatedComponent);
-						if (   (BlueprintComponentTemplate == SourceComponentTemplate || (GIsReinstancing && BlueprintComponentTemplate->GetFName() == SourceComponentTemplate->GetFName()))
-							&& (++FoundSerializedComponentsOfType == SourceComponentTypeSerializedIndex))
-						{
-							bMatches = (BlueprintCreatedComponent == Component);
-							break;
-						}
-					}
-				}
-			}
+			bMatches = (SourceComponentTypeSerializedIndex == Component->GetUCSSerializationIndex());
 		}
 	}
 	return bMatches;
@@ -609,13 +562,22 @@ void FComponentInstanceDataCache::ApplyToActor(AActor* Actor, const ECacheApplyP
 
 		for (UActorComponent* ComponentInstance : Components)
 		{
-			if (ComponentInstance && (bIsChildActor || ComponentInstance->IsCreatedByConstructionScript()))
+			if (bIsChildActor || ComponentInstance->IsCreatedByConstructionScript())
 			{
 				ComponentToArchetypeMap.Add(ComponentInstance, ComponentInstance->GetArchetype());
 			}
 		}
 
 		// Apply per-instance data.
+		TBitArray<> ComponentInstanceDataToConsider(true, ComponentsInstanceData.Num());
+		for (int32 Index = 0; Index < ComponentsInstanceData.Num(); ++Index)
+		{
+			if (!ComponentsInstanceData[Index].IsValid())
+			{
+				ComponentInstanceDataToConsider[Index] = false;
+			}
+		}
+
 		for (UActorComponent* ComponentInstance : Components)
 		{
 			if (ComponentInstance && (bIsChildActor || ComponentInstance->IsCreatedByConstructionScript())) // Only try and apply data to 'created by construction script' components
@@ -623,13 +585,14 @@ void FComponentInstanceDataCache::ApplyToActor(AActor* Actor, const ECacheApplyP
 				// Cache template here to avoid redundant calls in the loop below
 				const UObject* ComponentTemplate = ComponentToArchetypeMap.FindChecked(ComponentInstance);
 
-				for (const TStructOnScope<FActorComponentInstanceData>& ComponentInstanceData : ComponentsInstanceData)
+				for (TConstSetBitIterator<> ComponentInstanceDataIt(ComponentInstanceDataToConsider); ComponentInstanceDataIt; ++ComponentInstanceDataIt)
 				{
-					if (	ComponentInstanceData.IsValid()
-						&&	ComponentInstanceData->GetComponentClass() == ComponentTemplate->GetClass() // filter on class early to avoid unnecessary virtual and expensive tests
-						&&	ComponentInstanceData->MatchesComponent(ComponentInstance, ComponentTemplate, ComponentToArchetypeMap))
+					const TStructOnScope<FActorComponentInstanceData>& ComponentInstanceData = ComponentsInstanceData[ComponentInstanceDataIt.GetIndex()];
+					if (   ComponentInstanceData->GetComponentClass() == ComponentTemplate->GetClass() // filter on class early to avoid unnecessary virtual and expensive tests
+					    && ComponentInstanceData->MatchesComponent(ComponentInstance, ComponentTemplate, ComponentToArchetypeMap))
 					{
 						ComponentInstanceData->ApplyToComponent(ComponentInstance, CacheApplyPhase);
+						ComponentInstanceDataToConsider[ComponentInstanceDataIt.GetIndex()] = false;
 						break;
 					}
 				}
@@ -637,7 +600,7 @@ void FComponentInstanceDataCache::ApplyToActor(AActor* Actor, const ECacheApplyP
 		}
 
 		// Once we're done attaching, if we have any unattached instance components move them to the root
-		for (auto InstanceTransformPair : InstanceComponentTransformToRootMap)
+		for (const TPair<USceneComponent*, FTransform>& InstanceTransformPair : InstanceComponentTransformToRootMap)
 		{
 			check(Actor->GetRootComponent());
 
diff --git a/Engine/Source/Runtime/Engine/Private/Components/ActorComponent.cpp b/Engine/Source/Runtime/Engine/Private/Components/ActorComponent.cpp
--- a/Engine/Source/Runtime/Engine/Private/Components/ActorComponent.cpp
+++ b/Engine/Source/Runtime/Engine/Private/Components/ActorComponent.cpp
@@ -31,6 +31,7 @@
 #include "Engine/Engine.h"
 #include "HAL/LowLevelMemTracker.h"
 #include "Net/Core/PushModel/PushModel.h"
+#include "UObject/FrameworkObjectVersion.h"
 
 #if WITH_EDITOR
 #include "Kismet2/ComponentEditorUtils.h"
@@ -218,6 +219,7 @@ UActorComponent::UActorComponent(const FObjectInitializer& ObjectInitializer /*=
 	PrimaryComponentTick.SetTickFunctionEnable(false);
 
 	MarkedForEndOfFrameUpdateArrayIndex = INDEX_NONE;
+	UCSSerializationIndex = INDEX_NONE;
 
 	CreationMethod = EComponentCreationMethod::Native;
 
@@ -284,7 +286,7 @@ void UActorComponent::PostInitProperties()
 void UActorComponent::PostLoad()
 {
 	Super::PostLoad();
-
+	   
 #if WITH_EDITORONLY_DATA
 	if (GetLinkerUE4Version() < VER_UE4_ACTOR_COMPONENT_CREATION_METHOD)
 	{
@@ -335,6 +337,16 @@ void UActorComponent::PostLoad()
 	{
 		// For a brief period of time we were inadvertently storing these for all components, need to clear it out
 		UCSModifiedProperties.Empty();
+
+#if WITH_EDITORONLY_DATA
+		if (CreationMethod == EComponentCreationMethod::UserConstructionScript)
+		{
+			if (GetLinkerCustomVersion(FFrameworkObjectVersion::GUID) < FFrameworkObjectVersion::StoringUCSSerializationIndex)
+			{
+				bNeedsUCSSerializationIndexEvaluted = true;
+			}
+		}
+#endif
 	}
 
 #if WITH_EDITOR
@@ -420,6 +432,48 @@ bool UActorComponent::IsCreatedByConstructionScript() const
 	return ((CreationMethod == EComponentCreationMethod::SimpleConstructionScript) || (CreationMethod == EComponentCreationMethod::UserConstructionScript));
 }
 
+#if WITH_EDITORONLY_DATA
+void UActorComponent::DetermineUCSSerializationIndexForLegacyComponent()
+{
+	check(CreationMethod == EComponentCreationMethod::UserConstructionScript);
+	check(bNeedsUCSSerializationIndexEvaluted);
+	bNeedsUCSSerializationIndexEvaluted = false;
+
+	int32 ComputedSerializationIndex = INDEX_NONE;
+
+	if (AActor* ComponentOwner = GetOwner())
+	{
+		if (ComponentOwner->BlueprintCreatedComponents.Num() > 0)
+		{
+			UObject* ComponentTemplate = GetArchetype();
+
+			bool bFound = false;
+			for (const UActorComponent* BlueprintCreatedComponent : ComponentOwner->BlueprintCreatedComponents)
+			{
+				if (BlueprintCreatedComponent && BlueprintCreatedComponent->CreationMethod == EComponentCreationMethod::UserConstructionScript)
+				{
+					if (BlueprintCreatedComponent == this)
+					{
+						++ComputedSerializationIndex;
+						bFound = true;
+						break;
+					}
+					else if (BlueprintCreatedComponent->GetArchetype() == ComponentTemplate)
+					{
+						++ComputedSerializationIndex;
+					}
+				}
+			}
+			if (!bFound)
+			{
+				ComputedSerializationIndex = INDEX_NONE;
+			}
+		}
+	}
+	UCSSerializationIndex = ComputedSerializationIndex;
+}
+#endif
+
 #if WITH_EDITOR
 void UActorComponent::CheckForErrors()
 {
@@ -1953,6 +2007,8 @@ void UActorComponent::Serialize(FArchive& Ar)
 {
 	Super::Serialize(Ar);
 
+	Ar.UsingCustomVersion(FFrameworkObjectVersion::GUID);
+
 	if (Ar.IsLoading() && (Ar.HasAnyPortFlags(PPF_DuplicateForPIE)||!Ar.HasAnyPortFlags(PPF_Duplicate)) && !IsTemplate())
 	{
 		bHasBeenCreated = true;
