Add keywords to various material expression nodes so that they show up with more user-friendly searches

#rnx

diff --git a/Engine/Source/Runtime/Engine/Classes/Materials/MaterialExpressionComponentMask.h b/Engine/Source/Runtime/Engine/Classes/Materials/MaterialExpressionComponentMask.h
--- a/Engine/Source/Runtime/Engine/Classes/Materials/MaterialExpressionComponentMask.h
+++ b/Engine/Source/Runtime/Engine/Classes/Materials/MaterialExpressionComponentMask.h
@@ -34,6 +34,7 @@ class UMaterialExpressionComponentMask : public UMaterialExpression
 #if WITH_EDITOR
 	virtual int32 Compile(class FMaterialCompiler* Compiler, int32 OutputIndex) override;
 	virtual void GetCaption(TArray<FString>& OutCaptions) const override;
+	virtual FText GetKeywords() const override { return FText::FromString(TEXT("component mask")); }
 #endif
 	//~ End UMaterialExpression Interface
 };
diff --git a/Engine/Source/Runtime/Engine/Classes/Materials/MaterialExpressionDynamicParameter.h b/Engine/Source/Runtime/Engine/Classes/Materials/MaterialExpressionDynamicParameter.h
--- a/Engine/Source/Runtime/Engine/Classes/Materials/MaterialExpressionDynamicParameter.h
+++ b/Engine/Source/Runtime/Engine/Classes/Materials/MaterialExpressionDynamicParameter.h
@@ -53,6 +53,7 @@ class UMaterialExpressionDynamicParameter : public UMaterialExpression
 	virtual bool MatchesSearchQuery(const TCHAR* SearchQuery) override;
 	virtual int32 GetWidth() const override;
 	virtual int32 GetLabelPadding() override { return 8; }
+	virtual FText GetKeywords() const override { return FText::FromString(TEXT("dynamic")); }
 #endif
 	//~ End UMaterialExpression Interface
 
diff --git a/Engine/Source/Runtime/Engine/Classes/Materials/MaterialExpressionNormalize.h b/Engine/Source/Runtime/Engine/Classes/Materials/MaterialExpressionNormalize.h
--- a/Engine/Source/Runtime/Engine/Classes/Materials/MaterialExpressionNormalize.h
+++ b/Engine/Source/Runtime/Engine/Classes/Materials/MaterialExpressionNormalize.h
@@ -22,6 +22,7 @@ class UMaterialExpressionNormalize : public UMaterialExpression
 #if WITH_EDITOR
 	virtual int32 Compile(class FMaterialCompiler* Compiler, int32 OutputIndex) override;
 	virtual void GetCaption(TArray<FString>& OutCaptions) const override { OutCaptions.Add(TEXT("Normalize")); }
+	virtual FText GetKeywords() const override { return FText::FromString(TEXT("normal")); }
 #endif
 	//~ End UMaterialExpression Interface
 };
diff --git a/Engine/Source/Runtime/Engine/Classes/Materials/MaterialExpressionParticleColor.h b/Engine/Source/Runtime/Engine/Classes/Materials/MaterialExpressionParticleColor.h
--- a/Engine/Source/Runtime/Engine/Classes/Materials/MaterialExpressionParticleColor.h
+++ b/Engine/Source/Runtime/Engine/Classes/Materials/MaterialExpressionParticleColor.h
@@ -18,6 +18,7 @@ class UMaterialExpressionParticleColor : public UMaterialExpression
 #if WITH_EDITOR
 	virtual int32 Compile(class FMaterialCompiler* Compiler, int32 OutputIndex) override;
 	virtual void GetCaption(TArray<FString>& OutCaptions) const override;
+	virtual FText GetKeywords() const override { return FText::FromString(TEXT("particle color")); }
 #endif
 	//~ End UMaterialExpression Interface
 };
diff --git a/Engine/Source/Runtime/Engine/Classes/Materials/MaterialExpressionSphereMask.h b/Engine/Source/Runtime/Engine/Classes/Materials/MaterialExpressionSphereMask.h
--- a/Engine/Source/Runtime/Engine/Classes/Materials/MaterialExpressionSphereMask.h
+++ b/Engine/Source/Runtime/Engine/Classes/Materials/MaterialExpressionSphereMask.h
@@ -43,6 +43,7 @@ class UMaterialExpressionSphereMask : public UMaterialExpression
 #if WITH_EDITOR
 	virtual int32 Compile(class FMaterialCompiler* Compiler, int32 OutputIndex) override;
 	virtual void GetCaption(TArray<FString>& OutCaptions) const override;
+	virtual FText GetKeywords() const override { return FText::FromString(TEXT("sphere")); }
 #endif
 	//~ End UMaterialExpression Interface
 
diff --git a/Engine/Source/Runtime/Engine/Classes/Materials/MaterialExpressionVertexColor.h b/Engine/Source/Runtime/Engine/Classes/Materials/MaterialExpressionVertexColor.h
--- a/Engine/Source/Runtime/Engine/Classes/Materials/MaterialExpressionVertexColor.h
+++ b/Engine/Source/Runtime/Engine/Classes/Materials/MaterialExpressionVertexColor.h
@@ -18,6 +18,7 @@ class UMaterialExpressionVertexColor : public UMaterialExpression
 #if WITH_EDITOR
 	virtual int32 Compile(class FMaterialCompiler* Compiler, int32 OutputIndex) override;
 	virtual void GetCaption(TArray<FString>& OutCaptions) const override;
+	virtual FText GetKeywords() const override { return FText::FromString(TEXT("vertex color")); }
 #endif
 	//~ End UMaterialExpression Interface
 };
