Change RepLayout Push Model tracking to use a bitfield.

This should significantly reduce the number of cache misses incurred when properties haven't changed.

diff --git a/Engine/Source/Runtime/Engine/Private/RepLayout.cpp b/Engine/Source/Runtime/Engine/Private/RepLayout.cpp
--- a/Engine/Source/Runtime/Engine/Private/RepLayout.cpp
+++ b/Engine/Source/Runtime/Engine/Private/RepLayout.cpp
@@ -1134,6 +1134,7 @@ struct FComparePropertiesSharedParams
 	FRepChangelistState* const RepChangelistState;
 	FRepChangedPropertyTracker* const RepChangedPropertyTracker;
 	UE4PushModelPrivate::FPushModelPerNetDriverState* const PushModelState = nullptr;
+	const TBitArray<>* const PushModelProperties = nullptr;
 	const bool bValidateProperties = false;
 	const bool bIsNetworkProfilerActive = false;
 #if (WITH_PUSH_VALIDATION_SUPPORT || USE_NETWORK_PROFILER)
@@ -1269,7 +1270,7 @@ namespace UE4_RepLayout_Private
 		const FComparePropertiesSharedParams& SharedParams,
 		FComparePropertiesStackParams& StackParams)
 	{
-		return !EnumHasAnyFlags(SharedParams.Parents[ParentIndex].Flags, ERepParentFlags::UsePushModel) || SharedParams.PushModelState->IsPropertyDirty(ParentIndex);
+		return !(*SharedParams.PushModelProperties)[ParentIndex] || SharedParams.PushModelState->IsPropertyDirty(ParentIndex);
 	}
 #endif // WITH_PUSH_MODEL	
 }	
@@ -1313,7 +1314,7 @@ static void CompareParentProperties(
 				const bool bIsPropertyDirty = UE4_RepLayout_Private::IsPropertyDirty(ParentIndex, SharedParams, StackParams);
 				const bool bDidPropertyChange = UE4_RepLayout_Private::CompareParentPropertyHelper(ParentIndex, SharedParams, StackParams);
 
-				ensureAlwaysMsgf(bDidPropertyChange && !bIsPropertyDirty, TEXT("Push Model Property changed value, but was not marked dirty! Property=%s"), *SharedParams.Parents[ParentIndex].Property->GetPathName());
+				ensureAlwaysMsgf(!bDidPropertyChange || bIsPropertyDirty, TEXT("Push Model Property changed value, but was not marked dirty! Property=%s"), *SharedParams.Parents[ParentIndex].Property->GetPathName());
 			}	
 		}
 #endif // WITH_PUSH_VALIDATION_SUPPORT
@@ -1482,6 +1483,12 @@ bool FRepLayout::CompareProperties(
 	TArray<uint16>& Changed = NewHistoryItem.Changed;
 	Changed.Empty(1);
 
+#if WITH_PUSH_MODEL
+		const TBitArray<>* const LocalPushModelProperties = &PushModelProperties;
+#else
+		const TBitArray<>* const LocalPushModelProperties = nullptr;
+#endif		
+
 	FComparePropertiesSharedParams SharedParams{
 		/*bIsInitial=*/ !!RepFlags.bNetInitial,
 		/*bForceFail=*/ false,
@@ -1492,6 +1499,7 @@ bool FRepLayout::CompareProperties(
 		RepChangelistState,
 		(RepState ? RepState->RepChangedPropertyTracker.Get() : nullptr),
 		/*PushModelState=*/UE4_RepLayout_Private::GetPerNetDriverState(RepChangelistState),
+		/*PushModelProperties=*/ LocalPushModelProperties,	
 		/*bValidateProperties=*/GbPushModelValidateProperties,
 		/*bIsNetworkProfilerActive=*/UE4_RepLayout_Private::IsNetworkProfilerEnabled()
 	};
@@ -5454,6 +5462,10 @@ void FRepLayout::InitFromClass(
 
 	Object->GetLifetimeReplicatedProps(LifetimeProps);
 
+#if WITH_PUSH_MODEL
+	PushModelProperties.Init(false, Parents.Num());
+#endif
+
 	// Tracks the number of (non-delta) lifetime properties so we can check that against our
 	// Push Model Enabled properties.
 	int32 NumberOfLifetimeProperties = 0;
@@ -5499,7 +5511,7 @@ void FRepLayout::InitFromClass(
 			if (bIsPushModelEnabled && LifetimeProps[i].bIsPushBased)
 			{
 				++NumberOfPushModelProperties;
-				Parents[ParentIndex].Flags |= ERepParentFlags::UsePushModel;
+				PushModelProperties[ParentIndex] = true;
 			}
 #endif
 		}
diff --git a/Engine/Source/Runtime/Engine/Public/Net/RepLayout.h b/Engine/Source/Runtime/Engine/Public/Net/RepLayout.h
--- a/Engine/Source/Runtime/Engine/Public/Net/RepLayout.h
+++ b/Engine/Source/Runtime/Engine/Public/Net/RepLayout.h
@@ -716,7 +716,6 @@ enum class ERepParentFlags : uint32
 	IsZeroConstructible	= (1 << 6),	//! This property is ZeroConstructible.
 	IsFastArray			= (1 << 7),	//! This property is a FastArraySerializer. This can't be a ERepLayoutCmdType, because
 									//! these Custom Delta structs will have their inner properties tracked.
-	UsePushModel		= (1 << 8),	//! This object relies on PushModel, and can be skipped when not dirty.	
 };
 
 ENUM_CLASS_FLAGS(ERepParentFlags)
@@ -1015,13 +1014,6 @@ enum class ERepLayoutFlags : uint8
 };
 ENUM_CLASS_FLAGS(ERepLayoutFlags);
 
-enum class UE_DEPRECATED(4.24, "Use FRepLayout::IsEmpty() instead.") ERepLayoutState
-{
-	Uninitialized,	//! The RepLayout was never initiliazed, this should not be possible.
-	Empty,			//! The RepLayout was initialized, but doesn't have any RepCommands.
-					//! This can happen when replicating References to actors with no network state (e.g., Item Definitions, etc.).
-	Normal			//! The RepLayout was initialized, and contains commands.
-};
 
 /**
  * This class holds all replicated properties for a given type (either a UClass, UStruct, or UFunction).
@@ -1292,19 +1284,6 @@ public:
 	 */
 	void CallRepNotifies(FReceivingRepState* RepState, UObject* Object) const;
 
-	/**
-	 * Called after an object has finished replicating its properties.
-	 *
-	 * @param RepState		The RepState associated with the Object.
-	 *						This is expected to be valid.
-	 * @param PacketRange	Range / IDs of the Packets sent containing the property data.
-	 * @param bReliable		Whether or not the data was sent in reliable bunches.
-	 */
-	void PostReplicate(
-		FSendingRepState* RESTRICT RepState,
-		FPacketIdRange& PacketRange,
-		bool bReliable) const;
-
 	template<ERepDataBufferType DataType>
 	void ValidateWithChecksum(TConstRepDataBuffer<DataType> Data, FBitArchive & Ar) const;
 
@@ -1444,13 +1423,6 @@ public:
 		return nullptr;
 	}
 
-PRAGMA_DISABLE_DEPRECATION_WARNINGS
-
-	UE_DEPRECATED(4.24, "Use GetFlags instead.")
-	const ERepLayoutState GetRepLayoutState() const;
-
-PRAGMA_ENABLE_DEPRECATION_WARNINGS
-
 	const ERepLayoutFlags GetFlags() const
 	{
 		return Flags;
@@ -1815,14 +1787,9 @@ private:
 
 	/** Shared comparison to default state for multicast rpc */
 	TBitArray<> SharedInfoRPCParentsChanged;
-};
-
 
-PRAGMA_DISABLE_DEPRECATION_WARNINGS
-inline const ERepLayoutState FRepLayout::GetRepLayoutState() const
-{
-	// Because our constructor is hidden, we're guaranteed to be initialized.
-	// We're either Empty or Normal.
-	return IsEmpty() ? ERepLayoutState::Empty : ERepLayoutState::Normal;
-}
-PRAGMA_ENABLE_DEPRECATION_WARNINGS
+#if WITH_PUSH_MODEL
+	/** Properties that have push model enabled. */
+	TBitArray<> PushModelProperties;
+#endif
+};
\ No newline at end of file
