Non-determinism in cooking UStaticMeshComponent::LODData

Properly invalidate FStaticMeshComponentLODInfo::MapDataBuildId to ensure we have a deterministic cook

diff --git a/Engine/Source/Runtime/Engine/Classes/Components/StaticMeshComponent.h b/Engine/Source/Runtime/Engine/Classes/Components/StaticMeshComponent.h
--- a/Engine/Source/Runtime/Engine/Classes/Components/StaticMeshComponent.h
+++ b/Engine/Source/Runtime/Engine/Classes/Components/StaticMeshComponent.h
@@ -89,11 +89,6 @@ struct FStaticMeshComponentLODInfo
 	/** Uniquely identifies this LOD's built map data. */
 	FGuid MapBuildDataId;
 
-#if WITH_EDITOR
-	/** Check to see if MapBuildDataId was loaded - otherwise we need to display a warning on cook */
-	bool bMapBuildDataIdLoaded;
-#endif
-
 	/** Used during deserialization to temporarily store legacy lightmap data. */
 	FMeshMapBuildData* LegacyMapBuildData;
 
@@ -119,13 +114,20 @@ struct FStaticMeshComponentLODInfo
 
 	/** Default constructor */
 	FStaticMeshComponentLODInfo();
-	FStaticMeshComponentLODInfo(UStaticMeshComponent* InOwningComponent, int32 LodIndex);
+	FStaticMeshComponentLODInfo(UStaticMeshComponent* InOwningComponent);
 	/** Destructor */
 	~FStaticMeshComponentLODInfo();
 
 	/** Delete existing resources */
 	void CleanUp();
 
+	/** 
+	 * Ensure this LODInfo has a valid MapBuildDataId GUID.
+	 * @param LodIndex Index of the LOD this LODInfo represents.
+	 * @return true if a new GUID was created, false otherwise.
+	 */
+	bool CreateMapBuildDataId(int32 LodIndex);
+
 	/**
 	* Enqueues a rendering command to release the vertex colors.
 	* The game thread must block until the rendering thread has processed the command before deleting OverrideVertexColors.
diff --git a/Engine/Source/Runtime/Engine/Private/Components/StaticMeshComponent.cpp b/Engine/Source/Runtime/Engine/Private/Components/StaticMeshComponent.cpp
--- a/Engine/Source/Runtime/Engine/Private/Components/StaticMeshComponent.cpp
+++ b/Engine/Source/Runtime/Engine/Private/Components/StaticMeshComponent.cpp
@@ -110,9 +110,6 @@ bool FStaticMeshComponentInstanceData::ApplyVertexColorData(UStaticMeshComponent
 			if(CachedStaticLighting.IsValidIndex(LODIndex))
 			{
 				LODInfo.MapBuildDataId = CachedStaticLighting[LODIndex];
-#if WITH_EDITOR
-				LODInfo.bMapBuildDataIdLoaded = true;
-#endif
 			}
 		}
 
@@ -277,16 +274,6 @@ void UStaticMeshComponent::Serialize(FArchive& Ar)
 
 	Ar.UsingCustomVersion(FRenderingObjectVersion::GUID);
 
-#if WITH_EDITOR
-	const bool bCheckBuildGuids = false;//Ar.IsCooking() && !GetOutermost()->HasAnyPackageFlags(PKG_CompiledIn);
-
-	TArray<FGuid> MapBuildGuids;
-	if (bCheckBuildGuids)
-	{
-		Algo::Transform(LODData, MapBuildGuids, &FStaticMeshComponentLODInfo::MapBuildDataId);
-	}
-#endif
-
 #if WITH_EDITORONLY_DATA
 	if (Ar.IsCooking())
 	{
@@ -300,14 +287,6 @@ void UStaticMeshComponent::Serialize(FArchive& Ar)
 
 	Ar << LODData;
 
-#if WITH_EDITOR
-	if (bCheckBuildGuids)
-	{
-		// If we're cooking, display a deterministic cook warning if we didn't overwrite the generated GUIDs at load time
-		UE_CLOG(bCheckBuildGuids && !Algo::AllOf(LODData, &FStaticMeshComponentLODInfo::bMapBuildDataIdLoaded), LogStaticMesh, Warning, TEXT("%s contains a legacy UStaticMeshComponent and is being non-deterministically cooked - please resave the asset and recook."), *GetOutermost()->GetName());
-	}
-#endif
-
 	if (Ar.IsLoading())
 	{
 		for (int32 LODIndex = 0; LODIndex < LODData.Num(); LODIndex++)
@@ -2303,9 +2282,6 @@ void UStaticMeshComponent::ApplyComponentInstanceData(FStaticMeshComponentInstan
 			for (int32 i = 0; i < NumLODLightMaps; ++i)
 			{
 				LODData[i].MapBuildDataId = StaticMeshInstanceData->CachedStaticLighting[i];
-			#if WITH_EDITOR
-				LODData[i].bMapBuildDataIdLoaded = true;
-			#endif
 			}
 		}
 		else
@@ -2530,34 +2506,45 @@ FStaticMeshComponentLODInfo::FStaticMeshComponentLODInfo()
 	, OverrideVertexColors(NULL)
 	, OwningComponent(NULL)
 {
-	// Used by deserialization only, MapBuildDataId will be deserialized
+	// MapBuildDataId will be deserialized
 }
 
-FStaticMeshComponentLODInfo::FStaticMeshComponentLODInfo(UStaticMeshComponent* InOwningComponent, int32 LodIndex)
+FStaticMeshComponentLODInfo::FStaticMeshComponentLODInfo(UStaticMeshComponent* InOwningComponent)
 	: LegacyMapBuildData(NULL)
 	, OverrideVertexColors(NULL)
 	, OwningComponent(InOwningComponent)
+{
+	// MapBuildDataId is invalid for newly created FStaticMeshComponentLODInfo
+	// Will be assigned a valid GUID if we ever need to store data in the MapBuildData for it.
+	// See CreateMapBuildDataId()
+}
+
+bool FStaticMeshComponentLODInfo::CreateMapBuildDataId(int32 LodIndex)
+{
+	if (!MapBuildDataId.IsValid())
 	{
-	if (LodIndex == 0)
-	{
-		MapBuildDataId = FGuid::NewGuid();
-	}
-	else
-	{
-		FString GuidBaseString = OwningComponent->LODData[0].MapBuildDataId.ToString(EGuidFormats::Digits);
-		GuidBaseString += TEXT("LOD_") + FString::FromInt(LodIndex);
-		
-		FSHA1 Sha;
-		Sha.Update((uint8*)*GuidBaseString, GuidBaseString.Len() * sizeof(TCHAR));
-		Sha.Final();
-		// Retrieve the hash and use it to construct a pseudo-GUID.
-		uint32 Hash[5];
-		Sha.GetHash((uint8*)Hash);
-		MapBuildDataId = FGuid(Hash[0] ^ Hash[4], Hash[1], Hash[2], Hash[3]);
+		if (LodIndex == 0 || OwningComponent == nullptr)
+		{
+			MapBuildDataId = FGuid::NewGuid();
+		}
+		else
+		{
+			FString GuidBaseString = OwningComponent->LODData[0].MapBuildDataId.ToString(EGuidFormats::Digits);
+			GuidBaseString += TEXT("LOD_") + FString::FromInt(LodIndex);
+
+			FSHA1 Sha;
+			Sha.Update((uint8*)*GuidBaseString, GuidBaseString.Len() * sizeof(TCHAR));
+			Sha.Final();
+			// Retrieve the hash and use it to construct a pseudo-GUID.
+			uint32 Hash[5];
+			Sha.GetHash((uint8*)Hash);
+			MapBuildDataId = FGuid(Hash[0] ^ Hash[4], Hash[1], Hash[2], Hash[3]);
+		}
+
+		return true;
 	}
-#if WITH_EDITOR
-	bMapBuildDataIdLoaded = false;
-#endif
+
+	return false;
 }
 
 /** Destructor */
@@ -2726,9 +2713,6 @@ FArchive& operator<<(FArchive& Ar,FStaticMeshComponentLODInfo& I)
 		if (Ar.IsLoading() && Ar.CustomVer(FRenderingObjectVersion::GUID) < FRenderingObjectVersion::MapBuildDataSeparatePackage)
 		{
 			I.MapBuildDataId = FGuid::NewGuid();
-		#if WITH_EDITOR
-			I.bMapBuildDataIdLoaded = false;
-		#endif
 			I.LegacyMapBuildData = new FMeshMapBuildData();
 			Ar << I.LegacyMapBuildData->LightMap;
 			Ar << I.LegacyMapBuildData->ShadowMap;
@@ -2736,12 +2720,6 @@ FArchive& operator<<(FArchive& Ar,FStaticMeshComponentLODInfo& I)
 		else
 		{
 			Ar << I.MapBuildDataId;
-		#if WITH_EDITOR
-			if (Ar.IsLoading())
-			{
-				I.bMapBuildDataIdLoaded = true;
-			}
-		#endif
 		}
 	}
 
diff --git a/Engine/Source/Runtime/Engine/Private/InstancedStaticMesh.cpp b/Engine/Source/Runtime/Engine/Private/InstancedStaticMesh.cpp
--- a/Engine/Source/Runtime/Engine/Private/InstancedStaticMesh.cpp
+++ b/Engine/Source/Runtime/Engine/Private/InstancedStaticMesh.cpp
@@ -1843,10 +1843,17 @@ void UInstancedStaticMeshComponent::ApplyLightMapping(FStaticLightingTextureMapp
 
 		// Ensure LODData has enough entries in it, free not required.
 		SetLODDataCount(ResolvedMesh->GetNumLODs(), ResolvedMesh->GetNumLODs());
+		FStaticMeshComponentLODInfo& LODInfo = LODData[0];
+
+		// Ensure this LODInfo has a valid MapBuildDataId
+		if (LODInfo.CreateMapBuildDataId(0))
+		{
+			MarkPackageDirty();
+		}
 
 		ULevel* StorageLevel = LightingScenario ? LightingScenario : GetOwner()->GetLevel();
 		UMapBuildDataRegistry* Registry = StorageLevel->GetOrCreateMapBuildData();
-		FMeshMapBuildData& MeshBuildData = Registry->AllocateMeshBuildData(LODData[0].MapBuildDataId, true);
+		FMeshMapBuildData& MeshBuildData = Registry->AllocateMeshBuildData(LODInfo.MapBuildDataId, true);
 
 		MeshBuildData.PerInstanceLightmapData.Empty(AllQuantizedData.Num());
 		MeshBuildData.PerInstanceLightmapData.AddZeroed(AllQuantizedData.Num());
@@ -1858,11 +1865,11 @@ void UInstancedStaticMeshComponent::ApplyLightMapping(FStaticLightingTextureMapp
 		TRefCountPtr<FLightMap2D> NewLightMap = FLightMap2D::AllocateInstancedLightMap(Registry, this,
 			MoveTemp(AllQuantizedData),
 			bUseVirtualTextures ? MoveTemp(AllShadowMapData) : MoveTemp(EmptyShadowMapData),
-			Registry, LODData[0].MapBuildDataId, Bounds, PaddingType, LMF_Streamed);
+			Registry, LODInfo.MapBuildDataId, Bounds, PaddingType, LMF_Streamed);
 
 		// Create a shadow-map for the primitive, only needed when not using VT
 		TRefCountPtr<FShadowMap2D> NewShadowMap = (bNeedsShadowMap && !bUseVirtualTextures)
-			? FShadowMap2D::AllocateInstancedShadowMap(Registry, this, MoveTemp(AllShadowMapData), Registry, LODData[0].MapBuildDataId, Bounds, PaddingType, SMF_Streamed)
+			? FShadowMap2D::AllocateInstancedShadowMap(Registry, this, MoveTemp(AllShadowMapData), Registry, LODInfo.MapBuildDataId, Bounds, PaddingType, SMF_Streamed)
 			: nullptr;
 
 		MeshBuildData.LightMap = NewLightMap;
diff --git a/Engine/Source/Runtime/Engine/Private/StaticMeshLight.cpp b/Engine/Source/Runtime/Engine/Private/StaticMeshLight.cpp
--- a/Engine/Source/Runtime/Engine/Private/StaticMeshLight.cpp
+++ b/Engine/Source/Runtime/Engine/Private/StaticMeshLight.cpp
@@ -215,9 +215,7 @@ void FStaticMeshStaticLightingTextureMapping::Apply(FQuantizedLightmapData* Quan
 
 	UStaticMeshComponent* StaticMeshComponent = Primitive.Get();
 
-	if (StaticMeshComponent && StaticMeshComponent->GetOwner() && StaticMeshComponent->GetOwner()->GetLevel() && 
-		(LODIndex >= StaticMeshComponent->LODData.Num() || StaticMeshComponent->LODData[LODIndex].MapBuildDataId.IsValid())
-		)
+	if (StaticMeshComponent && StaticMeshComponent->GetOwner() && StaticMeshComponent->GetOwner()->GetLevel())
 	{
 		// Should have happened at a higher level
 		check(!StaticMeshComponent->IsRenderStateCreated());
@@ -233,7 +231,14 @@ void FStaticMeshStaticLightingTextureMapping::Apply(FQuantizedLightmapData* Quan
 			StaticMeshComponent->MarkPackageDirty();
 		}
 
-		const FStaticMeshComponentLODInfo& ComponentLODInfo = StaticMeshComponent->LODData[LODIndex];
+		FStaticMeshComponentLODInfo& ComponentLODInfo = StaticMeshComponent->LODData[LODIndex];
+
+		// Ensure this LODInfo has a valid MapBuildDataId
+		if (ComponentLODInfo.CreateMapBuildDataId(LODIndex))
+		{
+			StaticMeshComponent->MarkPackageDirty();
+		}
+
 		ELightMapPaddingType PaddingType = GAllowLightmapPadding ? LMPT_NormalPadding : LMPT_NoPadding;
 		const bool bHasNonZeroData = (QuantizedData != NULL && QuantizedData->HasNonZeroData());
 
@@ -456,10 +461,7 @@ void UStaticMeshComponent::InvalidateLightingCacheDetailed(bool bInvalidateBuild
 	for(int32 i = 0; i < LODData.Num(); i++)
 	{
 		FStaticMeshComponentLODInfo& LODDataElement = LODData[i];
-		LODDataElement.MapBuildDataId = FGuid::NewGuid();
-		#if WITH_EDITOR
-			LODDataElement.bMapBuildDataIdLoaded = false;
-		#endif
+		LODDataElement.MapBuildDataId.Invalidate();
 	}
 
 	MarkRenderStateDirty();
diff --git a/Engine/Source/Runtime/Engine/Private/StaticMeshRender.cpp b/Engine/Source/Runtime/Engine/Private/StaticMeshRender.cpp
--- a/Engine/Source/Runtime/Engine/Private/StaticMeshRender.cpp
+++ b/Engine/Source/Runtime/Engine/Private/StaticMeshRender.cpp
@@ -370,9 +370,8 @@ bool UStaticMeshComponent::SetLODDataCount( const uint32 MinSize, const uint32 M
 		uint32 ItemCountToAdd = MinSize - LODData.Num();
 		for(uint32 i = 0; i < ItemCountToAdd; ++i)
 		{
-			int32 LodIndex = LODData.Num();
 			// call constructor
-			new (LODData)FStaticMeshComponentLODInfo(this, LodIndex);
+			new (LODData)FStaticMeshComponentLODInfo(this);
 		}
 		return true;
 	}
