Fix amortized guid cache serialization if the per frame limit isn't set

diff --git a/Engine/Source/Runtime/Engine/Private/DemoNetDriver.cpp b/Engine/Source/Runtime/Engine/Private/DemoNetDriver.cpp
--- a/Engine/Source/Runtime/Engine/Private/DemoNetDriver.cpp
+++ b/Engine/Source/Runtime/Engine/Private/DemoNetDriver.cpp
@@ -1950,14 +1950,14 @@ bool UDemoNetDriver::SerializeGuidCache(const FRepActorsCheckpointParams& Params
 
 			++NumNetGuidsForRecording;
 
-			if (FPlatformTime::Seconds() >= Deadline)
+			if (Params.CheckpointMaxUploadTimePerFrame > 0 && (FPlatformTime::Seconds() >= Deadline))
 			{
 				break;
 			}
 		}
 	}
 
-	bool bCompleted = NextNetGuidForRecording == NetGuidCacheSnapshot.Num();
+	const bool bCompleted = NextNetGuidForRecording == NetGuidCacheSnapshot.Num();
 	if (bCompleted)
 	{
 		FArchivePos Pos = CheckpointArchive->Tell();
