TBA: Convert asset import data serialization to structured archive format

diff --git a/Engine/Source/Runtime/Engine/Classes/EditorFramework/AssetImportData.h b/Engine/Source/Runtime/Engine/Classes/EditorFramework/AssetImportData.h
--- a/Engine/Source/Runtime/Engine/Classes/EditorFramework/AssetImportData.h
+++ b/Engine/Source/Runtime/Engine/Classes/EditorFramework/AssetImportData.h
@@ -152,7 +152,7 @@ protected:
 	FString ResolveImportFilename(const FString& InRelativePath) const;
 
 	/** Overridden serialize function to write out the underlying data as json */
-	virtual void Serialize(FArchive& Ar) override;
+	virtual void Serialize(FStructuredArchive::FRecord Record) override;
 	
 #endif		// WITH_EDITORONLY_DATA
 };
diff --git a/Engine/Source/Runtime/Engine/Private/EditorFramework/AssetImportData.cpp b/Engine/Source/Runtime/Engine/Private/EditorFramework/AssetImportData.cpp
--- a/Engine/Source/Runtime/Engine/Private/EditorFramework/AssetImportData.cpp
+++ b/Engine/Source/Runtime/Engine/Private/EditorFramework/AssetImportData.cpp
@@ -361,36 +361,37 @@ FString UAssetImportData::ResolveImportFilename(const FString& InRelativePath) c
 	return ResolveImportFilename(InRelativePath, GetOutermost());
 }
 
-void UAssetImportData::Serialize(FArchive& Ar)
+void UAssetImportData::Serialize(FStructuredArchive::FRecord Record)
 {
-	if (Ar.UE4Ver() >= VER_UE4_ASSET_IMPORT_DATA_AS_JSON)
+	FArchive& BaseArchive = Record.GetUnderlyingArchive();
+
+	if (BaseArchive.UE4Ver() >= VER_UE4_ASSET_IMPORT_DATA_AS_JSON)
 	{
-		
-		if (!Ar.IsFilterEditorOnly())
+		if (!BaseArchive.IsFilterEditorOnly())
 		{
 			FString Json;
-			if (Ar.IsLoading())
+			if (BaseArchive.IsLoading())
 			{
-				Ar << Json;
+				Record << SA_VALUE(TEXT("Json"), Json);
 				TOptional<FAssetImportInfo> Copy = FAssetImportInfo::FromJson(MoveTemp(Json));
 				if (Copy.IsSet())
 				{
 					SourceData = MoveTemp(Copy.GetValue());
 				}
 			}
-			else if (Ar.IsSaving())
+			else if (BaseArchive.IsSaving())
 			{
 				Json = SourceData.ToJson();
-				Ar << Json;
+				Record << SA_VALUE(TEXT("Json"), Json);
 			}
 		}
 	}
 
-	Super::Serialize(Ar);
+	Super::Serialize(Record);
 
-	Ar.UsingCustomVersion(FAnimPhysObjectVersion::GUID);
+	BaseArchive.UsingCustomVersion(FAnimPhysObjectVersion::GUID);
 
-	if (Ar.CustomVer(FAnimPhysObjectVersion::GUID) < FAnimPhysObjectVersion::ThumbnailSceneInfoAndAssetImportDataAreTransactional)
+	if (BaseArchive.CustomVer(FAnimPhysObjectVersion::GUID) < FAnimPhysObjectVersion::ThumbnailSceneInfoAndAssetImportDataAreTransactional)
 	{
 		SetFlags(RF_Transactional);
 	}
