TBA: Convert datatable over to structured archive form

diff --git a/Engine/Source/Runtime/Engine/Classes/Engine/DataTable.h b/Engine/Source/Runtime/Engine/Classes/Engine/DataTable.h
--- a/Engine/Source/Runtime/Engine/Classes/Engine/DataTable.h
+++ b/Engine/Source/Runtime/Engine/Classes/Engine/DataTable.h
@@ -111,7 +111,7 @@ public:
 
 	//~ Begin UObject Interface.
 	ENGINE_API virtual void FinishDestroy() override;
-	ENGINE_API virtual void Serialize(FArchive& Ar) override;
+	ENGINE_API virtual void Serialize(FStructuredArchiveRecord Record) override;
 	ENGINE_API static void AddReferencedObjects(UObject* InThis, FReferenceCollector& Collector);
 	ENGINE_API virtual void GetPreloadDependencies(TArray<UObject*>& OutDeps) override;
 	ENGINE_API virtual void GetResourceSizeEx(FResourceSizeEx& CumulativeResourceSize) override;
@@ -365,8 +365,8 @@ public:
 	//~ End UDataTable Interface
 
 protected:
-	void SaveStructData(FArchive& Ar);
-	void LoadStructData(FArchive& Ar);
+	void SaveStructData(FStructuredArchiveSlot Slot);
+	void LoadStructData(FStructuredArchiveSlot Slot);
 
 	/**
 	 * Called whenever new data is imported into the data table via CreateTableFrom*; Alerts each imported row and gives the
diff --git a/Engine/Source/Runtime/Engine/Private/DataTable.cpp b/Engine/Source/Runtime/Engine/Private/DataTable.cpp
--- a/Engine/Source/Runtime/Engine/Private/DataTable.cpp
+++ b/Engine/Source/Runtime/Engine/Private/DataTable.cpp
@@ -6,6 +6,7 @@
 #include "Serialization/PropertyLocalizationDataGathering.h"
 #include "Serialization/ObjectWriter.h"
 #include "Serialization/ObjectReader.h"
+#include "Serialization/StructuredArchive.h"
 #include "UObject/LinkerLoad.h"
 #include "DataTableCSV.h"
 #include "Policies/PrettyJsonPrintPolicy.h"
@@ -88,7 +89,7 @@ void UDataTable::PostEditChangeProperty(FPropertyChangedEvent& PropertyChangedEv
 }
 #endif
 
-void UDataTable::LoadStructData(FArchive& Ar)
+void UDataTable::LoadStructData(FStructuredArchiveSlot Slot)
 {
 	UScriptStruct* LoadUsingStruct = RowStruct;
 	if (!LoadUsingStruct)
@@ -101,16 +102,18 @@ void UDataTable::LoadStructData(FArchive& Ar)
 	}
 
 	int32 NumRows;
-	Ar << NumRows;
+	FStructuredArchiveArray Array = Slot.EnterArray(NumRows);
 
 	DATATABLE_CHANGE_SCOPE();
 
 	RowMap.Reserve(NumRows);
 	for (int32 RowIdx = 0; RowIdx < NumRows; RowIdx++)
 	{
+		FStructuredArchiveRecord RowRecord = Array.EnterElement().EnterRecord();
+
 		// Load row name
 		FName RowName;
-		Ar << RowName;
+		RowRecord << SA_VALUE(TEXT("Name"), RowName);
 
 		// Load row data
 		uint8* RowData = (uint8*)FMemory::Malloc(LoadUsingStruct->GetStructureSize());
@@ -118,14 +121,14 @@ void UDataTable::LoadStructData(FArchive& Ar)
 		// And be sure to call DestroyScriptStruct later
 		LoadUsingStruct->InitializeStruct(RowData);
 
-		LoadUsingStruct->SerializeItem(Ar, RowData, nullptr);
+		LoadUsingStruct->SerializeItem(RowRecord.EnterField(SA_FIELD_NAME(TEXT("Value"))), RowData, nullptr);
 
 		// Add to map
 		RowMap.Add(RowName, RowData);
 	}
 }
 
-void UDataTable::SaveStructData(FArchive& Ar)
+void UDataTable::SaveStructData(FStructuredArchiveSlot Slot)
 {
 	UScriptStruct* SaveUsingStruct = RowStruct;
 	if (!SaveUsingStruct)
@@ -138,19 +141,20 @@ void UDataTable::SaveStructData(FArchive& Ar)
 	}
 
 	int32 NumRows = RowMap.Num();
-	Ar << NumRows;
+	FStructuredArchiveArray Array = Slot.EnterArray(NumRows);
 
 	// Now iterate over rows in the map
 	for (auto RowIt = RowMap.CreateIterator(); RowIt; ++RowIt)
 	{
 		// Save out name
 		FName RowName = RowIt.Key();
-		Ar << RowName;
+		FStructuredArchiveRecord Row = Array.EnterElement().EnterRecord();
+		Row << SA_VALUE(TEXT("Name"), RowName);
 
 		// Save out data
 		uint8* RowData = RowIt.Value();
 
-		SaveUsingStruct->SerializeItem(Ar, RowData, nullptr);
+		SaveUsingStruct->SerializeItem(Row.EnterField(SA_FIELD_NAME(TEXT("Value"))), RowData, nullptr);
 	}
 }
 
@@ -194,17 +198,19 @@ void UDataTable::OnPostDataImported(TArray<FString>& OutCollectedImportProblems)
 	OnDataTableChanged().Broadcast();
 }
 
-void UDataTable::Serialize( FArchive& Ar )
+void UDataTable::Serialize(FStructuredArchiveRecord Record)
 {
+	FArchive& BaseArchive = Record.GetUnderlyingArchive();
+
 #if WITH_EDITORONLY_DATA
 	// Make sure and update RowStructName before calling the parent Serialize (which will save the properties)
-	if (Ar.IsSaving() && RowStruct)
+	if (BaseArchive.IsSaving() && RowStruct)
 	{
 		RowStructName = RowStruct->GetFName();
 	}
 #endif	// WITH_EDITORONLY_DATA
 
-		Super::Serialize(Ar); // When loading, this should load our RowStruct!	
+	Super::Serialize(Record); // When loading, this should load our RowStruct!	
 
 	if (RowStruct && RowStruct->HasAnyFlags(RF_NeedLoad))
 	{
@@ -215,15 +221,15 @@ void UDataTable::Serialize( FArchive& Ar )
 		}
 	}
 
-	if(Ar.IsLoading())
+	if(BaseArchive.IsLoading())
 	{
 		DATATABLE_CHANGE_SCOPE();
 		EmptyTable();
-		LoadStructData(Ar);
+		LoadStructData(Record.EnterField(SA_FIELD_NAME(TEXT("Data"))));
 	}
-	else if(Ar.IsSaving())
+	else if(BaseArchive.IsSaving())
 	{
-		SaveStructData(Ar);
+		SaveStructData(Record.EnterField(SA_FIELD_NAME(TEXT("Data"))));
 	}
 }
 
@@ -445,7 +451,7 @@ void UDataTable::CleanBeforeStructChange()
 			};
 
 			FRawStructWriter MemoryWriter(RowsSerializedWithTags, TemporarilyReferencedObjects);
-			SaveStructData(MemoryWriter);
+			SaveStructData(FStructuredArchiveFromArchive(MemoryWriter).GetSlot());
 		}
 
 		EmptyTable();
@@ -474,7 +480,7 @@ void UDataTable::RestoreAfterStructChange()
 		};
 
 		FRawStructReader MemoryReader(RowsSerializedWithTags);
-		LoadStructData(MemoryReader);
+		LoadStructData(FStructuredArchiveFromArchive(MemoryReader).GetSlot());
 	}
 	TemporarilyReferencedObjects.Empty();
 	RowsSerializedWithTags.Empty();
