- Don't compile editor analytics support code into non-editor builds (CIS fix)

diff --git a/Engine/Source/Runtime/Engine/Private/Analytics/EditorSessionSummaryWriter.cpp b/Engine/Source/Runtime/Engine/Private/Analytics/EditorSessionSummaryWriter.cpp
--- a/Engine/Source/Runtime/Engine/Private/Analytics/EditorSessionSummaryWriter.cpp
+++ b/Engine/Source/Runtime/Engine/Private/Analytics/EditorSessionSummaryWriter.cpp
@@ -2,7 +2,8 @@
 
 #include "EditorSessionSummaryWriter.h"
 
-#include "EditorAnalyticsSession.h"
+#if WITH_EDITOR
+
 #include "EngineAnalytics.h"
 #include "EngineGlobals.h"
 #include "GeneralProjectSettings.h"
@@ -22,10 +23,9 @@
 #include "Misc/EngineVersion.h"
 #include "Misc/Guid.h"
 
-#if WITH_EDITOR
+#include "EditorAnalyticsSession.h"
 #include "IVREditorModule.h"
 #include "Kismet2/DebuggerCommands.h"
-#endif // WITH_EDITOR
 
 #define LOCTEXT_NAMESPACE "SessionSummary"
 
@@ -345,3 +345,4 @@ FString FEditorSessionSummaryWriter::GetUserActivityString() const
 }
 
 #undef LOCTEXT_NAMESPACE
+#endif // WITH_EDITOR
diff --git a/Engine/Source/Runtime/Engine/Private/Analytics/EditorSessionSummaryWriter.h b/Engine/Source/Runtime/Engine/Private/Analytics/EditorSessionSummaryWriter.h
--- a/Engine/Source/Runtime/Engine/Private/Analytics/EditorSessionSummaryWriter.h
+++ b/Engine/Source/Runtime/Engine/Private/Analytics/EditorSessionSummaryWriter.h
@@ -4,6 +4,8 @@
 
 #include "CoreMinimal.h"
 
+#if WITH_EDITOR
+
 struct FUserActivity;
 struct FEditorAnalyticsSession;
 
@@ -37,3 +39,5 @@ private:
 	float HeartbeatTimeElapsed;
 	bool bShutdown;
 };
+
+#endif
diff --git a/Engine/Source/Runtime/Engine/Private/EngineAnalytics.cpp b/Engine/Source/Runtime/Engine/Private/EngineAnalytics.cpp
--- a/Engine/Source/Runtime/Engine/Private/EngineAnalytics.cpp
+++ b/Engine/Source/Runtime/Engine/Private/EngineAnalytics.cpp
@@ -15,17 +15,23 @@
 #include "EngineSessionManager.h"
 #include "Misc/EngineVersion.h"
 #include "RHI.h"
+
+#if WITH_EDITOR
 #include "EditorAnalyticsSession.h"
 #include "EditorSessionSummarySender.h"
 #include "Analytics/EditorSessionSummaryWriter.h"
+#endif
+
 #include "GenericPlatform/GenericPlatformCrashContext.h"
 
 bool FEngineAnalytics::bIsInitialized;
 TSharedPtr<IAnalyticsProviderET> FEngineAnalytics::Analytics;
 TSharedPtr<FEngineSessionManager> FEngineAnalytics::SessionManager;
 
+#if WITH_EDITOR
 static TSharedPtr<FEditorSessionSummaryWriter> SessionSummaryWriter;
 static TSharedPtr<FEditorSessionSummarySender> SessionSummarySender;
+#endif
 
 /**
 * Default config func.
@@ -158,6 +164,7 @@ void FEngineAnalytics::Initialize()
 			SessionManager->Initialize();
 		}
 
+#if WITH_EDITOR
 		if (!SessionSummaryWriter.IsValid())
 		{
 			SessionSummaryWriter = MakeShareable(new FEditorSessionSummaryWriter());
@@ -172,6 +179,7 @@ void FEngineAnalytics::Initialize()
 				SessionSummarySender = MakeShareable(new FEditorSessionSummarySender(FEngineAnalytics::GetProvider(), TEXT("Editor")));
 			}
 		}
+#endif
 	}
 }
 
@@ -188,6 +196,7 @@ void FEngineAnalytics::Shutdown(bool bIsEngineShutdown)
 		SessionManager.Reset();
 	}
 
+#if WITH_EDITOR
 	if (SessionSummaryWriter.IsValid())
 	{
 		SessionSummaryWriter->Shutdown();
@@ -199,6 +208,7 @@ void FEngineAnalytics::Shutdown(bool bIsEngineShutdown)
 		SessionSummarySender->Shutdown();
 		SessionSummarySender.Reset();
 	}
+#endif
 }
 
 void FEngineAnalytics::Tick(float DeltaTime)
@@ -210,6 +220,7 @@ void FEngineAnalytics::Tick(float DeltaTime)
 		SessionManager->Tick(DeltaTime);
 	}
 
+#if WITH_EDITOR
 	if (SessionSummaryWriter.IsValid())
 	{
 		SessionSummaryWriter->Tick(DeltaTime);
@@ -219,4 +230,5 @@ void FEngineAnalytics::Tick(float DeltaTime)
 	{
 		SessionSummarySender->Tick(DeltaTime);
 	}
+#endif
 }
