Fixed missing HAL/IConsoleManager.h include

diff --git a/Engine/Source/Runtime/Engine/Private/AudioPluginUtilities.cpp b/Engine/Source/Runtime/Engine/Private/AudioPluginUtilities.cpp
--- a/Engine/Source/Runtime/Engine/Private/AudioPluginUtilities.cpp
+++ b/Engine/Source/Runtime/Engine/Private/AudioPluginUtilities.cpp
@@ -5,6 +5,7 @@
 #include "CoreGlobals.h"
 #include "CoreMinimal.h"
 #include "Features/IModularFeatures.h"
+#include "HAL/IConsoleManager.h"
 #include "Misc/ConfigCacheIni.h"
 
 
