AUTO MERGED/UNRESOLVED

Merging //UE4/Main @ 9360122 to //UE4/Dev-Core (stefan.boberg_rdu-p4merge-02_6162)

(ushell-p4-mergedown)

diff --git a/Engine/Config/BaseEditor.ini b/Engine/Config/BaseEditor.ini
--- a/Engine/Config/BaseEditor.ini
+++ b/Engine/Config/BaseEditor.ini
@@ -533,3 +533,6 @@ bValidateOnSave=true
 [/Script/UnrealEd.DiffAssetRegistriesCommandlet]
 P4EngineBasePath="/Engine/Content/"
 P4EngineAssetPath="/Engine/"
+
+[CoreRedirects]
++ClassRedirects=(OldName="IViewportLayoutEntity", NewName="IAssetViewportLayoutEntity")
diff --git a/Engine/Config/BaseEditorPerProjectUserSettings.ini b/Engine/Config/BaseEditorPerProjectUserSettings.ini
--- a/Engine/Config/BaseEditorPerProjectUserSettings.ini
+++ b/Engine/Config/BaseEditorPerProjectUserSettings.ini
@@ -852,6 +852,15 @@ bInfiniteKeyAreas=true
 bAutoSetTrackDefaults=true
 FrameNumberDisplayFormat=Frames
 
+[TemplateSequenceEditor SequencerSettings]
+bKeyInterpPropertiesOnly=true
+bShowRangeSlider=true
+bKeepPlayRangeInSectionBounds=false
+ZeroPadFrames=4
+bInfiniteKeyAreas=true
+bAutoSetTrackDefaults=true
+FrameNumberDisplayFormat=Frames
+
 [TakeRecorderSequenceEditor SequencerSettings]
 bKeyInterpPropertiesOnly=true
 bShowRangeSlider=true
diff --git a/Engine/Config/BaseEngine.ini b/Engine/Config/BaseEngine.ini
--- a/Engine/Config/BaseEngine.ini
+++ b/Engine/Config/BaseEngine.ini
@@ -692,6 +692,12 @@ bUseStreamingPause=false
 
 ; 4.24
 +PropertyRedirects=(OldName="NavDataConfig.NavigationDataClassName", NewName="NavDataConfig.NavDataClass")
++ClassRedirects=(OldName="/Script/EditorInteractiveToolsFramework.StaticMeshComponentTargetFactory",NewName="/Script/UnrealEd.StaticMeshComponentTargetFactory")
++ClassRedirects=(OldName="/Script/EditorInteractiveToolsFramework.StandardToolModeCommands",NewName="/Script/UnrealEd.StandardToolModeCommands")
++EnumRedirects==(OldName="/Script/EditorInteractiveToolsFramework.StandardToolModeCommands",NewName="/Script/UnrealEd.StandardToolModeCommands")
++ClassRedirects=(OldName="/Script/EditorInteractiveToolsFramework.InteractiveToolCommands",NewName="/Script/UnrealEd.InteractiveToolCommands")
++ClassRedirects=(OldName="/Script/EditorInteractiveToolsFramework.EditorToolAssetAPI",NewName="/Script/UnrealEd.EditorToolAssetAPI")
++ClassRedirects=(OldName="/Script/EditorInteractiveToolsFramework.StaticMeshComponentTarget",NewName="/Script/UnrealEd.StaticMeshComponentTarget")
 
 ; Rename ClothingSystemRuntime to ClothingSystemRuntimeNv
 +PackageRedirects=(OldName="/Script/ClothingSystemRuntime",NewName="/Script/ClothingSystemRuntimeNv")
@@ -1723,6 +1729,11 @@ bFilterAssetDataWithNoTags=false
 bFilterDependenciesWithNoTags=false
 bFilterSearchableNames=true
 
+; Control which collections are added as asset registry tags during cook (none by default)
+; If bTagAllCollections is true then you can exclude specfic collections via the CollectionsToExcludeAsTags arrray
+; If bTagAllCollections is false then you can include specfic collections via the CollectionsToIncludeAsTags arrray
+bTagAllCollections=false
+
 [AutomationTesting]
 ImportTestPath=../../Content/EditorAutomation/
 ImportTestPackagePath=/Engine/Content/EditorAutomation
diff --git a/Engine/Plugins/Experimental/ChaosEditor/Source/FractureEditor/Private/FractureEditorMode.cpp b/Engine/Plugins/Experimental/ChaosEditor/Source/FractureEditor/Private/FractureEditorMode.cpp
--- a/Engine/Plugins/Experimental/ChaosEditor/Source/FractureEditor/Private/FractureEditorMode.cpp
+++ b/Engine/Plugins/Experimental/ChaosEditor/Source/FractureEditor/Private/FractureEditorMode.cpp
@@ -118,10 +118,13 @@ void FFractureEditorMode::Render(const FSceneView* View, FViewport* Viewport, FP
 
 	FFractureEditorModeToolkit* FractureToolkit = (FFractureEditorModeToolkit*)Toolkit.Get();
  
-	UFractureTool* FractureTool = FractureToolkit->GetActiveTool();
+	if (UFractureTool* FractureTool = FractureToolkit->GetActiveTool())
+	{
+		auto Settings = FractureTool->GetSettingsObjects();
+		FractureTool->Render(View, Viewport, PDI);
+	}
+
 
-	auto Settings = FractureTool->GetSettingsObjects();
-	FractureTool->Render(View, Viewport, PDI);
 }
 
 bool FFractureEditorMode::UsesToolkits() const
diff --git a/Engine/Plugins/FX/Niagara/Source/Niagara/Private/NiagaraSpriteRendererProperties.cpp b/Engine/Plugins/FX/Niagara/Source/Niagara/Private/NiagaraSpriteRendererProperties.cpp
--- a/Engine/Plugins/FX/Niagara/Source/Niagara/Private/NiagaraSpriteRendererProperties.cpp
+++ b/Engine/Plugins/FX/Niagara/Source/Niagara/Private/NiagaraSpriteRendererProperties.cpp
@@ -49,7 +49,7 @@ UNiagaraSpriteRendererProperties::UNiagaraSpriteRendererProperties()
 FNiagaraRenderer* UNiagaraSpriteRendererProperties::CreateEmitterRenderer(ERHIFeatureLevel::Type FeatureLevel, const FNiagaraEmitterInstance* Emitter)
 {
 	FNiagaraRenderer* NewRenderer = new FNiagaraRendererSprites(FeatureLevel, this, Emitter);	
-	NewRenderer->Initialize(FeatureLevel, this, Emitter);
+	NewRenderer->Initialize(this, Emitter);
 	return NewRenderer;
 }
 
diff --git a/Engine/Source/Developer/MeshBoneReduction/Private/MeshBoneReduction.cpp b/Engine/Source/Developer/MeshBoneReduction/Private/MeshBoneReduction.cpp
--- a/Engine/Source/Developer/MeshBoneReduction/Private/MeshBoneReduction.cpp
+++ b/Engine/Source/Developer/MeshBoneReduction/Private/MeshBoneReduction.cpp
@@ -358,7 +358,7 @@ public:
 		}
 	}
 
-	bool ReduceBoneCounts(USkeletalMesh* SkeletalMesh, int32 DesiredLOD, const TArray<FName>* BoneNamesToRemove) override
+	bool ReduceBoneCounts(USkeletalMesh* SkeletalMesh, int32 DesiredLOD, const TArray<FName>* BoneNamesToRemove, bool bCallPostEditChange /*= true*/) override
 	{
 		check (SkeletalMesh);
 		USkeleton* Skeleton = SkeletalMesh->Skeleton;
@@ -371,7 +371,6 @@ public:
 		// Always restore all previously removed bones if not contained by BonesToRemove
 		SkeletalMesh->CalculateRequiredBones(SkeletalMesh->GetImportedModel()->LODModels[DesiredLOD], SkeletalMesh->RefSkeleton, &BonesToRemove);
 		
-		TComponentReregisterContext<USkinnedMeshComponent> ReregisterContext;
 		SkeletalMesh->ReleaseResources();
 		SkeletalMesh->ReleaseResourcesFence.Wait();
 
@@ -386,11 +385,6 @@ public:
 		{
 			NewModel = new FSkeletalMeshLODModel();
 			LODModels[DesiredLOD] = NewModel;
-			if (SkeletalMeshResource->OriginalReductionSourceMeshData.IsValidIndex(DesiredLOD))
-			{
-				SkeletalMeshResource->OriginalReductionSourceMeshData[DesiredLOD]->EmptyBulkData();
-				SkeletalMeshResource->OriginalReductionSourceMeshData.RemoveAt(DesiredLOD);
-			}
 
 			// Bulk data arrays need to be locked before a copy can be made.
 			SrcModel->RawPointIndices.Lock(LOCK_READ_ONLY);
@@ -483,10 +477,18 @@ public:
 		NewModel->ActiveBoneIndices.Sort();
 		NewModel->RequiredBones.Sort();
 
-		SkeletalMesh->PostEditChange();
-		SkeletalMesh->InitResources();
+		if (bCallPostEditChange)
+		{
+			SkeletalMesh->PostEditChange();
+			SkeletalMesh->InitResources();
+		}
 		SkeletalMesh->MarkPackageDirty();
 		
+		//Reregister skinned mesh component if we call post edit change
+		if (bCallPostEditChange)
+		{
+			TComponentReregisterContext<USkinnedMeshComponent> ReregisterContext;
+		}
 		return true;
 	}
 };
diff --git a/Engine/Source/Developer/MeshSimplifier/Private/QuadricMeshReduction.cpp b/Engine/Source/Developer/MeshSimplifier/Private/QuadricMeshReduction.cpp
--- a/Engine/Source/Developer/MeshSimplifier/Private/QuadricMeshReduction.cpp
+++ b/Engine/Source/Developer/MeshSimplifier/Private/QuadricMeshReduction.cpp
@@ -9,8 +9,7 @@
 #include "Templates/UniquePtr.h"
 #include "Features/IModularFeatures.h"
 #include "IMeshReductionInterfaces.h"
-#include "MeshDescription.h"
-#include "MeshAttributes.h"
+#include "StaticMeshAttributes.h"
 #include "RenderUtils.h"
 #include "Engine/StaticMesh.h"
 #include "MeshDescriptionOperations.h"
@@ -212,14 +211,14 @@ public:
 
 		int32 NumFaces = InMesh.Triangles().Num();
 		int32 NumWedges = NumFaces * 3;
-		FStaticMeshDescriptionConstAttributeGetter InMeshAttribute(&InMesh);
-		TVertexAttributesConstRef<FVector> InVertexPositions = InMeshAttribute.GetPositions();
-		TVertexInstanceAttributesConstRef<FVector> InVertexNormals = InMeshAttribute.GetNormals();
-		TVertexInstanceAttributesConstRef<FVector> InVertexTangents = InMeshAttribute.GetTangents();
-		TVertexInstanceAttributesConstRef<float> InVertexBinormalSigns = InMeshAttribute.GetBinormalSigns();
-		TVertexInstanceAttributesConstRef<FVector4> InVertexColors = InMeshAttribute.GetColors();
-		TVertexInstanceAttributesConstRef<FVector2D> InVertexUVs = InMeshAttribute.GetUVs();
-		TPolygonGroupAttributesConstRef<FName> InPolygonGroupMaterialNames = InMeshAttribute.GetPolygonGroupImportedMaterialSlotNames();
+		const FStaticMeshConstAttributes InMeshAttribute(InMesh);
+		TVertexAttributesConstRef<FVector> InVertexPositions = InMeshAttribute.GetVertexPositions();
+		TVertexInstanceAttributesConstRef<FVector> InVertexNormals = InMeshAttribute.GetVertexInstanceNormals();
+		TVertexInstanceAttributesConstRef<FVector> InVertexTangents = InMeshAttribute.GetVertexInstanceTangents();
+		TVertexInstanceAttributesConstRef<float> InVertexBinormalSigns = InMeshAttribute.GetVertexInstanceBinormalSigns();
+		TVertexInstanceAttributesConstRef<FVector4> InVertexColors = InMeshAttribute.GetVertexInstanceColors();
+		TVertexInstanceAttributesConstRef<FVector2D> InVertexUVs = InMeshAttribute.GetVertexInstanceUVs();
+		TPolygonGroupAttributesConstRef<FName> InPolygonGroupMaterialNames = InMeshAttribute.GetPolygonGroupMaterialSlotNames();
 
 		TPolygonGroupAttributesRef<FName> OutPolygonGroupMaterialNames = OutReducedMesh.PolygonGroupAttributes().GetAttributesRef<FName>(MeshAttribute::PolygonGroup::ImportedMaterialSlotName);
 
@@ -370,6 +369,47 @@ public:
 		float* ColorWeights = AttributeWeights + 3 + 3 + 3;
 		float* TexCoordWeights = ColorWeights + 4;
 
+		// Re-scale the weights for UV channels that exceed the expected 0-1 range.
+		// Otherwise garbage on the UVs will dominate the simplification quadric.
+		{
+			float XLength[MAX_STATIC_TEXCOORDS] = { 0 };
+			float YLength[MAX_STATIC_TEXCOORDS] = { 0 };
+			{
+				for (int32 TexCoordId = 0; TexCoordId < NumTexCoords; ++TexCoordId)
+				{
+					float XMax = -FLT_MAX;
+					float YMax = -FLT_MAX;
+					float XMin = FLT_MAX;
+					float YMin = FLT_MAX;
+					for (const TVertSimp< NumTexCoords >& SimpVert : Verts)
+					{
+						const FVector2D& UVs = SimpVert.TexCoords[TexCoordId];
+						XMax = FMath::Max(XMax, UVs.X);
+						XMin = FMath::Min(XMin, UVs.X);
+
+						YMax = FMath::Max(YMax, UVs.Y);
+						YMin = FMath::Min(YMin, UVs.Y);
+					}
+
+					XLength[TexCoordId] =  ( XMax > XMin ) ? XMax - XMin : 0.f;
+					YLength[TexCoordId] =  ( YMax > YMin ) ? YMax - YMin : 0.f;
+				}
+			}
+
+			for (int32 TexCoordId = 0; TexCoordId < NumTexCoords; ++TexCoordId)
+			{
+
+				if (XLength[TexCoordId] > 1.f)
+				{
+					TexCoordWeights[2 * TexCoordId + 0] /= XLength[TexCoordId];
+				}
+				if (YLength[TexCoordId] > 1.f)
+				{
+					TexCoordWeights[2 * TexCoordId + 1] /= YLength[TexCoordId];
+				}
+			}
+		}
+
 		// Zero out weights that aren't used
 		{
 			//TODO Check if we have vertex color
@@ -538,8 +578,7 @@ public:
 
 	virtual bool ReduceSkeletalMesh(
 		USkeletalMesh* SkeletalMesh,
-		int32 LODIndex,
-		bool bReregisterComponent = true
+		int32 LODIndex
 		) override
 	{
 		return false;
diff --git a/Engine/Source/Editor/GameProjectGeneration/Private/GameProjectUtils.cpp b/Engine/Source/Editor/GameProjectGeneration/Private/GameProjectUtils.cpp
--- a/Engine/Source/Editor/GameProjectGeneration/Private/GameProjectUtils.cpp
+++ b/Engine/Source/Editor/GameProjectGeneration/Private/GameProjectUtils.cpp
@@ -158,6 +158,50 @@ namespace
 
 		return MoveTemp(DefaultProjectSettings);
 	}
+
+	/** Set the state of XR plugins in OutProject based on the flags in InProjectInfo. */
+	void SetXRPluginStates(const FProjectInformation& InProjectInfo, FProjectDescriptor& OutProject)
+	{
+		static const FString XRPlugins[] = {
+			TEXT("MagicLeapMedia"),
+			TEXT("MagicLeap"),
+			TEXT("OculusVR"),
+			TEXT("SteamVR") };
+
+		if (!InProjectInfo.bEnableXR)
+		{
+			for (const FString& Plugin : XRPlugins)
+			{
+				int32 Index = OutProject.FindPluginReferenceIndex(Plugin);
+				if (Index == INDEX_NONE)
+				{
+					Index = OutProject.Plugins.AddDefaulted();
+					OutProject.Plugins[Index].Name = Plugin;
+				}
+
+				OutProject.Plugins[Index].bEnabled = false;
+			}
+		}
+	}
+
+	/** Get the configuration values for raytracing if enabled. */
+	FString GetRaytracingConfigString(const FProjectInformation& InProjectInfo)
+	{
+		FString RaytracingConfig;
+
+		if (InProjectInfo.bEnableRaytracing)
+		{
+			RaytracingConfig += LINE_TERMINATOR;
+			RaytracingConfig += TEXT("[/Script/WindowsTargetPlatform.WindowsTargetSettings]") LINE_TERMINATOR;
+			RaytracingConfig += TEXT("DefaultGraphicsRHI=DefaultGraphicsRHI_DX12") LINE_TERMINATOR;
+			RaytracingConfig += LINE_TERMINATOR;
+			RaytracingConfig += TEXT("[/Script/Engine.RendererSettings]") LINE_TERMINATOR;
+			RaytracingConfig += TEXT("r.SkinCache.CompileShaders=True") LINE_TERMINATOR;
+			RaytracingConfig += TEXT("r.RayTracing=True") LINE_TERMINATOR;
+		}
+
+		return MoveTemp(RaytracingConfig);
+	}
 } // namespace <>
 
 FText FNewClassInfo::GetClassName() const
@@ -1268,6 +1312,11 @@ UTemplateProjectDefs* GameProjectUtils::LoadTemplateDefs(const FString& ProjectD
 		TemplateDefs = NewObject<UTemplateProjectDefs>(GetTransientPackage(), ClassToConstruct);
 		TemplateDefs->LoadConfig(UTemplateProjectDefs::StaticClass(), *TemplateDefsIniFilename);
 
+		if (TemplateDefs->HiddenSettings.Num() > 1 && TemplateDefs->HiddenSettings.Contains(ETemplateSetting::All))
+		{
+			UE_LOG(LogGameProjectGeneration, Warning, TEXT("Template '%s' contains 'All' in HiddenSettings in addition to other entries. This is a mistake, and means that all settings will be hidden."), *ProjectDirectory);
+		}
+
 		TArray<TSharedPtr<FTemplateCategory>> AllTemplateCategories;
 		FGameProjectGenerationModule::Get().GetAllTemplateCategories(AllTemplateCategories);
 
@@ -1280,7 +1329,7 @@ UTemplateProjectDefs* GameProjectUtils::LoadTemplateDefs(const FString& ProjectD
 
 			if (!bCategoryExists)
 			{
-				UE_LOG(LogGameProjectGeneration, Warning, TEXT("Failed to find category definition named '%s', not defined in TemplateCategories.ini."), *CategoryKey.ToString());
+				UE_LOG(LogGameProjectGeneration, Warning, TEXT("Failed to find category definition named '%s' while loading template '%s', it is not defined in any TemplateCategories.ini."), *CategoryKey.ToString(), *ProjectDirectory);
 			}
 		}
 	}
@@ -1349,17 +1398,19 @@ bool GameProjectUtils::GenerateProjectFromScratch(const FProjectInformation& InP
 	// Generate the project file
 	{
 		// Set up the descriptor
-		FProjectDescriptor Descriptor;
+		FProjectDescriptor Project;
 		for(int32 Idx = 0; Idx < StartupModuleNames.Num(); Idx++)
 		{
-			Descriptor.Modules.Add(FModuleDescriptor(*StartupModuleNames[Idx]));
+			Project.Modules.Add(FModuleDescriptor(*StartupModuleNames[Idx]));
 		}
 
-		Descriptor.bIsEnterpriseProject = InProjectInfo.bIsEnterpriseProject;
+		Project.bIsEnterpriseProject = InProjectInfo.bIsEnterpriseProject;
+
+		SetXRPluginStates(InProjectInfo, Project);
 
 		// Try to save it
 		FText LocalFailReason;
-		if(!Descriptor.Save(InProjectInfo.ProjectFilename, LocalFailReason))
+		if(!Project.Save(InProjectInfo.ProjectFilename, LocalFailReason))
 		{
 			OutFailReason = LocalFailReason;
 			return false;
@@ -1592,6 +1643,8 @@ bool GameProjectUtils::CreateProjectFromTemplate(const FProjectInformation& InPr
 		FileContents += LINE_TERMINATOR;
 		FileContents += GetHardwareConfigString(InProjectInfo);
 
+		FileContents += GetRaytracingConfigString(InProjectInfo);
+
 		if ( !WriteOutputFile(DefaultEngineIniFilename, FileContents, OutFailReason) )
 		{
 			return false;
@@ -1757,6 +1810,9 @@ bool GameProjectUtils::CreateProjectFromTemplate(const FProjectInformation& InPr
 			ModuleInfo.Name = FName(*ModuleInfo.Name.ToString().Replace(*BaseSourceName, *BaseNewName));
 		}
 
+		// Disable XR plugins if desired
+		SetXRPluginStates(InProjectInfo, Project);
+
 		// Save it to disk
 		if(!Project.Save(InProjectInfo.ProjectFilename, OutFailReason))
 		{
@@ -1922,7 +1978,7 @@ FString GameProjectUtils::GetHardwareConfigString(const FProjectInformation& InP
 	HardwareTargeting += FString::Printf(TEXT("DefaultGraphicsPerformance=%s") LINE_TERMINATOR, *GraphicsPresetAsString);
 	HardwareTargeting += LINE_TERMINATOR;
 
-	return HardwareTargeting;
+	return MoveTemp(HardwareTargeting);
 }
 
 bool GameProjectUtils::GenerateConfigFiles(const FProjectInformation& InProjectInfo, TArray<FString>& OutCreatedFiles, FText& OutFailReason)
@@ -1986,6 +2042,9 @@ bool GameProjectUtils::GenerateConfigFiles(const FProjectInformation& InProjectI
 			}
 		}
 
+		FileContents += LINE_TERMINATOR;
+		FileContents += GetRaytracingConfigString(InProjectInfo);
+
 		if (WriteOutputFile(DefaultEngineIniFilename, FileContents, OutFailReason))
 		{
 			OutCreatedFiles.Add(DefaultEngineIniFilename);
@@ -4087,7 +4146,6 @@ bool GameProjectUtils::AddSharedContentToProject(const FProjectInformation &InPr
 			RequiredDetail = EFeaturePackDetailLevel::Standard;
 		}
 
-
 		TUniquePtr<FFeaturePackContentSource> TempFeaturePack = MakeUnique<FFeaturePackContentSource>();
 		bool bCopied = TempFeaturePack->InsertAdditionalResources(TemplateDefs->SharedContentPacks,RequiredDetail, DestFolder,CreatedFiles);
 		if( bCopied == false )
diff --git a/Engine/Source/Editor/UnrealEd/Classes/Editor/EditorEngine.h b/Engine/Source/Editor/UnrealEd/Classes/Editor/EditorEngine.h
--- a/Engine/Source/Editor/UnrealEd/Classes/Editor/EditorEngine.h
+++ b/Engine/Source/Editor/UnrealEd/Classes/Editor/EditorEngine.h
@@ -44,7 +44,7 @@ class FViewport;
 class IEngineLoop;
 class ILauncherWorker;
 class ILayers;
-class ILevelViewport;
+class IAssetViewport;
 class ITargetPlatform;
 class SViewport;
 class UActorFactory;
@@ -103,7 +103,7 @@ struct FSlatePlayInEditorInfo
 	TSharedPtr<class FSceneViewport>	SlatePlayInEditorWindowViewport;
 	
 	/** The slate viewport that should be used for play in viewport */
-	TWeakPtr<class ILevelViewport>		DestinationSlateViewport;
+	TWeakPtr<class IAssetViewport>		DestinationSlateViewport;
 
 	FSlatePlayInEditorInfo()
 	: SlatePlayInEditorWindow(NULL), DestinationSlateViewport(NULL)
@@ -619,7 +619,7 @@ public:
 	TMap<FName, FSlatePlayInEditorInfo>	SlatePlayInEditorMap;
 
 	/** Viewport the next PlaySession was requested to happen on */
-	TWeakPtr<class ILevelViewport>		RequestedDestinationSlateViewport;
+	TWeakPtr<class IAssetViewport>		RequestedDestinationSlateViewport;
 
 	/** When set to anything other than -1, indicates a specific In-Editor viewport index that PIE should use */
 	UPROPERTY()
@@ -1814,7 +1814,7 @@ public:
 	 * @param	bUseMobilePreview		True to enable mobile preview mode (PC platform only)
 	 * @param	bUseVRPreview			True to enable VR preview mode (PC platform only)
 	 */
-	void RequestPlaySession( bool bAtPlayerStart, TSharedPtr<class ILevelViewport> DestinationViewport, bool bInSimulateInEditor, const FVector* StartLocation = NULL, const FRotator* StartRotation = NULL, int32 DestinationConsole = -1, bool bUseMobilePreview = false, bool bUseVRPreview = false, bool bUseVulkanPreview = false);
+	void RequestPlaySession( bool bAtPlayerStart, TSharedPtr<class IAssetViewport> DestinationViewport, bool bInSimulateInEditor, const FVector* StartLocation = NULL, const FRotator* StartRotation = NULL, int32 DestinationConsole = -1, bool bUseMobilePreview = false, bool bUseVRPreview = false, bool bUseVulkanPreview = false);
 
 	// @todo gmp: temp hack for Rocket demo
 	void RequestPlaySession(const FVector* StartLocation, const FRotator* StartRotation, bool MobilePreview, bool VulkanPreview, const FString& MobilePreviewTargetDevice, FString AdditionalStandaloneLaunchParameters = TEXT(""));
diff --git a/Engine/Source/Editor/UnrealEd/Private/Commandlets/AssetRegistryGenerator.cpp b/Engine/Source/Editor/UnrealEd/Private/Commandlets/AssetRegistryGenerator.cpp
--- a/Engine/Source/Editor/UnrealEd/Private/Commandlets/AssetRegistryGenerator.cpp
+++ b/Engine/Source/Editor/UnrealEd/Private/Commandlets/AssetRegistryGenerator.cpp
@@ -706,6 +706,79 @@ void FAssetRegistryGenerator::UpdateKeptPackagesAssetData()
 	}
 }
 
+void FAssetRegistryGenerator::UpdateCollectionAssetData()
+{
+	// Read out the per-platform settings use to build the list of collections to tag
+	bool bTagAllCollections = false;
+	TArray<FString> CollectionsToIncludeOrExclude;
+	{
+		const FString PlatformIniName = TargetPlatform->IniPlatformName();
+
+		FConfigFile PlatformEngineIni;
+		FConfigCacheIni::LoadLocalIniFile(PlatformEngineIni, TEXT("Engine"), true, (!PlatformIniName.IsEmpty() ? *PlatformIniName : ANSI_TO_TCHAR(FPlatformProperties::IniPlatformName())));
+
+		// The list of collections will either be a inclusive or a exclusive depending on the value of bTagAllCollections
+		PlatformEngineIni.GetBool(TEXT("AssetRegistry"), TEXT("bTagAllCollections"), bTagAllCollections);
+		PlatformEngineIni.GetArray(TEXT("AssetRegistry"), bTagAllCollections ? TEXT("CollectionsToExcludeAsTags") : TEXT("CollectionsToIncludeAsTags"), CollectionsToIncludeOrExclude);
+	}
+
+	// Build the list of collections we should tag for each asset
+	TMap<FName, TArray<FName>> AssetPathNamesToCollectionTags;
+	{
+		ICollectionManager& CollectionManager = FCollectionManagerModule::GetModule().Get();
+
+		TArray<FCollectionNameType> CollectionNamesToTag;
+		CollectionManager.GetCollections(CollectionNamesToTag);
+		if (bTagAllCollections)
+		{
+			CollectionNamesToTag.RemoveAll([&CollectionsToIncludeOrExclude](const FCollectionNameType& CollectionNameAndType)
+			{
+				return CollectionsToIncludeOrExclude.Contains(CollectionNameAndType.Name.ToString());
+			});
+		}
+		else
+		{
+			CollectionNamesToTag.RemoveAll([&CollectionsToIncludeOrExclude](const FCollectionNameType& CollectionNameAndType)
+			{
+				return !CollectionsToIncludeOrExclude.Contains(CollectionNameAndType.Name.ToString());
+			});
+		}
+		
+		TArray<FName> TmpAssetPathNames;
+		for (const FCollectionNameType& CollectionNameToTag : CollectionNamesToTag)
+		{
+			const FName CollectionTagName = *FString::Printf(TEXT("%s%s"), FAssetData::GetCollectionTagPrefix(), *CollectionNameToTag.Name.ToString());
+
+			TmpAssetPathNames.Reset();
+			CollectionManager.GetAssetsInCollection(CollectionNameToTag.Name, CollectionNameToTag.Type, TmpAssetPathNames);
+
+			for (const FName AssetPathName : TmpAssetPathNames)
+			{
+				TArray<FName>& CollectionTagsForAsset = AssetPathNamesToCollectionTags.FindOrAdd(AssetPathName);
+				CollectionTagsForAsset.AddUnique(CollectionTagName);
+			}
+		}
+	}
+
+	// Apply the collection tags to the asset registry state
+	for (const auto& AssetPathNameToCollectionTagsPair : AssetPathNamesToCollectionTags)
+	{
+		const FName AssetPathName = AssetPathNameToCollectionTagsPair.Key;
+		const TArray<FName>& CollectionTagsForAsset = AssetPathNameToCollectionTagsPair.Value;
+
+		const FAssetData* AssetData = State.GetAssetByObjectPath(AssetPathName);
+		if (AssetData)
+		{
+			FAssetDataTagMap TagsAndValues = AssetData->TagsAndValues.GetMap();
+			for (const FName CollectionTagName : CollectionTagsForAsset)
+			{
+				TagsAndValues.Add(CollectionTagName, FString()); // TODO: Does this need a value to avoid being trimmed?
+			}
+			State.UpdateAssetData(FAssetData(AssetData->PackageName, AssetData->PackagePath, AssetData->AssetName, AssetData->AssetClass, TagsAndValues, AssetData->ChunkIDs, AssetData->PackageFlags));
+		}
+	}
+}
+
 void FAssetRegistryGenerator::Initialize(const TArray<FName> &InStartupPackages)
 {
 	StartupPackages.Append(InStartupPackages);
@@ -1045,6 +1118,7 @@ bool FAssetRegistryGenerator::SaveAssetRegistry(const FString& SandboxPath, bool
 	AssetRegistry.InitializeTemporaryAssetRegistryState(State, SaveOptions, true);
 	// Then possibly apply AssetData with TagsAndValues from a previous AssetRegistry for packages kept from a previous cook
 	UpdateKeptPackagesAssetData();
+	UpdateCollectionAssetData();
 
 	if (DevelopmentSaveOptions.bSerializeAssetRegistry && bSerializeDevelopmentAssetRegistry)
 	{
diff --git a/Engine/Source/Editor/UnrealEd/Private/EditorEngine.cpp b/Engine/Source/Editor/UnrealEd/Private/EditorEngine.cpp
--- a/Engine/Source/Editor/UnrealEd/Private/EditorEngine.cpp
+++ b/Engine/Source/Editor/UnrealEd/Private/EditorEngine.cpp
@@ -129,7 +129,7 @@
 #include "DirectoryWatcherModule.h"
 
 #include "Slate/SceneViewport.h"
-#include "ILevelViewport.h"
+#include "IAssetViewport.h"
 
 #include "ContentStreaming.h"
 #include "Components/SkeletalMeshComponent.h"
@@ -214,6 +214,7 @@
 #include "RenderTargetPool.h"
 #include "RenderGraphBuilder.h"
 #include "ToolMenus.h"
+#include "IToolMenusEditorModule.h"
 #include "Subsystems/AssetEditorSubsystem.h"
 
 
@@ -784,7 +785,10 @@ void UEditorEngine::InitEditor(IEngineLoop* InEngineLoop)
 
 	if (FSlateApplication::IsInitialized() && UToolMenus::IsToolMenuUIEnabled())
 	{
-		TWeakPtr<FTimerManager> WeakTimerManager;
+		UToolMenus::Get()->EditMenuIcon = FSlateIcon(FCoreStyle::Get().GetStyleSetName(), "MultiBox.GenericToolBarIcon.Small");
+		UToolMenus::Get()->EditToolbarIcon = FSlateIcon(FCoreStyle::Get().GetStyleSetName(), "MultiBox.GenericToolBarIcon");
+
+		TWeakPtr<FTimerManager> WeakTimerManager = TimerManager;
 		UToolMenus::Get()->AssignSetTimerForNextTickDelegate(FSimpleDelegate::CreateLambda([WeakTimerManager]()
 		{
 			if (WeakTimerManager.IsValid())
@@ -793,6 +797,18 @@ void UEditorEngine::InitEditor(IEngineLoop* InEngineLoop)
 			}
 		}));
 
+		bool bEnableEditToolMenusUI = false;
+		GConfig->GetBool(TEXT("/Script/UnrealEd.EditorExperimentalSettings"), TEXT("bEnableEditToolMenusUI"), bEnableEditToolMenusUI, GEditorPerProjectIni);
+		if (bEnableEditToolMenusUI)
+		{
+			IToolMenusEditorModule::Get().RegisterShowEditMenusModeCheckbox();
+
+			UToolMenus::Get()->EditMenuDelegate.BindLambda([](UToolMenu* InMenu)
+			{
+				IToolMenusEditorModule::Get().OpenEditToolMenuDialog(InMenu);
+			});
+		}
+
 		UToolMenus::Get()->ShouldDisplayExtensionPoints.BindStatic(&GetDisplayMultiboxHooks);
 
 		UToolMenus::Get()->RegisterStringCommandHandler("Command", FToolMenuExecuteString::CreateLambda([](const FString& InString, const FToolMenuContext& InContext)
@@ -2769,7 +2785,7 @@ void UEditorEngine::ApplyDeltaToComponent(USceneComponent* InComponent,
 		{
 			if ( bDelta )
 			{
-				const FRotator Rot = InComponent->RelativeRotation;
+				const FRotator Rot = InComponent->GetRelativeRotation();
 				FRotator ActorRotWind, ActorRotRem;
 				Rot.GetWindingAndRemainder(ActorRotWind, ActorRotRem);
 				const FQuat ActorQ = ActorRotRem.Quaternion();
@@ -2789,7 +2805,7 @@ void UEditorEngine::ApplyDeltaToComponent(USceneComponent* InComponent,
 
 			if ( bDelta )
 			{
-				FVector NewCompLocation = InComponent->RelativeLocation;
+				FVector NewCompLocation = InComponent->GetRelativeLocation();
 				NewCompLocation -= PivotLocation;
 				NewCompLocation = FRotationMatrix( InDeltaRot ).TransformPosition( NewCompLocation );
 				NewCompLocation += PivotLocation;
@@ -2804,7 +2820,7 @@ void UEditorEngine::ApplyDeltaToComponent(USceneComponent* InComponent,
 	{
 		if ( bDelta )
 		{
-			InComponent->SetRelativeLocation(InComponent->RelativeLocation + *InTrans);
+			InComponent->SetRelativeLocation(InComponent->GetRelativeLocation() + *InTrans);
 		}
 		else
 		{
@@ -2822,9 +2838,9 @@ void UEditorEngine::ApplyDeltaToComponent(USceneComponent* InComponent,
 		{
 			if ( bDelta )
 			{
-				InComponent->SetRelativeScale3D(InComponent->RelativeScale3D + InDeltaScale);
+				InComponent->SetRelativeScale3D(InComponent->GetRelativeScale3D() + InDeltaScale);
 
-				FVector NewCompLocation = InComponent->RelativeLocation;
+				FVector NewCompLocation = InComponent->GetRelativeLocation();
 				NewCompLocation -= PivotLocation;
 				NewCompLocation += FScaleMatrix( InDeltaScale ).TransformPosition( NewCompLocation );
 				NewCompLocation += PivotLocation;
@@ -3606,8 +3622,8 @@ private:
 		// Copy over actor properties.
 		Location				= Actor->GetActorLocation();
 		Rotation				= Actor->GetActorRotation();
-		DrawScale3D				= Actor->GetRootComponent() ? Actor->GetRootComponent()->RelativeScale3D : FVector(1.f,1.f,1.f);
-		bHidden					= Actor->bHidden;
+		DrawScale3D				= Actor->GetRootComponent() ? Actor->GetRootComponent()->GetRelativeScale3D() : FVector(1.f,1.f,1.f);
+		bHidden					= Actor->IsHidden();
 
 		// Record which actor properties differ from their defaults.
 		// we don't have properties for location, rotation, scale3D, so copy all the time. 
@@ -3629,16 +3645,25 @@ private:
 		}
 
 		// Set actor properties.
-		if ( bActorPropsDifferFromDefaults[0] ) Actor->SetActorLocation(Location, false);
-		if ( bActorPropsDifferFromDefaults[1] ) Actor->SetActorRotation(Rotation);
-		if ( bActorPropsDifferFromDefaults[4] )
+		if (bActorPropsDifferFromDefaults[0])
+		{
+			Actor->SetActorLocation(Location, false);
+		}
+		if (bActorPropsDifferFromDefaults[1])
+		{
+			Actor->SetActorRotation(Rotation);
+		}
+		if (bActorPropsDifferFromDefaults[4])
 		{
 			if( Actor->GetRootComponent() != NULL )
 			{
 				Actor->GetRootComponent()->SetRelativeScale3D( DrawScale3D );
 			}
 		}
-		if ( bActorPropsDifferFromDefaults[5] ) Actor->bHidden				= bHidden;
+		if (bActorPropsDifferFromDefaults[5])
+		{
+			Actor->SetHidden(bHidden);
+		}
 	}
 
 
@@ -4690,7 +4715,7 @@ TSharedPtr<SViewport> UEditorEngine::GetGameViewportWidget() const
 			return It.Value().SlatePlayInEditorWindowViewport->GetViewportWidget().Pin();
 		}
 
-		TSharedPtr<ILevelViewport> DestinationLevelViewport = It.Value().DestinationSlateViewport.Pin();
+		TSharedPtr<IAssetViewport> DestinationLevelViewport = It.Value().DestinationSlateViewport.Pin();
 		if (DestinationLevelViewport.IsValid())
 		{
 			return DestinationLevelViewport->GetViewportWidget().Pin();
@@ -5157,7 +5182,7 @@ void UEditorEngine::ReplaceActors(UActorFactory* Factory, const FAssetData& Asse
 				}
 				else
 				{
-					NewActorRootComponent->SetRelativeScale3D( OldActor->GetRootComponent()->RelativeScale3D );
+					NewActorRootComponent->SetRelativeScale3D( OldActor->GetRootComponent()->GetRelativeScale3D() );
 				}
 			}
 
@@ -5241,23 +5266,31 @@ void UEditorEngine::ReplaceActors(UActorFactory* Factory, const FAssetData& Asse
 	ReattachActorsHelper::ReattachActors(ConvertedMap, AttachmentInfo);
 
 	// Perform reference replacement on all Actors referenced by World
-	UWorld* CurrentEditorWorld = GetEditorWorldContext().World();
-	FArchiveReplaceObjectRef<AActor> Ar(CurrentEditorWorld, ConvertedMap, false, true, false);
+	TArray<UObject*> ReferencedLevels;
 
-	// Go through modified objects, marking their packages as dirty and informing them of property changes
-	for (const auto& MapItem : Ar.GetReplacedReferences())
+	for (const TPair<AActor*, AActor*>& ReplacedObj : ConvertedMap)
 	{
-		UObject* ModifiedObject = MapItem.Key;
+		ReferencedLevels.AddUnique(ReplacedObj.Value->GetLevel());
+	}
 
-		if (!ModifiedObject->HasAnyFlags(RF_Transient) && ModifiedObject->GetOutermost() != GetTransientPackage() && !ModifiedObject->RootPackageHasAnyFlags(PKG_CompiledIn))
-		{
-			ModifiedObject->MarkPackageDirty();
-		}
+	for (UObject* Referencer : ReferencedLevels)
+	{
+		FArchiveReplaceObjectRef<AActor> Ar(Referencer, ConvertedMap, false, true, false);
 
-		for (UProperty* Property : MapItem.Value)
+		for (const auto& MapItem : Ar.GetReplacedReferences())
 		{
-			FPropertyChangedEvent PropertyEvent(Property);
-			ModifiedObject->PostEditChangeProperty(PropertyEvent);
+			UObject* ModifiedObject = MapItem.Key;
+
+			if (!ModifiedObject->HasAnyFlags(RF_Transient) && ModifiedObject->GetOutermost() != GetTransientPackage() && !ModifiedObject->RootPackageHasAnyFlags(PKG_CompiledIn))
+			{
+				ModifiedObject->MarkPackageDirty();
+			}
+
+			for (UProperty* Property : MapItem.Value)
+			{
+				FPropertyChangedEvent PropertyEvent(Property);
+				ModifiedObject->PostEditChangeProperty(PropertyEvent);
+			}
 		}
 	}
 
diff --git a/Engine/Source/Programs/UnrealHeaderTool/Private/CodeGenerator.cpp b/Engine/Source/Programs/UnrealHeaderTool/Private/CodeGenerator.cpp
--- a/Engine/Source/Programs/UnrealHeaderTool/Private/CodeGenerator.cpp
+++ b/Engine/Source/Programs/UnrealHeaderTool/Private/CodeGenerator.cpp
@@ -2161,6 +2161,14 @@ void FNativeClassHeaderGenerator::ExportNativeGeneratedInitCode(FOutputDevice& O
 		GeneratedClassRegisterFunctionText.Logf(TEXT("\t\t{\r\n"));
 		GeneratedClassRegisterFunctionText.Logf(TEXT("\t\t\tUE4CodeGen_Private::ConstructUClass(OuterClass, %s::ClassParams);\r\n"), *StaticsStructName);
 
+		TArray<FString> SparseClassDataTypes;
+		((FClass*)Class)->GetSparseClassDataTypes(SparseClassDataTypes);
+		ensureMsgf(SparseClassDataTypes.Num() <= 1, TEXT("We don't currently support multiple sparse class data structures on a single object."));
+		for (const FString& SparseClassDataString : SparseClassDataTypes)
+		{
+			GeneratedClassRegisterFunctionText.Logf(TEXT("\t\t\tOuterClass->SetSparseClassDataStruct(F%s::StaticStruct());\r\n"), *SparseClassDataString);
+		}
+
 		if (bIsDynamic)
 		{
 			FString* CustomDynamicClassInitializationMD = MetaDataMap ? MetaDataMap->Find(TEXT("CustomDynamicClassInitialization")) : nullptr;
@@ -4764,6 +4772,8 @@ struct FNativeFunctionStringBuilder
 	FUHTStringBuilder RPCWrappers;
 	FUHTStringBuilder AutogeneratedBlueprintFunctionDeclarations;
 	FUHTStringBuilder AutogeneratedBlueprintFunctionDeclarationsOnlyNotDeclared;
+	FUHTStringBuilder AutogeneratedStaticData;
+	FUHTStringBuilder AutogeneratedStaticDataFuncs;
 };
 
 void FNativeClassHeaderGenerator::ExportNativeFunctions(FOutputDevice& OutGeneratedHeaderText, FOutputDevice& OutMacroCalls, FOutputDevice& OutNoPureDeclsMacroCalls, const FUnrealSourceFile& SourceFile, UClass* Class, FClassMetaData* ClassData)
@@ -4780,6 +4790,81 @@ void FNativeClassHeaderGenerator::ExportNativeFunctions(FOutputDevice& OutGenera
 		ClassRange.Validate();
 	}
 
+	// gather static class data
+	TArray<FString> SparseClassDataTypes;
+	((FClass*)Class)->GetSparseClassDataTypes(SparseClassDataTypes);
+	FString FullClassName = ((FClass*)Class)->GetNameWithPrefix();
+	for (const FString& SparseClassDataString : SparseClassDataTypes)
+	{
+		RuntimeStringBuilders.AutogeneratedStaticData += TEXT("virtual void* CreateSparseClassData() override\r\n");
+		RuntimeStringBuilders.AutogeneratedStaticData += TEXT("{\r\n");
+		RuntimeStringBuilders.AutogeneratedStaticData.Logf(TEXT("\tvoid* SparseClassDataObj = new F%s;\r\n"), *SparseClassDataString);
+		RuntimeStringBuilders.AutogeneratedStaticData.Logf(TEXT("\treturn SparseClassDataObj;\r\n"), *SparseClassDataString);
+		RuntimeStringBuilders.AutogeneratedStaticData += TEXT("}\r\n");
+
+		RuntimeStringBuilders.AutogeneratedStaticData.Logf(TEXT("F%s* Get%s()\r\n"), *SparseClassDataString, *SparseClassDataString);
+		RuntimeStringBuilders.AutogeneratedStaticData += TEXT("{\r\n");
+		RuntimeStringBuilders.AutogeneratedStaticData.Logf(TEXT("\treturn (F%s*)(GetClass()->GetOrCreateSparseClassData());\r\n"), *SparseClassDataString);
+		RuntimeStringBuilders.AutogeneratedStaticData += TEXT("}\r\n");
+
+		RuntimeStringBuilders.AutogeneratedStaticData.Logf(TEXT("F%s* Get%s() const\r\n"), *SparseClassDataString, *SparseClassDataString);
+		RuntimeStringBuilders.AutogeneratedStaticData += TEXT("{\r\n");
+		RuntimeStringBuilders.AutogeneratedStaticData.Logf(TEXT("\treturn (F%s*)(GetClass()->GetOrCreateSparseClassData());\r\n"), *SparseClassDataString);
+		RuntimeStringBuilders.AutogeneratedStaticData += TEXT("}\r\n");
+
+		UScriptStruct* SparseClassDataStruct = FindObjectSafe<UScriptStruct>(ANY_PACKAGE, *SparseClassDataString);
+		while (SparseClassDataStruct != nullptr)
+		{
+			const UProperty* Child = Cast<UProperty>(SparseClassDataStruct->Children);
+			while (Child)
+			{
+				FString ReturnExtendedType;
+				FString VarType = Child->GetCPPType(&ReturnExtendedType, EPropertyExportCPPFlags::CPPF_ArgumentOrReturnValue | EPropertyExportCPPFlags::CPPF_Implementation);
+				if (!ReturnExtendedType.IsEmpty())
+				{
+					VarType.Append(ReturnExtendedType);
+				}
+				FString VarName = Child->GetName();
+				FString CleanVarName = VarName;
+				if (Cast<UBoolProperty>(Child) && VarName.StartsWith(TEXT("b"), ESearchCase::CaseSensitive))
+				{
+					CleanVarName = VarName.RightChop(1);
+				}
+
+				if (!Child->HasMetaData("NoGetter"))
+				{
+					if (Child->HasMetaData("GetByRef"))
+					{
+						RuntimeStringBuilders.AutogeneratedStaticDataFuncs.Logf(TEXT("const %s& Get%s()\r\n"), *VarType, *CleanVarName);
+					}
+					else
+					{
+						RuntimeStringBuilders.AutogeneratedStaticDataFuncs.Logf(TEXT("%s Get%s()\r\n"), *VarType, *CleanVarName);
+					}
+					RuntimeStringBuilders.AutogeneratedStaticDataFuncs.Logf(TEXT("{\r\n"));
+					RuntimeStringBuilders.AutogeneratedStaticDataFuncs.Logf(TEXT("\treturn Get%s()->%s;\r\n"), *SparseClassDataString, *VarName);
+					RuntimeStringBuilders.AutogeneratedStaticDataFuncs.Logf(TEXT("}\r\n"));
+
+					if (Child->HasMetaData("GetByRef"))
+					{
+						RuntimeStringBuilders.AutogeneratedStaticDataFuncs.Logf(TEXT("const %s& Get%s() const\r\n"), *VarType, *CleanVarName);
+					}
+					else
+					{
+						RuntimeStringBuilders.AutogeneratedStaticDataFuncs.Logf(TEXT("%s Get%s() const\r\n"), *VarType, *CleanVarName);
+					}
+					RuntimeStringBuilders.AutogeneratedStaticDataFuncs.Logf(TEXT("{\r\n"));
+					RuntimeStringBuilders.AutogeneratedStaticDataFuncs.Logf(TEXT("\treturn Get%s()->%s;\r\n"), *SparseClassDataString, *VarName);
+					RuntimeStringBuilders.AutogeneratedStaticDataFuncs.Logf(TEXT("}\r\n"));
+				}
+
+				Child = Cast<UProperty>(Child->Next);
+			}
+
+			SparseClassDataStruct = Cast<UScriptStruct>(SparseClassDataStruct->GetSuperStruct());
+		}
+	}
+
 	// export the C++ stubs
 
 	for (UFunction* Function : TFieldRange<UFunction>(Class, EFieldIteratorFlags::ExcludeSuper))
@@ -4881,6 +4966,15 @@ void FNativeClassHeaderGenerator::ExportNativeFunctions(FOutputDevice& OutGenera
 		FuncStringBuilders.RPCWrappers += TEXT("\t}") LINE_TERMINATOR;
 	}
 
+	// static class data
+	{
+		FString MacroName = SourceFile.GetGeneratedMacroName(ClassData, TEXT("_SPARSE_DATA"));
+
+		WriteMacro(OutGeneratedHeaderText, MacroName, RuntimeStringBuilders.AutogeneratedStaticData + RuntimeStringBuilders.AutogeneratedStaticDataFuncs);
+		OutMacroCalls.Logf(TEXT("\t%s\r\n"), *MacroName);
+		OutNoPureDeclsMacroCalls.Logf(TEXT("\t%s\r\n"), *MacroName);
+	}
+
 	// Write runtime wrappers
 	{
 		FString MacroName = SourceFile.GetGeneratedMacroName(ClassData, TEXT("_RPC_WRAPPERS"));
diff --git a/Engine/Source/Programs/UnrealHeaderTool/Private/HeaderParser.cpp b/Engine/Source/Programs/UnrealHeaderTool/Private/HeaderParser.cpp
--- a/Engine/Source/Programs/UnrealHeaderTool/Private/HeaderParser.cpp
+++ b/Engine/Source/Programs/UnrealHeaderTool/Private/HeaderParser.cpp
@@ -5942,6 +5942,7 @@ UClass* FHeaderParser::CompileClassDeclaration(FClasses& AllClasses)
 	if (ClassDeclarationData->AutoCollapseCategories.Num()) { MetaData.Add("AutoCollapseCategories", FString::Join(ClassDeclarationData->AutoCollapseCategories, TEXT(" "))); }
 	if (ClassDeclarationData->HideCategories.Num()) { MetaData.Add("HideCategories", FString::Join(ClassDeclarationData->HideCategories, TEXT(" "))); }
 	if (ClassDeclarationData->ShowSubCatgories.Num()) { MetaData.Add("ShowCategories", FString::Join(ClassDeclarationData->ShowSubCatgories, TEXT(" "))); }
+	if (ClassDeclarationData->SparseClassDataTypes.Num()) { MetaData.Add("SparseClassDataTypes", FString::Join(ClassDeclarationData->SparseClassDataTypes, TEXT(" "))); }
 	if (ClassDeclarationData->HideFunctions.Num()) { MetaData.Add("HideFunctions", FString::Join(ClassDeclarationData->HideFunctions, TEXT(" "))); }
 	if (ClassDeclarationData->AutoExpandCategories.Num()) { MetaData.Add("AutoExpandCategories", FString::Join(ClassDeclarationData->AutoExpandCategories, TEXT(" "))); }
 
@@ -9709,6 +9710,10 @@ void FHeaderParser::ResetClassData()
 		{
 			CurrentClass->SetMetaData(TEXT("ShowCategories"), *SuperClass->GetMetaData("ShowCategories"));
 		}
+		if (SuperClass->HasMetaData(TEXT("SparseClassDataTypes")))
+		{
+			CurrentClass->SetMetaData(TEXT("SparseClassDataTypes"), *SuperClass->GetMetaData("SparseClassDataTypes"));
+		}
 		if (SuperClass->HasMetaData(TEXT("HideFunctions")))
 		{
 			CurrentClass->SetMetaData(TEXT("HideFunctions"), *SuperClass->GetMetaData("HideFunctions"));
@@ -9943,7 +9948,7 @@ void FHeaderParser::CheckDocumentationPolicyForStruct(UStruct* Struct, const TMa
 		for (UProperty* Property : TFieldRange<UProperty>(Struct, EFieldIteratorFlags::ExcludeSuper))
 		{
 			FString ToolTip = Property->GetToolTipText().ToString();
-			if (ToolTip.IsEmpty() || ToolTip.Equals(Property->GetName()))
+			if (ToolTip.IsEmpty() || ToolTip.Equals(Property->GetDisplayNameText().ToString()))
 			{
 				UE_LOG_ERROR_UHT(TEXT("Property '%s::%s' does not provide a tooltip / comment (DocumentationPolicy)."), *Struct->GetName(), *Property->GetName());
 				continue;
diff --git a/Engine/Source/Runtime/AssetRegistry/Public/AssetData.h b/Engine/Source/Runtime/AssetRegistry/Public/AssetData.h
--- a/Engine/Source/Runtime/AssetRegistry/Public/AssetData.h
+++ b/Engine/Source/Runtime/AssetRegistry/Public/AssetData.h
@@ -54,8 +54,15 @@ USTRUCT(BlueprintType)
 struct FAssetData
 {
 	GENERATED_BODY()
+
 public:
+	/** The prefix used for collection entries inside TagsAndValues */
+	static const TCHAR* GetCollectionTagPrefix()
+	{
+		return TEXT("CL_");
+	}
 
+public:
 	/** The object path for the asset in the form PackageName.AssetName. Only top level objects in a package can have AssetData */
 	UPROPERTY(BlueprintReadOnly, Category=AssetData, transient)
 	FName ObjectPath;
@@ -603,6 +610,53 @@ public:
 	}
 };
 
+/**
+ * Helper struct for FAssetIdentifier (e.g., for the FOnViewAssetIdentifiersInReferenceViewer delegate and Reference Viewer functions).
+ */
+#ifdef WITH_EDITORONLY_DATA
+struct FReferenceViewerParams
+{
+	FReferenceViewerParams()
+		// Displayed-on-graph options
+		: bShowReferencers(true)
+		, bShowDependencies(true)
+		// Slider-based options
+		, FixAndHideSearchDepthLimit(0)
+		, FixAndHideSearchBreadthLimit(0)
+		, bShowCollectionFilter(true)
+		// Checkbox options
+		, bShowShowReferencesOptions(true)
+		, bShowShowSearchableNames(true)
+		, bShowShowNativePackages(true)
+	{}
+
+	/* Whether to display the Referencers */
+	bool bShowReferencers;
+	/* Whether to display the Dependencies */
+	bool bShowDependencies;
+	/**
+	 * Whether to visually show to the user the option of "Search Depth Limit" or hide it and fix it to a default value:
+	 * - If 0 or negative, it will show to the user the option of "Search Depth Limit".
+	 * - If >0, it will hide that option and fix the Depth value to this value.
+	 */
+	int32 FixAndHideSearchDepthLimit;
+	/**
+	 * Whether to visually show to the user the option of "Search Breadth Limit" or hide it and fix it to a default value:
+	 * - If 0 or negative, it will show to the user the option of "Search Breadth Limit".
+	 * - If >0, it will hide that option and fix the Breadth value to this value.
+	 */
+	int32 FixAndHideSearchBreadthLimit;
+	/** Whether to visually show to the user the option of "Collection Filter" */
+	bool bShowCollectionFilter;
+	/** Whether to visually show to the user the options of "Show Soft/Hard/Management References" */
+	bool bShowShowReferencesOptions;
+	/** Whether to visually show to the user the option of "Show Searchable Names" */
+	bool bShowShowSearchableNames;
+	/** Whether to visually show to the user the option of "Show Native Packages" */
+	bool bShowShowNativePackages;
+};
+#endif // WITH_EDITORONLY_DATA
+
 /** A structure defining a thing that can be reference by something else in the asset registry. Represents either a package of a primary asset id */
 struct FAssetIdentifier
 {
diff --git a/Engine/Source/Runtime/Core/Private/Internationalization/TextLocalizationManager.cpp b/Engine/Source/Runtime/Core/Private/Internationalization/TextLocalizationManager.cpp
--- a/Engine/Source/Runtime/Core/Private/Internationalization/TextLocalizationManager.cpp
+++ b/Engine/Source/Runtime/Core/Private/Internationalization/TextLocalizationManager.cpp
@@ -885,12 +885,21 @@ void FTextLocalizationManager::LoadLocalizationResourcesForPrioritizedCultures(T
 		return;
 	}
 
+	// Leet-ify always needs the native text to operate on, so force native data if we're loading for LEET
+	ELocalizationLoadFlags FinalLocLoadFlags = LocLoadFlags;
+#if ENABLE_LOC_TESTING
+	if (PrioritizedCultureNames[0] == FLeetCulture::StaticGetName())
+	{
+		FinalLocLoadFlags |= ELocalizationLoadFlags::Native;
+	}
+#endif
+
 	// Load the resources from each text source
 	FTextLocalizationResource NativeResource;
 	FTextLocalizationResource LocalizedResource;
 	for (const TSharedPtr<ILocalizedTextSource>& LocalizedTextSource : LocalizedTextSources)
 	{
-		LocalizedTextSource->LoadLocalizedResources(LocLoadFlags, PrioritizedCultureNames, NativeResource, LocalizedResource);
+		LocalizedTextSource->LoadLocalizedResources(FinalLocLoadFlags, PrioritizedCultureNames, NativeResource, LocalizedResource);
 	}
 
 	// When loc testing is enabled, UpdateFromNative also takes care of restoring non-localized text which is why the condition below is gated
diff --git a/Engine/Source/Runtime/Core/Private/UObject/UnrealNames.cpp b/Engine/Source/Runtime/Core/Private/UObject/UnrealNames.cpp
--- a/Engine/Source/Runtime/Core/Private/UObject/UnrealNames.cpp
+++ b/Engine/Source/Runtime/Core/Private/UObject/UnrealNames.cpp
@@ -1312,17 +1312,6 @@ FString FName::NameToDisplayString( const FString& InDisplayName, const bool bIs
 			bInARun = false;
 		}
 
-		// We were running on uppercase letters before and still do, but the next character is a lowercase letter,
-		// so we should break the run here, like "3DWidget" should be "3D Widget"
-		if (bInARun && !bWasSpace && !bWasOpenParen && CharIndex < Chars.Num() - 1 && FChar::IsLower(Chars[CharIndex + 1]))
-		{
-			if (!bWasSpace && OutDisplayName.Len() > 0)
-			{
-				OutDisplayName += TEXT(' ');
-				bWasSpace = true;
-			}
-		}
-
 		// An underscore denotes a space, so replace it and continue the run
 		if( bIsUnderscore )
 		{
diff --git a/Engine/Source/Runtime/Core/Public/Templates/UnrealTemplate.h b/Engine/Source/Runtime/Core/Public/Templates/UnrealTemplate.h
--- a/Engine/Source/Runtime/Core/Public/Templates/UnrealTemplate.h
+++ b/Engine/Source/Runtime/Core/Public/Templates/UnrealTemplate.h
@@ -299,9 +299,9 @@ private:
  * Macro variant on TGuardValue<bool> that can deal with bitfields which cannot be passed by reference in to TGuardValue
  */
 #define FGuardValue_Bitfield(ReferenceValue, NewValue) \
-	const bool TempBitfield##__LINE__ = ReferenceValue; \
+	const bool PREPROCESSOR_JOIN(TempBitfield, __LINE__) = ReferenceValue; \
 	ReferenceValue = NewValue; \
-	const TGuardValue_Bitfield_Cleanup<TFunction<void()>> TempBitfieldCleanup##__LINE__([&](){ ReferenceValue = TempBitfield##__LINE__; });
+	const TGuardValue_Bitfield_Cleanup<TFunction<void()>> PREPROCESSOR_JOIN(TempBitfieldCleanup, __LINE__)([&](){ ReferenceValue = PREPROCESSOR_JOIN(TempBitfield, __LINE__); });
 
 /** 
  * Commonly used to make sure a value is incremented, and then decremented anyway the function can terminate.
diff --git a/Engine/Source/Runtime/CoreUObject/Private/UObject/Class.cpp b/Engine/Source/Runtime/CoreUObject/Private/UObject/Class.cpp
--- a/Engine/Source/Runtime/CoreUObject/Private/UObject/Class.cpp
+++ b/Engine/Source/Runtime/CoreUObject/Private/UObject/Class.cpp
@@ -2989,7 +2989,7 @@ class FRestoreClassInfo: public FRestoreForUObjectOverwrite
 	UClass*			Target;
 	/** Saved ClassWithin **/
 	UClass*			Within;
-	/** Save ClassGeneratedBy */
+	/** Saved ClassGeneratedBy */
 	UObject*		GeneratedBy;
 	/** Saved ClassDefaultObject **/
 	UObject*		DefaultObject;
@@ -3668,6 +3668,7 @@ void UClass::SetSuperStruct(UStruct* NewSuperStruct)
 	UnhashObject(this);
 	ClearFunctionMapsCaches();
 	Super::SetSuperStruct(NewSuperStruct);
+	SetSparseClassDataStruct(GetSparseClassDataArchetypeStruct());
 	HashObject(this);
 }
 
@@ -3968,6 +3969,33 @@ void UClass::SerializeDefaultObject(UObject* Object, FStructuredArchive::FSlot S
 	UnderlyingArchive.StopSerializingDefaults();
 }
 
+void UClass::SerializeSparseClassData(UObject* Object, FStructuredArchive::FSlot Slot)
+{
+	if (!SparseClassDataStruct)
+	{
+		return;
+	}
+
+	// tell the archive that it's allowed to load data for transient properties
+	FArchive& UnderlyingArchive = Slot.GetUnderlyingArchive();
+
+	if (((UnderlyingArchive.IsLoading() || UnderlyingArchive.IsSaving()) && !UnderlyingArchive.WantBinaryPropertySerialization()))
+	{
+		// class default objects do not always have a vtable when saved
+		// so use script serialization as opposed to native serialization to
+		// guarantee that all property data is loaded into the correct location
+		SparseClassDataStruct->SerializeItem(Slot, SparseClassData, GetArchetypeForSparseClassData());
+	}
+	else if (UnderlyingArchive.GetPortFlags() != 0)
+	{
+		SparseClassDataStruct->SerializeBinEx(Slot, (uint8*)SparseClassData, SparseClassDataStruct, GetSparseClassDataArchetypeStruct());
+	}
+	else
+	{
+		SparseClassDataStruct->SerializeBin(Slot, (uint8*)SparseClassData);
+	}
+}
+
 
 FArchive& operator<<(FArchive& Ar, FImplementedInterface& A)
 {
@@ -3978,9 +4006,21 @@ FArchive& operator<<(FArchive& Ar, FImplementedInterface& A)
 	return Ar;
 }
 
+void* UClass::GetArchetypeForSparseClassData() const
+{
+	UClass* SuperClass = GetSuperClass();
+	return SuperClass ? SuperClass->GetOrCreateSparseClassData() : nullptr;
+}
+
+UScriptStruct* UClass::GetSparseClassDataArchetypeStruct() const
+{
+	UClass* SuperClass = GetSuperClass();
+	return SuperClass ? SuperClass->GetSparseClassDataStruct() : nullptr;
+}
+
 UObject* UClass::GetArchetypeForCDO() const
 {
-	auto SuperClass = GetSuperClass();
+	UClass* SuperClass = GetSuperClass();
 	return SuperClass ? SuperClass->GetDefaultObject() : nullptr;
 }
 
@@ -4088,6 +4128,8 @@ UClass::UClass(const FObjectInitializer& ObjectInitializer)
 ,	ClassWithin( UObject::StaticClass() )
 ,	ClassGeneratedBy(nullptr)
 ,	ClassDefaultObject(nullptr)
+,	SparseClassData(nullptr)
+,	SparseClassDataStruct(nullptr)
 {
 	// If you add properties here, please update the other constructors and PurgeClass()
 
@@ -4111,6 +4153,8 @@ UClass::UClass(const FObjectInitializer& ObjectInitializer, UClass* InBaseClass)
 ,	ClassWithin(UObject::StaticClass())
 ,	ClassGeneratedBy(nullptr)
 ,	ClassDefaultObject(nullptr)
+,	SparseClassData(nullptr)
+,	SparseClassDataStruct(nullptr)
 {
 	// If you add properties here, please update the other constructors and PurgeClass()
 
@@ -4166,6 +4210,8 @@ UClass::UClass
 ,	ClassConfigName			()
 ,	NetFields				()
 ,	ClassDefaultObject		( nullptr )
+,	SparseClassData			( nullptr )
+,	SparseClassDataStruct	( nullptr )
 {
 	// If you add properties here, please update the other constructors and PurgeClass()
 
@@ -4178,6 +4224,41 @@ UClass::UClass
 	*(const TCHAR**)&ClassConfigName = InConfigName; //-V580
 }
 
+void* UClass::CreateSparseClassData()
+{
+	check(SparseClassData == nullptr);
+
+	SparseClassData = GetDefaultObject()->CreateSparseClassData();
+
+	if (SparseClassData)
+	{
+		// initialize per class data from the archetype if we have one
+		void* SparseArchetypeData = GetArchetypeForSparseClassData();
+		UStruct* SparseClassDataArchetypeStruct = GetSparseClassDataArchetypeStruct();
+
+		if (SparseArchetypeData)
+		{
+			for (UProperty* P = SparseClassDataArchetypeStruct->PropertyLink; P; P = P->PropertyLinkNext)
+			{
+				P->CopyCompleteValue_InContainer(SparseClassData, SparseArchetypeData);
+			}
+		}
+	}
+
+	return SparseClassData;
+}
+
+UScriptStruct* UClass::GetSparseClassDataStruct() const
+{
+	// this info is specified on the object via code generation so we use it instead of looking at the UClass
+	return SparseClassDataStruct;
+}
+
+void UClass::SetSparseClassDataStruct(UScriptStruct* InSparseClassDataStruct)
+{ 
+	SparseClassDataStruct = InSparseClassDataStruct; 
+}
+
 #if WITH_HOT_RELOAD
 
 bool UClass::HotReloadPrivateStaticClass(
diff --git a/Engine/Source/Runtime/CoreUObject/Private/UObject/Obj.cpp b/Engine/Source/Runtime/CoreUObject/Private/UObject/Obj.cpp
--- a/Engine/Source/Runtime/CoreUObject/Private/UObject/Obj.cpp
+++ b/Engine/Source/Runtime/CoreUObject/Private/UObject/Obj.cpp
@@ -1152,6 +1152,16 @@ void UObject::PostLoadSubobjects( FObjectInstancingGraph* OuterInstanceGraph/*=N
 	}
 }
 
+void* UObject::CreateSparseClassData()
+{
+	return nullptr;
+}
+
+UScriptStruct* UObject::GetSparseClassDataStruct() const
+{
+	UClass* Class = GetClass();
+	return Class ? Class->GetSparseClassDataStruct() : nullptr;
+}
 
 void UObject::ConditionalPostLoadSubobjects( FObjectInstancingGraph* OuterInstanceGraph/*=NULL*/ )
 {
@@ -1406,7 +1416,7 @@ void UObject::SerializeScriptProperties( FStructuredArchive::FSlot Slot ) const
 		FArchive::FScopeAddDebugData S(UnderlyingArchive, ObjClass->GetFName());
 #endif
 
-		ObjClass->SerializeTaggedProperties(Slot, (uint8*)this, HasAnyFlags(RF_ClassDefaultObject) ? ObjClass->GetSuperClass() : ObjClass, (uint8*)DiffObject, bBreakSerializationRecursion ? this : NULL);
+		ObjClass->SerializeTaggedProperties(Slot, (uint8*)this, HasAnyFlags(RF_ClassDefaultObject) ? ObjClass->GetSuperClass() : ObjClass, (uint8*)DiffObject, bBreakSerializationRecursion ? this : nullptr);
 	}
 	else if (UnderlyingArchive.GetPortFlags() != 0 && !UnderlyingArchive.ArUseCustomPropertyList )
 	{
@@ -1423,7 +1433,7 @@ void UObject::SerializeScriptProperties( FStructuredArchive::FSlot Slot ) const
 		ObjClass->SerializeBin(Slot, const_cast<UObject *>(this));
 	}
 
-	if( HasAnyFlags(RF_ClassDefaultObject) )
+	if (HasAnyFlags(RF_ClassDefaultObject))
 	{
 		UnderlyingArchive.StopSerializingDefaults();
 	}
diff --git a/Engine/Source/Runtime/CoreUObject/Private/UObject/UObjectGlobals.cpp b/Engine/Source/Runtime/CoreUObject/Private/UObject/UObjectGlobals.cpp
--- a/Engine/Source/Runtime/CoreUObject/Private/UObject/UObjectGlobals.cpp
+++ b/Engine/Source/Runtime/CoreUObject/Private/UObject/UObjectGlobals.cpp
@@ -4583,6 +4583,8 @@ namespace UE4CodeGen_Private
 #endif
 
 		NewClass->StaticLink();
+
+		NewClass->SetSparseClassDataStruct(NewClass->GetSparseClassDataArchetypeStruct());
 	}
 }
 
diff --git a/Engine/Source/Runtime/CoreUObject/Public/UObject/Class.h b/Engine/Source/Runtime/CoreUObject/Public/UObject/Class.h
--- a/Engine/Source/Runtime/CoreUObject/Public/UObject/Class.h
+++ b/Engine/Source/Runtime/CoreUObject/Public/UObject/Class.h
@@ -2430,6 +2430,7 @@ class COREUOBJECT_API UClass : public UStruct
 
 public:
 	friend class FRestoreClassInfo;
+	friend class FBlueprintEditorUtils;
 
 	typedef void		(*ClassConstructorType)				(const FObjectInitializer&);
 	typedef UObject*	(*ClassVTableHelperCtorCallerType)	(FVTableHelper& Helper);
@@ -2506,6 +2507,28 @@ public:
 	/** The class default object; used for delta serialization and object initialization */
 	UObject* ClassDefaultObject;
 
+protected:
+	/** This is where we store the data that is only changed per class instead of per instance */
+	UPROPERTY()
+	void* SparseClassData;
+
+	/** The struct used to store sparse class data. */
+	UPROPERTY()
+	UScriptStruct* SparseClassDataStruct;
+
+public:
+	/**
+	 * Returns a pointer to the sidecar data structure. This function will create an instance of the data structure if one has been specified and it has not yet been created.
+	 */
+	void* GetOrCreateSparseClassData() { return SparseClassData ? SparseClassData : CreateSparseClassData(); }
+
+	/**
+	 * Returns a pointer to the type of the sidecar data structure if one is specified.
+	 */
+	virtual UScriptStruct* GetSparseClassDataStruct() const;
+
+	void SetSparseClassDataStruct(UScriptStruct* InSparseClassDataStruct);
+
 	/** Assemble reference token streams for all classes if they haven't had it assembled already */
 	static void AssembleReferenceTokenStreams();
 
@@ -2517,6 +2540,8 @@ public:
 #endif // WITH_EDITOR
 
 private:
+	void* CreateSparseClassData();
+
 #if WITH_EDITOR
 	/** Provides access to attributes of the underlying C++ class. Should never be unset. */
 	TOptional<FCppClassTypeInfo> CppTypeInfo;
@@ -2928,7 +2953,7 @@ public:
 	 */
 	bool ImplementsInterface(const class UClass* SomeInterface) const;
 
-	/** serializes the passed in object as this class's default object using the given archive
+	/** serializes the passed in object as this class's default object using the given archive slot
 	 * @param Object the object to serialize as default
 	 * @param Slot the structured archive slot to serialize from
 	 */
@@ -2943,6 +2968,12 @@ public:
 		SerializeDefaultObject(Object, FStructuredArchiveFromArchive(Ar).GetSlot());
 	}
 
+	/** serializes the associated sparse class data for the passed in object using the given archive slot. This should only be called if the class has an associated sparse data structure.
+	 * @param Object the object to serialize as default
+	 * @param Slot the structured archive slot to serialize from
+	 */
+	void SerializeSparseClassData(UObject* Object, FStructuredArchive::FSlot Slot);
+
 	/** Wraps the PostLoad() call for the class default object.
 	 * @param Object the default object to call PostLoad() on
 	 */
@@ -2995,6 +3026,12 @@ public:
 	/** Returns archetype object for CDO */
 	virtual UObject* GetArchetypeForCDO() const;
 
+	/** Returns archetype for sparse class data */
+	virtual void* GetArchetypeForSparseClassData() const;
+
+	/** Returns the struct used by the sparse class data archetype */
+	UScriptStruct* GetSparseClassDataArchetypeStruct() const;
+
 	/**
 	* Returns all objects that should be preloaded before the class default object is serialized at load time. Only used by the EDL.
 	*
diff --git a/Engine/Source/Runtime/CoreUObject/Public/UObject/ObjectMacros.h b/Engine/Source/Runtime/CoreUObject/Public/UObject/ObjectMacros.h
--- a/Engine/Source/Runtime/CoreUObject/Public/UObject/ObjectMacros.h
+++ b/Engine/Source/Runtime/CoreUObject/Public/UObject/ObjectMacros.h
@@ -726,6 +726,9 @@ namespace UC
 
 		/// Marks this class as an 'early access' preview (while not considered production-ready, it's a step beyond 'experimental' and is being provided as a preview of things to come)
 		EarlyAccessPreview,
+
+		// Some properties are stored once per class in a sidecar structure and not on instances of the class
+		SparseClassDataType,
 	};
 }
 
diff --git a/Engine/Source/Runtime/Engine/Classes/Components/StaticMeshComponent.h b/Engine/Source/Runtime/Engine/Classes/Components/StaticMeshComponent.h
--- a/Engine/Source/Runtime/Engine/Classes/Components/StaticMeshComponent.h
+++ b/Engine/Source/Runtime/Engine/Classes/Components/StaticMeshComponent.h
@@ -403,7 +403,9 @@ public:
 	virtual void PostEditChangeProperty(FPropertyChangedEvent& PropertyChangedEvent) override;
 	virtual bool CanEditChange(const UProperty* InProperty) const override;
 #endif // WITH_EDITOR
+#if WITH_EDITORONLY_DATA
 	virtual void PreSave(const class ITargetPlatform* TargetPlatform) override;
+#endif // WITH_EDITORONLY_DATA
 	virtual void PostLoad() override;
 	virtual bool IsPostLoadThreadSafe() const override;
 	virtual bool AreNativePropertiesIdenticalTo( UObject* Other ) const override;
@@ -554,6 +556,7 @@ public:
 		int32& VertexLightMapMemoryUsage, int32& VertexShadowMapMemoryUsage,
 		int32& StaticLightingResolution, bool& bIsUsingTextureMapping, bool& bHasLightmapTexCoords) const;
 
+#if WITH_EDITORONLY_DATA
 	/**
 	 * Determines whether any of the component's LODs require override vertex color fixups
 	 *
@@ -576,6 +579,7 @@ public:
 	 * @param SourceComponent The component to copy vertex colors from
 	 */
 	void CopyInstanceVertexColorsIfCompatible( UStaticMeshComponent* SourceComponent );
+#endif
 
 	/**
 	 * Removes instance vertex colors from the specified LOD
@@ -583,13 +587,16 @@ public:
 	 */
 	void RemoveInstanceVertexColorsFromLOD( int32 LODToRemoveColorsFrom );
 
+#if WITH_EDITORONLY_DATA
 	/**
 	 * Removes instance vertex colors from all LODs
 	 */
 	void RemoveInstanceVertexColors();
+#endif
 
 	void UpdatePreCulledData(int32 LODIndex, const TArray<uint32>& PreCulledData, const TArray<int32>& NumTrianglesPerSection);
 
+#if WITH_EDITORONLY_DATA
 	/**
 	*	Sets the value of the SectionIndexPreview flag and reattaches the component as necessary.
 	*	@param	InSectionIndexPreview		New value of SectionIndexPreview.
@@ -601,7 +608,8 @@ public:
 	*	@param	InMaterialIndexPreview		New value of MaterialIndexPreview.
 	*/
 	void SetMaterialPreview(int32 InMaterialIndexPreview);
-	
+#endif
+
 	/** Sets the BodyInstance to use the mesh's body setup for external collision information*/
 	void UpdateCollisionFromStaticMesh();
 
@@ -615,8 +623,10 @@ private:
 	/** Initializes the resources used by the static mesh component. */
 	void InitResources();
 
+#if WITH_EDITOR
 	/** Update the vertex override colors */
 	void PrivateFixupOverrideColors();
+#endif
 protected:
 
 	/** Whether the component type supports static lighting. */
diff --git a/Engine/Source/Runtime/Engine/Classes/Engine/SkeletalMesh.h b/Engine/Source/Runtime/Engine/Classes/Engine/SkeletalMesh.h
--- a/Engine/Source/Runtime/Engine/Classes/Engine/SkeletalMesh.h
+++ b/Engine/Source/Runtime/Engine/Classes/Engine/SkeletalMesh.h
@@ -177,6 +177,10 @@ struct FSkeletalMeshLODInfo
 	TArray<FName> RemovedBones_DEPRECATED;
 #endif
 
+	/** build settings to apply when building render data. */
+	UPROPERTY(EditAnywhere, Category = BuildSettings)
+	FSkeletalMeshBuildSettings BuildSettings;
+
 	/** Reduction settings to apply when building render data. */
 	UPROPERTY(EditAnywhere, Category = ReductionSettings)
 	FSkeletalMeshOptimizationSettings ReductionSettings;
@@ -230,6 +234,13 @@ struct FSkeletalMeshLODInfo
 	 */
 	UPROPERTY()
 	uint8 bImportWithBaseMesh:1;
+
+	//Temporary build GUID data
+	//We use this GUID to store the LOD Key so we can now if the LOD need to be rebuild
+	//This GUID is set when we Cache the render data (build function)
+	FGuid BuildGUID;
+
+	ENGINE_API FGuid ComputeDeriveDataCacheKey(const FSkeletalMeshLODGroupSettings* SkeletalMeshLODGroupSettings);
 #endif
 
 	FSkeletalMeshLODInfo()
@@ -246,6 +257,9 @@ struct FSkeletalMeshLODInfo
 		, bImportWithBaseMesh(false)
 #endif
 	{
+#if WITH_EDITORONLY_DATA
+		BuildGUID.Invalidate();
+#endif
 	}
 
 };
@@ -420,6 +434,7 @@ struct FSkeletalMaterial
 #if WITH_EDITOR
 /** delegate type for pre skeletal mesh build events */
 DECLARE_MULTICAST_DELEGATE_OneParam(FOnPostMeshCache, class USkeletalMesh*);
+
 #endif
 
 #if WITH_EDITORONLY_DATA
@@ -625,12 +640,12 @@ public:
 	TEnumAsByte<EAxis::Type> SkelMirrorFlipAxis;
 
 	/** If true, use 32 bit UVs. If false, use 16 bit UVs to save memory */
-	UPROPERTY(EditAnywhere, Category=Mesh)
-	uint8 bUseFullPrecisionUVs:1;
+	UPROPERTY()
+	uint8 bUseFullPrecisionUVs_DEPRECATED :1;
 
 	/** If true, tangents will be stored at 16 bit vs 8 bit precision */
-	UPROPERTY(EditAnywhere, Category = Mesh)
-	uint8 bUseHighPrecisionTangentBasis : 1;
+	UPROPERTY()
+	uint8 bUseHighPrecisionTangentBasis_DEPRECATED : 1;
 
 	/** true if this mesh has ever been simplified with Simplygon. */
 	UPROPERTY()
@@ -892,6 +907,10 @@ public:
 	*/
 	void ReleaseResources();
 
+	/**
+	* Flush current render state
+	*/
+	void FlushRenderState();
 
 	/** Release CPU access version of buffer */
 	void ReleaseCPUResources();
@@ -922,6 +941,34 @@ public:
 
 	//~ Begin UObject Interface.
 #if WITH_EDITOR
+private:
+	int32 PostEditChangeStackCounter;
+public:
+	//We want to avoid calling post edit change multiple time during import and build process.
+
+	/*
+	 * This function will increment the PostEditChange stack counter.
+	 * It will return the stack counter value. (the value should be >= 1)
+	 */
+	int32 StackPostEditChange();
+	
+	/*
+	 * This function will decrement the stack counter.
+	 * It will return the stack counter value. (the value should be >= 0)
+	 */
+	int32 UnStackPostEditChange();
+
+	int32 GetPostEditChangeStackCounter() { return PostEditChangeStackCounter; }
+	void SetPostEditChangeStackCounter(int32 InPostEditChangeStackCounter)
+	{
+		PostEditChangeStackCounter = InPostEditChangeStackCounter;
+	}
+
+	/**
+	* If derive data cache key do not match, regenerate derived data and re-create any render state based on that.
+	*/
+	void Build();
+
 	virtual void PostEditChangeProperty(FPropertyChangedEvent& PropertyChangedEvent) override;
 
 	virtual void PostEditUndo() override;
@@ -1166,6 +1213,14 @@ private:
 public:
 	/** Get multicast delegate broadcast post to mesh data caching */
 	FOnPostMeshCache& OnPostMeshCached() { return PostMeshCached; }
+
+	/**
+	* Force the creation of a new GUID use to build the derive data cache key.
+	* Next time a build happen the whole skeletal mesh will be rebuild.
+	* Use this when you change stuff not in the skeletal mesh ddc key, like the geometry (import, re-import)
+	* Every big data should not be in the ddc key and should use this function, because its slow to create a key with big data.
+	*/
+	void InvalidateDeriveDataCacheGUID();
 #endif 
 
 private:
@@ -1178,16 +1233,6 @@ private:
 	/** Utility function to help with building the combined socket list */
 	bool IsSocketOnMesh( const FName& InSocketName ) const;
 
-	/**
-	* Flush current render state
-	*/
-	void FlushRenderState();
-
-	/**
-	* Create a new GUID for the source Model data, regenerate derived data and re-create any render state based on that.
-	*/
-	void InvalidateRenderData();
-
 #if WITH_EDITORONLY_DATA
 	/**
 	* In older data, the bEnableShadowCasting flag was stored in LODInfo
@@ -1304,12 +1349,18 @@ public:
 	/** Releases all allocated Skin Weight Profile resources, assumes none are currently in use */
 	void ReleaseSkinWeightProfileResources();
 
+#if WITH_EDITORONLY_DATA
+	/*Transient data use when we postload an old asset to use legacy ddc key, it is turn off so if the user change the asset it go back to the latest ddc code*/
+	bool UseLegacyMeshDerivedDataKey = false;
+#endif
+
 protected:
 	/** Set of skin weight profiles associated with this mesh */
 	UPROPERTY(EditAnywhere, Category = SkinWeights, EditFixedSize, Meta=(NoResetToDefault))
 	TArray<FSkinWeightProfileInfo> SkinWeightProfiles;
 };
 
+
 /**
  * Refresh Physics Asset Change
  * 
diff --git a/Engine/Source/Runtime/Engine/Classes/Engine/StaticMesh.h b/Engine/Source/Runtime/Engine/Classes/Engine/StaticMesh.h
--- a/Engine/Source/Runtime/Engine/Classes/Engine/StaticMesh.h
+++ b/Engine/Source/Runtime/Engine/Classes/Engine/StaticMesh.h
@@ -19,11 +19,8 @@
 #include "Templates/UniquePtr.h"
 #include "StaticMeshResources.h"
 #include "PerPlatformProperties.h"
-
-#include "MeshDescription.h"
-#include "MeshAttributes.h"
-#include "MeshAttributeArray.h"
 #include "RenderAssetUpdate.h"
+#include "MeshTypes.h"
 
 #include "StaticMesh.generated.h"
 
@@ -31,7 +28,9 @@ class FSpeedTreeWind;
 class UAssetUserData;
 class UMaterialInterface;
 class UNavCollisionBase;
+class UStaticMeshDescription;
 class FStaticMeshUpdate;
+struct FMeshDescription;
 struct FMeshDescriptionBulkData;
 struct FStaticMeshLODResources;
 
@@ -245,6 +244,9 @@ struct FStaticMeshSourceModel
 #if WITH_EDITOR
 	/** Serializes bulk data. */
 	void SerializeBulkData(FArchive& Ar, UObject* Owner);
+
+	/** Create a new MeshDescription object */
+	FMeshDescription* CreateMeshDescription();
 #endif
 };
 
@@ -476,69 +478,6 @@ struct FMaterialRemapIndex
 	TArray<int32> MaterialRemap;
 };
 
-struct FStaticMeshDescriptionConstAttributeGetter
-{
-	ENGINE_API FStaticMeshDescriptionConstAttributeGetter(const FMeshDescription* InMeshDescription)
-		: MeshDescription(InMeshDescription)
-	{}
-
-	const FMeshDescription* MeshDescription;
-
-	ENGINE_API TVertexAttributesConstRef<FVector> GetPositions() const { return MeshDescription->VertexAttributes().GetAttributesRef<FVector>(MeshAttribute::Vertex::Position); }
-
-	ENGINE_API TVertexInstanceAttributesConstRef<FVector> GetNormals() const { return MeshDescription->VertexInstanceAttributes().GetAttributesRef<FVector>(MeshAttribute::VertexInstance::Normal); }
-
-	ENGINE_API TVertexInstanceAttributesConstRef<FVector> GetTangents() const { return MeshDescription->VertexInstanceAttributes().GetAttributesRef<FVector>(MeshAttribute::VertexInstance::Tangent); }
-
-	ENGINE_API TVertexInstanceAttributesConstRef<float> GetBinormalSigns() const { return MeshDescription->VertexInstanceAttributes().GetAttributesRef<float>(MeshAttribute::VertexInstance::BinormalSign); }
-
-	ENGINE_API TVertexInstanceAttributesConstRef<FVector4> GetColors() const { return MeshDescription->VertexInstanceAttributes().GetAttributesRef<FVector4>(MeshAttribute::VertexInstance::Color); }
-
-	ENGINE_API TVertexInstanceAttributesConstRef<FVector2D> GetUVs() const { return MeshDescription->VertexInstanceAttributes().GetAttributesRef<FVector2D>(MeshAttribute::VertexInstance::TextureCoordinate); }
-
-	ENGINE_API TEdgeAttributesConstRef<bool> GetEdgeHardnesses() const { return MeshDescription->EdgeAttributes().GetAttributesRef<bool>(MeshAttribute::Edge::IsHard); }
-
-	ENGINE_API TEdgeAttributesConstRef<float> GetEdgeCreaseSharpnesses() const { return MeshDescription->EdgeAttributes().GetAttributesRef<float>(MeshAttribute::Edge::CreaseSharpness); }
-
-	ENGINE_API TPolygonGroupAttributesConstRef<FName> GetPolygonGroupImportedMaterialSlotNames() { return MeshDescription->PolygonGroupAttributes().GetAttributesRef<FName>(MeshAttribute::PolygonGroup::ImportedMaterialSlotName); }
-};
-
-struct FStaticMeshDescriptionAttributeGetter
-{
-	ENGINE_API FStaticMeshDescriptionAttributeGetter(FMeshDescription* InMeshDescription)
-		: MeshDescription(InMeshDescription)
-	{}
-	
-	FMeshDescription* MeshDescription;
-
-	ENGINE_API TVertexAttributesRef<FVector> GetPositions() const { return MeshDescription->VertexAttributes().GetAttributesRef<FVector>(MeshAttribute::Vertex::Position); }
-	ENGINE_API TVertexAttributesConstRef<FVector> GetPositionsConst() const { return MeshDescription->VertexAttributes().GetAttributesRef<FVector>(MeshAttribute::Vertex::Position); }
-
-	ENGINE_API TVertexInstanceAttributesRef<FVector> GetNormals() const { return MeshDescription->VertexInstanceAttributes().GetAttributesRef<FVector>(MeshAttribute::VertexInstance::Normal); }
-	ENGINE_API TVertexInstanceAttributesConstRef<FVector> GetNormalsConst() const { return MeshDescription->VertexInstanceAttributes().GetAttributesRef<FVector>(MeshAttribute::VertexInstance::Normal); }
-
-	ENGINE_API TVertexInstanceAttributesRef<FVector> GetTangents() const { return MeshDescription->VertexInstanceAttributes().GetAttributesRef<FVector>(MeshAttribute::VertexInstance::Tangent); }
-	ENGINE_API TVertexInstanceAttributesConstRef<FVector> GetTangentsConst() const { return MeshDescription->VertexInstanceAttributes().GetAttributesRef<FVector>(MeshAttribute::VertexInstance::Tangent); }
-
-	ENGINE_API TVertexInstanceAttributesRef<float> GetBinormalSigns() const { return MeshDescription->VertexInstanceAttributes().GetAttributesRef<float>(MeshAttribute::VertexInstance::BinormalSign); }
-	ENGINE_API TVertexInstanceAttributesConstRef<float> GetBinormalSignsConst() const { return MeshDescription->VertexInstanceAttributes().GetAttributesRef<float>(MeshAttribute::VertexInstance::BinormalSign); }
-
-	ENGINE_API TVertexInstanceAttributesRef<FVector4> GetColors() const { return MeshDescription->VertexInstanceAttributes().GetAttributesRef<FVector4>(MeshAttribute::VertexInstance::Color); }
-	ENGINE_API TVertexInstanceAttributesConstRef<FVector4> GetColorsConst() const { return MeshDescription->VertexInstanceAttributes().GetAttributesRef<FVector4>(MeshAttribute::VertexInstance::Color); }
-
-	ENGINE_API TVertexInstanceAttributesRef<FVector2D> GetUVs() const { return MeshDescription->VertexInstanceAttributes().GetAttributesRef<FVector2D>(MeshAttribute::VertexInstance::TextureCoordinate); }
-	ENGINE_API TVertexInstanceAttributesConstRef<FVector2D> GetUVsConst() const { return MeshDescription->VertexInstanceAttributes().GetAttributesRef<FVector2D>(MeshAttribute::VertexInstance::TextureCoordinate); }
-
-	ENGINE_API TEdgeAttributesRef<bool> GetEdgeHardnesses() const { return MeshDescription->EdgeAttributes().GetAttributesRef<bool>(MeshAttribute::Edge::IsHard); }
-	ENGINE_API TEdgeAttributesConstRef<bool> GetEdgeHardnessesConst() const { return MeshDescription->EdgeAttributes().GetAttributesRef<bool>(MeshAttribute::Edge::IsHard); }
-	
-	ENGINE_API TEdgeAttributesRef<float> GetEdgeCreaseSharpnesses() const { return MeshDescription->EdgeAttributes().GetAttributesRef<float>(MeshAttribute::Edge::CreaseSharpness); }
-	ENGINE_API TEdgeAttributesConstRef<float> GetEdgeCreaseSharpnessesConst() const { return MeshDescription->EdgeAttributes().GetAttributesRef<float>(MeshAttribute::Edge::CreaseSharpness); }
-
-	ENGINE_API TPolygonGroupAttributesRef<FName> GetPolygonGroupImportedMaterialSlotNames() { return MeshDescription->PolygonGroupAttributes().GetAttributesRef<FName>(MeshAttribute::PolygonGroup::ImportedMaterialSlotName); }
-	ENGINE_API TPolygonGroupAttributesConstRef<FName> GetPolygonGroupImportedMaterialSlotNamesConst() { return MeshDescription->PolygonGroupAttributes().GetAttributesRef<FName>(MeshAttribute::PolygonGroup::ImportedMaterialSlotName); }
-};
-
 
 /**
  * A StaticMesh is a piece of geometry that consists of a static set of polygons.
@@ -718,6 +657,11 @@ class UStaticMesh : public UStreamableRenderAsset, public IInterface_CollisionDa
 	UPROPERTY(EditAnywhere, AdvancedDisplay, Category = StaticMesh)
 	uint8 bSupportUniformlyDistributedSampling : 1;
 
+	/**
+	 * If true, StaticMesh has been built at runtime
+	 */
+	uint8 bIsBuiltAtRuntime : 1;
+
 protected:
 	/** Tracks whether InitResources has been called, and rendering resources are initialized. */
 	uint8 bRenderingResourcesInitialized:1;
@@ -825,7 +769,7 @@ protected:
 public:
 	/** The editable mesh representation of this static mesh */
 	// @todo: Maybe we don't want this visible in the details panel in the end; for now, this might aid debugging.
-	UPROPERTY(Instanced, VisibleAnywhere, Category = EditableMesh)
+	UPROPERTY(Instanced)
 	class UObject* EditableMesh;
 
 #if WITH_EDITORONLY_DATA
@@ -836,6 +780,7 @@ public:
 	/**
 	 * Registers the mesh attributes required by the mesh description for a static mesh.
 	 */
+	UE_DEPRECATED(4.25, "Please use FStaticMeshAttributes::Register to do this.")
 	ENGINE_API static void RegisterMeshAttributes( FMeshDescription& MeshDescription );
 
 #if WITH_EDITORONLY_DATA
@@ -851,18 +796,6 @@ public:
 	ENGINE_API void ClearMeshDescription(int32 LodIndex);
 	ENGINE_API void ClearMeshDescriptions();
 
-	UE_DEPRECATED(4.22, "Please use GetMeshDescription().")
-	FMeshDescription* GetOriginalMeshDescription(int32 LodIndex) const { return GetMeshDescription(LodIndex); }
-	
-	UE_DEPRECATED(4.22, "Please use CreateMeshDescription().")
-	FMeshDescription* CreateOriginalMeshDescription(int32 LodIndex) { return CreateMeshDescription(LodIndex); }
-
-	UE_DEPRECATED(4.22, "Please use CommitMeshDescription().")
-	void CommitOriginalMeshDescription(int32 LodIndex) { CommitMeshDescription(LodIndex); }
-
-	UE_DEPRECATED(4.22, "Please use ClearMeshDescription().")
-	void ClearOriginalMeshDescription(int32 LodIndex) { ClearMeshDescription(LodIndex); }
-
 	/**
 	 * Internal function use to make sure all imported material slot name are unique and non empty.
 	 */
@@ -902,6 +835,27 @@ public:
 
 #endif
 
+	/** Create an empty StaticMeshDescription object, to describe a static mesh at runtime */
+	UFUNCTION(BlueprintCallable, Category="StaticMesh")
+	static ENGINE_API UStaticMeshDescription* CreateStaticMeshDescription(UObject* Outer = nullptr);
+
+	/** Builds static mesh LODs from the array of StaticMeshDescriptions passed in */
+	UFUNCTION(BlueprintCallable, Category="StaticMesh")
+	ENGINE_API void BuildFromStaticMeshDescriptions(const TArray<UStaticMeshDescription*>& StaticMeshDescriptions);
+
+	/**
+	 * Builds static mesh render buffers from a list of MeshDescriptions, one per LOD.
+	 */
+	ENGINE_API bool BuildFromMeshDescriptions(const TArray<const FMeshDescription*>& MeshDescriptions);
+
+	/**
+	 * Builds simple collisions at runtime
+	 */
+	ENGINE_API bool BuildSimpleCollision();
+
+	/** Builds a LOD resource from a MeshDescription */
+	void BuildFromMeshDescription(const FMeshDescription& MeshDescription, FStaticMeshLODResources& LODResources);
+
 	/**
 	 * Returns the number of UV channels for the given LOD of a StaticMesh.
 	 * @param	LODIndex			Index of the StaticMesh LOD.
@@ -932,7 +886,10 @@ public:
 
 	//SourceModels API
 	ENGINE_API FStaticMeshSourceModel& AddSourceModel();
+
+	UFUNCTION(BlueprintCallable, Category="StaticMesh")
 	ENGINE_API void SetNumSourceModels(int32 Num);
+
 	ENGINE_API void RemoveSourceModel(int32 Index);
 	PRAGMA_DISABLE_DEPRECATION_WARNINGS
 	ENGINE_API TArray<FStaticMeshSourceModel>& GetSourceModels() { return SourceModels; }
@@ -1053,7 +1010,7 @@ public:
 	/**
 	 * Returns the number of LODs used by the mesh.
 	 */
-	UFUNCTION(BlueprintCallable, Category = "StaticMesh", meta=(ScriptName="GetNumLods"))
+	UFUNCTION(BlueprintPure, Category = "StaticMesh", meta=(ScriptName="GetNumLods"))
 	ENGINE_API int32 GetNumLODs() const;
 
 	/**
@@ -1070,11 +1027,11 @@ public:
 	ENGINE_API FBoxSphereBounds GetBounds() const;
 
 	/** Returns the bounding box, in local space including bounds extension(s), of the StaticMesh asset */
-	UFUNCTION(BlueprintCallable, Category="StaticMesh")
+	UFUNCTION(BlueprintPure, Category="StaticMesh")
 	ENGINE_API FBox GetBoundingBox() const;
 
 	/** Returns number of Sections that this StaticMesh has, in the supplied LOD (LOD 0 is the highest) */
-	UFUNCTION(BlueprintCallable, Category = "StaticMesh")
+	UFUNCTION(BlueprintPure, Category = "StaticMesh")
 	ENGINE_API int32 GetNumSections(int32 InLOD) const;
 
 	/**
@@ -1082,15 +1039,21 @@ public:
 	 *
 	 * @return Requested material
 	 */
-	UFUNCTION(BlueprintCallable, Category = "StaticMesh")
+	UFUNCTION(BlueprintPure, Category = "StaticMesh")
 	ENGINE_API UMaterialInterface* GetMaterial(int32 MaterialIndex) const;
 
 	/**
-	* Gets a Material index given a slot name
-	*
-	* @return Requested material
-	*/
+	 * Adds a new material and return its slot name
+	 */
 	UFUNCTION(BlueprintCallable, Category = "StaticMesh")
+	ENGINE_API FName AddMaterial(UMaterialInterface* Material);
+
+	/**
+	 * Gets a Material index given a slot name
+	 *
+	 * @return Requested material
+	 */
+	UFUNCTION(BlueprintPure, Category = "StaticMesh")
 	ENGINE_API int32 GetMaterialIndex(FName MaterialSlotName) const;
 
 	ENGINE_API int32 GetMaterialIndexFromImportedMaterialSlotName(FName ImportedMaterialSlotName) const;
@@ -1175,7 +1138,7 @@ public:
 	 *	Find a socket object in this StaticMesh by name.
 	 *	Entering NAME_None will return NULL. If there are multiple sockets with the same name, will return the first one.
 	 */
-	UFUNCTION(BlueprintCallable, Category = "StaticMesh")
+	UFUNCTION(BlueprintPure, Category = "StaticMesh")
 	ENGINE_API class UStaticMeshSocket* FindSocket(FName InSocketName) const;
 
 	/**
diff --git a/Engine/Source/Runtime/Engine/Engine.Build.cs b/Engine/Source/Runtime/Engine/Engine.Build.cs
--- a/Engine/Source/Runtime/Engine/Engine.Build.cs
+++ b/Engine/Source/Runtime/Engine/Engine.Build.cs
@@ -84,6 +84,7 @@ public class Engine : ModuleRules
 				"PacketHandler",
 				"AudioPlatformConfiguration",
 				"MeshDescription",
+				"StaticMeshDescription",
 				"PakFile",
                 "NetworkReplayStreaming",
 				"PhysicsCore",
diff --git a/Engine/Source/Runtime/Engine/Private/Actor.cpp b/Engine/Source/Runtime/Engine/Private/Actor.cpp
--- a/Engine/Source/Runtime/Engine/Private/Actor.cpp
+++ b/Engine/Source/Runtime/Engine/Private/Actor.cpp
@@ -98,7 +98,7 @@ void AActor::InitializeDefaults()
 
 	CustomTimeDilation = 1.0f;
 
-	Role = ROLE_Authority;
+	SetRole(ROLE_Authority);
 	RemoteRole = ROLE_None;
 	bReplicates = false;
 	NetPriority = 1.0f;
@@ -124,7 +124,7 @@ void AActor::InitializeDefaults()
 	bActorEnableCollision = true;
 	bActorSeamlessTraveled = false;
 	bBlockInput = false;
-	bCanBeDamaged = true;
+	SetCanBeDamaged(true);
 	bFindCameraComponentWhenViewTarget = true;
 	bAllowReceiveTickEventOnDedicatedServer = true;
 	bRelevantForNetworkReplays = true;
@@ -724,9 +724,9 @@ void AActor::PostLoadSubobjects(FObjectInstancingGraph* OuterInstanceGraph)
 	FVector OldScale;
 	if (bHadRoot)
 	{
-		OldRotation = OldRoot->RelativeRotation;
-		OldTranslation = OldRoot->RelativeLocation;
-		OldScale = OldRoot->RelativeScale3D;
+		OldRotation = OldRoot->GetRelativeRotation();
+		OldTranslation = OldRoot->GetRelativeLocation();
+		OldScale = OldRoot->GetRelativeScale3D();
 	}
 
 	Super::PostLoadSubobjects(OuterInstanceGraph);
@@ -736,9 +736,9 @@ void AActor::PostLoadSubobjects(FObjectInstancingGraph* OuterInstanceGraph)
 	if (RootComponent && bHadRoot && OldRoot != RootComponent && OldRoot->IsIn(this))
 	{
 		UE_LOG(LogActor, Log, TEXT("Root component has changed, relocating new root component to old position %s->%s"), *OldRoot->GetFullName(), *GetRootComponent()->GetFullName());
-		GetRootComponent()->RelativeRotation = OldRotation;
-		GetRootComponent()->RelativeLocation = OldTranslation;
-		GetRootComponent()->RelativeScale3D = OldScale;
+		GetRootComponent()->SetRelativeRotation_Direct(OldRotation);
+		GetRootComponent()->SetRelativeLocation_Direct(OldTranslation);
+		GetRootComponent()->SetRelativeScale3D_Direct(OldScale);
 		
 		// Migrate any attachment to the new root
 		if (OldRoot->GetAttachParent())
@@ -758,9 +758,9 @@ void AActor::PostLoadSubobjects(FObjectInstancingGraph* OuterInstanceGraph)
 		}
 
 		// Reset the transform on the old component
-		OldRoot->RelativeRotation = FRotator::ZeroRotator;
-		OldRoot->RelativeLocation = FVector::ZeroVector;
-		OldRoot->RelativeScale3D = FVector(1.0f, 1.0f, 1.0f);
+		OldRoot->SetRelativeRotation_Direct(FRotator::ZeroRotator);
+		OldRoot->SetRelativeLocation_Direct(FVector::ZeroVector);
+		OldRoot->SetRelativeScale3D_Direct(FVector(1.0f, 1.0f, 1.0f));
 	}
 }
 
@@ -968,7 +968,7 @@ UNetConnection* AActor::GetNetConnection() const
 UPlayer* AActor::GetNetOwningPlayer()
 {
 	// We can only replicate RPCs to the owning player
-	if (Role == ROLE_Authority)
+	if (GetLocalRole() == ROLE_Authority)
 	{
 		if (Owner)
 		{
@@ -1048,7 +1048,7 @@ bool AActor::ShouldTickIfViewportsOnly() const
 	return false;
 }
 
-void AActor::PreReplication( IRepChangedPropertyTracker & ChangedPropertyTracker )
+void AActor::PreReplication(IRepChangedPropertyTracker & ChangedPropertyTracker)
 {
 	// Attachment replication gets filled in by GatherCurrentMovement(), but in the case of a detached root we need to trigger remote detachment.
 	AttachmentReplication.AttachParent = nullptr;
@@ -1056,10 +1056,10 @@ void AActor::PreReplication( IRepChangedPropertyTracker & ChangedPropertyTracker
 
 	GatherCurrentMovement();
 
-	DOREPLIFETIME_ACTIVE_OVERRIDE( AActor, ReplicatedMovement, bReplicateMovement );
+	DOREPLIFETIME_ACTIVE_OVERRIDE(AActor, ReplicatedMovement, bReplicateMovement);
 
 	// Don't need to replicate AttachmentReplication if the root component replicates, because it already handles it.
-	DOREPLIFETIME_ACTIVE_OVERRIDE( AActor, AttachmentReplication, RootComponent && !RootComponent->GetIsReplicated() );
+	DOREPLIFETIME_ACTIVE_OVERRIDE(AActor, AttachmentReplication, RootComponent && !RootComponent->GetIsReplicated());
 
 	UBlueprintGeneratedClass* BPClass = Cast<UBlueprintGeneratedClass>(GetClass());
 	if (BPClass != nullptr)
@@ -1079,7 +1079,8 @@ void AActor::CallPreReplication(UNetDriver* NetDriver)
 
 	// PreReplication is only called on the server, except when we're recording a Client Replay.
 	// In that case we call PreReplication on the locally controlled Character as well.
-	if ((Role == ROLE_Authority) || ((Role == ROLE_AutonomousProxy) && GetWorld()->IsRecordingClientReplay()))
+	const ENetRole LocalRole = GetLocalRole();
+	if ((LocalRole == ROLE_Authority) || ((LocalRole == ROLE_AutonomousProxy) && GetWorld()->IsRecordingClientReplay()))
 	{
 		PreReplication(*ActorChangedPropertyTracker);
 	}
@@ -1295,7 +1296,7 @@ bool AActor::CheckStillInWorld()
 	}
 
 	// Only authority or non-networked actors should be destroyed, otherwise misprediction can destroy something the server is intending to keep alive.
-	if (!(HasAuthority() || Role == ROLE_None))
+	if (!(HasAuthority() || GetLocalRole() == ROLE_None))
 	{
 		return true;
 	}
@@ -1591,6 +1592,7 @@ void AActor::SetOwner(AActor* NewOwner)
 		}
 
 		Owner = NewOwner;
+		// MARK_PROPERTY_DIRTY_BY_NAME(AActor, Owner, this);
 
 		if (Owner != nullptr)
 		{
@@ -1680,9 +1682,9 @@ void AActor::OnRep_AttachmentReplication()
 
 			if (AttachParentComponent)
 			{
-				RootComponent->RelativeLocation = AttachmentReplication.LocationOffset;
-				RootComponent->RelativeRotation = AttachmentReplication.RotationOffset;
-				RootComponent->RelativeScale3D = AttachmentReplication.RelativeScale3D;
+				RootComponent->SetRelativeLocation_Direct(AttachmentReplication.LocationOffset);
+				RootComponent->SetRelativeRotation_Direct(AttachmentReplication.RotationOffset);
+				RootComponent->SetRelativeScale3D_Direct(AttachmentReplication.RelativeScale3D);
 
 				// If we're already attached to the correct Parent and Socket, then the update must be position only.
 				// AttachToComponent would early out in this case.
@@ -1764,10 +1766,11 @@ void AActor::DetachRootComponentFromParent(bool bMaintainWorldPosition)
 	{
 PRAGMA_DISABLE_DEPRECATION_WARNINGS
 		RootComponent->DetachFromParent(bMaintainWorldPosition);
-PRAGMA_ENABLE_DEPRECATION_WARNINGS
+PRAGMA_ENABLE_DEPRECATION_WARNINGS		
 
 		// Clear AttachmentReplication struct
 		AttachmentReplication = FRepAttachment();
+		// MARK_PROPERTY_DIRTY_FROM_NAME(AActor, AttachmentReplication, this);
 	}
 }
 
@@ -1960,7 +1963,7 @@ bool AActor::IsRelevancyOwnerFor(const AActor* ReplicatedActor, const AActor* Ac
 
 void AActor::ForceNetUpdate()
 {
-	if (Role == ROLE_Authority)
+	if (GetLocalRole() == ROLE_Authority)
 	{
 		// ForceNetUpdate on the game net driver only if we are the authority...
 		UNetDriver* NetDriver = GetNetDriver();
@@ -2221,9 +2224,8 @@ void AActor::TearOff()
 
 	if (NetMode == NM_ListenServer || NetMode == NM_DedicatedServer)
 	{
-		PRAGMA_DISABLE_DEPRECATION_WARNINGS
 		bTearOff = true;
-		PRAGMA_ENABLE_DEPRECATION_WARNINGS
+		// MARK_PROPERTY_DIRTY_FROM_NAME(AActor, bTearOff, this);
 		if (UNetDriver* NetDriver = GetNetDriver())
 		{
 			NetDriver->NotifyActorTearOff(this);
@@ -2241,7 +2243,7 @@ void AActor::Reset()
 void AActor::FellOutOfWorld(const UDamageType& dmgType)
 {
 	// Only authority or non-networked actors should be destroyed, otherwise misprediction can destroy something the server is intending to keep alive.
-	if (HasAuthority() || Role == ROLE_None)
+	if (HasAuthority() || GetLocalRole() == ROLE_None)
 	{
 		DisableComponentsSimulatePhysics();
 		SetActorHiddenInGame(true);
@@ -2252,7 +2254,7 @@ void AActor::FellOutOfWorld(const UDamageType& dmgType)
 
 void AActor::MakeNoise(float Loudness, APawn* NoiseInstigator, FVector NoiseLocation, float MaxRange, FName Tag)
 {
-	NoiseInstigator = NoiseInstigator ? NoiseInstigator : Instigator;
+	NoiseInstigator = NoiseInstigator ? NoiseInstigator : GetInstigator();
 	if ((GetNetMode() != NM_Client) && NoiseInstigator)
 	{
 		AActor::MakeNoiseDelegate.Execute(this, Loudness, NoiseInstigator
@@ -2451,7 +2453,7 @@ void AActor::DisplayDebug(UCanvas* Canvas, const FDebugDisplayInfo& DebugDisplay
 		if( GetNetMode() != NM_Standalone )
 		{
 			// networking attributes
-			T = FString::Printf(TEXT("ROLE: %i RemoteRole: %i NetNode: %i"), (int32)Role, (int32)RemoteRole, (int32)GetNetMode());
+			T = FString::Printf(TEXT("ROLE: %i RemoteRole: %i NetNode: %i"), (int32)GetLocalRole(), (int32)RemoteRole, (int32)GetNetMode());
 
 			if( GetTearOff() )
 			{
@@ -2502,7 +2504,7 @@ void AActor::DisplayDebug(UCanvas* Canvas, const FDebugDisplayInfo& DebugDisplay
 		DisplayDebugManager.DrawString(T);
 	}
 	DisplayDebugManager.DrawString(FString::Printf(TEXT(" Instigator: %s Owner: %s"), 
-		(Instigator ? *Instigator->GetName() : TEXT("None")), (Owner ? *Owner->GetName() : TEXT("None"))));
+		*GetNameSafe(GetInstigator()), *GetNameSafe(Owner)));
 
 	static FName NAME_Animation(TEXT("Animation"));
 	static FName NAME_Bones = FName(TEXT("Bones"));
@@ -2564,7 +2566,7 @@ void AActor::CalcCamera(float DeltaTime, FMinimalViewInfo& OutResult)
 
 		for (UCameraComponent* CameraComponent : Cameras)
 		{
-			if (CameraComponent->bIsActive)
+			if (CameraComponent->IsActive())
 			{
 				CameraComponent->GetCameraView(DeltaTime, OutResult);
 				return;
@@ -2585,7 +2587,7 @@ bool AActor::HasActiveCameraComponent() const
 			const UCameraComponent* CameraComponent = Cast<const UCameraComponent>(Component);
 			if (CameraComponent)
 			{
-				if (CameraComponent->bIsActive)
+				if (CameraComponent->IsActive())
 				{
 					return true;
 				}
@@ -2605,7 +2607,7 @@ bool AActor::HasActivePawnControlCameraComponent() const
 			const UCameraComponent* CameraComponent = Cast<const UCameraComponent>(Component);
 			if (CameraComponent)
 			{
-				if (CameraComponent->bIsActive && CameraComponent->bUsePawnControlRotation)
+				if (CameraComponent->IsActive() && CameraComponent->bUsePawnControlRotation)
 				{
 					return true;
 				}
@@ -2830,6 +2832,24 @@ TArray<UActorComponent*> AActor::GetComponentsByTag(TSubclassOf<UActorComponent>
 	return MoveTemp(ComponentsByTag);
 }
 
+TArray<UActorComponent*> AActor::GetComponentsByInterface(TSubclassOf<UInterface> Interface) const
+{
+	TArray<UActorComponent*> Components;
+
+	if (Interface)
+	{
+		for (UActorComponent* Component : GetComponents())
+		{
+			if (Component && Component->GetClass()->ImplementsInterface(Interface))
+			{
+				Components.Add(Component);
+			}
+		}
+	}
+
+	return Components;
+}
+
 void AActor::DisableComponentsSimulatePhysics()
 {
 	for (UActorComponent* Component : GetComponents())
@@ -2950,14 +2970,14 @@ void AActor::PostSpawnInitialize(FTransform const& UserSpawnTransform, AActor* I
 	CreationTime = (World ? World->GetTimeSeconds() : 0.f);
 
 	// Set network role.
-	check(Role == ROLE_Authority);
+	check(GetLocalRole() == ROLE_Authority);
 	ExchangeNetRoles(bRemoteOwned);
 
 	// Set owner.
 	SetOwner(InOwner);
 
 	// Set instigator
-	Instigator = InInstigator;
+	SetInstigator(InInstigator);
 
 	// Set the actor's world transform if it has a native rootcomponent.
 	USceneComponent* const SceneRootComponent = FixupNativeActorComponents(this);
@@ -3000,7 +3020,7 @@ void AActor::PostSpawnInitialize(FTransform const& UserSpawnTransform, AActor* I
 		{
 			// In the "normal" case we do respect any non-default transform value that the root component may have received from the archetype
 			// that's owned by the native CDO, so the final transform might not always necessarily equate to the passed-in UserSpawnTransform.
-			const FTransform RootTransform(SceneRootComponent->RelativeRotation, SceneRootComponent->RelativeLocation, SceneRootComponent->RelativeScale3D);
+			const FTransform RootTransform(SceneRootComponent->GetRelativeRotation(), SceneRootComponent->GetRelativeLocation(), SceneRootComponent->GetRelativeScale3D());
 			const FTransform FinalRootComponentTransform = RootTransform * UserSpawnTransform;
 			SceneRootComponent->SetWorldTransform(FinalRootComponentTransform);
 		}
@@ -3216,7 +3236,7 @@ void AActor::PostActorConstruction()
 
 void AActor::SetReplicates(bool bInReplicates)
 { 
-	if (Role == ROLE_Authority)
+	if (GetLocalRole() == ROLE_Authority)
 	{
 		const bool bNewlyReplicates = (bReplicates == false && bInReplicates == true);
 	
@@ -3238,6 +3258,8 @@ void AActor::SetReplicates(bool bInReplicates)
 				ForcePropertyCompare();
 			}
 		}
+
+		// MARK_PROPERTY_DIRTY_FROM_NAME(AActor, RemoteRole, this);
 	}
 	else
 	{
@@ -3247,7 +3269,7 @@ void AActor::SetReplicates(bool bInReplicates)
 
 void AActor::SetReplicateMovement(bool bInReplicateMovement)
 {
-	bReplicateMovement = bInReplicateMovement;
+	SetReplicatingMovement(bInReplicateMovement);
 }
 
 void AActor::SetAutonomousProxy(const bool bInAutonomousProxy, const bool bAllowForcePropertyCompare)
@@ -3259,7 +3281,9 @@ void AActor::SetAutonomousProxy(const bool bInAutonomousProxy, const bool bAllow
 		RemoteRole = (bInAutonomousProxy ? ROLE_AutonomousProxy : ROLE_SimulatedProxy);
 
 		if (bAllowForcePropertyCompare && RemoteRole != OldRemoteRole)
-		{ 
+		{
+			// MARK_PROPERTY_DIRTY_FROM_NAME(AActor, RemoteRole, this);
+
 			// We have to do this so the role change above will replicate (turn off shadow state sharing for a frame)
 			// This is because RemoteRole is special since it will change between connections, so we have to special case
 			ForcePropertyCompare();
@@ -3274,6 +3298,8 @@ void AActor::SetAutonomousProxy(const bool bInAutonomousProxy, const bool bAllow
 void AActor::CopyRemoteRoleFrom(const AActor* CopyFromActor)
 {
 	RemoteRole = CopyFromActor->GetRemoteRole();
+	// MARK_PROPERTY_DIRTY_FROM_NAME(AActor, RemoteRole, this);
+
 	if (RemoteRole != ROLE_None)
 	{
 		GetWorld()->AddNetworkActor(this);
@@ -3318,6 +3344,10 @@ void AActor::ExchangeNetRoles(bool bRemoteOwned)
 void AActor::SwapRoles()
 {
 	Swap(Role, RemoteRole);
+
+	// MARK_PROPERTY_DIRTY_FROM_NAME(AActor, RemoteRole, this);
+	// MARK_PROPERTY_DIRTY_FROM_NAME(AActor, Role, this);
+
 	ForcePropertyCompare();
 }
 
@@ -3811,16 +3841,16 @@ FVector AActor::GetActorRelativeScale3D() const
 {
 	if (RootComponent)
 	{
-		return RootComponent->RelativeScale3D;
+		return RootComponent->GetRelativeScale3D();
 	}
 	return FVector(1,1,1);
 }
 
 void AActor::SetActorHiddenInGame( bool bNewHidden )
 {
-	if (bHidden != bNewHidden)
+	if (IsHidden() != bNewHidden)
 	{
-		bHidden = bNewHidden;
+		SetHidden(bNewHidden);
 		MarkComponentsRenderStateDirty();
 	}
 }
@@ -3988,8 +4018,10 @@ int32 AActor::GetFunctionCallspace( UFunction* Function, FFrame* Stack )
 		return FunctionCallspace::Local;
 	}
 
+	const ENetRole LocalRole = GetLocalRole();
+
 	// If we are on a client and function is 'skip on client', absorb it
-	FunctionCallspace::Type Callspace = (Role < ROLE_Authority) && Function->HasAllFunctionFlags(FUNC_BlueprintAuthorityOnly) ? FunctionCallspace::Absorbed : FunctionCallspace::Local;
+	FunctionCallspace::Type Callspace = (LocalRole < ROLE_Authority) && Function->HasAllFunctionFlags(FUNC_BlueprintAuthorityOnly) ? FunctionCallspace::Absorbed : FunctionCallspace::Local;
 	
 	if (IsPendingKill())
 	{
@@ -4024,7 +4056,7 @@ int32 AActor::GetFunctionCallspace( UFunction* Function, FFrame* Stack )
 	// Quick reject 2. Has to be a network game to continue
 	if (NetMode == NM_Standalone)
 	{
-		if (Role < ROLE_Authority && (Function->FunctionFlags & FUNC_NetServer))
+		if (LocalRole < ROLE_Authority && (Function->FunctionFlags & FUNC_NetServer))
 		{
 			// Don't let clients call server functions (in edge cases where NetMode is standalone (net driver is null))
 			DEBUG_CALLSPACE(TEXT("GetFunctionCallspace No Authority Server Call Absorbed: %s"), *Function->GetName());
@@ -4095,7 +4127,7 @@ int32 AActor::GetFunctionCallspace( UFunction* Function, FFrame* Stack )
 	}
 
 	// Check if the actor can potentially call remote functions	
-	if (Role == ROLE_Authority)
+	if (LocalRole == ROLE_Authority)
 	{
 		UNetConnection* NetConnection = GetNetConnection();
 		if (NetConnection == nullptr)
@@ -4649,7 +4681,7 @@ void AActor::SetLifeSpan( float InLifespan )
 	// Store the new value
 	InitialLifeSpan = InLifespan;
 	// Initialize a timer for the actors lifespan if there is one. Otherwise clear any existing timer
-	if ((Role == ROLE_Authority || GetTearOff()) && !IsPendingKill())
+	if ((GetLocalRole() == ROLE_Authority || GetTearOff()) && !IsPendingKill())
 	{
 		if( InLifespan > 0.0f)
 		{
@@ -4938,4 +4970,45 @@ void AActor::SetLODParent(UPrimitiveComponent* InLODParent, float InParentDrawDi
 	}
 }
 
+void AActor::SetHidden(bool bInHidden)
+{
+	bHidden = bInHidden;
+	// MARK_PROPERTY_DIRTY_FROM_NAME(AActor, bHidden, this);
+}
+
+void AActor::SetReplicatingMovement(bool bInReplicateMovement)
+{
+	bReplicateMovement = bInReplicateMovement;
+	// MARK_PROPERTY_DIRTY_FROM_NAME(AActor, bReplicateMovement, this);
+}
+
+void AActor::SetCanBeDamaged(bool bInCanBeDamaged)
+{
+	bCanBeDamaged = bInCanBeDamaged;
+	// MARK_PROPERTY_DIRTY_FROM_NAME(AActor, bCanBeDamaged, this);
+}
+
+void AActor::SetRole(ENetRole InRole)
+{
+	Role = InRole;
+	// MARK_PROPERTY_DIRTY_FROM_NAME(AActor, Role, this);
+}
+
+FRepMovement& AActor::GetReplicatedMovement_Mutable()
+{
+	// MARK_PROPERTY_DIRTY_FROM_NAME(AActor, ReplicatedMovement, this);
+	return ReplicatedMovement;
+}
+
+void AActor::SetReplicatedMovement(const FRepMovement& InReplicatedMovement)
+{
+	GetReplicatedMovement_Mutable() = InReplicatedMovement;
+}
+
+void AActor::SetInstigator(APawn* InInstigator)
+{
+	Instigator = InInstigator;
+	// MARK_PROPERTY_DIRTY_FROM_NAME(AActor, Instigator, this);
+}
+
 #undef LOCTEXT_NAMESPACE
diff --git a/Engine/Source/Runtime/Engine/Private/AssetManager.cpp b/Engine/Source/Runtime/Engine/Private/AssetManager.cpp
--- a/Engine/Source/Runtime/Engine/Private/AssetManager.cpp
+++ b/Engine/Source/Runtime/Engine/Private/AssetManager.cpp
@@ -3834,7 +3834,7 @@ static FAutoConsoleCommandWithWorldAndArgs CVarLoadPrimaryAssetsWithType(
 // Cheat command to unload all assets of a given type
 static FAutoConsoleCommandWithWorldAndArgs CVarUnloadPrimaryAssetsWithType(
 	TEXT("AssetManager.UnloadPrimaryAssetsWithType"),
-	TEXT("Loads all assets of a given type"),
+	TEXT("Unloads all assets of a given type"),
 	FConsoleCommandWithWorldAndArgsDelegate::CreateStatic(
 		[](const TArray<FString>& Params, UWorld* World)
 {
@@ -3845,13 +3845,13 @@ static FAutoConsoleCommandWithWorldAndArgs CVarUnloadPrimaryAssetsWithType(
 
 	for (const FString& Param : Params)
 	{
-		const FPrimaryAssetType TypeToLoad(*Param);
+		const FPrimaryAssetType TypeToUnload(*Param);
 
 		FPrimaryAssetTypeInfo Info;
-		if (UAssetManager::Get().GetPrimaryAssetTypeInfo(TypeToLoad, /*out*/ Info))
+		if (UAssetManager::Get().GetPrimaryAssetTypeInfo(TypeToUnload, /*out*/ Info))
 		{
-			UE_LOG(LogAssetManager, Log, TEXT("UnloadPrimaryAssetsWithType(%s)"), *Param);
-			UAssetManager::Get().LoadPrimaryAssetsWithType(TypeToLoad);
+			int32 NumUnloaded = UAssetManager::Get().UnloadPrimaryAssetsWithType(TypeToUnload);
+			UE_LOG(LogAssetManager, Log, TEXT("UnloadPrimaryAssetsWithType(%s): Unloaded %d assets"), *Param, NumUnloaded);
 		}
 		else
 		{
diff --git a/Engine/Source/Runtime/Engine/Private/AudioDerivedData.cpp b/Engine/Source/Runtime/Engine/Private/AudioDerivedData.cpp
--- a/Engine/Source/Runtime/Engine/Private/AudioDerivedData.cpp
+++ b/Engine/Source/Runtime/Engine/Private/AudioDerivedData.cpp
@@ -140,6 +140,28 @@ static void GetStreamedAudioDerivedDataKey(
 	GetStreamedAudioDerivedDataKeyFromSuffix(KeySuffix, OutKey);
 }
 
+static ITargetPlatform* GetRunningTargetPlatform(ITargetPlatformManagerModule* TPM)
+{
+	ITargetPlatform* CurrentPlatform = NULL;
+	const TArray<ITargetPlatform*>& Platforms = TPM->GetActiveTargetPlatforms();
+
+	check(Platforms.Num());
+
+	CurrentPlatform = Platforms[0];
+
+	for (int32 Index = 1; Index < Platforms.Num(); Index++)
+	{
+		if (Platforms[Index]->IsRunningPlatform())
+		{
+			CurrentPlatform = Platforms[Index];
+			break;
+		}
+	}
+
+	check(CurrentPlatform != NULL);
+	return CurrentPlatform;
+}
+
 /**
  * Gets Wave format for a SoundWave on the current running platform
  * @param SoundWave - The SoundWave to get format for.
@@ -150,23 +172,8 @@ static FName GetWaveFormatForRunningPlatform(USoundWave& SoundWave)
 	ITargetPlatformManagerModule* TPM = GetTargetPlatformManager();
 	if (TPM)
 	{
-		ITargetPlatform* CurrentPlatform = NULL;
-		const TArray<ITargetPlatform*>& Platforms = TPM->GetActiveTargetPlatforms();
-
-		check(Platforms.Num());
+		ITargetPlatform* CurrentPlatform = GetRunningTargetPlatform(TPM);
 
-		CurrentPlatform = Platforms[0];
-
-		for (int32 Index = 1; Index < Platforms.Num(); Index++)
-		{
-			if (Platforms[Index]->IsRunningPlatform())
-			{
-				CurrentPlatform = Platforms[Index];
-				break;
-			}
-		}
-
-		check(CurrentPlatform != NULL);
 
 		return CurrentPlatform->GetWaveFormat(&SoundWave);
 	}
@@ -174,6 +181,22 @@ static FName GetWaveFormatForRunningPlatform(USoundWave& SoundWave)
 	return NAME_None;
 }
 
+static const FPlatformAudioCookOverrides* GetCookOverridesForRunningPlatform()
+{
+	ITargetPlatformManagerModule* TPM = GetTargetPlatformManager();
+	if (TPM)
+	{
+		ITargetPlatform* CurrentPlatform = GetRunningTargetPlatform(TPM);
+
+
+		return CurrentPlatform->GetAudioCompressionSettings();
+	}
+	else
+	{
+		return nullptr;
+	}
+}
+
 /**
  * Stores derived data in the DDC.
  * After this returns, all bulk data from streaming chunks will be sent separately to the DDC and the BulkData for those chunks removed.
@@ -431,7 +454,7 @@ class FStreamedAudioCacheDerivedDataWorker : public FNonAbandonableTask
 			}
 			else
 			{
-				UE_LOG(LogAudio, Warning, TEXT("Failed to retrieve compressed data for format %s and soundwave %s"),
+				UE_LOG(LogAudio, Display, TEXT("Failed to retrieve compressed data for format %s and soundwave %s."),
 					   *AudioFormatName.GetPlainNameString(),
 					   *SoundWave.GetPathName()
 					);
@@ -753,6 +776,12 @@ int32 FStreamedAudioPlatformData::DeserializeChunkFromDDC(TArray<uint8> TempData
 
 int32 FStreamedAudioPlatformData::GetChunkFromDDC(int32 ChunkIndex, uint8** OutChunkData, bool bMakeSureChunkIsLoaded /* = false */)
 {
+	if (Chunks.Num() == 0)
+	{
+		UE_LOG(LogAudioDerivedData, Warning, TEXT("No streamed audio chunks found!"));
+		return 0;
+	}
+
 	// if bMakeSureChunkIsLoaded is true, we don't actually know the size of the chunk's bulk data,
 	// so it will need to be allocated in GetCopy.
 	check(!bMakeSureChunkIsLoaded || (OutChunkData && (*OutChunkData == nullptr)));
@@ -1223,7 +1252,14 @@ static void CookSurroundWave( USoundWave* SoundWave, FName FormatName, const IAu
 			{
 				CompressionQualityModifier = CompressionOverrides->CompressionQualityModifier;
 			}
-			float ModifiedCompressionQuality = (float)SoundWave->CompressionQuality * CompressionQualityModifier;
+
+			float ModifiedCompressionQuality = (float)SoundWave->CompressionQuality;
+
+			if (!FMath::IsNearlyEqual(CompressionQualityModifier, 1.0f))
+			{
+				ModifiedCompressionQuality = (float)SoundWave->CompressionQuality * CompressionQualityModifier;
+			}
+			
 			if (ModifiedCompressionQuality >= 1.0f)
 			{
 				QualityInfo.Quality = FMath::FloorToInt(ModifiedCompressionQuality);
@@ -1415,7 +1451,7 @@ void USoundWave::CachePlatformData(bool bAsyncCache)
 {
 	FString DerivedDataKey;
 	FName AudioFormat = GetWaveFormatForRunningPlatform(*this);
-	const FPlatformAudioCookOverrides* CompressionOverrides = GetPlatformCompressionOverridesForCurrentPlatform();
+	const FPlatformAudioCookOverrides* CompressionOverrides = GetCookOverridesForRunningPlatform();
 	GetStreamedAudioDerivedDataKey(*this, AudioFormat, CompressionOverrides, DerivedDataKey);
 
 	if (RunningPlatformData == NULL || RunningPlatformData->DerivedDataKey != DerivedDataKey)
@@ -1587,7 +1623,7 @@ void USoundWave::FinishCachePlatformData()
 	{
 		FString DerivedDataKey;
 		FName AudioFormat = GetWaveFormatForRunningPlatform(*this);
-		const FPlatformAudioCookOverrides* CompressionOverrides = GetPlatformCompressionOverridesForCurrentPlatform();
+		const FPlatformAudioCookOverrides* CompressionOverrides = GetCookOverridesForRunningPlatform();
 		GetStreamedAudioDerivedDataKey(*this, AudioFormat, CompressionOverrides, DerivedDataKey);
 
 		check(RunningPlatformData->DerivedDataKey == DerivedDataKey);
diff --git a/Engine/Source/Runtime/Engine/Private/Level.cpp b/Engine/Source/Runtime/Engine/Private/Level.cpp
--- a/Engine/Source/Runtime/Engine/Private/Level.cpp
+++ b/Engine/Source/Runtime/Engine/Private/Level.cpp
@@ -453,7 +453,7 @@ bool ULevel::IsNetActor(const AActor* Actor)
 
 	// If this is a server, use RemoteRole.
 	// If this is a client, use Role.
-	const ENetRole NetRole = (!Actor->IsNetMode(NM_Client)) ? Actor->GetRemoteRole() : (ENetRole)Actor->Role;
+	const ENetRole NetRole = (!Actor->IsNetMode(NM_Client)) ? Actor->GetRemoteRole() : (ENetRole)Actor->GetLocalRole();
 
 	// This test will return true on clients for actors with ROLE_Authority, which might be counterintuitive,
 	// but clients will need to consider these actors in some cases, such as if their bTearOff flag is true.
diff --git a/Engine/Source/Runtime/Engine/Private/LevelStreaming.cpp b/Engine/Source/Runtime/Engine/Private/LevelStreaming.cpp
--- a/Engine/Source/Runtime/Engine/Private/LevelStreaming.cpp
+++ b/Engine/Source/Runtime/Engine/Private/LevelStreaming.cpp
@@ -1421,8 +1421,9 @@ FName ULevelStreaming::GetLoadedLevelPackageName() const
 
 void ULevelStreaming::SetWorldAssetByPackageName(FName InPackageName)
 {
+	// Need to strip PIE prefix from object name, only the package has it
 	const FString TargetWorldPackageName = InPackageName.ToString();
-	const FString TargetWorldObjectName = FPackageName::GetLongPackageAssetName(TargetWorldPackageName);
+	const FString TargetWorldObjectName = UWorld::RemovePIEPrefix(FPackageName::GetLongPackageAssetName(TargetWorldPackageName));
 	TSoftObjectPtr<UWorld> NewWorld;
 	NewWorld = FString::Printf(TEXT("%s.%s"), *TargetWorldPackageName, *TargetWorldObjectName);
 	SetWorldAsset(NewWorld);
diff --git a/Engine/Source/Runtime/Engine/Private/SkeletalMesh.cpp b/Engine/Source/Runtime/Engine/Private/SkeletalMesh.cpp
--- a/Engine/Source/Runtime/Engine/Private/SkeletalMesh.cpp
+++ b/Engine/Source/Runtime/Engine/Private/SkeletalMesh.cpp
@@ -288,6 +288,79 @@ FSkeletalMeshClothBuildParams::FSkeletalMeshClothBuildParams()
 
 }
 
+#if WITH_EDITOR
+FScopedSkeletalMeshPostEditChange::FScopedSkeletalMeshPostEditChange(USkeletalMesh* InSkeletalMesh, bool InbCallPostEditChange /*= true*/, bool InbReregisterComponents /*= true*/)
+{
+	SkeletalMesh = nullptr;
+	bReregisterComponents = InbReregisterComponents;
+	bCallPostEditChange = InbCallPostEditChange;
+	RecreateExistingRenderStateContext = nullptr;
+	ComponentReregisterContexts.Empty();
+	//Validation of the data
+	if (bCallPostEditChange && !bReregisterComponents)
+	{
+		//We never want to call PostEditChange without re register components, since PostEditChange will recreate the skeletalmesh render resources
+		ensure(bReregisterComponents);
+		bReregisterComponents = true;
+	}
+	if (InSkeletalMesh != nullptr)
+	{
+		//Only set a valid skeletal mesh
+		SetSkeletalMesh(InSkeletalMesh);
+	}
+}
+
+FScopedSkeletalMeshPostEditChange::~FScopedSkeletalMeshPostEditChange()
+{
+	//If decrementing the post edit change stack counter return 0 it mean we are the first scope call instance, so we have to call posteditchange and re register component
+	if (SkeletalMesh != nullptr && SkeletalMesh->UnStackPostEditChange() == 0)
+	{
+		if (bCallPostEditChange)
+		{
+			SkeletalMesh->PostEditChange();
+		}
+	}
+	//If there is some re register data it will be delete when the destructor go out of scope. This will re register
+}
+
+void FScopedSkeletalMeshPostEditChange::SetSkeletalMesh(USkeletalMesh* InSkeletalMesh)
+{
+	//We cannot set a different skeletal mesh, check that it was construct with null
+	check(SkeletalMesh == nullptr);
+	//We can only set a valid skeletal mesh
+	check(InSkeletalMesh != nullptr);
+
+	SkeletalMesh = InSkeletalMesh;
+	//If we are the first to increment, unregister the data we need to
+	if (SkeletalMesh->StackPostEditChange() == 1)
+	{
+		//Only allocate data if we re register
+		if (bReregisterComponents)
+		{
+			//Make sure all components using this skeletalmesh have there render ressources free
+			RecreateExistingRenderStateContext = new FSkinnedMeshComponentRecreateRenderStateContext(InSkeletalMesh, false);
+
+			// Now iterate over all skeletal mesh components and unregister them from the world, we will reregister them in the destructor
+			for (TObjectIterator<USkinnedMeshComponent> It; It; ++It)
+			{
+				USkinnedMeshComponent* SkinComp = *It;
+				if (SkinComp->SkeletalMesh == SkeletalMesh)
+				{
+					ComponentReregisterContexts.Add(new FComponentReregisterContext(SkinComp));
+				}
+			}
+		}
+
+		if (bCallPostEditChange)
+		{
+			//Make sure the render ressource use by the skeletalMesh is free, we will reconstruct them when a PostEditChange will be call
+			SkeletalMesh->FlushRenderState();
+		}
+	}
+}
+
+#endif //WITH_EDITOR
+
 /*-----------------------------------------------------------------------------
 	USkeletalMesh
 -----------------------------------------------------------------------------*/
@@ -306,7 +379,6 @@ USkeletalMesh::USkeletalMesh(const FObjectInitializer& ObjectInitializer)
 	// TODO: support saving some but not all optional LODs
 	MaxNumOptionalLODs.Default = 0;
 #endif
-	
 	MinLod.Default = 0;
 	DisableBelowMinLodStripping.Default = false;
 }
@@ -617,6 +689,46 @@ void USkeletalMesh::InitResources()
 	FSkeletalMeshRenderData* SkelMeshRenderData = GetResourceForRendering();
 	if (SkelMeshRenderData)
 	{
+#if WITH_EDITOR
+		//Editor sanity check, we must ensure all the data is in sync between LODModel, RenderData and UserSectionsData
+		if (ImportedModel.IsValid())
+		{
+			for (int32 LODIndex = 0; LODIndex < GetLODNum(); ++LODIndex)
+			{
+				if (!ImportedModel->LODModels.IsValidIndex(LODIndex) || !SkeletalMeshRenderData->LODRenderData.IsValidIndex(LODIndex))
+				{
+					continue;
+				}
+				const FSkeletalMeshLODModel& ImportLODModel = ImportedModel->LODModels[LODIndex];
+				FSkeletalMeshLODRenderData& RenderLODModel = SkeletalMeshRenderData->LODRenderData[LODIndex];
+				check(ImportLODModel.Sections.Num() == RenderLODModel.RenderSections.Num());
+				for (int32 SectionIndex = 0; SectionIndex < ImportLODModel.Sections.Num(); ++SectionIndex)
+				{
+					const FSkelMeshSection& ImportSection = ImportLODModel.Sections[SectionIndex];
+					
+					//In Editor we want to make sure the data is in sync between UserSectionsData and LODModel Sections
+					const FSkelMeshSourceSectionUserData& SectionUserData = ImportLODModel.UserSectionsData.FindChecked(ImportSection.OriginalDataSectionIndex);
+					bool bDataInSync = SectionUserData.bDisabled == ImportSection.bDisabled &&
+						SectionUserData.bCastShadow == ImportSection.bCastShadow &&
+						SectionUserData.bRecomputeTangent == ImportSection.bRecomputeTangent &&
+						SectionUserData.CorrespondClothAssetIndex == ImportSection.CorrespondClothAssetIndex &&
+						SectionUserData.ClothingData.AssetGuid == ImportSection.ClothingData.AssetGuid &&
+						SectionUserData.ClothingData.AssetLodIndex == ImportSection.ClothingData.AssetLodIndex;
+					check(bDataInSync);
+					
+					//In Editor we want to make sure the data is in sync between UserSectionsData and RenderSections
+					const FSkelMeshRenderSection& RenderSection = RenderLODModel.RenderSections[SectionIndex];
+					bDataInSync = SectionUserData.bDisabled == RenderSection.bDisabled &&
+						SectionUserData.bCastShadow == RenderSection.bCastShadow &&
+						SectionUserData.bRecomputeTangent == RenderSection.bRecomputeTangent &&
+						SectionUserData.CorrespondClothAssetIndex == RenderSection.CorrespondClothAssetIndex &&
+						SectionUserData.ClothingData.AssetGuid == RenderSection.ClothingData.AssetGuid &&
+						SectionUserData.ClothingData.AssetLodIndex == RenderSection.ClothingData.AssetLodIndex;
+					check(bDataInSync);
+				}
+			}
+		}
+#endif
 		SkelMeshRenderData->InitResources(bHasVertexColors, MorphTargets, this);
 	}
 
@@ -640,6 +752,9 @@ void USkeletalMesh::ReleaseResources()
 	FSkeletalMeshRenderData* SkelMeshRenderData = GetResourceForRendering();
 	if (SkelMeshRenderData)
 	{
+		//Flush the rendering command to be sure there is no command left that can create a rendering ressource
+		FlushRenderingCommands();
+
 		SkelMeshRenderData->ReleaseResources();
 	}
 
@@ -1073,29 +1188,44 @@ void RefreshSkelMeshOnPhysicsAssetChange(const USkeletalMesh* InSkeletalMesh)
 }
 
 #if WITH_EDITOR
+
+int32 USkeletalMesh::StackPostEditChange()
+{
+	check(PostEditChangeStackCounter >= 0);
+	//Return true if this is the first stack ID
+	PostEditChangeStackCounter++;
+	return PostEditChangeStackCounter;
+}
+
+int32 USkeletalMesh::UnStackPostEditChange()
+{
+	check(PostEditChangeStackCounter > 0);
+	PostEditChangeStackCounter--;
+	return PostEditChangeStackCounter;
+}
+
 void USkeletalMesh::PostEditChangeProperty(FPropertyChangedEvent& PropertyChangedEvent)
 {
+	if (PostEditChangeStackCounter > 0)
+	{
+		//Ignore those call when we have an active delay stack
+		return;
+	}
+	//Block any re-entrant call by incrementing PostEditChangeStackCounter. It will be decrement when we will go out of scope.
+	const bool bCallPostEditChange = false;
+	const bool bReRegisterComponents = false;
+	FScopedSkeletalMeshPostEditChange(this, bCallPostEditChange, bReRegisterComponents);
+
 	bool bFullPrecisionUVsReallyChanged = false;
 
 	UProperty* PropertyThatChanged = PropertyChangedEvent.Property;
 	
-	if( GIsEditor &&
-		PropertyThatChanged &&
-		PropertyThatChanged->GetFName() == FName(TEXT("bUseFullPrecisionUVs")) )
-	{
-		bFullPrecisionUVsReallyChanged = true;
-		if (!bUseFullPrecisionUVs && !GVertexElementTypeSupport.IsSupported(VET_Half2) )
-		{
-			bUseFullPrecisionUVs = true;
-			UE_LOG(LogSkeletalMesh, Warning, TEXT("16 bit UVs not supported. Reverting to 32 bit UVs"));			
-			bFullPrecisionUVsReallyChanged = false;
-		}
-	}
-	
+	bool bHasToReregisterComponent = false;
 	// Don't invalidate render data when dragging sliders, too slow
 	if (PropertyChangedEvent.ChangeType != EPropertyChangeType::Interactive)
 	{
-		InvalidateRenderData();
+		Build();
+		bHasToReregisterComponent = true;
 	}
 
 	if( GIsEditor &&
@@ -1121,7 +1251,6 @@ void USkeletalMesh::PostEditChangeProperty(FPropertyChangedEvent& PropertyChange
 		BuildPhysicsData();
 	}
 
-	bool bHasToReregisterComponent = false;
 	if(UProperty* MemberProperty = PropertyChangedEvent.MemberProperty)
 	{
 		if(MemberProperty->GetFName() == GET_MEMBER_NAME_CHECKED(USkeletalMesh, PositiveBoundsExtension) ||
@@ -1186,6 +1315,9 @@ void USkeletalMesh::PostEditChangeProperty(FPropertyChangedEvent& PropertyChange
 	}
 
 	Super::PostEditChangeProperty(PropertyChangedEvent);
+
+	//The stack counter after a PostEditChange should be 0
+	check(PostEditChangeStackCounter == 0);
 }
 
 void USkeletalMesh::PostEditUndo()
@@ -1470,6 +1602,26 @@ void USkeletalMesh::Serialize( FArchive& Ar )
 			}
 		}
 	}
+
+	if (Ar.IsLoading() && Ar.CustomVer(FEditorObjectVersion::GUID) < FEditorObjectVersion::SkeletalMeshBuildRefactor)
+	{
+		const int32 TotalLODNum = LODInfo.Num();
+		for (int32 LodIndex = 0; LodIndex < TotalLODNum; LodIndex++)
+		{
+			FSkeletalMeshLODInfo& ThisLODInfo = LODInfo[LodIndex];
+			// Restore the deprecated settings
+			ThisLODInfo.BuildSettings.bUseFullPrecisionUVs = bUseFullPrecisionUVs_DEPRECATED;
+			ThisLODInfo.BuildSettings.bUseHighPrecisionTangentBasis = bUseHighPrecisionTangentBasis_DEPRECATED;
+			ThisLODInfo.BuildSettings.bBuildAdjacencyBuffer = true;
+			ThisLODInfo.BuildSettings.bRemoveDegenerates = true;
+
+			//We cannot get back the imported build option here since those option are store in the UAssetImportData which FBX has derive in the UnrealEd module
+			//We are in engine module so there is no way to recover this data.
+			//Anyway because the asset was not re-import yet the build settings will not be shown in the UI and the asset will not be build
+			//With the new build until it get re-import (geo and skinning)
+			//So we will leave the default value for the rest of the new build settings
+		}
+	}
 }
 
 void USkeletalMesh::GetPreloadDependencies(TArray<UObject*>& OutDeps)
@@ -1493,22 +1645,15 @@ void USkeletalMesh::FlushRenderState()
 uint32 USkeletalMesh::GetVertexBufferFlags() const
 {
 	uint32 VertexFlags = ESkeletalMeshVertexFlags::None;
-	if (bUseFullPrecisionUVs)
-	{
-		VertexFlags |= ESkeletalMeshVertexFlags::UseFullPrecisionUVs;
-	}
 	if (bHasVertexColors)
 	{
 		VertexFlags |= ESkeletalMeshVertexFlags::HasVertexColors;
 	}
-	if (bUseHighPrecisionTangentBasis)
-	{
-		VertexFlags |= ESkeletalMeshVertexFlags::UseHighPrecisionTangentBasis;
-	}
 	return VertexFlags;
 }
 
-void USkeletalMesh::InvalidateRenderData()
+#if WITH_EDITOR
+void USkeletalMesh::Build()
 {
 	// Unregister all instances of this component
 	FSkinnedMeshComponentRecreateRenderStateContext RecreateRenderStateContext(this, false);
@@ -1520,13 +1665,9 @@ void USkeletalMesh::InvalidateRenderData()
 	// allocated, and potentially accessing the USkeletalMesh.
 	ReleaseResourcesFence.Wait();
 
-#if WITH_EDITOR
-	// Create new 
-	GetImportedModel()->GenerateNewGUID();
 
 	// rebuild render data from imported model
 	CacheDerivedData();
-#endif
 
 	// Do not need to fix up 16-bit UVs here, as we assume all editor platforms support them.
 	ensure(GVertexElementTypeSupport.IsSupported(VET_Half2));
@@ -1539,6 +1680,7 @@ void USkeletalMesh::InvalidateRenderData()
 		InitResources();
 	}
 }
+#endif
 
 void USkeletalMesh::PreSave(const class ITargetPlatform* TargetPlatform)
 {
@@ -1809,6 +1951,13 @@ void USkeletalMesh::RemoveLegacyClothingSections()
 
 						ClothingSectionCount++;
 					}
+					else
+					{
+						Section.CorrespondClothAssetIndex = INDEX_NONE;
+						Section.ClothingData.AssetGuid = FGuid();
+						Section.ClothingData.AssetLodIndex = INDEX_NONE;
+						Section.ClothMappingData.Empty();
+					}
 				}
 
 				if(BaseVertex != MAX_uint32 && BaseIndex != MAX_uint32)
@@ -1830,8 +1979,6 @@ void USkeletalMesh::RemoveLegacyClothingSections()
 							Index -= VertexCount;
 						}
 					}
-
-					PostEditChange();
 				}
 			}
 		}
@@ -1941,6 +2088,95 @@ void USkeletalMesh::PostLoad()
 
 		RemoveLegacyClothingSections();
 
+		UpdateGenerateUpToData();
+
+		if (GetLinkerCustomVersion(FEditorObjectVersion::GUID) < FEditorObjectVersion::SkeletalMeshBuildRefactor)
+		{
+			//We want to avoid changing the ddc if we load an old asset.
+			//This bool should be put to false at the end of the postload, if there is another posteditchange call after a new ddc will be created
+			UseLegacyMeshDerivedDataKey = true;
+			//Fill up the Section ChunkedParentSectionIndex and OriginalDataSectionIndex
+			//We also want to create the UserSectionsData structure so the user can change the section data
+			for (int32 LodIndex = 0; LodIndex < LODInfo.Num(); LodIndex++)
+			{
+				FSkeletalMeshLODModel& ThisLODModel = ImportedModel->LODModels[LodIndex];
+				FSkeletalMeshLODInfo* ThisLODInfo = GetLODInfo(LodIndex);
+				check(ThisLODInfo);
+				int32 LODModelSectionNum = ThisLODModel.Sections.Num();
+				//Fill the ChunkedParentSectionIndex data, we assume that every section using the same material are chunked
+				int32 LastMaterialIndex = INDEX_NONE;
+				int32 CurrentParentChunkIndex = INDEX_NONE;
+				int32 OriginalIndex = 0;
+				
+				//Reset the reduction setting to a non active state if the asset has active reduction but have no RawSkeletalMeshBulkData (we cannot reduce it)
+				if (IsReductionActive(LodIndex) && !ThisLODInfo->bHasBeenSimplified && ThisLODModel.RawSkeletalMeshBulkData.IsEmpty())
+				{
+					if (LodIndex > ThisLODInfo->ReductionSettings.BaseLOD)
+					{
+						ThisLODInfo->bHasBeenSimplified = true;
+					}
+					else if (LodIndex == ThisLODInfo->ReductionSettings.BaseLOD)
+					{
+						if (ThisLODInfo->ReductionSettings.TerminationCriterion == SkeletalMeshTerminationCriterion::SMTC_AbsNumOfTriangles
+							|| ThisLODInfo->ReductionSettings.TerminationCriterion == SkeletalMeshTerminationCriterion::SMTC_AbsNumOfVerts
+							|| ThisLODInfo->ReductionSettings.TerminationCriterion == SkeletalMeshTerminationCriterion::SMTC_AbsTriangleOrVert)
+						{
+							//MaxNum.... cannot be inactive, switch to NumOfTriangle
+							ThisLODInfo->ReductionSettings.TerminationCriterion = SMTC_NumOfTriangles;
+						}
+
+						//Now that we use trinagle or vert num, set an inactive value
+						if (ThisLODInfo->ReductionSettings.TerminationCriterion == SkeletalMeshTerminationCriterion::SMTC_NumOfTriangles
+							|| ThisLODInfo->ReductionSettings.TerminationCriterion == SkeletalMeshTerminationCriterion::SMTC_TriangleOrVert)
+						{
+							ThisLODInfo->ReductionSettings.NumOfTrianglesPercentage = 1.0f;
+						}
+						if (ThisLODInfo->ReductionSettings.TerminationCriterion == SkeletalMeshTerminationCriterion::SMTC_NumOfVerts
+							|| ThisLODInfo->ReductionSettings.TerminationCriterion == SkeletalMeshTerminationCriterion::SMTC_TriangleOrVert)
+						{
+							ThisLODInfo->ReductionSettings.NumOfVertPercentage = 1.0f;
+						}
+					}
+				}
+
+				for (int32 LODModelSectionIndex = 0; LODModelSectionIndex < LODModelSectionNum; ++LODModelSectionIndex)
+				{
+					//If we have cloth on a chunked section we treat the chunked section has a parent section (this is to get the same result has before the refactor)
+					if (ThisLODModel.Sections[LODModelSectionIndex].MaterialIndex == LastMaterialIndex && !ThisLODModel.Sections[LODModelSectionIndex].ClothingData.AssetGuid.IsValid())
+					{
+						ThisLODModel.Sections[LODModelSectionIndex].ChunkedParentSectionIndex = CurrentParentChunkIndex;
+						ThisLODModel.Sections[LODModelSectionIndex].OriginalDataSectionIndex = ThisLODModel.Sections[CurrentParentChunkIndex].OriginalDataSectionIndex;
+						//In case of a child section that was BONE chunked ensure it has the same setting has the original section
+						FSkelMeshSourceSectionUserData& SectionUserData = ThisLODModel.UserSectionsData.FindOrAdd(ThisLODModel.Sections[LODModelSectionIndex].OriginalDataSectionIndex);
+						ThisLODModel.Sections[LODModelSectionIndex].bDisabled = SectionUserData.bDisabled;
+						ThisLODModel.Sections[LODModelSectionIndex].bCastShadow = SectionUserData.bCastShadow;
+						ThisLODModel.Sections[LODModelSectionIndex].bRecomputeTangent = SectionUserData.bRecomputeTangent;
+						ThisLODModel.Sections[LODModelSectionIndex].GenerateUpToLodIndex = SectionUserData.GenerateUpToLodIndex;
+						//Chunked section cannot have cloth, a cloth section will be a parent section
+						ThisLODModel.Sections[LODModelSectionIndex].CorrespondClothAssetIndex = INDEX_NONE;
+						ThisLODModel.Sections[LODModelSectionIndex].ClothingData.AssetGuid = FGuid();
+						ThisLODModel.Sections[LODModelSectionIndex].ClothingData.AssetLodIndex = INDEX_NONE;
+					}
+					else
+					{
+						CurrentParentChunkIndex = LODModelSectionIndex;
+						FSkelMeshSourceSectionUserData& SectionUserData = ThisLODModel.UserSectionsData.FindOrAdd(OriginalIndex);
+						SectionUserData.bDisabled = ThisLODModel.Sections[LODModelSectionIndex].bDisabled;
+						SectionUserData.bCastShadow = ThisLODModel.Sections[LODModelSectionIndex].bCastShadow;
+						SectionUserData.bRecomputeTangent = ThisLODModel.Sections[LODModelSectionIndex].bRecomputeTangent;
+						SectionUserData.GenerateUpToLodIndex = ThisLODModel.Sections[LODModelSectionIndex].GenerateUpToLodIndex;
+						SectionUserData.CorrespondClothAssetIndex = ThisLODModel.Sections[LODModelSectionIndex].CorrespondClothAssetIndex;
+						SectionUserData.ClothingData.AssetGuid = ThisLODModel.Sections[LODModelSectionIndex].ClothingData.AssetGuid;
+						SectionUserData.ClothingData.AssetLodIndex = ThisLODModel.Sections[LODModelSectionIndex].ClothingData.AssetLodIndex;
+
+						ThisLODModel.Sections[LODModelSectionIndex].OriginalDataSectionIndex = OriginalIndex++;
+						ThisLODModel.Sections[LODModelSectionIndex].ChunkedParentSectionIndex = INDEX_NONE;
+					}
+					LastMaterialIndex = ThisLODModel.Sections[LODModelSectionIndex].MaterialIndex;
+				}
+			}
+		}
+
 		if (GetResourceForRendering() == nullptr)
 		{
 			CacheDerivedData();
@@ -1969,6 +2205,8 @@ void USkeletalMesh::PostLoad()
 				}
 				if (!bFound)
 				{
+					//No post edit change and no reregister, we just prevent the inner scope to call postedit change and reregister
+					FScopedSkeletalMeshPostEditChange ScopedPostEditChange(this, false, false);
 					//Make sure the asset is unbind, some old code path was allowing to have bind cloth asset not present in the imported model.
 					//The old inline reduction code was not rebinding the cloth asset nor unbind it.
 					MeshClothingAsset->UnbindFromSkeletalMesh(this);
@@ -2044,12 +2282,15 @@ void USkeletalMesh::PostLoad()
 		SamplingInfo.BuildWholeMesh(this);
 	}
 
-	UpdateGenerateUpToData();
 #endif
 
 #if !WITH_EDITOR
 	RebuildSocketMap();
 #endif // !WITH_EDITOR
+#if WITH_EDITORONLY_DATA
+	//Next postedit change will use the new ddc key scheme
+	UseLegacyMeshDerivedDataKey = false;
+#endif
 }
 
 #if WITH_EDITORONLY_DATA
@@ -2168,6 +2409,10 @@ void USkeletalMesh::DebugVerifySkeletalMeshLOD()
 
 void USkeletalMesh::InitMorphTargetsAndRebuildRenderData()
 {
+#if WITH_EDITOR
+	FScopedSkeletalMeshPostEditChange ScopedPostEditChange(this);
+#endif
+	
 	MarkPackageDirty();
 	// need to refresh the map
 	InitMorphTargets();
@@ -2180,9 +2425,6 @@ void USkeletalMesh::InitMorphTargetsAndRebuildRenderData()
 			It->RefreshMorphTargets();
 		}
 	}
-
-	// invalidate render data
-	InvalidateRenderData();
 }
 
 bool USkeletalMesh::RegisterMorphTarget(UMorphTarget* MorphTarget, bool bInvalidateRenderData)
@@ -2845,6 +3087,11 @@ void USkeletalMesh::AllocateResourceForRendering()
 }
 
 #if WITH_EDITOR
+void USkeletalMesh::InvalidateDeriveDataCacheGUID()
+{
+	// Create new DDC guid
+	GetImportedModel()->GenerateNewGUID();
+}
 
 void USkeletalMesh::CacheDerivedData()
 {
@@ -3518,6 +3765,102 @@ FText USkeletalMesh::GetSourceFileLabelFromIndex(int32 SourceFileIndex)
 	int32 RealSourceFileIndex = SourceFileIndex == INDEX_NONE ? 0 : SourceFileIndex;
 	return RealSourceFileIndex == 0 ? NSSkeletalMeshSourceFileLabels::GeoAndSkinningText() : RealSourceFileIndex == 1 ? NSSkeletalMeshSourceFileLabels::GeometryText() : NSSkeletalMeshSourceFileLabels::SkinningText();
 }
+
+static void SerializeReductionSettingsForDDC(FArchive& Ar, FSkeletalMeshOptimizationSettings& ReductionSettings)
+{
+	check(Ar.IsSaving());
+	// Note: this serializer is only used to build the mesh DDC key, no versioning is required
+	FArchive_Serialize_BitfieldBool(Ar, ReductionSettings.bRemapMorphTargets);
+	FArchive_Serialize_BitfieldBool(Ar, ReductionSettings.bRecalcNormals);
+	FArchive_Serialize_BitfieldBool(Ar, ReductionSettings.bEnforceBoneBoundaries);
+	FArchive_Serialize_BitfieldBool(Ar, ReductionSettings.bLockEdges);
+	Ar << ReductionSettings.TerminationCriterion;
+	Ar << ReductionSettings.NumOfTrianglesPercentage;
+	Ar << ReductionSettings.NumOfVertPercentage;
+	Ar << ReductionSettings.MaxNumOfTriangles;
+	Ar << ReductionSettings.MaxNumOfVerts;
+	Ar << ReductionSettings.MaxDeviationPercentage;
+	Ar << ReductionSettings.ReductionMethod;
+	Ar << ReductionSettings.SilhouetteImportance;
+	Ar << ReductionSettings.TextureImportance;
+	Ar << ReductionSettings.ShadingImportance;
+	Ar << ReductionSettings.SkinningImportance;
+	Ar << ReductionSettings.WeldingThreshold;
+	Ar << ReductionSettings.NormalsThreshold;
+	Ar << ReductionSettings.MaxBonesPerVertex;
+	Ar << ReductionSettings.VolumeImportance;
+	Ar << ReductionSettings.BaseLOD;
+}
+
+static void SerializeBuildSettingsForDDC(FArchive& Ar, FSkeletalMeshBuildSettings& BuildSettings)
+{
+	check(Ar.IsSaving());
+	// Note: this serializer is only used to build the mesh DDC key, no versioning is required
+	FArchive_Serialize_BitfieldBool(Ar, BuildSettings.bRecomputeNormals);
+	FArchive_Serialize_BitfieldBool(Ar, BuildSettings.bRecomputeTangents);
+	FArchive_Serialize_BitfieldBool(Ar, BuildSettings.bUseMikkTSpace);
+	FArchive_Serialize_BitfieldBool(Ar, BuildSettings.bComputeWeightedNormals);
+	FArchive_Serialize_BitfieldBool(Ar, BuildSettings.bRemoveDegenerates);
+	FArchive_Serialize_BitfieldBool(Ar, BuildSettings.bUseFullPrecisionUVs);
+	FArchive_Serialize_BitfieldBool(Ar, BuildSettings.bUseHighPrecisionTangentBasis);
+	Ar << BuildSettings.ThresholdPosition;
+	Ar << BuildSettings.ThresholdTangentNormal;
+	Ar << BuildSettings.ThresholdUV;
+}
+
+
+FGuid FSkeletalMeshLODInfo::ComputeDeriveDataCacheKey(const FSkeletalMeshLODGroupSettings* SkeletalMeshLODGroupSettings)
+{
+	const bool bIs16BitfloatBufferSupported = GVertexElementTypeSupport.IsSupported(VET_Half2);
+
+	// Serialize the LOD info members, the BuildSettings and the ReductionSettings into a temporary array.
+	TArray<uint8> TempBytes;
+	TempBytes.Reserve(64);
+	//The archive is flagged as persistent so that machines of different endianness produce identical binary results.
+	FMemoryWriter Ar(TempBytes, /*bIsPersistent=*/ true);
+
+	Ar << BonesToRemove;
+	Ar << BonesToPrioritize;
+	Ar << WeightOfPrioritization;
+	if (BakePose != nullptr)
+	{
+		FSoftObjectPath BakePosePath(BakePose);
+		Ar << BakePosePath;
+	}
+	if (BakePoseOverride != nullptr)
+	{
+		FSoftObjectPath BakePoseOverridePath(BakePoseOverride);
+		Ar << BakePoseOverridePath;
+	}
+	FArchive_Serialize_BitfieldBool(Ar, bAllowCPUAccess);
+	FArchive_Serialize_BitfieldBool(Ar, bSupportUniformlyDistributedSampling);
+
+	//Use the LOD settings asset if there is one
+	FSkeletalMeshBuildSettings RealBuildSettings = BuildSettings;
+	FSkeletalMeshOptimizationSettings RealReductionSettings = ReductionSettings;
+	if (SkeletalMeshLODGroupSettings != nullptr)
+	{
+		RealBuildSettings = SkeletalMeshLODGroupSettings->GetBuildSettings();
+		RealReductionSettings = SkeletalMeshLODGroupSettings->GetReductionSettings();
+	}
+
+	if (!RealBuildSettings.bUseFullPrecisionUVs && !bIs16BitfloatBufferSupported)
+	{
+		RealBuildSettings.bUseFullPrecisionUVs = true;
+	}
+	SerializeBuildSettingsForDDC(Ar, RealBuildSettings);
+	SerializeReductionSettingsForDDC(Ar, RealReductionSettings);
+
+	FSHA1 Sha;
+	Sha.Update(TempBytes.GetData(), TempBytes.Num() * TempBytes.GetTypeSize());
+	Sha.Final();
+	// Retrieve the hash and use it to construct a pseudo-GUID.
+	uint32 Hash[5];
+	Sha.GetHash((uint8*)Hash);
+	FGuid Guid = FGuid(Hash[0] ^ Hash[4], Hash[1], Hash[2], Hash[3]);
+	return Guid;
+}
+
 #endif //WITH_EDITOR
 
 
@@ -3792,9 +4135,8 @@ FSkeletalMeshSceneProxy::FSkeletalMeshSceneProxy(const USkinnedMeshComponent* Co
 
 			// If we are at a dropped LOD, route material index through the LODMaterialMap in the LODInfo struct.
 			int32 UseMaterialIndex = Section.MaterialIndex;			
-			if(LODIdx > 0)
 			{
-				if(Section.MaterialIndex < Info.LODMaterialMap.Num())
+				if(Section.MaterialIndex < Info.LODMaterialMap.Num() && Component->SkeletalMesh->Materials.IsValidIndex(Info.LODMaterialMap[Section.MaterialIndex]))
 				{
 					UseMaterialIndex = Info.LODMaterialMap[Section.MaterialIndex];
 					UseMaterialIndex = FMath::Clamp( UseMaterialIndex, 0, Component->SkeletalMesh->Materials.Num() );
diff --git a/Engine/Source/Runtime/Engine/Private/SkeletalMeshLODRenderData.cpp b/Engine/Source/Runtime/Engine/Private/SkeletalMeshLODRenderData.cpp
--- a/Engine/Source/Runtime/Engine/Private/SkeletalMeshLODRenderData.cpp
+++ b/Engine/Source/Runtime/Engine/Private/SkeletalMeshLODRenderData.cpp
@@ -533,6 +533,7 @@ void FSkeletalMeshLODRenderData::BuildFromLODModel(const FSkeletalMeshLODModel*
 	bool bUseFullPrecisionUVs = (BuildFlags & ESkeletalMeshVertexFlags::UseFullPrecisionUVs) != 0;
 	bool bUseHighPrecisionTangentBasis = (BuildFlags & ESkeletalMeshVertexFlags::UseHighPrecisionTangentBasis) != 0;
 	bool bHasVertexColors = (BuildFlags & ESkeletalMeshVertexFlags::HasVertexColors) != 0;
+	bool bBuildAdjacencyBuffer = (BuildFlags & ESkeletalMeshVertexFlags::BuildAdjacencyIndexBuffer) != 0;
 
 	// Copy required info from source sections
 	RenderSections.Empty();
@@ -602,10 +603,13 @@ void FSkeletalMeshLODRenderData::BuildFromLODModel(const FSkeletalMeshLODModel*
 
 	MultiSizeIndexContainer.RebuildIndexBuffer(DataTypeSize, ImportedModel->IndexBuffer);
 	
-	TArray<uint32> BuiltAdjacencyIndices;
 	IMeshUtilities& MeshUtilities = FModuleManager::Get().LoadModuleChecked<IMeshUtilities>("MeshUtilities");
-	MeshUtilities.BuildSkeletalAdjacencyIndexBuffer(Vertices, ImportedModel->NumTexCoords, ImportedModel->IndexBuffer, BuiltAdjacencyIndices);
-	AdjacencyMultiSizeIndexContainer.RebuildIndexBuffer(DataTypeSize, BuiltAdjacencyIndices);
+	if (bBuildAdjacencyBuffer)
+	{
+		TArray<uint32> BuiltAdjacencyIndices;
+		MeshUtilities.BuildSkeletalAdjacencyIndexBuffer(Vertices, ImportedModel->NumTexCoords, ImportedModel->IndexBuffer, BuiltAdjacencyIndices);
+		AdjacencyMultiSizeIndexContainer.RebuildIndexBuffer(DataTypeSize, BuiltAdjacencyIndices);
+	}
 
 	// MorphTargetVertexInfoBuffers are created in InitResources
 
diff --git a/Engine/Source/Runtime/Engine/Private/SoundCue.cpp b/Engine/Source/Runtime/Engine/Private/SoundCue.cpp
--- a/Engine/Source/Runtime/Engine/Private/SoundCue.cpp
+++ b/Engine/Source/Runtime/Engine/Private/SoundCue.cpp
@@ -279,7 +279,7 @@ void USoundCue::PostEditChangeProperty(struct FPropertyChangedEvent& PropertyCha
 	{
 		for (TObjectIterator<UAudioComponent> It; It; ++It)
 		{
-			if (It->Sound == this && It->bIsActive)
+			if (It->Sound == this && It->IsActive())
 			{
 				It->Stop();
 				It->Play();
diff --git a/Engine/Source/Runtime/Engine/Private/SoundWave.cpp b/Engine/Source/Runtime/Engine/Private/SoundWave.cpp
--- a/Engine/Source/Runtime/Engine/Private/SoundWave.cpp
+++ b/Engine/Source/Runtime/Engine/Private/SoundWave.cpp
@@ -47,6 +47,14 @@ FAutoConsoleVariableRef CVarLoadIntoCacheOnPostLoad(
 	TEXT("0: Honor the Play When Silent flag, 1: stop all silent non-procedural sources."),
 	ECVF_Default);
 
+static int32 ForceNonStreamingInEditorCVar = 0;
+FAutoConsoleVariableRef CVarForceNonStreamingInEditor(
+	TEXT("au.editor.ForceAudioNonStreaming"),
+	LoadIntoCacheOnPostLoadCVar,
+	TEXT("When set to 1, forces any audio played to be non-streaming May force a DDC miss.\n")
+	TEXT("0: Honor the Play When Silent flag, 1: stop all silent non-procedural sources."),
+	ECVF_Default);
+
 #if !UE_BUILD_SHIPPING
 static void DumpBakedAnalysisData(const TArray<FString>& Args)
 {
@@ -1744,7 +1752,7 @@ void USoundWave::Parse(FAudioDevice* AudioDevice, const UPTRINT NodeWaveInstance
 	// When the BypassVirtualizeWhenSilent cvar is enabled, we should only honor bVirtualizeWhenSilent for procedural sounds:
 	const bool bHasSubtitles = ActiveSound.bHandleSubtitles && (ActiveSound.bHasExternalSubtitles || (Subtitles.Num() > 0));
 	const bool bStarted = WaveInstanceVolume > KINDA_SMALL_NUMBER || ActiveSound.ComponentVolumeFader.IsFadingIn();
-	const bool bCanPlayWhenSilent = IsPlayWhenSilent() && (!BypassPlayWhenSilentCVar || bProcedural);
+	const bool bCanPlayWhenSilent = ActiveSound.IsPlayWhenSilent() && (!BypassPlayWhenSilentCVar || bProcedural);
 	if (bStarted || bCanPlayWhenSilent || bHasSubtitles)
 	{
 		WaveInstances.Add(WaveInstance);
@@ -1823,6 +1831,11 @@ bool USoundWave::IsStreaming(const FPlatformAudioCookOverrides* Overrides /* = n
 	// TODO: add in check on whether it's part of a streaming SoundGroup.
 	if (!Overrides)
 	{
+		if (GIsEditor && ForceNonStreamingInEditorCVar != 0)
+		{
+			return false;
+		}
+
 		Overrides = GetPlatformCompressionOverridesForCurrentPlatform();
 	}
 
@@ -1940,6 +1953,11 @@ void USoundWave::UpdatePlatformData()
 
 float USoundWave::GetSampleRateForCurrentPlatform()
 {
+	if (bProcedural)
+	{
+		return SampleRate;
+	}
+
 #if WITH_EDITOR
 	float SampleRateOverride = FPlatformCompressionUtilities::GetTargetSampleRateForPlatform(SampleRateQuality);
 	return (SampleRateOverride > 0) ? FMath::Min(SampleRateOverride, (float) SampleRate) : SampleRate;
@@ -1993,7 +2011,7 @@ bool USoundWave::GetChunkData(int32 ChunkIndex, uint8** OutChunkData, bool bMake
 		ForceRebuildPlatformData();
 		if (RunningPlatformData->GetChunkFromDDC(ChunkIndex, OutChunkData, bMakeSureChunkIsLoaded) == 0)
 		{
-			UE_LOG(LogAudio, Display, TEXT("Failed to build sound %s."), *GetPathName());
+			UE_LOG(LogAudio, Warning, TEXT("Failed to build sound %s."), *GetPathName());
 		}
 		else
 		{
diff --git a/Engine/Source/Runtime/Engine/Private/StaticMesh.cpp b/Engine/Source/Runtime/Engine/Private/StaticMesh.cpp
--- a/Engine/Source/Runtime/Engine/Private/StaticMesh.cpp
+++ b/Engine/Source/Runtime/Engine/Private/StaticMesh.cpp
@@ -27,6 +27,8 @@
 #include "Engine/AssetUserData.h"
 #include "StaticMeshResources.h"
 #include "StaticMeshVertexData.h"
+#include "StaticMeshAttributes.h"
+#include "StaticMeshDescription.h"
 #include "Interfaces/ITargetPlatform.h"
 #include "SpeedTreeWind.h"
 #include "DistanceFieldAtlas.h"
@@ -59,9 +61,6 @@
 
 #endif // #if WITH_EDITOR
 
-#include "MeshDescription.h"
-#include "MeshAttributes.h"
-
 #include "Engine/StaticMeshSocket.h"
 #include "EditorFramework/AssetImportData.h"
 #include "AI/Navigation/NavCollisionBase.h"
@@ -625,6 +624,8 @@ void FStaticMeshLODResources::Serialize(FArchive& Ar, UObject* Owner, int32 Inde
 	if (!StripFlags.IsDataStrippedForServer() && !bIsLODCookedOut)
 	{
 		FStaticMeshBuffersSize TmpBuffersSize;
+		TArray<uint8> TmpBuff;
+
 		if (bInlined)
 		{
 			SerializeBuffers(Ar, OwnerStaticMesh, ClassDataStripFlags, TmpBuffersSize);
@@ -633,6 +634,7 @@ void FStaticMeshLODResources::Serialize(FArchive& Ar, UObject* Owner, int32 Inde
 		}
 		else if (FPlatformProperties::RequiresCookedData() || Ar.IsCooking() || bUsingCookedData)
 		{
+			uint32 BulkDataSize = 0;
 #if WITH_EDITOR
 			if (Ar.IsSaving())
 			{
@@ -640,7 +642,6 @@ void FStaticMeshLODResources::Serialize(FArchive& Ar, UObject* Owner, int32 Inde
 				const int32 OptionalLODIdx = GetPlatformMinLODIdx(Ar.CookingTarget(), OwnerStaticMesh) - Index;
 				const bool bDiscardBulkData = OptionalLODIdx > MaxNumOptionalLODs;
 
-				TArray<uint8> TmpBuff;
 				if (!bDiscardBulkData)
 				{
 					FMemoryWriter MemWriter(TmpBuff, true);
@@ -674,11 +675,22 @@ void FStaticMeshLODResources::Serialize(FArchive& Ar, UObject* Owner, int32 Inde
 				TmpBulkData.Serialize(Ar, Owner, Index, false);
 				bIsOptionalLOD = !!(TmpBulkData.GetBulkDataFlags() & BULKDATA_OptionalPayload);
 
-				BulkDataStreamingToken = TmpBulkData.CreateStreamingToken();
+				BulkDataStreamingToken = TmpBulkData.CreateStreamingToken();			
+				BulkDataSize = BulkDataStreamingToken.GetSize();
 #else
 				StreamingBulkData.Serialize(Ar, Owner, Index, false);
 				bIsOptionalLOD = !!(StreamingBulkData.GetBulkDataFlags() & BULKDATA_OptionalPayload);
+				BulkDataSize = (uint32)StreamingBulkData.GetBulkDataSize();
 #endif
+
+				// Streaming CPU data in editor build isn't supported yet because tools and utils need access
+				if (bUsingCookedData && BulkDataSize > 0)
+				{
+					TmpBuff.Empty(BulkDataSize);
+					TmpBuff.AddUninitialized(BulkDataSize);
+					void* Dest = TmpBuff.GetData();
+					TmpBulkData.GetCopy(&Dest);
+				}
 			}
 
 			SerializeAvailabilityInfo(Ar);
@@ -690,6 +702,14 @@ void FStaticMeshLODResources::Serialize(FArchive& Ar, UObject* Owner, int32 Inde
 			{
 				ClearAvailabilityInfo();
 			}
+
+			if (Ar.IsLoading() && bUsingCookedData && BulkDataSize > 0)
+			{
+				ClearAvailabilityInfo();
+				FMemoryReader MemReader(TmpBuff, true);
+				MemReader.SetByteSwapping(Ar.IsByteSwapping());
+				SerializeBuffers(MemReader, OwnerStaticMesh, ClassDataStripFlags, TmpBuffersSize);
+			}
 		}
 	}
 }
@@ -2493,6 +2513,7 @@ UStaticMesh::UStaticMesh(const FObjectInitializer& ObjectInitializer)
 	MinLOD.Default = 0;
 
 	bSupportUniformlyDistributedSampling = false;
+	bIsBuiltAtRuntime = false;
 	bRenderingResourcesInitialized = false;
 #if WITH_EDITOR
 	BuildCacheAutomationTestGuid.Invalidate();
@@ -2555,7 +2576,7 @@ void UStaticMesh::InitResources()
 	{
 		LinkStreaming();
 	}
-	
+
 #if	STATS
 	UStaticMesh* This = this;
 	ENQUEUE_RENDER_COMMAND(UpdateMemoryStats)(
@@ -3174,11 +3195,14 @@ bool UStaticMesh::FixLODRequiresAdjacencyInformation(const int32 LODIndex, const
 		return false;
 	}
 	FStaticMeshSourceModel& SourceModel = GetSourceModel(LODIndex);
-	FMeshDescription* MeshDescription = GetMeshDescription(LODIndex);
+	const FMeshDescription* MeshDescription = GetMeshDescription(LODIndex);
+
 	//In preview mode we simulate a false BuildAdjacencyBuffer
 	if (MeshDescription && (!(SourceModel.BuildSettings.bBuildAdjacencyBuffer) || bPreviewMode))
 	{
-		TPolygonGroupAttributesRef<FName> PolygonGroupImportedMaterialSlotNames = MeshDescription->PolygonGroupAttributes().GetAttributesRef<FName>(MeshAttribute::PolygonGroup::ImportedMaterialSlotName);
+		FStaticMeshConstAttributes StaticMeshAttributes(*MeshDescription);
+
+		TPolygonGroupAttributesConstRef<FName> PolygonGroupImportedMaterialSlotNames = StaticMeshAttributes.GetPolygonGroupMaterialSlotNames();
 		int32 SectionIndex = 0;
 		
 		for (const FPolygonGroupID& PolygonGroupID : MeshDescription->PolygonGroups().GetElementIDs())
@@ -3418,7 +3442,7 @@ void FStaticMeshSourceModel::LoadRawMesh(FRawMesh& OutRawMesh) const
 			{
 				MaterialMap.Add(StaticMeshOwner->StaticMaterials[MaterialIndex].ImportedMaterialSlotName, MaterialIndex);
 			}
-			FMeshDescriptionOperations::ConvertToRawMesh(*CachedMeshDescription, OutRawMesh, MaterialMap);
+			FMeshDescriptionOperations::ConvertToRawMesh(*MeshDescription, OutRawMesh, MaterialMap);
 		}
 	}
 	else
@@ -3436,15 +3460,7 @@ void FStaticMeshSourceModel::SaveRawMesh(FRawMesh& InRawMesh, bool /* unused */)
 	//Save both format
 	RawMeshBulkData->SaveRawMesh(InRawMesh);
 
-	if (MeshDescription.IsValid())
-	{
-		MeshDescription->Empty();
-	}
-	else
-	{
-		MeshDescription = MakeUnique<FMeshDescription>();
-		UStaticMesh::RegisterMeshAttributes(*MeshDescription);
-	}
+	CreateMeshDescription();
 
 	TMap<int32, FName> MaterialMap;
 	check(StaticMeshOwner != nullptr);
@@ -3460,6 +3476,26 @@ void FStaticMeshSourceModel::SaveRawMesh(FRawMesh& InRawMesh, bool /* unused */)
 	MeshDescriptionBulkData->SaveMeshDescription(*MeshDescription);
 }
 
+FMeshDescription* FStaticMeshSourceModel::CreateMeshDescription()
+{
+	if (!MeshDescription.IsValid())
+	{
+		// If this is the first time a MeshDescription is being created, create it and register its attributes
+		MeshDescription = MakeUnique<FMeshDescription>();
+	}
+	else
+	{
+		// Otherwise, empty it completely
+		*MeshDescription = FMeshDescription();
+	}
+
+	// Register static mesh attributes on it
+	FStaticMeshAttributes StaticMeshAttributes(*MeshDescription);
+	StaticMeshAttributes.Register();
+
+	return MeshDescription.Get();
+}
+
 void FStaticMeshSourceModel::SerializeBulkData(FArchive& Ar, UObject* Owner)
 {
 	const bool bIsLoadingLegacyArchive = Ar.IsLoading() && Ar.CustomVer(FEditorObjectVersion::GUID) < FEditorObjectVersion::StaticMeshDeprecatedRawMesh;
@@ -3499,13 +3535,17 @@ void FStaticMeshSourceModel::SerializeBulkData(FArchive& Ar, UObject* Owner)
 			if (Ar.IsLoading())
 			{
 				MeshDescriptionBulkData = MakeUnique<FMeshDescriptionBulkData>();
-
-				// As we are loading a new mesh description bulkdata, reset any cached existing unpacked one
-				// @todo: would it be better to actually unpack the new one after deserializing it?
-				MeshDescription.Reset();
 			}
 
 			MeshDescriptionBulkData->Serialize(Ar, Owner);
+
+			// As we are loading a new mesh description bulkdata, if there's a cached existing unpacked MeshDescription,
+			// unpack the new one
+			if (Ar.IsLoading() && MeshDescription.IsValid())
+			{
+				CreateMeshDescription();
+				MeshDescriptionBulkData->LoadMeshDescription(*MeshDescription);
+			}
 		}
 
 		// For transactions only, serialize the unpacked mesh description here too.
@@ -3519,15 +3559,7 @@ void FStaticMeshSourceModel::SerializeBulkData(FArchive& Ar, UObject* Owner)
 			{
 				if (Ar.IsLoading())
 				{
-					if (MeshDescription.IsValid())
-					{
-						MeshDescription->Empty();
-					}
-					else
-					{
-						MeshDescription = MakeUnique<FMeshDescription>();
-						UStaticMesh::RegisterMeshAttributes(*MeshDescription);
-					}
+					CreateMeshDescription();
 				}
 
 				Ar << (*MeshDescription);
@@ -3659,27 +3691,10 @@ void FMeshSectionInfoMap::Serialize(FArchive& Ar)
 /**
  * Registers the mesh attributes required by the mesh description for a static mesh.
  */
-void UStaticMesh::RegisterMeshAttributes( FMeshDescription& MeshDescription )
+void UStaticMesh::RegisterMeshAttributes(FMeshDescription& MeshDescription)
 {
-	// Add basic vertex attributes
-	MeshDescription.VertexAttributes().RegisterAttribute<FVector>( MeshAttribute::Vertex::Position, 1, FVector::ZeroVector, EMeshAttributeFlags::Lerpable );
-	MeshDescription.VertexAttributes().RegisterAttribute<float>( MeshAttribute::Vertex::CornerSharpness, 1, 0.0f, EMeshAttributeFlags::Lerpable );
-
-	// Add basic vertex instance attributes
-	MeshDescription.VertexInstanceAttributes().RegisterAttribute<FVector2D>( MeshAttribute::VertexInstance::TextureCoordinate, 1, FVector2D::ZeroVector, EMeshAttributeFlags::Lerpable );
-	MeshDescription.VertexInstanceAttributes().RegisterAttribute<FVector>( MeshAttribute::VertexInstance::Normal, 1, FVector::ZeroVector, EMeshAttributeFlags::AutoGenerated );
-	MeshDescription.VertexInstanceAttributes().RegisterAttribute<FVector>( MeshAttribute::VertexInstance::Tangent, 1, FVector::ZeroVector, EMeshAttributeFlags::AutoGenerated );
-	MeshDescription.VertexInstanceAttributes().RegisterAttribute<float>( MeshAttribute::VertexInstance::BinormalSign, 1, 0.0f, EMeshAttributeFlags::AutoGenerated );
-	MeshDescription.VertexInstanceAttributes().RegisterAttribute<FVector4>( MeshAttribute::VertexInstance::Color, 1, FVector4( 1.0f, 1.0f, 1.0f, 1.0f ), EMeshAttributeFlags::Lerpable );
-
-	// Add basic edge attributes
-	MeshDescription.EdgeAttributes().RegisterAttribute<bool>( MeshAttribute::Edge::IsHard, 1, false );
-	MeshDescription.EdgeAttributes().RegisterAttribute<float>( MeshAttribute::Edge::CreaseSharpness, 1, 0.0f, EMeshAttributeFlags::Lerpable );
-
-	// Add basic polygon attributes
-
-	// Add basic polygon group attributes
-	MeshDescription.PolygonGroupAttributes().RegisterAttribute<FName>( MeshAttribute::PolygonGroup::ImportedMaterialSlotName ); //The unique key to match the mesh material slot
+	FStaticMeshAttributes Attributes(MeshDescription);
+	Attributes.Register();
 }
 
 
@@ -3727,7 +3742,9 @@ FMeshDescription* UStaticMesh::GetMeshDescription(int32 LodIndex) const
 
 	// Require a const_cast here, because GetMeshDescription should ostensibly have const semantics,
 	// but the lazy initialization (from the BulkData or the DDC) is a one-off event which breaks constness.
-	FStaticMeshSourceModel& SourceModel = const_cast<FStaticMeshSourceModel&>(GetSourceModel(LodIndex));
+	UStaticMesh* MutableThis = const_cast<UStaticMesh*>(this);
+
+	FStaticMeshSourceModel& SourceModel = MutableThis->GetSourceModel(LodIndex);
 
 	if (!SourceModel.MeshDescription.IsValid())
 	{
@@ -3771,8 +3788,7 @@ FMeshDescription* UStaticMesh::GetMeshDescription(int32 LodIndex) const
 	// If after all this we *still* don't have a valid MeshDescription, but there's a valid RawMesh, convert that to a MeshDescription.
 	if (!SourceModel.MeshDescription.IsValid() && !SourceModel.RawMeshBulkData->IsEmpty())
 	{
-		SourceModel.MeshDescription = MakeUnique<FMeshDescription>();
-		RegisterMeshAttributes(*SourceModel.MeshDescription);
+		SourceModel.CreateMeshDescription();
 
 		FRawMesh LodRawMesh;
 		SourceModel.LoadRawMesh(LodRawMesh);
@@ -3784,6 +3800,7 @@ FMeshDescription* UStaticMesh::GetMeshDescription(int32 LodIndex) const
 	return SourceModel.MeshDescription.Get();
 }
 
+
 bool UStaticMesh::IsMeshDescriptionValid(int32 LodIndex) const
 {
 	if (!IsSourceModelValid(LodIndex))
@@ -3797,18 +3814,17 @@ bool UStaticMesh::IsMeshDescriptionValid(int32 LodIndex) const
 	// If there is a valid MeshDescriptionBulkData, we know this implies a valid mesh description.
 	// If not, then it's a legacy asset which will have a mesh description built from the RawMeshBulkData, if non-empty.
 	return SourceModel.MeshDescription.IsValid() ||
-		   SourceModel.MeshDescriptionBulkData.IsValid() ||
-		   !SourceModel.RawMeshBulkData->IsEmpty();
+		SourceModel.MeshDescriptionBulkData.IsValid() ||
+		!SourceModel.RawMeshBulkData->IsEmpty();
 }
 
+
 FMeshDescription* UStaticMesh::CreateMeshDescription(int32 LodIndex)
 {
 	if (IsSourceModelValid(LodIndex))
 	{
 		FStaticMeshSourceModel& SourceModel = GetSourceModel(LodIndex);
-		SourceModel.MeshDescription = MakeUnique<FMeshDescription>();
-		UStaticMesh::RegisterMeshAttributes(*SourceModel.MeshDescription);
-		return SourceModel.MeshDescription.Get();
+		return SourceModel.CreateMeshDescription();
 	}
 
 	return nullptr;
@@ -3825,6 +3841,7 @@ FMeshDescription* UStaticMesh::CreateMeshDescription(int32 LodIndex, FMeshDescri
 	return NewMeshDescription;
 }
 
+
 void UStaticMesh::CommitMeshDescription(int32 LodIndex)
 {
 	// The source model must be created before calling this function
@@ -3871,6 +3888,7 @@ void UStaticMesh::ClearMeshDescription(int32 LodIndex)
 	}
 }
 
+
 void UStaticMesh::ClearMeshDescriptions()
 {
 	for (int LODIndex = 0; LODIndex < GetNumSourceModels(); LODIndex++)
@@ -3990,25 +4008,17 @@ void UStaticMesh::CacheMeshData()
 						SourceModel.RawMeshBulkData->LoadRawMesh(TempRawMesh);
 
 						// Create a new MeshDescription
-						if (SourceModel.MeshDescription.IsValid())
-						{
-							SourceModel.MeshDescription->Empty();
-						}
-						else
-						{
-							SourceModel.MeshDescription = MakeUnique<FMeshDescription>();
-							RegisterMeshAttributes(*SourceModel.MeshDescription);
-						}
+						FMeshDescription* MeshDescription = SourceModel.CreateMeshDescription();
 
 						// Convert the RawMesh to MeshDescription
 						TMap<int32, FName> MaterialMap;
 						FillMaterialName(StaticMaterials, MaterialMap);
-						FMeshDescriptionOperations::ConvertFromRawMesh(TempRawMesh, *SourceModel.MeshDescription, MaterialMap);
+						FMeshDescriptionOperations::ConvertFromRawMesh(TempRawMesh, *MeshDescription, MaterialMap);
 
 						// Pack MeshDescription into temporary bulk data, ready to write out to DDC.
 						// This will be reloaded from the DDC when needed if a MeshDescription is requested from the static mesh.
 						FMeshDescriptionBulkData MeshDescriptionBulkData;
-						MeshDescriptionBulkData.SaveMeshDescription(*SourceModel.MeshDescription);
+						MeshDescriptionBulkData.SaveMeshDescription(*MeshDescription);
 
 						// Write the DDC cache
 						TArray<uint8> DerivedData;
@@ -4142,7 +4152,9 @@ bool UStaticMesh::SetUVChannel(int32 LODIndex, int32 UVChannelIndex, const TMap<
 
 	Modify();
 
-	TMeshAttributesRef<FVertexInstanceID, FVector2D> UVs = MeshDescription->VertexInstanceAttributes().GetAttributesRef<FVector2D>(MeshAttribute::VertexInstance::TextureCoordinate);
+	FStaticMeshAttributes Attributes(*MeshDescription);
+
+	TMeshAttributesRef<FVertexInstanceID, FVector2D> UVs = Attributes.GetVertexInstanceUVs();
 	for (const FVertexInstanceID& VertexInstanceID : MeshDescription->VertexInstances().GetElementIDs())
 	{
 		if (const FVector2D* UVCoord = TexCoords.Find(VertexInstanceID))
@@ -4170,7 +4182,8 @@ int32 UStaticMesh::GetNumUVChannels(int32 LODIndex)
 	FMeshDescription* MeshDescription = GetMeshDescription(LODIndex);
 	if (MeshDescription)
 	{
-		NumUVChannels = MeshDescription->VertexInstanceAttributes().GetAttributeIndexCount<FVector2D>(MeshAttribute::VertexInstance::TextureCoordinate);
+		FStaticMeshConstAttributes Attributes(*MeshDescription);
+		NumUVChannels = Attributes.GetVertexInstanceUVs().GetNumIndices();
 	}
 #endif
 	return NumUVChannels;
@@ -4616,10 +4629,28 @@ void UStaticMesh::PostLoad()
 
 		FixupMaterialSlotName();
 
-		// This, among many other things, will build a MeshDescription from the legacy RawMesh if one has not already been serialized,
-		// or, failing that, if there is not already one in the DDC. This will remain cached until the end of PostLoad(), upon which it
-		// is then released, and can be reloaded on demand.
-		CacheDerivedData();
+		if (bIsBuiltAtRuntime)
+		{
+#if WITH_EDITOR
+			// If built at runtime, but an editor build, we cache the mesh descriptions so that they can be rebuilt within the editor if necessary.
+			// This is done through the fast build path for consistency
+			TArray<const FMeshDescription*> MeshDescriptions;
+			const int32 NumSourceModels = GetNumSourceModels();
+			MeshDescriptions.Reserve(NumSourceModels);
+			for (int32 SourceModelIndex = 0; SourceModelIndex < NumSourceModels; SourceModelIndex++)
+			{
+				MeshDescriptions.Add(GetMeshDescription(SourceModelIndex));
+			}
+			BuildFromMeshDescriptions(MeshDescriptions);
+#endif
+		}
+		else
+		{
+			// This, among many other things, will build a MeshDescription from the legacy RawMesh if one has not already been serialized,
+			// or, failing that, if there is not already one in the DDC. This will remain cached until the end of PostLoad(), upon which it
+			// is then released, and can be reloaded on demand.
+			CacheDerivedData();
+		}
 
 		//Fix up the material to remove redundant material, this is needed since the material refactor where we do not have anymore copy of the materials
 		//in the materials list
@@ -4855,6 +4886,285 @@ void UStaticMesh::PostLoad()
 	CreateNavCollision();
 }
 
+
+void UStaticMesh::BuildFromMeshDescription(const FMeshDescription& MeshDescription, FStaticMeshLODResources& LODResources)
+{
+	FStaticMeshConstAttributes MeshDescriptionAttributes(MeshDescription);
+
+	// Fill vertex buffers
+
+	int32 NumVertexInstances = MeshDescription.VertexInstances().GetArraySize();
+
+	TArray<FStaticMeshBuildVertex> StaticMeshBuildVertices;
+	StaticMeshBuildVertices.SetNum(NumVertexInstances);
+
+	TVertexAttributesConstRef<FVector> VertexPositions = MeshDescriptionAttributes.GetVertexPositions();
+	TVertexInstanceAttributesConstRef<FVector> VertexInstanceNormals = MeshDescriptionAttributes.GetVertexInstanceNormals();
+	TVertexInstanceAttributesConstRef<FVector> VertexInstanceTangents = MeshDescriptionAttributes.GetVertexInstanceTangents();
+	TVertexInstanceAttributesConstRef<float> VertexInstanceBinormalSigns = MeshDescriptionAttributes.GetVertexInstanceBinormalSigns();
+	TVertexInstanceAttributesConstRef<FVector4> VertexInstanceColors = MeshDescriptionAttributes.GetVertexInstanceColors();
+	TVertexInstanceAttributesConstRef<FVector2D> VertexInstanceUVs = MeshDescriptionAttributes.GetVertexInstanceUVs();
+
+	for (FVertexInstanceID VertexInstanceID : MeshDescription.VertexInstances().GetElementIDs())
+	{
+		FStaticMeshBuildVertex& StaticMeshVertex = StaticMeshBuildVertices[VertexInstanceID.GetValue()];
+
+		StaticMeshVertex.Position = VertexPositions[MeshDescription.GetVertexInstanceVertex(VertexInstanceID)];
+		StaticMeshVertex.TangentX = VertexInstanceTangents[VertexInstanceID];
+		StaticMeshVertex.TangentY = FVector::CrossProduct(VertexInstanceNormals[VertexInstanceID], VertexInstanceTangents[VertexInstanceID]).GetSafeNormal() * VertexInstanceBinormalSigns[VertexInstanceID];
+		StaticMeshVertex.TangentZ = VertexInstanceNormals[VertexInstanceID];
+
+		for (int32 UVIndex = 0; UVIndex < VertexInstanceUVs.GetNumIndices(); ++UVIndex)
+		{
+			StaticMeshVertex.UVs[UVIndex] = VertexInstanceUVs.Get(VertexInstanceID, UVIndex);
+		}
+	}
+
+	bool bHasVertexColors = false;
+	if (VertexInstanceColors.IsValid())
+	{
+		for (FVertexInstanceID VertexInstanceID : MeshDescription.VertexInstances().GetElementIDs())
+		{
+			FStaticMeshBuildVertex& StaticMeshVertex = StaticMeshBuildVertices[VertexInstanceID.GetValue()];
+
+			FLinearColor Color(VertexInstanceColors[VertexInstanceID]);
+			if (Color != FLinearColor::White)
+			{
+				bHasVertexColors = true;
+			}
+
+			StaticMeshVertex.Color = Color.ToFColor(true);
+		}
+	}
+
+	LODResources.VertexBuffers.PositionVertexBuffer.Init(StaticMeshBuildVertices);
+	LODResources.VertexBuffers.StaticMeshVertexBuffer.Init(StaticMeshBuildVertices, VertexInstanceUVs.GetNumIndices());
+
+	FColorVertexBuffer& ColorVertexBuffer = LODResources.VertexBuffers.ColorVertexBuffer;
+	if (bHasVertexColors)
+	{
+		ColorVertexBuffer.Init(StaticMeshBuildVertices);
+	}
+	else
+	{
+		ColorVertexBuffer.InitFromSingleColor(FColor::White, NumVertexInstances);
+	}
+
+	// Fill index buffer and sections array
+
+	int32 NumPolygonGroups = MeshDescription.PolygonGroups().Num();
+	int32 NumTriangles = MeshDescription.Triangles().Num();
+
+	TPolygonGroupAttributesConstRef<FName> MaterialSlotNames = MeshDescriptionAttributes.GetPolygonGroupMaterialSlotNames();
+
+	TArray<uint32> IndexBuffer;
+	IndexBuffer.SetNumZeroed(NumTriangles * 3);
+
+	TArray<FStaticMeshSection>& Sections = LODResources.Sections;
+
+	int32 SectionIndex = 0;
+	int32 IndexBufferIndex = 0;
+	EIndexBufferStride::Type IndexBufferStride = EIndexBufferStride::Force16Bit;
+
+	for (FPolygonGroupID PolygonGroupID : MeshDescription.PolygonGroups().GetElementIDs())
+	{
+		// Skip empty polygon groups - we do not want to build empty sections
+		if (MeshDescription.GetNumPolygonGroupPolygons(PolygonGroupID) == 0)
+		{
+			continue;
+		}
+
+		FStaticMeshSection& Section = Sections.AddDefaulted_GetRef();
+		Section.FirstIndex = IndexBufferIndex;
+
+		int32 TriangleCount = 0;
+		uint32 MinVertexIndex = TNumericLimits<uint32>::Max();
+		uint32 MaxVertexIndex = TNumericLimits<uint32>::Min();
+
+		for (FPolygonID PolygonID : MeshDescription.GetPolygonGroupPolygons(PolygonGroupID))
+		{
+			for (FTriangleID TriangleID : MeshDescription.GetPolygonTriangleIDs(PolygonID))
+			{
+				for (FVertexInstanceID TriangleVertexInstanceIDs : MeshDescription.GetTriangleVertexInstances(TriangleID))
+				{
+					uint32 VertexIndex = static_cast<uint32>(TriangleVertexInstanceIDs.GetValue());
+					MinVertexIndex = FMath::Min(MinVertexIndex, VertexIndex);
+					MaxVertexIndex = FMath::Max(MaxVertexIndex, VertexIndex);
+					IndexBuffer[IndexBufferIndex] = VertexIndex;
+					IndexBufferIndex++;
+				}
+
+				TriangleCount++;
+			}
+		}
+
+		Section.NumTriangles = TriangleCount;
+		Section.MinVertexIndex = MinVertexIndex;
+		Section.MaxVertexIndex = MaxVertexIndex;
+
+		const int32 MaterialIndex = StaticMaterials.IndexOfByPredicate(
+			[&MaterialSlotName = MaterialSlotNames[PolygonGroupID]](const FStaticMaterial& StaticMaterial) { return StaticMaterial.MaterialSlotName == MaterialSlotName; }
+		);
+
+		Section.MaterialIndex = MaterialIndex;
+		Section.bEnableCollision = true;
+		Section.bCastShadow = true;
+
+		if (MaxVertexIndex > TNumericLimits<uint16>::Max())
+		{
+			IndexBufferStride = EIndexBufferStride::Force32Bit;
+		}
+
+		SectionIndex++;
+	}
+	check(IndexBufferIndex == NumTriangles * 3);
+
+	LODResources.IndexBuffer.SetIndices(IndexBuffer, IndexBufferStride);
+
+	// Fill depth only index buffer
+
+	TArray<uint32> DepthOnlyIndexBuffer(IndexBuffer);
+	for (uint32& Index : DepthOnlyIndexBuffer)
+	{
+		// Compress all vertex instances into the same instance for each vertex
+		Index = MeshDescription.GetVertexVertexInstances(MeshDescription.GetVertexInstanceVertex(FVertexInstanceID(Index)))[0].GetValue();
+	}
+
+	LODResources.bHasDepthOnlyIndices = true;
+	LODResources.DepthOnlyIndexBuffer.SetIndices(DepthOnlyIndexBuffer, IndexBufferStride);
+
+	// Fill reversed index buffer
+	TArray<uint32> ReversedIndexBuffer(IndexBuffer);
+	for (int32 ReversedIndexBufferIndex = 0; ReversedIndexBufferIndex < IndexBuffer.Num(); ReversedIndexBufferIndex += 3)
+	{
+		Swap(ReversedIndexBuffer[ReversedIndexBufferIndex + 0], ReversedIndexBuffer[ReversedIndexBufferIndex + 2]);
+	}
+
+	LODResources.AdditionalIndexBuffers = new FAdditionalStaticMeshIndexBuffers();
+	LODResources.bHasReversedIndices = true;
+	LODResources.AdditionalIndexBuffers->ReversedIndexBuffer.SetIndices(ReversedIndexBuffer, IndexBufferStride);
+
+	// Fill reversed depth index buffer
+	TArray<uint32> ReversedDepthOnlyIndexBuffer(DepthOnlyIndexBuffer);
+	for (int32 ReversedIndexBufferIndex = 0; ReversedIndexBufferIndex < IndexBuffer.Num(); ReversedIndexBufferIndex += 3)
+	{
+		Swap(ReversedDepthOnlyIndexBuffer[ReversedIndexBufferIndex + 0], ReversedDepthOnlyIndexBuffer[ReversedIndexBufferIndex + 2]);
+	}
+
+	LODResources.bHasReversedDepthOnlyIndices = true;
+	LODResources.AdditionalIndexBuffers->ReversedDepthOnlyIndexBuffer.SetIndices(ReversedIndexBuffer, IndexBufferStride);
+
+	LODResources.bHasAdjacencyInfo = false;
+}
+
+
+UStaticMeshDescription* UStaticMesh::CreateStaticMeshDescription(UObject* Outer)
+{
+	if (Outer == nullptr)
+	{
+		Outer = GetTransientPackage();
+	}
+
+	UStaticMeshDescription* StaticMeshDescription = NewObject<UStaticMeshDescription>(Outer, NAME_None, RF_Transient);
+	StaticMeshDescription->RegisterAttributes();
+	return StaticMeshDescription;
+}
+
+
+void UStaticMesh::BuildFromStaticMeshDescriptions(const TArray<UStaticMeshDescription*>& StaticMeshDescriptions)
+{
+	TArray<const FMeshDescription*> MeshDescriptions;
+	MeshDescriptions.Reserve(StaticMeshDescriptions.Num());
+
+	for (UStaticMeshDescription* StaticMeshDescription : StaticMeshDescriptions)
+	{
+		MeshDescriptions.Emplace(&StaticMeshDescription->GetMeshDescription());
+	}
+
+	BuildFromMeshDescriptions(MeshDescriptions);
+}
+
+
+bool UStaticMesh::BuildFromMeshDescriptions(const TArray<const FMeshDescription*>& MeshDescriptions)
+{
+	// Set up
+
+	bIsBuiltAtRuntime = true;
+
+	TOptional<FStaticMeshComponentRecreateRenderStateContext> RecreateRenderStateContext;
+	
+	if (RenderData.IsValid())
+	{
+		const bool bInvalidateLighting = true;
+		const bool bRefreshBounds = true;
+		RecreateRenderStateContext = FStaticMeshComponentRecreateRenderStateContext(this, bInvalidateLighting, bRefreshBounds);
+	}
+
+	ReleaseResources();
+	ReleaseResourcesFence.Wait();
+
+	RenderData = MakeUnique<FStaticMeshRenderData>();
+	RenderData->AllocateLODResources(MeshDescriptions.Num());
+
+	// Build render data from each mesh description
+
+	int32 LODIndex = 0;
+	for (const FMeshDescription* MeshDescriptionPtr : MeshDescriptions)
+	{
+#if WITH_EDITOR
+		// Editor builds cache the mesh description so that it can be preserved during map reloads etc
+		CreateMeshDescription(LODIndex, *MeshDescriptionPtr);
+		CommitMeshDescription(LODIndex);
+#endif
+		check(MeshDescriptionPtr != nullptr);
+		FStaticMeshLODResources& LODResources = RenderData->LODResources[LODIndex];
+
+		BuildFromMeshDescription(*MeshDescriptionPtr, LODResources);
+
+		LODIndex++;
+	}
+
+	InitResources();
+
+	// Set up RenderData bounds and LOD data
+	RenderData->Bounds = MeshDescriptions[0]->GetBounds();
+	CalculateExtendedBounds();
+
+	for (int32 LOD = 0; LOD < MeshDescriptions.Num(); ++LOD)
+	{
+		// @todo: some way of customizing LOD screen size and/or calculate it based on mesh bounds
+		if (true)
+		{
+			const float LODPowerBase = 0.75f;
+			RenderData->ScreenSize[LOD].Default = FMath::Pow(LODPowerBase, LOD);
+		}
+		else
+		{
+			// Possible model for flexible LODs
+			const float MaxDeviation = 100.0f; // specify
+			const float PixelError = SMALL_NUMBER;
+			const float ViewDistance = (MaxDeviation * 960.0f) / PixelError;
+
+			// Generate a projection matrix.
+			const float HalfFOV = PI * 0.25f;
+			const float ScreenWidth = 1920.0f;
+			const float ScreenHeight = 1080.0f;
+			const FPerspectiveMatrix ProjMatrix(HalfFOV, ScreenWidth, ScreenHeight, 1.0f);
+
+			RenderData->ScreenSize[LOD].Default = ComputeBoundsScreenSize(FVector::ZeroVector, RenderData->Bounds.SphereRadius, FVector(0.0f, 0.0f, ViewDistance + RenderData->Bounds.SphereRadius), ProjMatrix);
+		}
+	}
+
+	// Set up physics-related data
+	CreateBodySetup();
+	check(BodySetup);
+	BodySetup->InvalidatePhysicsData();
+	BodySetup->CreatePhysicsMeshes();
+
+	return true;
+}
+
+
 bool UStaticMesh::CanBeClusterRoot() const
 {
 	return false;
@@ -4883,7 +5193,24 @@ int32 UStaticMesh::CalcNumOptionalMips() const
 #if !WITH_EDITOR
 	return MinLOD.Default;
 #else
-	return 0;
+	int32 NumOptionalLODs = 0;
+	if (RenderData)
+	{
+		const TIndirectArray<FStaticMeshLODResources>& LODResources = RenderData->LODResources;
+		for (int32 Idx = 0; Idx < LODResources.Num(); ++Idx)
+		{
+			const FStaticMeshLODResources& Resource = LODResources[Idx];
+			if (Resource.bIsOptionalLOD)
+			{
+				++NumOptionalLODs;
+			}
+			else
+			{
+				break;
+			}
+		}
+	}
+	return NumOptionalLODs;
 #endif
 }
 
@@ -4901,7 +5228,6 @@ int32 UStaticMesh::CalcCumulativeLODSize(int32 NumLODs) const
 
 bool UStaticMesh::GetMipDataFilename(const int32 MipIndex, FString& OutBulkDataFilename) const
 {
-#if !WITH_EDITOR
 	// TODO: this is slow. Should cache the name once per mesh
 	FString PackageName = GetOutermost()->FileName.ToString();
 	// Handle name redirection and localization
@@ -4917,9 +5243,6 @@ bool UStaticMesh::GetMipDataFilename(const int32 MipIndex, FString& OutBulkDataF
 	OutBulkDataFilename = FPaths::ChangeExtension(OutBulkDataFilename, MipIndex < MinLOD.Default ? TEXT(".uptnl") : TEXT(".ubulk"));
 	check(MipIndex < MinLOD.Default || IFileManager::Get().FileExists(*OutBulkDataFilename));
 	return true;
-#else
-	return false;
-#endif
 }
 
 bool UStaticMesh::IsReadyForStreaming() const
@@ -5767,6 +6090,30 @@ UMaterialInterface* UStaticMesh::GetMaterial(int32 MaterialIndex) const
 	return NULL;
 }
 
+
+FName UStaticMesh::AddMaterial(UMaterialInterface* Material)
+{
+	// Create a unique slot name for the material
+	FName MaterialName = Material->GetFName();
+	for (const FStaticMaterial& StaticMaterial : StaticMaterials)
+	{
+		const FName ExistingName = StaticMaterial.MaterialSlotName;
+		if (ExistingName.GetComparisonIndex() == MaterialName.GetComparisonIndex())
+		{
+			MaterialName = FName(MaterialName, FMath::Max(MaterialName.GetNumber(), ExistingName.GetNumber() + 1));
+		}
+	}
+
+#if WITH_EDITORONLY_DATA
+	StaticMaterials.Emplace(Material, MaterialName, MaterialName);
+#else
+	StaticMaterials.Emplace(Material, MaterialName);
+#endif
+
+	return MaterialName;
+}
+
+
 int32 UStaticMesh::GetMaterialIndex(FName MaterialSlotName) const
 {
 	for (int32 MaterialIndex = 0; MaterialIndex < StaticMaterials.Num(); ++MaterialIndex)
diff --git a/Engine/Source/Runtime/Engine/Private/TextureDerivedData.cpp b/Engine/Source/Runtime/Engine/Private/TextureDerivedData.cpp
--- a/Engine/Source/Runtime/Engine/Private/TextureDerivedData.cpp
+++ b/Engine/Source/Runtime/Engine/Private/TextureDerivedData.cpp
@@ -147,6 +147,7 @@ static void SerializeForKey(FArchive& Ar, const FTextureBuildSettings& Settings)
 		TempByte = Settings.bVirtualTextureEnableCompressZlib; Ar << TempByte;
 		TempByte = Settings.bVirtualTextureEnableCompressCrunch; Ar << TempByte;
 		TempByte = Settings.LossyCompressionAmount; Ar << TempByte; // Lossy compression currently only used by VT
+		TempByte = Settings.bApplyYCoCgBlockScale; Ar << TempByte; // YCoCg currently only used by VT
 	}
 }
 
@@ -310,6 +311,7 @@ static void FinalizeBuildSettingsForLayer(const UTexture& Texture, int32 LayerIn
 
 	OutSettings.bHDRSource = Texture.HasHDRSource(LayerIndex);
 	OutSettings.bSRGB = FormatSettings.SRGB;
+	OutSettings.bApplyYCoCgBlockScale = FormatSettings.CompressionYCoCg;
 
 	if (FormatSettings.CompressionSettings == TC_Displacementmap || FormatSettings.CompressionSettings == TC_DistanceFieldFont)
 	{
diff --git a/Engine/Source/Runtime/Engine/Private/World.cpp b/Engine/Source/Runtime/Engine/Private/World.cpp
--- a/Engine/Source/Runtime/Engine/Private/World.cpp
+++ b/Engine/Source/Runtime/Engine/Private/World.cpp
@@ -6119,7 +6119,7 @@ UWorld* FSeamlessTravelHandler::Tick()
 				const bool bIsInCurrentLevel	= TheActor->GetLevel() == CurrentWorld->PersistentLevel;
 				const bool bManuallyMarkedKeep	= KeepAnnotation.Get(TheActor);
 				const bool bDormant				= NetworkObjectInfo && NetDriver && NetDriver->ServerConnection && NetworkObjectInfo->DormantConnections.Contains(NetDriver->ServerConnection);
-				const bool bKeepNonOwnedActor	= TheActor->Role < ROLE_Authority && !bDormant && !TheActor->IsNetStartupActor();
+				const bool bKeepNonOwnedActor	= TheActor->GetLocalRole() < ROLE_Authority && !bDormant && !TheActor->IsNetStartupActor();
 				const bool bForceExcludeActor	= TheActor->IsA(ALevelScriptActor::StaticClass());
 
 				// Keep if it's in the current level AND it isn't specifically excluded AND it was either marked as should keep OR we don't own this actor
@@ -7509,6 +7509,63 @@ void UWorld::InitializeSubsystems()
 	}
 }
 
+static void DoPostProcessVolume(IInterface_PostProcessVolume* Volume, FVector ViewLocation, FSceneView* SceneView)
+{
+	const FPostProcessVolumeProperties VolumeProperties = Volume->GetProperties();
+	if (!VolumeProperties.bIsEnabled)
+	{
+		return;
+	}
+
+	float DistanceToPoint = 0.0f;
+	float LocalWeight = FMath::Clamp(VolumeProperties.BlendWeight, 0.0f, 1.0f);
+
+	if (!VolumeProperties.bIsUnbound)
+	{
+		float SquaredBlendRadius = VolumeProperties.BlendRadius * VolumeProperties.BlendRadius;
+		Volume->EncompassesPoint(ViewLocation, 0.0f, &DistanceToPoint);
+
+		if (DistanceToPoint >= 0)
+		{
+			if (DistanceToPoint > VolumeProperties.BlendRadius)
+			{
+				// outside
+				LocalWeight = 0.0f;
+			}
+			else
+			{
+				// to avoid div by 0
+				if (VolumeProperties.BlendRadius >= 1.0f)
+				{
+					LocalWeight *= 1.0f - DistanceToPoint / VolumeProperties.BlendRadius;
+
+					check(LocalWeight >= 0 && LocalWeight <= 1.0f);
+				}
+			}
+		}
+		else
+		{
+			LocalWeight = 0;
+		}
+	}
+
+	if (LocalWeight > 0)
+	{
+		SceneView->OverridePostProcessSettings(*VolumeProperties.Settings, LocalWeight);
+	}
+}
+
+void UWorld::AddPostProcessingSettings(FVector ViewLocation, FSceneView* SceneView)
+{
+	OnBeginPostProcessSettings.Broadcast(ViewLocation);
+
+	for (IInterface_PostProcessVolume* PPVolume : PostProcessVolumes)
+	{
+		DoPostProcessVolume(PPVolume, ViewLocation, SceneView);
+	}
+}
+
+
 /**
 * Dump visible actors in current world.
 */
diff --git a/Engine/Source/Runtime/EngineSettings/Classes/GameMapsSettings.h b/Engine/Source/Runtime/EngineSettings/Classes/GameMapsSettings.h
--- a/Engine/Source/Runtime/EngineSettings/Classes/GameMapsSettings.h
+++ b/Engine/Source/Runtime/EngineSettings/Classes/GameMapsSettings.h
@@ -190,4 +190,22 @@ private:
 	/** List of GameModes to load when game= is specified in the URL (e.g. "DM" could be an alias for "MyProject.MyGameModeMP_DM") */
 	UPROPERTY(config, EditAnywhere, Category = DefaultModes, AdvancedDisplay)
 	TArray<FGameModeName> GameModeClassAliases;
+
+public:
+
+	/** Returns the game local maps settings */
+	UFUNCTION(BlueprintPure, Category = Settings, meta=(DisplayName="Get Game Maps and Modes Settings"))
+	static UGameMapsSettings* GetGameMapsSettings();
+
+	/**
+	 * Modify "Skip Assigning Gamepad to Player 1" GameMapsSettings option
+	 * @param bSkipFirstPlayer		If set connected game pads will only be assigned to the second and subsequent players
+	 * @note This value is saved to local config when changed.
+	*/
+	UFUNCTION(BlueprintCallable, Category = "Utilities")
+	void SetSkipAssigningGamepadToPlayer1(bool bSkipFirstPlayer = true);
+
+	UFUNCTION(BlueprintPure, Category = "Utilities")
+	bool GetSkipAssigningGamepadToPlayer1() const;
+
 };
diff --git a/Engine/Source/Runtime/EngineSettings/Private/EngineSettingsModule.cpp b/Engine/Source/Runtime/EngineSettings/Private/EngineSettingsModule.cpp
--- a/Engine/Source/Runtime/EngineSettings/Private/EngineSettingsModule.cpp
+++ b/Engine/Source/Runtime/EngineSettings/Private/EngineSettingsModule.cpp
@@ -213,6 +213,21 @@ void UGameMapsSettings::SetGlobalDefaultGameMode( const FString& NewGameMode )
 	GameMapsSettings->GlobalDefaultGameMode = NewGameMode;
 }
 
+UGameMapsSettings* UGameMapsSettings::GetGameMapsSettings()
+{
+	return GetMutableDefault<UGameMapsSettings>();
+}
+
+void UGameMapsSettings::SetSkipAssigningGamepadToPlayer1(bool bSkipFirstPlayer)
+{
+	bOffsetPlayerGamepadIds = bSkipFirstPlayer;
+}
+
+bool UGameMapsSettings::GetSkipAssigningGamepadToPlayer1() const
+{
+	return bOffsetPlayerGamepadIds;
+}
+
 // Backwards compat for map strings
 void FixMapAssetRef(FSoftObjectPath& MapAssetReference)
 {
diff --git a/Engine/Source/Runtime/Landscape/Private/LandscapeGrass.cpp b/Engine/Source/Runtime/Landscape/Private/LandscapeGrass.cpp
--- a/Engine/Source/Runtime/Landscape/Private/LandscapeGrass.cpp
+++ b/Engine/Source/Runtime/Landscape/Private/LandscapeGrass.cpp
@@ -593,7 +593,8 @@ public:
 		check(RenderTargetTexture);
 		RenderTargetTexture->ClearColor = FLinearColor::White;
 		RenderTargetTexture->TargetGamma = 1.0f;
-		RenderTargetTexture->InitCustomFormat(TargetSize.X, TargetSize.Y, PF_B8G8R8A8, false);
+		const bool bForceLinearGamma = true;
+		RenderTargetTexture->InitCustomFormat(TargetSize.X, TargetSize.Y, PF_B8G8R8A8, bForceLinearGamma);
 		RenderTargetResource = RenderTargetTexture->GameThread_GetRenderTargetResource()->GetTextureRenderTarget2DResource();
 
 		// render
@@ -1476,7 +1477,7 @@ struct FGrassBuilderBase
 		const float DensityScale = bEnableDensityScaling ? CVarGrassDensityScale.GetValueOnAnyThread() : 1.0f;
 		GrassDensity = GrassVariety.GrassDensity.GetValueForFeatureLevel(FeatureLevel) * DensityScale;
 
-		DrawScale = Landscape->GetRootComponent()->RelativeScale3D;
+		DrawScale = Landscape->GetRootComponent()->GetRelativeScale3D();
 		DrawLoc = Landscape->GetActorLocation();
 		LandscapeSectionOffset = Landscape->LandscapeSectionOffset;
 
diff --git a/Engine/Source/Runtime/Launch/Private/LaunchEngineLoop.cpp b/Engine/Source/Runtime/Launch/Private/LaunchEngineLoop.cpp
--- a/Engine/Source/Runtime/Launch/Private/LaunchEngineLoop.cpp
+++ b/Engine/Source/Runtime/Launch/Private/LaunchEngineLoop.cpp
@@ -2429,7 +2429,7 @@ int32 FEngineLoop::PreInit(const TCHAR* CmdLine)
 					GUseRHIThread_InternalUseOnly = false;
 				}
 			}
-
+				
 			SCOPED_BOOT_TIMING("StartRenderingThread");
 			StartRenderingThread();
 		}
@@ -4994,10 +4994,6 @@ bool FEngineLoop::AppInit( )
 		FPlatformStackWalk::Init();
 	}
 
-#if WITH_EDITOR
-	FBlueprintSupport::InitializeCompilationManager();
-#endif
-
 	CheckForPrintTimesOverride();
 
 	// Check whether the project or any of its plugins are missing or are out of date
diff --git a/Engine/Source/Runtime/MeshDescription/Public/MeshDescription.h b/Engine/Source/Runtime/MeshDescription/Public/MeshDescription.h
--- a/Engine/Source/Runtime/MeshDescription/Public/MeshDescription.h
+++ b/Engine/Source/Runtime/MeshDescription/Public/MeshDescription.h
@@ -357,7 +357,7 @@ public:
 
 	// Mesh description should be a moveable type.
 	// Hence explicitly declare all the below as defaulted, to ensure they will be generated by the compiler.
-	FMeshDescription() = default;
+	FMeshDescription();
 	~FMeshDescription() = default;
 	FMeshDescription(const FMeshDescription&) = default;
 	FMeshDescription(FMeshDescription&&) = default;
@@ -1287,41 +1287,53 @@ public:
 	/** Remaps the element IDs in the mesh description according to the passed in object */
 	void Remap( const FElementIDRemappings& Remappings );
 
+	/** Returns bounds of vertices */
+	FBoxSphereBounds GetBounds() const;
+
+	/** Retriangulates the entire mesh */
 	void TriangulateMesh();
 
 	/** Set the polygon tangent and normal only for the specified polygonIDs */
+	UE_DEPRECATED(4.25, "Please use FStaticMeshOperations::ComputePolygonTangentsAndNormals.")
 	void ComputePolygonTangentsAndNormals(const TArray<FPolygonID>& PolygonIDs, float ComparisonThreshold = 0.0f);
 
 	/** Set the polygon tangent and normal for all polygons in the mesh description. */
+	UE_DEPRECATED(4.25, "Please use FStaticMeshOperations::ComputePolygonTangentsAndNormals.")
 	void ComputePolygonTangentsAndNormals(float ComparisonThreshold = 0.0f);
 	
 	/** Set the vertex instance tangent and normal only for the specified VertexInstanceIDs */
+	UE_DEPRECATED(4.25, "Please use FStaticMeshOperations::ComputeTangentsAndNormals.")
 	void ComputeTangentsAndNormals(const TArray<FVertexInstanceID>& VertexInstanceIDs, EComputeNTBsOptions ComputeNTBsOptions);
 
 	/** Set the vertex instance tangent and normal for all vertex instances in the mesh description. */
+	UE_DEPRECATED(4.25, "Please use FStaticMeshOperations::ComputeTangentsAndNormals.")
 	void ComputeTangentsAndNormals(EComputeNTBsOptions ComputeNTBsOptions);
 
 	/** Determine the edge hardnesses from existing normals */
+	UE_DEPRECATED(4.25, "Please use FStaticMeshOperations::DetermineEdgeHardnessesFromVertexInstanceNormals.")
 	void DetermineEdgeHardnessesFromVertexInstanceNormals( const float Tolerance = KINDA_SMALL_NUMBER );
 
 	/** Determine UV seams from existing vertex instance UVs */
+	UE_DEPRECATED(4.25, "UVSeam attribute is now deprecated in MeshDescription.")
 	void DetermineUVSeamsFromUVs( const int32 UVIndex, const float Tolerance = KINDA_SMALL_NUMBER );
 
 	/** Get polygons in the same UV chart as the specified polygon */
+	UE_DEPRECATED(4.25, "UVSeam attribute is now deprecated in MeshDescription.")
 	void GetPolygonsInSameChartAsPolygon( const FPolygonID PolygonID, TArray<FPolygonID>& OutPolygonIDs );
 
 	/** Get array of all UV charts */
+	UE_DEPRECATED(4.25, "UVSeam attribute is now deprecated in MeshDescription.")
 	void GetAllCharts( TArray<TArray<FPolygonID>>& OutCharts );
 
+	/** Reverses the winding order of all polygons in the mesh */
 	void ReverseAllPolygonFacing();
 
 	float GetPolygonCornerAngleForVertex(const FPolygonID PolygonID, const FVertexID VertexID) const;
-
-private:
-
 	FPlane ComputePolygonPlane(const FPolygonID PolygonID) const;
 	FVector ComputePolygonNormal(const FPolygonID PolygonID) const;
 
+private:
+
 	bool ComputePolygonTangentsAndNormals(
 		  const FPolygonID PolygonID
 		, float ComparisonThreshold
diff --git a/Engine/Source/Runtime/Renderer/Private/SceneRendering.h b/Engine/Source/Runtime/Renderer/Private/SceneRendering.h
--- a/Engine/Source/Runtime/Renderer/Private/SceneRendering.h
+++ b/Engine/Source/Runtime/Renderer/Private/SceneRendering.h
@@ -797,6 +797,9 @@ struct FPreviousViewInfo
 	// Temporal AA result of last frame
 	FTemporalAAHistory TemporalAAHistory;
 
+	// Half resolution version temporal AA result of last frame
+	TRefCountPtr<IPooledRenderTarget> HalfResTemporalAAHistory;
+
 	// Temporal AA history for diaphragm DOF.
 	FTemporalAAHistory DOFSetupHistory;
 	
