Merging //UE4/Main @ 9388175 to //UE4/Dev-Core (stefan.boberg_rdu-p4merge-02_6162)

(ushell-p4-mergedown)

diff --git a/Engine/Plugins/Experimental/GeometryCollectionPlugin/Source/GeometryCollectionExampleCore/Private/GeometryCollection/GeometryCollectionExampleUtility.cpp b/Engine/Plugins/Experimental/GeometryCollectionPlugin/Source/GeometryCollectionExampleCore/Private/GeometryCollection/GeometryCollectionExampleUtility.cpp
--- a/Engine/Plugins/Experimental/GeometryCollectionPlugin/Source/GeometryCollectionExampleCore/Private/GeometryCollection/GeometryCollectionExampleUtility.cpp
+++ b/Engine/Plugins/Experimental/GeometryCollectionPlugin/Source/GeometryCollectionExampleCore/Private/GeometryCollection/GeometryCollectionExampleUtility.cpp
@@ -40,14 +40,12 @@ namespace GeometryCollectionExample {
 
 	void FinalizeSolver(Chaos::FPBDRigidsSolver& InSolver)
 	{
-#if INCLUDE_CHAOS
 		InSolver.ForEachPhysicsProxy([](auto* Object)
 		{
 			Object->BufferPhysicsResults();
 			Object->FlipBuffer();
 			Object->PullFromPhysicsState();
 		});
-#endif
 	}
 
 
@@ -56,7 +54,6 @@ namespace GeometryCollectionExample {
 	TSharedPtr<FGeometryCollection>	CreateClusteredBody(FVector Position)
 	{
 		TSharedPtr<FGeometryCollection> RestCollection;
-#if INCLUDE_CHAOS
 		RestCollection = GeometryCollection::MakeCubeElement(FTransform(FQuat::MakeFromEuler(FVector(0.f)), Position), FVector(1.0));
 		RestCollection->AppendGeometry(*GeometryCollection::MakeCubeElement(FTransform(FQuat::MakeFromEuler(FVector(0.f)), FVector(30.f)), FVector(1.0)));
 
@@ -66,14 +63,12 @@ namespace GeometryCollectionExample {
 
 		GeometryCollectionAlgo::ParentTransforms(RestCollection.Get(), 2, { 0,1 });
 
-#endif
 		return RestCollection;
 	}
 
 	TSharedPtr<FGeometryCollection>	CreateClusteredBody_TwoParents_TwoBodies(FVector Position)
 	{
 		TSharedPtr<FGeometryCollection> RestCollection;
-#if INCLUDE_CHAOS
 		RestCollection = GeometryCollection::MakeCubeElement(FTransform(FQuat::MakeFromEuler(FVector(0.f)), Position), FVector(1.0));
 		RestCollection->AppendGeometry(*GeometryCollection::MakeCubeElement(FTransform(FQuat::MakeFromEuler(FVector(0.f)), FVector(30.f)), FVector(1.0)));
 
@@ -86,14 +81,12 @@ namespace GeometryCollectionExample {
 		GeometryCollectionAlgo::ParentTransforms(RestCollection.Get(), 2, { 0,1 });
 		GeometryCollectionAlgo::ParentTransforms(RestCollection.Get(), 3, { 2 });
 
-#endif
 		return RestCollection;
 	}
 
 	TSharedPtr<FGeometryCollection>	CreateClusteredBody_FourParents_OneBody(FVector Position)
 	{
 		TSharedPtr<FGeometryCollection> RestCollection;
-#if INCLUDE_CHAOS
 		RestCollection = GeometryCollection::MakeCubeElement(FTransform(FQuat::MakeFromEuler(FVector(0.f)), Position), FVector(1.0));
 
 		RestCollection->AddElements(4, FGeometryCollection::TransformGroup);
@@ -109,7 +102,6 @@ namespace GeometryCollectionExample {
 		GeometryCollectionAlgo::ParentTransforms(RestCollection.Get(), 3, { 2 });
 		GeometryCollectionAlgo::ParentTransforms(RestCollection.Get(), 4, { 3 });
 
-#endif
 		return RestCollection;
 	}
 
@@ -118,7 +110,6 @@ namespace GeometryCollectionExample {
 	TSharedPtr<FGeometryCollection>	CreateClusteredBody_TwoByTwo_ThreeTransform(FVector Position)
 	{
 		TSharedPtr<FGeometryCollection> RestCollection;
-#if INCLUDE_CHAOS
 		RestCollection = GeometryCollection::MakeCubeElement(FTransform(FQuat::MakeFromEuler(FVector(0.f)), FVector(0,0,0)), FVector(1.0));
 		RestCollection->AppendGeometry(*GeometryCollection::MakeCubeElement(FTransform(FQuat::MakeFromEuler(FVector(0.f)), FVector(100,0,0)), FVector(1.0)));
 		RestCollection->AppendGeometry(*GeometryCollection::MakeCubeElement(FTransform(FQuat::MakeFromEuler(FVector(0.f)), FVector(200,0,0)), FVector(1.0)));
@@ -140,7 +131,6 @@ namespace GeometryCollectionExample {
 		GeometryCollectionAlgo::ParentTransforms(RestCollection.Get(), 5, { 2,3 });
 		GeometryCollectionAlgo::ParentTransforms(RestCollection.Get(), 6, { 4,5 });
 
-#endif
 		return RestCollection;
 	}
 
@@ -149,7 +139,6 @@ namespace GeometryCollectionExample {
 	TSharedPtr<FGeometryCollection>	CreateClusteredBody_ThreeByTwo_ThreeTransform(FVector Position)
 	{
 		TSharedPtr<FGeometryCollection> RestCollection;
-#if INCLUDE_CHAOS
 		RestCollection = GeometryCollection::MakeCubeElement(FTransform(FQuat::MakeFromEuler(FVector(0.f)), FVector(0, 0, 0)), FVector(1.0));
 		RestCollection->AppendGeometry(*GeometryCollection::MakeCubeElement(FTransform(FQuat::MakeFromEuler(FVector(0.f)), FVector(100, 0, 0)), FVector(1.0)));
 		RestCollection->AppendGeometry(*GeometryCollection::MakeCubeElement(FTransform(FQuat::MakeFromEuler(FVector(0.f)), FVector(200, 0, 0)), FVector(1.0)));
@@ -175,7 +164,6 @@ namespace GeometryCollectionExample {
 		GeometryCollectionAlgo::ParentTransforms(RestCollection.Get(), 7, { 3,4,5 });
 		GeometryCollectionAlgo::ParentTransforms(RestCollection.Get(), 8, { 6,7 });
 
-#endif
 		return RestCollection;
 	}
 
@@ -183,7 +171,6 @@ namespace GeometryCollectionExample {
 	TSharedPtr<FGeometryCollection>	CreateClusteredBody_FracturedGeometry(FVector Position)
 	{
 		TSharedPtr<FGeometryCollection> RestCollection;
-#if INCLUDE_CHAOS
 
 		RestCollection = TSharedPtr<FGeometryCollection>(FGeometryCollection::NewGeometryCollection(
 			FracturedGeometry::RawVertexArray,
@@ -213,7 +200,6 @@ namespace GeometryCollectionExample {
 		RestCollection->SimulationType[12] = FGeometryCollection::ESimulationTypes::FST_Clustered;
 		RestCollection->SimulationType[0] = FGeometryCollection::ESimulationTypes::FST_Clustered;
 
-#endif
 		return RestCollection;
 	}
 
@@ -298,7 +284,6 @@ namespace GeometryCollectionExample {
 		return PhysObject;
 #else
 		return nullptr;
-#endif
 	}
 	template FGeometryCollectionPhysicsProxy* RigidBodySetup(
 		TUniquePtr<Chaos::TChaosPhysicsMaterial<float>> & PhysicalMaterial,
diff --git a/Engine/Source/Developer/SkeletalMeshUtilitiesCommon/Private/LODUtilities.cpp b/Engine/Source/Developer/SkeletalMeshUtilitiesCommon/Private/LODUtilities.cpp
--- a/Engine/Source/Developer/SkeletalMeshUtilitiesCommon/Private/LODUtilities.cpp
+++ b/Engine/Source/Developer/SkeletalMeshUtilitiesCommon/Private/LODUtilities.cpp
@@ -14,8 +14,8 @@
 #include "Engine/SkeletalMesh.h"
 #include "EditorFramework/AssetImportData.h"
 #include "MeshUtilities.h"
+#include "ClothingAsset.h"
 #include "OverlappingCorners.h"
-#include "Assets/ClothingAsset.h"
 #include "Framework/Commands/UIAction.h"
 
 #include "ObjectTools.h"
diff --git a/Engine/Source/Editor/ClothPainter/Private/SClothPaintTab.cpp b/Engine/Source/Editor/ClothPainter/Private/SClothPaintTab.cpp
--- a/Engine/Source/Editor/ClothPainter/Private/SClothPaintTab.cpp
+++ b/Engine/Source/Editor/ClothPainter/Private/SClothPaintTab.cpp
@@ -27,7 +27,7 @@
 #include "SClothPaintWidget.h"
 #include "IPersonaToolkit.h"
 #include "SClothAssetSelector.h"
-#include "Assets/ClothingAsset.h"
+#include "ClothingAsset.h"
 #include "Widgets/Layout/SScrollBox.h"
 #include "ComponentReregisterContext.h"
 
@@ -51,7 +51,7 @@ SClothPaintTab::~SClothPaintTab()
 
 void SClothPaintTab::Construct(const FArguments& InArgs)
 {
-	// Detail view for UClothingAsset
+	// Detail view for UClothingAssetCommon
 	FPropertyEditorModule& EditModule = FModuleManager::Get().GetModuleChecked<FPropertyEditorModule>("PropertyEditor");
 
 	FDetailsViewArgs DetailsViewArgs(
@@ -75,7 +75,7 @@ void SClothPaintTab::Construct(const FArguments& InArgs)
 
 	// Add the CDO by default
 	TArray<UObject*> Objects;
-	Objects.Add(UClothingAsset::StaticClass()->GetDefaultObject());
+	Objects.Add(UClothingAssetCommon::StaticClass()->GetDefaultObject());
 	DetailsView->SetObjects(Objects, true);
 
 	HostingApp = InArgs._InHostingApp;
@@ -165,7 +165,7 @@ void SClothPaintTab::UpdatePaintTools()
 
 			if(SelectorWidget.IsValid())
 			{
-				TWeakObjectPtr<UClothingAsset> WeakAsset = SelectorWidget->GetSelectedAsset();
+				TWeakObjectPtr<UClothingAssetCommon> WeakAsset = SelectorWidget->GetSelectedAsset();
 
 				if(WeakAsset.Get())
 				{
@@ -183,7 +183,7 @@ void SClothPaintTab::UpdatePaintTools()
 	}
 }
 
-void SClothPaintTab::OnAssetSelectionChanged(TWeakObjectPtr<UClothingAsset> InAssetPtr, int32 InLodIndex, int32 InMaskIndex)
+void SClothPaintTab::OnAssetSelectionChanged(TWeakObjectPtr<UClothingAssetCommon> InAssetPtr, int32 InLodIndex, int32 InMaskIndex)
 {
 	if(bPaintModeEnabled)
 	{
@@ -201,7 +201,7 @@ void SClothPaintTab::OnAssetSelectionChanged(TWeakObjectPtr<UClothingAsset> InAs
 		}
 	}
 
-	if(UClothingAsset* Asset = InAssetPtr.Get())
+	if(UClothingAssetCommon* Asset = InAssetPtr.Get())
 	{
 		TArray<UObject*> Objects;
 
@@ -220,7 +220,7 @@ bool SClothPaintTab::IsAssetDetailsPanelEnabled()
 
 		if(SelectedObjects.Num() > 0)
 		{
-			return SelectedObjects[0].Get() != UClothingAsset::StaticClass()->GetDefaultObject();
+			return SelectedObjects[0].Get() != UClothingAssetCommon::StaticClass()->GetDefaultObject();
 		}
 	}
 
diff --git a/Engine/Source/Editor/Persona/Private/PersonaMeshDetails.cpp b/Engine/Source/Editor/Persona/Private/PersonaMeshDetails.cpp
--- a/Engine/Source/Editor/Persona/Private/PersonaMeshDetails.cpp
+++ b/Engine/Source/Editor/Persona/Private/PersonaMeshDetails.cpp
@@ -51,7 +51,7 @@
 	#include "ApexClothingOptionWindow.h"
 #endif // #if WITH_APEX_CLOTHING
 
-#include "Assets/ClothingAsset.h"
+#include "ClothingAsset.h"
 
 #include "LODUtilities.h"
 #include "MeshUtilities.h"
@@ -2082,7 +2082,7 @@ void FPersonaMeshDetails::AddLODLevelCategories(IDetailLayoutBuilder& DetailLayo
 		const int32 SkelMeshLODCount = SkelMesh->GetLODNum();
 
 
-#if WITH_APEX_CLOTHING
+#if WITH_APEX_CLOTHING || WITH_CHAOS_CLOTHING
 		ClothComboBoxes.Reset();
 #endif
 
@@ -4060,8 +4060,8 @@ TSharedRef<SWidget> FPersonaMeshDetails::OnGenerateCustomSectionWidgetsForSectio
 		}
 	}
 
+#if WITH_APEX_CLOTHING || WITH_CHAOS_CLOTHING
 
-#if WITH_APEX_CLOTHING
 	UpdateClothingEntries();
 
 	ClothComboBoxes.AddDefaulted();
@@ -4101,7 +4101,7 @@ TSharedRef<SWidget> FPersonaMeshDetails::OnGenerateCustomSectionWidgetsForSectio
 			]
 		]
 	];
-#endif// #if WITH_APEX_CLOTHING
+#endif// #if WITH_APEX_CLOTHING || WITH_CHAOS_CLOTHING
 	SectionWidget->AddSlot()
 	.AutoHeight()
 	.Padding(0, 2, 0, 0)
@@ -5091,7 +5091,7 @@ void FPersonaMeshDetails::OnGenerateElementForClothingAsset( TSharedRef<IPropert
 		MakeClothingDetailsWidget(ElementIndex)
 	];	
 	
-	// Properties are now inside UClothingAsset, so we just add a new inspector and handle everything through that
+	// Properties are now inside UClothingAssetCommon, so we just add a new inspector and handle everything through that
 	FDetailWidgetRow& ClothPropRow = ChildrenBuilder.AddCustomRow(LOCTEXT("ClothingAsset_Search_Properties", "Properties"));
 
 	TSharedPtr<SKismetInspector> Inspector = nullptr;
@@ -5153,11 +5153,12 @@ TSharedRef<SUniformGridPanel> FPersonaMeshDetails::MakeClothingDetailsWidget(int
 
 		RowNumber++;
 
-		if (UClothingAsset* Asset = Cast<UClothingAsset>(ClothingAsset))
+		if (UClothingAssetCommon* Asset = Cast<UClothingAssetCommon>(ClothingAsset))
 		{
-			FClothLODData& LodData = Asset->LodData[LODIndex];
-			FClothPhysicalMeshData& PhysMeshData = LodData.PhysicalMeshData;
-			FClothCollisionData& CollisionData = LodData.CollisionData;
+			UClothLODDataBase* LodData = Asset->ClothLodData[LODIndex];
+			check(LodData->PhysicalMeshData);
+			UClothPhysicalMeshDataBase& PhysMeshData = *LodData->PhysicalMeshData;
+			FClothCollisionData& CollisionData = LodData->CollisionData;
 
 			Grid->AddSlot(0, RowNumber)
 				.HAlign(HAlign_Center)
@@ -5432,7 +5433,7 @@ void FPersonaMeshDetails::OnClothingComboBoxOpening()
 
 TSharedRef<SWidget> FPersonaMeshDetails::OnGenerateWidgetForClothingEntry(TSharedPtr<FClothingEntry> InEntry)
 {
-	UClothingAsset* Asset = Cast<UClothingAsset>(InEntry->Asset.Get());
+	UClothingAssetCommon* Asset = Cast<UClothingAssetCommon>(InEntry->Asset.Get());
 
 	FText EntryText;
 	if(Asset)
@@ -5454,7 +5455,7 @@ FText FPersonaMeshDetails::OnGetClothingComboText(int32 InLodIdx, int32 InSectio
 
 	if(Mesh)
 	{
-		UClothingAsset* ClothingAsset = Cast<UClothingAsset>(Mesh->GetSectionClothingAsset(InLodIdx, InSectionIdx));
+		UClothingAssetCommon* ClothingAsset = Cast<UClothingAssetCommon>(Mesh->GetSectionClothingAsset(InLodIdx, InSectionIdx));
 
 		if(ClothingAsset && ClothingAsset->LodMap.IsValidIndex(InLodIdx))
 		{
@@ -5494,7 +5495,7 @@ void FPersonaMeshDetails::OnClothingSelectionChanged(TSharedPtr<FClothingEntry>
 			OriginalSectionData.ClothingData.AssetGuid = FGuid();
 			OriginalSectionData.ClothingData.AssetLodIndex = INDEX_NONE;
 		};
-		if (UClothingAsset* ClothingAsset = Cast<UClothingAsset>(InNewEntry->Asset.Get()))
+		if (UClothingAssetCommon* ClothingAsset = Cast<UClothingAssetCommon>(InNewEntry->Asset.Get()))
 		{
 			// Look for a currently bound asset an unbind it if necessary first
 			if (UClothingAssetBase* CurrentAsset = Mesh->GetSectionClothingAsset(InLodIdx, InSectionIdx))
@@ -5541,20 +5542,15 @@ bool FPersonaMeshDetails::IsClothingPanelEnabled() const
 
 void FPersonaMeshDetails::OnFinishedChangingClothingProperties(const FPropertyChangedEvent& Event, int32 InAssetIndex)
 {
-	if(Event.ChangeType != EPropertyChangeType::Interactive)
+	USkeletalMesh* CurrentMesh = GetPersonaToolkit()->GetMesh();
+	if(CurrentMesh->MeshClothingAssets.IsValidIndex(InAssetIndex))
 	{
-		if(Event.Property->GetFName() == GET_MEMBER_NAME_CHECKED(FClothConfig, SelfCollisionRadius) ||
-			Event.Property->GetFName() == GET_MEMBER_NAME_CHECKED(FClothConfig, SelfCollisionCullScale))
+		UClothingAssetBase* Asset = CurrentMesh->MeshClothingAssets[InAssetIndex];
+		if (Asset)
 		{
-			USkeletalMesh* CurrentMesh = GetPersonaToolkit()->GetMesh();
-			if(CurrentMesh->MeshClothingAssets.IsValidIndex(InAssetIndex))
-			{
-				if(UClothingAssetBase* Asset = CurrentMesh->MeshClothingAssets[InAssetIndex])
-					Asset->BuildSelfCollisionData();
-			}
+			Asset->PostPropertyChangeCb(Event);
 		}
 	}
-
 	if(UDebugSkelMeshComponent* PreviewComponent = GetPersonaToolkit()->GetPreviewMeshComponent())
 	{
 		// Reregister our preview component to apply the change
diff --git a/Engine/Source/Editor/SkeletalMeshEditor/Private/SkeletalMeshEditor.cpp b/Engine/Source/Editor/SkeletalMeshEditor/Private/SkeletalMeshEditor.cpp
--- a/Engine/Source/Editor/SkeletalMeshEditor/Private/SkeletalMeshEditor.cpp
+++ b/Engine/Source/Editor/SkeletalMeshEditor/Private/SkeletalMeshEditor.cpp
@@ -23,7 +23,7 @@
 #include "Rendering/SkeletalMeshModel.h"
 
 #include "Animation/DebugSkelMeshComponent.h"
-#include "Assets/ClothingAsset.h"
+#include "ClothingAsset.h"
 #include "SCreateClothingSettingsPanel.h"
 #include "ClothingSystemEditorInterfaceModule.h"
 #include "Preferences/PersonaOptions.h"
@@ -676,12 +676,12 @@ void FSkeletalMeshEditor::FillApplyClothingAssetMenu(FMenuBuilder& MenuBuilder,
 	{
 		for(UClothingAssetBase* BaseAsset : Mesh->MeshClothingAssets)
 		{
-			UClothingAsset* ClothAsset = CastChecked<UClothingAsset>(BaseAsset);
+			UClothingAssetCommon* ClothAsset = CastChecked<UClothingAssetCommon>(BaseAsset);
 
 			FUIAction Action;
 			Action.CanExecuteAction = FCanExecuteAction::CreateSP(this, &FSkeletalMeshEditor::CanApplyClothing, InLodIndex, InSectionIndex);
 
-			const int32 NumClothLods = ClothAsset->LodData.Num();
+			const int32 NumClothLods = ClothAsset->GetNumLods();
 			for(int32 ClothLodIndex = 0; ClothLodIndex < NumClothLods; ++ClothLodIndex)
 			{
 				Action.ExecuteAction = FExecuteAction::CreateSP(this, &FSkeletalMeshEditor::OnApplyClothingAssetClicked, BaseAsset, InLodIndex, InSectionIndex, ClothLodIndex);
@@ -913,7 +913,7 @@ void FSkeletalMeshEditor::ApplyClothing(UClothingAssetBase* InAsset, int32 InLod
 			OriginalSectionData.ClothingData.AssetGuid = FGuid();
 			OriginalSectionData.ClothingData.AssetLodIndex = INDEX_NONE;
 		};
-		if (UClothingAsset* ClothingAsset = Cast<UClothingAsset>(InAsset))
+		if (UClothingAssetCommon* ClothingAsset = Cast<UClothingAssetCommon>(InAsset))
 		{
 			// Look for a currently bound asset an unbind it if necessary first
 			if (UClothingAssetBase* CurrentAsset = Mesh->GetSectionClothingAsset(InLodIndex, InSectionIndex))
diff --git a/Engine/Source/Editor/UnrealEd/Private/Animation/DebugSkelMeshComponent.cpp b/Engine/Source/Editor/UnrealEd/Private/Animation/DebugSkelMeshComponent.cpp
--- a/Engine/Source/Editor/UnrealEd/Private/Animation/DebugSkelMeshComponent.cpp
+++ b/Engine/Source/Editor/UnrealEd/Private/Animation/DebugSkelMeshComponent.cpp
@@ -17,7 +17,7 @@
 #include "Rendering/SkeletalMeshRenderData.h"
 #include "Rendering/SkeletalMeshModel.h"
 
-#include "Assets/ClothingAsset.h"
+#include "ClothingAsset.h"
 #include "ClothingSimulation.h"
 #include "Utils/ClothingMeshUtils.h"
 #include "DynamicMeshBuilder.h"
@@ -740,17 +740,15 @@ void UDebugSkelMeshComponent::RebuildClothingSectionsFixedVerts()
 
 				if(BaseAsset)
 				{
-					UClothingAsset* ConcreteAsset = Cast<UClothingAsset>(BaseAsset);
-					FClothLODData& LodData = ConcreteAsset->LodData[Section.ClothingData.AssetLodIndex];
+					UClothingAssetCommon* ConcreteAsset = Cast<UClothingAssetCommon>(BaseAsset);
+					UClothLODDataBase* LodData = ConcreteAsset->ClothLodData[Section.ClothingData.AssetLodIndex];
 
 					for(FMeshToMeshVertData& VertData : Section.ClothMappingData)
 					{
-						float TriangleDistanceMax = 0.0f;
-						TriangleDistanceMax += LodData.PhysicalMeshData.MaxDistances[VertData.SourceMeshVertIndices[0]];
-						TriangleDistanceMax += LodData.PhysicalMeshData.MaxDistances[VertData.SourceMeshVertIndices[1]];
-						TriangleDistanceMax += LodData.PhysicalMeshData.MaxDistances[VertData.SourceMeshVertIndices[2]];
-
-						if(TriangleDistanceMax == 0.0f)
+						if(LodData->PhysicalMeshData->IsFixed(
+							VertData.SourceMeshVertIndices[0],
+							VertData.SourceMeshVertIndices[1],
+							VertData.SourceMeshVertIndices[2]))
 						{
 							VertData.SourceMeshVertIndices[3] = 0xFFFF;
 						}
@@ -819,9 +817,9 @@ void UDebugSkelMeshComponent::RefreshSelectedClothingSkinnedPositions()
 
 		if(Asset)
 		{
-			UClothingAsset* ConcreteAsset = Cast<UClothingAsset>(*Asset);
+			UClothingAssetCommon* ConcreteAsset = Cast<UClothingAssetCommon>(*Asset);
 
-			if(ConcreteAsset->LodData.IsValidIndex(SelectedClothingLodForPainting))
+			if(ConcreteAsset->ClothLodData.IsValidIndex(SelectedClothingLodForPainting))
 			{
 				SkinnedSelectedClothingPositions.Reset();
 				SkinnedSelectedClothingNormals.Reset();
@@ -830,9 +828,9 @@ void UDebugSkelMeshComponent::RefreshSelectedClothingSkinnedPositions()
 				// Pass LOD0 to collect all bones
 				GetCurrentRefToLocalMatrices(RefToLocals, 0);
 
-				FClothLODData& LodData = ConcreteAsset->LodData[SelectedClothingLodForPainting];
+				UClothLODDataBase* LodData = ConcreteAsset->ClothLodData[SelectedClothingLodForPainting];
 
-				ClothingMeshUtils::SkinPhysicsMesh(ConcreteAsset->UsedBoneIndices, LodData.PhysicalMeshData, FTransform::Identity, RefToLocals.GetData(), RefToLocals.Num(), SkinnedSelectedClothingPositions, SkinnedSelectedClothingNormals);
+				ClothingMeshUtils::SkinPhysicsMesh(ConcreteAsset->UsedBoneIndices, *LodData->PhysicalMeshData, FTransform::Identity, RefToLocals.GetData(), RefToLocals.Num(), SkinnedSelectedClothingPositions, SkinnedSelectedClothingNormals);
 				RebuildCachedClothBounds();
 			}
 		}
@@ -1033,17 +1031,17 @@ FDebugSkelMeshDynamicData::FDebugSkelMeshDynamicData(UDebugSkelMeshComponent* In
 				{
 					ClothingSimDataIndexWhenPainting = ClothingAssetIndex;
 
-					if(UClothingAsset* ConcreteAsset = Cast<UClothingAsset>(BaseAsset))
+					if(UClothingAssetCommon* ConcreteAsset = Cast<UClothingAssetCommon>(BaseAsset))
 					{
-						if(ConcreteAsset->LodData.IsValidIndex(InComponent->SelectedClothingLodForPainting))
+						if(ConcreteAsset->ClothLodData.IsValidIndex(InComponent->SelectedClothingLodForPainting))
 						{
-							FClothLODData& LodData = ConcreteAsset->LodData[InComponent->SelectedClothingLodForPainting];
+							UClothLODDataBase* LodData = ConcreteAsset->ClothLodData[InComponent->SelectedClothingLodForPainting];
 
-							ClothingSimIndices = LodData.PhysicalMeshData.Indices;
+							ClothingSimIndices = LodData->PhysicalMeshData->Indices;
 
-							if(LodData.ParameterMasks.IsValidIndex(InComponent->SelectedClothingLodMaskForPainting))
+							if(LodData->ParameterMasks.IsValidIndex(InComponent->SelectedClothingLodMaskForPainting))
 							{
-								FClothParameterMask_PhysMesh& Mask = LodData.ParameterMasks[InComponent->SelectedClothingLodMaskForPainting];
+								FPointWeightMap& Mask = LodData->ParameterMasks[InComponent->SelectedClothingLodMaskForPainting];
 
 								ClothingVisiblePropertyValues = Mask.GetValueArray();
 							}
diff --git a/Engine/Source/Programs/UnrealBuildTool/Configuration/TargetRules.cs b/Engine/Source/Programs/UnrealBuildTool/Configuration/TargetRules.cs
--- a/Engine/Source/Programs/UnrealBuildTool/Configuration/TargetRules.cs
+++ b/Engine/Source/Programs/UnrealBuildTool/Configuration/TargetRules.cs
@@ -362,6 +362,12 @@ namespace UnrealBuildTool
 		[RequiresUniqueBuildEnvironment]
 		public bool bUseChaosChecked = false;
 
+		/// <summary>
+		/// Whether to compile in chaos memory tracking features
+		/// </summary>
+		[RequiresUniqueBuildEnvironment]
+		public bool bUseChaosMemoryTracking = false;
+
 		/// <summary>
 		/// Whether scene query acceleration is done by UE4. The physx scene query structure is still created, but we do not use it.
 		/// </summary>
@@ -1871,6 +1877,11 @@ namespace UnrealBuildTool
 			get { return Inner.bUseChaos; }
 		}
 
+		public bool bUseChaosMemoryTracking
+		{
+			get { return Inner.bUseChaosMemoryTracking; }
+		}
+
 		public bool bUseChaosChecked
 		{
 			get { return Inner.bUseChaosChecked; }
diff --git a/Engine/Source/Runtime/ClothingSystemRuntimeInterface/Public/ClothingAssetBase.h b/Engine/Source/Runtime/ClothingSystemRuntimeInterface/Public/ClothingAssetBase.h
--- a/Engine/Source/Runtime/ClothingSystemRuntimeInterface/Public/ClothingAssetBase.h
+++ b/Engine/Source/Runtime/ClothingSystemRuntimeInterface/Public/ClothingAssetBase.h
@@ -23,7 +23,7 @@ public:
 	* @param InSubmeshIdx Submesh in this asset to replace section with
 	* @param InAssetLodIndex Internal clothing LOD to use
 	*/
-	virtual bool BindToSkeletalMesh(USkeletalMesh* InSkelMesh, int32 InMeshLodIndex, int32 InSectionIndex, int32 InAssetLodIndex)
+	virtual bool BindToSkeletalMesh(USkeletalMesh* InSkelMesh, const int32 InMeshLodIndex, const int32 InSectionIndex, const int32 InAssetLodIndex, const bool bCallPostEditChange = true)
 	PURE_VIRTUAL(UClothingAssetBase::BindToSkeletalMesh, return false;);
 
 	/**
@@ -38,7 +38,7 @@ public:
 	* @param InSkelMesh skeletal mesh to remove from
 	* @param InMeshLodIndex Mesh LOD to remove this asset from (could still be bound to other LODs)
 	*/
-	virtual void UnbindFromSkeletalMesh(USkeletalMesh* InSkelMesh, int32 InMeshLodIndex)
+	virtual void UnbindFromSkeletalMesh(USkeletalMesh* InSkelMesh, const int32 InMeshLodIndex)
 	PURE_VIRTUAL(UClothingAssetBase::UnbindFromSkeletalMesh,);
 
 	/** 
@@ -65,10 +65,26 @@ public:
 	virtual int32 GetNumLods()
 	PURE_VIRTUAL(UClothingAssetBase::GetNumLods(), return 0;);
 
+	/** Add a new LOD class instance. */
+	virtual int32 AddNewLod()
+	PURE_VIRTUAL(UClothingAssetBase::AddNewLod(), return INDEX_NONE;);
+
+	/** Called from \c FPersonaMeshDetails::OnFinishedChangingClothingProperties(). */
+	virtual void PostPropertyChangeCb(const FPropertyChangedEvent& Event)
+	PURE_VIRTUAL(UClothingAssetBase::PostPropertyChangeCb, );
+
 	/** Builds self collision data */
 	virtual void BuildSelfCollisionData()
 	PURE_VIRTUAL(UClothingAssetBase::BuildSelfCollisionData(), );
 
+	/**
+	*	Builds the LOD transition data
+	*	When we transition between LODs we skin the incoming mesh to the outgoing mesh
+	*	in exactly the same way the render mesh is skinned to create a smooth swap
+	*/
+	virtual void BuildLodTransitionData()
+	PURE_VIRTUAL(UClothingAssetBase::BuildLodTransitionData(), );
+
 	/** Get the guid identifying this asset */
 	const FGuid& GetAssetGuid() const
 	{
diff --git a/Engine/Source/Runtime/Engine/Private/PhysicsEngine/ImmediatePhysics/ImmediatePhysicsChaos/ImmediatePhysicsSimulation_Chaos.cpp b/Engine/Source/Runtime/Engine/Private/PhysicsEngine/ImmediatePhysics/ImmediatePhysicsChaos/ImmediatePhysicsSimulation_Chaos.cpp
--- a/Engine/Source/Runtime/Engine/Private/PhysicsEngine/ImmediatePhysics/ImmediatePhysicsChaos/ImmediatePhysicsSimulation_Chaos.cpp
+++ b/Engine/Source/Runtime/Engine/Private/PhysicsEngine/ImmediatePhysics/ImmediatePhysicsChaos/ImmediatePhysicsSimulation_Chaos.cpp
@@ -1,23 +1,119 @@
 // Copyright 1998-2019 Epic Games, Inc. All Rights Reserved.
 
 #include "Physics/ImmediatePhysics/ImmediatePhysicsChaos/ImmediatePhysicsSimulation_Chaos.h"
-
-#if INCLUDE_CHAOS
-
 #include "Physics/ImmediatePhysics/ImmediatePhysicsChaos/ImmediatePhysicsActorHandle_Chaos.h"
 #include "Physics/ImmediatePhysics/ImmediatePhysicsChaos/ImmediatePhysicsJointHandle_Chaos.h"
 #include "Physics/ImmediatePhysics/ImmediatePhysicsStats.h"
 
+#include "Chaos/ChaosDebugDraw.h"
+#include "Chaos/DebugDrawQueue.h"
 #include "Chaos/ParticleHandle.h"
-#include "Chaos/PBD6DJointConstraints.h"
+#include "Chaos/PBDCollisionConstraint.h"
+#include "Chaos/PBDJointConstraints.h"
 #include "Chaos/PBDRigidParticles.h"
 #include "Chaos/PBDRigidsEvolutionGBF.h"
+#include "ChaosLog.h"
 
 #include "PhysicsEngine/BodyInstance.h"
 #include "PhysicsEngine/ConstraintInstance.h"
 
+// @todo(ccaulfield): remove when finished
+float ChaosImmediate_Evolution_DeltaTime = 0.03f;
+int32 ChaosImmediate_Evolution_Iterations = 10;
+int32 ChaosImmediate_Collision_Enabled = 1;
+int32 ChaosImmediate_Collision_ApplyEnabled = 0;
+int32 ChaosImmediate_Collision_PushOutIterations = 5;
+int32 ChaosImmediate_Collision_PushOutPairIterations = 2;
+float ChaosImmediate_Collision_Thickness = 0;
+
+FAutoConsoleVariableRef CVarChaosImmPhysDeltaTime(TEXT("p.Chaos.ImmPhys.DeltaTime"), ChaosImmediate_Evolution_DeltaTime, TEXT("Override chaos immediate physics delta time if non-zero"));
+FAutoConsoleVariableRef CVarChaosImmPhysIterations(TEXT("p.Chaos.ImmPhys.Iterations"), ChaosImmediate_Evolution_Iterations, TEXT("Number of constraint solver loops in immediate physics"));
+FAutoConsoleVariableRef CVarImmediatePhysicsDisableCollisions(TEXT("p.Chaos.ImmPhys.EnableCollisions"), ChaosImmediate_Collision_Enabled, TEXT("Enable/Disable collisions in Immediate Physics."));
+FAutoConsoleVariableRef CVarCollisionDisableApply(TEXT("p.Chaos.ImmPhys.CollisionEnableApply"), ChaosImmediate_Collision_ApplyEnabled, TEXT("Enable/Disable the Apply() (velocity, friction, restitution) method for collisions"));
+FAutoConsoleVariableRef CVarCollisionPushOutIterations(TEXT("p.Chaos.ImmPhys.CollisionPushOutIterations"), ChaosImmediate_Collision_PushOutIterations, TEXT("Set the ApplyPushOut() (position correction) iteration count (0 to disable)"));
+FAutoConsoleVariableRef CVarCollisionPushOutPairIterations(TEXT("p.Chaos.ImmPhys.CollisionPushOutPairIterations"), ChaosImmediate_Collision_PushOutPairIterations, TEXT("Set the ApplyPushOut() internal pair interations (position correction) iteration count (0 to disable)"));
+FAutoConsoleVariableRef CVarChaosImmPhysThickness(TEXT("p.Chaos.ImmPhys.CollisionThickness"), ChaosImmediate_Collision_Thickness, TEXT("ChaosImmediateThickness"));
+
+float ChaosImmediate_Joint_SwingTwistAngleTolerance = 1.0e-6f;
+int32 ChaosImmediate_Joint_EnableTwistLimits = 1;
+int32 ChaosImmediate_Joint_EnableSwingLimits = 1;
+int32 ChaosImmediate_Joint_EnableDrives = 1;
+int32 ChaosImmediate_Joint_FreezeIterations = 0;
+int32 ChaosImmediate_Joint_FrozenIterations = 1;
+float ChaosImmediate_Joint_PBDDriveStiffness = 0.0f;
+float ChaosImmediate_Joint_PBDMinParentMassRatio = 0.5f;
+float ChaosImmediate_Joint_PBDMaxInertiaRatio = 5.0f;
+FAutoConsoleVariableRef CVarSwingTwistAngleTolerance(TEXT("p.Chaos.ImmPhys.Joint.SwingTwistAngleTolerance"), ChaosImmediate_Joint_SwingTwistAngleTolerance, TEXT("SwingTwistAngleTolerance."));
+FAutoConsoleVariableRef CVarEnableTwistLimits(TEXT("p.Chaos.ImmPhys.Joint.EnableTwistLimits"), ChaosImmediate_Joint_EnableTwistLimits, TEXT("EnableTwistLimits."));
+FAutoConsoleVariableRef CVarEnableSwingLimits(TEXT("p.Chaos.ImmPhys.Joint.EnableSwingLimits"), ChaosImmediate_Joint_EnableSwingLimits, TEXT("EnableSwingLimits."));
+FAutoConsoleVariableRef CVarEnableDrives(TEXT("p.Chaos.ImmPhys.Joint.EnableDrives"), ChaosImmediate_Joint_EnableDrives, TEXT("EnableDrives."));
+FAutoConsoleVariableRef CVarFreezeIterations(TEXT("p.Chaos.ImmPhys.Joint.FreezeIterations"), ChaosImmediate_Joint_FreezeIterations, TEXT("FreezeIterations."));
+FAutoConsoleVariableRef CVarFrozenIterations(TEXT("p.Chaos.ImmPhys.Joint.FrozenIterations"), ChaosImmediate_Joint_FrozenIterations, TEXT("FrozenIterations."));
+FAutoConsoleVariableRef CVarDriveStiffness(TEXT("p.Chaos.ImmPhys.Joint.PBDDriveStiffness"), ChaosImmediate_Joint_PBDDriveStiffness, TEXT("6Dof joint drive stiffness override (if > 0)."));
+FAutoConsoleVariableRef CVarMinParentMassRatio(TEXT("p.Chaos.ImmPhys.Joint.MinParentMassRatio"), ChaosImmediate_Joint_PBDMinParentMassRatio, TEXT("6Dof joint PBDMinParentMassRatio (if > 0)"));
+FAutoConsoleVariableRef CVarMaxInertiaRatio(TEXT("p.Chaos.ImmPhys.Joint.MaxInertiaRatio"), ChaosImmediate_Joint_PBDMaxInertiaRatio, TEXT("6Dof joint PBDMaxInertiaRatio (if > 0)"));
+
+// DebugDraw CVars
+#if UE_BUILD_DEBUG
+int32 ChaosImmediate_DebugDrawParticles = 1;
+int32 ChaosImmediate_DebugDrawShapes = 1;
+int32 ChaosImmediate_DebugDrawCollisions = 3;
+int32 ChaosImmediate_DebugDrawJoints = 1;
+#else
+int32 ChaosImmediate_DebugDrawParticles = 0;
+int32 ChaosImmediate_DebugDrawShapes = 0;
+int32 ChaosImmediate_DebugDrawCollisions = 0;
+int32 ChaosImmediate_DebugDrawJoints = 0;
+#endif
+
+FAutoConsoleVariableRef CVarDebugDrawParticles(TEXT("p.Chaos.ImmPhys.DebugDrawParticles"), ChaosImmediate_DebugDrawParticles, TEXT("Draw Particle Transforms (0 = never; 1 = end of frame; 2 = begin and end of frame; 3 = each evolution phase)."));
+FAutoConsoleVariableRef CVarDebugDrawShapes(TEXT("p.Chaos.ImmPhys.DebugDrawShapes"), ChaosImmediate_DebugDrawShapes, TEXT("Draw Shapes (0 = never; 1 = end of frame; 2 = begin and end of frame; 3 = each evolution phase."));
+FAutoConsoleVariableRef CVarDebugDrawCollisions(TEXT("p.Chaos.ImmPhys.DebugDrawCollisions"), ChaosImmediate_DebugDrawCollisions, TEXT("Draw Collisions (0 = never; 1 = end of frame; 2 = begin and end of frame; 3 = each evolution phase)"));
+FAutoConsoleVariableRef CVarDebugDrawJoints(TEXT("p.Chaos.ImmPhys.DebugDrawJoints"), ChaosImmediate_DebugDrawJoints, TEXT("Draw Joints. (0 = never; 1 = end of frame; 2 = begin and end of frame; 3 = each evolution phase; 4 = each Apply step)."));
+
 namespace ImmediatePhysics_Chaos
 {
+	template<typename T, int d>
+	bool ShouldIgnoreCollisionConstraint(
+		const Chaos::TGeometryParticleHandle<T, d>* ParticleA,
+		const Chaos::TGeometryParticleHandle<T, d>* ParticleB,
+		const TMap<const Chaos::TGeometryParticleHandle<T, d>*, TSet<const Chaos::TGeometryParticleHandle<T, d>*>>& IgnoreSetMap)
+	{
+		using namespace Chaos;
+
+		if (!ChaosImmediate_Collision_Enabled)
+		{
+			return true;
+		}
+		if (const TSet<const TGeometryParticleHandle<T, d>*>* IgnoreSet = IgnoreSetMap.Find(ParticleA))
+		{
+			return IgnoreSet->Contains(ParticleB);
+		}
+		if (const TSet<const TGeometryParticleHandle<T, d>*>* IgnoreSet = IgnoreSetMap.Find(ParticleB))
+		{
+			return IgnoreSet->Contains(ParticleA);
+		}
+		return false;
+	}
+
+	template<typename T, int d>
+	void FilterCollisionConstraints(
+		TArray<Chaos::TRigidBodyContactConstraint<T, d>>& Constraints, 
+		const TMap<const Chaos::TGeometryParticleHandle<T, d>*, TSet<const Chaos::TGeometryParticleHandle<T, d>*>>& IgnoreSetMap)
+	{
+		using namespace Chaos;
+
+		for (int ConstraintIndex = 0; ConstraintIndex < Constraints.Num(); ++ConstraintIndex)
+		{
+			if (ShouldIgnoreCollisionConstraint(Constraints[ConstraintIndex].Particle, Constraints[ConstraintIndex].Levelset, IgnoreSetMap))
+			{
+				Constraints.RemoveAtSwap(ConstraintIndex);
+				--ConstraintIndex;
+				continue;
+			}
+		}
+	}
+
 
 	//
 	//
@@ -25,26 +121,85 @@ namespace ImmediatePhysics_Chaos
 
 	FSimulation::FSimulation()
 		: NumActiveActorHandles(0)
-		, Gravity(FVector::ZeroVector)
+		, SimulationSpaceTransform(FTransform::Identity)
 	{
 		using namespace Chaos;
 
-		const int DefaultSolverIterations = 3;
 		Particles = MakeUnique<TPBDRigidsSOAs<FReal, Dimensions>>();
-		Joints = MakeUnique<TPBD6DJointConstraints<FReal, Dimensions>>();
-		JointsRule = MakeUnique<TPBDConstraintIslandRule<TPBD6DJointConstraints<FReal, Dimensions>, FReal, Dimensions>>(*Joints);
-		Evolution = MakeUnique<TPBDRigidsEvolutionGBF<FReal, Dimensions>>(*Particles.Get(), DefaultSolverIterations);
+		Joints = MakeUnique<TPBDJointConstraints<FReal, Dimensions>>(TPBDJointSolverSettings<FReal, Dimensions>());
+		JointsRule = MakeUnique<TPBDConstraintIslandRule<TPBDJointConstraints<FReal, Dimensions>, FReal, Dimensions>>(*Joints);
+		Evolution = MakeUnique<TPBDRigidsEvolutionGBF<FReal, Dimensions>>(*Particles.Get(), 1);
+		TPBDCollisionConstraint<FReal, Dimensions>& Collisions = Evolution->GetCollisionConstraints();
 
 		Evolution->AddConstraintRule(JointsRule.Get());
 
-		Evolution->AddForceFunction([this](TTransientPBDRigidParticleHandle<FReal, Dimensions>& Particle, const FReal Dt)
+		Collisions.SetPostComputeCallback(
+			[this](TArray<TRigidBodyContactConstraint<FReal, Dimensions>>& Constraints)
+			{
+				FilterCollisionConstraints(Constraints, IgnoreCollisionParticlePairTable);
+			});
+		
+#if CHAOS_DEBUG_DRAW
+		Evolution->SetPostIntegrateCallback(
+			[this]()
+			{
+				DebugDrawElementsAll(3, 3, 0.3f);
+			});
+		Evolution->SetPostApplyCallback(
+			[this](int32 Island)
 			{
-				Particle.F() += Gravity * Particle.M();
+				DebugDrawElementsIsland(Island, 3, 3, 0.3f);
 			});
+		Evolution->SetPostApplyPushOutCallback(
+			[this](int32 Island)
+			{
+				DebugDrawElementsIsland(Island, 3, 3, 0.3f);
+			});
+		Collisions.SetPostApplyCallback(
+			[this](const float Dt, const TArray<TPBDCollisionConstraintHandle<float, 3>*>& InConstraintHandles)
+			{
+				if (ChaosImmediate_DebugDrawCollisions == 4)
+				{
+					DebugDraw::DrawCollisions(SimulationSpaceTransform, InConstraintHandles, 0.3f);
+				}
+			});
+		Collisions.SetPostApplyPushOutCallback(
+			[this](const float Dt, const TArray<TPBDCollisionConstraintHandle<float, 3>*>& InConstraintHandles, bool bRequiresAnotherIteration)
+			{
+				if (ChaosImmediate_DebugDrawCollisions == 4)
+				{
+					DebugDraw::DrawCollisions(SimulationSpaceTransform, InConstraintHandles, 0.3f);
+				}
+			});
+		Joints->SetPreApplyCallback(
+			[this](const float Dt, const TArray<TPBDJointConstraintHandle<float, 3>*>& InConstraintHandles)
+			{
+				if (ChaosImmediate_DebugDrawJoints == 4)
+				{
+					DebugDraw::DrawJointConstraints(SimulationSpaceTransform, InConstraintHandles, 0.3f);
+				}
+			});
+		Joints->SetPostApplyCallback(
+			[this](const float Dt, const TArray<TPBDJointConstraintHandle<float, 3>*>& InConstraintHandles)
+			{
+				if (ChaosImmediate_DebugDrawJoints == 4)
+				{
+					DebugDraw::DrawJointConstraints(SimulationSpaceTransform, InConstraintHandles, 1.0f);
+				}
+			});
+#endif
 	}
 
 	FSimulation::~FSimulation()
 	{
+		using namespace Chaos;
+
+		Evolution->GetCollisionConstraints().ClearPostComputeCallback();
+		Evolution->GetCollisionConstraints().ClearPostApplyCallback();
+		Evolution->GetCollisionConstraints().ClearPostApplyPushOutCallback();
+
+		//Evolution->RemoveConstraintRule(JointsRule.Get());
+
 		for (FActorHandle* ActorHandle : ActorHandles)
 		{
 			delete ActorHandle;
@@ -129,27 +284,95 @@ namespace ImmediatePhysics_Chaos
 
 	void FSimulation::SetIgnoreCollisionPairTable(const TArray<FIgnorePair>& InIgnoreCollisionPairTable)
 	{
-#if IMMEDIATEPHYSICS_CHAOS_TODO
-		IgnoreCollisionPairTable.Reset();
-		for (const FIgnorePair& IgnorePair : InIgnoreTable)
+		using namespace Chaos;
+
+		IgnoreCollisionParticlePairTable.Reset();
+		for (const FIgnorePair& IgnorePair : InIgnoreCollisionPairTable)
 		{
-			TSet<FActorHandle*>& IgnoreActorsA = IgnoreCollisionPairTable.FindOrAdd(IgnorePair.A);
-			IgnoreActorsA.Add(IgnorePair.B);
+			if(!IgnorePair.A || !IgnorePair.B)
+			{
+				continue;
+			}
+
+			const TGeometryParticleHandle<FReal, Dimensions>* ParticleA = IgnorePair.A->GetParticle();
+			const TGeometryParticleHandle<FReal, Dimensions>* ParticleB = IgnorePair.B->GetParticle();
 
-			TSet<FActorHandle*>& IgnoreActorsB = IgnoreCollisionPairTable.FindOrAdd(IgnorePair.B);
-			IgnoreActorsB.Add(IgnorePair.A);
+			IgnoreCollisionParticlePairTable.FindOrAdd(ParticleA).Add(ParticleB); 
+			IgnoreCollisionParticlePairTable.FindOrAdd(ParticleB).Add(ParticleA);
 		}
-		bRecreateIterationCache = true;
-#endif
 	}
 
 	void FSimulation::SetIgnoreCollisionActors(const TArray<FActorHandle*>& InIgnoreCollisionActors)
 	{
-#if IMMEDIATEPHYSICS_CHAOS_TODO
-		IgnoreCollisionActors.Reset();
-		IgnoreCollisionActors.Append(InIgnoreCollisionActors);
-		bRecreateIterationCache = true;
-#endif
+		using namespace Chaos;
+
+		for (FActorHandle* ActorHandle : InIgnoreCollisionActors)
+		{
+			TPBDRigidParticleHandle<FReal, Dimensions>* Particle = ActorHandle->GetParticle()->AsDynamic();
+			if (Particle != nullptr)
+			{
+				Particle->SetCollisionGroup(INDEX_NONE);
+			}
+		}
+	}
+
+	void FSimulation::ConditionConstraints()
+	{
+		// Assign levels to actors based on connection-distance to a non-dynamic actor
+		// This is used to sort constraints and set parent/child relationship on the constrained particles in a constraint
+		// @todo(ccaulfield): this should use the constraint graph and should only update when constraint connectivity changes or particles change type
+		using namespace Chaos;
+
+		TMap<FActorHandle*, TArray<FJointHandle*>> ActorJoints;
+		TArray<FActorHandle*> ActorQueue;
+		ActorQueue.Reserve(ActorHandles.Num());
+
+		// Reset all actor levels
+		for (FActorHandle* ActorHandle : ActorHandles)
+		{
+			ActorHandle->SetLevel(INDEX_NONE);
+			ActorJoints.Emplace(ActorHandle);
+
+			if (!ActorHandle->GetParticle()->AsDynamic())
+			{
+				ActorHandle->SetLevel(0);
+				ActorQueue.Add(ActorHandle);
+			}
+		}
+
+		// Build a list of joints per actor
+		for (FJointHandle* JointHandle : JointHandles)
+		{
+			TVector<FActorHandle*, 2> JointActorHandles = JointHandle->GetActorHandles();
+			ActorJoints[JointActorHandles[0]].Add(JointHandle);
+			ActorJoints[JointActorHandles[1]].Add(JointHandle);
+		}
+
+		// Breadth-first assign level
+		for (int32 ActorQueueIndex = 0; ActorQueueIndex < ActorQueue.Num(); ++ActorQueueIndex)
+		{
+			FActorHandle* ActorHandle = ActorQueue[ActorQueueIndex];
+			for (FJointHandle* JointHandle : ActorJoints[ActorHandle])
+			{
+				TVector<FActorHandle*, 2> JointActorHandles = JointHandle->GetActorHandles();
+				if (JointActorHandles[0]->GetLevel() == INDEX_NONE)
+				{
+					JointActorHandles[0]->SetLevel(ActorHandle->GetLevel() + 1);
+					ActorQueue.Add(JointActorHandles[0]);
+				}
+				if (JointActorHandles[1]->GetLevel() == INDEX_NONE)
+				{
+					JointActorHandles[1]->SetLevel(ActorHandle->GetLevel() + 1);
+					ActorQueue.Add(JointActorHandles[1]);
+				}
+			}
+		}
+
+		// Update constraint levels
+		for (FJointHandle* JointHandle : JointHandles)
+		{
+			JointHandle->UpdateLevels();
+		}
 	}
 
 	DECLARE_CYCLE_STAT(TEXT("FSimulation::Simulate_Chaos"), STAT_ImmediateSimulate_Chaos, STATGROUP_ImmediatePhysics);
@@ -157,11 +380,105 @@ namespace ImmediatePhysics_Chaos
 	void FSimulation::Simulate(float DeltaTime, const FVector& InGravity)
 	{
 		SCOPE_CYCLE_COUNTER(STAT_ImmediateSimulate_Chaos);
+		using namespace Chaos;
+
+		// TEMP: overrides
+		{
+			if (ChaosImmediate_Evolution_DeltaTime > 0)
+			{
+				DeltaTime = ChaosImmediate_Evolution_DeltaTime;
+			}
+			UE_LOG(LogChaosJoint, Verbose, TEXT("Simulate dt = %f"), DeltaTime);
+
+			// Overrides
+			TPBDJointSolverSettings<float, 3> JointsSettings = Joints->GetSettings();
+			JointsSettings.SwingTwistAngleTolerance = ChaosImmediate_Joint_SwingTwistAngleTolerance;
+			JointsSettings.PBDMinParentMassRatio = ChaosImmediate_Joint_PBDMinParentMassRatio;
+			JointsSettings.PBDMaxInertiaRatio = ChaosImmediate_Joint_PBDMaxInertiaRatio;
+			JointsSettings.FreezeIterations = ChaosImmediate_Joint_FreezeIterations;
+			JointsSettings.FrozenIterations = ChaosImmediate_Joint_FrozenIterations;
+			JointsSettings.bEnableTwistLimits = ChaosImmediate_Joint_EnableTwistLimits > 0;
+			JointsSettings.bEnableSwingLimits = ChaosImmediate_Joint_EnableSwingLimits > 0;
+			JointsSettings.bEnableDrives = ChaosImmediate_Joint_EnableDrives > 0;
+			JointsSettings.PBDDriveStiffness = ChaosImmediate_Joint_PBDDriveStiffness;
+			Joints->SetSettings(JointsSettings);
+
+			Evolution->SetNumIterations(ChaosImmediate_Evolution_Iterations);
 
-		Gravity = InGravity;
+			Evolution->GetCollisionConstraints().SetThickness(ChaosImmediate_Collision_Thickness);
+			Evolution->GetCollisionConstraints().SetVelocitySolveEnabled(ChaosImmediate_Collision_ApplyEnabled != 0);
+			Evolution->GetCollisionConstraints().SetPushOutPairIterations(ChaosImmediate_Collision_PushOutPairIterations);
+			Evolution->GetCollisionConstraintsRule().SetPushOutIterations(ChaosImmediate_Collision_PushOutIterations);
+		}
+
+		DebugDrawElementsAll(2, 2, 0.7f);
 
+		ConditionConstraints();
+
+		Evolution->GetGravityForces().SetAcceleration(InGravity);
 		Evolution->AdvanceOneTimeStep(DeltaTime);
+		Evolution->EndFrame(DeltaTime);
+
+		DebugDrawElementsAll(1, 3, 1.0f);
+	}
+
+	void FSimulation::DebugDrawElementsAll(const int32 MinDebugLevel, const int32 MaxDebugLevel, const float ColorScale)
+	{
+#if CHAOS_DEBUG_DRAW
+		using namespace Chaos;
+		if (FDebugDrawQueue::IsDebugDrawingEnabled())
+		{
+			if ((ChaosImmediate_DebugDrawParticles >= MinDebugLevel) && (ChaosImmediate_DebugDrawParticles <= MaxDebugLevel))
+			{
+				DebugDraw::DrawParticleTransforms(SimulationSpaceTransform, Evolution->GetParticles().GetAllParticlesView(), ColorScale);
+			}
+			if ((ChaosImmediate_DebugDrawShapes >= MinDebugLevel) && (ChaosImmediate_DebugDrawShapes <= MaxDebugLevel))
+			{
+				DebugDraw::DrawParticleShapes(SimulationSpaceTransform, Evolution->GetParticles().GetAllParticlesView(), ColorScale);
+			}
+			if ((ChaosImmediate_DebugDrawCollisions >= MinDebugLevel) && (ChaosImmediate_DebugDrawCollisions <= MaxDebugLevel))
+			{
+				DebugDraw::DrawCollisions(SimulationSpaceTransform, Evolution->GetCollisionConstraints(), ColorScale);
+			}
+			if ((ChaosImmediate_DebugDrawJoints >= MinDebugLevel) && (ChaosImmediate_DebugDrawJoints <= MaxDebugLevel))
+			{
+				DebugDraw::DrawJointConstraints(SimulationSpaceTransform, *Joints, ColorScale);
+			}
+		}
+#endif
 	}
-}
 
-#endif // INCLUDE_CHAOS
\ No newline at end of file
+	void FSimulation::DebugDrawElementsIsland(const int32 Island, const int32 MinDebugLevel, const int32 MaxDebugLevel, const float ColorScale)
+	{
+#if CHAOS_DEBUG_DRAW
+		using namespace Chaos;
+		if (FDebugDrawQueue::IsDebugDrawingEnabled())
+		{
+			if ((ChaosImmediate_DebugDrawParticles >= MinDebugLevel) && (ChaosImmediate_DebugDrawParticles <= MaxDebugLevel))
+			{
+				DebugDraw::DrawParticleTransforms(SimulationSpaceTransform, Evolution->GetIslandParticles(Island), ColorScale);
+			}
+			if ((ChaosImmediate_DebugDrawShapes >= MinDebugLevel) && (ChaosImmediate_DebugDrawShapes <= MaxDebugLevel))
+			{
+				DebugDraw::DrawParticleShapes(SimulationSpaceTransform, Evolution->GetIslandParticles(Island), ColorScale);
+			}
+			if ((ChaosImmediate_DebugDrawCollisions >= MinDebugLevel) && (ChaosImmediate_DebugDrawCollisions <= MaxDebugLevel))
+			{
+				Evolution->GetCollisionConstraintsRule().VisitIslandConstraints(Island,
+					[this, ColorScale](const TArray<TPBDCollisionConstraintHandle<float, 3>*>& ConstraintHandles)
+					{
+						DebugDraw::DrawCollisions(SimulationSpaceTransform, ConstraintHandles, ColorScale);
+					});
+			}
+			if ((ChaosImmediate_DebugDrawJoints >= MinDebugLevel) && (ChaosImmediate_DebugDrawJoints <= MaxDebugLevel))
+			{
+				JointsRule->VisitIslandConstraints(Island,
+					[this, ColorScale](const TArray<TPBDJointConstraintHandle<float, 3>*>& ConstraintHandles)
+					{
+						DebugDraw::DrawJointConstraints(SimulationSpaceTransform, *Joints, ColorScale);
+					});
+			}
+		}
+#endif
+	}
+}
diff --git a/Engine/Source/Runtime/Engine/Private/SkeletalRenderGPUSkin.cpp b/Engine/Source/Runtime/Engine/Private/SkeletalRenderGPUSkin.cpp
--- a/Engine/Source/Runtime/Engine/Private/SkeletalRenderGPUSkin.cpp
+++ b/Engine/Source/Runtime/Engine/Private/SkeletalRenderGPUSkin.cpp
@@ -575,7 +575,7 @@ void FSkeletalMeshObjectGPUSkin::ProcessUpdatedDynamicData(FGPUSkinCache* GPUSki
 			// Create a uniform buffer from the bone transforms.
 			TArray<FMatrix>& ReferenceToLocalMatrices = DynamicData->ReferenceToLocal;
 			bool bNeedFence = ShaderData.UpdateBoneData(RHICmdList, ReferenceToLocalMatrices, Section.BoneMap, RevisionNumber, false, FeatureLevel, bUseSkinCache);
-#if WITH_APEX_CLOTHING
+#if WITH_APEX_CLOTHING || WITH_CHAOS_CLOTHING
 			// Update uniform buffer for APEX cloth simulation mesh positions and normals
 			if( bClothFactory )
 			{				
@@ -588,7 +588,7 @@ void FSkeletalMeshObjectGPUSkin::ProcessUpdatedDynamicData(FGPUSkinCache* GPUSki
 					bNeedFence = ClothShaderData.UpdateClothSimulData(RHICmdList, DynamicData->ClothingSimData[ActorIdx].Positions, DynamicData->ClothingSimData[ActorIdx].Normals, FrameNumberToPrepare, FeatureLevel) || bNeedFence;
 				}
 			}
-#endif // WITH_APEX_CLOTHING
+#endif // WITH_APEX_CLOTHING || WITH_CHAOS_CLOTHING
 			// Try to use the GPU skinning cache if possible
 			if (bUseSkinCache)
 			{
diff --git a/Engine/Source/Runtime/Experimental/Chaos/Public/Chaos/KinematicGeometryParticles.h b/Engine/Source/Runtime/Experimental/Chaos/Public/Chaos/KinematicGeometryParticles.h
--- a/Engine/Source/Runtime/Experimental/Chaos/Public/Chaos/KinematicGeometryParticles.h
+++ b/Engine/Source/Runtime/Experimental/Chaos/Public/Chaos/KinematicGeometryParticles.h
@@ -10,7 +10,7 @@ template<class T, int d, EGeometryParticlesSimType SimType>
 class TKinematicGeometryParticlesImp : public TGeometryParticlesImp<T, d, SimType>
 {
   public:
-	TKinematicGeometryParticlesImp()
+	CHAOS_API TKinematicGeometryParticlesImp()
 	    : TGeometryParticlesImp<T, d, SimType>()
 	{
 		this->MParticleType = EParticleType::Kinematic;
@@ -18,14 +18,14 @@ class TKinematicGeometryParticlesImp : public TGeometryParticlesImp<T, d, SimTyp
 		TArrayCollection::AddArray(&MW);
 	}
 	TKinematicGeometryParticlesImp(const TKinematicGeometryParticlesImp<T, d, SimType>& Other) = delete;
-	TKinematicGeometryParticlesImp(TKinematicGeometryParticlesImp<T, d, SimType>&& Other)
+	CHAOS_API TKinematicGeometryParticlesImp(TKinematicGeometryParticlesImp<T, d, SimType>&& Other)
 	    : TGeometryParticlesImp<T, d, SimType>(MoveTemp(Other)), MV(MoveTemp(Other.MV)), MW(MoveTemp(Other.MW))
 	{
 		this->MParticleType = EParticleType::Kinematic;
 		TArrayCollection::AddArray(&MV);
 		TArrayCollection::AddArray(&MW);
 	}
-	~TKinematicGeometryParticlesImp() {}
+	CHAOS_API virtual ~TKinematicGeometryParticlesImp();
 
 	const TVector<T, d>& V(const int32 Index) const { return MV[Index]; }
 	TVector<T, d>& V(const int32 Index) { return MV[Index]; }
@@ -45,7 +45,7 @@ class TKinematicGeometryParticlesImp : public TGeometryParticlesImp<T, d, SimTyp
 	//cannot be reference because double pointer would allow for badness, but still useful to have non const access to handle
 	THandleType* Handle(int32 Index);
 
-	void Serialize(FChaosArchive& Ar)
+	CHAOS_API virtual void Serialize(FChaosArchive& Ar) override
 	{
 		TGeometryParticlesImp<T, d, SimType>::Serialize(Ar);
 		Ar << MV << MW;
diff --git a/Engine/Source/Runtime/Experimental/Chaos/Public/Chaos/RigidParticles.h b/Engine/Source/Runtime/Experimental/Chaos/Public/Chaos/RigidParticles.h
--- a/Engine/Source/Runtime/Experimental/Chaos/Public/Chaos/RigidParticles.h
+++ b/Engine/Source/Runtime/Experimental/Chaos/Public/Chaos/RigidParticles.h
@@ -1,6 +1,7 @@
 // Copyright 1998-2019 Epic Games, Inc. All Rights Reserved.
 #pragma once
 
+#include "Containers/Queue.h"
 #include "Chaos/ArrayCollectionArray.h"
 #include "Chaos/BVHParticles.h"
 #include "Chaos/GeometryParticles.h"
@@ -9,7 +10,6 @@
 #include "Chaos/Matrix.h"
 #include "Chaos/Particles.h"
 #include "Chaos/Rotation.h"
-#include "Chaos/TriangleMesh.h"
 
 namespace Chaos
 {
@@ -47,7 +47,7 @@ class TRigidParticles : public TKinematicGeometryParticles<T, d>
     using TParticles<T, d>::X;
     using TGeometryParticles<T, d>::R;
 
-	TRigidParticles()
+	CHAOS_API TRigidParticles()
 	    : TKinematicGeometryParticles<T, d>()
 	{
 		TArrayCollection::AddArray(&MF);
@@ -66,7 +66,7 @@ class TRigidParticles : public TKinematicGeometryParticles<T, d>
 		TArrayCollection::AddArray(&MToBeRemovedOnFracture);
 	}
 	TRigidParticles(const TRigidParticles<T, d>& Other) = delete;
-	TRigidParticles(TRigidParticles<T, d>&& Other)
+	CHAOS_API TRigidParticles(TRigidParticles<T, d>&& Other)
 	    : TKinematicGeometryParticles<T, d>(MoveTemp(Other)), MF(MoveTemp(Other.MF)), MT(MoveTemp(Other.MT)), MExternalForce(MoveTemp(Other.MExternalForce)), MExternalTorque(MoveTemp(Other.MExternalTorque)), MI(MoveTemp(Other.MI)), MInvI(MoveTemp(Other.MInvI)), MM(MoveTemp(Other.MM)), MInvM(MoveTemp(Other.MInvM)), MCollisionParticles(MoveTemp(Other.MCollisionParticles)), MCollisionGroup(MoveTemp(Other.MCollisionGroup)), MObjectState(MoveTemp(Other.MObjectState))
 	{
 		TArrayCollection::AddArray(&MF);
@@ -85,6 +85,9 @@ class TRigidParticles : public TKinematicGeometryParticles<T, d>
 		TArrayCollection::AddArray(&MToBeRemovedOnFracture);
 	}
 
+	CHAOS_API virtual ~TRigidParticles()
+	{}
+
 	const TVector<T, d>& Torque(const int32 Index) const { return MT[Index]; }
 	TVector<T, d>& Torque(const int32 Index) { return MT[Index]; }
 
@@ -130,6 +133,9 @@ class TRigidParticles : public TKinematicGeometryParticles<T, d>
 	const bool ToBeRemovedOnFracture(const int32 Index) const { return MToBeRemovedOnFracture[Index]; }
 	bool& ToBeRemovedOnFracture(const int32 Index) { return MToBeRemovedOnFracture[Index]; }
 
+	TQueue<TGeometryParticleHandle<T, d>*, EQueueMode::Mpsc>& GetSleepData() { return MSleepData; }
+	void AddSleepData(TGeometryParticleHandle<T, d>* Particle) { MSleepData.Enqueue(Particle); }
+
 	const EObjectStateType ObjectState(const int32 Index) const { return MObjectState[Index]; }
 	EObjectStateType& ObjectState(const int32 Index) { return MObjectState[Index]; }
 
@@ -148,7 +154,7 @@ class TRigidParticles : public TKinematicGeometryParticles<T, d>
 		return FString::Printf(TEXT("%s, MF:%s, MT:%s, MExternalForce:%s, MExternalTorque:%s, MI:%s, MInvI:%s, MM:%f, MInvM:%f, MCollisionParticles(num):%d, MCollisionGroup:%d, MDisabled:%d, MSleepring:%d, MIsland:%d"), *BaseString, *F(index).ToString(), *Torque(index).ToString(), *ExternalForce(index).ToString(), *ExternalTorque(index).ToString(), *I(index).ToString(), *InvI(index).ToString(), M(index), InvM(index), CollisionParticlesSize(index), CollisionGroup(index), Disabled(index), Sleeping(index), Island(index));
 	}
 
-	void Serialize(FChaosArchive& Ar)
+	CHAOS_API virtual void Serialize(FChaosArchive& Ar) override
 	{
 		TKinematicGeometryParticles<T,d>::Serialize(Ar);
 		Ar << MF << MT << MExternalForce << MExternalTorque << MI << MInvI << MM << MInvM;
@@ -170,6 +176,7 @@ class TRigidParticles : public TKinematicGeometryParticles<T, d>
 	TArrayCollectionArray<bool> MDisabled;
 	TArrayCollectionArray<bool> MToBeRemovedOnFracture;
 	TArrayCollectionArray<EObjectStateType> MObjectState;
+	TQueue<TGeometryParticleHandle<T, d>*, EQueueMode::Mpsc> MSleepData;
 };
 
 
@@ -181,6 +188,9 @@ FChaosArchive& operator<<(FChaosArchive& Ar, TRigidParticles<T, d>& Particles)
 	return Ar;
 }
 
+#ifdef __clang__
+extern template class CHAOS_API TRigidParticles<float, 3>;
+#else
 extern template class TRigidParticles<float, 3>;
-
+#endif
 }
