Fixed destruction order for analytics provider. SessionSummaryWriter held a naked reference but the lifetime is managed via shared pointers

diff --git a/Engine/Source/Developer/EditorAnalyticsSession/Private/EditorSessionSummarySender.cpp b/Engine/Source/Developer/EditorAnalyticsSession/Private/EditorSessionSummarySender.cpp
--- a/Engine/Source/Developer/EditorAnalyticsSession/Private/EditorSessionSummarySender.cpp
+++ b/Engine/Source/Developer/EditorAnalyticsSession/Private/EditorSessionSummarySender.cpp
@@ -35,6 +35,10 @@ FEditorSessionSummarySender::FEditorSessionSummarySender(IAnalyticsProvider& InA
 	SendStoredSessions();
 }
 
+FEditorSessionSummarySender::~FEditorSessionSummarySender()
+{
+}
+
 void FEditorSessionSummarySender::Tick(float DeltaTime)
 {
 	HeartbeatTimeElapsed += DeltaTime;
diff --git a/Engine/Source/Developer/EditorAnalyticsSession/Public/EditorSessionSummarySender.h b/Engine/Source/Developer/EditorAnalyticsSession/Public/EditorSessionSummarySender.h
--- a/Engine/Source/Developer/EditorAnalyticsSession/Public/EditorSessionSummarySender.h
+++ b/Engine/Source/Developer/EditorAnalyticsSession/Public/EditorSessionSummarySender.h
@@ -16,6 +16,7 @@ class EDITORANALYTICSSESSION_API FEditorSessionSummarySender
 {
 public:
 	FEditorSessionSummarySender(IAnalyticsProvider& InAnalyticsProvider, const FString& InSenderName, const int32 InCurrentSessionProcessId);
+	~FEditorSessionSummarySender();
 
 	void Tick(float DeltaTime);
 	void Shutdown();
diff --git a/Engine/Source/Runtime/Engine/Private/EngineAnalytics.cpp b/Engine/Source/Runtime/Engine/Private/EngineAnalytics.cpp
--- a/Engine/Source/Runtime/Engine/Private/EngineAnalytics.cpp
+++ b/Engine/Source/Runtime/Engine/Private/EngineAnalytics.cpp
@@ -185,10 +185,6 @@ void FEngineAnalytics::Initialize()
 
 void FEngineAnalytics::Shutdown(bool bIsEngineShutdown)
 {
-	ensure(!Analytics.IsValid() || Analytics.IsUnique());
-	Analytics.Reset();
-	bIsInitialized = false;
-
 	// Destroy the session manager singleton if it exists
 	if (SessionManager.IsValid() && bIsEngineShutdown)
 	{
@@ -209,6 +205,11 @@ void FEngineAnalytics::Shutdown(bool bIsEngineShutdown)
 		SessionSummarySender.Reset();
 	}
 #endif
+
+	bIsInitialized = false;
+
+	ensure(!Analytics.IsValid() || Analytics.IsUnique());
+	Analytics.Reset();
 }
 
 void FEngineAnalytics::Tick(float DeltaTime)
