Updated the fix for statc meshes to be able to load cooked data in the editor (9291668) to work with the new Bulkdata API.

Renamed bUsingCookedData to bUsingCookedEditorData to make it clearer that it controls code branches for cooked data in the editor and then define out the code entirely if it will never be used to avoid API issues with the new Bulkdata API.

diff --git a/Engine/Source/Runtime/Engine/Private/StaticMesh.cpp b/Engine/Source/Runtime/Engine/Private/StaticMesh.cpp
--- a/Engine/Source/Runtime/Engine/Private/StaticMesh.cpp
+++ b/Engine/Source/Runtime/Engine/Private/StaticMesh.cpp
@@ -600,9 +600,9 @@ void FStaticMeshLODResources::Serialize(FArchive& Ar, UObject* Owner, int32 Inde
 {
 	DECLARE_SCOPE_CYCLE_COUNTER(TEXT("FStaticMeshLODResources::Serialize"), STAT_StaticMeshLODResources_Serialize, STATGROUP_LoadTime);
 
-	bool bUsingCookedData = false;
+	bool bUsingCookedEditorData = false;
 #if WITH_EDITORONLY_DATA
-	bUsingCookedData = Owner->GetOutermost()->bIsCookedForEditor;
+	bUsingCookedEditorData = Owner->GetOutermost()->bIsCookedForEditor;
 #endif
 
 	UStaticMesh* OwnerStaticMesh = Cast<UStaticMesh>(Owner);
@@ -632,7 +632,7 @@ void FStaticMeshLODResources::Serialize(FArchive& Ar, UObject* Owner, int32 Inde
 			Ar << TmpBuffersSize;
 			BuffersSize = TmpBuffersSize.CalcBuffersSize();
 		}
-		else if (FPlatformProperties::RequiresCookedData() || Ar.IsCooking() || bUsingCookedData)
+		else if (FPlatformProperties::RequiresCookedData() || Ar.IsCooking() || bUsingCookedEditorData)
 		{
 			uint32 BulkDataSize = 0;
 #if WITH_EDITOR
@@ -683,14 +683,16 @@ void FStaticMeshLODResources::Serialize(FArchive& Ar, UObject* Owner, int32 Inde
 				BulkDataSize = (uint32)StreamingBulkData.GetBulkDataSize();
 #endif
 
+#if WITH_EDITORONLY_DATA
 				// Streaming CPU data in editor build isn't supported yet because tools and utils need access
-				if (bUsingCookedData && BulkDataSize > 0)
+				if (bUsingCookedEditorData && BulkDataSize > 0)
 				{
 					TmpBuff.Empty(BulkDataSize);
 					TmpBuff.AddUninitialized(BulkDataSize);
 					void* Dest = TmpBuff.GetData();
 					TmpBulkData.GetCopy(&Dest);
 				}
+#endif
 			}
 
 			SerializeAvailabilityInfo(Ar);
@@ -703,13 +705,15 @@ void FStaticMeshLODResources::Serialize(FArchive& Ar, UObject* Owner, int32 Inde
 				ClearAvailabilityInfo();
 			}
 
-			if (Ar.IsLoading() && bUsingCookedData && BulkDataSize > 0)
+#if WITH_EDITORONLY_DATA
+			if (Ar.IsLoading() && bUsingCookedEditorData && BulkDataSize > 0)
 			{
 				ClearAvailabilityInfo();
 				FMemoryReader MemReader(TmpBuff, true);
 				MemReader.SetByteSwapping(Ar.IsByteSwapping());
 				SerializeBuffers(MemReader, OwnerStaticMesh, ClassDataStripFlags, TmpBuffersSize);
 			}
+#endif
 		}
 	}
 }
