Gameplay Insights plugin

Initial implementation of object and animation tracks for Unreal Insights, using the new plugin API.
Object systems trace classes, objects and object events.
Animation systems trace poses and tick records, with some experimental support for anim nodes & pose links.
Enabled by default for QAGame.

diff --git a/Engine/Plugins/Animation/GameplayInsights/GameplayInsights.uplugin b/Engine/Plugins/Animation/GameplayInsights/GameplayInsights.uplugin
new file mode 100644
--- /dev/null
+++ b/Engine/Plugins/Animation/GameplayInsights/GameplayInsights.uplugin
@@ -0,0 +1,24 @@
+{
+	"FileVersion": 3,
+	"Version": 1,
+	"VersionName": "1.0",
+	"FriendlyName": "Gameplay Insights",
+	"Description": "Allows debugging of gameplay and animation systems via Unreal Insights",
+	"Category": "GameplayInsights",
+	"CreatedBy": "Epic Games, Inc.",
+	"CreatedByURL": "http://epicgames.com",
+	"DocsURL": "",
+	"MarketplaceURL": "",
+	"SupportURL": "",
+	"EnabledByDefault": false,
+	"CanContainContent": false,
+	"IsBetaVersion": false,
+	"Installed": false,
+	"Modules": [
+		{
+			"Name": "GameplayInsights",
+			"Type": "Editor",
+			"LoadingPhase": "Default"
+		}
+	]
+}
\ No newline at end of file
diff --git a/Engine/Plugins/Animation/GameplayInsights/Source/GameplayInsights/GameplayInsights.Build.cs b/Engine/Plugins/Animation/GameplayInsights/Source/GameplayInsights/GameplayInsights.Build.cs
new file mode 100644
--- /dev/null
+++ b/Engine/Plugins/Animation/GameplayInsights/Source/GameplayInsights/GameplayInsights.Build.cs
@@ -0,0 +1,29 @@
+// Copyright 1998-2019 Epic Games, Inc. All Rights Reserved.
+
+namespace UnrealBuildTool.Rules
+{
+	public class GameplayInsights : ModuleRules
+	{
+		public GameplayInsights(ReadOnlyTargetRules Target) : base(Target)
+		{
+			PrivateDependencyModuleNames.AddRange(
+			new string[]
+			{
+				"Core",
+				"CoreUObject",
+				"InputCore",
+				"SlateCore",
+				"Slate",
+				"TraceAnalysis",
+				"TraceServices",
+				"TraceInsights",
+			});
+
+			if (Target.bCompileAgainstEngine)
+			{
+				PrivateDependencyModuleNames.Add("Engine");
+			}
+		}
+	}
+}
+
diff --git a/Engine/Plugins/Animation/GameplayInsights/Source/GameplayInsights/Private/AnimationAnalyzer.cpp b/Engine/Plugins/Animation/GameplayInsights/Source/GameplayInsights/Private/AnimationAnalyzer.cpp
new file mode 100644
--- /dev/null
+++ b/Engine/Plugins/Animation/GameplayInsights/Source/GameplayInsights/Private/AnimationAnalyzer.cpp
@@ -0,0 +1,68 @@
+// Copyright 1998-2019 Epic Games, Inc. All Rights Reserved.
+
+#include "AnimationAnalyzer.h"
+#include "TraceServices/Model/AnalysisSession.h"
+#include "AnimationProvider.h"
+#include "Containers/ArrayView.h"
+
+FAnimationAnalyzer::FAnimationAnalyzer(Trace::IAnalysisSession& InSession, FAnimationProvider& InAnimationProvider)
+	: Session(InSession)
+	, AnimationProvider(InAnimationProvider)
+{
+}
+
+void FAnimationAnalyzer::OnAnalysisBegin(const FOnAnalysisContext& Context)
+{
+	auto& Builder = Context.InterfaceBuilder;
+
+	Builder.RouteEvent(RouteId_TickRecord, "Animation", "TickRecord");
+	Builder.RouteEvent(RouteId_SkeletalMesh, "Animation", "SkeletalMesh");
+	Builder.RouteEvent(RouteId_SkeletalMeshPose, "Animation", "SkeletalMeshPose");
+}
+
+bool FAnimationAnalyzer::OnEvent(uint16 RouteId, const FOnEventContext& Context)
+{
+	Trace::FAnalysisSessionEditScope _(Session);
+
+	const auto& EventData = Context.EventData;
+	switch (RouteId)
+	{
+	case RouteId_TickRecord:
+	{
+		uint64 Cycle = EventData.GetValue<uint64>("Cycle");
+		uint64 ComponentId = EventData.GetValue<uint64>("ComponentId");
+		uint64 AnimInstanceId = EventData.GetValue<uint64>("AnimInstanceId");
+		uint64 AssetId = EventData.GetValue<uint64>("AssetId");
+		float BlendWeight = EventData.GetValue<float>("BlendWeight");
+		float PlaybackTime = EventData.GetValue<float>("PlaybackTime");
+		float RootMotionWeight = EventData.GetValue<float>("RootMotionWeight");
+		float PlayRate = EventData.GetValue<float>("PlayRate");
+		uint16 FrameCounter = EventData.GetValue<uint16>("FrameCounter");
+		bool bLooping = EventData.GetValue<bool>("Looping");
+		AnimationProvider.AppendTickRecord(ComponentId, AnimInstanceId, Context.SessionContext.TimestampFromCycle(Cycle), AssetId, BlendWeight, PlaybackTime, RootMotionWeight, PlayRate, FrameCounter, bLooping);
+		break;
+	}
+	case RouteId_SkeletalMesh:
+	{
+		uint64 Id = EventData.GetValue<uint64>("Id");
+		uint32 BoneCount = EventData.GetValue<uint32>("BoneCount");
+		TArrayView<const int32> ParentIndices(reinterpret_cast<const int32*>(EventData.GetAttachment()), BoneCount);
+		AnimationProvider.AppendSkeletalMesh(Id, ParentIndices);
+		break;
+	}
+	case RouteId_SkeletalMeshPose:
+	{
+		uint64 Cycle = EventData.GetValue<uint64>("Cycle");
+		uint64 ComponentId = EventData.GetValue<uint64>("ComponentId");
+		uint64 MeshId = EventData.GetValue<uint64>("MeshId");
+		uint32 BoneCount = EventData.GetValue<uint32>("BoneCount");
+		uint16 FrameCounter = EventData.GetValue<uint16>("FrameCounter");
+		uint16 LodIndex = EventData.GetValue<uint16>("LodIndex");
+		TArrayView<const FTransform> Pose(reinterpret_cast<const FTransform*>(EventData.GetAttachment()), BoneCount);
+		AnimationProvider.AppendSkeletalMeshPose(ComponentId, MeshId, Context.SessionContext.TimestampFromCycle(Cycle), LodIndex, FrameCounter, Pose);
+		break;
+	}
+	}
+
+	return true;
+}
\ No newline at end of file
diff --git a/Engine/Plugins/Animation/GameplayInsights/Source/GameplayInsights/Private/AnimationAnalyzer.h b/Engine/Plugins/Animation/GameplayInsights/Source/GameplayInsights/Private/AnimationAnalyzer.h
new file mode 100644
--- /dev/null
+++ b/Engine/Plugins/Animation/GameplayInsights/Source/GameplayInsights/Private/AnimationAnalyzer.h
@@ -0,0 +1,29 @@
+// Copyright 1998-2019 Epic Games, Inc. All Rights Reserved.
+
+#pragma once
+
+#include "Trace/Analyzer.h"
+
+class FAnimationProvider;
+namespace Trace { class IAnalysisSession; }
+
+class FAnimationAnalyzer : public Trace::IAnalyzer
+{
+public:
+	FAnimationAnalyzer(Trace::IAnalysisSession& InSession, FAnimationProvider& InAnimationProvider);
+
+	virtual void OnAnalysisBegin(const FOnAnalysisContext& Context) override;
+	virtual void OnAnalysisEnd() override {}
+	virtual bool OnEvent(uint16 RouteId, const FOnEventContext& Context) override;
+
+private:
+	enum : uint16
+	{
+		RouteId_TickRecord,
+		RouteId_SkeletalMesh,
+		RouteId_SkeletalMeshPose,
+	};
+
+	Trace::IAnalysisSession& Session;
+	FAnimationProvider& AnimationProvider;
+};
\ No newline at end of file
diff --git a/Engine/Plugins/Animation/GameplayInsights/Source/GameplayInsights/Private/AnimationProvider.cpp b/Engine/Plugins/Animation/GameplayInsights/Source/GameplayInsights/Private/AnimationProvider.cpp
new file mode 100644
--- /dev/null
+++ b/Engine/Plugins/Animation/GameplayInsights/Source/GameplayInsights/Private/AnimationProvider.cpp
@@ -0,0 +1,263 @@
+// Copyright 1998-2019 Epic Games, Inc. All Rights Reserved.
+
+#include "AnimationProvider.h"
+#include "GameplayProvider.h"
+#include "Insights/ViewModels/TimingEventsTrack.h"
+
+FName FAnimationProvider::ProviderName("AnimationProvider");
+
+FAnimationProvider::FAnimationProvider(Trace::IAnalysisSession& InSession)
+	: Session(InSession)
+	, SkeletalMeshPoseTransforms(InSession.GetLinearAllocator(), 256)
+	, SkeletalMeshParentIndices(InSession.GetLinearAllocator(), 256)
+{
+}
+
+bool FAnimationProvider::ReadSkeletalMeshPoseTimeline(uint64 InObjectId, TFunctionRef<void(const SkeletalMeshPoseTimeline&)> Callback) const
+{
+	Session.ReadAccessCheck();
+
+	const uint32* IndexPtr = ObjectIdToSkeletalMeshPoseTimelines.Find(InObjectId);
+	if(IndexPtr != nullptr)
+	{
+		if (*IndexPtr < uint32(SkeletalMeshPoseTimelines.Num()))
+		{
+			Callback(*SkeletalMeshPoseTimelines[*IndexPtr]);
+			return true;
+		}
+	}
+
+	return false;
+}
+
+bool FAnimationProvider::ReadSkeletalMeshPoseMessage(uint64 InObjectId, uint64 InMessageId, TFunctionRef<void(const FSkeletalMeshPoseMessage&)> Callback) const
+{
+	return ReadSkeletalMeshPoseTimeline(InObjectId, [&Callback, &InMessageId](const SkeletalMeshPoseTimeline& InTimeline)
+	{
+		if(InMessageId < InTimeline.GetEventCount())
+		{
+			Callback(InTimeline.GetEvent(InMessageId));
+		}
+	});
+}
+
+void FAnimationProvider::EnumerateSkeletalMeshPose(const FSkeletalMeshPoseMessage& InMessage, const FSkeletalMeshInfo& InMeshInfo, TFunctionRef<void(const FTransform&, const FTransform&)> Callback) const
+{
+	Session.ReadAccessCheck();
+
+	if(InMeshInfo.BoneCount == InMessage.NumTransforms - 1)
+	{
+		// First transform is always component to world
+		const FTransform& ComponentToWorld = SkeletalMeshPoseTransforms[InMessage.TransformStartIndex];
+		uint64 StartTransformIndex = InMessage.TransformStartIndex + 1;
+		uint64 RootBoneIndex = InMeshInfo.ParentIndicesStartIndex;
+		uint64 EndTransformIndex = InMessage.TransformStartIndex + InMessage.NumTransforms;
+		uint64 TransformIndex;
+		uint64 BoneIndex;
+		for(TransformIndex = StartTransformIndex, BoneIndex = RootBoneIndex; TransformIndex < EndTransformIndex; ++TransformIndex, ++BoneIndex)
+		{
+			int32 ParentIndex = SkeletalMeshParentIndices[BoneIndex];
+			uint64 ParentTransformIndex = ParentIndex == INDEX_NONE ? TransformIndex : (StartTransformIndex + (uint64)ParentIndex);
+			Callback(SkeletalMeshPoseTransforms[TransformIndex] * ComponentToWorld, SkeletalMeshPoseTransforms[ParentTransformIndex] * ComponentToWorld);
+		}
+	}
+}
+
+void FAnimationProvider::EnumerateTickRecordTimelines(uint64 InObjectId, TFunctionRef<void(uint64, const TickRecordTimeline&)> Callback) const
+{
+	Session.ReadAccessCheck();
+
+	const uint32* IndexPtr = ObjectIdToTickRecordTimelineStorage.Find(InObjectId);
+	if(IndexPtr != nullptr)
+	{
+		if (*IndexPtr < uint32(PerObjectTimelineStorage.Num()))
+		{
+			const TSharedRef<FPerObjectTimelineStorage>& TimelineStorage = PerObjectTimelineStorage[*IndexPtr];
+			for(auto AssetIdPair : TimelineStorage->AssetIdToTickRecordTimeline)
+			{
+				Callback(AssetIdPair.Key, TimelineStorage->Timelines[AssetIdPair.Value].Get());
+			}
+		}
+	}
+}
+
+bool FAnimationProvider::ReadTickRecordTimeline(uint64 InObjectId, uint64 InAssetId, TFunctionRef<void(const TickRecordTimeline&)> Callback) const
+{
+	Session.ReadAccessCheck();
+
+	const uint32* IndexPtr = ObjectIdToTickRecordTimelineStorage.Find(InObjectId);
+	if(IndexPtr != nullptr)
+	{
+		if (*IndexPtr < uint32(PerObjectTimelineStorage.Num()))
+		{
+			const TSharedRef<FPerObjectTimelineStorage>& TimelineStorage = PerObjectTimelineStorage[*IndexPtr];
+			const uint32* TimelineIndexPtr = TimelineStorage->AssetIdToTickRecordTimeline.Find(InAssetId);
+			if(TimelineIndexPtr != nullptr)
+			{
+				Callback(TimelineStorage->Timelines[*TimelineIndexPtr].Get());
+				return true;
+			}
+		}
+	}
+
+	return false;
+}
+
+bool FAnimationProvider::ReadTickRecordMessage(uint64 InObjectId, uint64 InAssetId, uint64 InMessageId, TFunctionRef<void(const FTickRecordMessage&)> Callback) const
+{
+	return ReadTickRecordTimeline(InObjectId, InAssetId, [&Callback, &InMessageId](const TickRecordTimeline& InTimeline)
+	{
+		if(InMessageId < InTimeline.GetEventCount())
+		{
+			Callback(InTimeline.GetEvent(InMessageId));
+		}
+	});
+}
+
+const FSkeletalMeshInfo* FAnimationProvider::FindSkeletalMeshInfo(uint64 InObjectId) const
+{
+	const int32* ObjectIndex = SkeletalMeshIdToIndexMap.Find(InObjectId);
+	if(ObjectIndex != nullptr)
+	{
+		return &SkeletalMeshInfos[*ObjectIndex];
+	}
+
+	return nullptr;
+}
+
+void FAnimationProvider::AppendTickRecord(uint64 InObjectId, uint64 InSubObjectId, double InTime, uint64 InAssetId, float InBlendWeight, float InPlaybackTime, float InRootMotionWeight, float InPlayRate, uint16 InFrameCounter, bool bInLooping)
+{
+	Session.WriteAccessCheck();
+
+	TSharedPtr<Trace::TMonotonicTimeline<FTickRecordMessage>> Timeline;
+	TSharedPtr<FPerObjectTimelineStorage> TimelineStorage;
+	uint32* TimelineStorageIndexPtr = ObjectIdToTickRecordTimelineStorage.Find(InObjectId);
+	if(TimelineStorageIndexPtr != nullptr)
+	{
+		TimelineStorage = PerObjectTimelineStorage[*TimelineStorageIndexPtr];
+	}
+	else
+	{
+		ObjectIdToTickRecordTimelineStorage.Add(InObjectId, PerObjectTimelineStorage.Num());
+		TimelineStorage = PerObjectTimelineStorage.Add_GetRef(MakeShared<FPerObjectTimelineStorage>());
+	}
+
+	check(TimelineStorage.IsValid());
+
+	uint32* TimelineIndexPtr = TimelineStorage->AssetIdToTickRecordTimeline.Find(InAssetId);
+	if(TimelineIndexPtr != nullptr)
+	{
+		Timeline = TimelineStorage->Timelines[*TimelineIndexPtr];
+	}
+	else
+	{
+		Timeline = MakeShared<Trace::TMonotonicTimeline<FTickRecordMessage>>(Session.GetLinearAllocator());
+		TimelineStorage->AssetIdToTickRecordTimeline.Add(InAssetId, TimelineStorage->Timelines.Num());
+		TimelineStorage->Timelines.Add(Timeline.ToSharedRef());
+	}
+
+	// terminate existing scopes - check for continuity
+	uint64 NumEvents = Timeline->GetEventCount();
+	if(NumEvents > 0)
+	{
+		if(InFrameCounter - 1 != Timeline->GetEvent(NumEvents - 1).FrameCounter)
+		{
+			// discontinuous frame count, so add an end event at last time
+			Timeline->AppendEndEvent(Timeline->GetEndTime());
+		}
+		else
+		{
+			// continuous frame time, add end event at current time
+			Timeline->AppendEndEvent(InTime);
+		}
+	}
+
+	FTickRecordMessage Message;
+	Message.MessageId = Timeline->GetEventCount();
+	Message.ComponentId = InObjectId;
+	Message.AnimInstanceId = InSubObjectId;
+	Message.AssetId = InAssetId;
+	Message.BlendWeight = InBlendWeight;
+	Message.PlaybackTime = InPlaybackTime;
+	Message.RootMotionWeight = InRootMotionWeight;
+	Message.PlayRate = InPlayRate;
+	Message.FrameCounter = InFrameCounter;
+	Message.Looping = bInLooping;
+
+	Timeline->AppendBeginEvent(InTime, Message);
+
+	Session.UpdateDurationSeconds(InTime);
+}
+
+void FAnimationProvider::AppendSkeletalMesh(uint64 InObjectId, const TArrayView<const int32>& InParentIndices)
+{
+	Session.WriteAccessCheck();
+
+	if(SkeletalMeshIdToIndexMap.Find(InObjectId) == nullptr)
+	{
+		FSkeletalMeshInfo NewSkeletalMeshInfo;
+		NewSkeletalMeshInfo.Id = InObjectId;
+		NewSkeletalMeshInfo.BoneCount = (uint32)InParentIndices.Num();
+		NewSkeletalMeshInfo.ParentIndicesStartIndex = SkeletalMeshParentIndices.Num();
+
+		for(const int32& ParentIndex : InParentIndices)
+		{
+			SkeletalMeshParentIndices.PushBack() = ParentIndex;
+		}
+
+		int32 NewSkeletalMeshInfoIndex = SkeletalMeshInfos.Add(NewSkeletalMeshInfo);
+		SkeletalMeshIdToIndexMap.Add(InObjectId, NewSkeletalMeshInfoIndex);
+	}
+}
+
+void FAnimationProvider::AppendSkeletalMeshPose(uint64 InObjectId, uint64 InMeshId, double InTime, uint16 InLodIndex, uint16 InFrameCounter, const TArrayView<const FTransform>& InPose)
+{
+	Session.WriteAccessCheck();
+
+	TSharedPtr<Trace::TMonotonicTimeline<FSkeletalMeshPoseMessage>> Timeline;
+	uint32* IndexPtr = ObjectIdToSkeletalMeshPoseTimelines.Find(InObjectId);
+	if(IndexPtr != nullptr)
+	{
+		Timeline = SkeletalMeshPoseTimelines[*IndexPtr];
+	}
+	else
+	{
+		Timeline = MakeShared<Trace::TMonotonicTimeline<FSkeletalMeshPoseMessage>>(Session.GetLinearAllocator());
+		ObjectIdToSkeletalMeshPoseTimelines.Add(InObjectId, SkeletalMeshPoseTimelines.Num());
+		SkeletalMeshPoseTimelines.Add(Timeline.ToSharedRef());
+	}
+
+	// terminate existing scopes - check for continuity
+	uint64 NumEvents = Timeline->GetEventCount();
+	if(NumEvents > 0)
+	{
+		if(InFrameCounter - 1 != Timeline->GetEvent(NumEvents - 1).FrameCounter)
+		{
+			// discontinuous frame count, so add an end event at last time
+			Timeline->AppendEndEvent(Timeline->GetEndTime());
+		}
+		else
+		{
+			// continuous frame time, add end event at current time
+			Timeline->AppendEndEvent(InTime);
+		}
+	}
+
+	FSkeletalMeshPoseMessage Message;
+	Message.MessageId = Timeline->GetEventCount();
+	Message.TransformStartIndex = SkeletalMeshPoseTransforms.Num();
+	Message.ComponentId = InObjectId;
+	Message.MeshId = InMeshId;
+	Message.NumTransforms = (uint16)InPose.Num();
+	Message.LodIndex = InLodIndex;
+	Message.FrameCounter = InFrameCounter;
+
+	Timeline->AppendBeginEvent(InTime, Message);
+
+	for(const FTransform& Transform : InPose)
+	{
+		SkeletalMeshPoseTransforms.PushBack() = Transform;
+	}
+
+	Session.UpdateDurationSeconds(InTime);
+}
diff --git a/Engine/Plugins/Animation/GameplayInsights/Source/GameplayInsights/Private/AnimationProvider.h b/Engine/Plugins/Animation/GameplayInsights/Source/GameplayInsights/Private/AnimationProvider.h
new file mode 100644
--- /dev/null
+++ b/Engine/Plugins/Animation/GameplayInsights/Source/GameplayInsights/Private/AnimationProvider.h
@@ -0,0 +1,60 @@
+// Copyright 1998-2019 Epic Games, Inc. All Rights Reserved.
+
+#pragma once
+
+#include "IAnimationProvider.h"
+#include "Common/PagedArray.h"
+#include "Model/MonotonicTimeline.h"
+
+namespace Trace { class IAnalysisSession; }
+
+class FAnimationProvider : public IAnimationProvider
+{
+public:
+	static FName ProviderName;
+
+	FAnimationProvider(Trace::IAnalysisSession& InSession);
+
+	/** IAnimationProvider interface */
+	virtual bool ReadSkeletalMeshPoseTimeline(uint64 InObjectId, TFunctionRef<void(const SkeletalMeshPoseTimeline&)> Callback) const override;
+	virtual bool ReadSkeletalMeshPoseMessage(uint64 InObjectId, uint64 InMessageId, TFunctionRef<void(const FSkeletalMeshPoseMessage&)> Callback) const override;
+	virtual void EnumerateSkeletalMeshPose(const FSkeletalMeshPoseMessage& InMessage, const FSkeletalMeshInfo& InMeshInfo, TFunctionRef<void(const FTransform&, const FTransform&)> Callback) const override;
+	virtual void EnumerateTickRecordTimelines(uint64 InObjectId, TFunctionRef<void(uint64, const TickRecordTimeline&)> Callback) const override;
+	virtual bool ReadTickRecordTimeline(uint64 InObjectId, uint64 InAssetId, TFunctionRef<void(const TickRecordTimeline&)> Callback) const override;
+	virtual bool ReadTickRecordMessage(uint64 InObjectId, uint64 InAssetId, uint64 InMessageId, TFunctionRef<void(const FTickRecordMessage&)> Callback) const override;
+	virtual const FSkeletalMeshInfo* FindSkeletalMeshInfo(uint64 InObjectId) const override;
+
+	/** Add a tick record */
+	void AppendTickRecord(uint64 InObjectId, uint64 InSubObjectId, double InTime, uint64 InAssetId, float InBlendWeight, float InPlaybackTime, float InRootMotionWeight, float InPlayRate, uint16 InFrameCounter, bool bInLooping);
+
+	/** Add a skeletal mesh */
+	void AppendSkeletalMesh(uint64 InObjectId, const TArrayView<const int32>& ParentIndices);
+
+	/** Add a skeletal mesh pose */
+	void AppendSkeletalMeshPose(uint64 InObjectId, uint64 InMeshId, double InTime, uint16 InLodIndex, uint16 InFrameCounter, const TArrayView<const FTransform>& InTransforms);
+
+private:
+	Trace::IAnalysisSession& Session;
+
+	/** Maps into timeline arrays per-object */
+	TMap<uint64, uint32> ObjectIdToTickRecordTimelineStorage;
+	TMap<uint64, uint32> ObjectIdToSkeletalMeshPoseTimelines;
+
+	/** All the skeletal mesh info we have seen, grow only for stable indices */
+	TArray<FSkeletalMeshInfo> SkeletalMeshInfos;
+
+	/** Map of skeletal mesh->index into SkeletalMeshInfos index */
+	TMap<uint64, int32> SkeletalMeshIdToIndexMap;
+
+	struct FPerObjectTimelineStorage
+	{
+		TMap<uint64, uint32> AssetIdToTickRecordTimeline;
+		TArray<TSharedRef<Trace::TMonotonicTimeline<FTickRecordMessage>>> Timelines;
+	};
+
+	/** Message storage */
+	TArray<TSharedRef<FPerObjectTimelineStorage>> PerObjectTimelineStorage;
+	TArray<TSharedRef<Trace::TMonotonicTimeline<FSkeletalMeshPoseMessage>>> SkeletalMeshPoseTimelines;
+	TPagedArray<FTransform> SkeletalMeshPoseTransforms;
+	TPagedArray<int32> SkeletalMeshParentIndices;
+};
\ No newline at end of file
diff --git a/Engine/Plugins/Animation/GameplayInsights/Source/GameplayInsights/Private/AnimationSharedData.cpp b/Engine/Plugins/Animation/GameplayInsights/Source/GameplayInsights/Private/AnimationSharedData.cpp
new file mode 100644
--- /dev/null
+++ b/Engine/Plugins/Animation/GameplayInsights/Source/GameplayInsights/Private/AnimationSharedData.cpp
@@ -0,0 +1,227 @@
+// Copyright 1998-2019 Epic Games, Inc. All Rights Reserved.
+
+#include "AnimationSharedData.h"
+#include "TraceServices/Model/AnalysisSession.h"
+#include "GameplaySharedData.h"
+#include "ObjectEventsTrack.h"
+#include "SkeletalMeshPoseTrack.h"
+#include "AnimationTickRecordsTrack.h"
+#include "Insights/ITimingViewSession.h"
+#include "Insights/ViewModels/TimingEvent.h"
+
+#define LOCTEXT_NAMESPACE "AnimationSharedData"
+
+FAnimationSharedData::FAnimationSharedData(FGameplaySharedData& InGameplaySharedData)
+	: GameplaySharedData(InGameplaySharedData)
+	, AnalysisSession(nullptr)
+	, SelectedEventStartTime(0.0)
+	, SelectedEventEndTime(0.0)
+	, HoveredEventStartTime(0.0)
+	, HoveredEventEndTime(0.0)
+	, SelectionStartTime(0.0)
+	, SelectionEndTime(0.0)
+	, MarkerTime(0.0)
+	, bSelectedEventValid(false)
+	, bHoveredEventValid(false)
+	, bSelectionValid(false)
+	, bTimeMarkerValid(false)
+	, bAnimationTracksEnabled(true)
+{
+}
+
+void FAnimationSharedData::OnBeginSession(Insights::ITimingViewSession& InTimingViewSession)
+{
+	SkeletalMeshPoseTracks.Reset();
+	AnimationTickRecordsTracks.Reset();
+
+	SelectedEventChangedHandle = InTimingViewSession.OnSelectedEventChanged().AddRaw(this, &FAnimationSharedData::OnSelectedEventChanged);
+	HoveredEventChangedHandle = InTimingViewSession.OnHoveredEventChanged().AddRaw(this, &FAnimationSharedData::OnHoveredEventChanged);
+	TimeMarkerChangedHandle = InTimingViewSession.OnTimeMarkerChanged().AddRaw(this, &FAnimationSharedData::OnTimeMarkerChanged);
+	SelectionChangedHandle = InTimingViewSession.OnSelectionChanged().AddRaw(this, &FAnimationSharedData::OnSelectionChanged);
+}
+
+void FAnimationSharedData::OnEndSession(Insights::ITimingViewSession& InTimingViewSession)
+{
+	SkeletalMeshPoseTracks.Reset();
+	AnimationTickRecordsTracks.Reset();
+
+	InTimingViewSession.OnSelectedEventChanged().Remove(SelectedEventChangedHandle);
+	InTimingViewSession.OnHoveredEventChanged().Remove(HoveredEventChangedHandle);
+	InTimingViewSession.OnTimeMarkerChanged().Remove(TimeMarkerChangedHandle);
+	InTimingViewSession.OnSelectionChanged().Remove(SelectionChangedHandle);
+}
+
+void FAnimationSharedData::Tick(Insights::ITimingViewSession& InTimingViewSession, const Trace::IAnalysisSession& InAnalysisSession)
+{
+	AnalysisSession = &InAnalysisSession;
+
+	const FAnimationProvider* AnimationProvider = InAnalysisSession.ReadProvider<FAnimationProvider>(FAnimationProvider::ProviderName);
+	const FGameplayProvider* GameplayProvider = InAnalysisSession.ReadProvider<FGameplayProvider>(FGameplayProvider::ProviderName);
+
+	if(AnimationProvider && GameplayProvider)
+	{
+		// Add tracks for each tracked object's animation data
+		GameplayProvider->EnumerateObjects([this, &InTimingViewSession, &InAnalysisSession, &AnimationProvider, &GameplayProvider](const FObjectInfo& InObjectInfo)
+		{
+			FObjectEventsTrack* ObjectEventsTrack = GameplaySharedData.GetObjectEventsTrackForId(InTimingViewSession, InAnalysisSession, InObjectInfo);
+
+			AnimationProvider->ReadSkeletalMeshPoseTimeline(InObjectInfo.Id, [this, &InObjectInfo, &ObjectEventsTrack, &InTimingViewSession](const IAnimationProvider::SkeletalMeshPoseTimeline& InTimeline)
+			{
+				ObjectEventsTrack->SetVisibilityFlag(GameplaySharedData.AreGameplayTracksEnabled());
+
+				auto FindSkeletalMeshPoseTrack = [](const FBaseTimingTrack& InTrack)
+				{
+					if (InTrack.GetType() == FSkeletalMeshPoseTrack::TypeName &&
+						InTrack.GetSubType() == FSkeletalMeshPoseTrack::SubTypeName)
+					{
+						return true;
+					}
+
+					return false;
+				};
+
+				FSkeletalMeshPoseTrack* SkeletalMeshPoseTrack = static_cast<FSkeletalMeshPoseTrack*>(ObjectEventsTrack->GetGameplayTrack().FindChildTrack(InObjectInfo.Id, FindSkeletalMeshPoseTrack));
+				if(SkeletalMeshPoseTrack == nullptr)
+				{
+					SkeletalMeshPoseTrack = new FSkeletalMeshPoseTrack(*this, InObjectInfo.Id, InObjectInfo.Name);
+					SkeletalMeshPoseTrack->SetVisibilityFlag(bAnimationTracksEnabled);
+					SkeletalMeshPoseTracks.Add(SkeletalMeshPoseTrack);
+
+					InTimingViewSession.AddTimingEventsTrack(SkeletalMeshPoseTrack);
+
+					ObjectEventsTrack->GetGameplayTrack().AddChildTrack(SkeletalMeshPoseTrack->GetGameplayTrack());
+				}
+			});
+
+			AnimationProvider->EnumerateTickRecordTimelines(InObjectInfo.Id, [this, &InObjectInfo, &ObjectEventsTrack, &InTimingViewSession, &GameplayProvider](uint64 InAssetId, const IAnimationProvider::TickRecordTimeline& InTimeline)
+			{
+				ObjectEventsTrack->SetVisibilityFlag(GameplaySharedData.AreGameplayTracksEnabled());
+
+				auto FindTickRecordTrackWithAssetId = [&InAssetId](const FBaseTimingTrack& InTrack)
+				{
+					if (InTrack.GetType() == FAnimationTickRecordsTrack::TypeName &&
+						InTrack.GetSubType() == FAnimationTickRecordsTrack::SubTypeName)
+					{
+						const FAnimationTickRecordsTrack& AnimationTickRecordsTrack = *static_cast<const FAnimationTickRecordsTrack*>(&InTrack);
+						return AnimationTickRecordsTrack.GetAssetId() == InAssetId;
+					}
+
+					return false;
+				};
+
+				FAnimationTickRecordsTrack* AnimationTickRecordsTrack = static_cast<FAnimationTickRecordsTrack*>(ObjectEventsTrack->GetGameplayTrack().FindChildTrack(InObjectInfo.Id, FindTickRecordTrackWithAssetId));
+				if(AnimationTickRecordsTrack == nullptr)
+				{
+					const FObjectInfo* AssetObjectInfo = GameplayProvider->FindObjectInfo(InAssetId);
+					FString AssetName = AssetObjectInfo ? AssetObjectInfo->Name : LOCTEXT("UnknownAsset", "Unknown").ToString();
+					AnimationTickRecordsTrack = new FAnimationTickRecordsTrack(*this, InObjectInfo.Id, InAssetId, *AssetName);
+					AnimationTickRecordsTrack->SetVisibilityFlag(bAnimationTracksEnabled);
+					AnimationTickRecordsTracks.Add(AnimationTickRecordsTrack);
+
+					InTimingViewSession.AddTimingEventsTrack(AnimationTickRecordsTrack);
+
+					ObjectEventsTrack->GetGameplayTrack().AddChildTrack(AnimationTickRecordsTrack->GetGameplayTrack());
+				}
+			});
+		});
+	}
+}
+
+void FAnimationSharedData::ExtendFilterMenu(FMenuBuilder& InMenuBuilder)
+{
+	InMenuBuilder.AddMenuEntry(
+		LOCTEXT("ToggleAnimationTracks", "Animation Tracks"),
+		LOCTEXT("ToggleAnimationTracks_Tooltip", "Show/hide the animation tracks"),
+		FSlateIcon(),
+		FUIAction(
+			FExecuteAction::CreateRaw(this, &FAnimationSharedData::ToggleAnimationTracks),
+			FCanExecuteAction(),
+			FIsActionChecked::CreateRaw(this, &FAnimationSharedData::AreAnimationTracksEnabled)),
+		NAME_None,
+		EUserInterfaceActionType::ToggleButton
+	);
+}
+
+void FAnimationSharedData::ToggleAnimationTracks()
+{
+	bAnimationTracksEnabled = !bAnimationTracksEnabled;
+
+	for(FSkeletalMeshPoseTrack* PoseTrack : SkeletalMeshPoseTracks)
+	{
+		PoseTrack->SetVisibilityFlag(bAnimationTracksEnabled);
+	}
+
+	for(FAnimationTickRecordsTrack* TickRecordTrack : AnimationTickRecordsTracks)
+	{
+		TickRecordTrack->SetVisibilityFlag(bAnimationTracksEnabled);
+	}
+}
+
+bool FAnimationSharedData::AreAnimationTracksEnabled() const
+{
+	return bAnimationTracksEnabled;
+}
+
+void FAnimationSharedData::OnSelectedEventChanged(const FTimingEvent& InEvent)
+{
+	bSelectedEventValid = InEvent.IsValid();
+	SelectedEventTrack = InEvent.Track;
+	SelectedEventStartTime = InEvent.StartTime;
+	SelectedEventEndTime = InEvent.EndTime;
+}
+
+void FAnimationSharedData::OnHoveredEventChanged(const FTimingEvent& InEvent)
+{
+	bHoveredEventValid = InEvent.IsValid();
+	HoveredEventTrack = InEvent.Track;
+	HoveredEventStartTime = InEvent.StartTime;
+	HoveredEventEndTime = InEvent.EndTime;
+}
+
+void FAnimationSharedData::OnTimeMarkerChanged(Insights::ETimeChangedFlags InFlags, double InTimeMarker)
+{
+	bTimeMarkerValid = InTimeMarker != std::numeric_limits<double>::infinity();
+	MarkerTime = InTimeMarker;
+}
+
+void FAnimationSharedData::OnSelectionChanged(Insights::ETimeChangedFlags InFlags, double InStartTime, double InEndTime)
+{
+	bSelectionValid = InStartTime < InEndTime;
+	SelectionStartTime = InStartTime;
+	SelectionEndTime = InEndTime;
+}
+
+#if WITH_ENGINE
+
+void FAnimationSharedData::DrawPoses(ULineBatchComponent* InLineBatcher)
+{
+	for(FSkeletalMeshPoseTrack* PoseTrack : SkeletalMeshPoseTracks)
+	{
+		if(PoseTrack->IsVisible())
+		{
+			if(bSelectedEventValid && PoseTrack == SelectedEventTrack && PoseTrack->ShouldDrawSelectedEvent())
+			{
+				PoseTrack->DrawPoses(InLineBatcher, SelectedEventStartTime, SelectedEventEndTime);
+			}
+		
+			if(bHoveredEventValid && PoseTrack == HoveredEventTrack && PoseTrack->ShouldDrawHoveredEvent())
+			{
+				PoseTrack->DrawPoses(InLineBatcher, HoveredEventStartTime, HoveredEventEndTime);
+			}
+
+			if(bTimeMarkerValid && PoseTrack->ShouldDrawMarkerTime())
+			{
+				PoseTrack->DrawPoses(InLineBatcher, MarkerTime, MarkerTime);
+			}
+
+			if(bSelectionValid && PoseTrack->ShouldDrawSelection())
+			{
+				PoseTrack->DrawPoses(InLineBatcher, SelectionStartTime, SelectionEndTime);
+			}
+		}
+	}
+}
+
+#endif
+
+#undef LOCTEXT_NAMESPACE
\ No newline at end of file
diff --git a/Engine/Plugins/Animation/GameplayInsights/Source/GameplayInsights/Private/AnimationSharedData.h b/Engine/Plugins/Animation/GameplayInsights/Source/GameplayInsights/Private/AnimationSharedData.h
new file mode 100644
--- /dev/null
+++ b/Engine/Plugins/Animation/GameplayInsights/Source/GameplayInsights/Private/AnimationSharedData.h
@@ -0,0 +1,85 @@
+// Copyright 1998-2019 Epic Games, Inc. All Rights Reserved.
+
+#pragma once
+
+#include "Insights/ViewModels/TimingEventsTrack.h"
+
+class FGameplaySharedData;
+namespace Trace { class IAnalysisSession; }
+namespace Insights { class ITimingViewSession; }
+namespace Insights { enum class ETimeChangedFlags : int32; }
+class FSkeletalMeshPoseTrack;
+class FAnimationTickRecordsTrack;
+class FMenuBuilder;
+class ULineBatchComponent;
+
+class FAnimationSharedData
+{
+public:
+	FAnimationSharedData(FGameplaySharedData& InGameplaySharedData);
+
+	void OnBeginSession(Insights::ITimingViewSession& InTimingViewSession);
+	void OnEndSession(Insights::ITimingViewSession& InTimingViewSession);
+	void Tick(Insights::ITimingViewSession& InTimingViewSession, const Trace::IAnalysisSession& InAnalysisSession);
+	void ExtendFilterMenu(FMenuBuilder& InMenuBuilder);
+
+#if WITH_ENGINE
+	void DrawPoses(ULineBatchComponent* InLineBatcher);
+#endif
+
+	// Check whether animation tracks are enabled
+	bool AreAnimationTracksEnabled() const;
+
+	// Get the last cached analysis session
+	const Trace::IAnalysisSession& GetAnalysisSession() const { return *AnalysisSession; }
+
+	// Check whether the analysis session is valid
+	bool IsAnalysisSessionValid() const { return AnalysisSession != nullptr; }
+
+private:
+	// UI handlers
+	void ToggleAnimationTracks();
+	void OnSelectedEventChanged(const FTimingEvent& InEvent);
+	void OnHoveredEventChanged(const FTimingEvent& InEvent);
+	void OnTimeMarkerChanged(Insights::ETimeChangedFlags InFlags, double InTimeMarker);
+	void OnSelectionChanged(Insights::ETimeChangedFlags InFlags, double InStartTime, double InEndTime);
+
+private:
+	// The gameplay shared data we are linked to
+	FGameplaySharedData& GameplaySharedData;
+
+	// Cached analysis session, set in Tick()
+	const Trace::IAnalysisSession* AnalysisSession;
+
+	// All the tracks we manage
+	TArray<FSkeletalMeshPoseTrack*> SkeletalMeshPoseTracks;
+	TArray<FAnimationTickRecordsTrack*> AnimationTickRecordsTracks;
+
+	// Delegate handles for hooks into the timing view
+	FDelegateHandle SelectedEventChangedHandle;
+	FDelegateHandle HoveredEventChangedHandle;
+	FDelegateHandle TimeMarkerChangedHandle;
+	FDelegateHandle SelectionChangedHandle;
+
+	/** Selected/hovered tracks */
+	const FTimingEventsTrack* SelectedEventTrack;
+	const FTimingEventsTrack* HoveredEventTrack;
+
+	/** Various times and ranges */
+	double SelectedEventStartTime;
+	double SelectedEventEndTime;
+	double HoveredEventStartTime;
+	double HoveredEventEndTime;
+	double SelectionStartTime;
+	double SelectionEndTime;
+	double MarkerTime;
+
+	/** Validity flags for pose times/ranges */
+	bool bSelectedEventValid;
+	bool bHoveredEventValid;
+	bool bSelectionValid;
+	bool bTimeMarkerValid;
+
+	// Whether all of our animation tracks are enabled
+	bool bAnimationTracksEnabled;
+};
\ No newline at end of file
diff --git a/Engine/Plugins/Animation/GameplayInsights/Source/GameplayInsights/Private/AnimationTickRecordsTrack.cpp b/Engine/Plugins/Animation/GameplayInsights/Source/GameplayInsights/Private/AnimationTickRecordsTrack.cpp
new file mode 100644
--- /dev/null
+++ b/Engine/Plugins/Animation/GameplayInsights/Source/GameplayInsights/Private/AnimationTickRecordsTrack.cpp
@@ -0,0 +1,151 @@
+// Copyright 1998-2019 Epic Games, Inc. All Rights Reserved.
+
+#include "AnimationTickRecordsTrack.h"
+#include "Insights/ITimingViewDrawHelper.h"
+#include "GameplayProvider.h"
+#include "AnimationProvider.h"
+#include "Insights/ViewModels/TimingTrackViewport.h"
+#include "Insights/ViewModels/TimingEvent.h"
+#include "Insights/ViewModels/TooltipDrawState.h"
+
+#define LOCTEXT_NAMESPACE "AnimationTickRecordsTrack"
+
+const FName FAnimationTickRecordsTrack::TypeName(TEXT("Animation"));
+const FName FAnimationTickRecordsTrack::SubTypeName(TEXT("TickRecords"));
+
+FAnimationTickRecordsTrack::FAnimationTickRecordsTrack(const FAnimationSharedData& InSharedData, uint64 InObjectID, uint64 InAssetId, const TCHAR* InName)
+	: TGameplayTrackMixin<FTimingEventsTrack>(InObjectID, FAnimationTickRecordsTrack::TypeName, FAnimationTickRecordsTrack::SubTypeName, FText::FromString(FString(InName)))
+	, SharedData(InSharedData)
+	, AssetId(InAssetId)
+{
+}
+
+void FAnimationTickRecordsTrack::Draw(ITimingViewDrawHelper& Helper) const
+{
+	FAnimationTickRecordsTrack& Track = *const_cast<FAnimationTickRecordsTrack*>(this);
+
+	if (Helper.BeginTimeline(Track))
+	{
+		const FGameplayProvider* GameplayProvider = SharedData.GetAnalysisSession().ReadProvider<FGameplayProvider>(FGameplayProvider::ProviderName);
+		const FAnimationProvider* AnimationProvider = SharedData.GetAnalysisSession().ReadProvider<FAnimationProvider>(FAnimationProvider::ProviderName);
+
+		if(GameplayProvider && AnimationProvider)
+		{
+			Trace::FAnalysisSessionReadScope SessionReadScope(SharedData.GetAnalysisSession());
+
+			AnimationProvider->ReadTickRecordTimeline(GetGameplayTrack().GetObjectId(), GetAssetId(), [this, &GameplayProvider, &AnimationProvider, &Helper](const FAnimationProvider::TickRecordTimeline& InTimeline)
+			{
+				auto DrawEvents = [this, &Helper, &GameplayProvider](double InStartTime, double InEndTime, uint32 InDepth, const FTickRecordMessage& InMessage)
+				{
+					Helper.AddEvent(InStartTime, InEndTime, 0, *GetName());
+				};
+
+				if (FTimingEventsTrack::bUseDownSampling)
+				{
+					const double SecondsPerPixel = 1.0 / Helper.GetViewport().GetScaleX();
+					InTimeline.EnumerateEventsDownSampled(Helper.GetViewport().GetStartTime(), Helper.GetViewport().GetEndTime(), SecondsPerPixel, DrawEvents);
+				}
+				else
+				{
+					InTimeline.EnumerateEvents(Helper.GetViewport().GetStartTime(), Helper.GetViewport().GetEndTime(), DrawEvents);
+				}
+			});
+		}
+
+		Helper.EndTimeline(Track);
+	}
+}
+
+void FAnimationTickRecordsTrack::InitTooltip(FTooltipDrawState& Tooltip, const FTimingEvent& HoveredTimingEvent) const
+{
+	Tooltip.ResetContent();
+
+	Tooltip.AddTitle(GetName());
+
+	const FAnimationProvider* AnimationProvider = SharedData.GetAnalysisSession().ReadProvider<FAnimationProvider>(FAnimationProvider::ProviderName);
+	if(AnimationProvider)
+	{
+		Trace::FAnalysisSessionReadScope SessionReadScope(SharedData.GetAnalysisSession());
+
+		AnimationProvider->ReadTickRecordMessage(GetGameplayTrack().GetObjectId(), GetAssetId(), HoveredTimingEvent.TypeId, [this, &Tooltip, &HoveredTimingEvent](const FTickRecordMessage& InMessage)
+		{
+			Tooltip.AddNameValueTextLine(LOCTEXT("EventTime", "Time").ToString(), FText::AsNumber(HoveredTimingEvent.StartTime).ToString());
+			Tooltip.AddNameValueTextLine(LOCTEXT("BlendWeight", "Blend Weight").ToString(), FText::AsNumber(InMessage.BlendWeight).ToString());
+			Tooltip.AddNameValueTextLine(LOCTEXT("PlaybackTime", "Playback Time").ToString(), FText::AsNumber(InMessage.PlaybackTime).ToString());
+			Tooltip.AddNameValueTextLine(LOCTEXT("RootMotionWeight", "Root Motion Weight").ToString(), FText::AsNumber(InMessage.RootMotionWeight).ToString());
+			Tooltip.AddNameValueTextLine(LOCTEXT("PlayRate", "Play Rate").ToString(), FText::AsNumber(InMessage.PlayRate).ToString());
+			Tooltip.AddNameValueTextLine(LOCTEXT("Looping", "Looping").ToString(), InMessage.Looping ? LOCTEXT("True", "True").ToString() : LOCTEXT("False", "False").ToString());
+		});
+	}
+
+	Tooltip.UpdateLayout();
+}
+
+bool FAnimationTickRecordsTrack::SearchTimingEvent(const double InStartTime, const double InEndTime, TFunctionRef<bool(double, double, uint32)> InPredicate, FTimingEvent& InOutTimingEvent, bool bInStopAtFirstMatch, bool bInSearchForLargestEvent) const
+{
+	struct FSearchTimingEventContext
+	{
+		const double StartTime;
+		const double EndTime;
+		TFunctionRef<bool(double, double, uint32)> Predicate;
+		FTimingEvent& TimingEvent;
+		const bool bStopAtFirstMatch;
+		const bool bSearchForLargestEvent;
+		mutable bool bFound;
+		mutable bool bContinueSearching;
+		mutable double LargestDuration;
+
+		FSearchTimingEventContext(const double InStartTime, const double InEndTime, TFunctionRef<bool(double, double, uint32)> InPredicate, FTimingEvent& InOutTimingEvent, bool bInStopAtFirstMatch, bool bInSearchForLargestEvent)
+			: StartTime(InStartTime)
+			, EndTime(InEndTime)
+			, Predicate(InPredicate)
+			, TimingEvent(InOutTimingEvent)
+			, bStopAtFirstMatch(bInStopAtFirstMatch)
+			, bSearchForLargestEvent(bInSearchForLargestEvent)
+			, bFound(false)
+			, bContinueSearching(true)
+			, LargestDuration(-1.0)
+		{
+		}
+
+		void CheckMessage(double EventStartTime, double EventEndTime, uint32 EventDepth, uint64 InMessageId)
+		{
+			if (bContinueSearching && Predicate(EventStartTime, EventEndTime, EventDepth))
+			{
+				if (!bSearchForLargestEvent || EventEndTime - EventStartTime > LargestDuration)
+				{
+					LargestDuration = EventEndTime - EventStartTime;
+
+					TimingEvent.TypeId = InMessageId;
+					TimingEvent.Depth = EventDepth;
+					TimingEvent.StartTime = EventStartTime;
+					TimingEvent.EndTime = EventEndTime;
+
+					bFound = true;
+					bContinueSearching = !bStopAtFirstMatch || bSearchForLargestEvent;
+				}
+			}
+		}
+	};
+
+	const FAnimationProvider* AnimationProvider = SharedData.GetAnalysisSession().ReadProvider<FAnimationProvider>(FAnimationProvider::ProviderName);
+
+	if(AnimationProvider)
+	{
+		FSearchTimingEventContext Context(InStartTime, InEndTime, InPredicate, InOutTimingEvent, bInStopAtFirstMatch, bInSearchForLargestEvent);
+
+		Trace::FAnalysisSessionReadScope SessionReadScope(SharedData.GetAnalysisSession());
+
+		AnimationProvider->ReadTickRecordTimeline(GetGameplayTrack().GetObjectId(), GetAssetId(), [this, &Context, &InStartTime, &InEndTime](const FAnimationProvider::TickRecordTimeline& InTimeline)
+		{
+			InTimeline.EnumerateEvents(InStartTime, InEndTime, [&Context](double InEventStartTime, double InEventEndTime, uint32 InDepth, const FTickRecordMessage& InMessage)
+			{
+				Context.CheckMessage(InEventStartTime, InEventEndTime, 0, InMessage.MessageId);
+			});
+		});
+	}
+
+	return false;
+}
+
+#undef LOCTEXT_NAMESPACE
\ No newline at end of file
diff --git a/Engine/Plugins/Animation/GameplayInsights/Source/GameplayInsights/Private/AnimationTickRecordsTrack.h b/Engine/Plugins/Animation/GameplayInsights/Source/GameplayInsights/Private/AnimationTickRecordsTrack.h
new file mode 100644
--- /dev/null
+++ b/Engine/Plugins/Animation/GameplayInsights/Source/GameplayInsights/Private/AnimationTickRecordsTrack.h
@@ -0,0 +1,29 @@
+// Copyright 1998-2019 Epic Games, Inc. All Rights Reserved.
+
+#pragma once
+
+#include "GameplayTrack.h"
+#include "Insights/ViewModels/TimingEventsTrack.h"
+
+class FAnimationTickRecordsTrack : public TGameplayTrackMixin<FTimingEventsTrack>
+{
+public:
+	static const FName TypeName;
+	static const FName SubTypeName;
+
+	FAnimationTickRecordsTrack(const FAnimationSharedData& InSharedData, uint64 InObjectId, uint64 InAssetId, const TCHAR* InName);
+
+	virtual void Draw(ITimingViewDrawHelper& Helper) const override;
+	virtual void InitTooltip(FTooltipDrawState& Tooltip, const FTimingEvent& HoveredTimingEvent) const override;
+	virtual bool SearchTimingEvent(const double InStartTime, const double InEndTime, TFunctionRef<bool(double, double, uint32)> InPredicate, FTimingEvent& InOutTimingEvent, bool bInStopAtFirstMatch, bool bInSearchForLargestEvent) const override;
+
+	/** Get the asset ID that this track uses */
+	uint64 GetAssetId() const { return AssetId; }
+
+private:
+	/** The shared data */
+	const FAnimationSharedData& SharedData;
+
+	/** The asset ID that this track uses */
+	uint64 AssetId;
+};
\ No newline at end of file
diff --git a/Engine/Plugins/Animation/GameplayInsights/Source/GameplayInsights/Private/GameplayAnalyzer.cpp b/Engine/Plugins/Animation/GameplayInsights/Source/GameplayInsights/Private/GameplayAnalyzer.cpp
new file mode 100644
--- /dev/null
+++ b/Engine/Plugins/Animation/GameplayInsights/Source/GameplayInsights/Private/GameplayAnalyzer.cpp
@@ -0,0 +1,63 @@
+// Copyright 1998-2019 Epic Games, Inc. All Rights Reserved.
+
+#include "GameplayAnalyzer.h"
+#include "TraceServices/Model/AnalysisSession.h"
+#include "GameplayProvider.h"
+
+FGameplayAnalyzer::FGameplayAnalyzer(Trace::IAnalysisSession& InSession, FGameplayProvider& InGameplayProvider)
+	: Session(InSession)
+	, GameplayProvider(InGameplayProvider)
+{
+}
+
+void FGameplayAnalyzer::OnAnalysisBegin(const FOnAnalysisContext& Context)
+{
+	auto& Builder = Context.InterfaceBuilder;
+
+	Builder.RouteEvent(RouteId_Class, "Object", "Class");
+	Builder.RouteEvent(RouteId_Object, "Object", "Object");
+	Builder.RouteEvent(RouteId_ObjectEvent, "Object", "ObjectEvent");
+}
+
+bool FGameplayAnalyzer::OnEvent(uint16 RouteId, const FOnEventContext& Context)
+{
+	Trace::FAnalysisSessionEditScope _(Session);
+
+	const auto& EventData = Context.EventData;
+	switch (RouteId)
+	{
+	case RouteId_Class:
+	{
+		const TCHAR* ClassNameAndPathName = reinterpret_cast<const TCHAR*>(EventData.GetAttachment());
+		const TCHAR* ClassName = ClassNameAndPathName;
+		int32 ClassNameStringLength = EventData.GetValue<int32>("ClassNameStringLength");
+		const TCHAR* ClassPathName = ClassNameAndPathName + ClassNameStringLength;
+		uint64 Id = EventData.GetValue<uint64>("Id");
+		uint64 SuperId = EventData.GetValue<uint64>("SuperId");
+		GameplayProvider.AppendClass(Id, SuperId, ClassName, ClassPathName);
+		break;
+	}
+	case RouteId_Object:
+	{
+		const TCHAR* ObjectNameAndPathName = reinterpret_cast<const TCHAR*>(EventData.GetAttachment());
+		const TCHAR* ObjectName = ObjectNameAndPathName;
+		int32 NameStringLength = EventData.GetValue<int32>("ObjectNameStringLength");
+		const TCHAR* ObjectPathName = ObjectNameAndPathName + NameStringLength;
+		uint64 Id = EventData.GetValue<uint64>("Id");
+		uint64 OuterId = EventData.GetValue<uint64>("OuterId");
+		uint64 ClassId = EventData.GetValue<uint64>("ClassId");
+		GameplayProvider.AppendObject(Id, OuterId, ClassId, ObjectName, ObjectPathName);
+		break;
+	}
+	case RouteId_ObjectEvent:
+	{
+		uint64 Cycle = EventData.GetValue<uint64>("Cycle");
+		uint64 Id = EventData.GetValue<uint64>("Id");
+		const TCHAR* Event = reinterpret_cast<const TCHAR*>(EventData.GetAttachment());
+		GameplayProvider.AppendObjectEvent(Id, Context.SessionContext.TimestampFromCycle(Cycle), Event);
+		break;
+	}
+	}
+
+	return true;
+}
\ No newline at end of file
diff --git a/Engine/Plugins/Animation/GameplayInsights/Source/GameplayInsights/Private/GameplayAnalyzer.h b/Engine/Plugins/Animation/GameplayInsights/Source/GameplayInsights/Private/GameplayAnalyzer.h
new file mode 100644
--- /dev/null
+++ b/Engine/Plugins/Animation/GameplayInsights/Source/GameplayInsights/Private/GameplayAnalyzer.h
@@ -0,0 +1,29 @@
+// Copyright 1998-2019 Epic Games, Inc. All Rights Reserved.
+
+#pragma once
+
+#include "Trace/Analyzer.h"
+
+class FGameplayProvider;
+namespace Trace { class IAnalysisSession; }
+
+class FGameplayAnalyzer : public Trace::IAnalyzer
+{
+public:
+	FGameplayAnalyzer(Trace::IAnalysisSession& InSession, FGameplayProvider& InGameplayProvider);
+
+	virtual void OnAnalysisBegin(const FOnAnalysisContext& Context) override;
+	virtual void OnAnalysisEnd() override {}
+	virtual bool OnEvent(uint16 RouteId, const FOnEventContext& Context) override;
+
+private:
+	enum : uint16
+	{
+		RouteId_Class,
+		RouteId_Object,
+		RouteId_ObjectEvent,
+	};
+
+	Trace::IAnalysisSession& Session;
+	FGameplayProvider& GameplayProvider;
+};
\ No newline at end of file
diff --git a/Engine/Plugins/Animation/GameplayInsights/Source/GameplayInsights/Private/GameplayInsightsModule.cpp b/Engine/Plugins/Animation/GameplayInsights/Source/GameplayInsights/Private/GameplayInsightsModule.cpp
new file mode 100644
--- /dev/null
+++ b/Engine/Plugins/Animation/GameplayInsights/Source/GameplayInsights/Private/GameplayInsightsModule.cpp
@@ -0,0 +1,84 @@
+// Copyright 1998-2019 Epic Games, Inc. All Rights Reserved.
+
+#include "Modules/ModuleInterface.h"
+#include "Features/IModularFeatures.h"
+#include "GameplayTraceModule.h"
+#include "GameplayTimingViewExtender.h"
+#include "Insights/ITimingViewExtender.h"
+#include "Containers/Ticker.h"
+#include "Misc/CoreDelegates.h"
+
+#if WITH_ENGINE
+#include "Engine/Engine.h"
+#endif
+
+class FGameplayInsightsModule : public IModuleInterface
+{
+public:
+	// IModuleInterface interface
+	virtual void StartupModule() override
+	{
+		IModularFeatures::Get().RegisterModularFeature(Trace::ModuleFeatureName, &GameplayTraceModule);
+		IModularFeatures::Get().RegisterModularFeature(Insights::TimingViewExtenderFeatureName, &GameplayTimingViewExtender);
+
+
+		PostEngineInitHandle = FCoreDelegates::OnPostEngineInit.AddLambda([this]()
+		{
+#if WITH_ENGINE
+			if(GEngine)
+			{
+				for(const FWorldContext& WorldContext : GEngine->GetWorldContexts())
+				{
+					GameplayTimingViewExtender.AddVisualizerWorld(WorldContext.World());
+				}
+
+				OnWorldAddedHandle = GEngine->OnWorldAdded().AddRaw(this, &FGameplayInsightsModule::HandleWorldAdded);
+			}
+#endif
+		});
+
+		TickerHandle = FTicker::GetCoreTicker().AddTicker(TEXT("GameplayInsights"), 0.0f, [this](float DeltaTime)
+		{
+			GameplayTimingViewExtender.TickVisualizers(DeltaTime);
+			return true;
+		});
+	}
+
+	virtual void ShutdownModule() override
+	{
+		FTicker::GetCoreTicker().RemoveTicker(TickerHandle);
+
+		FCoreDelegates::OnPostEngineInit.Remove(PostEngineInitHandle);
+
+#if WITH_ENGINE
+		if(GEngine)
+		{
+			GEngine->OnWorldAdded().Remove(OnWorldAddedHandle);
+		}
+#endif
+
+		IModularFeatures::Get().UnregisterModularFeature(Trace::ModuleFeatureName, &GameplayTraceModule);
+		IModularFeatures::Get().UnregisterModularFeature(Insights::TimingViewExtenderFeatureName, &GameplayTimingViewExtender);
+	}
+
+#if WITH_ENGINE
+	void HandleWorldAdded(UWorld* InWorld)
+	{
+		if (InWorld->WorldType == EWorldType::Editor ||
+			InWorld->WorldType == EWorldType::PIE ||
+			InWorld->WorldType == EWorldType::Game)
+		{
+			GameplayTimingViewExtender.AddVisualizerWorld(InWorld);
+		}
+	}
+#endif
+
+	FGameplayTraceModule GameplayTraceModule;
+	FGameplayTimingViewExtender GameplayTimingViewExtender;
+
+	FDelegateHandle OnWorldAddedHandle;
+	FDelegateHandle TickerHandle;
+	FDelegateHandle PostEngineInitHandle;
+};
+
+IMPLEMENT_MODULE(FGameplayInsightsModule, GameplayInsights);
diff --git a/Engine/Plugins/Animation/GameplayInsights/Source/GameplayInsights/Private/GameplayProvider.cpp b/Engine/Plugins/Animation/GameplayInsights/Source/GameplayInsights/Private/GameplayProvider.cpp
new file mode 100644
--- /dev/null
+++ b/Engine/Plugins/Animation/GameplayInsights/Source/GameplayInsights/Private/GameplayProvider.cpp
@@ -0,0 +1,138 @@
+// Copyright 1998-2019 Epic Games, Inc. All Rights Reserved.
+
+#include "GameplayProvider.h"
+
+FName FGameplayProvider::ProviderName("GameplayProvider");
+
+FGameplayProvider::FGameplayProvider(Trace::IAnalysisSession& InSession)
+	: Session(InSession)
+{
+}
+
+bool FGameplayProvider::ReadObjectEventsTimeline(uint64 InObjectId, TFunctionRef<void(const ObjectEventsTimeline&)> Callback) const
+{
+	Session.ReadAccessCheck();
+
+	const uint32* IndexPtr = ObjectIdToEventTimelines.Find(InObjectId);
+	if(IndexPtr != nullptr)
+	{
+		if (*IndexPtr < uint32(EventTimelines.Num()))
+		{
+			Callback(*EventTimelines[*IndexPtr]);
+			return true;
+		}
+	}
+
+	return false;
+}
+
+bool FGameplayProvider::ReadObjectEvent(uint64 InObjectId, uint64 InMessageId, TFunctionRef<void(const FObjectEventMessage&)> Callback) const
+{
+	Session.ReadAccessCheck();
+
+	return ReadObjectEventsTimeline(InObjectId, [&Callback, &InMessageId](const ObjectEventsTimeline& InTimeline)
+	{
+		if(InMessageId < InTimeline.GetEventCount())
+		{
+			Callback(InTimeline.GetEvent(InMessageId));
+		}
+	});
+}
+
+void FGameplayProvider::EnumerateObjects(TFunctionRef<void(const FObjectInfo&)> Callback) const
+{
+	for(const FObjectInfo& ObjectInfo : ObjectInfos)
+	{
+		Callback(ObjectInfo);
+	}
+}
+
+const FClassInfo* FGameplayProvider::FindClassInfo(uint64 InClassId) const
+{
+	const int32* ClassIndex = ClassIdToIndexMap.Find(InClassId);
+	if(ClassIndex != nullptr)
+	{
+		return &ClassInfos[*ClassIndex];
+	}
+
+	return nullptr;
+}
+
+const FObjectInfo* FGameplayProvider::FindObjectInfo(uint64 InObjectId) const
+{
+	const int32* ObjectIndex = ObjectIdToIndexMap.Find(InObjectId);
+	if(ObjectIndex != nullptr)
+	{
+		return &ObjectInfos[*ObjectIndex];
+	}
+
+	return nullptr;
+}
+
+void FGameplayProvider::AppendClass(uint64 InClassId, uint64 InSuperId, const TCHAR* InClassName, const TCHAR* InClassPathName)
+{
+	Session.WriteAccessCheck();
+
+	if(ClassIdToIndexMap.Find(InClassId) == nullptr)
+	{
+		const TCHAR* NewClassName = Session.StoreString(InClassName);
+		const TCHAR* NewClassPathName = Session.StoreString(InClassPathName);
+
+		FClassInfo NewClassInfo;
+		NewClassInfo.Id = InClassId;
+		NewClassInfo.SuperId = InSuperId;
+		NewClassInfo.Name = NewClassName;
+		NewClassInfo.PathName = NewClassPathName;
+
+		int32 NewClassInfoIndex = ClassInfos.Add(NewClassInfo);
+		ClassIdToIndexMap.Add(InClassId, NewClassInfoIndex);
+	}
+}
+
+void FGameplayProvider::AppendObject(uint64 InObjectId, uint64 InOuterId, uint64 InClassId, const TCHAR* InObjectName, const TCHAR* InObjectPathName)
+{
+	Session.WriteAccessCheck();
+
+	if(ObjectIdToIndexMap.Find(InObjectId) == nullptr)
+	{
+		const TCHAR* NewObjectName = Session.StoreString(InObjectName);
+		const TCHAR* NewObjectPathName = Session.StoreString(InObjectPathName);
+
+		FObjectInfo NewObjectInfo;
+		NewObjectInfo.Id = InObjectId;
+		NewObjectInfo.OuterId = InOuterId;
+		NewObjectInfo.ClassId = InClassId;
+		NewObjectInfo.Name = NewObjectName;
+		NewObjectInfo.PathName = NewObjectPathName;
+
+		int32 NewObjectInfoIndex = ObjectInfos.Add(NewObjectInfo);
+		ObjectIdToIndexMap.Add(InObjectId, NewObjectInfoIndex);
+	}
+}
+
+void FGameplayProvider::AppendObjectEvent(uint64 InObjectId, double InTime, const TCHAR* InEventName)
+{
+	Session.WriteAccessCheck();
+
+	TSharedPtr<Trace::TPointTimeline<FObjectEventMessage>> Timeline;
+	uint32* IndexPtr = ObjectIdToEventTimelines.Find(InObjectId);
+	if(IndexPtr != nullptr)
+	{
+		Timeline = EventTimelines[*IndexPtr];
+	}
+	else
+	{
+		Timeline = MakeShared<Trace::TPointTimeline<FObjectEventMessage>>(Session.GetLinearAllocator());
+		ObjectIdToEventTimelines.Add(InObjectId, EventTimelines.Num());
+		EventTimelines.Add(Timeline.ToSharedRef());
+	}
+
+	FObjectEventMessage Message;
+	Message.MessageId = Timeline->GetEventCount();
+	Message.Id = InObjectId;
+	Message.Name = Session.StoreString(InEventName);
+
+	Timeline->AppendEvent(InTime, Message);
+
+	Session.UpdateDurationSeconds(InTime);
+}
diff --git a/Engine/Plugins/Animation/GameplayInsights/Source/GameplayInsights/Private/GameplayProvider.h b/Engine/Plugins/Animation/GameplayInsights/Source/GameplayInsights/Private/GameplayProvider.h
new file mode 100644
--- /dev/null
+++ b/Engine/Plugins/Animation/GameplayInsights/Source/GameplayInsights/Private/GameplayProvider.h
@@ -0,0 +1,53 @@
+// Copyright 1998-2019 Epic Games, Inc. All Rights Reserved.
+
+#pragma once
+
+#include "IGameplayProvider.h"
+#include "Model/PointTimeline.h"
+
+namespace Trace { class IAnalysisSession; }
+
+class FGameplayProvider : public IGameplayProvider
+{
+public:
+	static FName ProviderName;
+
+	FGameplayProvider(Trace::IAnalysisSession& InSession);
+
+	/** IGameplayProvider interface */
+	virtual bool ReadObjectEventsTimeline(uint64 InObjectId, TFunctionRef<void(const ObjectEventsTimeline&)> Callback) const override;
+	virtual bool ReadObjectEvent(uint64 InObjectId, uint64 InMessageId, TFunctionRef<void(const FObjectEventMessage&)> Callback) const override;
+	virtual void EnumerateObjects(TFunctionRef<void(const FObjectInfo&)> Callback) const override;
+	virtual const FClassInfo* FindClassInfo(uint64 InClassId) const override;
+	virtual const FObjectInfo* FindObjectInfo(uint64 InObjectId) const override;
+
+	/** Add a class message */
+	void AppendClass(uint64 InClassId, uint64 InSuperId, const TCHAR* InClassName, const TCHAR* InClassPathName);
+
+	/** Add an object message */
+	void AppendObject(uint64 InObjectId, uint64 InOuterId, uint64 InClassId, const TCHAR* InObjectName, const TCHAR* InObjectPathName);
+
+	/** Add an object event message */
+	void AppendObjectEvent(uint64 InObjectId, double InTime, const TCHAR* InEvent);
+
+private:
+	Trace::IAnalysisSession& Session;
+
+	/** All class info, grow only for stable indices */
+	TArray<FClassInfo> ClassInfos;
+
+	/** All object info, grow only for stable indices */
+	TArray<FObjectInfo> ObjectInfos;
+
+	/** Classes that are in use. Map from Id to ClassInfo index */
+	TMap<uint64, int32> ClassIdToIndexMap;
+
+	/** Objects that are in use. Map from Id to ObjectInfo index */
+	TMap<uint64, int32> ObjectIdToIndexMap;
+
+	/** Map from object Id to timeline index */
+	TMap<uint64, uint32> ObjectIdToEventTimelines;
+
+	/** Message storage */
+	TArray<TSharedRef<Trace::TPointTimeline<FObjectEventMessage>>> EventTimelines;
+};
\ No newline at end of file
diff --git a/Engine/Plugins/Animation/GameplayInsights/Source/GameplayInsights/Private/GameplaySharedData.cpp b/Engine/Plugins/Animation/GameplayInsights/Source/GameplayInsights/Private/GameplaySharedData.cpp
new file mode 100644
--- /dev/null
+++ b/Engine/Plugins/Animation/GameplayInsights/Source/GameplayInsights/Private/GameplaySharedData.cpp
@@ -0,0 +1,177 @@
+// Copyright 1998-2019 Epic Games, Inc. All Rights Reserved.
+
+#include "GameplaySharedData.h"
+#include "TraceServices/Model/AnalysisSession.h"
+#include "ObjectEventsTrack.h"
+#include "Algo/Sort.h"
+
+#define LOCTEXT_NAMESPACE "GameplaySharedData"
+
+FGameplaySharedData::FGameplaySharedData()
+	: AnalysisSession(nullptr)
+	, bObjectTracksEnabled(false)
+{
+}
+
+void FGameplaySharedData::OnBeginSession(Insights::ITimingViewSession& InTimingViewSession)
+{
+	ObjectTracks.Reset();
+}
+
+void FGameplaySharedData::OnEndSession(Insights::ITimingViewSession& InTimingViewSession)
+{
+	ObjectTracks.Reset();
+}
+
+FObjectEventsTrack* FGameplaySharedData::GetObjectEventsTrackForId(Insights::ITimingViewSession& InTimingViewSession, const Trace::IAnalysisSession& InAnalysisSession, const FObjectInfo& InObjectInfo)
+{
+	const FGameplayProvider* GameplayProvider = InAnalysisSession.ReadProvider<FGameplayProvider>(FGameplayProvider::ProviderName);
+	check(GameplayProvider);
+
+	FObjectEventsTrack* LeafObjectEventsTrack = ObjectTracks.FindRef(InObjectInfo.Id);
+	if(LeafObjectEventsTrack == nullptr)
+	{
+		LeafObjectEventsTrack = new FObjectEventsTrack(*this, InObjectInfo.Id, InObjectInfo.Name);
+		LeafObjectEventsTrack->SetVisibilityFlag(bObjectTracksEnabled);
+		ObjectTracks.Add(InObjectInfo.Id, LeafObjectEventsTrack);
+
+		InTimingViewSession.AddTimingEventsTrack(LeafObjectEventsTrack);
+	}
+
+	// Fill the outer chain
+	if(InObjectInfo.OuterId != 0)
+	{
+		FObjectEventsTrack* ObjectEventsTrack = LeafObjectEventsTrack;
+		uint64 OuterId = InObjectInfo.OuterId;
+		while(const FObjectInfo* OuterInfo = GameplayProvider->FindObjectInfo(OuterId))
+		{
+			FObjectEventsTrack* OuterEventsTrack = ObjectTracks.FindRef(OuterId);
+			if(OuterEventsTrack == nullptr)
+			{
+				OuterEventsTrack = new FObjectEventsTrack(*this, OuterId, OuterInfo->Name);
+				OuterEventsTrack->SetVisibilityFlag(bObjectTracksEnabled);
+				ObjectTracks.Add(OuterId, OuterEventsTrack);
+
+				InTimingViewSession.AddTimingEventsTrack(OuterEventsTrack);
+			}
+
+			// setup hierarchy
+			if(ObjectEventsTrack->GetGameplayTrack().GetParentTrack() == nullptr)
+			{
+				OuterEventsTrack->GetGameplayTrack().AddChildTrack(ObjectEventsTrack->GetGameplayTrack());
+			}
+
+			ObjectEventsTrack = OuterEventsTrack;
+			OuterId = OuterInfo->OuterId;
+		}
+	}
+
+	return LeafObjectEventsTrack;
+}
+
+static void UpdateTrackOrderRecursive(FObjectEventsTrack* InTrack, int32& InOutOrder)
+{
+	// recurse down object-track children, then non-object track leaf tracks to set 
+	// overall ordering based on depth-first traversal of the hierarchy
+
+	InTrack->SetOrder(InOutOrder++);
+
+	for(FGameplayTrack* ChildTrack : InTrack->GetGameplayTrack().GetChildTracks())
+	{
+		if(ChildTrack->GetTimingTrack().GetType() == FObjectEventsTrack::TypeName)
+		{
+			UpdateTrackOrderRecursive(static_cast<FObjectEventsTrack*>(&ChildTrack->GetTimingTrack()), InOutOrder);
+		}
+	}
+
+	for(FGameplayTrack* ChildTrack : InTrack->GetGameplayTrack().GetChildTracks())
+	{
+		if(ChildTrack->GetTimingTrack().GetType() != FObjectEventsTrack::TypeName)
+		{
+			ChildTrack->GetTimingTrack().SetOrder(InOutOrder++);
+		}
+	}
+}
+
+void FGameplaySharedData::Tick(Insights::ITimingViewSession& InTimingViewSession, const Trace::IAnalysisSession& InAnalysisSession)
+{
+	AnalysisSession = &InAnalysisSession;
+
+	const FGameplayProvider* GameplayProvider = InAnalysisSession.ReadProvider<FGameplayProvider>(FGameplayProvider::ProviderName);
+
+	if(GameplayProvider)
+	{
+		// Add a track for each tracked object
+		GameplayProvider->EnumerateObjects([this, &InTimingViewSession, &InAnalysisSession, &GameplayProvider](const FObjectInfo& InObjectInfo)
+		{
+			FObjectEventsTrack* ObjectEventsTrack = GetObjectEventsTrackForId(InTimingViewSession, InAnalysisSession, InObjectInfo);
+			ObjectEventsTrack->SetVisibilityFlag(false);
+
+			GameplayProvider->ReadObjectEventsTimeline(InObjectInfo.Id, [this, &ObjectEventsTrack](const IGameplayProvider::ObjectEventsTimeline& InTimeline)
+			{
+				ObjectEventsTrack->SetVisibilityFlag(bObjectTracksEnabled);
+			});
+		});
+	}
+}
+
+void FGameplaySharedData::ExtendFilterMenu(FMenuBuilder& InMenuBuilder)
+{
+	InMenuBuilder.AddMenuEntry(
+		LOCTEXT("ToggleGameplayTracks", "Gameplay Tracks"),
+		LOCTEXT("ToggleGameplayTracks_Tooltip", "Show/hide the gameplay tracks"),
+		FSlateIcon(),
+		FUIAction(
+			FExecuteAction::CreateRaw(this, &FGameplaySharedData::ToggleGameplayTracks),
+			FCanExecuteAction(),
+			FIsActionChecked::CreateRaw(this, &FGameplaySharedData::AreGameplayTracksEnabled)),
+		NAME_None,
+		EUserInterfaceActionType::ToggleButton
+	);
+}
+
+void FGameplaySharedData::OnTracksChanged(int32& InOutOrder)
+{
+	// Find current roots
+	static TArray<FObjectEventsTrack*> Roots;
+	check(Roots.Num() == 0);
+
+	for(auto ObjectTrackPair : ObjectTracks)
+	{
+		if(ObjectTrackPair.Value->GetGameplayTrack().GetParentTrack() == nullptr)
+		{
+			Roots.Add(ObjectTrackPair.Value);
+		}
+	}
+
+	// Sort roots alphabetically
+	Algo::Sort(Roots, [](FObjectEventsTrack* InTrack0, FObjectEventsTrack* InTrack1)
+	{
+		return InTrack0->GetName() < InTrack1->GetName();
+	});
+
+	// update ordering
+	for(FObjectEventsTrack* RootTrack : Roots)
+	{
+		UpdateTrackOrderRecursive(RootTrack, InOutOrder);
+	}
+
+	Roots.Reset();
+}
+
+void FGameplaySharedData::ToggleGameplayTracks()
+{
+	bObjectTracksEnabled = !bObjectTracksEnabled;
+
+	for(TPair<uint64, FObjectEventsTrack*> TrackPair : ObjectTracks)
+	{
+		TrackPair.Value->SetVisibilityFlag(bObjectTracksEnabled);
+	}
+}
+
+bool FGameplaySharedData::AreGameplayTracksEnabled() const
+{
+	return bObjectTracksEnabled;
+}
+
+#undef LOCTEXT_NAMESPACE
\ No newline at end of file
diff --git a/Engine/Plugins/Animation/GameplayInsights/Source/GameplayInsights/Private/GameplaySharedData.h b/Engine/Plugins/Animation/GameplayInsights/Source/GameplayInsights/Private/GameplaySharedData.h
new file mode 100644
--- /dev/null
+++ b/Engine/Plugins/Animation/GameplayInsights/Source/GameplayInsights/Private/GameplaySharedData.h
@@ -0,0 +1,46 @@
+// Copyright 1998-2019 Epic Games, Inc. All Rights Reserved.
+
+#pragma once
+
+#include "Insights/ViewModels/TimingEventsTrack.h"
+
+namespace Trace { class IAnalysisSession; }
+namespace Insights { class ITimingViewSession; }
+class FObjectEventsTrack;
+class FSkeletalMeshPoseTrack;
+class FAnimationTickRecordsTrack;
+
+class FGameplaySharedData
+{
+public:
+	FGameplaySharedData();
+
+	void OnBeginSession(Insights::ITimingViewSession& InTimingViewSession);
+	void OnEndSession(Insights::ITimingViewSession& InTimingViewSession);
+	void Tick(Insights::ITimingViewSession& InTimingViewSession, const Trace::IAnalysisSession& InAnalysisSession);
+	void ExtendFilterMenu(FMenuBuilder& InMenuBuilder);
+	void OnTracksChanged(int32& InOutOrder);
+
+	// Helper function. Builds object track hierarchy on-demand and returns a track for the supplied object info.
+	FObjectEventsTrack* GetObjectEventsTrackForId(Insights::ITimingViewSession& InTimingViewSession, const Trace::IAnalysisSession& InAnalysisSession, const FObjectInfo& InObjectInfo);
+
+	// Check whether gameplay tacks are enabled
+	bool AreGameplayTracksEnabled() const;
+
+	// Get the last cached analysis session
+	const Trace::IAnalysisSession& GetAnalysisSession() const { return *AnalysisSession; }
+
+private:
+	// UI handlers
+	void ToggleGameplayTracks();
+
+private:
+	// Track for each tracked object, mapped from Object ID -> track
+	TMap<uint64, FObjectEventsTrack*> ObjectTracks;
+
+	// Cached analysis session, set in Tick()
+	const Trace::IAnalysisSession* AnalysisSession;
+
+	// Whether all of our object tracks are enabled
+	bool bObjectTracksEnabled;
+};
\ No newline at end of file
diff --git a/Engine/Plugins/Animation/GameplayInsights/Source/GameplayInsights/Private/GameplayTimingViewExtender.cpp b/Engine/Plugins/Animation/GameplayInsights/Source/GameplayInsights/Private/GameplayTimingViewExtender.cpp
new file mode 100644
--- /dev/null
+++ b/Engine/Plugins/Animation/GameplayInsights/Source/GameplayInsights/Private/GameplayTimingViewExtender.cpp
@@ -0,0 +1,93 @@
+// Copyright 1998-2019 Epic Games, Inc. All Rights Reserved.
+
+#include "GameplayTimingViewExtender.h"
+#include "Insights/ITimingViewSession.h"
+#include "GameplaySharedData.h"
+#include "AnimationSharedData.h"
+
+void FGameplayTimingViewExtender::OnBeginSession(Insights::ITimingViewSession& InSession)
+{
+	FPerSessionData* PerSessionData = PerSessionDataMap.Find(&InSession);
+	if(PerSessionData == nullptr)
+	{
+		PerSessionData = &PerSessionDataMap.Add(&InSession);
+		PerSessionData->GameplaySharedData = new FGameplaySharedData();
+		PerSessionData->AnimationSharedData = new FAnimationSharedData(*PerSessionData->GameplaySharedData);
+	}
+
+	PerSessionData->GameplaySharedData->OnBeginSession(InSession);
+	PerSessionData->AnimationSharedData->OnBeginSession(InSession);
+}
+
+void FGameplayTimingViewExtender::OnEndSession(Insights::ITimingViewSession& InSession)
+{
+	FPerSessionData* PerSessionData = PerSessionDataMap.Find(&InSession);
+	if(PerSessionData != nullptr)
+	{
+		PerSessionData->GameplaySharedData->OnEndSession(InSession);
+		PerSessionData->AnimationSharedData->OnEndSession(InSession);
+
+		delete PerSessionData->GameplaySharedData;
+		PerSessionData->GameplaySharedData = nullptr;
+		delete PerSessionData->AnimationSharedData;
+		PerSessionData->AnimationSharedData = nullptr;
+	}
+
+	PerSessionDataMap.Remove(&InSession);
+}
+
+void FGameplayTimingViewExtender::Tick(Insights::ITimingViewSession& InSession, const Trace::IAnalysisSession& InAnalysisSession)
+{
+	FPerSessionData* PerSessionData = PerSessionDataMap.Find(&InSession);
+	if(PerSessionData != nullptr)
+	{
+		PerSessionData->GameplaySharedData->Tick(InSession, InAnalysisSession);
+		PerSessionData->AnimationSharedData->Tick(InSession, InAnalysisSession);
+	}
+}
+
+void FGameplayTimingViewExtender::ExtendFilterMenu(Insights::ITimingViewSession& InSession, FMenuBuilder& InMenuBuilder)
+{
+	FPerSessionData* PerSessionData = PerSessionDataMap.Find(&InSession);
+	if(PerSessionData != nullptr)
+	{
+		PerSessionData->GameplaySharedData->ExtendFilterMenu(InMenuBuilder);
+		PerSessionData->AnimationSharedData->ExtendFilterMenu(InMenuBuilder);
+	}
+}
+
+void FGameplayTimingViewExtender::OnTracksChanged(Insights::ITimingViewSession& InSession, int32& InOutOrder)
+{
+	FPerSessionData* PerSessionData = PerSessionDataMap.Find(&InSession);
+	if(PerSessionData != nullptr)
+	{
+		PerSessionData->GameplaySharedData->OnTracksChanged(InOutOrder);
+	}
+}
+
+#if WITH_ENGINE
+void FGameplayTimingViewExtender::AddVisualizerWorld(UWorld* InWorld)
+{
+	Worlds.Add(InWorld);
+}
+#endif
+
+void FGameplayTimingViewExtender::TickVisualizers(float DeltaTime)
+{
+#if WITH_ENGINE
+	// Trim invalid worlds
+	Worlds.RemoveAll([](const TWeakObjectPtr<UWorld>& InWorld){ return InWorld.Get() == nullptr; });
+
+	for(auto& PerSessionData : PerSessionDataMap)
+	{
+		// Draw using line batchers
+		for(TWeakObjectPtr<UWorld>& World : Worlds)
+		{
+			if(World->LineBatcher)
+			{
+				PerSessionData.Value.AnimationSharedData->DrawPoses(World->LineBatcher);
+			}
+		}
+	}
+#endif
+}
\ No newline at end of file
diff --git a/Engine/Plugins/Animation/GameplayInsights/Source/GameplayInsights/Private/GameplayTimingViewExtender.h b/Engine/Plugins/Animation/GameplayInsights/Source/GameplayInsights/Private/GameplayTimingViewExtender.h
new file mode 100644
--- /dev/null
+++ b/Engine/Plugins/Animation/GameplayInsights/Source/GameplayInsights/Private/GameplayTimingViewExtender.h
@@ -0,0 +1,45 @@
+// Copyright 1998-2019 Epic Games, Inc. All Rights Reserved.
+
+#pragma once
+
+#include "Insights/ITimingViewExtender.h"
+
+namespace Insights { class ITimingViewSession; }
+class UWorld;
+class FGameplaySharedData;
+class FAnimationSharedData;
+
+class FGameplayTimingViewExtender : public Insights::ITimingViewExtender
+{
+public:
+	// Insights::ITimingViewExtender interface
+	virtual void OnBeginSession(Insights::ITimingViewSession& InSession) override;
+	virtual void OnEndSession(Insights::ITimingViewSession& InSession) override;
+	virtual void Tick(Insights::ITimingViewSession& InSession, const Trace::IAnalysisSession& InAnalysisSession) override;
+	virtual void ExtendFilterMenu(Insights::ITimingViewSession& InSession, FMenuBuilder& InMenuBuilder) override;
+	virtual void OnTracksChanged(Insights::ITimingViewSession& InSession, int32& InOutOrder) override;
+
+#if WITH_ENGINE
+	// Add a world to be visualized
+	void AddVisualizerWorld(UWorld* InWorld);
+#endif
+
+	// Tick the visualizers
+	void TickVisualizers(float DeltaTime);
+
+private:
+	struct FPerSessionData
+	{
+		// Shared data
+		FGameplaySharedData* GameplaySharedData;
+		FAnimationSharedData* AnimationSharedData;
+	};
+
+	// The data we host per-session
+	TMap<Insights::ITimingViewSession*, FPerSessionData> PerSessionDataMap;
+
+#if WITH_ENGINE
+	// The worlds we will visualize
+	TArray<TWeakObjectPtr<UWorld>> Worlds;
+#endif
+};
diff --git a/Engine/Plugins/Animation/GameplayInsights/Source/GameplayInsights/Private/GameplayTraceModule.cpp b/Engine/Plugins/Animation/GameplayInsights/Source/GameplayInsights/Private/GameplayTraceModule.cpp
new file mode 100644
--- /dev/null
+++ b/Engine/Plugins/Animation/GameplayInsights/Source/GameplayInsights/Private/GameplayTraceModule.cpp
@@ -0,0 +1,42 @@
+// Copyright 1998-2019 Epic Games, Inc. All Rights Reserved.
+
+#include "GameplayTraceModule.h"
+#include "AnimationAnalyzer.h"
+#include "AnimationProvider.h"
+#include "GameplayProvider.h"
+
+FName FGameplayTraceModule::ModuleName("GameplayTrace");
+
+void FGameplayTraceModule::GetModuleInfo(Trace::FModuleInfo& OutModuleInfo)
+{
+	OutModuleInfo.Name = ModuleName;
+	OutModuleInfo.DisplayName = TEXT("Gameplay");
+}
+
+void FGameplayTraceModule::OnAnalysisBegin(Trace::IAnalysisSession& InSession)
+{
+	FGameplayProvider* GameplayProvider = new FGameplayProvider(InSession);
+	InSession.AddProvider(FGameplayProvider::ProviderName, GameplayProvider);
+	FAnimationProvider* AnimationProvider = new FAnimationProvider(InSession);
+	InSession.AddProvider(FAnimationProvider::ProviderName, AnimationProvider);
+
+	InSession.AddAnalyzer(new FAnimationAnalyzer(InSession, *AnimationProvider));
+	InSession.AddAnalyzer(new FGameplayAnalyzer(InSession, *GameplayProvider));
+}
+
+void FGameplayTraceModule::GetLoggers(TArray<const TCHAR *>& OutLoggers)
+{
+	OutLoggers.Add(TEXT("Object"));
+	OutLoggers.Add(TEXT("Animation"));
+}
+
+void FGameplayTraceModule::GenerateReports(const Trace::IAnalysisSession& Session, const TCHAR* CmdLine, const TCHAR* OutputDirectory)
+{
+
+}
+
+const IAnimationProvider* ReadAnimationProvider(const Trace::IAnalysisSession& Session)
+{
+	return Session.ReadProvider<IAnimationProvider>(FAnimationProvider::ProviderName);
+}
+
diff --git a/Engine/Plugins/Animation/GameplayInsights/Source/GameplayInsights/Private/GameplayTraceModule.h b/Engine/Plugins/Animation/GameplayInsights/Source/GameplayInsights/Private/GameplayTraceModule.h
new file mode 100644
--- /dev/null
+++ b/Engine/Plugins/Animation/GameplayInsights/Source/GameplayInsights/Private/GameplayTraceModule.h
@@ -0,0 +1,25 @@
+// Copyright 1998-2019 Epic Games, Inc. All Rights Reserved.
+
+#pragma once
+
+#include "CoreMinimal.h"
+#include "TraceServices/ModuleService.h"
+
+class IAnimationProvider;
+namespace Insights { class ITimingViewSession; }
+
+class FGameplayTraceModule : public Trace::IModule
+{
+public:
+	// Trace::IModule interface
+	virtual void GetModuleInfo(Trace::FModuleInfo& OutModuleInfo) override;
+	virtual void OnAnalysisBegin(Trace::IAnalysisSession& Session) override;
+	virtual void GetLoggers(TArray<const TCHAR *>& OutLoggers) override;
+	virtual void GenerateReports(const Trace::IAnalysisSession& Session, const TCHAR* CmdLine, const TCHAR* OutputDirectory) override;
+	virtual const TCHAR* GetCommandLineArgument() override { return TEXT("objecttrace"); }
+
+private:
+	static FName ModuleName;
+};
+
+const IAnimationProvider* ReadAnimationProvider(const Trace::IAnalysisSession& Session);
\ No newline at end of file
diff --git a/Engine/Plugins/Animation/GameplayInsights/Source/GameplayInsights/Private/GameplayTrack.h b/Engine/Plugins/Animation/GameplayInsights/Source/GameplayInsights/Private/GameplayTrack.h
new file mode 100644
--- /dev/null
+++ b/Engine/Plugins/Animation/GameplayInsights/Source/GameplayInsights/Private/GameplayTrack.h
@@ -0,0 +1,98 @@
+// Copyright 1998-2019 Epic Games, Inc. All Rights Reserved.
+
+#pragma once
+
+#include "Insights/ViewModels/TimingEventsTrack.h"
+#include "Insights/ViewModels/GraphTrack.h"
+#include "Algo/Sort.h"
+
+// Provides parent/child hierarchy structure and the owning object Id
+// Designed as a compositional member of outer timing tracks
+class FGameplayTrack
+{
+public:
+	FGameplayTrack(FBaseTimingTrack* InTimingTrack, uint64 InObjectId)
+		: TimingTrack(InTimingTrack)
+		, ObjectId(InObjectId)
+		, Parent(nullptr)
+	{
+	}
+
+	// Get the object ID for this track
+	uint64 GetObjectId() const { return ObjectId; }
+
+	/** Get the parent track */
+	FGameplayTrack* GetParentTrack() { return Parent; }
+
+	/** Get child tracks */
+	TArrayView<FGameplayTrack*> GetChildTracks() { return Children; }
+
+	/** Add a child track */
+	void AddChildTrack(FGameplayTrack& InChildTrack)
+	{
+		check(InChildTrack.Parent == nullptr);
+		InChildTrack.Parent = this;
+
+		Children.Add(&InChildTrack);
+
+		Algo::Sort(Children, [](FGameplayTrack* InTrack0, FGameplayTrack* InTrack1)
+		{
+			return InTrack0->GetTimingTrack().GetName() < InTrack1->GetTimingTrack().GetName();
+		});	
+	}
+
+	/** Find a child track using the specified callback */
+	FBaseTimingTrack* FindChildTrack(uint64 InObjectId, TFunctionRef<bool(const FBaseTimingTrack& InTrack)> Callback)
+	{
+		for(FGameplayTrack* ChildTrack : Children)
+		{
+			if( ChildTrack != nullptr &&
+				ChildTrack->ObjectId == InObjectId && 
+				Callback(ChildTrack->GetTimingTrack()))
+			{
+				return &ChildTrack->GetTimingTrack();
+			}
+		}
+
+		return nullptr;
+	}
+
+	/** Access the outer timing track */
+	FBaseTimingTrack& GetTimingTrack() { return *TimingTrack; }
+
+	/** Access the outer timing track */
+	const FBaseTimingTrack& GetTimingTrack() const { return *TimingTrack; }
+
+private:
+	/** Outer timing track */
+	FBaseTimingTrack* TimingTrack;
+
+	// The object ID for this track
+	uint64 ObjectId;
+
+	/** Parent track */
+	FGameplayTrack* Parent;
+
+	/** Our child tracks */
+	TArray<FGameplayTrack*> Children;
+};
+
+template <class Base>
+class TGameplayTrackMixin : public Base
+{
+public:
+	TGameplayTrackMixin(uint64 InObjectId, const FName& InType, const FName& InSubType, const FText& InName)
+		: Base(FBaseTimingTrack::GenerateId(), InType, InSubType, InName.ToString())
+		, GameplayTrack(this, InObjectId)
+	{
+	}
+
+	/** Access the underlying gameplay track */
+	FGameplayTrack& GetGameplayTrack() { return GameplayTrack; }
+
+	/** Access the underlying gameplay track */
+	const FGameplayTrack& GetGameplayTrack() const { return GameplayTrack; }
+
+private:
+	FGameplayTrack GameplayTrack;
+};
\ No newline at end of file
diff --git a/Engine/Plugins/Animation/GameplayInsights/Source/GameplayInsights/Private/IAnimationProvider.h b/Engine/Plugins/Animation/GameplayInsights/Source/GameplayInsights/Private/IAnimationProvider.h
new file mode 100644
--- /dev/null
+++ b/Engine/Plugins/Animation/GameplayInsights/Source/GameplayInsights/Private/IAnimationProvider.h
@@ -0,0 +1,55 @@
+// Copyright 1998-2019 Epic Games, Inc. All Rights Reserved.
+
+#pragma once
+
+#include "TraceServices/Model/AnalysisSession.h"
+#include "TraceServices/Containers/Timelines.h"
+
+struct FAnimTickRecordTrace;
+
+struct FSkeletalMeshInfo
+{
+	uint64 ParentIndicesStartIndex = 0;
+	uint64 Id = 0;
+	uint32 BoneCount = 0;
+};
+
+struct FSkeletalMeshPoseMessage
+{
+	uint64 MessageId = 0;
+	uint64 TransformStartIndex = 0;
+	uint64 ComponentId = 0;	
+	uint64 MeshId = 0;	
+	uint16 NumTransforms = 0;
+	uint16 FrameCounter = 0;
+	uint16 LodIndex = 0;
+};
+
+struct FTickRecordMessage
+{
+	uint64 MessageId = 0;
+	uint64 ComponentId = 0;
+	uint64 AnimInstanceId = 0;
+	uint64 AssetId = 0;
+	float BlendWeight = 0.0f;
+	float PlaybackTime = 0.0f;
+	float RootMotionWeight = 0.0f;
+	float PlayRate = 0.0f;
+	uint16 FrameCounter = 0;
+	bool Looping = false;
+};
+
+class IAnimationProvider : public Trace::IProvider
+{
+public:
+	typedef Trace::ITimeline<FTickRecordMessage> TickRecordTimeline;
+	typedef Trace::ITimeline<FSkeletalMeshPoseMessage> SkeletalMeshPoseTimeline;
+
+	virtual bool ReadSkeletalMeshPoseTimeline(uint64 InObjectId, TFunctionRef<void(const SkeletalMeshPoseTimeline&)> Callback) const = 0;
+	virtual bool ReadSkeletalMeshPoseMessage(uint64 InObjectId, uint64 InMessageId, TFunctionRef<void(const FSkeletalMeshPoseMessage&)> Callback) const = 0;
+	virtual void EnumerateSkeletalMeshPose(const FSkeletalMeshPoseMessage& InMessage, const FSkeletalMeshInfo& InMeshInfo, TFunctionRef<void(const FTransform&, const FTransform&)> Callback) const = 0;
+	virtual void EnumerateTickRecordTimelines(uint64 InObjectId, TFunctionRef<void(uint64, const TickRecordTimeline&)> Callback) const = 0;
+	virtual bool ReadTickRecordTimeline(uint64 InObjectId, uint64 InAssetId, TFunctionRef<void(const TickRecordTimeline&)> Callback) const = 0;
+	virtual bool ReadTickRecordMessage(uint64 InObjectId, uint64 InAssetId, uint64 InMessageId, TFunctionRef<void(const FTickRecordMessage&)> Callback) const = 0;
+	virtual const FSkeletalMeshInfo* FindSkeletalMeshInfo(uint64 InObjectId) const = 0;
+};
\ No newline at end of file
diff --git a/Engine/Plugins/Animation/GameplayInsights/Source/GameplayInsights/Private/IGameplayProvider.h b/Engine/Plugins/Animation/GameplayInsights/Source/GameplayInsights/Private/IGameplayProvider.h
new file mode 100644
--- /dev/null
+++ b/Engine/Plugins/Animation/GameplayInsights/Source/GameplayInsights/Private/IGameplayProvider.h
@@ -0,0 +1,41 @@
+// Copyright 1998-2019 Epic Games, Inc. All Rights Reserved.
+
+#pragma once
+
+#include "TraceServices/Model/AnalysisSession.h"
+
+struct FClassInfo
+{	
+	uint64 Id = 0;
+	uint64 SuperId = 0;
+	const TCHAR* Name = nullptr;
+	const TCHAR* PathName = nullptr;
+};
+
+struct FObjectInfo
+{	
+	uint64 Id = 0;
+	uint64 OuterId = 0;
+	uint64 ClassId = 0;
+	const TCHAR* Name = nullptr;
+	const TCHAR* PathName = nullptr;
+};
+
+struct FObjectEventMessage
+{
+	uint64 MessageId = 0;
+	uint64 Id = 0;
+	const TCHAR* Name = nullptr;
+};
+
+class IGameplayProvider : public Trace::IProvider
+{
+public:
+	typedef Trace::ITimeline<FObjectEventMessage> ObjectEventsTimeline;
+
+	virtual bool ReadObjectEventsTimeline(uint64 InObjectId, TFunctionRef<void(const ObjectEventsTimeline&)> Callback) const = 0;
+	virtual bool ReadObjectEvent(uint64 InObjectId, uint64 InMessageId, TFunctionRef<void(const FObjectEventMessage&)> Callback) const = 0;
+	virtual void EnumerateObjects(TFunctionRef<void(const FObjectInfo&)> Callback) const = 0;
+	virtual const FClassInfo* FindClassInfo(uint64 InClassId) const = 0;
+	virtual const FObjectInfo* FindObjectInfo(uint64 InObjectId) const = 0;
+};
\ No newline at end of file
diff --git a/Engine/Plugins/Animation/GameplayInsights/Source/GameplayInsights/Private/ObjectEventsTrack.cpp b/Engine/Plugins/Animation/GameplayInsights/Source/GameplayInsights/Private/ObjectEventsTrack.cpp
new file mode 100644
--- /dev/null
+++ b/Engine/Plugins/Animation/GameplayInsights/Source/GameplayInsights/Private/ObjectEventsTrack.cpp
@@ -0,0 +1,138 @@
+// Copyright 1998-2019 Epic Games, Inc. All Rights Reserved.
+
+#include "ObjectEventsTrack.h"
+#include "Insights/ITimingViewDrawHelper.h"
+#include "GameplayProvider.h"
+#include "Insights/ViewModels/TimingTrackViewport.h"
+#include "Insights/ViewModels/TimingEvent.h"
+#include "Algo/Sort.h"
+#include "Insights/ViewModels/TooltipDrawState.h"
+
+#define LOCTEXT_NAMESPACE "ObjectEventsTrack"
+
+const FName FObjectEventsTrack::TypeName(TEXT("Gameplay"));
+const FName FObjectEventsTrack::SubTypeName(TEXT("ObjectEvents"));
+
+FObjectEventsTrack::FObjectEventsTrack(const FGameplaySharedData& InSharedData, uint64 InObjectID, const TCHAR* InName)
+	: TGameplayTrackMixin<FTimingEventsTrack>(InObjectID, TypeName, SubTypeName, FText::Format(LOCTEXT("ObjectEventsTrackName", "{0}"), FText::FromString(FString(InName))))
+	, SharedData(InSharedData)
+{
+}
+
+void FObjectEventsTrack::Draw(ITimingViewDrawHelper& Helper) const
+{
+	FObjectEventsTrack& Track = *const_cast<FObjectEventsTrack*>(this);
+
+	if (Helper.BeginTimeline(Track))
+	{
+		const FGameplayProvider* GameplayProvider = SharedData.GetAnalysisSession().ReadProvider<FGameplayProvider>(FGameplayProvider::ProviderName);
+
+		if(GameplayProvider)
+		{
+			Trace::FAnalysisSessionReadScope SessionReadScope(SharedData.GetAnalysisSession());
+
+			// object events
+			GameplayProvider->ReadObjectEventsTimeline(GetGameplayTrack().GetObjectId(), [&Helper](const FGameplayProvider::ObjectEventsTimeline& InTimeline)
+			{
+				InTimeline.EnumerateEvents(Helper.GetViewport().GetStartTime(), Helper.GetViewport().GetEndTime(), [&Helper](double InStartTime, double InEndTime, uint32 InDepth, const FObjectEventMessage& InMessage)
+				{
+					Helper.AddEvent(InStartTime, InEndTime, 0, InMessage.Name);
+				});
+			});
+		}
+
+		Helper.EndTimeline(Track);
+	}
+}
+
+void FObjectEventsTrack::InitTooltip(FTooltipDrawState& Tooltip, const FTimingEvent& HoveredTimingEvent) const
+{
+	Tooltip.ResetContent();
+
+	const FGameplayProvider* GameplayProvider = SharedData.GetAnalysisSession().ReadProvider<FGameplayProvider>(FGameplayProvider::ProviderName);
+	if(GameplayProvider)
+	{
+		Trace::FAnalysisSessionReadScope SessionReadScope(SharedData.GetAnalysisSession());
+
+		GameplayProvider->ReadObjectEventsTimeline(GetGameplayTrack().GetObjectId(), [&Tooltip, &HoveredTimingEvent](const FGameplayProvider::ObjectEventsTimeline& InTimeline)
+		{
+			if(HoveredTimingEvent.TypeId < InTimeline.GetEventCount())
+			{
+				const FObjectEventMessage& Message = InTimeline.GetEvent(HoveredTimingEvent.TypeId);
+				Tooltip.AddTitle(FText::FromString(FString(Message.Name)).ToString());
+				Tooltip.AddNameValueTextLine(LOCTEXT("EventTime", "Time").ToString(), FText::AsNumber(HoveredTimingEvent.StartTime).ToString());
+			}
+		});
+	}
+
+	Tooltip.UpdateLayout();
+}
+
+bool FObjectEventsTrack::SearchTimingEvent(const double InStartTime, const double InEndTime, TFunctionRef<bool(double, double, uint32)> InPredicate, FTimingEvent& InOutTimingEvent, bool bInStopAtFirstMatch, bool bInSearchForLargestEvent) const
+{
+	struct FSearchTimingEventContext
+	{
+		const double StartTime;
+		const double EndTime;
+		TFunctionRef<bool(double, double, uint32)> Predicate;
+		FTimingEvent& TimingEvent;
+		const bool bStopAtFirstMatch;
+		const bool bSearchForLargestEvent;
+		mutable bool bFound;
+		mutable bool bContinueSearching;
+		mutable double LargestDuration;
+
+		FSearchTimingEventContext(const double InStartTime, const double InEndTime, TFunctionRef<bool(double, double, uint32)> InPredicate, FTimingEvent& InOutTimingEvent, bool bInStopAtFirstMatch, bool bInSearchForLargestEvent)
+			: StartTime(InStartTime)
+			, EndTime(InEndTime)
+			, Predicate(InPredicate)
+			, TimingEvent(InOutTimingEvent)
+			, bStopAtFirstMatch(bInStopAtFirstMatch)
+			, bSearchForLargestEvent(bInSearchForLargestEvent)
+			, bFound(false)
+			, bContinueSearching(true)
+			, LargestDuration(-1.0)
+		{
+		}
+
+		void CheckMessage(double EventStartTime, double EventEndTime, uint32 EventDepth, uint64 InMessageId)
+		{
+			if (bContinueSearching && Predicate(EventStartTime, EventEndTime, EventDepth))
+			{
+				if (!bSearchForLargestEvent || EventEndTime - EventStartTime > LargestDuration)
+				{
+					LargestDuration = EventEndTime - EventStartTime;
+
+					TimingEvent.TypeId = InMessageId;
+					TimingEvent.Depth = EventDepth;
+					TimingEvent.StartTime = EventStartTime;
+					TimingEvent.EndTime = EventEndTime;
+
+					bFound = true;
+					bContinueSearching = !bStopAtFirstMatch || bSearchForLargestEvent;
+				}
+			}
+		}
+	};
+
+	const FGameplayProvider* GameplayProvider = SharedData.GetAnalysisSession().ReadProvider<FGameplayProvider>(FGameplayProvider::ProviderName);
+
+	if(GameplayProvider)
+	{
+		FSearchTimingEventContext Context(InStartTime, InEndTime, InPredicate, InOutTimingEvent, bInStopAtFirstMatch, bInSearchForLargestEvent);
+
+		Trace::FAnalysisSessionReadScope SessionReadScope(SharedData.GetAnalysisSession());
+
+		GameplayProvider->ReadObjectEventsTimeline(GetGameplayTrack().GetObjectId(), [&Context, &InStartTime, &InEndTime](const FGameplayProvider::ObjectEventsTimeline& InTimeline)
+		{
+			InTimeline.EnumerateEvents(InStartTime, InEndTime, [&Context](double InEventStartTime, double InEventEndTime, uint32 InDepth, const FObjectEventMessage& InMessage)
+			{
+				Context.CheckMessage(InEventStartTime, InEventEndTime, 0, InMessage.MessageId);
+			});
+		});
+	}
+
+	return false;
+}
+
+#undef LOCTEXT_NAMESPACE
\ No newline at end of file
diff --git a/Engine/Plugins/Animation/GameplayInsights/Source/GameplayInsights/Private/ObjectEventsTrack.h b/Engine/Plugins/Animation/GameplayInsights/Source/GameplayInsights/Private/ObjectEventsTrack.h
new file mode 100644
--- /dev/null
+++ b/Engine/Plugins/Animation/GameplayInsights/Source/GameplayInsights/Private/ObjectEventsTrack.h
@@ -0,0 +1,23 @@
+// Copyright 1998-2019 Epic Games, Inc. All Rights Reserved.
+
+#pragma once
+
+#include "GameplayTrack.h"
+#include "Insights/ViewModels/TimingEventsTrack.h"
+#include "Containers/ArrayView.h"
+
+class FObjectEventsTrack : public TGameplayTrackMixin<FTimingEventsTrack>
+{
+public:
+	static const FName TypeName;
+	static const FName SubTypeName;
+
+	FObjectEventsTrack(const FGameplaySharedData& InSharedData, uint64 InObjectID, const TCHAR* InName);
+
+	virtual void Draw(ITimingViewDrawHelper& Helper) const override;
+	virtual void InitTooltip(FTooltipDrawState& Tooltip, const FTimingEvent& HoveredTimingEvent) const override;
+	virtual bool SearchTimingEvent(const double InStartTime, const double InEndTime, TFunctionRef<bool(double, double, uint32)> InPredicate, FTimingEvent& InOutTimingEvent, bool bInStopAtFirstMatch, bool bInSearchForLargestEvent) const override;
+
+private:
+	const FGameplaySharedData& SharedData;
+};
\ No newline at end of file
diff --git a/Engine/Plugins/Animation/GameplayInsights/Source/GameplayInsights/Private/SkeletalMeshPoseTrack.cpp b/Engine/Plugins/Animation/GameplayInsights/Source/GameplayInsights/Private/SkeletalMeshPoseTrack.cpp
new file mode 100644
--- /dev/null
+++ b/Engine/Plugins/Animation/GameplayInsights/Source/GameplayInsights/Private/SkeletalMeshPoseTrack.cpp
@@ -0,0 +1,257 @@
+// Copyright 1998-2019 Epic Games, Inc. All Rights Reserved.
+
+#include "SkeletalMeshPoseTrack.h"
+#include "Insights/ITimingViewDrawHelper.h"
+#include "GameplayProvider.h"
+#include "AnimationProvider.h"
+#include "Insights/ViewModels/TimingTrackViewport.h"
+#include "Insights/ViewModels/TimingEvent.h"
+
+#if WITH_ENGINE
+#include "Components/LineBatchComponent.h"
+#endif
+
+#define LOCTEXT_NAMESPACE "SkeletalMeshPoseTrack"
+
+const FName FSkeletalMeshPoseTrack::TypeName(TEXT("Animation"));
+const FName FSkeletalMeshPoseTrack::SubTypeName(TEXT("SkeletalMeshPose"));
+
+FSkeletalMeshPoseTrack::FSkeletalMeshPoseTrack(const FAnimationSharedData& InSharedData, uint64 InObjectID, const TCHAR* InName)
+	: TGameplayTrackMixin<FTimingEventsTrack>(InObjectID, FSkeletalMeshPoseTrack::TypeName, FSkeletalMeshPoseTrack::SubTypeName, FText::Format(LOCTEXT("SkelMeshPoseName", "{0} Pose"), FText::FromString(FString(InName))))
+	, SharedData(InSharedData)
+	, bDrawMarkerTime(false)
+	, bDrawSelectedEvent(false)
+	, bDrawHoveredEvent(false)
+	, bDrawSelection(false)
+{
+}
+
+void FSkeletalMeshPoseTrack::Reset()
+{
+	FTimingEventsTrack::Reset();
+
+	SetHeight(16.0f);
+}
+
+void FSkeletalMeshPoseTrack::Draw( ITimingViewDrawHelper& Helper) const
+{
+	FSkeletalMeshPoseTrack& Track = *const_cast<FSkeletalMeshPoseTrack*>(this);
+
+	if (Helper.BeginTimeline(Track))
+	{
+		const FAnimationProvider* AnimationProvider = SharedData.GetAnalysisSession().ReadProvider<FAnimationProvider>(FAnimationProvider::ProviderName);
+
+		if(AnimationProvider)
+		{
+			Trace::FAnalysisSessionReadScope SessionReadScope(SharedData.GetAnalysisSession());
+
+			AnimationProvider->ReadSkeletalMeshPoseTimeline(GetGameplayTrack().GetObjectId(), [&Helper](const FAnimationProvider::SkeletalMeshPoseTimeline& InTimeline)
+			{
+				auto DrawEvent = [&Helper](double InStartTime, double InEndTime, uint32 InDepth, const FSkeletalMeshPoseMessage& InMessage)
+				{
+					static TCHAR Buffer[256];
+					FCString::Snprintf(Buffer, 256, TEXT("%d Bones"), InMessage.NumTransforms);
+					Helper.AddEvent(InStartTime, InEndTime, 0, Buffer);
+				};
+
+				if (FTimingEventsTrack::bUseDownSampling)
+				{
+					const double SecondsPerPixel = 1.0 / Helper.GetViewport().GetScaleX();
+					InTimeline.EnumerateEventsDownSampled(Helper.GetViewport().GetStartTime(), Helper.GetViewport().GetEndTime(), SecondsPerPixel, DrawEvent);
+				}
+				else
+				{
+					InTimeline.EnumerateEvents(Helper.GetViewport().GetStartTime(), Helper.GetViewport().GetEndTime(), DrawEvent);
+				}
+			});
+		}
+
+		Helper.EndTimeline(Track);
+	}
+}
+
+void FSkeletalMeshPoseTrack::InitTooltip(FTooltipDrawState& Tooltip, const FTimingEvent& HoveredTimingEvent) const
+{
+	Tooltip.ResetContent();
+
+	Tooltip.AddTitle(LOCTEXT("SkeletalMeshPoseTooltipTitle", "Skeletal Mesh Pose").ToString());
+
+	const FAnimationProvider* AnimationProvider = SharedData.GetAnalysisSession().ReadProvider<FAnimationProvider>(FAnimationProvider::ProviderName);
+	if(AnimationProvider)
+	{
+		Trace::FAnalysisSessionReadScope SessionReadScope(SharedData.GetAnalysisSession());
+
+		AnimationProvider->ReadSkeletalMeshPoseMessage(GetGameplayTrack().GetObjectId(), HoveredTimingEvent.TypeId, [&Tooltip, &HoveredTimingEvent](const FSkeletalMeshPoseMessage& InMessage)
+		{
+			Tooltip.AddNameValueTextLine(LOCTEXT("EventTime", "Time").ToString(), FText::AsNumber(HoveredTimingEvent.StartTime).ToString());
+			Tooltip.AddNameValueTextLine(LOCTEXT("BoneCount", "Bone Count").ToString(), FText::AsNumber(InMessage.NumTransforms).ToString());
+		});
+	}
+
+	Tooltip.UpdateLayout();
+}
+
+bool FSkeletalMeshPoseTrack::SearchTimingEvent(const double InStartTime, const double InEndTime, TFunctionRef<bool(double, double, uint32)> InPredicate, FTimingEvent& InOutTimingEvent, bool bInStopAtFirstMatch, bool bInSearchForLargestEvent) const
+{
+	struct FSearchTimingEventContext
+	{
+		const double StartTime;
+		const double EndTime;
+		TFunctionRef<bool(double, double, uint32)> Predicate;
+		FTimingEvent& TimingEvent;
+		const bool bStopAtFirstMatch;
+		const bool bSearchForLargestEvent;
+		mutable bool bFound;
+		mutable bool bContinueSearching;
+		mutable double LargestDuration;
+
+		FSearchTimingEventContext(const double InStartTime, const double InEndTime, TFunctionRef<bool(double, double, uint32)> InPredicate, FTimingEvent& InOutTimingEvent, bool bInStopAtFirstMatch, bool bInSearchForLargestEvent)
+			: StartTime(InStartTime)
+			, EndTime(InEndTime)
+			, Predicate(InPredicate)
+			, TimingEvent(InOutTimingEvent)
+			, bStopAtFirstMatch(bInStopAtFirstMatch)
+			, bSearchForLargestEvent(bInSearchForLargestEvent)
+			, bFound(false)
+			, bContinueSearching(true)
+			, LargestDuration(-1.0)
+		{
+		}
+
+		void CheckMessage(double EventStartTime, double EventEndTime, uint32 EventDepth, uint64 InMessageId)
+		{
+			if (bContinueSearching && Predicate(EventStartTime, EventEndTime, EventDepth))
+			{
+				if (!bSearchForLargestEvent || EventEndTime - EventStartTime > LargestDuration)
+				{
+					LargestDuration = EventEndTime - EventStartTime;
+
+					TimingEvent.TypeId = InMessageId;
+					TimingEvent.Depth = EventDepth;
+					TimingEvent.StartTime = EventStartTime;
+					TimingEvent.EndTime = EventEndTime;
+
+					bFound = true;
+					bContinueSearching = !bStopAtFirstMatch || bSearchForLargestEvent;
+				}
+			}
+		}
+	};
+
+	const FAnimationProvider* AnimationProvider = SharedData.GetAnalysisSession().ReadProvider<FAnimationProvider>(FAnimationProvider::ProviderName);
+
+	if(AnimationProvider)
+	{
+		FSearchTimingEventContext Context(InStartTime, InEndTime, InPredicate, InOutTimingEvent, bInStopAtFirstMatch, bInSearchForLargestEvent);
+
+		Trace::FAnalysisSessionReadScope SessionReadScope(SharedData.GetAnalysisSession());
+
+		AnimationProvider->ReadSkeletalMeshPoseTimeline(GetGameplayTrack().GetObjectId(), [&Context, &InStartTime, &InEndTime](const FAnimationProvider::SkeletalMeshPoseTimeline& InTimeline)
+		{
+			InTimeline.EnumerateEvents(InStartTime, InEndTime, [&Context](double InEventStartTime, double InEventEndTime, uint32 InDepth, const FSkeletalMeshPoseMessage& InMessage)
+			{
+				Context.CheckMessage(InEventStartTime, InEventEndTime, 0, InMessage.MessageId);
+			});
+		});
+	}
+
+	return false;
+}
+
+void FSkeletalMeshPoseTrack::BuildContextMenu(FMenuBuilder& MenuBuilder)
+{
+	MenuBuilder.BeginSection(TEXT("ShowPoseSection"), LOCTEXT("ShowPose", "Show Pose"));
+	{
+		MenuBuilder.AddMenuEntry(
+			LOCTEXT("ToggleDrawMarkerTime", "Marker Time"),
+			LOCTEXT("ToggleDrawMarkerTime_Tooltip", "Draw the pose at the current marker time"),
+			FSlateIcon(),
+			FUIAction(
+				FExecuteAction::CreateLambda([this](){ bDrawMarkerTime = !bDrawMarkerTime; }),
+				FCanExecuteAction(),
+				FIsActionChecked::CreateLambda([this](){ return bDrawMarkerTime; })),
+			NAME_None,
+			EUserInterfaceActionType::ToggleButton
+		);
+
+		MenuBuilder.AddMenuEntry(
+			LOCTEXT("ToggleDrawSelection", "Selection"),
+			LOCTEXT("ToggleDrawSelection_Tooltip", "Draw the pose for the currently selected event"),
+			FSlateIcon(),
+			FUIAction(
+				FExecuteAction::CreateLambda([this](){ bDrawSelectedEvent = !bDrawSelectedEvent; }),
+				FCanExecuteAction(),
+				FIsActionChecked::CreateLambda([this](){ return bDrawSelectedEvent; })),
+			NAME_None,
+			EUserInterfaceActionType::ToggleButton
+		);
+
+		MenuBuilder.AddMenuEntry(
+			LOCTEXT("ToggleDrawHovered", "Hovered"),
+			LOCTEXT("ToggleDrawSelection_Tooltip", "Draw the pose for the currently hovered event"),
+			FSlateIcon(),
+			FUIAction(
+				FExecuteAction::CreateLambda([this](){ bDrawHoveredEvent = !bDrawHoveredEvent; }),
+				FCanExecuteAction(),
+				FIsActionChecked::CreateLambda([this](){ return bDrawHoveredEvent; })),
+			NAME_None,
+			EUserInterfaceActionType::ToggleButton
+		);
+
+		MenuBuilder.AddMenuEntry(
+			LOCTEXT("ToggleDrawSelectedRange", "Selected Range"),
+			LOCTEXT("ToggleDrawSelectedRange_Tooltip", "Draw poses for the currently selected range"),
+			FSlateIcon(),
+			FUIAction(
+				FExecuteAction::CreateLambda([this](){ bDrawSelection = !bDrawSelection; }),
+				FCanExecuteAction(),
+				FIsActionChecked::CreateLambda([this](){ return bDrawSelection; })),
+			NAME_None,
+			EUserInterfaceActionType::ToggleButton
+		);
+	}
+	MenuBuilder.EndSection();
+}
+
+#if WITH_ENGINE
+
+void FSkeletalMeshPoseTrack::DrawPoses(ULineBatchComponent* InLineBatcher, double SelectionStartTime, double SelectionEndTime)
+{
+	if(SharedData.IsAnalysisSessionValid())
+	{
+		const FAnimationProvider* AnimationProvider = SharedData.GetAnalysisSession().ReadProvider<FAnimationProvider>(FAnimationProvider::ProviderName);
+
+		if(AnimationProvider)
+		{
+			Trace::FAnalysisSessionReadScope SessionReadScope(SharedData.GetAnalysisSession());
+
+			AnimationProvider->ReadSkeletalMeshPoseTimeline(GetGameplayTrack().GetObjectId(), [this, &InLineBatcher, &AnimationProvider, &SelectionStartTime, &SelectionEndTime](const FAnimationProvider::SkeletalMeshPoseTimeline& InTimeline)
+			{
+				static TArray<FBatchedLine> Lines;
+
+				InTimeline.EnumerateEvents(SelectionStartTime, SelectionEndTime, [&InLineBatcher, &AnimationProvider, &SelectionStartTime, &SelectionEndTime](double InStartTime, double InEndTime, uint32 InDepth, const FSkeletalMeshPoseMessage& InMessage)
+				{
+					if(SelectionStartTime == SelectionEndTime || (InStartTime >= SelectionStartTime && InEndTime <= SelectionEndTime))
+					{
+						const FSkeletalMeshInfo* SkeletalMeshInfo = AnimationProvider->FindSkeletalMeshInfo(InMessage.MeshId);
+						if(SkeletalMeshInfo)
+						{
+							AnimationProvider->EnumerateSkeletalMeshPose(InMessage, *SkeletalMeshInfo, [&InLineBatcher](const FTransform& InTransform, const FTransform& InParentTransform)
+							{
+								Lines.Emplace(InParentTransform.GetLocation(), InTransform.GetLocation(), FLinearColor::Red, 1.0f, 0.2f, SDPG_Foreground);
+							});
+						}
+					}
+				});
+
+				InLineBatcher->DrawLines(Lines);
+
+				Lines.Reset();
+			});
+		}
+	}
+}
+
+#endif
+
+#undef LOCTEXT_NAMESPACE
\ No newline at end of file
diff --git a/Engine/Plugins/Animation/GameplayInsights/Source/GameplayInsights/Private/SkeletalMeshPoseTrack.h b/Engine/Plugins/Animation/GameplayInsights/Source/GameplayInsights/Private/SkeletalMeshPoseTrack.h
new file mode 100644
--- /dev/null
+++ b/Engine/Plugins/Animation/GameplayInsights/Source/GameplayInsights/Private/SkeletalMeshPoseTrack.h
@@ -0,0 +1,44 @@
+// Copyright 1998-2019 Epic Games, Inc. All Rights Reserved.
+
+#pragma once
+
+#include "GameplayTrack.h"
+#include "Insights/ViewModels/TimingEventsTrack.h"
+
+class ULineBatchComponent;
+
+class FSkeletalMeshPoseTrack : public TGameplayTrackMixin<FTimingEventsTrack>
+{
+public:
+	static const FName TypeName;
+	static const FName SubTypeName;
+
+	FSkeletalMeshPoseTrack(const FAnimationSharedData& InSharedData, uint64 InObjectID, const TCHAR* InName);
+
+	virtual void Reset() override;
+	virtual void Draw(ITimingViewDrawHelper& Helper) const override;
+	virtual void InitTooltip(FTooltipDrawState& Tooltip, const FTimingEvent& HoveredTimingEvent) const override;
+	virtual bool SearchTimingEvent(const double InStartTime, const double InEndTime, TFunctionRef<bool(double, double, uint32)> InPredicate, FTimingEvent& InOutTimingEvent, bool bInStopAtFirstMatch, bool bInSearchForLargestEvent) const override;
+	virtual void BuildContextMenu(FMenuBuilder& MenuBuilder) override;
+
+	// Access drawing flags
+	bool ShouldDrawMarkerTime() const { return bDrawMarkerTime; }
+	bool ShouldDrawSelectedEvent() const { return bDrawSelectedEvent; }
+	bool ShouldDrawHoveredEvent() const { return bDrawHoveredEvent; }
+	bool ShouldDrawSelection() const { return bDrawSelection; }
+
+#if WITH_ENGINE
+	// Draw poses in the specified time range
+	void DrawPoses(ULineBatchComponent* InLineBatcher, double SelectionStartTime, double SelectionEndTime);
+#endif
+
+private:
+	/** The shared data */
+	const FAnimationSharedData& SharedData;
+
+	/** Whether to draw poses at specified times */
+	bool bDrawMarkerTime;
+	bool bDrawSelectedEvent;
+	bool bDrawHoveredEvent;
+	bool bDrawSelection;
+};
\ No newline at end of file
diff --git a/Engine/Source/Runtime/Engine/Classes/Animation/AnimNodeBase.h b/Engine/Source/Runtime/Engine/Classes/Animation/AnimNodeBase.h
--- a/Engine/Source/Runtime/Engine/Classes/Animation/AnimNodeBase.h
+++ b/Engine/Source/Runtime/Engine/Classes/Animation/AnimNodeBase.h
@@ -12,6 +12,8 @@
 #include "BonePose.h"
 #include "Logging/TokenizedMessage.h"
 #include "Stats/StatsHierarchical.h"
+#include "Animation/AnimTrace.h"
+
 #include "AnimNodeBase.generated.h"
 
 #define DECLARE_SCOPE_HIERARCHICAL_COUNTER_ANIMNODE(Method) \
@@ -152,6 +154,21 @@ public:
 	ENGINE_API UAnimBlueprint* GetAnimBlueprint() const;
 #endif //WITH_EDITORONLY_DATA
 
+#if ANIM_TRACE_ENABLED
+	// Get the current node Id, set when we recurse into graph traversal functions from pose links
+	ENGINE_API int32 GetCurrentNodeId() const { return CurrentNodeId; }
+
+	// Get the previous node Id, set when we recurse into graph traversal functions from pose links
+	ENGINE_API int32 GetPreviousNodeId() const { return PreviousNodeId; }
+
+protected:
+	// The current node ID, set when we recurse into graph traversal functions from pose links
+	int32 CurrentNodeId;
+
+	// The previous node ID, set when we recurse into graph traversal functions from pose links
+	int32 PreviousNodeId;
+#endif
+
 protected:
 
 	/** Interface for node contexts to register log messages with the proxy */
@@ -271,6 +288,16 @@ public:
 		return Result;
 	}
 
+#if ANIM_TRACE_ENABLED
+	FAnimationUpdateContext WithNodeId(int32 InNodeId) const
+	{ 
+		FAnimationUpdateContext Result(*this);
+		Result.PreviousNodeId = CurrentNodeId;
+		Result.CurrentNodeId = InNodeId;
+		return Result; 
+	}
+#endif
+
 	// Add a node to the list of tracked ancestors
 	template<typename NodeType>
 	FScopedAnimNodeTracker TrackAncestor(NodeType* Node) const
diff --git a/Engine/Source/Runtime/Engine/Engine.Build.cs b/Engine/Source/Runtime/Engine/Engine.Build.cs
--- a/Engine/Source/Runtime/Engine/Engine.Build.cs
+++ b/Engine/Source/Runtime/Engine/Engine.Build.cs
@@ -108,7 +108,8 @@ public class Engine : ModuleRules
 				"AudioMixerCore",
 				"SignalProcessing",
 				"CrunchCompression",
-			}
+				"TraceLog",
+            }
 		);
 
 		// Cross platform Audio Codecs:
diff --git a/Engine/Source/Runtime/Engine/Private/Actor.cpp b/Engine/Source/Runtime/Engine/Private/Actor.cpp
--- a/Engine/Source/Runtime/Engine/Private/Actor.cpp
+++ b/Engine/Source/Runtime/Engine/Private/Actor.cpp
@@ -40,6 +40,7 @@
 #include "DeviceProfiles/DeviceProfileManager.h"
 #include "Interfaces/ITargetPlatform.h"
 #include "DeviceProfiles/DeviceProfile.h"
+#include "ObjectTrace.h"
 
 DEFINE_LOG_CATEGORY(LogActor);
 
@@ -2168,6 +2169,8 @@ void AActor::EndPlay(const EEndPlayReason::Type EndPlayReason)
 {
 	if (ActorHasBegunPlay == EActorBeginPlayState::HasBegunPlay)
 	{
+		TRACE_OBJECT_EVENT(this, EndPlay);
+
 		ActorHasBegunPlay = EActorBeginPlayState::HasNotBegunPlay;
 
 		// Dispatch the blueprint events
@@ -3433,6 +3436,8 @@ void AActor::DispatchBeginPlay(bool bFromLevelStreaming)
 
 void AActor::BeginPlay()
 {
+	TRACE_OBJECT_EVENT(this, BeginPlay);
+
 	ensureMsgf(ActorHasBegunPlay == EActorBeginPlayState::BeginningPlay, TEXT("BeginPlay was called on actor %s which was in state %d"), *GetPathName(), (int32)ActorHasBegunPlay);
 	SetLifeSpan( InitialLifeSpan );
 	RegisterAllActorTickFunctions(true, false); // Components are done below.
diff --git a/Engine/Source/Runtime/Engine/Private/Animation/AnimInstanceProxy.cpp b/Engine/Source/Runtime/Engine/Private/Animation/AnimInstanceProxy.cpp
--- a/Engine/Source/Runtime/Engine/Private/Animation/AnimInstanceProxy.cpp
+++ b/Engine/Source/Runtime/Engine/Private/Animation/AnimInstanceProxy.cpp
@@ -24,6 +24,7 @@
 #include "Animation/AnimMTStats.h"
 #include "Animation/AnimNode_Root.h"
 #include "Animation/AnimNode_LinkedAnimLayer.h"
+#include "Animation/AnimTrace.h"
 #undef DO_ANIMSTAT_PROCESSING
 
 #define LOCTEXT_NAMESPACE "AnimInstance"
@@ -522,6 +523,8 @@ void FAnimInstanceProxy::PostEvaluate(UAnimInstance* InAnimInstance)
 {
 	DECLARE_SCOPE_HIERARCHICAL_COUNTER_FUNC()
 
+	TRACE_ANIM_TICK_RECORDS(*this, SkeletalMeshComponent);
+
 	ClearObjects();
 
 #if ENABLE_ANIM_LOGGING
diff --git a/Engine/Source/Runtime/Engine/Private/Animation/AnimNodeBase.cpp b/Engine/Source/Runtime/Engine/Private/Animation/AnimNodeBase.cpp
--- a/Engine/Source/Runtime/Engine/Private/Animation/AnimNodeBase.cpp
+++ b/Engine/Source/Runtime/Engine/Private/Animation/AnimNodeBase.cpp
@@ -3,6 +3,7 @@
 #include "Animation/AnimNodeBase.h"
 #include "Animation/AnimClassInterface.h"
 #include "Animation/AnimInstanceProxy.h"
+#include "Animation/AnimTrace.h"
 
 void* FExposedValueCopyRecord::GetDestAddr(FAnimInstanceProxy* Proxy, const UProperty* NodeProperty) const
 {
@@ -57,17 +58,29 @@ const void* FExposedValueCopyRecord::GetSourceAddr(FAnimInstanceProxy* Proxy) co
 
 FAnimationBaseContext::FAnimationBaseContext()
 	: AnimInstanceProxy(nullptr)
+#if ANIM_TRACE_ENABLED
+	, CurrentNodeId(INDEX_NONE)
+	, PreviousNodeId(INDEX_NONE)
+#endif
 {
 }
 
 FAnimationBaseContext::FAnimationBaseContext(FAnimInstanceProxy* InAnimInstanceProxy)
 	: AnimInstanceProxy(InAnimInstanceProxy)
+#if ANIM_TRACE_ENABLED
+	, CurrentNodeId(INDEX_NONE)
+	, PreviousNodeId(INDEX_NONE)
+#endif
 {
 }
 
 FAnimationBaseContext::FAnimationBaseContext(const FAnimationBaseContext& InContext)
+	: AnimInstanceProxy(InContext.AnimInstanceProxy)
+#if ANIM_TRACE_ENABLED
+	, CurrentNodeId(INDEX_NONE)
+	, PreviousNodeId(INDEX_NONE)
+#endif
 {
-	AnimInstanceProxy = InContext.AnimInstanceProxy;
 }
 
 IAnimClassInterface* FAnimationBaseContext::GetAnimClass() const
@@ -329,7 +342,13 @@ void FPoseLinkBase::Update(const FAnimationUpdateContext& Context)
 
 	if (LinkedNode != NULL)
 	{
+#if ANIM_TRACE_ENABLED
+		FAnimationUpdateContext LinkContext(Context.WithNodeId(LinkID));
+		TRACE_POSE_LINK(LinkContext);
+		LinkedNode->Update_AnyThread(LinkContext);
+#else
 		LinkedNode->Update_AnyThread(Context);
+#endif
 	}
 }
 
diff --git a/Engine/Source/Runtime/Engine/Private/Animation/AnimTrace.cpp b/Engine/Source/Runtime/Engine/Private/Animation/AnimTrace.cpp
new file mode 100644
--- /dev/null
+++ b/Engine/Source/Runtime/Engine/Private/Animation/AnimTrace.cpp
@@ -0,0 +1,357 @@
+// Copyright 1998-2019 Epic Games, Inc. All Rights Reserved.
+
+#include "Animation/AnimTrace.h"
+
+#if ANIM_TRACE_ENABLED
+
+#include "Trace/Trace.h"
+#include "Animation/AnimationAsset.h"
+#include "Animation/AnimInstanceProxy.h"
+#include "ObjectTrace.h"
+
+UE_TRACE_EVENT_BEGIN(Animation, TickRecord)
+	UE_TRACE_EVENT_FIELD(uint64, Cycle)
+	UE_TRACE_EVENT_FIELD(uint64, ComponentId)
+	UE_TRACE_EVENT_FIELD(uint64, AnimInstanceId)
+	UE_TRACE_EVENT_FIELD(uint64, AssetId)
+	UE_TRACE_EVENT_FIELD(float, BlendWeight)
+	UE_TRACE_EVENT_FIELD(float, PlaybackTime)
+	UE_TRACE_EVENT_FIELD(float, RootMotionWeight)
+	UE_TRACE_EVENT_FIELD(float, PlayRate)
+	UE_TRACE_EVENT_FIELD(uint16, FrameCounter)
+	UE_TRACE_EVENT_FIELD(bool, Looping)
+UE_TRACE_EVENT_END()
+
+UE_TRACE_EVENT_BEGIN(Animation, SkeletalMesh, Important)
+	UE_TRACE_EVENT_FIELD(uint64, Id)
+	UE_TRACE_EVENT_FIELD(uint32, BoneCount)
+UE_TRACE_EVENT_END()
+
+UE_TRACE_EVENT_BEGIN(Animation, SkeletalMeshPose)
+	UE_TRACE_EVENT_FIELD(uint64, Cycle)
+	UE_TRACE_EVENT_FIELD(uint64, ComponentId)
+	UE_TRACE_EVENT_FIELD(uint64, MeshId)
+	UE_TRACE_EVENT_FIELD(uint32, BoneCount)
+	UE_TRACE_EVENT_FIELD(uint16, LodIndex)
+	UE_TRACE_EVENT_FIELD(uint16, FrameCounter)
+UE_TRACE_EVENT_END()
+
+UE_TRACE_EVENT_BEGIN(Animation, AnimNodeValueBool)
+	UE_TRACE_EVENT_FIELD(uint64, AnimInstanceId)
+	UE_TRACE_EVENT_FIELD(uint64, FrameId)
+	UE_TRACE_EVENT_FIELD(int32, NodeId)
+	UE_TRACE_EVENT_FIELD(int32, KeyLength)
+	UE_TRACE_EVENT_FIELD(bool, Value)
+UE_TRACE_EVENT_END()
+
+UE_TRACE_EVENT_BEGIN(Animation, AnimNodeValueInt)
+	UE_TRACE_EVENT_FIELD(uint64, AnimInstanceId)
+	UE_TRACE_EVENT_FIELD(uint64, FrameId)
+	UE_TRACE_EVENT_FIELD(int32, NodeId)
+	UE_TRACE_EVENT_FIELD(int32, KeyLength)
+	UE_TRACE_EVENT_FIELD(int32, Value)
+UE_TRACE_EVENT_END()
+
+UE_TRACE_EVENT_BEGIN(Animation, AnimNodeValueFloat)
+	UE_TRACE_EVENT_FIELD(uint64, AnimInstanceId)
+	UE_TRACE_EVENT_FIELD(uint64, FrameId)
+	UE_TRACE_EVENT_FIELD(int32, NodeId)
+	UE_TRACE_EVENT_FIELD(int32, KeyLength)
+	UE_TRACE_EVENT_FIELD(float, Value)
+UE_TRACE_EVENT_END()
+
+UE_TRACE_EVENT_BEGIN(Animation, AnimNodeValueString)
+	UE_TRACE_EVENT_FIELD(uint64, AnimInstanceId)
+	UE_TRACE_EVENT_FIELD(uint64, FrameId)
+	UE_TRACE_EVENT_FIELD(int32, NodeId)
+	UE_TRACE_EVENT_FIELD(int32, KeyLength)
+UE_TRACE_EVENT_END()
+
+UE_TRACE_EVENT_BEGIN(Animation, PoseLink)
+	UE_TRACE_EVENT_FIELD(uint64, AnimInstanceId)
+	UE_TRACE_EVENT_FIELD(uint64, FrameId)
+	UE_TRACE_EVENT_FIELD(int32, TargetLinkId)
+	UE_TRACE_EVENT_FIELD(int32, SourceLinkId)
+	UE_TRACE_EVENT_FIELD(float, Weight)
+	UE_TRACE_EVENT_FIELD(int32, TargetNameLength)
+UE_TRACE_EVENT_END()
+
+// Object annotations used for tracing
+FUObjectAnnotationSparseBool GSkeletalMeshTraceAnnotations;
+
+void FAnimTrace::Init()
+{
+	if (FParse::Param(FCommandLine::Get(), TEXT("objecttrace")))
+	{
+		UE_TRACE_EVENT_IS_ENABLED(Animation, TickRecord);
+		UE_TRACE_EVENT_IS_ENABLED(Animation, SkeletalMesh);
+		UE_TRACE_EVENT_IS_ENABLED(Animation, SkeletalMeshPose);
+// 		UE_TRACE_EVENT_IS_ENABLED(Animation, AnimNodeValueBool);
+// 		UE_TRACE_EVENT_IS_ENABLED(Animation, AnimNodeValueInt);
+// 		UE_TRACE_EVENT_IS_ENABLED(Animation, AnimNodeValueFloat);
+// 		UE_TRACE_EVENT_IS_ENABLED(Animation, AnimNodeValueString);
+// 		UE_TRACE_EVENT_IS_ENABLED(Animation, PoseLink);
+		Trace::ToggleEvent(TEXT("Animation"), true);
+	}
+}
+
+void FAnimTrace::OutputAnimTickRecords(const FAnimInstanceProxy& InProxy, const USkeletalMeshComponent* InComponent)
+{
+	bool bEventEnabled = UE_TRACE_EVENT_IS_ENABLED(Animation, TickRecord);
+	if (!bEventEnabled || InComponent == nullptr)
+	{
+		return;
+	}
+
+	TRACE_OBJECT(InComponent);
+
+	auto TraceTickRecord = [&InProxy, &InComponent](const FAnimTickRecord& InAnimTickRecord)
+	{
+		TRACE_OBJECT(InAnimTickRecord.SourceAsset);
+
+		UE_TRACE_LOG(Animation, TickRecord)
+			<< TickRecord.Cycle(FPlatformTime::Cycles64())
+			<< TickRecord.ComponentId(FObjectTrace::GetObjectId(InComponent))
+			<< TickRecord.AnimInstanceId(FObjectTrace::GetObjectId(InProxy.GetAnimInstanceObject()))
+			<< TickRecord.AssetId(FObjectTrace::GetObjectId(InAnimTickRecord.SourceAsset))
+			<< TickRecord.BlendWeight(InAnimTickRecord.EffectiveBlendWeight)
+			<< TickRecord.PlaybackTime(*InAnimTickRecord.TimeAccumulator)
+			<< TickRecord.RootMotionWeight(InAnimTickRecord.RootMotionWeightModifier)
+			<< TickRecord.PlayRate(InAnimTickRecord.PlayRateMultiplier)
+			<< TickRecord.FrameCounter((uint16)(GFrameCounter % 0xffff))
+			<< TickRecord.Looping(InAnimTickRecord.bLooping);
+	};
+
+	const TArray<FAnimGroupInstance>& SyncGroups = InProxy.SyncGroupArrays[InProxy.GetSyncGroupWriteIndex()];
+	const TArray<FAnimTickRecord>& UngroupedActivePlayers = InProxy.UngroupedActivePlayerArrays[InProxy.GetSyncGroupWriteIndex()];
+	for(const FAnimGroupInstance& SyncGroup : SyncGroups)
+	{
+		for(const FAnimTickRecord& ActivePlayer : SyncGroup.ActivePlayers)
+		{
+			TraceTickRecord(ActivePlayer);
+		}
+	}
+
+	for(const FAnimTickRecord& UngroupedActivePlayer : UngroupedActivePlayers)
+	{
+		TraceTickRecord(UngroupedActivePlayer);
+	}
+}
+
+void FAnimTrace::OutputSkeletalMesh(const USkeletalMesh* InMesh)
+{
+	bool bEventEnabled = UE_TRACE_EVENT_IS_ENABLED(Animation, SkeletalMesh);
+	if (!bEventEnabled || InMesh == nullptr)
+	{
+		return;
+	}
+
+	if(GSkeletalMeshTraceAnnotations.Get(InMesh))
+	{
+		return;
+	}
+
+	TRACE_OBJECT(InMesh);
+
+	uint32 BoneCount = (uint32)InMesh->RefSkeleton.GetRawBoneNum();
+
+	auto CopyParentIndices = [InMesh](uint8* Out)
+	{
+		int32* OutParentIndices = reinterpret_cast<int32*>(Out);
+		for(const FMeshBoneInfo& BoneInfo : InMesh->RefSkeleton.GetRawRefBoneInfo())
+		{
+			*OutParentIndices = BoneInfo.ParentIndex;
+			OutParentIndices++;
+		}
+	};
+
+	UE_TRACE_LOG(Animation, SkeletalMesh, BoneCount * sizeof(int32))
+		<< SkeletalMesh.Id(FObjectTrace::GetObjectId(InMesh))
+		<< SkeletalMesh.BoneCount(BoneCount)
+		<< SkeletalMesh.Attachment(CopyParentIndices);
+
+	GSkeletalMeshTraceAnnotations.Set(InMesh);
+}
+
+void FAnimTrace::OutputSkeletalMeshPose(const USkeletalMeshComponent* InComponent)
+{
+	bool bEventEnabled = UE_TRACE_EVENT_IS_ENABLED(Animation, SkeletalMeshPose);
+	if (!bEventEnabled || InComponent == nullptr)
+	{
+		return;
+	}
+
+	int32 BoneCount = InComponent->GetComponentSpaceTransforms().Num();
+	if(BoneCount > 0)
+	{
+		TRACE_OBJECT(InComponent);
+		TRACE_SKELETAL_MESH(InComponent->SkeletalMesh);
+
+		auto CopyTransforms = [&InComponent, &BoneCount](uint8* Out)
+		{
+			FPlatformMemory::Memcpy(Out, &InComponent->GetComponentToWorld(), sizeof(FTransform));
+			FPlatformMemory::Memcpy(Out + sizeof(FTransform), InComponent->GetComponentSpaceTransforms().GetData(), BoneCount * sizeof(FTransform));
+		};
+
+		UE_TRACE_LOG(Animation, SkeletalMeshPose, (BoneCount + 1) * sizeof(FTransform))
+			<< SkeletalMeshPose.Cycle(FPlatformTime::Cycles64())
+			<< SkeletalMeshPose.ComponentId(FObjectTrace::GetObjectId(InComponent))
+			<< SkeletalMeshPose.MeshId(FObjectTrace::GetObjectId(InComponent->SkeletalMesh))
+			<< SkeletalMeshPose.BoneCount((uint32)BoneCount + 1)
+			<< SkeletalMeshPose.LodIndex((uint16)InComponent->PredictedLODLevel)
+			<< SkeletalMeshPose.FrameCounter((uint16)(GFrameCounter % 0xffff))
+			<< SkeletalMeshPose.Attachment(CopyTransforms);
+	}
+}
+
+void FAnimTrace::OutputAnimNodeValue(const FAnimationBaseContext& InContext, const TCHAR* InKey, bool InValue)
+{
+	bool bEventEnabled = UE_TRACE_EVENT_IS_ENABLED(Animation, AnimNodeValueBool);
+	if (!bEventEnabled)
+	{
+		return;
+	}
+
+	check(InContext.AnimInstanceProxy);
+
+	int32 KeyLength = FCString::Strlen(InKey) + 1;
+
+	UE_TRACE_LOG(Animation, AnimNodeValueBool, KeyLength * sizeof(TCHAR))
+		<< AnimNodeValueBool.AnimInstanceId(FObjectTrace::GetObjectId(InContext.AnimInstanceProxy->GetAnimInstanceObject()))
+		<< AnimNodeValueBool.FrameId(GFrameCounter)
+		<< AnimNodeValueBool.NodeId(InContext.GetCurrentNodeId())
+		<< AnimNodeValueBool.KeyLength(KeyLength)
+		<< AnimNodeValueBool.Value(InValue)
+		<< AnimNodeValueBool.Attachment(InKey, KeyLength * sizeof(TCHAR));
+}
+
+void FAnimTrace::OutputAnimNodeValue(const FAnimationBaseContext& InContext, const TCHAR* InKey, int32 InValue)
+{
+	bool bEventEnabled = UE_TRACE_EVENT_IS_ENABLED(Animation, AnimNodeValueInt);
+	if (!bEventEnabled)
+	{
+		return;
+	}
+
+	check(InContext.AnimInstanceProxy);
+
+	int32 KeyLength = FCString::Strlen(InKey) + 1;
+
+	UE_TRACE_LOG(Animation, AnimNodeValueInt, KeyLength * sizeof(TCHAR))
+		<< AnimNodeValueInt.AnimInstanceId(FObjectTrace::GetObjectId(InContext.AnimInstanceProxy->GetAnimInstanceObject()))
+		<< AnimNodeValueInt.FrameId(GFrameCounter)
+		<< AnimNodeValueInt.NodeId(InContext.GetCurrentNodeId())
+		<< AnimNodeValueInt.KeyLength(KeyLength)
+		<< AnimNodeValueInt.Value(InValue)
+		<< AnimNodeValueInt.Attachment(InKey, KeyLength * sizeof(TCHAR));
+}
+
+void FAnimTrace::OutputAnimNodeValue(const FAnimationBaseContext& InContext, const TCHAR* InKey, float InValue)
+{
+	bool bEventEnabled = UE_TRACE_EVENT_IS_ENABLED(Animation, AnimNodeValueFloat);
+	if (!bEventEnabled)
+	{
+		return;
+	}
+
+	check(InContext.AnimInstanceProxy);
+
+	int32 KeyLength = FCString::Strlen(InKey) + 1;
+
+	UE_TRACE_LOG(Animation, AnimNodeValueFloat, KeyLength * sizeof(TCHAR))
+		<< AnimNodeValueFloat.AnimInstanceId(FObjectTrace::GetObjectId(InContext.AnimInstanceProxy->GetAnimInstanceObject()))
+		<< AnimNodeValueFloat.FrameId(GFrameCounter)
+		<< AnimNodeValueFloat.NodeId(InContext.GetCurrentNodeId())
+		<< AnimNodeValueFloat.KeyLength(KeyLength)
+		<< AnimNodeValueFloat.Value(InValue)
+		<< AnimNodeValueFloat.Attachment(InKey, KeyLength * sizeof(TCHAR));
+}
+
+void FAnimTrace::OutputAnimNodeValue(const FAnimationBaseContext& InContext, const TCHAR* InKey, const FName& InValue)
+{
+	bool bEventEnabled = UE_TRACE_EVENT_IS_ENABLED(Animation, AnimNodeValueFloat);
+	if (!bEventEnabled)
+	{
+		return;
+	}
+
+	check(InContext.AnimInstanceProxy);
+
+	int32 KeyLength = FCString::Strlen(InKey) + 1;
+	int32 ValueLength = InValue.GetStringLength() + 1;
+
+	auto StringCopyFunc = [KeyLength, ValueLength, InKey, &InValue](uint8* Out)
+	{
+		FCString::Strncpy(reinterpret_cast<TCHAR*>(Out), InKey, KeyLength);
+		InValue.ToString(reinterpret_cast<TCHAR*>(Out) + KeyLength, ValueLength);
+	};
+
+	UE_TRACE_LOG(Animation, AnimNodeValueString, KeyLength * sizeof(TCHAR))
+		<< AnimNodeValueString.AnimInstanceId(FObjectTrace::GetObjectId(InContext.AnimInstanceProxy->GetAnimInstanceObject()))
+		<< AnimNodeValueString.FrameId(GFrameCounter)
+		<< AnimNodeValueString.NodeId(InContext.GetCurrentNodeId())
+		<< AnimNodeValueString.KeyLength(KeyLength)
+		<< AnimNodeValueString.Attachment(InKey, KeyLength * sizeof(TCHAR));
+}
+
+void FAnimTrace::OutputAnimNodeValue(const FAnimationBaseContext& InContext, const TCHAR* InKey, const TCHAR* InValue)
+{
+	bool bEventEnabled = UE_TRACE_EVENT_IS_ENABLED(Animation, AnimNodeValueFloat);
+	if (!bEventEnabled)
+	{
+		return;
+	}
+
+	check(InContext.AnimInstanceProxy);
+
+	int32 KeyLength = FCString::Strlen(InKey) + 1;
+	int32 ValueLength = FCString::Strlen(InValue) + 1;
+
+	auto StringCopyFunc = [KeyLength, ValueLength, InKey, InValue](uint8* Out)
+	{
+		FCString::Strncpy(reinterpret_cast<TCHAR*>(Out), InKey, KeyLength);
+		FCString::Strncpy(reinterpret_cast<TCHAR*>(Out) + KeyLength, InValue, ValueLength);
+	};
+
+	UE_TRACE_LOG(Animation, AnimNodeValueString, KeyLength * sizeof(TCHAR))
+		<< AnimNodeValueString.AnimInstanceId(FObjectTrace::GetObjectId(InContext.AnimInstanceProxy->GetAnimInstanceObject()))
+		<< AnimNodeValueString.FrameId(GFrameCounter)
+		<< AnimNodeValueString.NodeId(InContext.GetCurrentNodeId())
+		<< AnimNodeValueString.KeyLength(KeyLength)
+		<< AnimNodeValueString.Attachment(InKey, KeyLength * sizeof(TCHAR));
+}
+
+void FAnimTrace::OutputPoseLink(const FAnimationUpdateContext& InContext)
+{
+	bool bEventEnabled = UE_TRACE_EVENT_IS_ENABLED(Animation, PoseLink);
+	if (!bEventEnabled)
+	{
+		return;
+	}
+
+	check(InContext.AnimInstanceProxy);
+
+	IAnimClassInterface* AnimBlueprintClass = InContext.GetAnimClass();
+	check(AnimBlueprintClass);
+	const TArray<UStructProperty*>& AnimNodeProperties = AnimBlueprintClass->GetAnimNodeProperties();
+	check(AnimNodeProperties.IsValidIndex(InContext.GetCurrentNodeId()));
+	UProperty* LinkedProperty = AnimNodeProperties[InContext.GetCurrentNodeId()];
+
+	int32 NameLength = LinkedProperty->GetFName().GetStringLength() + 1;
+
+	auto StringCopyFunc = [NameLength, LinkedProperty](uint8* Out)
+	{
+		LinkedProperty->GetFName().ToString(reinterpret_cast<TCHAR*>(Out), NameLength);
+	};
+
+	UE_TRACE_LOG(Animation, PoseLink)
+		<< PoseLink.AnimInstanceId(FObjectTrace::GetObjectId(InContext.AnimInstanceProxy->GetAnimInstanceObject()))
+		<< PoseLink.FrameId(GFrameCounter)
+		<< PoseLink.SourceLinkId(InContext.GetPreviousNodeId())
+		<< PoseLink.TargetLinkId(InContext.GetCurrentNodeId())
+		<< PoseLink.Weight(InContext.GetFinalBlendWeight())
+		<< PoseLink.TargetNameLength(NameLength)
+		<< PoseLink.Attachment(StringCopyFunc);
+}
+
+#endif
\ No newline at end of file
diff --git a/Engine/Source/Runtime/Engine/Private/Components/ActorComponent.cpp b/Engine/Source/Runtime/Engine/Private/Components/ActorComponent.cpp
--- a/Engine/Source/Runtime/Engine/Private/Components/ActorComponent.cpp
+++ b/Engine/Source/Runtime/Engine/Private/Components/ActorComponent.cpp
@@ -34,6 +34,7 @@
 #if WITH_EDITOR
 #include "Kismet2/ComponentEditorUtils.h"
 #endif
+#include "ObjectTrace.h"
 
 #define LOCTEXT_NAMESPACE "ActorComponent"
 
@@ -818,6 +819,8 @@ void UActorComponent::UninitializeComponent()
 
 void UActorComponent::BeginPlay()
 {
+	TRACE_OBJECT_EVENT(this, BeginPlay);
+
 	check(bRegistered);
 	check(!bHasBegunPlay);
 	checkSlow(bTickFunctionsRegistered); // If this fails, someone called BeginPlay() without first calling RegisterAllComponentTickFunctions().
@@ -832,6 +835,8 @@ void UActorComponent::BeginPlay()
 
 void UActorComponent::EndPlay(const EEndPlayReason::Type EndPlayReason)
 {
+	TRACE_OBJECT_EVENT(this, EndPlay);
+
 	check(bHasBegunPlay);
 
 	// If we're in the process of being garbage collected it is unsafe to call out to blueprints
diff --git a/Engine/Source/Runtime/Engine/Private/Components/SkeletalMeshComponent.cpp b/Engine/Source/Runtime/Engine/Private/Components/SkeletalMeshComponent.cpp
--- a/Engine/Source/Runtime/Engine/Private/Components/SkeletalMeshComponent.cpp
+++ b/Engine/Source/Runtime/Engine/Private/Components/SkeletalMeshComponent.cpp
@@ -39,6 +39,7 @@
 #include "UObject/AnimPhysObjectVersion.h"
 #include "SkeletalRenderPublic.h"
 #include "ContentStreaming.h"
+#include "Animation/AnimTrace.h"
 
 #define LOCTEXT_NAMESPACE "SkeletalMeshComponent"
 
@@ -3705,6 +3706,8 @@ void USkeletalMeshComponent::FinalizeBoneTransform()
 	ConditionallyDispatchQueuedAnimEvents();
 
 	OnBoneTransformsFinalized.Broadcast();
+
+	TRACE_SKELETALMESH_POSE(this);
 }
 
 void USkeletalMeshComponent::GetCurrentRefToLocalMatrices(TArray<FMatrix>& OutRefToLocals, int32 InLodIdx)
diff --git a/Engine/Source/Runtime/Engine/Private/Level.cpp b/Engine/Source/Runtime/Engine/Private/Level.cpp
--- a/Engine/Source/Runtime/Engine/Private/Level.cpp
+++ b/Engine/Source/Runtime/Engine/Private/Level.cpp
@@ -58,6 +58,7 @@ Level.cpp: Level-related functions
 #include "ComponentRecreateRenderStateContext.h"
 #include "Algo/Copy.h"
 #include "HAL/LowLevelMemTracker.h"
+#include "ObjectTrace.h"
 
 DEFINE_LOG_CATEGORY(LogLevel);
 
@@ -1803,6 +1804,8 @@ void ULevel::ReleaseRenderingResources()
 
 void ULevel::RouteActorInitialize()
 {
+	TRACE_OBJECT_EVENT(this, RouteActorInitialize);
+
 	// Send PreInitializeComponents and collect volumes.
 	for( int32 Index = 0; Index < Actors.Num(); ++Index )
 	{
diff --git a/Engine/Source/Runtime/Engine/Private/ObjectTrace.cpp b/Engine/Source/Runtime/Engine/Private/ObjectTrace.cpp
new file mode 100644
--- /dev/null
+++ b/Engine/Source/Runtime/Engine/Private/ObjectTrace.cpp
@@ -0,0 +1,182 @@
+// Copyright 1998-2019 Epic Games, Inc. All Rights Reserved.
+
+#include "ObjectTrace.h"
+
+#if OBJECT_TRACE_ENABLED
+
+#include "Trace/Trace.h"
+#include "UObject/Object.h"
+#include "UObject/UObjectAnnotation.h"
+
+UE_TRACE_EVENT_BEGIN(Object, Class, Important)
+	UE_TRACE_EVENT_FIELD(uint64, Id)
+	UE_TRACE_EVENT_FIELD(uint64, SuperId)
+	UE_TRACE_EVENT_FIELD(int32, ClassNameStringLength)
+UE_TRACE_EVENT_END()
+
+UE_TRACE_EVENT_BEGIN(Object, Object, Important)
+	UE_TRACE_EVENT_FIELD(uint64, Id)
+	UE_TRACE_EVENT_FIELD(uint64, ClassId)
+	UE_TRACE_EVENT_FIELD(uint64, OuterId)
+	UE_TRACE_EVENT_FIELD(int32, ObjectNameStringLength)
+UE_TRACE_EVENT_END()
+
+UE_TRACE_EVENT_BEGIN(Object, ObjectEvent)
+	UE_TRACE_EVENT_FIELD(uint64, Cycle)
+	UE_TRACE_EVENT_FIELD(uint64, Id)
+	UE_TRACE_EVENT_FIELD(uint8, Event)
+UE_TRACE_EVENT_END()
+
+// Object annotations used for tracing
+struct FTracedObjectAnnotation
+{
+	FTracedObjectAnnotation()
+		: bTraced(false)
+	{}
+
+	// Whether this object has been traced this session
+	bool bTraced;
+
+	/** Determine if this annotation is default - required for annotations */
+	FORCEINLINE bool IsDefault() const
+	{
+		return !bTraced;
+	}
+};
+
+// Object annotations used for tracing
+FUObjectAnnotationSparse<FTracedObjectAnnotation, true> GObjectTraceAnnotations;
+
+void FObjectTrace::Init()
+{
+	if (FParse::Param(FCommandLine::Get(), TEXT("objecttrace")))
+	{
+		UE_TRACE_EVENT_IS_ENABLED(Object, Class);
+		UE_TRACE_EVENT_IS_ENABLED(Object, Object);
+		UE_TRACE_EVENT_IS_ENABLED(Object, ObjectEvent);
+		Trace::ToggleEvent(TEXT("Object"), true);
+	}
+}
+
+uint64 FObjectTrace::GetObjectId(const UObject* InObject)
+{
+	// An object ID uses a combination of its own and its outer's hash
+	// We do this to represent objects that get renamed into different outers 
+	// as distinct traces (we don't attempt to link them).
+	// We use a weak object ptr hash to take into account object name-recycling, 
+	// and represent the new object as a distinct entry
+	uint64 ObjectHash = 0;
+	uint64 OuterHash = 0;
+	if(InObject != nullptr)
+	{
+		ObjectHash = GetTypeHash(TWeakObjectPtr<UObject>(InObject));
+		if(UObject* Outer = InObject->GetOuter())
+		{
+			OuterHash = GetTypeHash(TWeakObjectPtr<UObject>(Outer));
+		}
+	}
+
+	return ObjectHash | (OuterHash << 32);
+}
+
+void FObjectTrace::OutputClass(const UClass* InClass)
+{
+	bool bEventEnabled = UE_TRACE_EVENT_IS_ENABLED(Object, Class);
+	if (!bEventEnabled || InClass == nullptr)
+	{
+		return;
+	}
+
+	FTracedObjectAnnotation Annotation = GObjectTraceAnnotations.GetAnnotation(InClass);
+	if(Annotation.bTraced)
+	{
+		// Already traced, so skip
+		return;
+	}
+
+	Annotation.bTraced = true;
+	GObjectTraceAnnotations.AddAnnotation(InClass, Annotation);
+
+	int32 ClassNameStringLength = InClass->GetFName().GetStringLength() + 1;
+	int32 ClassFullNameStringLength = InClass->GetPathName().Len() + 1;
+
+	auto StringCopyFunc = [ClassNameStringLength, ClassFullNameStringLength, InClass](uint8* Out)
+	{
+		InClass->GetFName().ToString(reinterpret_cast<TCHAR*>(Out), ClassNameStringLength);
+		FPlatformMemory::Memcpy(reinterpret_cast<TCHAR*>(Out) + ClassNameStringLength, *InClass->GetPathName(), ClassFullNameStringLength * sizeof(TCHAR));
+	};
+
+	UE_TRACE_LOG(Object, Class, (ClassNameStringLength + ClassFullNameStringLength) * sizeof(TCHAR))
+		<< Class.ClassNameStringLength(ClassNameStringLength)
+		<< Class.Id(GetObjectId(InClass))
+		<< Class.SuperId(GetObjectId(InClass->GetSuperClass()))
+		<< Class.Attachment(StringCopyFunc);
+}
+
+void FObjectTrace::OutputObject(const UObject* InObject)
+{
+	bool bEventEnabled = UE_TRACE_EVENT_IS_ENABLED(Object, Object);
+	if (!bEventEnabled || InObject == nullptr)
+	{
+		return;
+	}
+
+	if(InObject->HasAnyFlags(RF_ClassDefaultObject))
+	{
+		return;
+	}
+
+	FTracedObjectAnnotation Annotation = GObjectTraceAnnotations.GetAnnotation(InObject);
+	if(Annotation.bTraced)
+	{
+		// Already traced, so skip
+		return;
+	}
+
+	Annotation.bTraced = true;
+	GObjectTraceAnnotations.AddAnnotation(InObject, Annotation);
+
+	// Trace the object's class first
+	TRACE_CLASS(InObject->GetClass());
+
+	int32 ObjectNameStringLength = InObject->GetFName().GetStringLength() + 1;
+	int32 ObjectPathNameStringLength = InObject->GetPathName().Len() + 1;
+
+	auto StringCopyFunc = [ObjectNameStringLength, ObjectPathNameStringLength, InObject](uint8* Out)
+	{
+		InObject->GetFName().ToString(reinterpret_cast<TCHAR*>(Out), ObjectNameStringLength);
+		FPlatformMemory::Memcpy(reinterpret_cast<TCHAR*>(Out) + ObjectNameStringLength, *InObject->GetPathName(), ObjectPathNameStringLength * sizeof(TCHAR));
+	};
+
+	UE_TRACE_LOG(Object, Object, (ObjectNameStringLength + ObjectPathNameStringLength) * sizeof(TCHAR))
+		<< Object.ObjectNameStringLength(ObjectNameStringLength)
+		<< Object.Id(GetObjectId(InObject))
+		<< Object.ClassId(GetObjectId(InObject->GetClass()))
+		<< Object.OuterId(GetObjectId(InObject->GetOuter()))
+		<< Object.Attachment(StringCopyFunc);
+}
+
+void FObjectTrace::OutputObjectEvent(const UObject* InObject, const TCHAR* InEvent)
+{
+	bool bEventEnabled = UE_TRACE_EVENT_IS_ENABLED(Object, ObjectEvent);
+	if (!bEventEnabled || InObject == nullptr)
+	{
+		return;
+	}
+
+	if(InObject->HasAnyFlags(RF_ClassDefaultObject))
+	{
+		return;
+	}
+
+	TRACE_OBJECT(InObject);
+
+	int32 StringBufferSize = (FCString::Strlen(InEvent) + 1) * sizeof(TCHAR);
+
+	UE_TRACE_LOG(Object, ObjectEvent, StringBufferSize)
+		<< ObjectEvent.Cycle(FPlatformTime::Cycles64())
+		<< ObjectEvent.Id(GetObjectId(InObject))
+		<< ObjectEvent.Attachment(InEvent, StringBufferSize);
+}
+
+#endif
\ No newline at end of file
diff --git a/Engine/Source/Runtime/Engine/Private/UnrealEngine.cpp b/Engine/Source/Runtime/Engine/Private/UnrealEngine.cpp
--- a/Engine/Source/Runtime/Engine/Private/UnrealEngine.cpp
+++ b/Engine/Source/Runtime/Engine/Private/UnrealEngine.cpp
@@ -228,6 +228,8 @@ UnrealEngine.cpp: Implements the UEngine class and helpers.
 
 #include "Particles/ParticleSystemManager.h"
 #include "Components/SkinnedMeshComponent.h"
+#include "ObjectTrace.h"
+#include "Animation/AnimTrace.h"
 
 DEFINE_LOG_CATEGORY(LogEngine);
 IMPLEMENT_MODULE( FEngineModule, Engine );
@@ -264,6 +266,14 @@ void FEngineModule::StartupModule()
 #if WITH_EDITOR
 	USkinnedMeshComponent::BindWorldDelegates();
 #endif
+
+#if OBJECT_TRACE_ENABLED
+	FObjectTrace::Init();
+#endif
+
+#if ANIM_TRACE_ENABLED
+	FAnimTrace::Init();
+#endif
 }
 
 void FEngineModule::ShutdownModule()
diff --git a/Engine/Source/Runtime/Engine/Private/World.cpp b/Engine/Source/Runtime/Engine/Private/World.cpp
--- a/Engine/Source/Runtime/Engine/Private/World.cpp
+++ b/Engine/Source/Runtime/Engine/Private/World.cpp
@@ -124,6 +124,7 @@
 #include "Engine/HLODProxy.h"
 #include "ProfilingDebugging/CsvProfiler.h"
 #include "Interfaces/Interface_PostProcessVolume.h"
+#include "ObjectTrace.h"
 
 #if INCLUDE_CHAOS
 #include "ChaosSolversModule.h"
@@ -3924,6 +3925,8 @@ bool UWorld::SetGameMode(const FURL& InURL)
 
 void UWorld::InitializeActorsForPlay(const FURL& InURL, bool bResetTime)
 {
+	TRACE_OBJECT_EVENT(this, InitializeActorsForPlay);
+
 	check(bIsWorldInitialized);
 	SCOPED_BOOT_TIMING("UWorld::InitializeActorsForPlay");
 	double StartTime = FPlatformTime::Seconds();
diff --git a/Engine/Source/Runtime/Engine/Public/Animation/AnimInstanceProxy.h b/Engine/Source/Runtime/Engine/Public/Animation/AnimInstanceProxy.h
--- a/Engine/Source/Runtime/Engine/Public/Animation/AnimInstanceProxy.h
+++ b/Engine/Source/Runtime/Engine/Public/Animation/AnimInstanceProxy.h
@@ -463,6 +463,7 @@ public:
 	friend class USkeletalMeshComponent;
 	friend struct FAnimNode_LinkedAnimGraph;
 	friend struct FAnimationBaseContext;
+	friend struct FAnimTrace;
 
 protected:
 	/** Called when our anim instance is being initialized */
diff --git a/Engine/Source/Runtime/Engine/Public/Animation/AnimTrace.h b/Engine/Source/Runtime/Engine/Public/Animation/AnimTrace.h
new file mode 100644
--- /dev/null
+++ b/Engine/Source/Runtime/Engine/Public/Animation/AnimTrace.h
@@ -0,0 +1,64 @@
+// Copyright 1998-2019 Epic Games, Inc. All Rights Reserved.
+
+#pragma once
+
+#include "CoreTypes.h"
+#include "ObjectTrace.h"
+
+#define ANIM_TRACE_ENABLED OBJECT_TRACE_ENABLED
+
+#if ANIM_TRACE_ENABLED
+
+struct FAnimTickRecord;
+struct FAnimationBaseContext;
+struct FAnimationUpdateContext;
+
+struct FAnimTrace
+{
+	/** Initialize animation tracing */
+	ENGINE_API static void Init();
+
+	/** Helper function to output a set of tick records */
+	ENGINE_API static void OutputAnimTickRecords(const FAnimInstanceProxy& InProxy, const USkeletalMeshComponent* InComponent);
+
+	/** Helper function to output a skeletal mesh */
+	ENGINE_API static void OutputSkeletalMesh(const USkeletalMesh* InMesh);
+
+	/** Helper function to output a skeletal mesh pose */
+	ENGINE_API static void OutputSkeletalMeshPose(const USkeletalMeshComponent* InComponent);
+
+	/** Helper function to output a pose link */
+	ENGINE_API static void OutputPoseLink(const FAnimationUpdateContext& InContext);
+
+	/** Helper function to output a tracked value for an anim node */
+	ENGINE_API static void OutputAnimNodeValue(const FAnimationBaseContext& InContext, const TCHAR* InKey, bool Value);
+	ENGINE_API static void OutputAnimNodeValue(const FAnimationBaseContext& InContext, const TCHAR* InKey, int32 Value);
+	ENGINE_API static void OutputAnimNodeValue(const FAnimationBaseContext& InContext, const TCHAR* InKey, float Value);
+	ENGINE_API static void OutputAnimNodeValue(const FAnimationBaseContext& InContext, const TCHAR* InKey, const FName& Value);
+	ENGINE_API static void OutputAnimNodeValue(const FAnimationBaseContext& InContext, const TCHAR* InKey, const TCHAR* Value);
+};
+
+#define TRACE_ANIM_TICK_RECORDS(Proxy, Component) \
+	FAnimTrace::OutputAnimTickRecords(Proxy, Component);
+
+#define TRACE_SKELETAL_MESH(Mesh) \
+	FAnimTrace::OutputSkeletalMesh(Mesh);
+
+#define TRACE_SKELETALMESH_POSE(Component) \
+	FAnimTrace::OutputSkeletalMeshPose(Component);
+
+#define TRACE_POSE_LINK(Context) \
+	FAnimTrace::OutputPoseLink(Context);
+
+#define TRACE_ANIM_NODE_VALUE(Context, Key, Value) \
+	FAnimTrace::OutputAnimNodeValue(Context, Key, Value);
+
+#else
+
+#define TRACE_ANIM_TICK_RECORDS
+#define TRACE_SKELETAL_MESH
+#define TRACE_SKELETAL_MESH_POSE
+#define TRACE_POSE_LINK
+#define TRACE_ANIM_NODE_VALUE
+
+#endif
\ No newline at end of file
diff --git a/Engine/Source/Runtime/Engine/Public/ObjectTrace.h b/Engine/Source/Runtime/Engine/Public/ObjectTrace.h
new file mode 100644
--- /dev/null
+++ b/Engine/Source/Runtime/Engine/Public/ObjectTrace.h
@@ -0,0 +1,48 @@
+// Copyright 1998-2019 Epic Games, Inc. All Rights Reserved.
+
+#pragma once
+
+#include "CoreTypes.h"
+
+#if !IS_PROGRAM && !UE_BUILD_SHIPPING
+#define OBJECT_TRACE_ENABLED 1
+#else
+#define OBJECT_TRACE_ENABLED 0
+#endif
+
+#if OBJECT_TRACE_ENABLED
+
+struct FObjectTrace
+{
+	/** Initialize object tracing */
+	ENGINE_API static void Init();
+
+	/** Helper function to output an object */
+	ENGINE_API static void OutputClass(const UClass* InClass);
+
+	/** Helper function to output an object */
+	ENGINE_API static void OutputObject(const UObject* InObject);
+
+	/** Helper function to output an object event */
+	ENGINE_API static void OutputObjectEvent(const UObject* InObject, const TCHAR* InEvent);
+
+	/** Helper function to get an object ID from a UObject */
+	ENGINE_API static uint64 GetObjectId(const UObject* InObject);
+};
+
+#define TRACE_CLASS(Class) \
+	FObjectTrace::OutputClass(Class);
+
+#define TRACE_OBJECT(Object) \
+	FObjectTrace::OutputObject(Object);
+
+#define TRACE_OBJECT_EVENT(Object, Event) \
+	FObjectTrace::OutputObjectEvent(Object, TEXT(#Event));
+
+#else
+
+#define TRACE_CLASS
+#define TRACE_OBJECT
+#define TRACE_OBJECT_EVENT
+
+#endif
\ No newline at end of file
