IWYU fixes

diff --git a/Engine/Source/Runtime/Engine/Private/Animation/AnimTrace.cpp b/Engine/Source/Runtime/Engine/Private/Animation/AnimTrace.cpp
--- a/Engine/Source/Runtime/Engine/Private/Animation/AnimTrace.cpp
+++ b/Engine/Source/Runtime/Engine/Private/Animation/AnimTrace.cpp
@@ -8,6 +8,11 @@
 #include "Animation/AnimationAsset.h"
 #include "Animation/AnimInstanceProxy.h"
 #include "ObjectTrace.h"
+#include "Components/SkeletalMeshComponent.h"
+#include "Misc/CommandLine.h"
+#include "Engine/SkeletalMesh.h"
+#include "Math/TransformNonVectorized.h"
+#include "Animation/AnimNodeBase.h"
 
 UE_TRACE_EVENT_BEGIN(Animation, TickRecord)
 	UE_TRACE_EVENT_FIELD(uint64, Cycle)
diff --git a/Engine/Source/Runtime/Engine/Private/ObjectTrace.cpp b/Engine/Source/Runtime/Engine/Private/ObjectTrace.cpp
--- a/Engine/Source/Runtime/Engine/Private/ObjectTrace.cpp
+++ b/Engine/Source/Runtime/Engine/Private/ObjectTrace.cpp
@@ -4,9 +4,13 @@
 
 #if OBJECT_TRACE_ENABLED
 
+#include "CoreMinimal.h"
 #include "Trace/Trace.h"
 #include "UObject/Object.h"
+#include "UObject/Class.h"
 #include "UObject/UObjectAnnotation.h"
+#include "UObject/WeakObjectPtrTemplates.h"
+#include "Misc/CommandLine.h"
 
 UE_TRACE_EVENT_BEGIN(Object, Class, Important)
 	UE_TRACE_EVENT_FIELD(uint64, Id)
diff --git a/Engine/Source/Runtime/Engine/Public/Animation/AnimTrace.h b/Engine/Source/Runtime/Engine/Public/Animation/AnimTrace.h
--- a/Engine/Source/Runtime/Engine/Public/Animation/AnimTrace.h
+++ b/Engine/Source/Runtime/Engine/Public/Animation/AnimTrace.h
@@ -9,9 +9,14 @@
 
 #if ANIM_TRACE_ENABLED
 
+struct FAnimInstanceProxy;
 struct FAnimTickRecord;
 struct FAnimationBaseContext;
 struct FAnimationUpdateContext;
+class USkeletalMesh;
+class USkeletalMeshComponent;
+struct FAnimationUpdateContext;
+struct FAnimationBaseContext;
 
 struct FAnimTrace
 {
