Non-PCH build fixes

diff --git a/Engine/Source/Runtime/Engine/Private/ObjectTrace.cpp b/Engine/Source/Runtime/Engine/Private/ObjectTrace.cpp
--- a/Engine/Source/Runtime/Engine/Private/ObjectTrace.cpp
+++ b/Engine/Source/Runtime/Engine/Private/ObjectTrace.cpp
@@ -10,6 +10,7 @@
 #include "UObject/Class.h"
 #include "UObject/UObjectAnnotation.h"
 #include "UObject/WeakObjectPtrTemplates.h"
+#include "UObject/WeakObjectPtr.h"
 #include "Misc/CommandLine.h"
 
 UE_TRACE_EVENT_BEGIN(Object, Class, Important)
diff --git a/Engine/Source/Runtime/Engine/Public/Animation/AnimTrace.h b/Engine/Source/Runtime/Engine/Public/Animation/AnimTrace.h
--- a/Engine/Source/Runtime/Engine/Public/Animation/AnimTrace.h
+++ b/Engine/Source/Runtime/Engine/Public/Animation/AnimTrace.h
@@ -17,6 +17,7 @@ class USkeletalMesh;
 class USkeletalMeshComponent;
 struct FAnimationUpdateContext;
 struct FAnimationBaseContext;
+class FName;
 
 struct FAnimTrace
 {
diff --git a/Engine/Source/Runtime/Engine/Public/ObjectTrace.h b/Engine/Source/Runtime/Engine/Public/ObjectTrace.h
--- a/Engine/Source/Runtime/Engine/Public/ObjectTrace.h
+++ b/Engine/Source/Runtime/Engine/Public/ObjectTrace.h
@@ -12,6 +12,9 @@
 
 #if OBJECT_TRACE_ENABLED
 
+class UClass;
+class UObject;
+
 struct FObjectTrace
 {
 	/** Initialize object tracing */
