Compile fixes if Object/Anim tracing is turned off.

#rnx

diff --git a/Engine/Source/Runtime/Engine/Public/Animation/AnimTrace.h b/Engine/Source/Runtime/Engine/Public/Animation/AnimTrace.h
--- a/Engine/Source/Runtime/Engine/Public/Animation/AnimTrace.h
+++ b/Engine/Source/Runtime/Engine/Public/Animation/AnimTrace.h
@@ -61,10 +61,10 @@ struct FAnimTrace
 
 #else
 
-#define TRACE_ANIM_TICK_RECORDS
-#define TRACE_SKELETAL_MESH
-#define TRACE_SKELETAL_MESH_POSE
-#define TRACE_POSE_LINK
-#define TRACE_ANIM_NODE_VALUE
+#define TRACE_ANIM_TICK_RECORDS(Proxy, Component)
+#define TRACE_SKELETAL_MESH(Mesh)
+#define TRACE_SKELETALMESH_POSE(Component)
+#define TRACE_POSE_LINK(Context)
+#define TRACE_ANIM_NODE_VALUE(Context, Key, Value)
 
 #endif
\ No newline at end of file
diff --git a/Engine/Source/Runtime/Engine/Public/ObjectTrace.h b/Engine/Source/Runtime/Engine/Public/ObjectTrace.h
--- a/Engine/Source/Runtime/Engine/Public/ObjectTrace.h
+++ b/Engine/Source/Runtime/Engine/Public/ObjectTrace.h
@@ -44,8 +44,8 @@ struct FObjectTrace
 
 #else
 
-#define TRACE_CLASS
-#define TRACE_OBJECT
-#define TRACE_OBJECT_EVENT
+#define TRACE_CLASS(Class)
+#define TRACE_OBJECT(Object)
+#define TRACE_OBJECT_EVENT(Object, Event)
 
 #endif
\ No newline at end of file
