Fixed Linux deprecation warning.

#rnx

diff --git a/Engine/Source/Runtime/Engine/Private/ObjectTrace.cpp b/Engine/Source/Runtime/Engine/Private/ObjectTrace.cpp
--- a/Engine/Source/Runtime/Engine/Private/ObjectTrace.cpp
+++ b/Engine/Source/Runtime/Engine/Private/ObjectTrace.cpp
@@ -74,7 +74,7 @@ uint64 FObjectTrace::GetObjectId(const UObject* InObject)
 	uint64 OuterHash = 0;
 	if(InObject != nullptr)
 	{
-		ObjectHash = GetTypeHash(TWeakObjectPtr<UObject>(InObject));
+		ObjectHash = GetTypeHash(TWeakObjectPtr<const UObject>(InObject));
 		if(UObject* Outer = InObject->GetOuter())
 		{
 			OuterHash = GetTypeHash(TWeakObjectPtr<UObject>(Outer));
