Avoid calling GetBulkDataSize twice as looking up the size via the IoDispatcher is a bit more expensive than returning a cached size.

diff --git a/Engine/Source/Runtime/Engine/Private/TextureDerivedData.cpp b/Engine/Source/Runtime/Engine/Private/TextureDerivedData.cpp
--- a/Engine/Source/Runtime/Engine/Private/TextureDerivedData.cpp
+++ b/Engine/Source/Runtime/Engine/Private/TextureDerivedData.cpp
@@ -1076,11 +1076,12 @@ bool FTexturePlatformData::TryLoadMips(int32 FirstMipToLoad, void** OutMipData)
 	for (int32 MipIndex = FirstMipToLoad; MipIndex < LoadableMips; ++MipIndex)
 	{
 		FTexture2DMipMap& Mip = Mips[MipIndex];
-		if (Mip.BulkData.GetBulkDataSize() > 0)
+		const int64 BulkDataSize = Mip.BulkData.GetBulkDataSize();
+		if (BulkDataSize > 0)
 		{
-			if (OutMipData)
+			if (OutMipData != nullptr)
 			{
-				OutMipData[MipIndex - FirstMipToLoad] = FMemory::Malloc(Mip.BulkData.GetBulkDataSize());
+				OutMipData[MipIndex - FirstMipToLoad] = FMemory::Malloc(BulkDataSize);
 
 #if PLATFORM_SUPPORTS_TEXTURE_STREAMING && !TEXTURE2DMIPMAP_USE_COMPACT_BULKDATA
 				// We want to make sure that any non-streamed mips are coming from the texture asset file, and not from an external bulk file.
