Some small optimizations when loading texture mips from inline bulkdata.

- Instead of allocating the memory buffer and passing it to ::GetCopy we should pass in a nullptr, that way if the data is reloadable and either set to singleuse or we signal to ::GetCopy to discard the internal bulkdata buffer (which we do) then we can just take ownership of that internal buffer, saving us a memory allocation and a memcpy.
- Checking the bulkdata filename for the ".ubulk" extension is slow and doesn't actually cover all cases now that there are multiple extensions for seperate bulkdata files. We can check for all of the different extensions by just calling ::InSeperateFile which is what we really want to know and is a lot faster as it's just a bitflag check rather than a stringcmp.

diff --git a/Engine/Source/Runtime/Engine/Private/TextureDerivedData.cpp b/Engine/Source/Runtime/Engine/Private/TextureDerivedData.cpp
--- a/Engine/Source/Runtime/Engine/Private/TextureDerivedData.cpp
+++ b/Engine/Source/Runtime/Engine/Private/TextureDerivedData.cpp
@@ -1087,15 +1087,13 @@ bool FTexturePlatformData::TryLoadMips(int32 FirstMipToLoad, void** OutMipData)
 		{
 			if (OutMipData != nullptr)
 			{
-				OutMipData[MipIndex - FirstMipToLoad] = FMemory::Malloc(BulkDataSize);
-
 #if PLATFORM_SUPPORTS_TEXTURE_STREAMING && !TEXTURE2DMIPMAP_USE_COMPACT_BULKDATA
 				// We want to make sure that any non-streamed mips are coming from the texture asset file, and not from an external bulk file.
 				// But because "r.TextureStreaming" is driven by the project setting as well as the command line option "-NoTextureStreaming", 
 				// is it possible for streaming mips to be loaded in non streaming ways.
 				if (CVarSetTextureStreaming.GetValueOnAnyThread() != 0)
 				{
-					UE_CLOG(Mip.BulkData.GetFilename().EndsWith(TEXT(".ubulk")), LogTexture, Error, TEXT("Loading non-streamed mips from an external bulk file.  This is not desireable.  File %s"), *(Mip.BulkData.GetFilename() ) );
+					UE_CLOG(Mip.BulkData.InSeperateFile(), LogTexture, Error, TEXT("Loading non-streamed mips from an external bulk file.  This is not desireable.  File %s"), *(Mip.BulkData.GetFilename() ) );
 				}
 #endif
 				Mip.BulkData.GetCopy(&OutMipData[MipIndex - FirstMipToLoad], true);
