Prevent loading of UVectorFieldStatic::SourceData data twice during game load.

::UpdateCPUData would discard the BulkData contents when calling ::GetCopy, but in all case but one, after ::UpdateCPUDatais called the FVectorFieldStaticResource is updated which also needs that BulkData which would have to be reloaded from disk and then discarded again.
As an interim fix there is now a parameter that allows the caller to control if the data is discarded or not, however this still has a few draw backs such as the call  to ::GetCopy in ::UpdateCPUData will still allocate on the heap and memcpy the internal data if the data is not to be discarded.
A better approach would to be to rework how these updates work and just take the data from the bulk data once (which would just return the internal pointer is discard is true), use it for the CPU update and then pass it to FVectorFieldStaticResource which would eventually delete it once done.

#rnx

diff --git a/Engine/Source/Runtime/Engine/Classes/VectorField/VectorFieldStatic.h b/Engine/Source/Runtime/Engine/Classes/VectorField/VectorFieldStatic.h
--- a/Engine/Source/Runtime/Engine/Classes/VectorField/VectorFieldStatic.h
+++ b/Engine/Source/Runtime/Engine/Classes/VectorField/VectorFieldStatic.h
@@ -78,8 +78,13 @@ public:
 	 */
 	ENGINE_API void InitResource();
 
-	/** Takes a local copy of the source bulk data so that it is readable at runtime on the CPU. */
-	ENGINE_API void UpdateCPUData();
+	/** Takes a local copy of the source bulk data so that it is readable at runtime on the CPU. 
+	* @param bDiscardData If the internal loaded data should be discarded after calling this or not.
+	*/
+	ENGINE_API void UpdateCPUData(bool bDiscardData);
+
+	UE_DEPRECATED(4.25, "Please call UpdateCPUData(bool bDiscardData) instead and choose if the internal data should be discarded or not")
+	ENGINE_API void UpdateCPUData() { UpdateCPUData(true); }
 
 #if WITH_EDITOR
 	/** Sets the bAllowCPUAccess flag and calls UpdateCPUData(). */
diff --git a/Engine/Source/Runtime/Engine/Private/VectorField.cpp b/Engine/Source/Runtime/Engine/Private/VectorField.cpp
--- a/Engine/Source/Runtime/Engine/Private/VectorField.cpp
+++ b/Engine/Source/Runtime/Engine/Private/VectorField.cpp
@@ -286,9 +286,9 @@ void UVectorFieldStatic::InitResource()
 	check(!Resource);
 
 	// Loads and copies the bulk data into CPUData if bAllowCPUAccess is set, otherwise clear CPUData. 
-	UpdateCPUData();
+	UpdateCPUData(false);
 
-	Resource = new FVectorFieldStaticResource(this);
+	Resource = new FVectorFieldStaticResource(this); // Will discard the contents of SourceData
 	Resource->AddRef(); // Increment refcount because of UVectorFieldStatic::Resource is not a TRefCountPtr.
 
 	BeginInitResource(Resource);
@@ -300,21 +300,21 @@ void UVectorFieldStatic::UpdateResource()
 	check(Resource);
 
 	// Loads and copies the bulk data into CPUData if bAllowCPUAccess is set, otherwise clears CPUData. 
-	UpdateCPUData();
+	UpdateCPUData(false);
 
 	FVectorFieldStaticResource* StaticResource = (FVectorFieldStaticResource*)Resource;
-	StaticResource->UpdateResource(this);
+	StaticResource->UpdateResource(this); // Will discard the contents of SourceData
 }
 
 #if WITH_EDITOR
 ENGINE_API void UVectorFieldStatic::SetCPUAccessEnabled()
 {
 	bAllowCPUAccess = true;
-	UpdateCPUData();
+	UpdateCPUData(true);
 }
 #endif // WITH_EDITOR
 
-void UVectorFieldStatic::UpdateCPUData()
+void UVectorFieldStatic::UpdateCPUData(bool bDiscardData)
 {
 	if (bAllowCPUAccess)
 	{
@@ -322,7 +322,7 @@ void UVectorFieldStatic::UpdateCPUData()
 		// If the data is already loaded it makes a copy and discards the old content,
 		// otherwise it simply loads the data directly from file into the pointer after allocating.
 		FFloat16Color *Ptr = nullptr;
-		SourceData.GetCopy((void**)&Ptr, /* bDiscardInternalCopy */ true);
+		SourceData.GetCopy((void**)&Ptr, bDiscardData);
 
 		// Make sure the data is actually valid. 
 		if (!ensure(Ptr))
@@ -340,7 +340,7 @@ void UVectorFieldStatic::UpdateCPUData()
 		}
 
 		// GetCopy should free/unload the data.
-		if (SourceData.IsBulkDataLoaded())
+		if (bDiscardData && SourceData.IsBulkDataLoaded())
 		{
 			// NOTE(mv): This assertion will fail in the case where the bulk data is still available even though the bDiscardInternalCopy
 			//           flag is toggled when FUntypedBulkData::CanLoadFromDisk() also fail. This happens when the user tries to allow 
