Replace the use of TArray with a FBulkDataBuffer to avoid extra allocations and memory copies when loading the first chunk of a USoundWave when audio stream caching is enabled.

Possible future work: Extend GetChunkFromDDC to also take a FBulkDataBuffer.

diff --git a/Engine/Source/Runtime/Engine/Classes/Sound/SoundWave.h b/Engine/Source/Runtime/Engine/Classes/Sound/SoundWave.h
--- a/Engine/Source/Runtime/Engine/Classes/Sound/SoundWave.h
+++ b/Engine/Source/Runtime/Engine/Classes/Sound/SoundWave.h
@@ -14,6 +14,7 @@
 #include "Sound/SoundBase.h"
 #include "Sound/SoundClass.h"
 #include "Serialization/BulkData.h"
+#include "Serialization/BulkDataBuffer.h"
 #include "Sound/SoundGroups.h"
 #include "AudioMixerTypes.h"
 #include "AudioCompressionSettings.h"
@@ -612,7 +613,7 @@ public:
 	const uint8* ResourceData;
 
 	/** Zeroth Chunk of audio for sources that use Load On Demand. */
-	TArray<uint8> ZerothChunkData;
+	FBulkDataBuffer<uint8> ZerothChunkData;
 
 	/** Uncompressed wav data 16 bit in mono or stereo - stereo not allowed for multichannel data */
 	FByteBulkData RawData;
diff --git a/Engine/Source/Runtime/Engine/Private/SoundWave.cpp b/Engine/Source/Runtime/Engine/Private/SoundWave.cpp
--- a/Engine/Source/Runtime/Engine/Private/SoundWave.cpp
+++ b/Engine/Source/Runtime/Engine/Private/SoundWave.cpp
@@ -729,7 +729,7 @@ FByteBulkData* USoundWave::GetCompressedData(FName Format, const FPlatformAudioC
 void USoundWave::InvalidateCompressedData(bool bFreeResources, bool bRebuildStreamingChunks)
 {
 	CompressedDataGuid = FGuid::NewGuid();
-	ZerothChunkData.Reset();
+	ZerothChunkData.Empty();
 	CompressedFormatData.FlushData();
 
 	if (bFreeResources)
@@ -877,7 +877,7 @@ void USoundWave::PostLoad()
 void USoundWave::EnsureZerothChunkIsLoaded()
 {
 	// If the zeroth chunk is already loaded, early exit.
-	if (ZerothChunkData.Num() > 0 || !ShouldUseStreamCaching())
+	if (ZerothChunkData.GetView().Num() > 0 || !ShouldUseStreamCaching())
 	{
 		return;
 	}
@@ -897,24 +897,16 @@ void USoundWave::EnsureZerothChunkIsLoaded()
 		return;
 	}
 
-	// TODO: Support passing a TArray by ref into FStreamedAudioPlatformData::GetChunkFromDDC.
-	// Currently not feasible unless FUntypedBulkData::GetCopy API was changed.
-	// For now in editor, we have an extra allocated buffer in this scope.
-	ZerothChunkData.Reset();
-	ZerothChunkData.AddUninitialized(ChunkSizeInBytes);
-	FMemory::Memcpy(ZerothChunkData.GetData(), TempChunkBuffer, ChunkSizeInBytes);
-	FMemory::Free(TempChunkBuffer);
+	ZerothChunkData.Reset(TempChunkBuffer, ChunkSizeInBytes);
+
 #else // WITH_EDITOR
 	// Otherwise, the zeroth chunk is cooked out to RunningPlatformData, and we just need to retrieve it.
 	check(RunningPlatformData && RunningPlatformData->Chunks.Num() > 0);
 	FStreamedAudioChunk& ZerothChunk = RunningPlatformData->Chunks[0];
 	// Some sanity checks to ensure that the bulk size set up
-	check(ZerothChunk.BulkData.GetBulkDataSize() == ZerothChunk.DataSize && ZerothChunk.DataSize >= ZerothChunk.AudioDataSize);
+	check(ZerothChunk.BulkData.GetBulkDataSize() == ZerothChunk.DataSize);
 
-	ZerothChunkData.Reset();
-	ZerothChunkData.AddUninitialized(ZerothChunk.AudioDataSize);
-	uint8* Destination = ZerothChunkData.GetData();
-	ZerothChunk.BulkData.GetCopy((void**)&Destination, true);
+	ZerothChunkData = ZerothChunk.BulkData.GetCopyAsBuffer(ZerothChunk.AudioDataSize, true);
 #endif // WITH_EDITOR
 }
 
@@ -1975,7 +1967,7 @@ TArrayView<const uint8> USoundWave::GetZerothChunk()
 			EnsureZerothChunkIsLoaded();
 		}
 
-		check(ZerothChunkData.Num() > 0);
+		check(ZerothChunkData.GetView().Num() > 0);
 
 		if (GetNumChunks() > 1)
 		{
@@ -1983,7 +1975,7 @@ TArrayView<const uint8> USoundWave::GetZerothChunk()
 			IStreamingManager::Get().GetAudioStreamingManager().RequestChunk(this, 1, [](EAudioChunkLoadResult InResult) {});
 		}
 
-		return TArrayView<const uint8>(ZerothChunkData);
+		return ZerothChunkData.GetView();
 	}
 	else
 	{
