Fix test/shipping builds

diff --git a/Engine/Source/Runtime/Engine/Private/SkeletalMesh.cpp b/Engine/Source/Runtime/Engine/Private/SkeletalMesh.cpp
--- a/Engine/Source/Runtime/Engine/Private/SkeletalMesh.cpp
+++ b/Engine/Source/Runtime/Engine/Private/SkeletalMesh.cpp
@@ -5008,6 +5008,7 @@ void FSkeletalMeshSceneProxy::DebugDrawPhysicsAsset(int32 ViewIndex, FMeshElemen
 
 void FSkeletalMeshSceneProxy::DebugDrawSkeleton(int32 ViewIndex, FMeshElementCollector& Collector, const FEngineShowFlags& EngineShowFlags) const
 {
+#if !(UE_BUILD_SHIPPING || UE_BUILD_TEST)
 	FMatrix ProxyLocalToWorld, WorldToLocal;
 	if (!GetWorldMatrices(ProxyLocalToWorld, WorldToLocal))
 	{
@@ -5062,6 +5063,7 @@ void FSkeletalMeshSceneProxy::DebugDrawSkeleton(int32 ViewIndex, FMeshElementCol
 			}
 		}
 	}
+#endif
 }
 
 /**
diff --git a/Engine/Source/Runtime/Engine/Public/Animation/AnimTrace.h b/Engine/Source/Runtime/Engine/Public/Animation/AnimTrace.h
--- a/Engine/Source/Runtime/Engine/Public/Animation/AnimTrace.h
+++ b/Engine/Source/Runtime/Engine/Public/Animation/AnimTrace.h
@@ -69,7 +69,7 @@ struct FAnimTrace
 
 #define TRACE_ANIM_TICK_RECORDS(Proxy, Component)
 #define TRACE_SKELETAL_MESH(Mesh)
-#define TRACE_SKELETALMESH_POSE(Component)
+#define TRACE_SKELETAL_MESH_COMPONENT(Component)
 #define TRACE_POSE_LINK(Context)
 #define TRACE_ANIM_NODE_VALUE(Context, Key, Value)
 
