PR #6383: Fix missing vtable for FActorTickFunction (Contributed by [REDACTED_PERSON])

diff --git a/Engine/Source/Runtime/Engine/Classes/Engine/EngineBaseTypes.h b/Engine/Source/Runtime/Engine/Classes/Engine/EngineBaseTypes.h
--- a/Engine/Source/Runtime/Engine/Classes/Engine/EngineBaseTypes.h
+++ b/Engine/Source/Runtime/Engine/Classes/Engine/EngineBaseTypes.h
@@ -438,7 +438,7 @@ struct TStructOpsTypeTraits<FTickFunction> : public TStructOpsTypeTraitsBase2<FT
 * Tick function that calls AActor::TickActor
 **/
 USTRUCT()
-struct FActorTickFunction : public FTickFunction
+struct ENGINE_VTABLE FActorTickFunction : public FTickFunction
 {
 	GENERATED_USTRUCT_BODY()
 
