Compile fix for editor

Probably should split off the BulkData utility classes to their own header file to reduce the scope of what we need to include.

#rnx

diff --git a/Engine/Source/Runtime/Engine/Private/Streaming/SkeletalMeshUpdate.h b/Engine/Source/Runtime/Engine/Private/Streaming/SkeletalMeshUpdate.h
--- a/Engine/Source/Runtime/Engine/Private/Streaming/SkeletalMeshUpdate.h
+++ b/Engine/Source/Runtime/Engine/Private/Streaming/SkeletalMeshUpdate.h
@@ -8,7 +8,7 @@ SkeletalMeshUpdate.h: Helpers to stream in and out skeletal mesh LODs.
 
 #include "CoreMinimal.h"
 #include "Engine/SkeletalMesh.h"
-#include "Async/AsyncFileHandle.h"
+#include "Serialization/BulkData.h"
 
 /**
 * A context used to update or proceed with the next update step.
