Reduce cost of GetDependentFunctions on UMaterialFunction.  Specific changes are:

1) Introduce IterateDependentFunctions as a method to reduce temporary allocations when computing dependent functions (original author: PJ Kack, taken from shelved CL6968386)
2) Add transient UMaterialFunction::DependentFunctionExpressionCandidates field (computed in PostLoad and PostEditChangeProperty) wth only function expression types that we might have dependent functions on to reduce the amount of time we spend iterating on FunctionExpressions and casting them to see if they're a type we could have to get dependent functions from.

diff --git a/Engine/Source/Runtime/Engine/Classes/Materials/Material.h b/Engine/Source/Runtime/Engine/Classes/Materials/Material.h
--- a/Engine/Source/Runtime/Engine/Classes/Materials/Material.h
+++ b/Engine/Source/Runtime/Engine/Classes/Materials/Material.h
@@ -1244,28 +1244,24 @@ public:
 	{
 		for (class UMaterialExpression* Expression : Expressions)
 		{
-			const ExpressionType* ParameterExpression = Cast<const ExpressionType>(Expression);
-			const UMaterialExpressionMaterialFunctionCall* FunctionExpression = Cast<const UMaterialExpressionMaterialFunctionCall>(Expression);
-			const UMaterialExpressionMaterialAttributeLayers* LayersExpression = Cast<const UMaterialExpressionMaterialAttributeLayers>(Expression);
-
 			FMaterialParameterInfo BaseParameterInfo;
 			BaseParameterInfo.Association = EMaterialParameterAssociation::GlobalParameter;
 			BaseParameterInfo.Index = INDEX_NONE;
 
 			// Note: Intentionally checking the requested type first as this catches MaterialLayers
 			// which are a top-level only parameter without having to deal with the below recursion
-			if (ParameterExpression)
+			if (const ExpressionType* ParameterExpression = Cast<const ExpressionType>(Expression))
 			{
 				ParameterExpression->GetAllParameterInfo(OutParameterInfo, OutParameterIds, BaseParameterInfo);
 			}
-			else if (FunctionExpression)
+			else if (const UMaterialExpressionMaterialFunctionCall* FunctionExpression = Cast<const UMaterialExpressionMaterialFunctionCall>(Expression))
 			{
 				if (FunctionExpression->MaterialFunction)
 				{
 					FunctionExpression->MaterialFunction->GetAllParameterInfo<ExpressionType>(OutParameterInfo, OutParameterIds, BaseParameterInfo);
 				}
 			}
-			else if (LayersExpression)
+			else if (const UMaterialExpressionMaterialAttributeLayers* LayersExpression = Cast<const UMaterialExpressionMaterialAttributeLayers>(Expression))
 			{
 				const TArray<UMaterialFunctionInterface*>* Layers = &LayersExpression->GetLayers();
 				const TArray<UMaterialFunctionInterface*>* Blends = &LayersExpression->GetBlends();
@@ -1337,6 +1333,7 @@ public:
 	ENGINE_API virtual void GetAllStaticSwitchParameterInfo(TArray<FMaterialParameterInfo>& OutParameterInfo, TArray<FGuid>& OutParameterIds) const override;
 	ENGINE_API virtual void GetAllStaticComponentMaskParameterInfo(TArray<FMaterialParameterInfo>& OutParameterInfo, TArray<FGuid>& OutParameterIds) const override;
 
+	ENGINE_API virtual bool IterateDependentFunctions(TFunctionRef<bool(UMaterialFunctionInterface*)> Predicate) const override;
 	ENGINE_API virtual void GetDependentFunctions(TArray<class UMaterialFunctionInterface*>& DependentFunctions) const override;
 
 	ENGINE_API virtual bool GetScalarParameterDefaultValue(const FMaterialParameterInfo& ParameterInfo, float& OutValue, bool bOveriddenOnly = false, bool bCheckOwnedGlobalOverrides = false) const override;
diff --git a/Engine/Source/Runtime/Engine/Classes/Materials/MaterialExpressionMaterialAttributeLayers.h b/Engine/Source/Runtime/Engine/Classes/Materials/MaterialExpressionMaterialAttributeLayers.h
--- a/Engine/Source/Runtime/Engine/Classes/Materials/MaterialExpressionMaterialAttributeLayers.h
+++ b/Engine/Source/Runtime/Engine/Classes/Materials/MaterialExpressionMaterialAttributeLayers.h
@@ -93,6 +93,7 @@ class UMaterialExpressionMaterialAttributeLayers : public UMaterialExpression
 	ENGINE_API void OverrideLayerGraph(const FMaterialLayersFunctions* OverrideLayers);
 	bool ValidateLayerConfiguration(FMaterialCompiler* Compiler, bool bReportErrors);
 
+	bool IterateDependentFunctions(TFunctionRef<bool(UMaterialFunctionInterface*)> Predicate) const;
 	void GetDependentFunctions(TArray<UMaterialFunctionInterface*>& DependentFunctions) const;
 	UMaterialFunctionInterface* GetParameterAssociatedFunction(const FMaterialParameterInfo& ParameterInfo) const;
 	void GetParameterAssociatedFunctions(const FMaterialParameterInfo& ParameterInfo, TArray<UMaterialFunctionInterface*>& AssociatedFunctions) const;
diff --git a/Engine/Source/Runtime/Engine/Classes/Materials/MaterialExpressionMaterialFunctionCall.h b/Engine/Source/Runtime/Engine/Classes/Materials/MaterialExpressionMaterialFunctionCall.h
--- a/Engine/Source/Runtime/Engine/Classes/Materials/MaterialExpressionMaterialFunctionCall.h
+++ b/Engine/Source/Runtime/Engine/Classes/Materials/MaterialExpressionMaterialFunctionCall.h
@@ -106,6 +106,7 @@ class UMaterialExpressionMaterialFunctionCall : public UMaterialExpression
 	virtual bool NeedsLoadForClient() const override;
 	//~ End UObject Interface.
 
+	ENGINE_API bool IterateDependentFunctions(TFunctionRef<bool(UMaterialFunctionInterface*)> Predicate) const;
 	ENGINE_API void GetDependentFunctions(TArray<UMaterialFunctionInterface*>& DependentFunctions) const;
 
 #if WITH_EDITOR
diff --git a/Engine/Source/Runtime/Engine/Classes/Materials/MaterialFunction.h b/Engine/Source/Runtime/Engine/Classes/Materials/MaterialFunction.h
--- a/Engine/Source/Runtime/Engine/Classes/Materials/MaterialFunction.h
+++ b/Engine/Source/Runtime/Engine/Classes/Materials/MaterialFunction.h
@@ -69,6 +69,9 @@ class UMaterialFunction : public UMaterialFunctionInterface
 
 	UPROPERTY(transient)
 	UMaterial* PreviewMaterial;
+
+	UPROPERTY()
+	TArray<class UMaterialExpressionMaterialFunctionCall*> DependentFunctionExpressionCandidates;
 #endif // WITH_EDITORONLY_DATA
 
 private:
@@ -112,6 +115,15 @@ public:
 	/** @return true if this function is dependent on the passed in function, directly or indirectly. */
 	virtual bool IsDependent(UMaterialFunctionInterface* OtherFunction) override;
 
+	/**
+	 * Iterates all functions that this function is dependent on, directly or indrectly.
+	 *
+	 * @param Predicate a visitor predicate returning true to continue iteration, false to break
+	 *
+	 * @return true if all dependent functions were visited, false if the Predicate did break iteration
+	 */
+	ENGINE_API virtual bool IterateDependentFunctions(TFunctionRef<bool(UMaterialFunctionInterface*)> Predicate) const override;
+
 	/** Returns an array of the functions that this function is dependent on, directly or indirectly. */
 	ENGINE_API virtual void GetDependentFunctions(TArray<UMaterialFunctionInterface*>& DependentFunctions) const override;
 
@@ -123,6 +135,8 @@ public:
 
 	virtual void UpdateInputOutputTypes() override;
 
+	virtual void UpdateDependentFunctionCandidates();
+
 	/**
 	 * Checks whether a Material Function is arranged in the old style, with inputs flowing from right to left
 	 */
diff --git a/Engine/Source/Runtime/Engine/Classes/Materials/MaterialFunctionInstance.h b/Engine/Source/Runtime/Engine/Classes/Materials/MaterialFunctionInstance.h
--- a/Engine/Source/Runtime/Engine/Classes/Materials/MaterialFunctionInstance.h
+++ b/Engine/Source/Runtime/Engine/Classes/Materials/MaterialFunctionInstance.h
@@ -83,6 +83,7 @@ class UMaterialFunctionInstance : public UMaterialFunctionInterface
 #endif
 
 	virtual bool IsDependent(UMaterialFunctionInterface* OtherFunction) override;
+	ENGINE_API virtual bool IterateDependentFunctions(TFunctionRef<bool(UMaterialFunctionInterface*)> Predicate) const override;
 	ENGINE_API virtual void GetDependentFunctions(TArray<UMaterialFunctionInterface*>& DependentFunctions) const override;
 	virtual void AppendReferencedTextures(TArray<UObject*>& InOutTextures) const override;
 
diff --git a/Engine/Source/Runtime/Engine/Classes/Materials/MaterialFunctionInterface.h b/Engine/Source/Runtime/Engine/Classes/Materials/MaterialFunctionInterface.h
--- a/Engine/Source/Runtime/Engine/Classes/Materials/MaterialFunctionInterface.h
+++ b/Engine/Source/Runtime/Engine/Classes/Materials/MaterialFunctionInterface.h
@@ -76,6 +76,16 @@ public:
 	ENGINE_API virtual bool IsDependent(UMaterialFunctionInterface* OtherFunction)
 		PURE_VIRTUAL(UMaterialFunctionInterface::IsDependent,return false;);
 
+	/**
+	 * Iterates all functions that this function is dependent on, directly or indrectly.
+	 *
+	 * @param Predicate a visitor predicate returning true to continue iteration, false to break
+	 *
+	 * @return true if all dependent functions were visited, false if the Predicate did break iteration
+	 */
+	ENGINE_API virtual bool IterateDependentFunctions(TFunctionRef<bool(UMaterialFunctionInterface*)> Predicate) const
+		PURE_VIRTUAL(UMaterialFunctionInterface::IterateDependentFunctions,return false;);
+
 	/** Returns an array of the functions that this function is dependent on, directly or indirectly. */
 	ENGINE_API virtual void GetDependentFunctions(TArray<UMaterialFunctionInterface*>& DependentFunctions) const
 		PURE_VIRTUAL(UMaterialFunctionInterface::GetDependentFunctions,);
@@ -133,9 +143,10 @@ public:
 	{
 		if (const UMaterialFunctionInterface* ParameterFunction = GetBaseFunction())
 		{
+			const UClass* TargetClass = UMaterialExpressionMaterialFunctionCall::StaticClass();
 			for (UMaterialExpression* Expression : *ParameterFunction->GetFunctionExpressions())
 			{
-				if (const UMaterialExpressionMaterialFunctionCall* FunctionExpression = Cast<const UMaterialExpressionMaterialFunctionCall>(Expression))
+				if (const UMaterialExpressionMaterialFunctionCall* FunctionExpression = (Expression && Expression->IsA(TargetClass)) ? (const UMaterialExpressionMaterialFunctionCall*)Expression : nullptr)
 				{
 					if (FunctionExpression->MaterialFunction)
 					{
@@ -160,15 +171,14 @@ public:
 
 		if (UMaterialFunctionInterface* ParameterFunction = GetBaseFunction())
 		{
-			TArray<UMaterialFunctionInterface*> Functions;
-			ParameterFunction->GetDependentFunctions(Functions);
-			Functions.AddUnique(ParameterFunction);
+			const UClass* TargetClass = ExpressionType::StaticClass();
 
-			for (UMaterialFunctionInterface* Function : Functions)
+			auto GetExpressionParameterByNamePredicate = 
+				[&ParameterInfo, &Parameter, &OwningFunction, TargetClass](UMaterialFunctionInterface* Function) -> bool
 			{
 				for (UMaterialExpression* FunctionExpression : *Function->GetFunctionExpressions())
 				{
-					if (ExpressionType* ExpressionParameter = Cast<ExpressionType>(FunctionExpression))
+					if (ExpressionType* ExpressionParameter = (FunctionExpression && FunctionExpression->IsA(TargetClass)) ? (ExpressionType *)FunctionExpression : nullptr)
 					{
 						if (ExpressionParameter->ParameterName == ParameterInfo.Name)
 						{
@@ -179,11 +189,19 @@ public:
 								(*OwningFunction) = Function;
 							}
 
-							return true;
+							return false; // found, stop iterating
 						}
 					}
 				}
+
+				return true; // not found, continue iterating
+			};
+			
+			if (!ParameterFunction->IterateDependentFunctions(GetExpressionParameterByNamePredicate))
+			{
+				return true;
 			}
+			return !GetExpressionParameterByNamePredicate(ParameterFunction);
 		}
 
 		return false;
diff --git a/Engine/Source/Runtime/Engine/Classes/Materials/MaterialInstance.h b/Engine/Source/Runtime/Engine/Classes/Materials/MaterialInstance.h
--- a/Engine/Source/Runtime/Engine/Classes/Materials/MaterialInstance.h
+++ b/Engine/Source/Runtime/Engine/Classes/Materials/MaterialInstance.h
@@ -609,6 +609,7 @@ public:
 	ENGINE_API virtual void GetAllStaticSwitchParameterInfo(TArray<FMaterialParameterInfo>& OutParameterInfo, TArray<FGuid>& OutParameterIds) const override;
 	ENGINE_API virtual void GetAllStaticComponentMaskParameterInfo(TArray<FMaterialParameterInfo>& OutParameterInfo, TArray<FGuid>& OutParameterIds) const override;
 
+	ENGINE_API virtual bool IterateDependentFunctions(TFunctionRef<bool(UMaterialFunctionInterface*)> Predicate) const override;
 	ENGINE_API virtual void GetDependentFunctions(TArray<class UMaterialFunctionInterface*>& DependentFunctions) const override;
 	
 	ENGINE_API virtual bool GetScalarParameterDefaultValue(const FMaterialParameterInfo& ParameterInfo, float& OutValue, bool bOveriddenOnly = false, bool bCheckOwnedGlobalOverrides = false) const override;
diff --git a/Engine/Source/Runtime/Engine/Classes/Materials/MaterialInterface.h b/Engine/Source/Runtime/Engine/Classes/Materials/MaterialInterface.h
--- a/Engine/Source/Runtime/Engine/Classes/Materials/MaterialInterface.h
+++ b/Engine/Source/Runtime/Engine/Classes/Materials/MaterialInterface.h
@@ -504,6 +504,8 @@ public:
 	virtual void GetAllStaticComponentMaskParameterInfo(TArray<FMaterialParameterInfo>& OutParameterInfo, TArray<FGuid>& OutParameterIds) const
 		PURE_VIRTUAL(UMaterialInterface::GetAllStaticComponentMaskParameterInfo,return;);
 
+	virtual bool IterateDependentFunctions(TFunctionRef<bool(UMaterialFunctionInterface*)> Predicate) const
+		PURE_VIRTUAL(UMaterialInterface::IterateDependentFunctions,return false;);
 	virtual void GetDependentFunctions(TArray<class UMaterialFunctionInterface*>& DependentFunctions) const
 		PURE_VIRTUAL(UMaterialInterface::GetDependentFunctions,return;);
 
diff --git a/Engine/Source/Runtime/Engine/Private/Materials/Material.cpp b/Engine/Source/Runtime/Engine/Private/Materials/Material.cpp
--- a/Engine/Source/Runtime/Engine/Private/Materials/Material.cpp
+++ b/Engine/Source/Runtime/Engine/Private/Materials/Material.cpp
@@ -1797,19 +1797,35 @@ void UMaterial::GetAllStaticComponentMaskParameterInfo(TArray<FMaterialParameter
 	GetAllParameterInfo<UMaterialExpressionStaticComponentMaskParameter>(OutParameterInfo, OutParameterIds);
 }
 
-void UMaterial::GetDependentFunctions(TArray<UMaterialFunctionInterface*>& DependentFunctions) const
+bool UMaterial::IterateDependentFunctions(TFunctionRef<bool(UMaterialFunctionInterface*)> Predicate) const
 {
 	for (UMaterialExpression* Expression : Expressions)
 	{
 		if (UMaterialExpressionMaterialFunctionCall* FunctionCall = Cast<UMaterialExpressionMaterialFunctionCall>(Expression))
 		{
-			FunctionCall->GetDependentFunctions(DependentFunctions);
+			if (!FunctionCall->IterateDependentFunctions(Predicate))
+			{
+				return false;
+			}
 		}
 		else if (UMaterialExpressionMaterialAttributeLayers* Layers = Cast<UMaterialExpressionMaterialAttributeLayers>(Expression))
 		{
-			Layers->GetDependentFunctions(DependentFunctions);
+			if (!Layers->IterateDependentFunctions(Predicate))
+			{
+				return false;
+			}
 		}
 	}
+	return true;
+}
+
+void UMaterial::GetDependentFunctions(TArray<UMaterialFunctionInterface*>& DependentFunctions) const
+{
+	IterateDependentFunctions([&DependentFunctions](UMaterialFunctionInterface* MaterialFunction) -> bool
+	{
+		DependentFunctions.AddUnique(MaterialFunction);
+		return true;
+	});
 }
 
 bool UMaterial::GetScalarParameterDefaultValue(const FMaterialParameterInfo& ParameterInfo, float& OutValue, bool bOveriddenOnly, bool bCheckOwnedGlobalOverrides) const
@@ -3245,23 +3261,18 @@ void UMaterial::FlushResourceShaderMaps()
 
 void UMaterial::RebuildMaterialFunctionInfo()
 {	
-	MaterialFunctionInfos.Empty();
+	MaterialFunctionInfos.Reset();
 
 	for (int32 ExpressionIndex = 0; ExpressionIndex < Expressions.Num(); ExpressionIndex++)
 	{
 		UMaterialExpression* Expression = Expressions[ExpressionIndex];
-		UMaterialExpressionMaterialFunctionCall* MaterialFunctionNode = Cast<UMaterialExpressionMaterialFunctionCall>(Expression);
-		UMaterialExpressionMaterialAttributeLayers* MaterialLayersNode = Cast<UMaterialExpressionMaterialAttributeLayers>(Expression);
 
-		if (MaterialFunctionNode)
+		if (UMaterialExpressionMaterialFunctionCall* MaterialFunctionNode = Cast<UMaterialExpressionMaterialFunctionCall>(Expression))
 		{
 			if (MaterialFunctionNode->MaterialFunction)
 			{
-				TArray<UMaterialFunctionInterface*> DependentFunctions;
-				MaterialFunctionNode->GetDependentFunctions(DependentFunctions);
-
-				// Handle nested functions
-				for (UMaterialFunctionInterface* Function : DependentFunctions)
+				MaterialFunctionNode->IterateDependentFunctions(
+					[this](UMaterialFunctionInterface* Function)
 				{
 					FMaterialFunctionInfo NewFunctionInfo;
 					if (Function)
@@ -3270,19 +3281,18 @@ void UMaterial::RebuildMaterialFunctionInfo()
 						NewFunctionInfo.StateId = Function->StateId;
 					}
 					MaterialFunctionInfos.Add(NewFunctionInfo);
-				}	
+					return true;
+				});
 			}
 
 			// Update the function call node, so it can relink inputs and outputs as needed
 			// Update even if MaterialFunctionNode->MaterialFunction is NULL, because we need to remove the invalid inputs in that case
 			MaterialFunctionNode->UpdateFromFunctionResource();
 		}
-		else if (MaterialLayersNode)
+		else if (UMaterialExpressionMaterialAttributeLayers* MaterialLayersNode = Cast<UMaterialExpressionMaterialAttributeLayers>(Expression))
 		{
-			TArray<UMaterialFunctionInterface*> DependentFunctions;
-			MaterialLayersNode->GetDependentFunctions(DependentFunctions);
-
-			for (UMaterialFunctionInterface* Function : DependentFunctions)
+			MaterialLayersNode->IterateDependentFunctions(
+				[this](UMaterialFunctionInterface* Function)
 			{
 				FMaterialFunctionInfo NewFunctionInfo;
 				if (Function)
@@ -3294,7 +3304,8 @@ void UMaterial::RebuildMaterialFunctionInfo()
 					Function->UpdateFromFunctionResource();
 				}
 				MaterialFunctionInfos.Add(NewFunctionInfo);
-			}
+				return true;
+			});
 
 			MaterialLayersNode->RebuildLayerGraph(false);
 		}
diff --git a/Engine/Source/Runtime/Engine/Private/Materials/MaterialExpressions.cpp b/Engine/Source/Runtime/Engine/Private/Materials/MaterialExpressions.cpp
--- a/Engine/Source/Runtime/Engine/Private/Materials/MaterialExpressions.cpp
+++ b/Engine/Source/Runtime/Engine/Private/Materials/MaterialExpressions.cpp
@@ -6545,7 +6545,7 @@ bool UMaterialExpressionMaterialAttributeLayers::ValidateLayerConfiguration(FMat
 #undef COMPILER_OR_LOG_ERROR
 }
 
-void UMaterialExpressionMaterialAttributeLayers::GetDependentFunctions(TArray<UMaterialFunctionInterface*>& DependentFunctions) const
+bool UMaterialExpressionMaterialAttributeLayers::IterateDependentFunctions(TFunctionRef<bool(UMaterialFunctionInterface*)> Predicate) const
 {
 	const TArray<UMaterialFunctionInterface*>& Layers = GetLayers();
 	const TArray<UMaterialFunctionInterface*>& Blends = GetBlends();
@@ -6554,8 +6554,14 @@ void UMaterialExpressionMaterialAttributeLayers::GetDependentFunctions(TArray<UM
 	{
 		if (Layer)
 		{		
-			Layer->GetDependentFunctions(DependentFunctions);
-			DependentFunctions.AddUnique(Layer);
+			if (!Layer->IterateDependentFunctions(Predicate))
+			{
+				return false;
+			}
+			if (!Predicate(Layer))
+			{
+				return false;
+			}
 		}
 	}
 
@@ -6563,10 +6569,26 @@ void UMaterialExpressionMaterialAttributeLayers::GetDependentFunctions(TArray<UM
 	{
 		if (Blend)
 		{
-			Blend->GetDependentFunctions(DependentFunctions);
-			DependentFunctions.AddUnique(Blend);
+			if (!Blend->IterateDependentFunctions(Predicate))
+			{
+				return false;
+			}
+			if (!Predicate(Blend))
+			{
+				return false;
+			}
 		}
 	}
+	return true;
+}
+
+void UMaterialExpressionMaterialAttributeLayers::GetDependentFunctions(TArray<UMaterialFunctionInterface*>& DependentFunctions) const
+{
+	IterateDependentFunctions([&DependentFunctions](UMaterialFunctionInterface* MaterialFunction) -> bool
+	{
+		DependentFunctions.AddUnique(MaterialFunction);
+		return true;
+	});
 }
 
 UMaterialFunctionInterface* UMaterialExpressionMaterialAttributeLayers::GetParameterAssociatedFunction(const FMaterialParameterInfo& ParameterInfo) const
@@ -7157,16 +7179,18 @@ void UMaterialExpressionParameter::GetAllParameterInfo(TArray<FMaterialParameter
 	FMaterialParameterInfo NewParameter(ParameterName, InBaseParameterInfo.Association, InBaseParameterInfo.Index);
 
 #if WITH_EDITOR
-	NewParameter.ParameterLocation = Material;
 	if (Function != nullptr)
 	{
 		NewParameter.ParameterLocation = Function;
 	}
-
+	else
+	{
+		NewParameter.ParameterLocation = Material;
+	}
 	if (HasConnectedOutputs())
 #endif
 	{
-		OutParameterInfo.AddUnique(NewParameter);
+		OutParameterInfo.AddUnique(MoveTemp(NewParameter));
 		if (CurrentSize != OutParameterInfo.Num())
 		{
 			OutParameterIds.Add(ExpressionGUID);
@@ -11218,13 +11242,26 @@ void UMaterialFunction::UpdateInputOutputTypes()
 		}
 	}
 }
-#endif
 
 #if WITH_EDITOR
+void UMaterialFunction::UpdateDependentFunctionCandidates()
+{
+	DependentFunctionExpressionCandidates.Reset();
+	for (UMaterialExpression* CurrentExpression : FunctionExpressions)
+	{
+		if (UMaterialExpressionMaterialFunctionCall* MaterialFunctionExpression = Cast<UMaterialExpressionMaterialFunctionCall>(CurrentExpression))
+		{
+			DependentFunctionExpressionCandidates.Add(MaterialFunctionExpression);
+		}
+	}
+}
+#endif
+
 void UMaterialFunction::PostEditChangeProperty(FPropertyChangedEvent& PropertyChangedEvent)
 {
 #if WITH_EDITORONLY_DATA
 	UpdateInputOutputTypes();
+	UpdateDependentFunctionCandidates();
 #endif
 
 	//@todo - recreate guid only when needed, not when a comment changes
@@ -11297,6 +11334,7 @@ void UMaterialFunction::PostLoad()
 	{
 		UpdateInputOutputTypes();
 	}
+	UpdateDependentFunctionCandidates();
 
 	if (GIsEditor)
 	{
@@ -11699,15 +11737,39 @@ bool UMaterialFunction::IsDependent(UMaterialFunctionInterface* OtherFunction)
 	return bIsDependent;
 }
 
-void UMaterialFunction::GetDependentFunctions(TArray<UMaterialFunctionInterface*>& DependentFunctions) const
+bool UMaterialFunction::IterateDependentFunctions(TFunctionRef<bool(UMaterialFunctionInterface*)> Predicate) const
 {
+#if WITH_EDITOR
+	check(!HasAnyFlags(RF_NeedPostLoad));
+	for (UMaterialExpressionMaterialFunctionCall* MaterialFunctionExpression : DependentFunctionExpressionCandidates)
+	{
+		if(!MaterialFunctionExpression->IterateDependentFunctions(Predicate))
+		{
+			return false;
+		}
+	}
+#else
 	for (UMaterialExpression* CurrentExpression : FunctionExpressions)
 	{
 		if (UMaterialExpressionMaterialFunctionCall* MaterialFunctionExpression = Cast<UMaterialExpressionMaterialFunctionCall>(CurrentExpression))
 		{
-			MaterialFunctionExpression->GetDependentFunctions(DependentFunctions);
+			if (!MaterialFunctionExpression->IterateDependentFunctions(Predicate))
+			{
+				return false;
+			}
 		}
 	}
+#endif
+	return true;
+}
+
+void UMaterialFunction::GetDependentFunctions(TArray<UMaterialFunctionInterface*>& DependentFunctions) const
+{
+	IterateDependentFunctions([&DependentFunctions](UMaterialFunctionInterface* MaterialFunction) -> bool
+	{
+		DependentFunctions.AddUnique(MaterialFunction);
+		return true;
+	});
 }
 
 void UMaterialFunction::AppendReferencedTextures(TArray<UObject*>& InOutTextures) const
@@ -12241,13 +12303,29 @@ bool UMaterialFunctionInstance::IsDependent(UMaterialFunctionInterface* OtherFun
 	return Parent ? Parent->IsDependent(OtherFunction) : false;
 }
 
-void UMaterialFunctionInstance::GetDependentFunctions(TArray<UMaterialFunctionInterface*>& DependentFunctions) const
+bool UMaterialFunctionInstance::IterateDependentFunctions(TFunctionRef<bool(UMaterialFunctionInterface*)> Predicate) const
 {
 	if (Parent)
 	{
-		Parent->GetDependentFunctions(DependentFunctions);
-		DependentFunctions.AddUnique(Parent);
+		if (!Parent->IterateDependentFunctions(Predicate))
+		{
+			return false;
+		}
+		if (!Predicate(Parent))
+		{
+			return false;
+		}
 	}
+	return true;
+}
+
+void UMaterialFunctionInstance::GetDependentFunctions(TArray<UMaterialFunctionInterface*>& DependentFunctions) const
+{
+	IterateDependentFunctions([&DependentFunctions](UMaterialFunctionInterface* MaterialFunction) -> bool
+	{
+		DependentFunctions.AddUnique(MaterialFunction);
+		return true;
+	});
 }
 
 void UMaterialFunctionInstance::AppendReferencedTextures(TArray<UObject*>& InOutTextures) const
@@ -12575,13 +12653,30 @@ bool UMaterialExpressionMaterialFunctionCall::NeedsLoadForClient() const
 	return true;
 }
 
-void UMaterialExpressionMaterialFunctionCall::GetDependentFunctions(TArray<UMaterialFunctionInterface*>& DependentFunctions) const
+
+bool UMaterialExpressionMaterialFunctionCall::IterateDependentFunctions(TFunctionRef<bool(UMaterialFunctionInterface*)> Predicate) const
 {
 	if (MaterialFunction)
 	{
-		MaterialFunction->GetDependentFunctions(DependentFunctions);
-		DependentFunctions.AddUnique(MaterialFunction);
+		if (!MaterialFunction->IterateDependentFunctions(Predicate))
+		{
+			return false;
+		}
+		if (!Predicate(MaterialFunction))
+		{
+			return false;
+		}
 	}
+	return true;
+}
+
+void UMaterialExpressionMaterialFunctionCall::GetDependentFunctions(TArray<UMaterialFunctionInterface*>& DependentFunctions) const
+{
+	IterateDependentFunctions([&DependentFunctions](UMaterialFunctionInterface* InMaterialFunction) -> bool
+	{
+		DependentFunctions.AddUnique(InMaterialFunction);
+		return true;
+	});
 }
 
 #if WITH_EDITOR
@@ -12926,13 +13021,12 @@ bool UMaterialExpressionMaterialFunctionCall::SetMaterialFunctionEx(
 
 void UMaterialExpressionMaterialFunctionCall::UpdateFromFunctionResource(bool bRecreateAndLinkNode)
 {
-	TArray<FFunctionExpressionInput> OriginalInputs = FunctionInputs;
-	TArray<FFunctionExpressionOutput> OriginalOutputs = FunctionOutputs;
-	TArray<FExpressionOutput> OriginalGraphOutputs = Outputs;
+	TArray<FFunctionExpressionInput> OriginalInputs = MoveTemp(FunctionInputs);
+	TArray<FFunctionExpressionOutput> OriginalOutputs = MoveTemp(FunctionOutputs);
 
-	FunctionInputs.Empty();
-	FunctionOutputs.Empty();
-	Outputs.Empty();
+	FunctionInputs.Reserve(OriginalInputs.Num());
+	FunctionOutputs.Reserve(OriginalOutputs.Num());
+	Outputs.Reset();
 
 	if (MaterialFunction)
 	{
diff --git a/Engine/Source/Runtime/Engine/Private/Materials/MaterialInstance.cpp b/Engine/Source/Runtime/Engine/Private/Materials/MaterialInstance.cpp
--- a/Engine/Source/Runtime/Engine/Private/Materials/MaterialInstance.cpp
+++ b/Engine/Source/Runtime/Engine/Private/Materials/MaterialInstance.cpp
@@ -2145,7 +2145,7 @@ void UMaterialInstance::GetAllStaticComponentMaskParameterInfo(TArray<FMaterialP
 	}
 }
 
-void UMaterialInstance::GetDependentFunctions(TArray<UMaterialFunctionInterface*>& DependentFunctions) const
+bool UMaterialInstance::IterateDependentFunctions(TFunctionRef<bool(UMaterialFunctionInterface*)> Predicate) const
 {
 	// Important that local function references are listed first so that traversing for a parameter
 	// value we always hit the highest material in the hierarchy that can give us a valid value
@@ -2157,7 +2157,10 @@ void UMaterialInstance::GetDependentFunctions(TArray<UMaterialFunctionInterface*
 			{
 				if (Layer)
 				{
-					DependentFunctions.AddUnique(Layer);
+					if (!Predicate(Layer))
+					{
+						return false;
+					}
 				}
 			}
 
@@ -2165,16 +2168,25 @@ void UMaterialInstance::GetDependentFunctions(TArray<UMaterialFunctionInterface*
 			{
 				if (Blend)
 				{
-					DependentFunctions.AddUnique(Blend);
+					if (!Predicate(Blend))
+					{
+						return false;
+					}
 				}
 			}
 		}
 	}
 
-	if (Parent)
-	{
-		Parent->GetDependentFunctions(DependentFunctions);
-	}
+	return Parent ? Parent->IterateDependentFunctions(Predicate) : true;
+}
+
+void UMaterialInstance::GetDependentFunctions(TArray<UMaterialFunctionInterface*>& DependentFunctions) const
+{
+	IterateDependentFunctions([&DependentFunctions](UMaterialFunctionInterface* MaterialFunction) -> bool
+		{
+			DependentFunctions.AddUnique(MaterialFunction);
+			return true;
+		});
 }
 
 bool UMaterialInstance::GetScalarParameterDefaultValue(const FMaterialParameterInfo& ParameterInfo, float& OutValue, bool bOveriddenOnly, bool bCheckOwnedGlobalOverrides) const
