Anim graph tracing, anim BP debugger connection + misc. fixes

Added tracing of anim nodes, anim node values, anim graphs, asset players, blendspace players and state machines. Moved tick records to trace at their point of origin, rather than all at once.
Duplicated all GatherDebugData() data into TRACE_ANIM_NODE_VALUE calls. Now the traced node info can directly emulate ShowDebug ANIMATION.
Various fixes to correctly get the node IDs as the graph is traversed (including into and out of linked anim instances). This includes inserting node indices into linked anim graph/layer nodes in the compiler.
Added tracing files for AnimGraphRuntime (for blend spaces).
Added blendspace debugging pop-out for blendspace nodes. This re-uses the blendspace widget from the blendspace editor.
Fixed cleanup and visibility issues with skeletal mesh pose tracks.
Moved asset player progress and state machine blend state to use their specific recorded data, rather than going to the source instance. This brings them into line with the node traces used for drawing weighted graph connections.

diff --git a/Engine/Plugins/Animation/GameplayInsights/Source/GameplayInsights/GameplayInsights.Build.cs b/Engine/Plugins/Animation/GameplayInsights/Source/GameplayInsights/GameplayInsights.Build.cs
--- a/Engine/Plugins/Animation/GameplayInsights/Source/GameplayInsights/GameplayInsights.Build.cs
+++ b/Engine/Plugins/Animation/GameplayInsights/Source/GameplayInsights/GameplayInsights.Build.cs
@@ -16,6 +16,7 @@ namespace UnrealBuildTool.Rules
 				"TraceAnalysis",
 				"TraceServices",
 				"TraceInsights",
+				"AssetRegistry",
 			});
 
 			if (Target.bCompileAgainstEngine)
@@ -28,11 +29,13 @@ namespace UnrealBuildTool.Rules
 				});
 			}
 
-			if (Target.bBuildEditor)
+			if (Target.Type == TargetType.Editor)
 			{
 				PrivateDependencyModuleNames.AddRange(
 				new string[]
 				{
+					"AnimationBlueprintEditor",
+					"Persona",
 					"UnrealEd",
 				});
 			}
diff --git a/Engine/Plugins/Animation/GameplayInsights/Source/GameplayInsights/Private/AnimNodesTrack.cpp b/Engine/Plugins/Animation/GameplayInsights/Source/GameplayInsights/Private/AnimNodesTrack.cpp
new file mode 100644
--- /dev/null
+++ b/Engine/Plugins/Animation/GameplayInsights/Source/GameplayInsights/Private/AnimNodesTrack.cpp
@@ -0,0 +1,423 @@
+// Copyright 1998-2019 Epic Games, Inc. All Rights Reserved.
+
+#include "AnimNodesTrack.h"
+#include "Insights/ViewModels/ITimingViewDrawHelper.h"
+#include "GameplayProvider.h"
+#include "AnimationProvider.h"
+#include "Insights/ViewModels/TimingTrackViewport.h"
+#include "Insights/ViewModels/TimingEvent.h"
+#include "AnimationSharedData.h"
+#include "Framework/Multibox/MultiBoxBuilder.h"
+#include "Insights/ViewModels/TimingEventSearch.h"
+#include "Insights/ViewModels/TooltipDrawState.h"
+#include "Insights/Common/TimeUtils.h"
+#include "TraceServices/Model/Frames.h"
+#include "SkeletalMeshPoseTrack.h"
+#include "GameplayTimingViewExtender.h"
+
+#if WITH_ENGINE
+#include "Animation/AnimTrace.h"
+#include "Components/SkeletalMeshComponent.h"
+#endif
+
+#if WITH_EDITOR
+#include "Animation/AnimBlueprintGeneratedClass.h"
+#include "IAnimationBlueprintEditor.h"
+#include "IPersonaToolkit.h"
+#include "Animation/AnimBlueprint.h"
+#include "Subsystems/AssetEditorSubsystem.h"
+#include "Animation/AnimInstance.h"
+#include "Animation/BlendSpaceBase.h"
+#endif
+
+#define LOCTEXT_NAMESPACE "AnimNodesTrack"
+
+const FName FAnimNodesTrack::TypeName(TEXT("Events"));
+const FName FAnimNodesTrack::SubTypeName(TEXT("Animation.AnimNodes"));
+
+FAnimNodesTrack::FAnimNodesTrack(const FAnimationSharedData& InSharedData, uint64 InObjectID, const TCHAR* InName)
+	: TGameplayTrackMixin<FTimingEventsTrack>(InObjectID, FAnimNodesTrack::TypeName, FAnimNodesTrack::SubTypeName, FText::Format(LOCTEXT("TrackNameFormat", "Graph - {0}"), FText::FromString(FString(InName))))
+	, SharedData(InSharedData)
+#if WITH_EDITOR
+	, InstanceClass(nullptr)
+	, AnimInstance(nullptr)
+#endif
+{
+#if WITH_EDITOR
+	const FGameplayProvider* GameplayProvider = SharedData.GetAnalysisSession().ReadProvider<FGameplayProvider>(FGameplayProvider::ProviderName);
+
+	if(GameplayProvider)
+	{
+		const FObjectInfo* AnimInstanceInfo = GameplayProvider->FindObjectInfo(GetGameplayTrack().GetObjectId());
+		if(AnimInstanceInfo)
+		{
+			const FClassInfo* AnimInstanceClassInfo = GameplayProvider->FindClassInfo(AnimInstanceInfo->ClassId);
+			if(AnimInstanceClassInfo)
+			{
+				InstanceClass = FindObject<UAnimBlueprintGeneratedClass>(ANY_PACKAGE, AnimInstanceClassInfo->PathName);
+			}
+		}
+	}
+#endif
+}
+
+static const TCHAR* GetPhaseName(EAnimGraphPhase InPhase)
+{
+#if WITH_ENGINE
+	static_assert((__underlying_type(EAnimGraphPhase))EAnimGraphPhase::Initialize == (__underlying_type(FAnimTrace::EPhase))FAnimTrace::EPhase::Initialize, "EAnimGraphPhase and FAnimTrace::EPhase must be kept in sync");
+	static_assert((__underlying_type(EAnimGraphPhase))EAnimGraphPhase::PreUpdate == (__underlying_type(FAnimTrace::EPhase))FAnimTrace::EPhase::PreUpdate, "EAnimGraphPhase and FAnimTrace::EPhase must be kept in sync");
+	static_assert((__underlying_type(EAnimGraphPhase))EAnimGraphPhase::Update == (__underlying_type(FAnimTrace::EPhase))FAnimTrace::EPhase::Update, "EAnimGraphPhase and FAnimTrace::EPhase must be kept in sync");
+	static_assert((__underlying_type(EAnimGraphPhase))EAnimGraphPhase::CacheBones == (__underlying_type(FAnimTrace::EPhase))FAnimTrace::EPhase::CacheBones, "EAnimGraphPhase and FAnimTrace::EPhase must be kept in sync");
+	static_assert((__underlying_type(EAnimGraphPhase))EAnimGraphPhase::Evaluate == (__underlying_type(FAnimTrace::EPhase))FAnimTrace::EPhase::Evaluate, "EAnimGraphPhase and FAnimTrace::EPhase must be kept in sync");
+#endif
+
+	switch(InPhase)
+	{
+	case EAnimGraphPhase::Initialize:
+		return TEXT("Initialize");
+	case EAnimGraphPhase::PreUpdate:
+		return TEXT("PreUpdate");
+	case EAnimGraphPhase::Update:
+		return TEXT("Update");
+	case EAnimGraphPhase::CacheBones:
+		return TEXT("CacheBones");
+	case EAnimGraphPhase::Evaluate:
+		return TEXT("Evaluate");
+	}
+
+	return TEXT("Unknown");
+}
+
+void FAnimNodesTrack::BuildDrawState(ITimingEventsTrackDrawStateBuilder& Builder, const ITimingTrackUpdateContext& Context)
+{
+	const FAnimationProvider* AnimationProvider = SharedData.GetAnalysisSession().ReadProvider<FAnimationProvider>(FAnimationProvider::ProviderName);
+
+	if(AnimationProvider)
+	{
+		Trace::FAnalysisSessionReadScope SessionReadScope(SharedData.GetAnalysisSession());
+
+		AnimationProvider->ReadAnimGraphTimeline(GetGameplayTrack().GetObjectId(), [&Context, &Builder](const FAnimationProvider::AnimGraphTimeline& InTimeline)
+		{
+			InTimeline.EnumerateEvents(Context.GetViewport().GetStartTime(), Context.GetViewport().GetEndTime(), [&Builder](double InStartTime, double InEndTime, uint32 InDepth, const FAnimGraphMessage& InMessage)
+			{
+				Builder.AddEvent(InStartTime, InEndTime, 0, GetPhaseName(InMessage.Phase));
+			});
+		});
+	}
+}
+
+void FAnimNodesTrack::Draw(const ITimingTrackDrawContext& Context) const
+{
+	DrawEvents(Context);
+	GetGameplayTrack().DrawHeaderForTimingTrack(Context, *this, false);
+}
+
+void FAnimNodesTrack::InitTooltip(FTooltipDrawState& Tooltip, const ITimingEvent& HoveredTimingEvent) const
+{
+	FTimingEventSearchParameters SearchParameters(HoveredTimingEvent.GetStartTime(), HoveredTimingEvent.GetEndTime(), ETimingEventSearchFlags::StopAtFirstMatch);
+
+	FindAnimGraphMessage(SearchParameters, [this, &Tooltip](double InFoundStartTime, double InFoundEndTime, uint32 InFoundDepth, const FAnimGraphMessage& InMessage)
+	{
+		Tooltip.ResetContent();
+
+		Tooltip.AddTitle(FText::Format(LOCTEXT("GraphPhaseFormat", "{0} Anim Graph"), FText::FromString(GetPhaseName(InMessage.Phase))).ToString());
+
+		Tooltip.AddNameValueTextLine(LOCTEXT("EventDuration", "Duration").ToString(), TimeUtils::FormatTimeAuto(InFoundEndTime - InFoundStartTime));
+		Tooltip.AddNameValueTextLine(LOCTEXT("EventTime", "Time").ToString(), TimeUtils::FormatTimeAuto(InFoundStartTime));
+
+		Tooltip.UpdateLayout();
+	});
+}
+
+const TSharedPtr<const ITimingEvent> FAnimNodesTrack::SearchEvent(const FTimingEventSearchParameters& InSearchParameters) const
+{
+	TSharedPtr<const ITimingEvent> FoundEvent;
+
+	FindAnimGraphMessage(InSearchParameters, [this, &FoundEvent](double InFoundStartTime, double InFoundEndTime, uint32 InFoundDepth, const FAnimGraphMessage& InFoundMessage)
+	{
+		FoundEvent = MakeShared<const FTimingEvent>(SharedThis(this), InFoundStartTime, InFoundEndTime, InFoundDepth);
+	});
+
+	return FoundEvent;
+}
+
+void FAnimNodesTrack::FindAnimGraphMessage(const FTimingEventSearchParameters& InParameters, TFunctionRef<void(double, double, uint32, const FAnimGraphMessage&)> InFoundPredicate) const
+{
+	TTimingEventSearch<FAnimGraphMessage>::Search(
+		InParameters,
+
+		[this](TTimingEventSearch<FAnimGraphMessage>::FContext& InContext)
+		{
+			const FAnimationProvider* AnimationProvider = SharedData.GetAnalysisSession().ReadProvider<FAnimationProvider>(FAnimationProvider::ProviderName);
+
+			if(AnimationProvider)
+			{
+				Trace::FAnalysisSessionReadScope SessionReadScope(SharedData.GetAnalysisSession());
+
+				AnimationProvider->ReadAnimGraphTimeline(GetGameplayTrack().GetObjectId(), [&InContext](const FAnimationProvider::AnimGraphTimeline& InTimeline)
+				{
+					InTimeline.EnumerateEvents(InContext.GetParameters().StartTime, InContext.GetParameters().EndTime, [&InContext](double InEventStartTime, double InEventEndTime, uint32 InDepth, const FAnimGraphMessage& InMessage)
+					{
+						InContext.Check(InEventStartTime, InEventEndTime, 0, InMessage);
+					});
+				});
+			}
+		},
+
+		[&InFoundPredicate](double InFoundStartTime, double InFoundEndTime, uint32 InFoundDepth, const FAnimGraphMessage& InEvent)
+		{
+			InFoundPredicate(InFoundStartTime, InFoundEndTime, InFoundDepth, InEvent);
+		});
+}
+
+void FAnimNodesTrack::BuildContextMenu(FMenuBuilder& MenuBuilder)
+{
+#if WITH_EDITOR
+	MenuBuilder.BeginSection(TEXT("DebugSection"), LOCTEXT("Debug", "Debug"));
+	{
+		MenuBuilder.AddMenuEntry(
+			LOCTEXT("ToggleDebug", "Debug this graph"),
+			LOCTEXT("ToggleDebug_Tooltip", "Debug this graph in the animation blueprint editor, opens editor for asset if it exists"),
+			FSlateIcon(),
+			FUIAction(
+				FExecuteAction::CreateLambda([this]()
+				{  
+					if(InstanceClass.Get())
+					{
+						if(UAnimBlueprint* AnimBlueprint = Cast<UAnimBlueprint>(InstanceClass->ClassGeneratedBy))
+						{
+							GEditor->GetEditorSubsystem<UAssetEditorSubsystem>()->OpenEditorForAsset(AnimBlueprint);
+
+							const FGameplayProvider* GameplayProvider = SharedData.GetAnalysisSession().ReadProvider<FGameplayProvider>(FGameplayProvider::ProviderName);
+							if(GameplayProvider)
+							{
+								Trace::FAnalysisSessionReadScope SessionReadScope(SharedData.GetAnalysisSession());
+
+								if(const FObjectInfo* AnimInstanceObjectInfo = GameplayProvider->FindObjectInfo(GetGameplayTrack().GetObjectId()))
+								{
+									// find a counterpart skeletal mesh pose track
+									TSharedPtr<FSkeletalMeshPoseTrack> SkeletalMeshPoseTrack = SharedData.FindSkeletalMeshPoseTrack(AnimInstanceObjectInfo->OuterId);
+									if(SkeletalMeshPoseTrack.IsValid())
+									{
+										SkeletalMeshPoseTrack->MarkPotentiallyDebugged();
+
+										if(USkeletalMeshComponent* SkeletalMeshComponent = SkeletalMeshPoseTrack->GetComponent(FGameplayTimingViewExtender::GetWorldToVisualize()))
+										{
+											UAnimInstance* Instance = LazyCreateAnimInstance(SkeletalMeshComponent);
+											AnimBlueprint->SetObjectBeingDebugged(Instance);
+										}
+									}
+								}
+							}
+						}
+					}
+				}),
+				FCanExecuteAction(),
+				FIsActionChecked::CreateLambda([this]()
+				{
+					if(InstanceClass.Get())
+					{
+						if(UAnimBlueprint* AnimBlueprint = Cast<UAnimBlueprint>(InstanceClass->ClassGeneratedBy))
+						{
+							const FGameplayProvider* GameplayProvider = SharedData.GetAnalysisSession().ReadProvider<FGameplayProvider>(FGameplayProvider::ProviderName);
+							if(GameplayProvider)
+							{
+								Trace::FAnalysisSessionReadScope SessionReadScope(SharedData.GetAnalysisSession());
+
+								if(const FObjectInfo* AnimInstanceObjectInfo = GameplayProvider->FindObjectInfo(GetGameplayTrack().GetObjectId()))
+								{
+									// find a counterpart skeletal mesh pose track
+									TSharedPtr<FSkeletalMeshPoseTrack> SkeletalMeshPoseTrack = SharedData.FindSkeletalMeshPoseTrack(AnimInstanceObjectInfo->OuterId);
+									if(SkeletalMeshPoseTrack.IsValid())
+									{
+										if(USkeletalMeshComponent* SkeletalMeshComponent = SkeletalMeshPoseTrack->GetComponent(FGameplayTimingViewExtender::GetWorldToVisualize()))
+										{
+											UAnimInstance* Instance = LazyCreateAnimInstance(SkeletalMeshComponent);
+											AnimBlueprint->IsObjectBeingDebugged(Instance);
+										}
+									}
+								}
+							}
+						}
+					}
+					return false;
+				})),
+			NAME_None,
+			EUserInterfaceActionType::ToggleButton
+		);
+	}
+	MenuBuilder.EndSection();
+#endif
+}
+
+#if WITH_EDITOR
+
+UAnimInstance* FAnimNodesTrack::LazyCreateAnimInstance(USkeletalMeshComponent* InComponent)
+{
+	if(InstanceClass.Get() != nullptr)
+	{
+		if(AnimInstance == nullptr)
+		{
+			AnimInstance = NewObject<UAnimInstance>(InComponent, InstanceClass.Get());
+		}
+		return AnimInstance;
+	}
+
+	return nullptr;
+}
+
+void FAnimNodesTrack::UpdateDebugData(double InTime)
+{
+	if(InstanceClass.Get())
+	{
+		const FAnimationProvider* AnimationProvider = SharedData.GetAnalysisSession().ReadProvider<FAnimationProvider>(FAnimationProvider::ProviderName);
+		const FGameplayProvider* GameplayProvider = SharedData.GetAnalysisSession().ReadProvider<FGameplayProvider>(FGameplayProvider::ProviderName);
+
+		if(AnimationProvider && GameplayProvider)
+		{
+			Trace::FAnalysisSessionReadScope SessionReadScope(SharedData.GetAnalysisSession());
+
+			AnimationProvider->ReadAnimGraphTimeline(GetGameplayTrack().GetObjectId(), [this, AnimationProvider, GameplayProvider, InTime](const FAnimationProvider::AnimGraphTimeline& InGraphTimeline)
+			{
+				const Trace::IFrameProvider& FramesProvider = Trace::ReadFrameProvider(SharedData.GetAnalysisSession());
+
+
+				FAnimBlueprintDebugData& DebugData = InstanceClass->GetAnimBlueprintDebugData();
+				DebugData.ResetNodeVisitSites();
+
+				// round to nearest frame boundary
+				Trace::FFrame Frame;
+				if(FramesProvider.GetFrameFromTime(ETraceFrameType::TraceFrameType_Game, InTime, Frame))
+				{
+					InGraphTimeline.EnumerateEvents(Frame.StartTime, Frame.EndTime, [this, AnimationProvider, GameplayProvider, &DebugData](double InGraphStartTime, double InGraphEndTime, uint32 InDepth, const FAnimGraphMessage& InMessage)
+					{
+						// Check for an update phase (which contains weights)
+						if(InMessage.Phase == EAnimGraphPhase::Update)
+						{
+							// Basic verification - check node count is the same
+							// @TODO: could add some form of node hash/CRC to the class to improve this
+							if(InMessage.NodeCount == InstanceClass->AnimNodeProperties.Num())
+							{
+								AnimationProvider->ReadAnimNodesTimeline(GetGameplayTrack().GetObjectId(), [InGraphStartTime, InGraphEndTime, &DebugData](const FAnimationProvider::AnimNodesTimeline& InNodesTimeline)
+								{
+									InNodesTimeline.EnumerateEvents(InGraphStartTime, InGraphEndTime, [&DebugData](double InStartTime, double InEndTime, uint32 InDepth, const FAnimNodeMessage& InMessage)
+									{
+										DebugData.RecordNodeVisit(InMessage.NodeId, InMessage.PreviousNodeId, InMessage.Weight);
+									});
+								});
+
+								AnimationProvider->ReadStateMachinesTimeline(GetGameplayTrack().GetObjectId(), [InGraphStartTime, InGraphEndTime, &DebugData](const FAnimationProvider::StateMachinesTimeline& InStateMachinesTimeline)
+								{
+									InStateMachinesTimeline.EnumerateEvents(InGraphStartTime, InGraphEndTime, [&DebugData](double InStartTime, double InEndTime, uint32 InDepth, const FAnimStateMachineMessage& InMessage)
+									{
+										DebugData.RecordStateData(InMessage.StateMachineIndex, InMessage.StateIndex, InMessage.StateWeight, InMessage.ElapsedTime);
+									});
+								});
+
+								AnimationProvider->ReadAnimNodeValuesTimeline(GetGameplayTrack().GetObjectId(), [InGraphStartTime, InGraphEndTime, GameplayProvider, &DebugData](const FAnimationProvider::AnimNodeValuesTimeline& InNodeValuesTimeline)
+								{
+									InNodeValuesTimeline.EnumerateEvents(InGraphStartTime, InGraphEndTime, [GameplayProvider, &DebugData](double InStartTime, double InEndTime, uint32 InDepth, const FAnimNodeValueMessage& InMessage)
+									{
+										FText Text;
+										const FText KeyValueFormat(LOCTEXT("KeyValueFormat", "{0} = {1}"));
+
+										switch(InMessage.Type)
+										{
+										case EAnimNodeValueType::Bool:
+											Text = FText::Format(KeyValueFormat, FText::FromString(InMessage.Key), InMessage.Bool.bValue ? LOCTEXT("True", "true") : LOCTEXT("False", "false"));
+											break;
+										case EAnimNodeValueType::Int32:
+											Text = FText::Format(KeyValueFormat, FText::FromString(InMessage.Key), FText::AsNumber(InMessage.Int32.Value));
+											break;
+										case EAnimNodeValueType::Float:
+											Text = FText::Format(KeyValueFormat, FText::FromString(InMessage.Key), FText::AsNumber(InMessage.Float.Value));
+											break;
+										case EAnimNodeValueType::Vector:
+											Text = FText::Format(KeyValueFormat, FText::FromString(InMessage.Key), FText::Format(LOCTEXT("VectorFormat", "({0}, {1}, {2})"), FText::AsNumber(InMessage.Vector.Value.X), FText::AsNumber(InMessage.Vector.Value.Y), FText::AsNumber(InMessage.Vector.Value.Z)));
+											break;
+										case EAnimNodeValueType::String:
+											Text = FText::Format(KeyValueFormat, FText::FromString(InMessage.Key), FText::FromString(InMessage.String.Value));
+											break;
+										case EAnimNodeValueType::Object:
+										{
+											const FObjectInfo& ObjectInfo = GameplayProvider->GetObjectInfo(InMessage.Object.Value);
+											Text = FText::Format(KeyValueFormat, FText::FromString(InMessage.Key), FText::FromString(ObjectInfo.PathName));
+											break;
+										}
+										case EAnimNodeValueType::Class:
+										{
+											const FClassInfo& ClassInfo = GameplayProvider->GetClassInfo(InMessage.Class.Value);
+											Text = FText::Format(KeyValueFormat, FText::FromString(InMessage.Key), FText::FromString(ClassInfo.PathName));
+											break;
+										}
+										}
+
+										DebugData.RecordNodeValue(InMessage.NodeId, Text.ToString());
+									});
+								});
+
+								AnimationProvider->ReadAnimSequencePlayersTimeline(GetGameplayTrack().GetObjectId(), [InGraphStartTime, InGraphEndTime, GameplayProvider, &DebugData](const FAnimationProvider::AnimSequencePlayersTimeline& InSequencePlayersTimeline)
+								{
+									InSequencePlayersTimeline.EnumerateEvents(InGraphStartTime, InGraphEndTime, [&DebugData](double InStartTime, double InEndTime, uint32 InDepth, const FAnimSequencePlayerMessage& InMessage)
+									{
+										DebugData.RecordSequencePlayer(InMessage.NodeId, InMessage.Position, InMessage.Length, InMessage.FrameCount);
+									});
+								});
+
+								AnimationProvider->ReadAnimBlendSpacePlayersTimeline(GetGameplayTrack().GetObjectId(), [InGraphStartTime, InGraphEndTime, GameplayProvider, &DebugData](const FAnimationProvider::BlendSpacePlayersTimeline& InBlendSpacePlayersTimeline)
+								{
+									InBlendSpacePlayersTimeline.EnumerateEvents(InGraphStartTime, InGraphEndTime, [GameplayProvider, &DebugData](double InStartTime, double InEndTime, uint32 InDepth, const FBlendSpacePlayerMessage& InMessage)
+									{
+										UBlendSpaceBase* BlendSpaceBase = nullptr;
+										const FObjectInfo* BlendSpaceInfo = GameplayProvider->FindObjectInfo(InMessage.BlendSpaceId);
+										if(BlendSpaceInfo)
+										{
+											BlendSpaceBase = TSoftObjectPtr<UBlendSpaceBase>(FSoftObjectPath(BlendSpaceInfo->PathName)).LoadSynchronous();
+										}
+
+										DebugData.RecordBlendSpacePlayer(InMessage.NodeId, BlendSpaceBase, InMessage.PositionX, InMessage.PositionY, InMessage.PositionZ);
+									});
+								});
+							}
+						}
+					});
+				}
+			});
+		}
+	}
+}
+
+void FAnimNodesTrack::GetCustomDebugObjects(const IAnimationBlueprintEditor& InAnimationBlueprintEditor, TArray<FCustomDebugObject>& OutDebugList)
+{
+	if(InstanceClass.Get())
+	{
+		if(UAnimBlueprint* AnimBlueprint = Cast<UAnimBlueprint>(InstanceClass->ClassGeneratedBy))
+		{
+			const FGameplayProvider* GameplayProvider = SharedData.GetAnalysisSession().ReadProvider<FGameplayProvider>(FGameplayProvider::ProviderName);
+			if(GameplayProvider)
+			{
+				Trace::FAnalysisSessionReadScope SessionReadScope(SharedData.GetAnalysisSession());
+
+				if(const FObjectInfo* AnimInstanceObjectInfo = GameplayProvider->FindObjectInfo(GetGameplayTrack().GetObjectId()))
+				{
+					// find a counterpart skeletal mesh pose track
+					TSharedPtr<FSkeletalMeshPoseTrack> SkeletalMeshPoseTrack = SharedData.FindSkeletalMeshPoseTrack(AnimInstanceObjectInfo->OuterId);
+					if(SkeletalMeshPoseTrack.IsValid())
+					{
+						SkeletalMeshPoseTrack->MarkPotentiallyDebugged();
+
+						if(USkeletalMeshComponent* SkeletalMeshComponent = SkeletalMeshPoseTrack->GetComponent(FGameplayTimingViewExtender::GetWorldToVisualize()))
+						{
+							UAnimInstance* Instance = LazyCreateAnimInstance(SkeletalMeshComponent);
+							OutDebugList.Emplace(Instance, FText::Format(LOCTEXT("PreviewObjectLabel", "Insights - {0}"), FText::FromString(GetName())).ToString());
+						}
+					}
+				}
+			}
+		}
+	}
+}
+
+#endif
+
+#undef LOCTEXT_NAMESPACE
\ No newline at end of file
diff --git a/Engine/Plugins/Animation/GameplayInsights/Source/GameplayInsights/Private/AnimNodesTrack.h b/Engine/Plugins/Animation/GameplayInsights/Source/GameplayInsights/Private/AnimNodesTrack.h
new file mode 100644
--- /dev/null
+++ b/Engine/Plugins/Animation/GameplayInsights/Source/GameplayInsights/Private/AnimNodesTrack.h
@@ -0,0 +1,62 @@
+// Copyright 1998-2019 Epic Games, Inc. All Rights Reserved.
+
+#pragma once
+
+#include "GameplayTrack.h"
+#include "Insights/ViewModels/TimingEventsTrack.h"
+
+#if WITH_EDITOR
+#include "Animation/AnimBlueprintGeneratedClass.h"
+#endif
+
+class FAnimationSharedData;
+class FTimingEventSearchParameters;
+struct FAnimGraphMessage;
+class IAnimationBlueprintEditor;
+struct FCustomDebugObject;
+class AActor;
+class USkeletalMeshComponent;
+class UAnimInstance;
+
+class FAnimNodesTrack : public TGameplayTrackMixin<FTimingEventsTrack>
+{
+public:
+	static const FName TypeName;
+	static const FName SubTypeName;
+
+	FAnimNodesTrack(const FAnimationSharedData& InSharedData, uint64 InObjectID, const TCHAR* InName);
+
+	virtual void BuildDrawState(ITimingEventsTrackDrawStateBuilder& Builder, const ITimingTrackUpdateContext& Context) override;
+	virtual void Draw(const ITimingTrackDrawContext& Context) const override;
+	virtual void InitTooltip(FTooltipDrawState& Tooltip, const ITimingEvent& HoveredTimingEvent) const override;
+	virtual const TSharedPtr<const ITimingEvent> SearchEvent(const FTimingEventSearchParameters& InSearchParameters) const override;
+	virtual void BuildContextMenu(FMenuBuilder& MenuBuilder) override;
+
+#if WITH_EDITOR
+	// Update the debug data for any classes that are running on this track
+	void UpdateDebugData(double InTime);
+
+	// Get custom debug objects for integration with anim blueprint debugging
+	void GetCustomDebugObjects(const IAnimationBlueprintEditor& InAnimationBlueprintEditor, TArray<FCustomDebugObject>& OutDebugList);
+#endif
+
+private:
+	// Helper function used to find an anim graph message
+	void FindAnimGraphMessage(const FTimingEventSearchParameters& InParameters, TFunctionRef<void(double, double, uint32, const FAnimGraphMessage&)> InFoundPredicate) const;
+
+#if WITH_EDITOR
+	UAnimInstance* LazyCreateAnimInstance(USkeletalMeshComponent* InComponent);
+#endif
+
+private:
+	/** The shared data */
+	const FAnimationSharedData& SharedData;
+
+#if WITH_EDITOR
+	/** Instance class, if any, used for instantiating debug info */
+	TWeakObjectPtr<UAnimBlueprintGeneratedClass> InstanceClass;
+
+	/** Data used for anim BP debugging */
+	UAnimInstance* AnimInstance;
+#endif
+};
\ No newline at end of file
diff --git a/Engine/Plugins/Animation/GameplayInsights/Source/GameplayInsights/Private/AnimationAnalyzer.cpp b/Engine/Plugins/Animation/GameplayInsights/Source/GameplayInsights/Private/AnimationAnalyzer.cpp
--- a/Engine/Plugins/Animation/GameplayInsights/Source/GameplayInsights/Private/AnimationAnalyzer.cpp
+++ b/Engine/Plugins/Animation/GameplayInsights/Source/GameplayInsights/Private/AnimationAnalyzer.cpp
@@ -18,6 +18,19 @@ void FAnimationAnalyzer::OnAnalysisBegin(const FOnAnalysisContext& Context)
 	Builder.RouteEvent(RouteId_TickRecord, "Animation", "TickRecord");
 	Builder.RouteEvent(RouteId_SkeletalMesh, "Animation", "SkeletalMesh");
 	Builder.RouteEvent(RouteId_SkeletalMeshComponent, "Animation", "SkeletalMeshComponent");
+	Builder.RouteEvent(RouteId_SkeletalMeshFrame, "Animation", "SkeletalMeshFrame");
+	Builder.RouteEvent(RouteId_AnimGraph, "Animation", "AnimGraph");
+	Builder.RouteEvent(RouteId_AnimNodeStart, "Animation", "AnimNodeStart");
+	Builder.RouteEvent(RouteId_AnimNodeValueBool, "Animation", "AnimNodeValueBool");
+	Builder.RouteEvent(RouteId_AnimNodeValueInt, "Animation", "AnimNodeValueInt");
+	Builder.RouteEvent(RouteId_AnimNodeValueFloat, "Animation", "AnimNodeValueFloat");
+	Builder.RouteEvent(RouteId_AnimNodeValueVector, "Animation", "AnimNodeValueVector");
+	Builder.RouteEvent(RouteId_AnimNodeValueString, "Animation", "AnimNodeValueString");
+	Builder.RouteEvent(RouteId_AnimNodeValueObject, "Animation", "AnimNodeValueObject");
+	Builder.RouteEvent(RouteId_AnimNodeValueClass, "Animation", "AnimNodeValueClass");
+	Builder.RouteEvent(RouteId_AnimSequencePlayer, "Animation", "AnimSequencePlayer");
+	Builder.RouteEvent(RouteId_BlendSpacePlayer, "Animation", "BlendSpacePlayer");
+	Builder.RouteEvent(RouteId_StateMachineState, "Animation", "StateMachineState");
 	Builder.RouteEvent(RouteId_Name, "Animation", "Name");
 }
 
@@ -31,16 +44,19 @@ bool FAnimationAnalyzer::OnEvent(uint16 RouteId, const FOnEventContext& Context)
 	case RouteId_TickRecord:
 	{
 		uint64 Cycle = EventData.GetValue<uint64>("Cycle");
-		uint64 ComponentId = EventData.GetValue<uint64>("ComponentId");
 		uint64 AnimInstanceId = EventData.GetValue<uint64>("AnimInstanceId");
 		uint64 AssetId = EventData.GetValue<uint64>("AssetId");
+		int32 NodeId = EventData.GetValue<int32>("NodeId");
 		float BlendWeight = EventData.GetValue<float>("BlendWeight");
 		float PlaybackTime = EventData.GetValue<float>("PlaybackTime");
 		float RootMotionWeight = EventData.GetValue<float>("RootMotionWeight");
 		float PlayRate = EventData.GetValue<float>("PlayRate");
+		float BlendSpacePositionX = EventData.GetValue<float>("BlendSpacePositionX");
+		float BlendSpacePositionY = EventData.GetValue<float>("BlendSpacePositionY");
 		uint16 FrameCounter = EventData.GetValue<uint16>("FrameCounter");
 		bool bLooping = EventData.GetValue<bool>("Looping");
-		AnimationProvider.AppendTickRecord(ComponentId, AnimInstanceId, Context.SessionContext.TimestampFromCycle(Cycle), AssetId, BlendWeight, PlaybackTime, RootMotionWeight, PlayRate, FrameCounter, bLooping);
+		bool bIsBlendSpace = EventData.GetValue<bool>("IsBlendSpace");
+		AnimationProvider.AppendTickRecord(AnimInstanceId, Context.SessionContext.TimestampFromCycle(Cycle), AssetId, NodeId, BlendWeight, PlaybackTime, RootMotionWeight, PlayRate, BlendSpacePositionX, BlendSpacePositionY, FrameCounter, bLooping, bIsBlendSpace);
 		break;
 	}
 	case RouteId_SkeletalMesh:
@@ -71,6 +87,136 @@ bool FAnimationAnalyzer::OnEvent(uint16 RouteId, const FOnEventContext& Context)
 		AnimationProvider.AppendName(Id, reinterpret_cast<const TCHAR*>(EventData.GetAttachment()));
 		break;
 	}
+	case RouteId_SkeletalMeshFrame:
+	{
+		uint64 Cycle = EventData.GetValue<uint64>("Cycle");
+		uint64 ComponentId = EventData.GetValue<uint64>("ComponentId");
+		uint16 FrameCounter = EventData.GetValue<uint16>("FrameCounter");
+		AnimationProvider.AppendSkeletalMeshFrame(ComponentId, Context.SessionContext.TimestampFromCycle(Cycle), FrameCounter);
+		break;
+	}
+	case RouteId_AnimGraph:
+	{
+		uint64 StartCycle = EventData.GetValue<uint64>("StartCycle");
+		uint64 EndCycle = EventData.GetValue<uint64>("EndCycle");
+		uint64 AnimInstanceId = EventData.GetValue<uint64>("AnimInstanceId");
+		int32 NodeCount = EventData.GetValue<int32>("NodeCount");
+		uint16 FrameCounter = EventData.GetValue<uint16>("FrameCounter");
+		uint8 Phase = EventData.GetValue<uint8>("Phase");
+		AnimationProvider.AppendAnimGraph(AnimInstanceId, Context.SessionContext.TimestampFromCycle(StartCycle), Context.SessionContext.TimestampFromCycle(EndCycle), NodeCount, FrameCounter, Phase);
+		break;
+	}
+	case RouteId_AnimNodeStart:
+	{
+		uint64 StartCycle = EventData.GetValue<uint64>("StartCycle");
+		uint64 AnimInstanceId = EventData.GetValue<uint64>("AnimInstanceId");
+		int32 NodeId = EventData.GetValue<int32>("NodeId");
+		int32 PreviousNodeId = EventData.GetValue<int32>("PreviousNodeId");
+		float Weight = EventData.GetValue<float>("Weight");
+		float RootMotionWeight = EventData.GetValue<float>("RootMotionWeight");
+		uint16 FrameCounter = EventData.GetValue<uint16>("FrameCounter");
+		uint8 Phase = EventData.GetValue<uint8>("Phase");
+		const TCHAR* TargetNodeName = reinterpret_cast<const TCHAR*>(EventData.GetAttachment());
+		AnimationProvider.AppendAnimNodeStart(AnimInstanceId, Context.SessionContext.TimestampFromCycle(StartCycle), FrameCounter, NodeId, PreviousNodeId, Weight, RootMotionWeight, TargetNodeName, Phase);
+		break;
+	}
+	case RouteId_AnimNodeValueBool:
+	case RouteId_AnimNodeValueInt:
+	case RouteId_AnimNodeValueFloat:
+	case RouteId_AnimNodeValueVector:
+	case RouteId_AnimNodeValueString:
+	case RouteId_AnimNodeValueObject:
+	case RouteId_AnimNodeValueClass:
+	{
+		uint64 Cycle = EventData.GetValue<uint64>("Cycle");
+		uint64 AnimInstanceId = EventData.GetValue<uint64>("AnimInstanceId");
+		int32 NodeId = EventData.GetValue<int32>("NodeId");
+		int32 KeyLength = EventData.GetValue<uint32>("KeyLength");
+		uint16 FrameCounter = EventData.GetValue<uint16>("FrameCounter");
+		const TCHAR* Key = reinterpret_cast<const TCHAR*>(EventData.GetAttachment());
+
+		switch (RouteId)
+		{
+		case RouteId_AnimNodeValueBool:
+		{
+			bool Value = EventData.GetValue<bool>("Value");
+			AnimationProvider.AppendAnimNodeValue(AnimInstanceId, Context.SessionContext.TimestampFromCycle(Cycle), FrameCounter, NodeId, Key, Value);
+			break;
+		}
+		case RouteId_AnimNodeValueInt:
+		{
+			int32 Value = EventData.GetValue<int32>("Value");
+			AnimationProvider.AppendAnimNodeValue(AnimInstanceId, Context.SessionContext.TimestampFromCycle(Cycle), FrameCounter, NodeId, Key, Value);
+			break;
+		}
+		case RouteId_AnimNodeValueFloat:
+		{
+			float Value = EventData.GetValue<float>("Value");
+			AnimationProvider.AppendAnimNodeValue(AnimInstanceId, Context.SessionContext.TimestampFromCycle(Cycle), FrameCounter, NodeId, Key, Value);
+			break;
+		}
+		case RouteId_AnimNodeValueVector:
+		{
+			FVector Value(EventData.GetValue<float>("ValueX"), EventData.GetValue<float>("ValueY"), EventData.GetValue<float>("ValueZ"));
+			AnimationProvider.AppendAnimNodeValue(AnimInstanceId, Context.SessionContext.TimestampFromCycle(Cycle), FrameCounter, NodeId, Key, Value);
+			break;
+		}
+		case RouteId_AnimNodeValueString:
+		{
+			const TCHAR* Value = reinterpret_cast<const TCHAR*>(EventData.GetAttachment()) + KeyLength;
+			AnimationProvider.AppendAnimNodeValue(AnimInstanceId, Context.SessionContext.TimestampFromCycle(Cycle), FrameCounter, NodeId, Key, Value);
+			break;
+		}
+		case RouteId_AnimNodeValueObject:
+		{
+			uint64 Value = EventData.GetValue<uint64>("Value");
+			AnimationProvider.AppendAnimNodeValueObject(AnimInstanceId, Context.SessionContext.TimestampFromCycle(Cycle), FrameCounter, NodeId, Key, Value);
+			break;
+		}
+		case RouteId_AnimNodeValueClass:
+		{
+			uint64 Value = EventData.GetValue<uint64>("Value");
+			AnimationProvider.AppendAnimNodeValueClass(AnimInstanceId, Context.SessionContext.TimestampFromCycle(Cycle), FrameCounter, NodeId, Key, Value);
+			break;
+		}
+		}
+		break;
+	}
+	case RouteId_AnimSequencePlayer:
+	{
+		uint64 Cycle = EventData.GetValue<uint64>("Cycle");
+		uint64 AnimInstanceId = EventData.GetValue<uint64>("AnimInstanceId");
+		int32 NodeId = EventData.GetValue<int32>("NodeId");
+		float Position = EventData.GetValue<float>("Position");
+		float Length = EventData.GetValue<float>("Length");
+		int32 FrameCount = EventData.GetValue<int32>("FrameCount");
+		AnimationProvider.AppendAnimSequencePlayer(AnimInstanceId, Context.SessionContext.TimestampFromCycle(Cycle), NodeId, Position, Length, FrameCount);
+		break;
+	}
+	case RouteId_BlendSpacePlayer:
+	{
+		uint64 Cycle = EventData.GetValue<uint64>("Cycle");
+		uint64 AnimInstanceId = EventData.GetValue<uint64>("AnimInstanceId");
+		uint64 BlendSpaceId = EventData.GetValue<uint64>("BlendSpaceId");
+		int32 NodeId = EventData.GetValue<int32>("NodeId");
+		float PositionX = EventData.GetValue<float>("PositionX");
+		float PositionY = EventData.GetValue<float>("PositionY");
+		float PositionZ = EventData.GetValue<float>("PositionZ");
+		AnimationProvider.AppendBlendSpacePlayer(AnimInstanceId, Context.SessionContext.TimestampFromCycle(Cycle), NodeId, BlendSpaceId, PositionX, PositionY, PositionZ);
+		break;
+	}
+	case RouteId_StateMachineState:
+	{
+		uint64 Cycle = EventData.GetValue<uint64>("Cycle");
+		uint64 AnimInstanceId = EventData.GetValue<uint64>("AnimInstanceId");
+		int32 NodeId = EventData.GetValue<int32>("NodeId");
+		int32 StateMachineIndex = EventData.GetValue<int32>("StateMachineIndex");
+		int32 StateIndex = EventData.GetValue<int32>("StateIndex");
+		float StateWeight = EventData.GetValue<float>("StateWeight");
+		float ElapsedTime = EventData.GetValue<float>("ElapsedTime");
+		AnimationProvider.AppendStateMachineState(AnimInstanceId, Context.SessionContext.TimestampFromCycle(Cycle), NodeId, StateMachineIndex, StateIndex, StateWeight, ElapsedTime);
+		break;
+	}
 	}
 
 	return true;
diff --git a/Engine/Plugins/Animation/GameplayInsights/Source/GameplayInsights/Private/AnimationAnalyzer.h b/Engine/Plugins/Animation/GameplayInsights/Source/GameplayInsights/Private/AnimationAnalyzer.h
--- a/Engine/Plugins/Animation/GameplayInsights/Source/GameplayInsights/Private/AnimationAnalyzer.h
+++ b/Engine/Plugins/Animation/GameplayInsights/Source/GameplayInsights/Private/AnimationAnalyzer.h
@@ -22,6 +22,19 @@ private:
 		RouteId_TickRecord,
 		RouteId_SkeletalMesh,
 		RouteId_SkeletalMeshComponent,
+		RouteId_SkeletalMeshFrame,
+		RouteId_AnimGraph,
+		RouteId_AnimNodeStart,
+		RouteId_AnimNodeValueBool,
+		RouteId_AnimNodeValueInt,
+		RouteId_AnimNodeValueFloat,
+		RouteId_AnimNodeValueVector,
+		RouteId_AnimNodeValueString,
+		RouteId_AnimNodeValueObject,
+		RouteId_AnimNodeValueClass,
+		RouteId_AnimSequencePlayer,
+		RouteId_BlendSpacePlayer,
+		RouteId_StateMachineState,
 		RouteId_Name,
 	};
 
diff --git a/Engine/Plugins/Animation/GameplayInsights/Source/GameplayInsights/Private/AnimationProvider.cpp b/Engine/Plugins/Animation/GameplayInsights/Source/GameplayInsights/Private/AnimationProvider.cpp
--- a/Engine/Plugins/Animation/GameplayInsights/Source/GameplayInsights/Private/AnimationProvider.cpp
+++ b/Engine/Plugins/Animation/GameplayInsights/Source/GameplayInsights/Private/AnimationProvider.cpp
@@ -33,17 +33,6 @@ bool FAnimationProvider::ReadSkeletalMeshPoseTimeline(uint64 InObjectId, TFuncti
 	return false;
 }
 
-bool FAnimationProvider::ReadSkeletalMeshPoseMessage(uint64 InObjectId, uint64 InMessageId, TFunctionRef<void(const FSkeletalMeshPoseMessage&)> Callback) const
-{
-	return ReadSkeletalMeshPoseTimeline(InObjectId, [&Callback, &InMessageId](const SkeletalMeshPoseTimeline& InTimeline)
-	{
-		if(InMessageId < InTimeline.GetEventCount())
-		{
-			Callback(InTimeline.GetEvent(InMessageId));
-		}
-	});
-}
-
 void FAnimationProvider::GetSkeletalMeshComponentSpacePose(const FSkeletalMeshPoseMessage& InMessage, const FSkeletalMeshInfo& InMeshInfo, FTransform& OutComponentToWorld, TArray<FTransform>& OutTransforms) const
 {
 	Session.ReadAccessCheck();
@@ -66,7 +55,7 @@ void FAnimationProvider::GetSkeletalMeshComponentSpacePose(const FSkeletalMeshPo
 	}
 }
 
-void FAnimationProvider::EnumerateTickRecordTimelines(uint64 InObjectId, TFunctionRef<void(uint64, const TickRecordTimeline&)> Callback) const
+void FAnimationProvider::EnumerateTickRecordTimelines(uint64 InObjectId, TFunctionRef<void(uint64, int32, const TickRecordTimeline&)> Callback) const
 {
 	Session.ReadAccessCheck();
 
@@ -76,15 +65,15 @@ void FAnimationProvider::EnumerateTickRecordTimelines(uint64 InObjectId, TFuncti
 		if (*IndexPtr < uint32(PerObjectTimelineStorage.Num()))
 		{
 			const TSharedRef<FPerObjectTimelineStorage>& TimelineStorage = PerObjectTimelineStorage[*IndexPtr];
-			for(auto AssetIdPair : TimelineStorage->AssetIdToTickRecordTimeline)
+			for(auto AssetIdPair : TimelineStorage->AssetIdAndPlayerToTickRecordTimeline)
 			{
-				Callback(AssetIdPair.Key, TimelineStorage->Timelines[AssetIdPair.Value].Get());
+				Callback(AssetIdPair.Key.Key, AssetIdPair.Key.Value, TimelineStorage->Timelines[AssetIdPair.Value].Get());
 			}
 		}
 	}
 }
 
-bool FAnimationProvider::ReadTickRecordTimeline(uint64 InObjectId, uint64 InAssetId, TFunctionRef<void(const TickRecordTimeline&)> Callback) const
+bool FAnimationProvider::ReadTickRecordTimeline(uint64 InObjectId, uint64 InAssetId, int32 InNodeId, TFunctionRef<void(const TickRecordTimeline&)> Callback) const
 {
 	Session.ReadAccessCheck();
 
@@ -94,7 +83,7 @@ bool FAnimationProvider::ReadTickRecordTimeline(uint64 InObjectId, uint64 InAsse
 		if (*IndexPtr < uint32(PerObjectTimelineStorage.Num()))
 		{
 			const TSharedRef<FPerObjectTimelineStorage>& TimelineStorage = PerObjectTimelineStorage[*IndexPtr];
-			const uint32* TimelineIndexPtr = TimelineStorage->AssetIdToTickRecordTimeline.Find(InAssetId);
+			const uint32* TimelineIndexPtr = TimelineStorage->AssetIdAndPlayerToTickRecordTimeline.Find(TTuple<uint64, int32>(InAssetId, InNodeId));
 			if(TimelineIndexPtr != nullptr)
 			{
 				Callback(TimelineStorage->Timelines[*TimelineIndexPtr].Get());
@@ -106,23 +95,119 @@ bool FAnimationProvider::ReadTickRecordTimeline(uint64 InObjectId, uint64 InAsse
 	return false;
 }
 
-bool FAnimationProvider::ReadTickRecordMessage(uint64 InObjectId, uint64 InAssetId, uint64 InMessageId, TFunctionRef<void(const FTickRecordMessage&)> Callback) const
+bool FAnimationProvider::ReadAnimGraphTimeline(uint64 InObjectId, TFunctionRef<void(const AnimGraphTimeline&)> Callback) const
+{
+	Session.ReadAccessCheck();
+
+	const uint32* IndexPtr = ObjectIdToAnimGraphTimelines.Find(InObjectId);
+	if(IndexPtr != nullptr)
+	{
+		if (*IndexPtr < uint32(AnimGraphTimelines.Num()))
+		{
+			Callback(*AnimGraphTimelines[*IndexPtr]);
+			return true;
+		}
+	}
+
+	return false;
+}
+
+bool FAnimationProvider::ReadAnimNodesTimeline(uint64 InObjectId, TFunctionRef<void(const AnimNodesTimeline&)> Callback) const
+{
+	Session.ReadAccessCheck();
+
+	const uint32* IndexPtr = ObjectIdToAnimNodeTimelines.Find(InObjectId);
+	if(IndexPtr != nullptr)
+	{
+		if (*IndexPtr < uint32(AnimNodeTimelines.Num()))
+		{
+			Callback(*AnimNodeTimelines[*IndexPtr]);
+			return true;
+		}
+	}
+
+	return false;
+}
+
+bool FAnimationProvider::ReadAnimNodeValuesTimeline(uint64 InObjectId, TFunctionRef<void(const AnimNodeValuesTimeline&)> Callback) const
+{
+	Session.ReadAccessCheck();
+
+	const uint32* IndexPtr = ObjectIdToAnimNodeValueTimelines.Find(InObjectId);
+	if(IndexPtr != nullptr)
+	{
+		if (*IndexPtr < uint32(AnimNodeValueTimelines.Num()))
+		{
+			Callback(*AnimNodeValueTimelines[*IndexPtr]);
+			return true;
+		}
+	}
+
+	return false;
+}
+
+bool FAnimationProvider::ReadAnimSequencePlayersTimeline(uint64 InObjectId, TFunctionRef<void(const AnimSequencePlayersTimeline&)> Callback) const
+{
+	Session.ReadAccessCheck();
+
+	const uint32* IndexPtr = ObjectIdToAnimSequencePlayerTimelines.Find(InObjectId);
+	if(IndexPtr != nullptr)
+	{
+		if (*IndexPtr < uint32(AnimSequencePlayerTimelines.Num()))
+		{
+			Callback(*AnimSequencePlayerTimelines[*IndexPtr]);
+			return true;
+		}
+	}
+
+	return false;
+}
+
+bool FAnimationProvider::ReadAnimBlendSpacePlayersTimeline(uint64 InObjectId, TFunctionRef<void(const BlendSpacePlayersTimeline&)> Callback) const
+{
+	Session.ReadAccessCheck();
+
+	const uint32* IndexPtr = ObjectIdToBlendSpacePlayerTimelines.Find(InObjectId);
+	if(IndexPtr != nullptr)
+	{
+		if (*IndexPtr < uint32(BlendSpacePlayerTimelines.Num()))
+		{
+			Callback(*BlendSpacePlayerTimelines[*IndexPtr]);
+			return true;
+		}
+	}
+
+	return false;
+}
+
+bool FAnimationProvider::ReadStateMachinesTimeline(uint64 InObjectId, TFunctionRef<void(const StateMachinesTimeline&)> Callback) const
 {
-	return ReadTickRecordTimeline(InObjectId, InAssetId, [&Callback, &InMessageId](const TickRecordTimeline& InTimeline)
+	Session.ReadAccessCheck();
+
+	const uint32* IndexPtr = ObjectIdToStateMachineTimelines.Find(InObjectId);
+	if(IndexPtr != nullptr)
 	{
-		if(InMessageId < InTimeline.GetEventCount())
+		if (*IndexPtr < uint32(StateMachineTimelines.Num()))
 		{
-			Callback(InTimeline.GetEvent(InMessageId));
+			Callback(*StateMachineTimelines[*IndexPtr]);
+			return true;
 		}
-	});
+	}
+
+	return false;
 }
 
 const FSkeletalMeshInfo* FAnimationProvider::FindSkeletalMeshInfo(uint64 InObjectId) const
 {
-	const int32* ObjectIndex = SkeletalMeshIdToIndexMap.Find(InObjectId);
-	if(ObjectIndex != nullptr)
+	Session.ReadAccessCheck();
+
+	const uint32* IndexPtr = SkeletalMeshIdToIndexMap.Find(InObjectId);
+	if(IndexPtr != nullptr)
 	{
-		return &SkeletalMeshInfos[*ObjectIndex];
+		if (*IndexPtr < uint32(SkeletalMeshInfos.Num()))
+		{
+			return &SkeletalMeshInfos[*IndexPtr];
+		}
 	}
 
 	return nullptr;
@@ -133,26 +218,26 @@ const TCHAR* FAnimationProvider::GetName(uint32 InId) const
 	return NameMap.FindRef(InId);
 }
 
-void FAnimationProvider::AppendTickRecord(uint64 InObjectId, uint64 InSubObjectId, double InTime, uint64 InAssetId, float InBlendWeight, float InPlaybackTime, float InRootMotionWeight, float InPlayRate, uint16 InFrameCounter, bool bInLooping)
+void FAnimationProvider::AppendTickRecord(uint64 InAnimInstanceId, double InTime, uint64 InAssetId, int32 InNodeId, float InBlendWeight, float InPlaybackTime, float InRootMotionWeight, float InPlayRate, float InBlendSpacePositionX, float InBlendSpacePositionY, uint16 InFrameCounter, bool bInLooping, bool bInIsBlendSpace)
 {
 	Session.WriteAccessCheck();
 
 	TSharedPtr<Trace::TPointTimeline<FTickRecordMessage>> Timeline;
 	TSharedPtr<FPerObjectTimelineStorage> TimelineStorage;
-	uint32* TimelineStorageIndexPtr = ObjectIdToTickRecordTimelineStorage.Find(InObjectId);
+	uint32* TimelineStorageIndexPtr = ObjectIdToTickRecordTimelineStorage.Find(InAnimInstanceId);
 	if(TimelineStorageIndexPtr != nullptr)
 	{
 		TimelineStorage = PerObjectTimelineStorage[*TimelineStorageIndexPtr];
 	}
 	else
 	{
-		ObjectIdToTickRecordTimelineStorage.Add(InObjectId, PerObjectTimelineStorage.Num());
+		ObjectIdToTickRecordTimelineStorage.Add(InAnimInstanceId, PerObjectTimelineStorage.Num());
 		TimelineStorage = PerObjectTimelineStorage.Add_GetRef(MakeShared<FPerObjectTimelineStorage>());
 	}
 
 	check(TimelineStorage.IsValid());
 
-	uint32* TimelineIndexPtr = TimelineStorage->AssetIdToTickRecordTimeline.Find(InAssetId);
+	uint32* TimelineIndexPtr = TimelineStorage->AssetIdAndPlayerToTickRecordTimeline.Find(TTuple<uint64, int32>(InAssetId, InNodeId));
 	if(TimelineIndexPtr != nullptr)
 	{
 		Timeline = TimelineStorage->Timelines[*TimelineIndexPtr];
@@ -160,20 +245,44 @@ void FAnimationProvider::AppendTickRecord(uint64 InObjectId, uint64 InSubObjectI
 	else
 	{
 		Timeline = MakeShared<Trace::TPointTimeline<FTickRecordMessage>>(Session.GetLinearAllocator());
-		TimelineStorage->AssetIdToTickRecordTimeline.Add(InAssetId, TimelineStorage->Timelines.Num());
+
+		// We need to enumerate outside range to make sure we draw lines to the edge of the viewport when iterating to draw
+		Timeline->SetEnumerateOutsideRange(true);
+		TimelineStorage->AssetIdAndPlayerToTickRecordTimeline.Add(TTuple<uint64, int32>(InAssetId, InNodeId), TimelineStorage->Timelines.Num());
 		TimelineStorage->Timelines.Add(Timeline.ToSharedRef());
 	}
 
+	// check for continuity
+	bool bContinuous = false;
+	const uint64 NumEvents = Timeline->GetEventCount();
+	if(NumEvents > 0)
+	{
+		if(InFrameCounter - 1 != Timeline->GetEvent(NumEvents - 1).FrameCounter)
+		{
+			// discontinuous frame count
+			bContinuous = false;
+		}
+		else
+		{
+			// continuous frame time
+			bContinuous = true;
+		}
+	}
+
 	FTickRecordMessage Message;
-	Message.ComponentId = InObjectId;
-	Message.AnimInstanceId = InSubObjectId;
+	Message.AnimInstanceId = InAnimInstanceId;
 	Message.AssetId = InAssetId;
+	Message.NodeId = InNodeId;
 	Message.BlendWeight = InBlendWeight;
 	Message.PlaybackTime = InPlaybackTime;
 	Message.RootMotionWeight = InRootMotionWeight;
 	Message.PlayRate = InPlayRate;
+	Message.BlendSpacePositionX = InBlendSpacePositionX;
+	Message.BlendSpacePositionY = InBlendSpacePositionY;
 	Message.FrameCounter = InFrameCounter;
-	Message.Looping = bInLooping;
+	Message.bLooping = bInLooping;
+	Message.bIsBlendSpace = bInIsBlendSpace;
+	Message.bContinuous = bContinuous;
 
 	Timeline->AppendEvent(InTime, Message);
 
@@ -196,7 +305,7 @@ void FAnimationProvider::AppendSkeletalMesh(uint64 InObjectId, const TArrayView<
 			SkeletalMeshParentIndices.PushBack() = ParentIndex;
 		}
 
-		int32 NewSkeletalMeshInfoIndex = SkeletalMeshInfos.Add(NewSkeletalMeshInfo);
+		uint32 NewSkeletalMeshInfoIndex = SkeletalMeshInfos.Add(NewSkeletalMeshInfo);
 		SkeletalMeshIdToIndexMap.Add(InObjectId, NewSkeletalMeshInfoIndex);
 	}
 }
@@ -272,4 +381,284 @@ void FAnimationProvider::HandleObjectEndPlay(uint64 InObjectId, double InTime, c
 			Timeline->EndEvent(NumEvents - 1, InTime);
 		}
 	}
-}
\ No newline at end of file
+}
+
+void FAnimationProvider::AppendSkeletalMeshFrame(uint64 InObjectId, double InTime, uint16 InFrameCounter)
+{
+	Session.WriteAccessCheck();
+
+	TSharedPtr<Trace::TIntervalTimeline<FSkeletalMeshFrameMessage>> Timeline;
+	uint32* IndexPtr = ObjectIdToSkeletalMeshFrameTimelines.Find(InObjectId);
+	if(IndexPtr != nullptr)
+	{
+		check(SkeletalMeshFrameTimelines.IsValidIndex(*IndexPtr));
+		Timeline = SkeletalMeshFrameTimelines[*IndexPtr];
+	}
+	else
+	{
+		Timeline = MakeShared<Trace::TIntervalTimeline<FSkeletalMeshFrameMessage>>(Session.GetLinearAllocator());
+		ObjectIdToSkeletalMeshFrameTimelines.Add(InObjectId, SkeletalMeshFrameTimelines.Num());
+		SkeletalMeshFrameTimelines.Add(Timeline.ToSharedRef());
+	}
+
+	// terminate existing scopes
+	uint64 NumEvents = Timeline->GetEventCount();
+	if(NumEvents > 0)
+	{
+		// Add end event at current time
+		Timeline->EndEvent(NumEvents - 1, InTime);
+	}
+
+	FSkeletalMeshFrameMessage Message;
+	Message.ComponentId = InObjectId;
+	Message.FrameCounter = InFrameCounter;
+
+	Timeline->AppendBeginEvent(InTime, Message);
+
+	Session.UpdateDurationSeconds(InTime);
+}
+
+void FAnimationProvider::AppendAnimGraph(uint64 InAnimInstanceId, double InStartTime, double InEndTime, int32 InNodeCount, uint16 InFrameCounter, uint8 InPhase)
+{
+	Session.WriteAccessCheck();
+
+	TSharedPtr<Trace::TIntervalTimeline<FAnimGraphMessage>> Timeline;
+	uint32* IndexPtr = ObjectIdToAnimGraphTimelines.Find(InAnimInstanceId);
+	if(IndexPtr != nullptr)
+	{
+		check(AnimGraphTimelines.IsValidIndex(*IndexPtr));
+		Timeline = AnimGraphTimelines[*IndexPtr];
+	}
+	else
+	{
+		Timeline = MakeShared<Trace::TIntervalTimeline<FAnimGraphMessage>>(Session.GetLinearAllocator());
+		ObjectIdToAnimGraphTimelines.Add(InAnimInstanceId, AnimGraphTimelines.Num());
+		AnimGraphTimelines.Add(Timeline.ToSharedRef());
+	}
+
+	FAnimGraphMessage Message;
+	Message.AnimInstanceId = InAnimInstanceId;
+	Message.NodeCount = InNodeCount;
+	Message.FrameCounter = InFrameCounter;
+	Message.Phase = (EAnimGraphPhase)InPhase;
+
+	uint64 EventIndex = Timeline->AppendBeginEvent(InStartTime, Message);
+	Timeline->EndEvent(EventIndex, InEndTime);
+
+	Session.UpdateDurationSeconds(InStartTime);
+}
+
+void FAnimationProvider::AppendAnimNodeStart(uint64 InAnimInstanceId, double InStartTime, uint16 InFrameCounter, int32 InNodeId, int32 PreviousNodeId, float InWeight, float InRootMotionWeight, const TCHAR* InTargetNodeName, uint8 InPhase)
+{
+	Session.WriteAccessCheck();
+
+	TSharedPtr<Trace::TPointTimeline<FAnimNodeMessage>> Timeline;
+	uint32* IndexPtr = ObjectIdToAnimNodeTimelines.Find(InAnimInstanceId);
+	if(IndexPtr != nullptr)
+	{
+		check(AnimNodeTimelines.IsValidIndex(*IndexPtr));
+		Timeline = AnimNodeTimelines[*IndexPtr];
+	}
+	else
+	{
+		Timeline = MakeShared<Trace::TPointTimeline<FAnimNodeMessage>>(Session.GetLinearAllocator());
+		ObjectIdToAnimNodeTimelines.Add(InAnimInstanceId, AnimNodeTimelines.Num());
+		AnimNodeTimelines.Add(Timeline.ToSharedRef());
+	}
+
+	FAnimNodeMessage Message;
+	Message.NodeName = Session.StoreString(InTargetNodeName);
+	Message.AnimInstanceId = InAnimInstanceId;
+	Message.PreviousNodeId = PreviousNodeId;
+	Message.NodeId = InNodeId;
+	Message.Weight = InWeight;
+	Message.RootMotionWeight = InRootMotionWeight;
+	Message.FrameCounter = InFrameCounter;
+	Message.Phase = (EAnimGraphPhase)InPhase;
+
+	Timeline->AppendEvent(InStartTime, Message);
+
+	Session.UpdateDurationSeconds(InStartTime);
+}
+
+void FAnimationProvider::AppendAnimSequencePlayer(uint64 InAnimInstanceId, double InTime, int32 InNodeId, float InPosition, float InLength, int32 InFrameCount)
+{
+	Session.WriteAccessCheck();
+
+	TSharedPtr<Trace::TPointTimeline<FAnimSequencePlayerMessage>> Timeline;
+	uint32* IndexPtr = ObjectIdToAnimSequencePlayerTimelines.Find(InAnimInstanceId);
+	if(IndexPtr != nullptr)
+	{
+		check(AnimSequencePlayerTimelines.IsValidIndex(*IndexPtr));
+		Timeline = AnimSequencePlayerTimelines[*IndexPtr];
+	}
+	else
+	{
+		Timeline = MakeShared<Trace::TPointTimeline<FAnimSequencePlayerMessage>>(Session.GetLinearAllocator());
+		ObjectIdToAnimSequencePlayerTimelines.Add(InAnimInstanceId, AnimSequencePlayerTimelines.Num());
+		AnimSequencePlayerTimelines.Add(Timeline.ToSharedRef());
+	}
+
+	FAnimSequencePlayerMessage Message;
+	Message.AnimInstanceId = InAnimInstanceId;
+	Message.NodeId = InNodeId;
+	Message.Position = InPosition;
+	Message.Length = InLength;
+	Message.FrameCount = InFrameCount;
+
+	Timeline->AppendEvent(InTime, Message);
+
+	Session.UpdateDurationSeconds(InTime);
+}
+
+void FAnimationProvider::AppendBlendSpacePlayer(uint64 InAnimInstanceId, double InTime, int32 InNodeId, uint64 InBlendSpaceId, float InPositionX, float InPositionY, float InPositionZ)
+{
+	Session.WriteAccessCheck();
+
+	TSharedPtr<Trace::TPointTimeline<FBlendSpacePlayerMessage>> Timeline;
+	uint32* IndexPtr = ObjectIdToBlendSpacePlayerTimelines.Find(InAnimInstanceId);
+	if(IndexPtr != nullptr)
+	{
+		check(BlendSpacePlayerTimelines.IsValidIndex(*IndexPtr));
+		Timeline = BlendSpacePlayerTimelines[*IndexPtr];
+	}
+	else
+	{
+		Timeline = MakeShared<Trace::TPointTimeline<FBlendSpacePlayerMessage>>(Session.GetLinearAllocator());
+		ObjectIdToBlendSpacePlayerTimelines.Add(InAnimInstanceId, BlendSpacePlayerTimelines.Num());
+		BlendSpacePlayerTimelines.Add(Timeline.ToSharedRef());
+	}
+
+	FBlendSpacePlayerMessage Message;
+	Message.AnimInstanceId = InAnimInstanceId;
+	Message.BlendSpaceId = InBlendSpaceId;
+	Message.NodeId = InNodeId;
+	Message.PositionX = InPositionX;
+	Message.PositionY = InPositionY;
+	Message.PositionZ = InPositionZ;
+
+	Timeline->AppendEvent(InTime, Message);
+
+	Session.UpdateDurationSeconds(InTime);
+}
+
+void FAnimationProvider::AppendAnimNodeValue(uint64 InAnimInstanceId, double InTime, uint16 InFrameCounter, int32 InNodeId, const TCHAR* InKey, bool bInValue)
+{
+	FAnimNodeValueMessage Message;
+	Message.Bool.bValue = bInValue;
+	Message.Type = EAnimNodeValueType::Bool;
+
+	AppendAnimNodeValue(InAnimInstanceId, InTime, InFrameCounter, InNodeId, InKey, Message);
+}
+
+void FAnimationProvider::AppendAnimNodeValue(uint64 InAnimInstanceId, double InTime, uint16 InFrameCounter, int32 InNodeId, const TCHAR* InKey, int32 InValue)
+{
+	FAnimNodeValueMessage Message;
+	Message.Int32.Value = InValue;
+	Message.Type = EAnimNodeValueType::Int32;
+
+	AppendAnimNodeValue(InAnimInstanceId, InTime, InFrameCounter, InNodeId, InKey, Message);
+}
+
+void FAnimationProvider::AppendAnimNodeValue(uint64 InAnimInstanceId, double InTime, uint16 InFrameCounter, int32 InNodeId, const TCHAR* InKey, float InValue)
+{
+	FAnimNodeValueMessage Message;
+	Message.Float.Value = InValue;
+	Message.Type = EAnimNodeValueType::Float;
+
+	AppendAnimNodeValue(InAnimInstanceId, InTime, InFrameCounter, InNodeId, InKey, Message);
+}
+
+void FAnimationProvider::AppendAnimNodeValue(uint64 InAnimInstanceId, double InTime, uint16 InFrameCounter, int32 InNodeId, const TCHAR* InKey, const FVector& InValue)
+{
+	FAnimNodeValueMessage Message;
+	Message.Vector.Value = InValue;
+	Message.Type = EAnimNodeValueType::Vector;
+
+	AppendAnimNodeValue(InAnimInstanceId, InTime, InFrameCounter, InNodeId, InKey, Message);
+}
+
+void FAnimationProvider::AppendAnimNodeValue(uint64 InAnimInstanceId, double InTime, uint16 InFrameCounter, int32 InNodeId, const TCHAR* InKey, const TCHAR* InValue)
+{
+	FAnimNodeValueMessage Message;
+	Message.String.Value = Session.StoreString(InValue);
+	Message.Type = EAnimNodeValueType::String;
+
+	AppendAnimNodeValue(InAnimInstanceId, InTime, InFrameCounter, InNodeId, InKey, Message);
+}
+
+void FAnimationProvider::AppendAnimNodeValueObject(uint64 InAnimInstanceId, double InTime, uint16 InFrameCounter, int32 InNodeId, const TCHAR* InKey, uint64 InValue)
+{
+	FAnimNodeValueMessage Message;
+	Message.Object.Value = InValue;
+	Message.Type = EAnimNodeValueType::Object;
+
+	AppendAnimNodeValue(InAnimInstanceId, InTime, InFrameCounter, InNodeId, InKey, Message);
+}
+
+void FAnimationProvider::AppendAnimNodeValueClass(uint64 InAnimInstanceId, double InTime, uint16 InFrameCounter, int32 InNodeId, const TCHAR* InKey, uint64 InValue)
+{
+	FAnimNodeValueMessage Message;
+	Message.Class.Value = InValue;
+	Message.Type = EAnimNodeValueType::Class;
+
+	AppendAnimNodeValue(InAnimInstanceId, InTime, InFrameCounter, InNodeId, InKey, Message);
+}
+
+void FAnimationProvider::AppendAnimNodeValue(uint64 InAnimInstanceId, double InTime, uint16 InFrameCounter, int32 InNodeId, const TCHAR* InKey, FAnimNodeValueMessage& InMessage)
+{
+	Session.WriteAccessCheck();
+
+	TSharedPtr<Trace::TPointTimeline<FAnimNodeValueMessage>> Timeline;
+	uint32* IndexPtr = ObjectIdToAnimNodeValueTimelines.Find(InAnimInstanceId);
+	if(IndexPtr != nullptr)
+	{
+		Timeline = AnimNodeValueTimelines[*IndexPtr];
+	}
+	else
+	{
+		Timeline = MakeShared<Trace::TPointTimeline<FAnimNodeValueMessage>>(Session.GetLinearAllocator());
+		ObjectIdToAnimNodeValueTimelines.Add(InAnimInstanceId, AnimNodeValueTimelines.Num());
+		AnimNodeValueTimelines.Add(Timeline.ToSharedRef());
+	}
+
+	InMessage.Key = Session.StoreString(InKey);
+	InMessage.AnimInstanceId = InAnimInstanceId;
+	InMessage.NodeId = InNodeId;
+	InMessage.FrameCounter = InFrameCounter;
+
+	Timeline->AppendEvent(InTime, InMessage);
+
+	Session.UpdateDurationSeconds(InTime);
+}
+
+void FAnimationProvider::AppendStateMachineState(uint64 InAnimInstanceId, double InTime, int32 InNodeId, int32 InStateMachineIndex, int32 InStateIndex, float InStateWeight, float InElapsedTime)
+{
+	Session.WriteAccessCheck();
+
+	TSharedPtr<Trace::TPointTimeline<FAnimStateMachineMessage>> Timeline;
+	uint32* IndexPtr = ObjectIdToStateMachineTimelines.Find(InAnimInstanceId);
+	if(IndexPtr != nullptr)
+	{
+		Timeline = StateMachineTimelines[*IndexPtr];
+	}
+	else
+	{
+		Timeline = MakeShared<Trace::TPointTimeline<FAnimStateMachineMessage>>(Session.GetLinearAllocator());
+		ObjectIdToStateMachineTimelines.Add(InAnimInstanceId, StateMachineTimelines.Num());
+		StateMachineTimelines.Add(Timeline.ToSharedRef());
+	}
+
+	FAnimStateMachineMessage Message;
+	Message.AnimInstanceId = InAnimInstanceId;
+	Message.NodeId = InNodeId;
+	Message.StateMachineIndex = InStateMachineIndex;
+	Message.StateIndex = InStateIndex;
+	Message.StateWeight = InStateWeight;
+	Message.ElapsedTime = InElapsedTime;
+
+	Timeline->AppendEvent(InTime, Message);
+
+	Session.UpdateDurationSeconds(InTime);
+}
+
diff --git a/Engine/Plugins/Animation/GameplayInsights/Source/GameplayInsights/Private/AnimationProvider.h b/Engine/Plugins/Animation/GameplayInsights/Source/GameplayInsights/Private/AnimationProvider.h
--- a/Engine/Plugins/Animation/GameplayInsights/Source/GameplayInsights/Private/AnimationProvider.h
+++ b/Engine/Plugins/Animation/GameplayInsights/Source/GameplayInsights/Private/AnimationProvider.h
@@ -7,6 +7,7 @@
 #include "Model/IntervalTimeline.h"
 #include "Model/PointTimeline.h"
 #include "Containers/ArrayView.h"
+#include "Model/IntervalTimeline.h"
 
 namespace Trace { class IAnalysisSession; }
 class FGameplayProvider;
@@ -22,16 +23,20 @@ public:
 
 	/** IAnimationProvider interface */
 	virtual bool ReadSkeletalMeshPoseTimeline(uint64 InObjectId, TFunctionRef<void(const SkeletalMeshPoseTimeline&)> Callback) const override;
-	virtual bool ReadSkeletalMeshPoseMessage(uint64 InObjectId, uint64 InMessageId, TFunctionRef<void(const FSkeletalMeshPoseMessage&)> Callback) const override;
 	virtual void GetSkeletalMeshComponentSpacePose(const FSkeletalMeshPoseMessage& InMessage, const FSkeletalMeshInfo& InMeshInfo, FTransform& OutComponentToWorld, TArray<FTransform>& OutTransforms) const override;
-	virtual void EnumerateTickRecordTimelines(uint64 InObjectId, TFunctionRef<void(uint64, const TickRecordTimeline&)> Callback) const override;
-	virtual bool ReadTickRecordTimeline(uint64 InObjectId, uint64 InAssetId, TFunctionRef<void(const TickRecordTimeline&)> Callback) const override;
-	virtual bool ReadTickRecordMessage(uint64 InObjectId, uint64 InAssetId, uint64 InMessageId, TFunctionRef<void(const FTickRecordMessage&)> Callback) const override;
+	virtual void EnumerateTickRecordTimelines(uint64 InObjectId, TFunctionRef<void(uint64, int32, const TickRecordTimeline&)> Callback) const override;
+	virtual bool ReadTickRecordTimeline(uint64 InObjectId, uint64 InAssetId, int32 InNodeId, TFunctionRef<void(const TickRecordTimeline&)> Callback) const override;
+	virtual bool ReadAnimGraphTimeline(uint64 InObjectId, TFunctionRef<void(const AnimGraphTimeline&)> Callback) const override;
+	virtual bool ReadAnimNodesTimeline(uint64 InObjectId, TFunctionRef<void(const AnimNodesTimeline&)> Callback) const override;
+	virtual bool ReadAnimNodeValuesTimeline(uint64 InObjectId, TFunctionRef<void(const AnimNodeValuesTimeline&)> Callback) const override;
+	virtual bool ReadAnimSequencePlayersTimeline(uint64 InObjectId, TFunctionRef<void(const AnimSequencePlayersTimeline&)> Callback) const override;
+	virtual bool ReadAnimBlendSpacePlayersTimeline(uint64 InObjectId, TFunctionRef<void(const BlendSpacePlayersTimeline&)> Callback) const override;
+	virtual bool ReadStateMachinesTimeline(uint64 InObjectId, TFunctionRef<void(const StateMachinesTimeline&)> Callback) const override;
 	virtual const FSkeletalMeshInfo* FindSkeletalMeshInfo(uint64 InObjectId) const override;
 	virtual const TCHAR* GetName(uint32 InId) const override;
 
 	/** Add a tick record */
-	void AppendTickRecord(uint64 InObjectId, uint64 InSubObjectId, double InTime, uint64 InAssetId, float InBlendWeight, float InPlaybackTime, float InRootMotionWeight, float InPlayRate, uint16 InFrameCounter, bool bInLooping);
+	void AppendTickRecord(uint64 InAnimInstanceId, double InTime, uint64 InAssetId, int32 InNodeId, float InBlendWeight, float InPlaybackTime, float InRootMotionWeight, float InPlayRate, float InBlendSpacePositionX, float InBlendSpacePositionY, uint16 InFrameCounter, bool bInLooping, bool bInIsBlendSpace);
 
 	/** Add a skeletal mesh */
 	void AppendSkeletalMesh(uint64 InObjectId, const TArrayView<const int32>& ParentIndices);
@@ -45,7 +50,37 @@ public:
 	/** Add a name, referenced by ID (for curves etc.) */
 	void AppendName(uint32 InId, const TCHAR* InName);
 
+	/** Add a skeletal mesh frame */
+	void AppendSkeletalMeshFrame(uint64 InObjectId, double InTime, uint16 InFrameCounter);
+	
+	/** Add an anim graph */
+	void AppendAnimGraph(uint64 InAnimInstanceId, double InStartTime, double InEndTime, int32 InNodeCount, uint16 InFrameCounter, uint8 InPhase);
+
+	/** Add an anim node */
+	void AppendAnimNodeStart(uint64 InAnimInstanceId, double InStartTime, uint16 InFrameCounter, int32 InTargetLinkId, int32 InSourceLinkId, float InWeight, float InRootMotionWeight, const TCHAR* InTargetNodeName, uint8 InPhase);
+
+	/** Add anim node values */
+	void AppendAnimNodeValue(uint64 InAnimInstanceId, double InTime, uint16 InFrameCounter, int32 InNodeId, const TCHAR* InKey, bool bInValue);
+	void AppendAnimNodeValue(uint64 InAnimInstanceId, double InTime, uint16 InFrameCounter, int32 InNodeId, const TCHAR* InKey, int32 InValue);
+	void AppendAnimNodeValue(uint64 InAnimInstanceId, double InTime, uint16 InFrameCounter, int32 InNodeId, const TCHAR* InKey, float InValue);
+	void AppendAnimNodeValue(uint64 InAnimInstanceId, double InTime, uint16 InFrameCounter, int32 InNodeId, const TCHAR* InKey, const FVector& InValue);
+	void AppendAnimNodeValue(uint64 InAnimInstanceId, double InTime, uint16 InFrameCounter, int32 InNodeId, const TCHAR* InKey, const TCHAR* InValue);
+	void AppendAnimNodeValueObject(uint64 InAnimInstanceId, double InTime, uint16 InFrameCounter, int32 InNodeId, const TCHAR* InKey, uint64 InValue);
+	void AppendAnimNodeValueClass(uint64 InAnimInstanceId, double InTime, uint16 InFrameCounter, int32 InNodeId, const TCHAR* InKey, uint64 InValue);
+
+	/** Add a sequence player record */
+	void AppendAnimSequencePlayer(uint64 InAnimInstanceId, double InTime, int32 InNodeId, float InPosition, float InLength, int32 InFrameCount);
+
+	/** Add a blend space player record */
+	void AppendBlendSpacePlayer(uint64 InAnimInstanceId, double InTime, int32 InNodeId, uint64 InBlendSpaceId, float InPositionX, float InPositionY, float InPositionZ);
+
+	/** Add a state machine state */
+	void AppendStateMachineState(uint64 InAnimInstanceId, double InTime, int32 InNodeId, int32 InStateMachineIndex, int32 InStateIndex, float InStateWeight, float InElapsedTime);
+
 private:
+	/** Add anim node values helper */
+	void AppendAnimNodeValue(uint64 InAnimInstanceId, double InTime, uint16 InFrameCounter, int32 InNodeId, const TCHAR* InKey, FAnimNodeValueMessage& InMessage);
+
 	/** Handle an object ending play - terminate existing scopes */
 	void HandleObjectEndPlay(uint64 InObjectId, double InTime, const FObjectInfo& InObjectInfo);
 
@@ -56,25 +91,39 @@ private:
 	/** Maps into timeline arrays per-object */
 	TMap<uint64, uint32> ObjectIdToTickRecordTimelineStorage;
 	TMap<uint64, uint32> ObjectIdToSkeletalMeshPoseTimelines;
+	TMap<uint64, uint32> ObjectIdToSkeletalMeshFrameTimelines;
+	TMap<uint64, uint32> ObjectIdToAnimGraphTimelines;
+	TMap<uint64, uint32> ObjectIdToAnimNodeTimelines;
+	TMap<uint64, uint32> ObjectIdToStateMachineTimelines;
+	TMap<uint64, uint32> ObjectIdToAnimNodeValueTimelines;
+	TMap<uint64, uint32> ObjectIdToAnimSequencePlayerTimelines;
+	TMap<uint64, uint32> ObjectIdToBlendSpacePlayerTimelines;
 
 	/** All the skeletal mesh info we have seen, grow only for stable indices */
 	TArray<FSkeletalMeshInfo> SkeletalMeshInfos;
 
 	/** Map of skeletal mesh->index into SkeletalMeshInfos index */
-	TMap<uint64, int32> SkeletalMeshIdToIndexMap;
+	TMap<uint64, uint32> SkeletalMeshIdToIndexMap;
 
 	/** Map of name IDs to name string */
 	TMap<uint32, const TCHAR*> NameMap;
 	
 	struct FPerObjectTimelineStorage
 	{
-		TMap<uint64, uint32> AssetIdToTickRecordTimeline;
+		TMap<TTuple<uint64, int32>, uint32> AssetIdAndPlayerToTickRecordTimeline;
 		TArray<TSharedRef<Trace::TPointTimeline<FTickRecordMessage>>> Timelines;
 	};
 
 	/** Message storage */
 	TArray<TSharedRef<FPerObjectTimelineStorage>> PerObjectTimelineStorage;
 	TArray<TSharedRef<Trace::TIntervalTimeline<FSkeletalMeshPoseMessage>>> SkeletalMeshPoseTimelines;
+	TArray<TSharedRef<Trace::TIntervalTimeline<FSkeletalMeshFrameMessage>>> SkeletalMeshFrameTimelines;
+	TArray<TSharedRef<Trace::TIntervalTimeline<FAnimGraphMessage>>> AnimGraphTimelines;
+	TArray<TSharedRef<Trace::TPointTimeline<FAnimNodeMessage>>> AnimNodeTimelines;
+	TArray<TSharedRef<Trace::TPointTimeline<FAnimNodeValueMessage>>> AnimNodeValueTimelines;
+	TArray<TSharedRef<Trace::TPointTimeline<FAnimSequencePlayerMessage>>> AnimSequencePlayerTimelines;
+	TArray<TSharedRef<Trace::TPointTimeline<FBlendSpacePlayerMessage>>> BlendSpacePlayerTimelines;
+	TArray<TSharedRef<Trace::TPointTimeline<FAnimStateMachineMessage>>> StateMachineTimelines;
 	TPagedArray<FTransform> SkeletalMeshPoseTransforms;
 	TPagedArray<FSkeletalMeshNamedCurve> SkeletalMeshCurves;
 	TPagedArray<int32> SkeletalMeshParentIndices;
diff --git a/Engine/Plugins/Animation/GameplayInsights/Source/GameplayInsights/Private/AnimationSharedData.cpp b/Engine/Plugins/Animation/GameplayInsights/Source/GameplayInsights/Private/AnimationSharedData.cpp
--- a/Engine/Plugins/Animation/GameplayInsights/Source/GameplayInsights/Private/AnimationSharedData.cpp
+++ b/Engine/Plugins/Animation/GameplayInsights/Source/GameplayInsights/Private/AnimationSharedData.cpp
@@ -11,6 +11,13 @@
 #include "AnimationProvider.h"
 #include "GameplayProvider.h"
 #include "Framework/MultiBox/MultiBoxBuilder.h"
+#include "AnimNodesTrack.h"
+#include "GameplayTimingViewExtender.h"
+
+#if WITH_EDITOR
+#include "EditorViewportClient.h"
+#include "Editor/EditorEngine.h"
+#endif
 
 #define LOCTEXT_NAMESPACE "AnimationSharedData"
 
@@ -70,10 +77,10 @@ void FAnimationSharedData::Tick(Insights::ITimingViewSession& InTimingViewSessio
 					return false;
 				};
 
-				TSharedPtr<FSkeletalMeshPoseTrack> SkeletalMeshPoseTrack = StaticCastSharedPtr<FSkeletalMeshPoseTrack>(ObjectEventsTrack->GetGameplayTrack().FindChildTrack(InObjectInfo.Id, FindSkeletalMeshPoseTrack));
-				if(!SkeletalMeshPoseTrack.IsValid())
+				TSharedPtr<FSkeletalMeshPoseTrack> ExistingSkeletalMeshPoseTrack = StaticCastSharedPtr<FSkeletalMeshPoseTrack>(ObjectEventsTrack->GetGameplayTrack().FindChildTrack(InObjectInfo.Id, FindSkeletalMeshPoseTrack));
+				if(!ExistingSkeletalMeshPoseTrack.IsValid())
 				{
-					SkeletalMeshPoseTrack = MakeShared<FSkeletalMeshPoseTrack>(*this, InObjectInfo.Id, InObjectInfo.Name);
+					TSharedPtr<FSkeletalMeshPoseTrack> SkeletalMeshPoseTrack = MakeShared<FSkeletalMeshPoseTrack>(*this, InObjectInfo.Id, InObjectInfo.Name);
 					SkeletalMeshPoseTrack->SetVisibilityFlag(bAnimationTracksEnabled);
 					SkeletalMeshPoseTracks.Add(SkeletalMeshPoseTrack.ToSharedRef());
 
@@ -84,7 +91,7 @@ void FAnimationSharedData::Tick(Insights::ITimingViewSession& InTimingViewSessio
 				}
 			});
 
-			AnimationProvider->EnumerateTickRecordTimelines(InObjectInfo.Id, [this, &InObjectInfo, &ObjectEventsTrack, &InTimingViewSession, &GameplayProvider](uint64 InAssetId, const IAnimationProvider::TickRecordTimeline& InTimeline)
+			AnimationProvider->EnumerateTickRecordTimelines(InObjectInfo.Id, [this, &InObjectInfo, &ObjectEventsTrack, &InTimingViewSession, &GameplayProvider](uint64 InAssetId, int32 InNodeId, const IAnimationProvider::TickRecordTimeline& InTimeline)
 			{
 				ObjectEventsTrack->SetVisibilityFlag(GameplaySharedData.AreGameplayTracksEnabled());
 
@@ -100,12 +107,12 @@ void FAnimationSharedData::Tick(Insights::ITimingViewSession& InTimingViewSessio
 					return false;
 				};
 
-				TSharedPtr<FAnimationTickRecordsTrack> AnimationTickRecordsTrack = StaticCastSharedPtr<FAnimationTickRecordsTrack>(ObjectEventsTrack->GetGameplayTrack().FindChildTrack(InObjectInfo.Id, FindTickRecordTrackWithAssetId));
-				if(!AnimationTickRecordsTrack.IsValid())
+				TSharedPtr<FAnimationTickRecordsTrack> ExistingAnimationTickRecordsTrack = StaticCastSharedPtr<FAnimationTickRecordsTrack>(ObjectEventsTrack->GetGameplayTrack().FindChildTrack(InObjectInfo.Id, FindTickRecordTrackWithAssetId));
+				if(!ExistingAnimationTickRecordsTrack.IsValid())
 				{
 					const FObjectInfo* AssetObjectInfo = GameplayProvider->FindObjectInfo(InAssetId);
 					FString AssetName = AssetObjectInfo ? AssetObjectInfo->Name : LOCTEXT("UnknownAsset", "Unknown").ToString();
-					AnimationTickRecordsTrack = MakeShared<FAnimationTickRecordsTrack>(*this, InObjectInfo.Id, InAssetId, *AssetName);
+					TSharedPtr<FAnimationTickRecordsTrack> AnimationTickRecordsTrack = MakeShared<FAnimationTickRecordsTrack>(*this, InObjectInfo.Id, InAssetId, InNodeId, *AssetName);
 					AnimationTickRecordsTrack->SetVisibilityFlag(bAnimationTracksEnabled);
 					AnimationTickRecordsTracks.Add(AnimationTickRecordsTrack.ToSharedRef());
 
@@ -115,6 +122,35 @@ void FAnimationSharedData::Tick(Insights::ITimingViewSession& InTimingViewSessio
 					ObjectEventsTrack->GetGameplayTrack().AddChildTrack(AnimationTickRecordsTrack->GetGameplayTrack());
 				}
 			});
+
+			AnimationProvider->ReadAnimGraphTimeline(InObjectInfo.Id, [this, &InObjectInfo, &ObjectEventsTrack, &InTimingViewSession](const IAnimationProvider::AnimGraphTimeline& InTimeline)
+			{
+				ObjectEventsTrack->SetVisibilityFlag(GameplaySharedData.AreGameplayTracksEnabled());
+
+				auto FindAnimNodesTrack = [](const FBaseTimingTrack& InTrack)
+				{
+					if (InTrack.GetType() == FAnimNodesTrack::TypeName &&
+						InTrack.GetSubType() == FAnimNodesTrack::SubTypeName)
+					{
+						return true;
+					}
+
+					return false;
+				};
+
+				TSharedPtr<FAnimNodesTrack> ExistingAnimNodesTrack = StaticCastSharedPtr<FAnimNodesTrack>(ObjectEventsTrack->GetGameplayTrack().FindChildTrack(InObjectInfo.Id, FindAnimNodesTrack));
+				if(!ExistingAnimNodesTrack.IsValid())
+				{
+					TSharedPtr<FAnimNodesTrack> AnimNodesTrack = MakeShared<FAnimNodesTrack>(*this, InObjectInfo.Id, InObjectInfo.Name);
+					AnimNodesTrack->SetVisibilityFlag(bAnimationTracksEnabled);
+					AnimNodesTracks.Add(AnimNodesTrack.ToSharedRef());
+
+					InTimingViewSession.AddScrollableTrack(AnimNodesTrack);
+					GameplaySharedData.InvalidateObjectTracksOrder();
+
+					ObjectEventsTrack->GetGameplayTrack().AddChildTrack(AnimNodesTrack->GetGameplayTrack());
+				}
+			});
 		});
 	}
 
@@ -171,8 +207,44 @@ void FAnimationSharedData::OnTimeMarkerChanged(Insights::ETimeChangedFlags InFla
 {
 	bTimeMarkerValid = InTimeMarker != std::numeric_limits<double>::infinity();
 	MarkerTime = InTimeMarker;
+
+#if WITH_EDITOR
+	for(const TSharedRef<FAnimNodesTrack> AnimNodesTrack : AnimNodesTracks)
+	{
+		AnimNodesTrack->UpdateDebugData(InTimeMarker);
+	}
+
+	InvalidateViewports();
+
+	// Update pose tracks even if they are disabled, as they may be being debugged
+	UWorld* WorldToUse = FGameplayTimingViewExtender::GetWorldToVisualize();
+	for(TSharedRef<FSkeletalMeshPoseTrack> PoseTrack : SkeletalMeshPoseTracks)
+	{
+		if(bTimeMarkerValid && PoseTrack->IsPotentiallyDebugged())
+		{
+			PoseTrack->DrawPoses(WorldToUse, MarkerTime);
+		}
+	}
+#endif
 }
 
+#if WITH_EDITOR
+void FAnimationSharedData::InvalidateViewports()
+{
+	UEditorEngine* Engine = Cast<UEditorEngine>(GEngine);
+	if (GIsEditor && Engine != nullptr)
+	{
+		for (FEditorViewportClient* ViewportClient : Engine->GetAllViewportClients())
+		{
+			if (ViewportClient)
+			{
+				ViewportClient->Invalidate();
+			}
+		}
+	}
+}
+#endif
+
 void FAnimationSharedData::EnumerateSkeletalMeshPoseTracks(TFunctionRef<void(const TSharedRef<FSkeletalMeshPoseTrack>&)> InCallback) const
 {
 	for(const TSharedRef<FSkeletalMeshPoseTrack>& Track : SkeletalMeshPoseTracks)
@@ -181,6 +253,44 @@ void FAnimationSharedData::EnumerateSkeletalMeshPoseTracks(TFunctionRef<void(con
 	}
 }
 
+TSharedPtr<FSkeletalMeshPoseTrack> FAnimationSharedData::FindSkeletalMeshPoseTrack(uint64 InComponentId) const
+{
+	TSharedPtr<FSkeletalMeshPoseTrack> FoundTrack;
+
+	EnumerateSkeletalMeshPoseTracks([&FoundTrack, InComponentId](const TSharedRef<FSkeletalMeshPoseTrack>& InTrack)
+	{
+		if(InTrack->GetGameplayTrack().GetObjectId() == InComponentId)
+		{
+			FoundTrack = InTrack;
+		}
+	});
+
+	return FoundTrack;
+}
+
+void FAnimationSharedData::EnumerateAnimNodesTracks(TFunctionRef<void(const TSharedRef<FAnimNodesTrack>&)> InCallback) const
+{
+	for(const TSharedRef<FAnimNodesTrack>& Track : AnimNodesTracks)
+	{
+		InCallback(Track);
+	}
+}
+
+TSharedPtr<FAnimNodesTrack> FAnimationSharedData::FindAnimNodesTrack(uint64 InAnimInstanceId) const
+{
+	TSharedPtr<FAnimNodesTrack> FoundTrack;
+
+	EnumerateAnimNodesTracks([&FoundTrack, InAnimInstanceId](const TSharedRef<FAnimNodesTrack>& InTrack)
+	{
+		if(InTrack->GetGameplayTrack().GetObjectId() == InAnimInstanceId)
+		{
+			FoundTrack = InTrack;
+		}
+	});
+
+	return FoundTrack;
+}
+
 #if WITH_ENGINE
 
 void FAnimationSharedData::DrawPoses(UWorld* InWorld)
@@ -199,4 +309,19 @@ void FAnimationSharedData::DrawPoses(UWorld* InWorld)
 
 #endif
 
+#if WITH_EDITOR
+
+void FAnimationSharedData::GetCustomDebugObjects(const IAnimationBlueprintEditor& InAnimationBlueprintEditor, TArray<FCustomDebugObject>& OutDebugList)
+{
+	for(const TSharedRef<FAnimNodesTrack> AnimNodesTrack : AnimNodesTracks)
+	{
+		if(AnimNodesTrack->IsVisible())
+		{
+			AnimNodesTrack->GetCustomDebugObjects(InAnimationBlueprintEditor, OutDebugList);
+		}
+	}
+}
+
+#endif
+
 #undef LOCTEXT_NAMESPACE
\ No newline at end of file
diff --git a/Engine/Plugins/Animation/GameplayInsights/Source/GameplayInsights/Private/AnimationSharedData.h b/Engine/Plugins/Animation/GameplayInsights/Source/GameplayInsights/Private/AnimationSharedData.h
--- a/Engine/Plugins/Animation/GameplayInsights/Source/GameplayInsights/Private/AnimationSharedData.h
+++ b/Engine/Plugins/Animation/GameplayInsights/Source/GameplayInsights/Private/AnimationSharedData.h
@@ -10,8 +10,11 @@ namespace Insights { class ITimingViewSession; }
 namespace Insights { enum class ETimeChangedFlags : int32; }
 class FSkeletalMeshPoseTrack;
 class FAnimationTickRecordsTrack;
+class FAnimNodesTrack;
 class FMenuBuilder;
 class UWorld;
+class IAnimationBlueprintEditor;
+struct FCustomDebugObject;
 
 class FAnimationSharedData
 {
@@ -27,6 +30,14 @@ public:
 	void DrawPoses(UWorld* InWorld);
 #endif
 
+#if WITH_EDITOR
+	// Get the debug objects to plug into the anim BP debugger
+	void GetCustomDebugObjects(const IAnimationBlueprintEditor& InAnimationBlueprintEditor, TArray<FCustomDebugObject>& OutDebugList);
+
+	// Helper function to invalidate all viewports so non-realtime viewports update correctly.
+	void InvalidateViewports();
+#endif
+
 	// Check whether animation tracks are enabled
 	bool AreAnimationTracksEnabled() const;
 
@@ -42,6 +53,15 @@ public:
 	// Enumerate skeletal mesh pose tracks
 	void EnumerateSkeletalMeshPoseTracks(TFunctionRef<void(const TSharedRef<FSkeletalMeshPoseTrack>&)> InCallback) const;
 
+	// Find a skeletal mesh track with the specified component ID
+	TSharedPtr<FSkeletalMeshPoseTrack> FindSkeletalMeshPoseTrack(uint64 InComponentId) const;
+
+	// Enumerate anim nodes tracks
+	void EnumerateAnimNodesTracks(TFunctionRef<void(const TSharedRef<FAnimNodesTrack>&)> InCallback) const;
+
+	// Find an anim nodes track with the specified anim instance ID
+	TSharedPtr<FAnimNodesTrack> FindAnimNodesTrack(uint64 InAnimInstanceId) const;
+
 private:
 	// UI handlers
 	void ToggleAnimationTracks();
@@ -59,6 +79,7 @@ private:
 	// All the tracks we manage
 	TArray<TSharedRef<FSkeletalMeshPoseTrack>> SkeletalMeshPoseTracks;
 	TArray<TSharedRef<FAnimationTickRecordsTrack>> AnimationTickRecordsTracks;
+	TArray<TSharedRef<FAnimNodesTrack>> AnimNodesTracks;
 
 	// Delegate handles for hooks into the timing view
 	FDelegateHandle TimeMarkerChangedHandle;
diff --git a/Engine/Plugins/Animation/GameplayInsights/Source/GameplayInsights/Private/AnimationTickRecordsTrack.cpp b/Engine/Plugins/Animation/GameplayInsights/Source/GameplayInsights/Private/AnimationTickRecordsTrack.cpp
--- a/Engine/Plugins/Animation/GameplayInsights/Source/GameplayInsights/Private/AnimationTickRecordsTrack.cpp
+++ b/Engine/Plugins/Animation/GameplayInsights/Source/GameplayInsights/Private/AnimationTickRecordsTrack.cpp
@@ -9,42 +9,361 @@
 #include "AnimationSharedData.h"
 #include "Insights/ViewModels/TimingEventSearch.h"
 #include "Insights/ViewModels/TooltipDrawState.h"
+#include "Templates/Invoke.h"
+#include "Insights/ViewModels/GraphTrackBuilder.h"
+#include "Insights/Common/TimeUtils.h"
+#include "GameplaySharedData.h"
+#include "Framework/MultiBox/MultiBoxBuilder.h"
+
+#if WITH_EDITOR
+#include "Animation/AnimBlueprintGeneratedClass.h"
+#include "IAnimationBlueprintEditor.h"
+#include "Animation/AnimBlueprint.h"
+#include "Subsystems/AssetEditorSubsystem.h"
+#include "EdGraph/EdGraphNode.h"
+#endif
 
 #define LOCTEXT_NAMESPACE "AnimationTickRecordsTrack"
 
-const FName FAnimationTickRecordsTrack::TypeName(TEXT("Animation"));
-const FName FAnimationTickRecordsTrack::SubTypeName(TEXT("TickRecords"));
+const FName FAnimationTickRecordsTrack::TypeName(TEXT("Graph"));
+const FName FAnimationTickRecordsTrack::SubTypeName(TEXT("Animation.TickRecords"));
+
+FString FTickRecordSeries::FormatValue(double Value) const
+{
+	switch (Type)
+	{
+	case ESeriesType::PlaybackTime:
+		return TimeUtils::FormatTimeAuto(Value);
+	case ESeriesType::BlendWeight:
+	case ESeriesType::RootMotionWeight:
+	case ESeriesType::PlayRate:
+	case ESeriesType::BlendSpacePositionX:
+	case ESeriesType::BlendSpacePositionY:
+		return FText::AsNumber(Value).ToString();
+	}
+
+	return FGraphSeries::FormatValue(Value);
+}
+
+void FTickRecordSeries::UpdateAutoZoom(const FTimingTrackViewport& InViewport, const FAnimationTickRecordsTrack& InTrack)
+{
+	const float TimelineDY = FMath::Max(1.0f, InViewport.GetLayout().TimelineDY);
+	const float TopY = FMath::Max(1.0f, TimelineDY);
+	const float BottomY = FMath::Max(TopY, InTrack.GetHeight() - TimelineDY);
+
+	FGraphSeries::UpdateAutoZoom(TopY, BottomY, CurrentMin, CurrentMax);
+}
+
+static FLinearColor MakeSeriesColor(uint32 InSeed, bool bInLine = false)
+{
+	FRandomStream Stream(InSeed);
+	const uint8 Hue = (uint8)(Stream.FRand() * 255.0f);
+	const uint8 SatVal = bInLine ? 196 : 128;
+	return FLinearColor::MakeFromHSV8(Hue, SatVal, SatVal);
+}
+
+static FLinearColor MakeSeriesColor(FTickRecordSeries::ESeriesType InSeed, bool bInLine = false)
+{
+	return MakeSeriesColor((uint32)InSeed, bInLine);
+}
 
-FAnimationTickRecordsTrack::FAnimationTickRecordsTrack(const FAnimationSharedData& InSharedData, uint64 InObjectID, uint64 InAssetId, const TCHAR* InName)
-	: TGameplayTrackMixin<FTimingEventsTrack>(InObjectID, FAnimationTickRecordsTrack::TypeName, FAnimationTickRecordsTrack::SubTypeName, FText::FromString(FString(InName)))
+FAnimationTickRecordsTrack::FAnimationTickRecordsTrack(const FAnimationSharedData& InSharedData, uint64 InObjectID, uint64 InAssetId, int32 InNodeId, const TCHAR* InName)
+	: TGameplayTrackMixin<FGraphTrack>(InObjectID, FAnimationTickRecordsTrack::SubTypeName, MakeTrackName(InSharedData.GetGameplaySharedData(), InAssetId, InName))
 	, SharedData(InSharedData)
 	, AssetId(InAssetId)
+	, NodeId(InNodeId)
+	, RequestedTrackSizeScale(1.0f)
 {
+	uint32 Hash = GetTypeHash(GetName());
+	MainSeriesLineColor = MakeSeriesColor(Hash, true);
+	MainSeriesFillColor = MakeSeriesColor(Hash);
+
+	AddAllSeries();
+
+	bDrawPoints = false;
+	bDrawBoxes = false;
+	bDrawBaseline = false;
+	bUseEventDuration = false;
+
+#if WITH_EDITOR
+	const FGameplayProvider* GameplayProvider = SharedData.GetAnalysisSession().ReadProvider<FGameplayProvider>(FGameplayProvider::ProviderName);
+
+	if(GameplayProvider)
+	{
+		const FObjectInfo* AnimInstanceInfo = GameplayProvider->FindObjectInfo(GetGameplayTrack().GetObjectId());
+		if(AnimInstanceInfo)
+		{
+			const FClassInfo* AnimInstanceClassInfo = GameplayProvider->FindClassInfo(AnimInstanceInfo->ClassId);
+			if(AnimInstanceClassInfo)
+			{
+				InstanceClass = FindObject<UAnimBlueprintGeneratedClass>(ANY_PACKAGE, AnimInstanceClassInfo->PathName);
+			}
+		}
+	}
+#endif
 }
 
-void FAnimationTickRecordsTrack::BuildDrawState(ITimingEventsTrackDrawStateBuilder& Builder, const ITimingTrackUpdateContext& Context)
+void FAnimationTickRecordsTrack::AddAllSeries()
 {
+	struct FSeriesDescription
+	{
+		FText Name;
+		FText Description;
+		FLinearColor LineColor;
+		FLinearColor FillColor;
+		FTickRecordSeries::ESeriesType Type;
+		bool bEnabled;
+	};
+
+	const FSeriesDescription SeriesDescriptions[] =
+	{
+		{ 
+			LOCTEXT("SeriesNameBlendWeight", "Blend Weight"), 
+			LOCTEXT("SeriesDescBlendWeight", "The final effective weight that this animation sequence was played at"), 
+			MainSeriesLineColor, 
+			MainSeriesFillColor, 
+			FTickRecordSeries::ESeriesType::BlendWeight,
+			true
+		},
+		{ 
+			LOCTEXT("SeriesNamePlaybackTime", "Playback Time"), 
+			LOCTEXT("SeriesDescPlaybackTime", "The playback time of this animation sequence"), 
+			MakeSeriesColor(FTickRecordSeries::ESeriesType::PlaybackTime, true), 
+			MakeSeriesColor(FTickRecordSeries::ESeriesType::PlaybackTime),
+			FTickRecordSeries::ESeriesType::PlaybackTime,
+			false
+		},
+		{ 
+			LOCTEXT("SeriesNameRootMotionWeight", "Root Motion Weight"),
+			LOCTEXT("SeriesDescRootMotionWeight", "The final effective root motion weight that this animation sequence was played at"), 
+			MakeSeriesColor(FTickRecordSeries::ESeriesType::RootMotionWeight, true), 
+			MakeSeriesColor(FTickRecordSeries::ESeriesType::RootMotionWeight), 
+			FTickRecordSeries::ESeriesType::RootMotionWeight, 
+			false
+		},
+		{ 
+			LOCTEXT("SeriesNamePlayRate", "Play Rate"), 
+			LOCTEXT("SeriesDescPlayRate", "The play rate/speed of this animation sequence"), 
+			MakeSeriesColor(FTickRecordSeries::ESeriesType::PlayRate, true),
+			MakeSeriesColor(FTickRecordSeries::ESeriesType::PlayRate), 
+			FTickRecordSeries::ESeriesType::PlayRate,
+			false
+		},
+	};
+
+	const FSeriesDescription BlendSpaceSeriesDescriptions[] =
+	{
+		{ 
+			LOCTEXT("SeriesNameBlendSpacePositionX", "BlendSpace Position X"), 
+			LOCTEXT("SeriesDescBlendSpacePositionX", "The X value used to sample this blend space"), 
+			MakeSeriesColor(FTickRecordSeries::ESeriesType::BlendSpacePositionX, true),
+			MakeSeriesColor(FTickRecordSeries::ESeriesType::BlendSpacePositionX),
+			FTickRecordSeries::ESeriesType::BlendSpacePositionX, 
+			false
+		},
+		{
+			LOCTEXT("SeriesNameBlendSpacePositionY", "BlendSpace Position Y"), 
+			LOCTEXT("SeriesDescBlendSpacePositionY", "The Y value used to sample this blend space"), 
+			MakeSeriesColor(FTickRecordSeries::ESeriesType::BlendSpacePositionY, true), 
+			MakeSeriesColor(FTickRecordSeries::ESeriesType::BlendSpacePositionY), 
+			FTickRecordSeries::ESeriesType::BlendSpacePositionY, 
+			false
+		},
+	};
+
+	auto AddSeries = [this](const FSeriesDescription& InSeriesDescription)
+	{
+		TSharedRef<FTickRecordSeries> Series = MakeShared<FTickRecordSeries>();
+		Series->SetName(InSeriesDescription.Name.ToString());
+		Series->SetDescription(InSeriesDescription.Description.ToString());
+		Series->SetColor(InSeriesDescription.LineColor, InSeriesDescription.LineColor, InSeriesDescription.FillColor);
+		Series->Type = InSeriesDescription.Type;
+		Series->SetVisibility(InSeriesDescription.bEnabled);
+		Series->SetBaselineY(25.0f);
+		Series->SetScaleY(20.0f);
+		Series->EnableAutoZoom();
+		AllSeries.Add(Series);	
+	};
+
+	for(const FSeriesDescription& SeriesDescription : SeriesDescriptions)
+	{
+		AddSeries(SeriesDescription);
+	}
+
 	const FGameplayProvider* GameplayProvider = SharedData.GetAnalysisSession().ReadProvider<FGameplayProvider>(FGameplayProvider::ProviderName);
-	const FAnimationProvider* AnimationProvider = SharedData.GetAnalysisSession().ReadProvider<FAnimationProvider>(FAnimationProvider::ProviderName);
+	if(GameplayProvider)
+	{
+		const FClassInfo& ClassInfo = GameplayProvider->GetClassInfoFromObject(AssetId);
+		if(FCString::Stristr(ClassInfo.Name, TEXT("BlendSpace")) != nullptr)
+		{
+			for(const FSeriesDescription& BlendSpaceSeriesDescription : BlendSpaceSeriesDescriptions)
+			{
+				AddSeries(BlendSpaceSeriesDescription);
+			}
+		}
+	}
+}
+
+template<typename ProjectionType>
+static void UpdateSeries(const FAnimationSharedData& InSharedData, FAnimationTickRecordsTrack& InTrack, FTickRecordSeries& InSeries, const FTimingTrackViewport& InViewport, ProjectionType Projection)
+{		
+	const FGameplayProvider* GameplayProvider = InSharedData.GetAnalysisSession().ReadProvider<FGameplayProvider>(FGameplayProvider::ProviderName);
+	const FAnimationProvider* AnimationProvider = InSharedData.GetAnalysisSession().ReadProvider<FAnimationProvider>(FAnimationProvider::ProviderName);
 
 	if(GameplayProvider && AnimationProvider)
 	{
-		Trace::FAnalysisSessionReadScope SessionReadScope(SharedData.GetAnalysisSession());
+		Trace::FAnalysisSessionReadScope SessionReadScope(InSharedData.GetAnalysisSession());
+
+		// Calc visible range before we build events (as builder uses scale factors calculated in auto zoom!)
+		InSeries.CurrentMin = 0.0;
+		InSeries.CurrentMax = 0.0;
+
+		AnimationProvider->ReadTickRecordTimeline(InTrack.GetGameplayTrack().GetObjectId(), InTrack.GetAssetId(), InTrack.GetNodeId(), [&InViewport, &InSeries, &Projection](const FAnimationProvider::TickRecordTimeline& InTimeline)
+		{
+			InTimeline.EnumerateEvents(InViewport.GetStartTime(), InViewport.GetEndTime(), [&InSeries, &Projection](double InStartTime, double InEndTime, uint32 InDepth, const FTickRecordMessage& InMessage)
+			{
+				const double Value = Invoke(Projection, InMessage);
+				InSeries.CurrentMin = FMath::Min(InSeries.CurrentMin, Value);
+				InSeries.CurrentMax = FMath::Max(InSeries.CurrentMax, Value);
+			});
+		});
 
-		AnimationProvider->ReadTickRecordTimeline(GetGameplayTrack().GetObjectId(), GetAssetId(), [this, &GameplayProvider, &AnimationProvider, &Context, &Builder](const FAnimationProvider::TickRecordTimeline& InTimeline)
+		InSeries.UpdateAutoZoom(InViewport, InTrack);
+
+		FGraphTrackBuilder Builder(InTrack, InSeries, InViewport);
+
+		AnimationProvider->ReadTickRecordTimeline(InTrack.GetGameplayTrack().GetObjectId(), InTrack.GetAssetId(), InTrack.GetNodeId(), [&InTrack, &GameplayProvider, &AnimationProvider, &Builder, &InViewport, &Projection](const FAnimationProvider::TickRecordTimeline& InTimeline)
 		{
-			InTimeline.EnumerateEvents(Context.GetViewport().GetStartTime(), Context.GetViewport().GetEndTime(), [this, &Builder, &GameplayProvider](double InStartTime, double InEndTime, uint32 InDepth, const FTickRecordMessage& InMessage)
+			InTimeline.EnumerateEvents(InViewport.GetStartTime(), InViewport.GetEndTime(), [&InTrack, &Builder, &GameplayProvider, &Projection](double InStartTime, double InEndTime, uint32 InDepth, const FTickRecordMessage& InMessage)
 			{
-				Builder.AddEvent(InStartTime, InEndTime, 0, *GetName());
+				Builder.AddEvent(InStartTime, InEndTime - InStartTime, Invoke(Projection, InMessage), InMessage.bContinuous);
+
+				InTrack.HeightInLanes = 1;
 			});
 		});
 	}
 }
 
+void FAnimationTickRecordsTrack::UpdateTrackHeight(const ITimingTrackUpdateContext& Context)
+{
+	const FTimingTrackViewport& Viewport = Context.GetViewport();
+
+	const float CurrentTrackHeight = GetHeight();
+	const float DesiredTrackHeight = Viewport.GetLayout().ComputeTrackHeight(HeightInLanes) * RequestedTrackSizeScale;
+
+	if (CurrentTrackHeight < DesiredTrackHeight)
+	{
+		float NewTrackHeight;
+		if (Viewport.IsDirty(ETimingTrackViewportDirtyFlags::VLayoutChanged))
+		{
+			NewTrackHeight = DesiredTrackHeight;
+		}
+		else
+		{
+			NewTrackHeight = FMath::CeilToFloat(CurrentTrackHeight * 0.9f + DesiredTrackHeight * 0.1f);
+		}
+
+		SetHeight(NewTrackHeight);
+
+		for (TSharedPtr<FGraphSeries>& Series : AllSeries)
+		{
+			StaticCastSharedPtr<FTickRecordSeries>(Series)->UpdateAutoZoom(Context.GetViewport(), *this);
+		}
+	}
+	else if (CurrentTrackHeight > DesiredTrackHeight)
+	{
+		float NewTrackHeight;
+		if (Viewport.IsDirty(ETimingTrackViewportDirtyFlags::VLayoutChanged))
+		{
+			NewTrackHeight = DesiredTrackHeight;
+		}
+		else
+		{
+			NewTrackHeight = FMath::FloorToFloat(CurrentTrackHeight * 0.9f + DesiredTrackHeight * 0.1f);
+		}
+
+		SetHeight(NewTrackHeight);
+
+		for (TSharedPtr<FGraphSeries>& Series : AllSeries)
+		{
+			StaticCastSharedPtr<FTickRecordSeries>(Series)->UpdateAutoZoom(Context.GetViewport(), *this);
+		}
+	}
+}
+
+void FAnimationTickRecordsTrack::PreUpdate(const ITimingTrackUpdateContext& Context)
+{
+	FGraphTrack::PreUpdate(Context);
+
+	// update border size
+	BorderY = Context.GetViewport().GetLayout().TimelineDY;
+
+	const bool bIsEntireGraphTrackDirty = IsDirty() || Context.GetViewport().IsHorizontalViewportDirty();
+	bool bNeedsUpdate = bIsEntireGraphTrackDirty;
+
+	if (!bNeedsUpdate)
+	{
+		for (TSharedPtr<FGraphSeries>& Series : AllSeries)
+		{
+			if (Series->IsVisible() && Series->IsDirty())
+			{
+				// At least one series is dirty.
+				bNeedsUpdate = true;
+				break;
+			}
+		}
+	}
+
+	if (bNeedsUpdate)
+	{
+		ClearDirtyFlag();
+
+		HeightInLanes = 0;
+
+		const FTimingTrackViewport& Viewport = Context.GetViewport();
+
+		for (TSharedPtr<FGraphSeries>& Series : AllSeries)
+		{
+			if (Series->IsVisible() && (bIsEntireGraphTrackDirty || Series->IsDirty()))
+			{
+				// Clear the flag before updating, because the update itself may further need to set the series as dirty.
+				Series->ClearDirtyFlag();
+
+				TSharedPtr<FTickRecordSeries> TickRecordSeries = StaticCastSharedPtr<FTickRecordSeries>(Series);
+				switch (TickRecordSeries->Type)
+				{
+				case FTickRecordSeries::ESeriesType::BlendWeight:
+					UpdateSeries(SharedData, *this, *TickRecordSeries, Viewport, &FTickRecordMessage::BlendWeight);
+					break;
+				case FTickRecordSeries::ESeriesType::PlaybackTime:
+					UpdateSeries(SharedData, *this, *TickRecordSeries, Viewport, &FTickRecordMessage::PlaybackTime);
+					break;
+				case FTickRecordSeries::ESeriesType::RootMotionWeight:
+					UpdateSeries(SharedData, *this, *TickRecordSeries, Viewport, &FTickRecordMessage::RootMotionWeight);
+					break;
+				case FTickRecordSeries::ESeriesType::PlayRate:
+					UpdateSeries(SharedData, *this, *TickRecordSeries, Viewport, &FTickRecordMessage::PlayRate);
+					break;
+				case FTickRecordSeries::ESeriesType::BlendSpacePositionX:
+					UpdateSeries(SharedData, *this, *TickRecordSeries, Viewport, &FTickRecordMessage::BlendSpacePositionX);
+					break;
+				case FTickRecordSeries::ESeriesType::BlendSpacePositionY:
+					UpdateSeries(SharedData, *this, *TickRecordSeries, Viewport, &FTickRecordMessage::BlendSpacePositionY);
+					break;
+				}
+			}
+		}
+
+		UpdateStats();
+	}
+
+	UpdateTrackHeight(Context);
+}
+
 void FAnimationTickRecordsTrack::Draw(const ITimingTrackDrawContext& Context) const
 {
-	DrawEvents(Context);
-	GetGameplayTrack().DrawHeaderForTimingTrack(Context, *this);
+	FGraphTrack::Draw(Context);
+	GetGameplayTrack().DrawHeaderForTimingTrack(Context, *this, true);
 }
 
 void FAnimationTickRecordsTrack::InitTooltip(FTooltipDrawState& Tooltip, const ITimingEvent& HoveredTimingEvent) const
@@ -59,10 +378,15 @@ void FAnimationTickRecordsTrack::InitTooltip(FTooltipDrawState& Tooltip, const I
 
 		Tooltip.AddNameValueTextLine(LOCTEXT("EventTime", "Time").ToString(), FText::AsNumber(InFoundStartTime).ToString());
 		Tooltip.AddNameValueTextLine(LOCTEXT("BlendWeight", "Blend Weight").ToString(), FText::AsNumber(InMessage.BlendWeight).ToString());
+		if(InMessage.bIsBlendSpace)
+		{
+			Tooltip.AddNameValueTextLine(LOCTEXT("BlendSpacePositionX", "X").ToString(), FText::AsNumber(InMessage.BlendSpacePositionX).ToString());
+			Tooltip.AddNameValueTextLine(LOCTEXT("BlendSpacePositionX", "Y").ToString(), FText::AsNumber(InMessage.BlendSpacePositionY).ToString());
+		}
 		Tooltip.AddNameValueTextLine(LOCTEXT("PlaybackTime", "Playback Time").ToString(), FText::AsNumber(InMessage.PlaybackTime).ToString());
 		Tooltip.AddNameValueTextLine(LOCTEXT("RootMotionWeight", "Root Motion Weight").ToString(), FText::AsNumber(InMessage.RootMotionWeight).ToString());
 		Tooltip.AddNameValueTextLine(LOCTEXT("PlayRate", "Play Rate").ToString(), FText::AsNumber(InMessage.PlayRate).ToString());
-		Tooltip.AddNameValueTextLine(LOCTEXT("Looping", "Looping").ToString(), InMessage.Looping ? LOCTEXT("True", "True").ToString() : LOCTEXT("False", "False").ToString());
+		Tooltip.AddNameValueTextLine(LOCTEXT("Looping", "Looping").ToString(), InMessage.bLooping ? LOCTEXT("True", "True").ToString() : LOCTEXT("False", "False").ToString());
 
 		Tooltip.UpdateLayout();
 	});
@@ -93,7 +417,7 @@ void FAnimationTickRecordsTrack::FindTickRecordMessage(const FTimingEventSearchP
 			{
 				Trace::FAnalysisSessionReadScope SessionReadScope(SharedData.GetAnalysisSession());
 
-				AnimationProvider->ReadTickRecordTimeline(GetGameplayTrack().GetObjectId(), GetAssetId(), [this, &InContext](const FAnimationProvider::TickRecordTimeline& InTimeline)
+				AnimationProvider->ReadTickRecordTimeline(GetGameplayTrack().GetObjectId(), AssetId, NodeId, [this, &InContext](const FAnimationProvider::TickRecordTimeline& InTimeline)
 				{
 					InTimeline.EnumerateEvents(InContext.GetParameters().StartTime, InContext.GetParameters().EndTime, [&InContext](double InEventStartTime, double InEventEndTime, uint32 InDepth, const FTickRecordMessage& InMessage)
 					{
@@ -103,10 +427,120 @@ void FAnimationTickRecordsTrack::FindTickRecordMessage(const FTimingEventSearchP
 			}
 		},
 
+		[&InParameters](double InFoundStartTime, double InFoundEndTime, uint32 InFoundDepth, const FTickRecordMessage& InEvent)
+		{
+			// Match the start time exactly here
+			return InFoundStartTime == InParameters.StartTime;
+		},
+
 		[&InFoundPredicate](double InFoundStartTime, double InFoundEndTime, uint32 InFoundDepth, const FTickRecordMessage& InEvent)
 		{
 			InFoundPredicate(InFoundStartTime, InFoundEndTime, InFoundDepth, InEvent);
-		});
+		},
+			
+		TTimingEventSearch<FTickRecordMessage>::NoMatch);
+}
+
+void FAnimationTickRecordsTrack::BuildContextMenu(FMenuBuilder& MenuBuilder)
+{
+#if WITH_EDITOR
+	MenuBuilder.BeginSection("TrackActions", LOCTEXT("TrackActionsMenuHeader", "Track Actions"));
+	{
+		MenuBuilder.AddMenuEntry
+		(
+			LOCTEXT("FindAssetPlayerNode", "Find Asset Player Node"),
+			LOCTEXT("FindAssetPlayerNode_Tooltip", "Open the animation blueprint that this animation was played from."),
+			FSlateIcon(),
+			FUIAction(
+				FExecuteAction::CreateLambda([this]()
+				{
+					if(InstanceClass)
+					{
+						if(UAnimBlueprint* AnimBlueprint = Cast<UAnimBlueprint>(InstanceClass->ClassGeneratedBy))
+						{
+							GEditor->GetEditorSubsystem<UAssetEditorSubsystem>()->OpenEditorForAsset(AnimBlueprint);
+
+							if(IAnimationBlueprintEditor* AnimBlueprintEditor = static_cast<IAnimationBlueprintEditor*>(GEditor->GetEditorSubsystem<UAssetEditorSubsystem>()->FindEditorForAsset(AnimBlueprint, true)))
+							{
+								int32 AnimNodeIndex = InstanceClass->AnimNodeProperties.Num() - NodeId - 1;
+								TWeakObjectPtr<const UEdGraphNode> GraphNode = InstanceClass->AnimBlueprintDebugData.NodePropertyIndexToNodeMap[AnimNodeIndex];
+								if(GraphNode.Get())
+								{
+									AnimBlueprintEditor->JumpToHyperlink(GraphNode.Get());
+								}
+							}
+						}
+					}
+				})
+			),
+			NAME_None,
+			EUserInterfaceActionType::Button
+		);
+	}
+	MenuBuilder.EndSection();
+#endif
+
+	FGraphTrack::BuildContextMenu(MenuBuilder);
+
+	MenuBuilder.BeginSection("TrackSize", LOCTEXT("TrackSizeMenuHeader", "Track Size"));
+	{
+		MenuBuilder.AddMenuEntry
+		(
+			LOCTEXT("NormalTrack", "Normal"),
+			LOCTEXT("NormalTrack_Tooltip", "Draw this track at the standard size."),
+			FSlateIcon(),
+			FUIAction(
+				FExecuteAction::CreateLambda([this](){ RequestedTrackSizeScale = 1.0f; }),
+				FCanExecuteAction(),
+				FIsActionChecked::CreateLambda([this](){ return RequestedTrackSizeScale == 1.0f; })
+			),
+			NAME_None,
+			EUserInterfaceActionType::RadioButton
+		);
+
+		MenuBuilder.AddMenuEntry
+		(
+			LOCTEXT("LargeTrack", "Large"),
+			LOCTEXT("LargeTrack_Tooltip", "Make this track larger than normal."),
+			FSlateIcon(),
+			FUIAction(
+				FExecuteAction::CreateLambda([this](){ RequestedTrackSizeScale = 2.0f; }),
+				FCanExecuteAction(),
+				FIsActionChecked::CreateLambda([this](){ return RequestedTrackSizeScale == 2.0f; })
+			),
+			NAME_None,
+			EUserInterfaceActionType::RadioButton
+		);
+
+		MenuBuilder.AddMenuEntry
+		(
+			LOCTEXT("ExtraLargeTrack", "Extra Large"),
+			LOCTEXT("ExtraLargeTrack_Tooltip", "Make this track much larger than normal."),
+			FSlateIcon(),
+			FUIAction(
+				FExecuteAction::CreateLambda([this](){ RequestedTrackSizeScale = 4.0f; }),
+				FCanExecuteAction(),
+				FIsActionChecked::CreateLambda([this](){ return RequestedTrackSizeScale == 4.0f; })
+			),
+			NAME_None,
+			EUserInterfaceActionType::RadioButton
+		);
+	}
+	MenuBuilder.EndSection();
+}
+
+FText FAnimationTickRecordsTrack::MakeTrackName(const FGameplaySharedData& InSharedData, uint64 InAssetID, const TCHAR* InName) const
+{
+	FText AssetTypeName = LOCTEXT("UnknownAsset", "Unknown");
+
+	const FGameplayProvider* GameplayProvider = InSharedData.GetAnalysisSession().ReadProvider<FGameplayProvider>(FGameplayProvider::ProviderName);
+	if(GameplayProvider)
+	{
+		const FClassInfo& ClassInfo = GameplayProvider->GetClassInfoFromObject(InAssetID);
+		AssetTypeName = FText::FromString(ClassInfo.Name);
+	}
+
+	return FText::Format(LOCTEXT("AnimationTickRecordsTrackName", "{0} - {1}"), AssetTypeName, FText::FromString(InName));
 }
 
 #undef LOCTEXT_NAMESPACE
\ No newline at end of file
diff --git a/Engine/Plugins/Animation/GameplayInsights/Source/GameplayInsights/Private/AnimationTickRecordsTrack.h b/Engine/Plugins/Animation/GameplayInsights/Source/GameplayInsights/Private/AnimationTickRecordsTrack.h
--- a/Engine/Plugins/Animation/GameplayInsights/Source/GameplayInsights/Private/AnimationTickRecordsTrack.h
+++ b/Engine/Plugins/Animation/GameplayInsights/Source/GameplayInsights/Private/AnimationTickRecordsTrack.h
@@ -3,36 +3,107 @@
 #pragma once
 
 #include "GameplayTrack.h"
-#include "Insights/ViewModels/TimingEventsTrack.h"
+#include "Insights/ViewModels/GraphTrack.h"
+#include "Insights/ViewModels/GraphSeries.h"
 
 class FAnimationSharedData;
 struct FTickRecordMessage;
 class FTimingEventSearchParameters;
+class FGameplaySharedData;
+class UAnimBlueprintGeneratedClass;
+class FAnimationTickRecordsTrack;
+class FTimingTrackViewport;
 
-class FAnimationTickRecordsTrack : public TGameplayTrackMixin<FTimingEventsTrack>
+class FTickRecordSeries : public FGraphSeries
+{
+public:
+	enum class ESeriesType : uint32
+	{
+		BlendWeight,
+		PlaybackTime,
+		RootMotionWeight,
+		PlayRate,
+		BlendSpacePositionX,
+		BlendSpacePositionY,
+
+		Count,
+	};
+
+	virtual FString FormatValue(double Value) const override;
+
+	// Custom overload for auto zoom
+	void UpdateAutoZoom(const FTimingTrackViewport& InViewport, const FAnimationTickRecordsTrack& InTrack);
+
+public:
+	ESeriesType Type;
+	double CurrentMin;
+	double CurrentMax;
+};
+
+class FAnimationTickRecordsTrack : public TGameplayTrackMixin<FGraphTrack>
 {
 public:
 	static const FName TypeName;
 	static const FName SubTypeName;
 
-	FAnimationTickRecordsTrack(const FAnimationSharedData& InSharedData, uint64 InObjectId, uint64 InAssetId, const TCHAR* InName);
+	FAnimationTickRecordsTrack(const FAnimationSharedData& InSharedData, uint64 InObjectId, uint64 InAssetId, int32 InNodeId, const TCHAR* InName);
 
-	virtual void BuildDrawState(ITimingEventsTrackDrawStateBuilder& Builder, const ITimingTrackUpdateContext& Context) override;
+	virtual void PreUpdate(const ITimingTrackUpdateContext& Context) override;
 	virtual void Draw(const ITimingTrackDrawContext& Context) const override;
 	virtual void InitTooltip(FTooltipDrawState& Tooltip, const ITimingEvent& HoveredTimingEvent) const override;
 	virtual const TSharedPtr<const ITimingEvent> SearchEvent(const FTimingEventSearchParameters& InSearchParameters) const override;
+	virtual void BuildContextMenu(FMenuBuilder& MenuBuilder) override;
+
+	virtual float GetBorderY() const { return BorderY; }
 
 	/** Get the asset ID that this track uses */
 	uint64 GetAssetId() const { return AssetId; }
 
+	/** Get the node ID that this track uses */
+	int32 GetNodeId() const { return NodeId; }
+
+	/** The current height in lanes */
+	int32 HeightInLanes;
+
 private:
+	// Add all the series that this track can represent
+	void AddAllSeries();
+
+	// Helper for PreUpdate to update the track height
+	void UpdateTrackHeight(const ITimingTrackUpdateContext& Context);
+
 	// Helper function used to find a tick record
 	void FindTickRecordMessage(const FTimingEventSearchParameters& InParameters, TFunctionRef<void(double, double, uint32, const FTickRecordMessage&)> InFoundPredicate) const;
 
+	// Override certain UI items to be disabled
+	virtual bool ContextMenu_UseEventDuration_CanExecute() { return false; }
+	virtual bool ContextMenu_ShowBars_CanExecute() { return false; }
+
+	// Helper used to build track name
+	FText MakeTrackName(const FGameplaySharedData& InSharedData, uint64 InAssetID, const TCHAR* InName) const;
+
 private:
 	/** The shared data */
 	const FAnimationSharedData& SharedData;
 
+#if WITH_EDITOR
+	/** The class that output this tick record */
+	UAnimBlueprintGeneratedClass* InstanceClass;
+#endif
+
+	/** Colors for line drawing */
+	FLinearColor MainSeriesLineColor;
+	FLinearColor MainSeriesFillColor;
+
 	/** The asset ID that this track uses */
 	uint64 AssetId;
+
+	/** The node ID that this track's data comes from */
+	int32 NodeId;
+
+	/** The track size we want to be displayed at */
+	float RequestedTrackSizeScale;
+
+	/** The size of the border between tracks */
+	float BorderY;
 };
\ No newline at end of file
diff --git a/Engine/Plugins/Animation/GameplayInsights/Source/GameplayInsights/Private/GameplayInsightsModule.cpp b/Engine/Plugins/Animation/GameplayInsights/Source/GameplayInsights/Private/GameplayInsightsModule.cpp
--- a/Engine/Plugins/Animation/GameplayInsights/Source/GameplayInsights/Private/GameplayInsightsModule.cpp
+++ b/Engine/Plugins/Animation/GameplayInsights/Source/GameplayInsights/Private/GameplayInsightsModule.cpp
@@ -8,10 +8,16 @@
 #include "Containers/Ticker.h"
 #include "Modules/ModuleManager.h"
 
+#if WITH_EDITOR
+#include "IAnimationBlueprintEditorModule.h"
+#endif
+
 #if WITH_ENGINE
 #include "Engine/Engine.h"
 #endif
 
+#define LOCTEXT_NAMESPACE "GameplayInsightsModule"
+
 class FGameplayInsightsModule : public IModuleInterface
 {
 public:
@@ -26,10 +32,26 @@ public:
 			GameplayTimingViewExtender.TickVisualizers(DeltaTime);
 			return true;
 		});
+
+#if WITH_EDITOR
+		IAnimationBlueprintEditorModule& AnimationBlueprintEditorModule = FModuleManager::LoadModuleChecked<IAnimationBlueprintEditorModule>("AnimationBlueprintEditor");
+		CustomDebugObjectHandle = AnimationBlueprintEditorModule.OnGetCustomDebugObjects().AddLambda([this](const IAnimationBlueprintEditor& InAnimationBlueprintEditor, TArray<FCustomDebugObject>& OutDebugList)
+		{
+			GameplayTimingViewExtender.GetCustomDebugObjects(InAnimationBlueprintEditor, OutDebugList);
+		});
+#endif
 	}
 
 	virtual void ShutdownModule() override
 	{
+#if WITH_EDITOR
+		IAnimationBlueprintEditorModule* AnimationBlueprintEditorModule = FModuleManager::GetModulePtr<IAnimationBlueprintEditorModule>("AnimationBlueprintEditor");
+		if(AnimationBlueprintEditorModule)
+		{
+			AnimationBlueprintEditorModule->OnGetCustomDebugObjects().Remove(CustomDebugObjectHandle);
+		}
+#endif
+
 		FTicker::GetCoreTicker().RemoveTicker(TickerHandle);
 
 		IModularFeatures::Get().UnregisterModularFeature(Trace::ModuleFeatureName, &GameplayTraceModule);
@@ -40,6 +62,12 @@ public:
 	FGameplayTimingViewExtender GameplayTimingViewExtender;
 
 	FDelegateHandle TickerHandle;
+
+#if WITH_EDITOR
+	FDelegateHandle CustomDebugObjectHandle;
+#endif
 };
 
 IMPLEMENT_MODULE(FGameplayInsightsModule, GameplayInsights);
+
+#undef LOCTEXT_NAMESPACE
\ No newline at end of file
diff --git a/Engine/Plugins/Animation/GameplayInsights/Source/GameplayInsights/Private/GameplayProvider.cpp b/Engine/Plugins/Animation/GameplayInsights/Source/GameplayInsights/Private/GameplayProvider.cpp
--- a/Engine/Plugins/Animation/GameplayInsights/Source/GameplayInsights/Private/GameplayProvider.cpp
+++ b/Engine/Plugins/Animation/GameplayInsights/Source/GameplayInsights/Private/GameplayProvider.cpp
@@ -4,6 +4,8 @@
 
 FName FGameplayProvider::ProviderName("GameplayProvider");
 
+#define LOCTEXT_NAMESPACE "GameplayProvider"
+
 FGameplayProvider::FGameplayProvider(Trace::IAnalysisSession& InSession)
 	: Session(InSession)
 	, EndPlayEvent(nullptr)
@@ -76,6 +78,49 @@ const FObjectInfo* FGameplayProvider::FindObjectInfo(uint64 InObjectId) const
 	return nullptr;
 }
 
+const FClassInfo& FGameplayProvider::GetClassInfo(uint64 InClassId) const
+{
+	const FClassInfo* ClassInfo = FindClassInfo(InClassId);
+	if(ClassInfo)
+	{
+		return *ClassInfo;
+	}
+
+	static FText UnknownText(LOCTEXT("Unknown", "Unknown"));
+	static FClassInfo DefaultClassInfo = { 0, 0, *UnknownText.ToString(), *UnknownText.ToString() };
+	return DefaultClassInfo;
+}
+
+const FClassInfo& FGameplayProvider::GetClassInfoFromObject(uint64 InObjectId) const
+{
+	const FObjectInfo* ObjectInfo = FindObjectInfo(InObjectId);
+	if(ObjectInfo)
+	{
+		const FClassInfo* ClassInfo = FindClassInfo(ObjectInfo->ClassId);
+		if(ClassInfo)
+		{
+			return *ClassInfo;
+		}
+	}
+
+	static FText UnknownText(LOCTEXT("Unknown", "Unknown"));
+	static FClassInfo DefaultClassInfo = { 0, 0, *UnknownText.ToString(), *UnknownText.ToString() };
+	return DefaultClassInfo;
+}
+
+const FObjectInfo& FGameplayProvider::GetObjectInfo(uint64 InObjectId) const
+{
+	const FObjectInfo* ObjectInfo = FindObjectInfo(InObjectId);
+	if(ObjectInfo)
+	{
+		return *ObjectInfo;
+	}
+
+	static FText UnknownText(LOCTEXT("Unknown", "Unknown"));
+	static FObjectInfo DefaultObjectInfo = { 0, 0, 0, *UnknownText.ToString(), *UnknownText.ToString() };
+	return DefaultObjectInfo;
+}
+
 void FGameplayProvider::AppendClass(uint64 InClassId, uint64 InSuperId, const TCHAR* InClassName, const TCHAR* InClassPathName)
 {
 	Session.WriteAccessCheck();
@@ -156,3 +201,5 @@ void FGameplayProvider::AppendObjectEvent(uint64 InObjectId, double InTime, cons
 
 	Session.UpdateDurationSeconds(InTime);
 }
+
+#undef LOCTEXT_NAMESPACE
\ No newline at end of file
diff --git a/Engine/Plugins/Animation/GameplayInsights/Source/GameplayInsights/Private/GameplayProvider.h b/Engine/Plugins/Animation/GameplayInsights/Source/GameplayInsights/Private/GameplayProvider.h
--- a/Engine/Plugins/Animation/GameplayInsights/Source/GameplayInsights/Private/GameplayProvider.h
+++ b/Engine/Plugins/Animation/GameplayInsights/Source/GameplayInsights/Private/GameplayProvider.h
@@ -22,6 +22,9 @@ public:
 	virtual void EnumerateObjects(TFunctionRef<void(const FObjectInfo&)> Callback) const override;
 	virtual const FClassInfo* FindClassInfo(uint64 InClassId) const override;
 	virtual const FObjectInfo* FindObjectInfo(uint64 InObjectId) const override;
+	virtual const FClassInfo& GetClassInfo(uint64 InClassId) const override;
+	virtual const FClassInfo& GetClassInfoFromObject(uint64 InObjectId) const override;
+	virtual const FObjectInfo& GetObjectInfo(uint64 InObjectId) const override;
 	virtual FOnObjectEndPlay& OnObjectEndPlay() override { return OnObjectEndPlayDelegate; }
 
 	/** Add a class message */
diff --git a/Engine/Plugins/Animation/GameplayInsights/Source/GameplayInsights/Private/GameplaySharedData.cpp b/Engine/Plugins/Animation/GameplayInsights/Source/GameplayInsights/Private/GameplaySharedData.cpp
--- a/Engine/Plugins/Animation/GameplayInsights/Source/GameplayInsights/Private/GameplaySharedData.cpp
+++ b/Engine/Plugins/Animation/GameplayInsights/Source/GameplayInsights/Private/GameplaySharedData.cpp
@@ -75,19 +75,19 @@ TSharedRef<FObjectEventsTrack> FGameplaySharedData::GetObjectEventsTrackForId(In
 	return LeafObjectEventsTrack.ToSharedRef();
 }
 
-static void UpdateTrackOrderRecursive(TSharedRef<FObjectEventsTrack> InTrack, int32& InOutOrder, uint32 InDepth)
+static void UpdateTrackOrderRecursive(TSharedRef<FObjectEventsTrack> InTrack, int32& InOutOrder)
 {
 	// recurse down object-track children, then non-object track leaf tracks to set 
 	// overall ordering based on depth-first traversal of the hierarchy
 
 	InTrack->SetOrder(InOutOrder++);
-	InTrack->GetGameplayTrack().SetIndent(InDepth);
 
 	for(FGameplayTrack* ChildTrack : InTrack->GetGameplayTrack().GetChildTracks())
 	{
 		if(ChildTrack->GetTimingTrack()->GetType() == FObjectEventsTrack::TypeName)
 		{
-			UpdateTrackOrderRecursive(StaticCastSharedRef<FObjectEventsTrack>(ChildTrack->GetTimingTrack()), InOutOrder, InDepth + 1);
+			ChildTrack->SetIndent(InTrack->GetGameplayTrack().GetIndent() + 1);
+			UpdateTrackOrderRecursive(StaticCastSharedRef<FObjectEventsTrack>(ChildTrack->GetTimingTrack()), InOutOrder);
 		}
 	}
 
@@ -95,8 +95,8 @@ static void UpdateTrackOrderRecursive(TSharedRef<FObjectEventsTrack> InTrack, in
 	{
 		if(ChildTrack->GetTimingTrack()->GetType() != FObjectEventsTrack::TypeName)
 		{
+			ChildTrack->SetIndent(InTrack->GetGameplayTrack().GetIndent() + 1);
 			ChildTrack->GetTimingTrack()->SetOrder(InOutOrder++);
-			ChildTrack->SetIndent(InDepth + 1);
 		}
 	}
 }
@@ -169,10 +169,10 @@ void FGameplaySharedData::SortTracks()
 		return InTrack0->GetName() < InTrack1->GetName();
 	});
 
-	// update ordering
+	// update ordering/indent
 	for(TSharedRef<FObjectEventsTrack> RootTrack : Roots)
 	{
-		UpdateTrackOrderRecursive(RootTrack, Order, 0);
+		UpdateTrackOrderRecursive(RootTrack, Order);
 	}
 
 	Roots.Reset();
diff --git a/Engine/Plugins/Animation/GameplayInsights/Source/GameplayInsights/Private/GameplayTimingViewExtender.cpp b/Engine/Plugins/Animation/GameplayInsights/Source/GameplayInsights/Private/GameplayTimingViewExtender.cpp
--- a/Engine/Plugins/Animation/GameplayInsights/Source/GameplayInsights/Private/GameplayTimingViewExtender.cpp
+++ b/Engine/Plugins/Animation/GameplayInsights/Source/GameplayInsights/Private/GameplayTimingViewExtender.cpp
@@ -64,7 +64,7 @@ void FGameplayTimingViewExtender::ExtendFilterMenu(Insights::ITimingViewSession&
 
 #if WITH_ENGINE
 
-static UWorld* GetWorldToVisualize()
+UWorld* FGameplayTimingViewExtender::GetWorldToVisualize()
 {
 	UWorld* World = nullptr;
 
@@ -99,4 +99,16 @@ void FGameplayTimingViewExtender::TickVisualizers(float DeltaTime)
 		}
 	}
 #endif
-}
\ No newline at end of file
+}
+
+#if WITH_EDITOR
+
+void FGameplayTimingViewExtender::GetCustomDebugObjects(const IAnimationBlueprintEditor& InAnimationBlueprintEditor, TArray<FCustomDebugObject>& OutDebugList)
+{
+	for(auto& PerSessionData : PerSessionDataMap)
+	{
+		PerSessionData.Value.AnimationSharedData->GetCustomDebugObjects(InAnimationBlueprintEditor, OutDebugList);
+	}
+}
+
+#endif
\ No newline at end of file
diff --git a/Engine/Plugins/Animation/GameplayInsights/Source/GameplayInsights/Private/GameplayTimingViewExtender.h b/Engine/Plugins/Animation/GameplayInsights/Source/GameplayInsights/Private/GameplayTimingViewExtender.h
--- a/Engine/Plugins/Animation/GameplayInsights/Source/GameplayInsights/Private/GameplayTimingViewExtender.h
+++ b/Engine/Plugins/Animation/GameplayInsights/Source/GameplayInsights/Private/GameplayTimingViewExtender.h
@@ -8,6 +8,8 @@ namespace Insights { class ITimingViewSession; }
 class UWorld;
 class FGameplaySharedData;
 class FAnimationSharedData;
+class IAnimationBlueprintEditor;
+struct FCustomDebugObject;
 
 class FGameplayTimingViewExtender : public Insights::ITimingViewExtender
 {
@@ -18,6 +20,14 @@ public:
 	virtual void Tick(Insights::ITimingViewSession& InSession, const Trace::IAnalysisSession& InAnalysisSession) override;
 	virtual void ExtendFilterMenu(Insights::ITimingViewSession& InSession, FMenuBuilder& InMenuBuilder) override;
 
+#if WITH_EDITOR
+	// Gets a world to perform visualizations within, depending on context
+	static UWorld* GetWorldToVisualize();
+
+	// Get custom debug objects for integration with anim blueprint debugging
+	void GetCustomDebugObjects(const IAnimationBlueprintEditor& InAnimationBlueprintEditor, TArray<FCustomDebugObject>& OutDebugList);
+#endif
+
 	// Tick the visualizers
 	void TickVisualizers(float DeltaTime);
 
diff --git a/Engine/Plugins/Animation/GameplayInsights/Source/GameplayInsights/Private/GameplayTrack.cpp b/Engine/Plugins/Animation/GameplayInsights/Source/GameplayInsights/Private/GameplayTrack.cpp
--- a/Engine/Plugins/Animation/GameplayInsights/Source/GameplayInsights/Private/GameplayTrack.cpp
+++ b/Engine/Plugins/Animation/GameplayInsights/Source/GameplayInsights/Private/GameplayTrack.cpp
@@ -47,30 +47,32 @@ static FORCEINLINE bool IntervalsIntersect(float Min1, float Max1, float Min2, f
 	return !(Max2 < Min1 || Max1 < Min2);
 }
 
-void FGameplayTrack::DrawHeaderForTimingTrack(const ITimingTrackDrawContext& InContext, const FTimingEventsTrack& InTrack) const
+void FGameplayTrack::DrawHeaderForTimingTrack(const ITimingTrackDrawContext& InContext, const FBaseTimingTrack& InTrack, bool bUsePreallocatedLayers) const
 {
 	const float X = (float)Indent * GameplayTrackConstants::IndentSize;
 	const float Y = InTrack.GetPosY();
 	const float H = InTrack.GetHeight();
 	const float TrackNameH = H > 7.0f ? 12.0f : H;
+	const int32 HeaderBackgroundLayerId = bUsePreallocatedLayers ? InContext.GetHelper().GetHeaderBackgroundLayerId() : InContext.GetDrawContext().LayerId++;
+	const int32 HeaderTextLayerId = bUsePreallocatedLayers ? InContext.GetHelper().GetHeaderTextLayerId() : InContext.GetDrawContext().LayerId++;
 
 	if (H > 0.0f &&
 		Y + H > InContext.GetViewport().GetTopOffset() &&
 		Y < InContext.GetViewport().GetHeight() - InContext.GetViewport().GetBottomOffset())
 	{
 		// Draw a horizontal line between timelines.
-		InContext.GetDrawContext().DrawBox(InContext.GetHelper().GetHeaderBackgroundLayerId(), X, Y, InContext.GetViewport().GetWidth(), 1.0f, InContext.GetHelper().GetWhiteBrush(), InContext.GetHelper().GetEdgeColor());
+		InContext.GetDrawContext().DrawBox(HeaderBackgroundLayerId, X, Y, InContext.GetViewport().GetWidth(), 1.0f, InContext.GetHelper().GetWhiteBrush(), InContext.GetHelper().GetEdgeColor());
 
 		if(H > 7.0f)
 		{
 			const TSharedRef<FSlateFontMeasure> FontMeasureService = FSlateApplication::Get().GetRenderer()->GetFontMeasureService();
 			const float NameWidth = FontMeasureService->Measure(InTrack.GetName(), InContext.GetHelper().GetEventFont()).X;
-			InContext.GetDrawContext().DrawBox(InContext.GetHelper().GetHeaderBackgroundLayerId(), X, Y + 1.0f, NameWidth + 4.0f, TrackNameH, InContext.GetHelper().GetWhiteBrush(), InContext.GetHelper().GetEdgeColor());
-			InContext.GetDrawContext().DrawText(InContext.GetHelper().GetHeaderTextLayerId(), X + 2.0f, Y, InTrack.GetName(), InContext.GetHelper().GetEventFont(), InContext.GetHelper().GetTrackNameTextColor(InTrack));
+			InContext.GetDrawContext().DrawBox(HeaderBackgroundLayerId, X, Y + 1.0f, NameWidth + 4.0f, TrackNameH, InContext.GetHelper().GetWhiteBrush(), InContext.GetHelper().GetEdgeColor());
+			InContext.GetDrawContext().DrawText(HeaderTextLayerId, X + 2.0f, Y, InTrack.GetName(), InContext.GetHelper().GetEventFont(), InContext.GetHelper().GetTrackNameTextColor(InTrack));
 		}
 		else
 		{
-			InContext.GetDrawContext().DrawBox(InContext.GetHelper().GetHeaderBackgroundLayerId(), X, Y + 1.0f, H, H, InContext.GetHelper().GetWhiteBrush(), InContext.GetHelper().GetEdgeColor());
+			InContext.GetDrawContext().DrawBox(HeaderBackgroundLayerId, X, Y + 1.0f, H, H, InContext.GetHelper().GetWhiteBrush(), InContext.GetHelper().GetEdgeColor());
 		}
 	}
 
@@ -82,8 +84,8 @@ void FGameplayTrack::DrawHeaderForTimingTrack(const ITimingTrackDrawContext& InC
 
 		if (IntervalsIntersect(ParentY, Y, InContext.GetViewport().GetTopOffset(), InContext.GetViewport().GetHeight() - InContext.GetViewport().GetBottomOffset()))
 		{
-			InContext.GetDrawContext().DrawBox(InContext.GetHelper().GetHeaderBackgroundLayerId(), ParentX, Y + (TrackNameH * 0.5f), X - ParentX, 1.0f, InContext.GetHelper().GetWhiteBrush(), InContext.GetHelper().GetEdgeColor());
-			InContext.GetDrawContext().DrawBox(InContext.GetHelper().GetHeaderBackgroundLayerId(), ParentX, ParentY, 1.0f, (Y - ParentY) + (TrackNameH * 0.5f), InContext.GetHelper().GetWhiteBrush(), InContext.GetHelper().GetEdgeColor());
+			InContext.GetDrawContext().DrawBox(HeaderBackgroundLayerId, ParentX, Y + (TrackNameH * 0.5f), X - ParentX, 1.0f, InContext.GetHelper().GetWhiteBrush(), InContext.GetHelper().GetEdgeColor());
+			InContext.GetDrawContext().DrawBox(HeaderBackgroundLayerId, ParentX, ParentY, 1.0f, (Y - ParentY) + (TrackNameH * 0.5f), InContext.GetHelper().GetWhiteBrush(), InContext.GetHelper().GetEdgeColor());
 		}
 	}
 }
diff --git a/Engine/Plugins/Animation/GameplayInsights/Source/GameplayInsights/Private/GameplayTrack.h b/Engine/Plugins/Animation/GameplayInsights/Source/GameplayInsights/Private/GameplayTrack.h
--- a/Engine/Plugins/Animation/GameplayInsights/Source/GameplayInsights/Private/GameplayTrack.h
+++ b/Engine/Plugins/Animation/GameplayInsights/Source/GameplayInsights/Private/GameplayTrack.h
@@ -33,7 +33,7 @@ public:
 	TSharedPtr<FBaseTimingTrack> FindChildTrack(uint64 InObjectId, TFunctionRef<bool(const FBaseTimingTrack& InTrack)> Callback) const;
 
 	/** Helper to draw the name for a timing track (uses indentation etc.) */
-	void DrawHeaderForTimingTrack(const ITimingTrackDrawContext& InContext, const FTimingEventsTrack& InTrack) const;
+	void DrawHeaderForTimingTrack(const ITimingTrackDrawContext& InContext, const FBaseTimingTrack& InTrack, bool bUsePreallocatedLayers) const;
 
 	/** Access the outer timing track */
 	TSharedRef<FBaseTimingTrack> GetTimingTrack() const { return TimingTrack.AsShared(); }
@@ -65,6 +65,12 @@ template <class Base>
 class TGameplayTrackMixin : public Base
 {
 public:
+	TGameplayTrackMixin(uint64 InObjectId, const FName& InSubType, const FText& InName)
+		: Base(InSubType, InName.ToString())
+		, GameplayTrack(*this, InObjectId)
+	{
+	}
+
 	TGameplayTrackMixin(uint64 InObjectId, const FName& InType, const FName& InSubType, const FText& InName)
 		: Base(InType, InSubType, InName.ToString())
 		, GameplayTrack(*this, InObjectId)
diff --git a/Engine/Plugins/Animation/GameplayInsights/Source/GameplayInsights/Private/IAnimationProvider.h b/Engine/Plugins/Animation/GameplayInsights/Source/GameplayInsights/Private/IAnimationProvider.h
--- a/Engine/Plugins/Animation/GameplayInsights/Source/GameplayInsights/Private/IAnimationProvider.h
+++ b/Engine/Plugins/Animation/GameplayInsights/Source/GameplayInsights/Private/IAnimationProvider.h
@@ -34,17 +34,146 @@ struct FSkeletalMeshPoseMessage
 	uint16 LodIndex = 0;
 };
 
+struct FSkeletalMeshFrameMessage
+{
+	uint64 ComponentId = 0;
+	uint16 FrameCounter = 0;
+};
+
 struct FTickRecordMessage
 {
 	uint64 ComponentId = 0;
 	uint64 AnimInstanceId = 0;
 	uint64 AssetId = 0;
+	int32 NodeId = -1;
 	float BlendWeight = 0.0f;
 	float PlaybackTime = 0.0f;
 	float RootMotionWeight = 0.0f;
 	float PlayRate = 0.0f;
+	float BlendSpacePositionX = 0.0f;
+	float BlendSpacePositionY = 0.0f;
 	uint16 FrameCounter = 0;
-	bool Looping = false;
+	bool bLooping = false;
+	bool bIsBlendSpace = false;
+	bool bContinuous = true;
+};
+
+enum class EAnimGraphPhase : uint8
+{
+	Initialize = 0,
+	PreUpdate = 1,
+	Update = 2,
+	CacheBones = 3,
+	Evaluate = 4,
+};
+
+struct FAnimGraphMessage
+{
+	uint64 AnimInstanceId = 0;
+	int32 NodeCount = 0;
+	uint16 FrameCounter = 0;
+	EAnimGraphPhase Phase = EAnimGraphPhase::Initialize;
+};
+
+struct FAnimNodeMessage
+{
+	const TCHAR* NodeName = nullptr;
+	uint64 AnimInstanceId = 0;
+	int32 PreviousNodeId = -1;
+	int32 NodeId = -1;
+	float Weight = 0.0f;
+	float RootMotionWeight = 0.0f;
+	uint16 FrameCounter = 0;
+	EAnimGraphPhase Phase = EAnimGraphPhase::Initialize;
+};
+
+enum class EAnimNodeValueType : uint8
+{
+	Bool,
+	Int32,
+	Float,
+	Vector,
+	String,
+	Object,
+	Class,
+};
+
+struct FAnimNodeValueMessage
+{
+	FAnimNodeValueMessage()
+		: Vector(FVector::ZeroVector)
+	{
+	}
+
+	const TCHAR* Key = nullptr;
+	uint64 AnimInstanceId = 0;
+	union
+	{
+		struct
+		{
+			bool bValue;
+		} Bool;
+		struct
+		{
+			int32 Value;
+		} Int32;
+		struct
+		{
+			float Value;
+		} Float;
+		struct FVectorEntry
+		{
+			FVectorEntry(const FVector& InVector)
+				: Value(InVector)
+			{}
+
+			FVector Value;
+		} Vector;
+		struct
+		{
+			const TCHAR* Value;
+		} String;
+		struct
+		{
+			uint64 Value;
+		} Object;
+		struct
+		{
+			uint64 Value;
+		} Class;
+	};
+	int32 NodeId = -1;
+	uint16 FrameCounter = 0;
+	EAnimNodeValueType Type = EAnimNodeValueType::Bool;
+};
+
+struct FAnimSequencePlayerMessage
+{
+	uint64 AnimInstanceId = 0;
+	int32 NodeId = -1;
+	float Position = 0.0f;
+	float Length = 0.0f;
+	int32 FrameCount = 0;
+};
+
+struct FBlendSpacePlayerMessage
+{
+	uint64 AnimInstanceId = 0;
+	uint64 BlendSpaceId = 0;
+	int32 NodeId = -1;
+	float PositionX = 0.0f;
+	float PositionY = 0.0f;
+	float PositionZ = 0.0f;
+};
+
+struct FAnimStateMachineMessage
+{
+	uint64 AnimInstanceId = 0;
+	int32 NodeId = -1;
+	int32 StateMachineIndex = -1;
+	int32 StateIndex = -1;
+	float StateWeight = 0.0f;
+	float ElapsedTime = 0.0f;
 };
 
 class IAnimationProvider : public Trace::IProvider
@@ -52,13 +181,23 @@ class IAnimationProvider : public Trace::IProvider
 public:
 	typedef Trace::ITimeline<FTickRecordMessage> TickRecordTimeline;
 	typedef Trace::ITimeline<FSkeletalMeshPoseMessage> SkeletalMeshPoseTimeline;
+	typedef Trace::ITimeline<FAnimGraphMessage> AnimGraphTimeline;
+	typedef Trace::ITimeline<FAnimNodeMessage> AnimNodesTimeline;
+	typedef Trace::ITimeline<FAnimNodeValueMessage> AnimNodeValuesTimeline;
+	typedef Trace::ITimeline<FAnimSequencePlayerMessage> AnimSequencePlayersTimeline;
+	typedef Trace::ITimeline<FAnimStateMachineMessage> StateMachinesTimeline;
+	typedef Trace::ITimeline<FBlendSpacePlayerMessage> BlendSpacePlayersTimeline;
 
 	virtual bool ReadSkeletalMeshPoseTimeline(uint64 InObjectId, TFunctionRef<void(const SkeletalMeshPoseTimeline&)> Callback) const = 0;
-	virtual bool ReadSkeletalMeshPoseMessage(uint64 InObjectId, uint64 InMessageId, TFunctionRef<void(const FSkeletalMeshPoseMessage&)> Callback) const = 0;
 	virtual void GetSkeletalMeshComponentSpacePose(const FSkeletalMeshPoseMessage& InMessage, const FSkeletalMeshInfo& InMeshInfo, FTransform& OutComponentToWorld, TArray<FTransform>& OutTransforms) const = 0;
-	virtual void EnumerateTickRecordTimelines(uint64 InObjectId, TFunctionRef<void(uint64, const TickRecordTimeline&)> Callback) const = 0;
-	virtual bool ReadTickRecordTimeline(uint64 InObjectId, uint64 InAssetId, TFunctionRef<void(const TickRecordTimeline&)> Callback) const = 0;
-	virtual bool ReadTickRecordMessage(uint64 InObjectId, uint64 InAssetId, uint64 InMessageId, TFunctionRef<void(const FTickRecordMessage&)> Callback) const = 0;
+	virtual void EnumerateTickRecordTimelines(uint64 InObjectId, TFunctionRef<void(uint64, int32, const TickRecordTimeline&)> Callback) const = 0;
+	virtual bool ReadTickRecordTimeline(uint64 InObjectId, uint64 InAssetId, int32 InNodeId, TFunctionRef<void(const TickRecordTimeline&)> Callback) const = 0;
+	virtual bool ReadAnimGraphTimeline(uint64 InObjectId, TFunctionRef<void(const AnimGraphTimeline&)> Callback) const = 0;
+	virtual bool ReadAnimNodesTimeline(uint64 InObjectId, TFunctionRef<void(const AnimNodesTimeline&)> Callback) const = 0;
+	virtual bool ReadAnimNodeValuesTimeline(uint64 InObjectId, TFunctionRef<void(const AnimNodeValuesTimeline&)> Callback) const = 0;
+	virtual bool ReadAnimSequencePlayersTimeline(uint64 InObjectId, TFunctionRef<void(const AnimSequencePlayersTimeline&)> Callback) const = 0;
+	virtual bool ReadAnimBlendSpacePlayersTimeline(uint64 InObjectId, TFunctionRef<void(const BlendSpacePlayersTimeline&)> Callback) const = 0;
+	virtual bool ReadStateMachinesTimeline(uint64 InObjectId, TFunctionRef<void(const StateMachinesTimeline&)> Callback) const = 0;
 	virtual const FSkeletalMeshInfo* FindSkeletalMeshInfo(uint64 InObjectId) const = 0;
 	virtual const TCHAR* GetName(uint32 InId) const = 0;
 };
\ No newline at end of file
diff --git a/Engine/Plugins/Animation/GameplayInsights/Source/GameplayInsights/Private/IGameplayProvider.h b/Engine/Plugins/Animation/GameplayInsights/Source/GameplayInsights/Private/IGameplayProvider.h
--- a/Engine/Plugins/Animation/GameplayInsights/Source/GameplayInsights/Private/IGameplayProvider.h
+++ b/Engine/Plugins/Animation/GameplayInsights/Source/GameplayInsights/Private/IGameplayProvider.h
@@ -42,5 +42,8 @@ public:
 	virtual void EnumerateObjects(TFunctionRef<void(const FObjectInfo&)> Callback) const = 0;
 	virtual const FClassInfo* FindClassInfo(uint64 InClassId) const = 0;
 	virtual const FObjectInfo* FindObjectInfo(uint64 InObjectId) const = 0;
+	virtual const FClassInfo& GetClassInfo(uint64 InClassId) const = 0;
+	virtual const FClassInfo& GetClassInfoFromObject(uint64 InObjectId) const = 0;
+	virtual const FObjectInfo& GetObjectInfo(uint64 InObjectId) const = 0;
 	virtual FOnObjectEndPlay& OnObjectEndPlay() = 0;
 };
\ No newline at end of file
diff --git a/Engine/Plugins/Animation/GameplayInsights/Source/GameplayInsights/Private/InsightsSkeletalMeshComponent.cpp b/Engine/Plugins/Animation/GameplayInsights/Source/GameplayInsights/Private/InsightsSkeletalMeshComponent.cpp
--- a/Engine/Plugins/Animation/GameplayInsights/Source/GameplayInsights/Private/InsightsSkeletalMeshComponent.cpp
+++ b/Engine/Plugins/Animation/GameplayInsights/Source/GameplayInsights/Private/InsightsSkeletalMeshComponent.cpp
@@ -1,9 +1,9 @@
 // Copyright 1998-2019 Epic Games, Inc. All Rights Reserved.
 
 #include "InsightsSkeletalMeshComponent.h"
-#include "AnimationProvider.h"
+#include "IAnimationProvider.h"
 
-void UInsightsSkeletalMeshComponent::SetPoseFromProvider(const FAnimationProvider& InProvider, const FSkeletalMeshPoseMessage& InMessage, const FSkeletalMeshInfo& SkeletalMeshInfo)
+void UInsightsSkeletalMeshComponent::SetPoseFromProvider(const IAnimationProvider& InProvider, const FSkeletalMeshPoseMessage& InMessage, const FSkeletalMeshInfo& SkeletalMeshInfo)
 {
 	// Grab transform and bone transforms from provider
 	FTransform ComponentToWorldTransform;
diff --git a/Engine/Plugins/Animation/GameplayInsights/Source/GameplayInsights/Private/InsightsSkeletalMeshComponent.h b/Engine/Plugins/Animation/GameplayInsights/Source/GameplayInsights/Private/InsightsSkeletalMeshComponent.h
--- a/Engine/Plugins/Animation/GameplayInsights/Source/GameplayInsights/Private/InsightsSkeletalMeshComponent.h
+++ b/Engine/Plugins/Animation/GameplayInsights/Source/GameplayInsights/Private/InsightsSkeletalMeshComponent.h
@@ -6,7 +6,7 @@
 #include "Components/SkeletalMeshComponent.h"
 #include "InsightsSkeletalMeshComponent.generated.h"
 
-class FAnimationProvider;
+class IAnimationProvider;
 struct FSkeletalMeshPoseMessage;
 struct FSkeletalMeshInfo;
 
@@ -17,7 +17,7 @@ class UInsightsSkeletalMeshComponent : public USkeletalMeshComponent
 
 public:
 	// Set this component up from a provider & message
-	void SetPoseFromProvider(const FAnimationProvider& InProvider, const FSkeletalMeshPoseMessage& InMessage, const FSkeletalMeshInfo& SkeletalMeshInfo);
+	void SetPoseFromProvider(const IAnimationProvider& InProvider, const FSkeletalMeshPoseMessage& InMessage, const FSkeletalMeshInfo& SkeletalMeshInfo);
 
 	// USkeletalMeshComponent interface
 	virtual void InitAnim(bool bForceReInit) override;
diff --git a/Engine/Plugins/Animation/GameplayInsights/Source/GameplayInsights/Private/ObjectEventsTrack.cpp b/Engine/Plugins/Animation/GameplayInsights/Source/GameplayInsights/Private/ObjectEventsTrack.cpp
--- a/Engine/Plugins/Animation/GameplayInsights/Source/GameplayInsights/Private/ObjectEventsTrack.cpp
+++ b/Engine/Plugins/Animation/GameplayInsights/Source/GameplayInsights/Private/ObjectEventsTrack.cpp
@@ -11,11 +11,11 @@
 
 #define LOCTEXT_NAMESPACE "ObjectEventsTrack"
 
-const FName FObjectEventsTrack::TypeName(TEXT("Gameplay"));
-const FName FObjectEventsTrack::SubTypeName(TEXT("ObjectEvents"));
+const FName FObjectEventsTrack::TypeName(TEXT("Events"));
+const FName FObjectEventsTrack::SubTypeName(TEXT("Gameplay.ObjectEvents"));
 
 FObjectEventsTrack::FObjectEventsTrack(const FGameplaySharedData& InSharedData, uint64 InObjectID, const TCHAR* InName)
-	: TGameplayTrackMixin<FTimingEventsTrack>(InObjectID, TypeName, SubTypeName, FText::Format(LOCTEXT("ObjectEventsTrackName", "{0}"), FText::FromString(FString(InName))))
+	: TGameplayTrackMixin<FTimingEventsTrack>(InObjectID, TypeName, SubTypeName, MakeTrackName(InSharedData, InObjectID, InName))
 	, SharedData(InSharedData)
 {
 }
@@ -42,7 +42,7 @@ void FObjectEventsTrack::BuildDrawState(ITimingEventsTrackDrawStateBuilder& Buil
 void FObjectEventsTrack::Draw(const ITimingTrackDrawContext& Context) const
 {
 	DrawEvents(Context);
-	GetGameplayTrack().DrawHeaderForTimingTrack(Context, *this);
+	GetGameplayTrack().DrawHeaderForTimingTrack(Context, *this, false);
 }
 
 void FObjectEventsTrack::InitTooltip(FTooltipDrawState& Tooltip, const ITimingEvent& HoveredTimingEvent) const
@@ -103,4 +103,23 @@ void FObjectEventsTrack::FindObjectEvent(const FTimingEventSearchParameters& InP
 		});
 }
 
+FText FObjectEventsTrack::MakeTrackName(const FGameplaySharedData& InSharedData, uint64 InObjectID, const TCHAR* InName) const
+{
+	FText ClassName = LOCTEXT("UnknownClass", "Unknown");
+
+	const FGameplayProvider* GameplayProvider = InSharedData.GetAnalysisSession().ReadProvider<FGameplayProvider>(FGameplayProvider::ProviderName);
+	if(GameplayProvider)	
+	{
+		if(const FObjectInfo* ObjectInfo = GameplayProvider->FindObjectInfo(InObjectID))
+		{
+			if(const FClassInfo* ClassInfo = GameplayProvider->FindClassInfo(ObjectInfo->ClassId))
+			{
+				ClassName = FText::FromString(ClassInfo->Name);
+			}
+		}
+	}
+
+	return FText::Format(LOCTEXT("ObjectEventsTrackName", "{0} - {1}"), ClassName, FText::FromString(InName));
+}
+
 #undef LOCTEXT_NAMESPACE
\ No newline at end of file
diff --git a/Engine/Plugins/Animation/GameplayInsights/Source/GameplayInsights/Private/ObjectEventsTrack.h b/Engine/Plugins/Animation/GameplayInsights/Source/GameplayInsights/Private/ObjectEventsTrack.h
--- a/Engine/Plugins/Animation/GameplayInsights/Source/GameplayInsights/Private/ObjectEventsTrack.h
+++ b/Engine/Plugins/Animation/GameplayInsights/Source/GameplayInsights/Private/ObjectEventsTrack.h
@@ -26,6 +26,9 @@ private:
 	// Helper function used to find an object event
 	void FindObjectEvent(const FTimingEventSearchParameters& InParameters, TFunctionRef<void(double, double, uint32, const FObjectEventMessage&)> InFoundPredicate) const;
 
+	// Helper function to build the track's name
+	FText MakeTrackName(const FGameplaySharedData& InSharedData, uint64 InObjectID, const TCHAR* InName) const;
+
 private:
 	const FGameplaySharedData& SharedData;
 };
\ No newline at end of file
diff --git a/Engine/Plugins/Animation/GameplayInsights/Source/GameplayInsights/Private/SkeletalMeshPoseTrack.cpp b/Engine/Plugins/Animation/GameplayInsights/Source/GameplayInsights/Private/SkeletalMeshPoseTrack.cpp
--- a/Engine/Plugins/Animation/GameplayInsights/Source/GameplayInsights/Private/SkeletalMeshPoseTrack.cpp
+++ b/Engine/Plugins/Animation/GameplayInsights/Source/GameplayInsights/Private/SkeletalMeshPoseTrack.cpp
@@ -16,13 +16,14 @@
 #include "Engine/SkeletalMesh.h"
 #include "UObject/SoftObjectPtr.h"
 #include "InsightsSkeletalMeshComponent.h"
+#include "GameFramework/Actor.h"
 #endif
 
 
 #define LOCTEXT_NAMESPACE "SkeletalMeshPoseTrack"
 
-const FName FSkeletalMeshPoseTrack::TypeName(TEXT("Animation"));
-const FName FSkeletalMeshPoseTrack::SubTypeName(TEXT("SkeletalMeshPose"));
+const FName FSkeletalMeshPoseTrack::TypeName(TEXT("Events"));
+const FName FSkeletalMeshPoseTrack::SubTypeName(TEXT("Animation.SkeletalMeshPose"));
 
 FSkeletalMeshPoseTrack::FSkeletalMeshPoseTrack(const FAnimationSharedData& InSharedData, uint64 InObjectID, const TCHAR* InName)
 	: TGameplayTrackMixin<FTimingEventsTrack>(InObjectID, FSkeletalMeshPoseTrack::TypeName, FSkeletalMeshPoseTrack::SubTypeName, FText::Format(LOCTEXT("TrackNameFormat", "Pose - {0}"), FText::FromString(FString(InName))))
@@ -30,6 +31,7 @@ FSkeletalMeshPoseTrack::FSkeletalMeshPoseTrack(const FAnimationSharedData& InSha
 	, Color(FLinearColor::MakeRandomColor())
 	, bDrawPose(false)
 	, bDrawSkeleton(false)
+	, bPotentiallyDebugged(false)
 {
 #if WITH_ENGINE
 	OnWorldDestroyedHandle = FWorldDelegates::OnWorldCleanup.AddRaw(this, &FSkeletalMeshPoseTrack::OnWorldCleanup);
@@ -40,6 +42,24 @@ FSkeletalMeshPoseTrack::~FSkeletalMeshPoseTrack()
 {
 #if WITH_ENGINE
 	FWorldDelegates::OnWorldCleanup.Remove(OnWorldDestroyedHandle);
+
+	for(auto& WorldCacheEntry : WorldCache)
+	{
+		if(WorldCacheEntry.Value.Component)
+		{
+			WorldCacheEntry.Value.Component->UnregisterComponent();
+			WorldCacheEntry.Value.Component->MarkPendingKill();
+			WorldCacheEntry.Value.Component = nullptr;
+		}
+
+		if(WorldCacheEntry.Value.Actor)
+		{
+			WorldCacheEntry.Value.Actor->Destroy();
+			WorldCacheEntry.Value.Actor = nullptr;
+		}
+	}
+
+	WorldCache.Empty();
 #endif
 }
 
@@ -66,7 +86,7 @@ void FSkeletalMeshPoseTrack::BuildDrawState(ITimingEventsTrackDrawStateBuilder&
 void FSkeletalMeshPoseTrack::Draw(const ITimingTrackDrawContext& Context) const
 {
 	DrawEvents(Context);
-	GetGameplayTrack().DrawHeaderForTimingTrack(Context, *this);
+	GetGameplayTrack().DrawHeaderForTimingTrack(Context, *this, false);
 }
 
 void FSkeletalMeshPoseTrack::InitTooltip(FTooltipDrawState& Tooltip, const ITimingEvent& HoveredTimingEvent) const
@@ -155,7 +175,11 @@ void FSkeletalMeshPoseTrack::BuildContextMenu(FMenuBuilder& MenuBuilder)
 			LOCTEXT("ToggleDrawPose_Tooltip", "Draw the poses in this track"),
 			FSlateIcon(),
 			FUIAction(
-				FExecuteAction::CreateLambda([this](){ bDrawPose = !bDrawPose; }),
+				FExecuteAction::CreateLambda([this]()
+				{
+					bDrawPose = !bDrawPose;
+					UpdateComponentVisibility();
+				}),
 				FCanExecuteAction(),
 				FIsActionChecked::CreateLambda([this](){ return bDrawPose; })),
 			NAME_None,
@@ -216,6 +240,7 @@ void FSkeletalMeshPoseTrack::BuildContextMenu(FMenuBuilder& MenuBuilder)
 									if(OtherComponentObjectInfo->OuterId == ActorObjectInfo->Id)
 									{
 										InTrack->bDrawPose = !bSetDrawPose;
+										InTrack->UpdateComponentVisibility();
 									}
 								});
 							}),
@@ -265,6 +290,7 @@ void FSkeletalMeshPoseTrack::BuildContextMenu(FMenuBuilder& MenuBuilder)
 									if(OtherComponentObjectInfo->OuterId == ActorObjectInfo->Id)
 									{
 										InTrack->bDrawSkeleton = !bSetDrawSkeleton;
+										InTrack->UpdateComponentVisibility();
 									}
 								});
 							}),
@@ -311,6 +337,16 @@ void FSkeletalMeshPoseTrack::OnWorldCleanup(UWorld* InWorld, bool bSessionEnded,
 	WorldCache.Remove(TWeakObjectPtr<UWorld>(InWorld));
 }
 
+USkeletalMeshComponent* FSkeletalMeshPoseTrack::GetComponent(UWorld* InWorld)
+{
+	if(InWorld)
+	{
+		return GetWorldCache(InWorld).GetComponent();
+	}
+
+	return nullptr;
+}
+
 void FSkeletalMeshPoseTrack::DrawPoses(UWorld* InWorld, double InTime)
 {
 	if(SharedData.IsAnalysisSessionValid())
@@ -339,7 +375,7 @@ void FSkeletalMeshPoseTrack::DrawPoses(UWorld* InWorld, double InTime)
 						if(SkeletalMeshInfo && SkeletalMeshObjectInfo)
 						{
 							UInsightsSkeletalMeshComponent* Component = CacheForWorld.GetComponent();
-							Component->SetVisibility(true);
+							Component->SetVisibility(bDrawPose);
 
 							if(CacheForWorld.Time != InTime)
 							{
@@ -373,11 +409,23 @@ FSkeletalMeshPoseTrack::FWorldComponentCache& FSkeletalMeshPoseTrack::GetWorldCa
 
 UInsightsSkeletalMeshComponent* FSkeletalMeshPoseTrack::FWorldComponentCache::GetComponent()
 {
+	if(Actor == nullptr)
+	{
+		Actor = World->SpawnActor<AActor>();
+		Actor->SetActorLabel(TEXT("Insights"));
+
+		Time = 0.0;
+	}
+
 	if(Component == nullptr)
 	{
-		Component = NewObject<UInsightsSkeletalMeshComponent>();
+		Component = NewObject<UInsightsSkeletalMeshComponent>(Actor);
 		Component->PrimaryComponentTick.bStartWithTickEnabled = false;
 		Component->PrimaryComponentTick.bCanEverTick = false;
+
+		Actor->AddInstanceComponent(Component);
+
+		Component->SetAnimationMode(EAnimationMode::AnimationCustomMode);
 		Component->RegisterComponentWithWorld(World);
 
 		Time = 0.0;
@@ -390,10 +438,22 @@ void FSkeletalMeshPoseTrack::AddReferencedObjects(FReferenceCollector& Collector
 {
 	for(auto& WorldCacheEntry : WorldCache)
 	{
+		Collector.AddReferencedObject(WorldCacheEntry.Value.Actor);
 		Collector.AddReferencedObject(WorldCacheEntry.Value.Component);
 	}
 }
 
+void FSkeletalMeshPoseTrack::UpdateComponentVisibility()
+{
+	for(auto& WorldCacheEntry : WorldCache)
+	{
+		if(WorldCacheEntry.Value.Component)
+		{
+			WorldCacheEntry.Value.Component->SetVisibility(bDrawPose);
+		}
+	}
+}
+
 #endif
 
 #undef LOCTEXT_NAMESPACE
\ No newline at end of file
diff --git a/Engine/Plugins/Animation/GameplayInsights/Source/GameplayInsights/Private/SkeletalMeshPoseTrack.h b/Engine/Plugins/Animation/GameplayInsights/Source/GameplayInsights/Private/SkeletalMeshPoseTrack.h
--- a/Engine/Plugins/Animation/GameplayInsights/Source/GameplayInsights/Private/SkeletalMeshPoseTrack.h
+++ b/Engine/Plugins/Animation/GameplayInsights/Source/GameplayInsights/Private/SkeletalMeshPoseTrack.h
@@ -11,7 +11,9 @@
 #endif
 
 #if WITH_ENGINE
+class USkeletalMeshComponent;
 class UInsightsSkeletalMeshComponent;
+class AActor;
 #endif
 
 class FAnimationSharedData;
@@ -40,7 +42,14 @@ public:
 	bool ShouldDrawPose() const { return bDrawPose; }
 	bool ShouldDrawSkeleton() const { return bDrawSkeleton; }
 
+	// Mark this track as potentially being debugged, so we might need to push through its transforms when time changes
+	void MarkPotentiallyDebugged() { bPotentiallyDebugged = true; }
+	bool IsPotentiallyDebugged() const { return bPotentiallyDebugged; }
+
 #if WITH_ENGINE
+	// Get the component for the specified world
+	USkeletalMeshComponent* GetComponent(UWorld* InWorld);
+
 	// Handle worlds being torn down
 	void OnWorldCleanup(UWorld* InWorld, bool bSessionEnded, bool bCleanupResources);
 
@@ -56,6 +65,9 @@ private:
 	// Helper function used to find a skeletal mesh pose
 	void FindSkeletalMeshPoseMessage(const FTimingEventSearchParameters& InParameters, TFunctionRef<void(double, double, uint32, const FSkeletalMeshPoseMessage&)> InFoundPredicate) const;
 
+	// Updates component visibility based on draw pose flag
+	void UpdateComponentVisibility();
+
 private:
 	/** The shared data */
 	const FAnimationSharedData& SharedData;
@@ -69,12 +81,16 @@ private:
 	/** Whether to draw the skeleton */
 	bool bDrawSkeleton;
 
+	/** Whether we might be being debugged */
+	bool bPotentiallyDebugged;
+
 #if WITH_ENGINE
 	/** Cached data per-world */
 	struct FWorldComponentCache
 	{
 		FWorldComponentCache()
 			: World(nullptr)
+			, Actor(nullptr)
 			, Component(nullptr)
 			, Time(0.0)
 		{}
@@ -85,6 +101,9 @@ private:
 		/** The world we populate */
 		UWorld* World;
 
+		/** Cached actor used to hang the component off of */
+		AActor* Actor;
+
 		/** Cached component used to visualize in this world */
 		UInsightsSkeletalMeshComponent* Component;
 
diff --git a/Engine/Plugins/Animation/LiveLink/Source/LiveLink/Private/AnimNode_LiveLinkPose.cpp b/Engine/Plugins/Animation/LiveLink/Source/LiveLink/Private/AnimNode_LiveLinkPose.cpp
--- a/Engine/Plugins/Animation/LiveLink/Source/LiveLink/Private/AnimNode_LiveLinkPose.cpp
+++ b/Engine/Plugins/Animation/LiveLink/Source/LiveLink/Private/AnimNode_LiveLinkPose.cpp
@@ -10,6 +10,7 @@
 #include "LiveLinkRemapAsset.h"
 #include "Roles/LiveLinkAnimationRole.h"
 #include "Roles/LiveLinkAnimationTypes.h"
+#include "Animation/AnimTrace.h"
 
 FAnimNode_LiveLinkPose::FAnimNode_LiveLinkPose() 
 	: RetargetAsset(ULiveLinkRemapAsset::StaticClass())
@@ -58,6 +59,8 @@ void FAnimNode_LiveLinkPose::Update_AnyThread(const FAnimationUpdateContext & Co
 
 	// Accumulate Delta time from update
 	CachedDeltaTime += Context.GetDeltaTime();
+
+	TRACE_ANIM_NODE_VALUE(Context, TEXT("SubjectName"), LiveLinkSubjectName.Name);
 }
 
 void FAnimNode_LiveLinkPose::Evaluate_AnyThread(FPoseContext& Output)
diff --git a/Engine/Plugins/Experimental/ControlRig/Source/ControlRig/Private/AnimNode_ControlRig.cpp b/Engine/Plugins/Experimental/ControlRig/Source/ControlRig/Private/AnimNode_ControlRig.cpp
--- a/Engine/Plugins/Experimental/ControlRig/Source/ControlRig/Private/AnimNode_ControlRig.cpp
+++ b/Engine/Plugins/Experimental/ControlRig/Source/ControlRig/Private/AnimNode_ControlRig.cpp
@@ -95,6 +95,8 @@ void FAnimNode_ControlRig::Update_AnyThread(const FAnimationUpdateContext& Conte
 
 	PropagateInputProperties(Context.AnimInstanceProxy->GetAnimInstanceObject());
 	Source.Update(Context);
+
+	TRACE_ANIM_NODE_VALUE(Context, TEXT("Class"), *GetNameSafe(ControlRigClass.Get()));
 }
 
 void FAnimNode_ControlRig::Initialize_AnyThread(const FAnimationInitializeContext& Context)
diff --git a/Engine/Plugins/Runtime/PhysXVehicles/Source/PhysXVehicles/Private/AnimNode_WheelHandler.cpp b/Engine/Plugins/Runtime/PhysXVehicles/Source/PhysXVehicles/Private/AnimNode_WheelHandler.cpp
--- a/Engine/Plugins/Runtime/PhysXVehicles/Source/PhysXVehicles/Private/AnimNode_WheelHandler.cpp
+++ b/Engine/Plugins/Runtime/PhysXVehicles/Source/PhysXVehicles/Private/AnimNode_WheelHandler.cpp
@@ -78,6 +78,22 @@ void FAnimNode_WheelHandler::EvaluateSkeletalControl_AnyThread(FComponentSpacePo
 			OutBoneTransforms.Add(FBoneTransform(WheelSimBoneIndex, NewBoneTM));
 		}
 	}
+
+#if ANIM_TRACE_ENABLED
+	for (const FWheelLookupData& Wheel : Wheels)
+	{
+		if (Wheel.BoneReference.BoneIndex != INDEX_NONE)
+		{
+			TRACE_ANIM_NODE_VALUE(Output, *FString::Printf(TEXT("Wheel %d Name"), Wheel.WheelIndex), *Wheel.BoneReference.BoneName.ToString());
+			TRACE_ANIM_NODE_VALUE(Output, *FString::Printf(TEXT("Wheel %d Rotation Offset"), Wheel.WheelIndex), WheelAnimData[Wheel.WheelIndex].RotOffset);
+			TRACE_ANIM_NODE_VALUE(Output, *FString::Printf(TEXT("Wheel %d Location Offset"), Wheel.WheelIndex), WheelAnimData[Wheel.WheelIndex].LocOffset);
+		}
+		else
+		{
+			TRACE_ANIM_NODE_VALUE(Output, *FString::Printf(TEXT("Wheel %d Name"), Wheel.WheelIndex), *FString::Printf(TEXT("%s (invalid)"), *Wheel.BoneReference.BoneName.ToString()));
+		}
+	}
+#endif
 }
 
 bool FAnimNode_WheelHandler::IsValidToEvaluate(const USkeleton* Skeleton, const FBoneContainer& RequiredBones) 
diff --git a/Engine/Source/Editor/AnimationBlueprintEditor/Private/AnimationBlueprintEditor.cpp b/Engine/Source/Editor/AnimationBlueprintEditor/Private/AnimationBlueprintEditor.cpp
--- a/Engine/Source/Editor/AnimationBlueprintEditor/Private/AnimationBlueprintEditor.cpp
+++ b/Engine/Source/Editor/AnimationBlueprintEditor/Private/AnimationBlueprintEditor.cpp
@@ -1285,7 +1285,7 @@ void FAnimationBlueprintEditor::ClearupPreviewMeshAnimNotifyStates()
 	}
 }
 
-void FAnimationBlueprintEditor::GetCustomDebugObjects(TArray<FCustomDebugObject>& DebugList) const
+UAnimInstance* FAnimationBlueprintEditor::GetPreviewInstance() const
 {
 	UDebugSkelMeshComponent* PreviewMeshComponent = PersonaToolkit->GetPreviewMeshComponent();
 	if (PreviewMeshComponent->IsAnimBlueprintInstanced())
@@ -1306,8 +1306,22 @@ void FAnimationBlueprintEditor::GetCustomDebugObjects(TArray<FCustomDebugObject>
 			}
 		}
 
+		return PreviewInstance;
+	}
+
+	return nullptr;
+}
+
+void FAnimationBlueprintEditor::GetCustomDebugObjects(TArray<FCustomDebugObject>& DebugList) const
+{
+	UAnimInstance* PreviewInstance = GetPreviewInstance();
+	if (PreviewInstance)
+	{
 		new (DebugList) FCustomDebugObject(PreviewInstance, LOCTEXT("PreviewObjectLabel", "Preview Instance").ToString());
 	}
+
+	FAnimationBlueprintEditorModule& AnimationBlueprintEditorModule = FModuleManager::GetModuleChecked<FAnimationBlueprintEditorModule>("AnimationBlueprintEditor");
+	AnimationBlueprintEditorModule.OnGetCustomDebugObjects().Broadcast(*this, DebugList);
 }
 
 void FAnimationBlueprintEditor::CreateDefaultTabContents(const TArray<UBlueprint*>& InBlueprints)
@@ -1624,7 +1638,7 @@ void FAnimationBlueprintEditor::HandleSetObjectBeingDebugged(UObject* InObject)
 			}
 			else
 			{
-				// Otherwise eet us to display the debugged instance via copy-pose
+				// Otherwise set us to display the debugged instance via copy-pose
 				GetPreviewScene()->GetPreviewMeshComponent()->EnablePreview(true, nullptr);
 				GetPreviewScene()->GetPreviewMeshComponent()->PreviewInstance->SetDebugSkeletalMeshComponent(SkeletalMeshComponent);
 			}
diff --git a/Engine/Source/Editor/AnimationBlueprintEditor/Private/AnimationBlueprintEditor.h b/Engine/Source/Editor/AnimationBlueprintEditor/Private/AnimationBlueprintEditor.h
--- a/Engine/Source/Editor/AnimationBlueprintEditor/Private/AnimationBlueprintEditor.h
+++ b/Engine/Source/Editor/AnimationBlueprintEditor/Private/AnimationBlueprintEditor.h
@@ -92,6 +92,7 @@ public:
 	virtual const FEdGraphPinType& GetLastGraphPinTypeUsed() const override { return LastGraphPinType; }
 	virtual void SetLastGraphPinTypeUsed(const FEdGraphPinType& InType) override { LastGraphPinType = InType; }
 	virtual IAnimationSequenceBrowser* GetAssetBrowser() const override;
+	virtual UAnimInstance* GetPreviewInstance() const override;
 
 	/** IHasPersonaToolkit interface */
 	virtual TSharedRef<class IPersonaToolkit> GetPersonaToolkit() const { return PersonaToolkit.ToSharedRef(); }
diff --git a/Engine/Source/Editor/AnimationBlueprintEditor/Private/AnimationBlueprintEditorModule.h b/Engine/Source/Editor/AnimationBlueprintEditor/Private/AnimationBlueprintEditorModule.h
--- a/Engine/Source/Editor/AnimationBlueprintEditor/Private/AnimationBlueprintEditorModule.h
+++ b/Engine/Source/Editor/AnimationBlueprintEditor/Private/AnimationBlueprintEditorModule.h
@@ -35,6 +35,8 @@ public:
 	virtual TSharedPtr<FExtensibilityManager> GetMenuExtensibilityManager() override { return MenuExtensibilityManager; }
 	virtual TSharedPtr<FExtensibilityManager> GetToolBarExtensibilityManager() override { return ToolBarExtensibilityManager; }
 
+	virtual FOnGetCustomDebugObjects& OnGetCustomDebugObjects() override { return OnGetCustomDebugObjectsDelegate; }
+
 private:
 	/** When a new AnimBlueprint is created, this will handle post creation work such as adding non-event default nodes */
 	void OnNewBlueprintCreated(class UBlueprint* InBlueprint);
@@ -48,4 +50,6 @@ private:
 	TSharedPtr<FAnimationGraphNodeFactory> AnimGraphNodeFactory;
 	TSharedPtr<FAnimationGraphPinFactory> AnimGraphPinFactory;
 	TSharedPtr<FAnimationGraphPinConnectionFactory> AnimGraphPinConnectionFactory;
+
+	FOnGetCustomDebugObjects OnGetCustomDebugObjectsDelegate;
 };
diff --git a/Engine/Source/Editor/AnimationBlueprintEditor/Private/AnimationNodes/SAnimationGraphNode.cpp b/Engine/Source/Editor/AnimationBlueprintEditor/Private/AnimationNodes/SAnimationGraphNode.cpp
--- a/Engine/Source/Editor/AnimationBlueprintEditor/Private/AnimationNodes/SAnimationGraphNode.cpp
+++ b/Engine/Source/Editor/AnimationBlueprintEditor/Private/AnimationNodes/SAnimationGraphNode.cpp
@@ -10,6 +10,8 @@
 #include "AnimGraphNode_Base.h"
 #include "IDocumentation.h"
 #include "AnimationEditorUtils.h"
+#include "Kismet2/BlueprintEditorUtils.h"
+#include "Animation/AnimInstance.h"
 
 class SPoseViewColourPickerPopup : public SCompoundWidget
 {
@@ -210,4 +212,35 @@ void SAnimationGraphNode::HandleNodeTitleChanged()
 	{
 		NodeTitle->MarkDirty();
 	}
+}
+
+void SAnimationGraphNode::GetNodeInfoPopups(FNodeInfoContext* Context, TArray<FGraphInformationPopupInfo>& Popups) const
+{
+	UAnimBlueprint* AnimBlueprint = Cast<UAnimBlueprint>(FBlueprintEditorUtils::FindBlueprintForNode(GraphNode));
+	if(AnimBlueprint)
+	{
+		UAnimInstance* ActiveObject = Cast<UAnimInstance>(AnimBlueprint->GetObjectBeingDebugged());
+		UAnimBlueprintGeneratedClass* Class = AnimBlueprint->GetAnimBlueprintGeneratedClass();
+
+		const FLinearColor Color(1.f, 0.5f, 0.25f);
+
+		// Display various types of debug data
+		if ((ActiveObject != NULL) && (Class != NULL))
+		{
+			if (Class->AnimNodeProperties.Num())
+			{
+				if(int32* NodeIndexPtr = Class->GetAnimBlueprintDebugData().NodePropertyToIndexMap.Find(TWeakObjectPtr<UAnimGraphNode_Base>(Cast<UAnimGraphNode_Base>(GraphNode))))
+				{
+					int32 AnimNodeIndex = *NodeIndexPtr;
+					// reverse node index temporarily because of a bug in NodeGuidToIndexMap
+					AnimNodeIndex = Class->AnimNodeProperties.Num() - AnimNodeIndex - 1;
+
+					if (FAnimBlueprintDebugData::FNodeValue* DebugInfo = Class->GetAnimBlueprintDebugData().NodeValuesThisFrame.FindByPredicate([AnimNodeIndex](const FAnimBlueprintDebugData::FNodeValue& InValue){ return InValue.NodeID == AnimNodeIndex; }))
+					{
+						Popups.Emplace(nullptr, Color, DebugInfo->Text);
+					}
+				}
+			}
+		}
+	}
 }
\ No newline at end of file
diff --git a/Engine/Source/Editor/AnimationBlueprintEditor/Private/AnimationNodes/SAnimationGraphNode.h b/Engine/Source/Editor/AnimationBlueprintEditor/Private/AnimationNodes/SAnimationGraphNode.h
--- a/Engine/Source/Editor/AnimationBlueprintEditor/Private/AnimationNodes/SAnimationGraphNode.h
+++ b/Engine/Source/Editor/AnimationBlueprintEditor/Private/AnimationNodes/SAnimationGraphNode.h
@@ -28,6 +28,7 @@ protected:
 	// SGraphNode interface
 	virtual TArray<FOverlayWidgetInfo> GetOverlayWidgets(bool bSelected, const FVector2D& WidgetSize) const override;
 	virtual TSharedRef<SWidget> CreateTitleWidget(TSharedPtr<SNodeTitle> InNodeTitle) override;
+	virtual void GetNodeInfoPopups(FNodeInfoContext* Context, TArray<FGraphInformationPopupInfo>& Popups) const override;
 	// End of SGraphNode interface
 
 private:
diff --git a/Engine/Source/Editor/AnimationBlueprintEditor/Private/AnimationNodes/SGraphNodeSequencePlayer.cpp b/Engine/Source/Editor/AnimationBlueprintEditor/Private/AnimationNodes/SGraphNodeSequencePlayer.cpp
--- a/Engine/Source/Editor/AnimationBlueprintEditor/Private/AnimationNodes/SGraphNodeSequencePlayer.cpp
+++ b/Engine/Source/Editor/AnimationBlueprintEditor/Private/AnimationNodes/SGraphNodeSequencePlayer.cpp
@@ -118,15 +118,31 @@ EVisibility SGraphNodeSequencePlayer::GetSliderVisibility() const
 
 bool SGraphNodeSequencePlayer::GetSequencePositionInfo(float& Out_Position, float& Out_Length, int32& Out_FrameCount) const
 {
-	if(FAnimNode_SequencePlayer* SequencePlayer = GetSequencePlayer())
+	if (UBlueprint* Blueprint = FBlueprintEditorUtils::FindBlueprintForNode(GraphNode))
 	{
-		if (UAnimSequenceBase* BoundSequence = SequencePlayer->Sequence)
+		if (UObject* ActiveObject = Blueprint->GetObjectBeingDebugged())
 		{
-			Out_Position = SequencePlayer->GetAccumulatedTime();
-			Out_Length = BoundSequence->SequenceLength;
-			Out_FrameCount = BoundSequence->GetNumberOfFrames();
-
-			return true;
+			if (UAnimGraphNode_SequencePlayer* VisualSequencePlayer = Cast<UAnimGraphNode_SequencePlayer>(GraphNode))
+			{
+				if (UAnimBlueprintGeneratedClass* Class = Cast<UAnimBlueprintGeneratedClass>((UObject*)ActiveObject->GetClass()))
+				{
+					if(int32* NodeIndexPtr = Class->GetAnimBlueprintDebugData().NodePropertyToIndexMap.Find(TWeakObjectPtr<UAnimGraphNode_Base>(Cast<UAnimGraphNode_Base>(GraphNode))))
+					{
+						int32 AnimNodeIndex = *NodeIndexPtr;
+						// reverse node index temporarily because of a bug in NodeGuidToIndexMap
+						AnimNodeIndex = Class->AnimNodeProperties.Num() - AnimNodeIndex - 1;
+
+						if (FAnimBlueprintDebugData::FSequencePlayerRecord* DebugInfo = Class->GetAnimBlueprintDebugData().SequencePlayerRecordsThisFrame.FindByPredicate([AnimNodeIndex](const FAnimBlueprintDebugData::FSequencePlayerRecord& InRecord){ return InRecord.NodeID == AnimNodeIndex; }))
+						{
+							Out_Position = DebugInfo->Position;
+							Out_Length = DebugInfo->Length;
+							Out_FrameCount = DebugInfo->FrameCount;
+
+							return true;
+						}
+					}
+				}
+			}
 		}
 	}
 
diff --git a/Engine/Source/Editor/AnimationBlueprintEditor/Private/AnimationStateNodes/SGraphNodeAnimState.cpp b/Engine/Source/Editor/AnimationBlueprintEditor/Private/AnimationStateNodes/SGraphNodeAnimState.cpp
--- a/Engine/Source/Editor/AnimationBlueprintEditor/Private/AnimationStateNodes/SGraphNodeAnimState.cpp
+++ b/Engine/Source/Editor/AnimationBlueprintEditor/Private/AnimationStateNodes/SGraphNodeAnimState.cpp
@@ -16,6 +16,8 @@
 #include "AnimGraphNode_StateMachineBase.h"
 #include "Widgets/Text/SInlineEditableTextBlock.h"
 
+#define LOCTEXT_NAMESPACE "SGraphNodeAnimState"
+
 /////////////////////////////////////////////////////
 // SStateMachineOutputPin
 
@@ -94,26 +96,31 @@ void SGraphNodeAnimState::GetStateInfoPopup(UEdGraphNode* GraphNode, TArray<FGra
 		// Display various types of debug data
 		if ((ActiveObject != NULL) && (Class != NULL))
 		{
-			if (FStateMachineDebugData* DebugInfo = Class->GetAnimBlueprintDebugData().StateMachineDebugData.Find(GraphNode->GetGraph()))
+			if (Class->AnimNodeProperties.Num())
 			{
-				if (Class->AnimNodeProperties.Num())
+				if (FStateMachineDebugData* DebugInfo = Class->GetAnimBlueprintDebugData().StateMachineDebugData.Find(GraphNode->GetGraph()))
 				{
-					UAnimationStateMachineGraph* TypedGraph = CastChecked<UAnimationStateMachineGraph>(GraphNode->GetGraph());
-
-					if (FAnimNode_StateMachine* CurrentInstance = Class->GetPropertyInstance<FAnimNode_StateMachine>(ActiveObject, TypedGraph->OwnerAnimGraphNode))  
+					if(int32* StateIndexPtr = DebugInfo->NodeToStateIndex.Find(GraphNode))
 					{
-						if (int32* pStateIndex = DebugInfo->NodeToStateIndex.Find(GraphNode))
+						for(const FStateMachineStateDebugData& StateData : Class->GetAnimBlueprintDebugData().StateData)
 						{
-							const float Weight = CurrentInstance->GetStateWeight(*pStateIndex);
-							if (Weight > 0.0f)
+							if(StateData.StateMachineIndex == DebugInfo->MachineIndex && StateData.StateIndex == *StateIndexPtr)
 							{
-								FString StateText = FString::Printf(TEXT("%.1f%%"), Weight * 100.0f);
-								if (*pStateIndex == CurrentInstance->GetCurrentState())
+								if (StateData.Weight > 0.0f)
 								{
-									StateText += FString::Printf(TEXT("\nActive for %.2f secs"), CurrentInstance->GetCurrentStateElapsedTime());
+									FText StateText;
+									if (StateData.ElapsedTime > 0.0f)
+									{
+										StateText = FText::Format(LOCTEXT("ActiveStateWeightFormat", "{0}\nActive for {1}s"), FText::AsPercent(StateData.Weight), FText::AsNumber(StateData.ElapsedTime));
+									}
+									else
+									{
+										StateText = FText::Format(LOCTEXT("StateWeightFormat", "{0}"), FText::AsPercent(StateData.Weight));
+									}
+
+									Popups.Emplace(nullptr, CurrentStateColor, StateText.ToString());
+									break;
 								}
-
-								new (Popups) FGraphInformationPopupInfo(NULL, CurrentStateColor, StateText);
 							}
 						}
 					}
@@ -145,18 +152,15 @@ FSlateColor SGraphNodeAnimState::GetBorderBackgroundColor() const
 		{
 			if (FStateMachineDebugData* DebugInfo = Class->GetAnimBlueprintDebugData().StateMachineDebugData.Find(GraphNode->GetGraph()))
 			{
-				if (Class->AnimNodeProperties.Num())
+				if(int32* StateIndexPtr = DebugInfo->NodeToStateIndex.Find(GraphNode))
 				{
-					UAnimationStateMachineGraph* TypedGraph = CastChecked<UAnimationStateMachineGraph>(GraphNode->GetGraph());
-
-					if (FAnimNode_StateMachine* CurrentInstance = Class->GetPropertyInstance<FAnimNode_StateMachine>(ActiveObject, TypedGraph->OwnerAnimGraphNode))
+					for(const FStateMachineStateDebugData& StateData : Class->GetAnimBlueprintDebugData().StateData)
 					{
-						if (int32* pStateIndex = DebugInfo->NodeToStateIndex.Find(GraphNode))
+						if(StateData.StateMachineIndex == DebugInfo->MachineIndex && StateData.StateIndex == *StateIndexPtr)
 						{
-							const float Weight = CurrentInstance->GetStateWeight(*pStateIndex);
-							if (Weight > 0.0f)
+							if (StateData.Weight > 0.0f)
 							{
-								return FMath::Lerp<FLinearColor>(ActiveStateColorDim, ActiveStateColorBright, Weight);
+								return FMath::Lerp<FLinearColor>(ActiveStateColorDim, ActiveStateColorBright, StateData.Weight);
 							}
 						}
 					}
@@ -353,3 +357,5 @@ const FSlateBrush* SGraphNodeAnimConduit::GetNameIcon() const
 {
 	return FEditorStyle::GetBrush( TEXT("Graph.ConduitNode.Icon") );
 }
+
+#undef LOCTEXT_NAMESPACE
\ No newline at end of file
diff --git a/Engine/Source/Editor/AnimationBlueprintEditor/Public/IAnimationBlueprintEditor.h b/Engine/Source/Editor/AnimationBlueprintEditor/Public/IAnimationBlueprintEditor.h
--- a/Engine/Source/Editor/AnimationBlueprintEditor/Public/IAnimationBlueprintEditor.h
+++ b/Engine/Source/Editor/AnimationBlueprintEditor/Public/IAnimationBlueprintEditor.h
@@ -7,6 +7,7 @@
 #include "IHasPersonaToolkit.h"
 
 class IAnimationSequenceBrowser;
+class UAnimInstance;
 
 class IAnimationBlueprintEditor : public FBlueprintEditor, public IHasPersonaToolkit
 {
@@ -19,4 +20,7 @@ public:
 
 	/** Get the asset browser we host */
 	virtual IAnimationSequenceBrowser* GetAssetBrowser() const = 0;
+
+	/** Get the preview anim instance we are using, which can be a linked instance depending on preview settings */
+	virtual UAnimInstance* GetPreviewInstance() const = 0;
 };
diff --git a/Engine/Source/Editor/AnimationBlueprintEditor/Public/IAnimationBlueprintEditorModule.h b/Engine/Source/Editor/AnimationBlueprintEditor/Public/IAnimationBlueprintEditorModule.h
--- a/Engine/Source/Editor/AnimationBlueprintEditor/Public/IAnimationBlueprintEditorModule.h
+++ b/Engine/Source/Editor/AnimationBlueprintEditor/Public/IAnimationBlueprintEditorModule.h
@@ -33,4 +33,7 @@ public:
 	/** Get all toolbar extenders */
 	DECLARE_DELEGATE_RetVal_TwoParams(TSharedRef<FExtender>, FAnimationBlueprintEditorToolbarExtender, const TSharedRef<FUICommandList> /*InCommandList*/, TSharedRef<IAnimationBlueprintEditor> /*InAnimationBlueprintEditor*/);
 	virtual TArray<FAnimationBlueprintEditorToolbarExtender>& GetAllAnimationBlueprintEditorToolbarExtenders() = 0;
+
+	DECLARE_MULTICAST_DELEGATE_TwoParams(FOnGetCustomDebugObjects, const IAnimationBlueprintEditor& /*AnimationBlueprintEditor*/, TArray<FCustomDebugObject>& /*DebugList*/);
+	virtual FOnGetCustomDebugObjects& OnGetCustomDebugObjects() = 0;
 };
diff --git a/Engine/Source/Editor/KismetCompiler/Private/AnimBlueprintCompiler.cpp b/Engine/Source/Editor/KismetCompiler/Private/AnimBlueprintCompiler.cpp
--- a/Engine/Source/Editor/KismetCompiler/Private/AnimBlueprintCompiler.cpp
+++ b/Engine/Source/Editor/KismetCompiler/Private/AnimBlueprintCompiler.cpp
@@ -1830,6 +1830,20 @@ void FAnimBlueprintCompilerContext::CopyTermDefaultsToDefaultObject(UObject* Def
 					NewLinkedInputPose.Graph = Cast<UAnimGraphNode_LinkedInputPose>(MessageLog.FindSourceObject(LinkedInputPoseNode))->GetGraph()->GetFName();
 					TargetProperty->CopyCompleteValue(DestinationPtr, &NewLinkedInputPose);
 				}
+				else if(UAnimGraphNode_LinkedAnimGraph* LinkedAnimGraphNode = ExactCast<UAnimGraphNode_LinkedAnimGraph>(VisualAnimNode))
+				{
+					// patch node index into linked anim graph nodes
+					FAnimNode_LinkedAnimGraph NewLinkedAnimGraph = *reinterpret_cast<FAnimNode_LinkedAnimGraph*>(SourcePtr);
+					NewLinkedAnimGraph.NodeIndex = LinkIndexCount;
+					TargetProperty->CopyCompleteValue(DestinationPtr, &NewLinkedAnimGraph);
+				}
+				else if(UAnimGraphNode_LinkedAnimLayer* LinkedAnimLayerNode = ExactCast<UAnimGraphNode_LinkedAnimLayer>(VisualAnimNode))
+				{
+					// patch node index into linked anim layer nodes
+					FAnimNode_LinkedAnimLayer NewLinkedAnimLayer = *reinterpret_cast<FAnimNode_LinkedAnimLayer*>(SourcePtr);
+					NewLinkedAnimLayer.NodeIndex = LinkIndexCount;
+					TargetProperty->CopyCompleteValue(DestinationPtr, &NewLinkedAnimLayer);
+				}
 				else
 				{
 					TargetProperty->CopyCompleteValue(DestinationPtr, SourcePtr);
diff --git a/Engine/Source/Runtime/AnimGraphRuntime/AnimGraphRuntime.Build.cs b/Engine/Source/Runtime/AnimGraphRuntime/AnimGraphRuntime.Build.cs
--- a/Engine/Source/Runtime/AnimGraphRuntime/AnimGraphRuntime.Build.cs
+++ b/Engine/Source/Runtime/AnimGraphRuntime/AnimGraphRuntime.Build.cs
@@ -14,21 +14,27 @@ public class AnimGraphRuntime : ModuleRules
 				"Core", 
 				"CoreUObject", 
 				"Engine",
-                "AnimationCore",
+				"AnimationCore",
 			}
 		);
 
-        SetupModulePhysicsSupport(Target);
+		PrivateDependencyModuleNames.AddRange(
+			new string[] {
+				"TraceLog",
+			}
+		);
+
+		SetupModulePhysicsSupport(Target);
 
-        if (Target.bCompileChaos || Target.bUseChaos)
-        {
-            PublicDependencyModuleNames.AddRange(
-                new string[] {
+		if (Target.bCompileChaos || Target.bUseChaos)
+		{
+			PublicDependencyModuleNames.AddRange(
+				new string[] {
 					"ChaosSolvers",
 					"GeometryCollectionEngine",
 					"GeometryCollectionSimulationCore"
-                }
-            );
-        }
-    }
+				}
+			);
+		}
+	}
 }
diff --git a/Engine/Source/Runtime/AnimGraphRuntime/Private/AnimGraphRuntimeModule.cpp b/Engine/Source/Runtime/AnimGraphRuntime/Private/AnimGraphRuntimeModule.cpp
--- a/Engine/Source/Runtime/AnimGraphRuntime/Private/AnimGraphRuntimeModule.cpp
+++ b/Engine/Source/Runtime/AnimGraphRuntime/Private/AnimGraphRuntimeModule.cpp
@@ -5,6 +5,7 @@
 #include "BoneControllers/AnimNode_AnimDynamics.h"
 #include "UObject/UObjectIterator.h"
 #include "Animation/AnimInstance.h"
+#include "AnimGraphRuntimeTrace.h"
 
 //////////////////////////////////////////////////////////////////////////
 // FAnimGraphRuntimeModule
@@ -14,6 +15,9 @@ class FAnimGraphRuntimeModule : public IModuleInterface
 public:
 	virtual void StartupModule() override
 	{
+#if ANIM_TRACE_ENABLED
+		FAnimGraphRuntimeTrace::Init();
+#endif
 	}
 
 	virtual void ShutdownModule() override
diff --git a/Engine/Source/Runtime/AnimGraphRuntime/Private/AnimGraphRuntimeTrace.cpp b/Engine/Source/Runtime/AnimGraphRuntime/Private/AnimGraphRuntimeTrace.cpp
new file mode 100644
--- /dev/null
+++ b/Engine/Source/Runtime/AnimGraphRuntime/Private/AnimGraphRuntimeTrace.cpp
@@ -0,0 +1,54 @@
+// Copyright 1998-2019 Epic Games, Inc. All Rights Reserved.
+
+#include "AnimGraphRuntimeTrace.h"
+
+#if ANIM_TRACE_ENABLED
+
+#include "Trace/Trace.h"
+#include "Animation/BlendSpaceBase.h"
+#include "AnimNodes/AnimNode_BlendSpacePlayer.h"
+#include "Animation/AnimInstanceProxy.h"
+
+UE_TRACE_EVENT_BEGIN(Animation, BlendSpacePlayer)
+	UE_TRACE_EVENT_FIELD(uint64, Cycle)
+	UE_TRACE_EVENT_FIELD(uint64, AnimInstanceId)
+	UE_TRACE_EVENT_FIELD(uint64, BlendSpaceId)
+	UE_TRACE_EVENT_FIELD(int32, NodeId)
+	UE_TRACE_EVENT_FIELD(float, PositionX)
+	UE_TRACE_EVENT_FIELD(float, PositionY)
+	UE_TRACE_EVENT_FIELD(float, PositionZ)
+UE_TRACE_EVENT_END()
+
+void FAnimGraphRuntimeTrace::Init()
+{
+	if (FParse::Param(FCommandLine::Get(), TEXT("objecttrace")))
+	{
+		UE_TRACE_EVENT_IS_ENABLED(Animation, BlendSpacePlayer);
+		Trace::ToggleEvent(TEXT("Animation"), true);
+	}
+}
+
+void FAnimGraphRuntimeTrace::OutputBlendSpacePlayer(const FAnimationBaseContext& InContext, const FAnimNode_BlendSpacePlayer& InNode)
+{
+	bool bEventEnabled = UE_TRACE_EVENT_IS_ENABLED(Animation, BlendSpacePlayer);
+	if (!bEventEnabled)
+	{
+		return;
+	}
+
+	check(InContext.AnimInstanceProxy);
+
+	TRACE_OBJECT(InContext.AnimInstanceProxy->GetAnimInstanceObject());
+	TRACE_OBJECT(InNode.BlendSpace);
+
+	UE_TRACE_LOG(Animation, BlendSpacePlayer)
+		<< BlendSpacePlayer.Cycle(FPlatformTime::Cycles64())
+		<< BlendSpacePlayer.AnimInstanceId(FObjectTrace::GetObjectId(InContext.AnimInstanceProxy->GetAnimInstanceObject()))
+		<< BlendSpacePlayer.BlendSpaceId(FObjectTrace::GetObjectId(InNode.BlendSpace))
+		<< BlendSpacePlayer.NodeId(InContext.GetCurrentNodeId())
+		<< BlendSpacePlayer.PositionX(InNode.X)
+		<< BlendSpacePlayer.PositionY(InNode.Y)
+		<< BlendSpacePlayer.PositionZ(InNode.Z);
+}
+
+#endif
\ No newline at end of file
diff --git a/Engine/Source/Runtime/AnimGraphRuntime/Private/AnimNodes/AnimNode_AimOffsetLookAt.cpp b/Engine/Source/Runtime/AnimGraphRuntime/Private/AnimNodes/AnimNode_AimOffsetLookAt.cpp
--- a/Engine/Source/Runtime/AnimGraphRuntime/Private/AnimNodes/AnimNode_AimOffsetLookAt.cpp
+++ b/Engine/Source/Runtime/AnimGraphRuntime/Private/AnimNodes/AnimNode_AimOffsetLookAt.cpp
@@ -9,6 +9,7 @@
 #include "DrawDebugHelpers.h"
 #include "EngineGlobals.h"
 #include "Engine/Engine.h"
+#include "Animation/AnimTrace.h"
 
 TAutoConsoleVariable<int32> CVarAimOffsetLookAtEnable(TEXT("a.AnimNode.AimOffsetLookAt.Enable"), 1, TEXT("Enable/Disable LookAt AimOffset"));
 TAutoConsoleVariable<int32> CVarAimOffsetLookAtDebug(TEXT("a.AnimNode.AimOffsetLookAt.Debug"), 0, TEXT("Toggle LookAt AimOffset debug"));
@@ -63,6 +64,8 @@ void FAnimNode_AimOffsetLookAt::UpdateAssetPlayer(const FAnimationUpdateContext&
 // 	}
 
 	BasePose.Update(Context);
+
+	TRACE_ANIM_NODE_VALUE(Context, TEXT("Play Time"), InternalTimeAccumulator);
 }
 
 void FAnimNode_AimOffsetLookAt::CacheBones_AnyThread(const FAnimationCacheBonesContext& Context)
diff --git a/Engine/Source/Runtime/AnimGraphRuntime/Private/AnimNodes/AnimNode_ApplyAdditive.cpp b/Engine/Source/Runtime/AnimGraphRuntime/Private/AnimNodes/AnimNode_ApplyAdditive.cpp
--- a/Engine/Source/Runtime/AnimGraphRuntime/Private/AnimNodes/AnimNode_ApplyAdditive.cpp
+++ b/Engine/Source/Runtime/AnimGraphRuntime/Private/AnimNodes/AnimNode_ApplyAdditive.cpp
@@ -3,6 +3,7 @@
 #include "AnimNodes/AnimNode_ApplyAdditive.h"
 #include "AnimationRuntime.h"
 #include "Animation/AnimInstanceProxy.h"
+#include "Animation/AnimTrace.h"
 
 /////////////////////////////////////////////////////
 // FAnimNode_ApplyAdditive
@@ -59,6 +60,8 @@ void FAnimNode_ApplyAdditive::Update_AnyThread(const FAnimationUpdateContext& Co
 			Additive.Update(Context.FractionalWeight(ActualAlpha));
 		}
 	}
+
+	TRACE_ANIM_NODE_VALUE(Context, TEXT("Alpha"), ActualAlpha);
 }
 
 void FAnimNode_ApplyAdditive::Evaluate_AnyThread(FPoseContext& Output)
diff --git a/Engine/Source/Runtime/AnimGraphRuntime/Private/AnimNodes/AnimNode_BlendBoneByChannel.cpp b/Engine/Source/Runtime/AnimGraphRuntime/Private/AnimNodes/AnimNode_BlendBoneByChannel.cpp
--- a/Engine/Source/Runtime/AnimGraphRuntime/Private/AnimNodes/AnimNode_BlendBoneByChannel.cpp
+++ b/Engine/Source/Runtime/AnimGraphRuntime/Private/AnimNodes/AnimNode_BlendBoneByChannel.cpp
@@ -3,6 +3,7 @@
 #include "AnimNodes/AnimNode_BlendBoneByChannel.h"
 #include "AnimationRuntime.h"
 #include "Animation/AnimInstanceProxy.h"
+#include "Animation/AnimTrace.h"
 
 /////////////////////////////////////////////////////
 // FAnimNode_BlendBoneByChannel
@@ -52,6 +53,8 @@ void FAnimNode_BlendBoneByChannel::Update_AnyThread(const FAnimationUpdateContex
 	{
 		B.Update(Context.FractionalWeight(InternalBlendAlpha));
 	}
+
+	TRACE_ANIM_NODE_VALUE(Context, TEXT("Alpha"), InternalBlendAlpha);
 }
 
 void FAnimNode_BlendBoneByChannel::Evaluate_AnyThread(FPoseContext& Output)
diff --git a/Engine/Source/Runtime/AnimGraphRuntime/Private/AnimNodes/AnimNode_BlendListBase.cpp b/Engine/Source/Runtime/AnimGraphRuntime/Private/AnimNodes/AnimNode_BlendListBase.cpp
--- a/Engine/Source/Runtime/AnimGraphRuntime/Private/AnimNodes/AnimNode_BlendListBase.cpp
+++ b/Engine/Source/Runtime/AnimGraphRuntime/Private/AnimNodes/AnimNode_BlendListBase.cpp
@@ -5,6 +5,7 @@
 #include "Animation/BlendProfile.h"
 #include "Animation/AnimInstanceProxy.h"
 #include "Animation/AnimNode_Inertialization.h"
+#include "Animation/AnimTrace.h"
 
 /////////////////////////////////////////////////////
 // FAnimNode_BlendListBase
@@ -226,6 +227,13 @@ void FAnimNode_BlendListBase::Update_AnyThread(const FAnimationUpdateContext& Co
 			FBlendSampleData::NormalizeDataWeight(PerBoneSampleData);
 		}
 	}
+
+#if ANIM_TRACE_ENABLED
+	const int32 ChildIndex = GetActiveChildIndex();
+	TRACE_ANIM_NODE_VALUE(Context, TEXT("Active Index"), ChildIndex);
+	TRACE_ANIM_NODE_VALUE(Context, TEXT("Active Weight"), BlendWeights[ChildIndex]);
+	TRACE_ANIM_NODE_VALUE(Context, TEXT("Active Blend Time"), BlendTime[ChildIndex]);
+#endif
 }
 
 void FAnimNode_BlendListBase::Evaluate_AnyThread(FPoseContext& Output)
diff --git a/Engine/Source/Runtime/AnimGraphRuntime/Private/AnimNodes/AnimNode_BlendSpaceEvaluator.cpp b/Engine/Source/Runtime/AnimGraphRuntime/Private/AnimNodes/AnimNode_BlendSpaceEvaluator.cpp
--- a/Engine/Source/Runtime/AnimGraphRuntime/Private/AnimNodes/AnimNode_BlendSpaceEvaluator.cpp
+++ b/Engine/Source/Runtime/AnimGraphRuntime/Private/AnimNodes/AnimNode_BlendSpaceEvaluator.cpp
@@ -2,6 +2,7 @@
 
 #include "AnimNodes/AnimNode_BlendSpaceEvaluator.h"
 #include "Animation/BlendSpaceBase.h"
+#include "Animation/AnimTrace.h"
 
 /////////////////////////////////////////////////////
 // FAnimNode_BlendSpaceEvaluator
@@ -19,8 +20,10 @@ void FAnimNode_BlendSpaceEvaluator::UpdateAssetPlayer(const FAnimationUpdateCont
 	PlayRate = 0.f;
 
 	UpdateInternal(Context);
-}
 
+	TRACE_ANIM_NODE_VALUE(Context, TEXT("Blend Space Name"), BlendSpace ? *BlendSpace->GetName() : TEXT("None"));
+	TRACE_ANIM_NODE_VALUE(Context, TEXT("Play Time"), InternalTimeAccumulator);
+}
 
 void FAnimNode_BlendSpaceEvaluator::GatherDebugData(FNodeDebugData& DebugData)
 {
diff --git a/Engine/Source/Runtime/AnimGraphRuntime/Private/AnimNodes/AnimNode_BlendSpacePlayer.cpp b/Engine/Source/Runtime/AnimGraphRuntime/Private/AnimNodes/AnimNode_BlendSpacePlayer.cpp
--- a/Engine/Source/Runtime/AnimGraphRuntime/Private/AnimNodes/AnimNode_BlendSpacePlayer.cpp
+++ b/Engine/Source/Runtime/AnimGraphRuntime/Private/AnimNodes/AnimNode_BlendSpacePlayer.cpp
@@ -4,6 +4,7 @@
 #include "Animation/BlendSpaceBase.h"
 #include "Animation/AnimSequence.h"
 #include "Animation/AnimInstanceProxy.h"
+#include "AnimGraphRuntimeTrace.h"
 
 /////////////////////////////////////////////////////
 // FAnimNode_BlendSpacePlayer
@@ -101,8 +102,22 @@ void FAnimNode_BlendSpacePlayer::UpdateInternal(const FAnimationUpdateContext& C
 			SyncGroup->TestTickRecordForLeadership(GroupRole);
 		}
 
+
+		TRACE_ANIM_TICK_RECORD(Context, TickRecord);
+
+#if ANIM_NODE_IDS_AVAILABLE && WITH_EDITORONLY_DATA
+		if (FAnimBlueprintDebugData* DebugData = Context.AnimInstanceProxy->GetAnimBlueprintDebugData())
+		{
+			DebugData->RecordBlendSpacePlayer(Context.GetCurrentNodeId(), BlendSpace, BlendInput.X, BlendInput.Y, BlendInput.Z);
+		}
+#endif
+
 		PreviousBlendSpace = BlendSpace;
 	}
+
+	TRACE_BLENDSPACE_PLAYER(Context, *this);
+	TRACE_ANIM_NODE_VALUE(Context, TEXT("Blend Space Name"), BlendSpace ? *BlendSpace->GetName() : TEXT("None"));
+	TRACE_ANIM_NODE_VALUE(Context, TEXT("Play Time"), InternalTimeAccumulator);
 }
 
 void FAnimNode_BlendSpacePlayer::Evaluate_AnyThread(FPoseContext& Output)
diff --git a/Engine/Source/Runtime/AnimGraphRuntime/Private/AnimNodes/AnimNode_CopyPoseFromMesh.cpp b/Engine/Source/Runtime/AnimGraphRuntime/Private/AnimNodes/AnimNode_CopyPoseFromMesh.cpp
--- a/Engine/Source/Runtime/AnimGraphRuntime/Private/AnimNodes/AnimNode_CopyPoseFromMesh.cpp
+++ b/Engine/Source/Runtime/AnimGraphRuntime/Private/AnimNodes/AnimNode_CopyPoseFromMesh.cpp
@@ -2,6 +2,7 @@
 
 #include "AnimNodes/AnimNode_CopyPoseFromMesh.h"
 #include "Animation/AnimInstanceProxy.h"
+#include "Animation/AnimTrace.h"
 
 /////////////////////////////////////////////////////
 // FAnimNode_CopyPoseFromMesh
@@ -152,6 +153,9 @@ void FAnimNode_CopyPoseFromMesh::Update_AnyThread(const FAnimationUpdateContext&
 	// This introduces a frame of latency in setting the pin-driven source component,
 	// but we cannot do the work to extract transforms on a worker thread as it is not thread safe.
 	GetEvaluateGraphExposedInputs().Execute(Context);
+
+	TRACE_ANIM_NODE_VALUE(Context, TEXT("Component"), *GetNameSafe(CurrentlyUsedSourceMeshComponent.IsValid() ? CurrentlyUsedSourceMeshComponent.Get() : nullptr));
+	TRACE_ANIM_NODE_VALUE(Context, TEXT("Mesh"), *GetNameSafe(CurrentlyUsedSourceMeshComponent.IsValid() ? CurrentlyUsedSourceMeshComponent.Get()->SkeletalMesh : nullptr));
 }
 
 void FAnimNode_CopyPoseFromMesh::Evaluate_AnyThread(FPoseContext& Output)
diff --git a/Engine/Source/Runtime/AnimGraphRuntime/Private/AnimNodes/AnimNode_CurveSource.cpp b/Engine/Source/Runtime/AnimGraphRuntime/Private/AnimNodes/AnimNode_CurveSource.cpp
--- a/Engine/Source/Runtime/AnimGraphRuntime/Private/AnimNodes/AnimNode_CurveSource.cpp
+++ b/Engine/Source/Runtime/AnimGraphRuntime/Private/AnimNodes/AnimNode_CurveSource.cpp
@@ -3,6 +3,7 @@
 #include "AnimNodes/AnimNode_CurveSource.h"
 #include "AnimationRuntime.h"
 #include "Animation/AnimInstanceProxy.h"
+#include "Animation/AnimTrace.h"
 
 FAnimNode_CurveSource::FAnimNode_CurveSource()
 	: SourceBinding(ICurveSourceInterface::DefaultBinding)
@@ -90,7 +91,10 @@ void FAnimNode_CurveSource::Evaluate_AnyThread(FPoseContext& Output)
 			{
 				const float CurrentValue = Output.Curve.Get(NameUID);
 				const float ClampedAlpha = FMath::Clamp(Alpha, 0.0f, 1.0f);
-				Output.Curve.Set(NameUID, FMath::Lerp(CurrentValue, NamedValue.Value, ClampedAlpha));
+				const float LerpedValue = FMath::Lerp(CurrentValue, NamedValue.Value, ClampedAlpha);
+				Output.Curve.Set(NameUID, LerpedValue);
+
+				TRACE_ANIM_NODE_VALUE(Output, *NamedValue.Name.ToString(), LerpedValue);
 			}
 		}
 	}
diff --git a/Engine/Source/Runtime/AnimGraphRuntime/Private/AnimNodes/AnimNode_LayeredBoneBlend.cpp b/Engine/Source/Runtime/AnimGraphRuntime/Private/AnimNodes/AnimNode_LayeredBoneBlend.cpp
--- a/Engine/Source/Runtime/AnimGraphRuntime/Private/AnimNodes/AnimNode_LayeredBoneBlend.cpp
+++ b/Engine/Source/Runtime/AnimGraphRuntime/Private/AnimNodes/AnimNode_LayeredBoneBlend.cpp
@@ -3,6 +3,7 @@
 #include "AnimNodes/AnimNode_LayeredBoneBlend.h"
 #include "AnimationRuntime.h"
 #include "Animation/AnimInstanceProxy.h"
+#include "Animation/AnimTrace.h"
 
 #define DEFAULT_SOURCEINDEX 0xFF
 /////////////////////////////////////////////////////
@@ -200,6 +201,8 @@ void FAnimNode_LayeredBoneBlend::Update_AnyThread(const FAnimationUpdateContext&
 	{
 		BasePose.Update(Context);
 	}
+
+	TRACE_ANIM_NODE_VALUE(Context, TEXT("Num Poses"), BlendPoses.Num());
 }
 
 void FAnimNode_LayeredBoneBlend::Evaluate_AnyThread(FPoseContext& Output)
diff --git a/Engine/Source/Runtime/AnimGraphRuntime/Private/AnimNodes/AnimNode_MakeDynamicAdditive.cpp b/Engine/Source/Runtime/AnimGraphRuntime/Private/AnimNodes/AnimNode_MakeDynamicAdditive.cpp
--- a/Engine/Source/Runtime/AnimGraphRuntime/Private/AnimNodes/AnimNode_MakeDynamicAdditive.cpp
+++ b/Engine/Source/Runtime/AnimGraphRuntime/Private/AnimNodes/AnimNode_MakeDynamicAdditive.cpp
@@ -2,6 +2,7 @@
 
 #include "AnimNodes/AnimNode_MakeDynamicAdditive.h"
 #include "AnimationRuntime.h"
+#include "Animation/AnimTrace.h"
 
 /////////////////////////////////////////////////////
 // FAnimNode_MakeDynamicAdditive
@@ -32,6 +33,8 @@ void FAnimNode_MakeDynamicAdditive::Update_AnyThread(const FAnimationUpdateConte
 	DECLARE_SCOPE_HIERARCHICAL_COUNTER_ANIMNODE(Update_AnyThread)
 	Base.Update(Context.FractionalWeight(1.f));
 	Additive.Update(Context.FractionalWeight(1.f));
+
+	TRACE_ANIM_NODE_VALUE(Context, TEXT("Mesh Space Additive"), bMeshSpaceAdditive);
 }
 
 void FAnimNode_MakeDynamicAdditive::Evaluate_AnyThread(FPoseContext& Output)
diff --git a/Engine/Source/Runtime/AnimGraphRuntime/Private/AnimNodes/AnimNode_PoseByName.cpp b/Engine/Source/Runtime/AnimGraphRuntime/Private/AnimNodes/AnimNode_PoseByName.cpp
--- a/Engine/Source/Runtime/AnimGraphRuntime/Private/AnimNodes/AnimNode_PoseByName.cpp
+++ b/Engine/Source/Runtime/AnimGraphRuntime/Private/AnimNodes/AnimNode_PoseByName.cpp
@@ -2,6 +2,7 @@
 
 #include "AnimNodes/AnimNode_PoseByName.h"
 #include "Animation/AnimInstanceProxy.h"
+#include "Animation/AnimTrace.h"
 
 /////////////////////////////////////////////////////
 // FAnimPoseByNameNode
@@ -35,6 +36,9 @@ void FAnimNode_PoseByName::UpdateAssetPlayer(const FAnimationUpdateContext& Cont
 		RebuildPoseList(Context.AnimInstanceProxy->GetRequiredBones(), PoseAsset);
 		CurrentPoseName = PoseName;
 	}
+
+	TRACE_ANIM_NODE_VALUE(Context, TEXT("Pose Asset"), CurrentPoseAsset.IsValid() ? *CurrentPoseAsset.Get()->GetName() : TEXT("None"));
+	TRACE_ANIM_NODE_VALUE(Context, TEXT("Pose"), *PoseName.ToString());
 }
 
 void FAnimNode_PoseByName::Evaluate_AnyThread(FPoseContext& Output)
diff --git a/Engine/Source/Runtime/AnimGraphRuntime/Private/AnimNodes/AnimNode_PoseHandler.cpp b/Engine/Source/Runtime/AnimGraphRuntime/Private/AnimNodes/AnimNode_PoseHandler.cpp
--- a/Engine/Source/Runtime/AnimGraphRuntime/Private/AnimNodes/AnimNode_PoseHandler.cpp
+++ b/Engine/Source/Runtime/AnimGraphRuntime/Private/AnimNodes/AnimNode_PoseHandler.cpp
@@ -2,6 +2,7 @@
 
 #include "AnimNodes/AnimNode_PoseHandler.h"
 #include "Animation/AnimInstanceProxy.h"
+#include "Animation/AnimTrace.h"
 
 /////////////////////////////////////////////////////
 // FAnimPoseByNameNode
@@ -77,6 +78,8 @@ void FAnimNode_PoseHandler::UpdateAssetPlayer(const FAnimationUpdateContext& Con
 	{
 		UpdatePoseAssetProperty(Context.AnimInstanceProxy);
 	}
+
+	TRACE_ANIM_NODE_VALUE(Context, TEXT("Pose Asset"), CurrentPoseAsset.IsValid() ? *CurrentPoseAsset.Get()->GetName() : TEXT("None"));
 }
 
 void FAnimNode_PoseHandler::OverrideAsset(UAnimationAsset* NewAsset)
diff --git a/Engine/Source/Runtime/AnimGraphRuntime/Private/AnimNodes/AnimNode_PoseSnapshot.cpp b/Engine/Source/Runtime/AnimGraphRuntime/Private/AnimNodes/AnimNode_PoseSnapshot.cpp
--- a/Engine/Source/Runtime/AnimGraphRuntime/Private/AnimNodes/AnimNode_PoseSnapshot.cpp
+++ b/Engine/Source/Runtime/AnimGraphRuntime/Private/AnimNodes/AnimNode_PoseSnapshot.cpp
@@ -2,6 +2,7 @@
 
 #include "AnimNodes/AnimNode_PoseSnapshot.h"
 #include "Animation/AnimInstanceProxy.h"
+#include "Animation/AnimTrace.h"
 
 /////////////////////////////////////////////////////
 // FAnimNode_PoseSnapshot
@@ -51,6 +52,8 @@ void FAnimNode_PoseSnapshot::Update_AnyThread(const FAnimationUpdateContext& Con
 	DECLARE_SCOPE_HIERARCHICAL_COUNTER_ANIMNODE(Update_AnyThread)
 	// Evaluate any BP logic plugged into this node
 	GetEvaluateGraphExposedInputs().Execute(Context);
+
+	TRACE_ANIM_NODE_VALUE(Context, TEXT("Snapshot Name"), Snapshot.SnapshotName);
 }
 
 void FAnimNode_PoseSnapshot::Evaluate_AnyThread(FPoseContext& Output)
diff --git a/Engine/Source/Runtime/AnimGraphRuntime/Private/AnimNodes/AnimNode_RandomPlayer.cpp b/Engine/Source/Runtime/AnimGraphRuntime/Private/AnimNodes/AnimNode_RandomPlayer.cpp
--- a/Engine/Source/Runtime/AnimGraphRuntime/Private/AnimNodes/AnimNode_RandomPlayer.cpp
+++ b/Engine/Source/Runtime/AnimGraphRuntime/Private/AnimNodes/AnimNode_RandomPlayer.cpp
@@ -4,6 +4,7 @@
 #include "AnimationRuntime.h"
 #include "Animation/AnimInstanceProxy.h"
 #include "Animation/AnimSequence.h"
+#include "Animation/AnimTrace.h"
 
 FAnimNode_RandomPlayer::FAnimNode_RandomPlayer()
 : CurrentEntry(INDEX_NONE)
@@ -159,13 +160,22 @@ void FAnimNode_RandomPlayer::Update_AnyThread(const FAnimationUpdateContext& Con
 			FAnimTickRecord& TickRecord = AnimProxy->CreateUninitializedTickRecord(INDEX_NONE, SyncGroup);
 			AnimProxy->MakeSequenceTickRecord(TickRecord, Entries[CurrentEntry].Sequence, true, CurrentData->PlayRate, CurrentData->BlendWeight, CurrentData->InternalTimeAccumulator, CurrentData->MarkerTickRecord);
 
+			TRACE_ANIM_TICK_RECORD(Context, TickRecord);
+
 			if(NextData->BlendWeight > 0.0f)
 			{
 				FAnimTickRecord& NextTickRecord = AnimProxy->CreateUninitializedTickRecord(INDEX_NONE, SyncGroup);
 				AnimProxy->MakeSequenceTickRecord(NextTickRecord, Entries[NextEntry].Sequence, true, NextData->PlayRate, NextData->BlendWeight, NextData->InternalTimeAccumulator, NextData->MarkerTickRecord);
+
+				TRACE_ANIM_TICK_RECORD(Context, NextTickRecord);
 			}
 		}
 	}
+
+	TRACE_ANIM_NODE_VALUE(Context, TEXT("Current Sequence"), Entries[CurrentEntry].Sequence);
+	TRACE_ANIM_NODE_VALUE(Context, TEXT("Current Weight"), CurrentData ? CurrentData->BlendWeight : 0.0f);
+	TRACE_ANIM_NODE_VALUE(Context, TEXT("Next Sequence"), Entries[NextEntry].Sequence);
+	TRACE_ANIM_NODE_VALUE(Context, TEXT("Next Weight"), NextData ? NextData->BlendWeight : 0.0f);
 }
 
 void FAnimNode_RandomPlayer::Evaluate_AnyThread(FPoseContext& Output)
diff --git a/Engine/Source/Runtime/AnimGraphRuntime/Private/AnimNodes/AnimNode_RefPose.cpp b/Engine/Source/Runtime/AnimGraphRuntime/Private/AnimNodes/AnimNode_RefPose.cpp
--- a/Engine/Source/Runtime/AnimGraphRuntime/Private/AnimNodes/AnimNode_RefPose.cpp
+++ b/Engine/Source/Runtime/AnimGraphRuntime/Private/AnimNodes/AnimNode_RefPose.cpp
@@ -1,10 +1,25 @@
 // Copyright 1998-2019 Epic Games, Inc. All Rights Reserved.
 
 #include "AnimNodes/AnimNode_RefPose.h"
+#include "Animation/AnimTrace.h"
 
 /////////////////////////////////////////////////////
 // FAnimRefPoseNode
 
+/** Helper for enum output... */
+#ifndef CASE_ENUM_TO_TEXT
+#define CASE_ENUM_TO_TEXT(txt) case txt: return TEXT(#txt);
+#endif
+
+const TCHAR* GetRefPostTypeText(ERefPoseType RefPose)
+{
+	switch (RefPose)
+	{
+		FOREACH_ENUM_EREFPOSETYPE(CASE_ENUM_TO_TEXT)
+	}
+	return TEXT("Unknown Ref Pose Type");
+}
+
 void FAnimNode_RefPose::Evaluate_AnyThread(FPoseContext& Output)
 {
 	DECLARE_SCOPE_HIERARCHICAL_COUNTER_ANIMNODE(Evaluate_AnyThread)
@@ -22,6 +37,8 @@ void FAnimNode_RefPose::Evaluate_AnyThread(FPoseContext& Output)
 		Output.ResetToAdditiveIdentity();
 		break;
 	}
+
+	TRACE_ANIM_NODE_VALUE(Output, TEXT("Ref Pose Type"), GetRefPostTypeText(RefPoseType));
 }
 
 void FAnimNode_MeshSpaceRefPose::EvaluateComponentSpace_AnyThread(FComponentSpacePoseContext& Output)
@@ -30,20 +47,6 @@ void FAnimNode_MeshSpaceRefPose::EvaluateComponentSpace_AnyThread(FComponentSpac
 	Output.ResetToRefPose();
 }
 
-/** Helper for enum output... */
-#ifndef CASE_ENUM_TO_TEXT
-#define CASE_ENUM_TO_TEXT(txt) case txt: return TEXT(#txt);
-#endif
-
-const TCHAR* GetRefPostTypeText(ERefPoseType RefPose)
-{
-	switch (RefPose)
-	{
-		FOREACH_ENUM_EREFPOSETYPE(CASE_ENUM_TO_TEXT)
-	}
-	return TEXT("Unknown Ref Pose Type");
-}
-
 void FAnimNode_RefPose::GatherDebugData(FNodeDebugData& DebugData)
 {
 	DECLARE_SCOPE_HIERARCHICAL_COUNTER_ANIMNODE(GatherDebugData)
diff --git a/Engine/Source/Runtime/AnimGraphRuntime/Private/AnimNodes/AnimNode_RotateRootBone.cpp b/Engine/Source/Runtime/AnimGraphRuntime/Private/AnimNodes/AnimNode_RotateRootBone.cpp
--- a/Engine/Source/Runtime/AnimGraphRuntime/Private/AnimNodes/AnimNode_RotateRootBone.cpp
+++ b/Engine/Source/Runtime/AnimGraphRuntime/Private/AnimNodes/AnimNode_RotateRootBone.cpp
@@ -1,6 +1,7 @@
 // Copyright 1998-2019 Epic Games, Inc. All Rights Reserved.
 
 #include "AnimNodes/AnimNode_RotateRootBone.h"
+#include "Animation/AnimTrace.h"
 
 /////////////////////////////////////////////////////
 // FAnimNode_RotateRootBone
@@ -30,6 +31,9 @@ void FAnimNode_RotateRootBone::Update_AnyThread(const FAnimationUpdateContext& C
 
 	ActualPitch = PitchScaleBiasClamp.ApplyTo(Pitch, Context.GetDeltaTime());
 	ActualYaw = YawScaleBiasClamp.ApplyTo(Yaw, Context.GetDeltaTime());
+
+	TRACE_ANIM_NODE_VALUE(Context, TEXT("Pitch"), ActualPitch);
+	TRACE_ANIM_NODE_VALUE(Context, TEXT("Yaw"), ActualYaw);
 }
 
 void FAnimNode_RotateRootBone::Evaluate_AnyThread(FPoseContext& Output)
diff --git a/Engine/Source/Runtime/AnimGraphRuntime/Private/AnimNodes/AnimNode_RotationOffsetBlendSpace.cpp b/Engine/Source/Runtime/AnimGraphRuntime/Private/AnimNodes/AnimNode_RotationOffsetBlendSpace.cpp
--- a/Engine/Source/Runtime/AnimGraphRuntime/Private/AnimNodes/AnimNode_RotationOffsetBlendSpace.cpp
+++ b/Engine/Source/Runtime/AnimGraphRuntime/Private/AnimNodes/AnimNode_RotationOffsetBlendSpace.cpp
@@ -3,6 +3,7 @@
 #include "AnimNodes/AnimNode_RotationOffsetBlendSpace.h"
 #include "Animation/AnimInstanceProxy.h"
 #include "AnimationRuntime.h"
+#include "Animation/AnimTrace.h"
 
 /////////////////////////////////////////////////////
 // FAnimNode_RotationOffsetBlendSpace
@@ -67,6 +68,9 @@ void FAnimNode_RotationOffsetBlendSpace::UpdateAssetPlayer(const FAnimationUpdat
 	}
 
 	BasePose.Update(Context);
+
+	TRACE_ANIM_NODE_VALUE(Context, TEXT("Alpha"), ActualAlpha);
+	TRACE_ANIM_NODE_VALUE(Context, TEXT("PlayTime"), InternalTimeAccumulator);
 }
 
 void FAnimNode_RotationOffsetBlendSpace::Evaluate_AnyThread(FPoseContext& Context)
diff --git a/Engine/Source/Runtime/AnimGraphRuntime/Private/AnimNodes/AnimNode_SequenceEvaluator.cpp b/Engine/Source/Runtime/AnimGraphRuntime/Private/AnimNodes/AnimNode_SequenceEvaluator.cpp
--- a/Engine/Source/Runtime/AnimGraphRuntime/Private/AnimNodes/AnimNode_SequenceEvaluator.cpp
+++ b/Engine/Source/Runtime/AnimGraphRuntime/Private/AnimNodes/AnimNode_SequenceEvaluator.cpp
@@ -2,6 +2,7 @@
 
 #include "AnimNodes/AnimNode_SequenceEvaluator.h"
 #include "Animation/AnimInstanceProxy.h"
+#include "Animation/AnimTrace.h"
 
 float FAnimNode_SequenceEvaluator::GetCurrentAssetTime()
 {
@@ -68,10 +69,10 @@ void FAnimNode_SequenceEvaluator::UpdateAssetPlayer(const FAnimationUpdateContex
 
 			// if you jump from front to end or end to front, your time jump is 0.f, so nothing moves
 			// to prevent that from happening, we set current accumulator to explicit time
- 			if (TimeJump == 0.f)
- 			{
- 				InternalTimeAccumulator = ExplicitTime;
- 			}
+			if (TimeJump == 0.f)
+			{
+				InternalTimeAccumulator = ExplicitTime;
+			}
 			
 			const float DeltaTime = Context.GetDeltaTime();
 			const float RateScale = Sequence->RateScale;
@@ -85,6 +86,10 @@ void FAnimNode_SequenceEvaluator::UpdateAssetPlayer(const FAnimationUpdateContex
 	}
 
 	bReinitialized = false;
+
+	TRACE_ANIM_NODE_VALUE(Context, TEXT("Sequence"), Sequence != nullptr ? Sequence->GetFName() : NAME_None);
+	TRACE_ANIM_NODE_VALUE(Context, TEXT("InputTime"), ExplicitTime);
+	TRACE_ANIM_NODE_VALUE(Context, TEXT("Time"), InternalTimeAccumulator);
 }
 
 void FAnimNode_SequenceEvaluator::Evaluate_AnyThread(FPoseContext& Output)
diff --git a/Engine/Source/Runtime/AnimGraphRuntime/Private/AnimNodes/AnimNode_Slot.cpp b/Engine/Source/Runtime/AnimGraphRuntime/Private/AnimNodes/AnimNode_Slot.cpp
--- a/Engine/Source/Runtime/AnimGraphRuntime/Private/AnimNodes/AnimNode_Slot.cpp
+++ b/Engine/Source/Runtime/AnimGraphRuntime/Private/AnimNodes/AnimNode_Slot.cpp
@@ -2,6 +2,7 @@
 
 #include "AnimNodes/AnimNode_Slot.h"
 #include "Animation/AnimInstanceProxy.h"
+#include "Animation/AnimTrace.h"
 
 /////////////////////////////////////////////////////
 // FAnimNode_Slot
@@ -43,6 +44,14 @@ void FAnimNode_Slot::Update_AnyThread(const FAnimationUpdateContext& Context)
 		const float SourceWeight = FMath::Max(FAnimWeight::GetSmallestRelevantWeight(), WeightData.SourceWeight);
 		Source.Update(Context.FractionalWeight(SourceWeight));
 	}
+
+#if ANIM_TRACE_ENABLED
+	TRACE_ANIM_NODE_VALUE(Context, TEXT("Slot Name"), SlotName);
+	TRACE_ANIM_NODE_VALUE(Context, TEXT("Weight"), WeightData.SlotNodeWeight);
+	TRACE_ANIM_NODE_VALUE(Context, TEXT("Pose Source"), (WeightData.SourceWeight <= ZERO_ANIMWEIGHT_THRESH));
+
+	Context.AnimInstanceProxy->TraceMontageEvaluationData(Context, SlotName);
+#endif
 }
 
 void FAnimNode_Slot::Evaluate_AnyThread(FPoseContext & Output)
diff --git a/Engine/Source/Runtime/AnimGraphRuntime/Private/AnimNodes/AnimNode_TwoWayBlend.cpp b/Engine/Source/Runtime/AnimGraphRuntime/Private/AnimNodes/AnimNode_TwoWayBlend.cpp
--- a/Engine/Source/Runtime/AnimGraphRuntime/Private/AnimNodes/AnimNode_TwoWayBlend.cpp
+++ b/Engine/Source/Runtime/AnimGraphRuntime/Private/AnimNodes/AnimNode_TwoWayBlend.cpp
@@ -100,6 +100,8 @@ void FAnimNode_TwoWayBlend::Update_AnyThread(const FAnimationUpdateContext& Cont
 		// Take all of A
 		A.Update(Context);
 	}
+
+	TRACE_ANIM_NODE_VALUE(Context, TEXT("Alpha"), InternalBlendAlpha);
 }
 
 void FAnimNode_TwoWayBlend::Evaluate_AnyThread(FPoseContext& Output)
diff --git a/Engine/Source/Runtime/AnimGraphRuntime/Private/BoneControllers/AnimNodeRotationMultiplier.cpp b/Engine/Source/Runtime/AnimGraphRuntime/Private/BoneControllers/AnimNodeRotationMultiplier.cpp
--- a/Engine/Source/Runtime/AnimGraphRuntime/Private/BoneControllers/AnimNodeRotationMultiplier.cpp
+++ b/Engine/Source/Runtime/AnimGraphRuntime/Private/BoneControllers/AnimNodeRotationMultiplier.cpp
@@ -4,6 +4,7 @@
 #include "BoneIndices.h"
 #include "BonePose.h"
 #include "BoneControllers/AnimNode_RotationMultiplier.h"
+#include "Animation/AnimTrace.h"
 
 /////////////////////////////////////////////////////
 // FAnimNode_RotationMultiplier
@@ -162,6 +163,10 @@ void FAnimNode_RotationMultiplier::EvaluateSkeletalControl_AnyThread(FComponentS
 			OutBoneTransforms.Add( FBoneTransform(TargetBoneIndex, NewLocalTransform) );
 		}
 	}
+
+	TRACE_ANIM_NODE_VALUE(Output, TEXT("Source Bone"), SourceBone.BoneName);
+	TRACE_ANIM_NODE_VALUE(Output, TEXT("Target Bone"), TargetBone.BoneName);
+	TRACE_ANIM_NODE_VALUE(Output, TEXT("Multiplier"), Multiplier);
 }
 
 bool FAnimNode_RotationMultiplier::IsValidToEvaluate(const USkeleton* Skeleton, const FBoneContainer& RequiredBones) 
diff --git a/Engine/Source/Runtime/AnimGraphRuntime/Private/BoneControllers/AnimNode_BoneDrivenController.cpp b/Engine/Source/Runtime/AnimGraphRuntime/Private/BoneControllers/AnimNode_BoneDrivenController.cpp
--- a/Engine/Source/Runtime/AnimGraphRuntime/Private/BoneControllers/AnimNode_BoneDrivenController.cpp
+++ b/Engine/Source/Runtime/AnimGraphRuntime/Private/BoneControllers/AnimNode_BoneDrivenController.cpp
@@ -4,6 +4,7 @@
 
 #include "Curves/CurveFloat.h"
 #include "Animation/AnimInstanceProxy.h"
+#include "Animation/AnimTrace.h"
 
 /////////////////////////////////////////////////////
 // FAnimNode_BoneDrivenController
@@ -169,6 +170,18 @@ void FAnimNode_BoneDrivenController::EvaluateSkeletalControl_AnyThread(FComponen
 	{
 		OutBoneTransforms.Add(FBoneTransform(TargetBoneIndex, ModifiedLocalTM));
 	}
+
+#if ANIM_TRACE_ENABLED
+	TRACE_ANIM_NODE_VALUE(Output, TEXT("Driving Bone"), SourceBone.BoneName);
+	if (DestinationMode == EDrivenDestinationMode::Bone)
+	{
+		TRACE_ANIM_NODE_VALUE(Output, TEXT("Driven Bone"), TargetBone.BoneName);
+	}
+	else
+	{
+		TRACE_ANIM_NODE_VALUE(Output, TEXT("Driven Parameter"), ParameterName);
+	}
+#endif
 }
 
 void FAnimNode_BoneDrivenController::EvaluateComponentSpaceInternal(FComponentSpacePoseContext& Context)
diff --git a/Engine/Source/Runtime/AnimGraphRuntime/Private/BoneControllers/AnimNode_Constraint.cpp b/Engine/Source/Runtime/AnimGraphRuntime/Private/BoneControllers/AnimNode_Constraint.cpp
--- a/Engine/Source/Runtime/AnimGraphRuntime/Private/BoneControllers/AnimNode_Constraint.cpp
+++ b/Engine/Source/Runtime/AnimGraphRuntime/Private/BoneControllers/AnimNode_Constraint.cpp
@@ -5,6 +5,8 @@
 #include "AnimationRuntime.h"
 #include "SceneManagement.h"
 #include "Components/SkeletalMeshComponent.h"
+#include "Animation/AnimTrace.h"
+
 /////////////////////////////////////////////////////
 // FAnimNode_Constraint
 
@@ -69,6 +71,13 @@ void FAnimNode_Constraint::EvaluateSkeletalControl_AnyThread(FComponentSpacePose
 		CachedConstrainedTransform = ConstrainedTransform;
 #endif // WITH_EDITOR
 	}
+
+#if ANIM_TRACE_ENABLED
+	for (int32 ConstraintIndex = 0; ConstraintIndex < ConstraintNum; ++ConstraintIndex)
+	{
+		TRACE_ANIM_NODE_VALUE(Output, *ConstraintSetup[ConstraintIndex].TargetBone.BoneName.ToString(), ConstraintWeights[ConstraintIndex]);
+	}
+#endif
 }
 
 bool FAnimNode_Constraint::IsValidToEvaluate(const USkeleton* Skeleton, const FBoneContainer& RequiredBones) 
diff --git a/Engine/Source/Runtime/AnimGraphRuntime/Private/BoneControllers/AnimNode_CopyBone.cpp b/Engine/Source/Runtime/AnimGraphRuntime/Private/BoneControllers/AnimNode_CopyBone.cpp
--- a/Engine/Source/Runtime/AnimGraphRuntime/Private/BoneControllers/AnimNode_CopyBone.cpp
+++ b/Engine/Source/Runtime/AnimGraphRuntime/Private/BoneControllers/AnimNode_CopyBone.cpp
@@ -3,6 +3,7 @@
 #include "BoneControllers/AnimNode_CopyBone.h"
 #include "AnimationRuntime.h"
 #include "Animation/AnimInstanceProxy.h"
+#include "Animation/AnimTrace.h"
 
 /////////////////////////////////////////////////////
 // FAnimNode_CopyBone
@@ -78,6 +79,9 @@ void FAnimNode_CopyBone::EvaluateSkeletalControl_AnyThread(FComponentSpacePoseCo
 
 	// Output new transform for current bone.
 	OutBoneTransforms.Add(FBoneTransform(TargetBoneIndex, CurrentBoneTM));
+
+	TRACE_ANIM_NODE_VALUE(Output, TEXT("Source Bone"), SourceBone.BoneName);
+	TRACE_ANIM_NODE_VALUE(Output, TEXT("Target Bone"), TargetBone.BoneName);
 }
 
 bool FAnimNode_CopyBone::IsValidToEvaluate(const USkeleton* Skeleton, const FBoneContainer& RequiredBones) 
diff --git a/Engine/Source/Runtime/AnimGraphRuntime/Private/BoneControllers/AnimNode_HandIKRetargeting.cpp b/Engine/Source/Runtime/AnimGraphRuntime/Private/BoneControllers/AnimNode_HandIKRetargeting.cpp
--- a/Engine/Source/Runtime/AnimGraphRuntime/Private/BoneControllers/AnimNode_HandIKRetargeting.cpp
+++ b/Engine/Source/Runtime/AnimGraphRuntime/Private/BoneControllers/AnimNode_HandIKRetargeting.cpp
@@ -1,6 +1,7 @@
 // Copyright 1998-2019 Epic Games, Inc. All Rights Reserved.
 
 #include "BoneControllers/AnimNode_HandIKRetargeting.h"
+#include "Animation/AnimTrace.h"
 
 /////////////////////////////////////////////////////
 // FAnimNode_HandIKRetargeting
@@ -94,6 +95,14 @@ void FAnimNode_HandIKRetargeting::EvaluateSkeletalControl_AnyThread(FComponentSp
 			OutBoneTransforms.Sort(FCompareBoneTransformIndex());
 		}
 	}
+
+#if ANIM_TRACE_ENABLED
+	TRACE_ANIM_NODE_VALUE(Output, TEXT("Hand FK Weight"), HandFKWeight);
+	for (int32 BoneIndex = 0; BoneIndex < IKBonesToMove.Num(); BoneIndex++)
+	{
+		TRACE_ANIM_NODE_VALUE(Output, *FString::Printf(TEXT("Bone %d"), BoneIndex), IKBonesToMove[BoneIndex].BoneName);
+	}
+#endif
 }
 
 bool FAnimNode_HandIKRetargeting::IsValidToEvaluate(const USkeleton* Skeleton, const FBoneContainer& RequiredBones)
diff --git a/Engine/Source/Runtime/AnimGraphRuntime/Private/BoneControllers/AnimNode_LookAt.cpp b/Engine/Source/Runtime/AnimGraphRuntime/Private/BoneControllers/AnimNode_LookAt.cpp
--- a/Engine/Source/Runtime/AnimGraphRuntime/Private/BoneControllers/AnimNode_LookAt.cpp
+++ b/Engine/Source/Runtime/AnimGraphRuntime/Private/BoneControllers/AnimNode_LookAt.cpp
@@ -7,6 +7,7 @@
 #include "AnimationCoreLibrary.h"
 #include "Engine/Engine.h"
 #include "Materials/MaterialInstanceDynamic.h"
+#include "Animation/AnimTrace.h"
 
 static const FVector DefaultLookAtAxis(0.f, 1.f, 0.f);
 static const FVector DefaultLookUpAxis(1.f, 0.f, 0.f);
@@ -142,6 +143,11 @@ void FAnimNode_LookAt::EvaluateSkeletalControl_AnyThread(FComponentSpacePoseCont
 #if !UE_BUILD_SHIPPING
 	CachedLookAtTransform = ComponentBoneTransform;
 #endif
+
+	TRACE_ANIM_NODE_VALUE(Output, TEXT("Bone"), BoneToModify.BoneName);
+	TRACE_ANIM_NODE_VALUE(Output, TEXT("Look At Target"), LookAtTarget.HasValidSetup() ? LookAtTarget.GetTargetSetup() : NAME_None);
+	TRACE_ANIM_NODE_VALUE(Output, TEXT("Look At Location"), LookAtLocation);
+	TRACE_ANIM_NODE_VALUE(Output, TEXT("Target Location"), CachedCurrentTargetLocation);
 }
 
 void FAnimNode_LookAt::EvaluateComponentSpaceInternal(FComponentSpacePoseContext& Context)
diff --git a/Engine/Source/Runtime/AnimGraphRuntime/Private/BoneControllers/AnimNode_ModifyBone.cpp b/Engine/Source/Runtime/AnimGraphRuntime/Private/BoneControllers/AnimNode_ModifyBone.cpp
--- a/Engine/Source/Runtime/AnimGraphRuntime/Private/BoneControllers/AnimNode_ModifyBone.cpp
+++ b/Engine/Source/Runtime/AnimGraphRuntime/Private/BoneControllers/AnimNode_ModifyBone.cpp
@@ -3,6 +3,7 @@
 #include "BoneControllers/AnimNode_ModifyBone.h"
 #include "AnimationRuntime.h"
 #include "Animation/AnimInstanceProxy.h"
+#include "Animation/AnimTrace.h"
 
 /////////////////////////////////////////////////////
 // FAnimNode_ModifyBone
@@ -103,6 +104,8 @@ void FAnimNode_ModifyBone::EvaluateSkeletalControl_AnyThread(FComponentSpacePose
 	}
 	
 	OutBoneTransforms.Add( FBoneTransform(BoneToModify.GetCompactPoseIndex(BoneContainer), NewBoneTM) );
+
+	TRACE_ANIM_NODE_VALUE(Output, TEXT("Target"), BoneToModify.BoneName);
 }
 
 bool FAnimNode_ModifyBone::IsValidToEvaluate(const USkeleton* Skeleton, const FBoneContainer& RequiredBones) 
diff --git a/Engine/Source/Runtime/AnimGraphRuntime/Private/BoneControllers/AnimNode_ObserveBone.cpp b/Engine/Source/Runtime/AnimGraphRuntime/Private/BoneControllers/AnimNode_ObserveBone.cpp
--- a/Engine/Source/Runtime/AnimGraphRuntime/Private/BoneControllers/AnimNode_ObserveBone.cpp
+++ b/Engine/Source/Runtime/AnimGraphRuntime/Private/BoneControllers/AnimNode_ObserveBone.cpp
@@ -3,6 +3,7 @@
 #include "BoneControllers/AnimNode_ObserveBone.h"
 #include "AnimationRuntime.h"
 #include "Animation/AnimInstanceProxy.h"
+#include "Animation/AnimTrace.h"
 
 /////////////////////////////////////////////////////
 // FAnimNode_ObserveBone
@@ -48,6 +49,11 @@ void FAnimNode_ObserveBone::EvaluateSkeletalControl_AnyThread(FComponentSpacePos
 	Translation = BoneTM.GetTranslation();
 	Rotation = BoneTM.GetRotation().Rotator();
 	Scale = BoneTM.GetScale3D();
+
+	TRACE_ANIM_NODE_VALUE(Output, TEXT("Bone"), BoneToObserve.BoneName);
+	TRACE_ANIM_NODE_VALUE(Output, TEXT("Translation"), Translation);
+	TRACE_ANIM_NODE_VALUE(Output, TEXT("Rotation"), Rotation.Euler());
+	TRACE_ANIM_NODE_VALUE(Output, TEXT("Scale"), Scale);
 }
 
 bool FAnimNode_ObserveBone::IsValidToEvaluate(const USkeleton* Skeleton, const FBoneContainer& RequiredBones)
diff --git a/Engine/Source/Runtime/AnimGraphRuntime/Private/BoneControllers/AnimNode_SkeletalControlBase.cpp b/Engine/Source/Runtime/AnimGraphRuntime/Private/BoneControllers/AnimNode_SkeletalControlBase.cpp
--- a/Engine/Source/Runtime/AnimGraphRuntime/Private/BoneControllers/AnimNode_SkeletalControlBase.cpp
+++ b/Engine/Source/Runtime/AnimGraphRuntime/Private/BoneControllers/AnimNode_SkeletalControlBase.cpp
@@ -113,6 +113,8 @@ void FAnimNode_SkeletalControlBase::Update_AnyThread(const FAnimationUpdateConte
 			UpdateInternal(Context);
 		}
 	}
+
+	TRACE_ANIM_NODE_VALUE(Context, TEXT("Alpha"), ActualAlpha);
 }
 
 bool ContainsNaN(const TArray<FBoneTransform> & BoneTransforms)
diff --git a/Engine/Source/Runtime/AnimGraphRuntime/Private/BoneControllers/AnimNode_SplineIK.cpp b/Engine/Source/Runtime/AnimGraphRuntime/Private/BoneControllers/AnimNode_SplineIK.cpp
--- a/Engine/Source/Runtime/AnimGraphRuntime/Private/BoneControllers/AnimNode_SplineIK.cpp
+++ b/Engine/Source/Runtime/AnimGraphRuntime/Private/BoneControllers/AnimNode_SplineIK.cpp
@@ -5,6 +5,7 @@
 #include "AnimationRuntime.h"
 #include "Animation/AnimInstanceProxy.h"
 #include "SplineIK.h"
+#include "Animation/AnimTrace.h"
 
 FAnimNode_SplineIK::FAnimNode_SplineIK() 
 	: BoneAxis(ESplineBoneAxis::X)
@@ -91,6 +92,9 @@ void FAnimNode_SplineIK::EvaluateSkeletalControl_AnyThread(FComponentSpacePoseCo
 			OutBoneTransforms.Emplace(CompactPoseBoneIndices[OutBoneIndex], OutTransforms[OutBoneIndex]);
 		}
 	}
+
+	TRACE_ANIM_NODE_VALUE(Output, TEXT("Start Bone"), StartBone.BoneName);
+	TRACE_ANIM_NODE_VALUE(Output, TEXT("End Bone"), EndBone.BoneName);
 }
 
 bool FAnimNode_SplineIK::IsValidToEvaluate(const USkeleton* Skeleton, const FBoneContainer& RequiredBones) 
diff --git a/Engine/Source/Runtime/AnimGraphRuntime/Private/BoneControllers/AnimNode_SpringBone.cpp b/Engine/Source/Runtime/AnimGraphRuntime/Private/BoneControllers/AnimNode_SpringBone.cpp
--- a/Engine/Source/Runtime/AnimGraphRuntime/Private/BoneControllers/AnimNode_SpringBone.cpp
+++ b/Engine/Source/Runtime/AnimGraphRuntime/Private/BoneControllers/AnimNode_SpringBone.cpp
@@ -3,6 +3,7 @@
 #include "BoneControllers/AnimNode_SpringBone.h"
 #include "GameFramework/WorldSettings.h"
 #include "Animation/AnimInstanceProxy.h"
+#include "Animation/AnimTrace.h"
 
 /////////////////////////////////////////////////////
 // FAnimNode_SpringBone
@@ -213,6 +214,8 @@ void FAnimNode_SpringBone::EvaluateSkeletalControl_AnyThread(FComponentSpacePose
 
 	// Output new transform for current bone.
 	OutBoneTransforms.Add(FBoneTransform(SpringBoneIndex, OutBoneTM));
+
+	TRACE_ANIM_NODE_VALUE(Output, TEXT("Remaining Time"), RemainingTime);
 }
 
 
diff --git a/Engine/Source/Runtime/AnimGraphRuntime/Private/BoneControllers/AnimNode_Trail.cpp b/Engine/Source/Runtime/AnimGraphRuntime/Private/BoneControllers/AnimNode_Trail.cpp
--- a/Engine/Source/Runtime/AnimGraphRuntime/Private/BoneControllers/AnimNode_Trail.cpp
+++ b/Engine/Source/Runtime/AnimGraphRuntime/Private/BoneControllers/AnimNode_Trail.cpp
@@ -3,6 +3,7 @@
 #include "BoneControllers/AnimNode_Trail.h"
 #include "Animation/AnimInstanceProxy.h"
 #include "AngularLimit.h"
+#include "Animation/AnimTrace.h"
 /////////////////////////////////////////////////////
 // FAnimNode_Trail
 
@@ -47,6 +48,8 @@ void FAnimNode_Trail::UpdateInternal(const FAnimationUpdateContext& Context)
 	FAnimNode_SkeletalControlBase::UpdateInternal(Context);
 
 	ThisTimstep += Context.GetDeltaTime();
+
+	TRACE_ANIM_NODE_VALUE(Context, TEXT("Active Bone"), TrailBone.BoneName);
 }
 
 void FAnimNode_Trail::GatherDebugData(FNodeDebugData& DebugData)
diff --git a/Engine/Source/Runtime/AnimGraphRuntime/Private/BoneControllers/AnimNode_TwoBoneIK.cpp b/Engine/Source/Runtime/AnimGraphRuntime/Private/BoneControllers/AnimNode_TwoBoneIK.cpp
--- a/Engine/Source/Runtime/AnimGraphRuntime/Private/BoneControllers/AnimNode_TwoBoneIK.cpp
+++ b/Engine/Source/Runtime/AnimGraphRuntime/Private/BoneControllers/AnimNode_TwoBoneIK.cpp
@@ -9,6 +9,7 @@
 #include "SceneManagement.h"
 #include "Materials/MaterialInstanceDynamic.h"
 #include "MaterialShared.h"
+#include "Animation/AnimTrace.h"
 
 DECLARE_CYCLE_STAT(TEXT("TwoBoneIK Eval"), STAT_TwoBoneIK_Eval, STATGROUP_Anim);
 
@@ -187,6 +188,8 @@ void FAnimNode_TwoBoneIK::EvaluateSkeletalControl_AnyThread(FComponentSpacePoseC
 
 	// Make sure we have correct number of bones
 	check(OutBoneTransforms.Num() == 3);
+
+	TRACE_ANIM_NODE_VALUE(Output, TEXT("IK Bone"), IKBone.BoneName);
 }
 
 bool FAnimNode_TwoBoneIK::IsValidToEvaluate(const USkeleton* Skeleton, const FBoneContainer& RequiredBones) 
diff --git a/Engine/Source/Runtime/AnimGraphRuntime/Public/AnimGraphRuntimeTrace.h b/Engine/Source/Runtime/AnimGraphRuntime/Public/AnimGraphRuntimeTrace.h
new file mode 100644
--- /dev/null
+++ b/Engine/Source/Runtime/AnimGraphRuntime/Public/AnimGraphRuntimeTrace.h
@@ -0,0 +1,29 @@
+// Copyright 1998-2019 Epic Games, Inc. All Rights Reserved.
+
+#pragma once
+
+#include "CoreMinimal.h"
+#include "Animation/AnimTrace.h"
+
+#if ANIM_TRACE_ENABLED
+
+struct FAnimationBaseContext;
+struct FAnimNode_BlendSpacePlayer;
+
+struct FAnimGraphRuntimeTrace
+{
+	/** Initialize animation graph runtime tracing */
+	ANIMGRAPHRUNTIME_API static void Init();
+
+	/** Helper function to output debug info for blendspace player nodes */
+	ANIMGRAPHRUNTIME_API static void OutputBlendSpacePlayer(const FAnimationBaseContext& InContext, const FAnimNode_BlendSpacePlayer& InNode);
+};
+
+#define TRACE_BLENDSPACE_PLAYER(Context, Node) \
+	FAnimGraphRuntimeTrace::OutputBlendSpacePlayer(Context, Node);
+
+#else
+
+#define TRACE_BLENDSPACE_PLAYER(Context, Node)
+
+#endif
\ No newline at end of file
diff --git a/Engine/Source/Runtime/AnimGraphRuntime/Public/AnimNodes/AnimNode_PoseHandler.h b/Engine/Source/Runtime/AnimGraphRuntime/Public/AnimNodes/AnimNode_PoseHandler.h
--- a/Engine/Source/Runtime/AnimGraphRuntime/Public/AnimNodes/AnimNode_PoseHandler.h
+++ b/Engine/Source/Runtime/AnimGraphRuntime/Public/AnimNodes/AnimNode_PoseHandler.h
@@ -41,7 +41,7 @@ public:
 	// End of FAnimNode_Base interface
 
 	// FAnimNode_AssetPlayerBase Interface
-	virtual float GetAccumulatedTime() {return 0.f;}
+	virtual float GetAccumulatedTime() const {return 0.f;}
 	virtual void SetAccumulatedTime(const float& NewTime) {}
 	virtual UAnimationAsset* GetAnimAsset() {return PoseAsset;}
 	// End of FAnimNode_AssetPlayerBase Interface
diff --git a/Engine/Source/Runtime/AnimGraphRuntime/Public/AnimNodes/AnimNode_SequenceEvaluator.h b/Engine/Source/Runtime/AnimGraphRuntime/Public/AnimNodes/AnimNode_SequenceEvaluator.h
--- a/Engine/Source/Runtime/AnimGraphRuntime/Public/AnimNodes/AnimNode_SequenceEvaluator.h
+++ b/Engine/Source/Runtime/AnimGraphRuntime/Public/AnimNodes/AnimNode_SequenceEvaluator.h
@@ -86,7 +86,7 @@ public:
 	// End of FAnimNode_Base interface
 
 	// FAnimNode_AssetPlayerBase Interface
-	virtual float GetAccumulatedTime() {return ExplicitTime;}
+	virtual float GetAccumulatedTime() const {return ExplicitTime;}
 	virtual void SetAccumulatedTime(const float& NewTime) {ExplicitTime = NewTime;}
 	virtual UAnimationAsset* GetAnimAsset() {return Sequence;}
 	// End of FAnimNode_AssetPlayerBase Interface
diff --git a/Engine/Source/Runtime/Engine/Classes/Animation/AnimBlueprintGeneratedClass.h b/Engine/Source/Runtime/Engine/Classes/Animation/AnimBlueprintGeneratedClass.h
--- a/Engine/Source/Runtime/Engine/Classes/Animation/AnimBlueprintGeneratedClass.h
+++ b/Engine/Source/Runtime/Engine/Classes/Animation/AnimBlueprintGeneratedClass.h
@@ -11,6 +11,7 @@
 #include "Animation/AnimStateMachineTypes.h"
 #include "Animation/AnimClassInterface.h"
 #include "Animation/AnimNodeBase.h"
+#include "Animation/BlendSpaceBase.h"
 
 #include "AnimBlueprintGeneratedClass.generated.h"
 
@@ -22,13 +23,52 @@ class UAnimStateTransitionNode;
 class UEdGraph;
 class USkeleton;
 
+// Represents the debugging information for a single state within a state machine
+USTRUCT()
+struct FStateMachineStateDebugData
+{
+	GENERATED_BODY()
+
+public:
+	FStateMachineStateDebugData()
+		: StateMachineIndex(INDEX_NONE)
+		, StateIndex(INDEX_NONE)
+		, Weight(0.0f)
+		, ElapsedTime(0.0f)
+	{
+	}
+
+	FStateMachineStateDebugData(int32 InStateMachineIndex, int32 InStateIndex, float InWeight, float InElapsedTime)
+		: StateMachineIndex(InStateMachineIndex)
+		, StateIndex(InStateIndex)
+		, Weight(InWeight)
+		, ElapsedTime(InElapsedTime)
+	{}
+
+	// The index of the state machine
+	int32 StateMachineIndex;
+
+	// The index of the state
+	int32 StateIndex;
+
+	// The last recorded weight for this state
+	float Weight;
+
+	// The time that this state has been active (only valid if this is the current state)
+	float ElapsedTime;
+};
+
 // This structure represents debugging information for a single state machine
 USTRUCT()
 struct FStateMachineDebugData
 {
-	GENERATED_USTRUCT_BODY()
+	GENERATED_BODY()
 
 public:
+	FStateMachineDebugData()
+		: MachineIndex(INDEX_NONE)
+	{}
+
 	// Map from state nodes to their state entry in a state machine
 	TMap<TWeakObjectPtr<UEdGraphNode>, int32> NodeToStateIndex;
 	TMap<TWeakObjectPtr<UEdGraphNode>, int32> NodeToTransitionIndex;
@@ -108,6 +148,9 @@ public:
 	// Map from animation node GUID to property index
 	TMap<FGuid, int32> NodeGuidToIndexMap;
 
+	// The debug data for each state machine state
+	TArray<FStateMachineStateDebugData> StateData;	
+	
 	// History of snapshots of animation data
 	TSimpleRingBuffer<FAnimationFrameSnapshot>* SnapshotBuffer;
 
@@ -129,6 +172,61 @@ public:
 	// History of activated nodes
 	TArray<FNodeVisit> UpdatedNodesThisFrame;
 
+	// Values output by nodes
+	struct FNodeValue
+	{
+		FString Text;
+		int32 NodeID;
+
+		FNodeValue(const FString& InText, int32 InNodeID)
+			: Text(InText)
+			, NodeID(InNodeID)
+		{}
+	};
+
+	// Values output by nodes
+	TArray<FNodeValue> NodeValuesThisFrame;
+
+	// Record of a sequence player's state
+	struct FSequencePlayerRecord
+	{
+		FSequencePlayerRecord(int32 InNodeID, float InPosition, float InLength, int32 InFrameCount)
+			: NodeID(InNodeID)
+			, Position(InPosition)
+			, Length(InLength) 
+			, FrameCount(InFrameCount)
+		{}
+
+		int32 NodeID;
+		float Position;
+		float Length;
+		int32 FrameCount;
+	};
+
+	// All sequence player records this frame
+	TArray<FSequencePlayerRecord> SequencePlayerRecordsThisFrame;
+
+	// Record of a blend space player's state
+	struct FBlendSpacePlayerRecord
+	{
+		FBlendSpacePlayerRecord(int32 InNodeID, UBlendSpaceBase* InBlendSpace, float InPositionX, float InPositionY, float InPositionZ)
+			: NodeID(InNodeID)
+			, BlendSpace(InBlendSpace)
+			, PositionX(InPositionX)
+			, PositionY(InPositionY) 
+			, PositionZ(InPositionY)
+		{}
+
+		int32 NodeID;
+		TWeakObjectPtr<const UBlendSpaceBase> BlendSpace;
+		float PositionX;
+		float PositionY;
+		float PositionZ;
+	};
+
+	// All blend space player records this frame
+	TArray<FBlendSpacePlayerRecord> BlendSpacePlayerRecordsThisFrame;
+
 	// Active pose watches to track
 	TArray<FAnimNodePoseWatch> AnimNodePoseWatch;
 
@@ -158,6 +256,10 @@ public:
 	void ResetNodeVisitSites();
 	void RecordNodeVisit(int32 TargetNodeIndex, int32 SourceNodeIndex, float BlendWeight);
 	void RecordNodeVisitArray(const TArray<FNodeVisit>& Nodes);
+	void RecordStateData(int32 StateMachineIndex, int32 StateIndex, float Weight, float ElapsedTime);
+	void RecordNodeValue(int32 InNodeID, const FString& InText);
+	void RecordSequencePlayer(int32 InNodeID, float InPosition, float InLength, int32 InFrameCount);
+	void RecordBlendSpacePlayer(int32 InNodeID, UBlendSpaceBase* InBlendSpace, float InPositionX, float InPositionY, float InPositionZ);
 
 	void AddPoseWatch(int32 NodeID, FColor Color);
 	void RemovePoseWatch(int32 NodeID);
diff --git a/Engine/Source/Runtime/Engine/Classes/Animation/AnimInstance.h b/Engine/Source/Runtime/Engine/Classes/Animation/AnimInstance.h
--- a/Engine/Source/Runtime/Engine/Classes/Animation/AnimInstance.h
+++ b/Engine/Source/Runtime/Engine/Classes/Animation/AnimInstance.h
@@ -1254,7 +1254,7 @@ public:
 	void AddCurveValue(const FName& CurveName, float Value);
 
 	/** Given a machine and state index, record a state weight for this frame */
-	void RecordStateWeight(const int32 InMachineClassIndex, const int32 InStateIndex, const float InStateWeight);
+	void RecordStateWeight(const int32 InMachineClassIndex, const int32 InStateIndex, const float InStateWeight, const float InElapsedTime);
 
 protected:
 #if WITH_EDITORONLY_DATA
diff --git a/Engine/Source/Runtime/Engine/Classes/Animation/AnimNodeBase.h b/Engine/Source/Runtime/Engine/Classes/Animation/AnimNodeBase.h
--- a/Engine/Source/Runtime/Engine/Classes/Animation/AnimNodeBase.h
+++ b/Engine/Source/Runtime/Engine/Classes/Animation/AnimNodeBase.h
@@ -19,6 +19,8 @@
 #define DECLARE_SCOPE_HIERARCHICAL_COUNTER_ANIMNODE(Method) \
 	DECLARE_SCOPE_HIERARCHICAL_COUNTER_FUNC()
 
+#define ANIM_NODE_IDS_AVAILABLE	(ANIM_TRACE_ENABLED || WITH_EDITORONLY_DATA)
+
 class IAnimClassInterface;
 class UAnimBlueprint;
 class UAnimInstance;
@@ -154,7 +156,7 @@ public:
 	ENGINE_API UAnimBlueprint* GetAnimBlueprint() const;
 #endif //WITH_EDITORONLY_DATA
 
-#if ANIM_TRACE_ENABLED
+#if ANIM_NODE_IDS_AVAILABLE
 	// Get the current node Id, set when we recurse into graph traversal functions from pose links
 	ENGINE_API int32 GetCurrentNodeId() const { return CurrentNodeId; }
 
@@ -237,6 +239,10 @@ public:
 		, RootMotionWeightModifier(Copy.RootMotionWeightModifier)
 		, DeltaTime(Copy.DeltaTime)
 	{
+#if ANIM_TRACE_ENABLED
+		CurrentNodeId = Copy.CurrentNodeId;
+		PreviousNodeId = Copy.PreviousNodeId;
+#endif
 	}
 
 public:
@@ -288,7 +294,7 @@ public:
 		return Result;
 	}
 
-#if ANIM_TRACE_ENABLED
+#if ANIM_NODE_IDS_AVAILABLE
 	FAnimationUpdateContext WithNodeId(int32 InNodeId) const
 	{ 
 		FAnimationUpdateContext Result(*this);
@@ -369,6 +375,14 @@ public:
 		Initialize(SourceContext.AnimInstanceProxy);
 	}
 
+#if ANIM_NODE_IDS_AVAILABLE
+	void SetNodeId(int32 InNodeId)
+	{ 
+		PreviousNodeId = CurrentNodeId;
+		CurrentNodeId = InNodeId;
+	}
+#endif
+
 	ENGINE_API void Initialize(FAnimInstanceProxy* InAnimInstanceProxy);
 
 	// Log evaluation message
@@ -446,6 +460,14 @@ public:
 		// No need to initialize, done through FA2CSPose::AllocateLocalPoses
 	}
 
+#if ANIM_NODE_IDS_AVAILABLE
+	void SetNodeId(int32 InNodeId)
+	{ 
+		PreviousNodeId = CurrentNodeId;
+		CurrentNodeId = InNodeId;
+	}
+#endif
+
 	ENGINE_API void ResetToRefPose();
 
 	ENGINE_API bool ContainsNaN() const;
diff --git a/Engine/Source/Runtime/Engine/Classes/Animation/AnimNode_AssetPlayerBase.h b/Engine/Source/Runtime/Engine/Classes/Animation/AnimNode_AssetPlayerBase.h
--- a/Engine/Source/Runtime/Engine/Classes/Animation/AnimNode_AssetPlayerBase.h
+++ b/Engine/Source/Runtime/Engine/Classes/Animation/AnimNode_AssetPlayerBase.h
@@ -17,13 +17,13 @@ struct ENGINE_API FAnimNode_AssetPlayerBase : public FAnimNode_Base
 	FAnimNode_AssetPlayerBase();
 
 	/** Get the last encountered blend weight for this node */
-	virtual float GetCachedBlendWeight();
+	virtual float GetCachedBlendWeight() const;
 	
 	/** Set the cached blendweight to zero */
 	void ClearCachedBlendWeight();
 
 	/** Get the currently referenced time within the asset player node */
-	virtual float GetAccumulatedTime();
+	virtual float GetAccumulatedTime() const;
 
 	/** Override the currently accumulated time */
 	virtual void SetAccumulatedTime(const float& NewTime);
diff --git a/Engine/Source/Runtime/Engine/Classes/Animation/AnimNode_LinkedAnimGraph.h b/Engine/Source/Runtime/Engine/Classes/Animation/AnimNode_LinkedAnimGraph.h
--- a/Engine/Source/Runtime/Engine/Classes/Animation/AnimNode_LinkedAnimGraph.h
+++ b/Engine/Source/Runtime/Engine/Classes/Animation/AnimNode_LinkedAnimGraph.h
@@ -48,6 +48,9 @@ public:
 	// The root node of the dynamically-linked graph
 	FAnimNode_Base* LinkedRoot;
 
+	// Our node index
+	int32 NodeIndex;
+
 	/** Dynamically set the anim class of this linked instance */
 	void SetAnimClass(TSubclassOf<UAnimInstance> InClass, const UAnimInstance* InOwningAnimInstance);
 
diff --git a/Engine/Source/Runtime/Engine/Classes/Animation/AnimNode_LinkedInputPose.h b/Engine/Source/Runtime/Engine/Classes/Animation/AnimNode_LinkedInputPose.h
--- a/Engine/Source/Runtime/Engine/Classes/Animation/AnimNode_LinkedInputPose.h
+++ b/Engine/Source/Runtime/Engine/Classes/Animation/AnimNode_LinkedInputPose.h
@@ -20,6 +20,7 @@ struct ENGINE_API FAnimNode_LinkedInputPose : public FAnimNode_Base
 	FAnimNode_LinkedInputPose()
 		: Name(DefaultInputPoseName)
 		, Graph(NAME_None)
+		, OuterGraphNodeIndex(INDEX_NONE)
 		, InputProxy(nullptr)
 	{
 	}
@@ -43,6 +44,9 @@ struct ENGINE_API FAnimNode_LinkedInputPose : public FAnimNode_Base
 	FCompactHeapPose CachedInputPose;
 	FBlendedHeapCurve CachedInputCurve;
 
+	// The node index of the currently-linked outer node
+	int32 OuterGraphNodeIndex;
+
 	// FAnimNode_Base interface
 #if ENABLE_ANIMGRAPH_TRAVERSAL_DEBUG
 	virtual void Initialize_AnyThread(const FAnimationInitializeContext& Context) override;
@@ -54,7 +58,7 @@ struct ENGINE_API FAnimNode_LinkedInputPose : public FAnimNode_Base
 	// End of FAnimNode_Base interface
 
 	/** Called by linked instance nodes to dynamically link this to an outer graph */
-	void DynamicLink(FAnimInstanceProxy* InInputProxy, FPoseLinkBase* InPoseLink);
+	void DynamicLink(FAnimInstanceProxy* InInputProxy, FPoseLinkBase* InPoseLink, int32 InOuterGraphNodeIndex);
 
 	/** Called by linked instance nodes to dynamically unlink this to an outer graph */
 	void DynamicUnlink();
diff --git a/Engine/Source/Runtime/Engine/Private/Animation/AnimBlueprintGeneratedClass.cpp b/Engine/Source/Runtime/Engine/Private/Animation/AnimBlueprintGeneratedClass.cpp
--- a/Engine/Source/Runtime/Engine/Private/Animation/AnimBlueprintGeneratedClass.cpp
+++ b/Engine/Source/Runtime/Engine/Private/Animation/AnimBlueprintGeneratedClass.cpp
@@ -135,11 +135,15 @@ void FAnimBlueprintDebugData::SetSnapshotIndexByTime(UAnimInstance* Instance, do
 void FAnimBlueprintDebugData::ResetNodeVisitSites()
 {
 	UpdatedNodesThisFrame.Empty(UpdatedNodesThisFrame.Num());
+	StateData.Empty(StateData.Num());
+	NodeValuesThisFrame.Empty(NodeValuesThisFrame.Num());
+	SequencePlayerRecordsThisFrame.Empty(SequencePlayerRecordsThisFrame.Num());
+	BlendSpacePlayerRecordsThisFrame.Empty(BlendSpacePlayerRecordsThisFrame.Num());
 }
 
 void FAnimBlueprintDebugData::RecordNodeVisit(int32 TargetNodeIndex, int32 SourceNodeIndex, float BlendWeight)
 {
-	new (UpdatedNodesThisFrame) FNodeVisit(SourceNodeIndex, TargetNodeIndex, BlendWeight);
+	UpdatedNodesThisFrame.Emplace(SourceNodeIndex, TargetNodeIndex, BlendWeight);
 }
 
 void FAnimBlueprintDebugData::RecordNodeVisitArray(const TArray<FNodeVisit>& Nodes)
@@ -147,6 +151,26 @@ void FAnimBlueprintDebugData::RecordNodeVisitArray(const TArray<FNodeVisit>& Nod
 	UpdatedNodesThisFrame.Append(Nodes);
 }
 
+void FAnimBlueprintDebugData::RecordStateData(int32 StateMachineIndex, int32 StateIndex, float Weight, float ElapsedTime)
+{
+	StateData.Emplace(StateMachineIndex, StateIndex, Weight, ElapsedTime);
+}
+
+void FAnimBlueprintDebugData::RecordNodeValue(int32 InNodeID, const FString& InText)
+{
+	NodeValuesThisFrame.Emplace(InText, InNodeID);
+}
+
+void FAnimBlueprintDebugData::RecordSequencePlayer(int32 InNodeID, float InPosition, float InLength, int32 InFrameCount)
+{
+	SequencePlayerRecordsThisFrame.Emplace(InNodeID, InPosition, InLength, InFrameCount);
+}
+
+void FAnimBlueprintDebugData::RecordBlendSpacePlayer(int32 InNodeID, UBlendSpaceBase* InBlendSpace, float InPositionX, float InPositionY, float InPositionZ)
+{
+	BlendSpacePlayerRecordsThisFrame.Emplace(InNodeID, InBlendSpace, InPositionX, InPositionY, InPositionZ);
+}
+
 void FAnimBlueprintDebugData::AddPoseWatch(int32 NodeID, FColor Color)
 {
 	for (FAnimNodePoseWatch& PoseWatch : AnimNodePoseWatch)
diff --git a/Engine/Source/Runtime/Engine/Private/Animation/AnimInstance.cpp b/Engine/Source/Runtime/Engine/Private/Animation/AnimInstance.cpp
--- a/Engine/Source/Runtime/Engine/Private/Animation/AnimInstance.cpp
+++ b/Engine/Source/Runtime/Engine/Private/Animation/AnimInstance.cpp
@@ -172,6 +172,8 @@ UWorld* UAnimInstance::GetWorld() const
 
 void UAnimInstance::InitializeAnimation(bool bInDeferRootNodeInitialization)
 {
+	TRACE_OBJECT_EVENT(this, InitializeAnimation);
+
 	FScopeCycleCounterUObject ContextScope(this);
 	SCOPE_CYCLE_COUNTER(STAT_AnimInitTime);
 	LLM_SCOPE(ELLMTag::Animation);
@@ -225,6 +227,8 @@ void UAnimInstance::InitializeAnimation(bool bInDeferRootNodeInitialization)
 
 void UAnimInstance::UninitializeAnimation()
 {
+	TRACE_OBJECT_EVENT(this, UninitializeAnimation);
+
 	NativeUninitializeAnimation();
 
 	GetProxyOnGameThread<FAnimInstanceProxy>().Uninitialize(this);
@@ -329,6 +333,11 @@ void UAnimInstance::UpdateMontageSyncGroup()
 					// the max count should be 2 as you had older one and you have newer one. After TestMontageTickRecordForLeadership, it should set to be 1
 					SyncGroup->TestMontageTickRecordForLeadership();
 				}
+
+#if ANIM_TRACE_ENABLED
+				FAnimationUpdateContext UpdateContext(&GetProxyOnGameThread<FAnimInstanceProxy>());
+				TRACE_ANIM_TICK_RECORD(UpdateContext, TickRecord);
+#endif
 			}
 			MontageInstance->bDidUseMarkerSyncThisTick = false;
 		}
@@ -3182,9 +3191,9 @@ void UAnimInstance::RecordMachineWeight(const int32 InMachineClassIndex, const f
 	GetProxyOnAnyThread<FAnimInstanceProxy>().RecordMachineWeight(InMachineClassIndex, InMachineWeight);
 }
 
-void UAnimInstance::RecordStateWeight(const int32 InMachineClassIndex, const int32 InStateIndex, const float InStateWeight)
+void UAnimInstance::RecordStateWeight(const int32 InMachineClassIndex, const int32 InStateIndex, const float InStateWeight, const float InElapsedTime)
 {
-	GetProxyOnAnyThread<FAnimInstanceProxy>().RecordStateWeight(InMachineClassIndex, InStateIndex, InStateWeight);
+	GetProxyOnAnyThread<FAnimInstanceProxy>().RecordStateWeight(InMachineClassIndex, InStateIndex, InStateWeight, InElapsedTime);
 }
 
 const FGraphTraversalCounter& UAnimInstance::GetUpdateCounter() const
diff --git a/Engine/Source/Runtime/Engine/Private/Animation/AnimInstanceProxy.cpp b/Engine/Source/Runtime/Engine/Private/Animation/AnimInstanceProxy.cpp
--- a/Engine/Source/Runtime/Engine/Private/Animation/AnimInstanceProxy.cpp
+++ b/Engine/Source/Runtime/Engine/Private/Animation/AnimInstanceProxy.cpp
@@ -24,7 +24,8 @@
 #include "Animation/AnimMTStats.h"
 #include "Animation/AnimNode_Root.h"
 #include "Animation/AnimNode_LinkedAnimLayer.h"
-#include "Animation/AnimTrace.h"
+#include "Animation/AnimMontage.h"
+
 #undef DO_ANIMSTAT_PROCESSING
 
 #define LOCTEXT_NAMESPACE "AnimInstance"
@@ -36,6 +37,8 @@ const FName NAME_AnimGraph(TEXT("AnimGraph"));
 
 void FAnimInstanceProxy::UpdateAnimationNode(const FAnimationUpdateContext& InContext)
 {
+	TRACE_SCOPED_ANIM_GRAPH(InContext);
+
 	UpdateAnimationNode_WithRoot(InContext, RootNode, NAME_AnimGraph);
 }
 
@@ -285,12 +288,20 @@ void FAnimInstanceProxy::InitializeRootNode_WithRoot(FAnimNode_Base* InRootNode)
 
 	if (InRootNode != nullptr)
 	{
+		FAnimationInitializeContext InitContext(this);
+
 		if(InRootNode == RootNode)
 		{
 			InitializationCounter.Increment();
+
+			TRACE_SCOPED_ANIM_GRAPH(InitContext);
+
+			InRootNode->Initialize_AnyThread(InitContext);
+		}
+		else
+		{
+			InRootNode->Initialize_AnyThread(InitContext);
 		}
-		FAnimationInitializeContext InitContext(this);
-		InRootNode->Initialize_AnyThread(InitContext);
 	}
 }
 
@@ -523,8 +534,6 @@ void FAnimInstanceProxy::PostEvaluate(UAnimInstance* InAnimInstance)
 {
 	DECLARE_SCOPE_HIERARCHICAL_COUNTER_FUNC()
 
-	TRACE_ANIM_TICK_RECORDS(*this, SkeletalMeshComponent);
-
 	ClearObjects();
 
 #if ENABLE_ANIM_LOGGING
@@ -1186,6 +1195,12 @@ void FAnimInstanceProxy::UpdateAnimation()
 
 	FAnimationUpdateSharedContext SharedContext;
 	FAnimationUpdateContext Context(this, CurrentDeltaSeconds, &SharedContext);
+#if ANIM_NODE_IDS_AVAILABLE
+	if(AnimClassInterface && AnimClassInterface->GetAnimBlueprintFunctions().Num() > 0)
+	{
+		Context = Context.WithNodeId(AnimClassInterface->GetAnimBlueprintFunctions()[0].OutputPoseNodeIndex);
+	}
+#endif
 	UpdateAnimation_WithRoot(Context, RootNode, NAME_AnimGraph);
 }
 
@@ -1276,6 +1291,8 @@ void FAnimInstanceProxy::PreEvaluateAnimation(UAnimInstance* InAnimInstance)
 
 void FAnimInstanceProxy::EvaluateAnimation(FPoseContext& Output)
 {
+	TRACE_SCOPED_ANIM_GRAPH(Output);
+
 	EvaluateAnimation_WithRoot(Output, RootNode);
 }
 
@@ -1313,6 +1330,9 @@ void FAnimInstanceProxy::CacheBones()
 
 		CachedBonesCounter.Increment();
 		FAnimationCacheBonesContext Context(this);
+
+		TRACE_SCOPED_ANIM_GRAPH(Context);
+
 		RootNode->CacheBones_AnyThread(Context);
 
 		CacheBonesRecursionCounter--;
@@ -2351,7 +2371,7 @@ float FAnimInstanceProxy::GetRecordedStateWeight(const int32 InMachineClassIndex
 	return 0.0f;
 }
 
-void FAnimInstanceProxy::RecordStateWeight(const int32 InMachineClassIndex, const int32 InStateIndex, const float InStateWeight)
+void FAnimInstanceProxy::RecordStateWeight(const int32 InMachineClassIndex, const int32 InStateIndex, const float InStateWeight, const float InElapsedTime)
 {
 	const int32* BaseIndexPtr = StateMachineClassIndexToWeightOffset.Find(InMachineClassIndex);
 
@@ -2360,6 +2380,13 @@ void FAnimInstanceProxy::RecordStateWeight(const int32 InMachineClassIndex, cons
 		const int32 StateIndex = *BaseIndexPtr + InStateIndex;
 		StateWeightArrays[GetSyncGroupWriteIndex()][StateIndex] = InStateWeight;
 	}
+
+#if WITH_EDITORONLY_DATA
+	if (FAnimBlueprintDebugData* DebugData = GetAnimBlueprintDebugData())
+	{
+		DebugData->RecordStateData(InMachineClassIndex, InStateIndex, InStateWeight, InElapsedTime);
+	}
+#endif
 }
 
 void FAnimInstanceProxy::ResetDynamics(ETeleportType InTeleportType)
@@ -2376,6 +2403,32 @@ void FAnimInstanceProxy::ResetDynamics()
 	ResetDynamics(ETeleportType::ResetPhysics);
 }
 
+#if ANIM_TRACE_ENABLED
+void FAnimInstanceProxy::TraceMontageEvaluationData(const FAnimationUpdateContext& InContext, const FName& InSlotName)
+{
+	for (const FMontageEvaluationState& MontageEvaluationState : MontageEvaluationData)
+	{
+		if (MontageEvaluationState.Montage != nullptr && MontageEvaluationState.Montage->IsValidSlot(InSlotName))
+		{
+			if (const FAnimTrack* const Track = MontageEvaluationState.Montage->GetAnimationData(InSlotName))
+			{
+				if (const FAnimSegment* const Segment = Track->GetSegmentAtTime(MontageEvaluationState.MontagePosition))
+				{
+					float CurrentAnimPos;
+					if (UAnimSequenceBase* Anim = Segment->GetAnimationData(MontageEvaluationState.MontagePosition, CurrentAnimPos))
+					{
+						TRACE_ANIM_NODE_VALUE(InContext, TEXT("Montage"), MontageEvaluationState.Montage.Get());
+						TRACE_ANIM_NODE_VALUE(InContext, TEXT("Sequence"), Anim);
+						TRACE_ANIM_NODE_VALUE(InContext, TEXT("Sequence Playback Time"), CurrentAnimPos);
+						break;
+					}
+				}
+			}
+		}
+	}
+}
+#endif
+
 TArray<FAnimNode_AssetPlayerBase*> FAnimInstanceProxy::GetInstanceAssetPlayers(const FName& GraphName)
 {
 	TArray<FAnimNode_AssetPlayerBase*> Nodes;
diff --git a/Engine/Source/Runtime/Engine/Private/Animation/AnimNodeBase.cpp b/Engine/Source/Runtime/Engine/Private/Animation/AnimNodeBase.cpp
--- a/Engine/Source/Runtime/Engine/Private/Animation/AnimNodeBase.cpp
+++ b/Engine/Source/Runtime/Engine/Private/Animation/AnimNodeBase.cpp
@@ -77,8 +77,8 @@ FAnimationBaseContext::FAnimationBaseContext(FAnimInstanceProxy* InAnimInstanceP
 FAnimationBaseContext::FAnimationBaseContext(const FAnimationBaseContext& InContext)
 	: AnimInstanceProxy(InContext.AnimInstanceProxy)
 #if ANIM_TRACE_ENABLED
-	, CurrentNodeId(INDEX_NONE)
-	, PreviousNodeId(INDEX_NONE)
+	, CurrentNodeId(InContext.CurrentNodeId)
+	, PreviousNodeId(InContext.PreviousNodeId)
 #endif
 {
 }
@@ -343,9 +343,11 @@ void FPoseLinkBase::Update(const FAnimationUpdateContext& Context)
 	if (LinkedNode != NULL)
 	{
 #if ANIM_TRACE_ENABLED
-		FAnimationUpdateContext LinkContext(Context.WithNodeId(LinkID));
-		TRACE_POSE_LINK(LinkContext);
-		LinkedNode->Update_AnyThread(LinkContext);
+		{
+			FAnimationUpdateContext LinkContext(Context.WithNodeId(LinkID));
+			TRACE_SCOPED_ANIM_NODE(LinkContext);
+			LinkedNode->Update_AnyThread(LinkContext);
+		}
 #else
 		LinkedNode->Update_AnyThread(Context);
 #endif
@@ -391,7 +393,15 @@ void FPoseLink::Evaluate(FPoseContext& Output)
 #if ENABLE_ANIMNODE_POSE_DEBUG
 		CurrentPose.ResetToAdditiveIdentity();
 #endif
-		LinkedNode->Evaluate_AnyThread(Output);
+
+		{
+#if ANIM_TRACE_ENABLED
+			Output.SetNodeId(LinkID);
+			TRACE_SCOPED_ANIM_NODE(Output);
+#endif
+			LinkedNode->Evaluate_AnyThread(Output);
+		}
+
 #if ENABLE_ANIMNODE_POSE_DEBUG
 		CurrentPose.CopyBonesFrom(Output.Pose);
 #endif
@@ -456,6 +466,9 @@ void FComponentSpacePoseLink::EvaluateComponentSpace(FComponentSpacePoseContext&
 
 	if (LinkedNode != NULL)
 	{
+#if ANIM_TRACE_ENABLED
+		Output.SetNodeId(LinkID);
+#endif
 		LinkedNode->EvaluateComponentSpace_AnyThread(Output);
 
 #if WITH_EDITOR
diff --git a/Engine/Source/Runtime/Engine/Private/Animation/AnimNode_ApplyMeshSpaceAdditive.cpp b/Engine/Source/Runtime/Engine/Private/Animation/AnimNode_ApplyMeshSpaceAdditive.cpp
--- a/Engine/Source/Runtime/Engine/Private/Animation/AnimNode_ApplyMeshSpaceAdditive.cpp
+++ b/Engine/Source/Runtime/Engine/Private/Animation/AnimNode_ApplyMeshSpaceAdditive.cpp
@@ -2,6 +2,7 @@
 
 #include "Animation/AnimNode_ApplyMeshSpaceAdditive.h"
 #include "AnimationRuntime.h"
+#include "Animation/AnimTrace.h"
 
 /////////////////////////////////////////////////////
 // FAnimNode_ApplyMeshSpaceAdditive
@@ -36,6 +37,8 @@ void FAnimNode_ApplyMeshSpaceAdditive::Update_AnyThread(const FAnimationUpdateCo
 			Additive.Update(Context.FractionalWeight(ActualAlpha));
 		}
 	}
+
+	TRACE_ANIM_NODE_VALUE(Context, TEXT("Alpha"), ActualAlpha);
 }
 
 void FAnimNode_ApplyMeshSpaceAdditive::Evaluate_AnyThread(FPoseContext& Output)
diff --git a/Engine/Source/Runtime/Engine/Private/Animation/AnimNode_AssetPlayerBase.cpp b/Engine/Source/Runtime/Engine/Private/Animation/AnimNode_AssetPlayerBase.cpp
--- a/Engine/Source/Runtime/Engine/Private/Animation/AnimNode_AssetPlayerBase.cpp
+++ b/Engine/Source/Runtime/Engine/Private/Animation/AnimNode_AssetPlayerBase.cpp
@@ -49,14 +49,16 @@ void FAnimNode_AssetPlayerBase::CreateTickRecordForNode(const FAnimationUpdateCo
 	{
 		SyncGroup->TestTickRecordForLeadership(GroupRole);
 	}
+
+	TRACE_ANIM_TICK_RECORD(Context, TickRecord);
 }
 
-float FAnimNode_AssetPlayerBase::GetCachedBlendWeight()
+float FAnimNode_AssetPlayerBase::GetCachedBlendWeight() const
 {
 	return BlendWeight;
 }
 
-float FAnimNode_AssetPlayerBase::GetAccumulatedTime()
+float FAnimNode_AssetPlayerBase::GetAccumulatedTime() const
 {
 	return InternalTimeAccumulator;
 }
diff --git a/Engine/Source/Runtime/Engine/Private/Animation/AnimNode_Inertialization.cpp b/Engine/Source/Runtime/Engine/Private/Animation/AnimNode_Inertialization.cpp
--- a/Engine/Source/Runtime/Engine/Private/Animation/AnimNode_Inertialization.cpp
+++ b/Engine/Source/Runtime/Engine/Private/Animation/AnimNode_Inertialization.cpp
@@ -236,6 +236,11 @@ void FAnimNode_Inertialization::Evaluate_AnyThread(FPoseContext& Output)
 	// Reset the time accumulator and teleport state
 	DeltaTime = 0.0f;
 	TeleportType = ETeleportType::None;
+
+	TRACE_ANIM_NODE_VALUE(Output, TEXT("State"), *UEnum::GetValueAsString(InertializationState));
+	TRACE_ANIM_NODE_VALUE(Output, TEXT("Elapsed Time"), InertializationElapsedTime);
+	TRACE_ANIM_NODE_VALUE(Output, TEXT("Duration"), InertializationDuration);
+	TRACE_ANIM_NODE_VALUE(Output, TEXT("Normalized Time"), InertializationDuration > KINDA_SMALL_NUMBER ? (InertializationElapsedTime / InertializationDuration) : 0.0f);
 }
 
 
diff --git a/Engine/Source/Runtime/Engine/Private/Animation/AnimNode_LinkedAnimGraph.cpp b/Engine/Source/Runtime/Engine/Private/Animation/AnimNode_LinkedAnimGraph.cpp
--- a/Engine/Source/Runtime/Engine/Private/Animation/AnimNode_LinkedAnimGraph.cpp
+++ b/Engine/Source/Runtime/Engine/Private/Animation/AnimNode_LinkedAnimGraph.cpp
@@ -6,6 +6,7 @@
 #include "Animation/AnimNode_Inertialization.h"
 #include "Animation/AnimNode_LinkedInputPose.h"
 #include "Animation/AnimNode_Root.h"
+#include "Animation/AnimTrace.h"
 
 static float GetBlendDuration(const IAnimClassInterface* PriorAnimBPClass, const IAnimClassInterface* NewAnimBPClass, FName Layer)
 {
@@ -32,6 +33,7 @@ FAnimNode_LinkedAnimGraph::FAnimNode_LinkedAnimGraph()
 	: InstanceClass(nullptr)
 	, Tag(NAME_None)
 	, LinkedRoot(nullptr)
+	, NodeIndex(INDEX_NONE)
 	, PendingBlendDuration(-1.0f)
 {
 }
@@ -133,6 +135,8 @@ void FAnimNode_LinkedAnimGraph::Update_AnyThread(const FAnimationUpdateContext&
 
 		PendingBlendDuration = -1.0f;
 	}
+
+	TRACE_ANIM_NODE_VALUE(InContext, TEXT("Target Class"), InstanceClass.Get());
 }
 
 void FAnimNode_LinkedAnimGraph::Evaluate_AnyThread(FPoseContext& Output)
@@ -331,7 +335,7 @@ void FAnimNode_LinkedAnimGraph::DynamicLink(UAnimInstance* InOwningAnimInstance)
 						{
 							FAnimNode_LinkedInputPose* LinkedInputPoseNode = AnimBlueprintFunction.InputPoseNodeProperties[InputPropertyIndex]->ContainerPtrToValuePtr<FAnimNode_LinkedInputPose>(LinkTargetInstance);
 							check(LinkedInputPoseNode->Name == AnimBlueprintFunction.InputPoseNames[InputPoseIndex]);
-							LinkedInputPoseNode->DynamicLink(NonConstProxy, &InputPoses[InputPoseIndex]);
+							LinkedInputPoseNode->DynamicLink(NonConstProxy, &InputPoses[InputPoseIndex], NodeIndex);
 						}
 						else
 						{
diff --git a/Engine/Source/Runtime/Engine/Private/Animation/AnimNode_LinkedInputPose.cpp b/Engine/Source/Runtime/Engine/Private/Animation/AnimNode_LinkedInputPose.cpp
--- a/Engine/Source/Runtime/Engine/Private/Animation/AnimNode_LinkedInputPose.cpp
+++ b/Engine/Source/Runtime/Engine/Private/Animation/AnimNode_LinkedInputPose.cpp
@@ -29,6 +29,9 @@ void FAnimNode_LinkedInputPose::Update_AnyThread(const FAnimationUpdateContext&
 	if(InputProxy)
 	{
 		FAnimationUpdateContext InputContext = Context.WithOtherProxy(InputProxy);
+#if ANIM_NODE_IDS_AVAILABLE
+		InputContext = InputContext.WithNodeId(OuterGraphNodeIndex);
+#endif
 		InputPose.Update(InputContext);
 	}
 }
@@ -67,12 +70,13 @@ void FAnimNode_LinkedInputPose::GatherDebugData(FNodeDebugData& DebugData)
 	}
 }
 
-void FAnimNode_LinkedInputPose::DynamicLink(FAnimInstanceProxy* InInputProxy, FPoseLinkBase* InPoseLink)
+void FAnimNode_LinkedInputPose::DynamicLink(FAnimInstanceProxy* InInputProxy, FPoseLinkBase* InPoseLink, int32 InOuterGraphNodeIndex)
 {
 	check(InputProxy == nullptr);			// Must be unlinked before re-linking
 
 	InputProxy = InInputProxy;
 	InputPose.SetDynamicLinkNode(InPoseLink);
+	OuterGraphNodeIndex = InOuterGraphNodeIndex;
 }
 
 void FAnimNode_LinkedInputPose::DynamicUnlink()
@@ -81,4 +85,5 @@ void FAnimNode_LinkedInputPose::DynamicUnlink()
 
 	InputProxy = nullptr;
 	InputPose.SetDynamicLinkNode(nullptr);
+	OuterGraphNodeIndex = INDEX_NONE;
 }
\ No newline at end of file
diff --git a/Engine/Source/Runtime/Engine/Private/Animation/AnimNode_SaveCachedPose.cpp b/Engine/Source/Runtime/Engine/Private/Animation/AnimNode_SaveCachedPose.cpp
--- a/Engine/Source/Runtime/Engine/Private/Animation/AnimNode_SaveCachedPose.cpp
+++ b/Engine/Source/Runtime/Engine/Private/Animation/AnimNode_SaveCachedPose.cpp
@@ -2,6 +2,7 @@
 
 #include "Animation/AnimNode_SaveCachedPose.h"
 #include "Animation/AnimInstanceProxy.h"
+#include "Animation/AnimTrace.h"
 
 /////////////////////////////////////////////////////
 // FAnimNode_SaveCachedPose
@@ -52,6 +53,8 @@ void FAnimNode_SaveCachedPose::Update_AnyThread(const FAnimationUpdateContext& C
 
 	// Store this context for the post update
 	CachedUpdate.Context = Context.WithOtherSharedContext(CachedUpdate.SharedContext.Get());
+
+	TRACE_ANIM_NODE_VALUE(Context, TEXT("Cached Pose Name"), CachePoseName);
 }
 
 void FAnimNode_SaveCachedPose::Evaluate_AnyThread(FPoseContext& Output)
diff --git a/Engine/Source/Runtime/Engine/Private/Animation/AnimNode_SequencePlayer.cpp b/Engine/Source/Runtime/Engine/Private/Animation/AnimNode_SequencePlayer.cpp
--- a/Engine/Source/Runtime/Engine/Private/Animation/AnimNode_SequencePlayer.cpp
+++ b/Engine/Source/Runtime/Engine/Private/Animation/AnimNode_SequencePlayer.cpp
@@ -3,6 +3,7 @@
 #include "Animation/AnimNode_SequencePlayer.h"
 #include "Animation/AnimInstanceProxy.h"
 #include "AnimEncoding.h"
+#include "Animation/AnimTrace.h"
 
 #define LOCTEXT_NAMESPACE "AnimNode_SequencePlayer"
 
@@ -72,6 +73,17 @@ void FAnimNode_SequencePlayer::UpdateAssetPlayer(const FAnimationUpdateContext&
 		const float AdjustedPlayRate = PlayRateScaleBiasClamp.ApplyTo(FMath::IsNearlyZero(PlayRateBasis) ? 0.f : (PlayRate / PlayRateBasis), Context.GetDeltaTime());
 		CreateTickRecordForNode(Context, Sequence, bLoopAnimation, AdjustedPlayRate);
 	}
+
+#if ANIM_NODE_IDS_AVAILABLE && WITH_EDITORONLY_DATA
+	if (FAnimBlueprintDebugData* DebugData = Context.AnimInstanceProxy->GetAnimBlueprintDebugData())
+	{
+		DebugData->RecordSequencePlayer(Context.GetCurrentNodeId(), GetAccumulatedTime(), Sequence != nullptr ? Sequence->SequenceLength : 0.0f, Sequence != nullptr ? Sequence->GetNumberOfFrames() : 0);
+	}
+#endif
+
+	TRACE_ANIM_SEQUENCE_PLAYER(Context, *this);
+	TRACE_ANIM_NODE_VALUE(Context, TEXT("Sequence"), Sequence);
+	TRACE_ANIM_NODE_VALUE(Context, TEXT("Playback Time"), InternalTimeAccumulator);
 }
 
 void FAnimNode_SequencePlayer::Evaluate_AnyThread(FPoseContext& Output)
diff --git a/Engine/Source/Runtime/Engine/Private/Animation/AnimNode_StateMachine.cpp b/Engine/Source/Runtime/Engine/Private/Animation/AnimNode_StateMachine.cpp
--- a/Engine/Source/Runtime/Engine/Private/Animation/AnimNode_StateMachine.cpp
+++ b/Engine/Source/Runtime/Engine/Private/Animation/AnimNode_StateMachine.cpp
@@ -9,6 +9,12 @@
 #include "Animation/BlendProfile.h"
 #include "Animation/AnimNode_LinkedAnimLayer.h"
 #include "Animation/AnimInstance.h"
+#include "Animation/AnimTrace.h"
+
+#if WITH_EDITORONLY_DATA
+#include "Animation/AnimBlueprintGeneratedClass.h"
+#include "Animation/AnimBlueprint.h"
+#endif
 
 #define LOCTEXT_NAMESPACE "AnimNode_StateMachine"
 
@@ -544,10 +550,25 @@ void FAnimNode_StateMachine::Update_AnyThread(const FAnimationUpdateContext& Con
 	if (ActiveTransitionArray.Num() == 0 && !IsAConduitState(CurrentState) && !StatesUpdated.Contains(CurrentState))
 	{
 		StatePoseLinks[CurrentState].Update(Context);
-		Context.AnimInstanceProxy->RecordStateWeight(StateMachineIndexInClass, CurrentState, GetStateWeight(CurrentState));
+
+		const float StateWeight = GetStateWeight(CurrentState);
+		const float CurrentStateElapsedTime = GetCurrentStateElapsedTime();
+		Context.AnimInstanceProxy->RecordStateWeight(StateMachineIndexInClass, CurrentState, StateWeight, CurrentStateElapsedTime);
+
+		TRACE_ANIM_STATE_MACHINE_STATE(Context, StateMachineIndexInClass, CurrentState, StateWeight, CurrentStateElapsedTime);
 	}
 
 	ElapsedTime += Context.GetDeltaTime();
+
+#if ANIM_TRACE_ENABLED
+	TRACE_ANIM_NODE_VALUE(Context, TEXT("State Machine Name"), GetMachineDescription()->MachineName);
+	TRACE_ANIM_NODE_VALUE(Context, TEXT("Current State"), GetStateInfo().StateName);
+	for (int32 PoseIndex = 0; PoseIndex < StatePoseLinks.Num(); ++PoseIndex)
+	{
+		const FString StateName = FString::Printf(TEXT("State %s Weight"), *GetStateInfo(PoseIndex).StateName.ToString());
+		TRACE_ANIM_NODE_VALUE(Context, *StateName, GetStateWeight(PoseIndex));
+	}
+#endif
 }
 
 FAnimNode_AssetPlayerBase* FAnimNode_StateMachine::GetRelevantAssetPlayerFromState(const FAnimationUpdateContext& Context, const FBakedAnimationState& StateInfo)
@@ -1057,7 +1078,11 @@ void FAnimNode_StateMachine::UpdateState(int32 StateIndex, const FAnimationUpdat
 		StatesUpdated.Add(StateIndex);
 		StatePoseLinks[StateIndex].Update(Context);
 
-		Context.AnimInstanceProxy->RecordStateWeight(StateMachineIndexInClass, StateIndex, GetStateWeight(StateIndex));
+		const float StateWeight = GetStateWeight(StateIndex);
+		const float CurrentStateElapsedTime = StateIndex == CurrentState ? GetCurrentStateElapsedTime() : 0.0f;
+		Context.AnimInstanceProxy->RecordStateWeight(StateMachineIndexInClass, StateIndex, StateWeight, CurrentStateElapsedTime);
+
+		TRACE_ANIM_STATE_MACHINE_STATE(Context, StateMachineIndexInClass, StateIndex, StateWeight, CurrentStateElapsedTime);
 	}
 }
 
diff --git a/Engine/Source/Runtime/Engine/Private/Animation/AnimNode_TransitionPoseEvaluator.cpp b/Engine/Source/Runtime/Engine/Private/Animation/AnimNode_TransitionPoseEvaluator.cpp
--- a/Engine/Source/Runtime/Engine/Private/Animation/AnimNode_TransitionPoseEvaluator.cpp
+++ b/Engine/Source/Runtime/Engine/Private/Animation/AnimNode_TransitionPoseEvaluator.cpp
@@ -2,6 +2,7 @@
 
 #include "Animation/AnimNode_TransitionPoseEvaluator.h"
 #include "Animation/AnimInstanceProxy.h"
+#include "Animation/AnimTrace.h"
 
 /////////////////////////////////////////////////////
 // FAnimNode_TransitionPoseEvaluator
@@ -40,6 +41,8 @@ void FAnimNode_TransitionPoseEvaluator::CacheBones_AnyThread(const FAnimationCac
 void FAnimNode_TransitionPoseEvaluator::Update_AnyThread(const FAnimationUpdateContext& Context)
 {
 	// updating is all handled in state machine
+
+	TRACE_ANIM_NODE_VALUE(Context, TEXT("Cached Frames Remaining"), CacheFramesRemaining);
 }
 
 void FAnimNode_TransitionPoseEvaluator::Evaluate_AnyThread(FPoseContext& Output)
diff --git a/Engine/Source/Runtime/Engine/Private/Animation/AnimNode_UseCachedPose.cpp b/Engine/Source/Runtime/Engine/Private/Animation/AnimNode_UseCachedPose.cpp
--- a/Engine/Source/Runtime/Engine/Private/Animation/AnimNode_UseCachedPose.cpp
+++ b/Engine/Source/Runtime/Engine/Private/Animation/AnimNode_UseCachedPose.cpp
@@ -1,6 +1,7 @@
 // Copyright 1998-2019 Epic Games, Inc. All Rights Reserved.
 
 #include "Animation/AnimNode_UseCachedPose.h"
+#include "Animation/AnimTrace.h"
 
 /////////////////////////////////////////////////////
 // FAnimNode_UseCachedPose
@@ -24,6 +25,8 @@ void FAnimNode_UseCachedPose::CacheBones_AnyThread(const FAnimationCacheBonesCon
 void FAnimNode_UseCachedPose::Update_AnyThread(const FAnimationUpdateContext& Context)
 {
 	LinkToCachingNode.Update(Context);
+
+	TRACE_ANIM_NODE_VALUE(Context, TEXT("Cached Pose Name"), CachePoseName);
 }
 
 void FAnimNode_UseCachedPose::Evaluate_AnyThread(FPoseContext& Output)
diff --git a/Engine/Source/Runtime/Engine/Private/Animation/AnimSingleNodeInstanceProxy.cpp b/Engine/Source/Runtime/Engine/Private/Animation/AnimSingleNodeInstanceProxy.cpp
--- a/Engine/Source/Runtime/Engine/Private/Animation/AnimSingleNodeInstanceProxy.cpp
+++ b/Engine/Source/Runtime/Engine/Private/Animation/AnimSingleNodeInstanceProxy.cpp
@@ -8,6 +8,7 @@
 #include "Animation/AnimStreamable.h"
 #include "Animation/AnimSingleNodeInstance.h"
 #include "AnimEncoding.h"
+#include "Animation/AnimTrace.h"
 
 FAnimSingleNodeInstanceProxy::~FAnimSingleNodeInstanceProxy()
 {
@@ -503,6 +504,8 @@ void FAnimNode_SingleNode::Update_AnyThread(const FAnimationUpdateContext& Conte
 		{
 			FAnimTickRecord& TickRecord = Proxy->CreateUninitializedTickRecord(INDEX_NONE, /*out*/ SyncGroup);
 			Proxy->MakeBlendSpaceTickRecord(TickRecord, BlendSpace, Proxy->BlendSpaceInput, Proxy->BlendSampleData, Proxy->BlendFilter, Proxy->bLooping, NewPlayRate, 1.f, /*inout*/ Proxy->CurrentTime, Proxy->MarkerTickRecord);
+
+			TRACE_ANIM_TICK_RECORD(Context, TickRecord);
 #if WITH_EDITORONLY_DATA
 			PreviewBasePose = BlendSpace->PreviewBasePose;
 #endif
@@ -511,6 +514,9 @@ void FAnimNode_SingleNode::Update_AnyThread(const FAnimationUpdateContext& Conte
 		{
 			FAnimTickRecord& TickRecord = Proxy->CreateUninitializedTickRecord(INDEX_NONE, /*out*/ SyncGroup);
 			Proxy->MakeSequenceTickRecord(TickRecord, Sequence, Proxy->bLooping, NewPlayRate, 1.f, /*inout*/ Proxy->CurrentTime, Proxy->MarkerTickRecord);
+
+			TRACE_ANIM_TICK_RECORD(Context, TickRecord);
+
 			// if it's not looping, just set play to be false when reached to end
 			if (!Proxy->bLooping)
 			{
@@ -525,6 +531,9 @@ void FAnimNode_SingleNode::Update_AnyThread(const FAnimationUpdateContext& Conte
 		{
 			FAnimTickRecord& TickRecord = Proxy->CreateUninitializedTickRecord(INDEX_NONE, /*out*/ SyncGroup);
 			Proxy->MakeSequenceTickRecord(TickRecord, Streamable, Proxy->bLooping, NewPlayRate, 1.f, /*inout*/ Proxy->CurrentTime, Proxy->MarkerTickRecord);
+
+			TRACE_ANIM_TICK_RECORD(Context, TickRecord);
+
 			// if it's not looping, just set play to be false when reached to end
 			if (!Proxy->bLooping)
 			{
@@ -539,6 +548,9 @@ void FAnimNode_SingleNode::Update_AnyThread(const FAnimationUpdateContext& Conte
 		{
 			FAnimTickRecord& TickRecord = Proxy->CreateUninitializedTickRecord(INDEX_NONE, /*out*/ SyncGroup);
 			Proxy->MakeSequenceTickRecord(TickRecord, Composite, Proxy->bLooping, NewPlayRate, 1.f, /*inout*/ Proxy->CurrentTime, Proxy->MarkerTickRecord);
+
+			TRACE_ANIM_TICK_RECORD(Context, TickRecord);
+
 			// if it's not looping, just set play to be false when reached to end
 			if (!Proxy->bLooping)
 			{
@@ -575,6 +587,8 @@ void FAnimNode_SingleNode::Update_AnyThread(const FAnimationUpdateContext& Conte
 		{
 			FAnimTickRecord& TickRecord = Proxy->CreateUninitializedTickRecord(INDEX_NONE, /*out*/ SyncGroup);
 			Proxy->MakePoseAssetTickRecord(TickRecord, PoseAsset, 1.f);
+
+			TRACE_ANIM_TICK_RECORD(Context, TickRecord);
 		}
 	}
 
diff --git a/Engine/Source/Runtime/Engine/Private/Animation/AnimTrace.cpp b/Engine/Source/Runtime/Engine/Private/Animation/AnimTrace.cpp
--- a/Engine/Source/Runtime/Engine/Private/Animation/AnimTrace.cpp
+++ b/Engine/Source/Runtime/Engine/Private/Animation/AnimTrace.cpp
@@ -13,18 +13,24 @@
 #include "Engine/SkeletalMesh.h"
 #include "Math/TransformNonVectorized.h"
 #include "Animation/AnimNodeBase.h"
+#include "Animation/AnimMontage.h"
+#include "Animation/BlendSpaceBase.h"
+#include "Animation/AnimNode_SequencePlayer.h"
 
 UE_TRACE_EVENT_BEGIN(Animation, TickRecord)
 	UE_TRACE_EVENT_FIELD(uint64, Cycle)
-	UE_TRACE_EVENT_FIELD(uint64, ComponentId)
 	UE_TRACE_EVENT_FIELD(uint64, AnimInstanceId)
 	UE_TRACE_EVENT_FIELD(uint64, AssetId)
+	UE_TRACE_EVENT_FIELD(int32, NodeId)
 	UE_TRACE_EVENT_FIELD(float, BlendWeight)
 	UE_TRACE_EVENT_FIELD(float, PlaybackTime)
 	UE_TRACE_EVENT_FIELD(float, RootMotionWeight)
 	UE_TRACE_EVENT_FIELD(float, PlayRate)
+	UE_TRACE_EVENT_FIELD(float, BlendSpacePositionX)
+	UE_TRACE_EVENT_FIELD(float, BlendSpacePositionY)
 	UE_TRACE_EVENT_FIELD(uint16, FrameCounter)
 	UE_TRACE_EVENT_FIELD(bool, Looping)
+	UE_TRACE_EVENT_FIELD(bool, IsBlendSpace)
 UE_TRACE_EVENT_END()
 
 UE_TRACE_EVENT_BEGIN(Animation, SkeletalMesh, Important)
@@ -42,44 +48,127 @@ UE_TRACE_EVENT_BEGIN(Animation, SkeletalMeshComponent)
 	UE_TRACE_EVENT_FIELD(uint16, FrameCounter)
 UE_TRACE_EVENT_END()
 
+UE_TRACE_EVENT_BEGIN(Animation, SkeletalMeshFrame)
+	UE_TRACE_EVENT_FIELD(uint64, Cycle)
+	UE_TRACE_EVENT_FIELD(uint64, ComponentId)
+	UE_TRACE_EVENT_FIELD(uint16, FrameCounter)
+UE_TRACE_EVENT_END()
+
+UE_TRACE_EVENT_BEGIN(Animation, AnimGraph)
+	UE_TRACE_EVENT_FIELD(uint64, StartCycle)
+	UE_TRACE_EVENT_FIELD(uint64, EndCycle)
+	UE_TRACE_EVENT_FIELD(uint64, AnimInstanceId)
+	UE_TRACE_EVENT_FIELD(int32, NodeCount)
+	UE_TRACE_EVENT_FIELD(uint16, FrameCounter)
+	UE_TRACE_EVENT_FIELD(uint8, Phase)
+UE_TRACE_EVENT_END()
+
+UE_TRACE_EVENT_BEGIN(Animation, AnimNodeStart)
+	UE_TRACE_EVENT_FIELD(uint64, StartCycle)
+	UE_TRACE_EVENT_FIELD(uint64, AnimInstanceId)
+	UE_TRACE_EVENT_FIELD(int32, PreviousNodeId)
+	UE_TRACE_EVENT_FIELD(int32, NodeId)
+	UE_TRACE_EVENT_FIELD(float, Weight)
+	UE_TRACE_EVENT_FIELD(uint16, FrameCounter)
+	UE_TRACE_EVENT_FIELD(uint8, Phase)
+UE_TRACE_EVENT_END()
+
+UE_TRACE_EVENT_BEGIN(Animation, AnimNodeEnd)
+	UE_TRACE_EVENT_FIELD(uint64, EndCycle)
+	UE_TRACE_EVENT_FIELD(uint64, AnimInstanceId)
+UE_TRACE_EVENT_END()
+
 UE_TRACE_EVENT_BEGIN(Animation, AnimNodeValueBool)
+	UE_TRACE_EVENT_FIELD(uint64, Cycle)
 	UE_TRACE_EVENT_FIELD(uint64, AnimInstanceId)
-	UE_TRACE_EVENT_FIELD(uint64, FrameId)
 	UE_TRACE_EVENT_FIELD(int32, NodeId)
 	UE_TRACE_EVENT_FIELD(int32, KeyLength)
+	UE_TRACE_EVENT_FIELD(uint16, FrameCounter)
 	UE_TRACE_EVENT_FIELD(bool, Value)
 UE_TRACE_EVENT_END()
 
 UE_TRACE_EVENT_BEGIN(Animation, AnimNodeValueInt)
+	UE_TRACE_EVENT_FIELD(uint64, Cycle)
 	UE_TRACE_EVENT_FIELD(uint64, AnimInstanceId)
-	UE_TRACE_EVENT_FIELD(uint64, FrameId)
 	UE_TRACE_EVENT_FIELD(int32, NodeId)
 	UE_TRACE_EVENT_FIELD(int32, KeyLength)
 	UE_TRACE_EVENT_FIELD(int32, Value)
+	UE_TRACE_EVENT_FIELD(uint16, FrameCounter)
 UE_TRACE_EVENT_END()
 
 UE_TRACE_EVENT_BEGIN(Animation, AnimNodeValueFloat)
+	UE_TRACE_EVENT_FIELD(uint64, Cycle)
 	UE_TRACE_EVENT_FIELD(uint64, AnimInstanceId)
-	UE_TRACE_EVENT_FIELD(uint64, FrameId)
 	UE_TRACE_EVENT_FIELD(int32, NodeId)
 	UE_TRACE_EVENT_FIELD(int32, KeyLength)
 	UE_TRACE_EVENT_FIELD(float, Value)
+	UE_TRACE_EVENT_FIELD(uint16, FrameCounter)
+UE_TRACE_EVENT_END()
+
+UE_TRACE_EVENT_BEGIN(Animation, AnimNodeValueVector)
+	UE_TRACE_EVENT_FIELD(uint64, Cycle)
+	UE_TRACE_EVENT_FIELD(uint64, AnimInstanceId)
+	UE_TRACE_EVENT_FIELD(int32, NodeId)
+	UE_TRACE_EVENT_FIELD(int32, KeyLength)
+	UE_TRACE_EVENT_FIELD(float, ValueX)
+	UE_TRACE_EVENT_FIELD(float, ValueY)
+	UE_TRACE_EVENT_FIELD(float, ValueZ)
+	UE_TRACE_EVENT_FIELD(uint16, FrameCounter)
 UE_TRACE_EVENT_END()
 
 UE_TRACE_EVENT_BEGIN(Animation, AnimNodeValueString)
+	UE_TRACE_EVENT_FIELD(uint64, Cycle)
 	UE_TRACE_EVENT_FIELD(uint64, AnimInstanceId)
-	UE_TRACE_EVENT_FIELD(uint64, FrameId)
 	UE_TRACE_EVENT_FIELD(int32, NodeId)
 	UE_TRACE_EVENT_FIELD(int32, KeyLength)
+	UE_TRACE_EVENT_FIELD(uint16, FrameCounter)
 UE_TRACE_EVENT_END()
 
-UE_TRACE_EVENT_BEGIN(Animation, PoseLink)
+UE_TRACE_EVENT_BEGIN(Animation, AnimNodeValueObject)
+	UE_TRACE_EVENT_FIELD(uint64, Cycle)
 	UE_TRACE_EVENT_FIELD(uint64, AnimInstanceId)
-	UE_TRACE_EVENT_FIELD(uint64, FrameId)
-	UE_TRACE_EVENT_FIELD(int32, TargetLinkId)
-	UE_TRACE_EVENT_FIELD(int32, SourceLinkId)
-	UE_TRACE_EVENT_FIELD(float, Weight)
-	UE_TRACE_EVENT_FIELD(int32, TargetNameLength)
+	UE_TRACE_EVENT_FIELD(uint64, Value)
+	UE_TRACE_EVENT_FIELD(int32, NodeId)
+	UE_TRACE_EVENT_FIELD(int32, KeyLength)
+	UE_TRACE_EVENT_FIELD(uint16, FrameCounter)
+UE_TRACE_EVENT_END()
+
+UE_TRACE_EVENT_BEGIN(Animation, AnimNodeValueClass)
+	UE_TRACE_EVENT_FIELD(uint64, Cycle)
+	UE_TRACE_EVENT_FIELD(uint64, AnimInstanceId)
+	UE_TRACE_EVENT_FIELD(uint64, Value)
+	UE_TRACE_EVENT_FIELD(int32, NodeId)
+	UE_TRACE_EVENT_FIELD(int32, KeyLength)
+	UE_TRACE_EVENT_FIELD(uint16, FrameCounter)
+UE_TRACE_EVENT_END()
+
+UE_TRACE_EVENT_BEGIN(Animation, AnimSequencePlayer)
+	UE_TRACE_EVENT_FIELD(uint64, Cycle)
+	UE_TRACE_EVENT_FIELD(uint64, AnimInstanceId)
+	UE_TRACE_EVENT_FIELD(int32, NodeId)
+	UE_TRACE_EVENT_FIELD(float, Position)
+	UE_TRACE_EVENT_FIELD(float, Length)
+	UE_TRACE_EVENT_FIELD(int32, FrameCount)
+UE_TRACE_EVENT_END()
+
+UE_TRACE_EVENT_BEGIN(Animation, BlendSpacePlayer)
+	UE_TRACE_EVENT_FIELD(uint64, Cycle)
+	UE_TRACE_EVENT_FIELD(uint64, AnimInstanceId)
+	UE_TRACE_EVENT_FIELD(uint64, BlendSpaceId)
+	UE_TRACE_EVENT_FIELD(int32, NodeId)
+	UE_TRACE_EVENT_FIELD(float, PositionX)
+	UE_TRACE_EVENT_FIELD(float, PositionY)
+	UE_TRACE_EVENT_FIELD(float, PositionZ)
+UE_TRACE_EVENT_END()
+
+UE_TRACE_EVENT_BEGIN(Animation, StateMachineState)
+	UE_TRACE_EVENT_FIELD(uint64, Cycle)
+	UE_TRACE_EVENT_FIELD(uint64, AnimInstanceId)
+	UE_TRACE_EVENT_FIELD(int32, NodeId)
+	UE_TRACE_EVENT_FIELD(int32, StateMachineIndex)
+	UE_TRACE_EVENT_FIELD(int32, StateIndex)
+	UE_TRACE_EVENT_FIELD(float, StateWeight)
+	UE_TRACE_EVENT_FIELD(float, ElapsedTime)
 UE_TRACE_EVENT_END()
 
 UE_TRACE_EVENT_BEGIN(Animation, Name, Important)
@@ -92,6 +181,76 @@ FUObjectAnnotationSparseBool GSkeletalMeshTraceAnnotations;
 // Map used for unique name output
 TMap<FName, uint32> GAnimTraceNames;
 
+FAnimTrace::FScopedAnimNodeTrace::FScopedAnimNodeTrace(const FAnimationInitializeContext& InContext)
+	: Context(InContext)
+{
+	OutputAnimNodeStart(InContext, FPlatformTime::Cycles64(), InContext.GetPreviousNodeId(), InContext.GetCurrentNodeId(), 0.0f, 0.0f, (__underlying_type(EPhase))EPhase::Initialize);
+}
+
+FAnimTrace::FScopedAnimNodeTrace::FScopedAnimNodeTrace(const FAnimationUpdateContext& InContext)
+	: Context(InContext)
+{
+	OutputAnimNodeStart(InContext, FPlatformTime::Cycles64(), InContext.GetPreviousNodeId(), InContext.GetCurrentNodeId(), InContext.GetFinalBlendWeight(), InContext.GetRootMotionWeightModifier(), (__underlying_type(EPhase))EPhase::Update);
+}
+
+FAnimTrace::FScopedAnimNodeTrace::FScopedAnimNodeTrace(const FAnimationCacheBonesContext& InContext)
+	: Context(InContext)
+{
+	OutputAnimNodeStart(InContext, FPlatformTime::Cycles64(), InContext.GetPreviousNodeId(), InContext.GetCurrentNodeId(), 0.0f, 0.0f, (__underlying_type(EPhase))EPhase::CacheBones);
+}
+
+FAnimTrace::FScopedAnimNodeTrace::FScopedAnimNodeTrace(const FPoseContext& InContext)
+	: Context(InContext)
+{
+	OutputAnimNodeStart(InContext, FPlatformTime::Cycles64(), InContext.GetPreviousNodeId(), InContext.GetCurrentNodeId(), 0.0f, 0.0f, (__underlying_type(EPhase))EPhase::Evaluate);
+}
+
+FAnimTrace::FScopedAnimNodeTrace::FScopedAnimNodeTrace(const FComponentSpacePoseContext& InContext)
+	: Context(InContext)
+{
+	OutputAnimNodeStart(InContext, FPlatformTime::Cycles64(), InContext.GetPreviousNodeId(), InContext.GetCurrentNodeId(), 0.0f, 0.0f, (__underlying_type(EPhase))EPhase::Evaluate);
+}
+
+FAnimTrace::FScopedAnimNodeTrace::~FScopedAnimNodeTrace()
+{
+	OutputAnimNodeEnd(Context, FPlatformTime::Cycles64());
+}
+
+FAnimTrace::FScopedAnimGraphTrace::FScopedAnimGraphTrace(const FAnimationInitializeContext& InContext)
+	: StartCycle(FPlatformTime::Cycles64())
+	, Context(InContext)
+	, Phase(EPhase::Initialize)
+{}
+
+FAnimTrace::FScopedAnimGraphTrace::FScopedAnimGraphTrace(const FAnimationUpdateContext& InContext)
+	: StartCycle(FPlatformTime::Cycles64())
+	, Context(InContext)
+	, Phase(EPhase::Update)
+{}
+
+FAnimTrace::FScopedAnimGraphTrace::FScopedAnimGraphTrace(const FAnimationCacheBonesContext& InContext)
+	: StartCycle(FPlatformTime::Cycles64())
+	, Context(InContext)
+	, Phase(EPhase::CacheBones)
+{}
+
+FAnimTrace::FScopedAnimGraphTrace::FScopedAnimGraphTrace(const FPoseContext& InContext)
+	: StartCycle(FPlatformTime::Cycles64())
+	, Context(InContext)
+	, Phase(EPhase::Evaluate)
+{}
+
+FAnimTrace::FScopedAnimGraphTrace::FScopedAnimGraphTrace(const FComponentSpacePoseContext& InContext)
+	: StartCycle(FPlatformTime::Cycles64())
+	, Context(InContext)
+	, Phase(EPhase::Evaluate)
+{}
+
+FAnimTrace::FScopedAnimGraphTrace::~FScopedAnimGraphTrace()
+{
+	OutputAnimGraph(Context, StartCycle, FPlatformTime::Cycles64(), (__underlying_type(EPhase))Phase);
+}
+
 void FAnimTrace::Init()
 {
 	if (FParse::Param(FCommandLine::Get(), TEXT("objecttrace")))
@@ -99,56 +258,65 @@ void FAnimTrace::Init()
 		UE_TRACE_EVENT_IS_ENABLED(Animation, TickRecord);
 		UE_TRACE_EVENT_IS_ENABLED(Animation, SkeletalMesh);
 		UE_TRACE_EVENT_IS_ENABLED(Animation, SkeletalMeshComponent);
-// 		UE_TRACE_EVENT_IS_ENABLED(Animation, AnimNodeValueBool);
-// 		UE_TRACE_EVENT_IS_ENABLED(Animation, AnimNodeValueInt);
-// 		UE_TRACE_EVENT_IS_ENABLED(Animation, AnimNodeValueFloat);
-// 		UE_TRACE_EVENT_IS_ENABLED(Animation, AnimNodeValueString);
-// 		UE_TRACE_EVENT_IS_ENABLED(Animation, PoseLink);
+		UE_TRACE_EVENT_IS_ENABLED(Animation, SkeletalMeshFrame);
+		UE_TRACE_EVENT_IS_ENABLED(Animation, AnimGraph);
+		UE_TRACE_EVENT_IS_ENABLED(Animation, AnimNodeStart);
+		UE_TRACE_EVENT_IS_ENABLED(Animation, AnimNodeEnd);
+		UE_TRACE_EVENT_IS_ENABLED(Animation, AnimNodeValueBool);
+		UE_TRACE_EVENT_IS_ENABLED(Animation, AnimNodeValueInt);
+		UE_TRACE_EVENT_IS_ENABLED(Animation, AnimNodeValueFloat);
+		UE_TRACE_EVENT_IS_ENABLED(Animation, AnimNodeValueVector);
+		UE_TRACE_EVENT_IS_ENABLED(Animation, AnimNodeValueString);
+		UE_TRACE_EVENT_IS_ENABLED(Animation, AnimSequencePlayer);
+		UE_TRACE_EVENT_IS_ENABLED(Animation, StateMachineState);
 		UE_TRACE_EVENT_IS_ENABLED(Animation, Name);
 		Trace::ToggleEvent(TEXT("Animation"), true);
 	}
 }
 
-void FAnimTrace::OutputAnimTickRecords(const FAnimInstanceProxy& InProxy, const USkeletalMeshComponent* InComponent)
+void FAnimTrace::OutputAnimTickRecord(const FAnimationBaseContext& InContext, const FAnimTickRecord& InTickRecord)
 {
 	bool bEventEnabled = UE_TRACE_EVENT_IS_ENABLED(Animation, TickRecord);
-	if (!bEventEnabled || InComponent == nullptr)
+	if (!bEventEnabled)
 	{
 		return;
 	}
 
-	TRACE_OBJECT(InComponent);
+	TRACE_OBJECT(InContext.AnimInstanceProxy->GetAnimInstanceObject());
 
-	auto TraceTickRecord = [&InProxy, &InComponent](const FAnimTickRecord& InAnimTickRecord)
+	if(InTickRecord.SourceAsset)
 	{
-		TRACE_OBJECT(InAnimTickRecord.SourceAsset);
+		TRACE_OBJECT(InTickRecord.SourceAsset);
 
-		UE_TRACE_LOG(Animation, TickRecord)
-			<< TickRecord.Cycle(FPlatformTime::Cycles64())
-			<< TickRecord.ComponentId(FObjectTrace::GetObjectId(InComponent))
-			<< TickRecord.AnimInstanceId(FObjectTrace::GetObjectId(InProxy.GetAnimInstanceObject()))
-			<< TickRecord.AssetId(FObjectTrace::GetObjectId(InAnimTickRecord.SourceAsset))
-			<< TickRecord.BlendWeight(InAnimTickRecord.EffectiveBlendWeight)
-			<< TickRecord.PlaybackTime(*InAnimTickRecord.TimeAccumulator)
-			<< TickRecord.RootMotionWeight(InAnimTickRecord.RootMotionWeightModifier)
-			<< TickRecord.PlayRate(InAnimTickRecord.PlayRateMultiplier)
-			<< TickRecord.FrameCounter((uint16)(GFrameCounter % 0xffff))
-			<< TickRecord.Looping(InAnimTickRecord.bLooping);
-	};
+		float PlaybackTime = *InTickRecord.TimeAccumulator;
+		if(InTickRecord.SourceAsset->IsA<UAnimMontage>())
+		{
+			PlaybackTime = InTickRecord.Montage.CurrentPosition;
+		}
 
-	const TArray<FAnimGroupInstance>& SyncGroups = InProxy.SyncGroupArrays[InProxy.GetSyncGroupWriteIndex()];
-	const TArray<FAnimTickRecord>& UngroupedActivePlayers = InProxy.UngroupedActivePlayerArrays[InProxy.GetSyncGroupWriteIndex()];
-	for(const FAnimGroupInstance& SyncGroup : SyncGroups)
-	{
-		for(const FAnimTickRecord& ActivePlayer : SyncGroup.ActivePlayers)
+		float BlendSpacePositionX = 0.0f;
+		float BlendSpacePositionY = 0.0f;
+		const bool bIsBlendSpace = InTickRecord.SourceAsset->IsA<UBlendSpaceBase>();
+		if(bIsBlendSpace)
 		{
-			TraceTickRecord(ActivePlayer);
+			BlendSpacePositionX = InTickRecord.BlendSpace.BlendSpacePositionX;
+			BlendSpacePositionY = InTickRecord.BlendSpace.BlendSpacePositionY;
 		}
-	}
 
-	for(const FAnimTickRecord& UngroupedActivePlayer : UngroupedActivePlayers)
-	{
-		TraceTickRecord(UngroupedActivePlayer);
+		UE_TRACE_LOG(Animation, TickRecord)
+			<< TickRecord.Cycle(FPlatformTime::Cycles64())
+			<< TickRecord.AnimInstanceId(FObjectTrace::GetObjectId(InContext.AnimInstanceProxy->GetAnimInstanceObject()))
+			<< TickRecord.AssetId(FObjectTrace::GetObjectId(InTickRecord.SourceAsset))
+			<< TickRecord.NodeId(InContext.GetCurrentNodeId())
+			<< TickRecord.BlendWeight(InTickRecord.EffectiveBlendWeight)
+			<< TickRecord.PlaybackTime(PlaybackTime)
+			<< TickRecord.RootMotionWeight(InTickRecord.RootMotionWeightModifier)
+			<< TickRecord.PlayRate(InTickRecord.PlayRateMultiplier)
+			<< TickRecord.BlendSpacePositionX(BlendSpacePositionX)
+			<< TickRecord.BlendSpacePositionY(BlendSpacePositionY)
+			<< TickRecord.FrameCounter((uint16)(GFrameCounter % 0xffff))
+			<< TickRecord.Looping(InTickRecord.bLooping)
+			<< TickRecord.IsBlendSpace(bIsBlendSpace);
 	}
 }
 
@@ -167,12 +335,12 @@ void FAnimTrace::OutputSkeletalMesh(const USkeletalMesh* InMesh)
 
 	TRACE_OBJECT(InMesh);
 
-	uint32 BoneCount = (uint32)InMesh->RefSkeleton.GetRawBoneNum();
+	uint32 BoneCount = (uint32)InMesh->RefSkeleton.GetNum();
 
 	auto CopyParentIndices = [InMesh](uint8* Out)
 	{
 		int32* OutParentIndices = reinterpret_cast<int32*>(Out);
-		for(const FMeshBoneInfo& BoneInfo : InMesh->RefSkeleton.GetRawRefBoneInfo())
+		for(const FMeshBoneInfo& BoneInfo : InMesh->RefSkeleton.GetRefBoneInfo())
 		{
 			*OutParentIndices = BoneInfo.ParentIndex;
 			OutParentIndices++;
@@ -280,6 +448,96 @@ void FAnimTrace::OutputSkeletalMeshComponent(const USkeletalMeshComponent* InCom
 	}
 }
 
+void FAnimTrace::OutputSkeletalMeshFrame(const USkeletalMeshComponent* InComponent)
+{
+	bool bEventEnabled = UE_TRACE_EVENT_IS_ENABLED(Animation, SkeletalMeshFrame);
+	if (!bEventEnabled || InComponent == nullptr)
+	{
+		return;
+	}
+
+	TRACE_OBJECT(InComponent);
+
+	UE_TRACE_LOG(Animation, SkeletalMeshFrame)
+		<< SkeletalMeshFrame.Cycle(FPlatformTime::Cycles64())
+		<< SkeletalMeshFrame.ComponentId(FObjectTrace::GetObjectId(InComponent))
+		<< SkeletalMeshFrame.FrameCounter((uint16)(GFrameCounter % 0xffff));
+}
+
+void FAnimTrace::OutputAnimGraph(const FAnimationBaseContext& InContext, uint64 InStartCycle, uint64 InEndCycle, uint8 InPhase)
+{
+	bool bEventEnabled = UE_TRACE_EVENT_IS_ENABLED(Animation, AnimGraph);
+	if (!bEventEnabled)
+	{
+		return;
+	}
+
+	check(InContext.AnimInstanceProxy);
+
+	const UAnimInstance* AnimInstance = Cast<UAnimInstance>(InContext.AnimInstanceProxy->GetAnimInstanceObject());
+	const UAnimBlueprintGeneratedClass* BPClass = Cast<UAnimBlueprintGeneratedClass>(AnimInstance->GetClass());
+
+	TRACE_OBJECT(AnimInstance);
+
+	UE_TRACE_LOG(Animation, AnimGraph)
+		<< AnimGraph.StartCycle(InStartCycle)
+		<< AnimGraph.EndCycle(InEndCycle)
+		<< AnimGraph.AnimInstanceId(FObjectTrace::GetObjectId(AnimInstance))
+		<< AnimGraph.NodeCount(BPClass ? BPClass->AnimNodeProperties.Num() : 0)
+		<< AnimGraph.FrameCounter((uint16)(GFrameCounter % 0xffff))
+		<< AnimGraph.Phase(InPhase);
+}
+
+void FAnimTrace::OutputAnimNodeStart(const FAnimationBaseContext& InContext, uint64 InStartCycle, int32 InPreviousNodeId, int32 InNodeId, float InBlendWeight, float InRootMotionWeight, uint8 InPhase)
+{
+	bool bEventEnabled = UE_TRACE_EVENT_IS_ENABLED(Animation, AnimNodeStart);
+	if (!bEventEnabled)
+	{
+		return;
+	}
+
+	check(InContext.AnimInstanceProxy);
+
+	TRACE_OBJECT(InContext.AnimInstanceProxy->GetAnimInstanceObject());
+
+	IAnimClassInterface* AnimBlueprintClass = InContext.GetAnimClass();
+	check(AnimBlueprintClass);
+	const TArray<UStructProperty*>& AnimNodeProperties = AnimBlueprintClass->GetAnimNodeProperties();
+	check(AnimNodeProperties.IsValidIndex(InNodeId));
+	UStructProperty* LinkedProperty = AnimNodeProperties[InNodeId];
+	check(LinkedProperty->Struct);
+
+	FString DisplayNameString = LinkedProperty->Struct->GetDisplayNameText().ToString();
+	DisplayNameString.RemoveFromStart(TEXT("Anim Node "));
+
+	UE_TRACE_LOG(Animation, AnimNodeStart, (DisplayNameString.Len() + 1) * sizeof(TCHAR))
+		<< AnimNodeStart.StartCycle(InStartCycle)
+		<< AnimNodeStart.AnimInstanceId(FObjectTrace::GetObjectId(InContext.AnimInstanceProxy->GetAnimInstanceObject()))
+		<< AnimNodeStart.PreviousNodeId(InPreviousNodeId)
+		<< AnimNodeStart.NodeId(InNodeId)
+		<< AnimNodeStart.Weight(InBlendWeight)
+		<< AnimNodeStart.FrameCounter((uint16)(GFrameCounter % 0xffff))
+		<< AnimNodeStart.Phase(InPhase)
+		<< AnimNodeStart.Attachment(*DisplayNameString, (DisplayNameString.Len() + 1) * sizeof(TCHAR));
+}
+
+void FAnimTrace::OutputAnimNodeEnd(const FAnimationBaseContext& InContext, uint64 InEndCycle)
+{
+	bool bEventEnabled = UE_TRACE_EVENT_IS_ENABLED(Animation, AnimNodeEnd);
+	if (!bEventEnabled)
+	{
+		return;
+	}
+
+	check(InContext.AnimInstanceProxy);
+
+	TRACE_OBJECT(InContext.AnimInstanceProxy->GetAnimInstanceObject());
+
+	UE_TRACE_LOG(Animation, AnimNodeEnd)
+		<< AnimNodeEnd.EndCycle(InEndCycle)
+		<< AnimNodeEnd.AnimInstanceId(FObjectTrace::GetObjectId(InContext.AnimInstanceProxy->GetAnimInstanceObject()));
+}
+
 void FAnimTrace::OutputAnimNodeValue(const FAnimationBaseContext& InContext, const TCHAR* InKey, bool InValue)
 {
 	bool bEventEnabled = UE_TRACE_EVENT_IS_ENABLED(Animation, AnimNodeValueBool);
@@ -290,14 +548,17 @@ void FAnimTrace::OutputAnimNodeValue(const FAnimationBaseContext& InContext, con
 
 	check(InContext.AnimInstanceProxy);
 
+	TRACE_OBJECT(InContext.AnimInstanceProxy->GetAnimInstanceObject());
+
 	int32 KeyLength = FCString::Strlen(InKey) + 1;
 
 	UE_TRACE_LOG(Animation, AnimNodeValueBool, KeyLength * sizeof(TCHAR))
+		<< AnimNodeValueBool.Cycle(FPlatformTime::Cycles64())
 		<< AnimNodeValueBool.AnimInstanceId(FObjectTrace::GetObjectId(InContext.AnimInstanceProxy->GetAnimInstanceObject()))
-		<< AnimNodeValueBool.FrameId(GFrameCounter)
 		<< AnimNodeValueBool.NodeId(InContext.GetCurrentNodeId())
 		<< AnimNodeValueBool.KeyLength(KeyLength)
 		<< AnimNodeValueBool.Value(InValue)
+		<< AnimNodeValueBool.FrameCounter((uint16)(GFrameCounter % 0xffff))
 		<< AnimNodeValueBool.Attachment(InKey, KeyLength * sizeof(TCHAR));
 }
 
@@ -311,14 +572,17 @@ void FAnimTrace::OutputAnimNodeValue(const FAnimationBaseContext& InContext, con
 
 	check(InContext.AnimInstanceProxy);
 
+	TRACE_OBJECT(InContext.AnimInstanceProxy->GetAnimInstanceObject());
+
 	int32 KeyLength = FCString::Strlen(InKey) + 1;
 
 	UE_TRACE_LOG(Animation, AnimNodeValueInt, KeyLength * sizeof(TCHAR))
+		<< AnimNodeValueInt.Cycle(FPlatformTime::Cycles64())
 		<< AnimNodeValueInt.AnimInstanceId(FObjectTrace::GetObjectId(InContext.AnimInstanceProxy->GetAnimInstanceObject()))
-		<< AnimNodeValueInt.FrameId(GFrameCounter)
 		<< AnimNodeValueInt.NodeId(InContext.GetCurrentNodeId())
 		<< AnimNodeValueInt.KeyLength(KeyLength)
 		<< AnimNodeValueInt.Value(InValue)
+		<< AnimNodeValueInt.FrameCounter((uint16)(GFrameCounter % 0xffff))
 		<< AnimNodeValueInt.Attachment(InKey, KeyLength * sizeof(TCHAR));
 }
 
@@ -332,20 +596,55 @@ void FAnimTrace::OutputAnimNodeValue(const FAnimationBaseContext& InContext, con
 
 	check(InContext.AnimInstanceProxy);
 
+	TRACE_OBJECT(InContext.AnimInstanceProxy->GetAnimInstanceObject());
+
 	int32 KeyLength = FCString::Strlen(InKey) + 1;
 
 	UE_TRACE_LOG(Animation, AnimNodeValueFloat, KeyLength * sizeof(TCHAR))
+		<< AnimNodeValueFloat.Cycle(FPlatformTime::Cycles64())
 		<< AnimNodeValueFloat.AnimInstanceId(FObjectTrace::GetObjectId(InContext.AnimInstanceProxy->GetAnimInstanceObject()))
-		<< AnimNodeValueFloat.FrameId(GFrameCounter)
 		<< AnimNodeValueFloat.NodeId(InContext.GetCurrentNodeId())
 		<< AnimNodeValueFloat.KeyLength(KeyLength)
 		<< AnimNodeValueFloat.Value(InValue)
+		<< AnimNodeValueFloat.FrameCounter((uint16)(GFrameCounter % 0xffff))
 		<< AnimNodeValueFloat.Attachment(InKey, KeyLength * sizeof(TCHAR));
 }
 
+void FAnimTrace::OutputAnimNodeValue(const FAnimationBaseContext& InContext, const TCHAR* InKey, const FRotator& Value)
+{
+	const FVector VectorValue(Value.Roll, Value.Pitch, Value.Yaw);
+	OutputAnimNodeValue(InContext, InKey, VectorValue);
+}
+
+void FAnimTrace::OutputAnimNodeValue(const FAnimationBaseContext& InContext, const TCHAR* InKey, const FVector& InValue)
+{
+	bool bEventEnabled = UE_TRACE_EVENT_IS_ENABLED(Animation, AnimNodeValueVector);
+	if (!bEventEnabled)
+	{
+		return;
+	}
+
+	check(InContext.AnimInstanceProxy);
+
+	TRACE_OBJECT(InContext.AnimInstanceProxy->GetAnimInstanceObject());
+
+	int32 KeyLength = FCString::Strlen(InKey) + 1;
+
+	UE_TRACE_LOG(Animation, AnimNodeValueVector, KeyLength * sizeof(TCHAR))
+		<< AnimNodeValueVector.Cycle(FPlatformTime::Cycles64())
+		<< AnimNodeValueVector.AnimInstanceId(FObjectTrace::GetObjectId(InContext.AnimInstanceProxy->GetAnimInstanceObject()))
+		<< AnimNodeValueVector.NodeId(InContext.GetCurrentNodeId())
+		<< AnimNodeValueVector.KeyLength(KeyLength)
+		<< AnimNodeValueVector.ValueX(InValue.X)
+		<< AnimNodeValueVector.ValueY(InValue.Y)
+		<< AnimNodeValueVector.ValueZ(InValue.Z)
+		<< AnimNodeValueVector.FrameCounter((uint16)(GFrameCounter % 0xffff))
+		<< AnimNodeValueVector.Attachment(InKey, KeyLength * sizeof(TCHAR));
+}
+
 void FAnimTrace::OutputAnimNodeValue(const FAnimationBaseContext& InContext, const TCHAR* InKey, const FName& InValue)
 {
-	bool bEventEnabled = UE_TRACE_EVENT_IS_ENABLED(Animation, AnimNodeValueFloat);
+	bool bEventEnabled = UE_TRACE_EVENT_IS_ENABLED(Animation, AnimNodeValueString);
 	if (!bEventEnabled)
 	{
 		return;
@@ -353,6 +652,8 @@ void FAnimTrace::OutputAnimNodeValue(const FAnimationBaseContext& InContext, con
 
 	check(InContext.AnimInstanceProxy);
 
+	TRACE_OBJECT(InContext.AnimInstanceProxy->GetAnimInstanceObject());
+
 	int32 KeyLength = FCString::Strlen(InKey) + 1;
 	int32 ValueLength = InValue.GetStringLength() + 1;
 
@@ -362,17 +663,18 @@ void FAnimTrace::OutputAnimNodeValue(const FAnimationBaseContext& InContext, con
 		InValue.ToString(reinterpret_cast<TCHAR*>(Out) + KeyLength, ValueLength);
 	};
 
-	UE_TRACE_LOG(Animation, AnimNodeValueString, KeyLength * sizeof(TCHAR))
+	UE_TRACE_LOG(Animation, AnimNodeValueString, (KeyLength + ValueLength) * sizeof(TCHAR))
+		<< AnimNodeValueString.Cycle(FPlatformTime::Cycles64())
 		<< AnimNodeValueString.AnimInstanceId(FObjectTrace::GetObjectId(InContext.AnimInstanceProxy->GetAnimInstanceObject()))
-		<< AnimNodeValueString.FrameId(GFrameCounter)
 		<< AnimNodeValueString.NodeId(InContext.GetCurrentNodeId())
 		<< AnimNodeValueString.KeyLength(KeyLength)
-		<< AnimNodeValueString.Attachment(InKey, KeyLength * sizeof(TCHAR));
+		<< AnimNodeValueString.FrameCounter((uint16)(GFrameCounter % 0xffff))
+		<< AnimNodeValueString.Attachment(StringCopyFunc);
 }
 
 void FAnimTrace::OutputAnimNodeValue(const FAnimationBaseContext& InContext, const TCHAR* InKey, const TCHAR* InValue)
 {
-	bool bEventEnabled = UE_TRACE_EVENT_IS_ENABLED(Animation, AnimNodeValueFloat);
+	bool bEventEnabled = UE_TRACE_EVENT_IS_ENABLED(Animation, AnimNodeValueString);
 	if (!bEventEnabled)
 	{
 		return;
@@ -380,6 +682,8 @@ void FAnimTrace::OutputAnimNodeValue(const FAnimationBaseContext& InContext, con
 
 	check(InContext.AnimInstanceProxy);
 
+	TRACE_OBJECT(InContext.AnimInstanceProxy->GetAnimInstanceObject());
+
 	int32 KeyLength = FCString::Strlen(InKey) + 1;
 	int32 ValueLength = FCString::Strlen(InValue) + 1;
 
@@ -389,17 +693,18 @@ void FAnimTrace::OutputAnimNodeValue(const FAnimationBaseContext& InContext, con
 		FCString::Strncpy(reinterpret_cast<TCHAR*>(Out) + KeyLength, InValue, ValueLength);
 	};
 
-	UE_TRACE_LOG(Animation, AnimNodeValueString, KeyLength * sizeof(TCHAR))
+	UE_TRACE_LOG(Animation, AnimNodeValueString, (KeyLength + ValueLength) * sizeof(TCHAR))
+		<< AnimNodeValueString.Cycle(FPlatformTime::Cycles64())
 		<< AnimNodeValueString.AnimInstanceId(FObjectTrace::GetObjectId(InContext.AnimInstanceProxy->GetAnimInstanceObject()))
-		<< AnimNodeValueString.FrameId(GFrameCounter)
 		<< AnimNodeValueString.NodeId(InContext.GetCurrentNodeId())
 		<< AnimNodeValueString.KeyLength(KeyLength)
-		<< AnimNodeValueString.Attachment(InKey, KeyLength * sizeof(TCHAR));
+		<< AnimNodeValueString.FrameCounter((uint16)(GFrameCounter % 0xffff))
+		<< AnimNodeValueString.Attachment(StringCopyFunc);
 }
 
-void FAnimTrace::OutputPoseLink(const FAnimationUpdateContext& InContext)
+void FAnimTrace::OutputAnimNodeValue(const FAnimationBaseContext& InContext, const TCHAR* InKey, const UObject* InValue)
 {
-	bool bEventEnabled = UE_TRACE_EVENT_IS_ENABLED(Animation, PoseLink);
+	bool bEventEnabled = UE_TRACE_EVENT_IS_ENABLED(Animation, AnimNodeValueObject);
 	if (!bEventEnabled)
 	{
 		return;
@@ -407,27 +712,87 @@ void FAnimTrace::OutputPoseLink(const FAnimationUpdateContext& InContext)
 
 	check(InContext.AnimInstanceProxy);
 
-	IAnimClassInterface* AnimBlueprintClass = InContext.GetAnimClass();
-	check(AnimBlueprintClass);
-	const TArray<UStructProperty*>& AnimNodeProperties = AnimBlueprintClass->GetAnimNodeProperties();
-	check(AnimNodeProperties.IsValidIndex(InContext.GetCurrentNodeId()));
-	UProperty* LinkedProperty = AnimNodeProperties[InContext.GetCurrentNodeId()];
+	TRACE_OBJECT(InContext.AnimInstanceProxy->GetAnimInstanceObject());
+	TRACE_OBJECT(InValue);
 
-	int32 NameLength = LinkedProperty->GetFName().GetStringLength() + 1;
+	int32 KeyLength = FCString::Strlen(InKey) + 1;
 
-	auto StringCopyFunc = [NameLength, LinkedProperty](uint8* Out)
+	UE_TRACE_LOG(Animation, AnimNodeValueObject, KeyLength * sizeof(TCHAR))
+		<< AnimNodeValueObject.Cycle(FPlatformTime::Cycles64())
+		<< AnimNodeValueObject.AnimInstanceId(FObjectTrace::GetObjectId(InContext.AnimInstanceProxy->GetAnimInstanceObject()))
+		<< AnimNodeValueObject.NodeId(InContext.GetCurrentNodeId())
+		<< AnimNodeValueObject.Value(FObjectTrace::GetObjectId(InValue))
+		<< AnimNodeValueObject.KeyLength(KeyLength)
+		<< AnimNodeValueObject.FrameCounter((uint16)(GFrameCounter % 0xffff))
+		<< AnimNodeValueObject.Attachment(InKey, KeyLength * sizeof(TCHAR));
+}
+
+void FAnimTrace::OutputAnimNodeValue(const FAnimationBaseContext& InContext, const TCHAR* InKey, const UClass* InValue)
+{
+	bool bEventEnabled = UE_TRACE_EVENT_IS_ENABLED(Animation, AnimNodeValueObject);
+	if (!bEventEnabled)
 	{
-		LinkedProperty->GetFName().ToString(reinterpret_cast<TCHAR*>(Out), NameLength);
-	};
+		return;
+	}
+
+	check(InContext.AnimInstanceProxy);
+
+	TRACE_OBJECT(InContext.AnimInstanceProxy->GetAnimInstanceObject());
+	TRACE_CLASS(InValue);
+
+	int32 KeyLength = FCString::Strlen(InKey) + 1;
+
+	UE_TRACE_LOG(Animation, AnimNodeValueClass, KeyLength * sizeof(TCHAR))
+		<< AnimNodeValueClass.Cycle(FPlatformTime::Cycles64())
+		<< AnimNodeValueClass.AnimInstanceId(FObjectTrace::GetObjectId(InContext.AnimInstanceProxy->GetAnimInstanceObject()))
+		<< AnimNodeValueClass.NodeId(InContext.GetCurrentNodeId())
+		<< AnimNodeValueClass.Value(FObjectTrace::GetObjectId(InValue))
+		<< AnimNodeValueClass.KeyLength(KeyLength)
+		<< AnimNodeValueClass.FrameCounter((uint16)(GFrameCounter % 0xffff))
+		<< AnimNodeValueClass.Attachment(InKey, KeyLength * sizeof(TCHAR));
+}
+
+void FAnimTrace::OutputAnimSequencePlayer(const FAnimationBaseContext& InContext, const FAnimNode_SequencePlayer& InNode)
+{
+	bool bEventEnabled = UE_TRACE_EVENT_IS_ENABLED(Animation, AnimSequencePlayer);
+	if (!bEventEnabled)
+	{
+		return;
+	}
+
+	check(InContext.AnimInstanceProxy);
+
+	TRACE_OBJECT(InContext.AnimInstanceProxy->GetAnimInstanceObject());
+
+	UE_TRACE_LOG(Animation, AnimSequencePlayer)
+		<< AnimSequencePlayer.Cycle(FPlatformTime::Cycles64())
+		<< AnimSequencePlayer.AnimInstanceId(FObjectTrace::GetObjectId(InContext.AnimInstanceProxy->GetAnimInstanceObject()))
+		<< AnimSequencePlayer.NodeId(InContext.GetCurrentNodeId())
+		<< AnimSequencePlayer.Position(InNode.GetAccumulatedTime())
+		<< AnimSequencePlayer.Length(InNode.Sequence ? InNode.Sequence->SequenceLength : 0.0f)
+		<< AnimSequencePlayer.FrameCount(InNode.Sequence ? InNode.Sequence->GetNumberOfFrames() : 0);
+}
+
+void FAnimTrace::OutputStateMachineState(const FAnimationBaseContext& InContext, int32 InStateMachineIndex, int32 InStateIndex, float InStateWeight, float InElapsedTime)
+{
+	bool bEventEnabled = UE_TRACE_EVENT_IS_ENABLED(Animation, StateMachineState);
+	if (!bEventEnabled)
+	{
+		return;
+	}
+
+	check(InContext.AnimInstanceProxy);
+
+	TRACE_OBJECT(InContext.AnimInstanceProxy->GetAnimInstanceObject());
 
-	UE_TRACE_LOG(Animation, PoseLink)
-		<< PoseLink.AnimInstanceId(FObjectTrace::GetObjectId(InContext.AnimInstanceProxy->GetAnimInstanceObject()))
-		<< PoseLink.FrameId(GFrameCounter)
-		<< PoseLink.SourceLinkId(InContext.GetPreviousNodeId())
-		<< PoseLink.TargetLinkId(InContext.GetCurrentNodeId())
-		<< PoseLink.Weight(InContext.GetFinalBlendWeight())
-		<< PoseLink.TargetNameLength(NameLength)
-		<< PoseLink.Attachment(StringCopyFunc);
+	UE_TRACE_LOG(Animation, StateMachineState)
+		<< StateMachineState.Cycle(FPlatformTime::Cycles64())
+		<< StateMachineState.AnimInstanceId(FObjectTrace::GetObjectId(InContext.AnimInstanceProxy->GetAnimInstanceObject()))
+		<< StateMachineState.NodeId(InContext.GetCurrentNodeId())
+		<< StateMachineState.StateMachineIndex(InStateMachineIndex)
+		<< StateMachineState.StateIndex(InStateIndex)
+		<< StateMachineState.StateWeight(InStateWeight)
+		<< StateMachineState.ElapsedTime(InElapsedTime);
 }
 
 #endif
\ No newline at end of file
diff --git a/Engine/Source/Runtime/Engine/Private/ObjectTrace.cpp b/Engine/Source/Runtime/Engine/Private/ObjectTrace.cpp
--- a/Engine/Source/Runtime/Engine/Private/ObjectTrace.cpp
+++ b/Engine/Source/Runtime/Engine/Private/ObjectTrace.cpp
@@ -12,6 +12,8 @@
 #include "UObject/WeakObjectPtrTemplates.h"
 #include "UObject/WeakObjectPtr.h"
 #include "Misc/CommandLine.h"
+#include "HAL/IConsoleManager.h"
+#include "Engine/World.h"
 
 UE_TRACE_EVENT_BEGIN(Object, Class, Important)
 	UE_TRACE_EVENT_FIELD(uint64, Id)
@@ -32,26 +34,46 @@ UE_TRACE_EVENT_BEGIN(Object, ObjectEvent)
 	UE_TRACE_EVENT_FIELD(uint8, Event)
 UE_TRACE_EVENT_END()
 
+TAutoConsoleVariable<int32> CVarRecordAllWorldTypes(
+	TEXT("Insights.RecordAllWorldTypes"),
+	0,
+	TEXT("Gameplay Insights recording by default only records Game and PIE worlds.")
+	TEXT("Toggle this value to 1 to record other world types."));
+
 // Object annotations used for tracing
 struct FTracedObjectAnnotation
 {
 	FTracedObjectAnnotation()
-		: bTraced(false)
+		: Id(0)
+		, bTraced(false)
 	{}
 
+	// Object ID
+	uint64 Id;
+
 	// Whether this object has been traced this session
 	bool bTraced;
 
 	/** Determine if this annotation is default - required for annotations */
 	FORCEINLINE bool IsDefault() const
 	{
-		return !bTraced;
+		return bTraced == false && Id == 0;
 	}
 };
 
 // Object annotations used for tracing
 FUObjectAnnotationSparse<FTracedObjectAnnotation, true> GObjectTraceAnnotations;
 
+static bool ShouldTraceObjectsWorld(const  UObject* InObject)
+{
+	UWorld* World = InObject->GetWorld();
+	return 
+		CVarRecordAllWorldTypes.GetValueOnAnyThread() != 0 ||
+		World == nullptr || 
+		World->WorldType == EWorldType::Game ||
+		World->WorldType == EWorldType::PIE;
+}
+
 void FObjectTrace::Init()
 {
 	if (FParse::Param(FCommandLine::Get(), TEXT("objecttrace")))
@@ -65,23 +87,37 @@ void FObjectTrace::Init()
 
 uint64 FObjectTrace::GetObjectId(const UObject* InObject)
 {
-	// An object ID uses a combination of its own and its outer's hash
+	// An object ID uses a combination of its own and its outer's index
 	// We do this to represent objects that get renamed into different outers 
 	// as distinct traces (we don't attempt to link them).
-	// We use a weak object ptr hash to take into account object name-recycling, 
-	// and represent the new object as a distinct entry
-	uint64 ObjectHash = 0;
-	uint64 OuterHash = 0;
-	if(InObject != nullptr)
+
+	auto GetObjectIdInner = [](const UObject* InObjectInner)
+	{
+		static uint64 CurrentId = 1;
+
+		FTracedObjectAnnotation Annotation = GObjectTraceAnnotations.GetAnnotation(InObjectInner);
+		if(Annotation.Id == 0)
+		{
+			Annotation.Id = CurrentId++;
+			GObjectTraceAnnotations.AddAnnotation(InObjectInner, Annotation);
+		}
+
+		return Annotation.Id;
+	};
+
+	uint64 Id = 0;
+	uint64 OuterId = 0;
+	if(InObject)
 	{
-		ObjectHash = GetTypeHash(TWeakObjectPtr<const UObject>(InObject));
-		if(UObject* Outer = InObject->GetOuter())
+		Id = GetObjectIdInner(InObject);
+
+		if(const UObject* Outer = InObject->GetOuter())
 		{
-			OuterHash = GetTypeHash(TWeakObjectPtr<UObject>(Outer));
+			OuterId = GetObjectIdInner(Outer);
 		}
 	}
 
-	return ObjectHash | (OuterHash << 32);
+	return Id | (OuterId << 32);
 }
 
 void FObjectTrace::OutputClass(const UClass* InClass)
@@ -141,6 +177,11 @@ void FObjectTrace::OutputObject(const UObject* InObject)
 	Annotation.bTraced = true;
 	GObjectTraceAnnotations.AddAnnotation(InObject, Annotation);
 
+	if(!ShouldTraceObjectsWorld(InObject))
+	{
+		return;
+	}
+
 	// Trace the object's class first
 	TRACE_CLASS(InObject->GetClass());
 
@@ -174,6 +215,11 @@ void FObjectTrace::OutputObjectEvent(const UObject* InObject, const TCHAR* InEve
 		return;
 	}
 
+	if(!ShouldTraceObjectsWorld(InObject))
+	{
+		return;
+	}
+
 	TRACE_OBJECT(InObject);
 
 	int32 StringBufferSize = (FCString::Strlen(InEvent) + 1) * sizeof(TCHAR);
diff --git a/Engine/Source/Runtime/Engine/Public/Animation/AnimInstanceProxy.h b/Engine/Source/Runtime/Engine/Public/Animation/AnimInstanceProxy.h
--- a/Engine/Source/Runtime/Engine/Public/Animation/AnimInstanceProxy.h
+++ b/Engine/Source/Runtime/Engine/Public/Animation/AnimInstanceProxy.h
@@ -17,6 +17,7 @@
 #include "Animation/AnimClassInterface.h"
 #include "Animation/AnimBlueprintGeneratedClass.h"
 #include "Logging/TokenizedMessage.h"
+#include "Animation/AnimTrace.h"
 
 #include "AnimInstanceProxy.generated.h"
 
@@ -436,7 +437,7 @@ public:
 	void RecordMachineWeight(const int32 InMachineClassIndex, const float InMachineWeight);
 
 	float GetRecordedStateWeight(const int32 InMachineClassIndex, const int32 InStateIndex) const;
-	void RecordStateWeight(const int32 InMachineClassIndex, const int32 InStateIndex, const float InStateWeight);
+	void RecordStateWeight(const int32 InMachineClassIndex, const int32 InStateIndex, const float InStateWeight, const float InElapsedTime);
 
 	bool IsSlotNodeRelevantForNotifies(const FName& SlotNodeName) const;
 	/** Reset any dynamics running simulation-style updates (e.g. on teleport, time skip etc.) */
@@ -457,6 +458,14 @@ public:
 	/** Get the transform of the actor we are running on */
 	const FTransform& GetActorTransform() { return ActorTransform; }
 
+#if ANIM_TRACE_ENABLED
+	// Trace montage debug data for the specified slot
+	void TraceMontageEvaluationData(const FAnimationUpdateContext& InContext, const FName& InSlotName);
+#endif
+
+	/** Get the debug data for this instance's anim bp */
+	FAnimBlueprintDebugData* GetAnimBlueprintDebugData() const;
+
 	/** Only restricted classes can access the protected interface */
 	friend class UAnimInstance;
 	friend class UAnimSingleNodeInstance;
@@ -788,9 +797,6 @@ protected:
 	void AddCurveValue(const FSmartNameMapping& Mapping, const FName& CurveName, float Value);
 
 private:
-	/** Get the debug data for this instance's anim bp */
-	FAnimBlueprintDebugData* GetAnimBlueprintDebugData() const;
-
 	/** The component to world transform of the component we are running on */
 	FTransform ComponentTransform;
 
diff --git a/Engine/Source/Runtime/Engine/Public/Animation/AnimTrace.h b/Engine/Source/Runtime/Engine/Public/Animation/AnimTrace.h
--- a/Engine/Source/Runtime/Engine/Public/Animation/AnimTrace.h
+++ b/Engine/Source/Runtime/Engine/Public/Animation/AnimTrace.h
@@ -4,6 +4,7 @@
 
 #include "CoreTypes.h"
 #include "ObjectTrace.h"
+#include "CoreMinimal.h"
 
 #define ANIM_TRACE_ENABLED OBJECT_TRACE_ENABLED
 
@@ -15,17 +16,84 @@ struct FAnimationBaseContext;
 struct FAnimationUpdateContext;
 class USkeletalMesh;
 class USkeletalMeshComponent;
+struct FAnimationInitializeContext;
 struct FAnimationUpdateContext;
 struct FAnimationBaseContext;
+struct FAnimationCacheBonesContext;
+struct FPoseContext;
+struct FComponentSpacePoseContext;
 class FName;
+struct FVector;
+struct FRotator;
+struct FAnimNode_SequencePlayer;
 
 struct FAnimTrace
 {
+	/** The various phases of anim graph processing */
+	enum class EPhase : uint8
+	{
+		Initialize = 0,
+		PreUpdate = 1,
+		Update = 2,
+		CacheBones = 3,
+		Evaluate = 4,
+	};
+
+	/** Helper for outputting anim nodes */
+	struct FScopedAnimNodeTrace
+	{
+		FScopedAnimNodeTrace(const FAnimationInitializeContext& InContext);
+		FScopedAnimNodeTrace(const FAnimationUpdateContext& InContext);
+		FScopedAnimNodeTrace(const FAnimationCacheBonesContext& InContext);
+		FScopedAnimNodeTrace(const FPoseContext& InContext);
+		FScopedAnimNodeTrace(const FComponentSpacePoseContext& InContext);
+
+		~FScopedAnimNodeTrace();
+
+		const FAnimationBaseContext& Context;
+	};
+
+	/** Helper for outputting anim graphs */
+	struct FScopedAnimGraphTrace
+	{
+		FScopedAnimGraphTrace(const FAnimationInitializeContext& InContext);
+		FScopedAnimGraphTrace(const FAnimationUpdateContext& InContext);
+		FScopedAnimGraphTrace(const FAnimationCacheBonesContext& InContext);
+		FScopedAnimGraphTrace(const FPoseContext& InContext);
+		FScopedAnimGraphTrace(const FComponentSpacePoseContext& InContext);
+
+		~FScopedAnimGraphTrace();
+
+		uint64 StartCycle;
+		const FAnimationBaseContext& Context;
+		EPhase Phase;
+	};
+
+	/** Describes a debug line output to the world */
+	struct FDebugLine
+	{
+		FDebugLine(const FVector& InStartLocation, const FVector& InEndLocation, const FColor& InColor, bool bInPersistentLines = false, float InLifeTime = -1.0f, float InThickness = 0.0f)
+			: StartLocation(InStartLocation)
+			, EndLocation(InEndLocation)
+			, Color(InColor)
+			, LifeTime(InLifeTime)
+			, Thickness(InThickness)
+			, bPersistentLines(bInPersistentLines)
+		{}
+
+		FVector StartLocation;
+		FVector EndLocation;
+		FColor Color;
+		float LifeTime;
+		float Thickness;
+		bool bPersistentLines;
+	};
+
 	/** Initialize animation tracing */
 	ENGINE_API static void Init();
 
-	/** Helper function to output a set of tick records */
-	ENGINE_API static void OutputAnimTickRecords(const FAnimInstanceProxy& InProxy, const USkeletalMeshComponent* InComponent);
+	/** Helper function to output a tick record */
+	ENGINE_API static void OutputAnimTickRecord(const FAnimationBaseContext& InContext, const FAnimTickRecord& InTickRecord);
 
 	/** Helper function to output a skeletal mesh */
 	ENGINE_API static void OutputSkeletalMesh(const USkeletalMesh* InMesh);
@@ -33,25 +101,42 @@ struct FAnimTrace
 	/** Helper function to output a skeletal mesh pose, curves etc. */
 	ENGINE_API static void OutputSkeletalMeshComponent(const USkeletalMeshComponent* InComponent);
 
-	/** Helper function to output a pose link */
-	ENGINE_API static void OutputPoseLink(const FAnimationUpdateContext& InContext);
+	/** Helper function to output a skeletal mesh frame marker */
+	ENGINE_API static void OutputSkeletalMeshFrame(const USkeletalMeshComponent* InComponent);
+
+	/** Helper function to output an anim graph's execution event */
+	ENGINE_API static void OutputAnimGraph(const FAnimationBaseContext& InContext, uint64 InStartCycle, uint64 InEndCycle, uint8 InPhase);
+
+	/** Helper function to output an anim node's execution event */
+	ENGINE_API static void OutputAnimNodeStart(const FAnimationBaseContext& InContext, uint64 InStartCycle, int32 InPreviousNodeId, int32 InNodeId, float InBlendWeight, float InRootMotionWeight, uint8 InPhase);
+	ENGINE_API static void OutputAnimNodeEnd(const FAnimationBaseContext& InContext, uint64 InEndCycle);
 
 	/** Helper function to output a tracked value for an anim node */
-	ENGINE_API static void OutputAnimNodeValue(const FAnimationBaseContext& InContext, const TCHAR* InKey, bool Value);
-	ENGINE_API static void OutputAnimNodeValue(const FAnimationBaseContext& InContext, const TCHAR* InKey, int32 Value);
-	ENGINE_API static void OutputAnimNodeValue(const FAnimationBaseContext& InContext, const TCHAR* InKey, float Value);
-	ENGINE_API static void OutputAnimNodeValue(const FAnimationBaseContext& InContext, const TCHAR* InKey, const FName& Value);
-	ENGINE_API static void OutputAnimNodeValue(const FAnimationBaseContext& InContext, const TCHAR* InKey, const TCHAR* Value);
+	ENGINE_API static void OutputAnimNodeValue(const FAnimationBaseContext& InContext, const TCHAR* InKey, bool InValue);
+	ENGINE_API static void OutputAnimNodeValue(const FAnimationBaseContext& InContext, const TCHAR* InKey, int32 InValue);
+	ENGINE_API static void OutputAnimNodeValue(const FAnimationBaseContext& InContext, const TCHAR* InKey, float InValue);
+	ENGINE_API static void OutputAnimNodeValue(const FAnimationBaseContext& InContext, const TCHAR* InKey, const FVector& InValue);
+	ENGINE_API static void OutputAnimNodeValue(const FAnimationBaseContext& InContext, const TCHAR* InKey, const FRotator& InValue);
+	ENGINE_API static void OutputAnimNodeValue(const FAnimationBaseContext& InContext, const TCHAR* InKey, const FName& InValue);
+	ENGINE_API static void OutputAnimNodeValue(const FAnimationBaseContext& InContext, const TCHAR* InKey, const TCHAR* InValue);
+	ENGINE_API static void OutputAnimNodeValue(const FAnimationBaseContext& InContext, const TCHAR* InKey, const UClass* InValue);
+	ENGINE_API static void OutputAnimNodeValue(const FAnimationBaseContext& InContext, const TCHAR* InKey, const UObject* InValue);
+
+	/** Helper function to output debug info for sequence player nodes */
+	ENGINE_API static void OutputAnimSequencePlayer(const FAnimationBaseContext& InContext, const FAnimNode_SequencePlayer& InNode);
 
 	/** 
 	 * Helper function to output a name to the trace stream, referenced by ID. 
 	 * @return the ID used to reference the name
 	 */
 	ENGINE_API static uint32 OutputName(const FName& InName);
+
+	/** Helper function to output a state machine state's info */
+	ENGINE_API static void OutputStateMachineState(const FAnimationBaseContext& InContext, int32 InStateMachineIndex, int32 InStateIndex, float InStateWeight, float InElapsedTime);
 };
 
-#define TRACE_ANIM_TICK_RECORDS(Proxy, Component) \
-	FAnimTrace::OutputAnimTickRecords(Proxy, Component);
+#define TRACE_ANIM_TICK_RECORD(Context, TickRecord) \
+	FAnimTrace::OutputAnimTickRecord(Context, TickRecord);
 
 #define TRACE_SKELETAL_MESH(Mesh) \
 	FAnimTrace::OutputSkeletalMesh(Mesh);
@@ -59,18 +144,34 @@ struct FAnimTrace
 #define TRACE_SKELETAL_MESH_COMPONENT(Component) \
 	FAnimTrace::OutputSkeletalMeshComponent(Component);
 
-#define TRACE_POSE_LINK(Context) \
-	FAnimTrace::OutputPoseLink(Context);
+#define TRACE_SKELETALMESH_FRAME(Component) \
+	FAnimTrace::OutputSkeletalMeshFrame(Component);
+
+#define TRACE_SCOPED_ANIM_GRAPH(Context) \
+	FAnimTrace::FScopedAnimGraphTrace _ScopedAnimGraphTrace(Context);
+
+#define TRACE_SCOPED_ANIM_NODE(Context) \
+	FAnimTrace::FScopedAnimNodeTrace _ScopedAnimNodeTrace(Context);
 
 #define TRACE_ANIM_NODE_VALUE(Context, Key, Value) \
 	FAnimTrace::OutputAnimNodeValue(Context, Key, Value);
 
+#define TRACE_ANIM_SEQUENCE_PLAYER(Context, Node) \
+	FAnimTrace::OutputAnimSequencePlayer(Context, Node);
+
+#define TRACE_ANIM_STATE_MACHINE_STATE(Context, StateMachineIndex, StateIndex, StateWeight, ElapsedTime) \
+	FAnimTrace::OutputStateMachineState(Context, StateMachineIndex, StateIndex, StateWeight, ElapsedTime);
+
 #else
 
-#define TRACE_ANIM_TICK_RECORDS(Proxy, Component)
+#define TRACE_ANIM_TICK_RECORD(Context, TickRecord)
 #define TRACE_SKELETAL_MESH(Mesh)
 #define TRACE_SKELETAL_MESH_COMPONENT(Component)
-#define TRACE_POSE_LINK(Context)
+#define TRACE_SKELETALMESH_FRAME(Component)
+#define TRACE_SCOPED_ANIM_GRAPH(Context)
+#define TRACE_SCOPED_ANIM_NODE(Context)
 #define TRACE_ANIM_NODE_VALUE(Context, Key, Value)
+#define TRACE_ANIM_SEQUENCE_PLAYER(Context, Node)
+#define TRACE_ANIM_STATE_MACHINE_STATE(Context, StateMachineIndex, StateIndex, StateWeight, ElapsedTime)
 
 #endif
\ No newline at end of file
