Optimized iteration heavy StreamableManager functions

diff --git a/Engine/Source/Runtime/Engine/Classes/Engine/StreamableManager.h b/Engine/Source/Runtime/Engine/Classes/Engine/StreamableManager.h
--- a/Engine/Source/Runtime/Engine/Classes/Engine/StreamableManager.h
+++ b/Engine/Source/Runtime/Engine/Classes/Engine/StreamableManager.h
@@ -149,6 +149,12 @@ private:
 	/** Callback when async load finishes, it's here so we can use a shared pointer for callback safety */
 	void AsyncLoadCallbackWrapper(const FName& PackageName, UPackage* LevelPackage, EAsyncLoadingResult::Type Result, FSoftObjectPath TargetName);
 
+	/** Notify all parents that a child completed loading */
+	void NotifyParentsOfCompletion();
+
+	/** Notify all parents that a child was canceled */
+	void NotifyParentsOfCancellation();
+
 	/** Called on meta handle when a child handle has completed/canceled */
 	void UpdateCombinedHandle();
 
@@ -194,6 +200,12 @@ private:
 	/** How many FStreamables is this waiting on to finish loading */
 	int32 StreamablesLoading;
 
+	/** How many of our children that have been completed */
+	int32 CompletedChildCount = 0;
+
+	/** How many of our children that have been canceled */
+	int32 CanceledChildCount = 0;
+
 	/** List of assets that were referenced by this handle */
 	TArray<FSoftObjectPath> RequestedAssets;
 
diff --git a/Engine/Source/Runtime/Engine/Private/StreamableManager.cpp b/Engine/Source/Runtime/Engine/Private/StreamableManager.cpp
--- a/Engine/Source/Runtime/Engine/Private/StreamableManager.cpp
+++ b/Engine/Source/Runtime/Engine/Private/StreamableManager.cpp
@@ -8,6 +8,7 @@
 #include "HAL/IConsoleManager.h"
 #include "Tickable.h"
 #include "Serialization/LoadTimeTrace.h"
+#include "Containers/List.h"
 
 DEFINE_LOG_CATEGORY_STATIC(LogStreamableManager, Log, All);
 
@@ -29,8 +30,13 @@ public:
 	void AddDelegate(const FStreamableDelegate& Delegate, const FStreamableDelegate& CancelDelegate, TSharedPtr<FStreamableHandle> AssociatedHandle)
 	{
 		DataLock.Lock();
-
-		PendingDelegates.Emplace(Delegate, CancelDelegate, AssociatedHandle);
+		
+		TUniquePtr<TDoubleLinkedList<FPendingDelegate>>& DelegatesForHandle = PendingDelegates.FindOrAdd(AssociatedHandle);
+		if (!DelegatesForHandle)
+		{
+			DelegatesForHandle.Reset(new TDoubleLinkedList<FPendingDelegate>());
+		}
+		DelegatesForHandle->AddTail(FPendingDelegate(Delegate, CancelDelegate, AssociatedHandle));
 
 		DataLock.Unlock();
 	}
@@ -38,34 +44,38 @@ public:
 	/** Cancels delegate for handle, this will either delete the delegate or replace with the cancel delegate */
 	void CancelDelegatesForHandle(TSharedPtr<FStreamableHandle> AssociatedHandle)
 	{
-		if (!AssociatedHandle.IsValid())
+		FScopeLock _(&DataLock);
+		
+		TUniquePtr<TDoubleLinkedList<FPendingDelegate>>* FindIt = PendingDelegates.Find(AssociatedHandle);
+		if (!FindIt)
 		{
 			return;
 		}
 
-		DataLock.Lock();
+		TDoubleLinkedList<FPendingDelegate>* DelegatesForHandle = FindIt->Get();
 
-		for (int32 DelegateIndex = 0; DelegateIndex < PendingDelegates.Num(); DelegateIndex++)
+		TDoubleLinkedList<FPendingDelegate>::TDoubleLinkedListNode* CurrentNode = DelegatesForHandle->GetHead();
+		while (CurrentNode)
 		{
-			FPendingDelegate& PendingDelegate = PendingDelegates[DelegateIndex];
-			if (PendingDelegate.RelatedHandle == AssociatedHandle)
+			TDoubleLinkedList<FPendingDelegate>::TDoubleLinkedListNode* NextNode = CurrentNode->GetNextNode();
+			FPendingDelegate& PendingDelegate = CurrentNode->GetValue();
+			if (PendingDelegate.CancelDelegate.IsBound())
 			{
-				if (PendingDelegate.CancelDelegate.IsBound())
-				{
-					// Replace with cancel delegate
-					PendingDelegate.Delegate = PendingDelegate.CancelDelegate;
-					PendingDelegate.CancelDelegate.Unbind();
-				}
-				else
-				{
-					// Remove entirely
-					PendingDelegates.RemoveAt(DelegateIndex);
-					DelegateIndex--;
-				}
+				// Replace with cancel delegate
+				PendingDelegate.Delegate = PendingDelegate.CancelDelegate;
+				PendingDelegate.CancelDelegate.Unbind();
+			}
+			else
+			{
+				// Remove entirely
+				DelegatesForHandle->RemoveNode(CurrentNode);
 			}
+			CurrentNode = NextNode;
+		}
+		if (DelegatesForHandle->Num() == 0)
+		{
+			PendingDelegates.Remove(AssociatedHandle);
 		}
-
-		DataLock.Unlock();
 	}
 
 	/** Calls all delegates, call from synchronous flushes */
@@ -88,15 +98,26 @@ public:
 
 		DataLock.Lock();
 
-		TArray<FPendingDelegate> DelegatesToCall;
-		DelegatesToCall.Reserve(PendingDelegates.Num());
-		for (int32 DelegateIndex = 0; DelegateIndex < PendingDelegates.Num(); DelegateIndex++)
+		TDoubleLinkedList<FPendingDelegate> DelegatesToCall;
+		for (auto It = PendingDelegates.CreateIterator(); It; ++It)
 		{
-			if (--PendingDelegates[DelegateIndex].DelayFrames <= 0)
+			TDoubleLinkedList<FPendingDelegate>* DelegatesForHandle = It->Value.Get();
+			TDoubleLinkedList<FPendingDelegate>::TDoubleLinkedListNode* CurrentNode = DelegatesForHandle->GetHead();
+			while (CurrentNode)
 			{
-				// Add to call array and remove from tracking one
-				DelegatesToCall.Emplace(PendingDelegates[DelegateIndex].Delegate, FStreamableDelegate(), PendingDelegates[DelegateIndex].RelatedHandle);
-				PendingDelegates.RemoveAt(DelegateIndex--);
+				TDoubleLinkedList<FPendingDelegate>::TDoubleLinkedListNode* NextNode = CurrentNode->GetNextNode();
+				FPendingDelegate& PendingDelegate = CurrentNode->GetValue();
+				if (--PendingDelegate.DelayFrames <= 0)
+				{
+					// Add to call array and remove from tracking one
+					DelegatesToCall.AddTail(CurrentNode);
+					DelegatesForHandle->RemoveNode(CurrentNode, false);
+				}
+				CurrentNode = NextNode;
+			}
+			if (DelegatesForHandle->Num() == 0)
+			{
+				It.RemoveCurrent();
 			}
 		}
 
@@ -107,6 +128,7 @@ public:
 			// Call delegates, these may add other deferred delegates
 			PendingDelegate.Delegate.ExecuteIfBound();
 		}
+		DelegatesToCall.Empty();
 
 		// When DelegatesToCall falls out of scope it may delete the referenced handles
 	}
@@ -135,6 +157,8 @@ private:
 
 	struct FPendingDelegate
 	{
+		FPendingDelegate* Next = nullptr;
+		
 		/** Delegate to call when frames are up */
 		FStreamableDelegate Delegate;
 
@@ -155,7 +179,7 @@ private:
 		{}
 	};
 
-	TArray<FPendingDelegate> PendingDelegates;
+	TMap<TSharedPtr<FStreamableHandle>, TUniquePtr<TDoubleLinkedList<FPendingDelegate>>> PendingDelegates;
 
 	FCriticalSection DataLock;
 };
@@ -397,10 +421,12 @@ void FStreamableHandle::CancelHandle()
 		}
 
 		bCanceled = true;
+		NotifyParentsOfCancellation();
 		return;
 	}
 
 	bCanceled = true;
+	NotifyParentsOfCancellation();
 
 	ExecuteDelegate(CancelDelegate, SharedThis);
 	UnbindDelegates();
@@ -425,6 +451,8 @@ void FStreamableHandle::CancelHandle()
 	}
 
 	ChildHandles.Empty();
+	CompletedChildCount = 0;
+	CanceledChildCount = 0;
 
 	OwningManager = nullptr;
 
@@ -478,6 +506,8 @@ void FStreamableHandle::ReleaseHandle()
 		}
 
 		ChildHandles.Empty();
+		CompletedChildCount = 0;
+		CanceledChildCount = 0;
 
 		OwningManager = nullptr;
 	}
@@ -525,6 +555,8 @@ void FStreamableHandle::CompleteLoad()
 		ExecuteDelegate(CompleteDelegate, AsShared(), CancelDelegate);
 		UnbindDelegates();
 
+		NotifyParentsOfCompletion();
+
 		if (ParentHandles.Num() > 0)
 		{
 			// Update any meta handles that are still active. Copy the array first as elements may be removed from original while iterating
@@ -542,6 +574,42 @@ void FStreamableHandle::CompleteLoad()
 	}
 }
 
+void FStreamableHandle::NotifyParentsOfCompletion()
+{
+	if (ParentHandles.Num() > 0)
+	{
+		// Update any meta handles that are still active. Copy the array first as elements may be removed from original while iterating
+		TArray<TWeakPtr<FStreamableHandle>> ParentHandlesCopy = ParentHandles;
+		for (TWeakPtr<FStreamableHandle>& WeakHandle : ParentHandlesCopy)
+		{
+			TSharedPtr<FStreamableHandle> Handle = WeakHandle.Pin();
+
+			if (Handle.IsValid())
+			{
+				++Handle->CompletedChildCount;
+			}
+		}
+	}
+}
+
+void FStreamableHandle::NotifyParentsOfCancellation()
+{
+	if (ParentHandles.Num() > 0)
+	{
+		// Update any meta handles that are still active. Copy the array first as elements may be removed from original while iterating
+		TArray<TWeakPtr<FStreamableHandle>> ParentHandlesCopy = ParentHandles;
+		for (TWeakPtr<FStreamableHandle>& WeakHandle : ParentHandlesCopy)
+		{
+			TSharedPtr<FStreamableHandle> Handle = WeakHandle.Pin();
+
+			if (Handle.IsValid())
+			{
+				++Handle->CanceledChildCount;
+			}
+		}
+	}
+}
+
 void FStreamableHandle::UpdateCombinedHandle()
 {
 	if (!IsActive())
@@ -554,6 +622,11 @@ void FStreamableHandle::UpdateCombinedHandle()
 		return;
 	}
 
+	if (CompletedChildCount + CanceledChildCount < ChildHandles.Num())
+	{
+		return;
+	}
+
 	// Check all our children, complete if done
 	bool bAllCompleted = true;
 	bool bAllCanceled = true;
@@ -708,6 +781,9 @@ struct FStreamable
 					FStreamableHandle::ExecuteDelegate(ActiveHandle->CancelDelegate, ActiveHandle);
 					ActiveHandle->UnbindDelegates();
 				}
+
+				ActiveHandle->NotifyParentsOfCancellation();
+
 			}
 		}
 		ActiveHandles.Empty();
@@ -1345,6 +1421,14 @@ TSharedPtr<FStreamableHandle> FStreamableManager::CreateCombinedHandle(const TAr
 
 		ChildHandle->ParentHandles.Add(NewRequest);
 		NewRequest->ChildHandles.Add(ChildHandle);
+		if (ChildHandle->bLoadCompleted)
+		{
+			++NewRequest->CompletedChildCount;
+		}
+		if (ChildHandle->bCanceled)
+		{
+			++NewRequest->CanceledChildCount;
+		}
 	}
 
 	// Add to pending list so these handles don't free when not referenced
