[UE-84853] Fix WriteTableAsJson to allow writing utf8 so that FillDataTableFromJSONFile and applications other than unreal can recognize and parse the file type.
Add ImportRowStruct argument to FillDataTableFromJSONFile so that it can be used in automatedtests, which require the ImportRowStruct to be passed in.
#rn minor

diff --git a/Engine/Source/Runtime/Engine/Classes/Engine/CurveTable.h b/Engine/Source/Runtime/Engine/Classes/Engine/CurveTable.h
--- a/Engine/Source/Runtime/Engine/Classes/Engine/CurveTable.h
+++ b/Engine/Source/Runtime/Engine/Classes/Engine/CurveTable.h
@@ -162,7 +162,8 @@ public:
 	ENGINE_API FString GetTableAsJSON() const;
 
 	/** Output entire contents of table as JSON. bAsArray true will write is as a JSON array, false will write it as a series of named objects*/
-	ENGINE_API bool WriteTableAsJSON(const TSharedRef< TJsonWriter<TCHAR, TPrettyJsonPrintPolicy<TCHAR> > >& JsonWriter,bool bAsArray = true) const;
+	template <typename CharType = TCHAR>
+	ENGINE_API bool WriteTableAsJSON(const TSharedRef< TJsonWriter<CharType, TPrettyJsonPrintPolicy<CharType> > >& JsonWriter,bool bAsArray = true) const;
 
 	/** 
 	 *	Create table from CSV style comma-separated string. 
diff --git a/Engine/Source/Runtime/Engine/Classes/Engine/DataTable.h b/Engine/Source/Runtime/Engine/Classes/Engine/DataTable.h
--- a/Engine/Source/Runtime/Engine/Classes/Engine/DataTable.h
+++ b/Engine/Source/Runtime/Engine/Classes/Engine/DataTable.h
@@ -317,13 +317,16 @@ public:
 	ENGINE_API FString GetTableAsJSON(const EDataTableExportFlags InDTExportFlags = EDataTableExportFlags::None) const;
 
 	/** Output entire contents of table as JSON */
-	ENGINE_API bool WriteTableAsJSON(const TSharedRef< TJsonWriter<TCHAR, TPrettyJsonPrintPolicy<TCHAR> > >& JsonWriter, const EDataTableExportFlags InDTExportFlags = EDataTableExportFlags::None) const;
+	template<typename CharType = TCHAR>
+	ENGINE_API bool WriteTableAsJSON(const TSharedRef< TJsonWriter<CharType, TPrettyJsonPrintPolicy<CharType> > >& JsonWriter, const EDataTableExportFlags InDTExportFlags = EDataTableExportFlags::None) const;
 
 	/** Output entire contents of table as a JSON Object*/
-	ENGINE_API bool WriteTableAsJSONObject(const TSharedRef< TJsonWriter<TCHAR, TPrettyJsonPrintPolicy<TCHAR> > >& JsonWriter, const EDataTableExportFlags InDTExportFlags = EDataTableExportFlags::None) const;
+	template<typename CharType = TCHAR>
+	ENGINE_API bool WriteTableAsJSONObject(const TSharedRef< TJsonWriter<CharType, TPrettyJsonPrintPolicy<CharType> > >& JsonWriter, const EDataTableExportFlags InDTExportFlags = EDataTableExportFlags::None) const;
 
 	/** Output the fields from a particular row (use RowMap to get RowData) to an existing JsonWriter */
-	ENGINE_API bool WriteRowAsJSON(const TSharedRef< TJsonWriter<TCHAR, TPrettyJsonPrintPolicy<TCHAR> > >& JsonWriter, const void* RowData, const EDataTableExportFlags InDTExportFlags = EDataTableExportFlags::None) const;
+	template<typename CharType = TCHAR>
+	ENGINE_API bool WriteRowAsJSON(const TSharedRef< TJsonWriter<CharType, TPrettyJsonPrintPolicy<CharType> > >& JsonWriter, const void* RowData, const EDataTableExportFlags InDTExportFlags = EDataTableExportFlags::None) const;
 
 	/** Copies all the import options from another table, this does not copy row dawta */
 	ENGINE_API bool CopyImportOptions(UDataTable* SourceTable);
diff --git a/Engine/Source/Runtime/Engine/Classes/Kismet/DataTableFunctionLibrary.h b/Engine/Source/Runtime/Engine/Classes/Kismet/DataTableFunctionLibrary.h
--- a/Engine/Source/Runtime/Engine/Classes/Kismet/DataTableFunctionLibrary.h
+++ b/Engine/Source/Runtime/Engine/Classes/Kismet/DataTableFunctionLibrary.h
@@ -138,6 +138,6 @@ class ENGINE_API UDataTableFunctionLibrary : public UBlueprintFunctionLibrary
 	 * @return	True if the operation succeeds, check the log for errors if it didn't succeed.
 	 */
 	UFUNCTION(BlueprintCallable, Category = "Editor Scripting | DataTable", DisplayName = "Fill Data Table from JSON File")
-	static bool FillDataTableFromJSONFile(UDataTable* DataTable, const FString& JSONFilePath);
+	static bool FillDataTableFromJSONFile(UDataTable* DataTable, const FString& JSONFilePath, UScriptStruct* ImportRowStruct = nullptr);
 #endif //WITH_EDITOR
 };
diff --git a/Engine/Source/Runtime/Engine/Private/CurveTable.cpp b/Engine/Source/Runtime/Engine/Private/CurveTable.cpp
--- a/Engine/Source/Runtime/Engine/Private/CurveTable.cpp
+++ b/Engine/Source/Runtime/Engine/Private/CurveTable.cpp
@@ -427,8 +427,8 @@ FString UCurveTable::GetTableAsJSON() const
 	return Result;
 }
 
-template<class T>
-void WriteTableAsJSON_Internal(const TMap<FName, T*>& RowMap, const TSharedRef< TJsonWriter<TCHAR, TPrettyJsonPrintPolicy<TCHAR> > >& JsonWriter, bool bAsArray)
+template<typename CharType, class T>
+void WriteTableAsJSON_Internal(const TMap<FName, T*>& RowMap, const TSharedRef< TJsonWriter<CharType, TPrettyJsonPrintPolicy<CharType> > >& JsonWriter, bool bAsArray)
 {
 	TArray<FName> Names;
 	TArray<T*> Curves;
@@ -481,7 +481,8 @@ void WriteTableAsJSON_Internal(const TMap<FName, T*>& RowMap, const TSharedRef<
 	}
 }
 
-bool UCurveTable::WriteTableAsJSON(const TSharedRef< TJsonWriter<TCHAR, TPrettyJsonPrintPolicy<TCHAR> > >& JsonWriter, bool bAsArray) const
+template <typename CharType>
+bool UCurveTable::WriteTableAsJSON(const TSharedRef< TJsonWriter<CharType, TPrettyJsonPrintPolicy<CharType> > >& JsonWriter, bool bAsArray) const
 {
 	if (RowMap.Num() <= 0)
 	{
@@ -490,16 +491,19 @@ bool UCurveTable::WriteTableAsJSON(const TSharedRef< TJsonWriter<TCHAR, TPrettyJ
 
 	if (CurveTableMode == ECurveTableMode::SimpleCurves)
 	{
-		WriteTableAsJSON_Internal<FSimpleCurve>(*reinterpret_cast<const TMap<FName, FSimpleCurve*>*>(&RowMap), JsonWriter, bAsArray);
+		WriteTableAsJSON_Internal<CharType, FSimpleCurve>(*reinterpret_cast<const TMap<FName, FSimpleCurve*>*>(&RowMap), JsonWriter, bAsArray);
 	}
 	else
 	{
-		WriteTableAsJSON_Internal<FRichCurve>(*reinterpret_cast<const TMap<FName, FRichCurve*>*>(&RowMap), JsonWriter, bAsArray);
+		WriteTableAsJSON_Internal<CharType, FRichCurve>(*reinterpret_cast<const TMap<FName, FRichCurve*>*>(&RowMap), JsonWriter, bAsArray);
 	}
 
 	return true;
 }
 
+template ENGINE_API bool UCurveTable::WriteTableAsJSON<TCHAR>(const TSharedRef< TJsonWriter<TCHAR, TPrettyJsonPrintPolicy<TCHAR> > >& JsonWriter, bool bAsArray) const;
+template ENGINE_API bool UCurveTable::WriteTableAsJSON<ANSICHAR>(const TSharedRef< TJsonWriter<ANSICHAR, TPrettyJsonPrintPolicy<ANSICHAR> > >& JsonWriter, bool bAsArray) const;
+
 void UCurveTable::EmptyTable()
 {
 	CURVETABLE_CHANGE_SCOPE();
diff --git a/Engine/Source/Runtime/Engine/Private/DataTable.cpp b/Engine/Source/Runtime/Engine/Private/DataTable.cpp
--- a/Engine/Source/Runtime/Engine/Private/DataTable.cpp
+++ b/Engine/Source/Runtime/Engine/Private/DataTable.cpp
@@ -554,11 +554,15 @@ FString UDataTable::GetTableAsJSON(const EDataTableExportFlags InDTExportFlags)
 	return Result;
 }
 
-bool UDataTable::WriteRowAsJSON(const TSharedRef< TJsonWriter<TCHAR, TPrettyJsonPrintPolicy<TCHAR> > >& JsonWriter, const void* RowData, const EDataTableExportFlags InDTExportFlags) const
+template<typename CharType>
+bool UDataTable::WriteRowAsJSON(const TSharedRef< TJsonWriter<CharType, TPrettyJsonPrintPolicy<CharType> > >& JsonWriter, const void* RowData, const EDataTableExportFlags InDTExportFlags) const
 {
-	return FDataTableExporterJSON(InDTExportFlags, JsonWriter).WriteRow(RowStruct, RowData);
+	return TDataTableExporterJSON<CharType>(InDTExportFlags, JsonWriter).WriteRow(RowStruct, RowData);
 }
 
+template ENGINE_API bool UDataTable::WriteRowAsJSON<TCHAR>(const TSharedRef< TJsonWriter<TCHAR, TPrettyJsonPrintPolicy<TCHAR> > >& JsonWriter, const void* RowData, const EDataTableExportFlags InDTExportFlags) const;
+template ENGINE_API bool UDataTable::WriteRowAsJSON<ANSICHAR>(const TSharedRef< TJsonWriter<ANSICHAR, TPrettyJsonPrintPolicy<ANSICHAR> > >& JsonWriter, const void* RowData, const EDataTableExportFlags InDTExportFlags) const;
+
 bool UDataTable::CopyImportOptions(UDataTable* SourceTable)
 {
 	// Only safe to call on an empty table
@@ -586,15 +590,23 @@ bool UDataTable::CopyImportOptions(UDataTable* SourceTable)
 	return true;
 }
 
-bool UDataTable::WriteTableAsJSON(const TSharedRef< TJsonWriter<TCHAR, TPrettyJsonPrintPolicy<TCHAR> > >& JsonWriter, const EDataTableExportFlags InDTExportFlags) const
+template<typename CharType>
+bool UDataTable::WriteTableAsJSON(const TSharedRef< TJsonWriter<CharType, TPrettyJsonPrintPolicy<CharType> > >& JsonWriter, const EDataTableExportFlags InDTExportFlags) const
 {
-	return FDataTableExporterJSON(InDTExportFlags, JsonWriter).WriteTable(*this);
+	return TDataTableExporterJSON<CharType>(InDTExportFlags, JsonWriter).WriteTable(*this);
 }
 
-bool UDataTable::WriteTableAsJSONObject(const TSharedRef< TJsonWriter<TCHAR, TPrettyJsonPrintPolicy<TCHAR> > >& JsonWriter, const EDataTableExportFlags InDTExportFlags) const
+template ENGINE_API bool UDataTable::WriteTableAsJSON<TCHAR>(const TSharedRef< TJsonWriter<TCHAR, TPrettyJsonPrintPolicy<TCHAR> > >& JsonWriter, const EDataTableExportFlags InDTExportFlags) const;
+template ENGINE_API bool UDataTable::WriteTableAsJSON<ANSICHAR>(const TSharedRef< TJsonWriter<ANSICHAR, TPrettyJsonPrintPolicy<ANSICHAR> > >& JsonWriter, const EDataTableExportFlags InDTExportFlags) const;
+
+template<typename CharType>
+bool UDataTable::WriteTableAsJSONObject(const TSharedRef< TJsonWriter<CharType, TPrettyJsonPrintPolicy<CharType> > >& JsonWriter, const EDataTableExportFlags InDTExportFlags) const
 {
-	return FDataTableExporterJSON(InDTExportFlags, JsonWriter).WriteTableAsObject(*this);
+	return TDataTableExporterJSON<CharType>(InDTExportFlags, JsonWriter).WriteTableAsObject(*this);
 }
+
+template ENGINE_API bool UDataTable::WriteTableAsJSONObject<TCHAR>(const TSharedRef< TJsonWriter<TCHAR, TPrettyJsonPrintPolicy<TCHAR> > >& JsonWriter, const EDataTableExportFlags InDTExportFlags) const;
+template ENGINE_API bool UDataTable::WriteTableAsJSONObject<ANSICHAR>(const TSharedRef< TJsonWriter<ANSICHAR, TPrettyJsonPrintPolicy<ANSICHAR> > >& JsonWriter, const EDataTableExportFlags InDTExportFlags) const;
 #endif
 
 TArray<UProperty*> UDataTable::GetTablePropertyArray(const TArray<const TCHAR*>& Cells, UStruct* InRowStruct, TArray<FString>& OutProblems, int32 KeyColumn)
diff --git a/Engine/Source/Runtime/Engine/Private/DataTableFunctionLibrary.cpp b/Engine/Source/Runtime/Engine/Private/DataTableFunctionLibrary.cpp
--- a/Engine/Source/Runtime/Engine/Private/DataTableFunctionLibrary.cpp
+++ b/Engine/Source/Runtime/Engine/Private/DataTableFunctionLibrary.cpp
@@ -170,7 +170,7 @@ bool UDataTableFunctionLibrary::FillDataTableFromJSONString(UDataTable* DataTabl
 	return Result != nullptr && !bWasCancelled;
 }
 
-bool UDataTableFunctionLibrary::FillDataTableFromJSONFile(UDataTable* DataTable, const FString& InFilePath)
+bool UDataTableFunctionLibrary::FillDataTableFromJSONFile(UDataTable* DataTable, const FString& InFilePath, UScriptStruct* ImportRowStruct)
 {
 	if (!DataTable)
 	{
@@ -192,6 +192,7 @@ bool UDataTableFunctionLibrary::FillDataTableFromJSONFile(UDataTable* DataTable,
 
 	DataTable->AssetImportData->Update(InFilePath);
 	UCSVImportFactory* ImportFactory = NewObject<UCSVImportFactory>();
+	ImportFactory->AutomatedImportSettings.ImportRowStruct = ImportRowStruct;
 	return ImportFactory->ReimportCSV(DataTable) == EReimportResult::Succeeded;
 }
 #endif //WITH_EDITOR
\ No newline at end of file
diff --git a/Engine/Source/Runtime/Engine/Private/DataTableJSON.cpp b/Engine/Source/Runtime/Engine/Private/DataTableJSON.cpp
--- a/Engine/Source/Runtime/Engine/Private/DataTableJSON.cpp
+++ b/Engine/Source/Runtime/Engine/Private/DataTableJSON.cpp
@@ -36,7 +36,8 @@ namespace
 	}
 #if WITH_EDITOR
 
-	void WriteJSONObjectStartWithOptionalIdentifier(FDataTableExporterJSON::FDataTableJsonWriter& InJsonWriter, const FString* InIdentifier)
+	template <typename CharType>
+	void WriteJSONObjectStartWithOptionalIdentifier(typename TDataTableExporterJSON<CharType>::FDataTableJsonWriter& InJsonWriter, const FString* InIdentifier)
 	{
 		if (InIdentifier)
 		{
@@ -48,8 +49,8 @@ namespace
 		}
 	}
 
-	template <typename ValueType>
-	void WriteJSONValueWithOptionalIdentifier(FDataTableExporterJSON::FDataTableJsonWriter& InJsonWriter, const FString* InIdentifier, const ValueType InValue)
+	template <typename CharType, typename ValueType>
+	void WriteJSONValueWithOptionalIdentifier(typename TDataTableExporterJSON<CharType>::FDataTableJsonWriter& InJsonWriter, const FString* InIdentifier, const ValueType InValue)
 	{
 		if (InIdentifier)
 		{
@@ -82,20 +83,21 @@ FString DataTableJSONUtils::GetKeyFieldName(const UDataTable& InDataTable)
 #if WITH_EDITOR
 
 FDataTableExporterJSON::FDataTableExporterJSON(const EDataTableExportFlags InDTExportFlags, FString& OutExportText)
-	: DTExportFlags(InDTExportFlags)
-	, JsonWriter(TJsonWriterFactory<TCHAR, TPrettyJsonPrintPolicy<TCHAR>>::Create(&OutExportText))
-	, bJsonWriterNeedsClose(true)
+	: TDataTableExporterJSON<TCHAR>(InDTExportFlags, TJsonWriterFactory<TCHAR, TPrettyJsonPrintPolicy<TCHAR>>::Create(&OutExportText))
 {
+	bJsonWriterNeedsClose = true;
 }
 
-FDataTableExporterJSON::FDataTableExporterJSON(const EDataTableExportFlags InDTExportFlags, TSharedRef<FDataTableJsonWriter> InJsonWriter)
+template<typename CharType>
+TDataTableExporterJSON<CharType>::TDataTableExporterJSON(const EDataTableExportFlags InDTExportFlags, TSharedRef<FDataTableJsonWriter> InJsonWriter)
 	: DTExportFlags(InDTExportFlags)
 	, JsonWriter(InJsonWriter)
 	, bJsonWriterNeedsClose(false)
 {
 }
 
-FDataTableExporterJSON::~FDataTableExporterJSON()
+template<typename CharType>
+TDataTableExporterJSON<CharType>::~TDataTableExporterJSON()
 {
 	if (bJsonWriterNeedsClose)
 	{
@@ -103,7 +105,8 @@ FDataTableExporterJSON::~FDataTableExporterJSON()
 	}
 }
 
-bool FDataTableExporterJSON::WriteTable(const UDataTable& InDataTable)
+template<typename CharType>
+bool TDataTableExporterJSON<CharType>::WriteTable(const UDataTable& InDataTable)
 {
 	if (!InDataTable.RowStruct)
 	{
@@ -134,7 +137,8 @@ bool FDataTableExporterJSON::WriteTable(const UDataTable& InDataTable)
 	return true;
 }
 
-bool FDataTableExporterJSON::WriteTableAsObject(const UDataTable& InDataTable)
+template<typename CharType>
+bool TDataTableExporterJSON<CharType>::WriteTableAsObject(const UDataTable& InDataTable)
 {
 	if (!InDataTable.RowStruct)
 	{
@@ -161,7 +165,8 @@ bool FDataTableExporterJSON::WriteTableAsObject(const UDataTable& InDataTable)
 	return true;
 }
 
-bool FDataTableExporterJSON::WriteRow(const UScriptStruct* InRowStruct, const void* InRowData, const FString* FieldToSkip)
+template<typename CharType>
+bool TDataTableExporterJSON<CharType>::WriteRow(const UScriptStruct* InRowStruct, const void* InRowData, const FString* FieldToSkip)
 {
 	if (!InRowStruct)
 	{
@@ -171,7 +176,8 @@ bool FDataTableExporterJSON::WriteRow(const UScriptStruct* InRowStruct, const vo
 	return WriteStruct(InRowStruct, InRowData, FieldToSkip);
 }
 
-bool FDataTableExporterJSON::WriteStruct(const UScriptStruct* InStruct, const void* InStructData, const FString* FieldToSkip)
+template<typename CharType>
+bool TDataTableExporterJSON<CharType>::WriteStruct(const UScriptStruct* InStruct, const void* InStructData, const FString* FieldToSkip)
 {
 	for (TFieldIterator<const UProperty> It(InStruct); It; ++It)
 	{
@@ -207,7 +213,8 @@ bool FDataTableExporterJSON::WriteStruct(const UScriptStruct* InStruct, const vo
 	return true;
 }
 
-bool FDataTableExporterJSON::WriteStructEntry(const void* InRowData, const UProperty* InProperty, const void* InPropertyData)
+template<typename CharType>
+bool TDataTableExporterJSON<CharType>::WriteStructEntry(const void* InRowData, const UProperty* InProperty, const void* InPropertyData)
 {
 	const FString Identifier = DataTableUtils::GetPropertyExportName(InProperty, DTExportFlags);
 
@@ -311,48 +318,49 @@ bool FDataTableExporterJSON::WriteStructEntry(const void* InRowData, const UProp
 	return true;
 }
 
-bool FDataTableExporterJSON::WriteContainerEntry(const UProperty* InProperty, const void* InPropertyData, const FString* InIdentifier)
+template<typename CharType>
+bool TDataTableExporterJSON<CharType>::WriteContainerEntry(const UProperty* InProperty, const void* InPropertyData, const FString* InIdentifier)
 {
 	if (const UEnumProperty* EnumProp = Cast<const UEnumProperty>(InProperty))
 	{
 		const FString PropertyValue = DataTableUtils::GetPropertyValueAsStringDirect(InProperty, (uint8*)InPropertyData, DTExportFlags);
-		WriteJSONValueWithOptionalIdentifier(*JsonWriter, InIdentifier, PropertyValue);
+		WriteJSONValueWithOptionalIdentifier<CharType>(*JsonWriter, InIdentifier, PropertyValue);
 	}
 	else if (const UNumericProperty *NumProp = Cast<const UNumericProperty>(InProperty))
 	{
 		if (NumProp->IsEnum())
 		{
 			const FString PropertyValue = DataTableUtils::GetPropertyValueAsStringDirect(InProperty, (uint8*)InPropertyData, DTExportFlags);
-			WriteJSONValueWithOptionalIdentifier(*JsonWriter, InIdentifier, PropertyValue);
+			WriteJSONValueWithOptionalIdentifier<CharType>(*JsonWriter, InIdentifier, PropertyValue);
 		}
 		else if (NumProp->IsInteger())
 		{
 			const int64 PropertyValue = NumProp->GetSignedIntPropertyValue(InPropertyData);
-			WriteJSONValueWithOptionalIdentifier(*JsonWriter, InIdentifier, PropertyValue);
+			WriteJSONValueWithOptionalIdentifier<CharType>(*JsonWriter, InIdentifier, PropertyValue);
 		}
 		else
 		{
 			const double PropertyValue = NumProp->GetFloatingPointPropertyValue(InPropertyData);
-			WriteJSONValueWithOptionalIdentifier(*JsonWriter, InIdentifier, PropertyValue);
+			WriteJSONValueWithOptionalIdentifier<CharType>(*JsonWriter, InIdentifier, PropertyValue);
 		}
 	}
 	else if (const UBoolProperty* BoolProp = Cast<const UBoolProperty>(InProperty))
 	{
 		const bool PropertyValue = BoolProp->GetPropertyValue(InPropertyData);
-		WriteJSONValueWithOptionalIdentifier(*JsonWriter, InIdentifier, PropertyValue);
+		WriteJSONValueWithOptionalIdentifier<CharType>(*JsonWriter, InIdentifier, PropertyValue);
 	}
 	else if (const UStructProperty* StructProp = Cast<const UStructProperty>(InProperty))
 	{
 		if (!!(DTExportFlags & EDataTableExportFlags::UseJsonObjectsForStructs))
 		{
-			WriteJSONObjectStartWithOptionalIdentifier(*JsonWriter, InIdentifier);
+			WriteJSONObjectStartWithOptionalIdentifier<CharType>(*JsonWriter, InIdentifier);
 			WriteStruct(StructProp->Struct, InPropertyData);
 			JsonWriter->WriteObjectEnd();
 		}
 		else
 		{
 			const FString PropertyValue = DataTableUtils::GetPropertyValueAsStringDirect(InProperty, (uint8*)InPropertyData, DTExportFlags);
-			WriteJSONValueWithOptionalIdentifier(*JsonWriter, InIdentifier, PropertyValue);
+			WriteJSONValueWithOptionalIdentifier<CharType>(*JsonWriter, InIdentifier, PropertyValue);
 		}
 	}
 	else if (const UArrayProperty* ArrayProp = Cast<const UArrayProperty>(InProperty))
@@ -373,12 +381,15 @@ bool FDataTableExporterJSON::WriteContainerEntry(const UProperty* InProperty, co
 	else
 	{
 		const FString PropertyValue = DataTableUtils::GetPropertyValueAsStringDirect(InProperty, (uint8*)InPropertyData, DTExportFlags);
-		WriteJSONValueWithOptionalIdentifier(*JsonWriter, InIdentifier, PropertyValue);
+		WriteJSONValueWithOptionalIdentifier<CharType>(*JsonWriter, InIdentifier, PropertyValue);
 	}
 
 	return true;
 }
 
+template class TDataTableExporterJSON<TCHAR>;
+template class TDataTableExporterJSON<ANSICHAR>;
+
 #endif // WITH_EDITOR
 
 
diff --git a/Engine/Source/Runtime/Engine/Private/DataTableJSON.h b/Engine/Source/Runtime/Engine/Private/DataTableJSON.h
--- a/Engine/Source/Runtime/Engine/Private/DataTableJSON.h
+++ b/Engine/Source/Runtime/Engine/Private/DataTableJSON.h
@@ -26,16 +26,18 @@ namespace DataTableJSONUtils
 
 #if WITH_EDITOR
 
-class FDataTableExporterJSON
+/** 
+ * Class to serialize a DataTable to a TJsonWriter.
+ */
+template<typename CharType = TCHAR>
+class TDataTableExporterJSON
 {
 public:
-	typedef TJsonWriter<TCHAR, TPrettyJsonPrintPolicy<TCHAR>> FDataTableJsonWriter;
+	typedef TJsonWriter<CharType, TPrettyJsonPrintPolicy<CharType>> FDataTableJsonWriter;
 
-	FDataTableExporterJSON(const EDataTableExportFlags InDTExportFlags, FString& OutExportText);
-
-	FDataTableExporterJSON(const EDataTableExportFlags InDTExportFlags, TSharedRef<FDataTableJsonWriter> InJsonWriter);
+	TDataTableExporterJSON(const EDataTableExportFlags InDTExportFlags, TSharedRef<FDataTableJsonWriter> InJsonWriter);
 
-	~FDataTableExporterJSON();
+	~TDataTableExporterJSON();
 
 	/** Writes the data table out as an array of objects */
 	bool WriteTable(const UDataTable& InDataTable);
@@ -49,7 +51,7 @@ public:
 	/** Writes the contents of a single row */
 	bool WriteStruct(const UScriptStruct* InStruct, const void* InStructData, const FString* FieldToSkip = nullptr);
 
-private:
+protected:
 	bool WriteStructEntry(const void* InRowData, const UProperty* InProperty, const void* InPropertyData);
 
 	bool WriteContainerEntry(const UProperty* InProperty, const void* InPropertyData, const FString* InIdentifier = nullptr);
@@ -59,6 +61,17 @@ private:
 	bool bJsonWriterNeedsClose;
 };
 
+/**
+ * TCHAR-specific instantiation of TDataTableExporterJSON that has a convenience constructor to write output to an FString instead of an external TJsonWriter
+ */
+class FDataTableExporterJSON : public TDataTableExporterJSON<TCHAR>
+{
+public:
+	using TDataTableExporterJSON<TCHAR>::TDataTableExporterJSON;
+// 
+	FDataTableExporterJSON(const EDataTableExportFlags InDTExportFlags, FString& OutExportText);
+};
+
 #endif // WITH_EDITOR
 
 class FDataTableImporterJSON
