Improve cook times by avoiding re-computing the static parameter values twice during calls to FMaterial::CacheShaders (once in FMaterialResource::GetShaderMapId and again in FMaterialResource::GetStaticParameterSet) by allowing the MaterialInstance to store a cached parameter value set and re-use it in selected scopes (only applied to FMaterial::CacheShaders for now, but may be applicable elsewhere).

Testing:
1) QAGame produced same WindowsNoEditor cooked data (barring metadata that has some non-determinism issues) with/without this change.
2) InfiltratorDemo looked visually correct on Win64 with and without cooked data

diff --git a/Engine/Source/Runtime/Engine/Classes/Materials/MaterialInstance.h b/Engine/Source/Runtime/Engine/Classes/Materials/MaterialInstance.h
--- a/Engine/Source/Runtime/Engine/Classes/Materials/MaterialInstance.h
+++ b/Engine/Source/Runtime/Engine/Classes/Materials/MaterialInstance.h
@@ -429,6 +429,11 @@ private:
 	UPROPERTY()
 	FStaticParameterSet StaticParameters;
 
+#ifdef WITH_EDITOR
+	mutable TOptional<FStaticParameterSet> CachedStaticParameterValues;
+	mutable uint8 AllowCachingStaticParameterValuesCounter = 0;
+#endif // WITH_EDITOR
+
 	/** Inline material resources serialized from disk. To be processed on game thread in PostLoad. */
 	TArray<FMaterialResource> LoadedMaterialResources;
 
@@ -665,6 +670,10 @@ public:
 	void SaveShaderStableKeys(const class ITargetPlatform* TP);
 	ENGINE_API virtual void SaveShaderStableKeysInner(const class ITargetPlatform* TP, const struct FStableShaderKeyAndValue& SaveKeyVal) override;
 
+#if WITH_EDITOR
+	void BeginAllowCachingStaticParameterValues();
+	void EndAllowCachingStaticParameterValues();
+#endif // WITH_EDITOR
 
 protected:
 
diff --git a/Engine/Source/Runtime/Engine/Private/Materials/Material.cpp b/Engine/Source/Runtime/Engine/Private/Materials/Material.cpp
--- a/Engine/Source/Runtime/Engine/Private/Materials/Material.cpp
+++ b/Engine/Source/Runtime/Engine/Private/Materials/Material.cpp
@@ -242,6 +242,24 @@ void FMaterialResource::GatherExpressionsForCustomInterpolators(TArray<UMaterial
 	Material->GetAllExpressionsForCustomInterpolators(OutExpressions);
 }
 
+#if WITH_EDITOR
+void FMaterialResource::BeginAllowCachingStaticParameterValues()
+{
+	if (MaterialInstance)
+	{
+		MaterialInstance->BeginAllowCachingStaticParameterValues();
+	}
+}
+
+void FMaterialResource::EndAllowCachingStaticParameterValues()
+{
+	if (MaterialInstance)
+	{
+		MaterialInstance->EndAllowCachingStaticParameterValues();
+	}
+}
+#endif // WITH_EDITOR
+
 void FMaterialResource::GetShaderMapId(EShaderPlatform Platform, FMaterialShaderMapId& OutId) const
 {
 	FMaterial::GetShaderMapId(Platform, OutId);
diff --git a/Engine/Source/Runtime/Engine/Private/Materials/MaterialInstance.cpp b/Engine/Source/Runtime/Engine/Private/Materials/MaterialInstance.cpp
--- a/Engine/Source/Runtime/Engine/Private/Materials/MaterialInstance.cpp
+++ b/Engine/Source/Runtime/Engine/Private/Materials/MaterialInstance.cpp
@@ -1959,6 +1959,12 @@ void UMaterialInstance::GetStaticParameterValues(FStaticParameterSet& OutStaticP
 {
 	check(IsInGameThread());
 
+	if ((AllowCachingStaticParameterValuesCounter > 0) && CachedStaticParameterValues.IsSet())
+	{
+		OutStaticParameters = CachedStaticParameterValues.GetValue();
+		return;
+	}
+
 	if (Parent)
 	{
 		UMaterial* ParentMaterial = Parent->GetMaterial();
@@ -2063,6 +2069,11 @@ void UMaterialInstance::GetStaticParameterValues(FStaticParameterSet& OutStaticP
 
 	// Custom parameters.
 	CustomStaticParametersGetters.Broadcast(OutStaticParameters, this);
+
+	if (AllowCachingStaticParameterValuesCounter > 0)
+	{
+		CachedStaticParameterValues = OutStaticParameters;
+	}
 }
 
 void UMaterialInstance::GetAllScalarParameterInfo(TArray<FMaterialParameterInfo>& OutParameterInfo, TArray<FGuid>& OutParameterIds) const
@@ -4728,6 +4739,23 @@ void UMaterialInstance::SaveShaderStableKeysInner(const class ITargetPlatform* T
 #endif
 }
 
+#if WITH_EDITOR
+void UMaterialInstance::BeginAllowCachingStaticParameterValues()
+{
+	++AllowCachingStaticParameterValuesCounter;
+}
+
+void UMaterialInstance::EndAllowCachingStaticParameterValues()
+{
+	check(AllowCachingStaticParameterValuesCounter > 0);
+	--AllowCachingStaticParameterValuesCounter;
+	if (AllowCachingStaticParameterValuesCounter == 0)
+	{
+		CachedStaticParameterValues.Reset();
+	}
+}
+#endif // WITH_EDITOR
+
 void UMaterialInstance::CopyMaterialUniformParametersInternal(UMaterialInterface* Source)
 {
 	LLM_SCOPE(ELLMTag::MaterialInstance);
diff --git a/Engine/Source/Runtime/Engine/Private/Materials/MaterialShared.cpp b/Engine/Source/Runtime/Engine/Private/Materials/MaterialShared.cpp
--- a/Engine/Source/Runtime/Engine/Private/Materials/MaterialShared.cpp
+++ b/Engine/Source/Runtime/Engine/Private/Materials/MaterialShared.cpp
@@ -53,6 +53,29 @@ FAutoConsoleVariableRef CVarDeferUniformExpressionCaching(
 	ECVF_RenderThreadSafe
 	);
 
+struct FAllowCachingStaticParameterValues
+{
+	FAllowCachingStaticParameterValues(FMaterial& InMaterial)
+#if WITH_EDITOR
+		: Material(InMaterial)
+#endif // WITH_EDITOR
+	{
+#if WITH_EDITOR
+		Material.BeginAllowCachingStaticParameterValues();
+#endif // WITH_EDITOR
+	};
+
+#if WITH_EDITOR
+	~FAllowCachingStaticParameterValues()
+	{
+		Material.EndAllowCachingStaticParameterValues();
+	}
+
+private:
+	FMaterial& Material;
+#endif // WITH_EDITOR
+};
+
 static FAutoConsoleCommand GFlushMaterialUniforms(
 	TEXT("r.FlushMaterialUniforms"),
 	TEXT(""),
@@ -1761,6 +1784,7 @@ void FMaterial::SetupMaterialEnvironment(
  */
 bool FMaterial::CacheShaders(EShaderPlatform Platform, const ITargetPlatform* TargetPlatform)
 {
+	FAllowCachingStaticParameterValues AllowCachingStaticParameterValues(*this);
 	FMaterialShaderMapId NoStaticParametersId;
 	GetShaderMapId(Platform, NoStaticParametersId);
 	FStaticParameterSet StaticParameterSet;
diff --git a/Engine/Source/Runtime/Engine/Public/MaterialShared.h b/Engine/Source/Runtime/Engine/Public/MaterialShared.h
--- a/Engine/Source/Runtime/Engine/Public/MaterialShared.h
+++ b/Engine/Source/Runtime/Engine/Public/MaterialShared.h
@@ -1745,6 +1745,11 @@ public:
 	static void RestoreEditorLoadedMaterialShadersFromMemory(const TMap<FMaterialShaderMap*, TUniquePtr<TArray<uint8> > >& ShaderMapToSerializedShaderData);
 #endif // WITH_EDITOR
 
+#if WITH_EDITOR
+	ENGINE_API virtual void BeginAllowCachingStaticParameterValues() {};
+	ENGINE_API virtual void EndAllowCachingStaticParameterValues() {};
+#endif // WITH_EDITOR
+
 protected:
 	
 	// shared code needed for GetUniformScalarParameterExpressions, GetUniformVectorParameterExpressions, GetUniformCubeTextureExpressions..
@@ -2198,6 +2203,10 @@ public:
 	// FMaterial interface.
 	ENGINE_API virtual void GetShaderMapId(EShaderPlatform Platform, FMaterialShaderMapId& OutId) const override;
 	ENGINE_API virtual void GetStaticParameterSet(EShaderPlatform Platform, FStaticParameterSet& OutSet) const override;
+#if WITH_EDITOR
+	ENGINE_API virtual void BeginAllowCachingStaticParameterValues() override;
+	ENGINE_API virtual void EndAllowCachingStaticParameterValues() override;
+#endif // WITH_EDITOR
 	ENGINE_API virtual EMaterialDomain GetMaterialDomain() const override;
 	ENGINE_API virtual bool IsTwoSided() const override;
 	ENGINE_API virtual bool IsDitheredLODTransition() const override;
