Fix non-editor builds

#rnx

diff --git a/Engine/Source/Runtime/Engine/Private/Animation/AnimTrace.cpp b/Engine/Source/Runtime/Engine/Private/Animation/AnimTrace.cpp
--- a/Engine/Source/Runtime/Engine/Private/Animation/AnimTrace.cpp
+++ b/Engine/Source/Runtime/Engine/Private/Animation/AnimTrace.cpp
@@ -507,7 +507,12 @@ void FAnimTrace::OutputAnimNodeStart(const FAnimationBaseContext& InContext, uin
 	UStructProperty* LinkedProperty = AnimNodeProperties[InNodeId];
 	check(LinkedProperty->Struct);
 
+#if WITH_EDITOR
 	FString DisplayNameString = LinkedProperty->Struct->GetDisplayNameText().ToString();
+#else
+	FString DisplayNameString = LinkedProperty->Struct->GetName();
+#endif
+
 	DisplayNameString.RemoveFromStart(TEXT("Anim Node "));
 
 	UE_TRACE_LOG(Animation, AnimNodeStart, (DisplayNameString.Len() + 1) * sizeof(TCHAR))
