More CIS Fixes

diff --git a/Engine/Source/Runtime/CoreUObject/Public/UObject/WeakFieldPtr.h b/Engine/Source/Runtime/CoreUObject/Public/UObject/WeakFieldPtr.h
--- a/Engine/Source/Runtime/CoreUObject/Public/UObject/WeakFieldPtr.h
+++ b/Engine/Source/Runtime/CoreUObject/Public/UObject/WeakFieldPtr.h
@@ -206,10 +206,10 @@ public:
 		return GetTypeHash(WeakObjectPtr.Field);
 	}
 
-	friend FArchive& operator<<(FArchive& Ar, TWeakFieldPtr& WeakObjectPtr)
+	friend FArchive& operator<<(FArchive& Ar, TWeakFieldPtr& InWeakFieldPtr)
 	{
-		Ar << Owner;
-		Ar << Field;
+		Ar << InWeakFieldPtr.Owner;
+		Ar << InWeakFieldPtr.Field;
 		return Ar;
 	}
 
diff --git a/Engine/Source/Runtime/Engine/Classes/Animation/AnimBlueprintGeneratedClass.h b/Engine/Source/Runtime/Engine/Classes/Animation/AnimBlueprintGeneratedClass.h
--- a/Engine/Source/Runtime/Engine/Classes/Animation/AnimBlueprintGeneratedClass.h
+++ b/Engine/Source/Runtime/Engine/Classes/Animation/AnimBlueprintGeneratedClass.h
@@ -446,7 +446,7 @@ public:
 	StructType& GetPropertyInstanceChecked(UObject* Object, UAnimGraphNode_Base* Node, EPropertySearchMode::Type SearchMode = EPropertySearchMode::OnlyThis)
 	{
 		const int32 Index = AnimBlueprintDebugData.NodePropertyToIndexMap.FindChecked(Node);
-		FStructProperty* AnimationProperty = AnimNodeProperties[AnimNodeProperties.Num() - 1 - Index];
+		FStructProperty* AnimationProperty = AnimNodeProperties[AnimNodeProperties.Num() - 1 - Index].Get();
 		check(AnimationProperty);
 		check(AnimationProperty->Struct->IsChildOf(StructType::StaticStruct()));
 		return AnimationProperty->ContainerPtrToValuePtr<StructType>((void*)Object);
