Fixing Android compile errors

diff --git a/Engine/Source/Runtime/Engine/Classes/Engine/MemberReference.h b/Engine/Source/Runtime/Engine/Classes/Engine/MemberReference.h
--- a/Engine/Source/Runtime/Engine/Classes/Engine/MemberReference.h
+++ b/Engine/Source/Runtime/Engine/Classes/Engine/MemberReference.h
@@ -282,27 +282,8 @@ protected:
 		return nullptr;
 	}
 
-	template<>
-	UFunction* ResolveUFunction() const
-	{
-		UFunction* ReturnField = nullptr;
-		FString const StringName = MemberName.ToString();
-		for (TObjectIterator<UPackage> PackageIt; PackageIt && (ReturnField == nullptr); ++PackageIt)
-		{
-			if (PackageIt->HasAnyPackageFlags(PKG_CompiledIn) == false)
-			{
-				continue;
-			}
-
-			// NOTE: this could return the wrong field (if there are 
-			//       two like-named delegates defined in separate packages)
-			ReturnField = FindObject<UFunction>(*PackageIt, *StringName);
-		}
-		return ReturnField;
-	}
-
-	template<class TFieldType, class TFieldClassType>
-	TFieldType* ResolveUField(TFieldClassType* InClass, UPackage* TargetPackage) const
+	template<class TFieldType>
+	TFieldType* ResolveUField(FFieldClass* InClass, UPackage* TargetPackage) const
 	{
 		return nullptr;
 	}
@@ -317,11 +298,6 @@ protected:
 	{
 		return Field->GetOuter();
 	}
-	template<>
-	UObject* GetFieldOuter(FField* Field) const
-	{
-		return Field->GetOwner().ToUObject();
-	}
 
 public:
 
@@ -372,16 +348,6 @@ public:
 	{
 		return ClassA == ClassB;
 	}
-	template <>
-	bool CompareClassesHelper(const UClass* ClassA, const FFieldClass* ClassB) const
-	{
-		return false;
-	}
-	template <>
-	bool CompareClassesHelper(const FFieldClass* ClassA, const UClass* ClassB) const
-	{
-		return false;
-	}
 
 	/** 
 	 *	Returns the member FProperty/UFunction this reference is pointing to, or NULL if it no longer exists 
@@ -594,3 +560,40 @@ public:
 		return Result;
 	}
 };
+
+
+template<>
+inline UFunction* FMemberReference::ResolveUFunction() const
+{
+	UFunction* ReturnField = nullptr;
+	FString const StringName = MemberName.ToString();
+	for (TObjectIterator<UPackage> PackageIt; PackageIt && (ReturnField == nullptr); ++PackageIt)
+	{
+		if (PackageIt->HasAnyPackageFlags(PKG_CompiledIn) == false)
+		{
+			continue;
+		}
+
+		// NOTE: this could return the wrong field (if there are 
+		//       two like-named delegates defined in separate packages)
+		ReturnField = FindObject<UFunction>(*PackageIt, *StringName);
+	}
+	return ReturnField;
+}
+
+template<>
+inline UObject* FMemberReference::GetFieldOuter(FField* Field) const
+{
+	return Field->GetOwner().ToUObject();
+}
+
+template <>
+inline bool FMemberReference::CompareClassesHelper(const UClass* ClassA, const FFieldClass* ClassB) const
+{
+	return false;
+}
+template <>
+inline bool FMemberReference::CompareClassesHelper(const FFieldClass* ClassA, const UClass* ClassB) const
+{
+	return false;
+}
\ No newline at end of file
