Undo changelist 10516234

diff --git a/Engine/Source/Runtime/Engine/Classes/Engine/StreamableManager.h b/Engine/Source/Runtime/Engine/Classes/Engine/StreamableManager.h
--- a/Engine/Source/Runtime/Engine/Classes/Engine/StreamableManager.h
+++ b/Engine/Source/Runtime/Engine/Classes/Engine/StreamableManager.h
@@ -149,12 +149,6 @@ private:
 	/** Callback when async load finishes, it's here so we can use a shared pointer for callback safety */
 	void AsyncLoadCallbackWrapper(const FName& PackageName, UPackage* LevelPackage, EAsyncLoadingResult::Type Result, FSoftObjectPath TargetName);
 
-	/** Notify all parents that a child completed loading */
-	void NotifyParentsOfCompletion();
-
-	/** Notify all parents that a child was canceled */
-	void NotifyParentsOfCancellation();
-
 	/** Called on meta handle when a child handle has completed/canceled */
 	void UpdateCombinedHandle();
 
@@ -200,12 +194,6 @@ private:
 	/** How many FStreamables is this waiting on to finish loading */
 	int32 StreamablesLoading;
 
-	/** How many of our children that have been completed */
-	int32 CompletedChildCount = 0;
-
-	/** How many of our children that have been canceled */
-	int32 CanceledChildCount = 0;
-
 	/** List of assets that were referenced by this handle */
 	TArray<FSoftObjectPath> RequestedAssets;
 
diff --git a/Engine/Source/Runtime/Engine/Private/StreamableManager.cpp b/Engine/Source/Runtime/Engine/Private/StreamableManager.cpp
--- a/Engine/Source/Runtime/Engine/Private/StreamableManager.cpp
+++ b/Engine/Source/Runtime/Engine/Private/StreamableManager.cpp
@@ -8,7 +8,6 @@
 #include "HAL/IConsoleManager.h"
 #include "Tickable.h"
 #include "Serialization/LoadTimeTrace.h"
-#include "Containers/List.h"
 
 DEFINE_LOG_CATEGORY_STATIC(LogStreamableManager, Log, All);
 
@@ -30,13 +29,8 @@ public:
 	void AddDelegate(const FStreamableDelegate& Delegate, const FStreamableDelegate& CancelDelegate, TSharedPtr<FStreamableHandle> AssociatedHandle)
 	{
 		DataLock.Lock();
-		
-		TUniquePtr<TDoubleLinkedList<FPendingDelegate>>& DelegatesForHandle = PendingDelegates.FindOrAdd(AssociatedHandle);
-		if (!DelegatesForHandle)
-		{
-			DelegatesForHandle.Reset(new TDoubleLinkedList<FPendingDelegate>());
-		}
-		DelegatesForHandle->AddTail(FPendingDelegate(Delegate, CancelDelegate, AssociatedHandle));
+
+		PendingDelegates.Emplace(Delegate, CancelDelegate, AssociatedHandle);
 
 		DataLock.Unlock();
 	}
@@ -44,38 +38,34 @@ public:
 	/** Cancels delegate for handle, this will either delete the delegate or replace with the cancel delegate */
 	void CancelDelegatesForHandle(TSharedPtr<FStreamableHandle> AssociatedHandle)
 	{
-		FScopeLock _(&DataLock);
-		
-		TUniquePtr<TDoubleLinkedList<FPendingDelegate>>* FindIt = PendingDelegates.Find(AssociatedHandle);
-		if (!FindIt)
+		if (!AssociatedHandle.IsValid())
 		{
 			return;
 		}
 
-		TDoubleLinkedList<FPendingDelegate>* DelegatesForHandle = FindIt->Get();
+		DataLock.Lock();
 
-		TDoubleLinkedList<FPendingDelegate>::TDoubleLinkedListNode* CurrentNode = DelegatesForHandle->GetHead();
-		while (CurrentNode)
+		for (int32 DelegateIndex = 0; DelegateIndex < PendingDelegates.Num(); DelegateIndex++)
 		{
-			TDoubleLinkedList<FPendingDelegate>::TDoubleLinkedListNode* NextNode = CurrentNode->GetNextNode();
-			FPendingDelegate& PendingDelegate = CurrentNode->GetValue();
-			if (PendingDelegate.CancelDelegate.IsBound())
-			{
-				// Replace with cancel delegate
-				PendingDelegate.Delegate = PendingDelegate.CancelDelegate;
-				PendingDelegate.CancelDelegate.Unbind();
-			}
-			else
+			FPendingDelegate& PendingDelegate = PendingDelegates[DelegateIndex];
+			if (PendingDelegate.RelatedHandle == AssociatedHandle)
 			{
-				// Remove entirely
-				DelegatesForHandle->RemoveNode(CurrentNode);
+				if (PendingDelegate.CancelDelegate.IsBound())
+				{
+					// Replace with cancel delegate
+					PendingDelegate.Delegate = PendingDelegate.CancelDelegate;
+					PendingDelegate.CancelDelegate.Unbind();
+				}
+				else
+				{
+					// Remove entirely
+					PendingDelegates.RemoveAt(DelegateIndex);
+					DelegateIndex--;
+				}
 			}
-			CurrentNode = NextNode;
-		}
-		if (DelegatesForHandle->Num() == 0)
-		{
-			PendingDelegates.Remove(AssociatedHandle);
 		}
+
+		DataLock.Unlock();
 	}
 
 	/** Calls all delegates, call from synchronous flushes */
@@ -98,26 +88,15 @@ public:
 
 		DataLock.Lock();
 
-		TDoubleLinkedList<FPendingDelegate> DelegatesToCall;
-		for (auto It = PendingDelegates.CreateIterator(); It; ++It)
+		TArray<FPendingDelegate> DelegatesToCall;
+		DelegatesToCall.Reserve(PendingDelegates.Num());
+		for (int32 DelegateIndex = 0; DelegateIndex < PendingDelegates.Num(); DelegateIndex++)
 		{
-			TDoubleLinkedList<FPendingDelegate>* DelegatesForHandle = It->Value.Get();
-			TDoubleLinkedList<FPendingDelegate>::TDoubleLinkedListNode* CurrentNode = DelegatesForHandle->GetHead();
-			while (CurrentNode)
+			if (--PendingDelegates[DelegateIndex].DelayFrames <= 0)
 			{
-				TDoubleLinkedList<FPendingDelegate>::TDoubleLinkedListNode* NextNode = CurrentNode->GetNextNode();
-				FPendingDelegate& PendingDelegate = CurrentNode->GetValue();
-				if (--PendingDelegate.DelayFrames <= 0)
-				{
-					// Add to call array and remove from tracking one
-					DelegatesToCall.AddTail(CurrentNode);
-					DelegatesForHandle->RemoveNode(CurrentNode, false);
-				}
-				CurrentNode = NextNode;
-			}
-			if (DelegatesForHandle->Num() == 0)
-			{
-				It.RemoveCurrent();
+				// Add to call array and remove from tracking one
+				DelegatesToCall.Emplace(PendingDelegates[DelegateIndex].Delegate, FStreamableDelegate(), PendingDelegates[DelegateIndex].RelatedHandle);
+				PendingDelegates.RemoveAt(DelegateIndex--);
 			}
 		}
 
@@ -128,7 +107,6 @@ public:
 			// Call delegates, these may add other deferred delegates
 			PendingDelegate.Delegate.ExecuteIfBound();
 		}
-		DelegatesToCall.Empty();
 
 		// When DelegatesToCall falls out of scope it may delete the referenced handles
 	}
@@ -157,8 +135,6 @@ private:
 
 	struct FPendingDelegate
 	{
-		FPendingDelegate* Next = nullptr;
-		
 		/** Delegate to call when frames are up */
 		FStreamableDelegate Delegate;
 
@@ -179,7 +155,7 @@ private:
 		{}
 	};
 
-	TMap<TSharedPtr<FStreamableHandle>, TUniquePtr<TDoubleLinkedList<FPendingDelegate>>> PendingDelegates;
+	TArray<FPendingDelegate> PendingDelegates;
 
 	FCriticalSection DataLock;
 };
@@ -421,12 +397,10 @@ void FStreamableHandle::CancelHandle()
 		}
 
 		bCanceled = true;
-		NotifyParentsOfCancellation();
 		return;
 	}
 
 	bCanceled = true;
-	NotifyParentsOfCancellation();
 
 	ExecuteDelegate(CancelDelegate, SharedThis);
 	UnbindDelegates();
@@ -451,8 +425,6 @@ void FStreamableHandle::CancelHandle()
 	}
 
 	ChildHandles.Empty();
-	CompletedChildCount = 0;
-	CanceledChildCount = 0;
 
 	OwningManager = nullptr;
 
@@ -506,8 +478,6 @@ void FStreamableHandle::ReleaseHandle()
 		}
 
 		ChildHandles.Empty();
-		CompletedChildCount = 0;
-		CanceledChildCount = 0;
 
 		OwningManager = nullptr;
 	}
@@ -555,8 +525,6 @@ void FStreamableHandle::CompleteLoad()
 		ExecuteDelegate(CompleteDelegate, AsShared(), CancelDelegate);
 		UnbindDelegates();
 
-		NotifyParentsOfCompletion();
-
 		if (ParentHandles.Num() > 0)
 		{
 			// Update any meta handles that are still active. Copy the array first as elements may be removed from original while iterating
@@ -574,42 +542,6 @@ void FStreamableHandle::CompleteLoad()
 	}
 }
 
-void FStreamableHandle::NotifyParentsOfCompletion()
-{
-	if (ParentHandles.Num() > 0)
-	{
-		// Update any meta handles that are still active. Copy the array first as elements may be removed from original while iterating
-		TArray<TWeakPtr<FStreamableHandle>> ParentHandlesCopy = ParentHandles;
-		for (TWeakPtr<FStreamableHandle>& WeakHandle : ParentHandlesCopy)
-		{
-			TSharedPtr<FStreamableHandle> Handle = WeakHandle.Pin();
-
-			if (Handle.IsValid())
-			{
-				++Handle->CompletedChildCount;
-			}
-		}
-	}
-}
-
-void FStreamableHandle::NotifyParentsOfCancellation()
-{
-	if (ParentHandles.Num() > 0)
-	{
-		// Update any meta handles that are still active. Copy the array first as elements may be removed from original while iterating
-		TArray<TWeakPtr<FStreamableHandle>> ParentHandlesCopy = ParentHandles;
-		for (TWeakPtr<FStreamableHandle>& WeakHandle : ParentHandlesCopy)
-		{
-			TSharedPtr<FStreamableHandle> Handle = WeakHandle.Pin();
-
-			if (Handle.IsValid())
-			{
-				++Handle->CanceledChildCount;
-			}
-		}
-	}
-}
-
 void FStreamableHandle::UpdateCombinedHandle()
 {
 	if (!IsActive())
@@ -622,11 +554,6 @@ void FStreamableHandle::UpdateCombinedHandle()
 		return;
 	}
 
-	if (CompletedChildCount + CanceledChildCount < ChildHandles.Num())
-	{
-		return;
-	}
-
 	// Check all our children, complete if done
 	bool bAllCompleted = true;
 	bool bAllCanceled = true;
@@ -781,9 +708,6 @@ struct FStreamable
 					FStreamableHandle::ExecuteDelegate(ActiveHandle->CancelDelegate, ActiveHandle);
 					ActiveHandle->UnbindDelegates();
 				}
-
-				ActiveHandle->NotifyParentsOfCancellation();
-
 			}
 		}
 		ActiveHandles.Empty();
@@ -1427,14 +1351,6 @@ TSharedPtr<FStreamableHandle> FStreamableManager::CreateCombinedHandle(const TAr
 
 		ChildHandle->ParentHandles.Add(NewRequest);
 		NewRequest->ChildHandles.Add(ChildHandle);
-		if (ChildHandle->bLoadCompleted)
-		{
-			++NewRequest->CompletedChildCount;
-		}
-		if (ChildHandle->bCanceled)
-		{
-			++NewRequest->CanceledChildCount;
-		}
 	}
 
 	// Add to pending list so these handles don't free when not referenced
