Add some simple tests for UTexture2D::CreateTransient.

A bug was found when creating a transient in some BulkData changes I submitted, but not during testing due to the limited number of uses we have in our code base. Adding an engine test should help future issues to come to light sooner.

Currently we only test creating transient textures with invalid parameters and creating a valid one, so very limited basic testing.

#rnx

diff --git a/Engine/Source/Runtime/Engine/Private/Tests/Texture2DTests.cpp b/Engine/Source/Runtime/Engine/Private/Tests/Texture2DTests.cpp
new file mode 100644
--- /dev/null
+++ b/Engine/Source/Runtime/Engine/Private/Tests/Texture2DTests.cpp
@@ -0,0 +1,36 @@
+// Copyright 1998-2019 Epic Games, Inc. All Rights Reserved.
+
+#include "Misc/AutomationTest.h"
+#include "Engine/Texture2D.h"
+
+#if WITH_DEV_AUTOMATION_TESTS
+
+#define TEST_NAME_ROOT "System.Engine.Texture2D"
+constexpr const uint32 TestFlags = EAutomationTestFlags::ApplicationContextMask | EAutomationTestFlags::EngineFilter;
+
+IMPLEMENT_SIMPLE_AUTOMATION_TEST(FTexture2DTestTransient, TEST_NAME_ROOT ".CreateTransient", TestFlags)
+bool FTexture2DTestTransient::RunTest(const FString& Parameters)
+{
+	// A simple test to make sure that basica functionality in UTexture2D::CreateTransient works as it seems to be a 
+	// fairly uncommon code path in our samples/test games etc.
+
+	{
+		// Each test in this scope is expected to give one warning about invalid parameters
+		AddExpectedError(TEXT("Invalid parameters specified for UTexture2D::CreateTransient()"), EAutomationExpectedErrorFlags::Contains, 3);
+
+		UTexture2D* ZeroSizedTexture = UTexture2D::CreateTransient(0, 0);
+		TestTrue(TEXT("Creating a transient texture with a zero length dimension should fail!"), ZeroSizedTexture == nullptr);
+
+		UTexture2D* ZeroWidthTexture = UTexture2D::CreateTransient(0, 32);
+		TestTrue(TEXT("Creating a transient texture with a zero length dimension should fail!"), ZeroWidthTexture == nullptr);
+
+		UTexture2D* ZeroHeightTexture = UTexture2D::CreateTransient(32, 0);
+		TestTrue(TEXT("Creating a transient texture with a zero length dimension should fail!"), ZeroHeightTexture == nullptr);
+	}
+
+	UTexture2D* TransientTexture = UTexture2D::CreateTransient(32, 32);
+	TestTrue(TEXT("Failed to create a 32*32 transient texture!"), TransientTexture != nullptr);
+	
+	return true;
+}
+#endif 
