Fix CIS StressTestUnity build.

- Moved the Texture2D tests to their own namespace (Texture2DTest) so that the code does not conflict with any actual runtime code when merged into once giant SCU for the entire module.

#rnx

diff --git a/Engine/Source/Runtime/Engine/Private/Tests/Texture2DTests.cpp b/Engine/Source/Runtime/Engine/Private/Tests/Texture2DTests.cpp
--- a/Engine/Source/Runtime/Engine/Private/Tests/Texture2DTests.cpp
+++ b/Engine/Source/Runtime/Engine/Private/Tests/Texture2DTests.cpp
@@ -5,32 +5,35 @@
 
 #if WITH_DEV_AUTOMATION_TESTS
 
-#define TEST_NAME_ROOT "System.Engine.Texture2D"
-constexpr const uint32 TestFlags = EAutomationTestFlags::ApplicationContextMask | EAutomationTestFlags::EngineFilter;
-
-IMPLEMENT_SIMPLE_AUTOMATION_TEST(FTexture2DTestTransient, TEST_NAME_ROOT ".CreateTransient", TestFlags)
-bool FTexture2DTestTransient::RunTest(const FString& Parameters)
+namespace Texture2DTest
 {
+	#define TEST_NAME_ROOT "System.Engine.Texture2D"
+	constexpr const uint32 TestFlags = EAutomationTestFlags::ApplicationContextMask | EAutomationTestFlags::EngineFilter;
+
 	// A simple test to make sure that basica functionality in UTexture2D::CreateTransient works as it seems to be a 
 	// fairly uncommon code path in our samples/test games etc.
-
+	IMPLEMENT_SIMPLE_AUTOMATION_TEST(FTexture2DTestTransient, TEST_NAME_ROOT ".CreateTransient", TestFlags)
+	bool FTexture2DTestTransient::RunTest(const FString& Parameters)
 	{
-		// Each test in this scope is expected to give one warning about invalid parameters
-		AddExpectedError(TEXT("Invalid parameters specified for UTexture2D::CreateTransient()"), EAutomationExpectedErrorFlags::Contains, 3);
+		{
+			// Each test in this scope is expected to give one warning about invalid parameters
+			AddExpectedError(TEXT("Invalid parameters specified for UTexture2D::CreateTransient()"), EAutomationExpectedErrorFlags::Contains, 3);
 
-		UTexture2D* ZeroSizedTexture = UTexture2D::CreateTransient(0, 0);
-		TestTrue(TEXT("Creating a transient texture with a zero length dimension should fail!"), ZeroSizedTexture == nullptr);
+			UTexture2D* ZeroSizedTexture = UTexture2D::CreateTransient(0, 0);
+			TestTrue(TEXT("Creating a transient texture with a zero length dimension should fail!"), ZeroSizedTexture == nullptr);
 
-		UTexture2D* ZeroWidthTexture = UTexture2D::CreateTransient(0, 32);
-		TestTrue(TEXT("Creating a transient texture with a zero length dimension should fail!"), ZeroWidthTexture == nullptr);
+			UTexture2D* ZeroWidthTexture = UTexture2D::CreateTransient(0, 32);
+			TestTrue(TEXT("Creating a transient texture with a zero length dimension should fail!"), ZeroWidthTexture == nullptr);
 
-		UTexture2D* ZeroHeightTexture = UTexture2D::CreateTransient(32, 0);
-		TestTrue(TEXT("Creating a transient texture with a zero length dimension should fail!"), ZeroHeightTexture == nullptr);
-	}
+			UTexture2D* ZeroHeightTexture = UTexture2D::CreateTransient(32, 0);
+			TestTrue(TEXT("Creating a transient texture with a zero length dimension should fail!"), ZeroHeightTexture == nullptr);
+		}
 
-	UTexture2D* TransientTexture = UTexture2D::CreateTransient(32, 32);
-	TestTrue(TEXT("Failed to create a 32*32 transient texture!"), TransientTexture != nullptr);
-	
-	return true;
+		UTexture2D* TransientTexture = UTexture2D::CreateTransient(32, 32);
+		TestTrue(TEXT("Failed to create a 32*32 transient texture!"), TransientTexture != nullptr);
+
+		return true;
+	}
 }
+
 #endif 
