Fixing nativized anim blueprints crashes and issues

diff --git a/Engine/Source/Developer/BlueprintNativeCodeGen/Private/BlueprintNativeCodeGenUtils.cpp b/Engine/Source/Developer/BlueprintNativeCodeGen/Private/BlueprintNativeCodeGenUtils.cpp
--- a/Engine/Source/Developer/BlueprintNativeCodeGen/Private/BlueprintNativeCodeGenUtils.cpp
+++ b/Engine/Source/Developer/BlueprintNativeCodeGen/Private/BlueprintNativeCodeGenUtils.cpp
@@ -423,7 +423,7 @@ void FBlueprintNativeCodeGenUtils::GenerateCppCode(UObject* Obj, TSharedPtr<FStr
 
 		FDisableGatheringDataOnScope DisableFib;
 
-		const FString TempPackageName = FString::Printf(TEXT("/Temp/__TEMP_BP__%s"), *InBlueprintObj->GetOutermost()->GetPathName());
+		const FString TempPackageName = FString::Printf(TEXT("%s%s"), *UDynamicClass::GetTempPackagePrefix(), *InBlueprintObj->GetOutermost()->GetPathName());
 		UPackage* TempPackage = CreatePackage(nullptr, *TempPackageName);
 		check(TempPackage);
 		ON_SCOPE_EXIT
diff --git a/Engine/Source/Editor/AnimGraph/Private/AnimGraphNode_LinkedAnimLayer.cpp b/Engine/Source/Editor/AnimGraph/Private/AnimGraphNode_LinkedAnimLayer.cpp
--- a/Engine/Source/Editor/AnimGraph/Private/AnimGraphNode_LinkedAnimLayer.cpp
+++ b/Engine/Source/Editor/AnimGraph/Private/AnimGraphNode_LinkedAnimLayer.cpp
@@ -518,9 +518,9 @@ void UAnimGraphNode_LinkedAnimLayer::GetExposableProperties(TArray<FProperty*>&
 			// Check name matches.
 			if(AnimBlueprintFunction.Name == Node.GetDynamicLinkFunctionName())
 			{
-				for(FProperty* Property : AnimBlueprintFunction.InputProperties)
+				for(const TFieldPath<FProperty>& Property : AnimBlueprintFunction.InputProperties)
 				{
-					OutExposableProperties.Add(Property);
+					OutExposableProperties.Add(Property.Get());
 				}
 			}
 		}
diff --git a/Engine/Source/Runtime/CoreUObject/Private/UObject/Class.cpp b/Engine/Source/Runtime/CoreUObject/Private/UObject/Class.cpp
--- a/Engine/Source/Runtime/CoreUObject/Private/UObject/Class.cpp
+++ b/Engine/Source/Runtime/CoreUObject/Private/UObject/Class.cpp
@@ -6069,6 +6069,12 @@ FStructProperty* UDynamicClass::FindStructPropertyChecked(const TCHAR* PropertyN
 	return FindFieldChecked<FStructProperty>(this, PropertyName);
 }
 
+const FString& UDynamicClass::GetTempPackagePrefix()
+{
+	static const FString PackagePrefix(TEXT("/Temp/__TEMP_BP__"));
+	return PackagePrefix;
+}
+
 IMPLEMENT_CORE_INTRINSIC_CLASS(UDynamicClass, UClass,
 {
 	Class->ClassAddReferencedObjects = &UDynamicClass::AddReferencedObjects;
diff --git a/Engine/Source/Runtime/CoreUObject/Private/UObject/FieldPath.cpp b/Engine/Source/Runtime/CoreUObject/Private/UObject/FieldPath.cpp
--- a/Engine/Source/Runtime/CoreUObject/Private/UObject/FieldPath.cpp
+++ b/Engine/Source/Runtime/CoreUObject/Private/UObject/FieldPath.cpp
@@ -136,7 +136,13 @@ FString FFieldPath::ToString() const
 	// Generate path from the outermost (last item) to the property (first item)
 	for (int32 PathIndex = Path.Num() - 1; PathIndex >= 0; --PathIndex)
 	{
-		Result += Path[PathIndex].ToString();		
+		// @todo: this should be handled by some kind of a flag passed to this function
+		FString ObjName = Path[PathIndex].ToString();
+		if (PathIndex == (Path.Num() - 1) && ObjName.StartsWith(UDynamicClass::GetTempPackagePrefix(), ESearchCase::IgnoreCase))
+		{
+			ObjName.RemoveFromStart(UDynamicClass::GetTempPackagePrefix(), ESearchCase::IgnoreCase);
+		}
+		Result += ObjName;
 		if (PathIndex > 0)
 		{
 			// Separator between the package name (last item) and the class (asset object - the second to last item is a '.', oterwise use SUBOBJECT_DELIMITER_CHAR)
@@ -171,4 +177,4 @@ void FFieldPath::OnFieldDeleted()
 {
 	checkSlow(IsInGameThread());
 	GlobalSerialNumber++;
-}
\ No newline at end of file
+}
diff --git a/Engine/Source/Runtime/CoreUObject/Public/UObject/Class.h b/Engine/Source/Runtime/CoreUObject/Public/UObject/Class.h
--- a/Engine/Source/Runtime/CoreUObject/Public/UObject/Class.h
+++ b/Engine/Source/Runtime/CoreUObject/Public/UObject/Class.h
@@ -3206,6 +3206,9 @@ public:
 	UObject* AnimClassImplementation;
 
 	DynamicClassInitializerType DynamicClassInitializer;
+
+	/** Prefix for the temporary package where the dynamic classes are stored when being generated */
+	static const FString& GetTempPackagePrefix();
 };
 
 /**
diff --git a/Engine/Source/Runtime/Engine/Classes/Animation/AnimClassData.h b/Engine/Source/Runtime/Engine/Classes/Animation/AnimClassData.h
--- a/Engine/Source/Runtime/Engine/Classes/Animation/AnimClassData.h
+++ b/Engine/Source/Runtime/Engine/Classes/Animation/AnimClassData.h
@@ -41,38 +41,31 @@ public:
 
 	// The array of anim nodes
 	UPROPERTY()
-	TArray<UStructProperty*> AnimNodeProperties_DEPRECATED;
-	TArray<FStructPropertyPath> AnimNodeProperties;
+	TArray< TFieldPath<FStructProperty> > AnimNodeProperties;
 
 	// The array of linked anim graph nodes
 	UPROPERTY()
-	TArray<UStructProperty*> LinkedAnimGraphNodeProperties_DEPRECATED;
-	TArray<FStructPropertyPath> LinkedAnimGraphNodeProperties;
+	TArray< TFieldPath<FStructProperty> > LinkedAnimGraphNodeProperties;
 
 	// The array of linked anim layer nodes
 	UPROPERTY()
-	TArray<UStructProperty*> LinkedAnimLayerNodeProperties_DEPRECATED;
-	TArray<FStructPropertyPath> LinkedAnimLayerNodeProperties;
+	TArray< TFieldPath<FStructProperty> > LinkedAnimLayerNodeProperties;
 
 	// Array of nodes that need a PreUpdate() call
 	UPROPERTY()
-	TArray<UStructProperty*> PreUpdateNodeProperties_DEPRECATED;
-	TArray<FStructPropertyPath> PreUpdateNodeProperties;
+	TArray< TFieldPath<FStructProperty> > PreUpdateNodeProperties;
 
 	// Array of nodes that need a DynamicReset() call
 	UPROPERTY()
-	TArray<UStructProperty*> DynamicResetNodeProperties_DEPRECATED;
-	TArray<FStructPropertyPath> DynamicResetNodeProperties;
+	TArray< TFieldPath<FStructProperty> > DynamicResetNodeProperties;
 
 	// Array of state machine nodes
 	UPROPERTY()
-	TArray<UStructProperty*> StateMachineNodeProperties_DEPRECATED;
-	TArray<FStructPropertyPath> StateMachineNodeProperties;
+	TArray< TFieldPath<FStructProperty> > StateMachineNodeProperties;
 
 	// Array of nodes that need an OnInitializeAnimInstance call
 	UPROPERTY()
-	TArray<UStructProperty*> InitializationNodeProperties_DEPRECATED;
-	TArray<FStructPropertyPath> InitializationNodeProperties;
+	TArray< TFieldPath<FStructProperty> > InitializationNodeProperties;
 
 	// Indices for any Asset Player found within a specific (named) Anim Layer Graph, or implemented Anim Interface Graph
 	UPROPERTY()
@@ -110,9 +103,6 @@ public:
 	virtual const TMap<FName, FGraphAssetPlayerInformation>& GetGraphAssetPlayerInformation() const { return GraphNameAssetPlayers; }
 	virtual const TMap<FName, FAnimGraphBlendOptions>& GetGraphBlendOptions() const { return GraphBlendOptions; }
 
-	// UObject interface
-	virtual void Serialize(FArchive& Ar) override;
-
 #if WITH_EDITOR
 	void CopyFrom(IAnimClassInterface* AnimClass)
 	{
diff --git a/Engine/Source/Runtime/Engine/Classes/Animation/AnimClassInterface.h b/Engine/Source/Runtime/Engine/Classes/Animation/AnimClassInterface.h
--- a/Engine/Source/Runtime/Engine/Classes/Animation/AnimClassInterface.h
+++ b/Engine/Source/Runtime/Engine/Classes/Animation/AnimClassInterface.h
@@ -10,6 +10,7 @@
 #include "UObject/Interface.h"
 #include "Animation/AnimTypes.h"
 #include "Animation/AnimStateMachineTypes.h"
+#include "UObject/FieldPath.h"
 
 #include "AnimClassInterface.generated.h"
 
@@ -64,13 +65,16 @@ struct FAnimBlueprintFunction
 	TArray<int32> InputPoseNodeIndices;
 
 	/** The property of the output node, patched up during link */
-	FStructProperty* OutputPoseNodeProperty;
+	UPROPERTY(transient)
+	TFieldPath<FStructProperty> OutputPoseNodeProperty;
 
 	/** The properties of the input nodes, patched up during link */
-	TArray<FStructProperty*> InputPoseNodeProperties;
+	UPROPERTY(transient)
+	TArray< TFieldPath<FStructProperty> > InputPoseNodeProperties;
 
 	/** The input properties themselves */
-	TArray<FProperty*> InputProperties;
+	UPROPERTY(transient)
+	TArray< TFieldPath<FProperty> > InputProperties;
 
 	/** Whether this function is actually implemented by this class - it could just be a stub */
 	UPROPERTY(transient)
diff --git a/Engine/Source/Runtime/Engine/Private/Animation/AnimBlueprintGeneratedClass.cpp b/Engine/Source/Runtime/Engine/Private/Animation/AnimBlueprintGeneratedClass.cpp
--- a/Engine/Source/Runtime/Engine/Private/Animation/AnimBlueprintGeneratedClass.cpp
+++ b/Engine/Source/Runtime/Engine/Private/Animation/AnimBlueprintGeneratedClass.cpp
@@ -532,7 +532,7 @@ void UAnimBlueprintGeneratedClass::LinkFunctionsToDefaultObjectNodes(UObject* De
 	// Determine whether functions are 'implemented'
 	for(FAnimBlueprintFunction& AnimBlueprintFunction : AnimBlueprintFunctions)
 	{
-		if(AnimBlueprintFunction.OutputPoseNodeProperty)
+		if(AnimBlueprintFunction.OutputPoseNodeProperty.Get())
 		{
 			FAnimNode_Root* RootNode = AnimBlueprintFunction.OutputPoseNodeProperty->ContainerPtrToValuePtr<FAnimNode_Root>(DefaultObject);
 			if(RootNode->Result.LinkID != INDEX_NONE)
diff --git a/Engine/Source/Runtime/Engine/Private/Animation/AnimData.cpp b/Engine/Source/Runtime/Engine/Private/Animation/AnimData.cpp
--- a/Engine/Source/Runtime/Engine/Private/Animation/AnimData.cpp
+++ b/Engine/Source/Runtime/Engine/Private/Animation/AnimData.cpp
@@ -6,56 +6,3 @@
 
 #include "Animation/AnimClassData.h"
 #include "UObject/UnrealTypePrivate.h"
-#include "UObject/CoreObjectVersion.h"
-
-void UAnimClassData::Serialize(FArchive& Ar)
-{
-	Ar.UsingCustomVersion(FCoreObjectVersion::GUID);
-
-	Super::Serialize(Ar);
-
-	if (Ar.CustomVer(FCoreObjectVersion::GUID) >= FCoreObjectVersion::FProperties)
-	{
-		Ar << AnimNodeProperties;
-		Ar << LinkedAnimGraphNodeProperties;
-		Ar << LinkedAnimLayerNodeProperties;
-		Ar << PreUpdateNodeProperties;
-		Ar << DynamicResetNodeProperties;
-		Ar << StateMachineNodeProperties;
-		Ar << InitializationNodeProperties;
-	}
-#if WITH_EDITORONLY_DATA
-	else if (Ar.IsLoading())
-	{
-		for (UStructProperty* Prop : AnimNodeProperties_DEPRECATED)
-		{
-			AnimNodeProperties.Add(Prop);
-		}
-		for (UStructProperty* Prop : LinkedAnimGraphNodeProperties_DEPRECATED)
-		{
-			LinkedAnimGraphNodeProperties.Add(Prop);
-		}
-		for (UStructProperty* Prop : LinkedAnimLayerNodeProperties_DEPRECATED)
-		{
-			LinkedAnimLayerNodeProperties.Add(Prop);
-		}
-		for (UStructProperty* Prop : PreUpdateNodeProperties_DEPRECATED)
-		{
-			PreUpdateNodeProperties.Add(Prop);
-		}
-		for (UStructProperty* Prop : DynamicResetNodeProperties_DEPRECATED)
-		{
-			DynamicResetNodeProperties.Add(Prop);
-		}
-		for (UStructProperty* Prop : StateMachineNodeProperties_DEPRECATED)
-		{
-			StateMachineNodeProperties.Add(Prop);
-		}
-		for (UStructProperty* Prop : InitializationNodeProperties_DEPRECATED)
-		{
-			InitializationNodeProperties.Add(Prop);
-		}
-
-	}
-#endif // WITH_EDITORONLY_DATA
-}
\ No newline at end of file
diff --git a/Engine/Source/Runtime/Engine/Private/Animation/AnimInstanceProxy.cpp b/Engine/Source/Runtime/Engine/Private/Animation/AnimInstanceProxy.cpp
--- a/Engine/Source/Runtime/Engine/Private/Animation/AnimInstanceProxy.cpp
+++ b/Engine/Source/Runtime/Engine/Private/Animation/AnimInstanceProxy.cpp
@@ -95,7 +95,7 @@ void FAnimInstanceProxy::Initialize(UAnimInstance* InAnimInstance)
 		RootNode = nullptr;
 		if(AnimClassInterface->GetAnimBlueprintFunctions().Num() > 0)
 		{
-			if(FProperty* RootNodeProperty = AnimClassInterface->GetAnimBlueprintFunctions()[0].OutputPoseNodeProperty)
+			if(FProperty* RootNodeProperty = AnimClassInterface->GetAnimBlueprintFunctions()[0].OutputPoseNodeProperty.Get())
 			{
 				RootNode = RootNodeProperty->ContainerPtrToValuePtr<FAnimNode_Root>(InAnimInstance);
 			}
diff --git a/Engine/Source/Runtime/Engine/Private/Animation/AnimNode_LinkedAnimGraph.cpp b/Engine/Source/Runtime/Engine/Private/Animation/AnimNode_LinkedAnimGraph.cpp
--- a/Engine/Source/Runtime/Engine/Private/Animation/AnimNode_LinkedAnimGraph.cpp
+++ b/Engine/Source/Runtime/Engine/Private/Animation/AnimNode_LinkedAnimGraph.cpp
@@ -331,7 +331,7 @@ void FAnimNode_LinkedAnimGraph::DynamicLink(UAnimInstance* InOwningAnimInstance)
 						InputPoses[InputPoseIndex].AttemptRelink(Context);
 
 						int32 InputPropertyIndex = FindFunctionInputIndex(AnimBlueprintFunction, AnimBlueprintFunction.InputPoseNames[InputPoseIndex]);
-						if(InputPropertyIndex != INDEX_NONE && AnimBlueprintFunction.InputPoseNodeProperties[InputPropertyIndex])
+						if(InputPropertyIndex != INDEX_NONE && AnimBlueprintFunction.InputPoseNodeProperties[InputPropertyIndex].Get())
 						{
 							FAnimNode_LinkedInputPose* LinkedInputPoseNode = AnimBlueprintFunction.InputPoseNodeProperties[InputPropertyIndex]->ContainerPtrToValuePtr<FAnimNode_LinkedInputPose>(LinkTargetInstance);
 							check(LinkedInputPoseNode->Name == AnimBlueprintFunction.InputPoseNames[InputPoseIndex]);
@@ -343,7 +343,7 @@ void FAnimNode_LinkedAnimGraph::DynamicLink(UAnimInstance* InOwningAnimInstance)
 						}
 					}
 
-					if(AnimBlueprintFunction.OutputPoseNodeProperty)
+					if(AnimBlueprintFunction.OutputPoseNodeProperty.Get())
 					{
 						LinkedRoot = AnimBlueprintFunction.OutputPoseNodeProperty->ContainerPtrToValuePtr<FAnimNode_Root>(LinkTargetInstance);
 					}
@@ -381,7 +381,7 @@ void FAnimNode_LinkedAnimGraph::DynamicUnlink(UAnimInstance* InOwningAnimInstanc
 					for(int32 InputPoseIndex = 0; InputPoseIndex < AnimBlueprintFunction.InputPoseNames.Num() && InputPoseIndex < InputPoses.Num(); ++InputPoseIndex)
 					{
 						int32 InputPropertyIndex = FindFunctionInputIndex(AnimBlueprintFunction, AnimBlueprintFunction.InputPoseNames[InputPoseIndex]);
-						if(InputPropertyIndex != INDEX_NONE && AnimBlueprintFunction.InputPoseNodeProperties[InputPropertyIndex])
+						if(InputPropertyIndex != INDEX_NONE && AnimBlueprintFunction.InputPoseNodeProperties[InputPropertyIndex].Get())
 						{
 							FAnimNode_LinkedInputPose* LinkedInputPoseNode = AnimBlueprintFunction.InputPoseNodeProperties[InputPropertyIndex]->ContainerPtrToValuePtr<FAnimNode_LinkedInputPose>(LinkTargetInstance);
 							check(LinkedInputPoseNode->Name == AnimBlueprintFunction.InputPoseNames[InputPoseIndex]);
diff --git a/Engine/Source/Runtime/Engine/Private/UserDefinedStruct.cpp b/Engine/Source/Runtime/Engine/Private/UserDefinedStruct.cpp
--- a/Engine/Source/Runtime/Engine/Private/UserDefinedStruct.cpp
+++ b/Engine/Source/Runtime/Engine/Private/UserDefinedStruct.cpp
@@ -343,9 +343,9 @@ FGuid UUserDefinedStruct::GetCustomGuid() const
 ENGINE_API FString GetPathPostfix(const UObject* ForObject)
 {
 	FString FullAssetName = ForObject->GetOutermost()->GetPathName();
-	if (FullAssetName.StartsWith(TEXT("/Temp/__TEMP_BP__"), ESearchCase::CaseSensitive))
+	if (FullAssetName.StartsWith(UDynamicClass::GetTempPackagePrefix(), ESearchCase::CaseSensitive))
 	{
-		FullAssetName.RemoveFromStart(TEXT("/Temp/__TEMP_BP__"), ESearchCase::CaseSensitive);
+		FullAssetName.RemoveFromStart(UDynamicClass::GetTempPackagePrefix(), ESearchCase::CaseSensitive);
 	}
 	FString AssetName = FPackageName::GetLongPackageAssetName(FullAssetName);
 	// append a hash of the path, this uniquely identifies assets with the same name, but different folders:
