Minor optimization, do not call GetBulkDataSize when loading FStaticMeshLODResources unless we are going to use the result which only occurs if WITH_EDITORONLY_DATA is enabled.

#rnx

diff --git a/Engine/Source/Runtime/Engine/Private/StaticMesh.cpp b/Engine/Source/Runtime/Engine/Private/StaticMesh.cpp
--- a/Engine/Source/Runtime/Engine/Private/StaticMesh.cpp
+++ b/Engine/Source/Runtime/Engine/Private/StaticMesh.cpp
@@ -642,7 +642,10 @@ void FStaticMeshLODResources::Serialize(FArchive& Ar, UObject* Owner, int32 Inde
 		}
 		else if (FPlatformProperties::RequiresCookedData() || Ar.IsCooking() || bUsingCookedEditorData)
 		{
+#if WITH_EDITORONLY_DATA
 			uint32 BulkDataSize = 0;
+#endif
+
 #if WITH_EDITOR
 			if (Ar.IsSaving())
 			{
@@ -688,9 +691,10 @@ void FStaticMeshLODResources::Serialize(FArchive& Ar, UObject* Owner, int32 Inde
 				StreamingBulkData.Serialize(Ar, Owner, Index, false);
 				bIsOptionalLOD = StreamingBulkData.IsOptional();
 #endif
-				BulkDataSize = (uint32)StreamingBulkData.GetBulkDataSize();
 
 #if WITH_EDITORONLY_DATA
+				BulkDataSize = (uint32)StreamingBulkData.GetBulkDataSize();
+
 				// Streaming CPU data in editor build isn't supported yet because tools and utils need access
 				if (bUsingCookedEditorData && BulkDataSize > 0)
 				{
