Renamed ::InSeperateFile() to ::IsInSeperateFile() to match the coding standards.

- Added a deprecation message to ::InSeperateFile() although I cannot imagine that any externals have started using it.

#rnx

diff --git a/Engine/Source/Runtime/CoreUObject/Private/Serialization/BulkData2.cpp b/Engine/Source/Runtime/CoreUObject/Private/Serialization/BulkData2.cpp
--- a/Engine/Source/Runtime/CoreUObject/Private/Serialization/BulkData2.cpp
+++ b/Engine/Source/Runtime/CoreUObject/Private/Serialization/BulkData2.cpp
@@ -886,7 +886,7 @@ bool FBulkDataBase::IsInlined() const
 	return	(GetBulkDataFlags() & BULKDATA_PayloadAtEndOfFile) == 0;
 }
 
-bool FBulkDataBase::InSeperateFile() const
+bool FBulkDataBase::IsInSeperateFile() const
 {
 	return	(GetBulkDataFlags() & BULKDATA_PayloadInSeperateFile) != 0;
 }
@@ -1175,7 +1175,7 @@ void FBulkDataBase::FreeData()
 
 FString FBulkDataBase::ConvertFilenameFromFlags(const FString& Filename) const
 {
-	if (!InSeperateFile())
+	if (!IsInSeperateFile())
 	{
 		return Filename;
 	}
diff --git a/Engine/Source/Runtime/CoreUObject/Public/Serialization/BulkData.h b/Engine/Source/Runtime/CoreUObject/Public/Serialization/BulkData.h
--- a/Engine/Source/Runtime/CoreUObject/Public/Serialization/BulkData.h
+++ b/Engine/Source/Runtime/CoreUObject/Public/Serialization/BulkData.h
@@ -504,11 +504,14 @@ public:
 		return   (GetBulkDataFlags() & BULKDATA_PayloadAtEndOfFile) == 0;
 	}
 
+	UE_DEPRECATED(4.25, "Use ::IsInSeperateFile() instead")
+	FORCEINLINE bool InSeperateFile() const { return IsInSeperateFile(); }
+
 	/**
 	* Returns whether this bulk data is currently stored in it's own file or not
 	* @return true if BULKDATA_PayloadInSeperateFile is not set
 	*/
-	bool InSeperateFile() const
+	bool IsInSeperateFile() const
 	{
 		return	(GetBulkDataFlags() & BULKDATA_PayloadInSeperateFile) != 0;
 	}
diff --git a/Engine/Source/Runtime/CoreUObject/Public/Serialization/BulkData2.h b/Engine/Source/Runtime/CoreUObject/Public/Serialization/BulkData2.h
--- a/Engine/Source/Runtime/CoreUObject/Public/Serialization/BulkData2.h
+++ b/Engine/Source/Runtime/CoreUObject/Public/Serialization/BulkData2.h
@@ -111,7 +111,9 @@ public:
 	bool IsDuplicateNonOptional() const;
 	bool IsOptional() const;
 	bool IsInlined() const;
-	bool InSeperateFile() const;
+	UE_DEPRECATED(4.25, "Use ::IsInSeperateFile() instead")
+	FORCEINLINE bool InSeperateFile() const { return IsInSeperateFile(); }
+	bool IsInSeperateFile() const;
 	bool IsSingleUse() const;
 	bool IsMemoryMapped() const;
 	bool IsUsingIODispatcher() const;
diff --git a/Engine/Source/Runtime/Engine/Private/Texture2D.cpp b/Engine/Source/Runtime/Engine/Private/Texture2D.cpp
--- a/Engine/Source/Runtime/Engine/Private/Texture2D.cpp
+++ b/Engine/Source/Runtime/Engine/Private/Texture2D.cpp
@@ -540,7 +540,7 @@ bool UTexture2D::GetMipDataFilename(const int32 MipIndex, FString& OutBulkDataFi
 #else
 			OutBulkDataFilename = PlatformData->CachedPackageFileName;
 
-			if (PlatformData->Mips[MipIndex].BulkData.InSeperateFile())
+			if (PlatformData->Mips[MipIndex].BulkData.IsInSeperateFile())
 			{	
 				const bool UseOptionalBulkDataFileName = PlatformData->Mips[MipIndex].BulkData.IsOptional();
 				OutBulkDataFilename = FPaths::ChangeExtension(OutBulkDataFilename, UseOptionalBulkDataFileName ? TEXT(".uptnl") : TEXT(".ubulk"));
diff --git a/Engine/Source/Runtime/Engine/Private/TextureDerivedData.cpp b/Engine/Source/Runtime/Engine/Private/TextureDerivedData.cpp
--- a/Engine/Source/Runtime/Engine/Private/TextureDerivedData.cpp
+++ b/Engine/Source/Runtime/Engine/Private/TextureDerivedData.cpp
@@ -1084,7 +1084,7 @@ bool FTexturePlatformData::TryLoadMips(int32 FirstMipToLoad, void** OutMipData)
 				// is it possible for streaming mips to be loaded in non streaming ways.
 				if (CVarSetTextureStreaming.GetValueOnAnyThread() != 0)
 				{
-					UE_CLOG(Mip.BulkData.InSeperateFile(), LogTexture, Error, TEXT("Loading non-streamed mips from an external bulk file.  This is not desireable.  File %s"), *(Mip.BulkData.GetFilename() ) );
+					UE_CLOG(Mip.BulkData.IsInSeperateFile(), LogTexture, Error, TEXT("Loading non-streamed mips from an external bulk file.  This is not desireable.  File %s"), *(Mip.BulkData.GetFilename() ) );
 				}
 #endif
 				Mip.BulkData.GetCopy(&OutMipData[MipIndex - FirstMipToLoad], true);
@@ -1147,7 +1147,7 @@ int32 FTexturePlatformData::GetNumNonStreamingMips() const
 
 		for (const FTexture2DMipMap& Mip : Mips)
 		{
-			if ( Mip.BulkData.InSeperateFile() || !Mip.BulkData.IsInlined() )
+			if ( Mip.BulkData.IsInSeperateFile() || !Mip.BulkData.IsInlined() )
 			{
 				--NumNonStreamingMips;
 			}
diff --git a/Engine/Source/Runtime/Engine/Public/TextureResource.h b/Engine/Source/Runtime/Engine/Public/TextureResource.h
--- a/Engine/Source/Runtime/Engine/Public/TextureResource.h
+++ b/Engine/Source/Runtime/Engine/Public/TextureResource.h
@@ -74,7 +74,9 @@ struct FTexture2DMipMap
 		bool IsAvailableForUse() const { return !(BulkDataFlags & BULKDATA_Unused); }
 		bool IsOptional() const { return (BulkDataFlags & BULKDATA_OptionalPayload) != 0; }
 		bool IsInlined() const { return (GetBulkDataFlags() & BULKDATA_PayloadInSeperateFile) == 0 && (GetBulkDataFlags() & BULKDATA_PayloadAtEndOfFile) == 0; }
-		bool InSeperateFile() const { return (GetBulkDataFlags() & BULKDATA_PayloadInSeperateFile) != 0; }
+		UE_DEPRECATED(4.25, "Use ::IsInSeperateFile() instead")
+		FORCEINLINE bool InSeperateFile() const { return IsInSeperateFile(); }
+		bool IsInSeperateFile() const { return (GetBulkDataFlags() & BULKDATA_PayloadInSeperateFile) != 0; }
 		bool IsBulkDataLoaded() const { return IsInlined(); }
 		bool IsAsyncLoadingComplete() const { return true; }
 		bool IsStoredCompressedOnDisk() const { return !!(BulkDataFlags & BULKDATA_SerializeCompressed); }
