Fix android compile error caused by my last submit.

#rnx

diff --git a/Engine/Source/Runtime/Engine/Private/Materials/MaterialStaticParameterValueResolver.h b/Engine/Source/Runtime/Engine/Private/Materials/MaterialStaticParameterValueResolver.h
--- a/Engine/Source/Runtime/Engine/Private/Materials/MaterialStaticParameterValueResolver.h
+++ b/Engine/Source/Runtime/Engine/Private/Materials/MaterialStaticParameterValueResolver.h
@@ -69,9 +69,9 @@ public:
 				const UClass* TargetClass = ParameterType::StaticClass();
 
 				auto GetExpressionParameterByNamePredicate =
-					[this, &PendingParametersForFunction, &ParameterIndicesForFunction, &ParametersAndOwners, TargetClass](UMaterialFunctionInterface* Function) -> bool
+					[this, &PendingParametersForFunction, &ParameterIndicesForFunction, &ParametersAndOwners, TargetClass](UMaterialFunctionInterface* InFunction) -> bool
 				{
-					for (UMaterialExpression* FunctionExpression : *Function->GetFunctionExpressions())
+					for (UMaterialExpression* FunctionExpression : *InFunction->GetFunctionExpressions())
 					{
 						if (ParameterType* ExpressionParameter = (FunctionExpression && FunctionExpression->IsA(TargetClass)) ? (ParameterType*)FunctionExpression : nullptr)
 						{
@@ -82,7 +82,7 @@ public:
 
 								if (ExpressionParameter->ParameterName == EvalContext.GetParameterInfo(ParameterIndicesForFunction[FunctionParameterIdx])->Name)
 								{
-									ParametersAndOwners[FunctionParameterIdx] = {ExpressionParameter, Function};
+									ParametersAndOwners[FunctionParameterIdx] = {ExpressionParameter, InFunction};
 
 									if (--PendingParametersForFunction == 0)
 									{
