Do not attempt to calculate a mips filename if we are using the new BulkData api.

UStaticMesh::GetMipDataFilename and USkeletalMesh::GetMipDataFilename do not work when AsyncLoader2 is active and the result is not used making invoking this expensive method a bit pointless.

#rnx

diff --git a/Engine/Source/Runtime/Engine/Private/Streaming/SkeletalMeshUpdate.cpp b/Engine/Source/Runtime/Engine/Private/Streaming/SkeletalMeshUpdate.cpp
--- a/Engine/Source/Runtime/Engine/Private/Streaming/SkeletalMeshUpdate.cpp
+++ b/Engine/Source/Runtime/Engine/Private/Streaming/SkeletalMeshUpdate.cpp
@@ -489,8 +489,14 @@ template <bool bRenderThread>
 void TSkeletalMeshStreamIn_IO<bRenderThread>::DoInitiateIO(const FContext& Context)
 {
 	check(Context.CurrentThread == TT_Async);
+
+#if USE_BULKDATA_STREAMING_TOKEN
 	const FString IOFilename = GetIOFilename(Context);
 	SetIORequest(Context, IOFilename);
+#else
+	SetIORequest(Context, FString());
+#endif
+
 	PushTask(Context, TT_Async, SRA_UPDATE_CALLBACK(DoSerializeLODData), TT_Async, SRA_UPDATE_CALLBACK(DoCancelIO));
 }
 
diff --git a/Engine/Source/Runtime/Engine/Private/Streaming/StaticMeshUpdate.cpp b/Engine/Source/Runtime/Engine/Private/Streaming/StaticMeshUpdate.cpp
--- a/Engine/Source/Runtime/Engine/Private/Streaming/StaticMeshUpdate.cpp
+++ b/Engine/Source/Runtime/Engine/Private/Streaming/StaticMeshUpdate.cpp
@@ -465,8 +465,14 @@ void TStaticMeshStreamIn_IO<bRenderThread>::DoInitiateIO(const FContext& Context
 {
 	DECLARE_SCOPE_CYCLE_COUNTER(TEXT("TStaticMeshStreamIn_IO::DoInitiateIO"), STAT_StaticMeshStreamInIO_DoInitiateIO, STATGROUP_StreamingDetails);
 	check(Context.CurrentThread == TT_Async);
+
+#if USE_BULKDATA_STREAMING_TOKEN
 	const FString IOFilename = GetIOFilename(Context);
 	SetIORequest(Context, IOFilename);
+#else
+	SetIORequest(Context, FString());
+#endif
+
 	PushTask(Context, TT_Async, SRA_UPDATE_CALLBACK(DoSerializeLODData), TT_Async, SRA_UPDATE_CALLBACK(DoCancelIO));
 }
 
