Fixing copyright notice to new (undated) format in a file I added previously.

#rnx

diff --git a/Engine/Source/Runtime/Engine/Private/Materials/MaterialStaticParameterValueResolver.h b/Engine/Source/Runtime/Engine/Private/Materials/MaterialStaticParameterValueResolver.h
--- a/Engine/Source/Runtime/Engine/Private/Materials/MaterialStaticParameterValueResolver.h
+++ b/Engine/Source/Runtime/Engine/Private/Materials/MaterialStaticParameterValueResolver.h
@@ -1,4 +1,4 @@
-// Copyright 1998-2019 Epic Games, Inc. All Rights Reserved.
+// Copyright Epic Games, Inc. All Rights Reserved.
 
 #pragma once
 
