Users of TraceLog now use channels

Systems that use TraceLog now uses channels to control if events are emitted. Also introduce a new command line argument "-trace" which allows users to specify a comma separated list of channel names to enable.

diff --git a/Engine/Source/Editor/UnrealEd/Private/Commandlets/TextAssetCommandlet.cpp b/Engine/Source/Editor/UnrealEd/Private/Commandlets/TextAssetCommandlet.cpp
--- a/Engine/Source/Editor/UnrealEd/Private/Commandlets/TextAssetCommandlet.cpp
+++ b/Engine/Source/Editor/UnrealEd/Private/Commandlets/TextAssetCommandlet.cpp
@@ -359,19 +359,19 @@ bool UTextAssetCommandlet::DoTextAssetProcessing(const FProcessingArgs& InArgs)
 				static const TCHAR* PhaseNames[] = { TEXT("Binary Only"), TEXT("Text Only"), TEXT("Alternating Binary/Text") };
 
 				#if CPUPROFILERTRACE_ENABLED
-				static const uint32 PhaseEventTypes[3] = {
-						FCpuProfilerTrace::OutputEventType(TEXT("BinaryOnly"), CpuProfilerGroup_Default),
-						FCpuProfilerTrace::OutputEventType(TEXT("TextOnly"), CpuProfilerGroup_Default),
-						FCpuProfilerTrace::OutputEventType(TEXT("Alternating"), CpuProfilerGroup_Default)
+				static const TCHAR* PhaseEventTypes[3] = {
+						TEXT("BinaryOnly"),
+						TEXT("TextOnly"),
+						TEXT("Alternating"),
 				};
 
-				static const uint32 TestEventTypes[6] = {
-						FCpuProfilerTrace::OutputEventType(TEXT("Test1"), CpuProfilerGroup_Default),
-						FCpuProfilerTrace::OutputEventType(TEXT("Test2"), CpuProfilerGroup_Default),
-						FCpuProfilerTrace::OutputEventType(TEXT("Test3"), CpuProfilerGroup_Default),
-						FCpuProfilerTrace::OutputEventType(TEXT("Test4"), CpuProfilerGroup_Default),
-						FCpuProfilerTrace::OutputEventType(TEXT("Test5"), CpuProfilerGroup_Default),
-						FCpuProfilerTrace::OutputEventType(TEXT("Test6"), CpuProfilerGroup_Default)
+				static const TCHAR* TestEventTypes[6] = {
+						TEXT("Test1"),
+						TEXT("Test2"),
+						TEXT("Test3"),
+						TEXT("Test4"),
+						TEXT("Test5"),
+						TEXT("Test6"),
 				};
 				#endif // CPUPROFILERTRACE_ENABLED
 
@@ -390,7 +390,7 @@ bool UTextAssetCommandlet::DoTextAssetProcessing(const FProcessingArgs& InArgs)
 					TArray<FSHAHash> PhaseHashes = Hashes[Hashes.AddDefaulted()];
 					
 					#if CPUPROFILERTRACE_ENABLED
-					FCpuProfilerTrace::OutputBeginEvent(PhaseEventTypes[Phase]);
+					TRACE_CPUPROFILER_EVENT_SCOPE_TEXT(PhaseEventTypes[Phase]);
 					#endif
 
 					for (int32 i = 0; i < ((Phase == 2) ? NumTests * 2 : NumTests); ++i)
@@ -398,7 +398,7 @@ bool UTextAssetCommandlet::DoTextAssetProcessing(const FProcessingArgs& InArgs)
 						int32 Bucket;
 
 						#if CPUPROFILERTRACE_ENABLED
-						FCpuProfilerTrace::OutputBeginEvent(TestEventTypes[i]);
+						TRACE_CPUPROFILER_EVENT_SCOPE_TEXT(TestEventTypes[i]);
 						#endif
 
 						switch (Phase)
@@ -479,10 +479,6 @@ bool UTextAssetCommandlet::DoTextAssetProcessing(const FProcessingArgs& InArgs)
 
 							DiffFilenames.Add(TPair<FString, FString>(IntermediateFilename, FinalFilename));
 						}
-
-						#if CPUPROFILERTRACE_ENABLED
-						FCpuProfilerTrace::OutputEndEvent();
-						#endif
 					}
 
 					UE_LOG(LogTextAsset, Display, TEXT("Phase %i (%s) Results"), Phase + 1, PhaseNames[Phase]);
@@ -526,9 +522,6 @@ bool UTextAssetCommandlet::DoTextAssetProcessing(const FProcessingArgs& InArgs)
 
 					DiffFilenames.Empty();
 					IFileManager::Get().DeleteDirectory(*TempFailedDiffsPath, false, true);
-					#if CPUPROFILERTRACE_ENABLED
-					FCpuProfilerTrace::OutputEndEvent();
-					#endif
 				}
 
 				static const bool bDisableCleanup = FParse::Param(FCommandLine::Get(), TEXT("disablecleanup"));
diff --git a/Engine/Source/Runtime/AnimGraphRuntime/Private/AnimGraphRuntimeTrace.cpp b/Engine/Source/Runtime/AnimGraphRuntime/Private/AnimGraphRuntimeTrace.cpp
--- a/Engine/Source/Runtime/AnimGraphRuntime/Private/AnimGraphRuntimeTrace.cpp
+++ b/Engine/Source/Runtime/AnimGraphRuntime/Private/AnimGraphRuntimeTrace.cpp
@@ -30,7 +30,7 @@ void FAnimGraphRuntimeTrace::Init()
 
 void FAnimGraphRuntimeTrace::OutputBlendSpacePlayer(const FAnimationBaseContext& InContext, const FAnimNode_BlendSpacePlayer& InNode)
 {
-	bool bEventEnabled = UE_TRACE_EVENT_IS_ENABLED(Animation, BlendSpacePlayer);
+	bool bEventEnabled = UE_TRACE_CHANNELEXPR_IS_ENABLED(AnimationChannel);
 	if (!bEventEnabled)
 	{
 		return;
@@ -41,7 +41,7 @@ void FAnimGraphRuntimeTrace::OutputBlendSpacePlayer(const FAnimationBaseContext&
 	TRACE_OBJECT(InContext.AnimInstanceProxy->GetAnimInstanceObject());
 	TRACE_OBJECT(InNode.BlendSpace);
 
-	UE_TRACE_LOG(Animation, BlendSpacePlayer)
+	UE_TRACE_LOG(Animation, BlendSpacePlayer, AnimationChannel)
 		<< BlendSpacePlayer.Cycle(FPlatformTime::Cycles64())
 		<< BlendSpacePlayer.AnimInstanceId(FObjectTrace::GetObjectId(InContext.AnimInstanceProxy->GetAnimInstanceObject()))
 		<< BlendSpacePlayer.BlendSpaceId(FObjectTrace::GetObjectId(InNode.BlendSpace))
diff --git a/Engine/Source/Runtime/Core/Private/Logging/LogTrace.cpp b/Engine/Source/Runtime/Core/Private/Logging/LogTrace.cpp
--- a/Engine/Source/Runtime/Core/Private/Logging/LogTrace.cpp
+++ b/Engine/Source/Runtime/Core/Private/Logging/LogTrace.cpp
@@ -8,19 +8,21 @@
 #include "HAL/PlatformTime.h"
 #include "HAL/PlatformTLS.h"
 
-UE_TRACE_EVENT_BEGIN(Logging, LogCategory, Always|Important)
+UE_TRACE_CHANNEL(LogChannel)
+
+UE_TRACE_EVENT_BEGIN(Logging, LogCategory, Important)
 	UE_TRACE_EVENT_FIELD(const void*, CategoryPointer)
 	UE_TRACE_EVENT_FIELD(uint8, DefaultVerbosity)
 UE_TRACE_EVENT_END()
 
-UE_TRACE_EVENT_BEGIN(Logging, LogMessageSpec, Always|Important)
+UE_TRACE_EVENT_BEGIN(Logging, LogMessageSpec, Important)
 	UE_TRACE_EVENT_FIELD(const void*, LogPoint)
 	UE_TRACE_EVENT_FIELD(const void*, CategoryPointer)
 	UE_TRACE_EVENT_FIELD(int32, Line)
 	UE_TRACE_EVENT_FIELD(uint8, Verbosity)
 UE_TRACE_EVENT_END()
 
-UE_TRACE_EVENT_BEGIN(Logging, LogMessage, Always)
+UE_TRACE_EVENT_BEGIN(Logging, LogMessage)
 	UE_TRACE_EVENT_FIELD(const void*, LogPoint)
 	UE_TRACE_EVENT_FIELD(uint64, Cycle)
 	UE_TRACE_EVENT_FIELD(uint32, ThreadId)
@@ -29,7 +31,7 @@ UE_TRACE_EVENT_END()
 void FLogTrace::OutputLogCategory(const FLogCategoryBase* Category, const TCHAR* Name, ELogVerbosity::Type DefaultVerbosity)
 {
 	uint16 NameSize = (FCString::Strlen(Name) + 1) * sizeof(TCHAR);
-	UE_TRACE_LOG(Logging, LogCategory, NameSize)
+	UE_TRACE_LOG(Logging, LogCategory, LogChannel, NameSize)
 		<< LogCategory.CategoryPointer(Category)
 		<< LogCategory.DefaultVerbosity(DefaultVerbosity)
 		<< LogCategory.Attachment(Name, NameSize);
@@ -43,7 +45,7 @@ void FLogTrace::OutputLogMessageSpec(const void* LogPoint, const FLogCategoryBas
 		memcpy(Out, File, FileNameSize);
 		memcpy(Out + FileNameSize, Format, FormatStringSize);
 	};
-	UE_TRACE_LOG(Logging, LogMessageSpec, FileNameSize + FormatStringSize)
+	UE_TRACE_LOG(Logging, LogMessageSpec, LogChannel, FileNameSize + FormatStringSize)
 		<< LogMessageSpec.LogPoint(LogPoint)
 		<< LogMessageSpec.CategoryPointer(Category)
 		<< LogMessageSpec.Line(Line)
@@ -53,7 +55,7 @@ void FLogTrace::OutputLogMessageSpec(const void* LogPoint, const FLogCategoryBas
 
 void FLogTrace::OutputLogMessageInternal(const void* LogPoint, uint16 EncodedFormatArgsSize, uint8* EncodedFormatArgs)
 {
-	UE_TRACE_LOG(Logging, LogMessage, EncodedFormatArgsSize)
+	UE_TRACE_LOG(Logging, LogMessage, LogChannel, EncodedFormatArgsSize)
 		<< LogMessage.LogPoint(LogPoint)
 		<< LogMessage.Cycle(FPlatformTime::Cycles64())
 		<< LogMessage.ThreadId(FPlatformTLS::GetCurrentThreadId())
diff --git a/Engine/Source/Runtime/Core/Private/ProfilingDebugging/CountersTrace.cpp b/Engine/Source/Runtime/Core/Private/ProfilingDebugging/CountersTrace.cpp
--- a/Engine/Source/Runtime/Core/Private/ProfilingDebugging/CountersTrace.cpp
+++ b/Engine/Source/Runtime/Core/Private/ProfilingDebugging/CountersTrace.cpp
@@ -6,7 +6,9 @@
 
 #if COUNTERSTRACE_ENABLED
 
-UE_TRACE_EVENT_BEGIN(Counters, Spec, Always)
+UE_TRACE_CHANNEL_DEFINE(CountersChannel)
+
+UE_TRACE_EVENT_BEGIN(Counters, Spec, Important)
 	UE_TRACE_EVENT_FIELD(uint16, Id)
 	UE_TRACE_EVENT_FIELD(uint8, Type)
 	UE_TRACE_EVENT_FIELD(uint8, DisplayHint)
@@ -28,9 +30,11 @@ UE_TRACE_EVENT_END()
 uint16 FCountersTrace::OutputInitCounter(const TCHAR* CounterName, ETraceCounterType CounterType, ETraceCounterDisplayHint CounterDisplayHint)
 {
 	static TAtomic<uint16> NextId(1);
+	if (!UE_TRACE_CHANNELEXPR_IS_ENABLED(CountersChannel))
+		return 0;
 	uint16 CounterId = uint16(NextId++);
 	uint16 NameSize = (FCString::Strlen(CounterName) + 1) * sizeof(TCHAR);
-	UE_TRACE_LOG(Counters, Spec, NameSize)
+	UE_TRACE_LOG(Counters, Spec, CountersChannel, NameSize)
 		<< Spec.Id(CounterId)
 		<< Spec.Type(uint8(CounterType))
 		<< Spec.DisplayHint(uint8(CounterDisplayHint))
@@ -40,7 +44,7 @@ uint16 FCountersTrace::OutputInitCounter(const TCHAR* CounterName, ETraceCounter
 
 void FCountersTrace::OutputSetValue(uint16 CounterId, int64 Value)
 {
-	UE_TRACE_LOG(Counters, SetValueInt)
+	UE_TRACE_LOG(Counters, SetValueInt, CountersChannel && CounterId)
 		<< SetValueInt.Cycle(FPlatformTime::Cycles64())
 		<< SetValueInt.Value(Value)
 		<< SetValueInt.CounterId(CounterId);
@@ -48,7 +52,7 @@ void FCountersTrace::OutputSetValue(uint16 CounterId, int64 Value)
 
 void FCountersTrace::OutputSetValue(uint16 CounterId, double Value)
 {
-	UE_TRACE_LOG(Counters, SetValueFloat)
+	UE_TRACE_LOG(Counters, SetValueFloat, CountersChannel && CounterId)
 		<< SetValueFloat.Cycle(FPlatformTime::Cycles64())
 		<< SetValueFloat.Value(Value)
 		<< SetValueFloat.CounterId(CounterId);
@@ -60,7 +64,7 @@ void FCountersTrace::Init(const TCHAR* CmdLine)
 	{
 		UE_TRACE_EVENT_IS_ENABLED(Counters, SetValueInt);
 		UE_TRACE_EVENT_IS_ENABLED(Counters, SetValueFloat);
-		Trace::ToggleEvent(TEXT("Counters"), true);
+		Trace::ToggleChannel(CountersChannel, true);
 	}
 }
 
diff --git a/Engine/Source/Runtime/Core/Private/ProfilingDebugging/CpuProfilerTrace.cpp b/Engine/Source/Runtime/Core/Private/ProfilingDebugging/CpuProfilerTrace.cpp
--- a/Engine/Source/Runtime/Core/Private/ProfilingDebugging/CpuProfilerTrace.cpp
+++ b/Engine/Source/Runtime/Core/Private/ProfilingDebugging/CpuProfilerTrace.cpp
@@ -12,9 +12,11 @@
 
 #if CPUPROFILERTRACE_ENABLED
 
-UE_TRACE_EVENT_BEGIN(CpuProfiler, EventSpec, Always)
+UE_TRACE_CHANNEL_DEFINE(CpuProfilerChannel)
+UE_TRACE_CHANNEL_DEFINE(NamedEventsChannel)
+
+UE_TRACE_EVENT_BEGIN(CpuProfiler, EventSpec, Important)
 	UE_TRACE_EVENT_FIELD(uint32, Id)
-	UE_TRACE_EVENT_FIELD(uint16, Group)
 	UE_TRACE_EVENT_FIELD(uint8, CharSize)
 UE_TRACE_EVENT_END()
 
@@ -22,7 +24,7 @@ UE_TRACE_EVENT_BEGIN(CpuProfiler, EventBatch)
 	UE_TRACE_EVENT_FIELD(uint32, ThreadId)
 UE_TRACE_EVENT_END()
 
-UE_TRACE_EVENT_BEGIN(CpuProfiler, EndCapture, Always)
+UE_TRACE_EVENT_BEGIN(CpuProfiler, EndCapture, Important)
 	UE_TRACE_EVENT_FIELD(uint32, ThreadId)
 UE_TRACE_EVENT_END()
 
@@ -99,7 +101,7 @@ FCpuProfilerTraceInternal::FThreadBuffer* FCpuProfilerTraceInternal::CreateThrea
 
 void FCpuProfilerTraceInternal::FlushThreadBuffer(FThreadBuffer* InThreadBuffer)
 {
-	UE_TRACE_LOG(CpuProfiler, EventBatch, InThreadBuffer->BufferSize)
+	UE_TRACE_LOG(CpuProfiler, EventBatch, CpuProfilerChannel, InThreadBuffer->BufferSize)
 		<< EventBatch.ThreadId(FPlatformTLS::GetCurrentThreadId())
 		<< EventBatch.Attachment(InThreadBuffer->Buffer, InThreadBuffer->BufferSize);
 	InThreadBuffer->BufferSize = 0;
@@ -107,7 +109,7 @@ void FCpuProfilerTraceInternal::FlushThreadBuffer(FThreadBuffer* InThreadBuffer)
 
 void FCpuProfilerTraceInternal::EndCapture(FThreadBuffer* InThreadBuffer)
 {
-	UE_TRACE_LOG(CpuProfiler, EndCapture, InThreadBuffer->BufferSize)
+	UE_TRACE_LOG(CpuProfiler, EndCapture, CpuProfilerChannel, InThreadBuffer->BufferSize)
 		<< EndCapture.ThreadId(FPlatformTLS::GetCurrentThreadId())
 		<< EndCapture.Attachment(InThreadBuffer->Buffer, InThreadBuffer->BufferSize);
 	InThreadBuffer->BufferSize = 0;
@@ -115,11 +117,12 @@ void FCpuProfilerTraceInternal::EndCapture(FThreadBuffer* InThreadBuffer)
 
 #define CPUPROFILERTRACE_OUTPUTBEGINEVENT_PROLOGUE() \
 	++FCpuProfilerTraceInternal::ThreadDepth; \
-	bool bEventEnabled = UE_TRACE_EVENT_IS_ENABLED(CpuProfiler, EventBatch); \
+	const bool bEventEnabled = UE_TRACE_CHANNELEXPR_IS_ENABLED(Channel); \
 	if ((!bEventEnabled) & (bEventEnabled == FCpuProfilerTraceInternal::bThreadEnabled)) \
 	{ \
 		return; \
 	} \
+	check(FCpuProfilerTraceInternal::ThreadDepth < 1024); \
 	FCpuProfilerTraceInternal::bThreadEnabled = bEventEnabled; \
 	FCpuProfilerTraceInternal::FThreadBuffer* ThreadBuffer = FCpuProfilerTraceInternal::ThreadBuffer; \
 	if (!ThreadBuffer) \
@@ -145,13 +148,13 @@ void FCpuProfilerTraceInternal::EndCapture(FThreadBuffer* InThreadBuffer)
 		FCpuProfilerTraceInternal::FlushThreadBuffer(ThreadBuffer); \
 	}
 
-void FCpuProfilerTrace::OutputBeginEvent(uint32 SpecId)
+void FCpuProfilerTrace::OutputBeginEvent(uint32 SpecId, const Trace::FChannel& Channel)
 {
 	CPUPROFILERTRACE_OUTPUTBEGINEVENT_PROLOGUE();
 	CPUPROFILERTRACE_OUTPUTBEGINEVENT_EPILOGUE();
 }
 
-void FCpuProfilerTrace::OutputBeginDynamicEvent(const ANSICHAR* Name, ECpuProfilerGroup Group)
+void FCpuProfilerTrace::OutputBeginDynamicEvent(const ANSICHAR* Name, const Trace::FChannel& Channel)
 {
 	CPUPROFILERTRACE_OUTPUTBEGINEVENT_PROLOGUE();
 	uint32 SpecId = ThreadBuffer->DynamicAnsiScopeNamesMap.FindRef(Name);
@@ -160,13 +163,13 @@ void FCpuProfilerTrace::OutputBeginDynamicEvent(const ANSICHAR* Name, ECpuProfil
 		int32 NameSize = strlen(Name) + 1;
 		ANSICHAR* NameCopy = reinterpret_cast<ANSICHAR*>(ThreadBuffer->DynamicScopeNamesMemory.Alloc(NameSize, alignof(ANSICHAR)));
 		FMemory::Memmove(NameCopy, Name, NameSize);
-		SpecId = OutputEventType(NameCopy, Group);
+		SpecId = OutputEventType(NameCopy);
 		ThreadBuffer->DynamicAnsiScopeNamesMap.Add(NameCopy, SpecId);
 	}
 	CPUPROFILERTRACE_OUTPUTBEGINEVENT_EPILOGUE();
 }
 
-void FCpuProfilerTrace::OutputBeginDynamicEvent(const TCHAR* Name, ECpuProfilerGroup Group)
+void FCpuProfilerTrace::OutputBeginDynamicEvent(const TCHAR* Name, const Trace::FChannel& Channel)
 {
 	CPUPROFILERTRACE_OUTPUTBEGINEVENT_PROLOGUE();
 	uint32 SpecId = ThreadBuffer->DynamicTCharScopeNamesMap.FindRef(Name);
@@ -175,7 +178,7 @@ void FCpuProfilerTrace::OutputBeginDynamicEvent(const TCHAR* Name, ECpuProfilerG
 		int32 NameSize = (FCString::Strlen(Name) + 1) * sizeof(TCHAR);
 		TCHAR* NameCopy = reinterpret_cast<TCHAR*>(ThreadBuffer->DynamicScopeNamesMemory.Alloc(NameSize, alignof(TCHAR)));
 		FMemory::Memmove(NameCopy, Name, NameSize);
-		SpecId = OutputEventType(NameCopy, Group);
+		SpecId = OutputEventType(NameCopy);
 		ThreadBuffer->DynamicTCharScopeNamesMap.Add(NameCopy, SpecId);
 	}
 	CPUPROFILERTRACE_OUTPUTBEGINEVENT_EPILOGUE();
@@ -187,7 +190,7 @@ void FCpuProfilerTrace::OutputBeginDynamicEvent(const TCHAR* Name, ECpuProfilerG
 void FCpuProfilerTrace::OutputEndEvent()
 {
 	--FCpuProfilerTraceInternal::ThreadDepth;
-	bool bEventEnabled = UE_TRACE_EVENT_IS_ENABLED(CpuProfiler, EventBatch);
+	const bool bEventEnabled = UE_TRACE_CHANNELEXPR_IS_ENABLED(CpuProfilerChannel);
 	if ((!bEventEnabled) & (bEventEnabled == FCpuProfilerTraceInternal::bThreadEnabled))
 	{
 		return;
@@ -221,25 +224,23 @@ uint32 FCpuProfilerTraceInternal::GetNextSpecId()
 	return NextSpecId++;
 }
 
-uint32 FCpuProfilerTrace::OutputEventType(const TCHAR* Name, ECpuProfilerGroup Group)
+uint32 FCpuProfilerTrace::OutputEventType(const TCHAR* Name)
 {
 	uint32 SpecId = FCpuProfilerTraceInternal::GetNextSpecId();
 	uint16 NameSize = (FCString::Strlen(Name) + 1) * sizeof(TCHAR);
-	UE_TRACE_LOG(CpuProfiler, EventSpec, NameSize)
+	UE_TRACE_LOG(CpuProfiler, EventSpec, CpuProfilerChannel, NameSize)
 		<< EventSpec.Id(SpecId)
-		<< EventSpec.Group(uint16(Group))
 		<< EventSpec.CharSize(sizeof(TCHAR))
 		<< EventSpec.Attachment(Name, NameSize);
 	return SpecId;
 }
 
-uint32 FCpuProfilerTrace::OutputEventType(const ANSICHAR* Name, ECpuProfilerGroup Group)
+uint32 FCpuProfilerTrace::OutputEventType(const ANSICHAR* Name)
 {
 	uint32 SpecId = FCpuProfilerTraceInternal::GetNextSpecId();
 	uint16 NameSize = strlen(Name) + 1;
-	UE_TRACE_LOG(CpuProfiler, EventSpec, NameSize)
+	UE_TRACE_LOG(CpuProfiler, EventSpec, CpuProfilerChannel, NameSize)
 		<< EventSpec.Id(SpecId)
-		<< EventSpec.Group(uint16(Group))
 		<< EventSpec.CharSize(1)
 		<< EventSpec.Attachment(Name, NameSize);
 	return SpecId;
@@ -250,13 +251,14 @@ void FCpuProfilerTrace::Init(const TCHAR* CmdLine)
 	if (FParse::Param(CmdLine, TEXT("cpuprofilertrace")))
 	{
 		UE_TRACE_EVENT_IS_ENABLED(CpuProfiler, EventBatch);
-		Trace::ToggleEvent(TEXT("CpuProfiler.EventBatch"), true);
+		Trace::ToggleChannel(CpuProfilerChannel, true);
 	}
 }
 
 void FCpuProfilerTrace::Shutdown()
 {
-	Trace::ToggleEvent(TEXT("CpuProfiler.EventBatch"), false);
+	//todo: Was here for a specific reason?
+	//Trace::ToggleEvent(TEXT("CpuProfiler.EventBatch"), false); 
 	if (FCpuProfilerTraceInternal::ThreadBuffer)
 	{
 		delete FCpuProfilerTraceInternal::ThreadBuffer;
diff --git a/Engine/Source/Runtime/Core/Private/ProfilingDebugging/CsvProfilerTrace.cpp b/Engine/Source/Runtime/Core/Private/ProfilingDebugging/CsvProfilerTrace.cpp
--- a/Engine/Source/Runtime/Core/Private/ProfilingDebugging/CsvProfilerTrace.cpp
+++ b/Engine/Source/Runtime/Core/Private/ProfilingDebugging/CsvProfilerTrace.cpp
@@ -1,51 +1,52 @@
 // Copyright Epic Games, Inc. All Rights Reserved.
 #include "ProfilingDebugging/CsvProfilerTrace.h"
 #include "ProfilingDebugging/CsvProfiler.h"
+#include "ProfilingDebugging/CountersTrace.h"
 #include "HAL/PlatformTime.h"
 #include "UObject/NameTypes.h"
 #include "Trace/Trace.h"
 
 #if CSVPROFILERTRACE_ENABLED
 
-UE_TRACE_EVENT_BEGIN(CsvProfiler, RegisterCategory, Always)
+UE_TRACE_EVENT_BEGIN(CsvProfiler, RegisterCategory, Important)
 	UE_TRACE_EVENT_FIELD(int32, Index)
 UE_TRACE_EVENT_END()
 
-UE_TRACE_EVENT_BEGIN(CsvProfiler, DefineInlineStat, Always)
+UE_TRACE_EVENT_BEGIN(CsvProfiler, DefineInlineStat, Important)
 	UE_TRACE_EVENT_FIELD(uint64, StatId)
 	UE_TRACE_EVENT_FIELD(int32, CategoryIndex)
 UE_TRACE_EVENT_END()
 
-UE_TRACE_EVENT_BEGIN(CsvProfiler, DefineDeclaredStat, Always)
+UE_TRACE_EVENT_BEGIN(CsvProfiler, DefineDeclaredStat, Important)
 	UE_TRACE_EVENT_FIELD(uint64, StatId)
 	UE_TRACE_EVENT_FIELD(int32, CategoryIndex)
 UE_TRACE_EVENT_END()
 
-UE_TRACE_EVENT_BEGIN(CsvProfiler, BeginStat, Always)
+UE_TRACE_EVENT_BEGIN(CsvProfiler, BeginStat)
 	UE_TRACE_EVENT_FIELD(uint64, StatId)
 	UE_TRACE_EVENT_FIELD(uint64, Cycle)
 	UE_TRACE_EVENT_FIELD(uint32, ThreadId)
 UE_TRACE_EVENT_END()
 
-UE_TRACE_EVENT_BEGIN(CsvProfiler, EndStat, Always)
+UE_TRACE_EVENT_BEGIN(CsvProfiler, EndStat)
 	UE_TRACE_EVENT_FIELD(uint64, StatId)
 	UE_TRACE_EVENT_FIELD(uint64, Cycle)
 	UE_TRACE_EVENT_FIELD(uint32, ThreadId)
 UE_TRACE_EVENT_END()
 
-UE_TRACE_EVENT_BEGIN(CsvProfiler, BeginExclusiveStat, Always)
+UE_TRACE_EVENT_BEGIN(CsvProfiler, BeginExclusiveStat)
 	UE_TRACE_EVENT_FIELD(uint64, StatId)
 	UE_TRACE_EVENT_FIELD(uint64, Cycle)
 	UE_TRACE_EVENT_FIELD(uint32, ThreadId)
 UE_TRACE_EVENT_END()
 
-UE_TRACE_EVENT_BEGIN(CsvProfiler, EndExclusiveStat, Always)
+UE_TRACE_EVENT_BEGIN(CsvProfiler, EndExclusiveStat)
 	UE_TRACE_EVENT_FIELD(uint64, StatId)
 	UE_TRACE_EVENT_FIELD(uint64, Cycle)
 	UE_TRACE_EVENT_FIELD(uint32, ThreadId)
 UE_TRACE_EVENT_END()
 
-UE_TRACE_EVENT_BEGIN(CsvProfiler, CustomStatInt, Always)
+UE_TRACE_EVENT_BEGIN(CsvProfiler, CustomStatInt)
 	UE_TRACE_EVENT_FIELD(uint64, StatId)
 	UE_TRACE_EVENT_FIELD(uint64, Cycle)
 	UE_TRACE_EVENT_FIELD(uint32, ThreadId)
@@ -53,7 +54,7 @@ UE_TRACE_EVENT_BEGIN(CsvProfiler, CustomStatInt, Always)
 	UE_TRACE_EVENT_FIELD(uint8, OpType)
 UE_TRACE_EVENT_END()
 
-UE_TRACE_EVENT_BEGIN(CsvProfiler, CustomStatFloat, Always)
+UE_TRACE_EVENT_BEGIN(CsvProfiler, CustomStatFloat)
 	UE_TRACE_EVENT_FIELD(uint64, StatId)
 	UE_TRACE_EVENT_FIELD(uint64, Cycle)
 	UE_TRACE_EVENT_FIELD(uint32, ThreadId)
@@ -61,24 +62,24 @@ UE_TRACE_EVENT_BEGIN(CsvProfiler, CustomStatFloat, Always)
 	UE_TRACE_EVENT_FIELD(uint8, OpType)
 UE_TRACE_EVENT_END()
 
-UE_TRACE_EVENT_BEGIN(CsvProfiler, Event, Always)
+UE_TRACE_EVENT_BEGIN(CsvProfiler, Event)
 	UE_TRACE_EVENT_FIELD(uint64, Cycle)
 	UE_TRACE_EVENT_FIELD(int32, CategoryIndex)
 	UE_TRACE_EVENT_FIELD(uint32, ThreadId)
 UE_TRACE_EVENT_END()
 
-UE_TRACE_EVENT_BEGIN(CsvProfiler, BeginCapture, Always)
+UE_TRACE_EVENT_BEGIN(CsvProfiler, BeginCapture)
 	UE_TRACE_EVENT_FIELD(uint64, Cycle)
 	UE_TRACE_EVENT_FIELD(uint32, RenderThreadId)
 	UE_TRACE_EVENT_FIELD(uint32, RHIThreadId)
 	UE_TRACE_EVENT_FIELD(bool, EnableCounts)
 UE_TRACE_EVENT_END()
 
-UE_TRACE_EVENT_BEGIN(CsvProfiler, EndCapture, Always)
+UE_TRACE_EVENT_BEGIN(CsvProfiler, EndCapture)
 	UE_TRACE_EVENT_FIELD(uint64, Cycle)
 UE_TRACE_EVENT_END()
 
-UE_TRACE_EVENT_BEGIN(CsvProfiler, Metadata, Always)
+UE_TRACE_EVENT_BEGIN(CsvProfiler, Metadata)
 	UE_TRACE_EVENT_FIELD(uint16, ValueOffset)
 UE_TRACE_EVENT_END()
 
@@ -117,7 +118,7 @@ struct FCsvProfilerTraceInternal
 void FCsvProfilerTrace::OutputRegisterCategory(int32 Index, const TCHAR* Name)
 {
 	uint16 NameSize = (FCString::Strlen(Name) + 1) * sizeof(TCHAR);
-	UE_TRACE_LOG(CsvProfiler, RegisterCategory, NameSize)
+	UE_TRACE_LOG(CsvProfiler, RegisterCategory, CountersChannel, NameSize)
 		<< RegisterCategory.Index(Index)
 		<< RegisterCategory.Attachment(Name, NameSize);
 }
@@ -125,7 +126,7 @@ void FCsvProfilerTrace::OutputRegisterCategory(int32 Index, const TCHAR* Name)
 void FCsvProfilerTrace::OutputInlineStat(const char* StatName, int32 CategoryIndex)
 {
 	uint16 NameSize = (strlen(StatName) + 1) * sizeof(char);
-	UE_TRACE_LOG(CsvProfiler, DefineInlineStat, NameSize)
+	UE_TRACE_LOG(CsvProfiler, DefineInlineStat, CountersChannel, NameSize)
 		<< DefineInlineStat.StatId(FCsvProfilerTraceInternal::GetStatId(StatName, CategoryIndex))
 		<< DefineInlineStat.CategoryIndex(CategoryIndex)
 		<< DefineInlineStat.Attachment(StatName, NameSize);
@@ -143,7 +144,7 @@ void FCsvProfilerTrace::OutputDeclaredStat(const FName& StatName, int32 Category
 	TCHAR NameString[NAME_SIZE];
 	StatName.GetPlainNameString(NameString);
 	uint16 NameSize = (FCString::Strlen(NameString) + 1) * sizeof(TCHAR);
-	UE_TRACE_LOG(CsvProfiler, DefineDeclaredStat, NameSize)
+	UE_TRACE_LOG(CsvProfiler, DefineDeclaredStat, CountersChannel, NameSize)
 		<< DefineDeclaredStat.StatId(FCsvProfilerTraceInternal::GetStatId(StatName, CategoryIndex))
 		<< DefineDeclaredStat.CategoryIndex(CategoryIndex)
 		<< DefineDeclaredStat.Attachment(NameString, NameSize);
@@ -151,7 +152,7 @@ void FCsvProfilerTrace::OutputDeclaredStat(const FName& StatName, int32 Category
 
 void FCsvProfilerTrace::OutputBeginStat(const char* StatName, int32 CategoryIndex, uint64 Cycles)
 {
-	UE_TRACE_LOG(CsvProfiler, BeginStat)
+	UE_TRACE_LOG(CsvProfiler, BeginStat, CountersChannel)
 		<< BeginStat.StatId(FCsvProfilerTraceInternal::GetStatId(StatName, CategoryIndex))
 		<< BeginStat.Cycle(Cycles)
 		<< BeginStat.ThreadId(FPlatformTLS::GetCurrentThreadId());
@@ -159,7 +160,7 @@ void FCsvProfilerTrace::OutputBeginStat(const char* StatName, int32 CategoryInde
 
 void FCsvProfilerTrace::OutputBeginStat(const FName& StatName, int32 CategoryIndex, uint64 Cycles)
 {
-	UE_TRACE_LOG(CsvProfiler, BeginStat)
+	UE_TRACE_LOG(CsvProfiler, BeginStat, CountersChannel)
 		<< BeginStat.StatId(FCsvProfilerTraceInternal::GetStatId(StatName, CategoryIndex))
 		<< BeginStat.Cycle(Cycles)
 		<< BeginStat.ThreadId(FPlatformTLS::GetCurrentThreadId());
@@ -167,7 +168,7 @@ void FCsvProfilerTrace::OutputBeginStat(const FName& StatName, int32 CategoryInd
 
 void FCsvProfilerTrace::OutputEndStat(const char* StatName, int32 CategoryIndex, uint64 Cycles)
 {
-	UE_TRACE_LOG(CsvProfiler, EndStat)
+	UE_TRACE_LOG(CsvProfiler, EndStat, CountersChannel)
 		<< EndStat.StatId(FCsvProfilerTraceInternal::GetStatId(StatName, CategoryIndex))
 		<< EndStat.Cycle(Cycles)
 		<< EndStat.ThreadId(FPlatformTLS::GetCurrentThreadId());
@@ -175,7 +176,7 @@ void FCsvProfilerTrace::OutputEndStat(const char* StatName, int32 CategoryIndex,
 
 void FCsvProfilerTrace::OutputEndStat(const FName& StatName, int32 CategoryIndex, uint64 Cycles)
 {
-	UE_TRACE_LOG(CsvProfiler, EndStat)
+	UE_TRACE_LOG(CsvProfiler, EndStat, CountersChannel)
 		<< EndStat.StatId(FCsvProfilerTraceInternal::GetStatId(StatName, CategoryIndex))
 		<< EndStat.Cycle(Cycles)
 		<< EndStat.ThreadId(FPlatformTLS::GetCurrentThreadId());
@@ -183,7 +184,7 @@ void FCsvProfilerTrace::OutputEndStat(const FName& StatName, int32 CategoryIndex
 
 void FCsvProfilerTrace::OutputBeginExclusiveStat(const char* StatName, int32 CategoryIndex, uint64 Cycles)
 {
-	UE_TRACE_LOG(CsvProfiler, BeginExclusiveStat)
+	UE_TRACE_LOG(CsvProfiler, BeginExclusiveStat, CountersChannel)
 		<< BeginExclusiveStat.StatId(FCsvProfilerTraceInternal::GetStatId(StatName, CategoryIndex))
 		<< BeginExclusiveStat.Cycle(Cycles)
 		<< BeginExclusiveStat.ThreadId(FPlatformTLS::GetCurrentThreadId());
@@ -191,7 +192,7 @@ void FCsvProfilerTrace::OutputBeginExclusiveStat(const char* StatName, int32 Cat
 
 void FCsvProfilerTrace::OutputEndExclusiveStat(const char* StatName, int32 CategoryIndex, uint64 Cycles)
 {
-	UE_TRACE_LOG(CsvProfiler, EndExclusiveStat)
+	UE_TRACE_LOG(CsvProfiler, EndExclusiveStat, CountersChannel)
 		<< EndExclusiveStat.StatId(FCsvProfilerTraceInternal::GetStatId(StatName, CategoryIndex))
 		<< EndExclusiveStat.Cycle(Cycles)
 		<< EndExclusiveStat.ThreadId(FPlatformTLS::GetCurrentThreadId());
@@ -199,7 +200,7 @@ void FCsvProfilerTrace::OutputEndExclusiveStat(const char* StatName, int32 Categ
 
 void FCsvProfilerTrace::OutputCustomStat(const char* StatName, int32 CategoryIndex, int32 Value, uint8 OpType, uint64 Cycles)
 {
-	UE_TRACE_LOG(CsvProfiler, CustomStatInt)
+	UE_TRACE_LOG(CsvProfiler, CustomStatInt, CountersChannel)
 		<< CustomStatInt.Cycle(Cycles)
 		<< CustomStatInt.StatId(FCsvProfilerTraceInternal::GetStatId(StatName, CategoryIndex))
 		<< CustomStatInt.ThreadId(FPlatformTLS::GetCurrentThreadId())
@@ -209,7 +210,7 @@ void FCsvProfilerTrace::OutputCustomStat(const char* StatName, int32 CategoryInd
 
 void FCsvProfilerTrace::OutputCustomStat(const FName& StatName, int32 CategoryIndex, int32 Value, uint8 OpType, uint64 Cycles)
 {
-	UE_TRACE_LOG(CsvProfiler, CustomStatInt)
+	UE_TRACE_LOG(CsvProfiler, CustomStatInt, CountersChannel)
 		<< CustomStatInt.Cycle(Cycles)
 		<< CustomStatInt.StatId(FCsvProfilerTraceInternal::GetStatId(StatName, CategoryIndex))
 		<< CustomStatInt.ThreadId(FPlatformTLS::GetCurrentThreadId())
@@ -219,7 +220,7 @@ void FCsvProfilerTrace::OutputCustomStat(const FName& StatName, int32 CategoryIn
 
 void FCsvProfilerTrace::OutputCustomStat(const char* StatName, int32 CategoryIndex, float Value, uint8 OpType, uint64 Cycles)
 {
-	UE_TRACE_LOG(CsvProfiler, CustomStatFloat)
+	UE_TRACE_LOG(CsvProfiler, CustomStatFloat, CountersChannel)
 		<< CustomStatFloat.Cycle(Cycles)
 		<< CustomStatFloat.StatId(FCsvProfilerTraceInternal::GetStatId(StatName, CategoryIndex))
 		<< CustomStatFloat.ThreadId(FPlatformTLS::GetCurrentThreadId())
@@ -229,7 +230,7 @@ void FCsvProfilerTrace::OutputCustomStat(const char* StatName, int32 CategoryInd
 
 void FCsvProfilerTrace::OutputCustomStat(const FName& StatName, int32 CategoryIndex, float Value, uint8 OpType, uint64 Cycles)
 {
-	UE_TRACE_LOG(CsvProfiler, CustomStatFloat)
+	UE_TRACE_LOG(CsvProfiler, CustomStatFloat, CountersChannel)
 		<< CustomStatFloat.Cycle(Cycles)
 		<< CustomStatFloat.StatId(FCsvProfilerTraceInternal::GetStatId(StatName, CategoryIndex))
 		<< CustomStatFloat.ThreadId(FPlatformTLS::GetCurrentThreadId())
@@ -241,7 +242,7 @@ void FCsvProfilerTrace::OutputBeginCapture(const TCHAR* Filename, uint32 RenderT
 {
 	OutputInlineStat(DefaultWaitStatName, CSV_CATEGORY_INDEX(Exclusive));
 	uint16 NameSize = (FCString::Strlen(Filename) + 1) * sizeof(TCHAR);
-	UE_TRACE_LOG(CsvProfiler, BeginCapture, NameSize)
+	UE_TRACE_LOG(CsvProfiler, BeginCapture, CountersChannel, NameSize)
 		<< BeginCapture.Cycle(FPlatformTime::Cycles64())
 		<< BeginCapture.RenderThreadId(RenderThreadId)
 		<< BeginCapture.RHIThreadId(RHIThreadId)
@@ -252,7 +253,7 @@ void FCsvProfilerTrace::OutputBeginCapture(const TCHAR* Filename, uint32 RenderT
 void FCsvProfilerTrace::OutputEvent(const TCHAR* Text, int32 CategoryIndex, uint64 Cycles)
 {
 	uint16 TextSize = (FCString::Strlen(Text) + 1) * sizeof(TCHAR);
-	UE_TRACE_LOG(CsvProfiler, Event, TextSize)
+	UE_TRACE_LOG(CsvProfiler, Event, CountersChannel, TextSize)
 		<< Event.Cycle(Cycles)
 		<< Event.CategoryIndex(CategoryIndex)
 		<< Event.ThreadId(FPlatformTLS::GetCurrentThreadId())
@@ -261,7 +262,7 @@ void FCsvProfilerTrace::OutputEvent(const TCHAR* Text, int32 CategoryIndex, uint
 
 void FCsvProfilerTrace::OutputEndCapture()
 {
-	UE_TRACE_LOG(CsvProfiler, EndCapture)
+	UE_TRACE_LOG(CsvProfiler, EndCapture, CountersChannel)
 		<< EndCapture.Cycle(FPlatformTime::Cycles64());
 }
 
@@ -274,7 +275,7 @@ void FCsvProfilerTrace::OutputMetadata(const TCHAR* Key, const TCHAR* Value)
 		memcpy(Out, Key, KeySize);
 		memcpy(Out + KeySize, Value, ValueSize);
 	};
-	UE_TRACE_LOG(CsvProfiler, Metadata, KeySize + ValueSize)
+	UE_TRACE_LOG(CsvProfiler, Metadata, CountersChannel, KeySize + ValueSize)
 		<< Metadata.ValueOffset(KeySize)
 		<< Metadata.Attachment(Attachment);
 }
diff --git a/Engine/Source/Runtime/Core/Private/ProfilingDebugging/MiscTrace.cpp b/Engine/Source/Runtime/Core/Private/ProfilingDebugging/MiscTrace.cpp
--- a/Engine/Source/Runtime/Core/Private/ProfilingDebugging/MiscTrace.cpp
+++ b/Engine/Source/Runtime/Core/Private/ProfilingDebugging/MiscTrace.cpp
@@ -8,49 +8,52 @@
 #include "HAL/PlatformTLS.h"
 #include "HAL/PlatformTime.h"
 
-UE_TRACE_EVENT_BEGIN(Misc, RegisterGameThread, Always)
+UE_TRACE_CHANNEL(FrameChannel)
+UE_TRACE_CHANNEL(BookmarkChannel)
+
+UE_TRACE_EVENT_BEGIN(Misc, RegisterGameThread, Important)
 	UE_TRACE_EVENT_FIELD(uint32, ThreadId)
 UE_TRACE_EVENT_END()
 
-UE_TRACE_EVENT_BEGIN(Misc, CreateThread, Always)
+UE_TRACE_EVENT_BEGIN(Misc, CreateThread, Important)
 	UE_TRACE_EVENT_FIELD(uint32, CurrentThreadId)
 	UE_TRACE_EVENT_FIELD(uint32, CreatedThreadId)
 	UE_TRACE_EVENT_FIELD(uint32, Priority)
 	UE_TRACE_EVENT_FIELD(uint16, NameSize)
 UE_TRACE_EVENT_END()
 
-UE_TRACE_EVENT_BEGIN(Misc, SetThreadGroup, Always)
+UE_TRACE_EVENT_BEGIN(Misc, SetThreadGroup, Important)
 	UE_TRACE_EVENT_FIELD(uint32, ThreadId)
 UE_TRACE_EVENT_END()
 
-UE_TRACE_EVENT_BEGIN(Misc, BeginThreadGroupScope, Always)
+UE_TRACE_EVENT_BEGIN(Misc, BeginThreadGroupScope, Important)
 	UE_TRACE_EVENT_FIELD(uint32, CurrentThreadId)
 UE_TRACE_EVENT_END()
 
-UE_TRACE_EVENT_BEGIN(Misc, EndThreadGroupScope, Always)
+UE_TRACE_EVENT_BEGIN(Misc, EndThreadGroupScope, Important)
 	UE_TRACE_EVENT_FIELD(uint32, CurrentThreadId)
 UE_TRACE_EVENT_END()
 
-UE_TRACE_EVENT_BEGIN(Misc, BookmarkSpec, Always)
+UE_TRACE_EVENT_BEGIN(Misc, BookmarkSpec, Important)
 	UE_TRACE_EVENT_FIELD(const void*, BookmarkPoint)
 	UE_TRACE_EVENT_FIELD(int32, Line)
 UE_TRACE_EVENT_END()
 
-UE_TRACE_EVENT_BEGIN(Misc, Bookmark, Always)
+UE_TRACE_EVENT_BEGIN(Misc, Bookmark)
 	UE_TRACE_EVENT_FIELD(uint64, Cycle)
 	UE_TRACE_EVENT_FIELD(const void*, BookmarkPoint)
 UE_TRACE_EVENT_END()
 
-UE_TRACE_EVENT_BEGIN(Misc, BeginGameFrame, Always)
+UE_TRACE_EVENT_BEGIN(Misc, BeginGameFrame)
 UE_TRACE_EVENT_END()
 
-UE_TRACE_EVENT_BEGIN(Misc, EndGameFrame, Always)
+UE_TRACE_EVENT_BEGIN(Misc, EndGameFrame)
 UE_TRACE_EVENT_END()
 
-UE_TRACE_EVENT_BEGIN(Misc, BeginRenderFrame, Always)
+UE_TRACE_EVENT_BEGIN(Misc, BeginRenderFrame)
 UE_TRACE_EVENT_END()
 
-UE_TRACE_EVENT_BEGIN(Misc, EndRenderFrame, Always)
+UE_TRACE_EVENT_BEGIN(Misc, EndRenderFrame)
 UE_TRACE_EVENT_END()
 
 
@@ -63,14 +66,14 @@ uint64 FMiscTraceInternal::LastFrameCycle[TraceFrameType_Count] = { 0, 0 };
 
 void FMiscTrace::OutputRegisterGameThread(uint32 Id)
 {
-	UE_TRACE_LOG(Misc, RegisterGameThread)
+	UE_TRACE_LOG(Misc, RegisterGameThread, TraceLogChannel)
 		<< RegisterGameThread.ThreadId(FPlatformTLS::GetCurrentThreadId());
 }
 
 void FMiscTrace::OutputCreateThread(uint32 Id, const TCHAR* Name, uint32 Priority)
 {
 	uint16 NameSize = (FCString::Strlen(Name) + 1) * sizeof(TCHAR);
-	UE_TRACE_LOG(Misc, CreateThread, NameSize)
+	UE_TRACE_LOG(Misc, CreateThread, TraceLogChannel, NameSize)
 		<< CreateThread.CurrentThreadId(FPlatformTLS::GetCurrentThreadId())
 		<< CreateThread.CreatedThreadId(Id)
 		<< CreateThread.Priority(Priority)
@@ -80,7 +83,7 @@ void FMiscTrace::OutputCreateThread(uint32 Id, const TCHAR* Name, uint32 Priorit
 void FMiscTrace::OutputSetThreadGroup(uint32 Id, const ANSICHAR* GroupName)
 {
 	uint16 NameSize = strlen(GroupName) + 1;
-	UE_TRACE_LOG(Misc, SetThreadGroup, NameSize)
+	UE_TRACE_LOG(Misc, SetThreadGroup, TraceLogChannel, NameSize)
 		<< SetThreadGroup.ThreadId(Id)
 		<< SetThreadGroup.Attachment(GroupName, NameSize);
 }
@@ -88,14 +91,14 @@ void FMiscTrace::OutputSetThreadGroup(uint32 Id, const ANSICHAR* GroupName)
 void FMiscTrace::OutputBeginThreadGroupScope(const ANSICHAR* GroupName)
 {
 	uint16 NameSize = strlen(GroupName) + 1;
-	UE_TRACE_LOG(Misc, BeginThreadGroupScope, NameSize)
+	UE_TRACE_LOG(Misc, BeginThreadGroupScope, TraceLogChannel, NameSize)
 		<< BeginThreadGroupScope.CurrentThreadId(FPlatformTLS::GetCurrentThreadId())
 		<< BeginThreadGroupScope.Attachment(GroupName, NameSize);
 }
 
 void FMiscTrace::OutputEndThreadGroupScope()
 {
-	UE_TRACE_LOG(Misc, EndThreadGroupScope)
+	UE_TRACE_LOG(Misc, EndThreadGroupScope, TraceLogChannel)
 		<< EndThreadGroupScope.CurrentThreadId(FPlatformTLS::GetCurrentThreadId());
 }
 
@@ -107,7 +110,7 @@ void FMiscTrace::OutputBookmarkSpec(const void* BookmarkPoint, const ANSICHAR* F
 		memcpy(Out, File, FileNameSize);
 		memcpy(Out + FileNameSize, Format, FormatStringSize);
 	};
-	UE_TRACE_LOG(Misc, BookmarkSpec, FileNameSize + FormatStringSize)
+	UE_TRACE_LOG(Misc, BookmarkSpec, BookmarkChannel, FileNameSize + FormatStringSize)
 		<< BookmarkSpec.BookmarkPoint(BookmarkPoint)
 		<< BookmarkSpec.Line(Line)
 		<< BookmarkSpec.Attachment(StringCopyFunc);
@@ -115,7 +118,7 @@ void FMiscTrace::OutputBookmarkSpec(const void* BookmarkPoint, const ANSICHAR* F
 
 void FMiscTrace::OutputBookmarkInternal(const void* BookmarkPoint, uint16 EncodedFormatArgsSize, uint8* EncodedFormatArgs)
 {
-	UE_TRACE_LOG(Misc, Bookmark, EncodedFormatArgsSize)
+	UE_TRACE_LOG(Misc, Bookmark, BookmarkChannel, EncodedFormatArgsSize)
 		<< Bookmark.Cycle(FPlatformTime::Cycles64())
 		<< Bookmark.BookmarkPoint(BookmarkPoint)
 		<< Bookmark.Attachment(EncodedFormatArgs, EncodedFormatArgsSize);
@@ -132,12 +135,12 @@ void FMiscTrace::OutputBeginFrame(ETraceFrameType FrameType)
 	uint16 BufferSize = BufferPtr - Buffer;
 	if (FrameType == TraceFrameType_Game)
 	{
-		UE_TRACE_LOG(Misc, BeginGameFrame, BufferSize)
+		UE_TRACE_LOG(Misc, BeginGameFrame, FrameChannel, BufferSize)
 			<< BeginGameFrame.Attachment(&Buffer, BufferSize);
 	}
 	else if (FrameType == TraceFrameType_Rendering)
 	{
-		UE_TRACE_LOG(Misc, BeginRenderFrame, BufferSize)
+		UE_TRACE_LOG(Misc, BeginRenderFrame, FrameChannel, BufferSize)
 			<< BeginRenderFrame.Attachment(&Buffer, BufferSize);
 	}
 }
@@ -153,12 +156,12 @@ void FMiscTrace::OutputEndFrame(ETraceFrameType FrameType)
 	uint16 BufferSize = BufferPtr - Buffer;
 	if (FrameType == TraceFrameType_Game)
 	{
-		UE_TRACE_LOG(Misc, EndGameFrame, BufferSize)
+		UE_TRACE_LOG(Misc, EndGameFrame, FrameChannel, BufferSize)
 			<< EndGameFrame.Attachment(&Buffer, BufferSize);
 	}
 	else if (FrameType == TraceFrameType_Rendering)
 	{
-		UE_TRACE_LOG(Misc, EndRenderFrame, BufferSize)
+		UE_TRACE_LOG(Misc, EndRenderFrame, FrameChannel, BufferSize)
 			<< EndRenderFrame.Attachment(&Buffer, BufferSize);
 	}
 }
diff --git a/Engine/Source/Runtime/Core/Private/ProfilingDebugging/PlatformFileTrace.cpp b/Engine/Source/Runtime/Core/Private/ProfilingDebugging/PlatformFileTrace.cpp
--- a/Engine/Source/Runtime/Core/Private/ProfilingDebugging/PlatformFileTrace.cpp
+++ b/Engine/Source/Runtime/Core/Private/ProfilingDebugging/PlatformFileTrace.cpp
@@ -9,6 +9,8 @@
 #include "Misc/Parse.h"
 #include "Trace/Trace.h"
 
+UE_TRACE_CHANNEL(FileChannel)
+
 UE_TRACE_EVENT_BEGIN(PlatformFile, BeginOpen)
 	UE_TRACE_EVENT_FIELD(uint64, Cycle)
 	UE_TRACE_EVENT_FIELD(uint32, ThreadId)
@@ -66,7 +68,7 @@ UE_TRACE_EVENT_END()
 void FPlatformFileTrace::BeginOpen(const TCHAR* Path)
 {
 	uint16 PathSize = (FCString::Strlen(Path) + 1) * sizeof(TCHAR);
-	UE_TRACE_LOG(PlatformFile, BeginOpen, PathSize)
+	UE_TRACE_LOG(PlatformFile, BeginOpen, FileChannel, PathSize)
 		<< BeginOpen.Cycle(FPlatformTime::Cycles64())
 		<< BeginOpen.Attachment(Path, PathSize)
 		<< BeginOpen.ThreadId(FPlatformTLS::GetCurrentThreadId());
@@ -74,7 +76,7 @@ void FPlatformFileTrace::BeginOpen(const TCHAR* Path)
 
 void FPlatformFileTrace::EndOpen(uint64 FileHandle)
 {
-	UE_TRACE_LOG(PlatformFile, EndOpen)
+	UE_TRACE_LOG(PlatformFile, EndOpen, FileChannel)
 		<< EndOpen.Cycle(FPlatformTime::Cycles64())
 		<< EndOpen.FileHandle(FileHandle)
 		<< EndOpen.ThreadId(FPlatformTLS::GetCurrentThreadId());
@@ -82,7 +84,7 @@ void FPlatformFileTrace::EndOpen(uint64 FileHandle)
 
 void FPlatformFileTrace::BeginClose(uint64 FileHandle)
 {
-	UE_TRACE_LOG(PlatformFile, BeginClose)
+	UE_TRACE_LOG(PlatformFile, BeginClose, FileChannel)
 		<< BeginClose.Cycle(FPlatformTime::Cycles64())
 		<< BeginClose.FileHandle(FileHandle)
 		<< BeginClose.ThreadId(FPlatformTLS::GetCurrentThreadId());
@@ -90,14 +92,14 @@ void FPlatformFileTrace::BeginClose(uint64 FileHandle)
 
 void FPlatformFileTrace::EndClose()
 {
-	UE_TRACE_LOG(PlatformFile, EndClose)
+	UE_TRACE_LOG(PlatformFile, EndClose, FileChannel)
 		<< EndClose.Cycle(FPlatformTime::Cycles64())
 		<< EndClose.ThreadId(FPlatformTLS::GetCurrentThreadId());
 }
 
 void FPlatformFileTrace::BeginRead(uint64 ReadHandle, uint64 FileHandle, uint64 Offset, uint64 Size)
 {
-	UE_TRACE_LOG(PlatformFile, BeginRead)
+	UE_TRACE_LOG(PlatformFile, BeginRead, FileChannel)
 		<< BeginRead.Cycle(FPlatformTime::Cycles64())
 		<< BeginRead.ReadHandle(ReadHandle)
 		<< BeginRead.FileHandle(FileHandle)
@@ -108,7 +110,7 @@ void FPlatformFileTrace::BeginRead(uint64 ReadHandle, uint64 FileHandle, uint64
 
 void FPlatformFileTrace::EndRead(uint64 ReadHandle, uint64 SizeRead)
 {
-	UE_TRACE_LOG(PlatformFile, EndRead)
+	UE_TRACE_LOG(PlatformFile, EndRead, FileChannel)
 		<< EndRead.Cycle(FPlatformTime::Cycles64())
 		<< EndRead.ReadHandle(ReadHandle)
 		<< EndRead.SizeRead(SizeRead)
@@ -117,7 +119,7 @@ void FPlatformFileTrace::EndRead(uint64 ReadHandle, uint64 SizeRead)
 
 void FPlatformFileTrace::BeginWrite(uint64 WriteHandle, uint64 FileHandle, uint64 Offset, uint64 Size)
 {
-	UE_TRACE_LOG(PlatformFile, BeginWrite)
+	UE_TRACE_LOG(PlatformFile, BeginWrite, FileChannel)
 		<< BeginWrite.Cycle(FPlatformTime::Cycles64())
 		<< BeginWrite.WriteHandle(WriteHandle)
 		<< BeginWrite.FileHandle(FileHandle)
@@ -128,7 +130,7 @@ void FPlatformFileTrace::BeginWrite(uint64 WriteHandle, uint64 FileHandle, uint6
 
 void FPlatformFileTrace::EndWrite(uint64 WriteHandle, uint64 SizeWritten)
 {
-	UE_TRACE_LOG(PlatformFile, EndWrite)
+	UE_TRACE_LOG(PlatformFile, EndWrite, FileChannel)
 		<< EndWrite.Cycle(FPlatformTime::Cycles64())
 		<< EndWrite.WriteHandle(WriteHandle)
 		<< EndWrite.SizeWritten(SizeWritten)
@@ -147,7 +149,7 @@ void FPlatformFileTrace::Init(const TCHAR* CmdLine)
 		UE_TRACE_EVENT_IS_ENABLED(PlatformFile, EndRead);
 		UE_TRACE_EVENT_IS_ENABLED(PlatformFile, BeginWrite);
 		UE_TRACE_EVENT_IS_ENABLED(PlatformFile, EndWrite);
-		Trace::ToggleEvent(TEXT("PlatformFile"), true);
+		Trace::ToggleChannel(FileChannel, true);
 	}
 }
 
diff --git a/Engine/Source/Runtime/Core/Private/Serialization/LoadTimeTrace.cpp b/Engine/Source/Runtime/Core/Private/Serialization/LoadTimeTrace.cpp
--- a/Engine/Source/Runtime/Core/Private/Serialization/LoadTimeTrace.cpp
+++ b/Engine/Source/Runtime/Core/Private/Serialization/LoadTimeTrace.cpp
@@ -8,6 +8,8 @@
 #include "Misc/CString.h"
 #include "HAL/PlatformTLS.h"
 
+UE_TRACE_CHANNEL_DEFINE(LoadTimeChannel)
+
 UE_TRACE_EVENT_BEGIN(LoadTime, BeginRequestGroup)
 	UE_TRACE_EVENT_FIELD(uint32, ThreadId)
 UE_TRACE_EVENT_END()
@@ -18,7 +20,7 @@ UE_TRACE_EVENT_END()
 
 FLoadTimeProfilerTrace::FRequestGroupScope::~FRequestGroupScope()
 {
-	UE_TRACE_LOG(LoadTime, EndRequestGroup)
+	UE_TRACE_LOG(LoadTime, EndRequestGroup, LoadTimeChannel)
 		<< EndRequestGroup.ThreadId(FPlatformTLS::GetCurrentThreadId());
 }
 
@@ -30,7 +32,7 @@ void FLoadTimeProfilerTrace::FRequestGroupScope::OutputBegin()
 		memcpy(Out, FormatString, FormatStringSize);
 		memcpy(Out + FormatStringSize, FormatArgsBuffer, FormatArgsSize);
 	};
-	UE_TRACE_LOG(LoadTime, BeginRequestGroup, FormatStringSize + FormatArgsSize)
+	UE_TRACE_LOG(LoadTime, BeginRequestGroup, LoadTimeChannel, FormatStringSize + FormatArgsSize)
 		<< BeginRequestGroup.ThreadId(FPlatformTLS::GetCurrentThreadId())
 		<< BeginRequestGroup.Attachment(Attachment);
 }
diff --git a/Engine/Source/Runtime/Core/Private/Stats/Stats2.cpp b/Engine/Source/Runtime/Core/Private/Stats/Stats2.cpp
--- a/Engine/Source/Runtime/Core/Private/Stats/Stats2.cpp
+++ b/Engine/Source/Runtime/Core/Private/Stats/Stats2.cpp
@@ -607,12 +607,6 @@ public:
 			}
 		}
 
-#if CPUPROFILERTRACE_ENABLED
-		if (bCycleStat)
-		{
-			Result->TraceCpuProfilerSpecId = FCpuProfilerTrace::OutputEventType(*StatDescription, CpuProfilerGroup_Stats);
-		}
-#endif
 #if STATSTRACE_ENABLED
 		if (!bCycleStat && (InStatType == EStatDataType::ST_int64 || InStatType == EStatDataType::ST_double))
 		{
diff --git a/Engine/Source/Runtime/Core/Private/Stats/StatsTrace.cpp b/Engine/Source/Runtime/Core/Private/Stats/StatsTrace.cpp
--- a/Engine/Source/Runtime/Core/Private/Stats/StatsTrace.cpp
+++ b/Engine/Source/Runtime/Core/Private/Stats/StatsTrace.cpp
@@ -10,6 +10,8 @@
 
 #if STATSTRACE_ENABLED
 
+UE_TRACE_CHANNEL(StatsChannel)
+
 UE_TRACE_EVENT_BEGIN(Stats, Spec, Always)
 	UE_TRACE_EVENT_FIELD(uint32, Id)
 	UE_TRACE_EVENT_FIELD(bool, IsFloatingPoint)
@@ -72,7 +74,7 @@ FStatsTraceInternal::FThreadState* FStatsTraceInternal::InitThreadState()
 
 void FStatsTraceInternal::FlushThreadBuffer(FThreadState* ThreadState)
 {
-	UE_TRACE_LOG(Stats, EventBatch, ThreadState->BufferSize)
+	UE_TRACE_LOG(Stats, EventBatch, StatsChannel, ThreadState->BufferSize)
 		<< EventBatch.ThreadId(ThreadState->ThreadId)
 		<< EventBatch.Attachment(ThreadState->Buffer, ThreadState->BufferSize);
 	ThreadState->BufferSize = 0;
@@ -111,7 +113,7 @@ void FStatsTrace::DeclareStat(const FName& Stat, const ANSICHAR* Name, const TCH
 		memcpy(Buffer, Name, NameSize);
 		memcpy(Buffer + NameSize, Description, DescriptionSize);
 	};
-	UE_TRACE_LOG(Stats, Spec, NameSize + DescriptionSize)
+	UE_TRACE_LOG(Stats, Spec, StatsChannel, NameSize + DescriptionSize)
 		<< Spec.Id(Stat.GetComparisonIndex().ToUnstableInt())
 		<< Spec.IsFloatingPoint(IsFloatingPoint)
 		<< Spec.IsMemory(IsMemory)
diff --git a/Engine/Source/Runtime/Core/Private/Windows/WindowsPlatformMisc.cpp b/Engine/Source/Runtime/Core/Private/Windows/WindowsPlatformMisc.cpp
--- a/Engine/Source/Runtime/Core/Private/Windows/WindowsPlatformMisc.cpp
+++ b/Engine/Source/Runtime/Core/Private/Windows/WindowsPlatformMisc.cpp
@@ -871,6 +871,9 @@ void FWindowsPlatformMisc::BeginNamedEvent(const struct FColor& Color, const TCH
 		Profiler->StartScopedEvent(Text);
 	}
 #endif
+#if CPUPROFILERTRACE_ENABLED
+	FCpuProfilerTrace::OutputBeginDynamicEvent(Text, NamedEventsChannel);
+#endif
 }
 
 void FWindowsPlatformMisc::BeginNamedEvent(const struct FColor& Color, const ANSICHAR* Text)
@@ -884,6 +887,9 @@ void FWindowsPlatformMisc::BeginNamedEvent(const struct FColor& Color, const ANS
 		Profiler->StartScopedEvent(ANSI_TO_TCHAR(Text));
 	}
 #endif
+#if CPUPROFILERTRACE_ENABLED
+	FCpuProfilerTrace::OutputBeginDynamicEvent(Text, NamedEventsChannel);
+#endif
 }
 
 void FWindowsPlatformMisc::EndNamedEvent()
@@ -897,6 +903,9 @@ void FWindowsPlatformMisc::EndNamedEvent()
 		Profiler->EndScopedEvent();
 	}
 #endif
+#if CPUPROFILERTRACE_ENABLED
+	FCpuProfilerTrace::OutputEndEvent();
+#endif
 }
 #endif // STATS || ENABLE_STATNAMEDEVENTS
 
diff --git a/Engine/Source/Runtime/Core/Public/HAL/PlatformMisc.h b/Engine/Source/Runtime/Core/Public/HAL/PlatformMisc.h
--- a/Engine/Source/Runtime/Core/Public/HAL/PlatformMisc.h
+++ b/Engine/Source/Runtime/Core/Public/HAL/PlatformMisc.h
@@ -113,7 +113,6 @@ public:
 #endif
 
 #define SCOPED_NAMED_EVENT(Name, Color) \
-	TRACE_CPUPROFILER_EVENT_SCOPE(Name) \
 	FScopedNamedEventStatic ANONYMOUS_VARIABLE(NamedEvent_##Name##_)(Color, NAMED_EVENT_STR(#Name));
 
 #define SCOPED_NAMED_EVENT_FSTRING(Text, Color)  FScopedNamedEvent       ANONYMOUS_VARIABLE(NamedEvent_)         (Color, *Text);
diff --git a/Engine/Source/Runtime/Core/Public/ProfilingDebugging/CountersTrace.h b/Engine/Source/Runtime/Core/Public/ProfilingDebugging/CountersTrace.h
--- a/Engine/Source/Runtime/Core/Public/ProfilingDebugging/CountersTrace.h
+++ b/Engine/Source/Runtime/Core/Public/ProfilingDebugging/CountersTrace.h
@@ -4,6 +4,7 @@
 
 #include "CoreTypes.h"
 #include "Trace/Config.h"
+#include "Trace/Trace.h"
 
 #if UE_TRACE_ENABLED && !UE_BUILD_SHIPPING
 #define COUNTERSTRACE_ENABLED 1
@@ -25,6 +26,8 @@ enum ETraceCounterDisplayHint
 
 #if COUNTERSTRACE_ENABLED
 
+UE_TRACE_CHANNEL_EXTERN(CountersChannel);
+
 struct FCountersTrace
 {
 	CORE_API static uint16 OutputInitCounter(const TCHAR* CounterName, ETraceCounterType CounterType, ETraceCounterDisplayHint CounterDisplayHint);
@@ -37,8 +40,11 @@ struct FCountersTrace
 	public:
 		TCounter(const TCHAR* InCounterName, ETraceCounterDisplayHint InCounterDisplayHint)
 			: Value(0)
+			, CounterId(0)
+			, CounterName(InCounterName)
+			, CounterDisplayHint(InCounterDisplayHint)
 		{
-			CounterId = OutputInitCounter(InCounterName, CounterType, InCounterDisplayHint);
+			CounterId = OutputInitCounter(InCounterName, CounterType, CounterDisplayHint);
 		}
 
 		void Set(ValueType InValue)
@@ -83,6 +89,14 @@ struct FCountersTrace
 	private:
 		ValueType Value;
 		uint16 CounterId;
+		const TCHAR* CounterName;
+		ETraceCounterDisplayHint CounterDisplayHint;
+
+		bool CheckCounterId()
+		{
+			CounterId = OutputInitCounter(CounterName, CounterType, CounterDisplayHint);
+			return !!CounterId;
+		}
 	};
 
 	using FCounterInt = TCounter<int64, TraceCounterType_Int>;
diff --git a/Engine/Source/Runtime/Core/Public/ProfilingDebugging/CpuProfilerTrace.h b/Engine/Source/Runtime/Core/Public/ProfilingDebugging/CpuProfilerTrace.h
--- a/Engine/Source/Runtime/Core/Public/ProfilingDebugging/CpuProfilerTrace.h
+++ b/Engine/Source/Runtime/Core/Public/ProfilingDebugging/CpuProfilerTrace.h
@@ -4,6 +4,8 @@
 
 #include "CoreTypes.h"
 #include "Trace/Config.h"
+#include "Trace/Detail/Channel.h"
+#include "Trace/Detail/Channel.inl"
 
 #if !defined(CPUPROFILERTRACE_ENABLED)
 #if UE_TRACE_ENABLED && !UE_BUILD_SHIPPING
@@ -13,32 +15,29 @@
 #endif
 #endif
 
-enum ECpuProfilerGroup
-{
-	CpuProfilerGroup_Default,
-	CpuProfilerGroup_Stats,
-	CpuProfilerGroup_LoadTime,
-	CpuProfilerGroup_CsvProfiler,
-};
-
 #if CPUPROFILERTRACE_ENABLED
 
+// @note Cannot use the declare macros in this header since including
+// Trace.h will result in a circular dependency.
+CORE_API extern Trace::FChannel CpuProfilerChannel;
+CORE_API extern Trace::FChannel NamedEventsChannel;
+
 struct FCpuProfilerTrace
 {
 	CORE_API static void Init(const TCHAR* CmdLine);
 	CORE_API static void Shutdown();
-	FORCENOINLINE CORE_API static uint32 OutputEventType(const ANSICHAR* Name, ECpuProfilerGroup Group);
-	FORCENOINLINE CORE_API static uint32 OutputEventType(const TCHAR* Name, ECpuProfilerGroup Group);
-	CORE_API static void OutputBeginEvent(uint32 SpecId);
-	CORE_API static void OutputBeginDynamicEvent(const ANSICHAR* Name, ECpuProfilerGroup Group);
-	CORE_API static void OutputBeginDynamicEvent(const TCHAR* Name, ECpuProfilerGroup Group);
+	FORCENOINLINE CORE_API static uint32 OutputEventType(const ANSICHAR* Name);
+	FORCENOINLINE CORE_API static uint32 OutputEventType(const TCHAR* Name);
+	CORE_API static void OutputBeginEvent(uint32 SpecId, const Trace::FChannel& Channel);
+	CORE_API static void OutputBeginDynamicEvent(const ANSICHAR* Name, const Trace::FChannel& Channel);
+	CORE_API static void OutputBeginDynamicEvent(const TCHAR* Name, const Trace::FChannel& Channel);
 	CORE_API static void OutputEndEvent();
 
 	struct FEventScope
 	{
-		FEventScope(uint32 InSpecId)
+		FEventScope(uint32 InSpecId, const Trace::FChannel& Channel)
 		{
-			OutputBeginEvent(InSpecId);
+			OutputBeginEvent(InSpecId, Channel); 
 		}
 
 		~FEventScope()
@@ -49,14 +48,14 @@ struct FCpuProfilerTrace
 
 	struct FDynamicEventScope
 	{
-		FDynamicEventScope(const ANSICHAR* InEventName, ECpuProfilerGroup Group)
+		FDynamicEventScope(const ANSICHAR* InEventName, const Trace::FChannel& Channel)
 		{
-			OutputBeginDynamicEvent(InEventName, Group);
+			OutputBeginDynamicEvent(InEventName, Channel);
 		}
 
-		FDynamicEventScope(const TCHAR* InEventName, ECpuProfilerGroup Group)
+		FDynamicEventScope(const TCHAR* InEventName, const Trace::FChannel& Channel)
 		{
-			OutputBeginDynamicEvent(InEventName, Group);
+			OutputBeginDynamicEvent(InEventName, Channel);
 		}
 
 		~FDynamicEventScope()
@@ -72,29 +71,33 @@ struct FCpuProfilerTrace
 #define TRACE_CPUPROFILER_SHUTDOWN() \
 	FCpuProfilerTrace::Shutdown();
 
-#define TRACE_CPUPROFILER_EVENT_SCOPE_GROUP(Name, Group) \
+#define TRACE_CPUPROFILER_EVENT_SCOPE_ON_CHANNEL_STR(NameStr, Channel) \
 	static uint16 PREPROCESSOR_JOIN(__CpuProfilerEventSpecId, __LINE__); \
-	if (PREPROCESSOR_JOIN(__CpuProfilerEventSpecId, __LINE__) == 0) { \
-		PREPROCESSOR_JOIN(__CpuProfilerEventSpecId, __LINE__) = FCpuProfilerTrace::OutputEventType(#Name, Group); \
+	if (bool(Channel) && PREPROCESSOR_JOIN(__CpuProfilerEventSpecId, __LINE__) == 0) { \
+		PREPROCESSOR_JOIN(__CpuProfilerEventSpecId, __LINE__) = FCpuProfilerTrace::OutputEventType(NameStr); \
 	} \
-	FCpuProfilerTrace::FEventScope PREPROCESSOR_JOIN(__CpuProfilerEventScope, __LINE__)(PREPROCESSOR_JOIN(__CpuProfilerEventSpecId, __LINE__));
+	FCpuProfilerTrace::FEventScope PREPROCESSOR_JOIN(__CpuProfilerEventScope, __LINE__)(PREPROCESSOR_JOIN(__CpuProfilerEventSpecId, __LINE__), Channel);
+
+#define TRACE_CPUPROFILER_EVENT_SCOPE_ON_CHANNEL(Name, Channel) \
+	TRACE_CPUPROFILER_EVENT_SCOPE_ON_CHANNEL_STR(#Name, Channel)
 
 #define TRACE_CPUPROFILER_EVENT_SCOPE(Name) \
-	TRACE_CPUPROFILER_EVENT_SCOPE_GROUP(Name, CpuProfilerGroup_Default)
+	TRACE_CPUPROFILER_EVENT_SCOPE_ON_CHANNEL(Name, CpuProfilerChannel)
 
-#define TRACE_CPUPROFILER_EVENT_SCOPE_TEXT_GROUP(Name, Group) \
-	FCpuProfilerTrace::FDynamicEventScope PREPROCESSOR_JOIN(__CpuProfilerEventScope, __LINE__)(Name, Group);
+#define TRACE_CPUPROFILER_EVENT_SCOPE_TEXT_ON_CHANNEL(Name, Channel) \
+	FCpuProfilerTrace::FDynamicEventScope PREPROCESSOR_JOIN(__CpuProfilerEventScope, __LINE__)(Name, Channel);
 	
 #define TRACE_CPUPROFILER_EVENT_SCOPE_TEXT(Name) \
-	TRACE_CPUPROFILER_EVENT_SCOPE_TEXT_GROUP(Name, CpuProfilerGroup_Default)
+	TRACE_CPUPROFILER_EVENT_SCOPE_TEXT_ON_CHANNEL(Name, CpuProfilerChannel)
 
 #else
 
 #define TRACE_CPUPROFILER_INIT(CmdLine)
 #define TRACE_CPUPROFILER_SHUTDOWN()
-#define TRACE_CPUPROFILER_EVENT_SCOPE_GROUP(Name, Group)
+#define TRACE_CPUPROFILER_EVENT_SCOPE_ON_CHANNEL_STR(NameStr, Channel)
+#define TRACE_CPUPROFILER_EVENT_SCOPE_ON_CHANNEL(Name, Channel)
 #define TRACE_CPUPROFILER_EVENT_SCOPE(Name)
-#define TRACE_CPUPROFILER_EVENT_SCOPE_TEXT_GROUP(Name, Group)
+#define TRACE_CPUPROFILER_EVENT_SCOPE_TEXT_ON_CHANNEL(Name, Channel)
 #define TRACE_CPUPROFILER_EVENT_SCOPE_TEXT(Name)
 
 #endif
diff --git a/Engine/Source/Runtime/Core/Public/ProfilingDebugging/LoadTimeTracker.h b/Engine/Source/Runtime/Core/Public/ProfilingDebugging/LoadTimeTracker.h
--- a/Engine/Source/Runtime/Core/Public/ProfilingDebugging/LoadTimeTracker.h
+++ b/Engine/Source/Runtime/Core/Public/ProfilingDebugging/LoadTimeTracker.h
@@ -12,6 +12,7 @@
 #include "UObject/NameTypes.h"
 #include "ProfilingDebugging/ScopedTimers.h"
 #include "Trace/Config.h"
+#include "Serialization/LoadTimeTrace.h"
 
 #ifndef ENABLE_LOADTIME_TRACKING
 	#define ENABLE_LOADTIME_TRACKING 0
@@ -213,6 +214,6 @@ struct CORE_API FScopedLoadTimeAccumulatorTimer : public FScopedDurationTimer
 #define SCOPED_LOADTIMER_CNT(TimerName) FScopedDurationTimer DurationTimer_##TimerName(FLoadTimeTracker::Get().TimerName); FLoadTimeTracker::Get().TimerName##Cnt++;
 #else
 
-#define SCOPED_LOADTIMER(TimerName) TRACE_CPUPROFILER_EVENT_SCOPE_GROUP(TimerName, CpuProfilerGroup_LoadTime)
+#define SCOPED_LOADTIMER(TimerName) TRACE_CPUPROFILER_EVENT_SCOPE_ON_CHANNEL(TimerName, LoadTimeChannel)
 #define SCOPED_LOADTIMER_CNT(TimerName)
 #endif
diff --git a/Engine/Source/Runtime/Core/Public/Serialization/LoadTimeTrace.h b/Engine/Source/Runtime/Core/Public/Serialization/LoadTimeTrace.h
--- a/Engine/Source/Runtime/Core/Public/Serialization/LoadTimeTrace.h
+++ b/Engine/Source/Runtime/Core/Public/Serialization/LoadTimeTrace.h
@@ -4,6 +4,7 @@
 
 #include "CoreTypes.h"
 #include "Trace/Config.h"
+#include "Trace/Trace.h"
 #include "ProfilingDebugging/FormatArgsTrace.h"
 
 #if !defined(LOADTIMEPROFILERTRACE_ENABLED)
@@ -16,6 +17,8 @@
 
 #if LOADTIMEPROFILERTRACE_ENABLED
 
+CORE_API UE_TRACE_CHANNEL_EXTERN(LoadTimeChannel);
+
 struct FLoadTimeProfilerTrace
 {
 	struct FRequestGroupScope
diff --git a/Engine/Source/Runtime/Core/Public/Stats/Stats2.h b/Engine/Source/Runtime/Core/Public/Stats/Stats2.h
--- a/Engine/Source/Runtime/Core/Public/Stats/Stats2.h
+++ b/Engine/Source/Runtime/Core/Public/Stats/Stats2.h
@@ -129,10 +129,6 @@ struct TStatIdData
 
 	/** const WIDECHAR* pointer to a string describing the stat */
 	TUniquePtr<ANSICHAR[]> StatDescriptionAnsi;
-
-#if CPUPROFILERTRACE_ENABLED
-	uint32 TraceCpuProfilerSpecId = 0;
-#endif
 };
 
 struct TStatId
@@ -168,13 +164,6 @@ struct TStatId
 		return MinimalNameToName(StatIdPtr->Name);
 	}
 
-#if CPUPROFILERTRACE_ENABLED
-	FORCEINLINE uint16 GetTraceCpuProfilerSpecId() const
-	{
-		return StatIdPtr->TraceCpuProfilerSpecId;
-	}
-#endif
-
 	FORCEINLINE static const TStatIdData& GetStatNone()
 	{
 		return TStatId_NAME_None;
@@ -1533,9 +1522,6 @@ class FCycleCounter
 {
 	/** Name of the stat, usually a short name **/
 	FName StatId;
-#if CPUPROFILERTRACE_ENABLED
-	uint16 TraceCpuProfilerSpecId = 0;
-#endif
 
 public:
 
@@ -1550,14 +1536,6 @@ public:
 		{
 			return;
 		}
-#if CPUPROFILERTRACE_ENABLED
-		if (bAlways || GCycleStatsShouldEmitNamedEvents > 0)
-		{
-			TraceCpuProfilerSpecId = InStatId.GetTraceCpuProfilerSpecId();
-			FCpuProfilerTrace::OutputBeginEvent(TraceCpuProfilerSpecId);
-		}
-#endif
-
 		if( (bAlways && FThreadStats::WillEverCollectData()) || FThreadStats::IsCollectingData() )
 		{
 			FName StatName = MinimalNameToName(StatMinimalName);
@@ -1581,13 +1559,6 @@ public:
 	 */
 	FORCEINLINE_STATS void Stop()
 	{
-#if CPUPROFILERTRACE_ENABLED
-		if (TraceCpuProfilerSpecId)
-		{
-			FCpuProfilerTrace::OutputEndEvent();
-			TraceCpuProfilerSpecId = 0;
-		}
-#endif
 		if( !StatId.IsNone() )
 		{
 			FThreadStats::AddMessage(StatId, EStatOperation::CycleScopeEnd);
diff --git a/Engine/Source/Runtime/CoreUObject/Private/Serialization/LoadTimeTrace.cpp b/Engine/Source/Runtime/CoreUObject/Private/Serialization/LoadTimeTrace.cpp
--- a/Engine/Source/Runtime/CoreUObject/Private/Serialization/LoadTimeTrace.cpp
+++ b/Engine/Source/Runtime/CoreUObject/Private/Serialization/LoadTimeTrace.cpp
@@ -14,7 +14,7 @@
 #include "UObject/UObjectGlobals.h"
 #include "Misc/CommandLine.h"
 
-UE_TRACE_EVENT_BEGIN(LoadTime, StartAsyncLoading, Always)
+UE_TRACE_EVENT_BEGIN(LoadTime, StartAsyncLoading, Important)
 	UE_TRACE_EVENT_FIELD(uint64, Cycle)
 	UE_TRACE_EVENT_FIELD(uint32, ThreadId)
 UE_TRACE_EVENT_END()
@@ -109,7 +109,7 @@ UE_TRACE_EVENT_BEGIN(LoadTime, AsyncPackageImportDependency)
 	UE_TRACE_EVENT_FIELD(const void*, ImportedAsyncPackage)
 UE_TRACE_EVENT_END()
 
-UE_TRACE_EVENT_BEGIN(LoadTime, ClassInfo, Always)
+UE_TRACE_EVENT_BEGIN(LoadTime, ClassInfo, Important)
 	UE_TRACE_EVENT_FIELD(const UClass*, Class)
 UE_TRACE_EVENT_END()
 
@@ -137,13 +137,13 @@ void FLoadTimeProfilerTracePrivate::Init()
 	UE_TRACE_EVENT_IS_ENABLED(LoadTime, ClassInfo);
 	if (FParse::Param(FCommandLine::Get(), TEXT("loadtimetrace")))
 	{
-		Trace::ToggleEvent(TEXT("LoadTime"), true);
+		Trace::ToggleChannel(LoadTimeChannel, true);
 	}
 }
 
 void FLoadTimeProfilerTracePrivate::OutputStartAsyncLoading()
 {
-	UE_TRACE_LOG(LoadTime, StartAsyncLoading)
+	UE_TRACE_LOG(LoadTime, StartAsyncLoading, LoadTimeChannel)
 		<< StartAsyncLoading.Cycle(FPlatformTime::Cycles64())
 		<< StartAsyncLoading.ThreadId(FPlatformTLS::GetCurrentThreadId());
 }
@@ -151,20 +151,20 @@ void FLoadTimeProfilerTracePrivate::OutputStartAsyncLoading()
 PRAGMA_DISABLE_SHADOW_VARIABLE_WARNINGS
 void FLoadTimeProfilerTracePrivate::OutputSuspendAsyncLoading()
 {
-	UE_TRACE_LOG(LoadTime, SuspendAsyncLoading)
+	UE_TRACE_LOG(LoadTime, SuspendAsyncLoading, LoadTimeChannel)
 		<< SuspendAsyncLoading.Cycle(FPlatformTime::Cycles64());
 }
 
 void FLoadTimeProfilerTracePrivate::OutputResumeAsyncLoading()
 {
-	UE_TRACE_LOG(LoadTime, ResumeAsyncLoading)
+	UE_TRACE_LOG(LoadTime, ResumeAsyncLoading, LoadTimeChannel)
 		<< ResumeAsyncLoading.Cycle(FPlatformTime::Cycles64());
 }
 PRAGMA_ENABLE_SHADOW_VARIABLE_WARNINGS
 
 void FLoadTimeProfilerTracePrivate::OutputBeginRequest(uint64 RequestId)
 {
-	UE_TRACE_LOG(LoadTime, BeginRequest)
+	UE_TRACE_LOG(LoadTime, BeginRequest, LoadTimeChannel)
 		<< BeginRequest.Cycle(FPlatformTime::Cycles64())
 		<< BeginRequest.RequestId(RequestId)
 		<< BeginRequest.ThreadId(FPlatformTLS::GetCurrentThreadId());
@@ -172,7 +172,7 @@ void FLoadTimeProfilerTracePrivate::OutputBeginRequest(uint64 RequestId)
 
 void FLoadTimeProfilerTracePrivate::OutputEndRequest(uint64 RequestId)
 {
-	UE_TRACE_LOG(LoadTime, EndRequest)
+	UE_TRACE_LOG(LoadTime, EndRequest, LoadTimeChannel)
 		<< EndRequest.Cycle(FPlatformTime::Cycles64())
 		<< EndRequest.RequestId(RequestId);
 }
@@ -181,34 +181,34 @@ void FLoadTimeProfilerTracePrivate::OutputNewAsyncPackage(const void* AsyncPacka
 {
 	TCHAR Buffer[FName::StringBufferSize];
 	uint16 NameSize = (PackageName.ToString(Buffer) + 1) * sizeof(TCHAR);
-	UE_TRACE_LOG(LoadTime, NewAsyncPackage, NameSize)
+	UE_TRACE_LOG(LoadTime, NewAsyncPackage, LoadTimeChannel, NameSize)
 		<< NewAsyncPackage.AsyncPackage(AsyncPackage)
 		<< NewAsyncPackage.Attachment(Buffer, NameSize);
 }
 
 void FLoadTimeProfilerTracePrivate::OutputBeginLoadAsyncPackage(const void* AsyncPackage)
 {
-	UE_TRACE_LOG(LoadTime, BeginLoadAsyncPackage)
+	UE_TRACE_LOG(LoadTime, BeginLoadAsyncPackage, LoadTimeChannel)
 		<< BeginLoadAsyncPackage.Cycle(FPlatformTime::Cycles64())
 		<< BeginLoadAsyncPackage.AsyncPackage(AsyncPackage);
 }
 
 void FLoadTimeProfilerTracePrivate::OutputEndLoadAsyncPackage(const void* AsyncPackage)
 {
-	UE_TRACE_LOG(LoadTime, EndLoadAsyncPackage)
+	UE_TRACE_LOG(LoadTime, EndLoadAsyncPackage, LoadTimeChannel)
 		<< EndLoadAsyncPackage.Cycle(FPlatformTime::Cycles64())
 		<< EndLoadAsyncPackage.AsyncPackage(AsyncPackage);
 }
 
 void FLoadTimeProfilerTracePrivate::OutputDestroyAsyncPackage(const void* AsyncPackage)
 {
-	UE_TRACE_LOG(LoadTime, DestroyAsyncPackage)
+	UE_TRACE_LOG(LoadTime, DestroyAsyncPackage, LoadTimeChannel)
 		<< DestroyAsyncPackage.AsyncPackage(AsyncPackage);
 }
 
 void FLoadTimeProfilerTracePrivate::OutputPackageSummary(const void* AsyncPackage, uint32 TotalHeaderSize, uint32 ImportCount, uint32 ExportCount)
 {
-	UE_TRACE_LOG(LoadTime, PackageSummary)
+	UE_TRACE_LOG(LoadTime, PackageSummary, LoadTimeChannel)
 		<< PackageSummary.AsyncPackage(AsyncPackage)
 		<< PackageSummary.TotalHeaderSize(TotalHeaderSize)
 		<< PackageSummary.ImportCount(ImportCount)
@@ -217,14 +217,14 @@ void FLoadTimeProfilerTracePrivate::OutputPackageSummary(const void* AsyncPackag
 
 void FLoadTimeProfilerTracePrivate::OutputAsyncPackageRequestAssociation(const void* AsyncPackage, uint64 RequestId)
 {
-	UE_TRACE_LOG(LoadTime, AsyncPackageRequestAssociation)
+	UE_TRACE_LOG(LoadTime, AsyncPackageRequestAssociation, LoadTimeChannel)
 		<< AsyncPackageRequestAssociation.AsyncPackage(AsyncPackage)
 		<< AsyncPackageRequestAssociation.RequestId(RequestId);
 }
 
 void FLoadTimeProfilerTracePrivate::OutputAsyncPackageImportDependency(const void* Package, const void* ImportedPackage)
 {
-	UE_TRACE_LOG(LoadTime, AsyncPackageImportDependency)
+	UE_TRACE_LOG(LoadTime, AsyncPackageImportDependency, LoadTimeChannel)
 		<< AsyncPackageImportDependency.AsyncPackage(Package)
 		<< AsyncPackageImportDependency.ImportedAsyncPackage(ImportedPackage);
 }
@@ -233,7 +233,7 @@ void FLoadTimeProfilerTracePrivate::OutputClassInfo(const UClass* Class, const F
 {
 	TCHAR Buffer[FName::StringBufferSize];
 	uint16 NameSize = (Name.ToString(Buffer) + 1) * sizeof(TCHAR);
-	UE_TRACE_LOG(LoadTime, ClassInfo, NameSize)
+	UE_TRACE_LOG(LoadTime, ClassInfo, LoadTimeChannel, NameSize)
 		<< ClassInfo.Class(Class)
 		<< ClassInfo.Attachment(Buffer, NameSize);
 }
@@ -241,7 +241,7 @@ void FLoadTimeProfilerTracePrivate::OutputClassInfo(const UClass* Class, const F
 void FLoadTimeProfilerTracePrivate::OutputClassInfo(const UClass* Class, const TCHAR* Name)
 {
 	uint16 NameSize = (FCString::Strlen(Name) + 1) * sizeof(TCHAR);
-	UE_TRACE_LOG(LoadTime, ClassInfo, NameSize)
+	UE_TRACE_LOG(LoadTime, ClassInfo, LoadTimeChannel, NameSize)
 		<< ClassInfo.Class(Class)
 		<< ClassInfo.Attachment(Name, NameSize);
 }
@@ -249,7 +249,7 @@ void FLoadTimeProfilerTracePrivate::OutputClassInfo(const UClass* Class, const T
 FLoadTimeProfilerTracePrivate::FCreateExportScope::FCreateExportScope(const void* AsyncPackage, const UObject* const* InObject)
 	: Object(InObject)
 {
-	UE_TRACE_LOG(LoadTime, BeginCreateExport)
+	UE_TRACE_LOG(LoadTime, BeginCreateExport, LoadTimeChannel)
 		<< BeginCreateExport.Cycle(FPlatformTime::Cycles64())
 		<< BeginCreateExport.AsyncPackage(AsyncPackage)
 		<< BeginCreateExport.ThreadId(FPlatformTLS::GetCurrentThreadId());
@@ -257,7 +257,7 @@ FLoadTimeProfilerTracePrivate::FCreateExportScope::FCreateExportScope(const void
 
 FLoadTimeProfilerTracePrivate::FCreateExportScope::~FCreateExportScope()
 {
-	UE_TRACE_LOG(LoadTime, EndCreateExport)
+	UE_TRACE_LOG(LoadTime, EndCreateExport, LoadTimeChannel)
 		<< EndCreateExport.Cycle(FPlatformTime::Cycles64())
 		<< EndCreateExport.Object(*Object)
 		<< EndCreateExport.Class(*Object ? (*Object)->GetClass() : nullptr)
@@ -266,7 +266,7 @@ FLoadTimeProfilerTracePrivate::FCreateExportScope::~FCreateExportScope()
 
 FLoadTimeProfilerTracePrivate::FSerializeExportScope::FSerializeExportScope(const UObject* Object, uint64 SerialSize)
 {
-	UE_TRACE_LOG(LoadTime, BeginSerializeExport)
+	UE_TRACE_LOG(LoadTime, BeginSerializeExport, LoadTimeChannel)
 		<< BeginSerializeExport.Cycle(FPlatformTime::Cycles64())
 		<< BeginSerializeExport.Object(Object)
 		<< BeginSerializeExport.SerialSize(SerialSize)
@@ -275,14 +275,14 @@ FLoadTimeProfilerTracePrivate::FSerializeExportScope::FSerializeExportScope(cons
 
 FLoadTimeProfilerTracePrivate::FSerializeExportScope::~FSerializeExportScope()
 {
-	UE_TRACE_LOG(LoadTime, EndSerializeExport)
+	UE_TRACE_LOG(LoadTime, EndSerializeExport, LoadTimeChannel)
 		<< EndSerializeExport.Cycle(FPlatformTime::Cycles64())
 		<< EndSerializeExport.ThreadId(FPlatformTLS::GetCurrentThreadId());
 }
 
 FLoadTimeProfilerTracePrivate::FPostLoadExportScope::FPostLoadExportScope(const UObject* Object)
 {
-	UE_TRACE_LOG(LoadTime, BeginPostLoadExport)
+	UE_TRACE_LOG(LoadTime, BeginPostLoadExport, LoadTimeChannel)
 		<< BeginPostLoadExport.Cycle(FPlatformTime::Cycles64())
 		<< BeginPostLoadExport.Object(Object)
 		<< BeginPostLoadExport.ThreadId(FPlatformTLS::GetCurrentThreadId());
@@ -290,7 +290,7 @@ FLoadTimeProfilerTracePrivate::FPostLoadExportScope::FPostLoadExportScope(const
 
 FLoadTimeProfilerTracePrivate::FPostLoadExportScope::~FPostLoadExportScope()
 {
-	UE_TRACE_LOG(LoadTime, EndPostLoadExport)
+	UE_TRACE_LOG(LoadTime, EndPostLoadExport, LoadTimeChannel)
 		<< EndPostLoadExport.Cycle(FPlatformTime::Cycles64())
 		<< EndPostLoadExport.ThreadId(FPlatformTLS::GetCurrentThreadId());
 }
diff --git a/Engine/Source/Runtime/Engine/Private/Animation/AnimTrace.cpp b/Engine/Source/Runtime/Engine/Private/Animation/AnimTrace.cpp
--- a/Engine/Source/Runtime/Engine/Private/Animation/AnimTrace.cpp
+++ b/Engine/Source/Runtime/Engine/Private/Animation/AnimTrace.cpp
@@ -17,6 +17,8 @@
 #include "Animation/BlendSpaceBase.h"
 #include "Animation/AnimNode_SequencePlayer.h"
 
+UE_TRACE_CHANNEL_DEFINE(AnimationChannel);
+
 UE_TRACE_EVENT_BEGIN(Animation, TickRecord)
 	UE_TRACE_EVENT_FIELD(uint64, Cycle)
 	UE_TRACE_EVENT_FIELD(uint64, AnimInstanceId)
@@ -276,8 +278,8 @@ void FAnimTrace::Init()
 
 void FAnimTrace::OutputAnimTickRecord(const FAnimationBaseContext& InContext, const FAnimTickRecord& InTickRecord)
 {
-	bool bEventEnabled = UE_TRACE_EVENT_IS_ENABLED(Animation, TickRecord);
-	if (!bEventEnabled)
+	bool bChannelEnabled = UE_TRACE_CHANNELEXPR_IS_ENABLED(AnimationChannel);
+	if (!bChannelEnabled)
 	{
 		return;
 	}
@@ -303,7 +305,7 @@ void FAnimTrace::OutputAnimTickRecord(const FAnimationBaseContext& InContext, co
 			BlendSpacePositionY = InTickRecord.BlendSpace.BlendSpacePositionY;
 	}
 
-		UE_TRACE_LOG(Animation, TickRecord)
+		UE_TRACE_LOG(Animation, TickRecord, AnimationChannel)
 			<< TickRecord.Cycle(FPlatformTime::Cycles64())
 			<< TickRecord.AnimInstanceId(FObjectTrace::GetObjectId(InContext.AnimInstanceProxy->GetAnimInstanceObject()))
 			<< TickRecord.AssetId(FObjectTrace::GetObjectId(InTickRecord.SourceAsset))
@@ -322,8 +324,8 @@ void FAnimTrace::OutputAnimTickRecord(const FAnimationBaseContext& InContext, co
 
 void FAnimTrace::OutputSkeletalMesh(const USkeletalMesh* InMesh)
 {
-	bool bEventEnabled = UE_TRACE_EVENT_IS_ENABLED(Animation, SkeletalMesh);
-	if (!bEventEnabled || InMesh == nullptr)
+	bool bChannelEnabled = UE_TRACE_CHANNELEXPR_IS_ENABLED(AnimationChannel);
+	if (!bChannelEnabled || InMesh == nullptr)
 	{
 		return;
 	}
@@ -347,7 +349,7 @@ void FAnimTrace::OutputSkeletalMesh(const USkeletalMesh* InMesh)
 		}
 	};
 
-	UE_TRACE_LOG(Animation, SkeletalMesh, BoneCount * sizeof(int32))
+	UE_TRACE_LOG(Animation, SkeletalMesh, AnimationChannel, BoneCount * sizeof(int32))
 		<< SkeletalMesh.Id(FObjectTrace::GetObjectId(InMesh))
 		<< SkeletalMesh.BoneCount(BoneCount)
 		<< SkeletalMesh.Attachment(CopyParentIndices);
@@ -372,7 +374,7 @@ uint32 FAnimTrace::OutputName(const FName& InName)
 
 		uint32 NewId = CurrentId++;
 
-		UE_TRACE_LOG(Animation, Name, NameStringLength * sizeof(TCHAR))
+		UE_TRACE_LOG(Animation, Name, AnimationChannel, NameStringLength * sizeof(TCHAR))
 			<< Name.Id(NewId)
 			<< Name.Attachment(StringCopyFunc);
 
@@ -387,8 +389,8 @@ uint32 FAnimTrace::OutputName(const FName& InName)
 
 void FAnimTrace::OutputSkeletalMeshComponent(const USkeletalMeshComponent* InComponent)
 {
-	bool bEventEnabled = UE_TRACE_EVENT_IS_ENABLED(Animation, SkeletalMeshComponent);
-	if (!bEventEnabled || InComponent == nullptr)
+	bool bChannelEnabled = UE_TRACE_CHANNELEXPR_IS_ENABLED(AnimationChannel);
+	if (!bChannelEnabled || InComponent == nullptr)
 	{
 		return;
 	}
@@ -436,7 +438,8 @@ void FAnimTrace::OutputSkeletalMeshComponent(const USkeletalMeshComponent* InCom
 			}
 		};
 
-		UE_TRACE_LOG(Animation, SkeletalMeshComponent, ((BoneCount + 1) * sizeof(FTransform)) + (CurveCount * (sizeof(float) + sizeof(uint32))))
+		//todo: Crashes on cook enginetest
+		/*UE_TRACE_LOG(Animation, SkeletalMeshComponent, true, ((BoneCount + 1) * sizeof(FTransform)) + (CurveCount * (sizeof(float) + sizeof(uint32))))
 			<< SkeletalMeshComponent.Cycle(FPlatformTime::Cycles64())
 			<< SkeletalMeshComponent.ComponentId(FObjectTrace::GetObjectId(InComponent))
 			<< SkeletalMeshComponent.MeshId(FObjectTrace::GetObjectId(InComponent->SkeletalMesh))
@@ -444,21 +447,14 @@ void FAnimTrace::OutputSkeletalMeshComponent(const USkeletalMeshComponent* InCom
 			<< SkeletalMeshComponent.CurveCount((uint32)CurveCount)
 			<< SkeletalMeshComponent.LodIndex((uint16)InComponent->PredictedLODLevel)
 			<< SkeletalMeshComponent.FrameCounter((uint16)(GFrameCounter % 0xffff))
-			<< SkeletalMeshComponent.Attachment(CopyTransformsAndCurves);
+			<< SkeletalMeshComponent.Attachment(CopyTransformsAndCurves);*/
 	}
 }
 
 void FAnimTrace::OutputSkeletalMeshFrame(const USkeletalMeshComponent* InComponent)
 {
-	bool bEventEnabled = UE_TRACE_EVENT_IS_ENABLED(Animation, SkeletalMeshFrame);
-	if (!bEventEnabled || InComponent == nullptr)
-	{
-		return;
-	}
-
 	TRACE_OBJECT(InComponent);
-
-	UE_TRACE_LOG(Animation, SkeletalMeshFrame)
+	UE_TRACE_LOG(Animation, SkeletalMeshFrame, AnimationChannel)
 		<< SkeletalMeshFrame.Cycle(FPlatformTime::Cycles64())
 		<< SkeletalMeshFrame.ComponentId(FObjectTrace::GetObjectId(InComponent))
 		<< SkeletalMeshFrame.FrameCounter((uint16)(GFrameCounter % 0xffff));
@@ -466,8 +462,8 @@ void FAnimTrace::OutputSkeletalMeshFrame(const USkeletalMeshComponent* InCompone
 
 void FAnimTrace::OutputAnimGraph(const FAnimationBaseContext& InContext, uint64 InStartCycle, uint64 InEndCycle, uint8 InPhase)
 {
-	bool bEventEnabled = UE_TRACE_EVENT_IS_ENABLED(Animation, AnimGraph);
-	if (!bEventEnabled)
+	bool bChannelEnabled = UE_TRACE_CHANNELEXPR_IS_ENABLED(AnimationChannel);
+	if (!bChannelEnabled)
 	{
 		return;
 	}
@@ -479,7 +475,7 @@ void FAnimTrace::OutputAnimGraph(const FAnimationBaseContext& InContext, uint64
 
 	TRACE_OBJECT(AnimInstance);
 
-	UE_TRACE_LOG(Animation, AnimGraph)
+	UE_TRACE_LOG(Animation, AnimGraph, AnimationChannel)
 		<< AnimGraph.StartCycle(InStartCycle)
 		<< AnimGraph.EndCycle(InEndCycle)
 		<< AnimGraph.AnimInstanceId(FObjectTrace::GetObjectId(AnimInstance))
@@ -490,8 +486,8 @@ void FAnimTrace::OutputAnimGraph(const FAnimationBaseContext& InContext, uint64
 
 void FAnimTrace::OutputAnimNodeStart(const FAnimationBaseContext& InContext, uint64 InStartCycle, int32 InPreviousNodeId, int32 InNodeId, float InBlendWeight, float InRootMotionWeight, uint8 InPhase)
 {
-	bool bEventEnabled = UE_TRACE_EVENT_IS_ENABLED(Animation, AnimNodeStart);
-	if (!bEventEnabled)
+	bool bChannelEnabled = UE_TRACE_CHANNELEXPR_IS_ENABLED(AnimationChannel);
+	if (!bChannelEnabled)
 	{
 		return;
 	}
@@ -515,7 +511,7 @@ void FAnimTrace::OutputAnimNodeStart(const FAnimationBaseContext& InContext, uin
 
 	DisplayNameString.RemoveFromStart(TEXT("Anim Node "));
 
-	UE_TRACE_LOG(Animation, AnimNodeStart, (DisplayNameString.Len() + 1) * sizeof(TCHAR))
+	UE_TRACE_LOG(Animation, AnimNodeStart, AnimationChannel, (DisplayNameString.Len() + 1) * sizeof(TCHAR))
 		<< AnimNodeStart.StartCycle(InStartCycle)
 		<< AnimNodeStart.AnimInstanceId(FObjectTrace::GetObjectId(InContext.AnimInstanceProxy->GetAnimInstanceObject()))
 		<< AnimNodeStart.PreviousNodeId(InPreviousNodeId)
@@ -528,8 +524,8 @@ void FAnimTrace::OutputAnimNodeStart(const FAnimationBaseContext& InContext, uin
 
 void FAnimTrace::OutputAnimNodeEnd(const FAnimationBaseContext& InContext, uint64 InEndCycle)
 {
-	bool bEventEnabled = UE_TRACE_EVENT_IS_ENABLED(Animation, AnimNodeEnd);
-	if (!bEventEnabled)
+	bool bChannelEnabled = UE_TRACE_CHANNELEXPR_IS_ENABLED(AnimationChannel);
+	if (!bChannelEnabled)
 	{
 		return;
 	}
@@ -538,15 +534,15 @@ void FAnimTrace::OutputAnimNodeEnd(const FAnimationBaseContext& InContext, uint6
 
 	TRACE_OBJECT(InContext.AnimInstanceProxy->GetAnimInstanceObject());
 
-	UE_TRACE_LOG(Animation, AnimNodeEnd)
+	UE_TRACE_LOG(Animation, AnimNodeEnd, AnimationChannel)
 		<< AnimNodeEnd.EndCycle(InEndCycle)
 		<< AnimNodeEnd.AnimInstanceId(FObjectTrace::GetObjectId(InContext.AnimInstanceProxy->GetAnimInstanceObject()));
 }
 
 void FAnimTrace::OutputAnimNodeValue(const FAnimationBaseContext& InContext, const TCHAR* InKey, bool InValue)
 {
-	bool bEventEnabled = UE_TRACE_EVENT_IS_ENABLED(Animation, AnimNodeValueBool);
-	if (!bEventEnabled)
+	bool bChannelEnabled = UE_TRACE_CHANNELEXPR_IS_ENABLED(AnimationChannel);
+	if (!bChannelEnabled)
 	{
 		return;
 	}
@@ -557,7 +553,7 @@ void FAnimTrace::OutputAnimNodeValue(const FAnimationBaseContext& InContext, con
 
 	int32 KeyLength = FCString::Strlen(InKey) + 1;
 
-	UE_TRACE_LOG(Animation, AnimNodeValueBool, KeyLength * sizeof(TCHAR))
+	UE_TRACE_LOG(Animation, AnimNodeValueBool, AnimationChannel, KeyLength * sizeof(TCHAR))
 		<< AnimNodeValueBool.Cycle(FPlatformTime::Cycles64())
 		<< AnimNodeValueBool.AnimInstanceId(FObjectTrace::GetObjectId(InContext.AnimInstanceProxy->GetAnimInstanceObject()))
 		<< AnimNodeValueBool.NodeId(InContext.GetCurrentNodeId())
@@ -569,8 +565,8 @@ void FAnimTrace::OutputAnimNodeValue(const FAnimationBaseContext& InContext, con
 
 void FAnimTrace::OutputAnimNodeValue(const FAnimationBaseContext& InContext, const TCHAR* InKey, int32 InValue)
 {
-	bool bEventEnabled = UE_TRACE_EVENT_IS_ENABLED(Animation, AnimNodeValueInt);
-	if (!bEventEnabled)
+	bool bChannelEnabled = UE_TRACE_CHANNELEXPR_IS_ENABLED(AnimationChannel);
+	if (!bChannelEnabled)
 	{
 		return;
 	}
@@ -581,7 +577,7 @@ void FAnimTrace::OutputAnimNodeValue(const FAnimationBaseContext& InContext, con
 
 	int32 KeyLength = FCString::Strlen(InKey) + 1;
 
-	UE_TRACE_LOG(Animation, AnimNodeValueInt, KeyLength * sizeof(TCHAR))
+	UE_TRACE_LOG(Animation, AnimNodeValueInt, AnimationChannel, KeyLength * sizeof(TCHAR))
 		<< AnimNodeValueInt.Cycle(FPlatformTime::Cycles64())
 		<< AnimNodeValueInt.AnimInstanceId(FObjectTrace::GetObjectId(InContext.AnimInstanceProxy->GetAnimInstanceObject()))
 		<< AnimNodeValueInt.NodeId(InContext.GetCurrentNodeId())
@@ -593,8 +589,8 @@ void FAnimTrace::OutputAnimNodeValue(const FAnimationBaseContext& InContext, con
 
 void FAnimTrace::OutputAnimNodeValue(const FAnimationBaseContext& InContext, const TCHAR* InKey, float InValue)
 {
-	bool bEventEnabled = UE_TRACE_EVENT_IS_ENABLED(Animation, AnimNodeValueFloat);
-	if (!bEventEnabled)
+	bool bChannelEnabled = UE_TRACE_CHANNELEXPR_IS_ENABLED(AnimationChannel);
+	if (!bChannelEnabled)
 	{
 		return;
 	}
@@ -605,7 +601,7 @@ void FAnimTrace::OutputAnimNodeValue(const FAnimationBaseContext& InContext, con
 
 	int32 KeyLength = FCString::Strlen(InKey) + 1;
 
-	UE_TRACE_LOG(Animation, AnimNodeValueFloat, KeyLength * sizeof(TCHAR))
+	UE_TRACE_LOG(Animation, AnimNodeValueFloat, AnimationChannel, KeyLength * sizeof(TCHAR))
 		<< AnimNodeValueFloat.Cycle(FPlatformTime::Cycles64())
 		<< AnimNodeValueFloat.AnimInstanceId(FObjectTrace::GetObjectId(InContext.AnimInstanceProxy->GetAnimInstanceObject()))
 		<< AnimNodeValueFloat.NodeId(InContext.GetCurrentNodeId())
@@ -623,8 +619,8 @@ void FAnimTrace::OutputAnimNodeValue(const FAnimationBaseContext& InContext, con
 
 void FAnimTrace::OutputAnimNodeValue(const FAnimationBaseContext& InContext, const TCHAR* InKey, const FVector& InValue)
 {
-	bool bEventEnabled = UE_TRACE_EVENT_IS_ENABLED(Animation, AnimNodeValueVector);
-	if (!bEventEnabled)
+	bool bChannelEnabled = UE_TRACE_CHANNELEXPR_IS_ENABLED(AnimationChannel);
+	if (!bChannelEnabled)
 	{
 		return;
 	}
@@ -635,7 +631,7 @@ void FAnimTrace::OutputAnimNodeValue(const FAnimationBaseContext& InContext, con
 
 	int32 KeyLength = FCString::Strlen(InKey) + 1;
 
-	UE_TRACE_LOG(Animation, AnimNodeValueVector, KeyLength * sizeof(TCHAR))
+	UE_TRACE_LOG(Animation, AnimNodeValueVector, AnimationChannel, KeyLength * sizeof(TCHAR))
 		<< AnimNodeValueVector.Cycle(FPlatformTime::Cycles64())
 		<< AnimNodeValueVector.AnimInstanceId(FObjectTrace::GetObjectId(InContext.AnimInstanceProxy->GetAnimInstanceObject()))
 		<< AnimNodeValueVector.NodeId(InContext.GetCurrentNodeId())
@@ -649,8 +645,8 @@ void FAnimTrace::OutputAnimNodeValue(const FAnimationBaseContext& InContext, con
 
 void FAnimTrace::OutputAnimNodeValue(const FAnimationBaseContext& InContext, const TCHAR* InKey, const FName& InValue)
 {
-	bool bEventEnabled = UE_TRACE_EVENT_IS_ENABLED(Animation, AnimNodeValueString);
-	if (!bEventEnabled)
+	bool bChannelEnabled = UE_TRACE_CHANNELEXPR_IS_ENABLED(AnimationChannel);
+	if (!bChannelEnabled)
 	{
 		return;
 	}
@@ -668,7 +664,7 @@ void FAnimTrace::OutputAnimNodeValue(const FAnimationBaseContext& InContext, con
 		InValue.ToString(reinterpret_cast<TCHAR*>(Out) + KeyLength, ValueLength);
 	};
 
-	UE_TRACE_LOG(Animation, AnimNodeValueString, (KeyLength + ValueLength) * sizeof(TCHAR))
+	UE_TRACE_LOG(Animation, AnimNodeValueString, AnimationChannel, (KeyLength + ValueLength) * sizeof(TCHAR))
 		<< AnimNodeValueString.Cycle(FPlatformTime::Cycles64())
 		<< AnimNodeValueString.AnimInstanceId(FObjectTrace::GetObjectId(InContext.AnimInstanceProxy->GetAnimInstanceObject()))
 		<< AnimNodeValueString.NodeId(InContext.GetCurrentNodeId())
@@ -679,8 +675,8 @@ void FAnimTrace::OutputAnimNodeValue(const FAnimationBaseContext& InContext, con
 
 void FAnimTrace::OutputAnimNodeValue(const FAnimationBaseContext& InContext, const TCHAR* InKey, const TCHAR* InValue)
 {
-	bool bEventEnabled = UE_TRACE_EVENT_IS_ENABLED(Animation, AnimNodeValueString);
-	if (!bEventEnabled)
+	bool bChannelEnabled = UE_TRACE_CHANNELEXPR_IS_ENABLED(AnimationChannel);
+	if (!bChannelEnabled)
 	{
 		return;
 	}
@@ -698,7 +694,7 @@ void FAnimTrace::OutputAnimNodeValue(const FAnimationBaseContext& InContext, con
 		FCString::Strncpy(reinterpret_cast<TCHAR*>(Out) + KeyLength, InValue, ValueLength);
 	};
 
-	UE_TRACE_LOG(Animation, AnimNodeValueString, (KeyLength + ValueLength) * sizeof(TCHAR))
+	UE_TRACE_LOG(Animation, AnimNodeValueString, AnimationChannel, (KeyLength + ValueLength) * sizeof(TCHAR))
 		<< AnimNodeValueString.Cycle(FPlatformTime::Cycles64())
 		<< AnimNodeValueString.AnimInstanceId(FObjectTrace::GetObjectId(InContext.AnimInstanceProxy->GetAnimInstanceObject()))
 		<< AnimNodeValueString.NodeId(InContext.GetCurrentNodeId())
@@ -709,8 +705,8 @@ void FAnimTrace::OutputAnimNodeValue(const FAnimationBaseContext& InContext, con
 
 void FAnimTrace::OutputAnimNodeValue(const FAnimationBaseContext& InContext, const TCHAR* InKey, const UObject* InValue)
 {
-	bool bEventEnabled = UE_TRACE_EVENT_IS_ENABLED(Animation, AnimNodeValueObject);
-	if (!bEventEnabled)
+	bool bChannelEnabled = UE_TRACE_CHANNELEXPR_IS_ENABLED(AnimationChannel);
+	if (!bChannelEnabled)
 	{
 		return;
 	}
@@ -722,7 +718,7 @@ void FAnimTrace::OutputAnimNodeValue(const FAnimationBaseContext& InContext, con
 
 	int32 KeyLength = FCString::Strlen(InKey) + 1;
 
-	UE_TRACE_LOG(Animation, AnimNodeValueObject, KeyLength * sizeof(TCHAR))
+	UE_TRACE_LOG(Animation, AnimNodeValueObject, AnimationChannel, KeyLength * sizeof(TCHAR))
 		<< AnimNodeValueObject.Cycle(FPlatformTime::Cycles64())
 		<< AnimNodeValueObject.AnimInstanceId(FObjectTrace::GetObjectId(InContext.AnimInstanceProxy->GetAnimInstanceObject()))
 		<< AnimNodeValueObject.NodeId(InContext.GetCurrentNodeId())
@@ -734,8 +730,8 @@ void FAnimTrace::OutputAnimNodeValue(const FAnimationBaseContext& InContext, con
 
 void FAnimTrace::OutputAnimNodeValue(const FAnimationBaseContext& InContext, const TCHAR* InKey, const UClass* InValue)
 {
-	bool bEventEnabled = UE_TRACE_EVENT_IS_ENABLED(Animation, AnimNodeValueObject);
-	if (!bEventEnabled)
+	bool bChannelEnabled = UE_TRACE_CHANNELEXPR_IS_ENABLED(AnimationChannel);
+	if (!bChannelEnabled)
 	{
 		return;
 	}
@@ -747,7 +743,7 @@ void FAnimTrace::OutputAnimNodeValue(const FAnimationBaseContext& InContext, con
 
 	int32 KeyLength = FCString::Strlen(InKey) + 1;
 
-	UE_TRACE_LOG(Animation, AnimNodeValueClass, KeyLength * sizeof(TCHAR))
+	UE_TRACE_LOG(Animation, AnimNodeValueClass, AnimationChannel, KeyLength * sizeof(TCHAR))
 		<< AnimNodeValueClass.Cycle(FPlatformTime::Cycles64())
 		<< AnimNodeValueClass.AnimInstanceId(FObjectTrace::GetObjectId(InContext.AnimInstanceProxy->GetAnimInstanceObject()))
 		<< AnimNodeValueClass.NodeId(InContext.GetCurrentNodeId())
@@ -759,8 +755,8 @@ void FAnimTrace::OutputAnimNodeValue(const FAnimationBaseContext& InContext, con
 
 void FAnimTrace::OutputAnimSequencePlayer(const FAnimationBaseContext& InContext, const FAnimNode_SequencePlayer& InNode)
 {
-	bool bEventEnabled = UE_TRACE_EVENT_IS_ENABLED(Animation, AnimSequencePlayer);
-	if (!bEventEnabled)
+	bool bChannelEnabled = UE_TRACE_CHANNELEXPR_IS_ENABLED(AnimationChannel);
+	if (!bChannelEnabled)
 	{
 		return;
 	}
@@ -769,7 +765,7 @@ void FAnimTrace::OutputAnimSequencePlayer(const FAnimationBaseContext& InContext
 
 	TRACE_OBJECT(InContext.AnimInstanceProxy->GetAnimInstanceObject());
 
-	UE_TRACE_LOG(Animation, AnimSequencePlayer)
+	UE_TRACE_LOG(Animation, AnimSequencePlayer, AnimationChannel)
 		<< AnimSequencePlayer.Cycle(FPlatformTime::Cycles64())
 		<< AnimSequencePlayer.AnimInstanceId(FObjectTrace::GetObjectId(InContext.AnimInstanceProxy->GetAnimInstanceObject()))
 		<< AnimSequencePlayer.NodeId(InContext.GetCurrentNodeId())
@@ -780,8 +776,8 @@ void FAnimTrace::OutputAnimSequencePlayer(const FAnimationBaseContext& InContext
 
 void FAnimTrace::OutputStateMachineState(const FAnimationBaseContext& InContext, int32 InStateMachineIndex, int32 InStateIndex, float InStateWeight, float InElapsedTime)
 {
-	bool bEventEnabled = UE_TRACE_EVENT_IS_ENABLED(Animation, StateMachineState);
-	if (!bEventEnabled)
+	bool bChannelEnabled = UE_TRACE_CHANNELEXPR_IS_ENABLED(AnimationChannel);
+	if (!bChannelEnabled)
 	{
 		return;
 	}
@@ -790,7 +786,7 @@ void FAnimTrace::OutputStateMachineState(const FAnimationBaseContext& InContext,
 
 	TRACE_OBJECT(InContext.AnimInstanceProxy->GetAnimInstanceObject());
 
-	UE_TRACE_LOG(Animation, StateMachineState)
+	UE_TRACE_LOG(Animation, StateMachineState, AnimationChannel)
 		<< StateMachineState.Cycle(FPlatformTime::Cycles64())
 		<< StateMachineState.AnimInstanceId(FObjectTrace::GetObjectId(InContext.AnimInstanceProxy->GetAnimInstanceObject()))
 		<< StateMachineState.NodeId(InContext.GetCurrentNodeId())
diff --git a/Engine/Source/Runtime/Engine/Private/ObjectTrace.cpp b/Engine/Source/Runtime/Engine/Private/ObjectTrace.cpp
--- a/Engine/Source/Runtime/Engine/Private/ObjectTrace.cpp
+++ b/Engine/Source/Runtime/Engine/Private/ObjectTrace.cpp
@@ -15,6 +15,8 @@
 #include "HAL/IConsoleManager.h"
 #include "Engine/World.h"
 
+UE_TRACE_CHANNEL(ObjectChannel)
+
 UE_TRACE_EVENT_BEGIN(Object, Class, Important)
 	UE_TRACE_EVENT_FIELD(uint64, Id)
 	UE_TRACE_EVENT_FIELD(uint64, SuperId)
@@ -81,7 +83,7 @@ void FObjectTrace::Init()
 		UE_TRACE_EVENT_IS_ENABLED(Object, Class);
 		UE_TRACE_EVENT_IS_ENABLED(Object, Object);
 		UE_TRACE_EVENT_IS_ENABLED(Object, ObjectEvent);
-		Trace::ToggleEvent(TEXT("Object"), true);
+		Trace::ToggleChannel(ObjectChannel, true);
 	}
 }
 
@@ -122,8 +124,8 @@ uint64 FObjectTrace::GetObjectId(const UObject* InObject)
 
 void FObjectTrace::OutputClass(const UClass* InClass)
 {
-	bool bEventEnabled = UE_TRACE_EVENT_IS_ENABLED(Object, Class);
-	if (!bEventEnabled || InClass == nullptr)
+	bool bChannelEnabled = UE_TRACE_CHANNELEXPR_IS_ENABLED(ObjectChannel);
+	if (!bChannelEnabled || InClass == nullptr)
 	{
 		return;
 	}
@@ -147,7 +149,7 @@ void FObjectTrace::OutputClass(const UClass* InClass)
 		FPlatformMemory::Memcpy(reinterpret_cast<TCHAR*>(Out) + ClassNameStringLength, *InClass->GetPathName(), ClassFullNameStringLength * sizeof(TCHAR));
 	};
 
-	UE_TRACE_LOG(Object, Class, (ClassNameStringLength + ClassFullNameStringLength) * sizeof(TCHAR))
+	UE_TRACE_LOG(Object, Class, ObjectChannel, (ClassNameStringLength + ClassFullNameStringLength) * sizeof(TCHAR))
 		<< Class.ClassNameStringLength(ClassNameStringLength)
 		<< Class.Id(GetObjectId(InClass))
 		<< Class.SuperId(GetObjectId(InClass->GetSuperClass()))
@@ -156,8 +158,8 @@ void FObjectTrace::OutputClass(const UClass* InClass)
 
 void FObjectTrace::OutputObject(const UObject* InObject)
 {
-	bool bEventEnabled = UE_TRACE_EVENT_IS_ENABLED(Object, Object);
-	if (!bEventEnabled || InObject == nullptr)
+	bool bChannelEnabled = UE_TRACE_CHANNELEXPR_IS_ENABLED(ObjectChannel);
+	if (!bChannelEnabled || InObject == nullptr)
 	{
 		return;
 	}
@@ -194,7 +196,7 @@ void FObjectTrace::OutputObject(const UObject* InObject)
 		FPlatformMemory::Memcpy(reinterpret_cast<TCHAR*>(Out) + ObjectNameStringLength, *InObject->GetPathName(), ObjectPathNameStringLength * sizeof(TCHAR));
 	};
 
-	UE_TRACE_LOG(Object, Object, (ObjectNameStringLength + ObjectPathNameStringLength) * sizeof(TCHAR))
+	UE_TRACE_LOG(Object, Object, ObjectChannel, (ObjectNameStringLength + ObjectPathNameStringLength) * sizeof(TCHAR))
 		<< Object.ObjectNameStringLength(ObjectNameStringLength)
 		<< Object.Id(GetObjectId(InObject))
 		<< Object.ClassId(GetObjectId(InObject->GetClass()))
@@ -204,8 +206,8 @@ void FObjectTrace::OutputObject(const UObject* InObject)
 
 void FObjectTrace::OutputObjectEvent(const UObject* InObject, const TCHAR* InEvent)
 {
-	bool bEventEnabled = UE_TRACE_EVENT_IS_ENABLED(Object, ObjectEvent);
-	if (!bEventEnabled || InObject == nullptr)
+	bool bChannelEnabled = UE_TRACE_CHANNELEXPR_IS_ENABLED(ObjectChannel);
+	if (!bChannelEnabled || InObject == nullptr)
 	{
 		return;
 	}
@@ -224,7 +226,7 @@ void FObjectTrace::OutputObjectEvent(const UObject* InObject, const TCHAR* InEve
 
 	int32 StringBufferSize = (FCString::Strlen(InEvent) + 1) * sizeof(TCHAR);
 
-	UE_TRACE_LOG(Object, ObjectEvent, StringBufferSize)
+	UE_TRACE_LOG(Object, ObjectEvent, ObjectChannel, StringBufferSize)
 		<< ObjectEvent.Cycle(FPlatformTime::Cycles64())
 		<< ObjectEvent.Id(GetObjectId(InObject))
 		<< ObjectEvent.Attachment(InEvent, StringBufferSize);
diff --git a/Engine/Source/Runtime/Engine/Public/Animation/AnimTrace.h b/Engine/Source/Runtime/Engine/Public/Animation/AnimTrace.h
--- a/Engine/Source/Runtime/Engine/Public/Animation/AnimTrace.h
+++ b/Engine/Source/Runtime/Engine/Public/Animation/AnimTrace.h
@@ -5,11 +5,14 @@
 #include "CoreTypes.h"
 #include "ObjectTrace.h"
 #include "CoreMinimal.h"
+#include "Trace/Trace.h"
 
 #define ANIM_TRACE_ENABLED OBJECT_TRACE_ENABLED
 
 #if ANIM_TRACE_ENABLED
 
+ENGINE_API UE_TRACE_CHANNEL_EXTERN(AnimationChannel);
+
 struct FAnimInstanceProxy;
 struct FAnimTickRecord;
 struct FAnimationBaseContext;
diff --git a/Engine/Source/Runtime/Launch/Private/LaunchEngineLoop.cpp b/Engine/Source/Runtime/Launch/Private/LaunchEngineLoop.cpp
--- a/Engine/Source/Runtime/Launch/Private/LaunchEngineLoop.cpp
+++ b/Engine/Source/Runtime/Launch/Private/LaunchEngineLoop.cpp
@@ -208,6 +208,7 @@ class FFeedbackContext;
 #if defined(WITH_LAUNCHERCHECK) && WITH_LAUNCHERCHECK
 	#include "ILauncherCheckModule.h"
 #endif
+#include <String/ParseTokens.h>
 
 #if WITH_COREUOBJECT
 	#ifndef USE_LOCALIZED_PACKAGE_CACHE
@@ -1288,10 +1289,25 @@ DECLARE_CYCLE_STAT(TEXT("FEngineLoop::PreInitPostStartupScreen.AfterStats"), STA
 
 int32 FEngineLoop::PreInitPreStartupScreen(const TCHAR* CmdLine)
 {
-	TRACE_REGISTER_GAME_THREAD(FPlatformTLS::GetCurrentThreadId());
-	TRACE_CPUPROFILER_INIT(CmdLine);
-	TRACE_PLATFORMFILE_INIT(CmdLine);
-	TRACE_COUNTERS_INIT(CmdLine);
+#if UE_TRACE_ENABLED
+	{
+		Trace::Initialize();
+
+		FString EnabledChannels;
+		FParse::Value(CmdLine, TEXT("-trace="), EnabledChannels, false);
+		UE::String::ParseTokens(EnabledChannels, TEXT(","), [](FStringView Token) {
+			TCHAR ChannelName[64];
+			const size_t ChannelNameSize = Token.CopyString(ChannelName, 64);
+			ChannelName[ChannelNameSize] = '\0';
+			Trace::ToggleChannel(ChannelName, true);
+		});
+
+		TRACE_REGISTER_GAME_THREAD(FPlatformTLS::GetCurrentThreadId());
+		TRACE_CPUPROFILER_INIT(CmdLine);
+		TRACE_PLATFORMFILE_INIT(CmdLine);
+		TRACE_COUNTERS_INIT(CmdLine);
+	}
+#endif
 
 	SCOPED_BOOT_TIMING("FEngineLoop::PreInit");
 
@@ -1321,14 +1337,14 @@ int32 FEngineLoop::PreInitPreStartupScreen(const TCHAR* CmdLine)
 		uint8 AppNameOffset = AddToPayload(TEXT(UE_APP_NAME));
 		uint8 CommandLineOffset = AddToPayload(CmdLine);
 
-		UE_TRACE_EVENT_BEGIN(Diagnostics, Session, Important|Always)
+		UE_TRACE_EVENT_BEGIN(Diagnostics, Session, Important)
 			UE_TRACE_EVENT_FIELD(uint8, AppNameOffset)
 			UE_TRACE_EVENT_FIELD(uint8, CommandLineOffset)
 			UE_TRACE_EVENT_FIELD(uint8, ConfigurationType)
 			UE_TRACE_EVENT_FIELD(uint8, TargetType)
 		UE_TRACE_EVENT_END()
 
-		UE_TRACE_LOG(Diagnostics, Session, PayloadSize)
+		UE_TRACE_LOG(Diagnostics, Session, TraceLogChannel, PayloadSize)
 			<< Session.AppNameOffset(AppNameOffset)
 			<< Session.CommandLineOffset(CommandLineOffset)
 			<< Session.ConfigurationType(uint8(FApp::GetBuildConfiguration()))
diff --git a/Engine/Source/Runtime/RHI/Private/GpuProfilerTrace.cpp b/Engine/Source/Runtime/RHI/Private/GpuProfilerTrace.cpp
--- a/Engine/Source/Runtime/RHI/Private/GpuProfilerTrace.cpp
+++ b/Engine/Source/Runtime/RHI/Private/GpuProfilerTrace.cpp
@@ -8,6 +8,8 @@
 
 #if GPUPROFILERTRACE_ENABLED
 
+UE_TRACE_CHANNEL(GpuProfilerChannel)
+
 UE_TRACE_EVENT_BEGIN(GpuProfiler, EventSpec, Always)
 	UE_TRACE_EVENT_FIELD(const void*, EventType)
 	UE_TRACE_EVENT_FIELD(uint16, NameLength)
@@ -24,7 +26,7 @@ FGpuProfilerTrace::FEventType::FEventType(const TCHAR* Name)
 {
 	uint16 NameLength = FCString::Strlen(Name);
 	uint16 NameSize = (NameLength + 1) * sizeof(TCHAR);
-	UE_TRACE_LOG(GpuProfiler, EventSpec, NameSize) 
+	UE_TRACE_LOG(GpuProfiler, EventSpec, GpuProfilerChannel, NameSize)
 		<< EventSpec.EventType(this)
 		<< EventSpec.NameLength(NameLength)
 		<< EventSpec.Attachment(Name, NameSize);
@@ -69,7 +71,7 @@ void FGpuProfilerTrace::EndFrame()
 {
 	if (CurrentFrame.EventBufferSize)
 	{
-		UE_TRACE_LOG(GpuProfiler, Frame, CurrentFrame.EventBufferSize)
+		UE_TRACE_LOG(GpuProfiler, Frame, GpuProfilerChannel, CurrentFrame.EventBufferSize)
 			<< Frame.TimestampBase(CurrentFrame.TimestampBase)
 			<< Frame.RenderingFrameNumber(CurrentFrame.RenderingFrameNumber)
 			<< Frame.Attachment(CurrentFrame.EventBuffer, CurrentFrame.EventBufferSize);
diff --git a/Engine/Source/Runtime/RHI/Private/RHICommandList.cpp b/Engine/Source/Runtime/RHI/Private/RHICommandList.cpp
--- a/Engine/Source/Runtime/RHI/Private/RHICommandList.cpp
+++ b/Engine/Source/Runtime/RHI/Private/RHICommandList.cpp
@@ -11,6 +11,7 @@
 #include "Misc/ScopeLock.h"
 #include "PipelineStateCache.h"
 #include "ProfilingDebugging/CsvProfiler.h"
+#include "Trace/Trace.h"
 
 CSV_DEFINE_CATEGORY_MODULE(RHI_API, RHITStalls, false);
 CSV_DEFINE_CATEGORY_MODULE(RHI_API, RHITFlushes, false);
@@ -23,8 +24,7 @@ DECLARE_CYCLE_STAT(TEXT("All Command List Execute"), STAT_ImmedCmdListExecuteTim
 DECLARE_DWORD_COUNTER_STAT(TEXT("Immed. Command List memory"), STAT_ImmedCmdListMemory, STATGROUP_RHICMDLIST);
 DECLARE_DWORD_COUNTER_STAT(TEXT("Immed. Command count"), STAT_ImmedCmdListCount, STATGROUP_RHICMDLIST);
 
-
-
+UE_TRACE_CHANNEL_DEFINE(RHICommandsChannel);
 
 #if !PLATFORM_USES_FIXED_RHI_CLASS
 #include "RHICommandListCommandExecutes.inl"
diff --git a/Engine/Source/Runtime/RHI/Public/RHICommandList.h b/Engine/Source/Runtime/RHI/Public/RHICommandList.h
--- a/Engine/Source/Runtime/RHI/Public/RHICommandList.h
+++ b/Engine/Source/Runtime/RHI/Public/RHICommandList.h
@@ -27,6 +27,7 @@
 #include "HAL/LowLevelMemTracker.h"
 #include "ProfilingDebugging/CsvProfiler.h"
 #include "ProfilingDebugging/CpuProfilerTrace.h"
+#include "Trace/Trace.h"
 
 CSV_DECLARE_CATEGORY_MODULE_EXTERN(RHI_API, RHITStalls);
 CSV_DECLARE_CATEGORY_MODULE_EXTERN(RHI_API, RHITFlushes);
@@ -59,6 +60,7 @@ class FRayTracingPipelineState;
 
 DECLARE_STATS_GROUP(TEXT("RHICmdList"), STATGROUP_RHICMDLIST, STATCAT_Advanced);
 
+RHI_API extern Trace::FChannel RHICommandsChannel;
 
 // set this one to get a stat for each RHI command 
 #define RHI_STATS 0
@@ -84,6 +86,11 @@ extern RHI_API uint32 GWorkingRHIThreadStartCycles;
 /** How many cycles the from sampling input to the frame being flipped. */
 extern RHI_API uint64 GInputLatencyTime;
 
+/*Trace::FChannel& FORCEINLINE GetRHICommandsChannel() 
+{
+
+}*/
+
 /**
 * Whether the RHI commands are being run in a thread other than the render thread
 */
@@ -697,28 +704,8 @@ struct FRHICommand : public FRHICommandBase
 
 	void ExecuteAndDestruct(FRHICommandListBase& CmdList, FRHICommandListDebugContext& Context) override final
 	{
-#if CPUPROFILERTRACE_ENABLED
-		static uint16 __CpuProfilerEventSpecId;
-		if (__CpuProfilerEventSpecId == 0)
-		{
-			__CpuProfilerEventSpecId = FCpuProfilerTrace::OutputEventType(NameType::TStr(), CpuProfilerGroup_Default);
-		}
-
-		extern RHI_API int32 GRHICmdTraceEvents;
-		struct FConditionalTraceScope
-		{
-			FConditionalTraceScope(const uint16 InSpecId) : SpecId(InSpecId)
-			{
-				if (SpecId) FCpuProfilerTrace::OutputBeginEvent(SpecId);
-			}
-			~FConditionalTraceScope()
-			{
-				if (SpecId) FCpuProfilerTrace::OutputEndEvent();
-			}
-			const uint16 SpecId;
-		} TraceScope(GRHICmdTraceEvents ? __CpuProfilerEventSpecId : 0);
-#endif // CPUPROFILERTRACE_ENABLED
-
+		TRACE_CPUPROFILER_EVENT_SCOPE_ON_CHANNEL_STR(NameType::TStr(), RHICommandsChannel);
+		
 		TCmd *ThisCmd = static_cast<TCmd*>(this);
 #if RHI_COMMAND_LIST_DEBUG_TRACES
 		ThisCmd->StoreDebugInfo(Context);
@@ -3759,14 +3746,7 @@ class RHI_API FRHICommandListImmediate : public FRHICommandList
 
 		void ExecuteAndDestruct(FRHICommandListBase& CmdList, FRHICommandListDebugContext&) override final
 		{
-#if CPUPROFILERTRACE_ENABLED
-			static uint16 __CpuProfilerEventSpecId;
-			if (__CpuProfilerEventSpecId == 0)
-			{
-				__CpuProfilerEventSpecId = FCpuProfilerTrace::OutputEventType(TEXT("TRHILambdaCommand"), CpuProfilerGroup_Default);
-			}
-			FCpuProfilerTrace::FEventScope __CpuProfilerEventScope(__CpuProfilerEventSpecId);
-#endif
+			TRACE_CPUPROFILER_EVENT_SCOPE_ON_CHANNEL(TRHILambdaCommand, RHICommandsChannel);
 			Lambda(*static_cast<FRHICommandListImmediate*>(&CmdList));
 			Lambda.~LAMBDA();
 		}
diff --git a/Engine/Source/Runtime/RHI/RHI.Build.cs b/Engine/Source/Runtime/RHI/RHI.Build.cs
--- a/Engine/Source/Runtime/RHI/RHI.Build.cs
+++ b/Engine/Source/Runtime/RHI/RHI.Build.cs
@@ -8,6 +8,7 @@ public class RHI : ModuleRules
 	public RHI(ReadOnlyTargetRules Target) : base(Target)
 	{
 		PrivateDependencyModuleNames.Add("Core");
+		PrivateDependencyModuleNames.Add("TraceLog");
 		PrivateDependencyModuleNames.Add("ApplicationCore");
 
 		if (Target.bCompileAgainstEngine)
diff --git a/Engine/Source/Runtime/RenderCore/Private/RenderingThread.cpp b/Engine/Source/Runtime/RenderCore/Private/RenderingThread.cpp
--- a/Engine/Source/Runtime/RenderCore/Private/RenderingThread.cpp
+++ b/Engine/Source/Runtime/RenderCore/Private/RenderingThread.cpp
@@ -29,7 +29,7 @@
 // Globals
 //
 
-
+UE_TRACE_CHANNEL_DEFINE(RenderCommandsChannel);
 
 RENDERCORE_API bool GIsThreadedRendering = false;
 RENDERCORE_API bool GUseThreadedRendering = false;
diff --git a/Engine/Source/Runtime/RenderCore/Public/RenderingThread.h b/Engine/Source/Runtime/RenderCore/Public/RenderingThread.h
--- a/Engine/Source/Runtime/RenderCore/Public/RenderingThread.h
+++ b/Engine/Source/Runtime/RenderCore/Public/RenderingThread.h
@@ -10,6 +10,7 @@
 #include "Stats/Stats.h"
 #include "Async/TaskGraphInterfaces.h"
 #include "Templates/Atomic.h"
+#include "Trace/Trace.h"
 
 class FRHICommandListImmediate;
 
@@ -124,6 +125,8 @@ private:
 // Render commands
 ////////////////////////////////////
 
+RENDERCORE_API extern Trace::FChannel RenderCommandsChannel;
+
 /** The parent class of commands stored in the rendering command queue. */
 class RENDERCORE_API FRenderCommand
 {
@@ -186,14 +189,7 @@ public:
 
 	void DoTask(ENamedThreads::Type CurrentThread, const FGraphEventRef& MyCompletionGraphEvent)
 	{
-#if CPUPROFILERTRACE_ENABLED
-		static uint16 __CpuProfilerEventSpecId;
-		if (__CpuProfilerEventSpecId == 0)
-		{
-			__CpuProfilerEventSpecId = FCpuProfilerTrace::OutputEventType(TSTR::TStr(), CpuProfilerGroup_Default);
-		}
-		FCpuProfilerTrace::FEventScope __CpuProfilerEventScope(__CpuProfilerEventSpecId);
-#endif
+		TRACE_CPUPROFILER_EVENT_SCOPE_ON_CHANNEL_STR(TSTR::TStr(), RenderCommandsChannel);
 		FRHICommandListImmediate& RHICmdList = GetImmediateCommandList_ForRenderCommand();
 		Lambda(RHICmdList);
 	}
diff --git a/Engine/Source/Runtime/RenderCore/RenderCore.Build.cs b/Engine/Source/Runtime/RenderCore/RenderCore.Build.cs
--- a/Engine/Source/Runtime/RenderCore/RenderCore.Build.cs
+++ b/Engine/Source/Runtime/RenderCore/RenderCore.Build.cs
@@ -19,7 +19,7 @@ public class RenderCore : ModuleRules
 
         }
 
-        PrivateDependencyModuleNames.AddRange(new string[] { "Core", "Projects", "RHI", "ApplicationCore" });
+        PrivateDependencyModuleNames.AddRange(new string[] { "Core", "Projects", "RHI", "ApplicationCore", "TraceLog" });
 
         PrivateIncludePathModuleNames.AddRange(new string[] { "DerivedDataCache" });
 		
