AsyncLoading2: Strip Pak files from package data

diff --git a/Engine/Source/Programs/AutomationTool/Scripts/CopyBuildToStagingDirectory.Automation.cs b/Engine/Source/Programs/AutomationTool/Scripts/CopyBuildToStagingDirectory.Automation.cs
--- a/Engine/Source/Programs/AutomationTool/Scripts/CopyBuildToStagingDirectory.Automation.cs
+++ b/Engine/Source/Programs/AutomationTool/Scripts/CopyBuildToStagingDirectory.Automation.cs
@@ -2204,13 +2204,6 @@ public partial class Project : CommandUtils
 							{
 								UnrealPakResponseFile.Add(Entry.Key, Entry.Value);
 							}
-
-							// TODO: Remove when file enumeration has been fixed
-							if (Path.GetExtension(Entry.Key).Contains(".uasset") ||
-								Path.GetExtension(Entry.Key).Contains(".umap"))
-							{
-								UnrealPakResponseFile.Add(Entry.Key, Entry.Value);
-							}
 						}
 						IoStoreCommands.Add(GetIoStoreCommand(IoStoreResponseFile, OutputLocation, PrimaryOrderFile, PakParams.bCompressed, CryptoSettings, PakParams.EncryptionKeyGuid, SecondaryOrderFile));
 					}
diff --git a/Engine/Source/Runtime/CoreUObject/Private/Misc/PackageName.cpp b/Engine/Source/Runtime/CoreUObject/Private/Misc/PackageName.cpp
--- a/Engine/Source/Runtime/CoreUObject/Private/Misc/PackageName.cpp
+++ b/Engine/Source/Runtime/CoreUObject/Private/Misc/PackageName.cpp
@@ -24,6 +24,7 @@
 #include "Misc/PathViews.h"
 #include "Misc/ScopeLock.h"
 #include "Misc/StringBuilder.h"
+#include "IO/IoDispatcher.h"
 
 DEFINE_LOG_CATEGORY_STATIC(LogPackageName, Log, All);
 
@@ -39,6 +40,9 @@ FPackageName::FOnContentPathMountedEvent FPackageName::OnContentPathMountedEvent
 /** Event that is triggered when a content path is dismounted */
 FPackageName::FOnContentPathDismountedEvent FPackageName::OnContentPathDismountedEvent;
 
+/** Delegate used to check whether a package exist without using the filesystem. */
+FPackageName::FDoesPackageExistOverride FPackageName::DoesPackageExistOverrideDelegate;
+
 namespace PackageNameConstants
 {
 	// Minimum theoretical package name length ("/A/B") is 4
@@ -938,6 +942,22 @@ bool FPackageName::DoesPackageExist(const FString& LongPackageName, const FGuid*
 		return false;
 	}
 
+	// Used when I/O dispatcher is enabled
+	if (DoesPackageExistOverrideDelegate.IsBound())
+	{
+		if (DoesPackageExistOverrideDelegate.Execute(FName(*PackageName)))
+		{
+			if (OutFilename)
+			{
+				*OutFilename = LongPackageNameToFilename(PackageName, TEXT(""));
+			}
+
+			return true;
+		}
+
+		return false;
+	}
+
 	// Convert to filename (no extension yet).
 	FString Filename = LongPackageNameToFilename(PackageName, TEXT(""));
 	// Find the filename (with extension).
@@ -1269,6 +1289,8 @@ bool FPackageName::IsPackageExtension( const TCHAR* Ext )
 
 bool FPackageName::FindPackagesInDirectory( TArray<FString>& OutPackages, const FString& RootDir )
 {
+	UE_CLOG(FIoDispatcher::IsInitialized(), LogPackageName, Error, TEXT("Can't search for packages using the filesystem when I/O dispatcher is enabled"));
+
 	// Find all files in RootDir, then filter by extension (we have two package extensions so they can't
 	// be included in the search wildcard.
 	TArray<FString> AllFiles;
diff --git a/Engine/Source/Runtime/CoreUObject/Private/Serialization/AsyncLoading2.cpp b/Engine/Source/Runtime/CoreUObject/Private/Serialization/AsyncLoading2.cpp
--- a/Engine/Source/Runtime/CoreUObject/Private/Serialization/AsyncLoading2.cpp
+++ b/Engine/Source/Runtime/CoreUObject/Private/Serialization/AsyncLoading2.cpp
@@ -464,6 +464,11 @@ public:
 
 			ScriptArcs = reinterpret_cast<int32*>(Data);
 		}
+
+		FPackageName::DoesPackageExistOverride().BindLambda([this](FName InPackageName)
+		{
+			return PackageNameToPackageId.Contains(InPackageName);
+		});
 	}
 
 	void FinalizeInitialLoad()
diff --git a/Engine/Source/Runtime/CoreUObject/Public/Misc/PackageName.h b/Engine/Source/Runtime/CoreUObject/Public/Misc/PackageName.h
--- a/Engine/Source/Runtime/CoreUObject/Public/Misc/PackageName.h
+++ b/Engine/Source/Runtime/CoreUObject/Public/Misc/PackageName.h
@@ -463,6 +463,13 @@ public:
 	UE_DEPRECATED(4.17, "Deprecated. Call TryConvertFilenameToLongPackageName instead")
 	static FString PackageFromPath(const TCHAR* InPathName);
 
+	/** Override whether a package exist or not. */
+	DECLARE_DELEGATE_RetVal_OneParam(bool, FDoesPackageExistOverride, FName);
+	static FDoesPackageExistOverride& DoesPackageExistOverride()
+	{
+		return DoesPackageExistOverrideDelegate;
+	}
+
 private:
 
 	static FString AssetPackageExtension;
@@ -483,5 +490,8 @@ private:
 
 	/** Event that is triggered when a new content path is removed */
 	static FOnContentPathDismountedEvent OnContentPathDismountedEvent;
+
+	/** Delegate used to check whether a package exist without using the filesystem. */
+	static FDoesPackageExistOverride DoesPackageExistOverrideDelegate;
 };
 
diff --git a/Engine/Source/Runtime/Engine/Private/EngineUtils.cpp b/Engine/Source/Runtime/Engine/Private/EngineUtils.cpp
--- a/Engine/Source/Runtime/Engine/Private/EngineUtils.cpp
+++ b/Engine/Source/Runtime/Engine/Private/EngineUtils.cpp
@@ -20,6 +20,9 @@
 #include "Misc/UObjectToken.h"
 #include "Misc/MapErrors.h"
 #include "EngineModule.h"
+#include "Engine/AssetManager.h"
+#include "Misc/PathViews.h"
+#include "IO/IoDispatcher.h"
 
 #include "ProfilingDebugging/DiagnosticTable.h"
 #include "Interfaces/ITargetPlatform.h"
@@ -275,29 +278,57 @@ bool EngineUtils::FindOrLoadAssetsByPath(const FString& Path, TArray<UObject*>&
 		return false;
 	}
 
-	// Convert the package path to a filename with no extension (directory)
-	const FString FilePath = FPackageName::LongPackageNameToFilename(Path);
+	using FPackageNames = TArray<FName, TInlineAllocator<16>>;
 
-	// Gather the package files in that directory and subdirectories
-	TArray<FString> Filenames;
-	FPackageName::FindPackagesInDirectory(Filenames, FilePath);
-
-	// Cull out map files
-	for (int32 FilenameIdx = Filenames.Num() - 1; FilenameIdx >= 0; --FilenameIdx)
+	auto GetPackageNamesFromPath = [](const FString& InPath, FPackageNames& OutPackageNames)
 	{
-		const FString Extension = FPaths::GetExtension(Filenames[FilenameIdx], true);
-		if ( Extension == FPackageName::GetMapPackageExtension() )
+		// There is no filesystem support for packages when using the I/O dispatcher
+		if (FIoDispatcher::IsInitialized())
 		{
-			Filenames.RemoveAt(FilenameIdx);
+			FAssetRegistryModule& AssetRegistryModule = FModuleManager::LoadModuleChecked<FAssetRegistryModule>(TEXT("AssetRegistry"));
+			IAssetRegistry& AssetRegistry = AssetRegistryModule.Get();
+
+			TArray<FAssetData> Assets;
+			AssetRegistry.GetAssetsByPath(FName(*InPath), Assets, true);
+
+			for (const FAssetData& Asset : Assets)
+			{
+				// Cull packages containing maps
+				if ((Asset.PackageFlags & PKG_ContainsMap) != PKG_ContainsMap)
+				{
+					OutPackageNames.Emplace(Asset.PackageName);
+				}
+			}
 		}
-	}
+		else
+		{
+			// Convert the package path to a filename with no extension (directory)
+			const FString FilePath = FPackageName::LongPackageNameToFilename(InPath);
 
-	// Load packages or find existing ones and fully load them
-	for (int32 FileIdx = 0; FileIdx < Filenames.Num(); ++FileIdx)
-	{
-		const FString& Filename = Filenames[FileIdx];
+			// Gather the package files in that directory and subdirectories
+			TArray<FString> Filenames;
+			FPackageName::FindPackagesInDirectory(Filenames, FilePath);
 
-		UPackage* Package = FindPackage(NULL, *FPackageName::FilenameToLongPackageName(Filename));
+			// Cull out map files
+			for (const FString& Filename : Filenames)
+			{
+				FStringView Extension = FPathViews::GetExtension(Filename, true);
+				if (Extension != FPackageName::GetMapPackageExtension())
+				{
+					OutPackageNames.Emplace(*FPackageName::FilenameToLongPackageName(Filename));
+				}
+			}
+		}
+	};
+
+	FPackageNames PackageNames;
+	GetPackageNamesFromPath(Path, PackageNames);
+	TCHAR PackageName[FName::StringBufferSize];
+
+	for (const FName& Name : PackageNames)
+	{
+		Name.ToString(PackageName);
+		UPackage* Package = FindPackage(NULL, PackageName);
 
 		if (Package)
 		{
@@ -305,7 +336,7 @@ bool EngineUtils::FindOrLoadAssetsByPath(const FString& Path, TArray<UObject*>&
 		}
 		else
 		{
-			Package = LoadPackage(NULL, *Filename, LOAD_None);
+			Package = LoadPackage(NULL, PackageName, LOAD_None);
 		}
 
 		if (Package)
@@ -595,4 +626,4 @@ void VirtualTextureUtils::CheckAndReportInvalidUsage(const UObject* Owner, const
 	}
 }
 
-#undef LOCTEXT_NAMESPACE
\ No newline at end of file
+#undef LOCTEXT_NAMESPACE
diff --git a/Engine/Source/Runtime/Engine/Private/UserInterface/Console.cpp b/Engine/Source/Runtime/Engine/Private/UserInterface/Console.cpp
--- a/Engine/Source/Runtime/Engine/Private/UserInterface/Console.cpp
+++ b/Engine/Source/Runtime/Engine/Private/UserInterface/Console.cpp
@@ -30,6 +30,8 @@
 #include "Misc/TextFilter.h"
 #include "HAL/PlatformApplicationMisc.h"
 #include "ProfilingDebugging/CsvProfiler.h"
+#include "Engine/AssetManager.h"
+#include "IO/IoDispatcher.h"
 
 static const uint32 MAX_AUTOCOMPLETION_LINES = 20;
 
@@ -264,47 +266,74 @@ void UConsole::BuildRuntimeAutoCompleteList(bool bForce)
 
 	// enumerate maps
 	{
+		auto FindPackagesInDirectory = [](TArray<FString>& OutPackages, const FString& InPath)
+		{
+			// Can't search packages using the filesystem when I/O dispatcher is enabled
+			if (FIoDispatcher::IsInitialized())
+			{
+				FString PackagePath;
+				if (FPackageName::TryConvertFilenameToLongPackageName(InPath, PackagePath))
+				{
+					if (FAssetRegistryModule* AssetRegistryModule = FModuleManager::LoadModulePtr<FAssetRegistryModule>(TEXT("AssetRegistry")))
+					{
+						TArray<FAssetData> Assets;
+						AssetRegistryModule->Get().GetAssetsByPath(FName(*PackagePath), Assets, true);
+
+						for (const FAssetData& Asset : Assets)
+						{
+							if (!!(Asset.PackageFlags & PKG_ContainsMap))
+							{
+								OutPackages.Emplace(Asset.AssetName.ToString());
+							}
+						}
+					}
+				}
+			}
+			else
+			{
+				TArray<FString> Filenames;
+				FPackageName::FindPackagesInDirectory(Filenames, InPath);
+
+				for (const FString& Filename : Filenames)
+				{
+					const int32 NameStartIdx = Filename.Find(TEXT("/"), ESearchCase::IgnoreCase, ESearchDir::FromEnd);
+					const int32 ExtIdx = Filename.Find(*FPackageName::GetMapPackageExtension(), ESearchCase::IgnoreCase, ESearchDir::FromEnd);
+
+					if (NameStartIdx != INDEX_NONE && ExtIdx != INDEX_NONE)
+					{
+						OutPackages.Emplace(Filename.Mid(NameStartIdx + 1, ExtIdx - NameStartIdx - 1));
+					}
+				}
+			}
+		};
+
 		TArray<FString> Packages;
-		for (int32 PathIdx = 0; PathIdx < ConsoleSettings->AutoCompleteMapPaths.Num(); ++PathIdx)
+		for (const FString& MapPath : ConsoleSettings->AutoCompleteMapPaths)
 		{
-			FPackageName::FindPackagesInDirectory(Packages, FString::Printf(TEXT("%s%s"), *FPaths::ProjectDir(), *ConsoleSettings->AutoCompleteMapPaths[PathIdx]));
+			FindPackagesInDirectory(Packages, FString::Printf(TEXT("%s%s"), *FPaths::ProjectDir(), *MapPath));
 		}
 
-		// also include maps in this user's developer dir
-		FPackageName::FindPackagesInDirectory(Packages, FPaths::GameUserDeveloperDir());
+		FindPackagesInDirectory(Packages, FPaths::GameUserDeveloperDir());
 
-		for (int32 PackageIndex = 0; PackageIndex < Packages.Num(); PackageIndex++)
+		for (const FString& MapName : Packages)
 		{
-			FString Pkg = MoveTemp(Packages[PackageIndex]);
-			const int32 ExtIdx = Pkg.Find(*FPackageName::GetMapPackageExtension(), ESearchCase::IgnoreCase, ESearchDir::FromEnd);
-
-			FString MapName;
-			if (ExtIdx != INDEX_NONE && Pkg.Split(TEXT("/"), nullptr, &MapName, ESearchCase::CaseSensitive, ESearchDir::FromEnd))
+			int32 NewIdx = 0;
+			// put _P maps at the front so that they match early, since those are generally the maps we want to actually open
+			if (MapName.EndsWith(TEXT("_P")))
 			{
-				// try to peel off the extension
-				FString TrimmedMapName;
-				if (!MapName.Split(TEXT("."), &TrimmedMapName, nullptr, ESearchCase::CaseSensitive, ESearchDir::FromEnd))
-				{
-					TrimmedMapName = MapName;
-				}
-				int32 NewIdx;
-				// put _P maps at the front so that they match early, since those are generally the maps we want to actually open
-				if (TrimmedMapName.EndsWith(TEXT("_P")))
-				{
-					NewIdx = 0;
-					AutoCompleteList.InsertDefaulted(0, 3);
-				}
-				else
-				{
-					NewIdx = AutoCompleteList.AddDefaulted(3);
-				}
-				AutoCompleteList[NewIdx].Command = FString::Printf(TEXT("open %s"), *TrimmedMapName);
-				AutoCompleteList[NewIdx].Color = ConsoleSettings->AutoCompleteCommandColor;
-				AutoCompleteList[NewIdx + 1].Command = FString::Printf(TEXT("travel %s"), *TrimmedMapName);
-				AutoCompleteList[NewIdx + 1].Color = ConsoleSettings->AutoCompleteCommandColor;
-				AutoCompleteList[NewIdx + 2].Command = FString::Printf(TEXT("servertravel %s"), *TrimmedMapName);
-				AutoCompleteList[NewIdx + 2].Color = ConsoleSettings->AutoCompleteCommandColor;
+				AutoCompleteList.InsertDefaulted(0, 3);
+			}
+			else
+			{
+				NewIdx = AutoCompleteList.AddDefaulted(3);
 			}
+
+			AutoCompleteList[NewIdx].Command = FString::Printf(TEXT("open %s"), *MapName);
+			AutoCompleteList[NewIdx].Color = ConsoleSettings->AutoCompleteCommandColor;
+			AutoCompleteList[NewIdx + 1].Command = FString::Printf(TEXT("travel %s"), *MapName);
+			AutoCompleteList[NewIdx + 1].Color = ConsoleSettings->AutoCompleteCommandColor;
+			AutoCompleteList[NewIdx + 2].Command = FString::Printf(TEXT("servertravel %s"), *MapName);
+			AutoCompleteList[NewIdx + 2].Color = ConsoleSettings->AutoCompleteCommandColor;
 		}
 	}
 
