Moved skeletal mesh & component traces to use the new array API

Also removed initialization functions as they are no longer needed with channels

diff --git a/Engine/Plugins/Animation/GameplayInsights/Source/GameplayInsights/Private/AnimationAnalyzer.cpp b/Engine/Plugins/Animation/GameplayInsights/Source/GameplayInsights/Private/AnimationAnalyzer.cpp
--- a/Engine/Plugins/Animation/GameplayInsights/Source/GameplayInsights/Private/AnimationAnalyzer.cpp
+++ b/Engine/Plugins/Animation/GameplayInsights/Source/GameplayInsights/Private/AnimationAnalyzer.cpp
@@ -17,7 +17,9 @@ void FAnimationAnalyzer::OnAnalysisBegin(const FOnAnalysisContext& Context)
 
 	Builder.RouteEvent(RouteId_TickRecord, "Animation", "TickRecord");
 	Builder.RouteEvent(RouteId_SkeletalMesh, "Animation", "SkeletalMesh");
+	Builder.RouteEvent(RouteId_SkeletalMesh2, "Animation", "SkeletalMesh2");
 	Builder.RouteEvent(RouteId_SkeletalMeshComponent, "Animation", "SkeletalMeshComponent");
+	Builder.RouteEvent(RouteId_SkeletalMeshComponent2, "Animation", "SkeletalMeshComponent2");
 	Builder.RouteEvent(RouteId_SkeletalMeshFrame, "Animation", "SkeletalMeshFrame");
 	Builder.RouteEvent(RouteId_AnimGraph, "Animation", "AnimGraph");
 	Builder.RouteEvent(RouteId_AnimNodeStart, "Animation", "AnimNodeStart");
@@ -67,6 +69,14 @@ bool FAnimationAnalyzer::OnEvent(uint16 RouteId, const FOnEventContext& Context)
 		AnimationProvider.AppendSkeletalMesh(Id, ParentIndices);
 		break;
 	}
+	case RouteId_SkeletalMesh2:
+	{
+		uint64 Id = EventData.GetValue<uint64>("Id");
+		uint32 BoneCount = EventData.GetValue<uint32>("BoneCount");
+		TArrayView<const int32> ParentIndices = EventData.GetArrayView<int32>("ParentIndices");
+		AnimationProvider.AppendSkeletalMesh(Id, ParentIndices);
+		break;
+	}
 	case RouteId_SkeletalMeshComponent:
 	{
 		uint64 Cycle = EventData.GetValue<uint64>("Cycle");
@@ -81,6 +91,26 @@ bool FAnimationAnalyzer::OnEvent(uint16 RouteId, const FOnEventContext& Context)
 		AnimationProvider.AppendSkeletalMeshComponent(ComponentId, MeshId, Context.SessionContext.TimestampFromCycle(Cycle), LodIndex, FrameCounter, Pose, Curves);
 		break;
 	}
+	case RouteId_SkeletalMeshComponent2:
+	{
+		uint64 Cycle = EventData.GetValue<uint64>("Cycle");
+		uint64 ComponentId = EventData.GetValue<uint64>("ComponentId");
+		uint64 MeshId = EventData.GetValue<uint64>("MeshId");
+		uint32 BoneCount = EventData.GetValue<uint32>("BoneCount");
+		uint32 CurveCount = EventData.GetValue<uint32>("CurveCount");
+		uint16 FrameCounter = EventData.GetValue<uint16>("FrameCounter");
+		uint16 LodIndex = EventData.GetValue<uint16>("LodIndex");
+		TArrayView<const float> ComponentToWorldFloatArray = EventData.GetArrayView<float>("ComponentToWorld");
+		check(ComponentToWorldFloatArray.Num() == sizeof(FTransform) / sizeof(float));
+		const FTransform& ComponentToWorld = *reinterpret_cast<const FTransform*>(ComponentToWorldFloatArray.GetData());
+		TArrayView<const float> PoseFloatArray = EventData.GetArrayView<float>("Pose");
+		TArrayView<const FTransform> Pose(reinterpret_cast<const FTransform*>(PoseFloatArray.GetData()), PoseFloatArray.Num() / (sizeof(FTransform) / sizeof(float)));
+		TArrayView<const uint32> CurveIds = EventData.GetArrayView<uint32>("CurveIds");
+		TArrayView<const float> CurveValues = EventData.GetArrayView<float>("CurveValues");
+		check(CurveIds.Num() == CurveValues.Num());
+		AnimationProvider.AppendSkeletalMeshComponent(ComponentId, MeshId, Context.SessionContext.TimestampFromCycle(Cycle), LodIndex, FrameCounter, ComponentToWorld, Pose, CurveIds, CurveValues);
+		break;
+	}
 	case RouteId_Name:
 	{
 		uint32 Id = EventData.GetValue<uint32>("Id");
diff --git a/Engine/Plugins/Animation/GameplayInsights/Source/GameplayInsights/Private/AnimationAnalyzer.h b/Engine/Plugins/Animation/GameplayInsights/Source/GameplayInsights/Private/AnimationAnalyzer.h
--- a/Engine/Plugins/Animation/GameplayInsights/Source/GameplayInsights/Private/AnimationAnalyzer.h
+++ b/Engine/Plugins/Animation/GameplayInsights/Source/GameplayInsights/Private/AnimationAnalyzer.h
@@ -21,7 +21,9 @@ private:
 	{
 		RouteId_TickRecord,
 		RouteId_SkeletalMesh,
+		RouteId_SkeletalMesh2,
 		RouteId_SkeletalMeshComponent,
+		RouteId_SkeletalMeshComponent2,
 		RouteId_SkeletalMeshFrame,
 		RouteId_AnimGraph,
 		RouteId_AnimNodeStart,
diff --git a/Engine/Plugins/Animation/GameplayInsights/Source/GameplayInsights/Private/AnimationProvider.cpp b/Engine/Plugins/Animation/GameplayInsights/Source/GameplayInsights/Private/AnimationProvider.cpp
--- a/Engine/Plugins/Animation/GameplayInsights/Source/GameplayInsights/Private/AnimationProvider.cpp
+++ b/Engine/Plugins/Animation/GameplayInsights/Source/GameplayInsights/Private/AnimationProvider.cpp
@@ -447,6 +447,61 @@ void FAnimationProvider::AppendSkeletalMeshComponent(uint64 InObjectId, uint64 I
 	Session.UpdateDurationSeconds(InTime);
 }
 
+void FAnimationProvider::AppendSkeletalMeshComponent(uint64 InObjectId, uint64 InMeshId, double InTime, uint16 InLodIndex, uint16 InFrameCounter, const FTransform& InComponentToworld, const TArrayView<const FTransform>& InPose, const TArrayView<const uint32>& InCurveIds, const TArrayView<const float>& InCurveValues)
+{
+	Session.WriteAccessCheck();
+
+	TSharedPtr<FSkeletalMeshTimelineStorage> TimelineStorage;
+	uint32* IndexPtr = ObjectIdToSkeletalMeshPoseTimelines.Find(InObjectId);
+	if(IndexPtr != nullptr)
+	{
+		TimelineStorage = SkeletalMeshPoseTimelineStorage[*IndexPtr];
+	}
+	else
+	{
+		TimelineStorage = MakeShared<FSkeletalMeshTimelineStorage>();
+		TimelineStorage->Timeline = MakeShared<Trace::TIntervalTimeline<FSkeletalMeshPoseMessage>>(Session.GetLinearAllocator());
+		ObjectIdToSkeletalMeshPoseTimelines.Add(InObjectId, SkeletalMeshPoseTimelineStorage.Num());
+		SkeletalMeshPoseTimelineStorage.Add(TimelineStorage.ToSharedRef());
+	}
+
+	// terminate existing scopes
+	uint64 NumEvents = TimelineStorage->Timeline->GetEventCount();
+	if(NumEvents > 0)
+	{
+		// Add end event at current time
+		TimelineStorage->Timeline->EndEvent(NumEvents - 1, InTime);
+	}
+
+	const int32 NumCurves = InCurveIds.Num();
+
+	FSkeletalMeshPoseMessage Message;
+	Message.ComponentToWorld = InComponentToworld;
+	Message.TransformStartIndex = SkeletalMeshPoseTransforms.Num();
+	Message.CurveStartIndex = SkeletalMeshCurves.Num();
+	Message.ComponentId = InObjectId;
+	Message.MeshId = InMeshId;
+	Message.NumTransforms = (uint16)InPose.Num();
+	Message.NumCurves = (uint16)NumCurves;
+	Message.LodIndex = InLodIndex;
+	Message.FrameCounter = InFrameCounter;
+
+	TimelineStorage->Timeline->AppendBeginEvent(InTime, Message);
+
+	for(const FTransform& Transform : InPose)
+	{
+		SkeletalMeshPoseTransforms.PushBack() = Transform;
+	}
+
+	for(int32 CurveIndex = 0; CurveIndex < NumCurves; ++CurveIndex)
+	{
+		SkeletalMeshCurves.PushBack() = { InCurveIds[CurveIndex], InCurveValues[CurveIndex] };
+		TimelineStorage->AllCurveIds.Add(InCurveIds[CurveIndex]);
+	}
+
+	Session.UpdateDurationSeconds(InTime);
+}
+
 void FAnimationProvider::AppendName(uint32 InId, const TCHAR* InName)
 {
 	Session.WriteAccessCheck();
diff --git a/Engine/Plugins/Animation/GameplayInsights/Source/GameplayInsights/Private/AnimationProvider.h b/Engine/Plugins/Animation/GameplayInsights/Source/GameplayInsights/Private/AnimationProvider.h
--- a/Engine/Plugins/Animation/GameplayInsights/Source/GameplayInsights/Private/AnimationProvider.h
+++ b/Engine/Plugins/Animation/GameplayInsights/Source/GameplayInsights/Private/AnimationProvider.h
@@ -46,7 +46,8 @@ public:
 	void AppendSkeletalMesh(uint64 InObjectId, const TArrayView<const int32>& ParentIndices);
 
 	/** Add a skeletal mesh pose/curves etc. */
-	void AppendSkeletalMeshComponent(uint64 InObjectId, uint64 InMeshId, double InTime, uint16 InLodIndex, uint16 InFrameCounter, const TArrayView<const FTransform>& InTransforms, const TArrayView<const FSkeletalMeshNamedCurve>& InCurves);
+	void AppendSkeletalMeshComponent(uint64 InObjectId, uint64 InMeshId, double InTime, uint16 InLodIndex, uint16 InFrameCounter, const TArrayView<const FTransform>& InPose, const TArrayView<const FSkeletalMeshNamedCurve>& InCurves);
+	void AppendSkeletalMeshComponent(uint64 InObjectId, uint64 InMeshId, double InTime, uint16 InLodIndex, uint16 InFrameCounter, const FTransform& InComponentToworld, const TArrayView<const FTransform>& InPose, const TArrayView<const uint32>& InCurveIds, const TArrayView<const float>& InCurveValues);
 
 	/** Get a skeletal mesh for a specified path. If the mesh no longer exists, make a fake one */
 	USkeletalMesh* GetSkeletalMesh(const TCHAR* InPath);
diff --git a/Engine/Source/Runtime/AnimGraphRuntime/Private/AnimGraphRuntimeModule.cpp b/Engine/Source/Runtime/AnimGraphRuntime/Private/AnimGraphRuntimeModule.cpp
--- a/Engine/Source/Runtime/AnimGraphRuntime/Private/AnimGraphRuntimeModule.cpp
+++ b/Engine/Source/Runtime/AnimGraphRuntime/Private/AnimGraphRuntimeModule.cpp
@@ -5,7 +5,6 @@
 #include "BoneControllers/AnimNode_AnimDynamics.h"
 #include "UObject/UObjectIterator.h"
 #include "Animation/AnimInstance.h"
-#include "AnimGraphRuntimeTrace.h"
 
 //////////////////////////////////////////////////////////////////////////
 // FAnimGraphRuntimeModule
@@ -15,9 +14,6 @@ class FAnimGraphRuntimeModule : public IModuleInterface
 public:
 	virtual void StartupModule() override
 	{
-#if ANIM_TRACE_ENABLED
-		FAnimGraphRuntimeTrace::Init();
-#endif
 	}
 
 	virtual void ShutdownModule() override
diff --git a/Engine/Source/Runtime/AnimGraphRuntime/Private/AnimGraphRuntimeTrace.cpp b/Engine/Source/Runtime/AnimGraphRuntime/Private/AnimGraphRuntimeTrace.cpp
--- a/Engine/Source/Runtime/AnimGraphRuntime/Private/AnimGraphRuntimeTrace.cpp
+++ b/Engine/Source/Runtime/AnimGraphRuntime/Private/AnimGraphRuntimeTrace.cpp
@@ -19,15 +19,6 @@ UE_TRACE_EVENT_BEGIN(Animation, BlendSpacePlayer)
 	UE_TRACE_EVENT_FIELD(float, PositionZ)
 UE_TRACE_EVENT_END()
 
-void FAnimGraphRuntimeTrace::Init()
-{
-	if (FParse::Param(FCommandLine::Get(), TEXT("objecttrace")))
-	{
-		UE_TRACE_EVENT_IS_ENABLED(Animation, BlendSpacePlayer);
-		Trace::ToggleEvent(TEXT("Animation"), true);
-	}
-}
-
 void FAnimGraphRuntimeTrace::OutputBlendSpacePlayer(const FAnimationBaseContext& InContext, const FAnimNode_BlendSpacePlayer& InNode)
 {
 	bool bEventEnabled = UE_TRACE_CHANNELEXPR_IS_ENABLED(AnimationChannel);
diff --git a/Engine/Source/Runtime/AnimGraphRuntime/Public/AnimGraphRuntimeTrace.h b/Engine/Source/Runtime/AnimGraphRuntime/Public/AnimGraphRuntimeTrace.h
--- a/Engine/Source/Runtime/AnimGraphRuntime/Public/AnimGraphRuntimeTrace.h
+++ b/Engine/Source/Runtime/AnimGraphRuntime/Public/AnimGraphRuntimeTrace.h
@@ -12,9 +12,6 @@ struct FAnimNode_BlendSpacePlayer;
 
 struct FAnimGraphRuntimeTrace
 {
-	/** Initialize animation graph runtime tracing */
-	ANIMGRAPHRUNTIME_API static void Init();
-
 	/** Helper function to output debug info for blendspace player nodes */
 	ANIMGRAPHRUNTIME_API static void OutputBlendSpacePlayer(const FAnimationBaseContext& InContext, const FAnimNode_BlendSpacePlayer& InNode);
 };
diff --git a/Engine/Source/Runtime/Engine/Private/Animation/AnimTrace.cpp b/Engine/Source/Runtime/Engine/Private/Animation/AnimTrace.cpp
--- a/Engine/Source/Runtime/Engine/Private/Animation/AnimTrace.cpp
+++ b/Engine/Source/Runtime/Engine/Private/Animation/AnimTrace.cpp
@@ -35,17 +35,19 @@ UE_TRACE_EVENT_BEGIN(Animation, TickRecord)
 	UE_TRACE_EVENT_FIELD(bool, IsBlendSpace)
 UE_TRACE_EVENT_END()
 
-UE_TRACE_EVENT_BEGIN(Animation, SkeletalMesh, Important)
+UE_TRACE_EVENT_BEGIN(Animation, SkeletalMesh2, Important)
 	UE_TRACE_EVENT_FIELD(uint64, Id)
-	UE_TRACE_EVENT_FIELD(uint32, BoneCount)
+	UE_TRACE_EVENT_FIELD(int32[], ParentIndices)
 UE_TRACE_EVENT_END()
 
-UE_TRACE_EVENT_BEGIN(Animation, SkeletalMeshComponent)
+UE_TRACE_EVENT_BEGIN(Animation, SkeletalMeshComponent2)
 	UE_TRACE_EVENT_FIELD(uint64, Cycle)
 	UE_TRACE_EVENT_FIELD(uint64, ComponentId)
 	UE_TRACE_EVENT_FIELD(uint64, MeshId)
-	UE_TRACE_EVENT_FIELD(uint32, BoneCount)
-	UE_TRACE_EVENT_FIELD(uint32, CurveCount)
+	UE_TRACE_EVENT_FIELD(float[], ComponentToWorld)
+	UE_TRACE_EVENT_FIELD(float[], Pose)
+	UE_TRACE_EVENT_FIELD(uint32[], CurveIds)
+	UE_TRACE_EVENT_FIELD(float[], CurveValues)
 	UE_TRACE_EVENT_FIELD(uint16, LodIndex)
 	UE_TRACE_EVENT_FIELD(uint16, FrameCounter)
 UE_TRACE_EVENT_END()
@@ -183,6 +185,18 @@ FUObjectAnnotationSparseBool GSkeletalMeshTraceAnnotations;
 // Map used for unique name output
 TMap<FName, uint32> GAnimTraceNames;
 
+// Scratch buffers for various traces to avoid allocation churn.
+// These can be removed when lambda support is added for array fields to remove a memcpy.
+struct FAnimTraceScratchBuffers : public TThreadSingleton<FAnimTraceScratchBuffers>
+{
+	// Curve values/IDs for skeletal mesh component
+	TArray<float> CurveValues;
+	TArray<uint32> CurveIds;
+
+	// Parent indices for skeletal meshes
+	TArray<int32> ParentIndices;
+};
+
 class FSuspendCounter : public TThreadSingleton<FSuspendCounter>
 {
 public:
@@ -292,31 +306,6 @@ FAnimTrace::FScopedAnimGraphTrace::~FScopedAnimGraphTrace()
 	OutputAnimGraph(Context, StartCycle, FPlatformTime::Cycles64(), (__underlying_type(EPhase))Phase);
 }
 
-void FAnimTrace::Init()
-{
-	if (FParse::Param(FCommandLine::Get(), TEXT("objecttrace")))
-	{
-		UE_TRACE_EVENT_IS_ENABLED(Animation, TickRecord);
-		UE_TRACE_EVENT_IS_ENABLED(Animation, SkeletalMesh);
-		UE_TRACE_EVENT_IS_ENABLED(Animation, SkeletalMeshComponent);
-		UE_TRACE_EVENT_IS_ENABLED(Animation, SkeletalMeshFrame);
-		UE_TRACE_EVENT_IS_ENABLED(Animation, AnimGraph);
-		UE_TRACE_EVENT_IS_ENABLED(Animation, AnimNodeStart);
-		UE_TRACE_EVENT_IS_ENABLED(Animation, AnimNodeEnd);
-		UE_TRACE_EVENT_IS_ENABLED(Animation, AnimNodeValueBool);
-		UE_TRACE_EVENT_IS_ENABLED(Animation, AnimNodeValueInt);
-		UE_TRACE_EVENT_IS_ENABLED(Animation, AnimNodeValueFloat);
-		UE_TRACE_EVENT_IS_ENABLED(Animation, AnimNodeValueVector);
-		UE_TRACE_EVENT_IS_ENABLED(Animation, AnimNodeValueString);
-		UE_TRACE_EVENT_IS_ENABLED(Animation, AnimNodeValueObject);
-		UE_TRACE_EVENT_IS_ENABLED(Animation, AnimNodeValueClass);
-		UE_TRACE_EVENT_IS_ENABLED(Animation, AnimSequencePlayer);
-		UE_TRACE_EVENT_IS_ENABLED(Animation, StateMachineState);
-		UE_TRACE_EVENT_IS_ENABLED(Animation, Name);
-		Trace::ToggleEvent(TEXT("Animation"), true);
-	}
-}
-
 void FAnimTrace::OutputAnimTickRecord(const FAnimationBaseContext& InContext, const FAnimTickRecord& InTickRecord)
 {
 	bool bChannelEnabled = UE_TRACE_CHANNELEXPR_IS_ENABLED(AnimationChannel);
@@ -380,20 +369,19 @@ void FAnimTrace::OutputSkeletalMesh(const USkeletalMesh* InMesh)
 
 	uint32 BoneCount = (uint32)InMesh->RefSkeleton.GetNum();
 
-	auto CopyParentIndices = [InMesh](uint8* Out)
+	TArray<int32>& ParentIndices = FAnimTraceScratchBuffers::Get().ParentIndices;
+	ParentIndices.Reset();
+	ParentIndices.SetNumUninitialized(BoneCount);
+
+	int32 BoneIndex = 0;
+	for(const FMeshBoneInfo& BoneInfo : InMesh->RefSkeleton.GetRefBoneInfo())
 	{
-		int32* OutParentIndices = reinterpret_cast<int32*>(Out);
-		for(const FMeshBoneInfo& BoneInfo : InMesh->RefSkeleton.GetRefBoneInfo())
-		{
-			*OutParentIndices = BoneInfo.ParentIndex;
-			OutParentIndices++;
-		}
-	};
+		ParentIndices[BoneIndex++] = BoneInfo.ParentIndex;
+	}
 
-	UE_TRACE_LOG(Animation, SkeletalMesh, AnimationChannel, BoneCount * sizeof(int32))
-		<< SkeletalMesh.Id(FObjectTrace::GetObjectId(InMesh))
-		<< SkeletalMesh.BoneCount(BoneCount)
-		<< SkeletalMesh.Attachment(CopyParentIndices);
+	UE_TRACE_LOG(Animation, SkeletalMesh2, AnimationChannel)
+		<< SkeletalMesh2.Id(FObjectTrace::GetObjectId(InMesh))
+		<< SkeletalMesh2.ParentIndices(ParentIndices.GetData(), ParentIndices.Num());
 
 	GSkeletalMeshTraceAnnotations.Set(InMesh);
 }
@@ -453,42 +441,37 @@ void FAnimTrace::OutputSkeletalMeshComponent(const USkeletalMeshComponent* InCom
 		TRACE_OBJECT(InComponent);
 		TRACE_SKELETAL_MESH(InComponent->SkeletalMesh);
 
-		auto CopyTransformsAndCurves = [&InComponent, &BoneCount, &CurveCount, &AnimInstance](uint8* Out)
-		{
-			FPlatformMemory::Memcpy(Out, &InComponent->GetComponentToWorld(), sizeof(FTransform));
-			Out += sizeof(FTransform);
+		TArray<float>& CurveValues = FAnimTraceScratchBuffers::Get().CurveValues;
+		CurveValues.Reset();
+		CurveValues.SetNumUninitialized(CurveCount);
+		TArray<uint32>& CurveIds = FAnimTraceScratchBuffers::Get().CurveIds;
+		CurveIds.Reset();
+		CurveIds.SetNumUninitialized(CurveCount);
 
-			if(BoneCount > 0)
-			{
-				const int32 BufferSize = BoneCount * sizeof(FTransform);
-				FPlatformMemory::Memcpy(Out, InComponent->GetComponentSpaceTransforms().GetData(), BufferSize);
-				Out += BufferSize;
-			}
-			if(CurveCount > 0 && AnimInstance)
+		if(CurveCount > 0 && AnimInstance)
+		{
+			int32 CurveIndex = 0;
+			for(EAnimCurveType CurveType : TEnumRange<EAnimCurveType>())
 			{
-				for(EAnimCurveType CurveType : TEnumRange<EAnimCurveType>())
+				for(TPair<FName, float> CurvePair : AnimInstance->GetAnimationCurveList(CurveType))
 				{
-					for(TPair<FName, float> CurvePair : AnimInstance->GetAnimationCurveList(CurveType))
-					{
-						*reinterpret_cast<uint32*>(Out) = OutputName(CurvePair.Key);
-						Out += sizeof(uint32);
-						*reinterpret_cast<float*>(Out) = CurvePair.Value;
-						Out += sizeof(float);
-					}
+					CurveIds[CurveIndex] = OutputName(CurvePair.Key);
+					CurveValues[CurveIndex] = CurvePair.Value;
+					CurveIndex++;
 				}
 			}
-		};
+		}
 
-		//todo: Crashes on cook enginetest
-		/*UE_TRACE_LOG(Animation, SkeletalMeshComponent, true, ((BoneCount + 1) * sizeof(FTransform)) + (CurveCount * (sizeof(float) + sizeof(uint32))))
-			<< SkeletalMeshComponent.Cycle(FPlatformTime::Cycles64())
-			<< SkeletalMeshComponent.ComponentId(FObjectTrace::GetObjectId(InComponent))
-			<< SkeletalMeshComponent.MeshId(FObjectTrace::GetObjectId(InComponent->SkeletalMesh))
-			<< SkeletalMeshComponent.BoneCount((uint32)BoneCount + 1)
-			<< SkeletalMeshComponent.CurveCount((uint32)CurveCount)
-			<< SkeletalMeshComponent.LodIndex((uint16)InComponent->PredictedLODLevel)
-			<< SkeletalMeshComponent.FrameCounter((uint16)(GFrameCounter % 0xffff))
-			<< SkeletalMeshComponent.Attachment(CopyTransformsAndCurves);*/
+		UE_TRACE_LOG(Animation, SkeletalMeshComponent2, AnimationChannel)
+			<< SkeletalMeshComponent2.Cycle(FPlatformTime::Cycles64())
+			<< SkeletalMeshComponent2.ComponentId(FObjectTrace::GetObjectId(InComponent))
+			<< SkeletalMeshComponent2.MeshId(FObjectTrace::GetObjectId(InComponent->SkeletalMesh))
+			<< SkeletalMeshComponent2.ComponentToWorld(reinterpret_cast<const float*>(&InComponent->GetComponentToWorld()), sizeof(FTransform) / sizeof(float))
+			<< SkeletalMeshComponent2.Pose(reinterpret_cast<const float*>(InComponent->GetComponentSpaceTransforms().GetData()), BoneCount * (sizeof(FTransform) / sizeof(float)))
+			<< SkeletalMeshComponent2.CurveIds(CurveIds.GetData(), CurveIds.Num())
+			<< SkeletalMeshComponent2.CurveValues(CurveValues.GetData(), CurveValues.Num())
+			<< SkeletalMeshComponent2.LodIndex((uint16)InComponent->PredictedLODLevel)
+			<< SkeletalMeshComponent2.FrameCounter((uint16)(GFrameCounter % 0xffff));
 	}
 }
 
diff --git a/Engine/Source/Runtime/Engine/Private/UnrealEngine.cpp b/Engine/Source/Runtime/Engine/Private/UnrealEngine.cpp
--- a/Engine/Source/Runtime/Engine/Private/UnrealEngine.cpp
+++ b/Engine/Source/Runtime/Engine/Private/UnrealEngine.cpp
@@ -228,8 +228,6 @@ UnrealEngine.cpp: Implements the UEngine class and helpers.
 
 #include "Particles/ParticleSystemManager.h"
 #include "Components/SkinnedMeshComponent.h"
-#include "ObjectTrace.h"
-#include "Animation/AnimTrace.h"
 
 DEFINE_LOG_CATEGORY(LogEngine);
 IMPLEMENT_MODULE( FEngineModule, Engine );
@@ -266,14 +264,6 @@ void FEngineModule::StartupModule()
 #if WITH_EDITOR
 	USkinnedMeshComponent::BindWorldDelegates();
 #endif
-
-#if OBJECT_TRACE_ENABLED
-	FObjectTrace::Init();
-#endif
-
-#if ANIM_TRACE_ENABLED
-	FAnimTrace::Init();
-#endif
 }
 
 void FEngineModule::ShutdownModule()
diff --git a/Engine/Source/Runtime/Engine/Public/Animation/AnimTrace.h b/Engine/Source/Runtime/Engine/Public/Animation/AnimTrace.h
--- a/Engine/Source/Runtime/Engine/Public/Animation/AnimTrace.h
+++ b/Engine/Source/Runtime/Engine/Public/Animation/AnimTrace.h
@@ -99,9 +99,6 @@ struct FAnimTrace
 		bool bPersistentLines;
 	};
 
-	/** Initialize animation tracing */
-	ENGINE_API static void Init();
-
 	/** Helper function to output a tick record */
 	ENGINE_API static void OutputAnimTickRecord(const FAnimationBaseContext& InContext, const FAnimTickRecord& InTickRecord);
 
