Fix anim node traces in the evaluation pass

Incorrect indices were being used for skeletal controls and most other component-space-related nodes

diff --git a/Engine/Source/Runtime/AnimGraphRuntime/Private/BoneControllers/AnimNode_SkeletalControlBase.cpp b/Engine/Source/Runtime/AnimGraphRuntime/Private/BoneControllers/AnimNode_SkeletalControlBase.cpp
--- a/Engine/Source/Runtime/AnimGraphRuntime/Private/BoneControllers/AnimNode_SkeletalControlBase.cpp
+++ b/Engine/Source/Runtime/AnimGraphRuntime/Private/BoneControllers/AnimNode_SkeletalControlBase.cpp
@@ -144,6 +144,12 @@ void FAnimNode_SkeletalControlBase::EvaluateComponentSpaceInternal(FComponentSpa
 void FAnimNode_SkeletalControlBase::EvaluateComponentSpace_AnyThread(FComponentSpacePoseContext& Output)
 {
 	DECLARE_SCOPE_HIERARCHICAL_COUNTER_ANIMNODE(EvaluateComponentSpace_AnyThread)
+
+#if ANIM_NODE_IDS_AVAILABLE
+	// Cache the incoming node IDs in a base context
+	FAnimationBaseContext CachedContext(Output);
+#endif
+
 	EvaluateComponentPose_AnyThread(Output);
 
 #if WITH_EDITORONLY_DATA
@@ -159,6 +165,10 @@ void FAnimNode_SkeletalControlBase::EvaluateComponentSpace_AnyThread(FComponentS
 	// Apply the skeletal control if it's valid
 	if (FAnimWeight::IsRelevant(ActualAlpha) && IsValidToEvaluate(Output.AnimInstanceProxy->GetSkeleton(), Output.AnimInstanceProxy->GetRequiredBones()))
 	{
+#if ANIM_NODE_IDS_AVAILABLE
+		Output.SetNodeIds(CachedContext);
+#endif
+
 		EvaluateComponentSpaceInternal(Output);
 
 		BoneTransforms.Reset(BoneTransforms.Num());
diff --git a/Engine/Source/Runtime/Engine/Classes/Animation/AnimNodeBase.h b/Engine/Source/Runtime/Engine/Classes/Animation/AnimNodeBase.h
--- a/Engine/Source/Runtime/Engine/Classes/Animation/AnimNodeBase.h
+++ b/Engine/Source/Runtime/Engine/Classes/Animation/AnimNodeBase.h
@@ -383,6 +383,11 @@ public:
 		, bExpectsAdditivePose(SourceContext.bExpectsAdditivePose || bInOverrideExpectsAdditivePose)
 	{
 		Initialize(SourceContext.AnimInstanceProxy);
+
+#if ANIM_NODE_IDS_AVAILABLE
+		CurrentNodeId = SourceContext.CurrentNodeId;
+		PreviousNodeId = SourceContext.PreviousNodeId;
+#endif
 	}
 
 #if ANIM_NODE_IDS_AVAILABLE
@@ -391,6 +396,12 @@ public:
 		PreviousNodeId = CurrentNodeId;
 		CurrentNodeId = InNodeId;
 	}
+
+	void SetNodeIds(const FAnimationBaseContext& InContext)
+	{ 
+		CurrentNodeId = InContext.GetCurrentNodeId();
+		PreviousNodeId = InContext.GetPreviousNodeId();
+	}
 #endif
 
 	ENGINE_API void Initialize(FAnimInstanceProxy* InAnimInstanceProxy);
@@ -468,6 +479,11 @@ public:
 		: FAnimationBaseContext(SourceContext.AnimInstanceProxy)
 	{
 		// No need to initialize, done through FA2CSPose::AllocateLocalPoses
+
+#if ANIM_NODE_IDS_AVAILABLE
+		CurrentNodeId = SourceContext.CurrentNodeId;
+		PreviousNodeId = SourceContext.PreviousNodeId;
+#endif
 	}
 
 #if ANIM_NODE_IDS_AVAILABLE
@@ -476,6 +492,12 @@ public:
 		PreviousNodeId = CurrentNodeId;
 		CurrentNodeId = InNodeId;
 	}
+
+	void SetNodeIds(const FAnimationBaseContext& InContext)
+	{ 
+		CurrentNodeId = InContext.GetCurrentNodeId();
+		PreviousNodeId = InContext.GetPreviousNodeId();
+	}
 #endif
 
 	ENGINE_API void ResetToRefPose();
diff --git a/Engine/Source/Runtime/Engine/Private/Animation/AnimInstanceProxy.cpp b/Engine/Source/Runtime/Engine/Private/Animation/AnimInstanceProxy.cpp
--- a/Engine/Source/Runtime/Engine/Private/Animation/AnimInstanceProxy.cpp
+++ b/Engine/Source/Runtime/Engine/Private/Animation/AnimInstanceProxy.cpp
@@ -1387,6 +1387,16 @@ void FAnimInstanceProxy::EvaluateAnimationNode_WithRoot(FPoseContext& Output, FA
 		{
 			EvaluationCounter.Increment();
 		}
+
+#if ANIM_NODE_IDS_AVAILABLE
+		if(AnimClassInterface && AnimClassInterface->GetAnimBlueprintFunctions().Num() > 0)
+		{
+			Output.SetNodeId(AnimClassInterface->GetAnimBlueprintFunctions()[0].OutputPoseNodeIndex);
+		}
+#endif
+
+		TRACE_SCOPED_ANIM_NODE(Output);
+
 		InRootNode->Evaluate_AnyThread(Output);
 	}
 	else
diff --git a/Engine/Source/Runtime/Engine/Private/Animation/AnimNodeSpaceConversions.cpp b/Engine/Source/Runtime/Engine/Private/Animation/AnimNodeSpaceConversions.cpp
--- a/Engine/Source/Runtime/Engine/Private/Animation/AnimNodeSpaceConversions.cpp
+++ b/Engine/Source/Runtime/Engine/Private/Animation/AnimNodeSpaceConversions.cpp
@@ -28,6 +28,10 @@ void FAnimNode_ConvertComponentToLocalSpace::Evaluate_AnyThread(FPoseContext & O
 {
 	// Evaluate the child and convert
 	FComponentSpacePoseContext InputCSPose(Output.AnimInstanceProxy);
+#if ANIM_NODE_IDS_AVAILABLE
+	// If trace is enabled, we need to preserve the node ID chain as we use the proxy-based constructor above
+	InputCSPose.SetNodeIds(Output);
+#endif
 	ComponentPose.EvaluateComponentSpace(InputCSPose);
 
 	checkSlow( InputCSPose.Pose.GetPose().IsValid() );
@@ -75,6 +79,10 @@ void FAnimNode_ConvertLocalToComponentSpace::EvaluateComponentSpace_AnyThread(FC
 {
 	// Evaluate the child and convert
 	FPoseContext InputPose(OutputCSPose.AnimInstanceProxy);
+#if ANIM_NODE_IDS_AVAILABLE
+	// If trace is enabled, we need to preserve the node ID chain as we use the proxy-based constructor above
+	InputPose.SetNodeIds(OutputCSPose);
+#endif
 	LocalPose.Evaluate(InputPose);
 
 	OutputCSPose.Pose.InitPose(MoveTemp(InputPose.Pose));
