AsyncLoading2: Crash workaround for ReloadMaterialResouce when package filename is none (when running with staged data and IoDispatcher).

Output when loading a game on latest dev-core with IoDispatcher:
[[REDACTED_PHONE]:412][  0]LogMaterial: Warning: UMaterial::CacheResourceShadersForRendering - Can't reload material resource 'Material /Water/Materials/WaterSurface/M_Water_Ocean_02_NewModel.M_Water_Ocean_02_NewModel'. File system based reload is unsupported in this build.
[[REDACTED_PHONE]:413][  0]LogMaterial: Tried to access an uncooked shader map ID in a cooked application
[[REDACTED_PHONE]:413][  0]LogMaterial: Warning: Invalid shader map ID caching shaders for 'M_Water_Ocean_02_NewModel', will use default material.
[[REDACTED_PHONE]:413][  0]LogMaterial: Can't compile M_Water_Ocean_02_NewModel with cooked content, will use default material instead
[[REDACTED_PHONE]:413][  0]LogMaterial: Warning: /Water/Materials/WaterSurface/M_Water_Ocean_02_NewModel.M_Water_Ocean_02_NewModel: Failed to compile Material for platform PCD3D_SM5, Default Material will be used in game.

Corresponding output when loading a game on latest dev-core with cooked data (no pak):
[[REDACTED_PHONE]:791][  0]LogMaterial: Warning: Failed to reload material resources for package /Water/Materials/WaterSurface/M_Water_Ocean_02_NewModel (file name: ../../../Engine/Plugins/NotForLicensees/Water/Content/Materials/WaterSurface/M_Water_Ocean_02_NewModel.uexp).
[[REDACTED_PHONE]:792][  0]LogMaterial: Tried to access an uncooked shader map ID in a cooked application
[[REDACTED_PHONE]:792][  0]LogMaterial: Warning: Invalid shader map ID caching shaders for 'M_Water_Ocean_02_NewModel', will use default material.
[[REDACTED_PHONE]:792][  0]LogMaterial: Can't compile M_Water_Ocean_02_NewModel with cooked content, will use default material instead
[[REDACTED_PHONE]:793][  0]LogMaterial: Warning: /Water/Materials/WaterSurface/M_Water_Ocean_02_NewModel.M_Water_Ocean_02_NewModel: Failed to compile Material for platform PCD3D_SM5, Default Material will be used in game.

diff --git a/Engine/Source/Runtime/Engine/Private/Materials/Material.cpp b/Engine/Source/Runtime/Engine/Private/Materials/Material.cpp
--- a/Engine/Source/Runtime/Engine/Private/Materials/Material.cpp
+++ b/Engine/Source/Runtime/Engine/Private/Materials/Material.cpp
@@ -3117,7 +3117,11 @@ void UMaterial::CacheResourceShadersForRendering(bool bRegenerateId)
 			if (MaterialResource && !MaterialResource->GetGameThreadShaderMap())
 			{
 				FMaterialResource Tmp;
-				if (ReloadMaterialResource(&Tmp, GetOutermost()->FileName.ToString(), OffsetToFirstResource, FeatureLevel, LocalActiveQL))
+				FName PackageFileName = GetOutermost()->FileName;
+				UE_CLOG(PackageFileName.IsNone(), LogMaterial, Warning,
+					TEXT("UMaterial::CacheResourceShadersForRendering - Can't reload material resource '%s'. File system based reload is unsupported in this build."),
+					*GetFullName());
+				if (!PackageFileName.IsNone() && ReloadMaterialResource(&Tmp, PackageFileName.ToString(), OffsetToFirstResource, FeatureLevel, LocalActiveQL))
 				{
 					MaterialResource->SetInlineShaderMap(Tmp.GetGameThreadShaderMap());
 				}
diff --git a/Engine/Source/Runtime/Engine/Private/Materials/MaterialInstance.cpp b/Engine/Source/Runtime/Engine/Private/Materials/MaterialInstance.cpp
--- a/Engine/Source/Runtime/Engine/Private/Materials/MaterialInstance.cpp
+++ b/Engine/Source/Runtime/Engine/Private/Materials/MaterialInstance.cpp
@@ -2967,7 +2967,11 @@ void UMaterialInstance::CacheResourceShadersForRendering()
 			if (!MaterialResource->GetGameThreadShaderMap())
 			{
 				FMaterialResource Tmp;
-				if (ReloadMaterialResource(&Tmp, GetOutermost()->FileName.ToString(), OffsetToFirstResource, FeatureLevel, LocalActiveQL))
+				FName PackageFileName = GetOutermost()->FileName;
+				UE_CLOG(PackageFileName.IsNone(), LogMaterial, Warning,
+					TEXT("UMaterialInstance::CacheResourceShadersForRendering - Can't reload material resource '%s'. File system based reload is unsupported in this build."),
+					*GetFullName());
+				if (!PackageFileName.IsNone() && ReloadMaterialResource(&Tmp, PackageFileName.ToString(), OffsetToFirstResource, FeatureLevel, LocalActiveQL))
 				{
 					MaterialResource->SetInlineShaderMap(Tmp.GetGameThreadShaderMap());
 				}
