Reduce cost of FMaterialUpdateContext and RecacheMaterialInstanceUniformExpressions to save 4 minutes off of a large dataset cook time.  Changes:

1) FMaterialUpdateContext doesn't create a FGlobalComponentRecreateRenderStateContext instance unless FApp::CanEverRender() is true.
2) FGlobalComponentRecreateRenderStateContext stores FComponentRecreateRenderStateContext instances in a TArray instead of a TIndirectArray to avoid excessive numbers of individual allocations and frees.
3) FGlobalComponentRecreateRenderStateContext doesn't create FComponentRecreateRenderStateContext instances for components that aren't registered or don't have a render state created.  Analysis of the code shows that constructing those component contexts is pointlessly hitting the allocator (or bloating the TArray) because they will not do any useful work in their destructor.
4) RecacheMaterialInstanceUniformExpressions doesn't incur the cost of a TObjectIterator if FApp::CanEverRender() is false, as static analysis of the code shows no useful work done within the iteration if FApp::CanEverRender() returns false.

diff --git a/Engine/Source/Runtime/Engine/Private/Components/ActorComponent.cpp b/Engine/Source/Runtime/Engine/Private/Components/ActorComponent.cpp
--- a/Engine/Source/Runtime/Engine/Private/Components/ActorComponent.cpp
+++ b/Engine/Source/Runtime/Engine/Private/Components/ActorComponent.cpp
@@ -177,7 +177,10 @@ FGlobalComponentRecreateRenderStateContext::FGlobalComponentRecreateRenderStateC
 	// recreate render state for all components.
 	for (UActorComponent* Component : TObjectRange<UActorComponent>())
 	{
-		ComponentContexts.Add(new FComponentRecreateRenderStateContext(Component, &ScenesToUpdateAllPrimitiveSceneInfos));
+		if (Component->IsRegistered() && Component->IsRenderStateCreated())
+		{
+			ComponentContexts.Emplace(Component, &ScenesToUpdateAllPrimitiveSceneInfos);
+		}
 	}
 
 	UpdateAllPrimitiveSceneInfos();
diff --git a/Engine/Source/Runtime/Engine/Private/Materials/MaterialInstance.cpp b/Engine/Source/Runtime/Engine/Private/Materials/MaterialInstance.cpp
--- a/Engine/Source/Runtime/Engine/Private/Materials/MaterialInstance.cpp
+++ b/Engine/Source/Runtime/Engine/Private/Materials/MaterialInstance.cpp
@@ -63,7 +63,7 @@ void CacheMaterialInstanceUniformExpressions(const UMaterialInstance* MaterialIn
  */
 void RecacheMaterialInstanceUniformExpressions(const UMaterialInterface* ParentMaterial, bool bRecreateUniformBuffer)
 {
-	if (GIsEditor)
+	if (GIsEditor && FApp::CanEverRender())
 	{
 		UE_LOG(LogMaterial,Verbose,TEXT("Recaching MI Uniform Expressions for parent %s"), *ParentMaterial->GetFullName());
 		TArray<FMICReentranceGuard> ReentranceGuards;
diff --git a/Engine/Source/Runtime/Engine/Private/Materials/MaterialShared.cpp b/Engine/Source/Runtime/Engine/Private/Materials/MaterialShared.cpp
--- a/Engine/Source/Runtime/Engine/Private/Materials/MaterialShared.cpp
+++ b/Engine/Source/Runtime/Engine/Private/Materials/MaterialShared.cpp
@@ -3144,7 +3144,7 @@ void FMaterial::SaveShaderStableKeys(EShaderPlatform TargetShaderPlatform, FStab
 FMaterialUpdateContext::FMaterialUpdateContext(uint32 Options, EShaderPlatform InShaderPlatform)
 {
 	bool bReregisterComponents = (Options & EOptions::ReregisterComponents) != 0;
-	bool bRecreateRenderStates = (Options & EOptions::RecreateRenderStates) != 0;
+	bool bRecreateRenderStates = ((Options & EOptions::RecreateRenderStates) != 0) && FApp::CanEverRender();
 
 	bSyncWithRenderingThread = (Options & EOptions::SyncWithRenderingThread) != 0;
 	if (bReregisterComponents)
diff --git a/Engine/Source/Runtime/Engine/Public/ComponentRecreateRenderStateContext.h b/Engine/Source/Runtime/Engine/Public/ComponentRecreateRenderStateContext.h
--- a/Engine/Source/Runtime/Engine/Public/ComponentRecreateRenderStateContext.h
+++ b/Engine/Source/Runtime/Engine/Public/ComponentRecreateRenderStateContext.h
@@ -61,7 +61,7 @@ public:
 
 private:
 	/** The recreate contexts for the individual components. */
-	TIndirectArray<FComponentRecreateRenderStateContext> ComponentContexts;
+	TArray<FComponentRecreateRenderStateContext> ComponentContexts;
 
 	TSet<FSceneInterface*> ScenesToUpdateAllPrimitiveSceneInfos;
 
