Fix Linux build

diff --git a/Engine/Source/Runtime/Engine/Public/Animation/AnimTrace.h b/Engine/Source/Runtime/Engine/Public/Animation/AnimTrace.h
--- a/Engine/Source/Runtime/Engine/Public/Animation/AnimTrace.h
+++ b/Engine/Source/Runtime/Engine/Public/Animation/AnimTrace.h
@@ -198,7 +198,7 @@ struct FAnimTrace
 	FAnimTrace::OutputStateMachineState(Context, StateMachineIndex, StateIndex, StateWeight, ElapsedTime);
 
 #define TRACE_ANIM_NOTIFY(AnimInstance, NotifyEvent, EventType) \
-	FAnimTrace::OutputAnimNotify(AnimInstance, NotifyEvent, FAnimTrace::ENotifyEventType::##EventType);
+	FAnimTrace::OutputAnimNotify(AnimInstance, NotifyEvent, FAnimTrace::ENotifyEventType::EventType);
 
 #define TRACE_ANIM_SYNC_MARKER(AnimInstance, SyncMarker) \
 	FAnimTrace::OutputAnimSyncMarker(AnimInstance, SyncMarker);
