Check-in missing files:
Engine level UObject Trace filtering

diff --git a/Engine/Source/Runtime/Engine/Private/Animation/AnimTrace.cpp b/Engine/Source/Runtime/Engine/Private/Animation/AnimTrace.cpp
--- a/Engine/Source/Runtime/Engine/Private/Animation/AnimTrace.cpp
+++ b/Engine/Source/Runtime/Engine/Private/Animation/AnimTrace.cpp
@@ -19,6 +19,7 @@
 #include "Animation/AnimNotifies/AnimNotify.h"
 #include "Animation/AnimNotifies/AnimNotifyState.h"
 #include "Animation/AnimTypes.h"
+#include "TraceFilter.h"
 
 UE_TRACE_CHANNEL_DEFINE(AnimationChannel);
 
@@ -360,6 +361,11 @@ void FAnimTrace::OutputAnimTickRecord(const FAnimationBaseContext& InContext, co
 		return;
 	}
 
+	if (CANNOT_TRACE_OBJECT(InContext.AnimInstanceProxy->GetSkelMeshComponent()))
+	{
+		return;
+	}
+
 	TRACE_OBJECT(InContext.AnimInstanceProxy->GetAnimInstanceObject());
 
 	if(InTickRecord.SourceAsset)
@@ -478,6 +484,11 @@ void FAnimTrace::OutputSkeletalMeshComponent(const USkeletalMeshComponent* InCom
 		return;
 	}
 
+	if (CANNOT_TRACE_OBJECT(InComponent))
+	{
+		return;
+	}
+
 	int32 BoneCount = InComponent->GetComponentSpaceTransforms().Num();
 	int32 CurveCount = 0;
 	UAnimInstance* AnimInstance = InComponent->GetAnimInstance();
@@ -546,6 +557,11 @@ void FAnimTrace::OutputAnimGraph(const FAnimationBaseContext& InContext, uint64
 		return;
 	}
 
+	if (CANNOT_TRACE_OBJECT(InContext.AnimInstanceProxy->GetSkelMeshComponent()))
+	{
+		return;
+	}
+
 	check(InContext.AnimInstanceProxy);
 
 	const UAnimInstance* AnimInstance = Cast<UAnimInstance>(InContext.AnimInstanceProxy->GetAnimInstanceObject());
@@ -575,6 +591,11 @@ void FAnimTrace::OutputAnimNodeStart(const FAnimationBaseContext& InContext, uin
 		return;
 	}
 
+	if (CANNOT_TRACE_OBJECT(InContext.AnimInstanceProxy->GetSkelMeshComponent()))
+	{
+		return;
+	}
+
 	check(InContext.AnimInstanceProxy);
 
 	TRACE_OBJECT(InContext.AnimInstanceProxy->GetAnimInstanceObject());
@@ -624,6 +645,11 @@ void FAnimTrace::OutputAnimNodeEnd(const FAnimationBaseContext& InContext, uint6
 
 	check(InContext.AnimInstanceProxy);
 
+	if (CANNOT_TRACE_OBJECT(InContext.AnimInstanceProxy->GetSkelMeshComponent()))
+	{
+		return;
+	}
+
 	TRACE_OBJECT(InContext.AnimInstanceProxy->GetAnimInstanceObject());
 
 	UE_TRACE_LOG(Animation, AnimNodeEnd, AnimationChannel)
@@ -641,6 +667,11 @@ void FAnimTrace::OutputAnimNodeValue(const FAnimationBaseContext& InContext, con
 
 	check(InContext.AnimInstanceProxy);
 
+	if (CANNOT_TRACE_OBJECT(InContext.AnimInstanceProxy->GetSkelMeshComponent()))
+	{
+		return;
+	}
+
 	TRACE_OBJECT(InContext.AnimInstanceProxy->GetAnimInstanceObject());
 
 	int32 KeyLength = FCString::Strlen(InKey) + 1;
@@ -665,6 +696,11 @@ void FAnimTrace::OutputAnimNodeValue(const FAnimationBaseContext& InContext, con
 
 	check(InContext.AnimInstanceProxy);
 
+	if (CANNOT_TRACE_OBJECT(InContext.AnimInstanceProxy->GetSkelMeshComponent()))
+	{
+		return;
+	}
+
 	TRACE_OBJECT(InContext.AnimInstanceProxy->GetAnimInstanceObject());
 
 	int32 KeyLength = FCString::Strlen(InKey) + 1;
@@ -689,6 +725,11 @@ void FAnimTrace::OutputAnimNodeValue(const FAnimationBaseContext& InContext, con
 
 	check(InContext.AnimInstanceProxy);
 
+	if (CANNOT_TRACE_OBJECT(InContext.AnimInstanceProxy->GetSkelMeshComponent()))
+	{
+		return;
+	}
+
 	TRACE_OBJECT(InContext.AnimInstanceProxy->GetAnimInstanceObject());
 
 	int32 KeyLength = FCString::Strlen(InKey) + 1;
@@ -745,6 +786,11 @@ void FAnimTrace::OutputAnimNodeValue(const FAnimationBaseContext& InContext, con
 
 	check(InContext.AnimInstanceProxy);
 
+	if (CANNOT_TRACE_OBJECT(InContext.AnimInstanceProxy->GetSkelMeshComponent()))
+	{
+		return;
+	}
+
 	TRACE_OBJECT(InContext.AnimInstanceProxy->GetAnimInstanceObject());
 
 	int32 KeyLength = FCString::Strlen(InKey) + 1;
@@ -771,6 +817,11 @@ void FAnimTrace::OutputAnimNodeValue(const FAnimationBaseContext& InContext, con
 
 	check(InContext.AnimInstanceProxy);
 
+	if (CANNOT_TRACE_OBJECT(InContext.AnimInstanceProxy->GetSkelMeshComponent()))
+	{
+		return;
+	}
+
 	TRACE_OBJECT(InContext.AnimInstanceProxy->GetAnimInstanceObject());
 
 	int32 KeyLength = FCString::Strlen(InKey) + 1;
@@ -801,6 +852,11 @@ void FAnimTrace::OutputAnimNodeValue(const FAnimationBaseContext& InContext, con
 
 	check(InContext.AnimInstanceProxy);
 
+	if (CANNOT_TRACE_OBJECT(InContext.AnimInstanceProxy->GetSkelMeshComponent()))
+	{
+		return;
+	}
+
 	TRACE_OBJECT(InContext.AnimInstanceProxy->GetAnimInstanceObject());
 
 	int32 KeyLength = FCString::Strlen(InKey) + 1;
@@ -831,6 +887,11 @@ void FAnimTrace::OutputAnimNodeValue(const FAnimationBaseContext& InContext, con
 
 	check(InContext.AnimInstanceProxy);
 
+	if (CANNOT_TRACE_OBJECT(InContext.AnimInstanceProxy->GetSkelMeshComponent()))
+	{
+		return;
+	}
+
 	TRACE_OBJECT(InContext.AnimInstanceProxy->GetAnimInstanceObject());
 	TRACE_OBJECT(InValue);
 
@@ -856,6 +917,11 @@ void FAnimTrace::OutputAnimNodeValue(const FAnimationBaseContext& InContext, con
 
 	check(InContext.AnimInstanceProxy);
 
+	if (CANNOT_TRACE_OBJECT(InContext.AnimInstanceProxy->GetSkelMeshComponent()))
+	{
+		return;
+	}
+
 	TRACE_OBJECT(InContext.AnimInstanceProxy->GetAnimInstanceObject());
 	TRACE_CLASS(InValue);
 
@@ -881,6 +947,11 @@ void FAnimTrace::OutputAnimSequencePlayer(const FAnimationBaseContext& InContext
 
 	check(InContext.AnimInstanceProxy);
 
+	if (CANNOT_TRACE_OBJECT(InContext.AnimInstanceProxy->GetSkelMeshComponent()))
+	{
+		return;
+	}
+
 	TRACE_OBJECT(InContext.AnimInstanceProxy->GetAnimInstanceObject());
 
 	UE_TRACE_LOG(Animation, AnimSequencePlayer, AnimationChannel)
@@ -902,6 +973,11 @@ void FAnimTrace::OutputStateMachineState(const FAnimationBaseContext& InContext,
 
 	check(InContext.AnimInstanceProxy);
 
+	if (CANNOT_TRACE_OBJECT(InContext.AnimInstanceProxy->GetSkelMeshComponent()))
+	{
+		return;
+	}
+
 	TRACE_OBJECT(InContext.AnimInstanceProxy->GetAnimInstanceObject());
 
 	UE_TRACE_LOG(Animation, StateMachineState, AnimationChannel)
@@ -921,6 +997,11 @@ void FAnimTrace::OutputAnimNotify(UAnimInstance* InAnimInstance, const FAnimNoti
 	{
 		return;
 	}
+	
+	if (CANNOT_TRACE_OBJECT(InAnimInstance->GetSkelMeshComponent()))
+	{
+		return;
+	}
 
 	TRACE_OBJECT(InAnimInstance);
 
@@ -961,6 +1042,11 @@ void FAnimTrace::OutputAnimSyncMarker(UAnimInstance* InAnimInstance, const FPass
 		return;
 	}
 
+	if (CANNOT_TRACE_OBJECT(InAnimInstance->GetSkelMeshComponent()))
+	{
+		return;
+	}
+
 	TRACE_OBJECT(InAnimInstance);
 
 	const uint32 NameId = OutputName(InPassedSyncMarker.PassedMarkerName);
@@ -981,6 +1067,11 @@ void FAnimTrace::OutputMontage(UAnimInstance* InAnimInstance, const FAnimMontage
 
 	if(InMontageInstance.Montage != nullptr)
 	{
+		if (CANNOT_TRACE_OBJECT(InAnimInstance->GetSkelMeshComponent()))
+		{
+			return;
+		}
+
 		TRACE_OBJECT(InAnimInstance);
 		TRACE_OBJECT(InMontageInstance.Montage);
 
diff --git a/Engine/Source/Runtime/Engine/Private/ObjectTrace.cpp b/Engine/Source/Runtime/Engine/Private/ObjectTrace.cpp
--- a/Engine/Source/Runtime/Engine/Private/ObjectTrace.cpp
+++ b/Engine/Source/Runtime/Engine/Private/ObjectTrace.cpp
@@ -14,6 +14,7 @@
 #include "Misc/CommandLine.h"
 #include "HAL/IConsoleManager.h"
 #include "Engine/World.h"
+#include "TraceFilter.h"
 
 UE_TRACE_CHANNEL(ObjectChannel)
 
@@ -36,12 +37,6 @@ UE_TRACE_EVENT_BEGIN(Object, ObjectEvent)
 	UE_TRACE_EVENT_FIELD(uint8, Event)
 UE_TRACE_EVENT_END()
 
-TAutoConsoleVariable<int32> CVarRecordAllWorldTypes(
-	TEXT("Insights.RecordAllWorldTypes"),
-	0,
-	TEXT("Gameplay Insights recording by default only records Game and PIE worlds.")
-	TEXT("Toggle this value to 1 to record other world types."));
-
 // Object annotations used for tracing
 struct FTracedObjectAnnotation
 {
@@ -66,16 +61,6 @@ struct FTracedObjectAnnotation
 // Object annotations used for tracing
 FUObjectAnnotationSparse<FTracedObjectAnnotation, true> GObjectTraceAnnotations;
 
-static bool ShouldTraceObjectsWorld(const  UObject* InObject)
-{
-	UWorld* World = InObject->GetWorld();
-	return 
-		CVarRecordAllWorldTypes.GetValueOnAnyThread() != 0 ||
-		World == nullptr || 
-		World->WorldType == EWorldType::Game ||
-		World->WorldType == EWorldType::PIE;
-}
-
 void FObjectTrace::Init()
 {
 	if (FParse::Param(FCommandLine::Get(), TEXT("objecttrace")))
@@ -166,6 +151,11 @@ void FObjectTrace::OutputObject(const UObject* InObject)
 		return;
 	}
 
+	if (CANNOT_TRACE_OBJECT(InObject->GetWorld()))
+	{
+		return;
+	}
+
 	FTracedObjectAnnotation Annotation = GObjectTraceAnnotations.GetAnnotation(InObject);
 	if(Annotation.bTraced)
 	{
@@ -176,11 +166,6 @@ void FObjectTrace::OutputObject(const UObject* InObject)
 	Annotation.bTraced = true;
 	GObjectTraceAnnotations.AddAnnotation(InObject, Annotation);
 
-	if(!ShouldTraceObjectsWorld(InObject))
-	{
-		return;
-	}
-
 	// Trace the object's class first
 	TRACE_CLASS(InObject->GetClass());
 
@@ -214,7 +199,7 @@ void FObjectTrace::OutputObjectEvent(const UObject* InObject, const TCHAR* InEve
 		return;
 	}
 
-	if(!ShouldTraceObjectsWorld(InObject))
+	if (CANNOT_TRACE_OBJECT(InObject->GetWorld()))
 	{
 		return;
 	}
diff --git a/Engine/Source/Runtime/Engine/Private/TraceFilter.cpp b/Engine/Source/Runtime/Engine/Private/TraceFilter.cpp
new file mode 100644
--- /dev/null
+++ b/Engine/Source/Runtime/Engine/Private/TraceFilter.cpp
@@ -0,0 +1,179 @@
+// Copyright Epic Games, Inc. All Rights Reserved.
+
+#include "TraceFilter.h"
+
+#if TRACE_FILTERING_ENABLED
+
+#include "TraceFilters.h"
+#include "HAL/IConsoleManager.h"
+#include "Algo/Transform.h"
+#include "Containers/Map.h"
+
+#include "UObject/UObjectBase.h"
+#include "Engine/World.h"
+#include "GameFramework/Actor.h"
+#include "Components/ActorComponent.h"
+#include "Logging/LogMacros.h"
+#include "ObjectTrace.h"
+
+struct FTraceFilterObjectAnnotation
+{
+	FTraceFilterObjectAnnotation()
+		: bIsTraceable(false)
+	{}
+
+	bool bIsTraceable;
+
+	/** Determine if this annotation is default - required for annotations */
+	FORCEINLINE bool IsDefault() const
+	{
+		return !bIsTraceable;
+	}
+};
+
+FUObjectAnnotationSparse<FTraceFilterObjectAnnotation, true> GObjectFilterAnnotations;
+
+DEFINE_LOG_CATEGORY_STATIC(TraceFiltering, Display, Display);
+
+/** Console command allowing to debug the current state of GObjectFilterAnnotations, to see which objects are Traceable*/
+FAutoConsoleCommand FlushFilterStateCommand(TEXT("TraceFilter.FlushState"), TEXT("Flushes the current trace filtering state to the output log."),
+		FConsoleCommandDelegate::CreateLambda([]()
+		{
+			TMap<const UWorld*, TArray<const AActor*>> WorldToActorMap;
+			TMap<const AActor*, TArray<const UActorComponent*>> ActorToComponentMap;
+			TSet<const UObject*> Objects;
+
+			/** Collect object type information */
+			const TMap<const UObjectBase*, FTraceFilterObjectAnnotation>& Annotations = GObjectFilterAnnotations.GetAnnotationMap();
+
+			/** Retrieve all annotated objects */
+			TArray<const UObjectBase*> BaseObjects;
+			Annotations.GenerateKeyArray(BaseObjects);
+
+			/** Filter each object according to its type */
+			for (const UObjectBase* ObjectBase : BaseObjects)
+			{
+				const UObject* Object = (const UObject*)ObjectBase;
+				if (const UWorld* World = Cast<UWorld>(Object))
+				{
+					WorldToActorMap.Add(World);
+				}
+			}
+
+			for (const UObjectBase* ObjectBase : BaseObjects)
+			{
+				const UObject* Object = (const UObject*)ObjectBase;
+				if (const AActor* Actor = Cast<AActor>(Object))
+				{
+					WorldToActorMap.FindChecked(Actor->GetWorld()).Add(Actor);
+					ActorToComponentMap.Add(Actor);
+				}
+			}
+
+			for (const UObjectBase* ObjectBase : BaseObjects)
+			{
+				const UObject* Object = (const UObject*)ObjectBase;
+				if (const UActorComponent* Component = Cast<UActorComponent>(Object))
+				{
+					ActorToComponentMap.FindChecked(Component->GetOwner()).Add(Component);
+				}
+			}
+
+			Algo::TransformIf(BaseObjects, Objects, 
+				[](const UObjectBase* ObjectBase)
+				{
+					const UObject* Object = (const UObject*)ObjectBase;
+					return !Cast<UWorld>(Object) && !Cast<AActor>(Object) && !Cast<UActorComponent>(Object);
+				},
+				[](const UObjectBase* ObjectBase)
+				{
+					const UObject* Object = (const UObject*)ObjectBase;
+					return Object;
+				}
+			);
+
+			/** Output collated data */
+			FString OutputString;
+
+			/** For each UWorld, output all AActors */
+			for (const TPair<const UWorld*, TArray<const AActor*>>& WorldToActor : WorldToActorMap)
+			{	
+				OutputString += TEXT("\n");
+				OutputString += WorldToActor.Key->GetName();
+				OutputString += TEXT(" [UWorld]\n");
+
+				/** For each AActor, output all UActorComponent */
+				for (const AActor* Actor : WorldToActor.Value)
+				{
+					OutputString += TEXT("\t- ");
+					OutputString += Actor->GetName();
+					OutputString += TEXT(" [Actor]\n");
+
+					TArray<const UActorComponent*> ActorComponents = ActorToComponentMap.FindChecked(Actor);
+					for (const UActorComponent* Component : ActorComponents)
+					{
+						OutputString += TEXT("\t\t* ");
+						OutputString += Component->GetName();
+						OutputString += TEXT(" [Component]\n");
+					}			
+				}
+
+				OutputString += TEXT("----------------------------------------------------\n");
+			}			
+
+			for (const UObject* Object : Objects)
+			{
+				OutputString += Object->GetName();
+				OutputString += TEXT(" [Object]\n");
+			}
+
+			UE_LOG(TraceFiltering, Display, TEXT("%s"), *OutputString);
+		})
+	);
+
+bool FTraceFilter::IsObjectTraceable(const UObject* InObject)
+{
+	return InObject ? GObjectFilterAnnotations.GetAnnotation(InObject).bIsTraceable : true;
+}
+
+void FTraceFilter::SetObjectIsTraceable(const UObject* InObject, bool bIsTraceable)
+{
+	ensure(InObject);
+
+	if (bIsTraceable)
+	{
+		FTraceFilterObjectAnnotation Annotation;
+		Annotation.bIsTraceable = true;
+		GObjectFilterAnnotations.AddAnnotation(InObject, Annotation);
+
+		TRACE_OBJECT(InObject);
+	}
+	else
+	{
+		GObjectFilterAnnotations.RemoveAnnotation(InObject);
+	}
+}
+
+void FTraceFilter::MarkObjectTraceable(const UObject* InObject)
+{
+	ensure(InObject);
+
+	FTraceFilterObjectAnnotation Annotation;
+	Annotation.bIsTraceable = true;
+	GObjectFilterAnnotations.AddAnnotation(InObject, Annotation);
+}
+
+void FTraceFilter::Init()
+{
+	FTraceActorFilter::Initialize();
+	FTraceWorldFilter::Initialize();
+}
+
+void FTraceFilter::Destroy()
+{
+	GObjectFilterAnnotations.RemoveAllAnnotations();
+	FTraceActorFilter::Destroy();
+	FTraceWorldFilter::Destroy();
+}
+
+#endif // TRACE_FILTERING_ENABLED
diff --git a/Engine/Source/Runtime/Engine/Private/TraceFilters.cpp b/Engine/Source/Runtime/Engine/Private/TraceFilters.cpp
new file mode 100644
--- /dev/null
+++ b/Engine/Source/Runtime/Engine/Private/TraceFilters.cpp
@@ -0,0 +1,83 @@
+// Copyright Epic Games, Inc. All Rights Reserved.
+
+#include "TraceFilters.h"
+
+#if TRACE_FILTERING_ENABLED
+
+#include "EngineUtils.h"
+
+FDelegateHandle FTraceActorFilter::WorldPreInitHandle;
+FDelegateHandle FTraceActorFilter::WorldPostInitHandle;
+FDelegateHandle FTraceActorFilter::OnActorSpawnedHandle;
+
+void FTraceActorFilter::Initialize()
+{
+	WorldPreInitHandle = FWorldDelegates::OnPreWorldInitialization.AddStatic(&FTraceActorFilter::OnPreWorldInit);
+	WorldPostInitHandle = FWorldDelegates::OnPostWorldInitialization.AddStatic(&FTraceActorFilter::OnPostWorldInit);
+}
+
+void FTraceActorFilter::Destroy()
+{
+	FWorldDelegates::OnPreWorldInitialization.Remove(WorldPreInitHandle);
+	FWorldDelegates::OnPostWorldInitialization.Remove(WorldPostInitHandle);
+}
+
+void FTraceActorFilter::OnPreWorldInit(UWorld* World, const UWorld::InitializationValues IVS)
+{
+	World->AddOnActorSpawnedHandler(FOnActorSpawned::FDelegate::CreateStatic(&FTraceActorFilter::OnActorSpawned));
+}
+
+void FTraceActorFilter::OnPostWorldInit(UWorld* World, const UWorld::InitializationValues IVS)
+{
+	if (CAN_TRACE_OBJECT(World))
+	{
+		// Handle actors that were duplicated over from a previous UWorld instance (PIE)
+		for (TActorIterator<AActor> It(World); It; ++It)
+		{
+			if (AActor* Actor = *It)
+			{
+				OnActorSpawned(Actor);
+			}
+		}
+	}
+}
+
+void FTraceActorFilter::OnActorSpawned(AActor* Actor)
+{
+	/** Mark traceable if UWorld is marked as such */
+	if (CAN_TRACE_OBJECT(Actor->GetWorld()))
+	{
+		FTraceFilter::MarkObjectTraceable(Actor);
+	}
+}
+
+TAutoConsoleVariable<int32> CVarRecordAllWorldTypes(
+	TEXT("Insights.RecordAllWorldTypes"),
+	0,
+	TEXT("Gameplay Insights recording by default only records Game and PIE worlds.")
+	TEXT("Toggle this value to 1 to record other world types."));
+
+FDelegateHandle FTraceWorldFilter::WorldInitHandle;
+
+void FTraceWorldFilter::Initialize()
+{
+	WorldInitHandle = FWorldDelegates::OnPreWorldInitialization.AddStatic(&FTraceWorldFilter::OnWorldInit);	
+}
+
+void FTraceWorldFilter::Destroy()
+{
+	FWorldDelegates::OnPreWorldInitialization.Remove(WorldInitHandle);
+}
+
+void FTraceWorldFilter::OnWorldInit(UWorld* World, const UWorld::InitializationValues IVS)
+{
+	if (CVarRecordAllWorldTypes.GetValueOnAnyThread() != 0 ||
+			World == nullptr ||
+			World->WorldType == EWorldType::Game ||
+			World->WorldType == EWorldType::PIE)
+	{
+		FTraceFilter::MarkObjectTraceable(World);
+	}
+}
+
+#endif // TRACE_FILTERING_ENABLED
diff --git a/Engine/Source/Runtime/Engine/Private/UnrealEngine.cpp b/Engine/Source/Runtime/Engine/Private/UnrealEngine.cpp
--- a/Engine/Source/Runtime/Engine/Private/UnrealEngine.cpp
+++ b/Engine/Source/Runtime/Engine/Private/UnrealEngine.cpp
@@ -229,6 +229,7 @@ UnrealEngine.cpp: Implements the UEngine class and helpers.
 #include "Particles/ParticleSystemManager.h"
 #include "Components/SkinnedMeshComponent.h"
 #include "StudioAnalytics.h"
+#include "TraceFilter.h"
 
 DEFINE_LOG_CATEGORY(LogEngine);
 IMPLEMENT_MODULE( FEngineModule, Engine );
@@ -265,10 +266,18 @@ void FEngineModule::StartupModule()
 #if WITH_EDITOR
 	USkinnedMeshComponent::BindWorldDelegates();
 #endif
+
+#if TRACE_FILTERING_ENABLED
+	FTraceFilter::Init();
+#endif
 }
 
 void FEngineModule::ShutdownModule()
 {
+#if TRACE_FILTERING_ENABLED
+	FTraceFilter::Destroy();
+#endif
+
 	FParticleSystemWorldManager::OnShutdown();
 }
 
diff --git a/Engine/Source/Runtime/Engine/Public/TraceFilter.h b/Engine/Source/Runtime/Engine/Public/TraceFilter.h
new file mode 100644
--- /dev/null
+++ b/Engine/Source/Runtime/Engine/Public/TraceFilter.h
@@ -0,0 +1,98 @@
+// Copyright Epic Games, Inc. All Rights Reserved.
+
+#pragma once
+
+#include "CoreTypes.h"
+#include "Templates/EnableIf.h"
+#include "Templates/PointerIsConvertibleFromTo.h"
+#include "GameFramework/Actor.h"
+#include "Engine/World.h"
+#include "Components/ActorComponent.h"
+#include "Templates/UnrealTemplate.h"
+#include "UObject/Object.h"
+
+#if !IS_PROGRAM && !UE_BUILD_SHIPPING
+#define TRACE_FILTERING_ENABLED 1
+#else
+#define TRACE_FILTERING_ENABLED 0
+#endif
+
+#if TRACE_FILTERING_ENABLED
+
+struct ENGINE_API FTraceFilter
+{
+	struct TObjectFilter
+	{
+	public:
+		template<typename T>
+		static typename TEnableIf<TPointerIsConvertibleFromTo<typename TRemovePointer<T>::Type, UWorld>::Value, bool>::Type FORCEINLINE CanTrace(const T* Object)
+		{
+			const UWorld* World = (const UWorld*)Object;
+			return FTraceFilter::IsObjectTraceable(World);
+		}
+
+		template<typename T>
+		static typename TEnableIf<TPointerIsConvertibleFromTo<typename TRemovePointer<T>::Type, AActor>::Value, bool>::Type FORCEINLINE CanTrace(const T* Object)
+		{
+			/** For an individual Actor, we expect it and the owning world to both be marked traceable */
+			const AActor* Actor = (const AActor*)Object;
+			return (FTraceFilter::IsObjectTraceable(Cast<UObject>(Actor)) && (Actor && CanTrace(Actor->GetWorld())));
+		}
+
+		template<typename T>
+		static typename TEnableIf<TPointerIsConvertibleFromTo<typename TRemovePointer<T>::Type, UActorComponent>::Value, bool>::Type FORCEINLINE CanTrace(const T* Object)
+		{
+			/** For an individual Actor Component, we expect it or the owning actor to be marked traceable */
+			const UActorComponent* ActorComponent = (const UActorComponent*)Object;
+			return FTraceFilter::IsObjectTraceable(ActorComponent) || (ActorComponent && TObjectFilter::CanTrace(ActorComponent->GetOwner()));
+		}
+
+		template<typename T>
+		static typename TEnableIf<(!TPointerIsConvertibleFromTo<typename TRemovePointer<T>::Type, UActorComponent>::Value && !TPointerIsConvertibleFromTo<typename TRemovePointer<T>::Type, AActor>::Value 
+			&& !TPointerIsConvertibleFromTo<typename TRemovePointer<T>::Type, UWorld>::Value && TPointerIsConvertibleFromTo<typename TRemovePointer<T>::Type, UObject>::Value), bool>::Type FORCEINLINE CanTrace(const T* Object)
+		{
+			/** For an individual UObject, we expect it or the owning world to be marked traceable */
+			const UObject* BaseObject = (const UObject*)Object;
+			return FTraceFilter::IsObjectTraceable(Object) || (BaseObject && FTraceFilter::IsObjectTraceable(BaseObject->GetWorld()));
+		}
+	};
+
+	/** Set whether or not an Object is Traceable, and eligible for outputting Trace Data */
+	static void SetObjectIsTraceable(const UObject* InObject, bool bIsTraceable);
+
+	/** Mark an Object as Traceable, and eligible for outputting Trace Data */
+	static void MarkObjectTraceable(const UObject* InObject);
+
+	/** Initializes any of the Filters which are part of the Engine */
+	static void Init();
+
+	/** Destroys any of the Filters which are part of the Engine */
+	static void Destroy();
+protected:
+	/** Returns whether or not an object is eligible to be outputted (output trace data) */
+	static bool IsObjectTraceable(const UObject* InObject);
+};
+
+#define CAN_TRACE_OBJECT(Object) \
+	FTraceFilter::TObjectFilter::CanTrace(Object)
+
+#define CANNOT_TRACE_OBJECT(Object) \
+	(!FTraceFilter::TObjectFilter::CanTrace(Object))
+
+#define MARK_OBJECT_TRACEABLE(Object) \
+	FTraceFilter::MarkObjectTraceable(Object)
+
+#define SET_OBJECT_TRACEABLE(Object, bIsTraceable) \
+	FTraceFilter::SetObjectIsTraceable(Object, bIsTraceable)
+#else
+
+#define CAN_TRACE_OBJECT(Object) \
+	true
+#define CANNOT_TRACE_OBJECT(Object) \
+	false
+
+#define MARK_OBJECT_TRACEABLE(Object)
+
+#define SET_OBJECT_TRACEABLE(Object, bIsTraceable) 
+
+#endif
diff --git a/Engine/Source/Runtime/Engine/Public/TraceFilters.h b/Engine/Source/Runtime/Engine/Public/TraceFilters.h
new file mode 100644
--- /dev/null
+++ b/Engine/Source/Runtime/Engine/Public/TraceFilters.h
@@ -0,0 +1,40 @@
+// Copyright Epic Games, Inc. All Rights Reserved.
+
+#pragma once
+
+#include "TraceFilter.h"
+
+#if TRACE_FILTERING_ENABLED
+
+/** AActor specific Trace filter, marks individual instances to not be traced out */
+struct FTraceActorFilter
+{
+	static void Initialize();
+	static void Destroy();
+protected:
+	static FDelegateHandle WorldPreInitHandle;
+	static FDelegateHandle WorldPostInitHandle;
+	static FDelegateHandle OnActorSpawnedHandle;
+	static void OnPreWorldInit(UWorld* World, const UWorld::InitializationValues IVS);
+	static void OnPostWorldInit(UWorld* World, const UWorld::InitializationValues IVS);
+	static void OnActorSpawned(AActor* Actor);
+};
+
+/** UWorld specific Trace filter, marks individual instances to not be traced out */
+struct ENGINE_API FTraceWorldFilter
+{
+	static void Initialize();
+	static void Destroy();
+protected:
+	static void OnWorldInit(UWorld* World, const UWorld::InitializationValues IVS);	
+	static FDelegateHandle WorldInitHandle;	
+};
+
+#define DISABLE_ENGINE_WORLD_TRACE_FILTERING() \
+	FTraceWorldFilter::Destroy()
+
+#else
+
+#define DISABLE_ENGINE_WORLD_TRACE_FILTERING() 
+
+#endif // TRACE_FILTERING_ENABLED
