Fixes for audio mixer and audio device manager refactor

diff --git a/Engine/Source/Runtime/Engine/Private/AudioDecompress.cpp b/Engine/Source/Runtime/Engine/Private/AudioDecompress.cpp
--- a/Engine/Source/Runtime/Engine/Private/AudioDecompress.cpp
+++ b/Engine/Source/Runtime/Engine/Private/AudioDecompress.cpp
@@ -639,13 +639,17 @@ namespace ADPCM
 /**
  * Worker for decompression on a separate thread
  */
-FAsyncAudioDecompressWorker::FAsyncAudioDecompressWorker(USoundWave* InWave, int32 InPrecacheBufferNumFrames)
+FAsyncAudioDecompressWorker::FAsyncAudioDecompressWorker(USoundWave* InWave, int32 InPrecacheBufferNumFrames, FAudioDevice* InAudioDevice)
 	: Wave(InWave)
 	, AudioInfo(nullptr)
 	, NumPrecacheFrames(InPrecacheBufferNumFrames)
 {
 	check(NumPrecacheFrames > 0);
-	if (GEngine && GEngine->GetMainAudioDevice())
+	if (InAudioDevice)
+	{
+		AudioInfo = InAudioDevice->CreateCompressedAudioInfo(Wave);
+	}
+	else if (GEngine && GEngine->GetMainAudioDevice())
 	{
 		AudioInfo = GEngine->GetMainAudioDevice()->CreateCompressedAudioInfo(Wave);
 	}
diff --git a/Engine/Source/Runtime/Engine/Private/AudioDevice.cpp b/Engine/Source/Runtime/Engine/Private/AudioDevice.cpp
--- a/Engine/Source/Runtime/Engine/Private/AudioDevice.cpp
+++ b/Engine/Source/Runtime/Engine/Private/AudioDevice.cpp
@@ -5872,14 +5872,14 @@ void FAudioDevice::Precache(USoundWave* SoundWave, bool bSynchronous, bool bTrac
 			if (bSynchronous)
 			{
 				// Create a worker to decompress the vorbis data
-				FAsyncAudioDecompress TempDecompress(SoundWave, GetNumPrecacheFrames());
+				FAsyncAudioDecompress TempDecompress(SoundWave, GetNumPrecacheFrames(), this);
 				TempDecompress.StartSynchronousTask();
 			}
 			else
 			{
 				// This is the one case where precaching will not be done when this function exits
 				checkf(SoundWave->GetPrecacheState() == ESoundWavePrecacheState::InProgress, TEXT("Bad PrecacheState %d on SoundWave %s"), static_cast<uint8>(SoundWave->GetPrecacheState()), *GetPathNameSafe(SoundWave));
-				SoundWave->AudioDecompressor = new FAsyncAudioDecompress(SoundWave, GetNumPrecacheFrames());
+				SoundWave->AudioDecompressor = new FAsyncAudioDecompress(SoundWave, GetNumPrecacheFrames(), this);
 				SoundWave->AudioDecompressor->StartBackgroundTask();
 				PrecachingSoundWaves.Add(SoundWave);
 			}
diff --git a/Engine/Source/Runtime/Engine/Private/AudioDeviceManager.cpp b/Engine/Source/Runtime/Engine/Private/AudioDeviceManager.cpp
--- a/Engine/Source/Runtime/Engine/Private/AudioDeviceManager.cpp
+++ b/Engine/Source/Runtime/Engine/Private/AudioDeviceManager.cpp
@@ -1198,14 +1198,18 @@ void FAudioDeviceHandle::AddStackDumpToAudioDeviceContainer()
 
 FAudioDeviceHandle::~FAudioDeviceHandle()
 {
-	if (GEngine && IsValid())
+	if (IsValid())
 	{
-		GEngine->GetAudioDeviceManager()->DecrementDevice(DeviceId, World);
+		FAudioDeviceManager* AudioDeviceManager = GEngine->GetAudioDeviceManager();
+		if (AudioDeviceManager)
+		{
+			AudioDeviceManager->DecrementDevice(DeviceId, World);
 
 #if INSTRUMENT_AUDIODEVICE_HANDLES
-		check(StackWalkID != INDEX_NONE);
-		GEngine->GetAudioDeviceManager()->RemoveStackWalkForContainer(DeviceId, StackWalkID);
+			check(StackWalkID != INDEX_NONE);
+			AudioDeviceManager->RemoveStackWalkForContainer(DeviceId, StackWalkID);
 #endif
+		}
 	}
 }
 
@@ -1221,7 +1225,7 @@ Audio::FDeviceId FAudioDeviceHandle::GetDeviceID() const
 
 bool FAudioDeviceHandle::IsValid() const
 {
-	return Device != nullptr;
+	return GEngine && GEngine->GetAudioDeviceManager() && Device != nullptr;
 }
 
 void FAudioDeviceHandle::Reset()
@@ -1231,7 +1235,6 @@ void FAudioDeviceHandle::Reset()
 
 FAudioDeviceHandle& FAudioDeviceHandle::operator=(const FAudioDeviceHandle& Other)
 {
-	// If this chunk was previously referencing another chunk, remove that chunk here.
 	if (IsValid())
 	{
 		check(FAudioDeviceManager::Get());
@@ -1243,12 +1246,15 @@ FAudioDeviceHandle& FAudioDeviceHandle::operator=(const FAudioDeviceHandle& Othe
 
 	if (IsValid())
 	{
-		// Increment the reference count for the new streaming chunk.
-		FAudioDeviceManager::Get()->IncrementDevice(DeviceId);
+		FAudioDeviceManager* AudioDeviceManager = FAudioDeviceManager::Get();
+		if (AudioDeviceManager)
+		{
+			AudioDeviceManager->IncrementDevice(DeviceId);
 
 #if INSTRUMENT_AUDIODEVICE_HANDLES
-		AddStackDumpToAudioDeviceContainer();
+			AddStackDumpToAudioDeviceContainer();
 #endif
+		}
 	}
 
 	return *this;
@@ -1256,7 +1262,6 @@ FAudioDeviceHandle& FAudioDeviceHandle::operator=(const FAudioDeviceHandle& Othe
 
 FAudioDeviceHandle& FAudioDeviceHandle::operator=(FAudioDeviceHandle&& Other)
 {
-	// If this chunk was previously referencing another chunk, remove that chunk here.
 	if (FAudioDeviceManager::Get() && IsValid())
 	{
 #if INSTRUMENT_AUDIODEVICE_HANDLES
@@ -1270,7 +1275,6 @@ FAudioDeviceHandle& FAudioDeviceHandle::operator=(FAudioDeviceHandle&& Other)
 	Device = Other.Device;
 	DeviceId = Other.DeviceId;
 
-	// If the chunk we're moving from is valid, we null it out without needing to decrement.
 	Other.Device = nullptr;
 	Other.DeviceId = INDEX_NONE;
 
diff --git a/Engine/Source/Runtime/Engine/Private/AudioVolume.cpp b/Engine/Source/Runtime/Engine/Private/AudioVolume.cpp
--- a/Engine/Source/Runtime/Engine/Private/AudioVolume.cpp
+++ b/Engine/Source/Runtime/Engine/Private/AudioVolume.cpp
@@ -91,15 +91,18 @@ void AAudioVolume::AddProxy() const
 {
 	if (UWorld* World = GetWorld())
 	{
-		if (FAudioDevice* AudioDevice = World->GetAudioDeviceRaw())
+		if (FAudioDeviceHandle AudioDeviceHandle = World->GetAudioDevice())
 		{
 			DECLARE_CYCLE_STAT(TEXT("FAudioThreadTask.AddAudioVolumeProxy"), STAT_AudioAddAudioVolumeProxy, STATGROUP_TaskGraphTasks);
 
 			FAudioVolumeProxy Proxy(this);
 
-			FAudioThread::RunCommandOnAudioThread([AudioDevice, Proxy]()
+			FAudioThread::RunCommandOnAudioThread([AudioDeviceHandle, Proxy]() mutable
 			{
-				AudioDevice->AddAudioVolumeProxy(Proxy);
+				if (AudioDeviceHandle)
+				{
+					AudioDeviceHandle->AddAudioVolumeProxy(Proxy);
+				}
 			}, GET_STATID(STAT_AudioAddAudioVolumeProxy));
 		}
 	}
@@ -122,14 +125,17 @@ void AAudioVolume::RemoveProxy() const
 	UWorld* World = GetWorld();
 	if (World)
 	{
-		if (FAudioDevice* AudioDevice = World->GetAudioDeviceRaw())
+		if (FAudioDeviceHandle AudioDeviceHandle = World->GetAudioDevice())
 		{
 			DECLARE_CYCLE_STAT(TEXT("FAudioThreadTask.RemoveAudioVolumeProxy"), STAT_AudioRemoveAudioVolumeProxy, STATGROUP_TaskGraphTasks);
 
 			const uint32 AudioVolumeID = GetUniqueID();
-			FAudioThread::RunCommandOnAudioThread([AudioDevice, AudioVolumeID]()
+			FAudioThread::RunCommandOnAudioThread([AudioDeviceHandle, AudioVolumeID]() mutable
 			{
-				AudioDevice->RemoveAudioVolumeProxy(AudioVolumeID);
+				if (AudioDeviceHandle)
+				{
+					AudioDeviceHandle->RemoveAudioVolumeProxy(AudioVolumeID);
+				}
 			}, GET_STATID(STAT_AudioRemoveAudioVolumeProxy));
 		}
 	}
diff --git a/Engine/Source/Runtime/Engine/Public/AudioDecompress.h b/Engine/Source/Runtime/Engine/Public/AudioDecompress.h
--- a/Engine/Source/Runtime/Engine/Public/AudioDecompress.h
+++ b/Engine/Source/Runtime/Engine/Public/AudioDecompress.h
@@ -327,7 +327,7 @@ public:
 	 *
 	 * @param	InWave		Wave data to decompress
 	 */
-	ENGINE_API FAsyncAudioDecompressWorker(USoundWave* InWave, int32 InNumPrecacheFrames);
+	ENGINE_API FAsyncAudioDecompressWorker(USoundWave* InWave, int32 InNumPrecacheFrames, FAudioDevice* InAudioDevice = nullptr);
 
 	/**
 	 * Performs the async audio decompression
diff --git a/Engine/Source/Runtime/Engine/Public/AudioDeviceManager.h b/Engine/Source/Runtime/Engine/Public/AudioDeviceManager.h
--- a/Engine/Source/Runtime/Engine/Public/AudioDeviceManager.h
+++ b/Engine/Source/Runtime/Engine/Public/AudioDeviceManager.h
@@ -123,21 +123,25 @@ public:
 	// For safety, we still require explicit casting to an FAudioDevice* to ensure ownership isn't lost due to programmer error.
 	const FAudioDevice& operator*() const
 	{
+		check(IsValid());
 		return *Device;
 	}
 
 	FAudioDevice& operator*()
 	{
+		check(IsValid());
 		return *Device;
 	}
 
 	const FAudioDevice* operator->() const
 	{
+		check(IsValid());
 		return Device;
 	}
 
 	FAudioDevice* operator->()
 	{
+		check(IsValid());
 		return Device;
 	}
 
@@ -168,6 +172,7 @@ public:
 
 	explicit operator FAudioDevice*() const
 	{
+		check(IsValid());
 		return Device;
 	}
 
