fix for LPCM audio asserting (hanging on shipping/test builds) on loop when stream caching is enabled.

diff --git a/Engine/Source/Runtime/Engine/Private/ADPCMAudioInfo.cpp b/Engine/Source/Runtime/Engine/Private/ADPCMAudioInfo.cpp
--- a/Engine/Source/Runtime/Engine/Private/ADPCMAudioInfo.cpp
+++ b/Engine/Source/Runtime/Engine/Private/ADPCMAudioInfo.cpp
@@ -770,7 +770,7 @@ bool FADPCMAudioInfo::StreamCompressedData(uint8* Destination, bool bLooping, ui
 			if(TotalSamplesStreamed >= TotalSamplesPerChannel)
 			{
 				ReachedEndOfSamples = true;
-				CurrentChunkIndex = 0;
+				CurrentChunkIndex = FirstChunkSampleDataIndex;
 				CurrentChunkBufferOffset = 0;
 				TotalSamplesStreamed = 0;
 				CurCompressedChunkData = nullptr;
