fix for emotes playing back delayed. Core issue is that there will be latency playing back a USoundWave/USoundCue immediately after loading the asset even if the audio asset is set to retained on its USoundClass, since we kick off the async load for the first chunk of audio on PostLoad. Fix is to add 64 KB of audio to the zeroth chunk of sounds set with the streaming flag.

diff --git a/Engine/Source/Runtime/AudioPlatformConfiguration/Public/AudioCompressionSettings.h b/Engine/Source/Runtime/AudioPlatformConfiguration/Public/AudioCompressionSettings.h
--- a/Engine/Source/Runtime/AudioPlatformConfiguration/Public/AudioCompressionSettings.h
+++ b/Engine/Source/Runtime/AudioPlatformConfiguration/Public/AudioCompressionSettings.h
@@ -30,8 +30,15 @@ struct FAudioStreamCachingSettings
 	// CacheSizeKB / 256.
 	int32 CacheSizeKB;
 
+	// Bool flag for keeping sounds flagged for streaming chunked in the style of the legacy streaming manager.
+	bool bForceLegacyStreamChunking;
+
+	int32 ZerothChunkSizeForLegacyStreamChunkingKB;
+
 	FAudioStreamCachingSettings()
 		: CacheSizeKB(DefaultCacheSize)
+		, bForceLegacyStreamChunking(false)
+		, ZerothChunkSizeForLegacyStreamChunkingKB(256)
 	{
 	}
 };
@@ -105,6 +112,12 @@ struct FPlatformAudioCookOverrides
 			// cache info:
 			OutSuffix.Append(TEXT("MEM_"));
 			OutSuffix.AppendInt(InOverrides->StreamCachingSettings.CacheSizeKB);
+
+			if (InOverrides->StreamCachingSettings.bForceLegacyStreamChunking)
+			{
+				OutSuffix.Append(TEXT("_LegacyChunking_"));
+				OutSuffix.AppendInt(InOverrides->StreamCachingSettings.ZerothChunkSizeForLegacyStreamChunkingKB);
+			}
 		}
 		
 
diff --git a/Engine/Source/Runtime/Engine/Private/AudioDerivedData.cpp b/Engine/Source/Runtime/Engine/Private/AudioDerivedData.cpp
--- a/Engine/Source/Runtime/Engine/Private/AudioDerivedData.cpp
+++ b/Engine/Source/Runtime/Engine/Private/AudioDerivedData.cpp
@@ -336,6 +336,7 @@ class FStreamedAudioCacheDerivedDataWorker : public FNonAbandonableTask
 
 				const int32 MinimumChunkSize = AudioFormat->GetMinimumSizeForInitialChunk(AudioFormatName, CompressedBuffer);
 				const bool bUseStreamCaching = CompressionOverrides && CompressionOverrides->bUseStreamCaching;
+				const bool bForceLegacyStreamChunking = SoundWave.bStreaming && CompressionOverrides && CompressionOverrides->StreamCachingSettings.bForceLegacyStreamChunking;
 
 				// If the initial chunk  for this sound wave was overridden, use that:
 				if (SoundWave.InitialChunkSize > 0)
@@ -347,8 +348,22 @@ class FStreamedAudioCacheDerivedDataWorker : public FNonAbandonableTask
 					// Ensure that the minimum chunk size is nonzero if our compressed buffer is not empty.
 					checkf(CompressedBuffer.Num() == 0 || MinimumChunkSize != 0, TEXT("To use Load On Demand, please override GetMinimumSizeForInitialChunk"));
 
-					// Otherwise if we're using Load On Demand, the first chunk should be as small as possible:
-					FirstChunkSize = MinimumChunkSize;
+					//
+					if (bForceLegacyStreamChunking)
+					{
+						int32 LegacyZerothChunkSize = CompressionOverrides->StreamCachingSettings.ZerothChunkSizeForLegacyStreamChunkingKB * 1024;
+						if (LegacyZerothChunkSize == 0)
+						{
+							LegacyZerothChunkSize = MaxChunkSize;
+						}
+
+						FirstChunkSize = LegacyZerothChunkSize;
+					}
+					else
+					{
+						// Otherwise if we're using Audio Stream Caching, the first chunk should be as small as possible:
+						FirstChunkSize = MinimumChunkSize;
+					}
 				}
 
 				if (bUseStreamCaching)
