Fix missing critical section on handle increment Could cause out of order ref count, which could cause audio device shutdown.
#b aaron.mcleran

diff --git a/Engine/Source/Runtime/Engine/Private/AudioDeviceManager.cpp b/Engine/Source/Runtime/Engine/Private/AudioDeviceManager.cpp
--- a/Engine/Source/Runtime/Engine/Private/AudioDeviceManager.cpp
+++ b/Engine/Source/Runtime/Engine/Private/AudioDeviceManager.cpp
@@ -498,6 +498,8 @@ bool FAudioDeviceManager::ShutdownAudioDevice(Audio::FDeviceId Handle)
 
 void FAudioDeviceManager::IncrementDevice(Audio::FDeviceId DeviceID)
 {
+	FScopeLock ScopeLock(&DeviceMapCriticalSection);
+
 	// If there is an FAudioDeviceHandle out in the world
 	check(Devices.Contains(DeviceID));
 
