Fix cook crash on older timeline components as they are used to be flagged as UCS created components, but now are SCS but still thought they needed to determine their serialization index
#rnx

diff --git a/Engine/Source/Runtime/Engine/Private/Components/ActorComponent.cpp b/Engine/Source/Runtime/Engine/Private/Components/ActorComponent.cpp
--- a/Engine/Source/Runtime/Engine/Private/Components/ActorComponent.cpp
+++ b/Engine/Source/Runtime/Engine/Private/Components/ActorComponent.cpp
@@ -438,38 +438,40 @@ bool UActorComponent::IsCreatedByConstructionScript() const
 #if WITH_EDITORONLY_DATA
 void UActorComponent::DetermineUCSSerializationIndexForLegacyComponent()
 {
-	check(CreationMethod == EComponentCreationMethod::UserConstructionScript);
 	check(bNeedsUCSSerializationIndexEvaluted);
 	bNeedsUCSSerializationIndexEvaluted = false;
 
 	int32 ComputedSerializationIndex = INDEX_NONE;
 
-	if (AActor* ComponentOwner = GetOwner())
+	if (CreationMethod == EComponentCreationMethod::UserConstructionScript)
 	{
-		if (ComponentOwner->BlueprintCreatedComponents.Num() > 0)
+		if (AActor* ComponentOwner = GetOwner())
 		{
-			UObject* ComponentTemplate = GetArchetype();
-
-			bool bFound = false;
-			for (const UActorComponent* BlueprintCreatedComponent : ComponentOwner->BlueprintCreatedComponents)
+			if (ComponentOwner->BlueprintCreatedComponents.Num() > 0)
 			{
-				if (BlueprintCreatedComponent && BlueprintCreatedComponent->CreationMethod == EComponentCreationMethod::UserConstructionScript)
+				UObject* ComponentTemplate = GetArchetype();
+
+				bool bFound = false;
+				for (const UActorComponent* BlueprintCreatedComponent : ComponentOwner->BlueprintCreatedComponents)
 				{
-					if (BlueprintCreatedComponent == this)
+					if (BlueprintCreatedComponent && BlueprintCreatedComponent->CreationMethod == EComponentCreationMethod::UserConstructionScript)
 					{
-						++ComputedSerializationIndex;
-						bFound = true;
-						break;
-					}
-					else if (BlueprintCreatedComponent->GetArchetype() == ComponentTemplate)
-					{
-						++ComputedSerializationIndex;
+						if (BlueprintCreatedComponent == this)
+						{
+							++ComputedSerializationIndex;
+							bFound = true;
+							break;
+						}
+						else if (BlueprintCreatedComponent->GetArchetype() == ComponentTemplate)
+						{
+							++ComputedSerializationIndex;
+						}
 					}
 				}
-			}
-			if (!bFound)
-			{
-				ComputedSerializationIndex = INDEX_NONE;
+				if (!bFound)
+				{
+					ComputedSerializationIndex = INDEX_NONE;
+				}
 			}
 		}
 	}
