Add missing G16 case in default texture format list

diff --git a/Engine/Source/Runtime/Engine/Private/Texture.cpp b/Engine/Source/Runtime/Engine/Private/Texture.cpp
--- a/Engine/Source/Runtime/Engine/Private/Texture.cpp
+++ b/Engine/Source/Runtime/Engine/Private/Texture.cpp
@@ -1304,7 +1304,7 @@ FName GetDefaultTextureFormatName( const ITargetPlatform* TargetPlatform, const
 	FName TextureFormatName = NAME_None;
 
 	/**
-	 * IF you add a format to this function don't forget to update GetAllDefaultTextureFormatNames 
+	 * IF you add a format to this function don't forget to update GetAllDefaultTextureFormats 
 	 */
 
 #if WITH_EDITOR
@@ -1488,6 +1488,7 @@ void GetAllDefaultTextureFormats(const class ITargetPlatform* TargetPlatform, TA
 	static FName NameBGRA8(TEXT("BGRA8"));
 	static FName NameXGXR8(TEXT("XGXR8"));
 	static FName NameG8(TEXT("G8"));
+	static FName NameG16(TEXT("G16"));
 	static FName NameVU8(TEXT("VU8"));
 	static FName NameRGBA16F(TEXT("RGBA16F"));
 	static FName NameBC6H(TEXT("BC6H"));
@@ -1503,6 +1504,7 @@ void GetAllDefaultTextureFormats(const class ITargetPlatform* TargetPlatform, TA
 	OutFormats.Add(NameBGRA8);
 	OutFormats.Add(NameXGXR8);
 	OutFormats.Add(NameG8);
+	OutFormats.Add(NameG16);
 	OutFormats.Add(NameVU8);
 	OutFormats.Add(NameRGBA16F);
 	if (bSupportDX11TextureFormats)
