Make sure we verify we have valid data before compressing it.
#rnx

diff --git a/Engine/Source/Runtime/Engine/Private/Rendering/SkeletalMeshLODImporterData.cpp b/Engine/Source/Runtime/Engine/Private/Rendering/SkeletalMeshLODImporterData.cpp
--- a/Engine/Source/Runtime/Engine/Private/Rendering/SkeletalMeshLODImporterData.cpp
+++ b/Engine/Source/Runtime/Engine/Private/Rendering/SkeletalMeshLODImporterData.cpp
@@ -657,18 +657,25 @@ FArchive& operator<<(FArchive& Ar, FSkeletalMeshImportData& RawMesh)
 
 	if (Ar.IsLoading() && Version < RAW_SKELETAL_MESH_BULKDATA_VER_CompressMorphTargetData)
 	{
-		//Compress the morph target data
-		for (int32 MorphTargetIndex = 0; MorphTargetIndex < RawMesh.MorphTargets.Num(); ++MorphTargetIndex)
+		if (RawMesh.MorphTargetModifiedPoints.Num() != 0)
 		{
-			const TSet<uint32>& ModifiedPoints = RawMesh.MorphTargetModifiedPoints[MorphTargetIndex];
-			FSkeletalMeshImportData& ToCompressShapeImportData = RawMesh.MorphTargets[MorphTargetIndex];
-			TArray<FVector> CompressPoints;
-			CompressPoints.Reserve(ToCompressShapeImportData.Points.Num());
-			for (uint32 PointIndex : ModifiedPoints)
+			//Compress the morph target data
+			for (int32 MorphTargetIndex = 0; MorphTargetIndex < RawMesh.MorphTargets.Num(); ++MorphTargetIndex)
 			{
-				CompressPoints.Add(ToCompressShapeImportData.Points[PointIndex]);
+				if (!RawMesh.MorphTargetModifiedPoints.IsValidIndex(MorphTargetIndex))
+				{
+					continue;
+				}
+				const TSet<uint32>& ModifiedPoints = RawMesh.MorphTargetModifiedPoints[MorphTargetIndex];
+				FSkeletalMeshImportData& ToCompressShapeImportData = RawMesh.MorphTargets[MorphTargetIndex];
+				TArray<FVector> CompressPoints;
+				CompressPoints.Reserve(ToCompressShapeImportData.Points.Num());
+				for (uint32 PointIndex : ModifiedPoints)
+				{
+					CompressPoints.Add(ToCompressShapeImportData.Points[PointIndex]);
+				}
+				ToCompressShapeImportData.Points = CompressPoints;
 			}
-			ToCompressShapeImportData.Points = CompressPoints;
 		}
 	}
 
