Copying Private-LoadTimes-4.24 stream to Main. Biggest changes are in Materials/Shader memory freezing.

diff --git a/Engine/Source/Runtime/Engine/Classes/Materials/Material.h b/Engine/Source/Runtime/Engine/Classes/Materials/Material.h
--- a/Engine/Source/Runtime/Engine/Classes/Materials/Material.h
+++ b/Engine/Source/Runtime/Engine/Classes/Materials/Material.h
@@ -9,6 +9,7 @@
 #include "RenderCommandFence.h"
 #include "Materials/MaterialInterface.h"
 #include "MaterialShared.h"
+#include "MaterialCachedData.h"
 #include "MaterialExpressionIO.h"
 #include "Materials/MaterialExpressionMaterialFunctionCall.h"
 #include "Materials/MaterialExpressionMaterialAttributeLayers.h"
@@ -262,50 +263,6 @@ struct FVector2MaterialInput : public FMaterialInput
 };
 #endif
 
-
-/** Stores information about a function that this material references, used to know when the material needs to be recompiled. */
-USTRUCT()
-struct FMaterialFunctionInfo
-{
-	GENERATED_USTRUCT_BODY()
-
-	/** Id that the function had when this material was last compiled. */
-	UPROPERTY()
-	FGuid StateId;
-
-	/** The function which this material has a dependency on. */
-	UPROPERTY()
-	class UMaterialFunctionInterface* Function;
-
-	FMaterialFunctionInfo()
-		: Function(NULL)
-	{}
-};
-
-/** Stores information about a parameter collection that this material references, used to know when the material needs to be recompiled. */
-USTRUCT()
-struct FMaterialParameterCollectionInfo
-{
-	GENERATED_USTRUCT_BODY()
-
-	/** Id that the collection had when this material was last compiled. */
-	UPROPERTY()
-	FGuid StateId;
-
-	/** The collection which this material has a dependency on. */
-	UPROPERTY()
-	class UMaterialParameterCollection* ParameterCollection;
-
-	FMaterialParameterCollectionInfo()
-		: ParameterCollection(NULL)
-	{}
-
-	bool operator==(const FMaterialParameterCollectionInfo& Other) const
-	{
-		return StateId == Other.StateId && ParameterCollection == Other.ParameterCollection;
-	}
-};
-
 USTRUCT()
 struct FParameterGroupData
 {
@@ -836,13 +793,11 @@ public:
 
 	UPROPERTY()
 	int32 EditorYaw;
-#endif // WITH_EDITORONLY_DATA
 
 	/** Array of material expressions, excluding Comments.  Used by the material editor. */
 	UPROPERTY()
 	TArray<class UMaterialExpression*> Expressions;
 
-#if WITH_EDITORONLY_DATA
 	/** Array of comments associated with this material; viewed in the material editor. */
 	UPROPERTY()
 	TArray<class UMaterialExpressionComment*> EditorComments;
@@ -850,15 +805,7 @@ public:
 	/** Controls where this parameter group is displayed in a material instance parameter list.  The lower the number the higher up in the parameter list. */
 	UPROPERTY(EditAnywhere, EditFixedSize, Category = "Group Sorting")
 	TArray<FParameterGroupData> ParameterGroupData;
-
 #endif // WITH_EDITORONLY_DATA
-	/** Array of all functions this material depends on. */
-	UPROPERTY()
-	TArray<struct FMaterialFunctionInfo> MaterialFunctionInfos;
-
-	/** Array of all parameter collections this material depends on. */
-	UPROPERTY()
-	TArray<struct FMaterialParameterCollectionInfo> MaterialParameterCollectionInfos;
 
 	/** true if this Material can be assumed Opaque when set to masked. */
 	UPROPERTY()
@@ -980,10 +927,6 @@ public:
 #endif //WITH_EDITORONLY_DATA
 
 private:
-	/** Caches the used quality levels of the material to save runtime graph traversal. */
-	UPROPERTY()
-	TArray<bool> CachedQualityLevelsUsed;
-
 	/** Inline material resources serialized from disk. To be processed on game thread in PostLoad. */
 	TArray<FMaterialResource> LoadedMaterialResources;
 
@@ -1001,12 +944,8 @@ private:
 	/** Fence used to guarantee that the RT is finished using various resources in this UMaterial before cleanup. */
 	FRenderCommandFence ReleaseFence;
 
-	/** 
-	 * Cached texture references from all expressions in the material (including nested functions). 
-	 * This is used to link uniform texture expressions which were stored in the DDC with the UTextures that they reference.
-	 */
-	UPROPERTY(transient)
-	TArray<UObject*> ExpressionTextureReferences;
+	UPROPERTY()
+	FMaterialCachedExpressionData CachedExpressionData;
 
 #if WITH_EDITORONLY_DATA
 	UPROPERTY()
@@ -1015,26 +954,28 @@ private:
 #endif // WITH_EDITORONLY_DATA
 public:
 
+	const FMaterialCachedExpressionData& GetCachedExpressionData() const { return CachedExpressionData; }
+
 	//~ Begin UMaterialInterface Interface.
 	ENGINE_API virtual UMaterial* GetMaterial() override;
 	ENGINE_API virtual const UMaterial* GetMaterial() const override;
-	ENGINE_API virtual const UMaterial* GetMaterial_Concurrent(TMicRecursionGuard& RecursionGuard) const override;
-	ENGINE_API virtual bool GetScalarParameterValue(const FMaterialParameterInfo& ParameterInfo,float& OutValue, bool bOveriddenOnly = false) const override;
-	ENGINE_API virtual bool IsScalarParameterUsedAsAtlasPosition(const FMaterialParameterInfo& ParameterInfo, bool& OutValue, TSoftObjectPtr<class UCurveLinearColor>& Curve, TSoftObjectPtr<class UCurveLinearColorAtlas>& Atlas) const override;
+	ENGINE_API virtual const UMaterial* GetMaterial_Concurrent(TMicRecursionGuard RecursionGuard = TMicRecursionGuard()) const override;
+	ENGINE_API virtual bool GetScalarParameterValue(const FHashedMaterialParameterInfo& ParameterInfo,float& OutValue, bool bOveriddenOnly = false) const override;
 #if WITH_EDITOR
-	ENGINE_API virtual bool GetScalarParameterSliderMinMax(const FMaterialParameterInfo& ParameterInfo, float& OutMinSlider, float& OutMaxSlider) const override;
+	ENGINE_API virtual bool IsScalarParameterUsedAsAtlasPosition(const FHashedMaterialParameterInfo& ParameterInfo, bool& OutValue, TSoftObjectPtr<class UCurveLinearColor>& Curve, TSoftObjectPtr<class UCurveLinearColorAtlas>& Atlas) const override;
+	ENGINE_API virtual bool GetScalarParameterSliderMinMax(const FHashedMaterialParameterInfo& ParameterInfo, float& OutMinSlider, float& OutMaxSlider) const override;
 #endif
-	ENGINE_API virtual bool GetVectorParameterValue(const FMaterialParameterInfo& ParameterInfo,FLinearColor& OutValue, bool bOveriddenOnly = false) const override;
-	ENGINE_API virtual bool IsVectorParameterUsedAsChannelMask(const FMaterialParameterInfo& ParameterInfo, bool& OutValue) const override;
+	ENGINE_API virtual bool GetVectorParameterValue(const FHashedMaterialParameterInfo& ParameterInfo,FLinearColor& OutValue, bool bOveriddenOnly = false) const override;
 #if WITH_EDITOR
-	ENGINE_API virtual bool GetVectorParameterChannelNames(const FMaterialParameterInfo& ParameterInfo, FParameterChannelNames& OutValue) const override;
+	ENGINE_API virtual bool IsVectorParameterUsedAsChannelMask(const FHashedMaterialParameterInfo& ParameterInfo, bool& OutValue) const override;
+	ENGINE_API virtual bool GetVectorParameterChannelNames(const FHashedMaterialParameterInfo& ParameterInfo, FParameterChannelNames& OutValue) const override;
 #endif
-	ENGINE_API virtual bool GetTextureParameterValue(const FMaterialParameterInfo& ParameterInfo,class UTexture*& OutValue, bool bOveriddenOnly = false) const override;
-	ENGINE_API virtual bool GetRuntimeVirtualTextureParameterValue(const FMaterialParameterInfo& ParameterInfo, class URuntimeVirtualTexture*& OutValue, bool bOveriddenOnly = false) const override;
+	ENGINE_API virtual bool GetTextureParameterValue(const FHashedMaterialParameterInfo& ParameterInfo,class UTexture*& OutValue, bool bOveriddenOnly = false) const override;
+	ENGINE_API virtual bool GetRuntimeVirtualTextureParameterValue(const FHashedMaterialParameterInfo& ParameterInfo, class URuntimeVirtualTexture*& OutValue, bool bOveriddenOnly = false) const override;
 #if WITH_EDITOR
-	ENGINE_API virtual bool GetTextureParameterChannelNames(const FMaterialParameterInfo& ParameterInfo, FParameterChannelNames& OutValue) const override;
+	ENGINE_API virtual bool GetTextureParameterChannelNames(const FHashedMaterialParameterInfo& ParameterInfo, FParameterChannelNames& OutValue) const override;
 #endif
-	ENGINE_API virtual bool GetFontParameterValue(const FMaterialParameterInfo& ParameterInfo,class UFont*& OutFontValue,int32& OutFontPage, bool bOveriddenOnly = false) const override;
+	ENGINE_API virtual bool GetFontParameterValue(const FHashedMaterialParameterInfo& ParameterInfo,class UFont*& OutFontValue,int32& OutFontPage, bool bOveriddenOnly = false) const override;
 	ENGINE_API virtual bool GetRefractionSettings(float& OutBiasValue) const override;
 	ENGINE_API virtual FMaterialRenderProxy* GetRenderProxy() const override;
 	ENGINE_API virtual UPhysicalMaterial* GetPhysicalMaterial() const override;
@@ -1043,22 +984,24 @@ public:
 	ENGINE_API virtual void GetUsedTextures(TArray<UTexture*>& OutTextures, EMaterialQualityLevel::Type QualityLevel, bool bAllQualityLevels, ERHIFeatureLevel::Type FeatureLevel, bool bAllFeatureLevels) const override;
 	ENGINE_API virtual void GetUsedTexturesAndIndices(TArray<UTexture*>& OutTextures, TArray< TArray<int32> >& OutIndices, EMaterialQualityLevel::Type QualityLevel, ERHIFeatureLevel::Type FeatureLevel) const override;
 	ENGINE_API virtual void OverrideTexture(const UTexture* InTextureToOverride, UTexture* OverrideTexture, ERHIFeatureLevel::Type InFeatureLevel) override;
-	ENGINE_API virtual void OverrideVectorParameterDefault(const FMaterialParameterInfo& ParameterInfo, const FLinearColor& Value, bool bOverride, ERHIFeatureLevel::Type FeatureLevel) override;
-	ENGINE_API virtual void OverrideScalarParameterDefault(const FMaterialParameterInfo& ParameterInfo, float Value, bool bOverride, ERHIFeatureLevel::Type FeatureLevel) override;
+	ENGINE_API virtual void OverrideVectorParameterDefault(const FHashedMaterialParameterInfo& ParameterInfo, const FLinearColor& Value, bool bOverride, ERHIFeatureLevel::Type FeatureLevel) override;
+	ENGINE_API virtual void OverrideScalarParameterDefault(const FHashedMaterialParameterInfo& ParameterInfo, float Value, bool bOverride, ERHIFeatureLevel::Type FeatureLevel) override;
 	ENGINE_API virtual bool CheckMaterialUsage(const EMaterialUsage Usage) override;
 	ENGINE_API virtual bool CheckMaterialUsage_Concurrent(const EMaterialUsage Usage) const override;
 	ENGINE_API virtual FMaterialResource* AllocateResource();
 	ENGINE_API virtual FMaterialResource* GetMaterialResource(ERHIFeatureLevel::Type InFeatureLevel, EMaterialQualityLevel::Type QualityLevel = EMaterialQualityLevel::Num) override;
 	ENGINE_API virtual const FMaterialResource* GetMaterialResource(ERHIFeatureLevel::Type InFeatureLevel, EMaterialQualityLevel::Type QualityLevel = EMaterialQualityLevel::Num) const override;
+#if WITH_EDITORONLY_DATA
 	ENGINE_API virtual bool GetStaticSwitchParameterValues(FStaticParamEvaluationContext& EvalContext, TBitArray<>& OutValues, FGuid* OutExpressionGuids, bool bCheckParent = true) const override;
 	ENGINE_API virtual bool GetStaticComponentMaskParameterValues(FStaticParamEvaluationContext& EvalContext, TBitArray<>& OutRGBAOrderedValues, FGuid* OutExpressionGuids, bool bCheckParent = true) const override;
-	ENGINE_API virtual bool GetTerrainLayerWeightParameterValue(const FMaterialParameterInfo& ParameterInfo, int32& OutWeightmapIndex, FGuid& OutExpressionGuid) const override;
-	ENGINE_API virtual bool GetMaterialLayersParameterValue(const FMaterialParameterInfo& ParameterInfo, FMaterialLayersFunctions& OutLayers, FGuid& OutExpressionGuid, bool bCheckParent = true) const override;
+	ENGINE_API virtual bool GetMaterialLayersParameterValue(const FHashedMaterialParameterInfo& ParameterInfo, FMaterialLayersFunctions& OutLayers, FGuid& OutExpressionGuid, bool bCheckParent = true) const override;
+#endif // WITH_EDITORONLY_DATA
+	ENGINE_API virtual bool GetTerrainLayerWeightParameterValue(const FHashedMaterialParameterInfo& ParameterInfo, int32& OutWeightmapIndex, FGuid& OutExpressionGuid) const override;
 	ENGINE_API virtual bool UpdateLightmassTextureTracking() override;
 #if WITH_EDITOR
-	ENGINE_API virtual bool GetGroupName(const FMaterialParameterInfo& ParameterInfo, FName& OutGroup) const override;
-	ENGINE_API virtual bool GetParameterDesc(const FMaterialParameterInfo& ParameterInfo, FString& OutDesc, const TArray<struct FStaticMaterialLayersParameter>* MaterialLayersParameters = nullptr) const override;
-	ENGINE_API virtual bool GetParameterSortPriority(const FMaterialParameterInfo& ParameterInfo, int32& OutSortPriority, const TArray<struct FStaticMaterialLayersParameter>* MaterialLayersParameters = nullptr) const override;
+	ENGINE_API virtual bool GetGroupName(const FHashedMaterialParameterInfo& ParameterInfo, FName& OutGroup) const override;
+	ENGINE_API virtual bool GetParameterDesc(const FHashedMaterialParameterInfo& ParameterInfo, FString& OutDesc, const TArray<struct FStaticMaterialLayersParameter>* MaterialLayersParameters = nullptr) const override;
+	ENGINE_API virtual bool GetParameterSortPriority(const FHashedMaterialParameterInfo& ParameterInfo, int32& OutSortPriority, const TArray<struct FStaticMaterialLayersParameter>* MaterialLayersParameters = nullptr) const override;
 	ENGINE_API virtual bool GetGroupSortPriority(const FString& InGroupName, int32& OutSortPriority) const override;
 	ENGINE_API virtual bool GetTexturesInPropertyChain(EMaterialProperty InProperty, TArray<UTexture*>& OutTextures, 
 		TArray<FName>* OutTextureParamNames, struct FStaticParameterSet* InStaticParameterSet,
@@ -1191,6 +1134,19 @@ public:
 	ENGINE_API virtual void GetLightingGuidChain(bool bIncludeTextures, TArray<FGuid>& OutGuids) const override;
 
 private:
+#if WITH_EDITOR
+	bool GetScalarParameterValue_Legacy(const FHashedMaterialParameterInfo& ParameterInfo, float& OutValue, bool bOveriddenOnly) const;
+	bool GetVectorParameterValue_Legacy(const FHashedMaterialParameterInfo& ParameterInfo, FLinearColor& OutValue, bool bOveriddenOnly) const;
+	bool GetTextureParameterValue_Legacy(const FHashedMaterialParameterInfo& ParameterInfo, class UTexture*& OutValue, bool bOveriddenOnly) const;
+	bool GetRuntimeVirtualTextureParameterValue_Legacy(const FHashedMaterialParameterInfo& ParameterInfo, class URuntimeVirtualTexture*& OutValue, bool bOveriddenOnly) const;
+	bool GetFontParameterValue_Legacy(const FHashedMaterialParameterInfo& ParameterInfo, class UFont*& OutFontValue, int32& OutFontPage, bool bOveriddenOnly) const;
+#endif // WITH_EDITOR
+	bool GetScalarParameterValue_New(const FHashedMaterialParameterInfo& ParameterInfo, float& OutValue, bool bOveriddenOnly) const;
+	bool GetVectorParameterValue_New(const FHashedMaterialParameterInfo& ParameterInfo, FLinearColor& OutValue, bool bOveriddenOnly) const;
+	bool GetTextureParameterValue_New(const FHashedMaterialParameterInfo& ParameterInfo, class UTexture*& OutValue, bool bOveriddenOnly) const;
+	bool GetRuntimeVirtualTextureParameterValue_New(const FHashedMaterialParameterInfo& ParameterInfo, class URuntimeVirtualTexture*& OutValue, bool bOveriddenOnly) const;
+	bool GetFontParameterValue_New(const FHashedMaterialParameterInfo& ParameterInfo, class UFont*& OutFontValue, int32& OutFontPage, bool bOveriddenOnly) const;
+
 	void BackwardsCompatibilityInputConversion();
 
 	/** Handles setting up an annotation for this object if a flag has changed value */
@@ -1238,25 +1194,7 @@ public:
 	 */
 	ENGINE_API bool NeedsSetMaterialUsage_Concurrent(bool &bOutHasUsage, const EMaterialUsage Usage) const;
 
-	/**
-	* @param	OutParameterNames		Storage array for the parameter names.
-	* @param	OutParameterIds			Storage array for the parameter id's.
-	*
-	* @return	Returns a array of parameter names used in this material for the specified expression type.
-	*/
-	template<typename ExpressionType>
-	UE_DEPRECATED(4.18, "This function is deprecated. Use GetAllParameterInfo instead, this contains a struct with FNames, layer association, and index")
-	void GetAllParameterNames(TArray<FName>& OutParameterNames, TArray<FGuid>& OutParameterIds, const TArray<FStaticMaterialLayersParameter>* MaterialLayersParameters = nullptr) const
-	{
-		OutParameterNames.Empty();
-		TArray<FMaterialParameterInfo> OutParameterInfo;
-		GetAllParameterInfo<ExpressionType>(OutParameterInfo, OutParameterIds, MaterialLayersParameters);
-		for (const FMaterialParameterInfo& Info : OutParameterInfo)
-		{
-			OutParameterNames.Add(Info.Name);
-		}
-	}
-
+#if WITH_EDITORONLY_DATA
 	/**
 	 * @param	OutParameterNames		Storage array for the parameter names.
 	 * @param	OutParameterIds			Storage array for the parameter id's.
@@ -1335,38 +1273,31 @@ public:
 			check(OutParameterInfo.Num() == OutParameterIds.Num());
 		}
 	}
-	
-	UE_DEPRECATED(4.18, "This function is deprecated. Use GetAllScalarParameterInfo instead, this contains a struct with FNames, layer association, and index")
-	ENGINE_API void GetAllScalarParameterNames(TArray<FName>& OutParameterNames, TArray<FGuid>& OutParameterIds) const;
-	UE_DEPRECATED(4.18, "This function is deprecated. Use GetAllVectorParameterInfo instead, this contains a struct with FNames, layer association, and index")
-	ENGINE_API void GetAllVectorParameterNames(TArray<FName>& OutParameterNames, TArray<FGuid>& OutParameterIds) const;
-	UE_DEPRECATED(4.18, "This function is deprecated. Use GetAllTextureParameterInfo instead, this contains a struct with FNames, layer association, and index")
-	ENGINE_API void GetAllTextureParameterNames(TArray<FName>& OutParameterNames, TArray<FGuid>& OutParameterIds) const;
-	UE_DEPRECATED(4.18, "This function is deprecated. Use GetAllFontParameterInfo instead, this contains a struct with FNames, layer association, and index")
-	ENGINE_API void GetAllFontParameterNames(TArray<FName>& OutParameterNames, TArray<FGuid>& OutParameterIds) const;
-	UE_DEPRECATED(4.18, "This function is deprecated. Use GetAllStaticSwitchParameterInfo instead, this contains a struct with FNames, layer association, and index")
-	ENGINE_API void GetAllStaticSwitchParameterNames(TArray<FName>& OutParameterNames, TArray<FGuid>& OutParameterIds) const;
-	UE_DEPRECATED(4.18, "This function is deprecated. Use GetAllStaticComponentMaskParameterInfo instead, this contains a struct with FNames, layer association, and index")
-	ENGINE_API void GetAllStaticComponentMaskParameterNames(TArray<FName>& OutParameterNames, TArray<FGuid>& OutParameterIds) const;
+#endif // WITH_EDITORONLY_DATA
 
 	ENGINE_API virtual void GetAllScalarParameterInfo(TArray<FMaterialParameterInfo>& OutParameterInfo, TArray<FGuid>& OutParameterIds) const override;
 	ENGINE_API virtual void GetAllVectorParameterInfo(TArray<FMaterialParameterInfo>& OutParameterInfo, TArray<FGuid>& OutParameterIds) const override;
 	ENGINE_API virtual void GetAllTextureParameterInfo(TArray<FMaterialParameterInfo>& OutParameterInfo, TArray<FGuid>& OutParameterIds) const override;
 	ENGINE_API virtual void GetAllFontParameterInfo(TArray<FMaterialParameterInfo>& OutParameterInfo, TArray<FGuid>& OutParameterIds) const override;
+	ENGINE_API virtual void GetAllRuntimeVirtualTextureParameterInfo(TArray<FMaterialParameterInfo>& OutParameterInfo, TArray<FGuid>& OutParameterIds) const override;
+#if WITH_EDITORONLY_DATA
 	ENGINE_API virtual void GetAllMaterialLayersParameterInfo(TArray<FMaterialParameterInfo>& OutParameterInfo, TArray<FGuid>& OutParameterIds) const override;
 	ENGINE_API virtual void GetAllStaticSwitchParameterInfo(TArray<FMaterialParameterInfo>& OutParameterInfo, TArray<FGuid>& OutParameterIds) const override;
 	ENGINE_API virtual void GetAllStaticComponentMaskParameterInfo(TArray<FMaterialParameterInfo>& OutParameterInfo, TArray<FGuid>& OutParameterIds) const override;
 
 	ENGINE_API virtual bool IterateDependentFunctions(TFunctionRef<bool(UMaterialFunctionInterface*)> Predicate) const override;
 	ENGINE_API virtual void GetDependentFunctions(TArray<class UMaterialFunctionInterface*>& DependentFunctions) const override;
+#endif // WITH_EDITORONLY_DATA
 
-	ENGINE_API virtual bool GetScalarParameterDefaultValue(const FMaterialParameterInfo& ParameterInfo, float& OutValue, bool bOveriddenOnly = false, bool bCheckOwnedGlobalOverrides = false) const override;
-	ENGINE_API virtual bool GetVectorParameterDefaultValue(const FMaterialParameterInfo& ParameterInfo, FLinearColor& OutValue, bool bOveriddenOnly = false, bool bCheckOwnedGlobalOverrides = false) const override;
-	ENGINE_API virtual bool GetTextureParameterDefaultValue(const FMaterialParameterInfo& ParameterInfo, class UTexture*& OutValue, bool bCheckOwnedGlobalOverrides = false) const override;
-	ENGINE_API virtual bool GetRuntimeVirtualTextureParameterDefaultValue(const FMaterialParameterInfo& ParameterInfo, class URuntimeVirtualTexture*& OutValue, bool bCheckOwnedGlobalOverrides = false) const override;
-	ENGINE_API virtual bool GetFontParameterDefaultValue(const FMaterialParameterInfo& ParameterInfo, class UFont*& OutFontValue, int32& OutFontPage, bool bCheckOwnedGlobalOverrides = false) const override;
-	ENGINE_API virtual bool GetStaticComponentMaskParameterDefaultValue(const FMaterialParameterInfo& ParameterInfo, bool& OutR, bool& OutG, bool& OutB, bool& OutA, FGuid& OutExpressionGuid, bool bCheckOwnedGlobalOverrides = false) const override;
-	ENGINE_API virtual bool GetStaticSwitchParameterDefaultValue(const FMaterialParameterInfo& ParameterInfo, bool& OutValue, FGuid& OutExpressionGuid, bool bCheckOwnedGlobalOverrides = false) const override;
+	ENGINE_API virtual bool GetScalarParameterDefaultValue(const FHashedMaterialParameterInfo& ParameterInfo, float& OutValue, bool bOveriddenOnly = false, bool bCheckOwnedGlobalOverrides = false) const override;
+	ENGINE_API virtual bool GetVectorParameterDefaultValue(const FHashedMaterialParameterInfo& ParameterInfo, FLinearColor& OutValue, bool bOveriddenOnly = false, bool bCheckOwnedGlobalOverrides = false) const override;
+	ENGINE_API virtual bool GetTextureParameterDefaultValue(const FHashedMaterialParameterInfo& ParameterInfo, class UTexture*& OutValue, bool bCheckOwnedGlobalOverrides = false) const override;
+	ENGINE_API virtual bool GetRuntimeVirtualTextureParameterDefaultValue(const FHashedMaterialParameterInfo& ParameterInfo, class URuntimeVirtualTexture*& OutValue, bool bCheckOwnedGlobalOverrides = false) const override;
+	ENGINE_API virtual bool GetFontParameterDefaultValue(const FHashedMaterialParameterInfo& ParameterInfo, class UFont*& OutFontValue, int32& OutFontPage, bool bCheckOwnedGlobalOverrides = false) const override;
+#if WITH_EDITOR
+	ENGINE_API virtual bool GetStaticComponentMaskParameterDefaultValue(const FHashedMaterialParameterInfo& ParameterInfo, bool& OutR, bool& OutG, bool& OutB, bool& OutA, FGuid& OutExpressionGuid, bool bCheckOwnedGlobalOverrides = false) const override;
+	ENGINE_API virtual bool GetStaticSwitchParameterDefaultValue(const FHashedMaterialParameterInfo& ParameterInfo, bool& OutValue, FGuid& OutExpressionGuid, bool bCheckOwnedGlobalOverrides = false) const override;
+#endif // WITH_EDITOR
 
 	/** Returns the material's decal blend mode, calculated from the DecalBlendMode property and what inputs are connected. */
 	uint32 GetDecalBlendMode() const { return DecalBlendMode; }
@@ -1374,6 +1305,7 @@ public:
 	/** Returns the material's decal response mode */
 	uint32 GetMaterialDecalResponse() const { return MaterialDecalResponse; }
 
+#if WITH_EDITORONLY_DATA
 	/**
 	 * Attempt to find a expression by its GUID.
 	 */
@@ -1487,6 +1419,7 @@ public:
 
 		return false;
 	}
+#endif // WITH_EDITORONLY_DATA
 
 	/** Determines whether each quality level has different nodes by inspecting the material's expressions. 
 	* Or is required by the material quality setting overrides.
@@ -1495,20 +1428,9 @@ public:
 	*/
 	void GetQualityLevelUsage(TArray<bool, TInlineAllocator<EMaterialQualityLevel::Num> >& QualityLevelsUsed, EShaderPlatform ShaderPlatform);
 
-	/** Determines whether each quality level has different nodes by inspecting the material's expressions.
-	* @param	QualityLevelsUsed	output array of used quality levels.
-	*/
-	void GetQualityLevelNodeUsage(TArray<bool, TInlineAllocator<EMaterialQualityLevel::Num> >& OutQualityLevelsUsed);
-
-	/**
-	 * Cache the expression texture references for this UMaterial 
-	 * if the cache is not filled then it will rebuild the texture references
-	 * see also RebuildExpressionTextureReferences
-	 */
-	void CacheExpressionTextureReferences();
-
-	/** Rebuild ExpressionTextureReferences with all textures referenced by expressions in this material. */
-	ENGINE_API void RebuildExpressionTextureReferences();
+#if WITH_EDITOR
+	ENGINE_API void UpdateCachedExpressionData();
+#endif
 
 	/** Attempts to add a new group name to the Group Data struct. True if new name was added. */
 	ENGINE_API bool AttemptInsertNewGroupName(const FString& InNewName);
@@ -1518,19 +1440,6 @@ private:
 	 * Flush existing resource shader maps and resets the material resource's Ids.
 	 */
 	ENGINE_API virtual void FlushResourceShaderMaps();
-
-#if WITH_EDITOR
-	/**
-	 * Rebuild the MaterialFunctionInfos array with the current state of the material's function dependencies,
-	 * And updates any function call nodes in this material so their inputs and outputs stay valid.
-	 */
-	void RebuildMaterialFunctionInfo();
-
-	/** 
-	 * Rebuild the MaterialParameterCollectionInfos array with the current state of the material's parameter collection dependencies.
-	 */
-	void RebuildMaterialParameterCollectionInfo();
-#endif // WITH_EDITOR
 	
 	/** 
 	 * Cache resource shaders for rendering. 
@@ -1552,17 +1461,19 @@ private:
 	/** Caches shader maps for an array of material resources. */
 	void CacheShadersForResources(EShaderPlatform ShaderPlatform, const TArray<FMaterialResource*>& ResourcesToCache, const ITargetPlatform* TargetPlatform = nullptr);
 
+#if WITH_EDITOR
 	/**
 	 * If there is some texture reference used by a TextureProperty node in any expressions, this function
 	 * will extract the current hash of TextureReferencesHash into a string  then append the texture guid used by the node
 	 * and recompute a new hash.
 	 */
 	void GetForceRecompileTextureIdsHash(FSHAHash &TextureReferencesHash);
+#endif // WITH_EDITOR
 
 public:
+#if WITH_EDITOR
 	ENGINE_API bool IsTextureForceRecompileCacheRessource(UTexture *Texture);
 
-#if WITH_EDITOR
 	/* Recompute the ddc cache key and reload the material in case the key is not the same.
 	 * It will also make sure lightmass texture reference are up to date
 	 */
@@ -1572,7 +1483,7 @@ public:
 	/**
 	 * Go through every material, flush the specified types and re-initialize the material's shader maps.
 	 */
-	ENGINE_API static void UpdateMaterialShaders(TArray<FShaderType*>& ShaderTypesToFlush, TArray<const FShaderPipelineType*>& ShaderPipelineTypesToFlush, TArray<const FVertexFactoryType*>& VFTypesToFlush, EShaderPlatform ShaderPlatform);
+	ENGINE_API static void UpdateMaterialShaders(TArray<const FShaderType*>& ShaderTypesToFlush, TArray<const FShaderPipelineType*>& ShaderPipelineTypesToFlush, TArray<const FVertexFactoryType*>& VFTypesToFlush, EShaderPlatform ShaderPlatform);
 
 	/** 
 	 * Backs up all material shaders to memory through serialization, organized by FMaterialShaderMap. 
@@ -1634,11 +1545,6 @@ public:
 	 */
 	ENGINE_API virtual bool HasDuplicateDynamicParameters(const UMaterialExpression* Expression);
 
-	/** Collect all material expressions fomr this material and all its functions and figure out which possible shading models exist in this material */
-	ENGINE_API void RebuildShadingModelField();
-
-#endif // WITH_EDITOR
-
 	/**
 	 * Iterate through all of the expression nodes and fix up changed properties on
 	 * matching dynamic parameters when a change occurs.
@@ -1647,6 +1553,11 @@ public:
 	 */
 	ENGINE_API virtual void UpdateExpressionDynamicParameters(const UMaterialExpression* Expression);
 
+	/** Collect all material expressions fomr this material and all its functions and figure out which possible shading models exist in this material */
+	ENGINE_API void RebuildShadingModelField();
+
+#endif // WITH_EDITOR
+
 	/**
 	 * Get the name of a parameter.
 	 * @param	Expression	The expression to retrieve the name from.
@@ -1701,9 +1612,11 @@ public:
 	ENGINE_API FExpressionInput* GetExpressionInputForProperty(EMaterialProperty InProperty);
 #endif
 
+#if WITH_EDITORONLY_DATA
 	/* Returns any UMaterialExpressionCustomOutput expressions */
 	ENGINE_API void GetAllCustomOutputExpressions(TArray<class UMaterialExpressionCustomOutput*>& OutCustomOutputs) const;
 	ENGINE_API void GetAllExpressionsForCustomInterpolators(TArray<class UMaterialExpression*>& OutExpressions) const;
+#endif // WITH_EDITORONLY_DATA
 
 #if WITH_EDITOR
 	/**
@@ -1739,7 +1652,7 @@ public:
 #endif
 
 	/** Appends textures referenced by expressions, including nested functions. */
-	ENGINE_API virtual void AppendReferencedTextures(TArray<UObject*>& InOutTextures) const override;
+	ENGINE_API virtual TArrayView<UObject* const> GetReferencedTextures() const override final { return CachedExpressionData.ReferencedTextures; }
 
 protected:
 
@@ -1842,6 +1755,7 @@ private:
 	// DO NOT CALL outside of FMaterialEditor!
 	ENGINE_API static void ForceNoCompilationInPostLoad(bool bForceNoCompilation);
 
+#if WITH_EDITORONLY_DATA
 	/* Helper function to help finding expression GUID taking into account UMaterialExpressionMaterialFunctionCall */
 	template<typename ExpressionType>
 	ExpressionType* FindExpressionByGUIDRecursive(const FGuid &InGUID, const TArray<UMaterialExpression*>& InMaterialExpression)
@@ -1904,6 +1818,7 @@ private:
 
 		return nullptr;
 	}
+#endif // WITH_EDITORONLY_DATA
 
 #if WITH_EDITOR
 
diff --git a/Engine/Source/Runtime/Engine/Classes/Materials/MaterialExpression.h b/Engine/Source/Runtime/Engine/Classes/Materials/MaterialExpression.h
--- a/Engine/Source/Runtime/Engine/Classes/Materials/MaterialExpression.h
+++ b/Engine/Source/Runtime/Engine/Classes/Materials/MaterialExpression.h
@@ -7,6 +7,7 @@
 #include "UObject/Object.h"
 #include "Misc/Guid.h"
 #include "MaterialShared.h"
+#include "MaterialCachedData.h"
 #include "MaterialExpressionIO.h"
 
 #include "MaterialExpression.generated.h"
@@ -96,29 +97,6 @@ struct FExpressionOutput
 };
 #endif
 
-USTRUCT()
-struct FParameterChannelNames
-{
-	GENERATED_USTRUCT_BODY()
-
-		UPROPERTY(EditAnywhere, Category = MaterialExpressionVectorParameter)
-		FText R;
-
-	UPROPERTY(EditAnywhere, Category = MaterialExpressionVectorParameter)
-		FText G;
-
-	UPROPERTY(EditAnywhere, Category = MaterialExpressionVectorParameter)
-		FText B;
-
-	UPROPERTY(EditAnywhere, Category = MaterialExpressionVectorParameter)
-		FText A;
-
-	FParameterChannelNames()
-	{
-
-	};
-};
-
 UCLASS(abstract, BlueprintType, hidecategories=Object)
 class ENGINE_API UMaterialExpression : public UObject
 {
diff --git a/Engine/Source/Runtime/Engine/Classes/Materials/MaterialExpressionDynamicParameter.h b/Engine/Source/Runtime/Engine/Classes/Materials/MaterialExpressionDynamicParameter.h
--- a/Engine/Source/Runtime/Engine/Classes/Materials/MaterialExpressionDynamicParameter.h
+++ b/Engine/Source/Runtime/Engine/Classes/Materials/MaterialExpressionDynamicParameter.h
@@ -57,11 +57,13 @@ class UMaterialExpressionDynamicParameter : public UMaterialExpression
 #endif
 	//~ End UMaterialExpression Interface
 
+#if WITH_EDITORONLY_DATA
 	/**
 	 * Iterate through all of the expression nodes until we find another 
 	 * dynamic parameter we can copy the properties from
 	 */
 	ENGINE_API void UpdateDynamicParameterProperties();
+#endif // WITH_EDITORONLY_DATA
 
 	/**
 	 * Copy the properties from the specified dynamic parameter
diff --git a/Engine/Source/Runtime/Engine/Classes/Materials/MaterialExpressionFontSampleParameter.h b/Engine/Source/Runtime/Engine/Classes/Materials/MaterialExpressionFontSampleParameter.h
--- a/Engine/Source/Runtime/Engine/Classes/Materials/MaterialExpressionFontSampleParameter.h
+++ b/Engine/Source/Runtime/Engine/Classes/Materials/MaterialExpressionFontSampleParameter.h
@@ -53,7 +53,7 @@ class UMaterialExpressionFontSampleParameter : public UMaterialExpressionFontSam
 	//~ End UMaterialExpression Interface
 	
 	/** Return whether this is the named parameter, and fill in its value */
-	bool IsNamedParameter(const FMaterialParameterInfo& ParameterInfo, UFont*& OutFontValue, int32& OutFontPage) const;
+	bool IsNamedParameter(const FHashedMaterialParameterInfo& ParameterInfo, UFont*& OutFontValue, int32& OutFontPage) const;
 
 #if WITH_EDITOR
 	bool SetParameterValue(FName InParameterName, UFont* InFontValue, int32 InFontPage);
diff --git a/Engine/Source/Runtime/Engine/Classes/Materials/MaterialExpressionMaterialAttributeLayers.h b/Engine/Source/Runtime/Engine/Classes/Materials/MaterialExpressionMaterialAttributeLayers.h
--- a/Engine/Source/Runtime/Engine/Classes/Materials/MaterialExpressionMaterialAttributeLayers.h
+++ b/Engine/Source/Runtime/Engine/Classes/Materials/MaterialExpressionMaterialAttributeLayers.h
@@ -89,14 +89,24 @@ class UMaterialExpressionMaterialAttributeLayers : public UMaterialExpression
 	virtual bool NeedsLoadForClient() const override { return true; };
 	//~ Begin UObject Interface
 
+#if WITH_EDITORONLY_DATA
 	ENGINE_API void RebuildLayerGraph(bool bReportErrors);
 	ENGINE_API void OverrideLayerGraph(const FMaterialLayersFunctions* OverrideLayers);
+#endif // WITH_EDITORONLY_DATA
+
+#if WITH_EDITOR
 	bool ValidateLayerConfiguration(FMaterialCompiler* Compiler, bool bReportErrors);
+#endif // WITH_EDITOR
 
+#if WITH_EDITORONLY_DATA
 	bool IterateDependentFunctions(TFunctionRef<bool(UMaterialFunctionInterface*)> Predicate) const;
 	void GetDependentFunctions(TArray<UMaterialFunctionInterface*>& DependentFunctions) const;
-	UMaterialFunctionInterface* GetParameterAssociatedFunction(const FMaterialParameterInfo& ParameterInfo) const;
+#endif
+
+	UMaterialFunctionInterface* GetParameterAssociatedFunction(const FHashedMaterialParameterInfo& ParameterInfo) const;
+#if 0
 	void GetParameterAssociatedFunctions(const FMaterialParameterInfo& ParameterInfo, TArray<UMaterialFunctionInterface*>& AssociatedFunctions) const;
+#endif
 
 	//~ Begin UMaterialExpression Interface
 #if WITH_EDITOR
@@ -122,7 +132,7 @@ class UMaterialExpressionMaterialAttributeLayers : public UMaterialExpression
 	//~ End UMaterialExpression Interface
 	
 	/** Return whether this is the named parameter, and fill in its value */
-	bool IsNamedParameter(const FMaterialParameterInfo& ParameterInfo, FMaterialLayersFunctions& OutLayers, FGuid& OutExpressionGuid) const;
+	bool IsNamedParameter(const FHashedMaterialParameterInfo& ParameterInfo, FMaterialLayersFunctions& OutLayers, FGuid& OutExpressionGuid) const;
 	
 	ENGINE_API virtual FGuid& GetParameterExpressionId() override
 	{
diff --git a/Engine/Source/Runtime/Engine/Classes/Materials/MaterialExpressionRuntimeVirtualTextureSampleParameter.h b/Engine/Source/Runtime/Engine/Classes/Materials/MaterialExpressionRuntimeVirtualTextureSampleParameter.h
--- a/Engine/Source/Runtime/Engine/Classes/Materials/MaterialExpressionRuntimeVirtualTextureSampleParameter.h
+++ b/Engine/Source/Runtime/Engine/Classes/Materials/MaterialExpressionRuntimeVirtualTextureSampleParameter.h
@@ -35,7 +35,7 @@ class UMaterialExpressionRuntimeVirtualTextureSampleParameter : public UMaterial
 #endif
 
 	/** Return whether this is the named parameter from this material expression, and if it is then return its value. */
-	bool IsNamedParameter(const FMaterialParameterInfo& ParameterInfo, URuntimeVirtualTexture*& OutValue) const;
+	bool IsNamedParameter(const FHashedMaterialParameterInfo& ParameterInfo, URuntimeVirtualTexture*& OutValue) const;
 	/** Adds to arrays of parameter info and id with the values used by this material expression. */
 	void GetAllParameterInfo(TArray<FMaterialParameterInfo> &OutParameterInfo, TArray<FGuid> &OutParameterIds, const FMaterialParameterInfo& InBaseParameterInfo) const;
 
diff --git a/Engine/Source/Runtime/Engine/Classes/Materials/MaterialExpressionScalarParameter.h b/Engine/Source/Runtime/Engine/Classes/Materials/MaterialExpressionScalarParameter.h
--- a/Engine/Source/Runtime/Engine/Classes/Materials/MaterialExpressionScalarParameter.h
+++ b/Engine/Source/Runtime/Engine/Classes/Materials/MaterialExpressionScalarParameter.h
@@ -47,7 +47,7 @@ class UMaterialExpressionScalarParameter : public UMaterialExpressionParameter
 	//~ End UMaterialExpression Interface
 
 	/** Return whether this is the named parameter, and fill in its value */
-	bool IsNamedParameter(const FMaterialParameterInfo& ParameterInfo, float& OutValue) const;
+	bool IsNamedParameter(const FHashedMaterialParameterInfo& ParameterInfo, float& OutValue) const;
 
 #if WITH_EDITOR
 	bool SetParameterValue(FName InParameterName, float InValue);
diff --git a/Engine/Source/Runtime/Engine/Classes/Materials/MaterialExpressionStaticBoolParameter.h b/Engine/Source/Runtime/Engine/Classes/Materials/MaterialExpressionStaticBoolParameter.h
--- a/Engine/Source/Runtime/Engine/Classes/Materials/MaterialExpressionStaticBoolParameter.h
+++ b/Engine/Source/Runtime/Engine/Classes/Materials/MaterialExpressionStaticBoolParameter.h
@@ -28,7 +28,7 @@ public:
 	//~ End UMaterialExpression Interface
 
 	/** Return whether this is the named parameter, and fill in its value */
-	bool IsNamedParameter(const FMaterialParameterInfo& ParameterInfo, bool& OutValue, FGuid& OutExpressionGuid) const;
+	bool IsNamedParameter(const FHashedMaterialParameterInfo& ParameterInfo, bool& OutValue, FGuid& OutExpressionGuid) const;
 
 #if WITH_EDITOR
 	bool SetParameterValue(FName InParameterName, bool OutValue, FGuid InExpressionGuid);
diff --git a/Engine/Source/Runtime/Engine/Classes/Materials/MaterialExpressionStaticComponentMaskParameter.h b/Engine/Source/Runtime/Engine/Classes/Materials/MaterialExpressionStaticComponentMaskParameter.h
--- a/Engine/Source/Runtime/Engine/Classes/Materials/MaterialExpressionStaticComponentMaskParameter.h
+++ b/Engine/Source/Runtime/Engine/Classes/Materials/MaterialExpressionStaticComponentMaskParameter.h
@@ -41,7 +41,7 @@ public:
 	//~ End UMaterialExpression Interface
 
 	/** Return whether this is the named parameter, and fill in its value */
-	bool IsNamedParameter(const FMaterialParameterInfo& ParameterInfo, bool& OutR, bool& OutG, bool& OutB, bool& OutA, FGuid&OutExpressionGuid) const;
+	bool IsNamedParameter(const FHashedMaterialParameterInfo& ParameterInfo, bool& OutR, bool& OutG, bool& OutB, bool& OutA, FGuid&OutExpressionGuid) const;
 
 #if WITH_EDITOR
 	bool SetParameterValue(FName InParameterName, bool InR, bool InG, bool InB, bool InA, FGuid InExpressionGuid);
diff --git a/Engine/Source/Runtime/Engine/Classes/Materials/MaterialExpressionTextureSampleParameter.h b/Engine/Source/Runtime/Engine/Classes/Materials/MaterialExpressionTextureSampleParameter.h
--- a/Engine/Source/Runtime/Engine/Classes/Materials/MaterialExpressionTextureSampleParameter.h
+++ b/Engine/Source/Runtime/Engine/Classes/Materials/MaterialExpressionTextureSampleParameter.h
@@ -55,7 +55,7 @@ class ENGINE_API UMaterialExpressionTextureSampleParameter : public UMaterialExp
 	//~ End UMaterialExpression Interface
 
 	/** Return whether this is the named parameter, and fill in its value */
-	bool IsNamedParameter(const FMaterialParameterInfo& ParameterInfo, UTexture*& OutValue) const;
+	bool IsNamedParameter(const FHashedMaterialParameterInfo& ParameterInfo, UTexture*& OutValue) const;
 
 #if WITH_EDITOR
 	bool SetParameterValue(FName InParameterName, UTexture* InValue);
diff --git a/Engine/Source/Runtime/Engine/Classes/Materials/MaterialExpressionVectorParameter.h b/Engine/Source/Runtime/Engine/Classes/Materials/MaterialExpressionVectorParameter.h
--- a/Engine/Source/Runtime/Engine/Classes/Materials/MaterialExpressionVectorParameter.h
+++ b/Engine/Source/Runtime/Engine/Classes/Materials/MaterialExpressionVectorParameter.h
@@ -36,7 +36,7 @@ class UMaterialExpressionVectorParameter : public UMaterialExpressionParameter
 	//~ End UMaterialExpression Interface
 
 	/** Return whether this is the named parameter, and fill in its value */
-	bool IsNamedParameter(const FMaterialParameterInfo& ParameterInfo, FLinearColor& OutValue) const;
+	bool IsNamedParameter(const FHashedMaterialParameterInfo& ParameterInfo, FLinearColor& OutValue) const;
 
 #if WITH_EDITOR
 	virtual bool SetParameterValue(FName InParameterName, FLinearColor InValue);
diff --git a/Engine/Source/Runtime/Engine/Classes/Materials/MaterialFunction.h b/Engine/Source/Runtime/Engine/Classes/Materials/MaterialFunction.h
--- a/Engine/Source/Runtime/Engine/Classes/Materials/MaterialFunction.h
+++ b/Engine/Source/Runtime/Engine/Classes/Materials/MaterialFunction.h
@@ -34,9 +34,11 @@ class UMaterialFunction : public UMaterialFunctionInterface
 	UPROPERTY(EditAnywhere, Category=MaterialFunction, AssetRegistrySearchable)
 	FString Description;
 
+#if WITH_EDITORONLY_DATA
 	/** Array of material expressions, excluding Comments.  Used by the material editor. */
 	UPROPERTY()
 	TArray<UMaterialExpression*> FunctionExpressions;
+#endif // WITH_EDITORONLY_DATA
 
 	/** Whether to list this function in the material function library, which is a window in the material editor that lists categorized functions. */
 	UPROPERTY(EditAnywhere, Category=MaterialFunction, AssetRegistrySearchable)
@@ -112,6 +114,7 @@ public:
 	virtual void UnlinkFromCaller() override;
 #endif
 
+#if WITH_EDITORONLY_DATA
 	/** @return true if this function is dependent on the passed in function, directly or indirectly. */
 	virtual bool IsDependent(UMaterialFunctionInterface* OtherFunction) override;
 
@@ -126,9 +129,7 @@ public:
 
 	/** Returns an array of the functions that this function is dependent on, directly or indirectly. */
 	ENGINE_API virtual void GetDependentFunctions(TArray<UMaterialFunctionInterface*>& DependentFunctions) const override;
-
-	/** Appends textures referenced by the expressions in this function. */
-	virtual void AppendReferencedTextures(TArray<UObject*>& InOutTextures) const override;
+#endif // WITH_EDITORONLY_DATA
 
 #if WITH_EDITOR
 	virtual UMaterialInterface* GetPreviewMaterial() override;
@@ -145,7 +146,9 @@ public:
 
 	virtual UMaterialFunctionInterface* GetBaseFunction() override { return this; }
 	virtual const UMaterialFunctionInterface* GetBaseFunction() const override { return this; }
+#if WITH_EDITORONLY_DATA
 	virtual const TArray<UMaterialExpression*>* GetFunctionExpressions() const override { return &FunctionExpressions; }
+#endif
 	virtual const FString* GetDescription() const override { return &Description; }
 
 	virtual bool GetReentrantFlag() const override { return bReentrantFlag; }
diff --git a/Engine/Source/Runtime/Engine/Classes/Materials/MaterialFunctionInstance.h b/Engine/Source/Runtime/Engine/Classes/Materials/MaterialFunctionInstance.h
--- a/Engine/Source/Runtime/Engine/Classes/Materials/MaterialFunctionInstance.h
+++ b/Engine/Source/Runtime/Engine/Classes/Materials/MaterialFunctionInstance.h
@@ -65,10 +65,10 @@ class UMaterialFunctionInstance : public UMaterialFunctionInterface
 	UPROPERTY(EditAnywhere, Category = MaterialFunctionInstance)
 	TArray<struct FRuntimeVirtualTextureParameterValue> RuntimeVirtualTextureParameterValues;
 
-	ENGINE_API void UpdateParameterSet();
 #if WITH_EDITOR
+	ENGINE_API void UpdateParameterSet();
 	ENGINE_API void OverrideMaterialInstanceParameterValues(class UMaterialInstance* Instance);
-#endif
+#endif // WITH_EDITOR
 
 	//~ Begin UMaterialFunctionInterface interface
 	virtual void UpdateFromFunctionResource() override;
@@ -83,9 +83,11 @@ class UMaterialFunctionInstance : public UMaterialFunctionInterface
 #endif
 
 	virtual bool IsDependent(UMaterialFunctionInterface* OtherFunction) override;
+
+#if WITH_EDITORONLY_DATA
 	ENGINE_API virtual bool IterateDependentFunctions(TFunctionRef<bool(UMaterialFunctionInterface*)> Predicate) const override;
 	ENGINE_API virtual void GetDependentFunctions(TArray<UMaterialFunctionInterface*>& DependentFunctions) const override;
-	virtual void AppendReferencedTextures(TArray<UObject*>& InOutTextures) const override;
+#endif
 
 #if WITH_EDITOR
 	virtual UMaterialInterface* GetPreviewMaterial() override;
@@ -135,11 +137,13 @@ class UMaterialFunctionInstance : public UMaterialFunctionInterface
 		return BasePtr;
 	}
 
+#if WITH_EDITORONLY_DATA
 	virtual const TArray<UMaterialExpression*>* GetFunctionExpressions() const override
 	{
 		const UMaterialFunctionInterface* BaseFunction = GetBaseFunction();
 		return BaseFunction ? BaseFunction->GetFunctionExpressions() : nullptr;
 	}
+#endif // WITH_EDITORONLY_DATA
 
 	virtual const FString* GetDescription() const override
 	{
@@ -162,13 +166,13 @@ class UMaterialFunctionInstance : public UMaterialFunctionInterface
 	}
 
 public:
-	virtual bool OverrideNamedScalarParameter(const FMaterialParameterInfo& ParameterInfo, float& OutValue) override;
-	virtual bool OverrideNamedVectorParameter(const FMaterialParameterInfo& ParameterInfo, FLinearColor& OutValue) override;
-	virtual bool OverrideNamedTextureParameter(const FMaterialParameterInfo& ParameterInfo, class UTexture*& OutValue) override;
-	virtual bool OverrideNamedRuntimeVirtualTextureParameter(const FMaterialParameterInfo& ParameterInfo, class URuntimeVirtualTexture*& OutValue) override;
-	virtual bool OverrideNamedFontParameter(const FMaterialParameterInfo& ParameterInfo, class UFont*& OutFontValue, int32& OutFontPage) override;
-	virtual bool OverrideNamedStaticSwitchParameter(const FMaterialParameterInfo& ParameterInfo, bool& OutValue, FGuid& OutExpressionGuid) override;
-	virtual bool OverrideNamedStaticComponentMaskParameter(const FMaterialParameterInfo& ParameterInfo, bool& OutR, bool& OutG, bool& OutB, bool& OutA, FGuid& OutExpressionGuid) override;
+	virtual bool OverrideNamedScalarParameter(const FHashedMaterialParameterInfo& ParameterInfo, float& OutValue) override;
+	virtual bool OverrideNamedVectorParameter(const FHashedMaterialParameterInfo& ParameterInfo, FLinearColor& OutValue) override;
+	virtual bool OverrideNamedTextureParameter(const FHashedMaterialParameterInfo& ParameterInfo, class UTexture*& OutValue) override;
+	virtual bool OverrideNamedRuntimeVirtualTextureParameter(const FHashedMaterialParameterInfo& ParameterInfo, class URuntimeVirtualTexture*& OutValue) override;
+	virtual bool OverrideNamedFontParameter(const FHashedMaterialParameterInfo& ParameterInfo, class UFont*& OutFontValue, int32& OutFontPage) override;
+	virtual bool OverrideNamedStaticSwitchParameter(const FHashedMaterialParameterInfo& ParameterInfo, bool& OutValue, FGuid& OutExpressionGuid) override;
+	virtual bool OverrideNamedStaticComponentMaskParameter(const FHashedMaterialParameterInfo& ParameterInfo, bool& OutR, bool& OutG, bool& OutB, bool& OutA, FGuid& OutExpressionGuid) override;
 	//~ End UMaterialFunctionInterface interface
 
 protected:
diff --git a/Engine/Source/Runtime/Engine/Classes/Materials/MaterialFunctionInterface.h b/Engine/Source/Runtime/Engine/Classes/Materials/MaterialFunctionInterface.h
--- a/Engine/Source/Runtime/Engine/Classes/Materials/MaterialFunctionInterface.h
+++ b/Engine/Source/Runtime/Engine/Classes/Materials/MaterialFunctionInterface.h
@@ -76,6 +76,7 @@ public:
 	ENGINE_API virtual bool IsDependent(UMaterialFunctionInterface* OtherFunction)
 		PURE_VIRTUAL(UMaterialFunctionInterface::IsDependent,return false;);
 
+#if WITH_EDITORONLY_DATA
 	/**
 	 * Iterates all functions that this function is dependent on, directly or indrectly.
 	 *
@@ -89,10 +90,7 @@ public:
 	/** Returns an array of the functions that this function is dependent on, directly or indirectly. */
 	ENGINE_API virtual void GetDependentFunctions(TArray<UMaterialFunctionInterface*>& DependentFunctions) const
 		PURE_VIRTUAL(UMaterialFunctionInterface::GetDependentFunctions,);
-
-	/** Appends textures referenced by the expressions in this function. */
-	ENGINE_API virtual void AppendReferencedTextures(TArray<UObject*>& InOutTextures) const
-		PURE_VIRTUAL(UMaterialFunctionInterface::AppendReferencedTextures,);
+#endif
 
 #if WITH_EDITOR
 	ENGINE_API virtual UMaterialInterface* GetPreviewMaterial()
@@ -124,8 +122,10 @@ public:
 	virtual const UMaterialFunctionInterface* GetBaseFunction() const
 		PURE_VIRTUAL(UMaterialFunctionInterface::GetBaseFunction,return nullptr;);
 
+#if WITH_EDITORONLY_DATA
 	virtual const TArray<UMaterialExpression*>* GetFunctionExpressions() const
 		PURE_VIRTUAL(UMaterialFunctionInterface::GetFunctionExpressions,return nullptr;);
+#endif // WITH_EDITORONLY_DATA
 
 	virtual const FString* GetDescription() const
 		PURE_VIRTUAL(UMaterialFunctionInterface::GetDescription,return nullptr;);
@@ -137,6 +137,7 @@ public:
 		PURE_VIRTUAL(UMaterialFunctionInterface::SetReentrantFlag,);
 
 public:
+#if WITH_EDITORONLY_DATA
 	/** Finds the names of all matching type parameters */
 	template<typename ExpressionType>
 	void GetAllParameterInfo(TArray<FMaterialParameterInfo>& OutParameterInfo, TArray<FGuid>& OutParameterIds, const FMaterialParameterInfo& InBaseParameterInfo) const
@@ -162,10 +163,12 @@ public:
 			check(OutParameterInfo.Num() == OutParameterIds.Num());
 		}
 	}
+#endif // WITH_EDITORONLY_DATA
 
+#if WITH_EDITOR
 	/** Finds the first matching parameter by name and type */
 	template<typename ExpressionType>
-	bool GetNamedParameterOfType(const FMaterialParameterInfo& ParameterInfo, ExpressionType*& Parameter, UMaterialFunctionInterface** OwningFunction = nullptr)
+	bool GetNamedParameterOfType(const FHashedMaterialParameterInfo& ParameterInfo, ExpressionType*& Parameter, UMaterialFunctionInterface** OwningFunction = nullptr)
 	{
 		Parameter = nullptr;
 
@@ -207,9 +210,8 @@ public:
 		return false;
 	}
 
-#if WITH_EDITOR
 	/** Finds the first matching parameter's group name */
-	bool GetParameterGroupName(const FMaterialParameterInfo& ParameterInfo, FName& OutGroup)
+	bool GetParameterGroupName(const FHashedMaterialParameterInfo& ParameterInfo, FName& OutGroup)
 	{
 		if (UMaterialFunctionInterface* ParameterFunction = GetBaseFunction())
 		{
@@ -253,7 +255,7 @@ public:
 	}
 
 	/** Finds the first matching parameter's group name */
-	bool GetParameterSortPriority(const FMaterialParameterInfo& ParameterInfo, int32& OutSortPriority)
+	bool GetParameterSortPriority(const FHashedMaterialParameterInfo& ParameterInfo, int32& OutSortPriority)
 	{
 		if (UMaterialFunctionInterface* ParameterFunction = GetBaseFunction())
 		{
@@ -297,7 +299,7 @@ public:
 	}
 
 	/** Finds the first matching parameter's description */
-	bool GetParameterDesc(const FMaterialParameterInfo& ParameterInfo, FString& OutDesc)
+	bool GetParameterDesc(const FHashedMaterialParameterInfo& ParameterInfo, FString& OutDesc)
 	{
 		if (UMaterialFunctionInterface* ParameterFunction = GetBaseFunction())
 		{
@@ -339,7 +341,6 @@ public:
 
 		return false;
 	}
-#endif // WITH_EDITOR
 
 	/** Returns if any of the matching parameters have changed */
 	template <typename ParameterType, typename ExpressionType>
@@ -372,7 +373,9 @@ public:
 
 		return bChanged;
 	}
+#endif // WITH_EDITOR
 
+#if WITH_EDITORONLY_DATA
 	/** Get all expressions of the requested type, recursing through any function expressions in the function */
 	template<typename ExpressionType>
 	bool HasAnyExpressionsOfType()
@@ -423,38 +426,39 @@ public:
 			}
 		}
 	}
+#endif // WITH_EDITORONLY_DATA
 
-	virtual bool OverrideNamedScalarParameter(const FMaterialParameterInfo& ParameterInfo, float& OutValue)
+	virtual bool OverrideNamedScalarParameter(const FHashedMaterialParameterInfo& ParameterInfo, float& OutValue)
 	{
 		return false;
 	}
 
-	virtual bool OverrideNamedVectorParameter(const FMaterialParameterInfo& ParameterInfo, FLinearColor& OutValue)
+	virtual bool OverrideNamedVectorParameter(const FHashedMaterialParameterInfo& ParameterInfo, FLinearColor& OutValue)
 	{
 		return false;
 	}
 
-	virtual bool OverrideNamedTextureParameter(const FMaterialParameterInfo& ParameterInfo, class UTexture*& OutValue)
+	virtual bool OverrideNamedTextureParameter(const FHashedMaterialParameterInfo& ParameterInfo, class UTexture*& OutValue)
 	{
 		return false;
 	}
 	
-	virtual bool OverrideNamedRuntimeVirtualTextureParameter(const FMaterialParameterInfo& ParameterInfo, class URuntimeVirtualTexture*& OutValue)
+	virtual bool OverrideNamedRuntimeVirtualTextureParameter(const FHashedMaterialParameterInfo& ParameterInfo, class URuntimeVirtualTexture*& OutValue)
 	{
 		return false;
 	}
 
-	virtual bool OverrideNamedFontParameter(const FMaterialParameterInfo& ParameterInfo, class UFont*& OutFontValue, int32& OutFontPage)
+	virtual bool OverrideNamedFontParameter(const FHashedMaterialParameterInfo& ParameterInfo, class UFont*& OutFontValue, int32& OutFontPage)
 	{
 		return false;
 	}
 
-	virtual bool OverrideNamedStaticSwitchParameter(const FMaterialParameterInfo& ParameterInfo, bool& OutValue, FGuid& OutExpressionGuid)
+	virtual bool OverrideNamedStaticSwitchParameter(const FHashedMaterialParameterInfo& ParameterInfo, bool& OutValue, FGuid& OutExpressionGuid)
 	{
 		return false;
 	}
 
-	virtual bool OverrideNamedStaticComponentMaskParameter(const FMaterialParameterInfo& ParameterInfo, bool& OutR, bool& OutG, bool& OutB, bool& OutA, FGuid& OutExpressionGuid)
+	virtual bool OverrideNamedStaticComponentMaskParameter(const FHashedMaterialParameterInfo& ParameterInfo, bool& OutR, bool& OutG, bool& OutB, bool& OutA, FGuid& OutExpressionGuid)
 	{
 		return false;
 	}
diff --git a/Engine/Source/Runtime/Engine/Classes/Materials/MaterialInstance.h b/Engine/Source/Runtime/Engine/Classes/Materials/MaterialInstance.h
--- a/Engine/Source/Runtime/Engine/Classes/Materials/MaterialInstance.h
+++ b/Engine/Source/Runtime/Engine/Classes/Materials/MaterialInstance.h
@@ -9,6 +9,7 @@
 #include "Materials/MaterialInterface.h"
 #include "StaticParameterSet.h"
 #include "MaterialShared.h"
+#include "MaterialCachedData.h"
 #include "Materials/MaterialExpressionMaterialFunctionCall.h"
 #include "Materials/Material.h"
 #include "Materials/MaterialInstanceBasePropertyOverrides.h"
@@ -325,11 +326,13 @@ class ENGINE_VTABLE UMaterialInstance : public UMaterialInterface
 	 */
 	DECLARE_DELEGATE_RetVal_TwoParams(bool, FCustomParameterSetUpdaterDelegate, FStaticParameterSet&, UMaterial*);
 
+#if WITH_EDITORONLY_DATA
 	// Custom static parameters getter delegate.
 	ENGINE_API static FCustomStaticParametersGetterDelegate CustomStaticParametersGetters;
 
 	// An array of custom parameter set updaters.
 	ENGINE_API static TArray<FCustomParameterSetUpdaterDelegate> CustomParameterSetUpdaters;
+#endif // WITH_EDITORONLY_DATA
 
 	/**
 	 * Gets static parameter set for this material.
@@ -413,13 +416,6 @@ class ENGINE_VTABLE UMaterialInstance : public UMaterialInterface
 	 */
 	class FMaterialInstanceResource* Resource;
 
-	/** 
-	 * Cached texture references from all expressions in the material (including nested functions). 
-	 * This is used to link uniform texture expressions which were stored in the DDC with the UTextures that they reference.
-	 */
-	UPROPERTY(transient)
-	TArray<UObject*> PermutationTextureReferences;
-
 	ENGINE_API virtual void PreSave(const class ITargetPlatform* TargetPlatform) override;
 	ENGINE_API virtual float GetTextureDensity(FName TextureName, const struct FMeshUVChannelInfo& UVChannelData) const override;
 
@@ -436,6 +432,16 @@ private:
 	UPROPERTY()
 	FStaticParameterSet StaticParameters;
 
+	UPROPERTY()
+	FMaterialCachedParameters CachedLayerParameters;
+
+	/**
+	 * Cached texture references from all expressions in the material (including nested functions).
+	 * This is used to link uniform texture expressions which were stored in the DDC with the UTextures that they reference.
+	 */
+	UPROPERTY()
+	TArray<UObject*> CachedReferencedTextures;
+
 #ifdef WITH_EDITOR
 	mutable TOptional<FStaticParameterSet> CachedStaticParameterValues;
 	mutable uint8 AllowCachingStaticParameterValuesCounter = 0;
@@ -463,38 +469,43 @@ public:
 	// Begin UMaterialInterface interface.
 	virtual ENGINE_API UMaterial* GetMaterial() override;
 	virtual ENGINE_API const UMaterial* GetMaterial() const override;
-	virtual ENGINE_API const UMaterial* GetMaterial_Concurrent(TMicRecursionGuard& RecursionGuard) const override;
+	virtual ENGINE_API const UMaterial* GetMaterial_Concurrent(TMicRecursionGuard RecursionGuard = TMicRecursionGuard()) const override;
 	virtual ENGINE_API FMaterialResource* AllocatePermutationResource();
 	virtual ENGINE_API FMaterialResource* GetMaterialResource(ERHIFeatureLevel::Type InFeatureLevel, EMaterialQualityLevel::Type QualityLevel = EMaterialQualityLevel::Num) override;
 	virtual ENGINE_API const FMaterialResource* GetMaterialResource(ERHIFeatureLevel::Type InFeatureLevel, EMaterialQualityLevel::Type QualityLevel = EMaterialQualityLevel::Num) const override;
+
 #if WITH_EDITOR
-	virtual ENGINE_API bool GetScalarParameterSliderMinMax(const FMaterialParameterInfo& ParameterInfo, float& OutSliderMin, float& OutSliderMax) const override;
+	virtual ENGINE_API bool GetScalarParameterSliderMinMax(const FHashedMaterialParameterInfo& ParameterInfo, float& OutSliderMin, float& OutSliderMax) const override;
 #endif
-	virtual ENGINE_API bool GetScalarParameterValue(const FMaterialParameterInfo& ParameterInfo, float& OutValue, bool bOveriddenOnly = false) const override;
-	virtual ENGINE_API bool IsScalarParameterUsedAsAtlasPosition(const FMaterialParameterInfo& ParameterInfo, bool& OutValue, TSoftObjectPtr<class UCurveLinearColor>& Curve, TSoftObjectPtr<class UCurveLinearColorAtlas>& Atlas) const override;
-	virtual ENGINE_API bool GetVectorParameterValue(const FMaterialParameterInfo& ParameterInfo, FLinearColor& OutValue, bool bOveriddenOnly = false) const override;
-	virtual ENGINE_API bool IsVectorParameterUsedAsChannelMask(const FMaterialParameterInfo& ParameterInfo, bool& OutValue) const override;
+	virtual ENGINE_API bool GetScalarParameterValue(const FHashedMaterialParameterInfo& ParameterInfo, float& OutValue, bool bOveriddenOnly = false) const override;
 #if WITH_EDITOR
-	virtual ENGINE_API bool GetVectorParameterChannelNames(const FMaterialParameterInfo& ParameterInfo, FParameterChannelNames& OutValue) const override;
+	virtual ENGINE_API bool IsScalarParameterUsedAsAtlasPosition(const FHashedMaterialParameterInfo& ParameterInfo, bool& OutValue, TSoftObjectPtr<class UCurveLinearColor>& Curve, TSoftObjectPtr<class UCurveLinearColorAtlas>& Atlas) const override;
 #endif
-	virtual ENGINE_API bool GetTextureParameterValue(const FMaterialParameterInfo& ParameterInfo, class UTexture*& OutValue, bool bOveriddenOnly = false) const override;
-	virtual ENGINE_API bool GetRuntimeVirtualTextureParameterValue(const FMaterialParameterInfo& ParameterInfo, class URuntimeVirtualTexture*& OutValue, bool bOveriddenOnly = false) const override;
+	virtual ENGINE_API bool GetVectorParameterValue(const FHashedMaterialParameterInfo& ParameterInfo, FLinearColor& OutValue, bool bOveriddenOnly = false) const override;
 #if WITH_EDITOR
-	virtual ENGINE_API bool GetTextureParameterChannelNames(const FMaterialParameterInfo& ParameterInfo, FParameterChannelNames& OutValue) const override;
+	virtual ENGINE_API bool IsVectorParameterUsedAsChannelMask(const FHashedMaterialParameterInfo& ParameterInfo, bool& OutValue) const override;
+	virtual ENGINE_API bool GetVectorParameterChannelNames(const FHashedMaterialParameterInfo& ParameterInfo, FParameterChannelNames& OutValue) const override;
 #endif
-	virtual ENGINE_API bool GetFontParameterValue(const FMaterialParameterInfo& ParameterInfo, class UFont*& OutFontValue, int32& OutFontPage, bool bOveriddenOnly = false) const override;
+	virtual ENGINE_API bool GetTextureParameterValue(const FHashedMaterialParameterInfo& ParameterInfo, class UTexture*& OutValue, bool bOveriddenOnly = false) const override;
+	virtual ENGINE_API bool GetRuntimeVirtualTextureParameterValue(const FHashedMaterialParameterInfo& ParameterInfo, class URuntimeVirtualTexture*& OutValue, bool bOveriddenOnly = false) const override;
+#if WITH_EDITOR
+	virtual ENGINE_API bool GetTextureParameterChannelNames(const FHashedMaterialParameterInfo& ParameterInfo, FParameterChannelNames& OutValue) const override;
+#endif
+	virtual ENGINE_API bool GetFontParameterValue(const FHashedMaterialParameterInfo& ParameterInfo, class UFont*& OutFontValue, int32& OutFontPage, bool bOveriddenOnly = false) const override;
 	virtual ENGINE_API void GetUsedTextures(TArray<UTexture*>& OutTextures, EMaterialQualityLevel::Type QualityLevel, bool bAllQualityLevels, ERHIFeatureLevel::Type FeatureLevel, bool bAllFeatureLevels) const override;
 	virtual ENGINE_API void GetUsedTexturesAndIndices(TArray<UTexture*>& OutTextures, TArray< TArray<int32> >& OutIndices, EMaterialQualityLevel::Type QualityLevel, ERHIFeatureLevel::Type FeatureLevel) const;
 	virtual ENGINE_API void OverrideTexture(const UTexture* InTextureToOverride, UTexture* OverrideTexture, ERHIFeatureLevel::Type InFeatureLevel) override;
-	virtual ENGINE_API void OverrideVectorParameterDefault(const FMaterialParameterInfo& ParameterInfo, const FLinearColor& Value, bool bOverride, ERHIFeatureLevel::Type FeatureLevel) override;
-	virtual ENGINE_API void OverrideScalarParameterDefault(const FMaterialParameterInfo& ParameterInfo, float Value, bool bOverride, ERHIFeatureLevel::Type FeatureLevel) override;
+	virtual ENGINE_API void OverrideVectorParameterDefault(const FHashedMaterialParameterInfo& ParameterInfo, const FLinearColor& Value, bool bOverride, ERHIFeatureLevel::Type FeatureLevel) override;
+	virtual ENGINE_API void OverrideScalarParameterDefault(const FHashedMaterialParameterInfo& ParameterInfo, float Value, bool bOverride, ERHIFeatureLevel::Type FeatureLevel) override;
 	virtual ENGINE_API bool CheckMaterialUsage(const EMaterialUsage Usage) override;
 	virtual ENGINE_API bool CheckMaterialUsage_Concurrent(const EMaterialUsage Usage) const override;
+#if WITH_EDITORONLY_DATA
 	virtual ENGINE_API bool GetStaticSwitchParameterValues(FStaticParamEvaluationContext& EvalContext, TBitArray<>& OutValues, FGuid* OutExpressionGuids, bool bCheckParent = true) const override;
 	virtual ENGINE_API bool GetStaticComponentMaskParameterValues(FStaticParamEvaluationContext& EvalContext, TBitArray<>& OutRGBAOrderedValues, FGuid* OutExpressionGuids, bool bCheckParent = true) const override;
-	virtual ENGINE_API bool GetTerrainLayerWeightParameterValue(const FMaterialParameterInfo& ParameterInfo, int32& OutWeightmapIndex, FGuid &OutExpressionGuid) const override;
-	virtual ENGINE_API bool GetMaterialLayersParameterValue(const FMaterialParameterInfo& ParameterInfo, FMaterialLayersFunctions& OutLayers, FGuid &OutExpressionGuid, bool bCheckParent = true) const override;
-			ENGINE_API bool UpdateMaterialLayersParameterValue(const FMaterialParameterInfo& ParameterInfo, const FMaterialLayersFunctions& LayersValue, const bool bOverridden, const FGuid& GUID);
+	virtual ENGINE_API bool GetMaterialLayersParameterValue(const FHashedMaterialParameterInfo& ParameterInfo, FMaterialLayersFunctions& OutLayers, FGuid &OutExpressionGuid, bool bCheckParent = true) const override;
+#endif // WITH_EDITORONLY_DATA
+	virtual ENGINE_API bool GetTerrainLayerWeightParameterValue(const FHashedMaterialParameterInfo& ParameterInfo, int32& OutWeightmapIndex, FGuid &OutExpressionGuid) const override;
+			ENGINE_API bool UpdateMaterialLayersParameterValue(const FHashedMaterialParameterInfo& ParameterInfo, const FMaterialLayersFunctions& LayersValue, const bool bOverridden, const FGuid& GUID);
 	virtual ENGINE_API bool IsDependent(UMaterialInterface* TestDependency) override;
 	virtual ENGINE_API FMaterialRenderProxy* GetRenderProxy() const override;
 	virtual ENGINE_API UPhysicalMaterial* GetPhysicalMaterial() const override;
@@ -506,8 +517,8 @@ public:
 	virtual ENGINE_API float GetDiffuseBoost() const override;
 	virtual ENGINE_API float GetExportResolutionScale() const override;
 #if WITH_EDITOR
-	virtual ENGINE_API bool GetParameterDesc(const FMaterialParameterInfo& ParameterInfo, FString& OutDesc, const TArray<struct FStaticMaterialLayersParameter>* MaterialLayersParameters = nullptr) const;	
-	virtual ENGINE_API bool GetParameterSortPriority(const FMaterialParameterInfo& ParameterInfo, int32& OutSortPriority, const TArray<struct FStaticMaterialLayersParameter>* MaterialLayersParameters = nullptr) const override;
+	virtual ENGINE_API bool GetParameterDesc(const FHashedMaterialParameterInfo& ParameterInfo, FString& OutDesc, const TArray<struct FStaticMaterialLayersParameter>* MaterialLayersParameters = nullptr) const override;
+	virtual ENGINE_API bool GetParameterSortPriority(const FHashedMaterialParameterInfo& ParameterInfo, int32& OutSortPriority, const TArray<struct FStaticMaterialLayersParameter>* MaterialLayersParameters = nullptr) const override;
 	virtual ENGINE_API bool GetGroupSortPriority(const FString& InGroupName, int32& OutSortPriority) const override;
 	virtual ENGINE_API bool GetTexturesInPropertyChain(EMaterialProperty InProperty, TArray<UTexture*>& OutTextures,
 		TArray<FName>* OutTextureParamNames, struct FStaticParameterSet* InStaticParameterSet,
@@ -597,6 +608,7 @@ public:
 	 */
 	void GetAllShaderMaps(TArray<FMaterialShaderMap*>& OutShaderMaps);
 
+#if WITH_EDITORONLY_DATA
 	/**
 	 * Builds a composited set of static parameters, including inherited and overridden values
 	 */
@@ -613,33 +625,42 @@ public:
 			Material->GetAllParameterInfo<ExpressionType>(OutParameterInfo, OutParameterIds);
 		}
 	}
+#endif // WITH_EDITORONLY_DATA
+
+	void GetAllParametersOfType(EMaterialParameterType Type, TArray<FMaterialParameterInfo>& OutParameterInfo, TArray<FGuid>& OutParameterIds) const;
 
 	ENGINE_API virtual void GetAllScalarParameterInfo(TArray<FMaterialParameterInfo>& OutParameterInfo, TArray<FGuid>& OutParameterIds) const override;
 	ENGINE_API virtual void GetAllVectorParameterInfo(TArray<FMaterialParameterInfo>& OutParameterInfo, TArray<FGuid>& OutParameterIds) const override;
 	ENGINE_API virtual void GetAllTextureParameterInfo(TArray<FMaterialParameterInfo>& OutParameterInfo, TArray<FGuid>& OutParameterIds) const override;
 	ENGINE_API virtual void GetAllRuntimeVirtualTextureParameterInfo(TArray<FMaterialParameterInfo>& OutParameterInfo, TArray<FGuid>& OutParameterIds) const override;
 	ENGINE_API virtual void GetAllFontParameterInfo(TArray<FMaterialParameterInfo>& OutParameterInfo, TArray<FGuid>& OutParameterIds) const override;
+
+#if WITH_EDITORONLY_DATA
 	ENGINE_API virtual void GetAllMaterialLayersParameterInfo(TArray<FMaterialParameterInfo>& OutParameterInfo, TArray<FGuid>& OutParameterIds) const override;
 	ENGINE_API virtual void GetAllStaticSwitchParameterInfo(TArray<FMaterialParameterInfo>& OutParameterInfo, TArray<FGuid>& OutParameterIds) const override;
 	ENGINE_API virtual void GetAllStaticComponentMaskParameterInfo(TArray<FMaterialParameterInfo>& OutParameterInfo, TArray<FGuid>& OutParameterIds) const override;
 
 	ENGINE_API virtual bool IterateDependentFunctions(TFunctionRef<bool(UMaterialFunctionInterface*)> Predicate) const override;
 	ENGINE_API virtual void GetDependentFunctions(TArray<class UMaterialFunctionInterface*>& DependentFunctions) const override;
-	
-	ENGINE_API virtual bool GetScalarParameterDefaultValue(const FMaterialParameterInfo& ParameterInfo, float& OutValue, bool bOveriddenOnly = false, bool bCheckOwnedGlobalOverrides = false) const override;
-	ENGINE_API virtual bool GetVectorParameterDefaultValue(const FMaterialParameterInfo& ParameterInfo, FLinearColor& OutValue, bool bOveriddenOnly = false, bool bCheckOwnedGlobalOverrides = false) const override;
-	ENGINE_API virtual bool GetTextureParameterDefaultValue(const FMaterialParameterInfo& ParameterInfo, UTexture*& OutValue, bool bCheckOwnedGlobalOverrides = false) const override;
-	ENGINE_API virtual bool GetRuntimeVirtualTextureParameterDefaultValue(const FMaterialParameterInfo& ParameterInfo, URuntimeVirtualTexture*& OutValue, bool bCheckOwnedGlobalOverrides = false) const override;
-	ENGINE_API virtual bool GetFontParameterDefaultValue(const FMaterialParameterInfo& ParameterInfo, UFont*& OutFontValue, int32& OutFontPage, bool bCheckOwnedGlobalOverrides = false) const override;
-	ENGINE_API virtual bool GetStaticSwitchParameterDefaultValue(const FMaterialParameterInfo& ParameterInfo, bool& OutValue, FGuid& OutExpressionGuid, bool bCheckOwnedGlobalOverrides = false) const override;
-	ENGINE_API virtual bool GetStaticComponentMaskParameterDefaultValue(const FMaterialParameterInfo& ParameterInfo, bool& OutR, bool& OutG, bool& OutB, bool& OutA, FGuid& OutExpressionGuid, bool bCheckOwnedGlobalOverrides = false) const override;
+#endif // WITH_EDITORONLY_DATA
 
+	ENGINE_API virtual bool GetScalarParameterDefaultValue(const FHashedMaterialParameterInfo& ParameterInfo, float& OutValue, bool bOveriddenOnly = false, bool bCheckOwnedGlobalOverrides = false) const override;
+	ENGINE_API virtual bool GetVectorParameterDefaultValue(const FHashedMaterialParameterInfo& ParameterInfo, FLinearColor& OutValue, bool bOveriddenOnly = false, bool bCheckOwnedGlobalOverrides = false) const override;
+	ENGINE_API virtual bool GetTextureParameterDefaultValue(const FHashedMaterialParameterInfo& ParameterInfo, UTexture*& OutValue, bool bCheckOwnedGlobalOverrides = false) const override;
+	ENGINE_API virtual bool GetRuntimeVirtualTextureParameterDefaultValue(const FHashedMaterialParameterInfo& ParameterInfo, URuntimeVirtualTexture*& OutValue, bool bCheckOwnedGlobalOverrides = false) const override;
+	ENGINE_API virtual bool GetFontParameterDefaultValue(const FHashedMaterialParameterInfo& ParameterInfo, UFont*& OutFontValue, int32& OutFontPage, bool bCheckOwnedGlobalOverrides = false) const override;
 #if WITH_EDITOR
-	ENGINE_API virtual bool GetGroupName(const FMaterialParameterInfo& ParameterInfo, FName& OutGroup) const override;
-#endif
+	ENGINE_API virtual bool GetStaticSwitchParameterDefaultValue(const FHashedMaterialParameterInfo& ParameterInfo, bool& OutValue, FGuid& OutExpressionGuid, bool bCheckOwnedGlobalOverrides = false) const override;
+	ENGINE_API virtual bool GetStaticComponentMaskParameterDefaultValue(const FHashedMaterialParameterInfo& ParameterInfo, bool& OutR, bool& OutG, bool& OutB, bool& OutA, FGuid& OutExpressionGuid, bool bCheckOwnedGlobalOverrides = false) const override;
+	ENGINE_API virtual bool GetGroupName(const FHashedMaterialParameterInfo& ParameterInfo, FName& OutGroup) const override;
+#endif // WITH_EDITOR
 
 	/** Appends textures referenced by expressions, including nested functions. */
-	ENGINE_API virtual void AppendReferencedTextures(TArray<UObject*>& InOutTextures) const override;
+	ENGINE_API virtual TArrayView<UObject* const> GetReferencedTextures() const override final { return CachedReferencedTextures; }
+
+#if WITH_EDITOR
+	ENGINE_API void UpdateCachedLayerParameters();
+#endif
 
 	void GetBasePropertyOverridesHash(FSHAHash& OutHash)const;
 	ENGINE_API virtual bool HasOverridenBaseProperties()const;
@@ -682,6 +703,14 @@ public:
 #if WITH_EDITOR
 	void BeginAllowCachingStaticParameterValues();
 	void EndAllowCachingStaticParameterValues();
+
+private:
+	bool GetScalarLayerParameterValue_Legacy(const FHashedMaterialParameterInfo& ParameterInfo, bool bOveriddenOnly, float& OutValue, bool& OutResult) const;
+	bool GetVectorLayerParameterValue_Legacy(const FHashedMaterialParameterInfo& ParameterInfo, bool bOveriddenOnly, FLinearColor& OutValue, bool& OutResult) const;
+	bool GetTextureLayerParameterValue_Legacy(const FHashedMaterialParameterInfo& ParameterInfo, bool bOveriddenOnly, UTexture*& OutValue, bool& OutResult) const;
+	bool GetFontLayerParameterValue_Legacy(const FHashedMaterialParameterInfo& ParameterInfo, bool bOveriddenOnly, UFont*& OutValue, int32& OutFontPage, bool& OutResult) const;
+	bool GetRuntimeVirtualTextureLayerParameterValue_Legacy(const FHashedMaterialParameterInfo& ParameterInfo, bool bOveriddenOnly, URuntimeVirtualTexture*& OutValue, bool& OutResult) const;
+
 #endif // WITH_EDITOR
 
 protected:
@@ -703,7 +732,6 @@ protected:
 
 	void GetTextureExpressionValues(const FMaterialResource* MaterialResource, TArray<UTexture*>& OutTextures, TArray< TArray<int32> >* OutIndices = nullptr) const;
 
-	void GetAtlasTextureValues(const FMaterialResource* MaterialResource, TArray<UTexture*>& OutTextures) const;
 	/**
 	 * Updates StaticPermutationMaterialResources based on the value of bHasStaticPermutationResource
 	 * This is a helper used when recompiling MI's with static parameters.  
@@ -717,7 +745,7 @@ protected:
 	*/
 	ENGINE_API void UpdateParameterNames();
 
-#endif
+#endif // WITH_EDITOR
 
 	/**
 	 * Internal interface for setting / updating values for material instances.
@@ -766,6 +794,7 @@ protected:
 	friend class FMaterialUpdateContext;
 };
 
+#if WITH_EDITOR
 /** Workaround - Similar to base call but evaluates all expressions found, not just the first */
 template<typename ExpressionType>
 void FindClosestExpressionByGUIDRecursive(const FName& InName, const FGuid& InGUID, const TArray<UMaterialExpression*>& InMaterialExpression, ExpressionType*& OutExpression)
@@ -933,3 +962,4 @@ bool UpdateParameterSet(TArray<ParameterType>& Parameters, UMaterial* ParentMate
 
 	return bChanged;
 }
+#endif // WITH_EDITOR
diff --git a/Engine/Source/Runtime/Engine/Classes/Materials/MaterialInterface.h b/Engine/Source/Runtime/Engine/Classes/Materials/MaterialInterface.h
--- a/Engine/Source/Runtime/Engine/Classes/Materials/MaterialInterface.h
+++ b/Engine/Source/Runtime/Engine/Classes/Materials/MaterialInterface.h
@@ -4,6 +4,7 @@
 #pragma once
 
 #include "CoreMinimal.h"
+#include "Containers/ArrayView.h"
 #include "UObject/ObjectMacros.h"
 #include "UObject/Object.h"
 #include "Misc/Guid.h"
@@ -165,6 +166,40 @@ struct FMaterialTextureInfo
 	ENGINE_API bool IsValid(bool bCheckTextureIndex = false) const; 
 };
 
+struct TMicRecursionGuard
+{
+	inline TMicRecursionGuard() = default;
+	inline TMicRecursionGuard(const TMicRecursionGuard& Parent)
+		: MaterialInterface(nullptr)
+		, PreviousLink(&Parent)
+	{
+	}
+
+	inline void Set(class UMaterialInterface const* InMaterialInterface)
+	{
+		check(MaterialInterface == nullptr);
+		MaterialInterface = InMaterialInterface;
+	}
+
+	inline bool Contains(class UMaterialInterface const* InMaterialInterface)
+	{
+		TMicRecursionGuard const* Link = this;
+		do
+		{
+			if (Link->MaterialInterface == InMaterialInterface)
+			{
+				return true;
+			}
+			Link = Link->PreviousLink;
+		} while (Link);
+		return false;
+	}
+
+private:
+	void const* MaterialInterface = nullptr;
+	TMicRecursionGuard const* PreviousLink = nullptr;
+};
+
 UCLASS(abstract, BlueprintType, MinimalAPI, HideCategories = (Thumbnail))
 class ENGINE_VTABLE UMaterialInterface : public UObject, public IBlendableInterface, public IInterface_AssetUserData
 {
@@ -288,8 +323,7 @@ public:
 	/**
 	 * Same as above, but can be called concurrently
 	 */
-	typedef TArray<class UMaterialInterface const*, TInlineAllocator<8> > TMicRecursionGuard;
-	virtual const class UMaterial* GetMaterial_Concurrent(TMicRecursionGuard& RecursionGuard) const PURE_VIRTUAL(UMaterialInterface::GetMaterial_Concurrent,return NULL;);
+	virtual const class UMaterial* GetMaterial_Concurrent(TMicRecursionGuard RecursionGuard = TMicRecursionGuard()) const PURE_VIRTUAL(UMaterialInterface::GetMaterial_Concurrent,return NULL;);
 
 	/**
 	* Test this material for dependency on a given material.
@@ -351,14 +385,14 @@ public:
 	 * This is used to implement realtime previewing of parameter defaults. 
 	 * Handles updating dependent MI's and cached uniform expressions.
 	 */
-	virtual void OverrideVectorParameterDefault(const FMaterialParameterInfo& ParameterInfo, const FLinearColor& Value, bool bOverride, ERHIFeatureLevel::Type FeatureLevel) PURE_VIRTUAL(UMaterialInterface::OverrideTexture, return;);
-	virtual void OverrideScalarParameterDefault(const FMaterialParameterInfo& ParameterInfo, float Value, bool bOverride, ERHIFeatureLevel::Type FeatureLevel) PURE_VIRTUAL(UMaterialInterface::OverrideTexture, return;);
+	virtual void OverrideVectorParameterDefault(const FHashedMaterialParameterInfo& ParameterInfo, const FLinearColor& Value, bool bOverride, ERHIFeatureLevel::Type FeatureLevel) PURE_VIRTUAL(UMaterialInterface::OverrideTexture, return;);
+	virtual void OverrideScalarParameterDefault(const FHashedMaterialParameterInfo& ParameterInfo, float Value, bool bOverride, ERHIFeatureLevel::Type FeatureLevel) PURE_VIRTUAL(UMaterialInterface::OverrideTexture, return;);
 
 	/**
 	 * DEPRECATED: Returns default value of the given parameter
 	 */
 	UE_DEPRECATED(4.19, "This function is deprecated. Use GetScalarParameterDefaultValue instead.")
-	virtual float GetScalarParameterDefault(const FMaterialParameterInfo& ParameterInfo, ERHIFeatureLevel::Type FeatureLevel)
+	virtual float GetScalarParameterDefault(const FHashedMaterialParameterInfo& ParameterInfo, ERHIFeatureLevel::Type FeatureLevel)
 	{
 		float Value;
 		GetScalarParameterDefaultValue(ParameterInfo, Value);
@@ -388,12 +422,13 @@ public:
 	 */
 	virtual const FMaterialResource* GetMaterialResource(ERHIFeatureLevel::Type InFeatureLevel, EMaterialQualityLevel::Type QualityLevel = EMaterialQualityLevel::Num) const { return NULL; }
 
+#if WITH_EDITORONLY_DATA
 	struct FStaticParamEvaluationContext
 	{
 	public:
-		FStaticParamEvaluationContext(int32 InParameterNum, const FMaterialParameterInfo* InParameterInfos) : PendingParameterNum(InParameterNum), ParameterInfos(InParameterInfos) { PendingParameters.Add(true, PendingParameterNum); ResolvedByOverride.Add(false, PendingParameterNum); }
+		FStaticParamEvaluationContext(int32 InParameterNum, const FHashedMaterialParameterInfo* InParameterInfos) : PendingParameterNum(InParameterNum), ParameterInfos(InParameterInfos) { PendingParameters.Add(true, PendingParameterNum); ResolvedByOverride.Add(false, PendingParameterNum); }
 
-		const FMaterialParameterInfo* GetParameterInfo(int32 ParamIndex) const { check(ParamIndex >= 0 && ParamIndex < PendingParameters.Num()); return ParameterInfos + ParamIndex; }
+		const FHashedMaterialParameterInfo* GetParameterInfo(int32 ParamIndex) const { check(ParamIndex >= 0 && ParamIndex < PendingParameters.Num()); return ParameterInfos + ParamIndex; }
 		bool AllResolved() const { return PendingParameterNum == 0; }
 		bool IsResolved(int32 ParamIndex) const { return !PendingParameters[ParamIndex]; }
 		bool IsResolvedByOverride(int32 ParamIndex) const { return ResolvedByOverride[ParamIndex]; }
@@ -405,12 +440,12 @@ public:
 		 * Perform an operation on the pending parameters until performed on all pending parameters, or until the operation asks to stop iterating by returning false.
 	 	 * @param Op - The operation to perform on each pending parameter.  Accepts the index for the parameter and the parameter info associated with it.  If the Op returns true, we keep iterating if there are more pending parameters, if it returns false, we stop immediately after this Op.
 		 */
-		void ForEachPendingParameter(TFunctionRef<bool(int32 ParamIndex, const FMaterialParameterInfo& ParamInfo)> Op);
+		void ForEachPendingParameter(TFunctionRef<bool(int32 ParamIndex, const FHashedMaterialParameterInfo& ParamInfo)> Op);
 	private:
 		TBitArray<> PendingParameters;
 		TBitArray<> ResolvedByOverride;
 		int32 PendingParameterNum;
-		const FMaterialParameterInfo* ParameterInfos;
+		const FHashedMaterialParameterInfo* ParameterInfos;
 	};
 
 	/**
@@ -420,7 +455,7 @@ public:
 	* @param	OutValue		Will contain the value of the parameter if successful
 	* @return					True if successful
 	*/
-	ENGINE_API bool GetStaticSwitchParameterValue(const FMaterialParameterInfo& ParameterInfo, bool& OutValue, FGuid& OutExpressionGuid, bool bOveriddenOnly = false, bool bCheckParent = true) const;
+	ENGINE_API bool GetStaticSwitchParameterValue(const FHashedMaterialParameterInfo& ParameterInfo,bool &OutValue,FGuid &OutExpressionGuid, bool bOveriddenOnly = false, bool bCheckParent = true) const;
 
 	/**
 	* Get the values of the given set of static switch parameters.
@@ -440,7 +475,7 @@ public:
 	* @param	R, G, B, A		Will contain the values of the parameter if successful
 	* @return					True if successful
 	*/
-	ENGINE_API bool GetStaticComponentMaskParameterValue(const FMaterialParameterInfo& ParameterInfo, bool& R, bool& G, bool& B, bool& A, FGuid& OutExpressionGuid, bool bOveriddenOnly = false, bool bCheckParent = true) const;
+	ENGINE_API bool GetStaticComponentMaskParameterValue(const FHashedMaterialParameterInfo& ParameterInfo, bool& R, bool& G, bool& B, bool& A, FGuid& OutExpressionGuid, bool bOveriddenOnly = false, bool bCheckParent = true) const;
 
 	/**
 	* Get the values of the given set of static component mask parameters
@@ -453,6 +488,17 @@ public:
 	virtual bool GetStaticComponentMaskParameterValues(FStaticParamEvaluationContext& EvalContext, TBitArray<>& OutRGBAOrderedValues, FGuid* OutExpressionGuids, bool bCheckParent = true) const
 		PURE_VIRTUAL(UMaterialInterface::GetStaticComponentMaskParameterValues,return false;);
 
+	/**
+	* Get the value of the given static material layers parameter
+	*
+	* @param	ParameterName	The name of the material layer parameter
+	* @param	OutValue		Will contain the value of the parameter if successful
+	* @return					True if successful
+	*/
+	virtual bool GetMaterialLayersParameterValue(const FHashedMaterialParameterInfo& ParameterInfo, FMaterialLayersFunctions& OutLayers, FGuid& OutExpressionGuid, bool bCheckParent = true) const
+		PURE_VIRTUAL(UMaterialInterface::GetMaterialLayersParameterValue, return false;);
+#endif // WITH_EDITORONLY_DATA
+
 	/**
 	* Get the weightmap index of the given terrain layer weight parameter
 	*
@@ -460,7 +506,7 @@ public:
 	* @param	OutWeightmapIndex	Will contain the values of the parameter if successful
 	* @return					True if successful
 	*/
-	virtual bool GetTerrainLayerWeightParameterValue(const FMaterialParameterInfo& ParameterInfo, int32& OutWeightmapIndex, FGuid &OutExpressionGuid) const
+	virtual bool GetTerrainLayerWeightParameterValue(const FHashedMaterialParameterInfo& ParameterInfo, int32& OutWeightmapIndex, FGuid &OutExpressionGuid) const
 		PURE_VIRTUAL(UMaterialInterface::GetTerrainLayerWeightParameterValue,return false;);
 
 #if WITH_EDITOR
@@ -471,7 +517,7 @@ public:
 	* @param	OutSortPriority	Will contain the sort priority of the parameter if successful
 	* @return					True if successful
 	*/
-	virtual bool GetParameterSortPriority(const FMaterialParameterInfo& ParameterInfo, int32& OutSortPriority, const TArray<struct FStaticMaterialLayersParameter>* MaterialLayersParameters = nullptr) const
+	virtual bool GetParameterSortPriority(const FHashedMaterialParameterInfo& ParameterInfo, int32& OutSortPriority, const TArray<struct FStaticMaterialLayersParameter>* MaterialLayersParameters = nullptr) const
 		PURE_VIRTUAL(UMaterialInterface::GetParameterSortPriority, return false;);
 #endif
 
@@ -484,16 +530,6 @@ public:
 	*/
 	virtual bool GetGroupSortPriority(const FString& InGroupName, int32& OutSortPriority) const
 		PURE_VIRTUAL(UMaterialInterface::GetGroupSortPriority, return false;);
-		
-	/**
-	* Get the value of the given static material layers parameter
-	*
-	* @param	ParameterName	The name of the material layer parameter
-	* @param	OutValue		Will contain the value of the parameter if successful
-	* @return					True if successful
-	*/
-	virtual bool GetMaterialLayersParameterValue(const FMaterialParameterInfo& ParameterInfo, FMaterialLayersFunctions& OutLayers, FGuid& OutExpressionGuid, bool bCheckParent = true) const
-		PURE_VIRTUAL(UMaterialInterface::GetMaterialLayersParameterValue,return false;);
 
 	virtual void GetAllScalarParameterInfo(TArray<FMaterialParameterInfo>& OutParameterInfo, TArray<FGuid>& OutParameterIds) const
 		PURE_VIRTUAL(UMaterialInterface::GetAllScalarParameterInfo,return;);
@@ -505,6 +541,8 @@ public:
 		PURE_VIRTUAL(UMaterialInterface::GetAllRuntimeVirtualTextureParameterInfo, return;);
 	virtual void GetAllFontParameterInfo(TArray<FMaterialParameterInfo>& OutParameterInfo, TArray<FGuid>& OutParameterIds) const
 		PURE_VIRTUAL(UMaterialInterface::GetAllFontParameterInfo,return;);
+
+#if WITH_EDITORONLY_DATA
 	virtual void GetAllMaterialLayersParameterInfo(TArray<FMaterialParameterInfo>& OutParameterInfo, TArray<FGuid>& OutParameterIds) const
 		PURE_VIRTUAL(UMaterialInterface::GetAllMaterialLayersParameterInfo,return;);
 	virtual void GetAllStaticSwitchParameterInfo(TArray<FMaterialParameterInfo>& OutParameterInfo, TArray<FGuid>& OutParameterIds) const
@@ -516,25 +554,29 @@ public:
 		PURE_VIRTUAL(UMaterialInterface::IterateDependentFunctions,return false;);
 	virtual void GetDependentFunctions(TArray<class UMaterialFunctionInterface*>& DependentFunctions) const
 		PURE_VIRTUAL(UMaterialInterface::GetDependentFunctions,return;);
+#endif // WITH_EDITORONLY_DATA
 
-	virtual bool GetScalarParameterDefaultValue(const FMaterialParameterInfo& ParameterInfo, float& OutValue, bool bOveriddenOnly = false, bool bCheckOwnedGlobalOverrides = false) const
+	virtual bool GetScalarParameterDefaultValue(const FHashedMaterialParameterInfo& ParameterInfo, float& OutValue, bool bOveriddenOnly = false, bool bCheckOwnedGlobalOverrides = false) const
 		PURE_VIRTUAL(UMaterialInterface::GetScalarParameterDefaultValue,return false;);
-	virtual bool GetVectorParameterDefaultValue(const FMaterialParameterInfo& ParameterInfo, FLinearColor& OutValue, bool bOveriddenOnly = false, bool bCheckOwnedGlobalOverrides = false) const
+	virtual bool GetVectorParameterDefaultValue(const FHashedMaterialParameterInfo& ParameterInfo, FLinearColor& OutValue, bool bOveriddenOnly = false, bool bCheckOwnedGlobalOverrides = false) const
 		PURE_VIRTUAL(UMaterialInterface::GetVectorParameterDefaultValue,return false;);
-	virtual bool GetTextureParameterDefaultValue(const FMaterialParameterInfo& ParameterInfo, class UTexture*& OutValue, bool bCheckOwnedGlobalOverrides = false) const
+	virtual bool GetTextureParameterDefaultValue(const FHashedMaterialParameterInfo& ParameterInfo, class UTexture*& OutValue, bool bCheckOwnedGlobalOverrides = false) const
 		PURE_VIRTUAL(UMaterialInterface::GetTextureParameterDefaultValue,return false;);
-	virtual bool GetRuntimeVirtualTextureParameterDefaultValue(const FMaterialParameterInfo& ParameterInfo, class URuntimeVirtualTexture*& OutValue, bool bCheckOwnedGlobalOverrides = false) const
+	virtual bool GetRuntimeVirtualTextureParameterDefaultValue(const FHashedMaterialParameterInfo& ParameterInfo, class URuntimeVirtualTexture*& OutValue, bool bCheckOwnedGlobalOverrides = false) const
 		PURE_VIRTUAL(UMaterialInterface::GetRuntimeVirtualTextureParameterDefaultValue, return false;);
-	virtual bool GetFontParameterDefaultValue(const FMaterialParameterInfo& ParameterInfo, class UFont*& OutFontValue, int32& OutFontPage, bool bCheckOwnedGlobalOverrides = false) const
+	virtual bool GetFontParameterDefaultValue(const FHashedMaterialParameterInfo& ParameterInfo, class UFont*& OutFontValue, int32& OutFontPage, bool bCheckOwnedGlobalOverrides = false) const
 		PURE_VIRTUAL(UMaterialInterface::GetFontParameterDefaultValue,return false;);
-	virtual bool GetStaticSwitchParameterDefaultValue(const FMaterialParameterInfo& ParameterInfo, bool& OutValue, FGuid& OutExpressionGuid, bool bCheckOwnedGlobalOverrides = false) const
+	
+#if WITH_EDITOR
+	virtual bool GetStaticSwitchParameterDefaultValue(const FHashedMaterialParameterInfo& ParameterInfo, bool& OutValue, FGuid& OutExpressionGuid, bool bCheckOwnedGlobalOverrides = false) const
 		PURE_VIRTUAL(UMaterialInterface::GetStaticSwitchParameterDefaultValue,return false;);
-	virtual bool GetStaticComponentMaskParameterDefaultValue(const FMaterialParameterInfo& ParameterInfo, bool& OutR, bool& OutG, bool& OutB, bool& OutA, FGuid& OutExpressionGuid, bool bCheckOwnedGlobalOverrides = false) const
+	virtual bool GetStaticComponentMaskParameterDefaultValue(const FHashedMaterialParameterInfo& ParameterInfo, bool& OutR, bool& OutG, bool& OutB, bool& OutA, FGuid& OutExpressionGuid, bool bCheckOwnedGlobalOverrides = false) const
 		PURE_VIRTUAL(UMaterialInterface::GetStaticComponentMaskParameterDefaultValue,return false;);
-		
-	/** Appends textures referenced by expressions, including nested functions. */
-	virtual void AppendReferencedTextures(TArray<UObject*>& InOutTextures) const
-		PURE_VIRTUAL(UMaterialInterface::AppendReferencedTextures,);
+#endif // WITH_EDITOR
+
+	/** Get textures referenced by expressions, including nested functions. */
+	virtual TArrayView<UObject* const> GetReferencedTextures() const
+		PURE_VIRTUAL(UMaterialInterface::GetReferencedTextures,return TArrayView<UObject* const>(););
 
 	virtual void SaveShaderStableKeysInner(const class ITargetPlatform* TP, const struct FStableShaderKeyAndValue& SaveKeyVal)
 		PURE_VIRTUAL(UMaterialInterface::SaveShaderStableKeysInner, );
@@ -711,27 +753,29 @@ public:
 		ERHIFeatureLevel::Type InFeatureLevel = ERHIFeatureLevel::Num, EMaterialQualityLevel::Type InQuality = EMaterialQualityLevel::Num)
 		PURE_VIRTUAL(UMaterialInterface::GetTexturesInPropertyChain,return false;);
 
-	ENGINE_API virtual bool GetGroupName(const FMaterialParameterInfo& ParameterInfo, FName& GroupName) const;
-	ENGINE_API virtual bool GetParameterDesc(const FMaterialParameterInfo& ParameterInfo, FString& OutDesc, const TArray<struct FStaticMaterialLayersParameter>* MaterialLayersParameters = nullptr) const;	
-	ENGINE_API virtual bool GetScalarParameterSliderMinMax(const FMaterialParameterInfo& ParameterInfo, float& OutSliderMin, float& OutSliderMax) const;
+	ENGINE_API virtual bool GetGroupName(const FHashedMaterialParameterInfo& ParameterInfo, FName& GroupName) const;
+	ENGINE_API virtual bool GetParameterDesc(const FHashedMaterialParameterInfo& ParameterInfo, FString& OutDesc, const TArray<struct FStaticMaterialLayersParameter>* MaterialLayersParameters = nullptr) const;
+	ENGINE_API virtual bool GetScalarParameterSliderMinMax(const FHashedMaterialParameterInfo& ParameterInfo, float& OutSliderMin, float& OutSliderMax) const;
+#endif // WITH_EDITOR
+	ENGINE_API virtual bool GetScalarParameterValue(const FHashedMaterialParameterInfo& ParameterInfo, float& OutValue, bool bOveriddenOnly = false) const;
+#if WITH_EDITOR
+	ENGINE_API virtual bool IsScalarParameterUsedAsAtlasPosition(const FHashedMaterialParameterInfo& ParameterInfo, bool& OutValue, TSoftObjectPtr<class UCurveLinearColor>& Curve, TSoftObjectPtr<class UCurveLinearColorAtlas>&  Atlas) const;
 #endif // WITH_EDITOR
-	ENGINE_API virtual bool GetScalarParameterValue(const FMaterialParameterInfo& ParameterInfo, float& OutValue, bool bOveriddenOnly = false) const;
-	ENGINE_API virtual bool IsScalarParameterUsedAsAtlasPosition(const FMaterialParameterInfo& ParameterInfo, bool& OutValue, TSoftObjectPtr<class UCurveLinearColor>& Curve, TSoftObjectPtr<class UCurveLinearColorAtlas>&  Atlas) const;
-	ENGINE_API virtual bool GetScalarCurveParameterValue(const FMaterialParameterInfo& ParameterInfo, FInterpCurveFloat& OutValue) const;
-	ENGINE_API virtual bool GetVectorParameterValue(const FMaterialParameterInfo& ParameterInfo, FLinearColor& OutValue, bool bOveriddenOnly = false) const;
-	ENGINE_API virtual bool IsVectorParameterUsedAsChannelMask(const FMaterialParameterInfo& ParameterInfo, bool& OutValue) const;
+	ENGINE_API virtual bool GetScalarCurveParameterValue(const FHashedMaterialParameterInfo& ParameterInfo, FInterpCurveFloat& OutValue) const;
+	ENGINE_API virtual bool GetVectorParameterValue(const FHashedMaterialParameterInfo& ParameterInfo, FLinearColor& OutValue, bool bOveriddenOnly = false) const;
 #if WITH_EDITOR
-	ENGINE_API virtual bool GetVectorParameterChannelNames(const FMaterialParameterInfo& ParameterInfo, FParameterChannelNames& OutValue) const;
+	ENGINE_API virtual bool IsVectorParameterUsedAsChannelMask(const FHashedMaterialParameterInfo& ParameterInfo, bool& OutValue) const;
+	ENGINE_API virtual bool GetVectorParameterChannelNames(const FHashedMaterialParameterInfo& ParameterInfo, FParameterChannelNames& OutValue) const;
 #endif
-	ENGINE_API virtual bool GetVectorCurveParameterValue(const FMaterialParameterInfo& ParameterInfo, FInterpCurveVector& OutValue) const;
-	ENGINE_API virtual bool GetLinearColorParameterValue(const FMaterialParameterInfo& ParameterInfo, FLinearColor& OutValue) const;
-	ENGINE_API virtual bool GetLinearColorCurveParameterValue(const FMaterialParameterInfo& ParameterInfo, FInterpCurveLinearColor& OutValue) const;
-	ENGINE_API virtual bool GetTextureParameterValue(const FMaterialParameterInfo& ParameterInfo, class UTexture*& OutValue, bool bOveriddenOnly = false) const;
-	ENGINE_API virtual bool GetRuntimeVirtualTextureParameterValue(const FMaterialParameterInfo& ParameterInfo, class URuntimeVirtualTexture*& OutValue, bool bOveriddenOnly = false) const;
+	ENGINE_API virtual bool GetVectorCurveParameterValue(const FHashedMaterialParameterInfo& ParameterInfo, FInterpCurveVector& OutValue) const;
+	ENGINE_API virtual bool GetLinearColorParameterValue(const FHashedMaterialParameterInfo& ParameterInfo, FLinearColor& OutValue) const;
+	ENGINE_API virtual bool GetLinearColorCurveParameterValue(const FHashedMaterialParameterInfo& ParameterInfo, FInterpCurveLinearColor& OutValue) const;
+	ENGINE_API virtual bool GetTextureParameterValue(const FHashedMaterialParameterInfo& ParameterInfo, class UTexture*& OutValue, bool bOveriddenOnly = false) const;
+	ENGINE_API virtual bool GetRuntimeVirtualTextureParameterValue(const FHashedMaterialParameterInfo& ParameterInfo, class URuntimeVirtualTexture*& OutValue, bool bOveriddenOnly = false) const;
 #if WITH_EDITOR
-	ENGINE_API virtual bool GetTextureParameterChannelNames(const FMaterialParameterInfo& ParameterInfo, FParameterChannelNames& OutValue) const;
+	ENGINE_API virtual bool GetTextureParameterChannelNames(const FHashedMaterialParameterInfo& ParameterInfo, FParameterChannelNames& OutValue) const;
 #endif
-	ENGINE_API virtual bool GetFontParameterValue(const FMaterialParameterInfo& ParameterInfo,class UFont*& OutFontValue, int32& OutFontPage, bool bOveriddenOnly = false) const;
+	ENGINE_API virtual bool GetFontParameterValue(const FHashedMaterialParameterInfo& ParameterInfo,class UFont*& OutFontValue, int32& OutFontPage, bool bOveriddenOnly = false) const;
 	ENGINE_API virtual bool GetRefractionSettings(float& OutBiasValue) const;
 
 	/**
diff --git a/Engine/Source/Runtime/Engine/Classes/Materials/MaterialLayersFunctions.h b/Engine/Source/Runtime/Engine/Classes/Materials/MaterialLayersFunctions.h
--- a/Engine/Source/Runtime/Engine/Classes/Materials/MaterialLayersFunctions.h
+++ b/Engine/Source/Runtime/Engine/Classes/Materials/MaterialLayersFunctions.h
@@ -63,18 +63,87 @@ struct ENGINE_API FMaterialParameterInfo
 		Ar << Ref.Name << Ref.Association << Ref.Index;
 		return Ar;
 	}
+};
 
-	FORCEINLINE bool operator==(const FMaterialParameterInfo& Other) const
-	{
-		return Name.IsEqual(Other.Name) && Association == Other.Association && Index == Other.Index;
-	}
+struct FHashedMaterialParameterInfo
+{
+	DECLARE_TYPE_LAYOUT(FHashedMaterialParameterInfo, NonVirtual);
+public:
+	FHashedMaterialParameterInfo(const FHashedName& InName = FHashedName(), EMaterialParameterAssociation InAssociation = EMaterialParameterAssociation::GlobalParameter, int32 InIndex = INDEX_NONE)
+		: Name(InName)
+		, Index(InIndex)
+		, Association(InAssociation)
+	{}
+
+	FHashedMaterialParameterInfo(const TCHAR* InName, EMaterialParameterAssociation InAssociation = EMaterialParameterAssociation::GlobalParameter, int32 InIndex = INDEX_NONE)
+		: Name(InName)
+		, Index(InIndex)
+		, Association(InAssociation)
+	{}
+
+	FHashedMaterialParameterInfo(const FName& InName, EMaterialParameterAssociation InAssociation = EMaterialParameterAssociation::GlobalParameter, int32 InIndex = INDEX_NONE)
+		: Name(InName)
+		, Index(InIndex)
+		, Association(InAssociation)
+	{}
+
+	FHashedMaterialParameterInfo(const FMaterialParameterInfo& Rhs)
+		: Name(Rhs.Name)
+		, Index(Rhs.Index)
+		, Association(Rhs.Association)
+	{}
+
+	FHashedMaterialParameterInfo(const FHashedMaterialParameterInfo& Rhs) = default;
 
-	FORCEINLINE bool operator!=(const FMaterialParameterInfo& Other) const
+	friend FArchive& operator<<(FArchive& Ar, FHashedMaterialParameterInfo& Ref)
 	{
-		return !Name.IsEqual(Other.Name) || Association != Other.Association || Index != Other.Index;
+		Ar << Ref.Name << Ref.Association << Ref.Index;
+		return Ar;
 	}
+
+	LAYOUT_FIELD(FHashedName, Name);
+	LAYOUT_FIELD(int32, Index);
+	LAYOUT_FIELD(TEnumAsByte<EMaterialParameterAssociation>, Association);
 };
 
+// For sorting/searching
+FORCEINLINE bool operator<(const FHashedMaterialParameterInfo& Lhs, const FHashedMaterialParameterInfo& Rhs)
+{
+	if (Lhs.Association != Rhs.Association) return Lhs.Association < Rhs.Association;
+	else if (Lhs.Index != Rhs.Index) return Lhs.Index < Rhs.Index;
+	return Lhs.Name < Rhs.Name;
+}
+
+FORCEINLINE bool operator==(const FMaterialParameterInfo& Lhs, const FMaterialParameterInfo& Rhs)
+{
+	return Lhs.Name.IsEqual(Rhs.Name) && Lhs.Association == Rhs.Association && Lhs.Index == Rhs.Index;
+}
+
+FORCEINLINE bool operator!=(const FMaterialParameterInfo& Lhs, const FMaterialParameterInfo& Rhs)
+{
+	return !operator==(Lhs, Rhs);
+}
+
+FORCEINLINE bool operator==(const FHashedMaterialParameterInfo& Lhs, const FHashedMaterialParameterInfo& Rhs)
+{
+	return Lhs.Name == Rhs.Name && Lhs.Association == Rhs.Association && Lhs.Index == Rhs.Index;
+}
+
+FORCEINLINE bool operator!=(const FHashedMaterialParameterInfo& Lhs, const FHashedMaterialParameterInfo& Rhs)
+{
+	return !operator==(Lhs, Rhs);
+}
+
+FORCEINLINE bool operator==(const FMaterialParameterInfo& Lhs, const FHashedMaterialParameterInfo& Rhs)
+{
+	return FHashedName(Lhs.Name) == Rhs.Name && Lhs.Index == Rhs.Index && Lhs.Association == Rhs.Association;
+}
+
+FORCEINLINE bool operator!=(const FMaterialParameterInfo& Lhs, const FHashedMaterialParameterInfo& Rhs)
+{
+	return !operator==(Lhs, Rhs);
+}
+
 USTRUCT()
 struct ENGINE_API FMaterialLayersFunctions
 {
diff --git a/Engine/Source/Runtime/Engine/Classes/Sound/SoundWave.h b/Engine/Source/Runtime/Engine/Classes/Sound/SoundWave.h
--- a/Engine/Source/Runtime/Engine/Classes/Sound/SoundWave.h
+++ b/Engine/Source/Runtime/Engine/Classes/Sound/SoundWave.h
@@ -887,7 +887,8 @@ public:
 #endif // WITH_EDITOR
 
 	/** Checks whether sound has been categorised as streaming. */
-	bool IsStreaming(const FPlatformAudioCookOverrides* Overrides = nullptr) const;
+	bool IsStreaming(const TCHAR* PlatformName = nullptr) const;
+	bool IsStreaming(const FPlatformAudioCookOverrides& Overrides) const;
 
 	/** Checks whether sound has seekable streaming enabled. */
 	bool IsSeekableStreaming() const;
diff --git a/Engine/Source/Runtime/Engine/Private/Actor.cpp b/Engine/Source/Runtime/Engine/Private/Actor.cpp
--- a/Engine/Source/Runtime/Engine/Private/Actor.cpp
+++ b/Engine/Source/Runtime/Engine/Private/Actor.cpp
@@ -4400,7 +4400,7 @@ static USceneComponent* GetUnregisteredParent(UActorComponent* Component)
 	return ParentComponent;
 }
 
-bool AActor::IncrementalRegisterComponents(int32 NumComponentsToRegister)
+bool AActor::IncrementalRegisterComponents(int32 NumComponentsToRegister, FRegisterComponentContext* Context)
 {
 	if (NumComponentsToRegister == 0)
 	{
@@ -4430,7 +4430,7 @@ bool AActor::IncrementalRegisterComponents(int32 NumComponentsToRegister)
 			// This should prevent unwanted components hanging around when undoing a copy/paste or duplication action.
 			RootComponent->Modify(false);
 
-			RootComponent->RegisterComponentWithWorld(World);
+			RootComponent->RegisterComponentWithWorld(World, Context);
 		}
 	}
 
@@ -4470,7 +4470,7 @@ bool AActor::IncrementalRegisterComponents(int32 NumComponentsToRegister)
 			// This should prevent unwanted components hanging around when undoing a copy/paste or duplication action.
 			Component->Modify(false);
 
-			Component->RegisterComponentWithWorld(World);
+			Component->RegisterComponentWithWorld(World, Context);
 			NumRegisteredComponentsThisRun++;
 		}
 
