Copying Private-LoadTimes-4.24 stream to Main. Biggest changes are in Materials/Shader memory freezing.

diff --git a/Engine/Source/Runtime/Engine/Private/Atmosphere/AtmosphericFogComponent.cpp b/Engine/Source/Runtime/Engine/Private/Atmosphere/AtmosphericFogComponent.cpp
--- a/Engine/Source/Runtime/Engine/Private/Atmosphere/AtmosphericFogComponent.cpp
+++ b/Engine/Source/Runtime/Engine/Private/Atmosphere/AtmosphericFogComponent.cpp
@@ -292,9 +292,9 @@ void UAtmosphericFogComponent::ReleaseResource()
 	}
 }
 
-void UAtmosphericFogComponent::CreateRenderState_Concurrent()
+void UAtmosphericFogComponent::CreateRenderState_Concurrent(FRegisterComponentContext* Context)
 {
-	Super::CreateRenderState_Concurrent();
+	Super::CreateRenderState_Concurrent(Context);
 	AddFogIfNeeded();
 }
 
diff --git a/Engine/Source/Runtime/Engine/Private/Audio.cpp b/Engine/Source/Runtime/Engine/Private/Audio.cpp
--- a/Engine/Source/Runtime/Engine/Private/Audio.cpp
+++ b/Engine/Source/Runtime/Engine/Private/Audio.cpp
@@ -1093,7 +1093,7 @@ bool FWaveInstance::IsSeekable() const
 
 bool FWaveInstance::IsStreaming() const
 {
-	return FPlatformProperties::SupportsAudioStreaming() && WaveData != nullptr && WaveData->IsStreaming();
+	return FPlatformProperties::SupportsAudioStreaming() && WaveData != nullptr && WaveData->IsStreaming(nullptr);
 }
 
 bool FWaveInstance::GetUseSpatialization() const
diff --git a/Engine/Source/Runtime/Engine/Private/AudioCompressionSettingsUtils.cpp b/Engine/Source/Runtime/Engine/Private/AudioCompressionSettingsUtils.cpp
--- a/Engine/Source/Runtime/Engine/Private/AudioCompressionSettingsUtils.cpp
+++ b/Engine/Source/Runtime/Engine/Private/AudioCompressionSettingsUtils.cpp
@@ -2,9 +2,9 @@
 
 #include "AudioCompressionSettingsUtils.h"
 #include "AudioCompressionSettings.h"
+#include "Misc/DataDrivenPlatformInfoRegistry.h"
 
 #define ENABLE_PLATFORM_COMPRESSION_OVERRIDES 1
-#define PLATFORM_SUPPORTS_COMPRESSION_OVERRIDES (PLATFORM_ANDROID && !PLATFORM_LUMIN) || PLATFORM_IOS || PLATFORM_SWITCH || PLATFORM_XBOXONE || PLATFORM_PS4 || PLATFORM_WINDOWS
 
 #if PLATFORM_ANDROID && !PLATFORM_LUMIN && ENABLE_PLATFORM_COMPRESSION_OVERRIDES
 #include "AndroidRuntimeSettings.h"
@@ -72,106 +72,240 @@ const FPlatformRuntimeAudioCompressionOverrides* FPlatformCompressionUtilities::
 	return nullptr;
 }
 
-void CacheCurrentPlatformAudioCookOverrides(FPlatformAudioCookOverrides& OutOverrides)
+void CacheAudioCookOverrides(FPlatformAudioCookOverrides& OutOverrides, const TCHAR* InPlatformName=nullptr)
 {
-#if PLATFORM_ANDROID && !PLATFORM_LUMIN
-	const TCHAR* CategoryName = TEXT("/Script/AndroidRuntimeSettings.AndroidRuntimeSettings");
-#elif PLATFORM_IOS
-	const TCHAR* CategoryName = TEXT("/Script/IOSRuntimeSettings.IOSRuntimeSettings");
-#elif PLATFORM_SWITCH
-	const TCHAR* CategoryName = TEXT("/Script/SwitchRuntimeSettings.SwitchRuntimeSettings");
-#elif PLATFORM_WINDOWS
-	const TCHAR* CategoryName = TEXT("/Script/WindowsTargetPlatform.WindowsTargetSettings");
-#elif PLATFORM_PS4
-	const TCHAR* CategoryName = TEXT("/Script/PS4PlatformEditor.PS4TargetSettings");
-#elif PLATFORM_XBOXONE
-	const TCHAR* CategoryName = TEXT("/Script/XboxOnePlatformEditor.XboxOneTargetSettings");
-#else
-	const TCHAR* CategoryName = TEXT("");
-#endif
+	// if the platform was passed in, use it, otherwise, get the runtime platform's name for looking up DDPI
+	FString PlatformName = InPlatformName ? FString(InPlatformName) : FString(FPlatformProperties::IniPlatformName());
+	
+	// now use that platform name to get the ini section out of DDPI
+	FString CategoryName = FDataDrivenPlatformInfoRegistry::GetPlatformInfo(PlatformName).AudioCompressionSettingsIniSectionName;
 
-	int32 SoundCueCookQualityIndex = INDEX_NONE;
-	if (GConfig->GetInt(CategoryName, TEXT("SoundCueCookQualityIndex"), SoundCueCookQualityIndex, GEngineIni))
+	// if we don't support platform overrides, then return 
+	if (CategoryName.Len() == 0)
 	{
-		OutOverrides.SoundCueCookQualityIndex = SoundCueCookQualityIndex;
+		OutOverrides = FPlatformAudioCookOverrides();
+		return;
 	}
 	
-	GConfig->GetBool(CategoryName, TEXT("bResampleForDevice"), OutOverrides.bResampleForDevice, GEngineIni);
+	FConfigFile PlatformFile;
+	FConfigCacheIni::LoadLocalIniFile(PlatformFile, TEXT("Engine"), true, *PlatformName);
 
-	GConfig->GetFloat(CategoryName, TEXT("CompressionQualityModifier"), OutOverrides.CompressionQualityModifier, GEngineIni);
 
-	GConfig->GetFloat(CategoryName, TEXT("AutoStreamingThreshold"), OutOverrides.AutoStreamingThreshold, GEngineIni);
+	int32 SoundCueQualityIndex = INDEX_NONE;
+	if (PlatformFile.GetInt(*CategoryName, TEXT("SoundCueCookQualityIndex"), SoundCueQualityIndex))
+	{
+		OutOverrides.SoundCueCookQualityIndex = SoundCueQualityIndex;
+	}
 
-	GConfig->GetBool(CategoryName, TEXT("bUseAudioStreamCaching"), OutOverrides.bUseStreamCaching, GEngineIni);
+	PlatformFile.GetBool(*CategoryName, TEXT("bUseAudioStreamCaching"), OutOverrides.bUseStreamCaching);
 
 	/** Memory Load On Demand Settings */
 	if (OutOverrides.bUseStreamCaching)
 	{
 		// Cache size:
 		int32 RetrievedCacheSize = 32 * 1024;
-		GConfig->GetInt(CategoryName, TEXT("CacheSizeKB"), RetrievedCacheSize, GEngineIni);
+		PlatformFile.GetInt(*CategoryName, TEXT("CacheSizeKB"), RetrievedCacheSize);
 		OutOverrides.StreamCachingSettings.CacheSizeKB = RetrievedCacheSize;
+
+		bool bForceLegacyStreamChunking = false;
+		PlatformFile.GetBool(*CategoryName, TEXT("bForceLegacyStreamChunking"), bForceLegacyStreamChunking);
+		OutOverrides.StreamCachingSettings.bForceLegacyStreamChunking = bForceLegacyStreamChunking;
+
+		int32 ZerothChunkSizeForLegacyStreamChunking = 0;
+		PlatformFile.GetInt(*CategoryName, TEXT("ZerothChunkSizeForLegacyStreamChunking"), ZerothChunkSizeForLegacyStreamChunking);
+		OutOverrides.StreamCachingSettings.ZerothChunkSizeForLegacyStreamChunkingKB = ZerothChunkSizeForLegacyStreamChunking;
+	}
+
+	PlatformFile.GetBool(*CategoryName, TEXT("bResampleForDevice"), OutOverrides.bResampleForDevice);
+
+	PlatformFile.GetFloat(*CategoryName, TEXT("CompressionQualityModifier"), OutOverrides.CompressionQualityModifier);
+
+	PlatformFile.GetFloat(*CategoryName, TEXT("AutoStreamingThreshold"), OutOverrides.AutoStreamingThreshold);
+
+#if 1
+	//Cache sample rate map:
+	float RetrievedSampleRate = -1.0f;
+
+	PlatformFile.GetFloat(*CategoryName, TEXT("MaxSampleRate"), RetrievedSampleRate);
+	float* FoundSampleRate = OutOverrides.PlatformSampleRates.Find(ESoundwaveSampleRateSettings::Max);
+
+	if (FoundSampleRate)
+	{
+		if (!FMath::IsNearlyEqual(*FoundSampleRate, RetrievedSampleRate))
+		{
+			*FoundSampleRate = RetrievedSampleRate;
+		}
+
+	}
+	else
+	{
+		OutOverrides.PlatformSampleRates.Add(ESoundwaveSampleRateSettings::Max, RetrievedSampleRate);
+	}
+
+	RetrievedSampleRate = -1.0f;
+
+	PlatformFile.GetFloat(*CategoryName, TEXT("HighSampleRate"), RetrievedSampleRate);
+	FoundSampleRate = OutOverrides.PlatformSampleRates.Find(ESoundwaveSampleRateSettings::High);
+
+	if (FoundSampleRate)
+	{
+		if (!FMath::IsNearlyEqual(*FoundSampleRate, RetrievedSampleRate))
+		{
+			*FoundSampleRate = RetrievedSampleRate;
+		}
+
+	}
+	else
+	{
+		OutOverrides.PlatformSampleRates.Add(ESoundwaveSampleRateSettings::High, RetrievedSampleRate);
+	}
+
+
+	RetrievedSampleRate = -1.0f;
+
+	PlatformFile.GetFloat(*CategoryName, TEXT("MedSampleRate"), RetrievedSampleRate);
+	FoundSampleRate = OutOverrides.PlatformSampleRates.Find(ESoundwaveSampleRateSettings::Medium);
+
+	if (FoundSampleRate)
+	{
+		if (!FMath::IsNearlyEqual(*FoundSampleRate, RetrievedSampleRate))
+		{
+			*FoundSampleRate = RetrievedSampleRate;
+		}
+	}
+	else
+	{
+		OutOverrides.PlatformSampleRates.Add(ESoundwaveSampleRateSettings::Medium, RetrievedSampleRate);
 	}
 
+	RetrievedSampleRate = -1.0f;
+
+	PlatformFile.GetFloat(*CategoryName, TEXT("LowSampleRate"), RetrievedSampleRate);
+	FoundSampleRate = OutOverrides.PlatformSampleRates.Find(ESoundwaveSampleRateSettings::Low);
+
+	if (FoundSampleRate)
+	{
+		if (!FMath::IsNearlyEqual(*FoundSampleRate, RetrievedSampleRate))
+		{
+			*FoundSampleRate = RetrievedSampleRate;
+		}
+	}
+	else
+	{
+		OutOverrides.PlatformSampleRates.Add(ESoundwaveSampleRateSettings::Low, RetrievedSampleRate);
+	}
+
+	RetrievedSampleRate = -1.0f;
+
+	PlatformFile.GetFloat(*CategoryName, TEXT("MinSampleRate"), RetrievedSampleRate);
+	FoundSampleRate = OutOverrides.PlatformSampleRates.Find(ESoundwaveSampleRateSettings::Min);
+
+	if (FoundSampleRate)
+	{
+		if (!FMath::IsNearlyEqual(*FoundSampleRate, RetrievedSampleRate))
+		{
+			*FoundSampleRate = RetrievedSampleRate;
+		}
+	}
+	else
+	{
+		OutOverrides.PlatformSampleRates.Add(ESoundwaveSampleRateSettings::Min, RetrievedSampleRate);
+	}
+
+#else
+
 	//Cache sample rate map.
 	OutOverrides.PlatformSampleRates.Reset();
 
 	float RetrievedSampleRate = -1.0f;
 
-	GConfig->GetFloat(CategoryName, TEXT("MaxSampleRate"), RetrievedSampleRate, GEngineIni);
+	PlatformFile.GetFloat(*CategoryName, TEXT("MaxSampleRate"), RetrievedSampleRate);
 	OutOverrides.PlatformSampleRates.Add(ESoundwaveSampleRateSettings::Max, RetrievedSampleRate);
 
 	RetrievedSampleRate = -1.0f;
 
-	GConfig->GetFloat(CategoryName, TEXT("HighSampleRate"), RetrievedSampleRate, GEngineIni);
+	PlatformFile.GetFloat(*CategoryName, TEXT("HighSampleRate"), RetrievedSampleRate);
 	OutOverrides.PlatformSampleRates.Add(ESoundwaveSampleRateSettings::High, RetrievedSampleRate);
 
 	RetrievedSampleRate = -1.0f;
 
-	GConfig->GetFloat(CategoryName, TEXT("MedSampleRate"), RetrievedSampleRate, GEngineIni);
+	PlatformFile.GetFloat(*CategoryName, TEXT("MedSampleRate"), RetrievedSampleRate);
 	OutOverrides.PlatformSampleRates.Add(ESoundwaveSampleRateSettings::Medium, RetrievedSampleRate);
 
 	RetrievedSampleRate = -1.0f;
 
-	GConfig->GetFloat(CategoryName, TEXT("LowSampleRate"), RetrievedSampleRate, GEngineIni);
+	PlatformFile.GetFloat(*CategoryName, TEXT("LowSampleRate"), RetrievedSampleRate);
 	OutOverrides.PlatformSampleRates.Add(ESoundwaveSampleRateSettings::Low, RetrievedSampleRate);
 
 	RetrievedSampleRate = -1.0f;
 
-	GConfig->GetFloat(CategoryName, TEXT("MinSampleRate"), RetrievedSampleRate, GEngineIni);
+	PlatformFile.GetFloat(*CategoryName, TEXT("MinSampleRate"), RetrievedSampleRate);
 	OutOverrides.PlatformSampleRates.Add(ESoundwaveSampleRateSettings::Min, RetrievedSampleRate);
+#endif
 }
 
-#if PLATFORM_SUPPORTS_COMPRESSION_OVERRIDES
-static FPlatformAudioCookOverrides OutOverrides = FPlatformAudioCookOverrides();
-#endif
 
+static bool PlatformSupportsCompressionOverrides(const FString& PlatformName)
+{
+	return FDataDrivenPlatformInfoRegistry::GetPlatformInfo(PlatformName).AudioCompressionSettingsIniSectionName.Len() > 0;
+}
 
+static inline FString GetCookOverridePlatformName(const TCHAR* PlatformName)
+{
+	return PlatformName ? FString(PlatformName) : FString(FPlatformProperties::IniPlatformName());
+}
+
+static bool PlatformSupportsCompressionOverrides(const TCHAR* PlatformName=nullptr)
+{
+	return PlatformSupportsCompressionOverrides(GetCookOverridePlatformName(PlatformName));
+}
 
-#if PLATFORM_SUPPORTS_COMPRESSION_OVERRIDES
 static FCriticalSection CookOverridesCriticalSection;
-#endif // PLATFORM_SUPPORTS_COMPRESSION_OVERRIDES
 
-void FPlatformCompressionUtilities::RecacheCookOverrides()
+static FPlatformAudioCookOverrides& GetCacheableOverridesByPlatform(const TCHAR* InPlatformName, bool& bNeedsToBeInitialized)
 {
-#if PLATFORM_SUPPORTS_COMPRESSION_OVERRIDES
 	FScopeLock ScopeLock(&CookOverridesCriticalSection);
-	CacheCurrentPlatformAudioCookOverrides(OutOverrides);
-#endif // PLATFORM_SUPPORTS_COMPRESSION_OVERRIDES
+
+	// registry of overrides by platform name, for cooking, etc that may need multiple platforms worth
+	static TMap<FString, FPlatformAudioCookOverrides> OverridesByPlatform;
+
+	// make sure we don't reallocate the memory later
+	if (OverridesByPlatform.Num() == 0)
+	{
+		// give enough space for all known platforms
+		OverridesByPlatform.Reserve(FDataDrivenPlatformInfoRegistry::GetNumDataDrivenIniFiles());
+	}
+
+	FString PlatformName = GetCookOverridePlatformName(InPlatformName);
+	// return one, or make one
+	FPlatformAudioCookOverrides* ExistingOverrides = OverridesByPlatform.Find(PlatformName);
+	if (ExistingOverrides != nullptr)
+	{
+		bNeedsToBeInitialized = false;
+		return *ExistingOverrides;
+	}
+
+	bNeedsToBeInitialized = true;
+	return OverridesByPlatform.Add(PlatformName, FPlatformAudioCookOverrides());
 }
 
-const FPlatformAudioCookOverrides* FPlatformCompressionUtilities::GetCookOverridesForCurrentPlatform(bool bForceRecache)
+void FPlatformCompressionUtilities::RecacheCookOverrides()
 {
-#if PLATFORM_SUPPORTS_COMPRESSION_OVERRIDES
-	static bool bCachedCookOverrides = false;
-
-#if !WITH_EDITOR
-	// In non-editor situations, we only need to cache the cook overrides once.
-	if (!bCachedCookOverrides)
+	if (PlatformSupportsCompressionOverrides())
 	{
-#else
+		FScopeLock ScopeLock(&CookOverridesCriticalSection);
+		bool bNeedsToBeInitialized;
+		CacheAudioCookOverrides(GetCacheableOverridesByPlatform(nullptr, bNeedsToBeInitialized));
+	}
+}
+
+const FPlatformAudioCookOverrides* FPlatformCompressionUtilities::GetCookOverrides(const TCHAR* PlatformName, bool bForceRecache)
+{
+	bool bNeedsToBeInitialized;
+	FPlatformAudioCookOverrides& Overrides = GetCacheableOverridesByPlatform(PlatformName, bNeedsToBeInitialized);
+
+#if WITH_EDITOR
 	// In editor situations, the settings can change at any time, so we need to retrieve them.
-	FScopeLock ScopeLock(&CookOverridesCriticalSection);
 	
 	static double LastCacheTime = 0.0;
 	double CurrentTime = FPlatformTime::Seconds();
@@ -179,27 +313,29 @@ const FPlatformAudioCookOverrides* FPlatformCompressionUtilities::GetCookOverrid
 
 	if (bForceRecache || TimeSinceLastCache > CookOverrideCachingIntervalCvar)
 	{
+		bNeedsToBeInitialized = true;
 		LastCacheTime = CurrentTime;
-#endif // WITH_EDITOR
-	
-		CacheCurrentPlatformAudioCookOverrides(OutOverrides);
-		bCachedCookOverrides = true;
 	}
-	return &OutOverrides;
-#else 
-	return nullptr;
 #endif
+	
+	if (bNeedsToBeInitialized)
+	{
+		FScopeLock ScopeLock(&CookOverridesCriticalSection);
+		CacheAudioCookOverrides(Overrides, PlatformName);
+	}
+
+	return &Overrides;
 }
 
 bool FPlatformCompressionUtilities::IsCurrentPlatformUsingStreamCaching()
 {
-	const FPlatformAudioCookOverrides* Settings = GetCookOverridesForCurrentPlatform();
+	const FPlatformAudioCookOverrides* Settings = GetCookOverrides();
 	return Settings && Settings->bUseStreamCaching;
 }
 
 const FAudioStreamCachingSettings& FPlatformCompressionUtilities::GetStreamCachingSettingsForCurrentPlatform()
 {
-	const FPlatformAudioCookOverrides* Settings = GetCookOverridesForCurrentPlatform();
+	const FPlatformAudioCookOverrides* Settings = GetCookOverrides();
 	checkf(Settings, TEXT("Please only use this function if FPlatformCompressionUtilities::IsCurrentPlatformUsingLoadOnDemand() returns true."));
 	return Settings->StreamCachingSettings;
 }
@@ -207,7 +343,7 @@ const FAudioStreamCachingSettings& FPlatformCompressionUtilities::GetStreamCachi
 FCachedAudioStreamingManagerParams FPlatformCompressionUtilities::BuildCachedStreamingManagerParams()
 {
 	const FAudioStreamCachingSettings& CacheSettings = GetStreamCachingSettingsForCurrentPlatform();
-	int32 MaxChunkSize = GetMaxChunkSizeForCookOverrides(GetCookOverridesForCurrentPlatform());
+	int32 MaxChunkSize = GetMaxChunkSizeForCookOverrides(GetCookOverrides());
 
 	// Our number of elements is tweakable based on the minimum cache usage we want to support.
 	const float MinimumCacheUsage = FMath::Clamp(MinimumCacheUsageCvar, 0.0f, 0.95f);
@@ -273,7 +409,7 @@ float FPlatformCompressionUtilities::GetCompressionDurationForCurrentPlatform()
 float FPlatformCompressionUtilities::GetTargetSampleRateForPlatform(ESoundwaveSampleRateSettings InSampleRateLevel /*= ESoundwaveSampleRateSettings::High*/)
 {
 	float SampleRate = -1.0f;
-	const FPlatformAudioCookOverrides* Settings = GetCookOverridesForCurrentPlatform();
+	const FPlatformAudioCookOverrides* Settings = GetCookOverrides();
 	if (Settings && Settings->bResampleForDevice)
 	{
 		const float* FoundSampleRate = Settings->PlatformSampleRates.Find(InSampleRateLevel);
diff --git a/Engine/Source/Runtime/Engine/Private/AudioDerivedData.cpp b/Engine/Source/Runtime/Engine/Private/AudioDerivedData.cpp
--- a/Engine/Source/Runtime/Engine/Private/AudioDerivedData.cpp
+++ b/Engine/Source/Runtime/Engine/Private/AudioDerivedData.cpp
@@ -183,18 +183,7 @@ static FName GetWaveFormatForRunningPlatform(USoundWave& SoundWave)
 
 static const FPlatformAudioCookOverrides* GetCookOverridesForRunningPlatform()
 {
-	ITargetPlatformManagerModule* TPM = GetTargetPlatformManager();
-	if (TPM)
-	{
-		ITargetPlatform* CurrentPlatform = GetRunningTargetPlatform(TPM);
-
-
-		return CurrentPlatform->GetAudioCompressionSettings();
-	}
-	else
-	{
-		return nullptr;
-	}
+	return FPlatformCompressionUtilities::GetCookOverrides(ANSI_TO_TCHAR(FPlatformProperties::IniPlatformName()));
 }
 
 /**
@@ -1100,7 +1089,9 @@ static void CookSimpleWave(USoundWave* SoundWave, FName FormatName, const IAudio
 		QualityInfo.NumChannels = *WaveInfo.pChannels;
 		QualityInfo.SampleRate = WaveSampleRate;
 		QualityInfo.SampleDataSize = Input.Num();
-		QualityInfo.bStreaming = SoundWave->IsStreaming();
+		// without overrides, we don't know the target platform's name to be able to look up, and passing nullptr will use editor platform's settings, which could be wrong
+		// @todo: Pass in TargetPlatform/PlatformName maybe?
+		QualityInfo.bStreaming = SoundWave->IsStreaming(CompressionOverrides ? *CompressionOverrides : FPlatformAudioCookOverrides());
 		QualityInfo.DebugName = SoundWave->GetFullName();
 
 		// Cook the data.
@@ -1306,7 +1297,7 @@ static void CookSurroundWave( USoundWave* SoundWave, FName FormatName, const IAu
 			QualityInfo.NumChannels = ChannelCount;
 			QualityInfo.SampleRate = WaveSampleRate;
 			QualityInfo.SampleDataSize = SampleDataSize;
-			QualityInfo.bStreaming = SoundWave->IsStreaming();
+			QualityInfo.bStreaming = SoundWave->IsStreaming(CompressionOverrides ? *CompressionOverrides : FPlatformAudioCookOverrides());
 			QualityInfo.DebugName = SoundWave->GetFullName();
 			//@todo tighten up the checking for empty results here
 			if(Format.CookSurround(FormatName, SourceBuffers, QualityInfo, Output))
@@ -1448,7 +1439,7 @@ void USoundWave::SerializeCookedPlatformData(FArchive& Ar)
 		check(!Ar.CookingTarget()->IsServerOnly());
 
 		FName PlatformFormat = Ar.CookingTarget()->GetWaveFormat(this);
-		const FPlatformAudioCookOverrides* CompressionOverrides = Ar.CookingTarget()->GetAudioCompressionSettings();
+		const FPlatformAudioCookOverrides* CompressionOverrides = FPlatformCompressionUtilities::GetCookOverrides(*Ar.CookingTarget()->IniPlatformName());
 		FString DerivedDataKey;
 
 		GetStreamedAudioDerivedDataKeySuffix(*this, PlatformFormat, CompressionOverrides, DerivedDataKey);
@@ -1521,13 +1512,14 @@ void USoundWave::BeginCachePlatformData()
 
 void USoundWave::BeginCacheForCookedPlatformData(const ITargetPlatform *TargetPlatform)
 {
-	if (TargetPlatform->SupportsFeature(ETargetPlatformFeatures::AudioStreaming) && IsStreaming())
+	const FPlatformAudioCookOverrides* CompressionOverrides = FPlatformCompressionUtilities::GetCookOverrides(*TargetPlatform->IniPlatformName());
+
+	if (TargetPlatform->SupportsFeature(ETargetPlatformFeatures::AudioStreaming) && IsStreaming(*CompressionOverrides))
 	{
 		// Retrieve format to cache for targetplatform.
 		FName PlatformFormat = TargetPlatform->GetWaveFormat(this);
 		uint32 CacheFlags = EStreamedAudioCacheFlags::Async | EStreamedAudioCacheFlags::InlineChunks;
 
-		const FPlatformAudioCookOverrides* CompressionOverrides = TargetPlatform->GetAudioCompressionSettings();
 
 		// If source data is resident in memory then allow the streamed audio to be built
 		// in a background thread.
@@ -1560,11 +1552,12 @@ void USoundWave::BeginCacheForCookedPlatformData(const ITargetPlatform *TargetPl
 
 bool USoundWave::IsCachedCookedPlatformDataLoaded( const ITargetPlatform* TargetPlatform )
 {
-	if (TargetPlatform->SupportsFeature(ETargetPlatformFeatures::AudioStreaming) && IsStreaming())
+	const FPlatformAudioCookOverrides* CompressionOverrides = FPlatformCompressionUtilities::GetCookOverrides(*TargetPlatform->IniPlatformName());
+
+	if (TargetPlatform->SupportsFeature(ETargetPlatformFeatures::AudioStreaming) && IsStreaming(*CompressionOverrides))
 	{
 		// Retrieve format to cache for targetplatform.
 		FName PlatformFormat = TargetPlatform->GetWaveFormat(this);
-		const FPlatformAudioCookOverrides* CompressionOverrides = TargetPlatform->GetAudioCompressionSettings();
 		// find format data by comparing derived data keys.
 		FString DerivedDataKey;
 		GetStreamedAudioDerivedDataKeySuffix(*this, PlatformFormat, CompressionOverrides, DerivedDataKey);
@@ -1607,11 +1600,12 @@ void USoundWave::ClearCachedCookedPlatformData( const ITargetPlatform* TargetPla
 {
 	Super::ClearCachedCookedPlatformData(TargetPlatform);
 
-	if (TargetPlatform->SupportsFeature(ETargetPlatformFeatures::AudioStreaming) && IsStreaming())
+	const FPlatformAudioCookOverrides* CompressionOverrides = FPlatformCompressionUtilities::GetCookOverrides(*TargetPlatform->IniPlatformName());
+
+	if (TargetPlatform->SupportsFeature(ETargetPlatformFeatures::AudioStreaming) && IsStreaming(*CompressionOverrides))
 	{
 		// Retrieve format to cache for targetplatform.
 		FName PlatformFormat = TargetPlatform->GetWaveFormat(this);
-		const FPlatformAudioCookOverrides* CompressionOverrides = TargetPlatform->GetAudioCompressionSettings();
 
 		// find format data by comparing derived data keys.
 		FString DerivedDataKey;
diff --git a/Engine/Source/Runtime/Engine/Private/AudioDevice.cpp b/Engine/Source/Runtime/Engine/Private/AudioDevice.cpp
--- a/Engine/Source/Runtime/Engine/Private/AudioDevice.cpp
+++ b/Engine/Source/Runtime/Engine/Private/AudioDevice.cpp
@@ -5838,7 +5838,7 @@ void FAudioDevice::Precache(USoundWave* SoundWave, bool bSynchronous, bool bTrac
 
 
 		// handle audio decompression
-		if (FPlatformProperties::SupportsAudioStreaming() && SoundWave->IsStreaming())
+		if (FPlatformProperties::SupportsAudioStreaming() && SoundWave->IsStreaming(nullptr))
 		{
 			SoundWave->DecompressionType = DTYPE_Streaming;
 			SoundWave->bCanProcessAsync = false;
diff --git a/Engine/Source/Runtime/Engine/Private/AudioStreaming.cpp b/Engine/Source/Runtime/Engine/Private/AudioStreaming.cpp
--- a/Engine/Source/Runtime/Engine/Private/AudioStreaming.cpp
+++ b/Engine/Source/Runtime/Engine/Private/AudioStreaming.cpp
@@ -709,7 +709,7 @@ void FLegacyAudioStreamingManager::NotifyLevelOffset(class ULevel* Level, const
 
 void FLegacyAudioStreamingManager::AddStreamingSoundWave(USoundWave* SoundWave)
 {
-	if (FPlatformProperties::SupportsAudioStreaming() && SoundWave->IsStreaming())
+	if (FPlatformProperties::SupportsAudioStreaming() && SoundWave->IsStreaming(nullptr))
 	{
 		FScopeLock Lock(&CriticalSection);
 
@@ -873,7 +873,7 @@ void FLegacyAudioStreamingManager::AddStreamingSoundSource(FSoundSource* SoundSo
 void FLegacyAudioStreamingManager::RemoveStreamingSoundSource(FSoundSource* SoundSource)
 {
 	const FWaveInstance* WaveInstance = SoundSource->GetWaveInstance();
-	if (WaveInstance && WaveInstance->WaveData && WaveInstance->WaveData->IsStreaming() && !WaveInstance->WaveData->ShouldUseStreamCaching())
+	if (WaveInstance && WaveInstance->WaveData && WaveInstance->WaveData->IsStreaming(nullptr) && !WaveInstance->WaveData->ShouldUseStreamCaching())
 	{
 		FScopeLock Lock(&CriticalSection);
 
diff --git a/Engine/Source/Runtime/Engine/Private/BatchedElements.cpp b/Engine/Source/Runtime/Engine/Private/BatchedElements.cpp
--- a/Engine/Source/Runtime/Engine/Private/BatchedElements.cpp
+++ b/Engine/Source/Runtime/Engine/Private/BatchedElements.cpp
@@ -444,7 +444,7 @@ float GBatchedElementSmoothWidth = 4;
 FAutoConsoleVariableRef CVarWellCanvasDistanceFieldSmoothWidth(TEXT("Canvas.DistanceFieldSmoothness"), GBatchedElementSmoothWidth, TEXT("Global sharpness of distance field fonts/shapes rendered by canvas."), ECVF_Default);
 
 template<class TSimpleElementPixelShader>
-static TSimpleElementPixelShader* GetPixelShader(bool bEncoded, ESimpleElementBlendMode BlendMode, ERHIFeatureLevel::Type FeatureLevel)
+static TShaderRef<TSimpleElementPixelShader> GetPixelShader(bool bEncoded, ESimpleElementBlendMode BlendMode, ERHIFeatureLevel::Type FeatureLevel)
 {
 	if (bEncoded)
 	{
@@ -452,34 +452,34 @@ static TSimpleElementPixelShader* GetPixelShader(bool bEncoded, ESimpleElementBl
 		switch (BlendMode)
 		{
 			case SE_BLEND_Opaque:
-				return *TShaderMapRef<FEncodedSimpleElement<TSimpleElementPixelShader, SE_BLEND_Opaque> >(GetGlobalShaderMap(FeatureLevel));
+				return TShaderMapRef<FEncodedSimpleElement<TSimpleElementPixelShader, SE_BLEND_Opaque> >(GetGlobalShaderMap(FeatureLevel));
 			case SE_BLEND_Masked:
-				return *TShaderMapRef<FEncodedSimpleElement<TSimpleElementPixelShader, SE_BLEND_Masked> >(GetGlobalShaderMap(FeatureLevel));
+				return TShaderMapRef<FEncodedSimpleElement<TSimpleElementPixelShader, SE_BLEND_Masked> >(GetGlobalShaderMap(FeatureLevel));
 			case SE_BLEND_Translucent:
-				return *TShaderMapRef<FEncodedSimpleElement<TSimpleElementPixelShader, SE_BLEND_Translucent> >(GetGlobalShaderMap(FeatureLevel));
+				return TShaderMapRef<FEncodedSimpleElement<TSimpleElementPixelShader, SE_BLEND_Translucent> >(GetGlobalShaderMap(FeatureLevel));
 			case SE_BLEND_Additive:
-				return *TShaderMapRef<FEncodedSimpleElement<TSimpleElementPixelShader, SE_BLEND_Additive> >(GetGlobalShaderMap(FeatureLevel));
+				return TShaderMapRef<FEncodedSimpleElement<TSimpleElementPixelShader, SE_BLEND_Additive> >(GetGlobalShaderMap(FeatureLevel));
 			case SE_BLEND_Modulate:
-				return *TShaderMapRef<FEncodedSimpleElement<TSimpleElementPixelShader, SE_BLEND_Modulate> >(GetGlobalShaderMap(FeatureLevel));
+				return TShaderMapRef<FEncodedSimpleElement<TSimpleElementPixelShader, SE_BLEND_Modulate> >(GetGlobalShaderMap(FeatureLevel));
 			case SE_BLEND_MaskedDistanceField:
-				return *TShaderMapRef<FEncodedSimpleElement<TSimpleElementPixelShader, SE_BLEND_MaskedDistanceField> >(GetGlobalShaderMap(FeatureLevel));
+				return TShaderMapRef<FEncodedSimpleElement<TSimpleElementPixelShader, SE_BLEND_MaskedDistanceField> >(GetGlobalShaderMap(FeatureLevel));
 			case SE_BLEND_MaskedDistanceFieldShadowed:
-				return *TShaderMapRef<FEncodedSimpleElement<TSimpleElementPixelShader, SE_BLEND_MaskedDistanceFieldShadowed> >(GetGlobalShaderMap(FeatureLevel));
+				return TShaderMapRef<FEncodedSimpleElement<TSimpleElementPixelShader, SE_BLEND_MaskedDistanceFieldShadowed> >(GetGlobalShaderMap(FeatureLevel));
 			case SE_BLEND_AlphaComposite:
-				return *TShaderMapRef<FEncodedSimpleElement<TSimpleElementPixelShader, SE_BLEND_AlphaComposite> >(GetGlobalShaderMap(FeatureLevel));
+				return TShaderMapRef<FEncodedSimpleElement<TSimpleElementPixelShader, SE_BLEND_AlphaComposite> >(GetGlobalShaderMap(FeatureLevel));
 			case SE_BLEND_AlphaHoldout:
-				return *TShaderMapRef<FEncodedSimpleElement<TSimpleElementPixelShader, SE_BLEND_AlphaHoldout> >(GetGlobalShaderMap(FeatureLevel));
+				return TShaderMapRef<FEncodedSimpleElement<TSimpleElementPixelShader, SE_BLEND_AlphaHoldout> >(GetGlobalShaderMap(FeatureLevel));
 			case SE_BLEND_AlphaBlend:
-				return *TShaderMapRef<FEncodedSimpleElement<TSimpleElementPixelShader, SE_BLEND_AlphaBlend> >(GetGlobalShaderMap(FeatureLevel));
+				return TShaderMapRef<FEncodedSimpleElement<TSimpleElementPixelShader, SE_BLEND_AlphaBlend> >(GetGlobalShaderMap(FeatureLevel));
 			case SE_BLEND_TranslucentAlphaOnly:
-				return *TShaderMapRef<FEncodedSimpleElement<TSimpleElementPixelShader, SE_BLEND_TranslucentAlphaOnly> >(GetGlobalShaderMap(FeatureLevel));
+				return TShaderMapRef<FEncodedSimpleElement<TSimpleElementPixelShader, SE_BLEND_TranslucentAlphaOnly> >(GetGlobalShaderMap(FeatureLevel));
 			case SE_BLEND_TranslucentAlphaOnlyWriteAlpha:
-				return *TShaderMapRef<FEncodedSimpleElement<TSimpleElementPixelShader, SE_BLEND_TranslucentAlphaOnlyWriteAlpha> >(GetGlobalShaderMap(FeatureLevel));
+				return TShaderMapRef<FEncodedSimpleElement<TSimpleElementPixelShader, SE_BLEND_TranslucentAlphaOnlyWriteAlpha> >(GetGlobalShaderMap(FeatureLevel));
 			default:
 				checkNoEntry();
 		}
 	}
-	return *TShaderMapRef<TSimpleElementPixelShader>(GetGlobalShaderMap(FeatureLevel));
+	return TShaderMapRef<TSimpleElementPixelShader>(GetGlobalShaderMap(FeatureLevel));
 }
 
 static bool Is32BppHDREncoded(const FSceneView* View, ERHIFeatureLevel::Type FeatureLevel)
@@ -612,14 +612,14 @@ void FBatchedElements::PrepareShaders(
 		TShaderMapRef<FSimpleElementVS> VertexShader(GetGlobalShaderMap(FeatureLevel));
 			
 		GraphicsPSOInit.BoundShaderState.VertexDeclarationRHI = GSimpleElementVertexDeclaration.VertexDeclarationRHI;
-		GraphicsPSOInit.BoundShaderState.VertexShaderRHI = GETSAFERHISHADER_VERTEX(*VertexShader);
+		GraphicsPSOInit.BoundShaderState.VertexShaderRHI = VertexShader.GetVertexShader();
 
 		if (bHitTesting)
 		{
 			SetHitTestingBlendState(RHICmdList, GraphicsPSOInit, BlendMode);
 
 			TShaderMapRef<FSimpleElementHitProxyPS> HitTestingPixelShader(GetGlobalShaderMap(FeatureLevel));
-			GraphicsPSOInit.BoundShaderState.PixelShaderRHI = GETSAFERHISHADER_PIXEL(*HitTestingPixelShader);
+			GraphicsPSOInit.BoundShaderState.PixelShaderRHI = HitTestingPixelShader.GetPixelShader();
 
 			SetGraphicsPipelineState(RHICmdList, GraphicsPSOInit);
 
@@ -635,9 +635,8 @@ void FBatchedElements::PrepareShaders(
 
 				if (Texture->bSRGB)
 				{
-					auto* MaskedPixelShader = GetPixelShader<FSimpleElementMaskedGammaPS_SRGB>(bEncodedHDR, BlendMode, FeatureLevel);
-					GraphicsPSOInit.BoundShaderState.PixelShaderRHI = GETSAFERHISHADER_PIXEL(MaskedPixelShader);
-
+					auto MaskedPixelShader = GetPixelShader<FSimpleElementMaskedGammaPS_SRGB>(bEncodedHDR, BlendMode, FeatureLevel);
+					GraphicsPSOInit.BoundShaderState.PixelShaderRHI = MaskedPixelShader.GetPixelShader();
 					SetGraphicsPipelineState(RHICmdList, GraphicsPSOInit);
 
 					MaskedPixelShader->SetEditorCompositingParameters(RHICmdList, View);
@@ -645,8 +644,8 @@ void FBatchedElements::PrepareShaders(
 				}
 				else
 				{
-					auto* MaskedPixelShader = GetPixelShader<FSimpleElementMaskedGammaPS_Linear>(bEncodedHDR, BlendMode, FeatureLevel);
-					GraphicsPSOInit.BoundShaderState.PixelShaderRHI = GETSAFERHISHADER_PIXEL(MaskedPixelShader);
+					auto MaskedPixelShader = GetPixelShader<FSimpleElementMaskedGammaPS_Linear>(bEncodedHDR, BlendMode, FeatureLevel);
+					GraphicsPSOInit.BoundShaderState.PixelShaderRHI = MaskedPixelShader.GetPixelShader();
 
 					SetGraphicsPipelineState(RHICmdList, GraphicsPSOInit);
 
@@ -679,7 +678,7 @@ void FBatchedElements::PrepareShaders(
 				}
 				
 				TShaderMapRef<FSimpleElementDistanceFieldGammaPS> DistanceFieldPixelShader(GetGlobalShaderMap(FeatureLevel));
-				GraphicsPSOInit.BoundShaderState.PixelShaderRHI = GETSAFERHISHADER_PIXEL(*DistanceFieldPixelShader);
+				GraphicsPSOInit.BoundShaderState.PixelShaderRHI = DistanceFieldPixelShader.GetPixelShader();
 
 				SetGraphicsPipelineState(RHICmdList, GraphicsPSOInit);
 
@@ -713,8 +712,8 @@ void FBatchedElements::PrepareShaders(
 
 				if (FMath::Abs(Gamma - 1.0f) < KINDA_SMALL_NUMBER)
 				{
-					auto* AlphaOnlyPixelShader = GetPixelShader<FSimpleElementAlphaOnlyPS>(bEncodedHDR, BlendMode, FeatureLevel);
-					GraphicsPSOInit.BoundShaderState.PixelShaderRHI = GETSAFERHISHADER_PIXEL(AlphaOnlyPixelShader);
+					auto AlphaOnlyPixelShader = GetPixelShader<FSimpleElementAlphaOnlyPS>(bEncodedHDR, BlendMode, FeatureLevel);
+					GraphicsPSOInit.BoundShaderState.PixelShaderRHI = AlphaOnlyPixelShader.GetPixelShader();
 
 					SetGraphicsPipelineState(RHICmdList, GraphicsPSOInit);
 
@@ -723,8 +722,8 @@ void FBatchedElements::PrepareShaders(
 				}
 				else
 				{
-					auto* GammaAlphaOnlyPixelShader = GetPixelShader<FSimpleElementGammaAlphaOnlyPS>(bEncodedHDR, BlendMode, FeatureLevel);
-					GraphicsPSOInit.BoundShaderState.PixelShaderRHI = GETSAFERHISHADER_PIXEL(GammaAlphaOnlyPixelShader);
+					auto GammaAlphaOnlyPixelShader = GetPixelShader<FSimpleElementGammaAlphaOnlyPS>(bEncodedHDR, BlendMode, FeatureLevel);
+					GraphicsPSOInit.BoundShaderState.PixelShaderRHI = GammaAlphaOnlyPixelShader.GetPixelShader();
 
 					SetGraphicsPipelineState(RHICmdList, GraphicsPSOInit);
 
@@ -735,7 +734,7 @@ void FBatchedElements::PrepareShaders(
 			else if(BlendMode >= SE_BLEND_RGBA_MASK_START && BlendMode <= SE_BLEND_RGBA_MASK_END)
 			{
 				TShaderMapRef<FSimpleElementColorChannelMaskPS> ColorChannelMaskPixelShader(GetGlobalShaderMap(FeatureLevel));
-				GraphicsPSOInit.BoundShaderState.PixelShaderRHI = GETSAFERHISHADER_PIXEL(*ColorChannelMaskPixelShader);
+				GraphicsPSOInit.BoundShaderState.PixelShaderRHI = ColorChannelMaskPixelShader.GetPixelShader();
 
 				SetGraphicsPipelineState(RHICmdList, GraphicsPSOInit);
 			
@@ -748,7 +747,7 @@ void FBatchedElements::PrepareShaders(
 				if (FMath::Abs(Gamma - 1.0f) < KINDA_SMALL_NUMBER)
 				{
 					TShaderMapRef<FSimpleElementPS> PixelShader(GetGlobalShaderMap(FeatureLevel));
-					GraphicsPSOInit.BoundShaderState.PixelShaderRHI = GETSAFERHISHADER_PIXEL(*PixelShader);
+					GraphicsPSOInit.BoundShaderState.PixelShaderRHI = PixelShader.GetPixelShader();
 
 					SetGraphicsPipelineState(RHICmdList, GraphicsPSOInit);
 
@@ -757,12 +756,19 @@ void FBatchedElements::PrepareShaders(
 				}
 				else
 				{
-					TShaderMapRef<FSimpleElementGammaPS_SRGB> PixelShader_SRGB(GetGlobalShaderMap(FeatureLevel));
-					TShaderMapRef<FSimpleElementGammaPS_Linear> PixelShader_Linear(GetGlobalShaderMap(FeatureLevel));
-					
-					FSimpleElementGammaBasePS* BasePixelShader = Texture->bSRGB ? static_cast<FSimpleElementGammaBasePS*>(*PixelShader_SRGB) : *PixelShader_Linear;
+					TShaderRef<FSimpleElementGammaBasePS> BasePixelShader;
+					if (Texture->bSRGB)
+					{
+						TShaderMapRef<FSimpleElementGammaPS_SRGB> PixelShader_SRGB(GetGlobalShaderMap(FeatureLevel));
+						BasePixelShader = PixelShader_SRGB;
+					}
+					else
+					{
+						TShaderMapRef<FSimpleElementGammaPS_Linear> PixelShader_Linear(GetGlobalShaderMap(FeatureLevel));
+						BasePixelShader = PixelShader_Linear;
+					}
 
-					GraphicsPSOInit.BoundShaderState.PixelShaderRHI = GETSAFERHISHADER_PIXEL(BasePixelShader);
+					GraphicsPSOInit.BoundShaderState.PixelShaderRHI = BasePixelShader.GetPixelShader();
 					SetGraphicsPipelineState(RHICmdList, GraphicsPSOInit);
 
 					BasePixelShader->SetParameters(RHICmdList, Texture, Gamma, BlendMode);
diff --git a/Engine/Source/Runtime/Engine/Private/Components/ActorComponent.cpp b/Engine/Source/Runtime/Engine/Private/Components/ActorComponent.cpp
--- a/Engine/Source/Runtime/Engine/Private/Components/ActorComponent.cpp
+++ b/Engine/Source/Runtime/Engine/Private/Components/ActorComponent.cpp
@@ -1096,7 +1096,7 @@ void UActorComponent::TickComponent(float DeltaTime, enum ELevelTick TickType, F
 	}
 }
 
-void UActorComponent::RegisterComponentWithWorld(UWorld* InWorld)
+void UActorComponent::RegisterComponentWithWorld(UWorld* InWorld, FRegisterComponentContext* Context)
 {
 	SCOPE_CYCLE_COUNTER(STAT_RegisterComponent);
 	FScopeCycleCounterUObject ComponentScope(this);
@@ -1157,7 +1157,7 @@ void UActorComponent::RegisterComponentWithWorld(UWorld* InWorld)
 
 	WorldPrivate = InWorld;
 
-	ExecuteRegisterEvents();
+	ExecuteRegisterEvents(Context);
 
 	// If not in a game world register ticks now, otherwise defer until BeginPlay. If no owner we won't trigger BeginPlay either so register now in that case as well.
 	if (!InWorld->IsGameWorld())
@@ -1320,7 +1320,7 @@ void UActorComponent::K2_DestroyComponent(UObject* Object)
 	}
 }
 
-void UActorComponent::CreateRenderState_Concurrent()
+void UActorComponent::CreateRenderState_Concurrent(FRegisterComponentContext* Context)
 {
 	check(IsRegistered());
 	check(WorldPrivate->Scene);
@@ -1423,7 +1423,7 @@ void UActorComponent::DestroyPhysicsState()
 	}
 }
 
-void UActorComponent::ExecuteRegisterEvents()
+void UActorComponent::ExecuteRegisterEvents(FRegisterComponentContext* Context)
 {
 	if(!bRegistered)
 	{
@@ -1436,7 +1436,7 @@ void UActorComponent::ExecuteRegisterEvents()
 	{
 		SCOPE_CYCLE_COUNTER(STAT_ComponentCreateRenderState);
 		LLM_SCOPE(ELLMTag::SceneRender);
-		CreateRenderState_Concurrent();
+		CreateRenderState_Concurrent(Context);
 		checkf(bRenderStateCreated, TEXT("Failed to route CreateRenderState_Concurrent (%s)"), *GetFullName());
 	}
 
@@ -1486,7 +1486,7 @@ void UActorComponent::RecreateRenderState_Concurrent()
 
 	if(IsRegistered() && WorldPrivate->Scene)
 	{
-		CreateRenderState_Concurrent();
+		CreateRenderState_Concurrent(nullptr);
 		checkf(bRenderStateCreated, TEXT("Failed to route CreateRenderState_Concurrent (%s)"), *GetFullName());
 	}
 }
diff --git a/Engine/Source/Runtime/Engine/Private/Components/DecalComponent.cpp b/Engine/Source/Runtime/Engine/Private/Components/DecalComponent.cpp
--- a/Engine/Source/Runtime/Engine/Private/Components/DecalComponent.cpp
+++ b/Engine/Source/Runtime/Engine/Private/Components/DecalComponent.cpp
@@ -280,9 +280,9 @@ void UDecalComponent::BeginPlay()
 	SetLifeSpan(FadeStartDelay + FadeDuration);
 }
 
-void UDecalComponent::CreateRenderState_Concurrent()
+void UDecalComponent::CreateRenderState_Concurrent(FRegisterComponentContext* Context)
 {
-	Super::CreateRenderState_Concurrent();
+	Super::CreateRenderState_Concurrent(Context);
 
 	// Mimics UPrimitiveComponent's visibility logic, although without the UPrimitiveCompoent visibility flags
 	if ( ShouldComponentAddToScene() && ShouldRender() )
diff --git a/Engine/Source/Runtime/Engine/Private/Components/HeightFogComponent.cpp b/Engine/Source/Runtime/Engine/Private/Components/HeightFogComponent.cpp
--- a/Engine/Source/Runtime/Engine/Private/Components/HeightFogComponent.cpp
+++ b/Engine/Source/Runtime/Engine/Private/Components/HeightFogComponent.cpp
@@ -56,9 +56,9 @@ void UExponentialHeightFogComponent::AddFogIfNeeded()
 	}
 }
 
-void UExponentialHeightFogComponent::CreateRenderState_Concurrent()
+void UExponentialHeightFogComponent::CreateRenderState_Concurrent(FRegisterComponentContext* Context)
 {
-	Super::CreateRenderState_Concurrent();
+	Super::CreateRenderState_Concurrent(Context);
 	AddFogIfNeeded();
 }
 
diff --git a/Engine/Source/Runtime/Engine/Private/Components/LightComponent.cpp b/Engine/Source/Runtime/Engine/Private/Components/LightComponent.cpp
--- a/Engine/Source/Runtime/Engine/Private/Components/LightComponent.cpp
+++ b/Engine/Source/Runtime/Engine/Private/Components/LightComponent.cpp
@@ -775,9 +775,9 @@ void ULightComponent::OnRegister()
 	ValidateLightGUIDs();
 }
 
-void ULightComponent::CreateRenderState_Concurrent()
+void ULightComponent::CreateRenderState_Concurrent(FRegisterComponentContext* Context)
 {
-	Super::CreateRenderState_Concurrent();
+	Super::CreateRenderState_Concurrent(Context);
 
 	if (bAffectsWorld)
 	{
diff --git a/Engine/Source/Runtime/Engine/Private/Components/PortalComponent.cpp b/Engine/Source/Runtime/Engine/Private/Components/PortalComponent.cpp
--- a/Engine/Source/Runtime/Engine/Private/Components/PortalComponent.cpp
+++ b/Engine/Source/Runtime/Engine/Private/Components/PortalComponent.cpp
@@ -76,9 +76,9 @@ void ULightmassPortalComponent::UpdatePreviewShape()
 	}
 }
 
-void ULightmassPortalComponent::CreateRenderState_Concurrent()
+void ULightmassPortalComponent::CreateRenderState_Concurrent(FRegisterComponentContext* Context)
 {
-	Super::CreateRenderState_Concurrent();
+	Super::CreateRenderState_Concurrent(Context);
 
 	UpdatePreviewShape();
 }
diff --git a/Engine/Source/Runtime/Engine/Private/Components/PrimitiveComponent.cpp b/Engine/Source/Runtime/Engine/Private/Components/PrimitiveComponent.cpp
--- a/Engine/Source/Runtime/Engine/Private/Components/PrimitiveComponent.cpp
+++ b/Engine/Source/Runtime/Engine/Private/Components/PrimitiveComponent.cpp
@@ -536,7 +536,7 @@ FORCEINLINE_DEBUGGABLE bool OwnerLevelHasRegisteredStaticComponentsInStreamingMa
 	return false;
 }
 
-void UPrimitiveComponent::CreateRenderState_Concurrent()
+void UPrimitiveComponent::CreateRenderState_Concurrent(FRegisterComponentContext* Context)
 {
 	// Make sure cached cull distance is up-to-date if its zero and we have an LD cull distance
 	if( CachedMaxDrawDistance == 0.f && LDMaxDrawDistance > 0.f )
@@ -545,14 +545,21 @@ void UPrimitiveComponent::CreateRenderState_Concurrent()
 		CachedMaxDrawDistance = bNeverCull ? 0.f : LDMaxDrawDistance;
 	}
 
-	Super::CreateRenderState_Concurrent();
+	Super::CreateRenderState_Concurrent(Context);
 
 	UpdateBounds();
 
 	// If the primitive isn't hidden and the detail mode setting allows it, add it to the scene.
 	if (ShouldComponentAddToScene())
 	{
-		GetWorld()->Scene->AddPrimitive(this);
+		if (Context != nullptr)
+		{
+			Context->AddPrimitiveBatches.Add(this);
+		}
+		else
+		{
+			GetWorld()->Scene->AddPrimitive(this);
+		}
 	}
 
 	// Components are either registered as static or dynamic in the streaming manager.
diff --git a/Engine/Source/Runtime/Engine/Private/Components/ReflectionCaptureComponent.cpp b/Engine/Source/Runtime/Engine/Private/Components/ReflectionCaptureComponent.cpp
--- a/Engine/Source/Runtime/Engine/Private/Components/ReflectionCaptureComponent.cpp
+++ b/Engine/Source/Runtime/Engine/Private/Components/ReflectionCaptureComponent.cpp
@@ -732,9 +732,9 @@ UReflectionCaptureComponent::UReflectionCaptureComponent(const FObjectInitialize
 	bNeedsRecaptureOrUpload = false;
 }
 
-void UReflectionCaptureComponent::CreateRenderState_Concurrent()
+void UReflectionCaptureComponent::CreateRenderState_Concurrent(FRegisterComponentContext* Context)
 {
-	Super::CreateRenderState_Concurrent();
+	Super::CreateRenderState_Concurrent(Context);
 
 	UpdatePreviewShape();
 
diff --git a/Engine/Source/Runtime/Engine/Private/Components/RuntimeVirtualTextureComponent.cpp b/Engine/Source/Runtime/Engine/Private/Components/RuntimeVirtualTextureComponent.cpp
--- a/Engine/Source/Runtime/Engine/Private/Components/RuntimeVirtualTextureComponent.cpp
+++ b/Engine/Source/Runtime/Engine/Private/Components/RuntimeVirtualTextureComponent.cpp
@@ -19,7 +19,7 @@ bool URuntimeVirtualTextureComponent::IsVisible() const
 	return Super::IsVisible() && UseVirtualTexturing(GetScene()->GetFeatureLevel());
 }
 
-void URuntimeVirtualTextureComponent::CreateRenderState_Concurrent()
+void URuntimeVirtualTextureComponent::CreateRenderState_Concurrent(FRegisterComponentContext* Context)
 {
 	if (ShouldRender() && VirtualTexture != nullptr)
 	{
@@ -27,7 +27,7 @@ void URuntimeVirtualTextureComponent::CreateRenderState_Concurrent()
 		GetScene()->AddRuntimeVirtualTexture(this);
 	}
 
-	Super::CreateRenderState_Concurrent();
+	Super::CreateRenderState_Concurrent(Context);
 }
 
 void URuntimeVirtualTextureComponent::SendRenderTransform_Concurrent()
diff --git a/Engine/Source/Runtime/Engine/Private/Components/SceneCaptureComponent.cpp b/Engine/Source/Runtime/Engine/Private/Components/SceneCaptureComponent.cpp
--- a/Engine/Source/Runtime/Engine/Private/Components/SceneCaptureComponent.cpp
+++ b/Engine/Source/Runtime/Engine/Private/Components/SceneCaptureComponent.cpp
@@ -850,11 +850,11 @@ void UPlanarReflectionComponent::Serialize(FArchive& Ar)
 	}
 }
 
-void UPlanarReflectionComponent::CreateRenderState_Concurrent()
+void UPlanarReflectionComponent::CreateRenderState_Concurrent(FRegisterComponentContext* Context)
 {
 	UpdatePreviewShape();
 
-	Super::CreateRenderState_Concurrent();
+	Super::CreateRenderState_Concurrent(Context);
 
 	if (ShouldComponentAddToScene() && ShouldRender())
 	{
diff --git a/Engine/Source/Runtime/Engine/Private/Components/SkinnedMeshComponent.cpp b/Engine/Source/Runtime/Engine/Private/Components/SkinnedMeshComponent.cpp
--- a/Engine/Source/Runtime/Engine/Private/Components/SkinnedMeshComponent.cpp
+++ b/Engine/Source/Runtime/Engine/Private/Components/SkinnedMeshComponent.cpp
@@ -511,7 +511,7 @@ void USkinnedMeshComponent::OnUnregister()
 	}
 }
 
-void USkinnedMeshComponent::CreateRenderState_Concurrent()
+void USkinnedMeshComponent::CreateRenderState_Concurrent(FRegisterComponentContext* Context)
 {
 	LLM_SCOPE(ELLMTag::SkeletalMesh);
 
@@ -572,7 +572,7 @@ void USkinnedMeshComponent::CreateRenderState_Concurrent()
 		}
 	}
 
-	Super::CreateRenderState_Concurrent();
+	Super::CreateRenderState_Concurrent(Context);
 
 	if (SkeletalMesh)
 	{
diff --git a/Engine/Source/Runtime/Engine/Private/Components/SkyAtmosphereComponent.cpp b/Engine/Source/Runtime/Engine/Private/Components/SkyAtmosphereComponent.cpp
--- a/Engine/Source/Runtime/Engine/Private/Components/SkyAtmosphereComponent.cpp
+++ b/Engine/Source/Runtime/Engine/Private/Components/SkyAtmosphereComponent.cpp
@@ -117,9 +117,9 @@ static bool SkyAtmosphereComponentStaticLightingBuilt(const USkyAtmosphereCompon
 	return true;	// The component has not been spawned in any world yet so let's mark it as built for now.
 }
 
-void USkyAtmosphereComponent::CreateRenderState_Concurrent()
+void USkyAtmosphereComponent::CreateRenderState_Concurrent(FRegisterComponentContext* Context)
 {
-	Super::CreateRenderState_Concurrent();
+	Super::CreateRenderState_Concurrent(Context);
 	// If one day we need to look up lightmass built data, lookup it up here using the guid from the correct MapBuildData.
 
 	bool bHidden = false;
diff --git a/Engine/Source/Runtime/Engine/Private/Components/SkyLightComponent.cpp b/Engine/Source/Runtime/Engine/Private/Components/SkyLightComponent.cpp
--- a/Engine/Source/Runtime/Engine/Private/Components/SkyLightComponent.cpp
+++ b/Engine/Source/Runtime/Engine/Private/Components/SkyLightComponent.cpp
@@ -333,9 +333,9 @@ TArray<USkyLightComponent*> USkyLightComponent::SkyCapturesToUpdate;
 TArray<USkyLightComponent*> USkyLightComponent::SkyCapturesToUpdateBlendDestinations;
 FCriticalSection USkyLightComponent::SkyCapturesToUpdateLock;
 
-void USkyLightComponent::CreateRenderState_Concurrent()
+void USkyLightComponent::CreateRenderState_Concurrent(FRegisterComponentContext* Context)
 {
-	Super::CreateRenderState_Concurrent();
+	Super::CreateRenderState_Concurrent(Context);
 
 	bool bHidden = false;
 #if WITH_EDITORONLY_DATA
diff --git a/Engine/Source/Runtime/Engine/Private/Components/SplineMeshComponent.cpp b/Engine/Source/Runtime/Engine/Private/Components/SplineMeshComponent.cpp
--- a/Engine/Source/Runtime/Engine/Private/Components/SplineMeshComponent.cpp
+++ b/Engine/Source/Runtime/Engine/Private/Components/SplineMeshComponent.cpp
@@ -91,13 +91,9 @@ void FSplineMeshVertexFactoryShaderParameters::GetElementShaderBindings(
 //////////////////////////////////////////////////////////////////////////
 // SplineMeshVertexFactory
 
-IMPLEMENT_VERTEX_FACTORY_TYPE(FSplineMeshVertexFactory, "/Engine/Private/LocalVertexFactory.ush", true, true, true, true, true);
-
+IMPLEMENT_VERTEX_FACTORY_PARAMETER_TYPE(FSplineMeshVertexFactory, SF_Vertex, FSplineMeshVertexFactoryShaderParameters);
 
-FVertexFactoryShaderParameters* FSplineMeshVertexFactory::ConstructShaderParameters(EShaderFrequency ShaderFrequency)
-{
-	return ShaderFrequency == SF_Vertex ? new FSplineMeshVertexFactoryShaderParameters() : NULL;
-}
+IMPLEMENT_VERTEX_FACTORY_TYPE(FSplineMeshVertexFactory, "/Engine/Private/LocalVertexFactory.ush", true, true, true, true, true);
 
 //////////////////////////////////////////////////////////////////////////
 // SplineMeshSceneProxy
diff --git a/Engine/Source/Runtime/Engine/Private/Components/StaticMeshComponent.cpp b/Engine/Source/Runtime/Engine/Private/Components/StaticMeshComponent.cpp
--- a/Engine/Source/Runtime/Engine/Private/Components/StaticMeshComponent.cpp
+++ b/Engine/Source/Runtime/Engine/Private/Components/StaticMeshComponent.cpp
@@ -594,10 +594,10 @@ void UStaticMeshComponent::OnUnregister()
 	Super::OnUnregister();
 }
 
-void UStaticMeshComponent::CreateRenderState_Concurrent()
+void UStaticMeshComponent::CreateRenderState_Concurrent(FRegisterComponentContext* Context)
 {
 	LLM_SCOPE(ELLMTag::StaticMesh);
-	Super::CreateRenderState_Concurrent();
+	Super::CreateRenderState_Concurrent(Context);
 }
 
 void UStaticMeshComponent::OnCreatePhysicsState()
diff --git a/Engine/Source/Runtime/Engine/Private/Components/TextRenderComponent.cpp b/Engine/Source/Runtime/Engine/Private/Components/TextRenderComponent.cpp
--- a/Engine/Source/Runtime/Engine/Private/Components/TextRenderComponent.cpp
+++ b/Engine/Source/Runtime/Engine/Private/Components/TextRenderComponent.cpp
@@ -307,6 +307,7 @@ float CalculateVerticalAlignmentOffset(
 /** Caches MIDs used by text render components to avoid excessive (re)allocation of MIDs when the SCS runs */
 class FTextRenderComponentMIDCache : public FGCObject
 {
+	FCriticalSection CriticalSection;
 public:
 	/** Array of MIDs for a particular material and font */
 	struct FMIDData
@@ -419,7 +420,7 @@ public:
 
 	FMIDDataRef GetMIDData(UMaterialInterface* InMaterial, UFont* InFont)
 	{
-		checkfSlow(IsInGameThread(), TEXT("FTextRenderComponentMIDCache::GetMIDData is only expected to be called from the game thread!"));
+		FScopeLock Lock(&CriticalSection);
 
 		check(InMaterial && InFont && InFont->FontCacheType == EFontCacheType::Offline);
 
@@ -440,6 +441,8 @@ public:
 
 	virtual void AddReferencedObjects(FReferenceCollector& Collector) override
 	{
+		FScopeLock Lock(&CriticalSection);
+
 		for (auto& MIDDataPair : CachedMIDs)
 		{
 			const FMIDDataPtr& MIDData = MIDDataPair.Value;
@@ -519,6 +522,8 @@ private:
 
 	void PurgeUnreferencedMIDs()
 	{
+		FScopeLock Lock(&CriticalSection);
+
         QUICK_SCOPE_CYCLE_COUNTER(STAT_FTextRenderComponentMIDCache_PurgeUnreferencedMIDs);
 
 		checkfSlow(IsInGameThread(), TEXT("FTextRenderComponentMIDCache::PurgeUnreferencedMIDs is only expected to be called from the game thread!"));
diff --git a/Engine/Source/Runtime/Engine/Private/CubemapUnwrapUtils.cpp b/Engine/Source/Runtime/Engine/Private/CubemapUnwrapUtils.cpp
--- a/Engine/Source/Runtime/Engine/Private/CubemapUnwrapUtils.cpp
+++ b/Engine/Source/Runtime/Engine/Private/CubemapUnwrapUtils.cpp
@@ -117,18 +117,18 @@ namespace CubemapHelpers
 
 void FCubemapTexturePropertiesVS::SetParameters( FRHICommandList& RHICmdList, const FMatrix& TransformValue )
 {
-	SetShaderValue(RHICmdList, GetVertexShader(), Transform, TransformValue);
+	SetShaderValue(RHICmdList, RHICmdList.GetBoundVertexShader(), Transform, TransformValue);
 }
 
 template<bool bHDROutput>
 void FCubemapTexturePropertiesPS<bHDROutput>::SetParameters( FRHICommandList& RHICmdList, const FTexture* Texture, const FMatrix& ColorWeightsValue, float MipLevel, float GammaValue )
 {
-	SetTextureParameter(RHICmdList, GetPixelShader(),CubeTexture,CubeTextureSampler,Texture);
+	SetTextureParameter(RHICmdList, RHICmdList.GetBoundPixelShader(),CubeTexture,CubeTextureSampler,Texture);
 
 	FVector4 PackedProperties0Value(MipLevel, 0, 0, 0);
-	SetShaderValue(RHICmdList, GetPixelShader(), PackedProperties0, PackedProperties0Value);
-	SetShaderValue(RHICmdList, GetPixelShader(), ColorWeights, ColorWeightsValue);
-	SetShaderValue(RHICmdList, GetPixelShader(), Gamma, GammaValue);
+	SetShaderValue(RHICmdList, RHICmdList.GetBoundPixelShader(), PackedProperties0, PackedProperties0Value);
+	SetShaderValue(RHICmdList, RHICmdList.GetBoundPixelShader(), ColorWeights, ColorWeightsValue);
+	SetShaderValue(RHICmdList, RHICmdList.GetBoundPixelShader(), Gamma, GammaValue);
 }
 
 void FMipLevelBatchedElementParameters::BindShaders(FRHICommandList& RHICmdList, FGraphicsPipelineStateInitializer& GraphicsPSOInit, ERHIFeatureLevel::Type InFeatureLevel, const FMatrix& InTransform, const float InGamma, const FMatrix& ColorWeights, const FTexture* Texture)
@@ -152,8 +152,8 @@ void FMipLevelBatchedElementParameters::BindShaders(FRHICommandList& RHICmdList,
 	TShaderMapRef<TPixelShader> PixelShader(GetGlobalShaderMap(InFeatureLevel));
 	
 	GraphicsPSOInit.BoundShaderState.VertexDeclarationRHI = GSimpleElementVertexDeclaration.VertexDeclarationRHI;
-	GraphicsPSOInit.BoundShaderState.VertexShaderRHI = GETSAFERHISHADER_VERTEX(*VertexShader);
-	GraphicsPSOInit.BoundShaderState.PixelShaderRHI = GETSAFERHISHADER_PIXEL(*PixelShader);
+	GraphicsPSOInit.BoundShaderState.VertexShaderRHI = VertexShader.GetVertexShader();
+	GraphicsPSOInit.BoundShaderState.PixelShaderRHI = PixelShader.GetPixelShader();
 	GraphicsPSOInit.PrimitiveType = PT_TriangleList;
 
 	SetGraphicsPipelineState(RHICmdList, GraphicsPSOInit, EApplyRendertargetOption::ForceApply);
@@ -164,7 +164,7 @@ void FMipLevelBatchedElementParameters::BindShaders(FRHICommandList& RHICmdList,
 
 void FIESLightProfilePS::SetParameters( FRHICommandList& RHICmdList, const FTexture* Texture, float InBrightnessInLumens )
 {
-	FRHIPixelShader* ShaderRHI = GetPixelShader();
+	FRHIPixelShader* ShaderRHI = RHICmdList.GetBoundPixelShader();
 	SetTextureParameter(RHICmdList, ShaderRHI, IESTexture, IESTextureSampler, Texture);
 
 	SetShaderValue(RHICmdList, ShaderRHI, BrightnessInLumens, InBrightnessInLumens);
@@ -178,8 +178,8 @@ void FIESLightProfileBatchedElementParameters::BindShaders( FRHICommandList& RHI
 	TShaderMapRef<FIESLightProfilePS> PixelShader(GetGlobalShaderMap(InFeatureLevel));
 
 	GraphicsPSOInit.BoundShaderState.VertexDeclarationRHI = GSimpleElementVertexDeclaration.VertexDeclarationRHI;
-	GraphicsPSOInit.BoundShaderState.VertexShaderRHI = GETSAFERHISHADER_VERTEX(*VertexShader);
-	GraphicsPSOInit.BoundShaderState.PixelShaderRHI = GETSAFERHISHADER_PIXEL(*PixelShader);
+	GraphicsPSOInit.BoundShaderState.VertexShaderRHI = VertexShader.GetVertexShader();
+	GraphicsPSOInit.BoundShaderState.PixelShaderRHI = PixelShader.GetPixelShader();
 	GraphicsPSOInit.PrimitiveType = PT_TriangleList;
 
 	SetGraphicsPipelineState(RHICmdList, GraphicsPSOInit, EApplyRendertargetOption::ForceApply);
diff --git a/Engine/Source/Runtime/Engine/Private/DebugViewModeInterface.cpp b/Engine/Source/Runtime/Engine/Private/DebugViewModeInterface.cpp
--- a/Engine/Source/Runtime/Engine/Private/DebugViewModeInterface.cpp
+++ b/Engine/Source/Runtime/Engine/Private/DebugViewModeInterface.cpp
@@ -37,14 +37,18 @@ void FDebugViewModeInterface::SetInterface(EDebugViewShaderMode InDebugViewMode,
 	}
 }
 
+bool FDebugViewModeInterface::AllowFallbackToDefaultMaterial(EMaterialTessellationMode TessellationMode, bool bHasVertexPositionOffsetConnected, bool bHasPixelDepthOffsetConnected)
+{
+	// Check for anything that could change the shape from the default material.
+	return !bHasVertexPositionOffsetConnected &&
+		!bHasPixelDepthOffsetConnected &&
+		TessellationMode == MTM_NoTessellation;
+}
+
 bool FDebugViewModeInterface::AllowFallbackToDefaultMaterial(const FMaterial* InMaterial)
 {
 	check(InMaterial);
-
-	// Check for anything that could change the shape from the default material.
-	return !InMaterial->HasVertexPositionOffsetConnected() && 
-		!InMaterial->HasPixelDepthOffsetConnected() &&
-		InMaterial->GetTessellationMode() == MTM_NoTessellation; 
+	return FDebugViewModeInterface::AllowFallbackToDefaultMaterial(InMaterial->GetTessellationMode(), InMaterial->HasVertexPositionOffsetConnected(), InMaterial->HasPixelDepthOffsetConnected());
 }
 
 
diff --git a/Engine/Source/Runtime/Engine/Private/DebugViewModeMaterialProxy.cpp b/Engine/Source/Runtime/Engine/Private/DebugViewModeMaterialProxy.cpp
--- a/Engine/Source/Runtime/Engine/Private/DebugViewModeMaterialProxy.cpp
+++ b/Engine/Source/Runtime/Engine/Private/DebugViewModeMaterialProxy.cpp
@@ -35,7 +35,7 @@ FDebugViewModeMaterialProxy::FDebugViewModeMaterialProxy(
 	const EShaderPlatform ShaderPlatform = GetFeatureLevelShaderPlatform(FeatureLevel);
 
 	Material = InMaterialInterface->GetMaterial();
-	MaterialInterface->AppendReferencedTextures(ReferencedTextures);
+	ReferencedTextures = MaterialInterface->GetReferencedTextures();
 
 	FMaterialResource* Resource = InMaterialInterface->GetMaterialResource(FeatureLevel);
 	if (Resource)
@@ -68,9 +68,6 @@ FDebugViewModeMaterialProxy::FDebugViewModeMaterialProxy(
 		FMaterialShaderMapId ResourceId;
 		Resource->GetShaderMapId(ShaderPlatform, ResourceId);
 
-		FStaticParameterSet StaticParamSet;
-		Resource->GetStaticParameterSet(ShaderPlatform, StaticParamSet);
-
 		{
 			TArray<FShaderType*> ShaderTypes;
 			TArray<FVertexFactoryType*> VFTypes;
@@ -84,7 +81,7 @@ FDebugViewModeMaterialProxy::FDebugViewModeMaterialProxy(
 
 		ResourceId.Usage = Usage;
 
-		CacheShaders(ResourceId, &StaticParamSet, ShaderPlatform);
+		CacheShaders(ResourceId, ShaderPlatform);
 	}
 	else
 	{
@@ -110,22 +107,22 @@ const FMaterial& FDebugViewModeMaterialProxy::GetMaterialWithFallback(ERHIFeatur
 	}
 }
 
-bool FDebugViewModeMaterialProxy::GetVectorValue(const FMaterialParameterInfo& ParameterInfo, FLinearColor* OutValue, const FMaterialRenderContext& Context) const
+bool FDebugViewModeMaterialProxy::GetVectorValue(const FHashedMaterialParameterInfo& ParameterInfo, FLinearColor* OutValue, const FMaterialRenderContext& Context) const
 {
 	return MaterialInterface->GetRenderProxy()->GetVectorValue(ParameterInfo, OutValue, Context);
 }
 
-bool FDebugViewModeMaterialProxy::GetScalarValue(const FMaterialParameterInfo& ParameterInfo, float* OutValue, const FMaterialRenderContext& Context) const
+bool FDebugViewModeMaterialProxy::GetScalarValue(const FHashedMaterialParameterInfo& ParameterInfo, float* OutValue, const FMaterialRenderContext& Context) const
 {
 	return MaterialInterface->GetRenderProxy()->GetScalarValue(ParameterInfo, OutValue, Context);
 }
 
-bool FDebugViewModeMaterialProxy::GetTextureValue(const FMaterialParameterInfo& ParameterInfo,const UTexture** OutValue, const FMaterialRenderContext& Context) const
+bool FDebugViewModeMaterialProxy::GetTextureValue(const FHashedMaterialParameterInfo& ParameterInfo,const UTexture** OutValue, const FMaterialRenderContext& Context) const
 {
 	return MaterialInterface->GetRenderProxy()->GetTextureValue(ParameterInfo,OutValue,Context);
 }
 
-bool FDebugViewModeMaterialProxy::GetTextureValue(const FMaterialParameterInfo& ParameterInfo, const URuntimeVirtualTexture** OutValue, const FMaterialRenderContext& Context) const
+bool FDebugViewModeMaterialProxy::GetTextureValue(const FHashedMaterialParameterInfo& ParameterInfo, const URuntimeVirtualTexture** OutValue, const FMaterialRenderContext& Context) const
 {
 	return MaterialInterface->GetRenderProxy()->GetTextureValue(ParameterInfo, OutValue, Context);
 }
diff --git a/Engine/Source/Runtime/Engine/Private/DebugViewModeMaterialProxy.h b/Engine/Source/Runtime/Engine/Private/DebugViewModeMaterialProxy.h
--- a/Engine/Source/Runtime/Engine/Private/DebugViewModeMaterialProxy.h
+++ b/Engine/Source/Runtime/Engine/Private/DebugViewModeMaterialProxy.h
@@ -80,7 +80,7 @@ public:
 		return  false;
 	}
 
-	virtual const TArray<UObject*>& GetReferencedTextures() const override
+	virtual TArrayView<UObject* const> GetReferencedTextures() const override
 	{
 		return ReferencedTextures;
 	}
@@ -127,10 +127,10 @@ public:
 	////////////////
 	// FMaterialRenderProxy interface.
 	virtual const FMaterial& GetMaterialWithFallback(ERHIFeatureLevel::Type InFeatureLevel, const FMaterialRenderProxy*& OutFallbackMaterialRenderProxy) const override;
-	virtual bool GetVectorValue(const FMaterialParameterInfo& ParameterInfo, FLinearColor* OutValue, const FMaterialRenderContext& Context) const override;
-	virtual bool GetScalarValue(const FMaterialParameterInfo& ParameterInfo, float* OutValue, const FMaterialRenderContext& Context) const override;
-	virtual bool GetTextureValue(const FMaterialParameterInfo& ParameterInfo, const UTexture** OutValue, const FMaterialRenderContext& Context) const override;
-	virtual bool GetTextureValue(const FMaterialParameterInfo& ParameterInfo, const URuntimeVirtualTexture** OutValue, const FMaterialRenderContext& Context) const override;
+	virtual bool GetVectorValue(const FHashedMaterialParameterInfo& ParameterInfo, FLinearColor* OutValue, const FMaterialRenderContext& Context) const override;
+	virtual bool GetScalarValue(const FHashedMaterialParameterInfo& ParameterInfo, float* OutValue, const FMaterialRenderContext& Context) const override;
+	virtual bool GetTextureValue(const FHashedMaterialParameterInfo& ParameterInfo,const UTexture** OutValue, const FMaterialRenderContext& Context) const override;
+	virtual bool GetTextureValue(const FHashedMaterialParameterInfo& ParameterInfo, const URuntimeVirtualTexture** OutValue, const FMaterialRenderContext& Context) const override;
 
 	virtual EMaterialDomain GetMaterialDomain() const override;
 	virtual bool IsTwoSided() const  override;
diff --git a/Engine/Source/Runtime/Engine/Private/DeviceProfiles/DeviceProfileManager.cpp b/Engine/Source/Runtime/Engine/Private/DeviceProfiles/DeviceProfileManager.cpp
--- a/Engine/Source/Runtime/Engine/Private/DeviceProfiles/DeviceProfileManager.cpp
+++ b/Engine/Source/Runtime/Engine/Private/DeviceProfiles/DeviceProfileManager.cpp
@@ -34,7 +34,6 @@ static TAutoConsoleVariable<int32> CVarAllowScalabilityGroupsToChangeAtRuntime(
 	TEXT("priority which allows them to be changed at runtime. Off by default."),
 	ECVF_Default);
 
-FString UDeviceProfileManager::DeviceProfileFileName;
 TMap<FString, FString> UDeviceProfileManager::DeviceProfileScalabilityCVars;
 
 UDeviceProfileManager* UDeviceProfileManager::DeviceProfileManagerSingleton = nullptr;
@@ -89,11 +88,8 @@ void UDeviceProfileManager::InitializeCVarsForActiveDeviceProfile(bool bPushSett
 
 	UE_LOG(LogInit, Log, TEXT("Applying CVar settings loaded from the selected device profile: [%s]"), *ActiveProfileName);
 
-	// Load the device profile config
-	FConfigCacheIni::LoadGlobalIniFile(DeviceProfileFileName, TEXT("DeviceProfiles"));
-
 	TArray< FString > AvailableProfiles;
-	GConfig->GetSectionNames( DeviceProfileFileName, AvailableProfiles );
+	GConfig->GetSectionNames( GDeviceProfilesIni, AvailableProfiles );
 
 	// Look up the ini for this tree as we are far too early to use the UObject system
 	AvailableProfiles.Remove( TEXT( "DeviceProfiles" ) );
@@ -123,7 +119,7 @@ void UDeviceProfileManager::InitializeCVarsForActiveDeviceProfile(bool bPushSett
 			UE_LOG(LogInit, Log, TEXT("Setting ConfigRules Device Profile CVar: [[%s:%s]]"), *CVarKey, *CVarValue);
 
 			// set it and remember it
-			OnSetCVarFromIniEntry(*DeviceProfileFileName, *CVarKey, *CVarValue, ECVF_SetByDeviceProfile);
+			OnSetCVarFromIniEntry(*GDeviceProfilesIni, *CVarKey, *CVarValue, ECVF_SetByDeviceProfile);
 			CVarsAlreadySetList.Add(CVarKey, CVarValue);
 		}
 	}
@@ -147,7 +143,7 @@ void UDeviceProfileManager::InitializeCVarsForActiveDeviceProfile(bool bPushSett
 				UE_LOG(LogInit, Log, TEXT("Setting CommandLine Device Profile CVar: [[%s:%s]]"), *CVarKey, *CVarValue);
 
 				// set it and remember it (no thanks, Ron Popeil)
-				OnSetCVarFromIniEntry(*DeviceProfileFileName, *CVarKey, *CVarValue, ECVF_SetByDeviceProfile);
+				OnSetCVarFromIniEntry(*GDeviceProfilesIni, *CVarKey, *CVarValue, ECVF_SetByDeviceProfile);
 				CVarsAlreadySetList.Add(CVarKey, CVarValue);
 			}
 		}
@@ -201,7 +197,7 @@ void UDeviceProfileManager::InitializeCVarsForActiveDeviceProfile(bool bPushSett
 				}
 
 				TArray< FString > CurrentProfilesCVars;
-				GConfig->GetArray(*CurrentSectionName, *ArrayName, CurrentProfilesCVars, DeviceProfileFileName);
+				GConfig->GetArray(*CurrentSectionName, *ArrayName, CurrentProfilesCVars, GDeviceProfilesIni);
 
 				// Iterate over the profile and make sure we do not have duplicate CVars
 				{
@@ -269,7 +265,7 @@ void UDeviceProfileManager::InitializeCVarsForActiveDeviceProfile(bool bPushSett
 							}
 
 							uint32 CVarPriority = bIsScalabilityBucket ? ECVF_SetByScalability : ECVF_SetByDeviceProfile;
-							OnSetCVarFromIniEntry(*DeviceProfileFileName, *CVarKey, *CVarValue, CVarPriority);
+							OnSetCVarFromIniEntry(*GDeviceProfilesIni, *CVarKey, *CVarValue, CVarPriority);
 							CVarsAlreadySetList.Add(CVarKey, CVarValue);
 						}
 					}
@@ -278,7 +274,7 @@ void UDeviceProfileManager::InitializeCVarsForActiveDeviceProfile(bool bPushSett
 
 			// Get the next device profile name, to look for CVars in, along the tree
 			FString NextBaseDeviceProfileName;
-			if( GConfig->GetString( *CurrentSectionName, TEXT("BaseProfileName"), NextBaseDeviceProfileName, DeviceProfileFileName ) )
+			if( GConfig->GetString( *CurrentSectionName, TEXT("BaseProfileName"), NextBaseDeviceProfileName, GDeviceProfilesIni ) )
 			{
 				BaseDeviceProfileName = NextBaseDeviceProfileName;
 				UE_LOG(LogInit, Log, TEXT("Going up to parent DeviceProfile [%s]"), *BaseDeviceProfileName);
@@ -296,11 +292,10 @@ void UDeviceProfileManager::InitializeCVarsForActiveDeviceProfile(bool bPushSett
 
 bool UDeviceProfileManager::DoActiveProfilesReference(const TSet<FString>& DeviceProfilesToQuery)
 {
-	FConfigCacheIni::LoadGlobalIniFile(DeviceProfileFileName, TEXT("DeviceProfiles"));
 	TArray< FString > AvailableProfiles;
-	GConfig->GetSectionNames(DeviceProfileFileName, AvailableProfiles);
+	GConfig->GetSectionNames(GDeviceProfilesIni, AvailableProfiles);
 
-	auto DoesProfileReference = [&AvailableProfiles, DeviceProfileFileName = DeviceProfileFileName](const FString& SearchProfile, const TSet<FString>& InDeviceProfilesToQuery)
+	auto DoesProfileReference = [&AvailableProfiles, GDeviceProfilesIni = GDeviceProfilesIni](const FString& SearchProfile, const TSet<FString>& InDeviceProfilesToQuery)
 	{
 		// For each device profile, starting with the selected and working our way up the BaseProfileName tree,
 		FString BaseDeviceProfileName = SearchProfile;
@@ -318,7 +313,7 @@ bool UDeviceProfileManager::DoActiveProfilesReference(const TSet<FString>& Devic
 
 				// Get the next device profile name
 				FString NextBaseDeviceProfileName;
-				if (GConfig->GetString(*CurrentSectionName, TEXT("BaseProfileName"), NextBaseDeviceProfileName, DeviceProfileFileName))
+				if (GConfig->GetString(*CurrentSectionName, TEXT("BaseProfileName"), NextBaseDeviceProfileName, GDeviceProfilesIni))
 				{
 					BaseDeviceProfileName = NextBaseDeviceProfileName;
 				}
@@ -397,15 +392,24 @@ UDeviceProfile* UDeviceProfileManager::CreateProfile(const FString& ProfileName,
 		// @todo config: we could likely cache local ini files to speed this up,
 		// along with the ones we load in LoadConfig
 		// NOTE: This happens at runtime, so maybe only do this if !RequiresCookedData()?
-		FConfigFile PlatformConfigFile;
-		FConfigCacheIni::LoadLocalIniFile(PlatformConfigFile, TEXT("DeviceProfiles"), true, ConfigPlatform);
+		FConfigFile* PlatformConfigFile;
+		FConfigFile LocalConfigFile;
+		if (FPlatformProperties::RequiresCookedData())
+		{
+			PlatformConfigFile = GConfig->Find(GDeviceProfilesIni, false);
+		}
+		else
+		{
+			FConfigCacheIni::LoadLocalIniFile(LocalConfigFile, TEXT("DeviceProfiles"), true, ConfigPlatform);
+			PlatformConfigFile = &LocalConfigFile;
+		}
 
 		// Build Parent objects first. Important for setup
 		FString ParentName = InSpecifyParentName;
 		if (ParentName.Len() == 0)
 		{
 			const FString SectionName = FString::Printf(TEXT("%s %s"), *ProfileName, *UDeviceProfile::StaticClass()->GetName());
-			PlatformConfigFile.GetString(*SectionName, TEXT("BaseProfileName"), ParentName);
+			PlatformConfigFile->GetString(*SectionName, TEXT("BaseProfileName"), ParentName);
 		}
 
 		UObject* ParentObject = nullptr;
@@ -415,7 +419,7 @@ UDeviceProfile* UDeviceProfileManager::CreateProfile(const FString& ProfileName,
 			ParentObject = FindObject<UDeviceProfile>(GetTransientPackage(), *ParentName);
 			if (ParentObject == nullptr)
 			{
-				TestProfileForCircularReferences(ProfileName, ParentName, PlatformConfigFile);
+				TestProfileForCircularReferences(ProfileName, ParentName, *PlatformConfigFile);
 				ParentObject = CreateProfile(ParentName, ProfileType, TEXT(""), ConfigPlatform);
 			}
 		}
@@ -480,11 +484,6 @@ UDeviceProfile* UDeviceProfileManager::FindProfile( const FString& ProfileName,
 	return FoundProfile;
 }
 
-const FString UDeviceProfileManager::GetDeviceProfileIniName() const
-{
-	return DeviceProfileFileName;
-}
-
 
 FOnDeviceProfileManagerUpdated& UDeviceProfileManager::OnManagerUpdated()
 {
@@ -590,7 +589,7 @@ void UDeviceProfileManager::SaveProfiles(bool bSaveToDefaults)
 				UDeviceProfile* CurrentProfile = CastChecked<UDeviceProfile>(Profiles[DeviceProfileIndex]);
 				FString DeviceProfileTypeNameCombo = FString::Printf(TEXT("%s,%s"), *CurrentProfile->GetName(), *CurrentProfile->DeviceType);
 
-				CurrentProfile->SaveConfig(CPF_Config, *DeviceProfileFileName);
+				CurrentProfile->SaveConfig(CPF_Config, *GDeviceProfilesIni);
 			}
 		}
 
diff --git a/Engine/Source/Runtime/Engine/Private/DistanceFieldAtlas.cpp b/Engine/Source/Runtime/Engine/Private/DistanceFieldAtlas.cpp
--- a/Engine/Source/Runtime/Engine/Private/DistanceFieldAtlas.cpp
+++ b/Engine/Source/Runtime/Engine/Private/DistanceFieldAtlas.cpp
@@ -122,7 +122,7 @@ static TAutoConsoleVariable<int32> CVarDistFieldForceMaxAtlasSize(
 	TEXT("When enabled, we'll always allocate the largest possible volume texture for the distance field atlas regardless of how many blocks we need.  This is an optimization to avoid re-packing the texture, for projects that are expected to always require the largest amount of space."),
 	ECVF_Default);
 
-static int32 GDistanceFieldParallelAtlasUpdate = 0;
+static int32 GDistanceFieldParallelAtlasUpdate = 1;
 static FAutoConsoleVariableRef CVarDistanceFieldParallelAtlasUpdate(
 	TEXT("r.DistanceFields.ParallelAtlasUpdate"),
 	GDistanceFieldParallelAtlasUpdate,
@@ -416,6 +416,8 @@ static void CopyToUpdateTextureData
 
 void FDistanceFieldVolumeTextureAtlas::UpdateAllocations(FRHICommandListImmediate& RHICmdList, ERHIFeatureLevel::Type InFeatureLevel)
 {
+	SCOPED_NAMED_EVENT(FDistanceFieldVolumeTextureAtlas_UpdateAllocations, FColor::Emerald);
+
 	{
 		uint32 TotalSurface = BlockAllocator.GetMaxSizeX() * BlockAllocator.GetMaxSizeY() * BlockAllocator.GetMaxSizeZ();
 		CSV_CUSTOM_STAT_GLOBAL(DFAtlasPercentageUsage, float((float(AllocatedPixels) / float(TotalSurface))*100.0f), ECsvCustomStatOp::Set);
@@ -534,7 +536,7 @@ void FDistanceFieldVolumeTextureAtlas::UpdateAllocations(FRHICommandListImmediat
 			{
 				// Remove all allocations from the layout so we have a clean slate
 				BlockAllocator = FTextureLayout3d(0, 0, 0, AtlasXY, AtlasXY, AtlasZ, false, false);
-
+				
 				Generation++;
 
 				MaxUsedAtlasX = 0;
@@ -616,28 +618,28 @@ void FDistanceFieldVolumeTextureAtlas::UpdateAllocations(FRHICommandListImmediat
 					*TextureUpdateDataPtr = RHIBeginUpdateTexture3D(VolumeTextureRHI, 0, UpdateRegion);
 				}
 
-				TArray<uint8> UncompressedData;
-
 				if (bRuntimeDownsampling)
 				{
 					UpdateDataArray.Empty(LocalPendingAllocations->Num());
 					UpdateDataArray.AddUninitialized(LocalPendingAllocations->Num());
+					DownsamplingTasks.AddDefaulted(LocalPendingAllocations->Num());
 				}
 
-				for (int32 AllocationIndex = 0; AllocationIndex < LocalPendingAllocations->Num(); AllocationIndex++)
+				// @todo arne @todo beni: can you verify i properly merged the optimizations here and in Main?
+				ParallelFor(LocalPendingAllocations->Num(), [FormatSize, this, bDataIsCompressed, bRuntimeDownsampling, &LocalPendingAllocations, &TextureUpdateDataPtr, &DownsamplingTasks, &UpdateDataArray](int32 AllocationIndex)
 				{
+					TArray<uint8> UncompressedData;
 					FDistanceFieldVolumeTexture* Texture = (*LocalPendingAllocations)[AllocationIndex];
 					FIntVector Size = Texture->VolumeData.Size;
 					
 					const TArray<uint8>* SourceDataPtr = NULL;
-
 					if (bDataIsCompressed)
 					{
 						const int32 UncompressedSize = Size.X * Size.Y * Size.Z * FormatSize;
 						UncompressedData.Reset(UncompressedSize);
 						UncompressedData.AddUninitialized(UncompressedSize);
 
-						verify(FCompression::UncompressMemory(NAME_Zlib, UncompressedData.GetData(), UncompressedSize, Texture->VolumeData.CompressedDistanceFieldVolume.GetData(), Texture->VolumeData.CompressedDistanceFieldVolume.Num()));
+						verify(FCompression::UncompressMemory(NAME_LZ4, UncompressedData.GetData(), UncompressedSize, Texture->VolumeData.CompressedDistanceFieldVolume.GetData(), Texture->VolumeData.CompressedDistanceFieldVolume.Num()));
 
 						SourceDataPtr = &UncompressedData;
 					}
@@ -652,7 +654,7 @@ void FDistanceFieldVolumeTextureAtlas::UpdateAllocations(FRHICommandListImmediat
 					
 					if (bRuntimeDownsampling)
 					{
-						FDistanceFieldDownsamplingDataTask& DownsamplingTask = DownsamplingTasks.Emplace_GetRef();
+						FDistanceFieldDownsamplingDataTask& DownsamplingTask = DownsamplingTasks[AllocationIndex];
 						TextureUpdateDataPtr = &UpdateDataArray[AllocationIndex];
 						FDistanceFieldDownsampling::FillDownsamplingTask(Size, Texture->SizeInAtlas, Texture->GetAllocationMin(), Format, DownsamplingTask, *TextureUpdateDataPtr);
 					}
@@ -662,9 +664,7 @@ void FDistanceFieldVolumeTextureAtlas::UpdateAllocations(FRHICommandListImmediat
 					}
 					
 					CopyToUpdateTextureData(Size, FormatSize, *SourceDataPtr, *TextureUpdateDataPtr, DstOffset);
-				}
-
-				UncompressedData.Empty();
+				}, !GDistanceFieldParallelAtlasUpdate, false);
 
 				if (!bRuntimeDownsampling)
 				{
@@ -725,12 +725,11 @@ void FDistanceFieldVolumeTextureAtlas::UpdateAllocations(FRHICommandListImmediat
 						TArray<uint8> UncompressedData;
 						UncompressedData.Empty(UncompressedSize);
 						UncompressedData.AddUninitialized(UncompressedSize);
-						verify(FCompression::UncompressMemory(NAME_Zlib, UncompressedData.GetData(), UncompressedSize, Texture->VolumeData.CompressedDistanceFieldVolume.GetData(), Texture->VolumeData.CompressedDistanceFieldVolume.Num()));
+						verify(FCompression::UncompressMemory(NAME_LZ4, UncompressedData.GetData(), UncompressedSize, Texture->VolumeData.CompressedDistanceFieldVolume.GetData(), Texture->VolumeData.CompressedDistanceFieldVolume.Num()));
 						
 						CopyToUpdateTextureData(Size, FormatSize, UncompressedData, UpdateDataArray[Idx], FIntVector::ZeroValue);
 					}
-				},
-				!GDistanceFieldParallelAtlasUpdate);
+				}, !GDistanceFieldParallelAtlasUpdate, false);
 
 			if (!bRuntimeDownsampling)
 			{
@@ -816,7 +815,7 @@ FDistanceFieldAsyncQueue* GDistanceFieldAsyncQueue = NULL;
 #if WITH_EDITOR
 
 // DDC key for distance field data, must be changed when modifying the generation code or data format
-#define DISTANCEFIELD_DERIVEDDATA_VER TEXT("E1AE9CB64EF64BA9A5EA17E72C88F9D")
+#define DISTANCEFIELD_DERIVEDDATA_VER TEXT("6CBBF5D788CA4699B140BAEC2A3B6B67")
 
 FString BuildDistanceFieldDerivedDataKey(const FString& InMeshKey)
 {
@@ -1573,8 +1572,6 @@ void FLandscapeTextureAtlas::UpdateAllocations(FRHICommandListImmediate& RHICmdL
 		if (SubAllocType == SAT_Height)
 		{
 			TShaderMapRef<FUploadHeightFieldToAtlasCS> ComputeShader(GetGlobalShaderMap(InFeatureLevel));
-			FUploadHeightFieldToAtlasCS* ComputeShaderPtr = *ComputeShader;
-
 			for (int32 Idx = 0; Idx < PendingUploads.Num(); ++Idx)
 			{
 				typename FUploadHeightFieldToAtlasCS::FParameters* Parameters = GraphBuilder.AllocParameters<typename FUploadHeightFieldToAtlasCS::FParameters>();
@@ -1584,21 +1581,19 @@ void FLandscapeTextureAtlas::UpdateAllocations(FRHICommandListImmediate& RHICmdL
 					RDG_EVENT_NAME("UploadHeightFieldToAtlas"),
 					Parameters,
 					ERDGPassFlags::Compute,
-					[Parameters, ComputeShaderPtr, UpdateRegion, bNeedBarrier](FRHICommandList& CmdList)
+					[Parameters, ComputeShader, UpdateRegion, bNeedBarrier](FRHICommandList& CmdList)
 				{
 					if (bNeedBarrier)
 					{
 						CmdList.TransitionResource(EResourceTransitionAccess::ERWNoBarrier, EResourceTransitionPipeline::EComputeToCompute, Parameters->RWHeightFieldAtlas);
 					}
-					FComputeShaderUtils::Dispatch(CmdList, ComputeShaderPtr, *Parameters, FIntVector(UpdateRegion.X, UpdateRegion.Y, 1));
+					FComputeShaderUtils::Dispatch(CmdList, ComputeShader, *Parameters, FIntVector(UpdateRegion.X, UpdateRegion.Y, 1));
 				});
 			}
 		}
 		else
 		{
 			TShaderMapRef<FUploadVisibilityToAtlasCS> ComputeShader(GetGlobalShaderMap(InFeatureLevel));
-			FUploadVisibilityToAtlasCS* ComputeShaderPtr = *ComputeShader;
-
 			for (int32 Idx = 0; Idx < PendingUploads.Num(); ++Idx)
 			{
 				typename FUploadVisibilityToAtlasCS::FParameters* Parameters = GraphBuilder.AllocParameters<typename FUploadVisibilityToAtlasCS::FParameters>();
@@ -1608,13 +1603,13 @@ void FLandscapeTextureAtlas::UpdateAllocations(FRHICommandListImmediate& RHICmdL
 					RDG_EVENT_NAME("UploadVisibilityToAtlas"),
 					Parameters,
 					ERDGPassFlags::Compute,
-					[Parameters, ComputeShaderPtr, UpdateRegion, bNeedBarrier](FRHICommandList& CmdList)
+					[Parameters, ComputeShader, UpdateRegion, bNeedBarrier](FRHICommandList& CmdList)
 				{
 					if (bNeedBarrier)
 					{
 						CmdList.TransitionResource(EResourceTransitionAccess::ERWNoBarrier, EResourceTransitionPipeline::EComputeToCompute, Parameters->RWVisibilityAtlas);
 					}
-					FComputeShaderUtils::Dispatch(CmdList, ComputeShaderPtr, *Parameters, FIntVector(UpdateRegion.X, UpdateRegion.Y, 1));
+					FComputeShaderUtils::Dispatch(CmdList, ComputeShader, *Parameters, FIntVector(UpdateRegion.X, UpdateRegion.Y, 1));
 				});
 			}
 		}
diff --git a/Engine/Source/Runtime/Engine/Private/DistanceFieldDownsampling.cpp b/Engine/Source/Runtime/Engine/Private/DistanceFieldDownsampling.cpp
--- a/Engine/Source/Runtime/Engine/Private/DistanceFieldDownsampling.cpp
+++ b/Engine/Source/Runtime/Engine/Private/DistanceFieldDownsampling.cpp
@@ -122,9 +122,8 @@ void FDistanceFieldDownsampling::DispatchDownsampleTasks(FRHICommandListImmediat
 		PassParameters->MeshDFSampler = TStaticSamplerState<SF_Bilinear, AM_Clamp, AM_Clamp, AM_Clamp>::GetRHI();
 		PassParameters->DFAtlas = DFAtlasUAV;	
 
-		TShaderMap<FGlobalShaderType>* GlobalShaderMap = GetGlobalShaderMap(FeatureLevel);
+		FGlobalShaderMap* GlobalShaderMap = GetGlobalShaderMap(FeatureLevel);
 		TShaderMapRef< FDistanceFieldDownsamplingCS > ComputeShader(GlobalShaderMap);
-		FDistanceFieldDownsamplingCS* ComputeShaderPtr = *ComputeShader;
 		const FIntVector GroupCount(FMath::DivideAndRoundUp(Task.DstSize.X, 8), FMath::DivideAndRoundUp(Task.DstSize.Y, 8), Task.DstSize.Z);
 
 		GraphBuilder.AddPass
@@ -132,9 +131,9 @@ void FDistanceFieldDownsampling::DispatchDownsampleTasks(FRHICommandListImmediat
 			RDG_EVENT_NAME("DownsampleMeshDF"),
 			PassParameters, 
 			ERDGPassFlags::Compute,
-			[PassParameters, ComputeShaderPtr, GroupCount, &Task, &DFAtlasUAV](FRHICommandList& CmdList)
+			[PassParameters, ComputeShader, GroupCount, &Task, &DFAtlasUAV](FRHICommandList& CmdList)
 			{
-				FComputeShaderUtils::Dispatch(CmdList, ComputeShaderPtr, *PassParameters, GroupCount);
+				FComputeShaderUtils::Dispatch(CmdList, ComputeShader, *PassParameters, GroupCount);
 				CmdList.TransitionResources(EResourceTransitionAccess::ERWNoBarrier, EResourceTransitionPipeline::EComputeToCompute, &DFAtlasUAV, 1); // No barrier needed
 				Task.VolumeTextureRHI = nullptr;
 			}
diff --git a/Engine/Source/Runtime/Engine/Private/GPUSkinCache.cpp b/Engine/Source/Runtime/Engine/Private/GPUSkinCache.cpp
--- a/Engine/Source/Runtime/Engine/Private/GPUSkinCache.cpp
+++ b/Engine/Source/Runtime/Engine/Private/GPUSkinCache.cpp
@@ -121,10 +121,10 @@ FAutoConsoleVariableRef CVarGPUSkinCacheBlendUsingVertexColorForRecomputeTangent
 
 static int32 GGPUSkinCacheFlushCounter = 0;
 
-static inline bool DoesPlatformSupportGPUSkinCache(EShaderPlatform Platform)
+static inline bool DoesPlatformSupportGPUSkinCache(const FStaticShaderPlatform Platform)
 {
 	return Platform == SP_PCD3D_SM5 || IsMetalSM5Platform(Platform) || IsVulkanSM5Platform(Platform) || Platform == SP_OPENGL_SM5 
-		|| FDataDrivenShaderPlatformInfo::GetInfo(Platform).bSupportsGPUSkinCache;
+		|| FDataDrivenShaderPlatformInfo::GetSupportsGPUSkinCache(Platform);
 }
 
 ENGINE_API bool IsGPUSkinCacheAvailable(EShaderPlatform Platform)
@@ -423,6 +423,7 @@ protected:
 
 class FBaseGPUSkinCacheCS : public FGlobalShader
 {
+	DECLARE_INLINE_TYPE_LAYOUT(FBaseGPUSkinCacheCS, NonVirtual);
 public:
 	FBaseGPUSkinCacheCS() {}
 
@@ -466,7 +467,7 @@ public:
 		FGPUSkinCacheEntry::FSectionDispatchData& DispatchData,
 		FRHIUnorderedAccessView* PositionUAV, FRHIUnorderedAccessView* TangentUAV)
 	{
-		FRHIComputeShader* ShaderRHI = GetComputeShader();
+		FRHIComputeShader* ShaderRHI = RHICmdList.GetBoundComputeShader();
 
 		SetShaderValue(RHICmdList, ShaderRHI, NumVertices, DispatchData.NumVertices);
 		SetShaderValue(RHICmdList, ShaderRHI, InputStreamStart, DispatchData.InputStreamStart);
@@ -513,60 +514,46 @@ public:
 
 	void UnsetParameters(FRHICommandList& RHICmdList)
 	{
-		FRHIComputeShader* ShaderRHI = GetComputeShader();
+		FRHIComputeShader* ShaderRHI = RHICmdList.GetBoundComputeShader();
 
 		SetUAVParameter(RHICmdList, ShaderRHI, PositionBufferUAV, 0);
 		SetUAVParameter(RHICmdList, ShaderRHI, TangentBufferUAV, 0);
 	}
 
-	virtual bool Serialize(FArchive& Ar) override
-	{
-		bool bShaderHasOutdatedParameters = FGlobalShader::Serialize(Ar);
-		Ar  << NumVertices << InputStreamStart << SkinCacheStart
-			<< TangentInputBuffer << PositionInputBuffer << PositionBufferUAV << TangentBufferUAV << BoneMatrices << MorphBuffer << MorphBufferOffset
-			<< SkinCacheDebug;
-
-		Ar << NumBoneInfluences << InputWeightIndexSize << InputWeightStart << InputWeightStride << InputWeightStream << InputWeightLookupStream;
-		
-        Ar << ClothBuffer << ClothPositionsAndNormalsBuffer << ClothBufferOffset << ClothBlendWeight << ClothLocalToWorld << ClothWorldToLocal;
-
-		//Ar << DebugParameter;
-		return bShaderHasOutdatedParameters;
-	}
-
 private:
+	
+	LAYOUT_FIELD(FShaderParameter, NumVertices)
+	LAYOUT_FIELD(FShaderParameter, SkinCacheDebug)
+	LAYOUT_FIELD(FShaderParameter, InputStreamStart)
+	LAYOUT_FIELD(FShaderParameter, SkinCacheStart)
+
+	//LAYOUT_FIELD(FShaderParameter, DebugParameter)
+
+	LAYOUT_FIELD(FShaderUniformBufferParameter, SkinUniformBuffer)
+
+	LAYOUT_FIELD(FShaderResourceParameter, BoneMatrices)
+	LAYOUT_FIELD(FShaderResourceParameter, TangentInputBuffer)
+	LAYOUT_FIELD(FShaderResourceParameter, PositionInputBuffer)
+	LAYOUT_FIELD(FShaderResourceParameter, PositionBufferUAV)
+	LAYOUT_FIELD(FShaderResourceParameter, TangentBufferUAV)
+
+	LAYOUT_FIELD(FShaderParameter, NumBoneInfluences);
+	LAYOUT_FIELD(FShaderParameter, InputWeightIndexSize);
+	LAYOUT_FIELD(FShaderParameter, InputWeightStart)
+	LAYOUT_FIELD(FShaderParameter, InputWeightStride)
+	LAYOUT_FIELD(FShaderResourceParameter, InputWeightStream)
+	LAYOUT_FIELD(FShaderResourceParameter, InputWeightLookupStream);
+
+	LAYOUT_FIELD(FShaderResourceParameter, MorphBuffer)
+	LAYOUT_FIELD(FShaderParameter, MorphBufferOffset)
+
+	LAYOUT_FIELD(FShaderResourceParameter, ClothBuffer)
+	LAYOUT_FIELD(FShaderResourceParameter, ClothPositionsAndNormalsBuffer)
+	LAYOUT_FIELD(FShaderParameter, ClothBufferOffset)
+	LAYOUT_FIELD(FShaderParameter, ClothBlendWeight)
+	LAYOUT_FIELD(FShaderParameter, ClothLocalToWorld)
+	LAYOUT_FIELD(FShaderParameter, ClothWorldToLocal)
 
-	FShaderParameter NumVertices;
-	FShaderParameter SkinCacheDebug;
-	FShaderParameter InputStreamStart;
-	FShaderParameter SkinCacheStart;
-
-	//FShaderParameter DebugParameter;
-
-	FShaderUniformBufferParameter SkinUniformBuffer;
-
-	FShaderResourceParameter BoneMatrices;
-	FShaderResourceParameter TangentInputBuffer;
-	FShaderResourceParameter PositionInputBuffer;
-	FShaderResourceParameter PositionBufferUAV;
-	FShaderResourceParameter TangentBufferUAV;
-
-	FShaderParameter NumBoneInfluences;
-	FShaderParameter InputWeightIndexSize;
-	FShaderParameter InputWeightStart;
-	FShaderParameter InputWeightStride;
-	FShaderResourceParameter InputWeightStream;
-	FShaderResourceParameter InputWeightLookupStream;
-
-	FShaderResourceParameter MorphBuffer;
-	FShaderParameter MorphBufferOffset;
-
-	FShaderResourceParameter ClothBuffer;
-	FShaderResourceParameter ClothPositionsAndNormalsBuffer;
-	FShaderParameter ClothBufferOffset;
-	FShaderParameter ClothBlendWeight;
-	FShaderParameter ClothLocalToWorld;
-	FShaderParameter ClothWorldToLocal;
 };
 
 /** Compute shader that skins a batch of vertices. */
@@ -721,6 +708,7 @@ void FGPUSkinCache::TransitionAllToWriteable(FRHICommandList& RHICmdList)
 /** base of the FRecomputeTangentsPerTrianglePassCS class */
 class FBaseRecomputeTangentsPerTriangleShader : public FGlobalShader
 {
+	DECLARE_INLINE_TYPE_LAYOUT(FBaseRecomputeTangentsPerTriangleShader, NonVirtual);
 public:
 	static bool ShouldCompilePermutation(const FGlobalShaderPermutationParameters& Parameters)
 	{
@@ -730,21 +718,6 @@ public:
 
 	static const uint32 ThreadGroupSizeX = 64;
 
-	FShaderResourceParameter IntermediateAccumBufferUAV;
-	FShaderParameter NumTriangles;
-	FShaderResourceParameter GPUPositionCacheBuffer;
-	FShaderResourceParameter GPUTangentCacheBuffer;
-	FShaderParameter SkinCacheStart;
-	FShaderResourceParameter IndexBuffer;
-	FShaderParameter IndexBufferOffset;
-	FShaderParameter InputStreamStart;
-	FShaderParameter NumTexCoords;
-	FShaderParameter SelectedTexCoord;
-	FShaderResourceParameter TangentInputBuffer;
-	FShaderResourceParameter UVsInputBuffer;
-	FShaderResourceParameter DuplicatedIndices;
-	FShaderResourceParameter DuplicatedIndicesIndices;
-
 	FBaseRecomputeTangentsPerTriangleShader()
 	{}
 
@@ -771,7 +744,7 @@ public:
 
 	void SetParameters(FRHICommandListImmediate& RHICmdList, FGPUSkinCacheEntry* Entry, FGPUSkinCacheEntry::FSectionDispatchData& DispatchData, FRWBuffer& StagingBuffer)
 	{
-		FRHIComputeShader* ShaderRHI = GetComputeShader();
+		FRHIComputeShader* ShaderRHI = RHICmdList.GetBoundComputeShader();
 
 //later		FGlobalShader::SetParameters<FViewUniformShaderParameters>(RHICmdList, ShaderRHI, View);
 
@@ -804,19 +777,25 @@ public:
 
 	void UnsetParameters(FRHICommandList& RHICmdList)
 	{
-		FRHIComputeShader* ShaderRHI = GetComputeShader();
+		FRHIComputeShader* ShaderRHI = RHICmdList.GetBoundComputeShader();
 
 		SetUAVParameter(RHICmdList, ShaderRHI, IntermediateAccumBufferUAV, 0);
 	}
 
-	virtual bool Serialize(FArchive& Ar) override
-	{
-		bool bShaderHasOutdatedParameters = FGlobalShader::Serialize(Ar);
-		Ar << IntermediateAccumBufferUAV << NumTriangles << GPUPositionCacheBuffer << GPUTangentCacheBuffer << SkinCacheStart << IndexBuffer << IndexBufferOffset
-			<< InputStreamStart << NumTexCoords << SelectedTexCoord << TangentInputBuffer << UVsInputBuffer;
-        Ar << DuplicatedIndices << DuplicatedIndicesIndices;
-		return bShaderHasOutdatedParameters;
-	}
+	LAYOUT_FIELD(FShaderResourceParameter, IntermediateAccumBufferUAV);
+	LAYOUT_FIELD(FShaderParameter, NumTriangles);
+	LAYOUT_FIELD(FShaderResourceParameter, GPUPositionCacheBuffer);
+	LAYOUT_FIELD(FShaderResourceParameter, GPUTangentCacheBuffer);
+	LAYOUT_FIELD(FShaderParameter, SkinCacheStart);
+	LAYOUT_FIELD(FShaderResourceParameter, IndexBuffer);
+	LAYOUT_FIELD(FShaderParameter, IndexBufferOffset);
+	LAYOUT_FIELD(FShaderParameter, InputStreamStart);
+	LAYOUT_FIELD(FShaderParameter, NumTexCoords);
+	LAYOUT_FIELD(FShaderParameter, SelectedTexCoord);
+	LAYOUT_FIELD(FShaderResourceParameter, TangentInputBuffer);
+	LAYOUT_FIELD(FShaderResourceParameter, UVsInputBuffer);
+	LAYOUT_FIELD(FShaderResourceParameter, DuplicatedIndices);
+	LAYOUT_FIELD(FShaderResourceParameter, DuplicatedIndicesIndices);
 };
 
 /** Encapsulates the RecomputeSkinTangents compute shader. */
@@ -854,6 +833,8 @@ IMPLEMENT_SHADER_TYPE(template<>, FRecomputeTangentsPerTrianglePassCS<3>, TEXT("
 /** Encapsulates the RecomputeSkinTangentsResolve compute shader. */
 class FBaseRecomputeTangentsPerVertexShader : public FGlobalShader
 {
+	DECLARE_INLINE_TYPE_LAYOUT(FBaseRecomputeTangentsPerVertexShader, NonVirtual);
+
 public:
 	static bool ShouldCompilePermutation(const FGlobalShaderPermutationParameters& Parameters)
 	{
@@ -863,13 +844,13 @@ public:
 
 	static const uint32 ThreadGroupSizeX = 64;
 
-	FShaderResourceParameter IntermediateAccumBufferUAV;
-	FShaderResourceParameter TangentBufferUAV;
-	FShaderResourceParameter TangentInputBuffer;
-	FShaderResourceParameter ColorInputBuffer;
-	FShaderParameter SkinCacheStart;
-	FShaderParameter NumVertices;
-	FShaderParameter InputStreamStart;
+	LAYOUT_FIELD(FShaderResourceParameter, IntermediateAccumBufferUAV);
+	LAYOUT_FIELD(FShaderResourceParameter, TangentBufferUAV);
+	LAYOUT_FIELD(FShaderResourceParameter, TangentInputBuffer);
+	LAYOUT_FIELD(FShaderResourceParameter, ColorInputBuffer);
+	LAYOUT_FIELD(FShaderParameter, SkinCacheStart);
+	LAYOUT_FIELD(FShaderParameter, NumVertices);
+	LAYOUT_FIELD(FShaderParameter, InputStreamStart);
 
 	FBaseRecomputeTangentsPerVertexShader() {}
 
@@ -887,7 +868,7 @@ public:
 
 	void SetParameters(FRHICommandListImmediate& RHICmdList, FGPUSkinCacheEntry* Entry, FGPUSkinCacheEntry::FSectionDispatchData& DispatchData, FRWBuffer& StagingBuffer)
 	{
-		FRHIComputeShader* ShaderRHI = GetComputeShader();
+		FRHIComputeShader* ShaderRHI = RHICmdList.GetBoundComputeShader();
 
 		check(StagingBuffer.UAV);
 
@@ -907,18 +888,11 @@ public:
 
 	void UnsetParameters(FRHICommandList& RHICmdList)
 	{
-		FRHIComputeShader* ShaderRHI = GetComputeShader();
+		FRHIComputeShader* ShaderRHI = RHICmdList.GetBoundComputeShader();
 
 		SetUAVParameter(RHICmdList, ShaderRHI, TangentBufferUAV, 0);
 		SetUAVParameter(RHICmdList, ShaderRHI, IntermediateAccumBufferUAV, 0);
 	}
-
-	virtual bool Serialize(FArchive& Ar) override
-	{
-		bool bShaderHasOutdatedParameters = FGlobalShader::Serialize(Ar);
-		Ar << IntermediateAccumBufferUAV << TangentBufferUAV << TangentInputBuffer << ColorInputBuffer << SkinCacheStart << NumVertices << InputStreamStart;
-		return bShaderHasOutdatedParameters;
-	}
 };
 
 template <int Permutation>
@@ -1008,18 +982,20 @@ void FGPUSkinCache::DispatchUpdateSkinTangents(FRHICommandListImmediate& RHICmdL
 
 			bool bFullPrecisionUV = LodData.StaticVertexBuffers.StaticMeshVertexBuffer.GetUseFullPrecisionUVs();
 
-			FBaseRecomputeTangentsPerTriangleShader* Shader = 0;
+			TShaderRef<FBaseRecomputeTangentsPerTriangleShader> Shader;
 
 			if (bFullPrecisionUV)
 			{
-				Shader = GAllowDupedVertsForRecomputeTangents ? (FBaseRecomputeTangentsPerTriangleShader*)*ComputeShader11 : (FBaseRecomputeTangentsPerTriangleShader*)*ComputeShader01;
+				if (GAllowDupedVertsForRecomputeTangents) Shader = ComputeShader11;
+				else Shader = ComputeShader01;
 			}
 			else
 			{
-				Shader = GAllowDupedVertsForRecomputeTangents ? (FBaseRecomputeTangentsPerTriangleShader*)*ComputeShader10 : (FBaseRecomputeTangentsPerTriangleShader*)*ComputeShader00;
+				if (GAllowDupedVertsForRecomputeTangents) Shader = ComputeShader10;
+				else Shader = ComputeShader00;
 			}
 
-			check(Shader);
+			check(Shader.IsValid());
 
 			uint32 NumTriangles = DispatchData.NumTriangles;
 			uint32 ThreadGroupCountValue = FMath::DivideAndRoundUp(NumTriangles, FBaseRecomputeTangentsPerTriangleShader::ThreadGroupSizeX);
@@ -1027,7 +1003,7 @@ void FGPUSkinCache::DispatchUpdateSkinTangents(FRHICommandListImmediate& RHICmdL
 			SCOPED_DRAW_EVENTF(RHICmdList, SkinTangents_PerTrianglePass, TEXT("TangentsTri IndexStart=%d Tri=%d BoneInfluenceType=%d UVPrecision=%d"),
 				DispatchData.IndexBufferOffsetValue, DispatchData.NumTriangles, Entry->BoneInfluenceType, bFullPrecisionUV);
 
-			FRHIComputeShader* ShaderRHI = Shader->GetComputeShader();
+			FRHIComputeShader* ShaderRHI = Shader.GetComputeShader();
 			RHICmdList.SetComputeShader(ShaderRHI);
 
 			RHICmdList.TransitionResource(EResourceTransitionAccess::ERWNoBarrier, EResourceTransitionPipeline::EComputeToCompute, StagingBuffer.UAV.GetReference());
@@ -1041,7 +1017,7 @@ void FGPUSkinCache::DispatchUpdateSkinTangents(FRHICommandListImmediate& RHICmdL
 
 			INC_DWORD_STAT_BY(STAT_GPUSkinCache_NumTrianglesForRecomputeTangents, NumTriangles);
 			Shader->SetParameters(RHICmdList, Entry, DispatchData, StagingBuffer);
-			DispatchComputeShader(RHICmdList, Shader, ThreadGroupCountValue, 1, 1);
+			DispatchComputeShader(RHICmdList, Shader.GetShader(), ThreadGroupCountValue, 1, 1);
 			Shader->UnsetParameters(RHICmdList);
 		}
 
@@ -1055,19 +1031,18 @@ void FGPUSkinCache::DispatchUpdateSkinTangents(FRHICommandListImmediate& RHICmdL
 			TShaderMapRef<FRecomputeTangentsPerVertexPassCS<2>> ComputeShader2(GlobalShaderMap);
 			TShaderMapRef<FRecomputeTangentsPerVertexPassCS<3>> ComputeShader3(GlobalShaderMap);
 			TShaderMapRef<FRecomputeTangentsPerVertexPassCS<4>> ComputeShader4(GlobalShaderMap);
-			FBaseRecomputeTangentsPerVertexShader* ComputeShader = nullptr;
+			TShaderRef<FBaseRecomputeTangentsPerVertexShader> ComputeShader;
 			if (GBlendUsingVertexColorForRecomputeTangents == 1)
-				ComputeShader = (FBaseRecomputeTangentsPerVertexShader*)*ComputeShader1;
+				ComputeShader = ComputeShader1;
 			else if (GBlendUsingVertexColorForRecomputeTangents == 2)
-				ComputeShader = (FBaseRecomputeTangentsPerVertexShader*)*ComputeShader2;
+				ComputeShader = ComputeShader2;
 			else if (GBlendUsingVertexColorForRecomputeTangents == 3)
-				ComputeShader = (FBaseRecomputeTangentsPerVertexShader*)*ComputeShader3;
+				ComputeShader = ComputeShader3;
 			else if (GBlendUsingVertexColorForRecomputeTangents == 4)
-				ComputeShader = (FBaseRecomputeTangentsPerVertexShader*)*ComputeShader4;
+				ComputeShader = ComputeShader4;
 			else
-				ComputeShader = (FBaseRecomputeTangentsPerVertexShader*)*ComputeShader0;
-			FRHIComputeShader* ShaderRHI = ComputeShader->GetComputeShader();
-			RHICmdList.SetComputeShader(ComputeShader->GetComputeShader());
+				ComputeShader = ComputeShader0;
+			RHICmdList.SetComputeShader(ComputeShader.GetComputeShader());
 
 			uint32 VertexCount = DispatchData.NumVertices;
 			uint32 ThreadGroupCountValue = FMath::DivideAndRoundUp(VertexCount, ComputeShader->ThreadGroupSizeX);
@@ -1075,7 +1050,7 @@ void FGPUSkinCache::DispatchUpdateSkinTangents(FRHICommandListImmediate& RHICmdL
 			RHICmdList.TransitionResource(EResourceTransitionAccess::ERWBarrier, EResourceTransitionPipeline::EComputeToCompute, StagingBuffer.UAV.GetReference());
 
 			ComputeShader->SetParameters(RHICmdList, Entry, DispatchData, StagingBuffer);
-			DispatchComputeShader(RHICmdList, ComputeShader, ThreadGroupCountValue, 1, 1);
+			DispatchComputeShader(RHICmdList, ComputeShader.GetShader(), ThreadGroupCountValue, 1, 1);
 			ComputeShader->UnsetParameters(RHICmdList);
 		}
 		// todo				RHICmdList.TransitionResource(EResourceTransitionAccess::EReadable, EResourceTransitionPipeline::EComputeToGfx, TangentsBlendBuffer.VertexBufferSRV);
@@ -1259,7 +1234,7 @@ void FGPUSkinCache::Release(FGPUSkinCacheEntry*& SkinCacheEntry)
 }
 
 void FGPUSkinCache::SetVertexStreams(FGPUSkinCacheEntry* Entry, int32 Section, FRHICommandList& RHICmdList,
-	FShader* Shader, const FGPUSkinPassthroughVertexFactory* VertexFactory,
+	FRHIVertexShader* ShaderRHI, const FGPUSkinPassthroughVertexFactory* VertexFactory,
 	uint32 BaseVertexIndex, FShaderResourceParameter GPUSkinCachePreviousPositionBuffer)
 {
 	INC_DWORD_STAT(STAT_GPUSkinCache_NumSetVertexStreams);
@@ -1275,7 +1250,6 @@ void FGPUSkinCache::SetVertexStreams(FGPUSkinCacheEntry* Entry, int32 Section, F
 		RHICmdList.SetStreamSource(VertexFactory->GetTangentStreamIndex(), DispatchData.GetTangentRWBuffer()->Buffer, 0);
 	}
 
-	FRHIVertexShader* ShaderRHI = Shader->GetVertexShader();
 	if (ShaderRHI && GPUSkinCachePreviousPositionBuffer.IsBound())
 	{
 		RHICmdList.SetShaderResourceViewParameter(ShaderRHI, GPUSkinCachePreviousPositionBuffer.GetBaseIndex(), DispatchData.GetPreviousPositionRWBuffer()->SRV);
@@ -1338,26 +1312,61 @@ void FGPUSkinCache::DispatchUpdateSkinning(FRHICommandListImmediate& RHICmdList,
 	TShaderMapRef<TGPUSkinCacheCS<21>>  SkinCacheCS111(GlobalShaderMap);	// 16bit_1, BoneInfluenceType_1, SkinType_1
 	TShaderMapRef<TGPUSkinCacheCS<22>>  SkinCacheCS112(GlobalShaderMap);	// 16bit_1, BoneInfluenceType_1, SkinType_2
 
-	FBaseGPUSkinCacheCS* Shader = 0;
-
+	TShaderRef<FBaseGPUSkinCacheCS> Shader;
 	switch (DispatchData.SkinType)
 	{
-	case 0: Shader = Entry->BoneInfluenceType == 0 ? (Entry->bUse16BitBoneIndex ? (FBaseGPUSkinCacheCS*)*SkinCacheCS100 : (FBaseGPUSkinCacheCS*)*SkinCacheCS000) :
-					(Entry->BoneInfluenceType == 1 ? (Entry->bUse16BitBoneIndex ? (FBaseGPUSkinCacheCS*)*SkinCacheCS110 : (FBaseGPUSkinCacheCS*)*SkinCacheCS010) :
-													 (FBaseGPUSkinCacheCS*)*SkinCacheCS020);
+	case 0:
+		if (Entry->BoneInfluenceType == 0)
+		{
+			if (Entry->bUse16BitBoneIndex) Shader = SkinCacheCS100;
+			else Shader = SkinCacheCS000;
+		}
+		else if (Entry->BoneInfluenceType == 1)
+		{
+			if (Entry->bUse16BitBoneIndex) Shader = SkinCacheCS110;
+			else Shader = SkinCacheCS010;
+		}
+		else
+		{
+			Shader = SkinCacheCS020;
+		}
 		break;
-	case 1: Shader = Entry->BoneInfluenceType == 0 ? (Entry->bUse16BitBoneIndex ? (FBaseGPUSkinCacheCS*)*SkinCacheCS101 : (FBaseGPUSkinCacheCS*)*SkinCacheCS001) :
-					(Entry->BoneInfluenceType == 1 ? (Entry->bUse16BitBoneIndex ? (FBaseGPUSkinCacheCS*)*SkinCacheCS111 : (FBaseGPUSkinCacheCS*)*SkinCacheCS011) :
-													 (FBaseGPUSkinCacheCS*)*SkinCacheCS021);
+	case 1:
+		if (Entry->BoneInfluenceType == 0)
+		{
+			if (Entry->bUse16BitBoneIndex) Shader = SkinCacheCS101;
+			else Shader = SkinCacheCS001;
+		}
+		else if (Entry->BoneInfluenceType == 1)
+		{
+			if (Entry->bUse16BitBoneIndex) Shader = SkinCacheCS111;
+			else Shader = SkinCacheCS011;
+		}
+		else
+		{
+			Shader = SkinCacheCS021;
+		}
 		break;
-	case 2: Shader = Entry->BoneInfluenceType == 0 ? (Entry->bUse16BitBoneIndex ? (FBaseGPUSkinCacheCS*)*SkinCacheCS102 : (FBaseGPUSkinCacheCS*)*SkinCacheCS002) :
-					(Entry->BoneInfluenceType == 1 ? (Entry->bUse16BitBoneIndex ? (FBaseGPUSkinCacheCS*)*SkinCacheCS112 : (FBaseGPUSkinCacheCS*)*SkinCacheCS012) :
-													 (FBaseGPUSkinCacheCS*)*SkinCacheCS022);
+	case 2:
+		if (Entry->BoneInfluenceType == 0)
+		{
+			if (Entry->bUse16BitBoneIndex) Shader = SkinCacheCS102;
+			else Shader = SkinCacheCS002;
+		}
+		else if (Entry->BoneInfluenceType == 1)
+		{
+			if (Entry->bUse16BitBoneIndex) Shader = SkinCacheCS112;
+			else Shader = SkinCacheCS012;
+		}
+		else
+		{
+			Shader = SkinCacheCS022;
+		}
 		break;
 	default:
 		check(0);
 	}
-	check(Shader);
+	check(Shader.IsValid());
 
 	const FVertexBufferAndSRV& BoneBuffer = ShaderData.GetBoneBufferForReading(false);
 	const FVertexBufferAndSRV& PrevBoneBuffer = ShaderData.GetBoneBufferForReading(true);
@@ -1377,7 +1386,7 @@ void FGPUSkinCache::DispatchUpdateSkinning(FRHICommandListImmediate& RHICmdList,
 		DispatchData.PreviousPositionBuffer = DispatchData.PositionTracker.Find(PrevBoneBuffer, PreviousRevision);
 		check(DispatchData.PreviousPositionBuffer)
 
-		RHICmdList.SetComputeShader(Shader->GetComputeShader());
+		RHICmdList.SetComputeShader(Shader.GetComputeShader());
 
 		if (bUseIntermediateTangentBuffer)
 		{
@@ -1422,7 +1431,7 @@ void FGPUSkinCache::DispatchUpdateSkinning(FRHICommandListImmediate& RHICmdList,
 		DispatchData.PositionBuffer = DispatchData.PositionTracker.Find(BoneBuffer, CurrentRevision);
 		check(DispatchData.PositionBuffer);
 
-		RHICmdList.SetComputeShader(Shader->GetComputeShader());
+		RHICmdList.SetComputeShader(Shader.GetComputeShader());
 		if (bUseIntermediateTangentBuffer)
 		{
 			Shader->SetParameters(RHICmdList, BoneBuffer, Entry, DispatchData, DispatchData.GetPositionRWBuffer()->UAV, DispatchData.IntermediateTangentBuffer ? DispatchData.IntermediateTangentBuffer->UAV : nullptr);
