Copying Private-LoadTimes-4.24 stream to Main. Biggest changes are in Materials/Shader memory freezing.

diff --git a/Engine/Source/Runtime/Engine/Private/GPUSkinVertexFactory.cpp b/Engine/Source/Runtime/Engine/Private/GPUSkinVertexFactory.cpp
--- a/Engine/Source/Runtime/Engine/Private/GPUSkinVertexFactory.cpp
+++ b/Engine/Source/Runtime/Engine/Private/GPUSkinVertexFactory.cpp
@@ -59,22 +59,22 @@ static FBoneMatricesUniformShaderParameters GBoneUniformStruct;
 	bSupportsPositionOnly, \
 	false, \
 	false, \
-	Construct##FactoryClass##ShaderParameters<BoneInfluenceType>, \
-	FactoryClass<BoneInfluenceType>::ShouldCompilePermutation, \
-	FactoryClass<BoneInfluenceType>::ModifyCompilationEnvironment, \
-	FactoryClass<BoneInfluenceType>::ValidateCompiledResult, \
-	FactoryClass<BoneInfluenceType>::SupportsTessellationShaders \
+	IMPLEMENT_VERTEX_FACTORY_VTABLE(FactoryClass<BoneInfluenceType>) \
 	); \
 	template <GPUSkinBoneInfluenceType BoneInfluenceType> inline FVertexFactoryType* FactoryClass<BoneInfluenceType>::GetType() const { return &StaticType; }
 
 
 #define IMPLEMENT_GPUSKINNING_VERTEX_FACTORY_TYPE(FactoryClass, ShaderFilename,bUsedWithMaterials,bSupportsStaticLighting,bSupportsDynamicLighting,bPrecisePrevWorldPos,bSupportsPositionOnly) \
-	template <GPUSkinBoneInfluenceType BoneInfluenceType> FVertexFactoryShaderParameters* Construct##FactoryClass##ShaderParameters(EShaderFrequency ShaderFrequency) { return FactoryClass<BoneInfluenceType>::ConstructShaderParameters(ShaderFrequency); } \
 	IMPLEMENT_GPUSKINNING_VERTEX_FACTORY_TYPE_INTERNAL(FactoryClass, ShaderFilename,bUsedWithMaterials,bSupportsStaticLighting,bSupportsDynamicLighting,bPrecisePrevWorldPos,bSupportsPositionOnly) \
 	template class FactoryClass<DefaultBoneInfluence>;	\
 	template class FactoryClass<ExtraBoneInfluence>;	\
 	template class FactoryClass<UnlimitedBoneInfluence>;
 
+#define IMPLEMENT_GPUSKINNING_VERTEX_FACTORY_PARAMETER_TYPE(FactoryClass, Frequency, ParameterType) \
+	IMPLEMENT_VERTEX_FACTORY_PARAMETER_TYPE(FactoryClass<DefaultBoneInfluence>, Frequency, ParameterType); \
+	IMPLEMENT_VERTEX_FACTORY_PARAMETER_TYPE(FactoryClass<ExtraBoneInfluence>, Frequency, ParameterType); \
+	IMPLEMENT_VERTEX_FACTORY_PARAMETER_TYPE(FactoryClass<UnlimitedBoneInfluence>, Frequency, ParameterType)
+
 #if !(UE_BUILD_SHIPPING || UE_BUILD_TEST)
 static TAutoConsoleVariable<int32> CVarVelocityTest(
 	TEXT("r.VelocityTest"),
@@ -365,26 +365,26 @@ TGPUSkinVertexFactory
 TGlobalResource<FBoneBufferPool> FGPUBaseSkinVertexFactory::BoneBufferPool;
 
 template <GPUSkinBoneInfluenceType BoneInfluenceType>
-bool TGPUSkinVertexFactory<BoneInfluenceType>::ShouldCompilePermutation(EShaderPlatform Platform, const class FMaterial* Material, const FShaderType* ShaderType)
+bool TGPUSkinVertexFactory<BoneInfluenceType>::ShouldCompilePermutation(const FVertexFactoryShaderPermutationParameters& Parameters)
 {
 	bool bLimit2BoneInfluences = (CVarGPUSkinLimit2BoneInfluences.GetValueOnAnyThread() != 0);
 	
 	// Skip trying to use extra bone influences on < SM4 or when project uses 2 bones influence
-	if (BoneInfluenceType != DefaultBoneInfluence && (GetMaxSupportedFeatureLevel(Platform) < ERHIFeatureLevel::ES3_1 || bLimit2BoneInfluences))
+	if (BoneInfluenceType != DefaultBoneInfluence && (GetMaxSupportedFeatureLevel(Parameters.Platform) < ERHIFeatureLevel::ES3_1 || bLimit2BoneInfluences))
 	{
 		return false;
 	}
 
 	bool bUnlimitedBoneInfluences = (BoneInfluenceType == UnlimitedBoneInfluence && GUnlimitedBoneInfluences);
-	return ((Material->IsUsedWithSkeletalMesh() && (BoneInfluenceType != UnlimitedBoneInfluence || bUnlimitedBoneInfluences)) || Material->IsSpecialEngineMaterial());
+	return ((Parameters.MaterialParameters.bIsUsedWithSkeletalMesh && (BoneInfluenceType != UnlimitedBoneInfluence || bUnlimitedBoneInfluences)) || Parameters.MaterialParameters.bIsSpecialEngineMaterial);
 }
 
 
 template <GPUSkinBoneInfluenceType BoneInfluenceType>
-void TGPUSkinVertexFactory<BoneInfluenceType>::ModifyCompilationEnvironment(const FVertexFactoryType* Type, EShaderPlatform Platform, const FMaterial* Material, FShaderCompilerEnvironment& OutEnvironment )
+void TGPUSkinVertexFactory<BoneInfluenceType>::ModifyCompilationEnvironment(const FVertexFactoryShaderPermutationParameters& Parameters, FShaderCompilerEnvironment& OutEnvironment )
 {
-	FVertexFactory::ModifyCompilationEnvironment(Type, Platform, Material, OutEnvironment);
-	const int32 MaxGPUSkinBones = GetFeatureLevelMaxNumberOfBones(GetMaxSupportedFeatureLevel(Platform));
+	FVertexFactory::ModifyCompilationEnvironment(Parameters, OutEnvironment);
+	const int32 MaxGPUSkinBones = GetFeatureLevelMaxNumberOfBones(GetMaxSupportedFeatureLevel(Parameters.Platform));
 	OutEnvironment.SetDefine(TEXT("MAX_SHADER_BONES"), MaxGPUSkinBones);
 	const uint32 UseExtraBoneInfluences = (BoneInfluenceType == ExtraBoneInfluence);
 	OutEnvironment.SetDefine(TEXT("GPUSKIN_USE_EXTRA_INFLUENCES"), UseExtraBoneInfluences);
@@ -392,7 +392,7 @@ void TGPUSkinVertexFactory<BoneInfluenceType>::ModifyCompilationEnvironment(cons
 		bool bLimit2BoneInfluences = (CVarGPUSkinLimit2BoneInfluences.GetValueOnAnyThread() != 0);
 		OutEnvironment.SetDefine(TEXT("GPUSKIN_LIMIT_2BONE_INFLUENCES"), (bLimit2BoneInfluences ? 1 : 0));
 	}
-	OutEnvironment.SetDefine(TEXT("GPUSKIN_USE_BONES_SRV_BUFFER"), SupportsBonesBufferSRV(Platform) ? 1 : 0);
+	OutEnvironment.SetDefine(TEXT("GPUSKIN_USE_BONES_SRV_BUFFER"), SupportsBonesBufferSRV(Parameters.Platform) ? 1 : 0);
 	OutEnvironment.SetDefine(TEXT("GPUSKIN_UNLIMITED_BONE_INFLUENCE"), BoneInfluenceType == UnlimitedBoneInfluence ? 1 : 0);
 }
 
@@ -539,12 +539,13 @@ TGPUSkinVertexFactoryShaderParameters
 /** Shader parameters for use with TGPUSkinVertexFactory */
 class FGPUSkinVertexFactoryShaderParameters : public FVertexFactoryShaderParameters
 {
+	DECLARE_INLINE_TYPE_LAYOUT(FGPUSkinVertexFactoryShaderParameters, NonVirtual);
 public:
 	/**
 	* Bind shader constants by name
 	* @param	ParameterMap - mapping of named shader constants to indices
 	*/
-	virtual void Bind(const FShaderParameterMap& ParameterMap) override
+	void Bind(const FShaderParameterMap& ParameterMap)
 	{
 		PerBoneMotionBlur.Bind(ParameterMap,TEXT("PerBoneMotionBlur"));
 		BoneMatrices.Bind(ParameterMap,TEXT("BoneMatrices"));
@@ -552,20 +553,8 @@ public:
 		InputWeightIndexSize.Bind(ParameterMap, TEXT("InputWeightIndexSize"));
 		InputWeightStream.Bind(ParameterMap, TEXT("InputWeightStream"));
 	}
-	/**
-	* Serialize shader params to an archive
-	* @param	Ar - archive to serialize to
-	*/
-	virtual void Serialize(FArchive& Ar) override
-	{
-		Ar << PerBoneMotionBlur;
-		Ar << BoneMatrices;
-		Ar << PreviousBoneMatrices;
-		Ar << InputWeightIndexSize;
-		Ar << InputWeightStream;
-	}
 
-	virtual void GetElementShaderBindings(
+	void GetElementShaderBindings(
 		const FSceneInterface* Scene,
 		const FSceneView* View,
 		const FMeshMaterialShader* Shader,
@@ -574,7 +563,7 @@ public:
 		const FVertexFactory* VertexFactory,
 		const FMeshBatchElement& BatchElement,
 		class FMeshDrawSingleShaderBindings& ShaderBindings,
-		FVertexInputStreamArray& VertexStreams) const override
+		FVertexInputStreamArray& VertexStreams) const
 	{
 		const FGPUBaseSkinVertexFactory::FShaderDataType& ShaderData = ((const FGPUBaseSkinVertexFactory*)VertexFactory)->GetShaderData();
 
@@ -612,21 +601,16 @@ public:
 		}
 	}
 
-	virtual uint32 GetSize() const override { return sizeof(*this); }
-
 private:
-	FShaderParameter PerBoneMotionBlur;
-	FShaderResourceParameter BoneMatrices;
-	FShaderResourceParameter PreviousBoneMatrices;
-	FShaderParameter InputWeightIndexSize;
-	FShaderResourceParameter InputWeightStream;
+	LAYOUT_FIELD(FShaderParameter, PerBoneMotionBlur)
+	LAYOUT_FIELD(FShaderResourceParameter, BoneMatrices)
+	LAYOUT_FIELD(FShaderResourceParameter, PreviousBoneMatrices)
+	LAYOUT_FIELD(FShaderParameter, InputWeightIndexSize);
+	LAYOUT_FIELD(FShaderResourceParameter, InputWeightStream);
+
 };
 
-template <GPUSkinBoneInfluenceType BoneInfluenceType>
-FVertexFactoryShaderParameters* TGPUSkinVertexFactory<BoneInfluenceType>::ConstructShaderParameters(EShaderFrequency ShaderFrequency)
-{
-	return (ShaderFrequency == SF_Vertex) ? new FGPUSkinVertexFactoryShaderParameters() : NULL;
-}
+IMPLEMENT_GPUSKINNING_VERTEX_FACTORY_PARAMETER_TYPE(TGPUSkinVertexFactory, SF_Vertex, FGPUSkinVertexFactoryShaderParameters);
 
 /** bind gpu skin vertex factory to its shader file and its shader parameters */
 IMPLEMENT_GPUSKINNING_VERTEX_FACTORY_TYPE(TGPUSkinVertexFactory, "/Engine/Private/GpuSkinVertexFactory.ush", true, false, true, false, false);
@@ -638,29 +622,20 @@ TGPUSkinVertexFactoryShaderParameters
 /** Shader parameters for use with TGPUSkinVertexFactory */
 class FGPUSkinVertexPassthroughFactoryShaderParameters : public FLocalVertexFactoryShaderParametersBase
 {
+	DECLARE_INLINE_TYPE_LAYOUT(FGPUSkinVertexPassthroughFactoryShaderParameters, NonVirtual);
 public:
 	/**
 	* Bind shader constants by name
 	* @param	ParameterMap - mapping of named shader constants to indices
 	*/
-	virtual void Bind(const FShaderParameterMap& ParameterMap) override
+	void Bind(const FShaderParameterMap& ParameterMap)
 	{
 		FLocalVertexFactoryShaderParametersBase::Bind(ParameterMap);
 		GPUSkinCachePositionBuffer.Bind(ParameterMap,TEXT("GPUSkinCachePositionBuffer"));
 		GPUSkinCachePreviousPositionBuffer.Bind(ParameterMap,TEXT("GPUSkinCachePreviousPositionBuffer"));
 	}
-	/**
-	* Serialize shader params to an archive
-	* @param	Ar - archive to serialize to
-	*/
-	virtual void Serialize(FArchive& Ar) override
-	{
-		FLocalVertexFactoryShaderParametersBase::Serialize(Ar);
-		Ar << GPUSkinCachePositionBuffer;
-		Ar << GPUSkinCachePreviousPositionBuffer;
-	}
 	
-	virtual void GetElementShaderBindings(
+	void GetElementShaderBindings(
 		const FSceneInterface* Scene,
 		const FSceneView* View,
 		const FMeshMaterialShader* Shader,
@@ -669,7 +644,7 @@ public:
 		const FVertexFactory* VertexFactory,
 		const FMeshBatchElement& BatchElement,
 		class FMeshDrawSingleShaderBindings& ShaderBindings,
-		FVertexInputStreamArray& VertexStreams) const override
+		FVertexInputStreamArray& VertexStreams) const
 	{
 		check(VertexFactory->GetType() == &FGPUSkinPassthroughVertexFactory::StaticType);
 		FGPUSkinBatchElementUserData* BatchUserData = (FGPUSkinBatchElementUserData*)BatchElement.VertexFactoryUserData;
@@ -684,17 +659,17 @@ public:
 		FGPUSkinCache::GetShaderBindings(BatchUserData->Entry, BatchUserData->Section, Shader, (const FGPUSkinPassthroughVertexFactory*)VertexFactory, BatchElement.MinVertexIndex, GPUSkinCachePositionBuffer, GPUSkinCachePreviousPositionBuffer, ShaderBindings, VertexStreams);
 	}
 
-	virtual uint32 GetSize() const override { return sizeof(*this); }
-
 private:
-	FShaderResourceParameter GPUSkinCachePositionBuffer;
-	FShaderResourceParameter GPUSkinCachePreviousPositionBuffer;
+	
+		LAYOUT_FIELD(FShaderResourceParameter, GPUSkinCachePositionBuffer)
+		LAYOUT_FIELD(FShaderResourceParameter, GPUSkinCachePreviousPositionBuffer)
+	
 };
 
 /*-----------------------------------------------------------------------------
 FGPUSkinPassthroughVertexFactory
 -----------------------------------------------------------------------------*/
-void FGPUSkinPassthroughVertexFactory::ModifyCompilationEnvironment( const FVertexFactoryType* Type, EShaderPlatform Platform, const FMaterial* Material, FShaderCompilerEnvironment& OutEnvironment )
+void FGPUSkinPassthroughVertexFactory::ModifyCompilationEnvironment(const FVertexFactoryShaderPermutationParameters& Parameters, FShaderCompilerEnvironment& OutEnvironment )
 {
 	const bool ContainsManualVertexFetch = OutEnvironment.GetDefinitions().Contains("MANUAL_VERTEX_FETCH");
 	if (!ContainsManualVertexFetch)
@@ -702,14 +677,17 @@ void FGPUSkinPassthroughVertexFactory::ModifyCompilationEnvironment( const FVert
 		OutEnvironment.SetDefine(TEXT("MANUAL_VERTEX_FETCH"), TEXT("1"));
 	}
 
-	Super::ModifyCompilationEnvironment(Type, Platform, Material, OutEnvironment);
+	Super::ModifyCompilationEnvironment(Parameters, OutEnvironment);
 	OutEnvironment.SetDefine(TEXT("GPUSKIN_PASS_THROUGH"),TEXT("1"));
 }
 
-bool FGPUSkinPassthroughVertexFactory::ShouldCompilePermutation(EShaderPlatform Platform, const class FMaterial* Material, const FShaderType* ShaderType)
+bool FGPUSkinPassthroughVertexFactory::ShouldCompilePermutation(const FVertexFactoryShaderPermutationParameters& Parameters)
 {
 	// Passthrough is only valid on platforms with Compute Shader support AND for (skeletal meshes or default materials)
-	return IsGPUSkinCacheAvailable(Platform) && IsFeatureLevelSupported(Platform, ERHIFeatureLevel::SM5) && Super::ShouldCompilePermutation(Platform, Material, ShaderType) && (Material->IsUsedWithSkeletalMesh() || Material->IsSpecialEngineMaterial());
+	return IsGPUSkinCacheAvailable(Parameters.Platform) &&
+		IsFeatureLevelSupported(Parameters.Platform, ERHIFeatureLevel::SM5) &&
+		Super::ShouldCompilePermutation(Parameters) &&
+		(Parameters.MaterialParameters.bIsUsedWithSkeletalMesh || Parameters.MaterialParameters.bIsSpecialEngineMaterial);
 }
 
 void FGPUSkinPassthroughVertexFactory::InternalUpdateVertexDeclaration(FGPUBaseSkinVertexFactory* SourceVertexFactory, struct FRWBuffer* PositionRWBuffer, struct FRWBuffer* TangentRWBuffer)
@@ -781,15 +759,11 @@ void FGPUSkinPassthroughVertexFactory::InternalUpdateVertexDeclaration(FGPUBaseS
 	checkf(PositionStreamIndex != -1, TEXT("Unable to find stream for RWBuffer Vertex buffer!"));
 }
 
-FVertexFactoryShaderParameters* FGPUSkinPassthroughVertexFactory::ConstructShaderParameters(EShaderFrequency ShaderFrequency)
-{
+IMPLEMENT_VERTEX_FACTORY_PARAMETER_TYPE(FGPUSkinPassthroughVertexFactory, SF_Vertex, FGPUSkinVertexPassthroughFactoryShaderParameters);
 #if RHI_RAYTRACING
-	return (ShaderFrequency == SF_Vertex || ShaderFrequency == SF_RayHitGroup || ShaderFrequency == SF_Compute) ? new FGPUSkinVertexPassthroughFactoryShaderParameters() : nullptr;
-#else // RHI_RAYTRACING
-	return (ShaderFrequency == SF_Vertex) ? new FGPUSkinVertexPassthroughFactoryShaderParameters() : nullptr;
+IMPLEMENT_VERTEX_FACTORY_PARAMETER_TYPE(FGPUSkinPassthroughVertexFactory, SF_RayHitGroup, FGPUSkinVertexPassthroughFactoryShaderParameters);
+IMPLEMENT_VERTEX_FACTORY_PARAMETER_TYPE(FGPUSkinPassthroughVertexFactory, SF_Compute, FGPUSkinVertexPassthroughFactoryShaderParameters);
 #endif // RHI_RAYTRACING
-}
-
 IMPLEMENT_VERTEX_FACTORY_TYPE(FGPUSkinPassthroughVertexFactory, "/Engine/Private/LocalVertexFactory.ush", true, false, true, false, false);
 
 /*-----------------------------------------------------------------------------
@@ -801,17 +775,17 @@ TGPUSkinMorphVertexFactory
 * @param OutEnvironment - shader compile environment to modify
 */
 template <GPUSkinBoneInfluenceType BoneInfluenceType>
-void TGPUSkinMorphVertexFactory<BoneInfluenceType>::ModifyCompilationEnvironment( const FVertexFactoryType* Type, EShaderPlatform Platform, const FMaterial* Material, FShaderCompilerEnvironment& OutEnvironment )
+void TGPUSkinMorphVertexFactory<BoneInfluenceType>::ModifyCompilationEnvironment( const FVertexFactoryShaderPermutationParameters& Parameters, FShaderCompilerEnvironment& OutEnvironment )
 {
-	Super::ModifyCompilationEnvironment(Type, Platform, Material, OutEnvironment);
+	Super::ModifyCompilationEnvironment(Parameters, OutEnvironment);
 	OutEnvironment.SetDefine(TEXT("GPUSKIN_MORPH_BLEND"),TEXT("1"));
 }
 
 template <GPUSkinBoneInfluenceType BoneInfluenceType>
-bool TGPUSkinMorphVertexFactory<BoneInfluenceType>::ShouldCompilePermutation(EShaderPlatform Platform, const class FMaterial* Material, const class FShaderType* ShaderType)
+bool TGPUSkinMorphVertexFactory<BoneInfluenceType>::ShouldCompilePermutation(const FVertexFactoryShaderPermutationParameters& Parameters)
 {
-	return (Material->IsUsedWithMorphTargets() || Material->IsSpecialEngineMaterial()) 
-		&& Super::ShouldCompilePermutation(Platform, Material, ShaderType);
+	return (Parameters.MaterialParameters.bIsUsedWithMorphTargets || Parameters.MaterialParameters.bIsSpecialEngineMaterial)
+		&& Super::ShouldCompilePermutation(Parameters);
 }
 
 /**
@@ -844,11 +818,7 @@ void TGPUSkinMorphVertexFactory<BoneInfluenceType>::InitRHI()
 	FVertexFactory::InitDeclaration(Elements);
 }
 
-template <GPUSkinBoneInfluenceType BoneInfluenceType>
-FVertexFactoryShaderParameters* TGPUSkinMorphVertexFactory<BoneInfluenceType>::ConstructShaderParameters(EShaderFrequency ShaderFrequency)
-{
-	return ShaderFrequency == SF_Vertex ? new FGPUSkinVertexFactoryShaderParameters() : NULL;
-}
+IMPLEMENT_GPUSKINNING_VERTEX_FACTORY_PARAMETER_TYPE(TGPUSkinMorphVertexFactory, SF_Vertex, FGPUSkinVertexFactoryShaderParameters);
 
 /** bind morph target gpu skin vertex factory to its shader file and its shader parameters */
 IMPLEMENT_GPUSKINNING_VERTEX_FACTORY_TYPE(TGPUSkinMorphVertexFactory, "/Engine/Private/GpuSkinVertexFactory.ush", true, false, true, false, false);
@@ -860,13 +830,14 @@ IMPLEMENT_GPUSKINNING_VERTEX_FACTORY_TYPE(TGPUSkinMorphVertexFactory, "/Engine/P
 /** Shader parameters for use with TGPUSkinAPEXClothVertexFactory */
 class TGPUSkinAPEXClothVertexFactoryShaderParameters : public FGPUSkinVertexFactoryShaderParameters
 {
+	DECLARE_INLINE_TYPE_LAYOUT(TGPUSkinAPEXClothVertexFactoryShaderParameters, NonVirtual);
 public:
 
 	/**
 	* Bind shader constants by name
 	* @param	ParameterMap - mapping of named shader constants to indices
 	*/
-	virtual void Bind(const FShaderParameterMap& ParameterMap) override
+	void Bind(const FShaderParameterMap& ParameterMap)
 	{
 		FGPUSkinVertexFactoryShaderParameters::Bind(ParameterMap);
 		ClothSimulVertsPositionsNormalsParameter.Bind(ParameterMap,TEXT("ClothSimulVertsPositionsNormals"));
@@ -877,23 +848,8 @@ public:
 		GPUSkinApexClothParameter.Bind(ParameterMap, TEXT("GPUSkinApexCloth"));
 		GPUSkinApexClothStartIndexOffsetParameter.Bind(ParameterMap, TEXT("GPUSkinApexClothStartIndexOffset"));
 	}
-	/**
-	* Serialize shader params to an archive
-	* @param	Ar - archive to serialize to
-	*/
-	virtual void Serialize(FArchive& Ar) override
-	{ 
-		FGPUSkinVertexFactoryShaderParameters::Serialize(Ar);
-		Ar << ClothSimulVertsPositionsNormalsParameter;
-		Ar << PreviousClothSimulVertsPositionsNormalsParameter;
-		Ar << ClothLocalToWorldParameter;
-		Ar << PreviousClothLocalToWorldParameter;
-		Ar << ClothBlendWeightParameter;
-		Ar << GPUSkinApexClothParameter;
-		Ar << GPUSkinApexClothStartIndexOffsetParameter;
-	}
 	
-	virtual void GetElementShaderBindings(
+	void GetElementShaderBindings(
 		const FSceneInterface* Scene,
 		const FSceneView* View,
 		const FMeshMaterialShader* Shader,
@@ -902,7 +858,7 @@ public:
 		const FVertexFactory* VertexFactory,
 		const FMeshBatchElement& BatchElement,
 		class FMeshDrawSingleShaderBindings& ShaderBindings,
-		FVertexInputStreamArray& VertexStreams) const override
+		FVertexInputStreamArray& VertexStreams) const
 	{
 		// Call regular GPU skinning shader parameters
 		FGPUSkinVertexFactoryShaderParameters::GetElementShaderBindings(Scene, View, Shader, InputStreamType, FeatureLevel, VertexFactory, BatchElement, ShaderBindings, VertexStreams);
@@ -939,13 +895,15 @@ public:
 	}
 
 protected:
-	FShaderResourceParameter ClothSimulVertsPositionsNormalsParameter;
-	FShaderResourceParameter PreviousClothSimulVertsPositionsNormalsParameter;
-	FShaderParameter ClothLocalToWorldParameter;
-	FShaderParameter PreviousClothLocalToWorldParameter;
-	FShaderParameter ClothBlendWeightParameter;
-	FShaderResourceParameter GPUSkinApexClothParameter;
-	FShaderParameter GPUSkinApexClothStartIndexOffsetParameter;
+	
+		LAYOUT_FIELD(FShaderResourceParameter, ClothSimulVertsPositionsNormalsParameter)
+		LAYOUT_FIELD(FShaderResourceParameter, PreviousClothSimulVertsPositionsNormalsParameter)
+		LAYOUT_FIELD(FShaderParameter, ClothLocalToWorldParameter)
+		LAYOUT_FIELD(FShaderParameter, PreviousClothLocalToWorldParameter)
+		LAYOUT_FIELD(FShaderParameter, ClothBlendWeightParameter)
+		LAYOUT_FIELD(FShaderResourceParameter, GPUSkinApexClothParameter)
+		LAYOUT_FIELD(FShaderParameter, GPUSkinApexClothStartIndexOffsetParameter)
+	
 };
 
 /*-----------------------------------------------------------------------------
@@ -1045,18 +1003,18 @@ TGlobalResource<FClothBufferPool> FGPUBaseSkinAPEXClothVertexFactory::ClothSimul
 * @param OutEnvironment - shader compile environment to modify
 */
 template <GPUSkinBoneInfluenceType BoneInfluenceType>
-void TGPUSkinAPEXClothVertexFactory<BoneInfluenceType>::ModifyCompilationEnvironment( const FVertexFactoryType* Type, EShaderPlatform Platform, const FMaterial* Material, FShaderCompilerEnvironment& OutEnvironment )
+void TGPUSkinAPEXClothVertexFactory<BoneInfluenceType>::ModifyCompilationEnvironment( const FVertexFactoryShaderPermutationParameters& Parameters, FShaderCompilerEnvironment& OutEnvironment )
 {
-	Super::ModifyCompilationEnvironment(Type, Platform, Material, OutEnvironment);
+	Super::ModifyCompilationEnvironment(Parameters, OutEnvironment);
 	OutEnvironment.SetDefine(TEXT("GPUSKIN_APEX_CLOTH"),TEXT("1"));
 }
 
 template <GPUSkinBoneInfluenceType BoneInfluenceType>
-bool TGPUSkinAPEXClothVertexFactory<BoneInfluenceType>::ShouldCompilePermutation(EShaderPlatform Platform, const class FMaterial* Material, const class FShaderType* ShaderType)
+bool TGPUSkinAPEXClothVertexFactory<BoneInfluenceType>::ShouldCompilePermutation(const FVertexFactoryShaderPermutationParameters& Parameters)
 {
-	return GetMaxSupportedFeatureLevel(Platform) >= ERHIFeatureLevel::SM5
-		&& (Material->IsUsedWithAPEXCloth() || Material->IsSpecialEngineMaterial()) 
-		&& Super::ShouldCompilePermutation(Platform, Material, ShaderType);
+	return GetMaxSupportedFeatureLevel(Parameters.Platform) >= ERHIFeatureLevel::SM5
+		&& (Parameters.MaterialParameters.bIsUsedWithAPEXCloth || Parameters.MaterialParameters.bIsSpecialEngineMaterial)
+		&& Super::ShouldCompilePermutation(Parameters);
 }
 
 /**
@@ -1095,14 +1053,10 @@ void TGPUSkinAPEXClothVertexFactory<BoneInfluenceType>::InitRHI()
 	FVertexFactory::InitDeclaration(Elements);
 }
 
-template <GPUSkinBoneInfluenceType BoneInfluenceType>
-FVertexFactoryShaderParameters* TGPUSkinAPEXClothVertexFactory<BoneInfluenceType>::ConstructShaderParameters(EShaderFrequency ShaderFrequency)
-{
-	return ShaderFrequency == SF_Vertex ? new TGPUSkinAPEXClothVertexFactoryShaderParameters() : NULL;
-}
+IMPLEMENT_GPUSKINNING_VERTEX_FACTORY_PARAMETER_TYPE(TGPUSkinAPEXClothVertexFactory, SF_Vertex, TGPUSkinAPEXClothVertexFactoryShaderParameters);
 
 /** bind cloth gpu skin vertex factory to its shader file and its shader parameters */
 IMPLEMENT_GPUSKINNING_VERTEX_FACTORY_TYPE(TGPUSkinAPEXClothVertexFactory, "/Engine/Private/GpuSkinVertexFactory.ush", true, false, true, false, false);
 
-
+#undef IMPLEMENT_GPUSKINNING_VERTEX_FACTORY_PARAMETER_TYPE
 #undef IMPLEMENT_GPUSKINNING_VERTEX_FACTORY_TYPE
diff --git a/Engine/Source/Runtime/Engine/Private/GPUSort.cpp b/Engine/Source/Runtime/Engine/Private/GPUSort.cpp
--- a/Engine/Source/Runtime/Engine/Private/GPUSort.cpp
+++ b/Engine/Source/Runtime/Engine/Private/GPUSort.cpp
@@ -248,20 +248,12 @@ public:
 		OutOffsets.Bind( Initializer.ParameterMap, TEXT("OutOffsets") );
 	}
 
-	/** Serialization. */
-	virtual bool Serialize( FArchive& Ar ) override
-	{
-		bool bShaderHasOutdatedParameters = FGlobalShader::Serialize( Ar );
-		Ar << OutOffsets;
-		return bShaderHasOutdatedParameters;
-	}
-
 	/**
 	 * Set parameters for this shader.
 	 */
 	void SetParameters(FRHICommandList& RHICmdList, const FRadixSortUniformBufferRef& UniformBuffer )
 	{
-		FRHIComputeShader* ComputeShaderRHI = GetComputeShader();
+		FRHIComputeShader* ComputeShaderRHI = RHICmdList.GetBoundComputeShader();
 		SetUniformBufferParameter(RHICmdList, ComputeShaderRHI, GetUniformBufferParameter<FRadixSortParameters>(), UniformBuffer );
 	}
 
@@ -270,7 +262,7 @@ public:
 	 */
 	void SetOutput(FRHICommandList& RHICmdList, FRHIUnorderedAccessView* OutOffsetsUAV)
 	{
-		FRHIComputeShader* ComputeShaderRHI = GetComputeShader();
+		FRHIComputeShader* ComputeShaderRHI = RHICmdList.GetBoundComputeShader();
 		if ( OutOffsets.IsBound() )
 		{
 			RHICmdList.SetUAVParameter(ComputeShaderRHI, OutOffsets.GetBaseIndex(), OutOffsetsUAV);
@@ -282,7 +274,7 @@ public:
 	 */
 	void UnbindBuffers(FRHICommandList& RHICmdList)
 	{
-		FRHIComputeShader* ComputeShaderRHI = GetComputeShader();
+		FRHIComputeShader* ComputeShaderRHI = RHICmdList.GetBoundComputeShader();
 		if ( OutOffsets.IsBound() )
 		{
 			RHICmdList.SetUAVParameter(ComputeShaderRHI, OutOffsets.GetBaseIndex(), nullptr);
@@ -292,7 +284,7 @@ public:
 private:
 
 	/** The buffer to which offsets will be written. */
-	FShaderResourceParameter OutOffsets;
+	LAYOUT_FIELD(FShaderResourceParameter, OutOffsets);
 };
 IMPLEMENT_SHADER_TYPE(,FRadixSortClearOffsetsCS,TEXT("/Engine/Private/RadixSortShaders.usf"),TEXT("RadixSort_ClearOffsets"),SF_Compute);
 
@@ -334,16 +326,6 @@ public:
 		OutOffsets.Bind( Initializer.ParameterMap, TEXT("OutOffsets") );
 	}
 
-	/** Serialization. */
-	virtual bool Serialize( FArchive& Ar ) override
-	{
-		bool bShaderHasOutdatedParameters = FGlobalShader::Serialize( Ar );
-		Ar << RadixSortParameterBuffer;
-		Ar << InKeys;
-		Ar << OutOffsets;
-		return bShaderHasOutdatedParameters;
-	}
-
 	/**
 	 * Returns true if this shader was compiled to require the constant buffer
 	 * workaround.
@@ -358,7 +340,7 @@ public:
 	 */
 	void SetParameters(FRHICommandList& RHICmdList, FRHIShaderResourceView* InKeysSRV, FRadixSortUniformBufferRef& RadixSortUniformBuffer, FRHIShaderResourceView* RadixSortParameterBufferSRV)
 	{
-		FRHIComputeShader* ComputeShaderRHI = GetComputeShader();
+		FRHIComputeShader* ComputeShaderRHI = RHICmdList.GetBoundComputeShader();
 		SetUniformBufferParameter(RHICmdList, ComputeShaderRHI, GetUniformBufferParameter<FRadixSortParameters>(), RadixSortUniformBuffer );
 		if ( InKeys.IsBound() )
 		{
@@ -375,7 +357,7 @@ public:
 	 */
 	void SetOutput(FRHICommandList& RHICmdList, FRHIUnorderedAccessView* OutOffsetsUAV)
 	{
-		FRHIComputeShader* ComputeShaderRHI = GetComputeShader();
+		FRHIComputeShader* ComputeShaderRHI = RHICmdList.GetBoundComputeShader();
 		if ( OutOffsets.IsBound() )
 		{
 			RHICmdList.SetUAVParameter(ComputeShaderRHI, OutOffsets.GetBaseIndex(), OutOffsetsUAV);
@@ -387,7 +369,7 @@ public:
 	 */
 	void UnbindBuffers(FRHICommandList& RHICmdList)
 	{
-		FRHIComputeShader* ComputeShaderRHI = GetComputeShader();
+		FRHIComputeShader* ComputeShaderRHI = RHICmdList.GetBoundComputeShader();
 		if ( RadixSortParameterBuffer.IsBound() )
 		{
 			RHICmdList.SetShaderResourceViewParameter(ComputeShaderRHI, RadixSortParameterBuffer.GetBaseIndex(), nullptr);
@@ -403,13 +385,12 @@ public:
 	}
 
 private:
-
 	/** Uniform parameters stored in a vertex buffer, used to workaround an NVIDIA driver bug. */
-	FShaderResourceParameter RadixSortParameterBuffer;
+	LAYOUT_FIELD(FShaderResourceParameter, RadixSortParameterBuffer);
 	/** The buffer containing input keys. */
-	FShaderResourceParameter InKeys;
+	LAYOUT_FIELD(FShaderResourceParameter, InKeys);
 	/** The buffer to which offsets will be written. */
-	FShaderResourceParameter OutOffsets;
+	LAYOUT_FIELD(FShaderResourceParameter, OutOffsets);
 };
 IMPLEMENT_SHADER_TYPE(,FRadixSortUpsweepCS,TEXT("/Engine/Private/RadixSortShaders.usf"),TEXT("RadixSort_Upsweep"),SF_Compute);
 
@@ -450,21 +431,12 @@ public:
 		OutOffsets.Bind( Initializer.ParameterMap, TEXT("OutOffsets") );
 	}
 
-	/** Serialization. */
-	virtual bool Serialize( FArchive& Ar ) override
-	{
-		bool bShaderHasOutdatedParameters = FGlobalShader::Serialize( Ar );
-		Ar << InOffsets;
-		Ar << OutOffsets;
-		return bShaderHasOutdatedParameters;
-	}
-
 	/**
 	 * Set parameters for this shader.
 	 */
 	void SetParameters(FRHICommandList& RHICmdList, FRHIShaderResourceView* InOffsetsSRV)
 	{
-		FRHIComputeShader* ComputeShaderRHI = GetComputeShader();
+		FRHIComputeShader* ComputeShaderRHI = RHICmdList.GetBoundComputeShader();
 		if ( InOffsets.IsBound() )
 		{
 			RHICmdList.SetShaderResourceViewParameter(ComputeShaderRHI, InOffsets.GetBaseIndex(), InOffsetsSRV);
@@ -476,7 +448,7 @@ public:
 	 */
 	void SetOutput(FRHICommandList& RHICmdList, FRHIUnorderedAccessView* OutOffsetsUAV)
 	{
-		FRHIComputeShader* ComputeShaderRHI = GetComputeShader();
+		FRHIComputeShader* ComputeShaderRHI = RHICmdList.GetBoundComputeShader();
 		if ( OutOffsets.IsBound() )
 		{
 			RHICmdList.SetUAVParameter(ComputeShaderRHI, OutOffsets.GetBaseIndex(), OutOffsetsUAV);
@@ -488,7 +460,7 @@ public:
 	 */
 	void UnbindBuffers(FRHICommandList& RHICmdList)
 	{
-		FRHIComputeShader* ComputeShaderRHI = GetComputeShader();
+		FRHIComputeShader* ComputeShaderRHI = RHICmdList.GetBoundComputeShader();
 		if ( InOffsets.IsBound() )
 		{
 			RHICmdList.SetShaderResourceViewParameter(ComputeShaderRHI, InOffsets.GetBaseIndex(), nullptr);
@@ -502,9 +474,9 @@ public:
 private:
 
 	/** The buffer containing input offsets. */
-	FShaderResourceParameter InOffsets;
+	LAYOUT_FIELD(FShaderResourceParameter, InOffsets);
 	/** The buffer to which offsets will be written. */
-	FShaderResourceParameter OutOffsets;
+	LAYOUT_FIELD(FShaderResourceParameter, OutOffsets);
 };
 IMPLEMENT_SHADER_TYPE(,FRadixSortSpineCS,TEXT("/Engine/Private/RadixSortShaders.usf"),TEXT("RadixSort_Spine"),SF_Compute);
 
@@ -551,19 +523,6 @@ public:
 		OutValues.Bind( Initializer.ParameterMap, TEXT("OutValues") );
 	}
 
-	/** Serialization. */
-	virtual bool Serialize( FArchive& Ar ) override
-	{
-		bool bShaderHasOutdatedParameters = FGlobalShader::Serialize( Ar );
-		Ar << RadixSortParameterBuffer;
-		Ar << InKeys;
-		Ar << InValues;
-		Ar << InOffsets;
-		Ar << OutKeys;
-		Ar << OutValues;
-		return bShaderHasOutdatedParameters;
-	}
-
 	/**
 	 * Returns true if this shader was compiled to require the constant buffer
 	 * workaround.
@@ -584,7 +543,7 @@ public:
 		FRadixSortUniformBufferRef& RadixSortUniformBuffer,
 		FRHIShaderResourceView* RadixSortParameterBufferSRV )
 	{
-		FRHIComputeShader* ComputeShaderRHI = GetComputeShader();
+		FRHIComputeShader* ComputeShaderRHI = RHICmdList.GetBoundComputeShader();
 		SetUniformBufferParameter(RHICmdList, ComputeShaderRHI, GetUniformBufferParameter<FRadixSortParameters>(), RadixSortUniformBuffer );
 		if ( RadixSortParameterBuffer.IsBound() )
 		{
@@ -609,7 +568,7 @@ public:
 	 */
 	void SetOutput(FRHICommandList& RHICmdList, FRHIUnorderedAccessView* OutKeysUAV, FRHIUnorderedAccessView* OutValuesUAV)
 	{
-		FRHIComputeShader* ComputeShaderRHI = GetComputeShader();
+		FRHIComputeShader* ComputeShaderRHI = RHICmdList.GetBoundComputeShader();
 		if ( OutKeys.IsBound() )
 		{
 			RHICmdList.SetUAVParameter(ComputeShaderRHI, OutKeys.GetBaseIndex(), OutKeysUAV);
@@ -625,7 +584,7 @@ public:
 	 */
 	void UnbindBuffers(FRHICommandList& RHICmdList)
 	{
-		FRHIComputeShader* ComputeShaderRHI = GetComputeShader();
+		FRHIComputeShader* ComputeShaderRHI = RHICmdList.GetBoundComputeShader();
 		if ( RadixSortParameterBuffer.IsBound() )
 		{
 			RHICmdList.SetShaderResourceViewParameter(ComputeShaderRHI, RadixSortParameterBuffer.GetBaseIndex(), nullptr);
@@ -653,19 +612,18 @@ public:
 	}
 
 private:
-
 	/** Uniform parameters stored in a vertex buffer, used to workaround an NVIDIA driver bug. */
-	FShaderResourceParameter RadixSortParameterBuffer;
+	LAYOUT_FIELD(FShaderResourceParameter, RadixSortParameterBuffer);
 	/** The buffer containing input keys. */
-	FShaderResourceParameter InKeys;
+	LAYOUT_FIELD(FShaderResourceParameter, InKeys);
 	/** The buffer containing input values. */
-	FShaderResourceParameter InValues;
+	LAYOUT_FIELD(FShaderResourceParameter, InValues);
 	/** The buffer from which offsets will be read. */
-	FShaderResourceParameter InOffsets;
+	LAYOUT_FIELD(FShaderResourceParameter, InOffsets);
 	/** The buffer to which sorted keys will be written. */
-	FShaderResourceParameter OutKeys;
+	LAYOUT_FIELD(FShaderResourceParameter, OutKeys);
 	/** The buffer to which sorted values will be written. */
-	FShaderResourceParameter OutValues;
+	LAYOUT_FIELD(FShaderResourceParameter, OutValues);
 };
 IMPLEMENT_SHADER_TYPE(,FRadixSortDownsweepCS,TEXT("/Engine/Private/RadixSortShaders.usf"),TEXT("RadixSort_Downsweep"),SF_Compute);
 
@@ -786,19 +744,19 @@ int32 SortGPUBuffers(FRHICommandListImmediate& RHICmdList, FGPUSortBuffers SortB
 			RHICmdList.TransitionResource(EResourceTransitionAccess::ERWBarrier, EResourceTransitionPipeline::EComputeToCompute, GSortOffsetBuffers.BufferUAVs[0]);
 			
 			// Clear the offsets buffer.			
-			RHICmdList.SetComputeShader(ClearOffsetsCS->GetComputeShader());			
+			RHICmdList.SetComputeShader(ClearOffsetsCS.GetComputeShader());			
 			ClearOffsetsCS->SetOutput(RHICmdList, GSortOffsetBuffers.BufferUAVs[0]);
-			DispatchComputeShader(RHICmdList, *ClearOffsetsCS, 1, 1 ,1 );
+			DispatchComputeShader(RHICmdList, ClearOffsetsCS.GetShader(), 1, 1 ,1 );
 			ClearOffsetsCS->UnbindBuffers(RHICmdList);
 
 			//make UAV safe for readback
 			RHICmdList.TransitionResource(EResourceTransitionAccess::ERWBarrier, EResourceTransitionPipeline::EComputeToCompute, GSortOffsetBuffers.BufferUAVs[0]);
 
 			// Phase 1: Scan upsweep to compute per-digit totals.
-			RHICmdList.SetComputeShader(UpsweepCS->GetComputeShader());
+			RHICmdList.SetComputeShader(UpsweepCS.GetComputeShader());
 			UpsweepCS->SetOutput(RHICmdList, GSortOffsetBuffers.BufferUAVs[0]);
 			UpsweepCS->SetParameters(RHICmdList, SortBuffers.RemoteKeySRVs[BufferIndex], SortUniformBufferRef, GRadixSortParametersBuffer.SortParametersBufferSRV );
-			DispatchComputeShader(RHICmdList, *UpsweepCS, GroupCount, 1, 1 );
+			DispatchComputeShader(RHICmdList, UpsweepCS.GetShader(), GroupCount, 1, 1 );
 			UpsweepCS->UnbindBuffers(RHICmdList);
 
 			//barrier both UAVS since for next step.
@@ -815,10 +773,10 @@ int32 SortGPUBuffers(FRHICommandListImmediate& RHICmdList, FGPUSortBuffers SortB
 			}
 
 			// Phase 2: Parallel prefix scan on the offsets buffer.
-			RHICmdList.SetComputeShader(SpineCS->GetComputeShader());
+			RHICmdList.SetComputeShader(SpineCS.GetComputeShader());
 			SpineCS->SetOutput(RHICmdList, GSortOffsetBuffers.BufferUAVs[1]);
 			SpineCS->SetParameters(RHICmdList, GSortOffsetBuffers.BufferSRVs[0] );
-			DispatchComputeShader(RHICmdList, *SpineCS, 1, 1, 1 );
+			DispatchComputeShader(RHICmdList, SpineCS.GetShader(), 1, 1, 1 );
 			SpineCS->UnbindBuffers(RHICmdList);
 
 			if (bDebugOffsets)
@@ -839,7 +797,7 @@ int32 SortGPUBuffers(FRHICommandListImmediate& RHICmdList, FGPUSortBuffers SortB
 
 			const bool bIsLastPass = ((PassBits << RADIX_BITS) & KeyMask) == 0;
 			// Phase 3: Downsweep to compute final offsets and scatter keys.
-			RHICmdList.SetComputeShader(DownsweepCS->GetComputeShader());
+			RHICmdList.SetComputeShader(DownsweepCS.GetComputeShader());
 			{
 				FRHIUnorderedAccessView* ValuesUAV = nullptr;
 				if (bIsLastPass && SortBuffers.FinalValuesUAV)
@@ -858,7 +816,7 @@ int32 SortGPUBuffers(FRHICommandListImmediate& RHICmdList, FGPUSortBuffers SortB
 				FRHIShaderResourceView* ValuesSRV = (PassIndex == 0 && SortBuffers.FirstValuesSRV) ? SortBuffers.FirstValuesSRV : SortBuffers.RemoteValueSRVs[BufferIndex];
 				DownsweepCS->SetParameters(RHICmdList, SortBuffers.RemoteKeySRVs[BufferIndex], ValuesSRV, GSortOffsetBuffers.BufferSRVs[1], SortUniformBufferRef, GRadixSortParametersBuffer.SortParametersBufferSRV );
 			}
-			DispatchComputeShader(RHICmdList, *DownsweepCS, GroupCount, 1, 1 );
+			DispatchComputeShader(RHICmdList, DownsweepCS.GetShader(), GroupCount, 1, 1 );
 			DownsweepCS->UnbindBuffers(RHICmdList);
 
 			RHICmdList.TransitionResources(EResourceTransitionAccess::ERWBarrier, EResourceTransitionPipeline::EComputeToCompute, PrePhase3BarrierUAVS, 2);
diff --git a/Engine/Source/Runtime/Engine/Private/GPUSortManager.cpp b/Engine/Source/Runtime/Engine/Private/GPUSortManager.cpp
--- a/Engine/Source/Runtime/Engine/Private/GPUSortManager.cpp
+++ b/Engine/Source/Runtime/Engine/Private/GPUSortManager.cpp
@@ -107,9 +107,6 @@ public:
 	/** Initialization constructor. */
 	explicit FCopyUIntBufferCS(const ShaderMetaType::CompiledShaderInitializerType& Initializer);
 
-	/** Serialization. */
-	virtual bool Serialize(FArchive& Ar) override;
-
 	/**
 	 * Set parameters.
 	 */
@@ -128,9 +125,9 @@ public:
 
 private:
 
-	FShaderParameter CopyParams;
-	FShaderResourceParameter SourceData;
-	FShaderResourceParameter DestData[COPYUINTCS_BUFFER_COUNT];
+	LAYOUT_FIELD(FShaderParameter, CopyParams);
+	LAYOUT_FIELD(FShaderResourceParameter, SourceData);
+	LAYOUT_ARRAY(FShaderResourceParameter, DestData, COPYUINTCS_BUFFER_COUNT);
 };
 
 //*****************************************************************************
@@ -147,17 +144,6 @@ FCopyUIntBufferCS::FCopyUIntBufferCS(const ShaderMetaType::CompiledShaderInitial
 	DestData[2].Bind(Initializer.ParameterMap, TEXT("DestData2"));
 }
 
-bool FCopyUIntBufferCS::Serialize(FArchive& Ar)
-{
-	bool bShaderHasOutdatedParameters = FGlobalShader::Serialize(Ar);
-	Ar << CopyParams;
-	Ar << SourceData;
-	Ar << DestData[0];
-	Ar << DestData[1];
-	Ar << DestData[2];
-	return bShaderHasOutdatedParameters;
-}
-
 void FCopyUIntBufferCS::SetParameters(
 	FRHICommandList& RHICmdList,
 	FRHIShaderResourceView* InSourceData,
@@ -168,7 +154,7 @@ void FCopyUIntBufferCS::SetParameters(
 {
 	static TGlobalResource<FGPUSortDummyUAV> NiagaraSortingDummyUAV[COPYUINTCS_BUFFER_COUNT];
 		
-	FRHIComputeShader* ComputeShaderRHI = GetComputeShader();
+	FRHIComputeShader* ComputeShaderRHI = RHICmdList.GetBoundComputeShader();
 	check(DestCount > 0 && DestCount <= COPYUINTCS_BUFFER_COUNT);
 
 	RHICmdList.SetShaderResourceViewParameter(ComputeShaderRHI, SourceData.GetBaseIndex(), InSourceData);
@@ -192,7 +178,7 @@ void FCopyUIntBufferCS::SetParameters(
 
 void FCopyUIntBufferCS::UnbindBuffers(FRHICommandList& RHICmdList)
 {
-	FRHIComputeShader* ComputeShaderRHI = GetComputeShader();
+	FRHIComputeShader* ComputeShaderRHI = RHICmdList.GetBoundComputeShader();
 	SetSRVParameter(RHICmdList, ComputeShaderRHI, SourceData, nullptr);
 	for (int32 Index = 0; Index < COPYUINTCS_BUFFER_COUNT; ++Index)
 	{
@@ -207,7 +193,7 @@ void CopyUIntBufferToTargets(FRHICommandListImmediate& RHICmdList, ERHIFeatureLe
 	// No that resource transition must be made outside this call as we don't know how the content of those have been generated, and will be used.
 
 	TShaderMapRef<FCopyUIntBufferCS> CopyBufferCS(GetGlobalShaderMap(FeatureLevel));
-	RHICmdList.SetComputeShader(CopyBufferCS->GetComputeShader());
+	RHICmdList.SetComputeShader(CopyBufferCS.GetComputeShader());
 	
 	int32 Index0InPass = 0;
 	while (Index0InPass < NumTargets)
@@ -216,7 +202,7 @@ void CopyUIntBufferToTargets(FRHICommandListImmediate& RHICmdList, ERHIFeatureLe
 		const int32 NumElementsInPass = TargetSizes[Index0InPass + NumTargetsInPass - 1] - StartingOffset;
 
 		CopyBufferCS->SetParameters(RHICmdList, SourceSRV, TargetUAVs + Index0InPass, TargetSizes + Index0InPass, StartingOffset, NumTargetsInPass);
-		DispatchComputeShader(RHICmdList, *CopyBufferCS, FMath::DivideAndRoundUp(NumElementsInPass, COPYUINTCS_THREAD_COUNT), 1, 1);
+		DispatchComputeShader(RHICmdList, CopyBufferCS, FMath::DivideAndRoundUp(NumElementsInPass, COPYUINTCS_THREAD_COUNT), 1, 1);
 
 		StartingOffset += NumElementsInPass;
 		Index0InPass += COPYUINTCS_BUFFER_COUNT;
diff --git a/Engine/Source/Runtime/Engine/Private/GeometryCacheVertexFactory.cpp b/Engine/Source/Runtime/Engine/Private/GeometryCacheVertexFactory.cpp
--- a/Engine/Source/Runtime/Engine/Private/GeometryCacheVertexFactory.cpp
+++ b/Engine/Source/Runtime/Engine/Private/GeometryCacheVertexFactory.cpp
@@ -21,13 +21,14 @@ IMPLEMENT_GLOBAL_SHADER_PARAMETER_STRUCT(FGeometryCacheManualVertexFetchUniformB
 /** Shader parameters for use with TGPUSkinVertexFactory */
 class FGeometryCacheVertexFactoryShaderParameters : public FVertexFactoryShaderParameters
 {
+	DECLARE_INLINE_TYPE_LAYOUT(FGeometryCacheVertexFactoryShaderParameters, NonVirtual);
 public:
 
 	/**
 	* Bind shader constants by name
 	* @param	ParameterMap - mapping of named shader constants to indices
 	*/
-	virtual void Bind(const FShaderParameterMap& ParameterMap) override
+	void Bind(const FShaderParameterMap& ParameterMap)
 	{
 		MeshOrigin.Bind(ParameterMap, TEXT("MeshOrigin"));
 		MeshExtension.Bind(ParameterMap, TEXT("MeshExtension"));
@@ -36,20 +37,7 @@ public:
 		MotionBlurPositionScale.Bind(ParameterMap, TEXT("MotionBlurPositionScale"));
 	}
 
-	/**
-	* Serialize shader params to an archive
-	* @param	Ar - archive to serialize to
-	*/
-	virtual void Serialize(FArchive& Ar) override
-	{
-		Ar << MeshOrigin;
-		Ar << MeshExtension;
-		Ar << MotionBlurDataOrigin;
-		Ar << MotionBlurDataExtension;
-		Ar << MotionBlurPositionScale;
-	}
-
-	virtual void GetElementShaderBindings(
+	void GetElementShaderBindings(
 		const class FSceneInterface* Scene,
 		const FSceneView* View,
 		const class FMeshMaterialShader* Shader,
@@ -58,7 +46,7 @@ public:
 		const FVertexFactory* GenericVertexFactory,
 		const FMeshBatchElement& BatchElement,
 		class FMeshDrawSingleShaderBindings& ShaderBindings,
-		FVertexInputStreamArray& VertexStreams) const override
+		FVertexInputStreamArray& VertexStreams) const
 	{
 		// Ensure the vertex factory matches this parameter object and cast relevant objects
 		check(GenericVertexFactory->GetType() == &FGeometryCacheVertexVertexFactory::StaticType);
@@ -83,30 +71,20 @@ public:
 		ShaderBindings.Add(Shader->GetUniformBufferParameter<FGeometryCacheManualVertexFetchUniformBufferParameters>(), BatchData->ManualVertexFetchUniformBuffer);
 	}
 
-	virtual uint32 GetSize() const override { return sizeof(*this); }
-
 private:
-	FShaderParameter MeshOrigin;
-	FShaderParameter MeshExtension;
-	FShaderParameter MotionBlurDataOrigin;
-	FShaderParameter MotionBlurDataExtension;
-	FShaderParameter MotionBlurPositionScale;
+	LAYOUT_FIELD(FShaderParameter, MeshOrigin);
+	LAYOUT_FIELD(FShaderParameter, MeshExtension);
+	LAYOUT_FIELD(FShaderParameter, MotionBlurDataOrigin);
+	LAYOUT_FIELD(FShaderParameter, MotionBlurDataExtension);
+	LAYOUT_FIELD(FShaderParameter, MotionBlurPositionScale);
 };
 
 /*-----------------------------------------------------------------------------
 FGPUSkinPassthroughVertexFactory
 -----------------------------------------------------------------------------*/
-void FGeometryCacheVertexVertexFactory::ModifyCompilationEnvironment(const FVertexFactoryType* Type, EShaderPlatform Platform, const FMaterial* Material, FShaderCompilerEnvironment& OutEnvironment)
+void FGeometryCacheVertexVertexFactory::ModifyCompilationEnvironment(const FVertexFactoryShaderPermutationParameters& Parameters, FShaderCompilerEnvironment& OutEnvironment)
 {
-	Super::ModifyCompilationEnvironment(Type, Platform, Material, OutEnvironment);
-}
-
-bool FGeometryCacheVertexVertexFactory::ShouldCache(EShaderPlatform Platform, const class FMaterial* Material, const FShaderType* ShaderType)
-{
-	//FIXME: This can probably be beter like this?!?
-	//return Material->IsUsedWithGeometryCache() || Material->IsSpecialEngineMaterial();
-
-	return true;
+	Super::ModifyCompilationEnvironment(Parameters, OutEnvironment);
 }
 
 void FGeometryCacheVertexVertexFactory::SetData(const FDataType& InData)
@@ -347,26 +325,16 @@ void FGeometryCacheVertexVertexFactory::CreateManualVertexFetchUniformBuffer(
 	OutUserData.ManualVertexFetchUniformBuffer = FGeometryCacheManualVertexFetchUniformBufferParametersRef::CreateUniformBufferImmediate(ManualVertexFetchParameters, UniformBuffer_SingleFrame);
 }
 
-FVertexFactoryShaderParameters* FGeometryCacheVertexVertexFactory::ConstructShaderParameters(EShaderFrequency ShaderFrequency)
-{
-	switch (ShaderFrequency)
-	{
-		case SF_Vertex:
-#if RHI_RAYTRACING
-		case SF_RayHitGroup:
-#endif
-			return new FGeometryCacheVertexFactoryShaderParameters();
-		default:
-			return nullptr;
-	}
-}
-
-bool FGeometryCacheVertexVertexFactory::ShouldCompilePermutation(EShaderPlatform Platform, const class FMaterial* Material, const class FShaderType* ShaderType)
+bool FGeometryCacheVertexVertexFactory::ShouldCompilePermutation(const FVertexFactoryShaderPermutationParameters& Parameters)
 {
 	// Should this be platform or mesh type based? Returning true should work in all cases, but maybe too expensive? 
 	// return IsFeatureLevelSupported(Platform, ERHIFeatureLevel::SM5) && !IsConsolePlatform(Platform);
 	// TODO currently GeomCache supports only 4 UVs which could cause compilation errors when trying to compile shaders which use > 4
-	return Material->IsUsedWithGeometryCache() || Material->IsSpecialEngineMaterial();
+	return Parameters.MaterialParameters.bIsUsedWithGeometryCache || Parameters.MaterialParameters.bIsSpecialEngineMaterial;
 }
 
+IMPLEMENT_VERTEX_FACTORY_PARAMETER_TYPE(FGeometryCacheVertexVertexFactory, SF_Vertex, FGeometryCacheVertexFactoryShaderParameters);
+#if RHI_RAYTRACING
+IMPLEMENT_VERTEX_FACTORY_PARAMETER_TYPE(FGeometryCacheVertexVertexFactory, SF_RayHitGroup, FGeometryCacheVertexFactoryShaderParameters);
+#endif
 IMPLEMENT_VERTEX_FACTORY_TYPE(FGeometryCacheVertexVertexFactory, "/Engine/Private/GeometryCacheVertexFactory.ush", true, false, true, false, true);
diff --git a/Engine/Source/Runtime/Engine/Private/GeometryCollectionRendering.cpp b/Engine/Source/Runtime/Engine/Private/GeometryCollectionRendering.cpp
--- a/Engine/Source/Runtime/Engine/Private/GeometryCollectionRendering.cpp
+++ b/Engine/Source/Runtime/Engine/Private/GeometryCollectionRendering.cpp
@@ -2,9 +2,44 @@
 
 #include "GeometryCollectionRendering.h"
 #include "MeshDrawShaderBindings.h"
+#include "MeshMaterialShader.h"
 
+IMPLEMENT_TYPE_LAYOUT(FGeometryCollectionVertexFactoryShaderParameters);
+IMPLEMENT_VERTEX_FACTORY_PARAMETER_TYPE(FGeometryCollectionVertexFactory, SF_Vertex, FGeometryCollectionVertexFactoryShaderParameters);
 IMPLEMENT_VERTEX_FACTORY_TYPE(FGeometryCollectionVertexFactory, "/Engine/Private/LocalVertexFactory.ush", true, true, true, true, true);
 
+bool FGeometryCollectionVertexFactory::ShouldCompilePermutation(const FVertexFactoryShaderPermutationParameters& Parameters)
+{
+	return (Parameters.MaterialParameters.bIsUsedWithGeometryCollections || Parameters.MaterialParameters.bIsSpecialEngineMaterial)
+		&& FLocalVertexFactory::ShouldCompilePermutation(Parameters);
+}
+
+//
+// Modify compile environment to enable instancing
+// @param OutEnvironment - shader compile environment to modify
+//
+void FGeometryCollectionVertexFactory::ModifyCompilationEnvironment(const FVertexFactoryShaderPermutationParameters& Parameters, FShaderCompilerEnvironment& OutEnvironment)
+{
+	const bool ContainsManualVertexFetch = OutEnvironment.GetDefinitions().Contains("MANUAL_VERTEX_FETCH");
+	if (!ContainsManualVertexFetch && RHISupportsManualVertexFetch(Parameters.Platform))
+	{
+		OutEnvironment.SetDefine(TEXT("MANUAL_VERTEX_FETCH"), TEXT("1"));
+	}
+
+	if (RHISupportsManualVertexFetch(Parameters.Platform))
+	{
+		OutEnvironment.SetDefine(TEXT("USE_INSTANCING"), TEXT("1"));
+
+		OutEnvironment.SetDefine(TEXT("USE_INSTANCING_BONEMAP"), TEXT("1"));
+		OutEnvironment.SetDefine(TEXT("USE_DITHERED_LOD_TRANSITION_FOR_INSTANCED"), TEXT("0"));
+	}
+
+	// Geometry collections use a custom hit proxy per bone
+	OutEnvironment.SetDefine(TEXT("USE_PER_VERTEX_HITPROXY_ID"), 1);
+
+	FLocalVertexFactory::ModifyCompilationEnvironment(Parameters, OutEnvironment);
+}
+
 void FGeometryCollectionVertexFactoryShaderParameters::GetElementShaderBindings(
 	const FSceneInterface* Scene,
 	const FSceneView* View,
diff --git a/Engine/Source/Runtime/Engine/Private/InstancedStaticMesh.cpp b/Engine/Source/Runtime/Engine/Private/InstancedStaticMesh.cpp
--- a/Engine/Source/Runtime/Engine/Private/InstancedStaticMesh.cpp
+++ b/Engine/Source/Runtime/Engine/Private/InstancedStaticMesh.cpp
@@ -30,6 +30,7 @@
 #include "DeviceProfiles/DeviceProfileManager.h"
 #endif // WITH_EDITOR
 #include "MeshMaterialShader.h"
+#include "ProfilingDebugging/LoadTimeTracker.h"
 
 #if RHI_RAYTRACING
 #include "RayTracingInstance.h"
@@ -45,6 +46,9 @@
 #include "UObject/RenderingObjectVersion.h"
 
 
+IMPLEMENT_TYPE_LAYOUT(FInstancedStaticMeshVertexFactoryShaderParameters);
+
+
 const int32 InstancedStaticMeshMaxTexCoord = 8;
 
 IMPLEMENT_HIT_PROXY(HInstancedStaticMeshInstance, HHitProxy);
@@ -377,6 +381,8 @@ void FStaticMeshInstanceBuffer::InitRHI()
 	if (InstanceData->GetNumInstances() > 0)
 	{
 		QUICK_SCOPE_CYCLE_COUNTER(STAT_FStaticMeshInstanceBuffer_InitRHI);
+		SCOPED_LOADTIMER(FStaticMeshInstanceBuffer_InitRHI);
+
 		LLM_SCOPE(ELLMTag::InstancedMesh);
 		auto AccessFlags = BUF_Static;
 		CreateVertexBuffer(InstanceData->GetOriginResourceArray(), AccessFlags | BUF_ShaderResource, 16, PF_A32B32G32R32F, InstanceOriginBuffer.VertexBufferRHI, InstanceOriginSRV);
@@ -592,10 +598,32 @@ void FStaticMeshInstanceData::Serialize(FArchive& Ar)
 /**
  * Should we cache the material's shadertype on this platform with this vertex factory? 
  */
-bool FInstancedStaticMeshVertexFactory::ShouldCompilePermutation(EShaderPlatform Platform, const class FMaterial* Material, const class FShaderType* ShaderType)
+bool FInstancedStaticMeshVertexFactory::ShouldCompilePermutation(const FVertexFactoryShaderPermutationParameters& Parameters)
 {
-	return (Material->IsUsedWithInstancedStaticMeshes() || Material->IsSpecialEngineMaterial()) 
-			&& FLocalVertexFactory::ShouldCompilePermutation(Platform, Material, ShaderType);
+	return (Parameters.MaterialParameters.bIsUsedWithInstancedStaticMeshes || Parameters.MaterialParameters.bIsSpecialEngineMaterial)
+			&& FLocalVertexFactory::ShouldCompilePermutation(Parameters);
+}
+
+void FInstancedStaticMeshVertexFactory::ModifyCompilationEnvironment(const FVertexFactoryShaderPermutationParameters& Parameters, FShaderCompilerEnvironment& OutEnvironment)
+{
+	const bool ContainsManualVertexFetch = OutEnvironment.GetDefinitions().Contains("MANUAL_VERTEX_FETCH");
+	if (!ContainsManualVertexFetch && RHISupportsManualVertexFetch(Parameters.Platform))
+	{
+		OutEnvironment.SetDefine(TEXT("MANUAL_VERTEX_FETCH"), TEXT("1"));
+	}
+
+	OutEnvironment.SetDefine(TEXT("USE_INSTANCING"), TEXT("1"));
+	if (IsFeatureLevelSupported(Parameters.Platform, ERHIFeatureLevel::SM5))
+	{
+		OutEnvironment.SetDefine(TEXT("USE_DITHERED_LOD_TRANSITION_FOR_INSTANCED"), ALLOW_DITHERED_LOD_FOR_INSTANCED_STATIC_MESHES);
+	}
+	else
+	{
+		// On mobile dithered LOD transition has to be explicitly enabled in material and project settings
+		OutEnvironment.SetDefine(TEXT("USE_DITHERED_LOD_TRANSITION_FOR_INSTANCED"), Parameters.MaterialParameters.bIsDitheredLODTransition && ALLOW_DITHERED_LOD_FOR_INSTANCED_STATIC_MESHES);
+	}
+
+	FLocalVertexFactory::ModifyCompilationEnvironment(Parameters, OutEnvironment);
 }
 
 
@@ -617,6 +645,8 @@ void FInstancedStaticMeshVertexFactory::Copy(const FInstancedStaticMeshVertexFac
 
 void FInstancedStaticMeshVertexFactory::InitRHI()
 {
+	SCOPED_LOADTIMER(FInstancedStaticMeshVertexFactory_InitRHI);
+
 	check(HasValidFeatureLevel());
 
 #if !ALLOW_DITHERED_LOD_FOR_INSTANCED_STATIC_MESHES // position(and normal) only shaders cannot work with dithered LOD
@@ -747,11 +777,7 @@ void FInstancedStaticMeshVertexFactory::InitRHI()
 	}
 }
 
-
-FVertexFactoryShaderParameters* FInstancedStaticMeshVertexFactory::ConstructShaderParameters(EShaderFrequency ShaderFrequency)
-{
-	return ShaderFrequency == SF_Vertex ? new FInstancedStaticMeshVertexFactoryShaderParameters() : NULL;
-}
+IMPLEMENT_VERTEX_FACTORY_PARAMETER_TYPE(FInstancedStaticMeshVertexFactory, SF_Vertex, FInstancedStaticMeshVertexFactoryShaderParameters);
 
 IMPLEMENT_VERTEX_FACTORY_TYPE_EX(FInstancedStaticMeshVertexFactory,"/Engine/Private/LocalVertexFactory.ush",true,true,true,true,true,true,false);
 
@@ -3196,14 +3222,14 @@ void FInstancedStaticMeshVertexFactoryShaderParameters::GetElementShaderBindings
 				FinalCull *= MaxDrawDistanceScale;
 
 				InstancingViewZCompare.Z = FinalCull;
-				if (BatchElement.InstancedLODIndex < InstancingUserData->MeshRenderData->LODResources.Num() - 1)
+				if (int(BatchElement.InstancedLODIndex) < InstancingUserData->MeshRenderData->LODResources.Num() - 1)
 				{
 					float NextCut = ComputeBoundsDrawDistance(InstancingUserData->MeshRenderData->ScreenSize[BatchElement.InstancedLODIndex + 1].GetValueForFeatureLevel(FeatureLevel), SphereRadius, View->ViewMatrices.GetProjectionMatrix()) * LODScale;
 					InstancingViewZCompare.Z = FMath::Min(NextCut, FinalCull);
 				}
 
 				InstancingViewZCompare.X = MIN_flt;
-				if (BatchElement.InstancedLODIndex > FirstLOD)
+				if (int(BatchElement.InstancedLODIndex) > FirstLOD)
 				{
 					float CurCut = ComputeBoundsDrawDistance(InstancingUserData->MeshRenderData->ScreenSize[BatchElement.InstancedLODIndex].GetValueForFeatureLevel(FeatureLevel), SphereRadius, View->ViewMatrices.GetProjectionMatrix()) * LODScale;
 					if (CurCut < FinalCull)
diff --git a/Engine/Source/Runtime/Engine/Private/InstancedStaticMesh.h b/Engine/Source/Runtime/Engine/Private/InstancedStaticMesh.h
--- a/Engine/Source/Runtime/Engine/Private/InstancedStaticMesh.h
+++ b/Engine/Source/Runtime/Engine/Private/InstancedStaticMesh.h
@@ -212,33 +212,13 @@ public:
 	/**
 	 * Should we cache the material's shadertype on this platform with this vertex factory? 
 	 */
-	static bool ShouldCompilePermutation(EShaderPlatform Platform, const class FMaterial* Material, const class FShaderType* ShaderType);
+	static bool ShouldCompilePermutation(const FVertexFactoryShaderPermutationParameters& Parameters);
 
 	/**
 	 * Modify compile environment to enable instancing
 	 * @param OutEnvironment - shader compile environment to modify
 	 */
-	static void ModifyCompilationEnvironment(const FVertexFactoryType* Type, EShaderPlatform Platform, const FMaterial* Material, FShaderCompilerEnvironment& OutEnvironment)
-	{
-		const bool ContainsManualVertexFetch = OutEnvironment.GetDefinitions().Contains("MANUAL_VERTEX_FETCH");
-		if (!ContainsManualVertexFetch && RHISupportsManualVertexFetch(Platform))
-		{
-			OutEnvironment.SetDefine(TEXT("MANUAL_VERTEX_FETCH"), TEXT("1"));
-		}
-
-		OutEnvironment.SetDefine(TEXT("USE_INSTANCING"),TEXT("1"));
-		if (IsFeatureLevelSupported(Platform, ERHIFeatureLevel::SM5))
-		{
-			OutEnvironment.SetDefine(TEXT("USE_DITHERED_LOD_TRANSITION_FOR_INSTANCED"), ALLOW_DITHERED_LOD_FOR_INSTANCED_STATIC_MESHES);
-		}
-		else
-		{
-			// On mobile dithered LOD transition has to be explicitly enabled in material and project settings
-			OutEnvironment.SetDefine(TEXT("USE_DITHERED_LOD_TRANSITION_FOR_INSTANCED"), Material->IsDitheredLODTransition() && ALLOW_DITHERED_LOD_FOR_INSTANCED_STATIC_MESHES);
-		}
-		
-		FLocalVertexFactory::ModifyCompilationEnvironment(Type, Platform, Material, OutEnvironment);
-	}
+	static void ModifyCompilationEnvironment(const FVertexFactoryShaderPermutationParameters& Parameters, FShaderCompilerEnvironment& OutEnvironment);
 
 	/**
 	 * An implementation of the interface used by TSynchronizedResource to update the resource with new data from the game thread.
@@ -313,7 +293,9 @@ private:
 
 class FInstancedStaticMeshVertexFactoryShaderParameters : public FLocalVertexFactoryShaderParametersBase
 {
-	virtual void Bind(const FShaderParameterMap& ParameterMap) override
+	DECLARE_TYPE_LAYOUT(FInstancedStaticMeshVertexFactoryShaderParameters, NonVirtual);
+public:
+	void Bind(const FShaderParameterMap& ParameterMap)
 	{
 		FLocalVertexFactoryShaderParametersBase::Bind(ParameterMap);
 
@@ -330,7 +312,7 @@ class FInstancedStaticMeshVertexFactoryShaderParameters : public FLocalVertexFac
 		InstanceOffset.Bind(ParameterMap, TEXT("InstanceOffset"));
 	}
 
-	virtual void GetElementShaderBindings(
+	void GetElementShaderBindings(
 		const class FSceneInterface* Scene,
 		const FSceneView* View,
 		const FMeshMaterialShader* Shader,
@@ -340,49 +322,22 @@ class FInstancedStaticMeshVertexFactoryShaderParameters : public FLocalVertexFac
 		const FMeshBatchElement& BatchElement,
 		FMeshDrawSingleShaderBindings& ShaderBindings,
 		FVertexInputStreamArray& VertexStreams
-		) const override;
-
-	void Serialize(FArchive& Ar) override
-	{
-		Ar.UsingCustomVersion(FRenderingObjectVersion::GUID);
-
-		FLocalVertexFactoryShaderParametersBase::Serialize(Ar);
-		Ar << InstancingFadeOutParamsParameter;
-		Ar << InstancingViewZCompareZeroParameter;
-		Ar << InstancingViewZCompareOneParameter;
-		Ar << InstancingViewZConstantParameter;
-		Ar << InstancingOffsetParameter;
-		Ar << InstancingWorldViewOriginZeroParameter;
-		Ar << InstancingWorldViewOriginOneParameter;
-		
-		if (Ar.IsLoading() && Ar.CustomVer(FRenderingObjectVersion::GUID) < FRenderingObjectVersion::RemovedEmulatedInstancing)
-		{
-			// Older version need to also load legacy emulated instancing shader parameters
-			FShaderParameter Dummy;
-			Ar << Dummy << Dummy << Dummy;
-		}
-
-		Ar << VertexFetch_InstanceOriginBufferParameter;
-		Ar << VertexFetch_InstanceTransformBufferParameter;
-		Ar << VertexFetch_InstanceLightmapBufferParameter;
-		Ar << InstanceOffset;
-	}
-
-	virtual uint32 GetSize() const override { return sizeof(*this); }
+		) const;
 
 private:
-	FShaderParameter InstancingFadeOutParamsParameter;
-	FShaderParameter InstancingViewZCompareZeroParameter;
-	FShaderParameter InstancingViewZCompareOneParameter;
-	FShaderParameter InstancingViewZConstantParameter;
-	FShaderParameter InstancingOffsetParameter;
-	FShaderParameter InstancingWorldViewOriginZeroParameter;
-	FShaderParameter InstancingWorldViewOriginOneParameter;
-
-	FShaderResourceParameter VertexFetch_InstanceOriginBufferParameter;
-	FShaderResourceParameter VertexFetch_InstanceTransformBufferParameter;
-	FShaderResourceParameter VertexFetch_InstanceLightmapBufferParameter;
-	FShaderParameter InstanceOffset;
+	
+	LAYOUT_FIELD(FShaderParameter, InstancingFadeOutParamsParameter)
+	LAYOUT_FIELD(FShaderParameter, InstancingViewZCompareZeroParameter)
+	LAYOUT_FIELD(FShaderParameter, InstancingViewZCompareOneParameter)
+	LAYOUT_FIELD(FShaderParameter, InstancingViewZConstantParameter)
+	LAYOUT_FIELD(FShaderParameter, InstancingOffsetParameter);
+	LAYOUT_FIELD(FShaderParameter, InstancingWorldViewOriginZeroParameter)
+	LAYOUT_FIELD(FShaderParameter, InstancingWorldViewOriginOneParameter)
+
+	LAYOUT_FIELD(FShaderResourceParameter, VertexFetch_InstanceOriginBufferParameter)
+	LAYOUT_FIELD(FShaderResourceParameter, VertexFetch_InstanceTransformBufferParameter)
+	LAYOUT_FIELD(FShaderResourceParameter, VertexFetch_InstanceLightmapBufferParameter)
+	LAYOUT_FIELD(FShaderParameter, InstanceOffset)
 };
 
 struct FInstanceUpdateCmdBuffer;
@@ -466,7 +421,7 @@ public:
 	TIndirectArray<FInstancedStaticMeshVertexFactory> VertexFactories;
 
 	/** LOD render data from the static mesh. */
-	TIndirectArray<FStaticMeshLODResources>& LODModels;
+	FStaticMeshLODResourcesArray& LODModels;
 
 	/** Feature level used when creating instance data */
 	ERHIFeatureLevel::Type FeatureLevel;
diff --git a/Engine/Source/Runtime/Engine/Private/LODActor.cpp b/Engine/Source/Runtime/Engine/Private/LODActor.cpp
--- a/Engine/Source/Runtime/Engine/Private/LODActor.cpp
+++ b/Engine/Source/Runtime/Engine/Private/LODActor.cpp
@@ -1096,7 +1096,7 @@ UMaterialInterface* ALODActor::GetImposterMaterial(const UStaticMeshComponent* I
 		const int32 LODIndex = StaticMesh->GetNumLODs() - 1;
 
 		// Retrieve the sections, we're expect 1 for imposter meshes
-		const TArray<FStaticMeshSection>& Sections = StaticMesh->RenderData->LODResources[LODIndex].Sections;
+		const FStaticMeshLODResources::FStaticMeshSectionArray& Sections = StaticMesh->RenderData->LODResources[LODIndex].Sections;
 		if (Sections.Num() == 1)
 		{
 			// Retrieve material for this section
diff --git a/Engine/Source/Runtime/Engine/Private/Level.cpp b/Engine/Source/Runtime/Engine/Private/Level.cpp
--- a/Engine/Source/Runtime/Engine/Private/Level.cpp
+++ b/Engine/Source/Runtime/Engine/Private/Level.cpp
@@ -45,6 +45,7 @@ Level.cpp: Level-related functions
 #include "PhysicsEngine/BodySetup.h"
 #include "EngineGlobals.h"
 #include "Engine/LevelBounds.h"
+#include "Async/ParallelFor.h"
 #include "UnrealEngine.h"
 #if WITH_EDITOR
 #include "Kismet2/KismetEditorUtilities.h"
@@ -828,10 +829,10 @@ struct FModelComponentKey
 	}
 };
 
-void ULevel::UpdateLevelComponents(bool bRerunConstructionScripts)
+void ULevel::UpdateLevelComponents(bool bRerunConstructionScripts, FRegisterComponentContext* Context)
 {
 	// Update all components in one swoop.
-	IncrementalUpdateComponents( 0, bRerunConstructionScripts );
+	IncrementalUpdateComponents( 0, bRerunConstructionScripts, Context);
 }
 
 /**
@@ -920,7 +921,8 @@ static void SortActorsHierarchy(TArray<AActor*>& Actors, UObject* Level)
 	}
 }
 
-void ULevel::IncrementalUpdateComponents(int32 NumComponentsToUpdate, bool bRerunConstructionScripts)
+DECLARE_CYCLE_STAT(TEXT("Deferred Init Bodies"), STAT_DeferredUpdateBodies, STATGROUP_Physics);
+void ULevel::IncrementalUpdateComponents(int32 NumComponentsToUpdate, bool bRerunConstructionScripts, FRegisterComponentContext* Context)
 {
 	// A value of 0 means that we want to update all components.
 	if (NumComponentsToUpdate != 0)
@@ -937,6 +939,7 @@ void ULevel::IncrementalUpdateComponents(int32 NumComponentsToUpdate, bool bReru
 		SortActorsHierarchy(Actors, this);
 	}
 
+	int32 PreviousIndex = CurrentActorIndexForUpdateComponents;
 	// Find next valid actor to process components registration
 	while (CurrentActorIndexForUpdateComponents < Actors.Num())
 	{
@@ -952,7 +955,7 @@ void ULevel::IncrementalUpdateComponents(int32 NumComponentsToUpdate, bool bReru
 				Actor->PreRegisterAllComponents();
 				bHasCurrentActorCalledPreRegister = true;
 			}
-			bAllComponentsRegistered = Actor->IncrementalRegisterComponents(NumComponentsToUpdate);
+			bAllComponentsRegistered = Actor->IncrementalRegisterComponents(NumComponentsToUpdate, Context);
 		}
 
 		if (bAllComponentsRegistered)
@@ -1011,6 +1014,32 @@ void ULevel::IncrementalUpdateComponents(int32 NumComponentsToUpdate, bool bReru
 		// The editor is never allowed to incrementally updated components.  Make sure to pass in a value of zero for NumActorsToUpdate.
 		check(OwningWorld->IsGameWorld());
 	}
+
+	{
+		SCOPE_CYCLE_COUNTER(STAT_DeferredUpdateBodies);
+		// Init Bodies
+		FGCScopeGuard ScopeLock;
+		int32 NumToInit = Actors.Num() - PreviousIndex;
+		TSubclassOf<UActorComponent> PrimitiveStaticClass = UPrimitiveComponent::StaticClass();
+		ParallelFor(NumToInit, [&](int32 Idx) {
+			int32 ActorIdx = Idx + PreviousIndex;
+			AActor* Actor = Actors[ActorIdx];
+			if (Actor && !Actor->IsPendingKill())
+			{
+				TArray<UActorComponent*> PrimtiveComponents;
+				Actor->GetComponents(PrimitiveStaticClass, PrimtiveComponents);
+				for (UActorComponent* PrimtiveComponent : PrimtiveComponents)
+				{
+					FBodyInstance* BodyInstance = Cast<UPrimitiveComponent>(PrimtiveComponent)->GetBodyInstance();
+					if (BodyInstance)
+					{
+						BodyInstance->InitAllBodies(GetWorld()->GetPhysicsScene());
+					}
+				}
+			}
+			// @TODO DO NOT CHECKIN
+		}, EParallelForFlags::ForceSingleThread);
+	}
 }
 
 
diff --git a/Engine/Source/Runtime/Engine/Private/LightMap.cpp b/Engine/Source/Runtime/Engine/Private/LightMap.cpp
--- a/Engine/Source/Runtime/Engine/Private/LightMap.cpp
+++ b/Engine/Source/Runtime/Engine/Private/LightMap.cpp
@@ -22,6 +22,7 @@
 #include "VT/VirtualTexture.h"
 #include "EngineModule.h"
 #include "Misc/PackageName.h"
+#include "ProfilingDebugging/LoadTimeTracker.h"
 
 #define VISUALIZE_PACKING 0
 
@@ -3294,15 +3295,15 @@ void FLightmapResourceCluster::UpdateUniformBuffer(ERHIFeatureLevel::Type InFeat
 	ENQUEUE_RENDER_COMMAND(SetFeatureLevel)(
 		[Cluster, InFeatureLevel](FRHICommandList& RHICmdList)
 	{
-		Cluster->FeatureLevel = InFeatureLevel;
+		Cluster->SetFeatureLevel(InFeatureLevel);
 		Cluster->UpdateUniformBuffer_RenderThread();
 	});
 }
 
 bool FLightmapResourceCluster::GetUseVirtualTexturing() const
 {
-	const bool bAllowHighQualityLightMaps = AllowHighQualityLightmaps(FeatureLevel);
-	return bAllowHighQualityLightMaps && (CVarVirtualTexturedLightMaps.GetValueOnRenderThread() != 0) && UseVirtualTexturing(FeatureLevel);
+	const bool bAllowHighQualityLightMaps = AllowHighQualityLightmaps(GetFeatureLevel());
+	return bAllowHighQualityLightMaps && (CVarVirtualTexturedLightMaps.GetValueOnRenderThread() != 0) && UseVirtualTexturing(GetFeatureLevel());
 }
 
 void FLightmapResourceCluster::UpdateUniformBuffer_RenderThread()
@@ -3310,7 +3311,7 @@ void FLightmapResourceCluster::UpdateUniformBuffer_RenderThread()
 	check(IsInRenderingThread());
 
 	FLightmapResourceClusterShaderParameters Parameters;
-	GetLightmapClusterResourceParameters(FeatureLevel, Input, GetUseVirtualTexturing() ? AcquireAllocatedVT() : nullptr, Parameters);
+	GetLightmapClusterResourceParameters(GetFeatureLevel(), Input, GetUseVirtualTexturing() ? AcquireAllocatedVT() : nullptr, Parameters);
 
 	RHIUpdateUniformBuffer(UniformBuffer, &Parameters);
 }
@@ -3388,14 +3389,16 @@ void FLightmapResourceCluster::ReleaseAllocatedVT()
 
 void FLightmapResourceCluster::InitRHI()
 {
+	SCOPED_LOADTIMER(FLightmapResourceCluster_InitRHI);
+
 	FLightmapResourceClusterShaderParameters Parameters;
 
 	// Lightmap resources are normally created before the feature level is known, so we'll use defaults and rely on a subsequent call to UpdateUniformBuffer()
 	// to set the correct level and update things accordingly. However, when we're coming from FRenderResource::ChangeFeatureLevel(), the feature level
 	// has already been set, so we can go ahead and use the correct level and input from the start (UpdateUniformBuffer() is not being called in that case).
-	if (FeatureLevel < ERHIFeatureLevel::Num)
+	if (GetFeatureLevel() < ERHIFeatureLevel::Num)
 	{
-		GetLightmapClusterResourceParameters(FeatureLevel, Input, GetUseVirtualTexturing() ? AcquireAllocatedVT() : nullptr, Parameters);
+		GetLightmapClusterResourceParameters(GetFeatureLevel(), Input, GetUseVirtualTexturing() ? AcquireAllocatedVT() : nullptr, Parameters);
 	}
 	else
 	{
diff --git a/Engine/Source/Runtime/Engine/Private/LocalVertexFactory.cpp b/Engine/Source/Runtime/Engine/Private/LocalVertexFactory.cpp
--- a/Engine/Source/Runtime/Engine/Private/LocalVertexFactory.cpp
+++ b/Engine/Source/Runtime/Engine/Private/LocalVertexFactory.cpp
@@ -11,6 +11,10 @@
 #include "ShaderParameterUtils.h"
 #include "Rendering/ColorVertexBuffer.h"
 #include "MeshMaterialShader.h"
+#include "ProfilingDebugging/LoadTimeTracker.h"
+
+IMPLEMENT_TYPE_LAYOUT(FLocalVertexFactoryShaderParametersBase);
+IMPLEMENT_TYPE_LAYOUT(FLocalVertexFactoryShaderParameters);
 
 class FSpeedTreeWindNullUniformBuffer : public TUniformBuffer<FSpeedTreeUniformParameters>
 {
@@ -36,12 +40,6 @@ void FLocalVertexFactoryShaderParametersBase::Bind(const FShaderParameterMap& Pa
 	bAnySpeedTreeParamIsBound = LODParameter.IsBound() || ParameterMap.ContainsParameterAllocation(TEXT("SpeedTreeData"));
 }
 
-void FLocalVertexFactoryShaderParametersBase::Serialize(FArchive& Ar)
-{
-	Ar << bAnySpeedTreeParamIsBound;
-	Ar << LODParameter;
-}
-
 IMPLEMENT_GLOBAL_SHADER_PARAMETER_STRUCT(FLocalVertexFactoryUniformShaderParameters, "LocalVF");
 
 TUniformBufferRef<FLocalVertexFactoryUniformShaderParameters> CreateLocalVFUniformBuffer(
@@ -106,7 +104,7 @@ void FLocalVertexFactoryShaderParametersBase::GetElementShaderBindingsBase(
 	) const
 {
 	const auto* LocalVertexFactory = static_cast<const FLocalVertexFactory*>(VertexFactory);
-
+	
 	if (LocalVertexFactory->SupportsManualVertexFetch(FeatureLevel) || UseGPUScene(GMaxRHIShaderPlatform, FeatureLevel))
 	{
 		if (!VertexFactoryUniformBuffer)
@@ -181,10 +179,10 @@ void FLocalVertexFactoryShaderParameters::GetElementShaderBindings(
 /**
  * Should we cache the material's shadertype on this platform with this vertex factory? 
  */
-bool FLocalVertexFactory::ShouldCompilePermutation(EShaderPlatform Platform, const class FMaterial* Material, const class FShaderType* ShaderType)
+bool FLocalVertexFactory::ShouldCompilePermutation(const FVertexFactoryShaderPermutationParameters& Parameters)
 {
 	// Only compile this permutation inside the editor - it's not applicable in games, but occasionally the editor needs it.
-	if (Material->GetMaterialDomain() == MD_UI)
+	if (Parameters.MaterialParameters.MaterialDomain == MD_UI)
 	{
 		return !!WITH_EDITOR;
 	}
@@ -192,18 +190,18 @@ bool FLocalVertexFactory::ShouldCompilePermutation(EShaderPlatform Platform, con
 	return true; 
 }
 
-void FLocalVertexFactory::ModifyCompilationEnvironment(const FVertexFactoryType* Type, EShaderPlatform Platform, const FMaterial* Material, FShaderCompilerEnvironment& OutEnvironment)
+void FLocalVertexFactory::ModifyCompilationEnvironment(const FVertexFactoryShaderPermutationParameters& Parameters, FShaderCompilerEnvironment& OutEnvironment)
 {
 	OutEnvironment.SetDefine(TEXT("VF_SUPPORTS_SPEEDTREE_WIND"),TEXT("1"));
 
 	const bool ContainsManualVertexFetch = OutEnvironment.GetDefinitions().Contains("MANUAL_VERTEX_FETCH");
-	if (!ContainsManualVertexFetch && RHISupportsManualVertexFetch(Platform))
+	if (!ContainsManualVertexFetch && RHISupportsManualVertexFetch(Parameters.Platform))
 	{
 		OutEnvironment.SetDefine(TEXT("MANUAL_VERTEX_FETCH"), TEXT("1"));
 	}
 
-	OutEnvironment.SetDefine(TEXT("VF_SUPPORTS_PRIMITIVE_SCENE_DATA"), Type->SupportsPrimitiveIdStream() && UseGPUScene(Platform, GetMaxSupportedFeatureLevel(Platform)));
-	OutEnvironment.SetDefine(TEXT("VF_GPU_SCENE_TEXTURE"), Type->SupportsPrimitiveIdStream() && UseGPUScene(Platform, GetMaxSupportedFeatureLevel(Platform)) && GPUSceneUseTexture2D(Platform));
+	OutEnvironment.SetDefine(TEXT("VF_SUPPORTS_PRIMITIVE_SCENE_DATA"), Parameters.VertexFactoryType->SupportsPrimitiveIdStream() && UseGPUScene(Parameters.Platform, GetMaxSupportedFeatureLevel(Parameters.Platform)));
+	OutEnvironment.SetDefine(TEXT("VF_GPU_SCENE_TEXTURE"), Parameters.VertexFactoryType->SupportsPrimitiveIdStream() && UseGPUScene(Parameters.Platform, GetMaxSupportedFeatureLevel(Parameters.Platform)) && GPUSceneUseTexture2D(Parameters.Platform));
 }
 
 void FLocalVertexFactory::ValidateCompiledResult(const FVertexFactoryType* Type, EShaderPlatform Platform, const FShaderParameterMap& ParameterMap, TArray<FString>& OutErrors)
@@ -257,6 +255,8 @@ void FLocalVertexFactory::Copy(const FLocalVertexFactory& Other)
 
 void FLocalVertexFactory::InitRHI()
 {
+	SCOPED_LOADTIMER(FLocalVertexFactory_InitRHI);
+
 	// We create different streams based on feature level
 	check(HasValidFeatureLevel());
 
@@ -290,6 +290,7 @@ void FLocalVertexFactory::InitRHI()
 
 			InitDeclaration(StreamElements, InputStreamType);
 		};
+
 		AddDeclaration(EVertexInputStreamType::PositionOnly, false);
 		AddDeclaration(EVertexInputStreamType::PositionAndNormalOnly, true);
 	}
@@ -379,31 +380,17 @@ void FLocalVertexFactory::InitRHI()
 	const int32 DefaultBaseVertexIndex = 0;
 	if (RHISupportsManualVertexFetch(GMaxRHIShaderPlatform) || bCanUseGPUScene)
 	{
+		SCOPED_LOADTIMER(FLocalVertexFactory_InitRHI_CreateLocalVFUniformBuffer);
 		UniformBuffer = CreateLocalVFUniformBuffer(this, Data.LODLightmapDataIndex, nullptr, DefaultBaseVertexIndex);
 	}
 
 	check(IsValidRef(GetDeclaration()));
 }
 
-FVertexFactoryShaderParameters* FLocalVertexFactory::ConstructShaderParameters(EShaderFrequency ShaderFrequency)
-{
-	if (ShaderFrequency == SF_Vertex)
-	{
-		return new FLocalVertexFactoryShaderParameters();
-	}
-
+IMPLEMENT_VERTEX_FACTORY_PARAMETER_TYPE(FLocalVertexFactory, SF_Vertex, FLocalVertexFactoryShaderParameters);
 #if RHI_RAYTRACING
-	if (ShaderFrequency == SF_RayHitGroup)
-	{
-		return new FLocalVertexFactoryShaderParameters();
-	} 
-	else if (ShaderFrequency == SF_Compute)
-	{
-		return new FLocalVertexFactoryShaderParameters();
-	}
+IMPLEMENT_VERTEX_FACTORY_PARAMETER_TYPE(FLocalVertexFactory, SF_RayHitGroup, FLocalVertexFactoryShaderParameters);
+IMPLEMENT_VERTEX_FACTORY_PARAMETER_TYPE(FLocalVertexFactory, SF_Compute, FLocalVertexFactoryShaderParameters);
 #endif // RHI_RAYTRACING
 
-	return NULL;
-}
-
 IMPLEMENT_VERTEX_FACTORY_TYPE_EX(FLocalVertexFactory,"/Engine/Private/LocalVertexFactory.ush",true,true,true,true,true,true,true);
diff --git a/Engine/Source/Runtime/Engine/Private/Materials/HLSLMaterialTranslator.cpp b/Engine/Source/Runtime/Engine/Private/Materials/HLSLMaterialTranslator.cpp
--- a/Engine/Source/Runtime/Engine/Private/Materials/HLSLMaterialTranslator.cpp
+++ b/Engine/Source/Runtime/Engine/Private/Materials/HLSLMaterialTranslator.cpp
@@ -1032,6 +1032,36 @@ bool FHLSLMaterialTranslator::Translate()
 			ResourcesString += CustomOutputImplementations[ExpressionIndex] + "\r\n\r\n";
 		}
 
+		for (const FMaterialUniformExpression* ScalarExpression : UniformScalarExpressions)
+		{
+			FMaterialUniformPreshaderHeader& Preshader = MaterialCompilationOutput.UniformExpressionSet.UniformScalarPreshaders.AddDefaulted_GetRef();
+			Preshader.OpcodeOffset = MaterialCompilationOutput.UniformExpressionSet.UniformPreshaderData.Num();
+			ScalarExpression->WriteNumberOpcodes(MaterialCompilationOutput.UniformExpressionSet.UniformPreshaderData);
+			Preshader.OpcodeSize = MaterialCompilationOutput.UniformExpressionSet.UniformPreshaderData.Num() - Preshader.OpcodeOffset;
+		}
+
+		for (FMaterialUniformExpression* VectorExpression : UniformVectorExpressions)
+		{
+			FMaterialUniformPreshaderHeader& Preshader = MaterialCompilationOutput.UniformExpressionSet.UniformVectorPreshaders.AddDefaulted_GetRef();
+			Preshader.OpcodeOffset = MaterialCompilationOutput.UniformExpressionSet.UniformPreshaderData.Num();
+			VectorExpression->WriteNumberOpcodes(MaterialCompilationOutput.UniformExpressionSet.UniformPreshaderData);
+			Preshader.OpcodeSize = MaterialCompilationOutput.UniformExpressionSet.UniformPreshaderData.Num() - Preshader.OpcodeOffset;
+		}
+
+		for (uint32 TypeIndex = 0u; TypeIndex < NumMaterialTextureParameterTypes; ++TypeIndex)
+		{
+			MaterialCompilationOutput.UniformExpressionSet.UniformTextureParameters[TypeIndex].Empty(UniformTextureExpressions[TypeIndex].Num());
+			for (FMaterialUniformExpressionTexture* TextureExpression : UniformTextureExpressions[TypeIndex])
+			{
+				TextureExpression->GetTextureParameterInfo(MaterialCompilationOutput.UniformExpressionSet.UniformTextureParameters[TypeIndex].AddDefaulted_GetRef());
+			}
+		}
+		MaterialCompilationOutput.UniformExpressionSet.UniformExternalTextureParameters.Empty(UniformExternalTextureExpressions.Num());
+		for (FMaterialUniformExpressionExternalTexture* TextureExpression : UniformExternalTextureExpressions)
+		{
+			TextureExpression->GetExternalTextureParameterInfo(MaterialCompilationOutput.UniformExpressionSet.UniformExternalTextureParameters.AddDefaulted_GetRef());
+		}
+
 		LoadShaderSourceFileChecked(TEXT("/Engine/Private/MaterialTemplate.ush"), GetShaderPlatform(), MaterialTemplate);
 
 		// Find the string index of the '#line' statement in MaterialTemplate.usf
@@ -2157,7 +2187,7 @@ int32 FHLSLMaterialTranslator::AccessUniformExpression(int32 Index)
 	if(CodeChunk.Type == MCT_Float)
 	{
 		const static TCHAR IndexToMask[] = {'x', 'y', 'z', 'w'};
-		const int32 ScalarInputIndex = MaterialCompilationOutput.UniformExpressionSet.UniformScalarExpressions.AddUnique(CodeChunk.UniformExpression);
+		const int32 ScalarInputIndex = UniformScalarExpressions.AddUnique(CodeChunk.UniformExpression);
 		// Update the above FMemory::Malloc if this FCString::Sprintf grows in size, e.g. %s, ...
 		FCString::Sprintf(FormattedCode, TEXT("Material.ScalarExpressions[%u].%c"), ScalarInputIndex / 4, IndexToMask[ScalarInputIndex % 4]);
 	}
@@ -2173,7 +2203,7 @@ int32 FHLSLMaterialTranslator::AccessUniformExpression(int32 Index)
 		default: Mask = TEXT(""); break;
 		};
 
-		const int32 VectorInputIndex = MaterialCompilationOutput.UniformExpressionSet.UniformVectorExpressions.AddUnique(CodeChunk.UniformExpression);
+		const int32 VectorInputIndex = UniformVectorExpressions.AddUnique(CodeChunk.UniformExpression);
 		FCString::Sprintf(FormattedCode, TEXT("Material.VectorExpressions[%u]%s"), VectorInputIndex, Mask);
 	}
 	else if(CodeChunk.Type & MCT_Texture)
@@ -2184,27 +2214,27 @@ int32 FHLSLMaterialTranslator::AccessUniformExpression(int32 Index)
 		switch(CodeChunk.Type)
 		{
 		case MCT_Texture2D:
-			TextureInputIndex = MaterialCompilationOutput.UniformExpressionSet.Uniform2DTextureExpressions.AddUnique(TextureUniformExpression);
+			TextureInputIndex = UniformTextureExpressions[(uint32)EMaterialTextureParameterType::Standard2D].AddUnique(TextureUniformExpression);
 			BaseName = TEXT("Texture2D");
 			break;
 		case MCT_TextureCube:
-			TextureInputIndex = MaterialCompilationOutput.UniformExpressionSet.UniformCubeTextureExpressions.AddUnique(TextureUniformExpression);
+			TextureInputIndex = UniformTextureExpressions[(uint32)EMaterialTextureParameterType::Cube].AddUnique(TextureUniformExpression);
 			BaseName = TEXT("TextureCube");
 			break;
 		case MCT_Texture2DArray:
-			TextureInputIndex = MaterialCompilationOutput.UniformExpressionSet.Uniform2DArrayTextureExpressions.AddUnique(TextureUniformExpression);
+			TextureInputIndex = UniformTextureExpressions[(uint32)EMaterialTextureParameterType::Array2D].AddUnique(TextureUniformExpression);
 			BaseName = TEXT("Texture2DArray");
 			break;
 		case MCT_VolumeTexture:
-			TextureInputIndex = MaterialCompilationOutput.UniformExpressionSet.UniformVolumeTextureExpressions.AddUnique(TextureUniformExpression);
+			TextureInputIndex = UniformTextureExpressions[(uint32)EMaterialTextureParameterType::Volume].AddUnique(TextureUniformExpression);
 			BaseName = TEXT("VolumeTexture");
 			break;
 		case MCT_TextureExternal:
-			TextureInputIndex = MaterialCompilationOutput.UniformExpressionSet.UniformExternalTextureExpressions.AddUnique(ExternalTextureUniformExpression);
+			TextureInputIndex = UniformExternalTextureExpressions.AddUnique(ExternalTextureUniformExpression);
 			BaseName = TEXT("ExternalTexture");
 			break;
 		case MCT_TextureVirtual:
-			TextureInputIndex = MaterialCompilationOutput.UniformExpressionSet.UniformVirtualTextureExpressions.AddUnique(TextureUniformExpression);
+			TextureInputIndex = UniformTextureExpressions[(uint32)EMaterialTextureParameterType::Virtual].AddUnique(TextureUniformExpression);
 			GenerateCode = false;
 			break;
 		default: UE_LOG(LogMaterial, Fatal,TEXT("Unrecognized texture material value type: %u"),(int32)CodeChunk.Type);
@@ -2827,14 +2857,53 @@ int32 FHLSLMaterialTranslator::ScalarParameter(FName ParameterName, float Defaul
 {
 	FMaterialParameterInfo ParameterInfo = GetParameterAssociationInfo();
 	ParameterInfo.Name = ParameterName;
-	return AddUniformExpression(new FMaterialUniformExpressionScalarParameter(ParameterInfo,DefaultValue),MCT_Float,TEXT(""));
+	int32 ParameterIndex = INDEX_NONE;
+	for (int32 i = 0; i < MaterialCompilationOutput.UniformExpressionSet.UniformScalarParameters.Num(); ++i)
+	{
+		const FMaterialScalarParameterInfo& Parameter = MaterialCompilationOutput.UniformExpressionSet.UniformScalarParameters[i];
+		if (Parameter.ParameterInfo == ParameterInfo)
+		{
+			ParameterIndex = i;
+			break;
+		}
+	}
+	if (ParameterIndex == INDEX_NONE)
+	{
+		ParameterIndex = MaterialCompilationOutput.UniformExpressionSet.UniformScalarParameters.Num();
+		FMaterialScalarParameterInfo& Parameter = MaterialCompilationOutput.UniformExpressionSet.UniformScalarParameters.AddDefaulted_GetRef();
+		Parameter.ParameterInfo = ParameterInfo;
+		Parameter.ParameterName = ParameterName.ToString();
+		Parameter.DefaultValue = DefaultValue;
+	}
+
+	return AddUniformExpression(new FMaterialUniformExpressionScalarParameter(ParameterInfo, ParameterIndex), MCT_Float, TEXT(""));
 }
 
 int32 FHLSLMaterialTranslator::VectorParameter(FName ParameterName, const FLinearColor& DefaultValue)
 {
 	FMaterialParameterInfo ParameterInfo = GetParameterAssociationInfo();
 	ParameterInfo.Name = ParameterName;
-	return AddUniformExpression(new FMaterialUniformExpressionVectorParameter(ParameterInfo,DefaultValue),MCT_Float4,TEXT(""));
+
+	int32 ParameterIndex = INDEX_NONE;
+	for (int32 i = 0; i < MaterialCompilationOutput.UniformExpressionSet.UniformVectorParameters.Num(); ++i)
+	{
+		const FMaterialVectorParameterInfo& Parameter = MaterialCompilationOutput.UniformExpressionSet.UniformVectorParameters[i];
+		if (Parameter.ParameterInfo == ParameterInfo)
+		{
+			ParameterIndex = i;
+			break;
+		}
+	}
+	if (ParameterIndex == INDEX_NONE)
+	{
+		ParameterIndex = MaterialCompilationOutput.UniformExpressionSet.UniformVectorParameters.Num();
+		FMaterialVectorParameterInfo& Parameter = MaterialCompilationOutput.UniformExpressionSet.UniformVectorParameters.AddDefaulted_GetRef();
+		Parameter.ParameterInfo = ParameterInfo;
+		Parameter.ParameterName = ParameterName.ToString();
+		Parameter.DefaultValue = DefaultValue;
+	}
+
+	return AddUniformExpression(new FMaterialUniformExpressionVectorParameter(ParameterInfo, ParameterIndex), MCT_Float4, TEXT(""));
 }
 
 int32 FHLSLMaterialTranslator::Constant(float X)
@@ -4322,8 +4391,8 @@ int32 FHLSLMaterialTranslator::TextureSample(
 			return Errorf(TEXT("The provided uniform expression is not a texture"));
 		}
 
-		VirtualTextureIndex = MaterialCompilationOutput.UniformExpressionSet.UniformVirtualTextureExpressions.Find(TextureUniformExpression);
-		check(MaterialCompilationOutput.UniformExpressionSet.UniformVirtualTextureExpressions.IsValidIndex(VirtualTextureIndex));
+		VirtualTextureIndex = UniformTextureExpressions[(uint32)EMaterialTextureParameterType::Virtual].Find(TextureUniformExpression);
+		check(UniformTextureExpressions[(uint32)EMaterialTextureParameterType::Virtual].IsValidIndex(VirtualTextureIndex));
 
 		if (SamplerSource != SSM_FromTextureAsset)
 		{
@@ -4402,13 +4471,13 @@ int32 FHLSLMaterialTranslator::TextureSample(
 		//todo[vt]: Support feedback from other shader types
 		const bool bGenerateFeedback = ShaderFrequency == SF_Pixel;
 
-		VTLayerIndex = MaterialCompilationOutput.UniformExpressionSet.UniformVirtualTextureExpressions[VirtualTextureIndex]->GetTextureLayerIndex();
+		VTLayerIndex = UniformTextureExpressions[(uint32)EMaterialTextureParameterType::Virtual][VirtualTextureIndex]->GetTextureLayerIndex();
 		if (VTLayerIndex != INDEX_NONE)
 		{
 			// The layer index in the virtual texture stack is already known
 			// Create a page table sample for each new combination of virtual texture and sample parameters
 			VTStackIndex = AcquireVTStackIndex(MipValueMode, AddressU, AddressV, 1.0f, CoordinateIndex, MipValue0Index, MipValue1Index, TextureReferenceIndex, bGenerateFeedback);
-			VTPageTableIndex = MaterialCompilationOutput.UniformExpressionSet.UniformVirtualTextureExpressions[VirtualTextureIndex]->GetPageTableLayerIndex();
+			VTPageTableIndex = UniformTextureExpressions[(uint32)EMaterialTextureParameterType::Virtual][VirtualTextureIndex]->GetPageTableLayerIndex();
 		}
 		else
 		{
@@ -4836,7 +4905,7 @@ int32 FHLSLMaterialTranslator::Texture(UTexture* InTexture, int32& TextureRefere
 	// UE-3518: Additional pre-assert logging to help determine the cause of this failure.
 	if (TextureReferenceIndex == INDEX_NONE)
 	{
-		const TArray<UObject*>& ReferencedTextures = Material->GetReferencedTextures();
+		const auto ReferencedTextures = Material->GetReferencedTextures();
 		UE_LOG(LogMaterial, Error, TEXT("Compiler->Texture() failed to find texture '%s' in referenced list of size '%i':"), *InTexture->GetName(), ReferencedTextures.Num());
 		for (int32 i = 0; i < ReferencedTextures.Num(); ++i)
 		{
diff --git a/Engine/Source/Runtime/Engine/Private/Materials/HLSLMaterialTranslator.h b/Engine/Source/Runtime/Engine/Private/Materials/HLSLMaterialTranslator.h
--- a/Engine/Source/Runtime/Engine/Private/Materials/HLSLMaterialTranslator.h
+++ b/Engine/Source/Runtime/Engine/Private/Materials/HLSLMaterialTranslator.h
@@ -179,6 +179,11 @@ protected:
 	// Uniform expressions used across all material properties
 	TArray<FShaderCodeChunk> UniformExpressions;
 
+	TArray<TRefCountPtr<FMaterialUniformExpression> > UniformVectorExpressions;
+	TArray<TRefCountPtr<FMaterialUniformExpression> > UniformScalarExpressions;
+	TArray<TRefCountPtr<FMaterialUniformExpressionTexture> > UniformTextureExpressions[NumMaterialTextureParameterTypes];
+	TArray<TRefCountPtr<FMaterialUniformExpressionExternalTexture>> UniformExternalTextureExpressions;
+
 	/** Parameter collections referenced by this material.  The position in this array is used as an index on the shader parameter. */
 	TArray<UMaterialParameterCollection*> ParameterCollections;
 
