Copying Private-LoadTimes-4.24 stream to Main. Biggest changes are in Materials/Shader memory freezing.

diff --git a/Engine/Source/Runtime/Engine/Private/Materials/Material.cpp b/Engine/Source/Runtime/Engine/Private/Materials/Material.cpp
--- a/Engine/Source/Runtime/Engine/Private/Materials/Material.cpp
+++ b/Engine/Source/Runtime/Engine/Private/Materials/Material.cpp
@@ -137,7 +137,7 @@ int32 FMaterialResource::CompilePropertyAndSetMaterialProperty(EMaterialProperty
 	UMaterialInterface* MaterialInterface = MaterialInstance ? static_cast<UMaterialInterface*>(MaterialInstance) : Material;
 
 	int32 Ret = INDEX_NONE;
-	
+
 	switch(Property)
 	{
 		case MP_EmissiveColor:
@@ -233,6 +233,7 @@ int32 FMaterialResource::CompileCustomAttribute(const FGuid& AttributeID, FMater
 }
 #endif
 
+#if WITH_EDITORONLY_DATA
 void FMaterialResource::GatherCustomOutputExpressions(TArray<UMaterialExpressionCustomOutput*>& OutCustomOutputs) const
 {
 	Material->GetAllCustomOutputExpressions(OutCustomOutputs);
@@ -242,6 +243,7 @@ void FMaterialResource::GatherExpressionsForCustomInterpolators(TArray<UMaterial
 {
 	Material->GetAllExpressionsForCustomInterpolators(OutExpressions);
 }
+#endif // WITH_EDITORONLY_DATA
 
 #if WITH_EDITOR
 void FMaterialResource::BeginAllowCachingStaticParameterValues()
@@ -281,7 +283,7 @@ void FMaterialResource::GetShaderMapId(EShaderPlatform Platform, FMaterialShader
 #endif // WITH_EDITOR
 }
 
-
+#if WITH_EDITORONLY_DATA
 void FMaterialResource::GetStaticParameterSet(EShaderPlatform Platform, FStaticParameterSet& OutSet) const
 {
 	FMaterial::GetStaticParameterSet(Platform, OutSet);
@@ -292,6 +294,7 @@ void FMaterialResource::GetStaticParameterSet(EShaderPlatform Platform, FStaticP
 		MaterialInstance->GetStaticParameterValues(OutSet);
 	}
 }
+#endif // WITH_EDITORONLY_DATA
 
 /**
  * A resource which represents the default instance of a UMaterial to the renderer.
@@ -346,7 +349,7 @@ public:
 		return Material;
 	}
 
-	virtual bool GetVectorValue(const FMaterialParameterInfo& ParameterInfo, FLinearColor* OutValue, const FMaterialRenderContext& Context) const
+	virtual bool GetVectorValue(const FHashedMaterialParameterInfo& ParameterInfo, FLinearColor* OutValue, const FMaterialRenderContext& Context) const
 	{
 		const FMaterialResource* MaterialResource = Material->GetMaterialResource(Context.Material.GetFeatureLevel());
 		if(MaterialResource && MaterialResource->GetRenderingThreadShaderMap())
@@ -358,12 +361,12 @@ public:
 			return GetFallbackRenderProxy().GetVectorValue(ParameterInfo, OutValue, Context);
 		}
 	}
-	virtual bool GetScalarValue(const FMaterialParameterInfo& ParameterInfo, float* OutValue, const FMaterialRenderContext& Context) const
+	virtual bool GetScalarValue(const FHashedMaterialParameterInfo& ParameterInfo, float* OutValue, const FMaterialRenderContext& Context) const
 	{
 		const FMaterialResource* MaterialResource = Material->GetMaterialResource(Context.Material.GetFeatureLevel());
 		if(MaterialResource && MaterialResource->GetRenderingThreadShaderMap())
 		{
-			static FName NameSubsurfaceProfile(TEXT("__SubsurfaceProfile"));
+			static FHashedName NameSubsurfaceProfile(TEXT("__SubsurfaceProfile"));
 			if (ParameterInfo.Name == NameSubsurfaceProfile)
 			{
 				const USubsurfaceProfile* MySubsurfaceProfileRT = GetSubsurfaceProfileRT();
@@ -392,7 +395,7 @@ public:
 			return GetFallbackRenderProxy().GetScalarValue(ParameterInfo, OutValue, Context);
 		}
 	}
-	virtual bool GetTextureValue(const FMaterialParameterInfo& ParameterInfo,const UTexture** OutValue, const FMaterialRenderContext& Context) const
+	virtual bool GetTextureValue(const FHashedMaterialParameterInfo& ParameterInfo,const UTexture** OutValue, const FMaterialRenderContext& Context) const
 	{
 		const FMaterialResource* MaterialResource = Material->GetMaterialResource(Context.Material.GetFeatureLevel());
 		if(MaterialResource && MaterialResource->GetRenderingThreadShaderMap())
@@ -405,7 +408,7 @@ public:
 		}
 	}
 
-	virtual bool GetTextureValue(const FMaterialParameterInfo& ParameterInfo, const URuntimeVirtualTexture** OutValue, const FMaterialRenderContext& Context) const
+	virtual bool GetTextureValue(const FHashedMaterialParameterInfo& ParameterInfo, const URuntimeVirtualTexture** OutValue, const FMaterialRenderContext& Context) const
 	{
 		const FMaterialResource* MaterialResource = Material->GetMaterialResource(Context.Material.GetFeatureLevel());
 		if (MaterialResource && MaterialResource->GetRenderingThreadShaderMap())
@@ -960,6 +963,7 @@ void UMaterial::PreSave(const class ITargetPlatform* TargetPlatform)
 	Super::PreSave(TargetPlatform);
 #if WITH_EDITOR
 	GMaterialsWithDirtyUsageFlags.RemoveAnnotation(this);
+	UpdateCachedExpressionData();
 #endif
 }
 
@@ -1008,23 +1012,13 @@ void UMaterial::GetUsedTextures(TArray<UTexture*>& OutTextures, EMaterialQuality
 				if (CurrentResource == nullptr || (FeatureLevelIndex != FeatureLevel && !bAllFeatureLevels))
 					continue;
 
-				const TArray<TRefCountPtr<FMaterialUniformExpressionTexture> >* ExpressionsByType[5] =
-				{
-					&CurrentResource->GetUniform2DTextureExpressions(),
-					&CurrentResource->GetUniformCubeTextureExpressions(),
-					&CurrentResource->GetUniform2DArrayTextureExpressions(),
-					&CurrentResource->GetUniformVolumeTextureExpressions(),
-					&CurrentResource->GetUniformVirtualTextureExpressions()
-				};
-				for (int32 TypeIndex = 0; TypeIndex < UE_ARRAY_COUNT(ExpressionsByType); TypeIndex++)
+				for (int32 TypeIndex = 0; TypeIndex < NumMaterialTextureParameterTypes; TypeIndex++)
 				{
 					// Iterate over each of the material's texture expressions.
-					for (FMaterialUniformExpressionTexture* Expression : *ExpressionsByType[TypeIndex])
+					for (const FMaterialTextureParameterInfo& Parameter : CurrentResource->GetUniformTextureExpressions((EMaterialTextureParameterType)TypeIndex))
 					{
-						const bool bAllowOverride = false;
 						UTexture* Texture = NULL;
-						Expression->GetGameThreadTextureValue(this, *CurrentResource, Texture, bAllowOverride);
-
+						Parameter.GetGameThreadTextureValue(this, *CurrentResource, Texture);
 						if (Texture)
 						{
 							OutTextures.AddUnique(Texture);
@@ -1032,29 +1026,33 @@ void UMaterial::GetUsedTextures(TArray<UTexture*>& OutTextures, EMaterialQuality
 					}
 				}
 #if WITH_EDITOR
-				// Also look for any scalar parameters that are acting as lookups for an atlas texture, and store the atlas texture
-				const TArray<TRefCountPtr<FMaterialUniformExpression> >* AtlasExpressions[1] =
-				{
-					&CurrentResource->GetUniformScalarParameterExpressions()
-				};
-				for (int32 TypeIndex = 0; TypeIndex < UE_ARRAY_COUNT(AtlasExpressions); TypeIndex++)
+
+				const UMaterialInstance* MaterialInstance = Cast<UMaterialInstance>(this);
+				if (MaterialInstance)
 				{
-					// Iterate over each of the material's texture expressions.
-					for (FMaterialUniformExpression* Expression : *AtlasExpressions[TypeIndex])
+					// Also look for any scalar parameters that are acting as lookups for an atlas texture, and store the atlas texture
+					const TArray<FMaterialScalarParameterInfo, FMemoryImageAllocator>* AtlasExpressions[1] =
 					{
-						const FMaterialUniformExpressionScalarParameter* ScalarExpression = static_cast<const FMaterialUniformExpressionScalarParameter*>(Expression);
-						bool bIsUsedAsAtlasPosition;
-						TSoftObjectPtr<class UCurveLinearColor> Curve;
-						TSoftObjectPtr<class UCurveLinearColorAtlas> Atlas;
-						ScalarExpression->GetGameThreadUsedAsAtlas(this, bIsUsedAsAtlasPosition, Curve, Atlas);
-
-						if (Atlas)
+						&CurrentResource->GetUniformScalarParameterExpressions()
+					};
+					for (int32 TypeIndex = 0; TypeIndex < UE_ARRAY_COUNT(AtlasExpressions); TypeIndex++)
+					{
+						// Iterate over each of the material's texture expressions.
+						for (const FMaterialScalarParameterInfo& Parameter : *AtlasExpressions[TypeIndex])
 						{
-							OutTextures.AddUnique(Atlas.Get());
+							bool bIsUsedAsAtlasPosition;
+							TSoftObjectPtr<class UCurveLinearColor> Curve;
+							TSoftObjectPtr<class UCurveLinearColorAtlas> Atlas;
+							MaterialInstance->IsScalarParameterUsedAsAtlasPosition(Parameter.ParameterInfo, bIsUsedAsAtlasPosition, Curve, Atlas);
+
+							if (Atlas)
+							{
+								OutTextures.AddUnique(Atlas.Get());
+							}
 						}
 					}
 				}
-#endif 
+#endif // WITH_EDITOR
 			}
 		}
 	}
@@ -1073,26 +1071,24 @@ void UMaterial::GetUsedTexturesAndIndices(TArray<UTexture*>& OutTextures, TArray
 
 		if (CurrentResource)
 		{
-			const TArray<TRefCountPtr<FMaterialUniformExpressionTexture> >* ExpressionsByType[5] =
+			const TArray<FMaterialTextureParameterInfo, FMemoryImageAllocator>* ExpressionsByType[NumMaterialTextureParameterTypes];
+			uint32 NumTextures = 0u;
+			for (uint32 TypeIndex = 0u; TypeIndex < NumMaterialTextureParameterTypes; ++TypeIndex)
 			{
-				&CurrentResource->GetUniform2DTextureExpressions(),
-				&CurrentResource->GetUniformCubeTextureExpressions(),
-				&CurrentResource->GetUniform2DArrayTextureExpressions(),
-				&CurrentResource->GetUniformVolumeTextureExpressions(),
-				&CurrentResource->GetUniformVirtualTextureExpressions()
-			};
+				ExpressionsByType[TypeIndex] = &CurrentResource->GetUniformTextureExpressions((EMaterialTextureParameterType)TypeIndex);
+				NumTextures += ExpressionsByType[TypeIndex]->Num();
+			}
 
 			// Try to prevent resizing since this would be expensive.
-			OutIndices.Empty(ExpressionsByType[0]->Num() + ExpressionsByType[1]->Num() + ExpressionsByType[2]->Num() + ExpressionsByType[3]->Num());
+			OutIndices.Empty(NumTextures);
 
 			for (int32 TypeIndex = 0; TypeIndex < UE_ARRAY_COUNT(ExpressionsByType); TypeIndex++)
 			{
 				// Iterate over each of the material's texture expressions.
-				for (FMaterialUniformExpressionTexture* Expression : *ExpressionsByType[TypeIndex])
+				for (const FMaterialTextureParameterInfo& Parameter : *ExpressionsByType[TypeIndex])
 				{
-					const bool bAllowOverride = false;
 					UTexture* Texture = NULL;
-					Expression->GetGameThreadTextureValue(this, *CurrentResource, Texture, bAllowOverride);
+					Parameter.GetGameThreadTextureValue(this, *CurrentResource, Texture);
 
 					if (Texture)
 					{
@@ -1101,7 +1097,7 @@ void UMaterial::GetUsedTexturesAndIndices(TArray<UTexture*>& OutTextures, TArray
 						{
 							OutIndices.AddDefaulted(InsertIndex - OutIndices.Num() + 1);
 						}
-						OutIndices[InsertIndex].Add(Expression->GetTextureIndex());
+						OutIndices[InsertIndex].Add(Parameter.TextureIndex);
 					}
 				}
 			}
@@ -1132,20 +1128,12 @@ void UMaterial::LogMaterialsAndTextures(FOutputDevice& Ar, int32 Indent) const
 				TArray<UTexture*> Textures;
 				// GetTextureExpressionValues(MaterialResource, Textures);
 				{
-					const TArray<TRefCountPtr<FMaterialUniformExpressionTexture> >* ExpressionsByType[5]= 
-					{
-						&MaterialResource->GetUniform2DTextureExpressions(), 
-						&MaterialResource->GetUniformCubeTextureExpressions(),
-						&MaterialResource->GetUniform2DArrayTextureExpressions(),
-						&MaterialResource->GetUniformVolumeTextureExpressions(),
-						&MaterialResource->GetUniformVirtualTextureExpressions(),
-					};
-					for (int32 TypeIndex = 0; TypeIndex < UE_ARRAY_COUNT(ExpressionsByType); TypeIndex++)
+					for (int32 TypeIndex = 0; TypeIndex < NumMaterialTextureParameterTypes; TypeIndex++)
 					{
-						for (FMaterialUniformExpressionTexture* Expression : *ExpressionsByType[TypeIndex])
+						for (const FMaterialTextureParameterInfo& Parameter : MaterialResource->GetUniformTextureExpressions((EMaterialTextureParameterType)TypeIndex))
 						{
 							UTexture* Texture = NULL;
-							Expression->GetGameThreadTextureValue(this, *MaterialResource, Texture, false);
+							Parameter.GetGameThreadTextureValue(this, *MaterialResource, Texture);
 							if (Texture)
 							{
 								Textures.AddUnique(Texture);
@@ -1185,32 +1173,25 @@ void UMaterial::OverrideTexture(const UTexture* InTextureToOverride, UTexture* O
 	ERHIFeatureLevel::Type FeatureLevelsToUpdate[1] = { InFeatureLevel };
 	int32 NumFeatureLevelsToUpdate = 1;
 	
-	for (int32 i = 0; i < NumFeatureLevelsToUpdate; ++i)
+	for (int32 FeatureLevelIndex = 0; FeatureLevelIndex < NumFeatureLevelsToUpdate; ++FeatureLevelIndex)
 	{
-		FMaterialResource* Resource = GetMaterialResource(FeatureLevelsToUpdate[i]);
+		FMaterialResource* Resource = GetMaterialResource(FeatureLevelsToUpdate[FeatureLevelIndex]);
 		// Iterate over both the 2D textures and cube texture expressions.
-		const TArray<TRefCountPtr<FMaterialUniformExpressionTexture> >* ExpressionsByType[5] =
-		{
-			&Resource->GetUniform2DTextureExpressions(),
-			&Resource->GetUniformCubeTextureExpressions(),
-			&Resource->GetUniform2DArrayTextureExpressions(),
-			&Resource->GetUniformVolumeTextureExpressions(),
-			&Resource->GetUniformVirtualTextureExpressions()
-		};
-		for(int32 TypeIndex = 0;TypeIndex < UE_ARRAY_COUNT(ExpressionsByType);TypeIndex++)
+		for(int32 TypeIndex = 0;TypeIndex < NumMaterialTextureParameterTypes; TypeIndex++)
 		{
+			const TArray<FMaterialTextureParameterInfo, FMemoryImageAllocator>& Parameters = Resource->GetUniformTextureExpressions((EMaterialTextureParameterType)TypeIndex);
 			// Iterate over each of the material's texture expressions.
-			for (FMaterialUniformExpressionTexture* Expression : *ExpressionsByType[TypeIndex])
+			for (int32 ParameterIndex = 0; ParameterIndex < Parameters.Num(); ++ParameterIndex)
 			{
+				const FMaterialTextureParameterInfo& Parameter = Parameters[ParameterIndex];
+
 				// Evaluate the expression in terms of this material instance.
-				const bool bAllowOverride = false;
 				UTexture* Texture = NULL;
-				Expression->GetGameThreadTextureValue(this,*Resource,Texture,bAllowOverride);
-
+				Parameter.GetGameThreadTextureValue(this,*Resource,Texture);
 				if( Texture != NULL && Texture == InTextureToOverride )
 				{
 					// Override this texture!
-					Expression->SetTransientOverrideTextureValue( OverrideTexture );
+					Resource->TransientOverrides.SetTextureOverride((EMaterialTextureParameterType)TypeIndex, ParameterIndex, OverrideTexture);
 					bShouldRecacheMaterialExpressions = true;
 				}
 			}
@@ -1221,29 +1202,25 @@ void UMaterial::OverrideTexture(const UTexture* InTextureToOverride, UTexture* O
 	{
 		RecacheUniformExpressions(false);
 	}
-#endif // #if WITH_EDITOR
+#endif // WITH_EDITOR
 }
 
-void UMaterial::OverrideVectorParameterDefault(const FMaterialParameterInfo& ParameterInfo, const FLinearColor& Value, bool bOverride, ERHIFeatureLevel::Type InFeatureLevel)
+void UMaterial::OverrideVectorParameterDefault(const FHashedMaterialParameterInfo& ParameterInfo, const FLinearColor& Value, bool bOverride, ERHIFeatureLevel::Type InFeatureLevel)
 {
 #if WITH_EDITOR
 	bool bShouldRecacheMaterialExpressions = false;
 
 	FMaterialResource* Resource = GetMaterialResource(InFeatureLevel);
-	const TArray<TRefCountPtr<FMaterialUniformExpression> >& UniformExpressions = Resource->GetUniformVectorParameterExpressions();
+	const TArray<FMaterialVectorParameterInfo, FMemoryImageAllocator>& Parameters = Resource->GetUniformVectorParameterExpressions();
 
 	// Iterate over each of the material's vector expressions.
-	for (FMaterialUniformExpression* UniformExpression : UniformExpressions)
+	for (int32 i = 0; i < Parameters.Num(); ++i)
 	{
-		if (UniformExpression->GetType() == &FMaterialUniformExpressionVectorParameter::StaticType)
+		const FMaterialVectorParameterInfo& Parameter = Parameters[i];
+		if (Parameter.ParameterInfo == ParameterInfo)
 		{
-			FMaterialUniformExpressionVectorParameter* VectorExpression = static_cast<FMaterialUniformExpressionVectorParameter*>(UniformExpression);
-
-			if (VectorExpression->GetParameterInfo() == ParameterInfo)
-			{
-				VectorExpression->SetTransientOverrideDefaultValue(Value, bOverride);
-				bShouldRecacheMaterialExpressions = true;
-			}
+			Resource->TransientOverrides.SetVectorOverride(i, Value, bOverride);
+			bShouldRecacheMaterialExpressions = true;
 		}
 	}
 
@@ -1254,27 +1231,22 @@ void UMaterial::OverrideVectorParameterDefault(const FMaterialParameterInfo& Par
 #endif // #if WITH_EDITOR
 }
 
-void UMaterial::OverrideScalarParameterDefault(const FMaterialParameterInfo& ParameterInfo, float Value, bool bOverride, ERHIFeatureLevel::Type InFeatureLevel)
+void UMaterial::OverrideScalarParameterDefault(const FHashedMaterialParameterInfo& ParameterInfo, float Value, bool bOverride, ERHIFeatureLevel::Type InFeatureLevel)
 {
 #if WITH_EDITOR
 	bool bShouldRecacheMaterialExpressions = false;
 
 	FMaterialResource* Resource = GetMaterialResource(InFeatureLevel);
-	// Iterate over both the 2D textures and cube texture expressions.
-	const TArray<TRefCountPtr<FMaterialUniformExpression> >& UniformExpressions = Resource->GetUniformScalarParameterExpressions();
+	const TArray<FMaterialScalarParameterInfo, FMemoryImageAllocator>& Parameters = Resource->GetUniformScalarParameterExpressions();
 
-	// Iterate over each of the material's texture expressions.
-	for (FMaterialUniformExpression* UniformExpression : UniformExpressions)
+	// Iterate over each of the material's vector expressions.
+	for (int32 i = 0; i < Parameters.Num(); ++i)
 	{
-		if (UniformExpression->GetType() == &FMaterialUniformExpressionScalarParameter::StaticType)
+		const FMaterialScalarParameterInfo& Parameter = Parameters[i];
+		if (Parameter.ParameterInfo == ParameterInfo)
 		{
-			FMaterialUniformExpressionScalarParameter* ScalarExpression = static_cast<FMaterialUniformExpressionScalarParameter*>(UniformExpression);
-
-			if (ScalarExpression->GetParameterInfo() == ParameterInfo)
-			{
-				ScalarExpression->SetTransientOverrideDefaultValue(Value, bOverride);
-				bShouldRecacheMaterialExpressions = true;
-			}
+			Resource->TransientOverrides.SetScalarOverride(i, Value, bOverride);
+			bShouldRecacheMaterialExpressions = true;
 		}
 	}
 
@@ -1707,103 +1679,34 @@ void UMaterial::FixupMaterialUsageAfterLoad()
 	// All we need to do here is mark the package dirty as the usage itself was set on load.
 	MarkPackageDirty();
 }
-#endif
-
-
-void UMaterial::GetAllScalarParameterNames(TArray<FName>& OutParameterNames, TArray<FGuid>& OutParameterIds) const
-{
-	OutParameterNames.Reset();
-	TArray<FMaterialParameterInfo> OutParameterInfo;
-	GetAllScalarParameterInfo(OutParameterInfo, OutParameterIds);
-	for (FMaterialParameterInfo Info : OutParameterInfo)
-	{
-		OutParameterNames.Add(Info.Name);
-	}
-}
-
-void UMaterial::GetAllVectorParameterNames(TArray<FName>& OutParameterNames, TArray<FGuid>& OutParameterIds) const
-{
-	OutParameterNames.Reset();
-	TArray<FMaterialParameterInfo> OutParameterInfo;
-	GetAllVectorParameterInfo(OutParameterInfo, OutParameterIds);
-	for (FMaterialParameterInfo Info : OutParameterInfo)
-	{
-		OutParameterNames.Add(Info.Name);
-	}
-}
-
-void UMaterial::GetAllTextureParameterNames(TArray<FName>& OutParameterNames, TArray<FGuid>& OutParameterIds) const
-{
-	OutParameterNames.Reset();
-	TArray<FMaterialParameterInfo> OutParameterInfo;
-	GetAllTextureParameterInfo(OutParameterInfo, OutParameterIds);
-	for (FMaterialParameterInfo Info : OutParameterInfo)
-	{
-		OutParameterNames.Add(Info.Name);
-	}
-}
-
-void UMaterial::GetAllFontParameterNames(TArray<FName>& OutParameterNames, TArray<FGuid>& OutParameterIds) const
-{
-	OutParameterNames.Reset();
-	TArray<FMaterialParameterInfo> OutParameterInfo;
-	GetAllFontParameterInfo(OutParameterInfo, OutParameterIds);
-	for (FMaterialParameterInfo Info : OutParameterInfo)
-	{
-		OutParameterNames.Add(Info.Name);
-	}
-}
-
-void UMaterial::GetAllStaticSwitchParameterNames(TArray<FName>& OutParameterNames, TArray<FGuid>& OutParameterIds) const
-{
-	OutParameterNames.Reset();
-	TArray<FMaterialParameterInfo> OutParameterInfo;
-	GetAllStaticSwitchParameterInfo(OutParameterInfo, OutParameterIds);
-	for (FMaterialParameterInfo Info : OutParameterInfo)
-	{
-		OutParameterNames.Add(Info.Name);
-	}
-}
-
-void UMaterial::GetAllStaticComponentMaskParameterNames(TArray<FName>& OutParameterNames, TArray<FGuid>& OutParameterIds) const
-{
-	OutParameterNames.Reset();
-	TArray<FMaterialParameterInfo> OutParameterInfo;
-	GetAllStaticComponentMaskParameterInfo(OutParameterInfo, OutParameterIds);
-	for (FMaterialParameterInfo Info : OutParameterInfo)
-	{
-		OutParameterNames.Add(Info.Name);
-	}
-}
+#endif // WITH_EDITOR
 
 void UMaterial::GetAllScalarParameterInfo(TArray<FMaterialParameterInfo>& OutParameterInfo, TArray<FGuid>& OutParameterIds) const
 {
-	OutParameterInfo.Reset();
-	OutParameterIds.Reset();
-	GetAllParameterInfo<UMaterialExpressionScalarParameter>(OutParameterInfo, OutParameterIds);
+	CachedExpressionData.Parameters.GetAllParameterInfoOfType(EMaterialParameterType::Scalar, true, OutParameterInfo, OutParameterIds);
 }
 
 void UMaterial::GetAllVectorParameterInfo(TArray<FMaterialParameterInfo>& OutParameterInfo, TArray<FGuid>& OutParameterIds) const
 {
-	OutParameterInfo.Reset();
-	OutParameterIds.Reset();
-	GetAllParameterInfo<UMaterialExpressionVectorParameter>(OutParameterInfo, OutParameterIds);
+	CachedExpressionData.Parameters.GetAllParameterInfoOfType(EMaterialParameterType::Vector, true, OutParameterInfo, OutParameterIds);
 }
 
 void UMaterial::GetAllTextureParameterInfo(TArray<FMaterialParameterInfo>& OutParameterInfo, TArray<FGuid>& OutParameterIds) const
 {
-	OutParameterInfo.Reset();
-	OutParameterIds.Reset();
-	GetAllParameterInfo<UMaterialExpressionTextureSampleParameter>(OutParameterInfo, OutParameterIds);
+	CachedExpressionData.Parameters.GetAllParameterInfoOfType(EMaterialParameterType::Texture, true, OutParameterInfo, OutParameterIds);
 }
 
 void UMaterial::GetAllFontParameterInfo(TArray<FMaterialParameterInfo>& OutParameterInfo, TArray<FGuid>& OutParameterIds) const
 {
-	OutParameterInfo.Reset();
-	OutParameterIds.Reset();
-	GetAllParameterInfo<UMaterialExpressionFontSampleParameter>(OutParameterInfo, OutParameterIds);
+	CachedExpressionData.Parameters.GetAllParameterInfoOfType(EMaterialParameterType::Font, true, OutParameterInfo, OutParameterIds);
+}
+
+void UMaterial::GetAllRuntimeVirtualTextureParameterInfo(TArray<FMaterialParameterInfo>& OutParameterInfo, TArray<FGuid>& OutParameterIds) const
+{
+	CachedExpressionData.Parameters.GetAllParameterInfoOfType(EMaterialParameterType::RuntimeVirtualTexture, true, OutParameterInfo, OutParameterIds);
 }
 
+#if WITH_EDITORONLY_DATA
 void UMaterial::GetAllMaterialLayersParameterInfo(TArray<FMaterialParameterInfo>& OutParameterInfo, TArray<FGuid>& OutParameterIds) const
 {
 	OutParameterInfo.Reset();
@@ -1855,41 +1758,44 @@ void UMaterial::GetDependentFunctions(TArray<UMaterialFunctionInterface*>& Depen
 		return true;
 	});
 }
+#endif // WITH_EDITORONLY_DATA
 
-bool UMaterial::GetScalarParameterDefaultValue(const FMaterialParameterInfo& ParameterInfo, float& OutValue, bool bOveriddenOnly, bool bCheckOwnedGlobalOverrides) const
+bool UMaterial::GetScalarParameterDefaultValue(const FHashedMaterialParameterInfo& ParameterInfo, float& OutValue, bool bOveriddenOnly, bool bCheckOwnedGlobalOverrides) const
 {
 	return GetScalarParameterValue(ParameterInfo, OutValue, bOveriddenOnly);
 }
 
-bool UMaterial::GetVectorParameterDefaultValue(const FMaterialParameterInfo& ParameterInfo, FLinearColor& OutValue, bool bOveriddenOnly, bool bCheckOwnedGlobalOverrides) const
+bool UMaterial::GetVectorParameterDefaultValue(const FHashedMaterialParameterInfo& ParameterInfo, FLinearColor& OutValue, bool bOveriddenOnly, bool bCheckOwnedGlobalOverrides) const
 {
 	return GetVectorParameterValue(ParameterInfo, OutValue, bOveriddenOnly);
 }
 
-bool UMaterial::GetTextureParameterDefaultValue(const FMaterialParameterInfo& ParameterInfo, class UTexture*& OutValue, bool bCheckOwnedGlobalOverrides) const
+bool UMaterial::GetTextureParameterDefaultValue(const FHashedMaterialParameterInfo& ParameterInfo, class UTexture*& OutValue, bool bCheckOwnedGlobalOverrides) const
 {
 	return GetTextureParameterValue(ParameterInfo, OutValue);
 }
 
-bool UMaterial::GetRuntimeVirtualTextureParameterDefaultValue(const FMaterialParameterInfo& ParameterInfo, class URuntimeVirtualTexture*& OutValue, bool bCheckOwnedGlobalOverrides) const
+bool UMaterial::GetRuntimeVirtualTextureParameterDefaultValue(const FHashedMaterialParameterInfo& ParameterInfo, class URuntimeVirtualTexture*& OutValue, bool bCheckOwnedGlobalOverrides) const
 {
 	return GetRuntimeVirtualTextureParameterValue(ParameterInfo, OutValue);
 }
 
-bool UMaterial::GetFontParameterDefaultValue(const FMaterialParameterInfo& ParameterInfo, class UFont*& OutFontValue, int32& OutFontPage, bool bCheckOwnedGlobalOverrides) const
+bool UMaterial::GetFontParameterDefaultValue(const FHashedMaterialParameterInfo& ParameterInfo, class UFont*& OutFontValue, int32& OutFontPage, bool bCheckOwnedGlobalOverrides) const
 {
 	return GetFontParameterValue(ParameterInfo, OutFontValue, OutFontPage);
 }
 
-bool UMaterial::GetStaticComponentMaskParameterDefaultValue(const FMaterialParameterInfo& ParameterInfo, bool& OutR, bool& OutG, bool& OutB, bool& OutA, FGuid& OutExpressionGuid, bool bCheckOwnedGlobalOverrides) const
+#if WITH_EDITOR
+bool UMaterial::GetStaticComponentMaskParameterDefaultValue(const FHashedMaterialParameterInfo& ParameterInfo, bool& OutR, bool& OutG, bool& OutB, bool& OutA, FGuid& OutExpressionGuid, bool bCheckOwnedGlobalOverrides) const
 {
 	return GetStaticComponentMaskParameterValue(ParameterInfo, OutR, OutG, OutB, OutA, OutExpressionGuid);
 }
 
-bool UMaterial::GetStaticSwitchParameterDefaultValue(const FMaterialParameterInfo& ParameterInfo, bool& OutValue, FGuid& OutExpressionGuid, bool bCheckOwnedGlobalOverrides) const
+bool UMaterial::GetStaticSwitchParameterDefaultValue(const FHashedMaterialParameterInfo& ParameterInfo, bool& OutValue, FGuid& OutExpressionGuid, bool bCheckOwnedGlobalOverrides) const
 {
 	return GetStaticSwitchParameterValue(ParameterInfo, OutValue, OutExpressionGuid);
 }
+#endif // WITH_EDITOR
 
 extern FPostProcessMaterialNode* IteratePostProcessMaterialNodes(const FFinalPostProcessSettings& Dest, const UMaterial* Material, FBlendableEntry*& Iterator);
 
@@ -1964,13 +1870,13 @@ const UMaterial* UMaterial::GetMaterial() const
 	return this;
 }
 
-const UMaterial* UMaterial::GetMaterial_Concurrent(TMicRecursionGuard&) const
+const UMaterial* UMaterial::GetMaterial_Concurrent(TMicRecursionGuard) const
 {
 	return this;
 }
 
 #if WITH_EDITOR
-bool UMaterial::GetGroupName(const FMaterialParameterInfo& ParameterInfo, FName& OutGroup) const
+bool UMaterial::GetGroupName(const FHashedMaterialParameterInfo& ParameterInfo, FName& OutGroup) const
 {
 	for (const UMaterialExpression* Expression : Expressions)
 	{
@@ -2052,7 +1958,7 @@ bool UMaterial::GetGroupName(const FMaterialParameterInfo& ParameterInfo, FName&
 	return false;
 }
 
-bool UMaterial::GetParameterDesc(const FMaterialParameterInfo& ParameterInfo, FString& OutDesc, const TArray<FStaticMaterialLayersParameter>* MaterialLayersParameters) const
+bool UMaterial::GetParameterDesc(const FHashedMaterialParameterInfo& ParameterInfo, FString& OutDesc, const TArray<FStaticMaterialLayersParameter>* MaterialLayersParameters) const
 {
 	for (const UMaterialExpression* Expression : Expressions)
 	{
@@ -2134,9 +2040,56 @@ bool UMaterial::GetParameterDesc(const FMaterialParameterInfo& ParameterInfo, FS
 
 	return false;
 }
+void UMaterial::UpdateCachedExpressionData()
+{
+	CachedExpressionData.Reset();
+	CachedExpressionData.UpdateForExpressions(Expressions, EMaterialParameterAssociation::GlobalParameter, -1);
+
+	// Always append a default high quality level if nothing else is set
+	bool bAnyQualityLevelsSet = false;
+	for (bool Value : CachedExpressionData.QualityLevelsUsed)
+	{
+		if (Value)
+		{
+			bAnyQualityLevelsSet = true;
+			break;
+		}
+	}
+
+	if (!bAnyQualityLevelsSet)
+	{
+		CachedExpressionData.QualityLevelsUsed[EMaterialQualityLevel::High] = true;
+	}
+}
 #endif // WITH_EDITOR
 
-bool UMaterial::GetScalarParameterValue(const FMaterialParameterInfo& ParameterInfo, float& OutValue, bool bOveriddenOnly) const
+bool UMaterial::GetScalarParameterValue(const FHashedMaterialParameterInfo& ParameterInfo, float& OutValue, bool bOveriddenOnly) const
+{
+	const bool bResult = GetScalarParameterValue_New(ParameterInfo, OutValue, bOveriddenOnly);
+#if WITH_EDITOR
+	float OldValue = 0.0f;
+	const bool bOldResult = GetScalarParameterValue_Legacy(ParameterInfo, OldValue, bOveriddenOnly);
+
+	check(bOldResult == bResult);
+	check(!bResult || OutValue == OldValue);
+#endif
+	return bResult;
+}
+
+bool UMaterial::GetScalarParameterValue_New(const FHashedMaterialParameterInfo& ParameterInfo, float& OutValue, bool bOveriddenOnly) const
+{
+	const int32 Index = CachedExpressionData.Parameters.FindParameterIndex(EMaterialParameterType::Scalar, ParameterInfo, bOveriddenOnly);
+	if (Index == INDEX_NONE)
+	{
+		return false;
+	}
+
+	OutValue = CachedExpressionData.Parameters.ScalarValues[Index];
+	return true;
+}
+
+#if WITH_EDITOR
+bool UMaterial::GetScalarParameterValue_Legacy(const FHashedMaterialParameterInfo& ParameterInfo, float& OutValue, bool bOveriddenOnly) const
 {
 	if (bOveriddenOnly && !AreExperimentalMaterialLayersEnabled())
 	{
@@ -2213,140 +2166,69 @@ bool UMaterial::GetScalarParameterValue(const FMaterialParameterInfo& ParameterI
 	return false;
 }
 
-bool UMaterial::IsScalarParameterUsedAsAtlasPosition(const FMaterialParameterInfo& ParameterInfo, bool& OutValue, TSoftObjectPtr<UCurveLinearColor>& Curve, TSoftObjectPtr<UCurveLinearColorAtlas>& Atlas) const
+bool UMaterial::IsScalarParameterUsedAsAtlasPosition(const FHashedMaterialParameterInfo& ParameterInfo, bool& OutValue, TSoftObjectPtr<UCurveLinearColor>& OutCurve, TSoftObjectPtr<UCurveLinearColorAtlas>& OutAtlas) const
 {
-	// In the case of duplicate parameters with different values, this will return the
-	// first matching expression found, not necessarily the one that's used for rendering
-	UMaterialExpressionScalarParameter* Parameter = nullptr;
-	float TempScalar;
+	const int32 Index = CachedExpressionData.Parameters.FindParameterIndex(EMaterialParameterType::Scalar, ParameterInfo);
+	if (Index == INDEX_NONE)
+	{
+		return false;
+	}
 
-	for (UMaterialExpression* Expression : Expressions)
+	UCurveLinearColor* Curve = CachedExpressionData.Parameters.ScalarCurveValues[Index];
+	UCurveLinearColorAtlas* Atlas = CachedExpressionData.Parameters.ScalarCurveAtlasValues[Index];
+	if (Curve && Atlas)
 	{
-		// Only need to check parameters that match in associated scope
-		if (ParameterInfo.Association == EMaterialParameterAssociation::GlobalParameter)
-		{
-			if (UMaterialExpressionScalarParameter* ExpressionParameter = Cast<UMaterialExpressionScalarParameter>(Expression))
-			{
-				if (ExpressionParameter->IsNamedParameter(ParameterInfo, TempScalar))
-				{
-					OutValue = ExpressionParameter->IsUsedAsAtlasPosition();
-					if (OutValue)
-					{
-						UMaterialExpressionCurveAtlasRowParameter* AtlasParameter = Cast<UMaterialExpressionCurveAtlasRowParameter>(Expression);
-						Curve = TSoftObjectPtr<UCurveLinearColor>(FSoftObjectPath(AtlasParameter->Curve->GetPathName()));
-						Atlas = TSoftObjectPtr<UCurveLinearColorAtlas>(FSoftObjectPath(AtlasParameter->Atlas->GetPathName()));
-					}
-					return true;
-				}
-			}
-			else if (UMaterialExpressionMaterialFunctionCall* FunctionCall = Cast<UMaterialExpressionMaterialFunctionCall>(Expression))
-			{
-				UMaterialFunctionInterface* Function = FunctionCall->MaterialFunction;
-				UMaterialFunctionInterface* ParameterOwner = nullptr;
+		OutCurve = TSoftObjectPtr<UCurveLinearColor>(FSoftObjectPath(Curve->GetPathName()));
+		OutAtlas = TSoftObjectPtr<UCurveLinearColorAtlas>(FSoftObjectPath(Atlas->GetPathName()));
+		OutValue = true;
+	}
+	else
+	{
+		OutValue = false;
+	}
 
-				if (Function && Function->GetNamedParameterOfType(ParameterInfo, Parameter, &ParameterOwner))
-				{
-					if (Parameter)
-					{
-						OutValue = Parameter->IsUsedAsAtlasPosition();
-						if (OutValue)
-						{
-							UMaterialExpressionCurveAtlasRowParameter* AtlasParameter = Cast<UMaterialExpressionCurveAtlasRowParameter>(Parameter);
-							if (AtlasParameter && AtlasParameter->ParameterName == ParameterInfo.Name)
-							{
-								Curve = TSoftObjectPtr<UCurveLinearColor>(FSoftObjectPath(AtlasParameter->Curve->GetPathName()));
-								Atlas = TSoftObjectPtr<UCurveLinearColorAtlas>(FSoftObjectPath(AtlasParameter->Atlas->GetPathName()));
-								return true;
-							}
-						}
-					}
-				}
-			}
-		}
-		else
-		{
-			if (UMaterialExpressionMaterialAttributeLayers* LayersExpression = Cast<UMaterialExpressionMaterialAttributeLayers>(Expression))
-			{
-				UMaterialFunctionInterface* Function = LayersExpression->GetParameterAssociatedFunction(ParameterInfo);
-				UMaterialFunctionInterface* ParameterOwner = nullptr;
+	return true;
+}
 
-				if (Function && Function->GetNamedParameterOfType(ParameterInfo, Parameter, &ParameterOwner))
-				{
-					if (Parameter)
-					{
-						OutValue = Parameter->IsUsedAsAtlasPosition();
-						if (OutValue)
-						{
-							UMaterialExpressionCurveAtlasRowParameter* AtlasParameter = Cast<UMaterialExpressionCurveAtlasRowParameter>(Parameter);
-							if (AtlasParameter && AtlasParameter->ParameterName == ParameterInfo.Name)
-							{
-								Curve = TSoftObjectPtr<UCurveLinearColor>(FSoftObjectPath(AtlasParameter->Curve->GetPathName()));
-								Atlas = TSoftObjectPtr<UCurveLinearColorAtlas>(FSoftObjectPath(AtlasParameter->Atlas->GetPathName()));
-								return true;
-							}
-						}
-					}
-				}
-			}
-		}
+bool UMaterial::GetScalarParameterSliderMinMax(const FHashedMaterialParameterInfo& ParameterInfo, float& OutSliderMin, float& OutSliderMax) const
+{
+	const int32 Index = CachedExpressionData.Parameters.FindParameterIndex(EMaterialParameterType::Scalar, ParameterInfo);
+	if (Index == INDEX_NONE)
+	{
+		return false;
 	}
-
-	return false;
+	OutSliderMin = CachedExpressionData.Parameters.ScalarMinMaxValues[Index].X;
+	OutSliderMax = CachedExpressionData.Parameters.ScalarMinMaxValues[Index].Y;
+	return true;
 }
+#endif // WITH_EDITOR
 
-#if WITH_EDITOR
-bool UMaterial::GetScalarParameterSliderMinMax(const FMaterialParameterInfo& ParameterInfo, float& OutSliderMin, float& OutSliderMax) const
+bool UMaterial::GetVectorParameterValue(const FHashedMaterialParameterInfo& ParameterInfo, FLinearColor& OutValue, bool bOveriddenOnly) const
 {
-	// In the case of duplicate parameters with different values, this will return the
-	// first matching expression found, not necessarily the one that's used for rendering
-	UMaterialExpressionScalarParameter* Parameter = nullptr;
-	float TempValue = 0.f;
+	const bool bResult = GetVectorParameterValue_New(ParameterInfo, OutValue, bOveriddenOnly);
+#if WITH_EDITOR
+	FLinearColor OldValue(ForceInitToZero);
+	const bool bOldResult = GetVectorParameterValue_Legacy(ParameterInfo, OldValue, bOveriddenOnly);
 
-	for (UMaterialExpression* Expression : Expressions)
+	check(bOldResult == bResult);
+	check(!bResult || OutValue == OldValue);
+#endif
+	return bResult;
+}
+
+bool UMaterial::GetVectorParameterValue_New(const FHashedMaterialParameterInfo& ParameterInfo, FLinearColor& OutValue, bool bOveriddenOnly) const
+{
+	const int32 Index = CachedExpressionData.Parameters.FindParameterIndex(EMaterialParameterType::Vector, ParameterInfo, bOveriddenOnly);
+	if (Index == INDEX_NONE)
 	{
-		// Only need to check parameters that match in associated scope
-		if (ParameterInfo.Association == EMaterialParameterAssociation::GlobalParameter)
-		{
-			if (UMaterialExpressionScalarParameter* ExpressionParameter = Cast<UMaterialExpressionScalarParameter>(Expression))
-			{
-				if (ExpressionParameter->IsNamedParameter(ParameterInfo, TempValue))
-				{
-					OutSliderMin = ExpressionParameter->SliderMin;
-					OutSliderMax = ExpressionParameter->SliderMax;
-					return true;
-				}
-			}
-			else if (UMaterialExpressionMaterialFunctionCall* FunctionCall = Cast<UMaterialExpressionMaterialFunctionCall>(Expression))
-			{
-				UMaterialFunctionInterface* Function = FunctionCall->MaterialFunction;
-				if (Function && Function->GetNamedParameterOfType(ParameterInfo, Parameter))
-				{
-					OutSliderMin = Parameter->SliderMin;
-					OutSliderMax = Parameter->SliderMax;
-					return true;
-				}
-			}
-		}
-		else
-		{
-			if (UMaterialExpressionMaterialAttributeLayers* LayersExpression = Cast<UMaterialExpressionMaterialAttributeLayers>(Expression))
-			{
-				UMaterialFunctionInterface* Function = LayersExpression->GetParameterAssociatedFunction(ParameterInfo);
-				if (Function && Function->GetNamedParameterOfType(ParameterInfo, Parameter))
-				{
-					OutSliderMin = Parameter->SliderMin;
-					OutSliderMax = Parameter->SliderMax;
-					return true;
-				}
-			}
-		}
+		return false;
 	}
-
-	return false;
+	OutValue = CachedExpressionData.Parameters.VectorValues[Index];
+	return true;
 }
-#endif // WITH_EDITOR
 
-bool UMaterial::GetVectorParameterValue(const FMaterialParameterInfo& ParameterInfo, FLinearColor& OutValue, bool bOveriddenOnly) const
+#if WITH_EDITOR
+bool UMaterial::GetVectorParameterValue_Legacy(const FHashedMaterialParameterInfo& ParameterInfo, FLinearColor& OutValue, bool bOveriddenOnly) const
 {
 	if (bOveriddenOnly && !AreExperimentalMaterialLayersEnabled())
 	{
@@ -2420,113 +2302,57 @@ bool UMaterial::GetVectorParameterValue(const FMaterialParameterInfo& ParameterI
 	return false;
 }
 
-bool UMaterial::IsVectorParameterUsedAsChannelMask(const FMaterialParameterInfo& ParameterInfo, bool& OutValue) const
+bool UMaterial::IsVectorParameterUsedAsChannelMask(const FHashedMaterialParameterInfo& ParameterInfo, bool& OutValue) const
 {
-	// In the case of duplicate parameters with different values, this will return the
-	// first matching expression found, not necessarily the one that's used for rendering
-	UMaterialExpressionVectorParameter* Parameter = nullptr;
-	FLinearColor TempColor;
-
-	for (UMaterialExpression* Expression : Expressions)
+	const int32 Index = CachedExpressionData.Parameters.FindParameterIndex(EMaterialParameterType::Vector, ParameterInfo);
+	if (Index == INDEX_NONE)
 	{
-		// Only need to check parameters that match in associated scope
-		if (ParameterInfo.Association == EMaterialParameterAssociation::GlobalParameter)
-		{
-			if (UMaterialExpressionVectorParameter* ExpressionParameter = Cast<UMaterialExpressionVectorParameter>(Expression))
-			{
-				if (ExpressionParameter->IsNamedParameter(ParameterInfo, TempColor))
-				{
-					OutValue = ExpressionParameter->IsUsedAsChannelMask();
-					return true;
-				}
-			}
-			else if (UMaterialExpressionMaterialFunctionCall* FunctionCall = Cast<UMaterialExpressionMaterialFunctionCall>(Expression))
-			{
-				UMaterialFunctionInterface* Function = FunctionCall->MaterialFunction;
-				UMaterialFunctionInterface* ParameterOwner = nullptr;
-
-				if (Function && Function->GetNamedParameterOfType(ParameterInfo, Parameter, &ParameterOwner))
-				{
-					OutValue = Parameter->IsUsedAsChannelMask();
-					return true;
-				}
-			}
-		}
-		else
-		{
-			if (UMaterialExpressionMaterialAttributeLayers* LayersExpression = Cast<UMaterialExpressionMaterialAttributeLayers>(Expression))
-			{
-				UMaterialFunctionInterface* Function = LayersExpression->GetParameterAssociatedFunction(ParameterInfo);
-				UMaterialFunctionInterface* ParameterOwner = nullptr;
-
-				if (Function && Function->GetNamedParameterOfType(ParameterInfo, Parameter, &ParameterOwner))
-				{
-					OutValue = Parameter->IsUsedAsChannelMask();
-					return true;
-				}
-			}
-		}
+		return false;
 	}
-
-	return false;
+	OutValue = CachedExpressionData.Parameters.VectorUsedAsChannelMaskValues[Index];
+	return true;
 }
 
-#if WITH_EDITOR
-bool UMaterial::GetVectorParameterChannelNames(const FMaterialParameterInfo& ParameterInfo, FParameterChannelNames& OutValue) const
+bool UMaterial::GetVectorParameterChannelNames(const FHashedMaterialParameterInfo& ParameterInfo, FParameterChannelNames& OutValue) const
 {
-	// In the case of duplicate parameters with different values, this will return the
-	// first matching expression found, not necessarily the one that's used for rendering
-	UMaterialExpressionVectorParameter* Parameter = nullptr;
-	FLinearColor TempColor;
-
-	for (UMaterialExpression* Expression : Expressions)
+	const int32 Index = CachedExpressionData.Parameters.FindParameterIndex(EMaterialParameterType::Vector, ParameterInfo);
+	if (Index == INDEX_NONE)
 	{
-		// Only need to check parameters that match in associated scope
-		if (ParameterInfo.Association == EMaterialParameterAssociation::GlobalParameter)
-		{
-			if (UMaterialExpressionVectorParameter* ExpressionParameter = Cast<UMaterialExpressionVectorParameter>(Expression))
-			{
-				if (ExpressionParameter->IsNamedParameter(ParameterInfo, TempColor))
-				{
-					OutValue = ExpressionParameter->GetVectorChannelNames();
-					return true;
-				}
-			}
-			else if (UMaterialExpressionMaterialFunctionCall* FunctionCall = Cast<UMaterialExpressionMaterialFunctionCall>(Expression))
-			{
-				UMaterialFunctionInterface* Function = FunctionCall->MaterialFunction;
-				UMaterialFunctionInterface* ParameterOwner = nullptr;
+		return false;
+	}
+	OutValue = CachedExpressionData.Parameters.VectorChannelNameValues[Index];
+	return true;
+}
 
-				if (Function && Function->GetNamedParameterOfType(ParameterInfo, Parameter, &ParameterOwner))
-				{
-					OutValue = Parameter->GetVectorChannelNames();
-					return true;
-				}
-			}
-		}
-		else
-		{
-			if (UMaterialExpressionMaterialAttributeLayers* LayersExpression = Cast<UMaterialExpressionMaterialAttributeLayers>(Expression))
-			{
-				UMaterialFunctionInterface* Function = LayersExpression->GetParameterAssociatedFunction(ParameterInfo);
-				UMaterialFunctionInterface* ParameterOwner = nullptr;
+#endif // WITH_EDITOR
 
-				if (Function && Function->GetNamedParameterOfType(ParameterInfo, Parameter, &ParameterOwner))
-				{
-					OutValue = Parameter->GetVectorChannelNames();
-					return true;
-				}
-			}
-		}
-	}
 
-	return false;
-}
+bool UMaterial::GetTextureParameterValue(const FHashedMaterialParameterInfo& ParameterInfo, UTexture*& OutValue, bool bOveriddenOnly) const
+{
+	const bool bResult = GetTextureParameterValue_New(ParameterInfo, OutValue, bOveriddenOnly);
+#if WITH_EDITOR
+	UTexture* OldValue = nullptr;
+	const bool bOldResult = GetTextureParameterValue_Legacy(ParameterInfo, OldValue, bOveriddenOnly);
 
+	check(OutValue == OldValue);
+	check(bOldResult == bResult);
 #endif
+	return bResult;
+}
 
+bool UMaterial::GetTextureParameterValue_New(const FHashedMaterialParameterInfo& ParameterInfo, UTexture*& OutValue, bool bOveriddenOnly) const
+{
+	const int32 Index = CachedExpressionData.Parameters.FindParameterIndex(EMaterialParameterType::Texture, ParameterInfo, bOveriddenOnly);
+	if (Index == INDEX_NONE)
+	{
+		return false;
+	}
+	OutValue = CachedExpressionData.Parameters.TextureValues[Index];
+	return true;
+}
 
-bool UMaterial::GetTextureParameterValue(const FMaterialParameterInfo& ParameterInfo, UTexture*& OutValue, bool bOveriddenOnly) const
+#if WITH_EDITOR
+bool UMaterial::GetTextureParameterValue_Legacy(const FHashedMaterialParameterInfo& ParameterInfo, UTexture*& OutValue, bool bOveriddenOnly) const
 {
 	if (bOveriddenOnly && !AreExperimentalMaterialLayersEnabled())
 	{
@@ -2595,8 +2421,34 @@ bool UMaterial::GetTextureParameterValue(const FMaterialParameterInfo& Parameter
 
 	return false;
 }
+#endif // WITH_EDITOR
+
+bool UMaterial::GetRuntimeVirtualTextureParameterValue(const FHashedMaterialParameterInfo& ParameterInfo, URuntimeVirtualTexture*& OutValue, bool bOveriddenOnly) const
+{
+	const bool bResult = GetRuntimeVirtualTextureParameterValue_New(ParameterInfo, OutValue, bOveriddenOnly);
+#if WITH_EDITOR
+	URuntimeVirtualTexture* OldValue = nullptr;
+	const bool bOldResult = GetRuntimeVirtualTextureParameterValue_Legacy(ParameterInfo, OldValue, bOveriddenOnly);
+
+	check(bOldResult == bResult);
+	check(!bResult || OutValue == OldValue);
+#endif
+	return bResult;
+}
 
-bool UMaterial::GetRuntimeVirtualTextureParameterValue(const FMaterialParameterInfo& ParameterInfo, URuntimeVirtualTexture*& OutValue, bool bOveriddenOnly) const
+bool UMaterial::GetRuntimeVirtualTextureParameterValue_New(const FHashedMaterialParameterInfo& ParameterInfo, URuntimeVirtualTexture*& OutValue, bool bOveriddenOnly) const
+{
+	const int32 Index = CachedExpressionData.Parameters.FindParameterIndex(EMaterialParameterType::RuntimeVirtualTexture, ParameterInfo, bOveriddenOnly);
+	if (Index == INDEX_NONE)
+	{
+		return false;
+	}
+	OutValue = CachedExpressionData.Parameters.RuntimeVirtualTextureValues[Index];
+	return true;
+}
+
+#if WITH_EDITOR
+bool UMaterial::GetRuntimeVirtualTextureParameterValue_Legacy(const FHashedMaterialParameterInfo& ParameterInfo, URuntimeVirtualTexture*& OutValue, bool bOveriddenOnly) const
 {
 	if (bOveriddenOnly && !AreExperimentalMaterialLayersEnabled())
 	{
@@ -2649,67 +2501,14 @@ bool UMaterial::GetRuntimeVirtualTextureParameterValue(const FMaterialParameterI
 				if (Function && Function->OverrideNamedRuntimeVirtualTextureParameter(ParameterInfo, OutValue))
 				{
 					return true;
-				}
-
-				if (Function && Function->GetNamedParameterOfType(ParameterInfo, Parameter, &ParameterOwner))
-				{
-					if (!ParameterOwner->OverrideNamedRuntimeVirtualTextureParameter(ParameterInfo, OutValue))
-					{
-						Parameter->IsNamedParameter(ParameterInfo, OutValue);
-					}
-					return true;
-				}
-			}
-		}
-	}
-
-	return false;
-}
-
-
-#if WITH_EDITOR
-bool UMaterial::GetTextureParameterChannelNames(const FMaterialParameterInfo& ParameterInfo, FParameterChannelNames& OutValue) const
-{
-	// In the case of duplicate parameters with different values, this will return the
-	// first matching expression found, not necessarily the one that's used for rendering
-	UMaterialExpressionTextureSampleParameter* Parameter = nullptr;
-	UTexture* TempTexture;
-
-	for (UMaterialExpression* Expression : Expressions)
-	{
-		// Only need to check parameters that match in associated scope
-		if (ParameterInfo.Association == EMaterialParameterAssociation::GlobalParameter)
-		{
-			if (UMaterialExpressionTextureSampleParameter* ExpressionParameter = Cast<UMaterialExpressionTextureSampleParameter>(Expression))
-			{
-				if (ExpressionParameter->IsNamedParameter(ParameterInfo, TempTexture))
-				{
-					OutValue = ExpressionParameter->GetTextureChannelNames();
-					return true;
-				}
-			}
-			else if (UMaterialExpressionMaterialFunctionCall* FunctionCall = Cast<UMaterialExpressionMaterialFunctionCall>(Expression))
-			{
-				UMaterialFunctionInterface* Function = FunctionCall->MaterialFunction;
-				UMaterialFunctionInterface* ParameterOwner = nullptr;
-
-				if (Function && Function->GetNamedParameterOfType(ParameterInfo, Parameter, &ParameterOwner))
-				{
-					OutValue = Parameter->GetTextureChannelNames();
-					return true;
-				}
-			}
-		}
-		else
-		{
-			if (UMaterialExpressionMaterialAttributeLayers* LayersExpression = Cast<UMaterialExpressionMaterialAttributeLayers>(Expression))
-			{
-				UMaterialFunctionInterface* Function = LayersExpression->GetParameterAssociatedFunction(ParameterInfo);
-				UMaterialFunctionInterface* ParameterOwner = nullptr;
+				}
 
 				if (Function && Function->GetNamedParameterOfType(ParameterInfo, Parameter, &ParameterOwner))
 				{
-					OutValue = Parameter->GetTextureChannelNames();
+					if (!ParameterOwner->OverrideNamedRuntimeVirtualTextureParameter(ParameterInfo, OutValue))
+					{
+						Parameter->IsNamedParameter(ParameterInfo, OutValue);
+					}
 					return true;
 				}
 			}
@@ -2718,10 +2517,49 @@ bool UMaterial::GetTextureParameterChannelNames(const FMaterialParameterInfo& Pa
 
 	return false;
 }
+#endif // WITH_EDITOR
+
+#if WITH_EDITOR
+bool UMaterial::GetTextureParameterChannelNames(const FHashedMaterialParameterInfo& ParameterInfo, FParameterChannelNames& OutValue) const
+{
+	const int32 Index = CachedExpressionData.Parameters.FindParameterIndex(EMaterialParameterType::Texture, ParameterInfo);
+	if (Index == INDEX_NONE)
+	{
+		return false;
+	}
+	OutValue = CachedExpressionData.Parameters.TextureChannelNameValues[Index];
+	return true;
+}
+#endif // WITH_EDITOR
+
+bool UMaterial::GetFontParameterValue(const FHashedMaterialParameterInfo& ParameterInfo, UFont*& OutFontValue, int32& OutFontPage, bool bOveriddenOnly) const
+{
+	const bool bResult = GetFontParameterValue_New(ParameterInfo, OutFontValue, OutFontPage, bOveriddenOnly);
+#if WITH_EDITOR
+	UFont* OldValue = nullptr;
+	int32 OldPage = INDEX_NONE;
+	const bool bOldResult = GetFontParameterValue_Legacy(ParameterInfo, OldValue, OldPage, bOveriddenOnly);
 
+	check(bOldResult == bResult);
+	check(!bResult || (OutFontValue == OldValue && OldPage == OutFontPage));
 #endif
+	return bResult;
+}
+
+bool UMaterial::GetFontParameterValue_New(const FHashedMaterialParameterInfo& ParameterInfo, UFont*& OutFontValue, int32& OutFontPage, bool bOveriddenOnly) const
+{
+	const int32 Index = CachedExpressionData.Parameters.FindParameterIndex(EMaterialParameterType::Font, ParameterInfo, bOveriddenOnly);
+	if (Index == INDEX_NONE)
+	{
+		return false;
+	}
+	OutFontValue = CachedExpressionData.Parameters.FontValues[Index];
+	OutFontPage = CachedExpressionData.Parameters.FontPageValues[Index];
+	return true;
+}
 
-bool UMaterial::GetFontParameterValue(const FMaterialParameterInfo& ParameterInfo, UFont*& OutFontValue, int32& OutFontPage, bool bOveriddenOnly) const
+#if WITH_EDITOR
+bool UMaterial::GetFontParameterValue_Legacy(const FHashedMaterialParameterInfo& ParameterInfo, UFont*& OutFontValue, int32& OutFontPage, bool bOveriddenOnly) const
 {
 	// In the case of duplicate parameters with different values, this will return the
 	// first matching expression found, not necessarily the one that's used for rendering
@@ -2787,7 +2625,9 @@ bool UMaterial::GetFontParameterValue(const FMaterialParameterInfo& ParameterInf
 
 	return false;
 }
+#endif // WITH_EDITOR
 
+#if WITH_EDITORONLY_DATA
 bool UMaterial::GetStaticSwitchParameterValues(FStaticParamEvaluationContext& EvalContext, TBitArray<>& OutValues, FGuid* OutExpressionGuids, bool bCheckParent /*= true*/) const
 {
 	if (EvalContext.AllResolved())
@@ -2802,7 +2642,7 @@ bool UMaterial::GetStaticSwitchParameterValues(FStaticParamEvaluationContext& Ev
 	{
 		TMaterialStaticParameterValueResolver<UMaterialExpressionStaticBoolParameter, TBitArray<>> ValueResolver(EvalContext, OutValues, OutExpressionGuids);
 
-		EvalContext.ForEachPendingParameter([&EvalContext, Expression, &ValueResolver, &OutValues, &OutExpressionGuids](int32 ParamIndex, const FMaterialParameterInfo& ParameterInfo) -> bool
+		EvalContext.ForEachPendingParameter([&EvalContext, Expression, &ValueResolver, &OutValues, &OutExpressionGuids](int32 ParamIndex, const FHashedMaterialParameterInfo& ParameterInfo) -> bool
 			{
 				if (ParameterInfo.Association == EMaterialParameterAssociation::GlobalParameter)
 				{
@@ -2855,7 +2695,7 @@ bool UMaterial::GetStaticComponentMaskParameterValues(FStaticParamEvaluationCont
 	{
 		TMaterialStaticParameterValueResolver<UMaterialExpressionStaticComponentMaskParameter, TBitArray<>> ValueResolver(EvalContext, OutRGBAOrderedValues, OutExpressionGuids);
 
-		EvalContext.ForEachPendingParameter([&EvalContext, Expression, &ValueResolver, &OutRGBAOrderedValues, &OutExpressionGuids](int32 ParamIndex, const FMaterialParameterInfo& ParameterInfo) -> bool
+		EvalContext.ForEachPendingParameter([&EvalContext, Expression, &ValueResolver, &OutRGBAOrderedValues, &OutExpressionGuids](int32 ParamIndex, const FHashedMaterialParameterInfo& ParameterInfo) -> bool
 			{
 				if (ParameterInfo.Association == EMaterialParameterAssociation::GlobalParameter)
 				{
@@ -2898,17 +2738,9 @@ bool UMaterial::GetStaticComponentMaskParameterValues(FStaticParamEvaluationCont
 	return EvalContext.AllResolved();
 }
 
-bool UMaterial::GetTerrainLayerWeightParameterValue(const FMaterialParameterInfo& ParameterInfo, int32& OutWeightmapIndex, FGuid& OutExpressionGuid) const
-{
-	bool bSuccess = false;
-	OutWeightmapIndex = INDEX_NONE;
-	bSuccess = true;
-	return bSuccess;
-}
-
-bool UMaterial::GetMaterialLayersParameterValue(const FMaterialParameterInfo& ParameterInfo, FMaterialLayersFunctions& OutLayers, FGuid& OutExpressionGuid, bool bCheckParent /*= true*/) const
+bool UMaterial::GetMaterialLayersParameterValue(const FHashedMaterialParameterInfo& ParameterInfo, FMaterialLayersFunctions& OutLayers, FGuid& OutExpressionGuid, bool bCheckParent /*= true*/) const
 {
-	UMaterialExpressionStaticComponentMaskParameter* Parameter = nullptr;	
+	UMaterialExpressionStaticComponentMaskParameter* Parameter = nullptr;
 
 	for (UMaterialExpression* Expression : Expressions)
 	{
@@ -2919,13 +2751,22 @@ bool UMaterial::GetMaterialLayersParameterValue(const FMaterialParameterInfo& Pa
 			{
 				return true;
 			}
-		}	
+		}
 	}
 
 	OutLayers.Layers.Empty();
 	OutLayers.Blends.Empty();
 	return false;
 }
+#endif // WITH_EDITORONLY_DATA
+
+bool UMaterial::GetTerrainLayerWeightParameterValue(const FHashedMaterialParameterInfo& ParameterInfo, int32& OutWeightmapIndex, FGuid& OutExpressionGuid) const
+{
+	bool bSuccess = false;
+	OutWeightmapIndex = INDEX_NONE;
+	bSuccess = true;
+	return bSuccess;
+}
 
 bool UMaterial::GetRefractionSettings(float& OutBiasValue) const
 {
@@ -3011,6 +2852,7 @@ static FAutoConsoleVariable GCompileMaterialsForShaderFormatCVar(
 	TEXT("Note that these shaders are compiled and immediately tossed. This is only useful when directly inspecting output via r.DebugDumpShaderInfo.")
 	);
 
+#if WITH_EDITOR
 void UMaterial::GetForceRecompileTextureIdsHash(FSHAHash &TextureReferencesHash)
 {
 	TArray<UTexture*> ForceRecompileTextures;
@@ -3068,8 +2910,6 @@ bool UMaterial::IsTextureForceRecompileCacheRessource(UTexture *Texture)
 	return false;
 }
 
-#if WITH_EDITOR
-
 void UMaterial::UpdateMaterialShaderCacheAndTextureReferences()
 {
 	// If the material changes, then the debug view material must reset to prevent parameters mismatch
@@ -3228,8 +3068,9 @@ void UMaterial::CacheResourceShadersForCooking(EShaderPlatform ShaderPlatform, T
 
 void UMaterial::CacheShadersForResources(EShaderPlatform ShaderPlatform, const TArray<FMaterialResource*>& ResourcesToCache, const ITargetPlatform* TargetPlatform)
 {
-	RebuildExpressionTextureReferences();
-
+#if WITH_EDITOR
+	UpdateCachedExpressionData();
+#endif
 	for (int32 ResourceIndex = 0; ResourceIndex < ResourcesToCache.Num(); ResourceIndex++)
 	{
 		FMaterialResource* CurrentResource = ResourcesToCache[ResourceIndex];
@@ -3276,135 +3117,6 @@ void UMaterial::FlushResourceShaderMaps()
 	}
 }
 
-#if WITH_EDITOR
-
-void UMaterial::RebuildMaterialFunctionInfo()
-{	
-	MaterialFunctionInfos.Reset();
-
-	for (int32 ExpressionIndex = 0; ExpressionIndex < Expressions.Num(); ExpressionIndex++)
-	{
-		UMaterialExpression* Expression = Expressions[ExpressionIndex];
-
-		if (UMaterialExpressionMaterialFunctionCall* MaterialFunctionNode = Cast<UMaterialExpressionMaterialFunctionCall>(Expression))
-		{
-			if (MaterialFunctionNode->MaterialFunction)
-			{
-				MaterialFunctionNode->IterateDependentFunctions(
-					[this](UMaterialFunctionInterface* Function)
-				{
-					FMaterialFunctionInfo NewFunctionInfo;
-					if (Function)
-					{
-						NewFunctionInfo.Function = Function;
-						NewFunctionInfo.StateId = Function->StateId;
-					}
-					MaterialFunctionInfos.Add(NewFunctionInfo);
-					return true;
-				});
-			}
-
-			// Update the function call node, so it can relink inputs and outputs as needed
-			// Update even if MaterialFunctionNode->MaterialFunction is NULL, because we need to remove the invalid inputs in that case
-			MaterialFunctionNode->UpdateFromFunctionResource();
-		}
-		else if (UMaterialExpressionMaterialAttributeLayers* MaterialLayersNode = Cast<UMaterialExpressionMaterialAttributeLayers>(Expression))
-		{
-			MaterialLayersNode->IterateDependentFunctions(
-				[this](UMaterialFunctionInterface* Function)
-			{
-				FMaterialFunctionInfo NewFunctionInfo;
-				if (Function)
-				{
-					NewFunctionInfo.Function = Function;
-					NewFunctionInfo.StateId = Function->StateId;
-
-					// Update active functions, so they can relink child function inputs and outputs as needed
-					Function->UpdateFromFunctionResource();
-				}
-				MaterialFunctionInfos.Add(NewFunctionInfo);
-				return true;
-			});
-
-			MaterialLayersNode->RebuildLayerGraph(false);
-		}
-	}
-}
-
-void UMaterial::RebuildMaterialParameterCollectionInfo()
-{
-	MaterialParameterCollectionInfos.Empty();
-
-	for (int32 ExpressionIndex = 0; ExpressionIndex < Expressions.Num(); ExpressionIndex++)
-	{
-		UMaterialExpression* Expression = Expressions[ExpressionIndex];
-		UMaterialExpressionCollectionParameter* CollectionParameter = Cast<UMaterialExpressionCollectionParameter>(Expression);
-		UMaterialExpressionMaterialFunctionCall* MaterialFunctionNode = Cast<UMaterialExpressionMaterialFunctionCall>(Expression);
-		UMaterialExpressionMaterialAttributeLayers* MaterialLayersNode = Cast<UMaterialExpressionMaterialAttributeLayers>(Expression);
-
-		if (CollectionParameter && CollectionParameter->Collection)
-		{
-			FMaterialParameterCollectionInfo NewInfo;
-			NewInfo.ParameterCollection = CollectionParameter->Collection;
-			NewInfo.StateId = CollectionParameter->Collection->StateId;
-			MaterialParameterCollectionInfos.AddUnique(NewInfo);
-		}
-		else if (MaterialFunctionNode && MaterialFunctionNode->MaterialFunction)
-		{
-			TArray<UMaterialFunctionInterface*> DependentFunctions;
-			MaterialFunctionNode->GetDependentFunctions(DependentFunctions);
-
-			// Handle nested functions
-			for (UMaterialFunctionInterface* CurrentFunction : DependentFunctions)
-			{
-				for (UMaterialExpression* CurrentExpression : *CurrentFunction->GetFunctionExpressions())
-				{
-					UMaterialExpressionCollectionParameter* FunctionCollectionParameter = Cast<UMaterialExpressionCollectionParameter>(CurrentExpression);
-
-					if (FunctionCollectionParameter && FunctionCollectionParameter->Collection)
-					{
-						FMaterialParameterCollectionInfo NewInfo;
-						NewInfo.ParameterCollection = FunctionCollectionParameter->Collection;
-						NewInfo.StateId = FunctionCollectionParameter->Collection->StateId;
-						MaterialParameterCollectionInfos.AddUnique(NewInfo);
-					}
-				}
-			}
-		}
-		else if (MaterialLayersNode)
-		{
-			TArray<UMaterialFunctionInterface*> DependentFunctions;
-			MaterialLayersNode->GetDependentFunctions(DependentFunctions);
-
-			for (UMaterialFunctionInterface* CurrentFunction : DependentFunctions)
-			{
-				for (UMaterialExpression* CurrentExpression : *CurrentFunction->GetFunctionExpressions())
-				{
-					UMaterialExpressionCollectionParameter* FunctionCollectionParameter = Cast<UMaterialExpressionCollectionParameter>(CurrentExpression);
-
-					if (FunctionCollectionParameter && FunctionCollectionParameter->Collection)
-					{
-						FMaterialParameterCollectionInfo NewInfo;
-						NewInfo.ParameterCollection = FunctionCollectionParameter->Collection;
-						NewInfo.StateId = FunctionCollectionParameter->Collection->StateId;
-						MaterialParameterCollectionInfos.AddUnique(NewInfo);
-					}
-				}
-			}
-		}
-	}
-}
-
-#endif // WITH_EDITOR
-
-void UMaterial::CacheExpressionTextureReferences()
-{
-	if ( ExpressionTextureReferences.Num() <= 0 )
-	{
-		RebuildExpressionTextureReferences();
-	}
-}
-
 bool UMaterial::AttemptInsertNewGroupName(const FString & InNewName)
 {
 #if WITH_EDITOR
@@ -3425,24 +3137,6 @@ bool UMaterial::AttemptInsertNewGroupName(const FString & InNewName)
 	return false;
 }
 
-void UMaterial::RebuildExpressionTextureReferences()
-{
-#if WITH_EDITOR
-	// Note: builds without editor only data will have an incorrect shader map id due to skipping this
-	// That's ok, FMaterial::CacheShaders handles this 
-	if (FPlatformProperties::HasEditorOnlyData())
-	{
-		// Rebuild all transient cached material properties which are based off of the editor-only data (expressions) and need to be up to date for compiling
-		// Update the cached material function information, which will store off information about the functions this material uses
-		RebuildMaterialFunctionInfo();
-		RebuildMaterialParameterCollectionInfo();
-	}
-#endif // WITH_EDITOR
-
-	ExpressionTextureReferences.Empty();
-	AppendReferencedTextures(ExpressionTextureReferences);
-}
-
 FMaterialResource* UMaterial::GetMaterialResource(ERHIFeatureLevel::Type InFeatureLevel, EMaterialQualityLevel::Type QualityLevel)
 {
 	return const_cast<FMaterialResource*>(static_cast<const UMaterial*>(this)->GetMaterialResource(InFeatureLevel, QualityLevel));
@@ -3473,24 +3167,6 @@ void UMaterial::Serialize(FArchive& Ar)
 
 	Super::Serialize(Ar);
 
-	if(Ar.IsLoading())
-	{
-#if WITH_EDITOR
-		bool bCooked = FPlatformProperties::RequiresCookedData() || GetOutermost()->bIsCookedForEditor;
-#else
-		bool bCooked = FPlatformProperties::RequiresCookedData();
-#endif
-		if (bCooked)
-		{
-			Expressions.Remove(nullptr);
-		}
-		else if (IsRunningCommandlet())
-		{
-			// invalidate the cache of quality level to make sure they are rebuilt. GetQualityLevelNodeUsage() is not going to recalculate it otherwise
-			CachedQualityLevelsUsed.Empty();
-		}
-	}
-
 	if (Ar.UE4Ver() >= VER_UE4_PURGED_FMATERIAL_COMPILE_OUTPUTS)
 	{
 #if WITH_EDITOR
@@ -3566,13 +3242,6 @@ void UMaterial::Serialize(FArchive& Ar)
 		SaveShaderStableKeys(Ar.CookingTarget());
 	}
 #endif
-
-	if (Ar.CustomVer(FFortniteMainBranchObjectVersion::GUID) < FFortniteMainBranchObjectVersion::CachedMaterialQualityNodeUsage)
-	{
-		TArray<bool, TInlineAllocator<EMaterialQualityLevel::Num>> QualityLevelsUsed;
-		GetQualityLevelNodeUsage(QualityLevelsUsed);
-		CachedQualityLevelsUsed = QualityLevelsUsed;
-	}
 }
 
 void UMaterial::PostDuplicate(bool bDuplicateForPIE)
@@ -3625,31 +3294,11 @@ void UMaterial::BackwardsCompatibilityInputConversion()
 
 void UMaterial::GetQualityLevelUsage(TArray<bool, TInlineAllocator<EMaterialQualityLevel::Num> >& OutQualityLevelsUsed, EShaderPlatform ShaderPlatform)
 {
-#if WITH_EDITOR
-	// Assume that a commandlet cannot change nodes once Cached levels are set
-	// FIXME: assumption may not work for a commandlet that edits the material graph?
-	if (IsRunningCommandlet() && CachedQualityLevelsUsed.Num() == EMaterialQualityLevel::Num)
-	{
-		OutQualityLevelsUsed = CachedQualityLevelsUsed;
-	}
-	else
-	{
-		GetQualityLevelNodeUsage(OutQualityLevelsUsed);
-		CachedQualityLevelsUsed = OutQualityLevelsUsed;
-	}
-#else
-	if (CachedQualityLevelsUsed.Num() != EMaterialQualityLevel::Num)
-	{
-		// Allow an initial call to cache the data during PostInit
-		// which will return default values and be updated later
-		GetQualityLevelNodeUsage(OutQualityLevelsUsed);
-	}
-	else
+	OutQualityLevelsUsed = CachedExpressionData.QualityLevelsUsed;
+	if (OutQualityLevelsUsed.Num() == 0)
 	{
-		OutQualityLevelsUsed = CachedQualityLevelsUsed;
+		OutQualityLevelsUsed.AddDefaulted(EMaterialQualityLevel::Num);
 	}
-#endif
-
 	if (ShaderPlatform != SP_NumPlatforms)
 	{
 		const UShaderPlatformQualitySettings* MaterialQualitySettings = UMaterialShaderQualitySettings::Get()->GetShaderPlatformQualitySettings(ShaderPlatform);
@@ -3677,109 +3326,6 @@ void UMaterial::GetQualityLevelUsage(TArray<bool, TInlineAllocator<EMaterialQual
 	}
 }
 
-void UMaterial::GetQualityLevelNodeUsage(TArray<bool, TInlineAllocator<EMaterialQualityLevel::Num> >& OutQualityLevelsUsed)
-{
-	OutQualityLevelsUsed.AddZeroed(EMaterialQualityLevel::Num);
-
-	for (int32 ExpressionIndex = 0, NumExpressions = Expressions.Num(); ExpressionIndex < NumExpressions; ExpressionIndex++)
-	{
-		UMaterialExpression* Expression = Expressions[ExpressionIndex];
-		UMaterialExpressionQualitySwitch* QualitySwitchNode = Cast<UMaterialExpressionQualitySwitch>(Expression);
-		UMaterialExpressionMaterialFunctionCall* MaterialFunctionNode = Cast<UMaterialExpressionMaterialFunctionCall>(Expression);
-		UMaterialExpressionMaterialAttributeLayers* MaterialLayersNode = Cast<UMaterialExpressionMaterialAttributeLayers>(Expression);
-
-		if (QualitySwitchNode)
-		{
-			for (int32 InputIndex = 0; InputIndex < EMaterialQualityLevel::Num; InputIndex++)
-			{
-				if (QualitySwitchNode->Inputs[InputIndex].IsConnected())
-				{
-					OutQualityLevelsUsed[InputIndex] = true;
-				}
-			}
-
-			if (QualitySwitchNode->Default.IsConnected())
-			{
-				OutQualityLevelsUsed[EMaterialQualityLevel::High] = true;
-			}
-		}
-		else if (MaterialFunctionNode && MaterialFunctionNode->MaterialFunction)
-		{
-			TArray<UMaterialFunctionInterface*> DependentFunctions;
-			MaterialFunctionNode->GetDependentFunctions(DependentFunctions);
-
-			// Handle nested functions
-			for (UMaterialFunctionInterface* CurrentFunction : DependentFunctions)
-			{
-				for (UMaterialExpression* CurrentExpression : *CurrentFunction->GetFunctionExpressions())
-				{
-					UMaterialExpressionQualitySwitch* SwitchNode = Cast<UMaterialExpressionQualitySwitch>(CurrentExpression);
-
-					if (SwitchNode)
-					{
-						for (int32 InputIndex = 0; InputIndex < EMaterialQualityLevel::Num; InputIndex++)
-						{
-							if (SwitchNode->Inputs[InputIndex].IsConnected())
-							{
-								OutQualityLevelsUsed[InputIndex] = true;
-							}
-						}
-
-						if (SwitchNode->Default.IsConnected())
-						{
-							OutQualityLevelsUsed[EMaterialQualityLevel::High] = true;
-						}
-					}
-				}
-			}
-		}
-		else if (MaterialLayersNode)
-		{
-			TArray<UMaterialFunctionInterface*> DependentFunctions;
-			MaterialLayersNode->GetDependentFunctions(DependentFunctions);
-
-			for (UMaterialFunctionInterface* CurrentFunction : DependentFunctions)
-			{
-				for (UMaterialExpression* CurrentExpression : *CurrentFunction->GetFunctionExpressions())
-				{
-					if (UMaterialExpressionQualitySwitch* SwitchNode = Cast<UMaterialExpressionQualitySwitch>(CurrentExpression))
-					{
-						for (int32 InputIndex = 0; InputIndex < EMaterialQualityLevel::Num; InputIndex++)
-						{
-							if (SwitchNode->Inputs[InputIndex].IsConnected())
-							{
-								OutQualityLevelsUsed[InputIndex] = true;
-							}
-						}
-
-						if (SwitchNode->Default.IsConnected())
-						{
-							OutQualityLevelsUsed[EMaterialQualityLevel::High] = true;
-						}
-					}
-				}
-			}
-		}
-	}
-
-	// Always append a default high quality level if nothing else is set
-	
-	bool bAnyQualityLevelsSet = false;
-	for (bool Value : OutQualityLevelsUsed)
-	{
-		if (Value)
-		{
-			bAnyQualityLevelsSet = true;
-			break;
-		}
-	}
-
-	if (!bAnyQualityLevelsSet)
-	{
-		OutQualityLevelsUsed[EMaterialQualityLevel::High] = true;
-	}
-}
-
 void UMaterial::UpdateResourceAllocations(FMaterialResourceDeferredDeletionArray* ResourcesToFree)
 {
 	if (FApp::CanEverRender())
@@ -3852,6 +3398,10 @@ void UMaterial::PostLoad()
 
 	Super::PostLoad();
 
+#if WITH_EDITOR
+	UpdateCachedExpressionData();
+#endif // WITH_EDITOR
+
 	if (FApp::CanEverRender())
 	{
 		// Resources can be processed / registered now that we're back on the main thread
@@ -3916,6 +3466,7 @@ void UMaterial::PostLoad()
 		MaterialDomain = MD_UI;
 	}
 
+#if WITH_EDITORONLY_DATA
 	// Ensure expressions have been postloaded before we use them for compiling
 	// Any UObjects used by material compilation must be postloaded here
 	for (int32 ExpressionIndex = 0; ExpressionIndex < Expressions.Num(); ExpressionIndex++)
@@ -3925,12 +3476,13 @@ void UMaterial::PostLoad()
 			Expressions[ExpressionIndex]->ConditionalPostLoad();
 		}
 	}
+#endif // WITH_EDITORONLY_DATA
 
-	for (int32 CollectionIndex = 0; CollectionIndex < MaterialParameterCollectionInfos.Num(); CollectionIndex++)
+	for (int32 CollectionIndex = 0; CollectionIndex < CachedExpressionData.ParameterCollectionInfos.Num(); CollectionIndex++)
 	{
-		if (MaterialParameterCollectionInfos[CollectionIndex].ParameterCollection)
+		if (CachedExpressionData.ParameterCollectionInfos[CollectionIndex].ParameterCollection)
 		{
-			MaterialParameterCollectionInfos[CollectionIndex].ParameterCollection->ConditionalPostLoad();
+			CachedExpressionData.ParameterCollectionInfos[CollectionIndex].ParameterCollection->ConditionalPostLoad();
 		}
 	}
 
@@ -4048,8 +3600,7 @@ void UMaterial::PostLoad()
 		{
 			// Before caching shader resources we have to make sure all referenced textures have been post loaded
 			// as we depend on their resources being valid.
-			RebuildExpressionTextureReferences();
-			for (UObject* Texture : ExpressionTextureReferences)
+			for (UObject* Texture : CachedExpressionData.ReferencedTextures)
 			{
 				if (Texture)
 				{
@@ -4447,6 +3998,8 @@ void UMaterial::PostEditChangeProperty(FPropertyChangedEvent& PropertyChangedEve
 	
 	if (bRequiresCompilation)
 	{
+		UpdateCachedExpressionData();
+
 		// When redirecting an object pointer, we trust that the DDC hash will detect the change and that we don't need to force a recompile.
 		const bool bRegenerateId = PropertyChangedEvent.ChangeType != EPropertyChangeType::Redirected;
 		CacheResourceShadersForRendering(bRegenerateId);
@@ -4632,7 +4185,6 @@ bool UMaterial::HasDuplicateDynamicParameters(const UMaterialExpression* Express
 	}
 	return false;
 }
-#endif // WITH_EDITOR
 
 void UMaterial::UpdateExpressionDynamicParameters(const UMaterialExpression* Expression)
 {
@@ -4644,15 +4196,12 @@ void UMaterial::UpdateExpressionDynamicParameters(const UMaterialExpression* Exp
 			UMaterialExpressionDynamicParameter* CheckParam = Cast<UMaterialExpressionDynamicParameter>(Expressions[ExpIndex]);
 			if (CheckParam && CheckParam->CopyDynamicParameterProperties(DynParam))
 			{
-#if WITH_EDITOR
 				CheckParam->GraphNode->ReconstructNode();
-#endif // WITH_EDITOR
 			}
 		}
 	}
 }
 
-#if WITH_EDITOR
 void UMaterial::PropagateExpressionParameterChanges(UMaterialExpression* Parameter)
 {
 	FName ParmName;
@@ -4935,10 +4484,13 @@ bool UMaterial::CanBeClusterRoot() const
 
 void UMaterial::GetAssetRegistryTags(TArray<FAssetRegistryTag>& OutTags) const
 {
+#if WITH_EDITORONLY_DATA
 	// Search for any scene color nodes
 	bool bHasSceneColor = HasAnyExpressionsInMaterialAndFunctionsOfType<UMaterialExpressionSceneColor>();
 	OutTags.Add(FAssetRegistryTag("HasSceneColor", bHasSceneColor ? TEXT("True") : TEXT("False"), FAssetRegistryTag::TT_Alphabetical));
-
+#else
+	check(false); // TODO Do we need to handle this??
+#endif
 	Super::GetAssetRegistryTags(OutTags);
 }
 
@@ -4955,7 +4507,7 @@ void UMaterial::CancelOutstandingCompilation()
 }
 #endif
 
-void UMaterial::UpdateMaterialShaders(TArray<FShaderType*>& ShaderTypesToFlush, TArray<const FShaderPipelineType*>& ShaderPipelineTypesToFlush, TArray<const FVertexFactoryType*>& VFTypesToFlush, EShaderPlatform ShaderPlatform)
+void UMaterial::UpdateMaterialShaders(TArray<const FShaderType*>& ShaderTypesToFlush, TArray<const FShaderPipelineType*>& ShaderPipelineTypesToFlush, TArray<const FVertexFactoryType*>& VFTypesToFlush, EShaderPlatform ShaderPlatform)
 {
 	// Create a material update context so we can safely update materials.
 	{
@@ -5233,8 +4785,9 @@ FExpressionInput* UMaterial::GetExpressionInputForProperty(EMaterialProperty InP
 
 	return nullptr;
 }
-#endif
+#endif // WITH_EDITOR
 
+#if WITH_EDITORONLY_DATA
 void UMaterial::GetAllCustomOutputExpressions(TArray<class UMaterialExpressionCustomOutput*>& OutCustomOutputs) const
 {
 	for (UMaterialExpression* Expression : Expressions)
@@ -5260,6 +4813,7 @@ void UMaterial::GetAllExpressionsForCustomInterpolators(TArray<class UMaterialEx
 		}
 	}
 }
+#endif // WITH_EDITORONLY_DATA
 
 #if WITH_EDITOR
 bool UMaterial::GetAllReferencedExpressions(TArray<UMaterialExpression*>& OutExpressions, struct FStaticParameterSet* InStaticParameterSet,
@@ -5364,7 +4918,7 @@ bool UMaterial::GetExpressionsInPropertyChain(EMaterialProperty InProperty,
 	return true;
 }
 
-bool UMaterial::GetParameterSortPriority(const FMaterialParameterInfo& ParameterInfo, int32& OutSortPriority, const TArray<FStaticMaterialLayersParameter>* MaterialLayersParameters) const
+bool UMaterial::GetParameterSortPriority(const FHashedMaterialParameterInfo& ParameterInfo, int32& OutSortPriority, const TArray<FStaticMaterialLayersParameter>* MaterialLayersParameters) const
 {
 	for (UMaterialExpression* Expression : Expressions)
 	{
@@ -5655,50 +5209,6 @@ bool UMaterial::RecursiveGetExpressionChain(UMaterialExpression* InExpression, T
 }
 #endif // WITH_EDITOR
 
-void UMaterial::AppendReferencedTextures(TArray<UObject*>& InOutTextures) const
-{
-	for (int32 ExpressionIndex = 0; ExpressionIndex < Expressions.Num(); ExpressionIndex++)
-	{
-		UMaterialExpression* Expression = Expressions[ExpressionIndex];
-		UMaterialExpressionMaterialFunctionCall* MaterialFunctionNode = Cast<UMaterialExpressionMaterialFunctionCall>(Expression);
-		UMaterialExpressionMaterialAttributeLayers* MaterialLayersNode = Cast<UMaterialExpressionMaterialAttributeLayers>(Expression);
-
-		if (MaterialFunctionNode && MaterialFunctionNode->MaterialFunction)
-		{
-			TArray<UMaterialFunctionInterface*> DependentFunctions;
-			MaterialFunctionNode->GetDependentFunctions(DependentFunctions);
-
-			// Handle nested functions
-			for (UMaterialFunctionInterface* CurrentFunction : DependentFunctions)
-			{
-				CurrentFunction->AppendReferencedTextures(InOutTextures);
-			}
-		}
-		else if (MaterialLayersNode)
-		{
-			TArray<UMaterialFunctionInterface*> DependentFunctions;
-			MaterialLayersNode->GetDependentFunctions(DependentFunctions);
-
-			for (UMaterialFunctionInterface* CurrentFunction : DependentFunctions)
-			{
-				CurrentFunction->AppendReferencedTextures(InOutTextures);
-			}
-		}
-		// Not all expressions are cooked
-		else if (Expression)
-		{
-			// Append even if null as textures can be stripped at cook without our knowledge
-			// so we want to maintain the indices. This will waste a small amount of memory
-			UObject* ReferencedTexture = Expression->GetReferencedTexture();
-			checkf(!ReferencedTexture || Expression->CanReferenceTexture(), TEXT("This expression type missing an override for CanReferenceTexture?"));
-			if (Expression->CanReferenceTexture())
-			{
-				InOutTextures.AddUnique(ReferencedTexture);
-			}
-		}
-	}
-}
-
 #if WITH_EDITOR
 void UMaterial::RecursiveUpdateRealtimePreview( UMaterialExpression* InExpression, TArray<UMaterialExpression*>& InOutExpressionsToProcess )
 {
@@ -5756,17 +5266,17 @@ void UMaterial::RecursiveUpdateRealtimePreview( UMaterialExpression* InExpressio
 
 void UMaterial::AppendReferencedFunctionIdsTo(TArray<FGuid>& Ids) const
 {
-	for (int32 FunctionIndex = 0; FunctionIndex < MaterialFunctionInfos.Num(); FunctionIndex++)
+	for (int32 FunctionIndex = 0; FunctionIndex < CachedExpressionData.FunctionInfos.Num(); FunctionIndex++)
 	{
-		Ids.AddUnique(MaterialFunctionInfos[FunctionIndex].StateId);
+		Ids.AddUnique(CachedExpressionData.FunctionInfos[FunctionIndex].StateId);
 	}
 }
 
 void UMaterial::AppendReferencedParameterCollectionIdsTo(TArray<FGuid>& Ids) const
 {
-	for (int32 CollectionIndex = 0; CollectionIndex < MaterialParameterCollectionInfos.Num(); CollectionIndex++)
+	for (int32 CollectionIndex = 0; CollectionIndex < CachedExpressionData.ParameterCollectionInfos.Num(); CollectionIndex++)
 	{
-		Ids.AddUnique(MaterialParameterCollectionInfos[CollectionIndex].StateId);
+		Ids.AddUnique(CachedExpressionData.ParameterCollectionInfos[CollectionIndex].StateId);
 	}
 }
 
diff --git a/Engine/Source/Runtime/Engine/Private/Materials/MaterialCachedData.cpp b/Engine/Source/Runtime/Engine/Private/Materials/MaterialCachedData.cpp
new file mode 100644
--- /dev/null
+++ b/Engine/Source/Runtime/Engine/Private/Materials/MaterialCachedData.cpp
@@ -0,0 +1,425 @@
+#include "MaterialCachedData.h"
+#include "Materials/MaterialExpressionBreakMaterialAttributes.h"
+#include "Materials/MaterialExpressionShadingModel.h"
+#include "Materials/MaterialExpressionReroute.h"
+#include "Materials/MaterialExpressionSingleLayerWaterMaterialOutput.h"
+#include "Materials/MaterialExpressionCollectionParameter.h"
+#include "Materials/MaterialExpressionCustomOutput.h"
+#include "Materials/MaterialExpressionDynamicParameter.h"
+#include "Materials/MaterialExpressionFontSampleParameter.h"
+#include "Materials/MaterialExpressionQualitySwitch.h"
+#include "Materials/MaterialExpressionFeatureLevelSwitch.h"
+#include "Materials/MaterialExpressionShadingPathSwitch.h"
+#include "Materials/MaterialExpressionShaderStageSwitch.h"
+#include "Materials/MaterialExpressionMakeMaterialAttributes.h"
+#include "Materials/MaterialExpressionParameter.h"
+#include "Materials/MaterialExpressionRuntimeVirtualTextureSampleParameter.h"
+#include "Materials/MaterialExpressionScalarParameter.h"
+#include "Materials/MaterialExpressionStaticBoolParameter.h"
+#include "Materials/MaterialExpressionStaticComponentMaskParameter.h"
+#include "Materials/MaterialExpressionStaticSwitchParameter.h"
+#include "Materials/MaterialExpressionTextureBase.h"
+#include "Materials/MaterialExpressionTextureSample.h"
+#include "Materials/MaterialExpressionTextureSampleParameter.h"
+#include "Materials/MaterialExpressionVectorParameter.h"
+#include "Materials/MaterialExpressionVertexInterpolator.h"
+#include "Materials/MaterialExpressionSceneColor.h"
+#include "Materials/MaterialExpressionRuntimeVirtualTextureOutput.h"
+#include "Materials/MaterialExpressionLandscapeGrassOutput.h"
+#include "Materials/MaterialExpressionCurveAtlasRowParameter.h"
+#include "Materials/MaterialFunction.h"
+#include "Materials/MaterialFunctionInstance.h"
+#include "Materials/MaterialExpressionMaterialFunctionCall.h"
+#include "Materials/MaterialParameterCollection.h"
+#include "LandscapeGrassType.h"
+#include "ProfilingDebugging/LoadTimeTracker.h"
+
+void FMaterialCachedExpressionData::Reset()
+{
+	Parameters.Reset();
+	ReferencedTextures.Reset();
+	FunctionInfos.Reset();
+	ParameterCollectionInfos.Reset();
+	GrassTypes.Reset();
+	DynamicParameterNames.Reset();
+	QualityLevelsUsed.Reset();
+	QualityLevelsUsed.AddDefaulted(EMaterialQualityLevel::Num);
+	bHasRuntimeVirtualTextureOutput = false;
+}
+
+static int32 FindParameterLowerBoundIndex(const FMaterialCachedParameterEntry& Entry, const FHashedMaterialParameterInfo& HashedParameterInfo)
+{
+	// Parameters are first sorted by name hash
+	const uint64 NameHash = HashedParameterInfo.Name.GetHash();
+	const int32 LowerIndex = Algo::LowerBound(Entry.NameHashes, NameHash);
+	if (LowerIndex < Entry.NameHashes.Num())
+	{
+		const TConstArrayView<uint64> NameHashesUpper = MakeArrayView(&Entry.NameHashes[LowerIndex], Entry.NameHashes.Num() - LowerIndex);
+		const int32 UpperIndex = LowerIndex + Algo::UpperBound(NameHashesUpper, NameHash);
+		if (UpperIndex - LowerIndex > 0)
+		{
+			// more than 1 entry with the same name, next sort by Association/Index
+			auto ProjectionFunc = [](const FMaterialParameterInfo& InParameterInfo)
+			{
+				return FHashedMaterialParameterInfo(FHashedName(), InParameterInfo.Association, InParameterInfo.Index);
+			};
+			auto CompareFunc = [](const FHashedMaterialParameterInfo& Lhs, const FHashedMaterialParameterInfo& Rhs)
+			{
+				if (Lhs.Association != Rhs.Association) return Lhs.Association < Rhs.Association;
+				return Lhs.Index < Rhs.Index;
+			};
+
+			const TConstArrayView<FMaterialParameterInfo> ParameterInfos = MakeArrayView(&Entry.ParameterInfos[LowerIndex], UpperIndex - LowerIndex);
+			return LowerIndex + Algo::LowerBoundBy(ParameterInfos, HashedParameterInfo, ProjectionFunc, CompareFunc);
+		}
+	}
+	return LowerIndex;
+}
+
+#if WITH_EDITOR
+static int32 TryAddParameter(FMaterialCachedParameters& CachedParameters, EMaterialParameterType Type, const FMaterialParameterInfo& ParameterInfo, const FGuid& ExpressionGuid, bool bOverride = false)
+{
+	FMaterialCachedParameterEntry& Entry = CachedParameters.Entries[(int32)Type];
+	const FHashedMaterialParameterInfo HashedParameterInfo(ParameterInfo);
+	int32 Index = FindParameterLowerBoundIndex(Entry, HashedParameterInfo);
+	if (Index >= Entry.NameHashes.Num() || Entry.ParameterInfos[Index] != ParameterInfo)
+	{
+		Entry.NameHashes.Insert(HashedParameterInfo.Name.GetHash(), Index);
+		Entry.ParameterInfos.Insert(ParameterInfo, Index);
+		Entry.ExpressionGuids.Insert(ExpressionGuid, Index);
+		Entry.Overrides.Insert(bOverride, Index);
+		return Index;
+	}
+	return INDEX_NONE;
+}
+
+void FMaterialCachedExpressionData::UpdateForFunction(UMaterialFunctionInterface* Function, EMaterialParameterAssociation Association, int32 ParameterIndex)
+{
+	if (!Function)
+	{
+		return;
+	}
+
+	FMaterialFunctionInfo NewFunctionInfo;
+	NewFunctionInfo.Function = Function;
+	NewFunctionInfo.StateId = Function->StateId;
+	FunctionInfos.Add(NewFunctionInfo);
+
+	UMaterialFunctionInstance* FunctionInstance = Cast<UMaterialFunctionInstance>(Function);
+	if (FunctionInstance)
+	{
+		for (const FScalarParameterValue& Param : FunctionInstance->ScalarParameterValues)
+		{
+			const FMaterialParameterInfo ParameterInfo(Param.ParameterInfo.Name, Association, ParameterIndex);
+			const int32 Index = TryAddParameter(Parameters, EMaterialParameterType::Scalar, ParameterInfo, FGuid(), true);
+			if (Index != INDEX_NONE)
+			{
+				Parameters.ScalarValues.Insert(Param.ParameterValue, Index);
+			}
+		}
+
+		for (const FVectorParameterValue& Param : FunctionInstance->VectorParameterValues)
+		{
+			const FMaterialParameterInfo ParameterInfo(Param.ParameterInfo.Name, Association, ParameterIndex);
+			const int32 Index = TryAddParameter(Parameters, EMaterialParameterType::Vector, ParameterInfo, FGuid(), true);
+			if (Index != INDEX_NONE)
+			{
+				Parameters.VectorValues.Insert(Param.ParameterValue, Index);
+			}
+		}
+
+		for (const FTextureParameterValue& Param : FunctionInstance->TextureParameterValues)
+		{
+			const FMaterialParameterInfo ParameterInfo(Param.ParameterInfo.Name, Association, ParameterIndex);
+			const int32 Index = TryAddParameter(Parameters, EMaterialParameterType::Texture, ParameterInfo, FGuid(), true);
+			if (Index != INDEX_NONE)
+			{
+				Parameters.TextureValues.Insert(Param.ParameterValue, Index);
+			}
+		}
+
+		for (const FRuntimeVirtualTextureParameterValue& Param : FunctionInstance->RuntimeVirtualTextureParameterValues)
+		{
+			const FMaterialParameterInfo ParameterInfo(Param.ParameterInfo.Name, Association, ParameterIndex);
+			const int32 Index = TryAddParameter(Parameters, EMaterialParameterType::RuntimeVirtualTexture, ParameterInfo, FGuid(), true);
+			if (Index != INDEX_NONE)
+			{
+				Parameters.RuntimeVirtualTextureValues.Insert(Param.ParameterValue, Index);
+			}
+		}
+
+		for (const FFontParameterValue& Param : FunctionInstance->FontParameterValues)
+		{
+			const FMaterialParameterInfo ParameterInfo(Param.ParameterInfo.Name, Association, ParameterIndex);
+			const int32 Index = TryAddParameter(Parameters, EMaterialParameterType::Font, ParameterInfo, FGuid(), true);
+			if (Index != INDEX_NONE)
+			{
+				Parameters.FontValues.Insert(Param.FontValue, Index);
+				Parameters.FontPageValues.Insert(Param.FontPage, Index);
+			}
+		}
+	}
+
+	const TArray<UMaterialExpression*>* FunctionExpressions = Function->GetFunctionExpressions();
+	if (FunctionExpressions)
+	{
+		UpdateForExpressions(*FunctionExpressions, Association, ParameterIndex);
+	}
+}
+
+void FMaterialCachedExpressionData::UpdateForLayerFunctions(const FMaterialLayersFunctions& LayerFunctions)
+{
+	for (int32 LayerIndex = 0; LayerIndex < LayerFunctions.Layers.Num(); ++LayerIndex)
+	{
+		UpdateForFunction(LayerFunctions.Layers[LayerIndex], LayerParameter, LayerIndex);
+	}
+
+	for (int32 BlendIndex = 0; BlendIndex < LayerFunctions.Blends.Num(); ++BlendIndex)
+	{
+		UpdateForFunction(LayerFunctions.Blends[BlendIndex], BlendParameter, BlendIndex);
+	}
+}
+
+void FMaterialCachedExpressionData::UpdateForExpressions(const TArray<UMaterialExpression*>& Expressions, EMaterialParameterAssociation Association, int32 ParameterIndex)
+{
+	for (UMaterialExpression* Expression : Expressions)
+	{
+		UObject* ReferencedTexture = Expression->GetReferencedTexture();
+		checkf(!ReferencedTexture || Expression->CanReferenceTexture(), TEXT("This expression type missing an override for CanReferenceTexture?"));
+		if (Expression->CanReferenceTexture())
+		{
+			ReferencedTextures.AddUnique(ReferencedTexture);
+		}
+
+		if (UMaterialExpressionScalarParameter* ExpressionScalarParameter = Cast<UMaterialExpressionScalarParameter>(Expression))
+		{
+			const FMaterialParameterInfo ParameterInfo(ExpressionScalarParameter->GetParameterName(), Association, ParameterIndex);
+			const int32 Index = TryAddParameter(Parameters, EMaterialParameterType::Scalar, ParameterInfo, ExpressionScalarParameter->ExpressionGUID);
+			if (Index != INDEX_NONE)
+			{
+				Parameters.ScalarValues.Insert(ExpressionScalarParameter->DefaultValue, Index);
+				Parameters.ScalarMinMaxValues.Insert(FVector2D(ExpressionScalarParameter->SliderMin, ExpressionScalarParameter->SliderMax), Index);
+				if (ExpressionScalarParameter->IsUsedAsAtlasPosition())
+				{
+					UMaterialExpressionCurveAtlasRowParameter* ExpressionAtlasParameter = Cast<UMaterialExpressionCurveAtlasRowParameter>(ExpressionScalarParameter);
+					Parameters.ScalarCurveValues.Insert(ExpressionAtlasParameter->Curve, Index);
+					Parameters.ScalarCurveAtlasValues.Insert(ExpressionAtlasParameter->Atlas, Index);
+				}
+				else
+				{
+					Parameters.ScalarCurveValues.Insert(nullptr, Index);
+					Parameters.ScalarCurveAtlasValues.Insert(nullptr, Index);
+				}
+			}
+		}
+		else if (UMaterialExpressionVectorParameter* ExpressionVectorParameter = Cast<UMaterialExpressionVectorParameter>(Expression))
+		{
+			const FMaterialParameterInfo ParameterInfo(ExpressionVectorParameter->GetParameterName(), Association, ParameterIndex);
+			const int32 Index = TryAddParameter(Parameters, EMaterialParameterType::Vector, ParameterInfo, ExpressionVectorParameter->ExpressionGUID);
+			if (Index != INDEX_NONE)
+			{
+				Parameters.VectorValues.Insert(ExpressionVectorParameter->DefaultValue, Index);
+				Parameters.VectorChannelNameValues.Insert(ExpressionVectorParameter->ChannelNames, Index);
+				Parameters.VectorUsedAsChannelMaskValues.Insert(ExpressionVectorParameter->IsUsedAsChannelMask(), Index);
+			}
+		}
+		else if (UMaterialExpressionTextureSampleParameter* ExpressionTextureParameter = Cast<UMaterialExpressionTextureSampleParameter>(Expression))
+		{
+			const FMaterialParameterInfo ParameterInfo(ExpressionTextureParameter->GetParameterName(), Association, ParameterIndex);
+			const int32 Index = TryAddParameter(Parameters, EMaterialParameterType::Texture, ParameterInfo, ExpressionTextureParameter->ExpressionGUID);
+			if (Index != INDEX_NONE)
+			{
+				Parameters.TextureValues.Insert(ExpressionTextureParameter->Texture, Index);
+				Parameters.TextureChannelNameValues.Insert(ExpressionTextureParameter->ChannelNames, Index);
+			}
+		}
+		else if (UMaterialExpressionFontSampleParameter* ExpressionFontParameter = Cast<UMaterialExpressionFontSampleParameter>(Expression))
+		{
+			const FMaterialParameterInfo ParameterInfo(ExpressionFontParameter->GetParameterName(), Association, ParameterIndex);
+			const int32 Index = TryAddParameter(Parameters, EMaterialParameterType::Font, ParameterInfo, ExpressionFontParameter->ExpressionGUID);
+			if (Index != INDEX_NONE)
+			{
+				Parameters.FontValues.Insert(ExpressionFontParameter->Font, Index);
+				Parameters.FontPageValues.Insert(ExpressionFontParameter->FontTexturePage, Index);
+			}
+		}
+		else if (UMaterialExpressionRuntimeVirtualTextureSampleParameter* ExpressionRTVParameter = Cast<UMaterialExpressionRuntimeVirtualTextureSampleParameter>(Expression))
+		{
+			const FMaterialParameterInfo ParameterInfo(ExpressionRTVParameter->GetParameterName(), Association, ParameterIndex);
+			const int32 Index = TryAddParameter(Parameters, EMaterialParameterType::RuntimeVirtualTexture, ParameterInfo, ExpressionRTVParameter->ExpressionGUID);
+			if (Index != INDEX_NONE)
+			{
+				Parameters.RuntimeVirtualTextureValues.Insert(ExpressionRTVParameter->VirtualTexture, Index);
+			}
+		}
+		else if (UMaterialExpressionCollectionParameter* ExpressionCollectionParameter = Cast<UMaterialExpressionCollectionParameter>(Expression))
+		{
+			UMaterialParameterCollection* Collection = ExpressionCollectionParameter->Collection;
+			if (Collection)
+			{
+				FMaterialParameterCollectionInfo NewInfo;
+				NewInfo.ParameterCollection = Collection;
+				NewInfo.StateId = Collection->StateId;
+				ParameterCollectionInfos.AddUnique(NewInfo);
+			}
+		}
+		else if (UMaterialExpressionDynamicParameter* ExpressionDynamicParameter = Cast< UMaterialExpressionDynamicParameter>(Expression))
+		{
+			DynamicParameterNames.Empty(ExpressionDynamicParameter->ParamNames.Num());
+			for (const FString& Name : ExpressionDynamicParameter->ParamNames)
+			{
+				DynamicParameterNames.Add(*Name);
+			}
+		}
+		else if (UMaterialExpressionLandscapeGrassOutput* ExpressionGrassOutput = Cast<UMaterialExpressionLandscapeGrassOutput>(Expression))
+		{
+			for (const auto& Type : ExpressionGrassOutput->GrassTypes)
+			{
+				GrassTypes.AddUnique(Type.GrassType);
+			}
+		}
+		else if (UMaterialExpressionQualitySwitch* QualitySwitchNode = Cast<UMaterialExpressionQualitySwitch>(Expression))
+		{
+			for (int32 InputIndex = 0; InputIndex < EMaterialQualityLevel::Num; InputIndex++)
+			{
+				if (QualitySwitchNode->Inputs[InputIndex].IsConnected())
+				{
+					QualityLevelsUsed[InputIndex] = true;
+				}
+			}
+
+			if (QualitySwitchNode->Default.IsConnected())
+			{
+				QualityLevelsUsed[EMaterialQualityLevel::High] = true;
+			}
+		}
+		else if (Expression->IsA(UMaterialExpressionRuntimeVirtualTextureOutput::StaticClass()))
+		{
+			bHasRuntimeVirtualTextureOutput = true;
+		}
+		else if (UMaterialExpressionMaterialFunctionCall* FunctionCall = Cast<UMaterialExpressionMaterialFunctionCall>(Expression))
+		{
+			UpdateForFunction(FunctionCall->MaterialFunction, GlobalParameter, -1);
+
+			// Update the function call node, so it can relink inputs and outputs as needed
+			// Update even if MaterialFunctionNode->MaterialFunction is NULL, because we need to remove the invalid inputs in that case
+			FunctionCall->UpdateFromFunctionResource();
+		}
+		else if (UMaterialExpressionMaterialAttributeLayers* LayersExpression = Cast<UMaterialExpressionMaterialAttributeLayers>(Expression))
+		{
+			checkf(Association == GlobalParameter, TEXT("UMaterialExpressionMaterialAttributeLayers can't be nested"));
+			UpdateForLayerFunctions(LayersExpression->DefaultLayers);
+
+			LayersExpression->RebuildLayerGraph(false);
+		}
+	}
+}
+#endif // WITH_EDITOR
+
+void FMaterialCachedParameterEntry::Reset()
+{
+	NameHashes.Reset();
+	ParameterInfos.Reset();
+	ExpressionGuids.Reset();
+	Overrides.Reset();
+}
+
+void FMaterialCachedParameters::Reset()
+{
+	for (int32 i = 0; i < NumMaterialParameterTypes; ++i)
+	{
+		Entries[i].Reset();
+	}
+
+	ScalarValues.Reset();
+	VectorValues.Reset();
+	TextureValues.Reset();
+	FontValues.Reset();
+	FontPageValues.Reset();
+	RuntimeVirtualTextureValues.Reset();
+
+#if WITH_EDITORONLY_DATA
+	ScalarMinMaxValues.Reset();
+	ScalarCurveValues.Reset();
+	ScalarCurveAtlasValues.Reset();
+	VectorChannelNameValues.Reset();
+	VectorUsedAsChannelMaskValues.Reset();
+	TextureChannelNameValues.Reset();
+#endif // WITH_EDITORONLY_DATA
+}
+
+
+int32 FMaterialCachedParameters::FindParameterIndex(EMaterialParameterType Type, const FHashedMaterialParameterInfo& HashedParameterInfo, bool bOveriddenOnly) const
+{
+	const int32 Index = FindParameterIndex(Type, HashedParameterInfo);
+	if (Index != INDEX_NONE)
+	{
+		if (IsParameterValid(Type, Index, bOveriddenOnly))
+		{
+			return Index;
+		}
+	}
+	return INDEX_NONE;
+}
+
+int32 FMaterialCachedParameters::FindParameterIndex(EMaterialParameterType Type, const FHashedMaterialParameterInfo& HashedParameterInfo) const
+{
+	const FMaterialCachedParameterEntry& Entry = Entries[(int32)Type];
+	const int32 Index = FindParameterLowerBoundIndex(Entry, HashedParameterInfo);
+	if (Index < Entry.NameHashes.Num() &&
+		Entry.NameHashes[Index] == HashedParameterInfo.Name.GetHash() &&
+		Entry.ParameterInfos[Index].Association == HashedParameterInfo.Association &&
+		Entry.ParameterInfos[Index].Index == HashedParameterInfo.Index)
+	{
+		return Index;
+	}
+	
+	return INDEX_NONE;
+}
+
+bool FMaterialCachedParameters::IsParameterValid(EMaterialParameterType Type, int32 Index, bool bOveriddenOnly) const
+{
+	const FMaterialCachedParameterEntry& Entry = Entries[(int32)Type];
+	return !bOveriddenOnly || Entry.Overrides[Index];
+}
+
+void FMaterialCachedParameters::GetAllParameterInfoOfType(EMaterialParameterType Type, bool bEmptyOutput, TArray<FMaterialParameterInfo>& OutParameterInfo, TArray<FGuid>& OutParameterIds) const
+{
+	const FMaterialCachedParameterEntry& Entry = Entries[(int32)Type];
+	const int32 NumParameters = Entry.NameHashes.Num();
+	if (bEmptyOutput)
+	{
+		OutParameterInfo.Empty(NumParameters);
+		OutParameterIds.Empty(NumParameters);
+	}
+
+	for (int32 i = 0; i < NumParameters; ++i)
+	{
+		if (!Entry.Overrides[i])
+		{
+			OutParameterInfo.Add(Entry.ParameterInfos[i]);
+			OutParameterIds.Add(Entry.ExpressionGuids[i]);
+		}
+	}
+}
+
+void FMaterialCachedParameters::GetAllGlobalParameterInfoOfType(EMaterialParameterType Type, bool bEmptyOutput, TArray<FMaterialParameterInfo>& OutParameterInfo, TArray<FGuid>& OutParameterIds) const
+{
+	const FMaterialCachedParameterEntry& Entry = Entries[(int32)Type];
+	const int32 NumParameters = Entry.NameHashes.Num();
+	if (bEmptyOutput)
+	{
+		OutParameterInfo.Empty(NumParameters);
+		OutParameterIds.Empty(NumParameters);
+	}
+
+	for (int32 i = 0; i < NumParameters; ++i)
+	{
+		const FMaterialParameterInfo& ParameterInfo = Entry.ParameterInfos[i];
+		if (!Entry.Overrides[i] && ParameterInfo.Association == GlobalParameter)
+		{
+			OutParameterInfo.Add(ParameterInfo);
+			OutParameterIds.Add(Entry.ExpressionGuids[i]);
+		}
+	}
+}
