Copying Private-LoadTimes-4.24 stream to Main. Biggest changes are in Materials/Shader memory freezing.

diff --git a/Engine/Source/Runtime/Engine/Private/Materials/MaterialExpressions.cpp b/Engine/Source/Runtime/Engine/Private/Materials/MaterialExpressions.cpp
--- a/Engine/Source/Runtime/Engine/Private/Materials/MaterialExpressions.cpp
+++ b/Engine/Source/Runtime/Engine/Private/Materials/MaterialExpressions.cpp
@@ -39,6 +39,7 @@
 #include "Engine/VolumeTexture.h"
 #include "Styling/CoreStyle.h"
 #include "VT/RuntimeVirtualTexture.h"
+#include "ProfilingDebugging/LoadTimeTracker.h"
 
 #include "Materials/MaterialExpressionAbs.h"
 #include "Materials/MaterialExpressionActorPositionWS.h"
@@ -711,6 +712,7 @@ void UMaterialExpression::CopyMaterialExpressions(const TArray<UMaterialExpressi
 
 void UMaterialExpression::Serialize(FStructuredArchive::FRecord Record)
 {
+	SCOPED_LOADTIMER(UMaterialExpression::Serialize);
 	Super::Serialize(Record);
 
 #if WITH_EDITORONLY_DATA
@@ -740,6 +742,7 @@ void UMaterialExpression::PostInitProperties()
 
 void UMaterialExpression::PostLoad()
 {
+	SCOPED_LOADTIMER(UMaterialExpression::PostLoad);
 	Super::PostLoad();
 
 	if (!Material && GetOuter()->IsA(UMaterial::StaticClass()))
@@ -2406,7 +2409,7 @@ UMaterialExpressionRuntimeVirtualTextureSampleParameter::UMaterialExpressionRunt
 #endif
 }
 
-bool UMaterialExpressionRuntimeVirtualTextureSampleParameter::IsNamedParameter(const FMaterialParameterInfo& ParameterInfo, URuntimeVirtualTexture*& OutValue) const
+bool UMaterialExpressionRuntimeVirtualTextureSampleParameter::IsNamedParameter(const FHashedMaterialParameterInfo& ParameterInfo, URuntimeVirtualTexture*& OutValue) const
 {
 	if (ParameterInfo.Name == ParameterName)
 	{
@@ -2716,7 +2719,7 @@ void UMaterialExpressionTextureSampleParameter::SetValueToMatchingExpression(UMa
 }
 #endif // WITH_EDITOR
 
-bool UMaterialExpressionTextureSampleParameter::IsNamedParameter(const FMaterialParameterInfo& ParameterInfo, UTexture*& OutValue) const
+bool UMaterialExpressionTextureSampleParameter::IsNamedParameter(const FHashedMaterialParameterInfo& ParameterInfo, UTexture*& OutValue) const
 {
 	if (ParameterInfo.Name == ParameterName)
 	{
@@ -4427,7 +4430,7 @@ void UMaterialExpressionStaticComponentMaskParameter::SetValueToMatchingExpressi
 }
 #endif // WITH_EDITOR
 
-bool UMaterialExpressionStaticComponentMaskParameter::IsNamedParameter(const FMaterialParameterInfo& ParameterInfo, bool& OutR, bool& OutG, bool& OutB, bool& OutA, FGuid&OutExpressionGuid) const
+bool UMaterialExpressionStaticComponentMaskParameter::IsNamedParameter(const FHashedMaterialParameterInfo& ParameterInfo, bool& OutR, bool& OutG, bool& OutB, bool& OutA, FGuid&OutExpressionGuid) const
 {
 	if (ParameterInfo.Name == ParameterName)
 	{
@@ -6278,7 +6281,9 @@ void UMaterialExpressionMaterialAttributeLayers::PostLoad()
 		}
 	}
 
+#if WITH_EDITORONLY_DATA
 	RebuildLayerGraph(false);
+#endif // WITH_EDITORONLY_DATA
 }
 
 #if WITH_EDITOR
@@ -6287,8 +6292,9 @@ void UMaterialExpressionMaterialAttributeLayers::PostEditChangeProperty(FPropert
 	Super::PostEditChangeProperty(PropertyChangedEvent);
 	RebuildLayerGraph(false);
 }
-#endif
+#endif // WITH_EDITOR
 
+#if WITH_EDITORONLY_DATA
 void UMaterialExpressionMaterialAttributeLayers::RebuildLayerGraph(bool bReportErrors)
 {
 	const TArray<UMaterialFunctionInterface*>& Layers = GetLayers();
@@ -6431,7 +6437,9 @@ void UMaterialExpressionMaterialAttributeLayers::OverrideLayerGraph(const FMater
 		RebuildLayerGraph(false);
 	}
 }
+#endif // WITH_EDITORONLY_DATA
 
+#if WITH_EDITOR
 bool UMaterialExpressionMaterialAttributeLayers::ValidateLayerConfiguration(FMaterialCompiler* Compiler, bool bReportErrors)
 {
 #define COMPILER_OR_LOG_ERROR(Format, ...)								\
@@ -6547,7 +6555,9 @@ bool UMaterialExpressionMaterialAttributeLayers::ValidateLayerConfiguration(FMat
 
 #undef COMPILER_OR_LOG_ERROR
 }
+#endif // WITH_EDITOR
 
+#if WITH_EDITORONLY_DATA
 bool UMaterialExpressionMaterialAttributeLayers::IterateDependentFunctions(TFunctionRef<bool(UMaterialFunctionInterface*)> Predicate) const
 {
 	const TArray<UMaterialFunctionInterface*>& Layers = GetLayers();
@@ -6593,8 +6603,9 @@ void UMaterialExpressionMaterialAttributeLayers::GetDependentFunctions(TArray<UM
 		return true;
 	});
 }
+#endif // WITH_EDITORONLY_DATA
 
-UMaterialFunctionInterface* UMaterialExpressionMaterialAttributeLayers::GetParameterAssociatedFunction(const FMaterialParameterInfo& ParameterInfo) const
+UMaterialFunctionInterface* UMaterialExpressionMaterialAttributeLayers::GetParameterAssociatedFunction(const FHashedMaterialParameterInfo& ParameterInfo) const
 {
 	check(ParameterInfo.Association != EMaterialParameterAssociation::GlobalParameter);
 
@@ -6621,36 +6632,6 @@ UMaterialFunctionInterface* UMaterialExpressionMaterialAttributeLayers::GetParam
 	return LayersFunction;
 }
 
-void UMaterialExpressionMaterialAttributeLayers::GetParameterAssociatedFunctions(const FMaterialParameterInfo& ParameterInfo, TArray<UMaterialFunctionInterface*>& AssociatedFunctions) const
-{
-	check(ParameterInfo.Association != EMaterialParameterAssociation::GlobalParameter);
-
-	// Grab the associated layer or blend
-	UMaterialFunctionInterface* LayersFunction = nullptr;
-
-	if (ParameterInfo.Association == EMaterialParameterAssociation::LayerParameter)
-	{
-		const TArray<UMaterialFunctionInterface*>& Layers = GetLayers();
-		if (Layers.IsValidIndex(ParameterInfo.Index))
-		{
-			LayersFunction = Layers[ParameterInfo.Index];
-		}	
-	}
-	else if (ParameterInfo.Association == EMaterialParameterAssociation::BlendParameter)
-	{
-		const TArray<UMaterialFunctionInterface*>& Blends = GetBlends();
-		if (Blends.IsValidIndex(ParameterInfo.Index))
-		{
-			LayersFunction = Blends[ParameterInfo.Index];
-		}
-	}
-
-	if (LayersFunction)
-	{
-		LayersFunction->GetDependentFunctions(AssociatedFunctions);
-	}
-}
-
 #if WITH_EDITOR
 int32 UMaterialExpressionMaterialAttributeLayers::Compile(FMaterialCompiler* Compiler, int32 OutputIndex)
 {
@@ -6741,7 +6722,7 @@ uint32 UMaterialExpressionMaterialAttributeLayers::GetInputType(int32 InputIndex
 }
 #endif
 
-bool UMaterialExpressionMaterialAttributeLayers::IsNamedParameter(const FMaterialParameterInfo& ParameterInfo, FMaterialLayersFunctions& OutLayers, FGuid& OutExpressionGuid) const
+bool UMaterialExpressionMaterialAttributeLayers::IsNamedParameter(const FHashedMaterialParameterInfo& ParameterInfo, FMaterialLayersFunctions& OutLayers, FGuid& OutExpressionGuid) const
 {
 	if (ParameterInfo.Name == ParameterName)
 	{
@@ -7281,7 +7262,7 @@ void UMaterialExpressionVectorParameter::GetCaption(TArray<FString>& OutCaptions
 }
 #endif // WITH_EDITOR
 
-bool UMaterialExpressionVectorParameter::IsNamedParameter(const FMaterialParameterInfo& ParameterInfo, FLinearColor& OutValue) const
+bool UMaterialExpressionVectorParameter::IsNamedParameter(const FHashedMaterialParameterInfo& ParameterInfo, FLinearColor& OutValue) const
 {
 	if (ParameterInfo.Name == ParameterName)
 	{
@@ -7602,7 +7583,7 @@ void UMaterialExpressionScalarParameter::GetCaption(TArray<FString>& OutCaptions
 }
 #endif // WITH_EDITOR
 
-bool UMaterialExpressionScalarParameter::IsNamedParameter(const FMaterialParameterInfo& ParameterInfo, float& OutValue) const
+bool UMaterialExpressionScalarParameter::IsNamedParameter(const FHashedMaterialParameterInfo& ParameterInfo, float& OutValue) const
 {
 	if (ParameterInfo.Name == ParameterName)
 	{
@@ -7823,7 +7804,7 @@ void UMaterialExpressionStaticBoolParameter::SetValueToMatchingExpression(UMater
 }
 #endif // WITH_EDITOR
 
-bool UMaterialExpressionStaticBoolParameter::IsNamedParameter(const FMaterialParameterInfo& ParameterInfo, bool& OutValue, FGuid& OutExpressionGuid) const
+bool UMaterialExpressionStaticBoolParameter::IsNamedParameter(const FHashedMaterialParameterInfo& ParameterInfo, bool& OutValue, FGuid& OutExpressionGuid) const
 {
 	if (ParameterInfo.Name == ParameterName)
 	{
@@ -8747,6 +8728,7 @@ bool UMaterialExpressionDynamicParameter::NeedsLoadForClient() const
 	return true;
 }
 
+#if WITH_EDITORONLY_DATA
 void UMaterialExpressionDynamicParameter::UpdateDynamicParameterProperties()
 {
 	check(Material);
@@ -8759,6 +8741,7 @@ void UMaterialExpressionDynamicParameter::UpdateDynamicParameterProperties()
 		}
 	}
 }
+#endif // WITH_EDITORONLY_DATA
 
 bool UMaterialExpressionDynamicParameter::CopyDynamicParameterProperties(const UMaterialExpressionDynamicParameter* FromParam)
 {
@@ -10335,7 +10318,7 @@ void UMaterialExpressionFontSampleParameter::SetValueToMatchingExpression(UMater
 }
 #endif // WITH_EDITOR
 
-bool UMaterialExpressionFontSampleParameter::IsNamedParameter(const FMaterialParameterInfo& ParameterInfo, UFont*& OutFontValue, int32& OutFontPage) const
+bool UMaterialExpressionFontSampleParameter::IsNamedParameter(const FHashedMaterialParameterInfo& ParameterInfo, UFont*& OutFontValue, int32& OutFontPage) const
 {
 	if (ParameterInfo.Name == ParameterName)
 	{
@@ -11278,11 +11261,6 @@ void UMaterialFunction::Serialize(FArchive& Ar)
 {
 	Super::Serialize(Ar);
 
-	if (FPlatformProperties::RequiresCookedData() && Ar.IsLoading())
-	{
-		FunctionExpressions.Remove(nullptr);
-	}
-
 #if WITH_EDITOR
 	if (Ar.UE4Ver() < VER_UE4_FLIP_MATERIAL_COORDS)
 	{
@@ -11323,6 +11301,7 @@ void UMaterialFunction::PostLoad()
 		StateId = FGuid::NewGuid();
 	}
 
+#if WITH_EDITORONLY_DATA
 	for (int32 ExpressionIndex = 0; ExpressionIndex < FunctionExpressions.Num(); ExpressionIndex++)
 	{
 		// Expressions whose type was removed can be nullptr
@@ -11331,6 +11310,7 @@ void UMaterialFunction::PostLoad()
 			FunctionExpressions[ExpressionIndex]->ConditionalPostLoad();
 		}
 	}
+#endif // WITH_EDITORONLY_DATA
 
 #if WITH_EDITOR
 	if (CombinedOutputTypes == 0)
@@ -11701,6 +11681,7 @@ void UMaterialFunction::UnlinkFromCaller()
 
 #endif // WITH_EDITOR
 
+#if WITH_EDITORONLY_DATA
 bool UMaterialFunction::IsDependent(UMaterialFunctionInterface* OtherFunction)
 {
 	if (!OtherFunction)
@@ -11774,24 +11755,7 @@ void UMaterialFunction::GetDependentFunctions(TArray<UMaterialFunctionInterface*
 		return true;
 	});
 }
-
-void UMaterialFunction::AppendReferencedTextures(TArray<UObject*>& InOutTextures) const
-{
-	for (UMaterialExpression* CurrentExpression : FunctionExpressions)
-	{
-		if(CurrentExpression)
-		{
-			// Append even if null as textures can be stripped at cook without our knowledge
-			// so we want to maintain the indices. This will waste a small amount of memory
-			UObject* ReferencedTexture = CurrentExpression->GetReferencedTexture();
-			checkf(!ReferencedTexture || CurrentExpression->CanReferenceTexture(), TEXT("This expression type missing an override for CanReferenceTexture?"));
-			if (CurrentExpression->CanReferenceTexture())
-			{
-				InOutTextures.Add(ReferencedTexture);
-			}
-		}
-	}
-}
+#endif // WITH_EDITORONLY_DATA
 
 #if WITH_EDITOR
 bool UMaterialFunction::HasFlippedCoordinates() const
@@ -12134,6 +12098,7 @@ UMaterialFunctionInstance::UMaterialFunctionInstance(const FObjectInitializer& O
 #endif
 }
 
+#if WITH_EDITOR
 void UMaterialFunctionInstance::UpdateParameterSet()
 {
 	if (UMaterialFunction* BaseFunction = Cast<UMaterialFunction>(GetBaseFunction()))
@@ -12229,7 +12194,6 @@ void UMaterialFunctionInstance::UpdateParameterSet()
 	}
 }
 
-#if WITH_EDITOR
 void UMaterialFunctionInstance::OverrideMaterialInstanceParameterValues(UMaterialInstance* Instance)
 {
 	// Dynamic parameters
@@ -12245,7 +12209,7 @@ void UMaterialFunctionInstance::OverrideMaterialInstanceParameterValues(UMateria
 	StaticParametersOverride.StaticComponentMaskParameters = StaticComponentMaskParameterValues;
 	Instance->UpdateStaticPermutation(StaticParametersOverride);
 }
-#endif
+#endif // WITH_EDITOR
 
 void UMaterialFunctionInstance::UpdateFromFunctionResource()
 {
@@ -12306,6 +12270,7 @@ bool UMaterialFunctionInstance::IsDependent(UMaterialFunctionInterface* OtherFun
 	return Parent ? Parent->IsDependent(OtherFunction) : false;
 }
 
+#if WITH_EDITORONLY_DATA
 bool UMaterialFunctionInstance::IterateDependentFunctions(TFunctionRef<bool(UMaterialFunctionInterface*)> Predicate) const
 {
 	if (Parent)
@@ -12323,27 +12288,14 @@ bool UMaterialFunctionInstance::IterateDependentFunctions(TFunctionRef<bool(UMat
 }
 
 void UMaterialFunctionInstance::GetDependentFunctions(TArray<UMaterialFunctionInterface*>& DependentFunctions) const
-{
-	IterateDependentFunctions([&DependentFunctions](UMaterialFunctionInterface* MaterialFunction) -> bool
-	{
-		DependentFunctions.AddUnique(MaterialFunction);
-		return true;
-	});
-}
-
-void UMaterialFunctionInstance::AppendReferencedTextures(TArray<UObject*>& InOutTextures) const
 {
 	if (Parent)
 	{
-		Parent->AppendReferencedTextures(InOutTextures);
-	}
-
-	// @TODO: This should be able to replace base textures rather than append
-	for (const FTextureParameterValue& TextureParam : TextureParameterValues)
-	{
-		InOutTextures.Add(TextureParam.ParameterValue);
+		Parent->GetDependentFunctions(DependentFunctions);
+		DependentFunctions.AddUnique(Parent);
 	}
 }
+#endif // WITH_EDITORONLY_DATA
 
 #if WITH_EDITOR
 UMaterialInterface* UMaterialFunctionInstance::GetPreviewMaterial()
@@ -12374,7 +12326,7 @@ bool UMaterialFunctionInstance::HasFlippedCoordinates() const
 }
 #endif
 
-bool UMaterialFunctionInstance::OverrideNamedScalarParameter(const FMaterialParameterInfo& ParameterInfo, float& OutValue)
+bool UMaterialFunctionInstance::OverrideNamedScalarParameter(const FHashedMaterialParameterInfo& ParameterInfo, float& OutValue)
 {
 	for (const FScalarParameterValue& ScalarParameter : ScalarParameterValues)
 	{
@@ -12388,7 +12340,7 @@ bool UMaterialFunctionInstance::OverrideNamedScalarParameter(const FMaterialPara
 	return false;
 }
 
-bool UMaterialFunctionInstance::OverrideNamedVectorParameter(const FMaterialParameterInfo& ParameterInfo, FLinearColor& OutValue)
+bool UMaterialFunctionInstance::OverrideNamedVectorParameter(const FHashedMaterialParameterInfo& ParameterInfo, FLinearColor& OutValue)
 {
 	for (const FVectorParameterValue& VectorParameter : VectorParameterValues)
 	{
@@ -12402,7 +12354,7 @@ bool UMaterialFunctionInstance::OverrideNamedVectorParameter(const FMaterialPara
 	return false;
 }
 
-bool UMaterialFunctionInstance::OverrideNamedTextureParameter(const FMaterialParameterInfo& ParameterInfo, UTexture*& OutValue)
+bool UMaterialFunctionInstance::OverrideNamedTextureParameter(const FHashedMaterialParameterInfo& ParameterInfo, UTexture*& OutValue)
 {
 	for (const FTextureParameterValue& TextureParameter : TextureParameterValues)
 	{
@@ -12416,7 +12368,7 @@ bool UMaterialFunctionInstance::OverrideNamedTextureParameter(const FMaterialPar
 	return false;
 }
 
-bool UMaterialFunctionInstance::OverrideNamedRuntimeVirtualTextureParameter(const FMaterialParameterInfo& ParameterInfo, URuntimeVirtualTexture*& OutValue)
+bool UMaterialFunctionInstance::OverrideNamedRuntimeVirtualTextureParameter(const FHashedMaterialParameterInfo& ParameterInfo, URuntimeVirtualTexture*& OutValue)
 {
 	for (const FRuntimeVirtualTextureParameterValue& RuntimeVirtualTextureParameter : RuntimeVirtualTextureParameterValues)
 	{
@@ -12430,7 +12382,7 @@ bool UMaterialFunctionInstance::OverrideNamedRuntimeVirtualTextureParameter(cons
 	return false;
 }
 
-bool UMaterialFunctionInstance::OverrideNamedFontParameter(const FMaterialParameterInfo& ParameterInfo, UFont*& OutFontValue, int32& OutFontPage)
+bool UMaterialFunctionInstance::OverrideNamedFontParameter(const FHashedMaterialParameterInfo& ParameterInfo, UFont*& OutFontValue, int32& OutFontPage)
 {
 	for (const FFontParameterValue& FontParameter : FontParameterValues)
 	{
@@ -12445,7 +12397,7 @@ bool UMaterialFunctionInstance::OverrideNamedFontParameter(const FMaterialParame
 	return false;
 }
 
-bool UMaterialFunctionInstance::OverrideNamedStaticSwitchParameter(const FMaterialParameterInfo& ParameterInfo, bool& OutValue, FGuid& OutExpressionGuid)
+bool UMaterialFunctionInstance::OverrideNamedStaticSwitchParameter(const FHashedMaterialParameterInfo& ParameterInfo, bool& OutValue, FGuid& OutExpressionGuid)
 {
 	for (const FStaticSwitchParameter& StaticSwitchParameter : StaticSwitchParameterValues)
 	{
@@ -12460,7 +12412,7 @@ bool UMaterialFunctionInstance::OverrideNamedStaticSwitchParameter(const FMateri
 	return false;
 }
 
-bool UMaterialFunctionInstance::OverrideNamedStaticComponentMaskParameter(const FMaterialParameterInfo& ParameterInfo, bool& OutR, bool& OutG, bool& OutB, bool& OutA, FGuid& OutExpressionGuid)
+bool UMaterialFunctionInstance::OverrideNamedStaticComponentMaskParameter(const FHashedMaterialParameterInfo& ParameterInfo, bool& OutR, bool& OutG, bool& OutB, bool& OutA, FGuid& OutExpressionGuid)
 {
 	for (const FStaticComponentMaskParameter& StaticComponentMaskParameter : StaticComponentMaskParameterValues)
 	{
@@ -12656,7 +12608,7 @@ bool UMaterialExpressionMaterialFunctionCall::NeedsLoadForClient() const
 	return true;
 }
 
-
+#if WITH_EDITORONLY_DATA
 bool UMaterialExpressionMaterialFunctionCall::IterateDependentFunctions(TFunctionRef<bool(UMaterialFunctionInterface*)> Predicate) const
 {
 	if (MaterialFunction)
@@ -12681,6 +12633,7 @@ void UMaterialExpressionMaterialFunctionCall::GetDependentFunctions(TArray<UMate
 		return true;
 	});
 }
+#endif // WITH_EDITORONLY_DATA
 
 #if WITH_EDITOR
 void UMaterialExpressionMaterialFunctionCall::PreEditChange(FProperty* PropertyAboutToChange)
@@ -13455,6 +13408,7 @@ void UMaterialExpressionFunctionInput::ConditionallyGenerateId(bool bForce)
 	}
 }
 
+#if WITH_EDITOR
 void UMaterialExpressionFunctionInput::ValidateName()
 {
 	if (Material)
@@ -13491,7 +13445,6 @@ void UMaterialExpressionFunctionInput::ValidateName()
 	}
 }
 
-#if WITH_EDITOR
 bool UMaterialExpressionFunctionInput::IsResultMaterialAttributes(int32 OutputIndex)
 {
 	if( FunctionInput_MaterialAttributes == InputType )
@@ -13539,7 +13492,7 @@ uint32 UMaterialExpressionFunctionInput::GetOutputType(int32 OutputIndex)
 {
 	return GetInputType(0);
 }
-#endif
+#endif // WITH_EDITOR
 
 
 ///////////////////////////////////////////////////////////////////////////////
@@ -13662,6 +13615,7 @@ void UMaterialExpressionFunctionOutput::ConditionallyGenerateId(bool bForce)
 	}
 }
 
+#if WITH_EDITOR
 void UMaterialExpressionFunctionOutput::ValidateName()
 {
 	if (Material)
@@ -13698,7 +13652,6 @@ void UMaterialExpressionFunctionOutput::ValidateName()
 	}
 }
 
-#if WITH_EDITOR
 bool UMaterialExpressionFunctionOutput::IsResultMaterialAttributes(int32 OutputIndex)
 {
 	// If there is a loop anywhere in this expression's inputs then we can't risk checking them
@@ -16888,7 +16841,7 @@ int32 UMaterialExpressionCurveAtlasRowParameter::Compile(class FMaterialCompiler
 	{
 		// Retrieve the curve index directly from the atlas rather than relying on the scalar parameter defaults
 		int32 CurveIndex = 0;
-		
+
 		if (Atlas->GetCurveIndex(Curve, CurveIndex))
 		{
 		DefaultValue = (float)CurveIndex;
