Copying Private-LoadTimes-4.24 stream to Main. Biggest changes are in Materials/Shader memory freezing.

diff --git a/Engine/Source/Runtime/Engine/Private/Materials/MaterialInstance.cpp b/Engine/Source/Runtime/Engine/Private/Materials/MaterialInstance.cpp
--- a/Engine/Source/Runtime/Engine/Private/Materials/MaterialInstance.cpp
+++ b/Engine/Source/Runtime/Engine/Private/Materials/MaterialInstance.cpp
@@ -197,14 +197,14 @@ UMaterialInterface* FMaterialInstanceResource::GetMaterialInterface() const
 }
 
 bool FMaterialInstanceResource::GetScalarValue(
-	const FMaterialParameterInfo& ParameterInfo, 
+	const FHashedMaterialParameterInfo& ParameterInfo,
 	float* OutValue,
 	const FMaterialRenderContext& Context
 	) const
 {
 	checkSlow(IsInParallelRenderingThread());
 
-	static FName NameSubsurfaceProfile(TEXT("__SubsurfaceProfile"));
+	static FHashedName NameSubsurfaceProfile(TEXT("__SubsurfaceProfile"));
 	if (ParameterInfo.Name == NameSubsurfaceProfile)
 	{
 		check(ParameterInfo.Association == EMaterialParameterAssociation::GlobalParameter);
@@ -243,7 +243,7 @@ bool FMaterialInstanceResource::GetScalarValue(
 }
 
 bool FMaterialInstanceResource::GetVectorValue(
-	const FMaterialParameterInfo& ParameterInfo, 
+	const FHashedMaterialParameterInfo& ParameterInfo,
 	FLinearColor* OutValue,
 	const FMaterialRenderContext& Context
 	) const
@@ -266,7 +266,7 @@ bool FMaterialInstanceResource::GetVectorValue(
 }
 
 bool FMaterialInstanceResource::GetTextureValue(
-	const FMaterialParameterInfo& ParameterInfo,
+	const FHashedMaterialParameterInfo& ParameterInfo,
 	const UTexture** OutValue,
 	const FMaterialRenderContext& Context
 	) const
@@ -289,7 +289,7 @@ bool FMaterialInstanceResource::GetTextureValue(
 }
 
 bool FMaterialInstanceResource::GetTextureValue(
-	const FMaterialParameterInfo& ParameterInfo,
+	const FHashedMaterialParameterInfo& ParameterInfo,
 	const URuntimeVirtualTexture** OutValue,
 	const FMaterialRenderContext& Context
 ) const
@@ -321,7 +321,8 @@ void UMaterialInstance::PropagateDataToMaterialProxy()
 
 void FMaterialInstanceResource::GameThread_SetParent(UMaterialInterface* ParentMaterialInterface)
 {
-	check(IsInGameThread() || IsAsyncLoading());
+	// @todo loadtimes: this is no longer valid because of the ParallelFor calling AddPrimitive in UnrealEngine.cpp
+	// check(IsInGameThread() || IsAsyncLoading());
 
 	if (GameThreadParent != ParentMaterialInterface)
 	{
@@ -530,14 +531,14 @@ const UMaterial* UMaterialInstance::GetMaterial() const
 	}
 }
 
-const UMaterial* UMaterialInstance::GetMaterial_Concurrent(TMicRecursionGuard& RecursionGuard) const
+const UMaterial* UMaterialInstance::GetMaterial_Concurrent(TMicRecursionGuard RecursionGuard) const
 {
 	if(!Parent || RecursionGuard.Contains(this))
 	{
 		return UMaterial::GetDefaultMaterial(MD_Surface);
 	}
 
-	RecursionGuard.Add(this);
+	RecursionGuard.Set(this);
 	return Parent->GetMaterial_Concurrent(RecursionGuard);
 }
 
@@ -560,7 +561,7 @@ UMaterial* UMaterialInstance::GetMaterial()
 }
 
 #if WITH_EDITOR
-bool UMaterialInstance::GetScalarParameterSliderMinMax(const FMaterialParameterInfo& ParameterInfo, float& OutSliderMin, float& OutSliderMax) const
+bool UMaterialInstance::GetScalarParameterSliderMinMax(const FHashedMaterialParameterInfo& ParameterInfo, float& OutSliderMin, float& OutSliderMax) const
 {
 	if (GetReentrantFlag())
 	{
@@ -569,46 +570,200 @@ bool UMaterialInstance::GetScalarParameterSliderMinMax(const FMaterialParameterI
 
 	if (ParameterInfo.Association != EMaterialParameterAssociation::GlobalParameter)
 	{
-		UMaterialExpressionScalarParameter* Parameter = nullptr;
+		const int32 ParameterIndex = CachedLayerParameters.FindParameterIndex(EMaterialParameterType::Scalar, ParameterInfo);
+		if (ParameterIndex != INDEX_NONE)
+		{
+			OutSliderMin = CachedLayerParameters.ScalarMinMaxValues[ParameterIndex].X;
+			OutSliderMax = CachedLayerParameters.ScalarMinMaxValues[ParameterIndex].Y;
+			return true;
+		}
+	}
 
-		for (const FStaticMaterialLayersParameter& LayersParam : StaticParameters.MaterialLayersParameters)
+	if (Parent)
+	{
+		FMICReentranceGuard	Guard(this);
+		return Parent->GetScalarParameterSliderMinMax(ParameterInfo, OutSliderMin, OutSliderMax);
+	}
+
+	return false;
+}
+#endif // WITH_EDITOR
+
+#if WITH_EDITOR
+
+#define checkLayerParameter(x) check(x)
+
+bool UMaterialInstance::GetScalarLayerParameterValue_Legacy(const FHashedMaterialParameterInfo& ParameterInfo, bool bOveriddenOnly, float& OutValue, bool& OutResult) const
+{
+	UMaterialExpressionScalarParameter* Parameter = nullptr;
+	for (const FStaticMaterialLayersParameter& LayersParam : StaticParameters.MaterialLayersParameters)
+	{
+		if (LayersParam.bOverride)
 		{
-			if (LayersParam.bOverride)
+			UMaterialFunctionInterface* Function = LayersParam.GetParameterAssociatedFunction(ParameterInfo);
+			UMaterialFunctionInterface* ParameterOwner = nullptr;
+
+			if (Function && Function->OverrideNamedScalarParameter(ParameterInfo, OutValue))
+			{
+				OutResult = true;
+				return true;
+			}
+
+			if (Function && Function->GetNamedParameterOfType(ParameterInfo, Parameter, &ParameterOwner))
 			{
-				for (UMaterialFunctionInterface* Layer : LayersParam.Value.Layers)
+				if (ParameterOwner->OverrideNamedScalarParameter(ParameterInfo, OutValue))
 				{
-					if (Layer && Layer->GetNamedParameterOfType<UMaterialExpressionScalarParameter>(ParameterInfo, Parameter))
-					{
-						OutSliderMin = Parameter->SliderMin;
-						OutSliderMax = Parameter->SliderMax;
-						return true;
-					}
+					OutResult = true;
+					return true;
 				}
+				
+				Parameter->IsNamedParameter(ParameterInfo, OutValue);
+				OutResult = !bOveriddenOnly;
+				return true;
+			}
+		}
+	}
+
+	return false;
+}
+
+bool UMaterialInstance::GetVectorLayerParameterValue_Legacy(const FHashedMaterialParameterInfo& ParameterInfo, bool bOveriddenOnly, FLinearColor& OutValue, bool& OutResult) const
+{
+	UMaterialExpressionVectorParameter* Parameter = nullptr;
+	for (const FStaticMaterialLayersParameter& LayersParam : StaticParameters.MaterialLayersParameters)
+	{
+		if (LayersParam.bOverride)
+		{
+			UMaterialFunctionInterface* Function = LayersParam.GetParameterAssociatedFunction(ParameterInfo);
+			UMaterialFunctionInterface* ParameterOwner = nullptr;
 
-				for (UMaterialFunctionInterface* Blend : LayersParam.Value.Blends)
+			if (Function && Function->OverrideNamedVectorParameter(ParameterInfo, OutValue))
+			{
+				OutResult = true;
+				return true;
+			}
+
+			if (Function && Function->GetNamedParameterOfType(ParameterInfo, Parameter, &ParameterOwner))
+			{
+				if (ParameterOwner->OverrideNamedVectorParameter(ParameterInfo, OutValue))
 				{
-					if (Blend && Blend->GetNamedParameterOfType<UMaterialExpressionScalarParameter>(ParameterInfo, Parameter))
-					{
-						OutSliderMin = Parameter->SliderMin;
-						OutSliderMax = Parameter->SliderMax;
-						return true;
-					}
+					OutResult = true;
+					return true;
 				}
+
+				Parameter->IsNamedParameter(ParameterInfo, OutValue);
+				OutResult = !bOveriddenOnly;
+				return true;
 			}
 		}
 	}
 
-	if (Parent)
+	return false;
+}
+
+bool UMaterialInstance::GetTextureLayerParameterValue_Legacy(const FHashedMaterialParameterInfo& ParameterInfo, bool bOveriddenOnly, UTexture*& OutValue, bool& OutResult) const
+{
+	UMaterialExpressionTextureSampleParameter* Parameter = nullptr;
+	for (const FStaticMaterialLayersParameter& LayersParam : StaticParameters.MaterialLayersParameters)
 	{
-		FMICReentranceGuard	Guard(this);
-		return Parent->GetScalarParameterSliderMinMax(ParameterInfo, OutSliderMin, OutSliderMax);
+		if (LayersParam.bOverride)
+		{
+			UMaterialFunctionInterface* Function = LayersParam.GetParameterAssociatedFunction(ParameterInfo);
+			UMaterialFunctionInterface* ParameterOwner = nullptr;
+
+			if (Function && Function->OverrideNamedTextureParameter(ParameterInfo, OutValue))
+			{
+				OutResult = true;
+				return true;
+			}
+
+			if (Function && Function->GetNamedParameterOfType(ParameterInfo, Parameter, &ParameterOwner))
+			{
+				// TODO - Different logic than other parameter types for 'OutResult', intentional??
+				if (!ParameterOwner->OverrideNamedTextureParameter(ParameterInfo, OutValue))
+				{
+					Parameter->IsNamedParameter(ParameterInfo, OutValue);
+				}
+				OutResult = true;
+				return true;
+			}
+		}
 	}
 
 	return false;
 }
+
+bool UMaterialInstance::GetFontLayerParameterValue_Legacy(const FHashedMaterialParameterInfo& ParameterInfo, bool bOveriddenOnly, UFont*& OutFontValue, int32& OutFontPage, bool& OutResult) const
+{
+	UMaterialExpressionFontSampleParameter* Parameter = nullptr;
+	for (const FStaticMaterialLayersParameter& LayersParam : StaticParameters.MaterialLayersParameters)
+	{
+		if (LayersParam.bOverride)
+		{
+			UMaterialFunctionInterface* Function = LayersParam.GetParameterAssociatedFunction(ParameterInfo);
+			UMaterialFunctionInterface* ParameterOwner = nullptr;
+
+			if (Function && Function->OverrideNamedFontParameter(ParameterInfo, OutFontValue, OutFontPage))
+			{
+				OutResult = true;
+				return true;
+			}
+
+			if (Function && Function->GetNamedParameterOfType(ParameterInfo, Parameter, &ParameterOwner))
+			{
+				if (ParameterOwner->OverrideNamedFontParameter(ParameterInfo, OutFontValue, OutFontPage))
+				{
+					OutResult = true;
+					return true;
+				}
+				Parameter->IsNamedParameter(ParameterInfo, OutFontValue, OutFontPage);
+				OutResult = !bOveriddenOnly;
+				return true;
+			}
+		}
+	}
+
+	return false;
+}
+
+bool UMaterialInstance::GetRuntimeVirtualTextureLayerParameterValue_Legacy(const FHashedMaterialParameterInfo& ParameterInfo, bool bOveriddenOnly, URuntimeVirtualTexture*& OutValue, bool& OutResult) const
+{
+	UMaterialExpressionRuntimeVirtualTextureSampleParameter* Parameter = nullptr;
+	for (const FStaticMaterialLayersParameter& LayersParam : StaticParameters.MaterialLayersParameters)
+	{
+		if (LayersParam.bOverride)
+		{
+			UMaterialFunctionInterface* Function = LayersParam.GetParameterAssociatedFunction(ParameterInfo);
+			UMaterialFunctionInterface* ParameterOwner = nullptr;
+
+			if (Function && Function->OverrideNamedRuntimeVirtualTextureParameter(ParameterInfo, OutValue))
+			{
+				OutResult = true;
+				return true;
+			}
+
+			if (Function && Function->GetNamedParameterOfType(ParameterInfo, Parameter, &ParameterOwner))
+			{
+				if (!ParameterOwner->OverrideNamedRuntimeVirtualTextureParameter(ParameterInfo, OutValue))
+				{
+					Parameter->IsNamedParameter(ParameterInfo, OutValue);
+				}
+				OutResult = true;
+				return true;
+			}
+		}
+	}
+
+	return false;
+}
+
 #endif // WITH_EDITOR
 
-bool UMaterialInstance::GetScalarParameterValue(const FMaterialParameterInfo& ParameterInfo, float& OutValue, bool bOveriddenOnly) const
+#if !defined(checkLayerParameter)
+#define checkLayerParameter(x)
+#endif
+
+bool UMaterialInstance::GetScalarParameterValue(const FHashedMaterialParameterInfo& ParameterInfo, float& OutValue, bool bOveriddenOnly) const
 {
 	bool bFoundAValue = false;
 
@@ -628,31 +783,25 @@ bool UMaterialInstance::GetScalarParameterValue(const FMaterialParameterInfo& Pa
 	// Instance-included default
 	if (ParameterInfo.Association != EMaterialParameterAssociation::GlobalParameter)
 	{
-		UMaterialExpressionScalarParameter* Parameter = nullptr;
-		for (const FStaticMaterialLayersParameter& LayersParam : StaticParameters.MaterialLayersParameters)
-		{
-			if (LayersParam.bOverride)
-			{
-				UMaterialFunctionInterface* Function = LayersParam.GetParameterAssociatedFunction(ParameterInfo);
-				UMaterialFunctionInterface* ParameterOwner = nullptr;
-
-				if (Function && Function->OverrideNamedScalarParameter(ParameterInfo, OutValue))
-				{
-					return true;
-				}
+#if WITH_EDITOR
+		bool bLegacyFoundValue = false;
+		bool bLegacyResult = false;
+		float LegacyValue = 0.0f;
+		bLegacyFoundValue = GetScalarLayerParameterValue_Legacy(ParameterInfo, bOveriddenOnly, LegacyValue, bLegacyResult);
+#endif // WITH_EDITOR
 
-				if (Function && Function->GetNamedParameterOfType(ParameterInfo, Parameter, &ParameterOwner))
-				{
-					if (ParameterOwner->OverrideNamedScalarParameter(ParameterInfo, OutValue))
-					{
-						return true;		
-					}
-					
-					Parameter->IsNamedParameter(ParameterInfo, OutValue);
-					return !bOveriddenOnly;
-				}
-			}
+		const int32 ParameterIndex = CachedLayerParameters.FindParameterIndex(EMaterialParameterType::Scalar, ParameterInfo);
+		if (ParameterIndex != INDEX_NONE)
+		{
+			checkLayerParameter(bLegacyFoundValue);
+			const bool bResult = CachedLayerParameters.IsParameterValid(EMaterialParameterType::Scalar, ParameterIndex, bOveriddenOnly);
+			OutValue = CachedLayerParameters.ScalarValues[ParameterIndex];
+			checkLayerParameter(OutValue == LegacyValue);
+			checkLayerParameter(bLegacyResult == bResult);
+			return bResult;
 		}
+
+		checkLayerParameter(!bLegacyFoundValue);
 	}
 	
 	// Next material in hierarchy
@@ -665,7 +814,8 @@ bool UMaterialInstance::GetScalarParameterValue(const FMaterialParameterInfo& Pa
 	return false;
 }
 
-bool UMaterialInstance::IsScalarParameterUsedAsAtlasPosition(const FMaterialParameterInfo& ParameterInfo, bool& OutValue, TSoftObjectPtr<UCurveLinearColor>& Curve, TSoftObjectPtr<UCurveLinearColorAtlas>& Atlas) const
+#if WITH_EDITOR
+bool UMaterialInstance::IsScalarParameterUsedAsAtlasPosition(const FHashedMaterialParameterInfo& ParameterInfo, bool& OutValue, TSoftObjectPtr<UCurveLinearColor>& OutCurve, TSoftObjectPtr<UCurveLinearColorAtlas>& OutAtlas) const
 {
 	bool bFoundAValue = false;
 
@@ -680,8 +830,8 @@ bool UMaterialInstance::IsScalarParameterUsedAsAtlasPosition(const FMaterialPara
 	if (ParameterValue && ParameterValue->AtlasData.Curve.Get() != nullptr && ParameterValue->AtlasData.Atlas.Get() != nullptr)
 	{
 		OutValue = ParameterValue->AtlasData.bIsUsedAsAtlasPosition;
-		Curve = ParameterValue->AtlasData.Curve;
-		Atlas = ParameterValue->AtlasData.Atlas;
+		OutCurve = ParameterValue->AtlasData.Curve;
+		OutAtlas = ParameterValue->AtlasData.Atlas;
 		return true;
 	}
 #endif
@@ -689,26 +839,22 @@ bool UMaterialInstance::IsScalarParameterUsedAsAtlasPosition(const FMaterialPara
 	// Instance-included default
 	if (ParameterInfo.Association != EMaterialParameterAssociation::GlobalParameter)
 	{
-		UMaterialExpressionScalarParameter* Parameter = nullptr;
-		for (const FStaticMaterialLayersParameter& LayersParam : StaticParameters.MaterialLayersParameters)
+		const int32 ParameterIndex = CachedLayerParameters.FindParameterIndex(EMaterialParameterType::Scalar, ParameterInfo);
+		if (ParameterIndex != INDEX_NONE)
 		{
-			if (LayersParam.bOverride)
+			UCurveLinearColor* Curve = CachedLayerParameters.ScalarCurveValues[ParameterIndex];
+			UCurveLinearColorAtlas* Atlas = CachedLayerParameters.ScalarCurveAtlasValues[ParameterIndex];
+			if (Curve && Atlas)
 			{
-				UMaterialFunctionInterface* Function = LayersParam.GetParameterAssociatedFunction(ParameterInfo);
-				UMaterialFunctionInterface* ParameterOwner = nullptr;
-
-				if (Function && Function->GetNamedParameterOfType(ParameterInfo, Parameter, &ParameterOwner))
-				{
-					OutValue = Parameter->IsUsedAsAtlasPosition();
-					if (OutValue)
-					{
-						UMaterialExpressionCurveAtlasRowParameter* AtlasParameter = Cast<UMaterialExpressionCurveAtlasRowParameter>(Parameter);
-						Curve = TSoftObjectPtr<UCurveLinearColor>(FSoftObjectPath(AtlasParameter->Curve->GetPathName()));
-						Atlas = TSoftObjectPtr<UCurveLinearColorAtlas> (FSoftObjectPath(AtlasParameter->Atlas->GetPathName()));
-					}
-					return true;
-				}
+				OutCurve = TSoftObjectPtr<UCurveLinearColor>(FSoftObjectPath(Curve->GetPathName()));
+				OutAtlas = TSoftObjectPtr<UCurveLinearColorAtlas>(FSoftObjectPath(Atlas->GetPathName()));
+				OutValue = true;
 			}
+			else
+			{
+				OutValue = false;
+			}
+			return true;
 		}
 	}
 
@@ -716,13 +862,14 @@ bool UMaterialInstance::IsScalarParameterUsedAsAtlasPosition(const FMaterialPara
 	if (Parent)
 	{
 		FMICReentranceGuard	Guard(this);
-		return Parent->IsScalarParameterUsedAsAtlasPosition(ParameterInfo, OutValue, Curve, Atlas);
+		return Parent->IsScalarParameterUsedAsAtlasPosition(ParameterInfo, OutValue, OutCurve, OutAtlas);
 	}
 
 	return false;
 }
+#endif // WITH_EDITOR
 
-bool UMaterialInstance::GetVectorParameterValue(const FMaterialParameterInfo& ParameterInfo, FLinearColor& OutValue, bool bOveriddenOnly) const
+bool UMaterialInstance::GetVectorParameterValue(const FHashedMaterialParameterInfo& ParameterInfo, FLinearColor& OutValue, bool bOveriddenOnly) const
 {
 	bool bFoundAValue = false;
 
@@ -742,31 +889,25 @@ bool UMaterialInstance::GetVectorParameterValue(const FMaterialParameterInfo& Pa
 	// Instance-included default
 	if (ParameterInfo.Association != EMaterialParameterAssociation::GlobalParameter)
 	{
-		UMaterialExpressionVectorParameter* Parameter = nullptr;
-		for (const FStaticMaterialLayersParameter& LayersParam : StaticParameters.MaterialLayersParameters)
-		{
-			if (LayersParam.bOverride)
-			{
-				UMaterialFunctionInterface* Function = LayersParam.GetParameterAssociatedFunction(ParameterInfo);
-				UMaterialFunctionInterface* ParameterOwner = nullptr;
-
-				if (Function && Function->OverrideNamedVectorParameter(ParameterInfo, OutValue))
-				{
-					return true;
-				}
+#if WITH_EDITOR
+		bool bLegacyFoundValue = false;
+		bool bLegacyResult = false;
+		FLinearColor LegacyValue(ForceInitToZero);
+		bLegacyFoundValue = GetVectorLayerParameterValue_Legacy(ParameterInfo, bOveriddenOnly, LegacyValue, bLegacyResult);
+#endif // WITH_EDITOR
 
-				if (Function && Function->GetNamedParameterOfType(ParameterInfo, Parameter, &ParameterOwner))
-				{
-					if (ParameterOwner->OverrideNamedVectorParameter(ParameterInfo, OutValue))
-					{
-						return true;
-					}
-					
-					Parameter->IsNamedParameter(ParameterInfo, OutValue);
-					return !bOveriddenOnly;
-				}
-			}
+		const int32 ParameterIndex = CachedLayerParameters.FindParameterIndex(EMaterialParameterType::Vector, ParameterInfo);
+		if (ParameterIndex != INDEX_NONE)
+		{
+			checkLayerParameter(bLegacyFoundValue);
+			const bool bResult = CachedLayerParameters.IsParameterValid(EMaterialParameterType::Vector, ParameterIndex, bOveriddenOnly);
+			OutValue = CachedLayerParameters.VectorValues[ParameterIndex];
+			checkLayerParameter(OutValue == LegacyValue);
+			checkLayerParameter(bLegacyResult == bResult);
+			return bResult;
 		}
+
+		checkLayerParameter(!bLegacyFoundValue);
 	}
 	
 	// Next material in hierarchy
@@ -779,7 +920,8 @@ bool UMaterialInstance::GetVectorParameterValue(const FMaterialParameterInfo& Pa
 	return false;
 }
 
-bool UMaterialInstance::IsVectorParameterUsedAsChannelMask(const FMaterialParameterInfo& ParameterInfo, bool& OutValue) const
+#if WITH_EDITOR
+bool UMaterialInstance::IsVectorParameterUsedAsChannelMask(const FHashedMaterialParameterInfo& ParameterInfo, bool& OutValue) const
 {
 	bool bFoundAValue = false;
 
@@ -791,20 +933,11 @@ bool UMaterialInstance::IsVectorParameterUsedAsChannelMask(const FMaterialParame
 	// Instance-included default
 	if (ParameterInfo.Association != EMaterialParameterAssociation::GlobalParameter)
 	{
-		UMaterialExpressionVectorParameter* Parameter = nullptr;
-		for (const FStaticMaterialLayersParameter& LayersParam : StaticParameters.MaterialLayersParameters)
+		const int32 ParameterIndex = CachedLayerParameters.FindParameterIndex(EMaterialParameterType::Vector, ParameterInfo);
+		if (ParameterIndex != INDEX_NONE)
 		{
-			if (LayersParam.bOverride)
-			{
-				UMaterialFunctionInterface* Function = LayersParam.GetParameterAssociatedFunction(ParameterInfo);
-				UMaterialFunctionInterface* ParameterOwner = nullptr;
-
-				if (Function && Function->GetNamedParameterOfType(ParameterInfo, Parameter, &ParameterOwner))
-				{					
-					OutValue = Parameter->IsUsedAsChannelMask();
-					return true;
-				}
-			}
+			OutValue = CachedLayerParameters.VectorUsedAsChannelMaskValues[ParameterIndex];
+			return true;
 		}
 	}
 	
@@ -818,8 +951,7 @@ bool UMaterialInstance::IsVectorParameterUsedAsChannelMask(const FMaterialParame
 	return false;
 }
 
-#if WITH_EDITOR
-bool  UMaterialInstance::GetVectorParameterChannelNames(const FMaterialParameterInfo& ParameterInfo, FParameterChannelNames& OutValue) const
+bool  UMaterialInstance::GetVectorParameterChannelNames(const FHashedMaterialParameterInfo& ParameterInfo, FParameterChannelNames& OutValue) const
 {
 	bool bFoundAValue = false;
 
@@ -831,20 +963,11 @@ bool  UMaterialInstance::GetVectorParameterChannelNames(const FMaterialParameter
 	// Instance-included default
 	if (ParameterInfo.Association != EMaterialParameterAssociation::GlobalParameter)
 	{
-		UMaterialExpressionVectorParameter* Parameter = nullptr;
-		for (const FStaticMaterialLayersParameter& LayersParam : StaticParameters.MaterialLayersParameters)
+		const int32 ParameterIndex = CachedLayerParameters.FindParameterIndex(EMaterialParameterType::Vector, ParameterInfo);
+		if (ParameterIndex != INDEX_NONE)
 		{
-			if (LayersParam.bOverride)
-			{
-				UMaterialFunctionInterface* Function = LayersParam.GetParameterAssociatedFunction(ParameterInfo);
-				UMaterialFunctionInterface* ParameterOwner = nullptr;
-
-				if (Function && Function->GetNamedParameterOfType(ParameterInfo, Parameter, &ParameterOwner))
-				{
-					OutValue = Parameter->GetVectorChannelNames();
-					return true;
-				}
-			}
+			OutValue = CachedLayerParameters.VectorChannelNameValues[ParameterIndex];
+			return true;
 		}
 	}
 
@@ -857,9 +980,9 @@ bool  UMaterialInstance::GetVectorParameterChannelNames(const FMaterialParameter
 
 	return false;
 };
-#endif
+#endif // WITH_EDITOR
 
-bool UMaterialInstance::GetTextureParameterValue(const FMaterialParameterInfo& ParameterInfo, UTexture*& OutValue, bool bOveriddenOnly) const
+bool UMaterialInstance::GetTextureParameterValue(const FHashedMaterialParameterInfo& ParameterInfo, UTexture*& OutValue, bool bOveriddenOnly) const
 {
 	bool bFoundAValue = false;
 
@@ -879,29 +1002,25 @@ bool UMaterialInstance::GetTextureParameterValue(const FMaterialParameterInfo& P
 	// Instance-included default
 	if (ParameterInfo.Association != EMaterialParameterAssociation::GlobalParameter)
 	{
-		UMaterialExpressionTextureSampleParameter* Parameter = nullptr;
-		for (const FStaticMaterialLayersParameter& LayersParam : StaticParameters.MaterialLayersParameters)
-		{
-			if (LayersParam.bOverride)
-			{
-				UMaterialFunctionInterface* Function = LayersParam.GetParameterAssociatedFunction(ParameterInfo);
-				UMaterialFunctionInterface* ParameterOwner = nullptr;
-
-				if (Function && Function->OverrideNamedTextureParameter(ParameterInfo, OutValue))
-				{
-					return true;
-				}
+#if WITH_EDITOR
+		bool bLegacyFoundValue = false;
+		bool bLegacyResult = false;
+		UTexture* LegacyValue = nullptr;
+		bLegacyFoundValue = GetTextureLayerParameterValue_Legacy(ParameterInfo, bOveriddenOnly, LegacyValue, bLegacyResult);
+#endif // WITH_EDITOR
 
-				if (Function && Function->GetNamedParameterOfType(ParameterInfo, Parameter, &ParameterOwner))
-				{
-					if (!ParameterOwner->OverrideNamedTextureParameter(ParameterInfo, OutValue))
-					{
-						Parameter->IsNamedParameter(ParameterInfo, OutValue);
-					}
-					return true;
-				}
-			}
+		const int32 ParameterIndex = CachedLayerParameters.FindParameterIndex(EMaterialParameterType::Texture, ParameterInfo);
+		if (ParameterIndex != INDEX_NONE)
+		{
+			checkLayerParameter(bLegacyFoundValue);
+			const bool bResult = CachedLayerParameters.IsParameterValid(EMaterialParameterType::Texture, ParameterIndex, bOveriddenOnly);
+			OutValue = CachedLayerParameters.TextureValues[ParameterIndex];
+			checkLayerParameter(OutValue == LegacyValue);
+			checkLayerParameter(bLegacyResult == bResult);
+			return bResult;
 		}
+
+		checkLayerParameter(!bLegacyFoundValue);
 	}
 	
 	// Next material in hierarchy
@@ -914,7 +1033,7 @@ bool UMaterialInstance::GetTextureParameterValue(const FMaterialParameterInfo& P
 	return false;
 }
 
-bool UMaterialInstance::GetRuntimeVirtualTextureParameterValue(const FMaterialParameterInfo& ParameterInfo, URuntimeVirtualTexture*& OutValue, bool bOveriddenOnly) const
+bool UMaterialInstance::GetRuntimeVirtualTextureParameterValue(const FHashedMaterialParameterInfo& ParameterInfo, URuntimeVirtualTexture*& OutValue, bool bOveriddenOnly) const
 {
 	bool bFoundAValue = false;
 
@@ -934,29 +1053,25 @@ bool UMaterialInstance::GetRuntimeVirtualTextureParameterValue(const FMaterialPa
 	// Instance-included default
 	if (ParameterInfo.Association != EMaterialParameterAssociation::GlobalParameter)
 	{
-		UMaterialExpressionRuntimeVirtualTextureSampleParameter* Parameter = nullptr;
-		for (const FStaticMaterialLayersParameter& LayersParam : StaticParameters.MaterialLayersParameters)
-		{
-			if (LayersParam.bOverride)
-			{
-				UMaterialFunctionInterface* Function = LayersParam.GetParameterAssociatedFunction(ParameterInfo);
-				UMaterialFunctionInterface* ParameterOwner = nullptr;
-
-				if (Function && Function->OverrideNamedRuntimeVirtualTextureParameter(ParameterInfo, OutValue))
-				{
-					return true;
-				}
+#if WITH_EDITOR
+		bool bLegacyFoundValue = false;
+		bool bLegacyResult = false;
+		URuntimeVirtualTexture* LegacyValue = nullptr;
+		bLegacyFoundValue = GetRuntimeVirtualTextureLayerParameterValue_Legacy(ParameterInfo, bOveriddenOnly, LegacyValue, bLegacyResult);
+#endif // WITH_EDITOR
 
-				if (Function && Function->GetNamedParameterOfType(ParameterInfo, Parameter, &ParameterOwner))
-				{
-					if (!ParameterOwner->OverrideNamedRuntimeVirtualTextureParameter(ParameterInfo, OutValue))
-					{
-						Parameter->IsNamedParameter(ParameterInfo, OutValue);
-					}
-					return true;
-				}
-			}
+		const int32 ParameterIndex = CachedLayerParameters.FindParameterIndex(EMaterialParameterType::RuntimeVirtualTexture, ParameterInfo);
+		if (ParameterIndex != INDEX_NONE)
+		{
+			checkLayerParameter(bLegacyFoundValue);
+			const bool bResult = CachedLayerParameters.IsParameterValid(EMaterialParameterType::RuntimeVirtualTexture, ParameterIndex, bOveriddenOnly);
+			OutValue = CachedLayerParameters.RuntimeVirtualTextureValues[ParameterIndex];
+			checkLayerParameter(OutValue == LegacyValue);
+			checkLayerParameter(bLegacyResult == bResult);
+			return bResult;
 		}
+
+		checkLayerParameter(!bLegacyFoundValue);
 	}
 
 	// Next material in hierarchy
@@ -971,7 +1086,7 @@ bool UMaterialInstance::GetRuntimeVirtualTextureParameterValue(const FMaterialPa
 
 
 #if WITH_EDITOR
-bool  UMaterialInstance::GetTextureParameterChannelNames(const FMaterialParameterInfo& ParameterInfo, FParameterChannelNames& OutValue) const
+bool  UMaterialInstance::GetTextureParameterChannelNames(const FHashedMaterialParameterInfo& ParameterInfo, FParameterChannelNames& OutValue) const
 {
 	bool bFoundAValue = false;
 
@@ -983,20 +1098,11 @@ bool  UMaterialInstance::GetTextureParameterChannelNames(const FMaterialParamete
 	// Instance-included default
 	if (ParameterInfo.Association != EMaterialParameterAssociation::GlobalParameter)
 	{
-		UMaterialExpressionTextureSampleParameter* Parameter = nullptr;
-		for (const FStaticMaterialLayersParameter& LayersParam : StaticParameters.MaterialLayersParameters)
+		const int32 ParameterIndex = CachedLayerParameters.FindParameterIndex(EMaterialParameterType::Texture, ParameterInfo);
+		if (ParameterIndex != INDEX_NONE)
 		{
-			if (LayersParam.bOverride)
-			{
-				UMaterialFunctionInterface* Function = LayersParam.GetParameterAssociatedFunction(ParameterInfo);
-				UMaterialFunctionInterface* ParameterOwner = nullptr;
-
-				if (Function && Function->GetNamedParameterOfType(ParameterInfo, Parameter, &ParameterOwner))
-				{
-					OutValue = Parameter->GetTextureChannelNames();
-					return true;
-				}
-			}
+			OutValue = CachedLayerParameters.TextureChannelNameValues[ParameterIndex];
+			return true;
 		}
 	}
 
@@ -1009,9 +1115,9 @@ bool  UMaterialInstance::GetTextureParameterChannelNames(const FMaterialParamete
 
 	return false;
 };
-#endif
+#endif // WITH_EDITOR
 
-bool UMaterialInstance::GetFontParameterValue(const FMaterialParameterInfo& ParameterInfo,class UFont*& OutFontValue, int32& OutFontPage, bool bOveriddenOnly) const
+bool UMaterialInstance::GetFontParameterValue(const FHashedMaterialParameterInfo& ParameterInfo,class UFont*& OutFontValue, int32& OutFontPage, bool bOveriddenOnly) const
 {
 	bool bFoundAValue = false;
 
@@ -1032,30 +1138,28 @@ bool UMaterialInstance::GetFontParameterValue(const FMaterialParameterInfo& Para
 	// Instance-included default
 	if (ParameterInfo.Association != EMaterialParameterAssociation::GlobalParameter)
 	{
-		UMaterialExpressionFontSampleParameter* Parameter = nullptr;
-		for (const FStaticMaterialLayersParameter& LayersParam : StaticParameters.MaterialLayersParameters)
-		{
-			if (LayersParam.bOverride)
-			{
-				UMaterialFunctionInterface* Function = LayersParam.GetParameterAssociatedFunction(ParameterInfo);
-				UMaterialFunctionInterface* ParameterOwner = nullptr;
-
-				if (Function && Function->OverrideNamedFontParameter(ParameterInfo, OutFontValue, OutFontPage))
-				{
-					return true;
-				}
+#if WITH_EDITOR
+		bool bLegacyFoundValue = false;
+		bool bLegacyResult = false;
+		UFont* LegacyValue = nullptr;
+		int32 LegacyFontPageValue = INDEX_NONE;
+		bLegacyFoundValue = GetFontLayerParameterValue_Legacy(ParameterInfo, bOveriddenOnly, LegacyValue, LegacyFontPageValue, bLegacyResult);
+#endif // WITH_EDITOR
 
-				if (Function && Function->GetNamedParameterOfType(ParameterInfo, Parameter, &ParameterOwner))
-				{
-					if (ParameterOwner->OverrideNamedFontParameter(ParameterInfo, OutFontValue, OutFontPage))
-					{
-						return true;
-					}
-					Parameter->IsNamedParameter(ParameterInfo, OutFontValue, OutFontPage);
-					return !bOveriddenOnly;
-				}
-			}
+		const int32 ParameterIndex = CachedLayerParameters.FindParameterIndex(EMaterialParameterType::Font, ParameterInfo);
+		if (ParameterIndex != INDEX_NONE)
+		{
+			checkLayerParameter(bLegacyFoundValue);
+			const bool bResult = CachedLayerParameters.IsParameterValid(EMaterialParameterType::Font, ParameterIndex, bOveriddenOnly);
+			OutFontValue = CachedLayerParameters.FontValues[ParameterIndex];
+			OutFontPage = CachedLayerParameters.FontPageValues[ParameterIndex];
+			checkLayerParameter(OutFontValue == LegacyValue);
+			checkLayerParameter(OutFontPage == LegacyFontPageValue);
+			checkLayerParameter(bLegacyResult == bResult);
+			return bResult;
 		}
+
+		checkLayerParameter(!bLegacyFoundValue);
 	}
 	
 	// Next material in hierarchy
@@ -1102,78 +1206,45 @@ bool UMaterialInstance::GetRefractionSettings(float& OutBiasValue) const
 
 void UMaterialInstance::GetTextureExpressionValues(const FMaterialResource* MaterialResource, TArray<UTexture*>& OutTextures, TArray< TArray<int32> >* OutIndices) const
 {
-	const TArray<TRefCountPtr<FMaterialUniformExpressionTexture> >* ExpressionsByType[5];
-
 	check(MaterialResource);
-
-	ExpressionsByType[0] = &MaterialResource->GetUniform2DTextureExpressions();
-	ExpressionsByType[1] = &MaterialResource->GetUniformCubeTextureExpressions();
-	ExpressionsByType[2] = &MaterialResource->GetUniform2DArrayTextureExpressions();
-	ExpressionsByType[3] = &MaterialResource->GetUniformVolumeTextureExpressions();
-	ExpressionsByType[4] = &MaterialResource->GetUniformVirtualTextureExpressions();
+	const FUniformExpressionSet& UniformExpressions = MaterialResource->GetUniformExpressions();
 
 	if (OutIndices) // Try to prevent resizing since this would be expensive.
 	{
-		OutIndices->Empty(ExpressionsByType[0]->Num() + ExpressionsByType[1]->Num() + ExpressionsByType[2]->Num() + ExpressionsByType[3]->Num() + ExpressionsByType[4]->Num());
+		uint32 NumTextures = 0u;
+		for (uint32 TypeIndex = 0u; TypeIndex < NumMaterialTextureParameterTypes; ++TypeIndex)
+		{
+			NumTextures += UniformExpressions.GetNumTextures((EMaterialTextureParameterType)TypeIndex);
+		}
+		OutIndices->Empty(NumTextures);
 	}
 
-	for(int32 TypeIndex = 0;TypeIndex < UE_ARRAY_COUNT(ExpressionsByType);TypeIndex++)
+	for(int32 TypeIndex = 0;TypeIndex < NumMaterialTextureParameterTypes;TypeIndex++)
 	{
-		const TArray<TRefCountPtr<FMaterialUniformExpressionTexture> >& Expressions = *ExpressionsByType[TypeIndex];
-
 		// Iterate over each of the material's texture expressions.
-		for(int32 ExpressionIndex = 0;ExpressionIndex < Expressions.Num();ExpressionIndex++)
+		for(int32 TextureIndex = 0; TextureIndex < UniformExpressions.GetNumTextures((EMaterialTextureParameterType)TypeIndex); TextureIndex++)
 		{
-			FMaterialUniformExpressionTexture* Expression = Expressions[ExpressionIndex];
-
 			// Evaluate the expression in terms of this material instance.
 			UTexture* Texture = NULL;
-			Expression->GetGameThreadTextureValue(this,*MaterialResource,Texture, true);
+			UniformExpressions.GetGameThreadTextureValue((EMaterialTextureParameterType)TypeIndex, TextureIndex, this,*MaterialResource,Texture, true);
 			
 			if (Texture)
 			{
-				int32 InsertIndex = OutTextures.AddUnique(Texture);
-
+				const int32 InsertIndex = OutTextures.AddUnique(Texture);
 				if (OutIndices)
 				{
+					const FMaterialTextureParameterInfo& Parameter = UniformExpressions.GetTextureParameter((EMaterialTextureParameterType)TypeIndex, TextureIndex);
 					if (InsertIndex >= OutIndices->Num())
 					{
 						OutIndices->AddDefaulted(InsertIndex - OutIndices->Num() + 1);
 					}
-					(*OutIndices)[InsertIndex].Add(Expression->GetTextureIndex());
+					(*OutIndices)[InsertIndex].Add(Parameter.TextureIndex);
 				}
 			}
 		}
 	}
 }
 
-void UMaterialInstance::GetAtlasTextureValues(const FMaterialResource* MaterialResource, TArray<UTexture*>& OutTextures) const
-{
-	check(MaterialResource);
-
-	const TArray<TRefCountPtr<FMaterialUniformExpression> >* AtlasExpressions[1] =
-	{
-		&MaterialResource->GetUniformScalarParameterExpressions()
-	};
-	for (int32 TypeIndex = 0; TypeIndex < UE_ARRAY_COUNT(AtlasExpressions); TypeIndex++)
-	{
-		// Iterate over each of the material's scalar expressions.
-		for (FMaterialUniformExpression* Expression : *AtlasExpressions[TypeIndex])
-		{
-			const FMaterialUniformExpressionScalarParameter* ScalarExpression = static_cast<const FMaterialUniformExpressionScalarParameter*>(Expression);
-			bool bIsUsedAsAtlasPosition = false;
-			TSoftObjectPtr<class UCurveLinearColor> Curve;
-			TSoftObjectPtr<class UCurveLinearColorAtlas> Atlas;
-			ScalarExpression->GetGameThreadUsedAsAtlas(this, bIsUsedAsAtlasPosition, Curve, Atlas);
-
-			if (Atlas)
-			{
-				OutTextures.AddUnique(Atlas.Get());
-			}
-		}
-	}
-}
-
 void UMaterialInstance::GetUsedTextures(TArray<UTexture*>& OutTextures, EMaterialQualityLevel::Type QualityLevel, bool bAllQualityLevels, ERHIFeatureLevel::Type FeatureLevel, bool bAllFeatureLevels) const
 {
 	OutTextures.Empty();
@@ -1331,83 +1402,71 @@ void UMaterialInstance::ValidateTextureOverrides(ERHIFeatureLevel::Type InFeatur
 		// Fatal to call getmaterial in a non-game thread or async loading
 		return;
 	}
-	
+
 	const UMaterial* Material = GetMaterial();
 	const FMaterialResource* CurrentResource = Material->GetMaterialResource(InFeatureLevel);
-	
 
 	if (!CurrentResource)
 	{
 		return;
 	}
 
-	const TArray<TRefCountPtr<FMaterialUniformExpressionTexture>>* TextureExpressions[] =
-	{
-		&CurrentResource->GetUniform2DTextureExpressions(),
-		&CurrentResource->GetUniformCubeTextureExpressions(),
-		&CurrentResource->GetUniform2DArrayTextureExpressions(),
-		&CurrentResource->GetUniformVolumeTextureExpressions(),
-		&CurrentResource->GetUniformVirtualTextureExpressions(),
-	};
-
 	FString MaterialName;
 	GetName(MaterialName);
 
-	for (int32 TypeIndex = 0; TypeIndex < UE_ARRAY_COUNT(TextureExpressions); ++TypeIndex)
+	for (uint32 TypeIndex = 0; TypeIndex < NumMaterialTextureParameterTypes; ++TypeIndex)
 	{
-		for (FMaterialUniformExpressionTexture* TextureExpression : *TextureExpressions[TypeIndex])
+		const EMaterialTextureParameterType ParameterType = (EMaterialTextureParameterType)TypeIndex;
+		for (const FMaterialTextureParameterInfo& TextureInfo : CurrentResource->GetUniformTextureExpressions(ParameterType))
 		{
-			FMaterialUniformExpressionTextureParameter* ParameterExpression = TextureExpression->GetTextureParameterUniformExpression();
-			if (ParameterExpression)
+			UTexture* Texture = nullptr;
+			TextureInfo.GetGameThreadTextureValue(this, *CurrentResource, Texture);
+			if (Texture)
 			{
-				UTexture* Texture = nullptr;
-				const FMaterialParameterInfo& ParameterInfo = ParameterExpression->GetParameterInfo();
-				if (GetTextureParameterValue(ParameterInfo, Texture, true) && Texture != nullptr)
+				const EMaterialValueType TextureType = Texture->GetMaterialType();
+				switch (ParameterType)
 				{
-					const EMaterialValueType TextureType = Texture->GetMaterialType();
-					switch (TypeIndex)
+				case EMaterialTextureParameterType::Standard2D:
+					if (!(TextureType & (MCT_Texture2D | MCT_TextureExternal | MCT_TextureVirtual)))
 					{
-					case 0:
-						if (!(TextureType & (MCT_Texture2D | MCT_TextureExternal | MCT_TextureVirtual)))
-						{
-							UE_LOG(LogMaterial, Error, TEXT("MaterialInstance \"%s\" parameter '%s' assigned texture \"%s\" has invalid type, required 2D texture"), *MaterialName, *ParameterInfo.Name.ToString(), *Texture->GetName());
-						}
-						else if (TextureType & MCT_TextureVirtual)
-						{
-							UE_LOG(LogMaterial, Error, TEXT("MaterialInstance \"%s\" parameter '%s' assigned texture \"%s\" requires non-virtual texture"), *MaterialName, *ParameterInfo.Name.ToString(), *Texture->GetName());
-						}
-						break;
-					case 1:
-						if (!(TextureType & MCT_TextureCube))
-						{
-							UE_LOG(LogMaterial, Error, TEXT("MaterialInstance \"%s\" parameter '%s' assigned texture \"%s\" has invalid type, required Cube texture"), *MaterialName, *ParameterInfo.Name.ToString(), *Texture->GetName());
-						}
-						break;
-					case 2:
-						if (!(TextureType & MCT_Texture2DArray))
-						{
-							UE_LOG(LogMaterial, Error, TEXT("MaterialInstance \"%s\" parameter '%s' assigned texture \"%s\" has invalid type, required texture array"), *MaterialName, *ParameterInfo.Name.ToString(), *Texture->GetName());
-						}
-						break;
-					case 3:
-						if (!(TextureType & MCT_VolumeTexture))
-						{
-							UE_LOG(LogMaterial, Error, TEXT("MaterialInstance \"%s\" parameter '%s' assigned texture \"%s\" has invalid type, required Volume texture"), *MaterialName, *ParameterInfo.Name.ToString(), *Texture->GetName());
-						}
-						break;
-					case 4:
-						if (!(TextureType & (MCT_Texture2D | MCT_TextureExternal | MCT_TextureVirtual)))
-						{
-							UE_LOG(LogMaterial, Error, TEXT("MaterialInstance \"%s\" parameter '%s' assigned texture \"%s\" has invalid type, required 2D texture"), *MaterialName, *ParameterInfo.Name.ToString(), *Texture->GetName());
-						}
-						else if (!(TextureType & MCT_TextureVirtual))
-						{
-							UE_LOG(LogMaterial, Error, TEXT("MaterialInstance \"%s\" parameter '%s' assigned texture \"%s\" requires virtual texture"), *MaterialName, *ParameterInfo.Name.ToString(), *Texture->GetName());
-						}
-						break;
-					default:
-						checkNoEntry();
+						UE_LOG(LogMaterial, Error, TEXT("MaterialInstance \"%s\" parameter '%s' assigned texture \"%s\" has invalid type, required 2D texture"), *MaterialName, *TextureInfo.ParameterName, *Texture->GetName());
+					}
+					else if (TextureType & MCT_TextureVirtual)
+					{
+						UE_LOG(LogMaterial, Error, TEXT("MaterialInstance \"%s\" parameter '%s' assigned texture \"%s\" requires non-virtual texture"), *MaterialName, *TextureInfo.ParameterName, *Texture->GetName());
+					}
+					break;
+				case EMaterialTextureParameterType::Cube:
+					if (!(TextureType & MCT_TextureCube))
+					{
+						UE_LOG(LogMaterial, Error, TEXT("MaterialInstance \"%s\" parameter '%s' assigned texture \"%s\" has invalid type, required Cube texture"), *MaterialName, *TextureInfo.ParameterName, *Texture->GetName());
+					}
+					break;
+				case EMaterialTextureParameterType::Array2D:
+					if (!(TextureType & MCT_Texture2DArray))
+					{
+						UE_LOG(LogMaterial, Error, TEXT("MaterialInstance \"%s\" parameter '%s' assigned texture \"%s\" has invalid type, required texture array"), *MaterialName, *TextureInfo.ParameterName, *Texture->GetName());
+					}
+					break;
+				case EMaterialTextureParameterType::Volume:
+					if (!(TextureType & MCT_VolumeTexture))
+					{
+						UE_LOG(LogMaterial, Error, TEXT("MaterialInstance \"%s\" parameter '%s' assigned texture \"%s\" has invalid type, required Volume texture"), *MaterialName, *TextureInfo.ParameterName, *Texture->GetName());
+					}
+					break;
+				case EMaterialTextureParameterType::Virtual:
+					if (!(TextureType & (MCT_Texture2D | MCT_TextureExternal | MCT_TextureVirtual)))
+					{
+						UE_LOG(LogMaterial, Error, TEXT("MaterialInstance \"%s\" parameter '%s' assigned texture \"%s\" has invalid type, required 2D texture"), *MaterialName, *TextureInfo.ParameterName, *Texture->GetName());
+					}
+					else if (!(TextureType & MCT_TextureVirtual))
+					{
+						UE_LOG(LogMaterial, Error, TEXT("MaterialInstance \"%s\" parameter '%s' assigned texture \"%s\" requires virtual texture"), *MaterialName, *TextureInfo.ParameterName, *Texture->GetName());
 					}
+					break;
+				default:
+					checkNoEntry();
+					break;
 				}
 			}
 		}
@@ -1460,51 +1519,33 @@ void UMaterialInstance::OverrideTexture(const UTexture* InTextureToOverride, UTe
 #if WITH_EDITOR
 	bool bShouldRecacheMaterialExpressions = false;
 	
-	const TArray<TRefCountPtr<FMaterialUniformExpressionTexture> >* ExpressionsByType[5];
-
-	const FMaterialResource* SourceMaterialResource = NULL;
+	FMaterialResource* SourceMaterialResource = nullptr;
 	if (bHasStaticPermutationResource)
 	{
 		SourceMaterialResource = GetMaterialResource(InFeatureLevel);
-		// Iterate over both the 2D textures and cube texture expressions.
-		ExpressionsByType[0] = &SourceMaterialResource->GetUniform2DTextureExpressions();
-		ExpressionsByType[1] = &SourceMaterialResource->GetUniformCubeTextureExpressions();
-		ExpressionsByType[2] = &SourceMaterialResource->GetUniform2DArrayTextureExpressions();
-		ExpressionsByType[3] = &SourceMaterialResource->GetUniformVolumeTextureExpressions();
-		ExpressionsByType[4] = &SourceMaterialResource->GetUniformVirtualTextureExpressions();
 	}
 	else
 	{
 		//@todo - this isn't handling chained MIC's correctly, where a parent in the chain has static parameters
 		UMaterial* Material = GetMaterial();
 		SourceMaterialResource = Material->GetMaterialResource(InFeatureLevel);
-			
-		// Iterate over both the 2D textures and cube texture expressions.
-		ExpressionsByType[0] = &SourceMaterialResource->GetUniform2DTextureExpressions();
-		ExpressionsByType[1] = &SourceMaterialResource->GetUniformCubeTextureExpressions();
-		ExpressionsByType[2] = &SourceMaterialResource->GetUniform2DArrayTextureExpressions();
-		ExpressionsByType[3] = &SourceMaterialResource->GetUniformVolumeTextureExpressions();
-		ExpressionsByType[4] = &SourceMaterialResource->GetUniformVirtualTextureExpressions();
 	}
 		
-	for(int32 TypeIndex = 0;TypeIndex < UE_ARRAY_COUNT(ExpressionsByType);TypeIndex++)
+	for (int32 TypeIndex = 0; TypeIndex < NumMaterialTextureParameterTypes; TypeIndex++)
 	{
-		const TArray<TRefCountPtr<FMaterialUniformExpressionTexture> >& Expressions = *ExpressionsByType[TypeIndex];
-
+		const TArray<FMaterialTextureParameterInfo, FMemoryImageAllocator>& Parameters = SourceMaterialResource->GetUniformTextureExpressions((EMaterialTextureParameterType)TypeIndex);
 		// Iterate over each of the material's texture expressions.
-		for(int32 ExpressionIndex = 0;ExpressionIndex < Expressions.Num();ExpressionIndex++)
+		for (int32 i = 0; i < Parameters.Num(); ++i)
 		{
-			FMaterialUniformExpressionTexture* Expression = Expressions[ExpressionIndex];
+			const FMaterialTextureParameterInfo& Parameter = Parameters[i];
 
 			// Evaluate the expression in terms of this material instance.
-			const bool bAllowOverride = false;
 			UTexture* Texture = NULL;
-			Expression->GetGameThreadTextureValue(this,*SourceMaterialResource,Texture,bAllowOverride);
-
-			if( Texture != NULL && Texture == InTextureToOverride )
+			Parameter.GetGameThreadTextureValue(this, *SourceMaterialResource, Texture);
+			if (Texture != NULL && Texture == InTextureToOverride)
 			{
 				// Override this texture!
-				Expression->SetTransientOverrideTextureValue( OverrideTexture );
+				SourceMaterialResource->TransientOverrides.SetTextureOverride((EMaterialTextureParameterType)TypeIndex, i, OverrideTexture);
 				bShouldRecacheMaterialExpressions = true;
 			}
 		}
@@ -1517,29 +1558,22 @@ void UMaterialInstance::OverrideTexture(const UTexture* InTextureToOverride, UTe
 #endif // #if WITH_EDITOR
 }
 
-void UMaterialInstance::OverrideVectorParameterDefault(const FMaterialParameterInfo& ParameterInfo, const FLinearColor& Value, bool bOverride, ERHIFeatureLevel::Type InFeatureLevel)
+void UMaterialInstance::OverrideVectorParameterDefault(const FHashedMaterialParameterInfo& ParameterInfo, const FLinearColor& Value, bool bOverride, ERHIFeatureLevel::Type InFeatureLevel)
 {
 #if WITH_EDITOR
 	bool bShouldRecacheMaterialExpressions = false;
-
 	if (bHasStaticPermutationResource)
 	{
-		const FMaterialResource* SourceMaterialResource = GetMaterialResource(InFeatureLevel);
-		const TArray<TRefCountPtr<FMaterialUniformExpression> >& UniformExpressions = SourceMaterialResource->GetUniformVectorParameterExpressions();
+		FMaterialResource* SourceMaterialResource = GetMaterialResource(InFeatureLevel);
+		const TArray<FMaterialVectorParameterInfo, FMemoryImageAllocator>& Parameters = SourceMaterialResource->GetUniformVectorParameterExpressions();
 
-		// Iterate over each of the material's texture expressions.
-		for (int32 ExpressionIndex = 0; ExpressionIndex < UniformExpressions.Num(); ExpressionIndex++)
+		for (int32 i = 0; i < Parameters.Num(); ++i)
 		{
-			FMaterialUniformExpression* UniformExpression = UniformExpressions[ExpressionIndex];
-			if (UniformExpression->GetType() == &FMaterialUniformExpressionVectorParameter::StaticType)
+			const FMaterialVectorParameterInfo& Parameter = Parameters[i];
+			if (Parameter.ParameterInfo == ParameterInfo)
 			{
-				FMaterialUniformExpressionVectorParameter* VectorExpression = static_cast<FMaterialUniformExpressionVectorParameter*>(UniformExpression);
-
-				if (VectorExpression->GetParameterInfo() == ParameterInfo)
-				{
-					VectorExpression->SetTransientOverrideDefaultValue(Value, bOverride);
-					bShouldRecacheMaterialExpressions = true;
-				}
+				SourceMaterialResource->TransientOverrides.SetVectorOverride(i, Value, bOverride);
+				bShouldRecacheMaterialExpressions = true;
 			}
 		}
 	}
@@ -1551,29 +1585,22 @@ void UMaterialInstance::OverrideVectorParameterDefault(const FMaterialParameterI
 #endif // #if WITH_EDITOR
 }
 
-void UMaterialInstance::OverrideScalarParameterDefault(const FMaterialParameterInfo& ParameterInfo, float Value, bool bOverride, ERHIFeatureLevel::Type InFeatureLevel)
+void UMaterialInstance::OverrideScalarParameterDefault(const FHashedMaterialParameterInfo& ParameterInfo, float Value, bool bOverride, ERHIFeatureLevel::Type InFeatureLevel)
 {
 #if WITH_EDITOR
 	bool bShouldRecacheMaterialExpressions = false;
-
 	if (bHasStaticPermutationResource)
 	{
-		const FMaterialResource* SourceMaterialResource = GetMaterialResource(InFeatureLevel);
-		const TArray<TRefCountPtr<FMaterialUniformExpression> >& UniformExpressions = SourceMaterialResource->GetUniformScalarParameterExpressions();
+		FMaterialResource* SourceMaterialResource = GetMaterialResource(InFeatureLevel);
+		const TArray<FMaterialScalarParameterInfo, FMemoryImageAllocator>& Parameters = SourceMaterialResource->GetUniformScalarParameterExpressions();
 
-		// Iterate over each of the material's texture expressions.
-		for (int32 ExpressionIndex = 0; ExpressionIndex < UniformExpressions.Num(); ExpressionIndex++)
+		for (int32 i = 0; i < Parameters.Num(); ++i)
 		{
-			FMaterialUniformExpression* UniformExpression = UniformExpressions[ExpressionIndex];
-			if (UniformExpression->GetType() == &FMaterialUniformExpressionScalarParameter::StaticType)
+			const FMaterialScalarParameterInfo& Parameter = Parameters[i];
+			if (Parameter.ParameterInfo == ParameterInfo)
 			{
-				FMaterialUniformExpressionScalarParameter* ScalarExpression = static_cast<FMaterialUniformExpressionScalarParameter*>(UniformExpression);
-
-				if (ScalarExpression->GetParameterInfo() == ParameterInfo)
-				{
-					ScalarExpression->SetTransientOverrideDefaultValue(Value, bOverride);
-					bShouldRecacheMaterialExpressions = true;
-				}
+				SourceMaterialResource->TransientOverrides.SetScalarOverride(i, Value, bOverride);
+				bShouldRecacheMaterialExpressions = true;
 			}
 		}
 	}
@@ -1608,8 +1635,7 @@ bool UMaterialInstance::CheckMaterialUsage(const EMaterialUsage Usage)
 
 bool UMaterialInstance::CheckMaterialUsage_Concurrent(const EMaterialUsage Usage) const
 {
-	TMicRecursionGuard RecursionGuard;
-	UMaterial const* Material = GetMaterial_Concurrent(RecursionGuard);
+	UMaterial const* Material = GetMaterial_Concurrent();
 	if(Material)
 	{
 		bool bUsageSetSuccessfully = false;
@@ -1696,64 +1722,46 @@ void UMaterialInstanceDynamic::CopyScalarAndVectorParameters(const UMaterialInte
 		ClearParameterValuesInternal(false);
 
 		// scalar
-		{
-			const TArray<TRefCountPtr<FMaterialUniformExpression> >& Array = MaterialResource->GetUniformScalarParameterExpressions();
-
-			for (int32 i = 0, Count = Array.Num(); i < Count; ++i)
-			{
-				const FMaterialUniformExpression* UniformExpression = Array[i];
-
-				// the array can have non scalar parameters in it, those we don't want to interpolate
-				if (UniformExpression->GetType() == &FMaterialUniformExpressionScalarParameter::StaticType)
-				{
-					const FMaterialUniformExpressionScalarParameter* ScalarExpression = static_cast<const FMaterialUniformExpressionScalarParameter*>(UniformExpression);
-
-					float Value;
-					ScalarExpression->GetGameThreadNumberValue(&SourceMaterialToCopyFrom, Value);
-
-					const FMaterialParameterInfo& ParameterInfo = ScalarExpression->GetParameterInfo();
+		{
+			const TArray<FMaterialScalarParameterInfo, FMemoryImageAllocator>& Array = MaterialResource->GetUniformScalarParameterExpressions();
 
-					FScalarParameterValue* ParameterValue = GameThread_FindParameterByName(ScalarParameterValues, ParameterInfo);
+			for (int32 i = 0, Count = Array.Num(); i < Count; ++i)
+			{
+				const FMaterialScalarParameterInfo& Parameter = Array[i];
 
-					if (!ParameterValue)
-					{
-						ParameterValue = new(ScalarParameterValues)FScalarParameterValue;
-						ParameterValue->ParameterInfo = ParameterInfo;
-					}
+				float Value;
+				Parameter.GetGameThreadNumberValue(&SourceMaterialToCopyFrom, Value);
 
-					ParameterValue->ParameterValue = Value;
+				FScalarParameterValue* ParameterValue = GameThread_FindParameterByName(ScalarParameterValues, Parameter.ParameterInfo);
+				if (!ParameterValue)
+				{
+					ParameterValue = new(ScalarParameterValues)FScalarParameterValue;
+					ParameterValue->ParameterInfo = Parameter.GetParameterInfo();
 				}
+
+				ParameterValue->ParameterValue = Value;
 			}
 		}
 
 		// vector
 		{
-			const TArray<TRefCountPtr<FMaterialUniformExpression> >& Array = MaterialResource->GetUniformVectorParameterExpressions();
+			const TArray<FMaterialVectorParameterInfo, FMemoryImageAllocator>& Array = MaterialResource->GetUniformVectorParameterExpressions();
 
 			for (int32 i = 0, Count = Array.Num(); i < Count; ++i)
 			{
-				const FMaterialUniformExpression* UniformExpression = Array[i];
+				const FMaterialVectorParameterInfo& Parameter = Array[i];
 
-				// the array can have non vector parameters in it, those we don't want to inetrpolate
-				if (UniformExpression->GetType() == &FMaterialUniformExpressionVectorParameter::StaticType)
-				{
-					const FMaterialUniformExpressionVectorParameter* VectorExpression = static_cast<const FMaterialUniformExpressionVectorParameter*>(UniformExpression);
-
-					FLinearColor Value;
-					VectorExpression->GetGameThreadNumberValue(&SourceMaterialToCopyFrom, Value);
-
-					const FMaterialParameterInfo& ParameterInfo = VectorExpression->GetParameterInfo();
-
-					FVectorParameterValue* ParameterValue = GameThread_FindParameterByName(VectorParameterValues, ParameterInfo);
+				FLinearColor Value;
+				Parameter.GetGameThreadNumberValue(&SourceMaterialToCopyFrom, Value);
 
-					if (!ParameterValue)
-					{
-						ParameterValue = new(VectorParameterValues)FVectorParameterValue;
-						ParameterValue->ParameterInfo = ParameterInfo;
-					}
-
-					ParameterValue->ParameterValue = Value;
+				FVectorParameterValue* ParameterValue = GameThread_FindParameterByName(VectorParameterValues, Parameter.ParameterInfo);
+				if (!ParameterValue)
+				{
+					ParameterValue = new(VectorParameterValues)FVectorParameterValue;
+					ParameterValue->ParameterInfo = Parameter.GetParameterInfo();
 				}
+
+				ParameterValue->ParameterValue = Value;
 			}
 		}
 
@@ -1976,6 +1984,7 @@ UPhysicalMaterial* UMaterialInstance::GetPhysicalMaterialFromMap(int32 Index) co
 	return PhysicalMaterialMap[Index];
 }
 
+#if WITH_EDITORONLY_DATA
 void UMaterialInstance::GetStaticParameterValues(FStaticParameterSet& OutStaticParameters)
 {
 	check(IsInGameThread());
@@ -2025,7 +2034,14 @@ void UMaterialInstance::GetStaticParameterValues(FStaticParameterSet& OutStaticP
 		GetAllStaticSwitchParameterInfo(OutParameterInfo, Guids);
 		OutStaticParameters.StaticSwitchParameters.AddZeroed(OutParameterInfo.Num());
 
-		FStaticParamEvaluationContext SwitchEvalContext(OutParameterInfo.Num(), OutParameterInfo.GetData());
+		// make a temporary hashed verison of OutParameterInfo
+		TArray<FHashedMaterialParameterInfo> StaticSwitchHashedParameterInfo;
+		for (const FMaterialParameterInfo& Info : OutParameterInfo)
+		{
+			new (StaticSwitchHashedParameterInfo) FHashedMaterialParameterInfo(Info);
+		}
+
+		FStaticParamEvaluationContext SwitchEvalContext(StaticSwitchHashedParameterInfo.Num(), StaticSwitchHashedParameterInfo.GetData());
 		TBitArray<> StaticSwitchValues(false, OutParameterInfo.Num());
 		ensure(GetStaticSwitchParameterValues(SwitchEvalContext, StaticSwitchValues, Guids.GetData()));
 		for(int32 ParameterIdx=0; ParameterIdx<OutParameterInfo.Num(); ParameterIdx++)
@@ -2059,7 +2075,14 @@ void UMaterialInstance::GetStaticParameterValues(FStaticParameterSet& OutStaticP
 		GetAllStaticComponentMaskParameterInfo(OutParameterInfo, Guids);
 		OutStaticParameters.StaticComponentMaskParameters.AddZeroed(OutParameterInfo.Num());
 
-		FStaticParamEvaluationContext ComponentMaskEvalContext(OutParameterInfo.Num(), OutParameterInfo.GetData());
+		// make a temporary hashed verison of OutParameterInfo
+		TArray<FHashedMaterialParameterInfo> StaticComponentMaskHashedParameterInfo;
+		for (const FMaterialParameterInfo& Info : OutParameterInfo)
+		{
+			new (StaticComponentMaskHashedParameterInfo) FHashedMaterialParameterInfo(Info);
+		}
+
+		FStaticParamEvaluationContext ComponentMaskEvalContext(StaticComponentMaskHashedParameterInfo.Num(), StaticComponentMaskHashedParameterInfo.GetData());
 		TBitArray<> StaticComponentMaskValues(false, OutParameterInfo.Num()*4);
 		ensure(GetStaticComponentMaskParameterValues(ComponentMaskEvalContext, StaticComponentMaskValues, Guids.GetData()));
 		for(int32 ParameterIdx=0; ParameterIdx<OutParameterInfo.Num(); ParameterIdx++)
@@ -2105,57 +2128,52 @@ void UMaterialInstance::GetStaticParameterValues(FStaticParameterSet& OutStaticP
 		CachedStaticParameterValues = OutStaticParameters;
 	}
 }
+#endif // WITH_EDITORONLY_DATA
 
-void UMaterialInstance::GetAllScalarParameterInfo(TArray<FMaterialParameterInfo>& OutParameterInfo, TArray<FGuid>& OutParameterIds) const
+void UMaterialInstance::GetAllParametersOfType(EMaterialParameterType Type, TArray<FMaterialParameterInfo>& OutParameterInfo, TArray<FGuid>& OutParameterIds) const
 {
-	OutParameterInfo.Empty();
-	OutParameterIds.Empty();
-	if (const UMaterial* Material = GetMaterial())
+	const UMaterial* Material = GetMaterial();
+	int32 NumParameters = CachedLayerParameters.GetNumParameters(Type);
+	if (Material)
+	{
+		NumParameters += Material->GetCachedExpressionData().Parameters.GetNumParameters(Type);
+	}
+
+	OutParameterInfo.Empty(NumParameters);
+	OutParameterIds.Empty(NumParameters);
+	CachedLayerParameters.GetAllParameterInfoOfType(Type, false, OutParameterInfo, OutParameterIds);
+	if (Material)
 	{
-		Material->GetAllParameterInfo<UMaterialExpressionScalarParameter>(OutParameterInfo, OutParameterIds, &StaticParameters.MaterialLayersParameters);
+		Material->GetCachedExpressionData().Parameters.GetAllGlobalParameterInfoOfType(Type, false, OutParameterInfo, OutParameterIds);
 	}
 }
 
+void UMaterialInstance::GetAllScalarParameterInfo(TArray<FMaterialParameterInfo>& OutParameterInfo, TArray<FGuid>& OutParameterIds) const
+{
+	GetAllParametersOfType(EMaterialParameterType::Scalar, OutParameterInfo, OutParameterIds);
+}
+
 void UMaterialInstance::GetAllVectorParameterInfo(TArray<FMaterialParameterInfo>& OutParameterInfo, TArray<FGuid>& OutParameterIds) const
 {
-	OutParameterInfo.Empty();
-	OutParameterIds.Empty();
-	if (const UMaterial* Material = GetMaterial())
-	{
-		Material->GetAllParameterInfo<UMaterialExpressionVectorParameter>(OutParameterInfo, OutParameterIds, &StaticParameters.MaterialLayersParameters);
-	}
+	GetAllParametersOfType(EMaterialParameterType::Vector, OutParameterInfo, OutParameterIds);
 }
 
 void UMaterialInstance::GetAllTextureParameterInfo(TArray<FMaterialParameterInfo>& OutParameterInfo, TArray<FGuid>& OutParameterIds) const
 {
-	OutParameterInfo.Empty();
-	OutParameterIds.Empty();
-	if (const UMaterial* Material = GetMaterial())
-	{
-		Material->GetAllParameterInfo<UMaterialExpressionTextureSampleParameter>(OutParameterInfo, OutParameterIds, &StaticParameters.MaterialLayersParameters);
-	}
+	GetAllParametersOfType(EMaterialParameterType::Texture, OutParameterInfo, OutParameterIds);
 }
 
 void UMaterialInstance::GetAllRuntimeVirtualTextureParameterInfo(TArray<FMaterialParameterInfo>& OutParameterInfo, TArray<FGuid>& OutParameterIds) const
 {
-	OutParameterInfo.Empty();
-	OutParameterIds.Empty();
-	if (const UMaterial* Material = GetMaterial())
-	{
-		Material->GetAllParameterInfo<UMaterialExpressionRuntimeVirtualTextureSampleParameter>(OutParameterInfo, OutParameterIds, &StaticParameters.MaterialLayersParameters);
-	}
+	GetAllParametersOfType(EMaterialParameterType::RuntimeVirtualTexture, OutParameterInfo, OutParameterIds);
 }
 
 void UMaterialInstance::GetAllFontParameterInfo(TArray<FMaterialParameterInfo>& OutParameterInfo, TArray<FGuid>& OutParameterIds) const
 {
-	OutParameterInfo.Empty();
-	OutParameterIds.Empty();
-	if (const UMaterial* Material = GetMaterial())
-	{
-		Material->GetAllParameterInfo<UMaterialExpressionFontSampleParameter>(OutParameterInfo, OutParameterIds, &StaticParameters.MaterialLayersParameters);
-	}
+	GetAllParametersOfType(EMaterialParameterType::Font, OutParameterInfo, OutParameterIds);
 }
 
+#if WITH_EDITORONLY_DATA
 void UMaterialInstance::GetAllMaterialLayersParameterInfo(TArray<FMaterialParameterInfo>& OutParameterInfo, TArray<FGuid>& OutParameterIds) const
 {
 	OutParameterInfo.Empty();
@@ -2229,8 +2247,9 @@ void UMaterialInstance::GetDependentFunctions(TArray<UMaterialFunctionInterface*
 			return true;
 		});
 }
+#endif // WITH_EDITORONLY_DATA
 
-bool UMaterialInstance::GetScalarParameterDefaultValue(const FMaterialParameterInfo& ParameterInfo, float& OutValue, bool bOveriddenOnly, bool bCheckOwnedGlobalOverrides) const
+bool UMaterialInstance::GetScalarParameterDefaultValue(const FHashedMaterialParameterInfo& ParameterInfo, float& OutValue, bool bOveriddenOnly, bool bCheckOwnedGlobalOverrides) const
 {
 #if WITH_EDITOR
 	if (GetReentrantFlag())
@@ -2241,33 +2260,27 @@ bool UMaterialInstance::GetScalarParameterDefaultValue(const FMaterialParameterI
 
 	// In the case of duplicate parameters with different values, this will return the
 	// first matching expression found, not necessarily the one that's used for rendering
-	UMaterialExpressionScalarParameter* Parameter = nullptr;
-
 	if (ParameterInfo.Association != EMaterialParameterAssociation::GlobalParameter)
 	{
-		// Parameters introduced by this instance's layer stack
-		for (const FStaticMaterialLayersParameter& LayersParam : StaticParameters.MaterialLayersParameters)
-		{
-			if (LayersParam.bOverride)
-			{
-				UMaterialFunctionInterface* Function = LayersParam.GetParameterAssociatedFunction(ParameterInfo);
-				UMaterialFunctionInterface* ParameterOwner = nullptr;
-
-				if (Function && Function->OverrideNamedScalarParameter(ParameterInfo, OutValue))
-				{
-					return true;
-				}
+#if WITH_EDITOR
+		bool bLegacyFoundValue = false;
+		bool bLegacyResult = false;
+		float LegacyValue = 0.0f;
+		// Legacy code doesn't check bOveriddenOnly here
+		bLegacyFoundValue = GetScalarLayerParameterValue_Legacy(ParameterInfo, false, LegacyValue, bLegacyResult);
+#endif // WITH_EDITOR
 
-				if (Function && Function->GetNamedParameterOfType(ParameterInfo, Parameter, &ParameterOwner))
-				{
-					if (!ParameterOwner->OverrideNamedScalarParameter(ParameterInfo, OutValue))
-					{
-						Parameter->IsNamedParameter(ParameterInfo, OutValue);
-					}
-					return true;
-				}
-			}
+		const int32 ParameterIndex = CachedLayerParameters.FindParameterIndex(EMaterialParameterType::Scalar, ParameterInfo);
+		if (ParameterIndex != INDEX_NONE)
+		{
+			checkLayerParameter(bLegacyFoundValue);
+			OutValue = CachedLayerParameters.ScalarValues[ParameterIndex];
+			checkLayerParameter(OutValue == LegacyValue);
+			checkLayerParameter(bLegacyResult);
+			return true;
 		}
+
+		checkLayerParameter(!bLegacyFoundValue);
 	}
 	else if (bCheckOwnedGlobalOverrides)
 	{
@@ -2293,7 +2306,7 @@ bool UMaterialInstance::GetScalarParameterDefaultValue(const FMaterialParameterI
 	return false;
 }
 
-bool UMaterialInstance::GetVectorParameterDefaultValue(const FMaterialParameterInfo& ParameterInfo, FLinearColor& OutValue, bool bOveriddenOnly, bool bCheckOwnedGlobalOverrides) const
+bool UMaterialInstance::GetVectorParameterDefaultValue(const FHashedMaterialParameterInfo& ParameterInfo, FLinearColor& OutValue, bool bOveriddenOnly, bool bCheckOwnedGlobalOverrides) const
 {
 #if WITH_EDITOR
 	if (GetReentrantFlag())
@@ -2308,29 +2321,25 @@ bool UMaterialInstance::GetVectorParameterDefaultValue(const FMaterialParameterI
 
 	if (ParameterInfo.Association != EMaterialParameterAssociation::GlobalParameter)
 	{
-		// Parameters introduced by this instance's layer stack
-		for (const FStaticMaterialLayersParameter& LayersParam : StaticParameters.MaterialLayersParameters)
-		{
-			if (LayersParam.bOverride)
-			{
-				UMaterialFunctionInterface* Function = LayersParam.GetParameterAssociatedFunction(ParameterInfo);
-				UMaterialFunctionInterface* ParameterOwner = nullptr;
-
-				if (Function && Function->OverrideNamedVectorParameter(ParameterInfo, OutValue))
-				{
-					return true;
-				}
+#if WITH_EDITOR
+		bool bLegacyFoundValue = false;
+		bool bLegacyResult = false;
+		FLinearColor LegacyValue(ForceInitToZero);
+		// Legacy code doesn't check bOveriddenOnly here
+		bLegacyFoundValue = GetVectorLayerParameterValue_Legacy(ParameterInfo, false, LegacyValue, bLegacyResult);
+#endif // WITH_EDITOR
 
-				if (Function && Function->GetNamedParameterOfType(ParameterInfo, Parameter, &ParameterOwner))
-				{		
-					if (!ParameterOwner->OverrideNamedVectorParameter(ParameterInfo, OutValue))
-					{
-						Parameter->IsNamedParameter(ParameterInfo, OutValue);
-					}
-					return true;
-				}
-			}
+		const int32 ParameterIndex = CachedLayerParameters.FindParameterIndex(EMaterialParameterType::Vector, ParameterInfo);
+		if (ParameterIndex != INDEX_NONE)
+		{
+			checkLayerParameter(bLegacyFoundValue);
+			OutValue = CachedLayerParameters.VectorValues[ParameterIndex];
+			checkLayerParameter(OutValue == LegacyValue);
+			checkLayerParameter(bLegacyResult);
+			return true;
 		}
+
+		checkLayerParameter(!bLegacyFoundValue);
 	}
 	else if (bCheckOwnedGlobalOverrides)
 	{
@@ -2356,7 +2365,7 @@ bool UMaterialInstance::GetVectorParameterDefaultValue(const FMaterialParameterI
 	return false;
 }
 
-bool UMaterialInstance::GetTextureParameterDefaultValue(const FMaterialParameterInfo& ParameterInfo, UTexture*& OutValue, bool bCheckOwnedGlobalOverrides) const
+bool UMaterialInstance::GetTextureParameterDefaultValue(const FHashedMaterialParameterInfo& ParameterInfo, UTexture*& OutValue, bool bCheckOwnedGlobalOverrides) const
 {
 #if WITH_EDITOR
 	if (GetReentrantFlag())
@@ -2371,29 +2380,25 @@ bool UMaterialInstance::GetTextureParameterDefaultValue(const FMaterialParameter
 
 	if (ParameterInfo.Association != EMaterialParameterAssociation::GlobalParameter)
 	{
-		// Parameters introduced by this instance's layer stack
-		for (const FStaticMaterialLayersParameter& LayersParam : StaticParameters.MaterialLayersParameters)
-		{
-			if (LayersParam.bOverride)
-			{
-				UMaterialFunctionInterface* Function = LayersParam.GetParameterAssociatedFunction(ParameterInfo);
-				UMaterialFunctionInterface* ParameterOwner = nullptr;
-
-				if (Function && Function->OverrideNamedTextureParameter(ParameterInfo, OutValue))
-				{
-					return true;
-				}
+#if WITH_EDITOR
+		bool bLegacyFoundValue = false;
+		bool bLegacyResult = false;
+		UTexture* LegacyValue = nullptr;
+		// Legacy code doesn't check bOveriddenOnly here
+		bLegacyFoundValue = GetTextureLayerParameterValue_Legacy(ParameterInfo, false, LegacyValue, bLegacyResult);
+#endif // WITH_EDITOR
 
-				if (Function && Function->GetNamedParameterOfType(ParameterInfo, Parameter, &ParameterOwner))
-				{
-					if (!ParameterOwner->OverrideNamedTextureParameter(ParameterInfo, OutValue))
-					{
-						Parameter->IsNamedParameter(ParameterInfo, OutValue);
-					}
-					return true;
-				}
-			}
+		const int32 ParameterIndex = CachedLayerParameters.FindParameterIndex(EMaterialParameterType::Texture, ParameterInfo);
+		if (ParameterIndex != INDEX_NONE)
+		{
+			checkLayerParameter(bLegacyFoundValue);
+			OutValue = CachedLayerParameters.TextureValues[ParameterIndex];
+			checkLayerParameter(OutValue == LegacyValue);
+			checkLayerParameter(bLegacyResult);
+			return true;
 		}
+
+		checkLayerParameter(!bLegacyFoundValue);
 	}
 	else if (bCheckOwnedGlobalOverrides)
 	{
@@ -2419,7 +2424,7 @@ bool UMaterialInstance::GetTextureParameterDefaultValue(const FMaterialParameter
 	return false;
 }
 
-bool UMaterialInstance::GetRuntimeVirtualTextureParameterDefaultValue(const FMaterialParameterInfo& ParameterInfo, URuntimeVirtualTexture*& OutValue, bool bCheckOwnedGlobalOverrides) const
+bool UMaterialInstance::GetRuntimeVirtualTextureParameterDefaultValue(const FHashedMaterialParameterInfo& ParameterInfo, URuntimeVirtualTexture*& OutValue, bool bCheckOwnedGlobalOverrides) const
 {
 #if WITH_EDITOR
 	if (GetReentrantFlag())
@@ -2434,29 +2439,25 @@ bool UMaterialInstance::GetRuntimeVirtualTextureParameterDefaultValue(const FMat
 
 	if (ParameterInfo.Association != EMaterialParameterAssociation::GlobalParameter)
 	{
-		// Parameters introduced by this instance's layer stack
-		for (const FStaticMaterialLayersParameter& LayersParam : StaticParameters.MaterialLayersParameters)
-		{
-			if (LayersParam.bOverride)
-			{
-				UMaterialFunctionInterface* Function = LayersParam.GetParameterAssociatedFunction(ParameterInfo);
-				UMaterialFunctionInterface* ParameterOwner = nullptr;
-
-				if (Function && Function->OverrideNamedRuntimeVirtualTextureParameter(ParameterInfo, OutValue))
-				{
-					return true;
-				}
+#if WITH_EDITOR
+		bool bLegacyFoundValue = false;
+		bool bLegacyResult = false;
+		URuntimeVirtualTexture* LegacyValue = nullptr;
+		// Legacy code doesn't check bOveriddenOnly here
+		bLegacyFoundValue = GetRuntimeVirtualTextureLayerParameterValue_Legacy(ParameterInfo, false, LegacyValue, bLegacyResult);
+#endif // WITH_EDITOR
 
-				if (Function && Function->GetNamedParameterOfType(ParameterInfo, Parameter, &ParameterOwner))
-				{
-					if (!ParameterOwner->OverrideNamedRuntimeVirtualTextureParameter(ParameterInfo, OutValue))
-					{
-						Parameter->IsNamedParameter(ParameterInfo, OutValue);
-					}
-					return true;
-				}
-			}
+		const int32 ParameterIndex = CachedLayerParameters.FindParameterIndex(EMaterialParameterType::RuntimeVirtualTexture, ParameterInfo);
+		if (ParameterIndex != INDEX_NONE)
+		{
+			checkLayerParameter(bLegacyFoundValue);
+			OutValue = CachedLayerParameters.RuntimeVirtualTextureValues[ParameterIndex];
+			checkLayerParameter(OutValue == LegacyValue);
+			checkLayerParameter(bLegacyResult);
+			return true;
 		}
+
+		checkLayerParameter(!bLegacyFoundValue);
 	}
 	else if (bCheckOwnedGlobalOverrides)
 	{
@@ -2482,7 +2483,7 @@ bool UMaterialInstance::GetRuntimeVirtualTextureParameterDefaultValue(const FMat
 	return false;
 }
 
-bool UMaterialInstance::GetFontParameterDefaultValue(const FMaterialParameterInfo& ParameterInfo, UFont*& OutFontValue, int32& OutFontPage, bool bCheckOwnedGlobalOverrides) const
+bool UMaterialInstance::GetFontParameterDefaultValue(const FHashedMaterialParameterInfo& ParameterInfo, UFont*& OutFontValue, int32& OutFontPage, bool bCheckOwnedGlobalOverrides) const
 {
 #if WITH_EDITOR
 	if (GetReentrantFlag())
@@ -2497,29 +2498,28 @@ bool UMaterialInstance::GetFontParameterDefaultValue(const FMaterialParameterInf
 
 	if (ParameterInfo.Association != EMaterialParameterAssociation::GlobalParameter)
 	{
-		// Parameters introduced by this instance's layer stack
-		for (const FStaticMaterialLayersParameter& LayersParam : StaticParameters.MaterialLayersParameters)
-		{
-			if (LayersParam.bOverride)
-			{
-				UMaterialFunctionInterface* Function = LayersParam.GetParameterAssociatedFunction(ParameterInfo);
-				UMaterialFunctionInterface* ParameterOwner = nullptr;
-
-				if (Function && Function->OverrideNamedFontParameter(ParameterInfo, OutFontValue, OutFontPage))
-				{
-					return true;
-				}
+#if WITH_EDITOR
+		bool bLegacyFoundValue = false;
+		bool bLegacyResult = false;
+		UFont* LegacyValue = nullptr;
+		int32 LegacyFontPageValue = INDEX_NONE;
+		// Legacy code doesn't check bOveriddenOnly here
+		bLegacyFoundValue = GetFontLayerParameterValue_Legacy(ParameterInfo, false, LegacyValue, LegacyFontPageValue, bLegacyResult);
+#endif // WITH_EDITOR
 
-				if (Function && Function->GetNamedParameterOfType(ParameterInfo, Parameter, &ParameterOwner))
-				{
-					if (!ParameterOwner->OverrideNamedFontParameter(ParameterInfo, OutFontValue, OutFontPage))
-					{
-						Parameter->IsNamedParameter(ParameterInfo, OutFontValue, OutFontPage);
-					}
-					return true;
-				}
-			}
+		const int32 ParameterIndex = CachedLayerParameters.FindParameterIndex(EMaterialParameterType::Font, ParameterInfo);
+		if (ParameterIndex != INDEX_NONE)
+		{
+			checkLayerParameter(bLegacyFoundValue);
+			OutFontValue = CachedLayerParameters.FontValues[ParameterIndex];
+			OutFontPage = CachedLayerParameters.FontPageValues[ParameterIndex];
+			checkLayerParameter(OutFontValue == LegacyValue);
+			checkLayerParameter(OutFontPage == LegacyFontPageValue);
+			checkLayerParameter(bLegacyResult);
+			return true;
 		}
+
+		checkLayerParameter(!bLegacyFoundValue);
 	}
 	else if (bCheckOwnedGlobalOverrides)
 	{
@@ -2546,7 +2546,8 @@ bool UMaterialInstance::GetFontParameterDefaultValue(const FMaterialParameterInf
 	return false;
 }
 
-bool UMaterialInstance::GetStaticSwitchParameterDefaultValue(const FMaterialParameterInfo& ParameterInfo, bool& OutValue, FGuid& OutExpressionGuid, bool bCheckOwnedGlobalOverrides) const
+#if WITH_EDITOR
+bool UMaterialInstance::GetStaticSwitchParameterDefaultValue(const FHashedMaterialParameterInfo& ParameterInfo, bool& OutValue, FGuid& OutExpressionGuid, bool bCheckOwnedGlobalOverrides) const
 {
 	if (GetReentrantFlag())
 	{
@@ -2606,7 +2607,7 @@ bool UMaterialInstance::GetStaticSwitchParameterDefaultValue(const FMaterialPara
 	return false;
 }
 
-bool UMaterialInstance:: GetStaticComponentMaskParameterDefaultValue(const FMaterialParameterInfo& ParameterInfo, bool& OutR, bool& OutG, bool& OutB, bool& OutA, FGuid& OutExpressionGuid, bool bCheckOwnedGlobalOverrides) const
+bool UMaterialInstance:: GetStaticComponentMaskParameterDefaultValue(const FHashedMaterialParameterInfo& ParameterInfo, bool& OutR, bool& OutG, bool& OutB, bool& OutA, FGuid& OutExpressionGuid, bool bCheckOwnedGlobalOverrides) const
 {
 	if (GetReentrantFlag())
 	{
@@ -2669,8 +2670,7 @@ bool UMaterialInstance:: GetStaticComponentMaskParameterDefaultValue(const FMate
 	return false;
 }
 
-#if WITH_EDITOR
-bool UMaterialInstance::GetGroupName(const FMaterialParameterInfo& ParameterInfo, FName& OutGroup) const
+bool UMaterialInstance::GetGroupName(const FHashedMaterialParameterInfo& ParameterInfo, FName& OutGroup) const
 {
 	if (GetReentrantFlag())
 	{
@@ -2719,38 +2719,6 @@ bool UMaterialInstance::GetGroupName(const FMaterialParameterInfo& ParameterInfo
 }
 #endif // WITH_EDITOR
 
-void UMaterialInstance::AppendReferencedTextures(TArray<UObject*>& InOutTextures) const
-{
-	// @TODO: This shouldn't append all textures, instead replace those included by overridden layer stacks
-	for (const FStaticMaterialLayersParameter& LayersParam : StaticParameters.MaterialLayersParameters)
-	{
-		if (LayersParam.bOverride)
-		{
-			for (UMaterialFunctionInterface* Layer : LayersParam.Value.Layers)
-			{
-				// @TODO: We can skip hidden layers
-				if (Layer)
-				{
-					Layer->AppendReferencedTextures(InOutTextures);
-				}
-			}
-
-			for (UMaterialFunctionInterface* Blend : LayersParam.Value.Blends)
-			{
-				if (Blend)
-				{
-					Blend->AppendReferencedTextures(InOutTextures);
-				}
-			}
-		}
-	}
-
-	if (Parent)
-	{
-		Parent->AppendReferencedTextures(InOutTextures);
-	}
-}
-
 #if WITH_EDITOR
 void UMaterialInstance::ForceRecompileForRendering()
 {
@@ -3066,11 +3034,10 @@ void UMaterialInstance::CacheResourceShadersForCooking(EShaderPlatform ShaderPla
 void UMaterialInstance::CacheShadersForResources(EShaderPlatform ShaderPlatform, const TArray<FMaterialResource*>& ResourcesToCache, const ITargetPlatform* TargetPlatform)
 {
 	UMaterial* BaseMaterial = GetMaterial();
-	BaseMaterial->CacheExpressionTextureReferences();
-
-	// Build list of textures referenced by this instance permutation
-	PermutationTextureReferences.Empty();
-	AppendReferencedTextures(PermutationTextureReferences);
+#if WITH_EDITOR
+	BaseMaterial->UpdateCachedExpressionData();
+	UpdateCachedLayerParameters();
+#endif
 
 	for (int32 ResourceIndex = 0; ResourceIndex < ResourcesToCache.Num(); ResourceIndex++)
 	{
@@ -3097,6 +3064,7 @@ void UMaterialInstance::CacheShadersForResources(EShaderPlatform ShaderPlatform,
 	}
 }
 
+#if WITH_EDITORONLY_DATA
 bool UMaterialInstance::GetStaticSwitchParameterValues(FStaticParamEvaluationContext& EvalContext, TBitArray<>& OutValues, FGuid* OutExpressionGuids, bool bCheckParent /*= true*/) const
 {
 	if (EvalContext.AllResolved())
@@ -3118,7 +3086,7 @@ bool UMaterialInstance::GetStaticSwitchParameterValues(FStaticParamEvaluationCon
 		const FStaticSwitchParameter& Param = StaticParameters.StaticSwitchParameters[ValueIndex];
 		if (Param.bOverride)
 		{
-			EvalContext.ForEachPendingParameter([&EvalContext, &Param, &OutValues, &OutExpressionGuids](int32 ParamIndex, const FMaterialParameterInfo& ParameterInfo) -> bool
+			EvalContext.ForEachPendingParameter([&EvalContext, &Param, &OutValues, &OutExpressionGuids](int32 ParamIndex, const FHashedMaterialParameterInfo& ParameterInfo) -> bool
 				{
 					if (Param.ParameterInfo == ParameterInfo)
 					{
@@ -3145,7 +3113,7 @@ bool UMaterialInstance::GetStaticSwitchParameterValues(FStaticParamEvaluationCon
 		{
 			TMaterialStaticParameterValueResolver<UMaterialExpressionStaticBoolParameter, TBitArray<>> ValueResolver(EvalContext, OutValues, OutExpressionGuids);
 
-			EvalContext.ForEachPendingParameter([&EvalContext, &LayersParam, &ValueResolver, &OutValues, &OutExpressionGuids](int32 ParamIndex, const FMaterialParameterInfo& ParameterInfo) -> bool
+			EvalContext.ForEachPendingParameter([&EvalContext, &LayersParam, &ValueResolver, &OutValues, &OutExpressionGuids](int32 ParamIndex, const FHashedMaterialParameterInfo& ParameterInfo) -> bool
 				{
 					if (ParameterInfo.Association == EMaterialParameterAssociation::GlobalParameter)
 					{
@@ -3196,7 +3164,7 @@ bool UMaterialInstance::GetStaticComponentMaskParameterValues(FStaticParamEvalua
 		const FStaticComponentMaskParameter& Param = StaticParameters.StaticComponentMaskParameters[ValueIndex];
 		if (Param.bOverride)
 		{
-			EvalContext.ForEachPendingParameter([&EvalContext, &Param, &OutRGBAOrderedValues, &OutExpressionGuids](int32 ParamIndex, const FMaterialParameterInfo& ParameterInfo) -> bool
+			EvalContext.ForEachPendingParameter([&EvalContext, &Param, &OutRGBAOrderedValues, &OutExpressionGuids](int32 ParamIndex, const FHashedMaterialParameterInfo& ParameterInfo) -> bool
 				{
 					if (Param.ParameterInfo == ParameterInfo)
 					{
@@ -3227,7 +3195,7 @@ bool UMaterialInstance::GetStaticComponentMaskParameterValues(FStaticParamEvalua
 		{
 			TMaterialStaticParameterValueResolver<UMaterialExpressionStaticComponentMaskParameter, TBitArray<>> ValueResolver(EvalContext, OutRGBAOrderedValues, OutExpressionGuids);
 
-			EvalContext.ForEachPendingParameter([&EvalContext, &LayersParam, &ValueResolver, &OutRGBAOrderedValues, &OutExpressionGuids](int32 ParamIndex, const FMaterialParameterInfo& ParameterInfo) -> bool
+			EvalContext.ForEachPendingParameter([&EvalContext, &LayersParam, &ValueResolver, &OutRGBAOrderedValues, &OutExpressionGuids](int32 ParamIndex, const FHashedMaterialParameterInfo& ParameterInfo) -> bool
 				{
 					if (ParameterInfo.Association == EMaterialParameterAssociation::GlobalParameter)
 					{
@@ -3257,56 +3225,57 @@ bool UMaterialInstance::GetStaticComponentMaskParameterValues(FStaticParamEvalua
 	return EvalContext.AllResolved();
 }
 
-bool UMaterialInstance::GetTerrainLayerWeightParameterValue(const FMaterialParameterInfo& ParameterInfo, int32& OutWeightmapIndex, FGuid &OutExpressionGuid) const
+bool UMaterialInstance::GetMaterialLayersParameterValue(const FHashedMaterialParameterInfo& ParameterInfo, FMaterialLayersFunctions& OutLayers, FGuid& OutExpressionGuid, bool bCheckParent /*= true*/) const
 {
-	if(GetReentrantFlag())
+	if (GetReentrantFlag())
 	{
 		return false;
 	}
 
-	for (int32 ValueIndex = 0;ValueIndex < StaticParameters.TerrainLayerWeightParameters.Num();ValueIndex++)
+	for (const FStaticMaterialLayersParameter& Param : StaticParameters.MaterialLayersParameters)
 	{
-		const FStaticTerrainLayerWeightParameter& Param = StaticParameters.TerrainLayerWeightParameters[ValueIndex];
 		if (Param.bOverride && Param.ParameterInfo == ParameterInfo)
 		{
-			OutWeightmapIndex = Param.WeightmapIndex;
+			OutLayers = Param.Value;
 			OutExpressionGuid = Param.ExpressionGUID;
 			return true;
 		}
 	}
 
-	if(Parent)
+	if (Parent && bCheckParent)
 	{
 		FMICReentranceGuard	Guard(this);
-		return Parent->GetTerrainLayerWeightParameterValue(ParameterInfo, OutWeightmapIndex, OutExpressionGuid);
+		return Parent->GetMaterialLayersParameterValue(ParameterInfo, OutLayers, OutExpressionGuid);
 	}
 	else
 	{
 		return false;
 	}
 }
+#endif // WITH_EDITORONLY_DATA
 
-bool UMaterialInstance::GetMaterialLayersParameterValue(const FMaterialParameterInfo& ParameterInfo, FMaterialLayersFunctions& OutLayers, FGuid& OutExpressionGuid, bool bCheckParent /*= true*/) const
+bool UMaterialInstance::GetTerrainLayerWeightParameterValue(const FHashedMaterialParameterInfo& ParameterInfo, int32& OutWeightmapIndex, FGuid &OutExpressionGuid) const
 {
-	if( GetReentrantFlag() )
+	if(GetReentrantFlag())
 	{
 		return false;
 	}
 
-	for (const FStaticMaterialLayersParameter& Param : StaticParameters.MaterialLayersParameters)
+	for (int32 ValueIndex = 0;ValueIndex < StaticParameters.TerrainLayerWeightParameters.Num();ValueIndex++)
 	{
+		const FStaticTerrainLayerWeightParameter& Param = StaticParameters.TerrainLayerWeightParameters[ValueIndex];
 		if (Param.bOverride && Param.ParameterInfo == ParameterInfo)
 		{
-			OutLayers = Param.Value;
+			OutWeightmapIndex = Param.WeightmapIndex;
 			OutExpressionGuid = Param.ExpressionGUID;
 			return true;
 		}
 	}
 
-	if( Parent && bCheckParent)
+	if(Parent)
 	{
 		FMICReentranceGuard	Guard(this);
-		return Parent->GetMaterialLayersParameterValue(ParameterInfo, OutLayers, OutExpressionGuid);
+		return Parent->GetTerrainLayerWeightParameterValue(ParameterInfo, OutWeightmapIndex, OutExpressionGuid);
 	}
 	else
 	{
@@ -3314,7 +3283,7 @@ bool UMaterialInstance::GetMaterialLayersParameterValue(const FMaterialParameter
 	}
 }
 
-bool UMaterialInstance::UpdateMaterialLayersParameterValue(const FMaterialParameterInfo& ParameterInfo, const FMaterialLayersFunctions& LayersValue, const bool bOverridden, const FGuid& GUID)
+bool UMaterialInstance::UpdateMaterialLayersParameterValue(const FHashedMaterialParameterInfo& ParameterInfo, const FMaterialLayersFunctions& LayersValue, const bool bOverridden, const FGuid& GUID)
 {
 	for (int32 ValueIndex = 0;ValueIndex < StaticParameters.MaterialLayersParameters.Num();ValueIndex++)
 	{
@@ -3600,7 +3569,7 @@ void UMaterialInstance::PostLoad()
 			}
 		}
 	}
-#endif
+#endif // WITH_EDITOR
 
 	if (FApp::CanEverRender())
 	{
@@ -3682,6 +3651,10 @@ void UMaterialInstance::PostLoad()
 		}
 	}
 
+#if WITH_EDITOR
+	UpdateCachedLayerParameters();
+#endif
+
 	// called before we cache the uniform expression as a call to SubsurfaceProfileRT affects the dta in there
 	PropagateDataToMaterialProxy();
 
@@ -3847,6 +3820,9 @@ void UMaterialInstance::SetParentInternal(UMaterialInterface* NewParent, bool Re
 		{
 			InitResources();
 		}
+#if WITH_EDITOR
+		UpdateCachedLayerParameters();
+#endif
 	}
 }
 
@@ -4115,6 +4091,7 @@ void UMaterialInstance::UpdateStaticPermutation(const FStaticParameterSet& NewPa
 		bHasStaticPermutationResource = bWantsStaticPermutationResource;
 		StaticParameters = CompareParameters;
 
+		UpdateCachedLayerParameters();
 		CacheResourceShadersForRendering();
 		RecacheUniformExpressions(true);
 
@@ -4132,6 +4109,23 @@ void UMaterialInstance::UpdateStaticPermutation(const FStaticParameterSet& NewPa
 	}
 }
 
+void UMaterialInstance::UpdateCachedLayerParameters()
+{
+	FMaterialCachedExpressionData CachedExpressionData;
+	if (Parent)
+	{
+		CachedExpressionData.ReferencedTextures = Parent->GetReferencedTextures();
+	}
+	
+	for (FStaticMaterialLayersParameter& LayerParameters : StaticParameters.MaterialLayersParameters)
+	{
+		CachedExpressionData.UpdateForLayerFunctions(LayerParameters.Value);
+	}
+
+	CachedLayerParameters = MoveTemp(CachedExpressionData.Parameters);
+	CachedReferencedTextures = MoveTemp(CachedExpressionData.ReferencedTextures);
+}
+
 void UMaterialInstance::UpdateStaticPermutation(const FStaticParameterSet& NewParameters, FMaterialUpdateContext* MaterialUpdateContext)
 {
 	UpdateStaticPermutation(NewParameters, BasePropertyOverrides, false, MaterialUpdateContext);
@@ -4152,7 +4146,7 @@ void UMaterialInstance::UpdateParameterNames()
 		InitResources();
 	}
 }
-#endif
+#endif // WITH_EDITOR
 
 void UMaterialInstance::RecacheUniformExpressions(bool bRecreateUniformBuffer) const
 {	
@@ -4180,6 +4174,8 @@ void UMaterialInstance::PostEditChangeProperty(FPropertyChangedEvent& PropertyCh
 	{
 		RecacheMaterialInstanceUniformExpressions(this, false);
 	}
+
+	UpdateCachedLayerParameters();
 }
 
 #endif // WITH_EDITOR
@@ -4288,7 +4284,7 @@ float UMaterialInstance::GetExportResolutionScale() const
 }
 
 #if WITH_EDITOR
-bool UMaterialInstance::GetParameterDesc(const FMaterialParameterInfo& ParameterInfo, FString& OutDesc, const TArray<struct FStaticMaterialLayersParameter>* MaterialLayersParameters) const
+bool UMaterialInstance::GetParameterDesc(const FHashedMaterialParameterInfo& ParameterInfo, FString& OutDesc, const TArray<struct FStaticMaterialLayersParameter>* MaterialLayersParameters) const
 {
 	const UMaterial* BaseMaterial = GetMaterial();
 	if (BaseMaterial && BaseMaterial->GetParameterDesc(ParameterInfo, OutDesc, &StaticParameters.MaterialLayersParameters))
@@ -4299,7 +4295,7 @@ bool UMaterialInstance::GetParameterDesc(const FMaterialParameterInfo& Parameter
 	return false;
 }
 
-bool UMaterialInstance::GetParameterSortPriority(const FMaterialParameterInfo& ParameterInfo, int32& OutSortPriority, const TArray<struct FStaticMaterialLayersParameter>* MaterialLayersParameters) const
+bool UMaterialInstance::GetParameterSortPriority(const FHashedMaterialParameterInfo& ParameterInfo, int32& OutSortPriority, const TArray<struct FStaticMaterialLayersParameter>* MaterialLayersParameters) const
 {
 	const UMaterial* BaseMaterial = GetMaterial();
 	if (BaseMaterial && BaseMaterial->GetParameterSortPriority(ParameterInfo, OutSortPriority, &StaticParameters.MaterialLayersParameters))
@@ -4985,54 +4981,33 @@ void UMaterialInstance::CopyMaterialUniformParametersInternal(UMaterialInterface
 				if (MaterialResource)
 				{
 					// Scalars
-					const TArray<TRefCountPtr<FMaterialUniformExpression>>& ScalarExpressions = MaterialResource->GetUniformScalarParameterExpressions();
-					for (FMaterialUniformExpression* ScalarExpression : ScalarExpressions)
+					const TArray<FMaterialScalarParameterInfo, FMemoryImageAllocator>& ScalarExpressions = MaterialResource->GetUniformScalarParameterExpressions();
+					for (const FMaterialScalarParameterInfo& Parameter : ScalarExpressions)
 					{
-						if (ScalarExpression->GetType() == &FMaterialUniformExpressionScalarParameter::StaticType)
-						{
-							FMaterialUniformExpressionScalarParameter* ScalarParameter = static_cast<FMaterialUniformExpressionScalarParameter*>(ScalarExpression);
-
-							FScalarParameterValue* ParameterValue = new(ScalarParameterValues) FScalarParameterValue;
-							ParameterValue->ParameterInfo.Name = ScalarParameter->GetParameterInfo().Name;
-							ScalarParameter->GetDefaultValue(ParameterValue->ParameterValue);
-						}
+						FScalarParameterValue* ParameterValue = new(ScalarParameterValues) FScalarParameterValue;
+						ParameterValue->ParameterInfo.Name = *Parameter.ParameterName;
+						Parameter.GetDefaultValue(ParameterValue->ParameterValue);
 					}
 
 					// Vectors
-					const TArray<TRefCountPtr<FMaterialUniformExpression>>& VectorExpressions = MaterialResource->GetUniformVectorParameterExpressions();
-					for (FMaterialUniformExpression* VectorExpression : VectorExpressions)
+					const TArray<FMaterialVectorParameterInfo, FMemoryImageAllocator>& VectorExpressions = MaterialResource->GetUniformVectorParameterExpressions();
+					for (const FMaterialVectorParameterInfo& Parameter : VectorExpressions)
 					{
-						if (VectorExpression->GetType() == &FMaterialUniformExpressionVectorParameter::StaticType)
-						{
-							FMaterialUniformExpressionVectorParameter* VectorParameter = static_cast<FMaterialUniformExpressionVectorParameter*>(VectorExpression);
-
-							FVectorParameterValue* ParameterValue = new(VectorParameterValues) FVectorParameterValue;
-							ParameterValue->ParameterInfo.Name = VectorParameter->GetParameterInfo().Name;
-							VectorParameter->GetDefaultValue(ParameterValue->ParameterValue);
-						}
+						FVectorParameterValue* ParameterValue = new(VectorParameterValues) FVectorParameterValue;
+						ParameterValue->ParameterInfo.Name = *Parameter.ParameterName;
+						Parameter.GetDefaultValue(ParameterValue->ParameterValue);
 					}
 
 					// Textures
-					const TArray<TRefCountPtr<FMaterialUniformExpressionTexture>>* TextureExpressions[5] =
+					for (int32 TypeIndex = 0; TypeIndex < NumMaterialTextureParameterTypes; TypeIndex++)
 					{
-						&MaterialResource->GetUniform2DTextureExpressions(),
-						&MaterialResource->GetUniformCubeTextureExpressions(),
-						&MaterialResource->GetUniform2DArrayTextureExpressions(),
-						&MaterialResource->GetUniformVolumeTextureExpressions(),
-						&MaterialResource->GetUniformVirtualTextureExpressions()
-					};
-
-					for (int32 TypeIndex = 0; TypeIndex < UE_ARRAY_COUNT(TextureExpressions); TypeIndex++)
-					{
-						for (FMaterialUniformExpressionTexture* TextureExpression : *TextureExpressions[TypeIndex])
+						for (const FMaterialTextureParameterInfo& Parameter : MaterialResource->GetUniformTextureExpressions((EMaterialTextureParameterType)TypeIndex))
 						{
-							if (TextureExpression->GetType() == &FMaterialUniformExpressionTextureParameter::StaticType)
+							if (!Parameter.ParameterInfo.Name.IsNone())
 							{
-								FMaterialUniformExpressionTextureParameter* TextureParameter = static_cast<FMaterialUniformExpressionTextureParameter*>(TextureExpression);
-
 								FTextureParameterValue* ParameterValue = new(TextureParameterValues) FTextureParameterValue;
-								ParameterValue->ParameterInfo.Name = TextureParameter->GetParameterName();
-								TextureParameter->GetGameThreadTextureValue(AsMaterial, *MaterialResource, ParameterValue->ParameterValue, false);
+								ParameterValue->ParameterInfo.Name = *Parameter.ParameterName;
+								Parameter.GetGameThreadTextureValue(AsMaterial, *MaterialResource, ParameterValue->ParameterValue);
 							}
 						}
 					}
@@ -5044,6 +5019,7 @@ void UMaterialInstance::CopyMaterialUniformParametersInternal(UMaterialInterface
 	}
 }
 
-
+#if WITH_EDITORONLY_DATA
 UMaterialInstance::FCustomStaticParametersGetterDelegate UMaterialInstance::CustomStaticParametersGetters;
 TArray<UMaterialInstance::FCustomParameterSetUpdaterDelegate> UMaterialInstance::CustomParameterSetUpdaters;
+#endif // WITH_EDITORONLY_DATA
diff --git a/Engine/Source/Runtime/Engine/Private/Materials/MaterialInstanceSupport.h b/Engine/Source/Runtime/Engine/Private/Materials/MaterialInstanceSupport.h
--- a/Engine/Source/Runtime/Engine/Private/Materials/MaterialInstanceSupport.h
+++ b/Engine/Source/Runtime/Engine/Private/Materials/MaterialInstanceSupport.h
@@ -66,7 +66,7 @@ public:
 	template <typename ValueType>
 	struct TNamedParameter
 	{
-		FMaterialParameterInfo Info;
+		FHashedMaterialParameterInfo Info;
 		ValueType Value;
 	};
 
@@ -96,10 +96,10 @@ public:
 	virtual FMaterial* GetMaterialNoFallback(ERHIFeatureLevel::Type FeatureLevel) const override;
 	virtual UMaterialInterface* GetMaterialInterface() const override;
 	
-	virtual bool GetVectorValue(const FMaterialParameterInfo& ParameterInfo, FLinearColor* OutValue, const FMaterialRenderContext& Context) const override;
-	virtual bool GetScalarValue(const FMaterialParameterInfo& ParameterInfo, float* OutValue, const FMaterialRenderContext& Context) const override;
-	virtual bool GetTextureValue(const FMaterialParameterInfo& ParameterInfo, const UTexture** OutValue, const FMaterialRenderContext& Context) const override;
-	virtual bool GetTextureValue(const FMaterialParameterInfo& ParameterInfo, const URuntimeVirtualTexture** OutValue, const FMaterialRenderContext& Context) const override;
+	virtual bool GetVectorValue(const FHashedMaterialParameterInfo& ParameterInfo, FLinearColor* OutValue, const FMaterialRenderContext& Context) const override;
+	virtual bool GetScalarValue(const FHashedMaterialParameterInfo& ParameterInfo, float* OutValue, const FMaterialRenderContext& Context) const override;
+	virtual bool GetTextureValue(const FHashedMaterialParameterInfo& ParameterInfo, const UTexture** OutValue, const FMaterialRenderContext& Context) const override;
+	virtual bool GetTextureValue(const FHashedMaterialParameterInfo& ParameterInfo, const URuntimeVirtualTexture** OutValue, const FMaterialRenderContext& Context) const override;
 
 	void GameThread_SetParent(UMaterialInterface* ParentMaterialInterface);
 
@@ -119,7 +119,7 @@ public:
 	 * Updates a named parameter on the render thread.
 	 */
 	template <typename ValueType>
-	void RenderThread_UpdateParameter(const FMaterialParameterInfo& ParameterInfo, const ValueType& Value )
+	void RenderThread_UpdateParameter(const FHashedMaterialParameterInfo& ParameterInfo, const ValueType& Value )
 	{
 		LLM_SCOPE(ELLMTag::MaterialInstance);
 
@@ -145,7 +145,7 @@ public:
 	 * Retrieves a parameter by name.
 	 */
 	template <typename ValueType>
-	const ValueType* RenderThread_FindParameterByName(const FMaterialParameterInfo& ParameterInfo) const
+	const ValueType* RenderThread_FindParameterByName(const FHashedMaterialParameterInfo& ParameterInfo) const
 	{
 		const TArray<TNamedParameter<ValueType> >& ValueArray = GetValueArray<ValueType>();
 		const int32 ParameterCount = ValueArray.Num();
@@ -197,7 +197,7 @@ template <> FORCEINLINE const TArray<FMaterialInstanceResource::TNamedParameter<
 
 /** Finds a parameter by name from the game thread. */
 template <typename ParameterType>
-ParameterType* GameThread_FindParameterByName(TArray<ParameterType>& Parameters, const FMaterialParameterInfo& ParameterInfo)
+ParameterType* GameThread_FindParameterByName(TArray<ParameterType>& Parameters, const FHashedMaterialParameterInfo& ParameterInfo)
 {
 	for (int32 ParameterIndex = 0; ParameterIndex < Parameters.Num(); ParameterIndex++)
 	{
@@ -210,7 +210,7 @@ ParameterType* GameThread_FindParameterByName(TArray<ParameterType>& Parameters,
 	return NULL;
 }
 template <typename ParameterType>
-const ParameterType* GameThread_FindParameterByName(const TArray<ParameterType>& Parameters, const FMaterialParameterInfo& ParameterInfo)
+const ParameterType* GameThread_FindParameterByName(const TArray<ParameterType>& Parameters, const FHashedMaterialParameterInfo& ParameterInfo)
 {
 	for (int32 ParameterIndex = 0; ParameterIndex < Parameters.Num(); ParameterIndex++)
 	{
@@ -224,7 +224,7 @@ const ParameterType* GameThread_FindParameterByName(const TArray<ParameterType>&
 }
 
 template <typename ParameterType>
-int32 GameThread_FindParameterIndexByName(const TArray<ParameterType>& Parameters, const FMaterialParameterInfo& ParameterInfo)
+int32 GameThread_FindParameterIndexByName(const TArray<ParameterType>& Parameters, const FHashedMaterialParameterInfo& ParameterInfo)
 {
 	for (int32 ParameterIndex = 0; ParameterIndex < Parameters.Num(); ++ParameterIndex)
 	{
