Copying Private-LoadTimes-4.24 stream to Main. Biggest changes are in Materials/Shader memory freezing.

diff --git a/Engine/Source/Runtime/Engine/Private/Materials/MaterialInterface.cpp b/Engine/Source/Runtime/Engine/Private/Materials/MaterialInterface.cpp
--- a/Engine/Source/Runtime/Engine/Private/Materials/MaterialInterface.cpp
+++ b/Engine/Source/Runtime/Engine/Private/Materials/MaterialInterface.cpp
@@ -91,7 +91,8 @@ void UMaterialInterface::GetUsedTexturesAndIndices(TArray<UTexture*>& OutTexture
 	OutIndices.AddDefaulted(OutTextures.Num());
 }
 
-bool UMaterialInterface::GetStaticSwitchParameterValue(const FMaterialParameterInfo& ParameterInfo, bool& OutValue, FGuid& OutExpressionGuid, bool bOveriddenOnly /*= false*/, bool bCheckParent /*= true*/) const
+#if WITH_EDITORONLY_DATA
+bool UMaterialInterface::GetStaticSwitchParameterValue(const FHashedMaterialParameterInfo& ParameterInfo, bool& OutValue, FGuid& OutExpressionGuid, bool bOveriddenOnly /*= false*/, bool bCheckParent /*= true*/) const
 {
 	TBitArray<> Output(false, 1); // Relying on the default allocator to be inline to avoid allocation here.
 	FStaticParamEvaluationContext EvalContext(1, &ParameterInfo);
@@ -110,7 +111,7 @@ bool UMaterialInterface::GetStaticSwitchParameterValue(const FMaterialParameterI
 	return true;
 }
 
-bool UMaterialInterface::GetStaticComponentMaskParameterValue(const FMaterialParameterInfo& ParameterInfo, bool& R, bool& G, bool& B, bool& A, FGuid& OutExpressionGuid, bool bOveriddenOnly /*= false*/, bool bCheckParent /*= true*/) const
+bool UMaterialInterface::GetStaticComponentMaskParameterValue(const FHashedMaterialParameterInfo& ParameterInfo, bool& R, bool& G, bool& B, bool& A, FGuid& OutExpressionGuid, bool bOveriddenOnly /*= false*/, bool bCheckParent /*= true*/) const
 {
 	TBitArray<> Output(false, 4); // Relying on the default allocator to be inline to avoid allocation here.
 	FStaticParamEvaluationContext EvalContext(1, &ParameterInfo);
@@ -131,6 +132,7 @@ bool UMaterialInterface::GetStaticComponentMaskParameterValue(const FMaterialPar
 
 	return true;
 }
+#endif
 
 FMaterialRelevance UMaterialInterface::GetRelevance_Internal(const UMaterial* Material, ERHIFeatureLevel::Type InFeatureLevel) const
 {
@@ -214,8 +216,7 @@ FMaterialRelevance UMaterialInterface::GetRelevance(ERHIFeatureLevel::Type InFea
 FMaterialRelevance UMaterialInterface::GetRelevance_Concurrent(ERHIFeatureLevel::Type InFeatureLevel) const
 {
 	// Find the interface's concrete material.
-	TMicRecursionGuard RecursionGuard;
-	const UMaterial* Material = GetMaterial_Concurrent(RecursionGuard);
+	const UMaterial* Material = GetMaterial_Concurrent();
 	return GetRelevance_Internal(Material, InFeatureLevel);
 }
 
@@ -333,78 +334,80 @@ void UMaterialInterface::GetLightingGuidChain(bool bIncludeTextures, TArray<FGui
 #endif // WITH_EDITORONLY_DATA
 }
 
-bool UMaterialInterface::GetVectorParameterValue(const FMaterialParameterInfo& ParameterInfo, FLinearColor& OutValue, bool bOveriddenOnly) const
+bool UMaterialInterface::GetVectorParameterValue(const FHashedMaterialParameterInfo& ParameterInfo, FLinearColor& OutValue, bool bOveriddenOnly) const
 {
 	// is never called but because our system wants a UMaterialInterface instance we cannot use "virtual =0"
 	return false;
 }
 
-bool UMaterialInterface::IsVectorParameterUsedAsChannelMask(const FMaterialParameterInfo& ParameterInfo, bool& OutValue) const
+#if WITH_EDITOR
+bool UMaterialInterface::IsVectorParameterUsedAsChannelMask(const FHashedMaterialParameterInfo& ParameterInfo, bool& OutValue) const
 {
 	return false;
 }
 
-#if WITH_EDITOR
-bool UMaterialInterface::GetVectorParameterChannelNames(const FMaterialParameterInfo& ParameterInfo, FParameterChannelNames& OutValue) const
+bool UMaterialInterface::GetVectorParameterChannelNames(const FHashedMaterialParameterInfo& ParameterInfo, FParameterChannelNames& OutValue) const
 {
 	return false;
 }
 
-bool UMaterialInterface::GetScalarParameterSliderMinMax(const FMaterialParameterInfo& ParameterInfo, float& OutSliderMin, float& OutSliderMax) const
+bool UMaterialInterface::GetScalarParameterSliderMinMax(const FHashedMaterialParameterInfo& ParameterInfo, float& OutSliderMin, float& OutSliderMax) const
 {
 	return false;
 }
-#endif
+#endif // WITH_EDITOR
 
-bool UMaterialInterface::GetScalarParameterValue(const FMaterialParameterInfo& ParameterInfo, float& OutValue, bool bOveriddenOnly) const
+bool UMaterialInterface::GetScalarParameterValue(const FHashedMaterialParameterInfo& ParameterInfo, float& OutValue, bool bOveriddenOnly) const
 {
 	// is never called but because our system wants a UMaterialInterface instance we cannot use "virtual =0"
 	return false;
 }
 
-bool UMaterialInterface::IsScalarParameterUsedAsAtlasPosition(const FMaterialParameterInfo& ParameterInfo, bool& OutValue, TSoftObjectPtr<class UCurveLinearColor>& Curve, TSoftObjectPtr<class UCurveLinearColorAtlas>& Atlas) const
+#if WITH_EDITOR
+bool UMaterialInterface::IsScalarParameterUsedAsAtlasPosition(const FHashedMaterialParameterInfo& ParameterInfo, bool& OutValue, TSoftObjectPtr<class UCurveLinearColor>& Curve, TSoftObjectPtr<class UCurveLinearColorAtlas>& Atlas) const
 {
 	return false;
 }
+#endif // WITH_EDITOR
 
-bool UMaterialInterface::GetScalarCurveParameterValue(const FMaterialParameterInfo& ParameterInfo, FInterpCurveFloat& OutValue) const
+bool UMaterialInterface::GetScalarCurveParameterValue(const FHashedMaterialParameterInfo& ParameterInfo, FInterpCurveFloat& OutValue) const
 {
 	return false;
 }
 
-bool UMaterialInterface::GetVectorCurveParameterValue(const FMaterialParameterInfo& ParameterInfo, FInterpCurveVector& OutValue) const
+bool UMaterialInterface::GetVectorCurveParameterValue(const FHashedMaterialParameterInfo& ParameterInfo, FInterpCurveVector& OutValue) const
 {
 	return false;
 }
 
-bool UMaterialInterface::GetLinearColorParameterValue(const FMaterialParameterInfo& ParameterInfo, FLinearColor& OutValue) const
+bool UMaterialInterface::GetLinearColorParameterValue(const FHashedMaterialParameterInfo& ParameterInfo, FLinearColor& OutValue) const
 {
 	return false;
 }
 
-bool UMaterialInterface::GetLinearColorCurveParameterValue(const FMaterialParameterInfo& ParameterInfo, FInterpCurveLinearColor& OutValue) const
+bool UMaterialInterface::GetLinearColorCurveParameterValue(const FHashedMaterialParameterInfo& ParameterInfo, FInterpCurveLinearColor& OutValue) const
 {
 	return false;
 }
 
-bool UMaterialInterface::GetTextureParameterValue(const FMaterialParameterInfo& ParameterInfo, UTexture*& OutValue, bool bOveriddenOnly) const
+bool UMaterialInterface::GetTextureParameterValue(const FHashedMaterialParameterInfo& ParameterInfo, UTexture*& OutValue, bool bOveriddenOnly) const
 {
 	return false;
 }
 
-bool UMaterialInterface::GetRuntimeVirtualTextureParameterValue(const FMaterialParameterInfo& ParameterInfo, URuntimeVirtualTexture*& OutValue, bool bOveriddenOnly) const
+bool UMaterialInterface::GetRuntimeVirtualTextureParameterValue(const FHashedMaterialParameterInfo& ParameterInfo, URuntimeVirtualTexture*& OutValue, bool bOveriddenOnly) const
 {
 	return false;
 }
 
 #if WITH_EDITOR
-bool UMaterialInterface::GetTextureParameterChannelNames(const FMaterialParameterInfo& ParameterInfo, FParameterChannelNames& OutValue) const
+bool UMaterialInterface::GetTextureParameterChannelNames(const FHashedMaterialParameterInfo& ParameterInfo, FParameterChannelNames& OutValue) const
 {
 	return false;
 }
 #endif
 
-bool UMaterialInterface::GetFontParameterValue(const FMaterialParameterInfo& ParameterInfo, class UFont*& OutFontValue, int32& OutFontPage, bool bOveriddenOnly) const
+bool UMaterialInterface::GetFontParameterValue(const FHashedMaterialParameterInfo& ParameterInfo, class UFont*& OutFontValue, int32& OutFontPage, bool bOveriddenOnly) const
 {
 	return false;
 }
@@ -415,12 +418,12 @@ bool UMaterialInterface::GetRefractionSettings(float& OutBiasValue) const
 }
 
 #if WITH_EDITOR
-bool UMaterialInterface::GetParameterDesc(const FMaterialParameterInfo& ParameterInfo, FString& OutDesc, const TArray<struct FStaticMaterialLayersParameter>* MaterialLayersParameters) const
+bool UMaterialInterface::GetParameterDesc(const FHashedMaterialParameterInfo& ParameterInfo, FString& OutDesc, const TArray<struct FStaticMaterialLayersParameter>* MaterialLayersParameters) const
 {
 	return false;
 }
 
-bool UMaterialInterface::GetGroupName(const FMaterialParameterInfo& ParameterInfo, FName& OutDesc) const
+bool UMaterialInterface::GetGroupName(const FHashedMaterialParameterInfo& ParameterInfo, FName& OutDesc) const
 {
 	return false;
 }
@@ -606,7 +609,7 @@ void UMaterialInterface::SortTextureStreamingData(bool bForceSort, bool bFinalSo
 		TArray<UObject*> UsedTextures;
 		if (bFinalSort)
 		{
-			AppendReferencedTextures(UsedTextures);
+			UsedTextures = GetReferencedTextures();
 			for (int32 TextureIndex = 0; TextureIndex < UsedTextures.Num(); ++TextureIndex)
 			{
 				UTexture* UsedTexture = Cast<UTexture>(UsedTextures[TextureIndex]);
@@ -789,6 +792,7 @@ void UMaterialInterface::RemoveUserDataOfClass(TSubclassOf<UAssetUserData> InUse
 	}
 }
 
+#if WITH_EDITORONLY_DATA
 void UMaterialInterface::FStaticParamEvaluationContext::MarkParameterResolved(int32 ParamIndex, bool bIsOverride)
 {
 	FBitReference BitRef = PendingParameters[ParamIndex];
@@ -798,7 +802,7 @@ void UMaterialInterface::FStaticParamEvaluationContext::MarkParameterResolved(in
 	--PendingParameterNum;
 }
 
-void UMaterialInterface::FStaticParamEvaluationContext::ForEachPendingParameter(TFunctionRef<bool(int32, const FMaterialParameterInfo&)> Op)
+void UMaterialInterface::FStaticParamEvaluationContext::ForEachPendingParameter(TFunctionRef<bool(int32, const FHashedMaterialParameterInfo&)> Op)
 {
 	for (TConstSetBitIterator<> It(PendingParameters); It; ++It)
 	{
@@ -809,3 +813,5 @@ void UMaterialInterface::FStaticParamEvaluationContext::ForEachPendingParameter(
 		}
 	}
 }
+
+#endif
