Copying Private-LoadTimes-4.24 stream to Main. Biggest changes are in Materials/Shader memory freezing.

diff --git a/Engine/Source/Runtime/Engine/Private/Materials/MaterialShader.cpp b/Engine/Source/Runtime/Engine/Private/Materials/MaterialShader.cpp
--- a/Engine/Source/Runtime/Engine/Private/Materials/MaterialShader.cpp
+++ b/Engine/Source/Runtime/Engine/Private/Materials/MaterialShader.cpp
@@ -8,13 +8,16 @@
 #include "Materials/MaterialFunctionInterface.h"
 #include "ProfilingDebugging/DiagnosticTable.h"
 #include "MeshMaterialShaderType.h"
+#include "MaterialShaderMapLayout.h"
 #include "ShaderCompiler.h"
 #include "DerivedDataCacheInterface.h"
 #include "Interfaces/ITargetPlatformManagerModule.h"
+#include "Interfaces/ITargetPlatform.h"
 #include "ShaderDerivedDataVersion.h"
 #include "ProfilingDebugging/CookStats.h"
 #include "UObject/ReleaseObjectVersion.h"
 #include "UObject/EditorObjectVersion.h"
+#include "ProfilingDebugging/LoadTimeTracker.h"
 
 #if ENABLE_COOK_STATS
 namespace MaterialShaderCookStats
@@ -31,7 +34,6 @@ namespace MaterialShaderCookStats
 }
 #endif
 
-
 //
 // Globals
 //
@@ -123,12 +125,23 @@ FString GetBlendModeString(EBlendMode BlendMode)
 
 #if WITH_EDITOR
 /** Creates a string key for the derived data cache given a shader map id. */
-static FString GetMaterialShaderMapKeyString(const FMaterialShaderMapId& ShaderMapId, EShaderPlatform Platform)
+static FString GetMaterialShaderMapKeyString(const FMaterialShaderMapId& ShaderMapId, EShaderPlatform Platform, const ITargetPlatform* TargetPlatform)
 {
 	FName Format = LegacyShaderPlatformToShaderFormat(Platform);
 	FString ShaderMapKeyString = Format.ToString() + TEXT("_") + FString(FString::FromInt(GetTargetPlatformManagerRef().ShaderFormatVersion(Format))) + TEXT("_");
+
+	FPlatformTypeLayoutParameters LayoutParameters;
+	if (TargetPlatform)
+	{
+		LayoutParameters.InitializeForPlatform(TargetPlatform->IniPlatformName(), TargetPlatform->HasEditorOnlyData());
+	}
+	else
+	{
+		LayoutParameters.InitializeForCurrent();
+	}
+
 	ShaderMapAppendKeyString(Platform, ShaderMapKeyString);
-	ShaderMapId.AppendKeyString(ShaderMapKeyString);
+	ShaderMapId.AppendKeyString(LayoutParameters, ShaderMapKeyString);
 	FMaterialAttributeDefinitionMap::AppendDDCKeyString(ShaderMapKeyString);
 	return FDerivedDataCacheInterface::BuildCacheKey(TEXT("MATSM"), MATERIALSHADERMAP_DERIVEDDATA_VER, *ShaderMapKeyString);
 }
@@ -183,7 +196,7 @@ const FStaticMaterialLayersParameter::ID FStaticMaterialLayersParameter::GetID()
 }
 
 
-UMaterialFunctionInterface* FStaticMaterialLayersParameter::GetParameterAssociatedFunction(const FMaterialParameterInfo& InParameterInfo) const
+UMaterialFunctionInterface* FStaticMaterialLayersParameter::GetParameterAssociatedFunction(const FHashedMaterialParameterInfo& InParameterInfo) const
 {
 	check(InParameterInfo.Association != EMaterialParameterAssociation::GlobalParameter);
 
@@ -208,34 +221,6 @@ UMaterialFunctionInterface* FStaticMaterialLayersParameter::GetParameterAssociat
 	return Function;
 }
 
-void FStaticMaterialLayersParameter::GetParameterAssociatedFunctions(const FMaterialParameterInfo& InParameterInfo, TArray<UMaterialFunctionInterface*>& AssociatedFunctions) const
-{
-	check(InParameterInfo.Association != EMaterialParameterAssociation::GlobalParameter);
-
-	// Grab the associated layer or blend
-	UMaterialFunctionInterface* Function = nullptr;
-
-	if (InParameterInfo.Association == EMaterialParameterAssociation::LayerParameter)
-	{
-		if (Value.Layers.IsValidIndex(InParameterInfo.Index))
-		{
-			Function = Value.Layers[InParameterInfo.Index];
-		}
-	}
-	else if (InParameterInfo.Association == EMaterialParameterAssociation::BlendParameter)
-	{
-		if (Value.Blends.IsValidIndex(InParameterInfo.Index))
-		{
-			Function = Value.Blends[InParameterInfo.Index];
-		}
-	}
-
-	if (Function)
-	{
-		Function->GetDependentFunctions(AssociatedFunctions);
-	}
-}
-
 /** 
 * Tests this set against another for equality, disregarding override settings.
 * 
@@ -304,8 +289,11 @@ bool FStaticParameterSet::Equivalent(const FStaticParameterSet& ReferenceSet) co
 
 void FMaterialShaderMapId::Serialize(FArchive& Ar, bool bLoadedByCookedMaterial)
 {
+	SCOPED_LOADTIMER(FMaterialShaderMapId_Serialize);
+
 	// Note: FMaterialShaderMapId is saved both in packages (legacy UMaterialInstance) and the DDC (FMaterialShaderMap)
 	// Backwards compatibility only works with FMaterialShaderMapId's stored in packages.  
+	// Only serialized in legacy packages if UE4Ver() < VER_UE4_PURGED_FMATERIAL_COMPILE_OUTPUTS
 	// You must bump MATERIALSHADERMAP_DERIVEDDATA_VER as well if changing the serialization of FMaterialShaderMapId.
 	Ar.UsingCustomVersion(FEditorObjectVersion::GUID);
 	Ar.UsingCustomVersion(FReleaseObjectVersion::GUID);
@@ -675,7 +663,7 @@ void FMaterialShaderMapId::UpdateFromParameterSet(const FStaticParameterSet& Sta
 #endif // WITH_EDITOR
 
 #if WITH_EDITOR
-void FMaterialShaderMapId::AppendKeyString(FString& KeyString) const
+void FMaterialShaderMapId::AppendKeyString(const FPlatformTypeLayoutParameters& LayoutParams, FString& KeyString) const
 {
 	check(IsContentValid());
 	KeyString += BaseMaterialId.ToString();
@@ -689,6 +677,8 @@ void FMaterialShaderMapId::AppendKeyString(FString& KeyString) const
 	GetFeatureLevelName(FeatureLevel, FeatureLevelString);
 	KeyString += FeatureLevelString + TEXT("_");
 
+	LayoutParams.AppendKeyString(KeyString);
+
 	for (const FStaticSwitchParameter& StaticSwitchParameter : StaticSwitchParameters)
 	{
 		StaticSwitchParameter.AppendKeyString(KeyString);
@@ -716,6 +706,13 @@ void FMaterialShaderMapId::AppendKeyString(FString& KeyString) const
 		KeyString += ReferencedFunctions[FunctionIndex].ToString();
 	}
 
+	{
+		const FSHAHash LayoutHash = Freeze::HashLayout(StaticGetTypeLayoutDesc<FMaterialShaderMapContent>(), LayoutParams);
+		KeyString += TEXT("_");
+		KeyString += LayoutHash.ToString();
+		KeyString += TEXT("_");
+	}
+
 	KeyString += TEXT("_");
 
 	for (int32 CollectionIndex = 0; CollectionIndex < ReferencedParameterCollections.Num(); CollectionIndex++)
@@ -729,12 +726,16 @@ void FMaterialShaderMapId::AppendKeyString(FString& KeyString) const
 	for (int32 ShaderIndex = 0; ShaderIndex < ShaderTypeDependencies.Num(); ShaderIndex++)
 	{
 		const FShaderTypeDependency& ShaderTypeDependency = ShaderTypeDependencies[ShaderIndex];
+		const FShaderType* ShaderType = FindShaderTypeByName(ShaderTypeDependency.ShaderTypeName);
+
 		KeyString += TEXT("_");
-		KeyString += ShaderTypeDependency.ShaderType->GetName();
+		KeyString += ShaderType->GetName();
 		KeyString += ShaderTypeDependency.SourceHash.ToString();
-		ShaderTypeDependency.ShaderType->GetSerializationHistory().AppendKeyString(KeyString);
 
-		const TMap<const TCHAR*,FCachedUniformBufferDeclaration>& ReferencedUniformBufferStructsCache = ShaderTypeDependency.ShaderType->GetReferencedUniformBufferStructsCache();
+		const FSHAHash LayoutHash = Freeze::HashLayout(ShaderType->GetLayout(), LayoutParams);
+		KeyString += LayoutHash.ToString();
+
+		const TMap<const TCHAR*,FCachedUniformBufferDeclaration>& ReferencedUniformBufferStructsCache = ShaderType->GetReferencedUniformBufferStructsCache();
 
 		for (TMap<const TCHAR*,FCachedUniformBufferDeclaration>::TConstIterator It(ReferencedUniformBufferStructsCache); It; ++It)
 		{
@@ -746,11 +747,13 @@ void FMaterialShaderMapId::AppendKeyString(FString& KeyString) const
 	for (int32 TypeIndex = 0; TypeIndex < ShaderPipelineTypeDependencies.Num(); TypeIndex++)
 	{
 		const FShaderPipelineTypeDependency& Dependency = ShaderPipelineTypeDependencies[TypeIndex];
+		const FShaderPipelineType* ShaderPipelineType = FShaderPipelineType::GetShaderPipelineTypeByName(Dependency.ShaderPipelineTypeName);
+
 		KeyString += TEXT("_");
-		KeyString += Dependency.ShaderPipelineType->GetName();
+		KeyString += ShaderPipelineType->GetName();
 		KeyString += Dependency.StagesSourceHash.ToString();
 
-		for (const FShaderType* ShaderType : Dependency.ShaderPipelineType->GetStages())
+		for (const FShaderType* ShaderType : ShaderPipelineType->GetStages())
 		{
 			const TMap<const TCHAR*, FCachedUniformBufferDeclaration>& ReferencedUniformBufferStructsCache = ShaderType->GetReferencedUniformBufferStructsCache();
 
@@ -768,15 +771,22 @@ void FMaterialShaderMapId::AppendKeyString(FString& KeyString) const
 		KeyString += TEXT("_");
 
 		const FVertexFactoryTypeDependency& VFDependency = VertexFactoryTypeDependencies[VFIndex];
-		KeyString += VFDependency.VertexFactoryType->GetName();
+		const FVertexFactoryType* VertexFactoryType = FVertexFactoryType::GetVFByName(VFDependency.VertexFactoryTypeName);
+
+		KeyString += VertexFactoryType->GetName();
 		KeyString += VFDependency.VFSourceHash.ToString();
 
 		for (int32 Frequency = 0; Frequency < SF_NumFrequencies; Frequency++)
 		{
-			VFDependency.VertexFactoryType->GetSerializationHistory((EShaderFrequency)Frequency)->AppendKeyString(KeyString);
+			const FTypeLayoutDesc* ParameterLayout = VertexFactoryType->GetShaderParameterLayout((EShaderFrequency)Frequency);
+			if (ParameterLayout)
+			{
+				const FSHAHash LayoutHash = Freeze::HashLayout(*ParameterLayout, LayoutParams);
+				KeyString += LayoutHash.ToString();
+			}
 		}
 
-		const TMap<const TCHAR*,FCachedUniformBufferDeclaration>& ReferencedUniformBufferStructsCache = VFDependency.VertexFactoryType->GetReferencedUniformBufferStructsCache();
+		const TMap<const TCHAR*,FCachedUniformBufferDeclaration>& ReferencedUniformBufferStructsCache = VertexFactoryType->GetReferencedUniformBufferStructsCache();
 
 		for (TMap<const TCHAR*,FCachedUniformBufferDeclaration>::TConstIterator It(ReferencedUniformBufferStructsCache); It; ++It)
 		{
@@ -810,7 +820,7 @@ void FMaterialShaderMapId::SetShaderDependencies(const TArray<FShaderType*>& Sha
 		for (int32 ShaderTypeIndex = 0; ShaderTypeIndex < ShaderTypes.Num(); ShaderTypeIndex++)
 		{
 			FShaderTypeDependency Dependency;
-			Dependency.ShaderType = ShaderTypes[ShaderTypeIndex];
+			Dependency.ShaderTypeName = ShaderTypes[ShaderTypeIndex]->GetHashedName();
 			Dependency.SourceHash = ShaderTypes[ShaderTypeIndex]->GetSourceHash(ShaderPlatform);
 			ShaderTypeDependencies.Add(Dependency);
 		}
@@ -818,7 +828,7 @@ void FMaterialShaderMapId::SetShaderDependencies(const TArray<FShaderType*>& Sha
 		for (int32 VFTypeIndex = 0; VFTypeIndex < VFTypes.Num(); VFTypeIndex++)
 		{
 			FVertexFactoryTypeDependency Dependency;
-			Dependency.VertexFactoryType = VFTypes[VFTypeIndex];
+			Dependency.VertexFactoryTypeName = VFTypes[VFTypeIndex]->GetHashedName();
 			Dependency.VFSourceHash = VFTypes[VFTypeIndex]->GetSourceHash(ShaderPlatform);
 			VertexFactoryTypeDependencies.Add(Dependency);
 		}
@@ -827,7 +837,7 @@ void FMaterialShaderMapId::SetShaderDependencies(const TArray<FShaderType*>& Sha
 		{
 			const FShaderPipelineType* Pipeline = ShaderPipelineTypes[TypeIndex];
 			FShaderPipelineTypeDependency Dependency;
-			Dependency.ShaderPipelineType = Pipeline;
+			Dependency.ShaderPipelineTypeName = Pipeline->GetHashedName();
 			Dependency.StagesSourceHash = Pipeline->GetSourceHash(ShaderPlatform);
 			ShaderPipelineTypeDependencies.Add(Dependency);
 		}
@@ -891,7 +901,7 @@ void FMaterialShaderType::BeginCompileShaderPipeline(
 	const FMaterial* Material,
 	FShaderCompilerEnvironment* MaterialEnvironment,
 	const FShaderPipelineType* ShaderPipeline,
-	const TArray<FMaterialShaderType*>& ShaderStages,
+	const TArray<const FShaderType*>& ShaderStages,
 	TArray<TSharedRef<FShaderCommonCompileJob, ESPMode::ThreadSafe>>& NewJobs,
 	const FString& DebugDescription,
 	const FString& DebugExtension)
@@ -904,7 +914,7 @@ void FMaterialShaderType::BeginCompileShaderPipeline(
 	auto* NewPipelineJob = new FShaderPipelineCompileJob(ShaderMapId, ShaderPipeline, ShaderStages.Num());
 	for (int32 Index = 0; Index < ShaderStages.Num(); ++Index)
 	{
-		auto* ShaderStage = ShaderStages[Index];
+		FMaterialShaderType* ShaderStage = const_cast<FMaterialShaderType*>(static_cast<const FMaterialShaderType*>(ShaderStages[Index]));
 		ShaderStage->BeginCompileShader(ShaderMapId, kUniqueShaderPermutationId, Material, MaterialEnvironment, ShaderPipeline, Platform, NewPipelineJob->StageJobs, DebugDescription, DebugExtension);
 	}
 
@@ -921,14 +931,15 @@ FShader* FMaterialShaderType::FinishCompileShader(
 	const FSHAHash& MaterialShaderMapHash,
 	const FShaderCompileJob& CurrentJob,
 	const FShaderPipelineType* ShaderPipelineType,
-	const FString& InDebugDescription
+	const FString& InDebugDescription,
+	FShaderMapResourceBuilder& ResourceBuilder
 	)
 {
 	check(CurrentJob.bSucceeded);
 
 	// Reuse an existing resource with the same key or create a new one based on the compile output
 	// This allows FShaders to share compiled bytecode and RHI shader references
-	TRefCountPtr<FShaderResource> Resource = FShaderResource::FindOrCreate(CurrentJob.Output, 0);
+	const int32 ResourceIndex = ResourceBuilder.FindOrAddCode(CurrentJob.Output);
 
 	if (ShaderPipelineType && !ShaderPipelineType->ShouldOptimizeUnusedOutputs(CurrentJob.Input.Target.GetPlatform()))
 	{
@@ -936,17 +947,35 @@ FShader* FMaterialShaderType::FinishCompileShader(
 		ShaderPipelineType = nullptr;
 	}
 
-	// Find a shader with the same key in memory
-	FShader* Shader = CurrentJob.ShaderType->FindShaderByKey(FShaderKey(MaterialShaderMapHash, ShaderPipelineType, CurrentJob.VFType, CurrentJob.PermutationId, CurrentJob.Input.Target.GetPlatform()));
+	FShader* Shader = ConstructCompiled(CompiledShaderInitializerType(this, CurrentJob.PermutationId, CurrentJob.Output, ResourceIndex, UniformExpressionSet, MaterialShaderMapHash, ShaderPipelineType, nullptr, InDebugDescription));
+	CurrentJob.Output.ParameterMap.VerifyBindingsAreComplete(GetName(), CurrentJob.Output.Target, CurrentJob.VFType);
+
+	return Shader;
+}
+
+bool FMaterialShaderType::ShouldCompilePermutation(EShaderPlatform Platform, const FMaterialShaderParameters& MaterialParameters, int32 PermutationId) const
+{
+	return FShaderType::ShouldCompilePermutation(FMaterialShaderPermutationParameters(Platform, MaterialParameters, PermutationId));
+}
 
-	// There was no shader with the same key so create a new one with the compile output, which will bind shader parameters
-	if (!Shader)
+bool FMaterialShaderType::ShouldCompilePipeline(const FShaderPipelineType* ShaderPipelineType, EShaderPlatform Platform, const FMaterialShaderParameters& MaterialParameters)
+{
+	const FMaterialShaderPermutationParameters Parameters(Platform, MaterialParameters, kUniqueShaderPermutationId);
+	for (const FShaderType* ShaderType : ShaderPipelineType->GetStages())
 	{
-		Shader = (*ConstructCompiledRef)(CompiledShaderInitializerType(this, CurrentJob.PermutationId, CurrentJob.Output, MoveTemp(Resource), UniformExpressionSet, MaterialShaderMapHash, ShaderPipelineType, nullptr, InDebugDescription));
-		CurrentJob.Output.ParameterMap.VerifyBindingsAreComplete(GetName(), CurrentJob.Output.Target, CurrentJob.VFType);
+		checkSlow(ShaderType->GetMaterialShaderType());
+		if (!ShaderType->ShouldCompilePermutation(Parameters))
+		{
+			return false;
+		}
 	}
+	return true;
+}
 
-	return Shader;
+void FMaterialShaderType::SetupCompileEnvironment(EShaderPlatform Platform, const FMaterialShaderParameters& MaterialParameters, int32 PermutationId, FShaderCompilerEnvironment& Environment)
+{
+	// Allow the shader type to modify its compile environment.
+	ModifyCompilationEnvironment(FMaterialShaderPermutationParameters(Platform, MaterialParameters, PermutationId), Environment);
 }
 
 /**
@@ -966,7 +995,7 @@ TRefCountPtr<FMaterialShaderMap> FMaterialShaderMap::FindId(const FMaterialShade
 
 #if ALLOW_SHADERMAP_DEBUG_DATA
 /** Flushes the given shader types from any loaded FMaterialShaderMap's. */
-void FMaterialShaderMap::FlushShaderTypes(TArray<FShaderType*>& ShaderTypesToFlush, TArray<const FShaderPipelineType*>& ShaderPipelineTypesToFlush, TArray<const FVertexFactoryType*>& VFTypesToFlush)
+void FMaterialShaderMap::FlushShaderTypes(TArray<const FShaderType*>& ShaderTypesToFlush, TArray<const FShaderPipelineType*>& ShaderPipelineTypesToFlush, TArray<const FVertexFactoryType*>& VFTypesToFlush)
 {
 	for (int32 ShaderMapIndex = 0; ShaderMapIndex < AllMaterialShaderMaps.Num(); ShaderMapIndex++)
 	{
@@ -988,52 +1017,18 @@ void FMaterialShaderMap::FlushShaderTypes(TArray<FShaderType*>& ShaderTypesToFlu
 }
 #endif
 
-void FMaterialShaderMap::FixupShaderTypes(EShaderPlatform Platform, const TMap<FShaderType*, FString>& ShaderTypeNames, const TMap<const FShaderPipelineType*, FString>& ShaderPipelineTypeNames, const TMap<FVertexFactoryType*, FString>& VertexFactoryTypeNames)
-{
 #if WITH_EDITOR
-	FScopeLock ScopeLock(&GIdToMaterialShaderMapCS);
-	
-	TArray<FMaterialShaderMapId> Keys;
-	FMaterialShaderMap::GIdToMaterialShaderMap[Platform].GenerateKeyArray(Keys);
-
-	TArray<FMaterialShaderMap*> Values;
-	FMaterialShaderMap::GIdToMaterialShaderMap[Platform].GenerateValueArray(Values);
-
-	//@todo - what about the shader maps in AllMaterialShaderMaps that are not in GIdToMaterialShaderMap?
-	FMaterialShaderMap::GIdToMaterialShaderMap[Platform].Empty();
-
-	for (int32 PairIndex = 0; PairIndex < Keys.Num(); PairIndex++)
+void FMaterialShaderMap::GetAllOutdatedTypes(TArray<const FShaderType*>& OutdatedShaderTypes, TArray<const FShaderPipelineType*>& OutdatedShaderPipelineTypes, TArray<const FVertexFactoryType*>& OutdatedFactoryTypes)
+{
+#if ALLOW_SHADERMAP_DEBUG_DATA
+	for (const FMaterialShaderMap* ShaderMap : AllMaterialShaderMaps)
 	{
-		FMaterialShaderMapId& Key = Keys[PairIndex];
-
-		for (int32 ShaderTypeIndex = 0; ShaderTypeIndex < Key.ShaderTypeDependencies.Num(); ShaderTypeIndex++)
-		{
-			const FString& ShaderTypeName = ShaderTypeNames.FindChecked(Key.ShaderTypeDependencies[ShaderTypeIndex].ShaderType);
-			FShaderType* FoundShaderType = FShaderType::GetShaderTypeByName(*ShaderTypeName);
-			Key.ShaderTypeDependencies[ShaderTypeIndex].ShaderType = FoundShaderType;
-		}
-
-		for (int32 ShaderPipelineIndex = 0; ShaderPipelineIndex < Key.ShaderPipelineTypeDependencies.Num(); ShaderPipelineIndex++)
-		{
-			const FString& ShaderPipelineTypeName = ShaderPipelineTypeNames.FindChecked(Key.ShaderPipelineTypeDependencies[ShaderPipelineIndex].ShaderPipelineType);
-			const FShaderPipelineType* FoundShaderPipelineType = FShaderPipelineType::GetShaderPipelineTypeByName(*ShaderPipelineTypeName);
-			Key.ShaderPipelineTypeDependencies[ShaderPipelineIndex].ShaderPipelineType = FoundShaderPipelineType;
-		}
-
-		for (int32 VFTypeIndex = 0; VFTypeIndex < Key.VertexFactoryTypeDependencies.Num(); VFTypeIndex++)
-		{
-			const FString& VFTypeName = VertexFactoryTypeNames.FindChecked(Key.VertexFactoryTypeDependencies[VFTypeIndex].VertexFactoryType);
-			FVertexFactoryType* FoundVFType = FVertexFactoryType::GetVFByName(VFTypeName);
-			Key.VertexFactoryTypeDependencies[VFTypeIndex].VertexFactoryType = FoundVFType;
-		}
-
-		FMaterialShaderMap::GIdToMaterialShaderMap[Platform].Add(Key, Values[PairIndex]);
+		ShaderMap->GetOutdatedTypes(OutdatedShaderTypes, OutdatedShaderPipelineTypes, OutdatedFactoryTypes);
 	}
-#endif // WITH_EDITOR
+#endif
 }
 
-#if WITH_EDITOR
-void FMaterialShaderMap::LoadFromDerivedDataCache(const FMaterial* Material, const FMaterialShaderMapId& ShaderMapId, EShaderPlatform InPlatform, TRefCountPtr<FMaterialShaderMap>& InOutShaderMap)
+void FMaterialShaderMap::LoadFromDerivedDataCache(const FMaterial* Material, const FMaterialShaderMapId& ShaderMapId, EShaderPlatform InPlatform, const ITargetPlatform* TargetPlatform, TRefCountPtr<FMaterialShaderMap>& InOutShaderMap)
 {
 	if (InOutShaderMap != NULL)
 	{
@@ -1050,19 +1045,19 @@ void FMaterialShaderMap::LoadFromDerivedDataCache(const FMaterial* Material, con
 			COOK_STAT(auto Timer = MaterialShaderCookStats::UsageStats.TimeSyncWork());
 
 			TArray<uint8> CachedData;
-			const FString DataKey = GetMaterialShaderMapKeyString(ShaderMapId, InPlatform);
+			const FString DataKey = GetMaterialShaderMapKeyString(ShaderMapId, InPlatform, TargetPlatform);
 
 			if (GetDerivedDataCacheRef().GetSynchronous(*DataKey, CachedData, Material->GetFriendlyName()))
 			{
 				COOK_STAT(Timer.AddHit(CachedData.Num()));
-				InOutShaderMap = new FMaterialShaderMap(InPlatform);
+				InOutShaderMap = new FMaterialShaderMap();
 				FMemoryReader Ar(CachedData, true);
 
 				// Deserialize from the cached data
 				InOutShaderMap->Serialize(Ar);
-				InOutShaderMap->RegisterSerializedShaders(false);
+				//InOutShaderMap->RegisterSerializedShaders(false);
 
-				const FString InDataKey = GetMaterialShaderMapKeyString(InOutShaderMap->GetShaderMapId(), InPlatform);
+				const FString InDataKey = GetMaterialShaderMapKeyString(InOutShaderMap->GetShaderMapId(), InPlatform, TargetPlatform);
 				checkSlow(InOutShaderMap->GetShaderMapId() == ShaderMapId);
 
 				// Register in the global map
@@ -1079,14 +1074,14 @@ void FMaterialShaderMap::LoadFromDerivedDataCache(const FMaterial* Material, con
 	}
 }
 
-void FMaterialShaderMap::SaveToDerivedDataCache()
+void FMaterialShaderMap::SaveToDerivedDataCache(const ITargetPlatform* TargetPlatform)
 {
 	COOK_STAT(auto Timer = MaterialShaderCookStats::UsageStats.TimeSyncWork());
 	TArray<uint8> SaveData;
 	FMemoryWriter Ar(SaveData, true);
 	Serialize(Ar);
 
-	GetDerivedDataCacheRef().Put(*GetMaterialShaderMapKeyString(ShaderMapId, GetShaderPlatform()), SaveData, FriendlyName);
+	GetDerivedDataCacheRef().Put(*GetMaterialShaderMapKeyString(ShaderMapId, GetShaderPlatform(), TargetPlatform), SaveData, FStringView(GetFriendlyName()));
 	COOK_STAT(Timer.AddMiss(SaveData.Num()));
 }
 #endif // WITH_EDITOR
@@ -1094,42 +1089,54 @@ void FMaterialShaderMap::SaveToDerivedDataCache()
 TArray<uint8>* FMaterialShaderMap::BackupShadersToMemory()
 {
 	TArray<uint8>* SavedShaderData = new TArray<uint8>();
+#if 0
 	FMemoryWriter Ar(*SavedShaderData);
 	
-	for (int32 Index = 0; Index < MeshShaderMaps.Num(); Index++)
+	for (FMeshMaterialShaderMap* MeshShaderMap : Content->OrderedMeshShaderMaps)
 	{
-		// Serialize data needed to handle shader key changes in between the save and the load of the FShaders
-		const bool bHandleShaderKeyChanges = true;
-		MeshShaderMaps[Index].SerializeInline(Ar, true, bHandleShaderKeyChanges, false);
-		MeshShaderMaps[Index].Empty();
+		if (MeshShaderMap)
+		{
+			// Serialize data needed to handle shader key changes in between the save and the load of the FShaders
+			const bool bHandleShaderKeyChanges = true;
+			MeshShaderMap->SerializeInline(Ar, true, bHandleShaderKeyChanges, false);
+			MeshShaderMap->Empty();
+		}
 	}
 
-	SerializeInline(Ar, true, true, false);
-	Empty();
-
+	Content->SerializeInline(Ar, true, true, false);
+	Content->Empty();
+#endif
+	check(false);
 	return SavedShaderData;
 }
 
 void FMaterialShaderMap::RestoreShadersFromMemory(const TArray<uint8>& ShaderData)
 {
+	check(false);
+#if 0
 	FMemoryReader Ar(ShaderData);
 
-	for (int32 Index = 0; Index < MeshShaderMaps.Num(); Index++)
+	for(FMeshMaterialShaderMap* MeshShaderMap : Content->OrderedMeshShaderMaps)
 	{
-		// Use the serialized shader key data to detect when the saved shader is no longer valid and skip it
-		const bool bHandleShaderKeyChanges = true;
-		MeshShaderMaps[Index].SerializeInline(Ar, true, bHandleShaderKeyChanges, false);
-		MeshShaderMaps[Index].RegisterSerializedShaders(false);
+		if (MeshShaderMap)
+		{
+			// Use the serialized shader key data to detect when the saved shader is no longer valid and skip it
+			const bool bHandleShaderKeyChanges = true;
+			MeshShaderMap->SerializeInline(Ar, true, bHandleShaderKeyChanges, false);
+			//MeshShaderMaps[Index].RegisterSerializedShaders(false);
+		}
 	}
 
-	SerializeInline(Ar, true, true, false);
-	RegisterSerializedShaders(false);
+	Content->SerializeInline(Ar, true, true, false);
+	//RegisterSerializedShaders(false);
+#endif
 }
 
 void FMaterialShaderMap::SaveForRemoteRecompile(FArchive& Ar, const TMap<FString, TArray<TRefCountPtr<FMaterialShaderMap> > >& CompiledShaderMaps, const TArray<FShaderResourceId>& ClientResourceIds)
 {
 	UE_LOG(LogMaterial, Display, TEXT("Looking for unique resources, %d were on client"), ClientResourceIds.Num());	
-
+	check(false);
+#if 0
 	// first, we look for the unique shader resources
 	TArray<TRefCountPtr<FShaderResource>> UniqueResources;
 	int32 NumSkippedResources = 0;
@@ -1145,33 +1152,20 @@ void FMaterialShaderMap::SaveForRemoteRecompile(FArchive& Ar, const TMap<FString
 			if (ShaderMap)
 			{
 				// get all shaders in the shader map
-				TMap<FShaderId, FShader*> ShaderList;
+				TMap<FShaderId, TShaderRef<FShader>> ShaderList;
 				ShaderMap->GetShaderList(ShaderList);
 
-				// get shaders from shader pipelines
-				TArray<FShaderPipeline*> ShaderPipelineList;
-				ShaderMap->GetShaderPipelineList(ShaderPipelineList);
-
-				for (FShaderPipeline* ShaderPipeline : ShaderPipelineList)
-				{
-					for (FShader* Shader : ShaderPipeline->GetShaders())
-					{
-						FShaderId ShaderId = Shader->GetId();
-						ShaderList.Add(ShaderId, Shader);
-					}
-				}
-
 				// get the resources from the shaders
 				for (auto& KeyValue : ShaderList)
 				{
-					FShader* Shader = KeyValue.Value;
-					FShaderResourceId ShaderId = Shader->GetResourceId();
+					const TShaderRef<FShader>& Shader = KeyValue.Value;
+					FShaderResourceId ShaderId = Shader.GetResourceId();
 
 					// skip this shader if the Id was already on the client (ie, it didn't change)
 					if (ClientResourceIds.Contains(ShaderId) == false)
 					{
 						// lookup the resource by ID and add it if it's unique
-						UniqueResources.AddUnique(FShaderResource::FindById(ShaderId));
+						UniqueResources.AddUnique(Shader.GetResourceChecked());
 					}
 					else
 					{
@@ -1190,7 +1184,7 @@ void FMaterialShaderMap::SaveForRemoteRecompile(FArchive& Ar, const TMap<FString
 
 	for (int32 Index = 0; Index < NumUniqueResources; Index++)
 	{
-		UniqueResources[Index]->Serialize(Ar, false);
+		UniqueResources[Index]->SaveCodeToArchive(Ar);
 	}
 
 	// now we serialize a map (for each material), but without inline the resources, since they are above
@@ -1224,11 +1218,13 @@ void FMaterialShaderMap::SaveForRemoteRecompile(FArchive& Ar, const TMap<FString
 			}
 		}
 	}
+#endif
 }
 
 void FMaterialShaderMap::LoadForRemoteRecompile(FArchive& Ar, EShaderPlatform ShaderPlatform, const TArray<FString>& MaterialsForShaderMaps)
 {
-#if WITH_EDITOR
+	check(false);
+#if 0// WITH_EDITOR
 	int32 NumResources;
 	Ar << NumResources;
 
@@ -1240,10 +1236,7 @@ void FMaterialShaderMap::LoadForRemoteRecompile(FArchive& Ar, EShaderPlatform Sh
 	for (int32 Index = 0; Index < NumResources; Index++)
 	{
 		// Load the inlined shader resource
-		FShaderResource ResourceTemp;
-		ResourceTemp.Serialize(Ar, false);
-
-		KeepAliveReferences[Index] = FShaderResource::FindOrClone(MoveTemp(ResourceTemp));
+		KeepAliveReferences[Index] = FShaderResource_InlineCode::FindOrLoad(Ar);
 	}
 
 	int32 MapSize;
@@ -1268,7 +1261,7 @@ void FMaterialShaderMap::LoadForRemoteRecompile(FArchive& Ar, EShaderPlatform Sh
 
 			if (bIsValid)
 			{
-				FMaterialShaderMap* ShaderMap = new FMaterialShaderMap(ShaderPlatform);
+				FMaterialShaderMap* ShaderMap = new FMaterialShaderMap();
 
 				// serialize the id and the material shader map
 				ShaderMap->Serialize(Ar, false);
@@ -1313,7 +1306,47 @@ void FMaterialShaderMap::LoadForRemoteRecompile(FArchive& Ar, EShaderPlatform Sh
 #endif // WITH_EDITOR
 }
 
+void FMaterialShaderMap::FinalizeContent()
+{
+	FMaterialShaderMapContent* LocalContent = GetMutableContent();
+
+	FSHA1 Hasher;
+	for (const FShader* Shader : LocalContent->GetShaders())
+	{
+		const FHashedName& TypeName = Shader->GetTypeName();
+		const int32 PermutationId = Shader->GetPermutationId();
+		Hasher.Update((uint8*)&TypeName, sizeof(TypeName));
+		Hasher.Update((uint8*)&PermutationId, sizeof(PermutationId));
+	}
+
+	for (const FShaderPipeline* Pipeline : LocalContent->GetShaderPipelines())
+	{
+		const FHashedName& TypeName = Pipeline->TypeName;
+		Hasher.Update((uint8*)&TypeName, sizeof(TypeName));
+	}
+
+	for (FMeshMaterialShaderMap* MeshShaderMap : LocalContent->OrderedMeshShaderMaps)
+	{
+		for (const FShader* Shader : MeshShaderMap->GetShaders())
+		{
+			const FHashedName& TypeName = Shader->GetTypeName();
+			const int32 PermutationId = Shader->GetPermutationId();
+			Hasher.Update((uint8*)&TypeName, sizeof(TypeName));
+			Hasher.Update((uint8*)&PermutationId, sizeof(PermutationId));
+		}
+
+		for (const FShaderPipeline* Pipeline : MeshShaderMap->GetShaderPipelines())
+		{
+			const FHashedName& TypeName = Pipeline->TypeName;
+			Hasher.Update((uint8*)&TypeName, sizeof(TypeName));
+		}
+	}
 
+	Hasher.Final();
+	Hasher.GetHash(LocalContent->ShaderContentHash.Hash);
+
+	FShaderMapBase::FinalizeContent();
+}
 
 
 /**
@@ -1361,58 +1394,68 @@ void FMaterialShaderMap::Compile(
 #if DEBUG_INFINITESHADERCOMPILE
 			UE_LOG(LogTemp, Display, TEXT("Added material ShaderMap 0x%08X%08X with Material 0x%08X%08X to ShaderMapsBeingCompiled"), (int)((int64)(this) >> 32), (int)((int64)(this)), (int)((int64)(Material) >> 32), (int)((int64)(Material)));
 #endif  
-			// Setup the material compilation environment.
-			Material->SetupMaterialEnvironment(InPlatform, InMaterialCompilationOutput.UniformExpressionSet, *MaterialEnvironment);
   
-			MaterialCompilationOutput = InMaterialCompilationOutput;
+			FMaterialShaderMapContent* NewContent = new FMaterialShaderMapContent(InPlatform);
+			NewContent->MaterialCompilationOutput = InMaterialCompilationOutput;
+			AssignContent(NewContent);
+
 			ShaderMapId = InShaderMapId;
-			Platform = InPlatform;
 			bIsPersistent = Material->IsPersistent();
 
+			// Setup the material compilation environment.
+			{
+				FShaderParametersMetadata* UniformBufferStruct = NewContent->MaterialCompilationOutput.UniformExpressionSet.CreateBufferStruct();
+				Material->SetupMaterialEnvironment(InPlatform, *UniformBufferStruct, InMaterialCompilationOutput.UniformExpressionSet, *MaterialEnvironment);
+				delete UniformBufferStruct;
+			}
+ 
 #if ALLOW_SHADERMAP_DEBUG_DATA && WITH_EDITOR
 			// Store the material name for debugging purposes.
 			// Note: Material instances with static parameters will have the same FriendlyName for their shader maps!
-			FriendlyName = Material->GetFriendlyName();
+			NewContent->FriendlyName = Material->GetFriendlyName();
+			NewContent->MaterialPath = Material->GetBaseMaterialPathName();
+
+			GetMutableContent()->MaterialPath = FMemoryImageString(Material->GetBaseMaterialPathName());
 
 			// Log debug information about the material being compiled.
 			const FString MaterialUsage = Material->GetMaterialUsageDescription();
-			DebugDescription = FString::Printf(
+			FString WorkingDebugDescription = FString::Printf(
 				TEXT("Compiling %s: Platform=%s, Usage=%s"),
-				*FriendlyName,
+				*NewContent->FriendlyName,
 				*LegacyShaderPlatformToShaderFormat(InPlatform).ToString(),
 				*MaterialUsage
-				);
-			for(int32 StaticSwitchIndex = 0;StaticSwitchIndex < ShaderMapId.GetStaticSwitchParameters().Num();++StaticSwitchIndex)
+			);
+			for (int32 StaticSwitchIndex = 0; StaticSwitchIndex < ShaderMapId.GetStaticSwitchParameters().Num(); ++StaticSwitchIndex)
 			{
 				const FStaticSwitchParameter& StaticSwitchParameter = ShaderMapId.GetStaticSwitchParameters()[StaticSwitchIndex];
-				DebugDescription += FString::Printf(
+				WorkingDebugDescription += FString::Printf(
 					TEXT(", StaticSwitch'%s'=%s"),
 					*StaticSwitchParameter.ParameterInfo.ToString(),
 					StaticSwitchParameter.Value ? TEXT("True") : TEXT("False")
-					);
+				);
 			}
-			for(int32 StaticMaskIndex = 0;StaticMaskIndex < ShaderMapId.GetStaticComponentMaskParameters().Num();++StaticMaskIndex)
+			for (int32 StaticMaskIndex = 0; StaticMaskIndex < ShaderMapId.GetStaticComponentMaskParameters().Num(); ++StaticMaskIndex)
 			{
 				const FStaticComponentMaskParameter& StaticComponentMaskParameter = ShaderMapId.GetStaticComponentMaskParameters()[StaticMaskIndex];
-				DebugDescription += FString::Printf(
+				WorkingDebugDescription += FString::Printf(
 					TEXT(", StaticMask'%s'=%s%s%s%s"),
 					*StaticComponentMaskParameter.ParameterInfo.ToString(),
 					StaticComponentMaskParameter.R ? TEXT("R") : TEXT(""),
 					StaticComponentMaskParameter.G ? TEXT("G") : TEXT(""),
 					StaticComponentMaskParameter.B ? TEXT("B") : TEXT(""),
 					StaticComponentMaskParameter.A ? TEXT("A") : TEXT("")
-					);
+				);
 			}
-			for(int32 StaticLayerIndex = 0;StaticLayerIndex < ShaderMapId.GetTerrainLayerWeightParameters().Num();++StaticLayerIndex)
+			for (int32 StaticLayerIndex = 0; StaticLayerIndex < ShaderMapId.GetTerrainLayerWeightParameters().Num(); ++StaticLayerIndex)
 			{
 				const FStaticTerrainLayerWeightParameter& StaticTerrainLayerWeightParameter = ShaderMapId.GetTerrainLayerWeightParameters()[StaticLayerIndex];
-				DebugDescription += FString::Printf(
+				WorkingDebugDescription += FString::Printf(
 					TEXT(", StaticTerrainLayer'%s'=%s"),
 					*StaticTerrainLayerWeightParameter.ParameterInfo.ToString(),
-					*FString::Printf(TEXT("Weightmap%u"),StaticTerrainLayerWeightParameter.WeightmapIndex)
-					);
+					*FString::Printf(TEXT("Weightmap%u"), StaticTerrainLayerWeightParameter.WeightmapIndex)
+				);
 			}
-			for (const auto &LayerParameterID :ShaderMapId.GetMaterialLayersParameterIDs())
+			for (const auto &LayerParameterID : ShaderMapId.GetMaterialLayersParameterIDs())
 			{
 				FString UUIDs;
 				UUIDs += TEXT("Layers:");
@@ -1437,162 +1480,125 @@ void FMaterialShaderMap::Compile(
 					UUIDs += (State ? TEXT("1") : TEXT("0"));
 					StartWithComma = true;
 				}
-				DebugDescription += FString::Printf(TEXT(", LayersParameter'%s'=[%s]"), *LayerParameterID.ParameterID.ParameterInfo.ToString(), *UUIDs);
+				WorkingDebugDescription += FString::Printf(TEXT(", LayersParameter'%s'=[%s]"), *LayerParameterID.ParameterID.ParameterInfo.ToString(), *UUIDs);
 			}
   
-			UE_LOG(LogShaders, Display, TEXT("	%s"), *DebugDescription);
+			UE_LOG(LogShaders, Display, TEXT("	%s"), *WorkingDebugDescription);
+			NewContent->DebugDescription = *WorkingDebugDescription;
 
 			FString DebugExtension = FString::Printf( TEXT("_%08x%08x"), ShaderMapId.BaseMaterialId.A, ShaderMapId.BaseMaterialId.B);
-#elif ALLOW_SHADERMAP_DEBUG_DATA && !WITH_EDITOR
-			FString DebugExtension = "";
-			DebugDescription = "";
 #else
 			FString DebugExtension = "";
-			FString DebugDescription = "";
-#endif // ALLOW_SHADERMAP_DEBUG_DATA && WITH_EDITOR
-
+			FString WorkingDebugDescription = "";
+#endif 
 			uint32 NumShaders = 0;
 			uint32 NumVertexFactories = 0;
 			TArray<TSharedRef<FShaderCommonCompileJob, ESPMode::ThreadSafe>> NewJobs;
   
+			const FMaterialShaderParameters MaterialParameters(Material);
+			const FMaterialShaderMapLayout& Layout = AcquireMaterialShaderMapLayout(InPlatform, MaterialParameters);
+
 			// Iterate over all vertex factory types.
-			for(TLinkedList<FVertexFactoryType*>::TIterator VertexFactoryTypeIt(FVertexFactoryType::GetTypeList());VertexFactoryTypeIt;VertexFactoryTypeIt.Next())
+			for (const FMeshMaterialShaderMapLayout& MeshLayout : Layout.MeshShaderMaps)
 			{
-				FVertexFactoryType* VertexFactoryType = *VertexFactoryTypeIt;
-				check(VertexFactoryType);
-  
-				if(VertexFactoryType->IsUsedWithMaterials())
+				FMeshMaterialShaderMap* MeshShaderMap = GetMeshShaderMap(MeshLayout.VertexFactoryType);
+				if (MeshShaderMap == nullptr)
 				{
-					FMeshMaterialShaderMap* MeshShaderMap = nullptr;
+					// Create a new mesh material shader map.
+					MeshShaderMap = new FMeshMaterialShaderMap(InPlatform, MeshLayout.VertexFactoryType);
+					NewContent->AddMeshShaderMap(MeshLayout.VertexFactoryType, MeshShaderMap);
+				}
   
-					// look for existing map for this vertex factory type
-					int32 MeshShaderMapIndex = INDEX_NONE;
-					for (int32 ShaderMapIndex = 0; ShaderMapIndex < MeshShaderMaps.Num(); ShaderMapIndex++)
-					{
-						if (MeshShaderMaps[ShaderMapIndex].GetVertexFactoryType() == VertexFactoryType)
-						{
-							MeshShaderMap = &MeshShaderMaps[ShaderMapIndex];
-							MeshShaderMapIndex = ShaderMapIndex;
-							break;
-						}
-					}
+				// Enqueue compilation all mesh material shaders for this material and vertex factory type combo.
+				const uint32 MeshShaders = MeshShaderMap->BeginCompile(
+					CompilingId,
+					InShaderMapId,
+					Material,
+					MeshLayout,
+					MaterialEnvironment,
+					InPlatform,
+					NewJobs,
+					WorkingDebugDescription,
+					DebugExtension
+				);
+				NumShaders += MeshShaders;
+				if (MeshShaders > 0)
+				{
+					NumVertexFactories++;
+				}
+			}
   
-					if (MeshShaderMap == nullptr)
-					{
-						// Create a new mesh material shader map.
-						MeshShaderMapIndex = MeshShaderMaps.Num();
-						MeshShaderMap = new FMeshMaterialShaderMap(InPlatform, VertexFactoryType);
-						MeshShaderMaps.Add(MeshShaderMap);
-					}
+			// Iterate over all material shader types.
+			TMap<TShaderTypePermutation<const FShaderType>, FShaderCompileJob*> SharedShaderJobs;
+			for (const FShaderLayoutEntry& Shader : Layout.Shaders)
+			{
+				FMaterialShaderType* ShaderType = static_cast<FMaterialShaderType*>(Shader.ShaderType);
+				if (!Material->ShouldCache(InPlatform, ShaderType, nullptr))
+				{
+					continue;
+				}
+
+#if WITH_EDITOR
+				// Verify that the shader map Id contains inputs for any shaders that will be put into this shader map
+				check(InShaderMapId.ContainsShaderType(ShaderType, kUniqueShaderPermutationId));
+#endif
+				// Compile this material shader for this material.
+				TArray<FString> ShaderErrors;
   
-					// Enqueue compilation all mesh material shaders for this material and vertex factory type combo.
-					const uint32 MeshShaders = MeshShaderMap->BeginCompile(
+				// Only compile the shader if we don't already have it
+				if (!NewContent->HasShader(ShaderType, Shader.PermutationId))
+				{
+					FShaderCompileJob* Job = ShaderType->BeginCompileShader(
 						CompilingId,
-						InShaderMapId,
+						Shader.PermutationId,
 						Material,
 						MaterialEnvironment,
+						nullptr,
 						InPlatform,
 						NewJobs,
-						DebugDescription, 
+						WorkingDebugDescription,
 						DebugExtension
-						);
-					NumShaders += MeshShaders;
-					if (MeshShaders > 0)
-					{
-						NumVertexFactories++;
-					}
+					);
+
+					TShaderTypePermutation<const FShaderType> ShaderTypePermutation(ShaderType, Shader.PermutationId);
+					check(!SharedShaderJobs.Find(ShaderTypePermutation));
+					SharedShaderJobs.Add(ShaderTypePermutation, Job);
 				}
+				NumShaders++;
 			}
   
-			// Iterate over all material shader types.
-			TMap<TShaderTypePermutation<const FShaderType>, FShaderCompileJob*> SharedShaderJobs;
-			for(TLinkedList<FShaderType*>::TIterator ShaderTypeIt(FShaderType::GetTypeList());ShaderTypeIt;ShaderTypeIt.Next())
+			if (RHISupportsShaderPipelines(InPlatform))
 			{
-				FMaterialShaderType* ShaderType = ShaderTypeIt->GetMaterialShaderType();
-				const int32 PermutationCount = ShaderType ? ShaderType->GetPermutationCount() : 0;
-				for (int32 PermutationId = 0; PermutationId < PermutationCount; ++PermutationId)
+				for (FShaderPipelineType* Pipeline : Layout.ShaderPipelines)
 				{
-					if (ShouldCacheMaterialShader(ShaderType, InPlatform, Material, PermutationId))
+					if (!Material->ShouldCachePipeline(InPlatform, Pipeline, nullptr))
 					{
-#if WITH_EDITOR
-						// Verify that the shader map Id contains inputs for any shaders that will be put into this shader map
-						check(InShaderMapId.ContainsShaderType(ShaderType, kUniqueShaderPermutationId));
-#endif
-						// Compile this material shader for this material.
-						TArray<FString> ShaderErrors;
-
-						// Only compile the shader if we don't already have it
-						if (!HasShader(ShaderType, PermutationId))
-						{
-							FShaderCompileJob* Job = ShaderType->BeginCompileShader(
-								CompilingId,
-								PermutationId,
-								Material,
-								MaterialEnvironment,
-								nullptr,
-								InPlatform,
-								NewJobs,
-								DebugDescription, DebugExtension
-							);
-
-							TShaderTypePermutation<const FShaderType> ShaderTypePermutation(ShaderType, PermutationId);
-							check(!SharedShaderJobs.Find(ShaderTypePermutation));
-							SharedShaderJobs.Add(ShaderTypePermutation, Job);
-						}
-						NumShaders++;
+						continue;
 					}
-				}
-			}
   
-			const bool bHasTessellation = Material->GetTessellationMode() != MTM_NoTessellation;
-			for (TLinkedList<FShaderPipelineType*>::TIterator ShaderPipelineIt(FShaderPipelineType::GetTypeList());ShaderPipelineIt;ShaderPipelineIt.Next())
-			{
-				const FShaderPipelineType* Pipeline = *ShaderPipelineIt;
-				if (Pipeline->IsMaterialTypePipeline() && Pipeline->HasTessellation() == bHasTessellation && RHISupportsShaderPipelines(InPlatform))
-				{
 					auto& StageTypes = Pipeline->GetStages();
-					TArray<FMaterialShaderType*> ShaderStagesToCompile;
-					for (int32 Index = 0; Index < StageTypes.Num(); ++Index)
-					{
-						FMaterialShaderType* ShaderType = (FMaterialShaderType*)(StageTypes[Index]->GetMaterialShaderType());
-						if (ShaderType && ShouldCacheMaterialShader(ShaderType, InPlatform, Material, kUniqueShaderPermutationId))
-						{
+
+					// Verify that the shader map Id contains inputs for any shaders that will be put into this shader map
 #if WITH_EDITOR
-							// Verify that the shader map Id contains inputs for any shaders that will be put into this shader map
-							check(InShaderMapId.ContainsShaderType(ShaderType, kUniqueShaderPermutationId));
+					check(InShaderMapId.ContainsShaderPipelineType(Pipeline));
 #endif
-							ShaderStagesToCompile.Add(ShaderType);
-						}
-						else
-						{
-							break;
-						}
+					if (Pipeline->ShouldOptimizeUnusedOutputs(InPlatform))
+					{
+						NumShaders += StageTypes.Num();
+						FMaterialShaderType::BeginCompileShaderPipeline(CompilingId, InPlatform, Material, MaterialEnvironment, Pipeline, StageTypes, NewJobs, WorkingDebugDescription, DebugExtension);
 					}
-
-					if (ShaderStagesToCompile.Num() == StageTypes.Num())
+					else
 					{
-						// Verify that the shader map Id contains inputs for any shaders that will be put into this shader map
-#if WITH_EDITOR
-						check(InShaderMapId.ContainsShaderPipelineType(Pipeline));
-#endif
-						if (Pipeline->ShouldOptimizeUnusedOutputs(InPlatform))
+						// If sharing shaders amongst pipelines, add this pipeline as a dependency of an existing job
+						for (const FShaderType* ShaderType : StageTypes)
 						{
-							NumShaders += ShaderStagesToCompile.Num();
-							FMaterialShaderType::BeginCompileShaderPipeline(CompilingId, InPlatform, Material, MaterialEnvironment, Pipeline, ShaderStagesToCompile, NewJobs, DebugDescription, DebugExtension);
-						}
-						else
-						{
-							// If sharing shaders amongst pipelines, add this pipeline as a dependency of an existing job
-							for (const FShaderType* ShaderType : StageTypes)
-							{
-								TShaderTypePermutation<const FShaderType> ShaderTypePermutation(ShaderType, kUniqueShaderPermutationId);
-								FShaderCompileJob** Job = SharedShaderJobs.Find(ShaderTypePermutation);
-								checkf(Job, TEXT("Couldn't find existing shared job for material shader %s on pipeline %s!"), ShaderType->GetName(), Pipeline->GetName());
-								auto* SingleJob = (*Job)->GetSingleShaderJob();
-								auto& PipelinesToShare = SingleJob->SharingPipelines.FindOrAdd(nullptr);
-								check(!PipelinesToShare.Contains(Pipeline));
-								PipelinesToShare.Add(Pipeline);
-							}
+							TShaderTypePermutation<const FShaderType> ShaderTypePermutation(ShaderType, kUniqueShaderPermutationId);
+							FShaderCompileJob** Job = SharedShaderJobs.Find(ShaderTypePermutation);
+							checkf(Job, TEXT("Couldn't find existing shared job for material shader %s on pipeline %s!"), ShaderType->GetName(), Pipeline->GetName());
+							auto* SingleJob = (*Job)->GetSingleShaderJob();
+							auto& PipelinesToShare = SingleJob->SharingPipelines.FindOrAdd(nullptr);
+							check(!PipelinesToShare.Contains(Pipeline));
+							PipelinesToShare.Add(Pipeline);
 						}
 					}
 				}
@@ -1602,6 +1608,8 @@ void FMaterialShaderMap::Compile(
 
 			// Register this shader map in the global map with the material's ID.
 			Register(InPlatform);
+
+			CompilingTargetPlatform = MaterialEnvironment->TargetPlatform;
   
 			// Mark the shader map as not having been finalized with ProcessCompilationResults
 			bCompilationFinalized = false;
@@ -1621,18 +1629,14 @@ void FMaterialShaderMap::Compile(
 				TArray<int32> CurrentShaderMapId;
 				CurrentShaderMapId.Add(CompilingId);
 				GShaderCompilingManager->FinishCompilation(
-#if ALLOW_SHADERMAP_DEBUG_DATA
-				   *FriendlyName,
-#else
-				   nullptr,
-#endif
-				   CurrentShaderMapId);
+					GetFriendlyName(),
+					CurrentShaderMapId);
 			}
 		}
 	}
 }
 
-FShader* FMaterialShaderMap::ProcessCompilationResultsForSingleJob(FShaderCompileJob* SingleJob, const FShaderPipelineType* ShaderPipeline, const FSHAHash& MaterialShaderMapHash)
+FShader* FMaterialShaderMap::ProcessCompilationResultsForSingleJob(FShaderCompileJob* SingleJob, const FShaderPipelineType* ShaderPipeline, const FSHAHash& MaterialShaderMapHash, FShaderMapResourceBuilder& ResourceBuilder)
 {
 	check(SingleJob);
 	const FShaderCompileJob& CurrentJob = *SingleJob;
@@ -1646,58 +1650,36 @@ FShader* FMaterialShaderMap::ProcessCompilationResultsForSingleJob(FShaderCompil
 	{
 		FVertexFactoryType* VertexFactoryType = CurrentJob.VFType;
 		check(VertexFactoryType->IsUsedWithMaterials());
-		FMeshMaterialShaderMap* MeshShaderMap = nullptr;
-		int32 MeshShaderMapIndex = INDEX_NONE;
-
-		// look for existing map for this vertex factory type
-		for (int32 ShaderMapIndex = 0; ShaderMapIndex < MeshShaderMaps.Num(); ShaderMapIndex++)
-		{
-			if (MeshShaderMaps[ShaderMapIndex].GetVertexFactoryType() == VertexFactoryType)
-			{
-				MeshShaderMap = &MeshShaderMaps[ShaderMapIndex];
-				MeshShaderMapIndex = ShaderMapIndex;
-				break;
-			}
-		}
+		FMeshMaterialShaderMap* MeshShaderMap = GetMeshShaderMap(VertexFactoryType);
 
 		check(MeshShaderMap);
 		FMeshMaterialShaderType* MeshMaterialShaderType = CurrentJob.ShaderType->GetMeshMaterialShaderType();
 		check(MeshMaterialShaderType);
-		Shader = MeshMaterialShaderType->FinishCompileShader(MaterialCompilationOutput.UniformExpressionSet, MaterialShaderMapHash, CurrentJob, ShaderPipeline,
-#if ALLOW_SHADERMAP_DEBUG_DATA
-			FriendlyName);
-#else
-			FString());
-#endif
+		Shader = MeshMaterialShaderType->FinishCompileShader(GetContent()->MaterialCompilationOutput.UniformExpressionSet, MaterialShaderMapHash, CurrentJob, ShaderPipeline, GetFriendlyName(), ResourceBuilder);
 		check(Shader);
 		if (!ShaderPipeline)
 		{
 			check(!MeshShaderMap->HasShader(MeshMaterialShaderType, CurrentJob.PermutationId));
-			MeshShaderMap->AddShader(MeshMaterialShaderType, CurrentJob.PermutationId, Shader);
+			Shader = MeshShaderMap->FindOrAddShader(Shader);
 		}
 	}
 	else
 	{
 		FMaterialShaderType* MaterialShaderType = CurrentJob.ShaderType->GetMaterialShaderType();
 		check(MaterialShaderType);
-		Shader = MaterialShaderType->FinishCompileShader(MaterialCompilationOutput.UniformExpressionSet, MaterialShaderMapHash, CurrentJob, ShaderPipeline,
-#if ALLOW_SHADERMAP_DEBUG_DATA
-			FriendlyName);
-#else
-			FString());
-#endif
+		Shader = MaterialShaderType->FinishCompileShader(GetContent()->MaterialCompilationOutput.UniformExpressionSet, MaterialShaderMapHash, CurrentJob, ShaderPipeline, GetFriendlyName(), ResourceBuilder);
 
 		check(Shader);
 		if (!ShaderPipeline)
 		{
-			check(!HasShader(MaterialShaderType, CurrentJob.PermutationId));
-			AddShader(MaterialShaderType, CurrentJob.PermutationId, Shader);
+			check(!GetContent()->HasShader(MaterialShaderType, CurrentJob.PermutationId));
+			Shader = GetMutableContent()->FindOrAddShader(Shader);
 		}
 	}
 
 #if WITH_EDITOR
 	// add shader source to 
-	ShaderProcessedSource.Add(CurrentJob.ShaderType->GetFName(), CurrentJob.Output.OptionalFinalShaderSource);
+	new(GetMutableContent()->ShaderProcessedSource) FMaterialProcessedSource(CurrentJob.ShaderType->GetFName(), *CurrentJob.Output.OptionalFinalShaderSource);
 #endif
 
 	return Shader;
@@ -1714,12 +1696,13 @@ bool FMaterialShaderMap::ProcessCompilationResults(const TArray<TSharedRef<FShad
 	FSHAHash MaterialShaderMapHash;
 	ShaderMapId.GetMaterialHash(MaterialShaderMapHash);
 
+	FShaderMapResourceBuilder ResourceBuilder(GetResourceCode());
 	do
 	{
 		FShaderCompileJob* SingleJob = InCompilationResults[InOutJobIndex]->GetSingleShaderJob();
 		if (SingleJob)
 		{
-			ProcessCompilationResultsForSingleJob(SingleJob, nullptr, MaterialShaderMapHash);
+			ProcessCompilationResultsForSingleJob(SingleJob, nullptr, MaterialShaderMapHash, ResourceBuilder);
 			for (auto Pair : SingleJob->SharingPipelines)
 			{
 				auto& SharedPipelinesPerVF = SharedPipelines.FindOrAdd(SingleJob->VFType);
@@ -1742,7 +1725,7 @@ bool FMaterialShaderMap::ProcessCompilationResults(const TArray<TSharedRef<FShad
 			for (int32 Index = 0; Index < CurrentJob.StageJobs.Num(); ++Index)
 			{
 				SingleJob = CurrentJob.StageJobs[Index]->GetSingleShaderJob();
-				FShader* Shader = ProcessCompilationResultsForSingleJob(SingleJob, PipelineJob->ShaderPipeline, MaterialShaderMapHash);
+				FShader* Shader = ProcessCompilationResultsForSingleJob(SingleJob, PipelineJob->ShaderPipeline, MaterialShaderMapHash, ResourceBuilder);
 				ShaderStages.Add(Shader);
 				check(VertexFactoryType == CurrentJob.StageJobs[Index]->GetSingleShaderJob()->VFType);
 			}
@@ -1753,28 +1736,15 @@ bool FMaterialShaderMap::ProcessCompilationResults(const TArray<TSharedRef<FShad
 				if (VertexFactoryType)
 				{
 					check(VertexFactoryType->IsUsedWithMaterials());
-					FMeshMaterialShaderMap* MeshShaderMap = nullptr;
-					int32 MeshShaderMapIndex = INDEX_NONE;
-
-					// look for existing map for this vertex factory type
-					for (int32 ShaderMapIndex = 0; ShaderMapIndex < MeshShaderMaps.Num(); ShaderMapIndex++)
-					{
-						if (MeshShaderMaps[ShaderMapIndex].GetVertexFactoryType() == VertexFactoryType)
-						{
-							MeshShaderMap = &MeshShaderMaps[ShaderMapIndex];
-							MeshShaderMapIndex = ShaderMapIndex;
-							break;
-						}
-					}
-
+					FMeshMaterialShaderMap* MeshShaderMap = GetMeshShaderMap(VertexFactoryType);
 					check(MeshShaderMap);
-					check(!MeshShaderMap->HasShaderPipeline(ShaderPipeline->PipelineType));
-					MeshShaderMap->AddShaderPipeline(PipelineJob->ShaderPipeline, ShaderPipeline);
+					check(!MeshShaderMap->HasShaderPipeline(ShaderPipeline->TypeName));
+					MeshShaderMap->FindOrAddShaderPipeline(ShaderPipeline);
 				}
 				else
 				{
-					check(!HasShaderPipeline(ShaderPipeline->PipelineType));
-					AddShaderPipeline(PipelineJob->ShaderPipeline, ShaderPipeline);
+					check(!GetContent()->HasShaderPipeline(ShaderPipeline->TypeName));
+					GetMutableContent()->FindOrAddShaderPipeline(ShaderPipeline);
 				}
 			}
 		}
@@ -1791,30 +1761,33 @@ bool FMaterialShaderMap::ProcessCompilationResults(const TArray<TSharedRef<FShad
 	{
 		{
 			// Process the mesh shader pipelines that share shaders
-			for (int32 ShaderMapIndex = 0; ShaderMapIndex < MeshShaderMaps.Num(); ShaderMapIndex++)
+			for (int32 ShaderMapIndex = 0; ShaderMapIndex < GetContent()->OrderedMeshShaderMaps.Num(); ShaderMapIndex++)
 			{
-				auto* MeshShaderMap = &MeshShaderMaps[ShaderMapIndex];
-				auto* VertexFactory = MeshShaderMap->GetVertexFactoryType();
-				auto* FoundSharedPipelines = SharedPipelines.Find(VertexFactory);
-				if (VertexFactory && FoundSharedPipelines)
+				FMeshMaterialShaderMap* MeshShaderMap = GetContent()->OrderedMeshShaderMaps[ShaderMapIndex];
+				if (MeshShaderMap)
 				{
-					for (const FShaderPipelineType* ShaderPipelineType : *FoundSharedPipelines)
+					FVertexFactoryType* VertexFactory = FVertexFactoryType::GetVFByName(MeshShaderMap->GetVertexFactoryTypeName());
+					TArray<const FShaderPipelineType*>* FoundSharedPipelines = SharedPipelines.Find(VertexFactory);
+					if (VertexFactory && FoundSharedPipelines)
 					{
-						if (ShaderPipelineType->IsMeshMaterialTypePipeline() && !MeshShaderMap->HasShaderPipeline(ShaderPipelineType))
+						for (const FShaderPipelineType* ShaderPipelineType : *FoundSharedPipelines)
 						{
-							auto& StageTypes = ShaderPipelineType->GetStages();
-							TArray<FShader*> ShaderStages;
-							for (int32 Index = 0; Index < StageTypes.Num(); ++Index)
+							if (ShaderPipelineType->IsMeshMaterialTypePipeline() && !MeshShaderMap->HasShaderPipeline(ShaderPipelineType))
 							{
-								FMeshMaterialShaderType* ShaderType = ((FMeshMaterialShaderType*)(StageTypes[Index]))->GetMeshMaterialShaderType();
-								FShader* Shader = MeshShaderMap->GetShader(ShaderType, kUniqueShaderPermutationId);
-								check(Shader);
-								ShaderStages.Add(Shader);
+								auto& StageTypes = ShaderPipelineType->GetStages();
+								TArray<FShader*> ShaderStages;
+								for (int32 Index = 0; Index < StageTypes.Num(); ++Index)
+								{
+									FMeshMaterialShaderType* ShaderType = ((FMeshMaterialShaderType*)(StageTypes[Index]))->GetMeshMaterialShaderType();
+									FShader* Shader = MeshShaderMap->GetShader(ShaderType, kUniqueShaderPermutationId);
+									check(Shader);
+									ShaderStages.Add(Shader);
+								}
+
+								checkf(StageTypes.Num() == ShaderStages.Num(), TEXT("Internal Error adding MeshMaterial ShaderPipeline %s"), ShaderPipelineType->GetName());
+								FShaderPipeline* ShaderPipeline = new FShaderPipeline(ShaderPipelineType, ShaderStages);
+								MeshShaderMap->FindOrAddShaderPipeline(ShaderPipeline);
 							}
-
-							checkf(StageTypes.Num() == ShaderStages.Num(), TEXT("Internal Error adding MeshMaterial ShaderPipeline %s"), ShaderPipelineType->GetName());
-							FShaderPipeline* ShaderPipeline = new FShaderPipeline(ShaderPipelineType, ShaderStages);
-							MeshShaderMap->AddShaderPipeline(ShaderPipelineType, ShaderPipeline);
 						}
 					}
 				}
@@ -1826,47 +1799,39 @@ bool FMaterialShaderMap::ProcessCompilationResults(const TArray<TSharedRef<FShad
 			{
 				for (const FShaderPipelineType* ShaderPipelineType : *FoundSharedPipelines)
 				{
-					if (ShaderPipelineType->IsMaterialTypePipeline() && !HasShaderPipeline(ShaderPipelineType))
+					if (ShaderPipelineType->IsMaterialTypePipeline() && !GetContent()->HasShaderPipeline(ShaderPipelineType))
 					{
 						auto& StageTypes = ShaderPipelineType->GetStages();
 						TArray<FShader*> ShaderStages;
 						for (int32 Index = 0; Index < StageTypes.Num(); ++Index)
 						{
 							FMaterialShaderType* ShaderType = ((FMaterialShaderType*)(StageTypes[Index]))->GetMaterialShaderType();
-							FShader* Shader = GetShader(ShaderType, kUniqueShaderPermutationId);
+							FShader* Shader = GetContent()->GetShader(ShaderType, kUniqueShaderPermutationId);
 							check(Shader);
 							ShaderStages.Add(Shader);
 						}
 
 						checkf(StageTypes.Num() == ShaderStages.Num(), TEXT("Internal Error adding Material ShaderPipeline %s"), ShaderPipelineType->GetName());
 						FShaderPipeline* ShaderPipeline = new FShaderPipeline(ShaderPipelineType, ShaderStages);
-						AddShaderPipeline(ShaderPipelineType, ShaderPipeline);
+						GetMutableContent()->FindOrAddShaderPipeline(ShaderPipeline);
 					}
 				}
 			}
 		}
 
-		for (int32 ShaderMapIndex = MeshShaderMaps.Num() - 1; ShaderMapIndex >= 0; ShaderMapIndex--)
-		{
-			if (MeshShaderMaps[ShaderMapIndex].GetNumShaders() == 0 && MeshShaderMaps[ShaderMapIndex].GetNumShaderPipelines() == 0)
-			{
-				// If the mesh material shader map is complete and empty, discard it.
-				MeshShaderMaps.RemoveAt(ShaderMapIndex);
-			}
-		}
-
-		// Reinitialize the ordered mesh shader maps
-		InitOrderedMeshShaderMaps();
+		FinalizeContent();
 
 		// Add the persistent shaders to the local shader cache.
 		if (bIsPersistent)
 		{
-			SaveToDerivedDataCache();
+			SaveToDerivedDataCache(CompilingTargetPlatform);
 		}
 
 		// The shader map can now be used on the rendering thread
 		bCompilationFinalized = true;
 
+		CompilingTargetPlatform = nullptr;
+
 		return true;
 	}
 
@@ -1891,38 +1856,151 @@ bool FMaterialShaderMap::TryToAddToExistingCompilationTask(FMaterial* Material)
 	return false;
 }
 
-bool FMaterialShaderMap::IsMaterialShaderComplete(
-	const FMaterial* Material,
-	const FMaterialShaderType* ShaderType,
-	const FShaderPipelineType* Pipeline,
-	int32 PermutationId,
-	bool bSilent) const
+class FMaterialShaderMapLayoutCache
 {
-	const bool bShouldCacheMaterialShader = ShouldCacheMaterialShader(ShaderType, Platform, Material, PermutationId);
-	const bool bMissingShaderPipeline = (Pipeline && !HasShaderPipeline(Pipeline));
-	const bool bMissingShader = (!Pipeline && !HasShader((FShaderType*)ShaderType, PermutationId));
+public:
+	static FMaterialShaderMapLayoutCache& Get()
+	{
+		static FMaterialShaderMapLayoutCache Instance;
+		return Instance;
+	}
 
-	if (bShouldCacheMaterialShader && (bMissingShaderPipeline || bMissingShader))
+	const FMaterialShaderMapLayout& AcquireLayout(EShaderPlatform Platform, const FMaterialShaderParameters& MaterialParameters)
 	{
-		if (!bSilent)
+		const uint64 ParameterHash = CityHash64((char*)&MaterialParameters, sizeof(MaterialParameters));
+		
+		int32 Index = INDEX_NONE;
+		for (Index = ShaderMapHashTable.First(ParameterHash); ShaderMapHashTable.IsValid(Index); Index = ShaderMapHashTable.Next(Index))
 		{
-			if (Pipeline)
+			if (MaterialParameterHashes[Index] == ParameterHash)
 			{
-				UE_LOG(LogShaders, Warning, TEXT("Incomplete material %s, missing pipeline %s."), *Material->GetFriendlyName(), Pipeline->GetName());
+				break;
 			}
-			else
+		}
+		if (Index == INDEX_NONE)
+		{
+			Index = MaterialParameterHashes.Add(ParameterHash);
+			MaterialShaderParameters.Add(MaterialParameters);
+			ShaderMapHashTable.Add(ParameterHash, Index);
+			FMaterialShaderMapLayout& Layout = ShaderMapLayouts.AddDefaulted_GetRef();
+			check(MaterialParameterHashes.Num() == ShaderMapLayouts.Num());
+			check(MaterialShaderParameters.Num() == ShaderMapLayouts.Num());
+			CreateLayout(Layout, Platform, MaterialParameters);
+		}
+
+		return ShaderMapLayouts[Index];
+	}
+
+private:
+	void CreateLayout(FMaterialShaderMapLayout& Layout, EShaderPlatform Platform, const FMaterialShaderParameters& MaterialParameters)
+	{
+		SCOPED_LOADTIMER(FMaterialShaderMapLayoutCache_CreateLayout);
+
+		const bool bHasTessellation = MaterialParameters.TessellationMode != MTM_NoTessellation;
+		const TArray<FShaderType*>& SortedMaterialShaderTypes = FShaderType::GetSortedTypes(FShaderType::EShaderTypeForDynamicCast::Material);
+		const TArray<FShaderType*>& SortedMeshMaterialShaderTypes = FShaderType::GetSortedTypes(FShaderType::EShaderTypeForDynamicCast::MeshMaterial);
+		const TArray<FShaderPipelineType*>& SortedMaterialPipelineTypes = FShaderPipelineType::GetSortedTypes(FShaderType::EShaderTypeForDynamicCast::Material);
+		const TArray<FShaderPipelineType*>& SortedMeshMaterialPipelineTypes = FShaderPipelineType::GetSortedTypes(FShaderType::EShaderTypeForDynamicCast::MeshMaterial);
+
+		FSHA1 Hasher;
+
+		for (FShaderType* BaseShaderType : SortedMaterialShaderTypes)
+		{
+			// Find this shader type in the material's shader map.
+			FMaterialShaderType* ShaderType = static_cast<FMaterialShaderType*>(BaseShaderType);
+			const int32 PermutationCount = ShaderType->GetPermutationCount();
+			for (int32 PermutationId = 0; PermutationId < PermutationCount; ++PermutationId)
 			{
-				UE_LOG(LogShaders, Warning, TEXT("Incomplete material %s, missing FMaterialShader (%s, %d)."), *Material->GetFriendlyName(), ShaderType->GetName(), PermutationId);
+				if (ShaderType->ShouldCompilePermutation(Platform, MaterialParameters, PermutationId))
+				{
+					Layout.Shaders.Add(FShaderLayoutEntry(ShaderType, PermutationId));
+
+					const FHashedName& TypeName = ShaderType->GetHashedName();
+					Hasher.Update((uint8*)&TypeName, sizeof(TypeName));
+					Hasher.Update((uint8*)&PermutationId, sizeof(PermutationId));
+				}
 			}
 		}
-		return false;
+
+		// Iterate over all pipeline types
+		for (FShaderPipelineType* ShaderPipelineType : SortedMaterialPipelineTypes)
+		{
+			if (ShaderPipelineType->HasTessellation() == bHasTessellation &&
+				FMaterialShaderType::ShouldCompilePipeline(ShaderPipelineType, Platform, MaterialParameters))
+			{
+				Layout.ShaderPipelines.Add(ShaderPipelineType);
+
+				const FHashedName& TypeName = ShaderPipelineType->GetHashedName();
+				Hasher.Update((uint8*)&TypeName, sizeof(TypeName));
+			}
+		}
+
+		for (FVertexFactoryType* VertexFactoryType : FVertexFactoryType::GetSortedMaterialTypes())
+		{
+			if (!FMeshMaterialShaderType::ShouldCompileVertexFactoryPermutation(VertexFactoryType, Platform, MaterialParameters))
+			{
+				continue;
+			}
+
+			FMeshMaterialShaderMapLayout* MeshLayout = nullptr;
+			for (FShaderType* BaseShaderType : SortedMeshMaterialShaderTypes)
+			{
+				FMeshMaterialShaderType* ShaderType = static_cast<FMeshMaterialShaderType*>(BaseShaderType);
+				const int32 PermutationCount = ShaderType->GetPermutationCount();
+				for (int32 PermutationId = 0; PermutationId < PermutationCount; ++PermutationId)
+				{
+					if (ShaderType->ShouldCompilePermutation(Platform, MaterialParameters, VertexFactoryType, PermutationId))
+					{
+						if (!MeshLayout)
+						{
+							MeshLayout = new(Layout.MeshShaderMaps) FMeshMaterialShaderMapLayout(VertexFactoryType);
+						}
+						MeshLayout->Shaders.Add(FShaderLayoutEntry(ShaderType, PermutationId));
+
+						const FHashedName& TypeName = ShaderType->GetHashedName();
+						Hasher.Update((uint8*)&TypeName, sizeof(TypeName));
+						Hasher.Update((uint8*)&PermutationId, sizeof(PermutationId));
+					}
+				}
+			}
+
+			for (FShaderPipelineType* ShaderPipelineType : SortedMeshMaterialPipelineTypes)
+			{
+				if (ShaderPipelineType->HasTessellation() == bHasTessellation &&
+					FMeshMaterialShaderType::ShouldCompilePipeline(ShaderPipelineType, Platform, MaterialParameters, VertexFactoryType))
+				{
+					// Now check the completeness of the shader map
+					if (!MeshLayout)
+					{
+						MeshLayout = new(Layout.MeshShaderMaps) FMeshMaterialShaderMapLayout(VertexFactoryType);
+					}
+					MeshLayout->ShaderPipelines.Add(ShaderPipelineType);
+
+					const FHashedName& TypeName = ShaderPipelineType->GetHashedName();
+					Hasher.Update((uint8*)&TypeName, sizeof(TypeName));
+				}
+			}
+		}
+
+		Hasher.Final();
+		Hasher.GetHash(Layout.ShaderMapHash.Hash);
 	}
 
-	return true;
+	TArray<FMaterialShaderMapLayout> ShaderMapLayouts;
+	TArray<FMaterialShaderParameters> MaterialShaderParameters;
+	TArray<uint64> MaterialParameterHashes;
+	FHashTable ShaderMapHashTable;
+};
+
+const FMaterialShaderMapLayout& AcquireMaterialShaderMapLayout(EShaderPlatform Platform, const FMaterialShaderParameters& MaterialParameters)
+{
+	return FMaterialShaderMapLayoutCache::Get().AcquireLayout(Platform, MaterialParameters);
 }
 
 bool FMaterialShaderMap::IsComplete(const FMaterial* Material, bool bSilent)
 {
+	SCOPED_LOADTIMER(FMaterialShaderMap_IsComplete);
+
 	// Make sure we are operating on a referenced shader map or the below Find will cause this shader map to be deleted,
 	// Since it creates a temporary ref counted pointer.
 	check(NumRefs > 0);
@@ -1934,84 +2012,105 @@ bool FMaterialShaderMap::IsComplete(const FMaterial* Material, bool bSilent)
 		return false;
 	}
 
-	// Iterate over all vertex factory types.
-	for(TLinkedList<FVertexFactoryType*>::TIterator VertexFactoryTypeIt(FVertexFactoryType::GetTypeList());VertexFactoryTypeIt;VertexFactoryTypeIt.Next())
+	const FMaterialShaderMapContent* LocalContent = GetContent();
+	const EShaderPlatform Platform = LocalContent->GetShaderPlatform();
+	const FMaterialShaderParameters MaterialParameters(Material);
+
+	const FMaterialShaderMapLayout& Layout = AcquireMaterialShaderMapLayout(Platform, MaterialParameters);
+	if (Layout.ShaderMapHash == LocalContent->ShaderContentHash)
 	{
-		FVertexFactoryType* VertexFactoryType = *VertexFactoryTypeIt;
+		return true;
+	}
 
-		if(VertexFactoryType->IsUsedWithMaterials())
+	// If our hash doesn't match the cached layout hash, shader map may still be complete
+	// This can happen if FMaterial::ShouldCache is set to return false for any shaders that are included in the cached layout
+	for (const FShaderLayoutEntry& Shader : Layout.Shaders)
+	{
+		if (!LocalContent->HasShader(Shader.ShaderType, Shader.PermutationId))
 		{
-			// Find the shaders for this vertex factory type.
-			const FMeshMaterialShaderMap* MeshShaderMap = GetMeshShaderMap(VertexFactoryType);
-			if (!FMeshMaterialShaderMap::IsComplete(MeshShaderMap,GetShaderPlatform(),Material,VertexFactoryType,bSilent))
+			if (Material->ShouldCache(Platform, Shader.ShaderType, nullptr))
 			{
-				if (!MeshShaderMap && !bSilent)
+				if (!bSilent)
 				{
-					UE_LOG(LogShaders, Warning, TEXT("Incomplete material %s, missing Vertex Factory %s."), *Material->GetFriendlyName(), VertexFactoryType->GetName());
+					UE_LOG(LogShaders, Warning, TEXT("Incomplete material %s, missing FMaterialShader (%s, %d)."), *Material->GetFriendlyName(), Shader.ShaderType->GetName(), Shader.PermutationId);
 				}
 				return false;
 			}
 		}
 	}
 
-	// Iterate over all material shader types.
-	for(TLinkedList<FShaderType*>::TIterator ShaderTypeIt(FShaderType::GetTypeList());ShaderTypeIt;ShaderTypeIt.Next())
+	for (const FShaderPipelineType* Pipeline : Layout.ShaderPipelines)
 	{
-		// Find this shader type in the material's shader map.
-		FMaterialShaderType* ShaderType = ShaderTypeIt->GetMaterialShaderType();
-		const int32 PermutationCount = ShaderType ? ShaderType->GetPermutationCount() : 0;
-		for (int32 PermutationId = 0; PermutationId < PermutationCount; ++PermutationId)
+		if (!LocalContent->HasShaderPipeline(Pipeline))
 		{
-			if (!IsMaterialShaderComplete(Material, ShaderType, nullptr, PermutationId, bSilent))
+			if (Material->ShouldCachePipeline(Platform, Pipeline, nullptr))
 			{
+				if (!bSilent)
+				{
+					UE_LOG(LogShaders, Warning, TEXT("Incomplete material %s, missing pipeline %s."), *Material->GetFriendlyName(), Pipeline->GetName());
+				}
 				return false;
 			}
 		}
 	}
 
-	// Iterate over all pipeline types
-	const bool bHasTessellation = Material->GetTessellationMode() != MTM_NoTessellation;
-	for (TLinkedList<FShaderPipelineType*>::TIterator ShaderPipelineIt(FShaderPipelineType::GetTypeList());ShaderPipelineIt;ShaderPipelineIt.Next())
+	for (const FMeshMaterialShaderMapLayout& MeshLayout : Layout.MeshShaderMaps)
 	{
-		const FShaderPipelineType* Pipeline = *ShaderPipelineIt;
-		if (Pipeline->IsMaterialTypePipeline() && Pipeline->HasTessellation() == bHasTessellation)
-		{
-			auto& StageTypes = Pipeline->GetStages();
+		const FMeshMaterialShaderMap* MeshShaderMap = LocalContent->GetMeshShaderMap(MeshLayout.VertexFactoryType->GetHashedName());
 
-			int32 NumShouldCache = 0;
-			for (int32 Index = 0; Index < StageTypes.Num(); ++Index)
+		for (const FShaderLayoutEntry& Shader : MeshLayout.Shaders)
+		{
+			if (Material->ShouldCache(Platform, Shader.ShaderType, MeshLayout.VertexFactoryType))
 			{
-				auto* ShaderType = StageTypes[Index]->GetMaterialShaderType();
-				if (ShouldCacheMaterialShader(ShaderType, GetShaderPlatform(), Material, kUniqueShaderPermutationId))
+				if (!MeshShaderMap || !MeshShaderMap->HasShader(Shader.ShaderType, Shader.PermutationId))
 				{
-					++NumShouldCache;
-				}
-				else
-				{
-					break;
+					if (!bSilent)
+					{
+						if (!MeshShaderMap)
+						{
+							UE_LOG(LogShaders, Warning, TEXT("Incomplete material %s, missing Vertex Factory %s."), *Material->GetFriendlyName(), MeshLayout.VertexFactoryType->GetName());
+						}
+						else
+						{
+							UE_LOG(LogShaders, Warning, TEXT("Incomplete material %s, missing (%s, %d) from %s."), *Material->GetFriendlyName(), Shader.ShaderType->GetName(), Shader.PermutationId, MeshLayout.VertexFactoryType->GetName());
+						}
+					}
+					return false;
 				}
 			}
+		}
 
-			if (NumShouldCache == StageTypes.Num())
+		for (const FShaderPipelineType* Pipeline : MeshLayout.ShaderPipelines)
+		{
+			if (!MeshShaderMap || !MeshShaderMap->HasShaderPipeline(Pipeline))
 			{
-				for (int32 Index = 0; Index < StageTypes.Num(); ++Index)
+				if (Material->ShouldCachePipeline(Platform, Pipeline, MeshLayout.VertexFactoryType))
 				{
-					auto* ShaderType = StageTypes[Index]->GetMaterialShaderType();
-					if (!IsMaterialShaderComplete(Material, ShaderType, Pipeline, kUniqueShaderPermutationId, bSilent))
+					if (!bSilent)
 					{
-						return false;
+						if (!MeshShaderMap)
+						{
+							UE_LOG(LogShaders, Warning, TEXT("Incomplete material %s, missing Vertex Factory %s."), *Material->GetFriendlyName(), MeshLayout.VertexFactoryType->GetName());
+						}
+						else
+						{
+							UE_LOG(LogShaders, Warning, TEXT("Incomplete material %s, missing pipeline %s from %s."), *Material->GetFriendlyName(), Pipeline->GetName(), MeshLayout.VertexFactoryType->GetName());
+						}
 					}
+					return false;
 				}
 			}
 		}
 	}
 
+	// Was missing some shaders from the initial layout, but all of those shaders were explicitly exluced by our FMaterial::ShouldCache implementation
 	return true;
 }
 
 #if WITH_EDITOR
 void FMaterialShaderMap::LoadMissingShadersFromMemory(const FMaterial* Material)
 {
+#if 0
 	// Make sure we are operating on a referenced shader map or the below Find will cause this shader map to be deleted,
 	// Since it creates a temporary ref counted pointer.
 	check(NumRefs > 0);
@@ -2107,52 +2206,70 @@ void FMaterialShaderMap::LoadMissingShadersFromMemory(const FMaterial* Material)
 			}
 		}
 	}
+#endif
 }
 #endif // WITH_EDITOR
 
 #if WITH_EDITOR
-const FString* FMaterialShaderMap::GetShaderSource(const FName ShaderTypeName) const
+const FMemoryImageString* FMaterialShaderMap::GetShaderSource(const FName ShaderTypeName) const
 {
-	const FString* Source = ShaderProcessedSource.Find(ShaderTypeName);
-	return Source;
+	for (const FMaterialProcessedSource& Source : GetContent()->ShaderProcessedSource)
+	{
+		if (Source.Name == ShaderTypeName)
+		{
+			return &Source.Source;
+		}
+	}
+
+	return nullptr;
 }
-#endif
+#endif // WITH_EDITOR
 
-void FMaterialShaderMap::GetShaderList(TMap<FShaderId, FShader*>& OutShaders) const
+void FMaterialShaderMap::GetShaderList(TMap<FShaderId, TShaderRef<FShader>>& OutShaders) const
 {
-	TShaderMap<FMaterialShaderType>::GetShaderList(OutShaders);
-	for(int32 Index = 0;Index < MeshShaderMaps.Num();Index++)
+	FSHAHash MaterialShaderMapHash;
+#if WITH_EDITOR
+	// TODO
+	ShaderMapId.GetMaterialHash(MaterialShaderMapHash);
+#endif
+
+	GetContent()->GetShaderList(*this, MaterialShaderMapHash, OutShaders);
+	for (FMeshMaterialShaderMap* MeshShaderMap : GetContent()->OrderedMeshShaderMaps)
 	{
-		MeshShaderMaps[Index].GetShaderList(OutShaders);
+		if (MeshShaderMap)
+		{
+			MeshShaderMap->GetShaderList(*this, MaterialShaderMapHash, OutShaders);
+		}
 	}
 }
 
-void FMaterialShaderMap::GetShaderList(TMap<FName, FShader*>& OutShaders) const
+void FMaterialShaderMap::GetShaderList(TMap<FHashedName, TShaderRef<FShader>>& OutShaders) const
 {
-	TShaderMap<FMaterialShaderType>::GetShaderList(OutShaders);
-	for (int32 Index = 0; Index < MeshShaderMaps.Num(); Index++)
+	GetContent()->GetShaderList(*this, OutShaders);
+	for (FMeshMaterialShaderMap* MeshShaderMap : GetContent()->OrderedMeshShaderMaps)
 	{
-		MeshShaderMaps[Index].GetShaderList(OutShaders);
+		if (MeshShaderMap)
+		{
+			MeshShaderMap->GetShaderList(*this, OutShaders);
+		}
 	}
 }
 
-void FMaterialShaderMap::GetShaderPipelineList(TArray<FShaderPipeline*>& OutShaderPipelines) const
+void FMaterialShaderMap::GetShaderPipelineList(TArray<FShaderPipelineRef>& OutShaderPipelines) const
 {
-	TShaderMap<FMaterialShaderType>::GetShaderPipelineList(OutShaderPipelines, FShaderPipeline::EAll);
-	for (int32 Index = 0;Index < MeshShaderMaps.Num();Index++)
+	GetContent()->GetShaderPipelineList(*this, OutShaderPipelines, FShaderPipeline::EAll);
+	for(FMeshMaterialShaderMap* MeshShaderMap : GetContent()->OrderedMeshShaderMaps)
 	{
-		MeshShaderMaps[Index].GetShaderPipelineList(OutShaderPipelines, FShaderPipeline::EAll);
+		if (MeshShaderMap)
+		{
+			MeshShaderMap->GetShaderPipelineList(*this, OutShaderPipelines, FShaderPipeline::EAll);
+		}
 	}
 }
 
 uint32 FMaterialShaderMap::GetShaderNum() const
 {
-	uint32 Count = Shaders.Num();
-	for (int32 Index = 0; Index < MeshShaderMaps.Num(); Index++)
-	{
-		Count += MeshShaderMaps[Index].GetShaderNum();
-	}
-	return Count;
+	return GetContent()->GetNumShaders();
 }
 
 /**
@@ -2161,34 +2278,14 @@ uint32 FMaterialShaderMap::GetShaderNum() const
 void FMaterialShaderMap::Register(EShaderPlatform InShaderPlatform)
 {
 	extern int32 GCreateShadersOnLoad;
-	if (GCreateShadersOnLoad && GetShaderPlatform() == InShaderPlatform)
+	if (GCreateShadersOnLoad && GetShaderPlatform() == InShaderPlatform && GetResource())
 	{
-		for (auto KeyValue : GetShaders())
-		{
-			FShader* Shader = KeyValue.Value;
-			if (Shader)
-			{
-				Shader->BeginInitializeResources();
-			}
-		}
-
-		for (int32 Index = 0; Index < MeshShaderMaps.Num(); Index++)
-		{
-			for (TMap<FMeshMaterialShaderMap::FShaderPrimaryKey, TRefCountPtr<FShader> >::TConstIterator It(MeshShaderMaps[Index].GetShaders()); It; ++It)
-			{
-				FShader* Shader = It.Value();
-				if (Shader)
-				{
-					Shader->BeginInitializeResources();
-				}
-			}
-		}
+		// TODO
 	}
 
 	if (!bRegistered)
 	{
 		INC_DWORD_STAT(STAT_Shaders_NumShaderMaps);
-		INC_DWORD_STAT_BY(STAT_Shaders_ShaderMapMemory, GetSizeBytes());
 	}
 
 	{
@@ -2235,7 +2332,6 @@ void FMaterialShaderMap::Release()
 			{
 				bRegistered = false;
 				DEC_DWORD_STAT(STAT_Shaders_NumShaderMaps);
-				DEC_DWORD_STAT_BY(STAT_Shaders_ShaderMapMemory, GetSizeBytes());
 
 				FMaterialShaderMap *CachedMap = GIdToMaterialShaderMap[GetShaderPlatform()].FindRef(ShaderMapId);
 
@@ -2264,8 +2360,8 @@ void FMaterialShaderMap::Release()
 	}
 }
 
-FMaterialShaderMap::FMaterialShaderMap(EShaderPlatform InPlatform) :
-	TShaderMap<FMaterialShaderType>(InPlatform),
+FMaterialShaderMap::FMaterialShaderMap() :
+	CompilingTargetPlatform(nullptr),
 	CompilingId(1),
 	NumRefs(0),
 	bDeletedThroughDeferredCleanup(false),
@@ -2287,10 +2383,14 @@ FMaterialShaderMap::~FMaterialShaderMap()
 	check(bDeletedThroughDeferredCleanup);
 	check(!bRegistered);
 #if ALLOW_SHADERMAP_DEBUG_DATA
-	if(GShaderCompilerStats != 0)
+	if(GShaderCompilerStats != 0 && GetContent())
 	{
-		FString Path = !MaterialPath.IsEmpty() ? MaterialPath : GetFriendlyName();
-		GShaderCompilerStats->RegisterCookedShaders(GetShaderNum(), CompileTime, Platform, Path, GetDebugDescription());
+		FString Path = GetMaterialPath();
+		if (Path.IsEmpty())
+		{
+			Path = GetFriendlyName();
+		}
+		GShaderCompilerStats->RegisterCookedShaders(GetShaderNum(), CompileTime, GetShaderPlatform(), Path, GetDebugDescription());
 	}
 	AllMaterialShaderMaps.RemoveSwap(this);
 #endif
@@ -2300,12 +2400,16 @@ FMaterialShaderMap::~FMaterialShaderMap()
  * Removes all entries in the cache with exceptions based on a shader type
  * @param ShaderType - The shader type to flush
  */
-void FMaterialShaderMap::FlushShadersByShaderType(FShaderType* ShaderType)
+void FMaterialShaderMap::FlushShadersByShaderType(const FShaderType* ShaderType)
 {
 	// flush from all the vertex factory shader maps
-	for(int32 Index = 0;Index < MeshShaderMaps.Num();Index++)
+	for (int32 Index = 0; Index < GetContent()->OrderedMeshShaderMaps.Num(); Index++)
 	{
-		MeshShaderMaps[Index].FlushShadersByShaderType(ShaderType);
+		FMeshMaterialShaderMap* MeshShaderMap = GetContent()->OrderedMeshShaderMaps[Index];
+		if (MeshShaderMap)
+		{
+			MeshShaderMap->FlushShadersByShaderType(ShaderType);
+		}
 	}
 
 	if (ShaderType->GetMaterialShaderType())
@@ -2313,7 +2417,7 @@ void FMaterialShaderMap::FlushShadersByShaderType(FShaderType* ShaderType)
 		const int32 PermutationCount = ShaderType->GetPermutationCount();
 		for (int32 PermutationId = 0; PermutationId < PermutationCount; ++PermutationId)
 		{
-			RemoveShaderTypePermutaion(ShaderType->GetMaterialShaderType(), PermutationId);
+			GetMutableContent()->RemoveShaderTypePermutaion(ShaderType->GetMaterialShaderType(), PermutationId);
 		}
 	}
 }
@@ -2321,14 +2425,18 @@ void FMaterialShaderMap::FlushShadersByShaderType(FShaderType* ShaderType)
 void FMaterialShaderMap::FlushShadersByShaderPipelineType(const FShaderPipelineType* ShaderPipelineType)
 {
 	// flush from all the vertex factory shader maps
-	for (int32 Index = 0;Index < MeshShaderMaps.Num();Index++)
+	for (int32 Index = 0;Index < GetContent()->OrderedMeshShaderMaps.Num();Index++)
 	{
-		MeshShaderMaps[Index].FlushShadersByShaderPipelineType(ShaderPipelineType);
+		FMeshMaterialShaderMap* MeshShaderMap = GetContent()->OrderedMeshShaderMaps[Index];
+		if (MeshShaderMap)
+		{
+			MeshShaderMap->FlushShadersByShaderPipelineType(ShaderPipelineType);
+		}
 	}
 
 	if (ShaderPipelineType->IsMaterialTypePipeline())
 	{
-		RemoveShaderPipelineType(ShaderPipelineType);
+		GetMutableContent()->RemoveShaderPipelineType(ShaderPipelineType);
 	}
 }
 
@@ -2339,187 +2447,24 @@ void FMaterialShaderMap::FlushShadersByShaderPipelineType(const FShaderPipelineT
  */
 void FMaterialShaderMap::FlushShadersByVertexFactoryType(const FVertexFactoryType* VertexFactoryType)
 {
-	for (int32 Index = 0; Index < MeshShaderMaps.Num(); Index++)
-	{
-		FVertexFactoryType* VFType = MeshShaderMaps[Index].GetVertexFactoryType();
-		// determine if this shaders vertex factory type should be flushed
-		if (VFType == VertexFactoryType)
-		{
-			// remove the shader map
-			MeshShaderMaps.RemoveAt(Index);
-			// fix up the counter
-			Index--;
-		}
-	}
-
-	// reset the OrderedMeshShaderMap to remove references to the removed maps
-	InitOrderedMeshShaderMaps();
+	GetMutableContent()->RemoveMeshShaderMap(VertexFactoryType);
 }
 
-struct FCompareMeshShaderMaps
-{
-	FORCEINLINE bool operator()(const FMeshMaterialShaderMap& A, const FMeshMaterialShaderMap& B ) const
-	{
-		return FCString::Strncmp(
-			A.GetVertexFactoryType()->GetName(), 
-			B.GetVertexFactoryType()->GetName(), 
-			FMath::Min(FCString::Strlen(A.GetVertexFactoryType()->GetName()), FCString::Strlen(B.GetVertexFactoryType()->GetName()))) > 0;
-	}
-};
-
 void FMaterialShaderMap::Serialize(FArchive& Ar, bool bInlineShaderResources, bool bLoadedByCookedMaterial)
 {
-	LLM_SCOPE(ELLMTag::Shaders);
+	SCOPED_LOADTIMER(FMaterialShaderMap_Serialize);
 	// Note: This is saved to the DDC, not into packages (except when cooked)
 	// Backwards compatibility therefore will not work based on the version of Ar
 	// Instead, just bump MATERIALSHADERMAP_DERIVEDDATA_VER
-
 	ShaderMapId.Serialize(Ar, bLoadedByCookedMaterial);
-
-	// serialize the platform enum as a uint8
-	int32 TempPlatform = (int32)GetShaderPlatform();
-	Ar << TempPlatform;
-	Platform = (EShaderPlatform)TempPlatform;
-
-#if ALLOW_SHADERMAP_DEBUG_DATA
-	Ar << FriendlyName;
-	Ar << MaterialPath;
-#else
-	FString TempString;
-	Ar << TempString;
-	Ar << TempString;
-#endif
-
-	MaterialCompilationOutput.Serialize(Ar);
-
-#if ALLOW_SHADERMAP_DEBUG_DATA
-	Ar << DebugDescription;
-#else
-	Ar << TempString;
-#endif
-
-	if (Ar.IsSaving())
-	{
-		// Material shaders
-		TShaderMap<FMaterialShaderType>::SerializeInline(Ar, bInlineShaderResources, false, false);
-
-		// Mesh material shaders
-		int32 NumMeshShaderMaps = 0;
-
-		for (int32 VFIndex = 0; VFIndex < OrderedMeshShaderMaps.Num(); VFIndex++)
-		{
-			FMeshMaterialShaderMap* MeshShaderMap = OrderedMeshShaderMaps[VFIndex];
-
-			if (MeshShaderMap)
-			{
-				// Count the number of non-empty mesh shader maps
-				NumMeshShaderMaps++;
-			}
-		}
-
-		Ar << NumMeshShaderMaps;
-
-		TArray<FMeshMaterialShaderMap*> SortedMeshShaderMaps;
-		SortedMeshShaderMaps.Empty(MeshShaderMaps.Num());
-
-		for (int32 MapIndex = 0; MapIndex < MeshShaderMaps.Num(); MapIndex++)
-		{
-			SortedMeshShaderMaps.Add(&MeshShaderMaps[MapIndex]);
-		}
-
-		// Sort mesh shader maps by VF name so that the DDC entry always has the same binary result for a given key
-		SortedMeshShaderMaps.Sort(FCompareMeshShaderMaps());
-
-		for (int32 MapIndex = 0; MapIndex < SortedMeshShaderMaps.Num(); MapIndex++)
-		{
-			FMeshMaterialShaderMap* MeshShaderMap = SortedMeshShaderMaps[MapIndex];
-
-			if (MeshShaderMap)
-			{
-				FVertexFactoryType* VFType = MeshShaderMap->GetVertexFactoryType();
-				check(VFType);
-
-				Ar << VFType;
-
-				MeshShaderMap->SerializeInline(Ar, bInlineShaderResources, false, false);
-			}
-		}
-
-		bool bCooking = Ar.IsCooking();
-		Ar << bCooking;
-
-#if WITH_EDITOR
-		if (!bCooking)
-		{
-			int32 NrShaderSources = ShaderProcessedSource.Num();
-			Ar << NrShaderSources;
-			for (auto SourceEntry : ShaderProcessedSource)
-			{
-				Ar << SourceEntry.Key;
-				Ar << SourceEntry.Value;
-			}
-		}
-#endif
-	}
-
-	if (Ar.IsLoading())
-	{
-		MeshShaderMaps.Empty();
-
-		// Material shaders
-		TShaderMap<FMaterialShaderType>::SerializeInline(Ar, bInlineShaderResources, false, bLoadedByCookedMaterial);
-
-		// Mesh material shaders
-		int32 NumMeshShaderMaps = 0;
-		Ar << NumMeshShaderMaps;
-
-		MeshShaderMaps.Empty(NumMeshShaderMaps);
-		OrderedMeshShaderMaps.Empty(FVertexFactoryType::GetNumVertexFactoryTypes());
-		OrderedMeshShaderMaps.AddZeroed(FVertexFactoryType::GetNumVertexFactoryTypes());
-
-		for (int32 VFIndex = 0; VFIndex < NumMeshShaderMaps; VFIndex++)
-		{
-			FVertexFactoryType* VFType = nullptr;
-			Ar << VFType;
-			// Not handling missing vertex factory types on cooked data
-			// The cooker binary and running binary are assumed to be on the same code version
-			check(VFType);
-
-			FMeshMaterialShaderMap* MeshShaderMap = new FMeshMaterialShaderMap(GetShaderPlatform(), VFType);
-			MeshShaderMap->SerializeInline(Ar, bInlineShaderResources, false, bLoadedByCookedMaterial);
-			MeshShaderMaps.Add(MeshShaderMap);
-			OrderedMeshShaderMaps[VFType->GetId()] = MeshShaderMap;
-		}
-
-		bool bCooked;
-		Ar << bCooked;
-
-#if WITH_EDITOR
-		if (!bCooked)
-		{
-			int32 NrShaderSources = 0;
-			Ar << NrShaderSources;
-
-			ShaderProcessedSource.Empty();
-			FString ShaderSourceSrc;
-			FName ShaderName;
-			for (int32 i = 0; i < NrShaderSources; ++i)
-			{
-				Ar << ShaderName;
-				Ar << ShaderSourceSrc;
-
-				ShaderProcessedSource.Add(ShaderName, ShaderSourceSrc);
-			}
-		}
-#endif
-	}
+	Super::Serialize(Ar, bInlineShaderResources, bLoadedByCookedMaterial);
 }
 
-void FMaterialShaderMap::RegisterSerializedShaders(bool bLoadedByCookedMaterial)
+/*void FMaterialShaderMap::RegisterSerializedShaders(bool bLoadedByCookedMaterial)
 {
 	check(IsInGameThread());
 
-	TShaderMap<FMaterialShaderType>::RegisterSerializedShaders(bLoadedByCookedMaterial);
+	TShaderMapContent<FMaterialShaderType>::RegisterSerializedShaders(bLoadedByCookedMaterial);
 	
 	for (FMeshMaterialShaderMap* MeshShaderMap : OrderedMeshShaderMaps)
 	{
@@ -2549,7 +2494,7 @@ void FMaterialShaderMap::RegisterSerializedShaders(bool bLoadedByCookedMaterial)
 
 void FMaterialShaderMap::DiscardSerializedShaders()
 {
-	TShaderMap<FMaterialShaderType>::DiscardSerializedShaders();
+	TShaderMapContent<FMaterialShaderType>::DiscardSerializedShaders();
 
 	for (int32 VFIndex = 0; VFIndex < OrderedMeshShaderMaps.Num(); VFIndex++)
 	{
@@ -2561,7 +2506,7 @@ void FMaterialShaderMap::DiscardSerializedShaders()
 		MeshShaderMaps[Index].DiscardSerializedShaders();
 	}
 	MeshShaderMaps.Empty();
-}
+}*/
 
 void FMaterialShaderMap::RemovePendingMaterial(FMaterial* Material)
 {
@@ -2600,39 +2545,49 @@ const FMaterialShaderMap* FMaterialShaderMap::GetShaderMapBeingCompiled(const FM
 #if WITH_EDITOR
 uint32 FMaterialShaderMap::GetMaxTextureSamplers() const
 {
-	uint32 MaxTextureSamplers = GetMaxTextureSamplersShaderMap();
+	uint32 MaxTextureSamplers = GetContent()->GetMaxTextureSamplersShaderMap(*this);
 
-	for (int32 Index = 0;Index < MeshShaderMaps.Num();Index++)
+	for (int32 Index = 0;Index < GetContent()->OrderedMeshShaderMaps.Num();Index++)
 	{
-		MaxTextureSamplers = FMath::Max(MaxTextureSamplers, MeshShaderMaps[Index].GetMaxTextureSamplersShaderMap());
+		const FMeshMaterialShaderMap* MeshShaderMap = GetContent()->OrderedMeshShaderMaps[Index];
+		if (MeshShaderMap)
+		{
+			MaxTextureSamplers = FMath::Max(MaxTextureSamplers, MeshShaderMap->GetMaxTextureSamplersShaderMap(*this));
+		}
 	}
 
 	return MaxTextureSamplers;
 }
 #endif // WITH_EDITOR
 
-const FMeshMaterialShaderMap* FMaterialShaderMap::GetMeshShaderMap(FVertexFactoryType* VertexFactoryType) const
+FMeshMaterialShaderMap* FMaterialShaderMapContent::GetMeshShaderMap(const FHashedName& VertexFactoryTypeName) const
 {
-	checkSlow(bCompilationFinalized);
-#if WITH_EDITOR 
-	// Attempt to get some more info for a rare crash (UE-35937)
-	checkf(OrderedMeshShaderMaps.Num() > 0 && bCompilationFinalized, TEXT("OrderedMeshShaderMaps.Num() is %d. bCompilationFinalized is %d. This may relate to bug UE-35937"), OrderedMeshShaderMaps.Num(), (int)bCompilationFinalized);
-#endif
-	const FMeshMaterialShaderMap* MeshShaderMap = OrderedMeshShaderMaps[VertexFactoryType->GetId()];
-	checkSlow(!MeshShaderMap || MeshShaderMap->GetVertexFactoryType() == VertexFactoryType);
-	return MeshShaderMap;
+	const int32 Index = Algo::BinarySearchBy(OrderedMeshShaderMaps, VertexFactoryTypeName, FProjectMeshShaderMapToKey());
+	if (Index != INDEX_NONE)
+	{
+		FMeshMaterialShaderMap* MeshShaderMap = OrderedMeshShaderMaps[Index].Get();
+		checkSlow(MeshShaderMap->GetVertexFactoryTypeName() == VertexFactoryTypeName);
+		return MeshShaderMap;
+	}
+	return nullptr;
 }
 
-void FMaterialShaderMap::InitOrderedMeshShaderMaps()
+void FMaterialShaderMapContent::AddMeshShaderMap(const FVertexFactoryType* VertexFactoryType, FMeshMaterialShaderMap* MeshShaderMap)
 {
-	OrderedMeshShaderMaps.Empty(FVertexFactoryType::GetNumVertexFactoryTypes());
-	OrderedMeshShaderMaps.AddZeroed(FVertexFactoryType::GetNumVertexFactoryTypes());
+	check(VertexFactoryType->GetHashedName() == MeshShaderMap->GetVertexFactoryTypeName());
+	checkSlow(GetMeshShaderMap(VertexFactoryType->GetHashedName()) == nullptr);
+	const int32 Index = Algo::LowerBoundBy(OrderedMeshShaderMaps, VertexFactoryType->GetHashedName(), FProjectMeshShaderMapToKey());
+	OrderedMeshShaderMaps.Insert(MeshShaderMap, Index);
+}
 
-	for (int32 Index = 0;Index < MeshShaderMaps.Num();Index++)
+void FMaterialShaderMapContent::RemoveMeshShaderMap(const FVertexFactoryType* VertexFactoryType)
+{
+	const int32 Index = Algo::BinarySearchBy(OrderedMeshShaderMaps, VertexFactoryType->GetHashedName(), FProjectMeshShaderMapToKey());
+	if (Index != INDEX_NONE)
 	{
-		check(MeshShaderMaps[Index].GetVertexFactoryType());
-		const int32 VFIndex = MeshShaderMaps[Index].GetVertexFactoryType()->GetId();
-		OrderedMeshShaderMaps[VFIndex] = &MeshShaderMaps[Index];
+		FMeshMaterialShaderMap* MeshShaderMap = OrderedMeshShaderMaps[Index].Get();
+		delete MeshShaderMap;
+		OrderedMeshShaderMaps.RemoveAt(Index);
 	}
 }
 
@@ -2643,41 +2598,34 @@ void FMaterialShaderMap::DumpDebugInfo()
 	const FString& DebugDescriptionS = GetDebugDescription();
 	UE_LOG(LogConsoleResponse, Display, TEXT("  DebugDescription %s"), *DebugDescriptionS);
 
-	TMap<FShaderId, FShader*> ShadersL;
+	TMap<FShaderId, TShaderRef<FShader>> ShadersL;
 	GetShaderList(ShadersL);
 	UE_LOG(LogConsoleResponse, Display, TEXT("  --- %d shaders"), ShadersL.Num());
 	int32 Index = 0;
 	for (auto& KeyValue : ShadersL)
 	{
 		UE_LOG(LogConsoleResponse, Display, TEXT("    --- shader %d"), Index);
-		FShader* Shader = KeyValue.Value;
-		Shader->DumpDebugInfo();
+		FShader* Shader = KeyValue.Value.GetShader();
+		Shader->DumpDebugInfo(GetPointerTable());
 		Index++;
 	}
 	//TArray<FShaderPipeline*> ShaderPipelinesL;
 	//GetShaderPipelineList(ShaderPipelinesL);
 }
 
+#if WITH_EDITOR
 void FMaterialShaderMap::SaveShaderStableKeys(EShaderPlatform TargetShaderPlatform, const FStableShaderKeyAndValue& SaveKeyVal)
 {
-#if WITH_EDITOR
-	TMap<FShaderId, FShader*> ShadersL;
-	GetShaderList(ShadersL);
-	for (auto& KeyValue : ShadersL)
-	{
-		FShader* Shader = KeyValue.Value;
-		Shader->SaveShaderStableKeys(TargetShaderPlatform, SaveKeyVal);
-	}
-
-	TArray<FShaderPipeline*> ShadersPipelinesL;
-	GetShaderPipelineList(ShadersPipelinesL);
-	for (FShaderPipeline* Pipeline : ShadersPipelinesL)
+	FShaderMapBase::SaveShaderStableKeys(TargetShaderPlatform, SaveKeyVal);
+	for (FMeshMaterialShaderMap* MeshShaderMap : GetContent()->OrderedMeshShaderMaps)
 	{
-		Pipeline->SaveShaderStableKeys(TargetShaderPlatform, SaveKeyVal);
+		if (MeshShaderMap)
+		{
+			MeshShaderMap->SaveShaderStableKeys(*this, TargetShaderPlatform, SaveKeyVal);
+		}
 	}
-#endif
 }
-
+#endif // WITH_EDITOR
 
 /**
  * Dump material stats for a given platform.
@@ -2692,7 +2640,7 @@ void DumpMaterialStats(EShaderPlatform Platform)
 	//#todo-rco: Pipelines
 
 	// Mapping from friendly material name to shaders associated with it.
-	TMultiMap<FString,FShader*> MaterialToShaderMap;
+	TMultiMap<FString, TShaderRef<FShader>> MaterialToShaderMap;
 	TMultiMap<FString, FShaderPipeline*> MaterialToShaderPipelineMap;
 
 	// Set of material names.
@@ -2702,8 +2650,8 @@ void DumpMaterialStats(EShaderPlatform Platform)
 	for (int32 ShaderMapIndex = 0; ShaderMapIndex < FMaterialShaderMap::AllMaterialShaderMaps.Num(); ShaderMapIndex++)
 	{
 		FMaterialShaderMap* MaterialShaderMap = FMaterialShaderMap::AllMaterialShaderMaps[ShaderMapIndex];
-		TMap<FShaderId, FShader*> Shaders;
-		TArray<FShaderPipeline*> ShaderPipelines;
+		TMap<FShaderId, TShaderRef<FShader>> Shaders;
+		TArray<FShaderPipelineRef> ShaderPipelines;
 		MaterialShaderMap->GetShaderList(Shaders);
 		MaterialShaderMap->GetShaderPipelineList(ShaderPipelines);
 
@@ -2714,17 +2662,16 @@ void DumpMaterialStats(EShaderPlatform Platform)
 		// Add shaders to mapping per friendly name as there might be multiple
 		for (auto& KeyValue : Shaders)
 		{
-			FShader* Shader = KeyValue.Value;
-			MaterialToShaderMap.AddUnique(FriendlyName, Shader);
+			MaterialToShaderMap.AddUnique(FriendlyName, KeyValue.Value);
 		}
 
-		for (FShaderPipeline* Pipeline : ShaderPipelines)
+		for (const FShaderPipelineRef& Pipeline : ShaderPipelines)
 		{
-			for (FShader* Shader : Pipeline->GetShaders())
+			for (const TShaderRef<FShader>& Shader : Pipeline.GetShaders())
 			{
 				MaterialToShaderMap.AddUnique(FriendlyName, Shader);
 			}
-			MaterialToShaderPipelineMap.AddUnique(FriendlyName, Pipeline);
+			MaterialToShaderPipelineMap.AddUnique(FriendlyName, Pipeline.GetPipeline());
 		}
 	}
 
@@ -2742,7 +2689,7 @@ void DumpMaterialStats(EShaderPlatform Platform)
 	for( TSet<FString>::TConstIterator It(MaterialNames); It; ++It )
 	{
 		// Retrieve list of shaders in map.
-		TArray<FShader*> Shaders;
+		TArray<TShaderRef<FShader>> Shaders;
 		MaterialToShaderMap.MultiFind( *It, Shaders );
 		TArray<FShaderPipeline*> ShaderPipelines;
 		MaterialToShaderPipelineMap.MultiFind(*It, ShaderPipelines);
@@ -2751,8 +2698,8 @@ void DumpMaterialStats(EShaderPlatform Platform)
 		int32 CodeSize = 0;
 		for( int32 ShaderIndex=0; ShaderIndex<Shaders.Num(); ShaderIndex++ )
 		{
-			FShader* Shader = Shaders[ShaderIndex];
-			CodeSize += Shader->GetCode().Num();
+			const TShaderRef<FShader>& Shader = Shaders[ShaderIndex];
+			CodeSize += Shader->GetCodeSize();
 		}
 
 		TotalCodeSize += CodeSize;
