Copying Private-LoadTimes-4.24 stream to Main. Biggest changes are in Materials/Shader memory freezing.

diff --git a/Engine/Source/Runtime/Engine/Private/Materials/MaterialShared.cpp b/Engine/Source/Runtime/Engine/Private/Materials/MaterialShared.cpp
--- a/Engine/Source/Runtime/Engine/Private/Materials/MaterialShared.cpp
+++ b/Engine/Source/Runtime/Engine/Private/Materials/MaterialShared.cpp
@@ -17,12 +17,12 @@
 #include "Materials/MaterialInstanceBasePropertyOverrides.h"
 #include "Materials/MaterialInstance.h"
 #include "Materials/MaterialUniformExpressions.h"
+#include "Materials/MaterialShaderMapLayout.h"
 #include "UObject/UObjectIterator.h"
 #include "ComponentReregisterContext.h"
 #include "Materials/MaterialExpressionBreakMaterialAttributes.h"
 #include "Materials/MaterialExpressionShadingModel.h"
 #include "Materials/MaterialExpressionReroute.h"
-#include "Materials/MaterialExpressionRuntimeVirtualTextureOutput.h"
 #include "Materials/MaterialExpressionSingleLayerWaterMaterialOutput.h"
 #include "ShaderCompiler.h"
 #include "MaterialCompiler.h"
@@ -47,6 +47,20 @@
 
 DEFINE_LOG_CATEGORY(LogMaterial);
 
+IMPLEMENT_TYPE_LAYOUT(FHashedMaterialParameterInfo);
+IMPLEMENT_TYPE_LAYOUT(FUniformExpressionSet);
+IMPLEMENT_TYPE_LAYOUT(FMaterialCompilationOutput);
+IMPLEMENT_TYPE_LAYOUT(FMeshMaterialShaderMap);
+IMPLEMENT_TYPE_LAYOUT(FMaterialProcessedSource);
+IMPLEMENT_TYPE_LAYOUT(FMaterialShaderMapContent);
+IMPLEMENT_TYPE_LAYOUT(FMaterialPreshaderData);
+IMPLEMENT_TYPE_LAYOUT(FMaterialUniformPreshaderHeader);
+IMPLEMENT_TYPE_LAYOUT(FMaterialScalarParameterInfo);
+IMPLEMENT_TYPE_LAYOUT(FMaterialVectorParameterInfo);
+IMPLEMENT_TYPE_LAYOUT(FMaterialTextureParameterInfo);
+IMPLEMENT_TYPE_LAYOUT(FMaterialExternalTextureParameterInfo);
+IMPLEMENT_TYPE_LAYOUT(FMaterialVirtualTextureStack);
+
 int32 GDeferUniformExpressionCaching = 1;
 FAutoConsoleVariableRef CVarDeferUniformExpressionCaching(
 	TEXT("r.DeferUniformExpressionCaching"),
@@ -194,7 +208,7 @@ bool ReloadMaterialResource(
 	Tmp.SerializeInlineShaderMap(Ar);
 	if (Tmp.GetGameThreadShaderMap())
 	{
-		Tmp.GetGameThreadShaderMap()->RegisterSerializedShaders(false);
+		//Tmp.GetGameThreadShaderMap()->RegisterSerializedShaders(false);
 		return true;
 	}
 	UE_LOG(LogMaterial, Warning, TEXT("Failed to reload material resources for package %s (file name: %s)."), *PackageName, *Filename);
@@ -202,17 +216,20 @@ bool ReloadMaterialResource(
 }
 #endif
 
-static inline SIZE_T AddShaderSize(FShader* Shader, TSet<FShaderResourceId>& UniqueShaderResourceIds)
+static inline SIZE_T AddShaderSize(const TShaderRef<FShader>& Shader, TSet<FShaderResourceId>& UniqueShaderResourceIds)
 {
 	SIZE_T ResourceSize = 0;
-	FShaderResourceId ResourceId = Shader->GetResourceId();
+	// @todo fnhack
+//	check(false);
+#if 0
+	FShaderResourceId ResourceId = Shader.GetResourceId();
 	bool bCountedResource = false;
 	UniqueShaderResourceIds.Add(ResourceId, &bCountedResource);
 	if (!bCountedResource)
 	{
-		ResourceSize += Shader->GetResourceSizeBytes();
+		ResourceSize += Shader.GetResource()->GetSizeBytes();
 	}
-	ResourceSize += Shader->GetSizeBytes();
+#endif
 	return ResourceSize;
 }
 
@@ -511,51 +528,6 @@ int32 FMaterialAttributesInput::CompileWithDefault(class FMaterialCompiler* Comp
 }
 #endif  // WITH_EDITOR
 
-void FMaterialCompilationOutput::Serialize(FArchive& Ar)
-{
-	UniformExpressionSet.Serialize(Ar);
-
-	Ar << UsedSceneTextures;
-
-#if WITH_EDITOR
-	Ar << EstimatedNumTextureSamplesVS;
-	Ar << EstimatedNumTextureSamplesPS;
-	Ar << EstimatedNumVirtualTextureLookups;
-	Ar << NumUsedUVScalars;
-	Ar << NumUsedCustomInterpolatorScalars;
-#else
-	uint16 Tmp16 = 0;
-	uint8 Tmp8 = 0; 
-
-	Ar << Tmp16;
-	Ar << Tmp16;
-	Ar << Tmp16;
-	Ar << Tmp8;
-	Ar << Tmp8;
-#endif
-
-	Ar << RuntimeVirtualTextureOutputAttributeMask;
-
-	uint8 PackedFlags = 0;
-	PackedFlags |= (bUsesEyeAdaptation << 0);
-	PackedFlags |= (bModifiesMeshPosition << 1);
-	PackedFlags |= (bUsesWorldPositionOffset << 2);
-	PackedFlags |= (bUsesGlobalDistanceField << 3);
-	PackedFlags |= (bUsesPixelDepthOffset << 4);
-	PackedFlags |= (bUsesDistanceCullFade << 5);
-	PackedFlags |= (bHasRuntimeVirtualTextureOutputNode << 6);
-
-	Ar << PackedFlags;
-
-	bUsesEyeAdaptation = (PackedFlags >> 0) & 1;
-	bModifiesMeshPosition = (PackedFlags >> 1) & 1;
-	bUsesWorldPositionOffset = (PackedFlags >> 2) & 1;
-	bUsesGlobalDistanceField = (PackedFlags >> 3) & 1;
-	bUsesPixelDepthOffset = (PackedFlags >> 4) & 1;
-	bUsesDistanceCullFade = (PackedFlags >> 5) & 1;
-	bHasRuntimeVirtualTextureOutputNode = (PackedFlags >> 6) & 1;
-}
-
 void FMaterial::GetShaderMapId(EShaderPlatform Platform, FMaterialShaderMapId& OutId) const
 { 
 	if (bLoadedCookedShaderMapId)
@@ -597,13 +569,13 @@ void FMaterial::GetShaderMapId(EShaderPlatform Platform, FMaterialShaderMapId& O
 	}
 }
 
-
+#if WITH_EDITORONLY_DATA
 void FMaterial::GetStaticParameterSet(EShaderPlatform Platform, FStaticParameterSet& OutSet) const
 {
 	// Clear the set in default implementation
 	OutSet = FStaticParameterSet();
 }
-
+#endif // WITH_EDITORONLY_DATA
 
 EMaterialTessellationMode FMaterial::GetTessellationMode() const 
 { 
@@ -697,95 +669,31 @@ const FMaterialShaderMap* FMaterial::GetShaderMapToUse() const
 	return ShaderMapToUse;
 }
 
-const TArray<TRefCountPtr<FMaterialUniformExpressionTexture> >& FMaterial::GetUniform2DTextureExpressions() const 
-{ 
-	const FMaterialShaderMap* ShaderMapToUse = GetShaderMapToUse();
-
-	if (ShaderMapToUse)
-	{
-		return ShaderMapToUse->GetUniformExpressionSet().Uniform2DTextureExpressions; 
-	}
-	
-	static const TArray<TRefCountPtr<FMaterialUniformExpressionTexture> > EmptyExpressions;
-	return EmptyExpressions;
-}
-
-const TArray<TRefCountPtr<FMaterialUniformExpressionTexture> >& FMaterial::GetUniformCubeTextureExpressions() const 
-{ 
-	const FMaterialShaderMap* ShaderMapToUse = GetShaderMapToUse();
-
-	if (ShaderMapToUse)
-	{
-		return ShaderMapToUse->GetUniformExpressionSet().UniformCubeTextureExpressions; 
-	}
-
-	static const TArray<TRefCountPtr<FMaterialUniformExpressionTexture> > EmptyExpressions;
-	return EmptyExpressions;
-}
-
-const TArray<TRefCountPtr<FMaterialUniformExpressionTexture> >& FMaterial::GetUniform2DArrayTextureExpressions() const
-{
-	const FMaterialShaderMap* ShaderMapToUse = GetShaderMapToUse();
-
-	if (ShaderMapToUse)
-	{
-		return ShaderMapToUse->GetUniformExpressionSet().Uniform2DArrayTextureExpressions;
-	}
-
-	static const TArray<TRefCountPtr<FMaterialUniformExpressionTexture> > EmptyExpressions;
-	return EmptyExpressions;
-}
-
-const TArray<TRefCountPtr<FMaterialUniformExpressionTexture> >& FMaterial::GetUniformVolumeTextureExpressions() const 
+const FUniformExpressionSet& FMaterial::GetUniformExpressions() const
 { 
 	const FMaterialShaderMap* ShaderMapToUse = GetShaderMapToUse();
-
 	if (ShaderMapToUse)
 	{
-		return ShaderMapToUse->GetUniformExpressionSet().UniformVolumeTextureExpressions; 
+		return ShaderMapToUse->GetUniformExpressionSet();
 	}
 
-	static const TArray<TRefCountPtr<FMaterialUniformExpressionTexture> > EmptyExpressions;
+	static const FUniformExpressionSet EmptyExpressions;
 	return EmptyExpressions;
 }
 
-const TArray<TRefCountPtr<FMaterialUniformExpressionTexture> >& FMaterial::GetUniformVirtualTextureExpressions() const
+const TArray<FMaterialTextureParameterInfo, FMemoryImageAllocator>& FMaterial::GetUniformTextureExpressions(EMaterialTextureParameterType Type) const
 {
-	const FMaterialShaderMap* ShaderMapToUse = GetShaderMapToUse();
-
-	if (ShaderMapToUse)
-	{
-		return ShaderMapToUse->GetUniformExpressionSet().UniformVirtualTextureExpressions;
-	}
-
-	static const TArray<TRefCountPtr<FMaterialUniformExpressionTexture> > EmptyExpressions;
-	return EmptyExpressions;
+	return GetUniformExpressions().UniformTextureParameters[(uint32)Type];
 }
 
-const TArray<TRefCountPtr<FMaterialUniformExpression> >& FMaterial::GetUniformVectorParameterExpressions() const 
+const TArray<FMaterialVectorParameterInfo, FMemoryImageAllocator>& FMaterial::GetUniformVectorParameterExpressions() const
 { 
-	const FMaterialShaderMap* ShaderMapToUse = GetShaderMapToUse();
-
-	if (ShaderMapToUse)
-	{
-		return ShaderMapToUse->GetUniformExpressionSet().UniformVectorExpressions; 
-	}
-
-	static const TArray<TRefCountPtr<FMaterialUniformExpression> > EmptyExpressions;
-	return EmptyExpressions;
+	return GetUniformExpressions().UniformVectorParameters;
 }
 
-const TArray<TRefCountPtr<FMaterialUniformExpression> >& FMaterial::GetUniformScalarParameterExpressions() const 
+const TArray<FMaterialScalarParameterInfo, FMemoryImageAllocator>& FMaterial::GetUniformScalarParameterExpressions() const
 { 
-	const FMaterialShaderMap* ShaderMapToUse = GetShaderMapToUse();
-
-	if (ShaderMapToUse)
-	{
-		return ShaderMapToUse->GetUniformExpressionSet().UniformScalarExpressions; 
-	}
-
-	static const TArray<TRefCountPtr<FMaterialUniformExpression> > EmptyExpressions;
-	return EmptyExpressions;
+	return GetUniformExpressions().UniformScalarParameters;
 }
 
 bool FMaterial::RequiresSceneColorCopy_GameThread() const
@@ -1047,7 +955,7 @@ void FMaterial::RegisterInlineShaderMap(bool bLoadedByCookedMaterial)
 		{
 			RenderingThreadShaderMap = GameThreadShaderMap;
 		}
-		GameThreadShaderMap->RegisterSerializedShaders(bLoadedByCookedMaterial);
+		//GameThreadShaderMap->RegisterSerializedShaders(bLoadedByCookedMaterial);
 	}
 }
 
@@ -1065,18 +973,23 @@ void FMaterialResource::LegacySerialize(FArchive& Ar)
 	}
 }
 
-const TArray<UObject*>& FMaterialResource::GetReferencedTextures() const
+TArrayView<UObject* const> FMaterialResource::GetReferencedTextures() const
 {
-	if (MaterialInstance && MaterialInstance->PermutationTextureReferences.Num())
+	if (MaterialInstance)
 	{
-		return MaterialInstance->PermutationTextureReferences;
+		const TArrayView<UObject* const> Textures = MaterialInstance->GetReferencedTextures();
+		if (Textures.Num())
+		{
+			return Textures;
+		}
 	}
-	else if (Material)
+	
+	if (Material)
 	{
-		return Material->ExpressionTextureReferences;
+		return Material->GetReferencedTextures();
 	}
 
-	return UMaterial::GetDefaultMaterial(MD_Surface)->ExpressionTextureReferences;
+	return UMaterial::GetDefaultMaterial(MD_Surface)->GetReferencedTextures();
 }
 
 void FMaterialResource::AddReferencedObjects(FReferenceCollector& Collector)
@@ -1112,7 +1025,7 @@ void FMaterial::DiscardShaderMap()
 	check(RenderingThreadShaderMap == nullptr);
 	if (GameThreadShaderMap)
 	{
-		GameThreadShaderMap->DiscardSerializedShaders();
+		//GameThreadShaderMap->DiscardSerializedShaders();
 		GameThreadShaderMap = nullptr;
 	}
 }
@@ -1485,15 +1398,7 @@ uint32 FMaterialResource::GetStencilCompare() const
 
 bool FMaterialResource::HasRuntimeVirtualTextureOutput() const
 {
-	// Slow check. Used for ShouldCompilePermutation() calls. 
-	for (UMaterialExpression* Expression : Material->Expressions)
-	{
-		if (Expression->IsA(UMaterialExpressionRuntimeVirtualTextureOutput::StaticClass()))
-		{
-			return true;
-		}
-	}
-	return false;
+	return Material->GetCachedExpressionData().bHasRuntimeVirtualTextureOutput;
 }
 
 bool FMaterialResource::CastsRayTracedShadows() const
@@ -1516,11 +1421,10 @@ void FMaterialResource::NotifyCompilationFinished()
 void FMaterialResource::GetResourceSizeEx(FResourceSizeEx& CumulativeResourceSize)
 {
 	TSet<const FMaterialShaderMap*> UniqueShaderMaps;
-	TMap<FShaderId, FShader*> UniqueShaders;
-	TArray<FShaderPipeline*> ShaderPipelines;
+	TMap<FShaderId, TShaderRef<FShader>> UniqueShaders;
+	TArray<FShaderPipelineRef> ShaderPipelines;
 	TSet<FShaderResourceId> UniqueShaderResourceIds;
 
-	CumulativeResourceSize.AddDedicatedSystemMemoryBytes(sizeof(FMaterialResource));
 	UniqueShaderMaps.Add(GetGameThreadShaderMap());
 
 	for (TSet<const FMaterialShaderMap*>::TConstIterator It(UniqueShaderMaps); It; ++It)
@@ -1528,7 +1432,7 @@ void FMaterialResource::GetResourceSizeEx(FResourceSizeEx& CumulativeResourceSiz
 		const FMaterialShaderMap* MaterialShaderMap = *It;
 		if (MaterialShaderMap)
 		{
-			CumulativeResourceSize.AddDedicatedSystemMemoryBytes(MaterialShaderMap->GetSizeBytes());
+			CumulativeResourceSize.AddDedicatedSystemMemoryBytes(MaterialShaderMap->GetFrozenContentSize());
 
 			// Shaders are shared, so only count them in total mode
 			if (CumulativeResourceSize.GetResourceSizeMode() == EResourceSizeMode::EstimatedTotal)
@@ -1541,25 +1445,24 @@ void FMaterialResource::GetResourceSizeEx(FResourceSizeEx& CumulativeResourceSiz
 
 	for (auto& KeyValue : UniqueShaders)
 	{
-		FShader* Shader = KeyValue.Value;
-		if (Shader)
+		const TShaderRef<FShader>& Shader = KeyValue.Value;
+		if (Shader.IsValid())
 		{
 			CumulativeResourceSize.AddDedicatedSystemMemoryBytes(AddShaderSize(Shader, UniqueShaderResourceIds));
 		}
 	}
 
-	for (FShaderPipeline* Pipeline : ShaderPipelines)
+	for (const FShaderPipelineRef& Pipeline : ShaderPipelines)
 	{
-		if (Pipeline)
+		if (Pipeline.IsValid())
 		{
-			for (FShader* Shader : Pipeline->GetShaders())
+			for (const TShaderRef<FShader>& Shader : Pipeline.GetShaders())
 			{
-				if (Shader)
+				if (Shader.IsValid())
 				{
 					CumulativeResourceSize.AddDedicatedSystemMemoryBytes(AddShaderSize(Shader, UniqueShaderResourceIds));
 				}
 			}
-			CumulativeResourceSize.AddDedicatedSystemMemoryBytes(Pipeline->GetSizeBytes());
 		}
 	}	
 }
@@ -1582,20 +1485,26 @@ FMaterial::~FMaterial()
 #endif // WITH_EDITOR
 
 	FMaterialShaderMap::RemovePendingMaterial(this);
+
+	if (GameThreadShaderMap.GetRefCount() > 1)
+	{
+		int a = 0;
+	}
 }
 
 /** Populates OutEnvironment with defines needed to compile shaders for this material. */
 void FMaterial::SetupMaterialEnvironment(
 	EShaderPlatform Platform,
+	const FShaderParametersMetadata& InUniformBufferStruct,
 	const FUniformExpressionSet& InUniformExpressionSet,
 	FShaderCompilerEnvironment& OutEnvironment
 	) const
 {
 	// Add the material uniform buffer definition.
-	FShaderUniformBufferParameter::ModifyCompilationEnvironment(TEXT("Material"),InUniformExpressionSet.GetUniformBufferStruct(),Platform,OutEnvironment);
+	FShaderUniformBufferParameter::ModifyCompilationEnvironment(TEXT("Material"), InUniformBufferStruct, Platform, OutEnvironment);
 
 	// Mark as using external texture if uniform expression contains external texture
-	if (InUniformExpressionSet.UniformExternalTextureExpressions.Num() > 0)
+	if (InUniformExpressionSet.UniformExternalTextureParameters.Num() > 0)
 	{
 		OutEnvironment.CompilerFlags.Add(CFLAG_UsesExternalTexture);
 	}
@@ -1811,22 +1720,14 @@ bool FMaterial::CacheShaders(EShaderPlatform Platform, const ITargetPlatform* Ta
 	FAllowCachingStaticParameterValues AllowCachingStaticParameterValues(*this);
 	FMaterialShaderMapId NoStaticParametersId;
 	GetShaderMapId(Platform, NoStaticParametersId);
-	
-	if (FPlatformProperties::RequiresCookedData())
-	{ 
-		return CacheShaders(NoStaticParametersId, nullptr, Platform, TargetPlatform);
-	}
-
-	FStaticParameterSet StaticParameterSet;
-	GetStaticParameterSet(Platform, StaticParameterSet);
-	return CacheShaders(NoStaticParametersId, &StaticParameterSet, Platform, TargetPlatform);
+	return CacheShaders(NoStaticParametersId, Platform, TargetPlatform);
 }
 
 /**
  * Caches the material shaders for the given static parameter set and platform.
  * This is used by material resources of UMaterialInstances.
  */
-bool FMaterial::CacheShaders(const FMaterialShaderMapId& ShaderMapId, const FStaticParameterSet *StaticParameterSet, EShaderPlatform Platform, const ITargetPlatform* TargetPlatform)
+bool FMaterial::CacheShaders(const FMaterialShaderMapId& ShaderMapId, EShaderPlatform Platform, const ITargetPlatform* TargetPlatform)
 {
 	bool bSucceeded = false;
 	UE_CLOG(!ShaderMapId.IsValid(), LogMaterial, Warning, TEXT("Invalid shader map ID caching shaders for '%s', will use default material."), *GetFriendlyName());
@@ -1868,7 +1769,7 @@ bool FMaterial::CacheShaders(const FMaterialShaderMapId& ShaderMapId, const FSta
 			if (!ShaderMap && !FPlatformProperties::RequiresCookedData()) 
 			{
 				TRefCountPtr<FMaterialShaderMap> LoadedShaderMap;
-				FMaterialShaderMap::LoadFromDerivedDataCache(this, ShaderMapId, Platform, LoadedShaderMap);
+				FMaterialShaderMap::LoadFromDerivedDataCache(this, ShaderMapId, Platform, TargetPlatform, LoadedShaderMap);
 				ShaderMap = LoadedShaderMap;
 			}
 		}
@@ -1948,11 +1849,14 @@ bool FMaterial::CacheShaders(const FMaterialShaderMapId& ShaderMapId, const FSta
 			UE_LOG(LogMaterial, Display, TEXT("%s cached shader map for material %s, compiling. %s"),ShaderMapCondition,*GetFriendlyName(), IsSpecialEngineMaterial() ? TEXT("Is special engine material.") : TEXT("") );
 
 			TRefCountPtr<FMaterialShaderMap> ShaderMap;
+#if WITH_EDITORONLY_DATA
+			FStaticParameterSet StaticParameterSet;
+			GetStaticParameterSet(Platform, StaticParameterSet);
 
 			// If there's no cached shader map for this material, compile a new one.
 			// This is just kicking off the async compile, GameThreadShaderMap will not be complete yet
-			check(StaticParameterSet != nullptr);
-			bSucceeded = BeginCompileShaderMap(ShaderMapId, *StaticParameterSet, Platform, ShaderMap, TargetPlatform);
+			bSucceeded = BeginCompileShaderMap(ShaderMapId, StaticParameterSet, Platform, ShaderMap, TargetPlatform);
+#endif // WITH_EDITORONLY_DATA
 
 			if (!bSucceeded)
 			{
@@ -2017,7 +1921,6 @@ bool FMaterial::BeginCompileShaderMap(
 
 	SCOPE_SECONDS_COUNTER(MaterialCompileTime);
 
-	NewShaderMap->SetMaterialPath(GetBaseMaterialPathName());
 	// Generate the material shader code.
 	FMaterialCompilationOutput NewCompilationOutput;
 	FHLSLMaterialTranslator MaterialTranslator(this, NewCompilationOutput, StaticParameterSet, Platform,GetQualityLevel(), ShaderMapId.FeatureLevel, TargetPlatform);
@@ -2078,6 +1981,20 @@ bool FMaterial::ShouldCache(EShaderPlatform Platform, const FShaderType* ShaderT
 	return true;
 }
 
+bool FMaterial::ShouldCachePipeline(EShaderPlatform Platform, const FShaderPipelineType* PipelineType, const FVertexFactoryType* VertexFactoryType) const
+{
+	for (const FShaderType* ShaderType : PipelineType->GetStages())
+	{
+		if (!ShouldCache(Platform, ShaderType, VertexFactoryType))
+		{
+			return false;
+		}
+	}
+
+	// Only include the pipeline if all shaders should be cached
+	return true;
+}
+
 //
 // FColoredMaterialRenderProxy implementation.
 //
@@ -2090,8 +2007,9 @@ const FMaterial& FColoredMaterialRenderProxy::GetMaterialWithFallback(ERHIFeatur
 /**
  * Finds the shader matching the template type and the passed in vertex factory, asserts if not found.
  */
-FShader* FMaterial::GetShader(FMeshMaterialShaderType* ShaderType, FVertexFactoryType* VertexFactoryType, int32 PermutationId, bool bFatalIfMissing) const
+TShaderRef<FShader> FMaterial::GetShader(FMeshMaterialShaderType* ShaderType, FVertexFactoryType* VertexFactoryType, int32 PermutationId, bool bFatalIfMissing) const
 {
+	TRACE_CPUPROFILER_EVENT_SCOPE(FMaterial::GetShader);
 #if WITH_EDITOR && DO_CHECK
 	// Attempt to get some more info for a rare crash (UE-35937)
 	FMaterialShaderMap* GameThreadShaderMapPtr = GameThreadShaderMap;
@@ -2101,120 +2019,115 @@ FShader* FMaterial::GetShader(FMeshMaterialShaderType* ShaderType, FVertexFactor
 	FShader* Shader = MeshShaderMap ? MeshShaderMap->GetShader(ShaderType, PermutationId) : nullptr;
 	if (!Shader)
 	{
-		// we don't care about thread safety because we are about to crash 
-		const auto CachedGameThreadShaderMap = GameThreadShaderMap;
-		const auto CachedGameMeshShaderMap = CachedGameThreadShaderMap ? CachedGameThreadShaderMap->GetMeshShaderMap(VertexFactoryType) : nullptr;
-		bool bShaderWasFoundInGameShaderMap = CachedGameMeshShaderMap && CachedGameMeshShaderMap->GetShader(ShaderType, PermutationId) != nullptr;
-
-		// Get the ShouldCache results that determine whether the shader should be compiled
-		auto ShaderPlatform = GShaderPlatformForFeatureLevel[GetFeatureLevel()];
-		bool bMaterialShouldCache = ShouldCache(ShaderPlatform, ShaderType, VertexFactoryType);
-		bool bVFShouldCache = VertexFactoryType->ShouldCache(ShaderPlatform, this, ShaderType);
-		bool bShaderShouldCache = ShaderType->ShouldCompilePermutation(ShaderPlatform, this, VertexFactoryType, PermutationId);
-		FString MaterialUsage = GetMaterialUsageDescription();
-
-		int BreakPoint = 0;
-
-		// Assert with detailed information if the shader wasn't found for rendering.  
-		// This is usually the result of an incorrect ShouldCache function.
-		UE_LOG(LogMaterial, Error,
-			TEXT("Couldn't find Shader (%s, %d) for Material Resource %s!\n")
-			TEXT("		RenderMeshShaderMap %d, RenderThreadShaderMap %d\n")
-			TEXT("		GameMeshShaderMap %d, GameThreadShaderMap %d, bShaderWasFoundInGameShaderMap %d\n")
-			TEXT("		With VF=%s, Platform=%s\n")
-			TEXT("		ShouldCache: Mat=%u, VF=%u, Shader=%u \n")
-			TEXT("		MaterialUsageDesc: %s"),
-			ShaderType->GetName(), PermutationId, *GetFriendlyName(),
-			MeshShaderMap != nullptr, RenderingThreadShaderMap != nullptr,
-			CachedGameMeshShaderMap != nullptr, CachedGameThreadShaderMap != nullptr, bShaderWasFoundInGameShaderMap,
-			VertexFactoryType->GetName(), *LegacyShaderPlatformToShaderFormat(ShaderPlatform).ToString(),
-			bMaterialShouldCache, bVFShouldCache, bShaderShouldCache,
-			*MaterialUsage
+		auto noinline_lambda = [&](...) FORCENOINLINE
+		{
+			// we don't care about thread safety because we are about to crash 
+			const auto CachedGameThreadShaderMap = GameThreadShaderMap;
+			const auto CachedGameMeshShaderMap = CachedGameThreadShaderMap ? CachedGameThreadShaderMap->GetMeshShaderMap(VertexFactoryType) : nullptr;
+			bool bShaderWasFoundInGameShaderMap = CachedGameMeshShaderMap && CachedGameMeshShaderMap->GetShader(ShaderType, PermutationId) != nullptr;
+
+			// Get the ShouldCache results that determine whether the shader should be compiled
+			auto ShaderPlatform = GShaderPlatformForFeatureLevel[GetFeatureLevel()];
+			bool bMaterialShouldCache = ShouldCache(ShaderPlatform, ShaderType, VertexFactoryType);
+			bool bVFShouldCache = FMeshMaterialShaderType::ShouldCompileVertexFactoryPermutation(VertexFactoryType, ShaderPlatform, this);
+			bool bShaderShouldCache = ShaderType->ShouldCompilePermutation(ShaderPlatform, this, VertexFactoryType, PermutationId);
+			FString MaterialUsage = GetMaterialUsageDescription();
+
+			int BreakPoint = 0;
+
+			// Assert with detailed information if the shader wasn't found for rendering.  
+			// This is usually the result of an incorrect ShouldCache function.
+			UE_LOG(LogMaterial, Error,
+				TEXT("Couldn't find Shader (%s, %d) for Material Resource %s!\n")
+				TEXT("		RenderMeshShaderMap %d, RenderThreadShaderMap %d\n")
+				TEXT("		GameMeshShaderMap %d, GameThreadShaderMap %d, bShaderWasFoundInGameShaderMap %d\n")
+				TEXT("		With VF=%s, Platform=%s\n")
+				TEXT("		ShouldCache: Mat=%u, VF=%u, Shader=%u \n")
+				TEXT("		MaterialUsageDesc: %s"),
+				ShaderType->GetName(), PermutationId, *GetFriendlyName(),
+				MeshShaderMap != nullptr, RenderingThreadShaderMap != nullptr,
+				CachedGameMeshShaderMap != nullptr, CachedGameThreadShaderMap != nullptr, bShaderWasFoundInGameShaderMap,
+				VertexFactoryType->GetName(), *LegacyShaderPlatformToShaderFormat(ShaderPlatform).ToString(),
+				bMaterialShouldCache, bVFShouldCache, bShaderShouldCache,
+				*MaterialUsage
 			);
 
-		if (MeshShaderMap)
-		{
-			TMap<FShaderId, FShader*> List;
-			MeshShaderMap->GetShaderList(List);
-
-			for (const auto& ShaderPair : List)
+			if (MeshShaderMap)
 			{
-				FString TypeName = ShaderPair.Value->GetType()->GetName();
-				UE_LOG(LogMaterial, Error, TEXT("ShaderType found in MaterialMap: %s"), *TypeName);
+				TMap<FHashedName, TShaderRef<FShader>> List;
+				MeshShaderMap->GetShaderList(*RenderingThreadShaderMap, List);
+
+				for (const auto& ShaderPair : List)
+				{
+					FString TypeName = ShaderPair.Value.GetType()->GetName();
+					UE_LOG(LogMaterial, Error, TEXT("ShaderType found in MaterialMap: %s"), *TypeName);
+				}
 			}
-		}
 
-		UE_LOG(LogMaterial, Fatal, TEXT("Fatal Error Material not found"));
+			UE_LOG(LogMaterial, Fatal, TEXT("Fatal Error Material not found"));
+		};
+		noinline_lambda();
 	}
 
-	return Shader;
+	return TShaderRef<FShader>(Shader, *RenderingThreadShaderMap);
 }
 
-FShaderPipeline* FMaterial::GetShaderPipeline(class FShaderPipelineType* ShaderPipelineType, FVertexFactoryType* VertexFactoryType, bool bFatalIfNotFound) const
+FShaderPipelineRef FMaterial::GetShaderPipeline(class FShaderPipelineType* ShaderPipelineType, FVertexFactoryType* VertexFactoryType, bool bFatalIfNotFound) const
 {
 	const FMeshMaterialShaderMap* MeshShaderMap = RenderingThreadShaderMap->GetMeshShaderMap(VertexFactoryType);
 	FShaderPipeline* ShaderPipeline = MeshShaderMap ? MeshShaderMap->GetShaderPipeline(ShaderPipelineType) : nullptr;
 	if (!ShaderPipeline && bFatalIfNotFound)
 	{
-		// Get the ShouldCache results that determine whether the shader should be compiled
-		auto ShaderPlatform = GShaderPlatformForFeatureLevel[GetFeatureLevel()];
-		FString MaterialUsage = GetMaterialUsageDescription();
+		auto noinline_lambda = [&](...) FORCENOINLINE
+		{
+			// Get the ShouldCache results that determine whether the shader should be compiled
+			auto ShaderPlatform = GShaderPlatformForFeatureLevel[GetFeatureLevel()];
+			FString MaterialUsage = GetMaterialUsageDescription();
 
-		UE_LOG(LogMaterial, Error,
-			TEXT("Couldn't find ShaderPipeline %s for Material Resource %s!"), ShaderPipelineType->GetName(), *GetFriendlyName());
+			UE_LOG(LogMaterial, Error,
+				TEXT("Couldn't find ShaderPipeline %s for Material Resource %s!"), ShaderPipelineType->GetName(), *GetFriendlyName());
 
-		for (auto* ShaderType : ShaderPipelineType->GetStages())
-		{
-			FShader* Shader = MeshShaderMap ? MeshShaderMap->GetShader((FShaderType*)ShaderType) : RenderingThreadShaderMap->GetShader((FShaderType*)ShaderType);
-			if (!Shader)
-			{
-				UE_LOG(LogMaterial, Error, TEXT("Missing %s shader %s!"), GetShaderFrequencyString(ShaderType->GetFrequency(), false), ShaderType->GetName());
-			}
-			else if (ShaderType->GetMeshMaterialShaderType())
+			for (auto* ShaderType : ShaderPipelineType->GetStages())
 			{
-				bool bMaterialShouldCache = ShouldCache(ShaderPlatform, ShaderType->GetMeshMaterialShaderType(), VertexFactoryType);
-				bool bVFShouldCache = VertexFactoryType->ShouldCache(ShaderPlatform, this, ShaderType->GetMeshMaterialShaderType());
-				bool bShaderShouldCache = ShaderType->GetMeshMaterialShaderType()->ShouldCompilePermutation(ShaderPlatform, this, VertexFactoryType, kUniqueShaderPermutationId);
+				FShader* Shader = MeshShaderMap ? MeshShaderMap->GetShader((FShaderType*)ShaderType) : RenderingThreadShaderMap->GetShader((FShaderType*)ShaderType).GetShader();
+				if (!Shader)
+				{
+					UE_LOG(LogMaterial, Error, TEXT("Missing %s shader %s!"), GetShaderFrequencyString(ShaderType->GetFrequency(), false), ShaderType->GetName());
+				}
+				else if (ShaderType->GetMeshMaterialShaderType())
+				{
+					bool bMaterialShouldCache = ShouldCache(ShaderPlatform, ShaderType->GetMeshMaterialShaderType(), VertexFactoryType);
+					bool bVFShouldCache = FMeshMaterialShaderType::ShouldCompileVertexFactoryPermutation(VertexFactoryType, ShaderPlatform, this);
+					bool bShaderShouldCache = ShaderType->GetMeshMaterialShaderType()->ShouldCompilePermutation(ShaderPlatform, this, VertexFactoryType, kUniqueShaderPermutationId);
 
-				UE_LOG(LogMaterial, Error, TEXT("%s %s ShouldCache: Mat=%u, VF=%u, Shader=%u"),
-					GetShaderFrequencyString(ShaderType->GetFrequency(), false), ShaderType->GetName(), bMaterialShouldCache, bVFShouldCache, bShaderShouldCache);
-			}
-			else if (ShaderType->GetMaterialShaderType())
-			{
-				bool bMaterialShouldCache = ShouldCache(ShaderPlatform, ShaderType->GetMaterialShaderType(), VertexFactoryType);
-				bool bShaderShouldCache = ShaderType->GetMaterialShaderType()->ShouldCompilePermutation(ShaderPlatform, this, kUniqueShaderPermutationId);
+					UE_LOG(LogMaterial, Error, TEXT("%s %s ShouldCache: Mat=%u, VF=%u, Shader=%u"),
+						GetShaderFrequencyString(ShaderType->GetFrequency(), false), ShaderType->GetName(), bMaterialShouldCache, bVFShouldCache, bShaderShouldCache);
+				}
+				else if (ShaderType->GetMaterialShaderType())
+				{
+					bool bMaterialShouldCache = ShouldCache(ShaderPlatform, ShaderType->GetMaterialShaderType(), VertexFactoryType);
+					bool bShaderShouldCache = ShaderType->GetMaterialShaderType()->ShouldCompilePermutation(ShaderPlatform, this, kUniqueShaderPermutationId);
 
-				UE_LOG(LogMaterial, Error, TEXT("%s %s ShouldCache: Mat=%u, NO VF, Shader=%u"),
-					GetShaderFrequencyString(ShaderType->GetFrequency(), false), ShaderType->GetName(), bMaterialShouldCache, bShaderShouldCache);
+					UE_LOG(LogMaterial, Error, TEXT("%s %s ShouldCache: Mat=%u, NO VF, Shader=%u"),
+						GetShaderFrequencyString(ShaderType->GetFrequency(), false), ShaderType->GetName(), bMaterialShouldCache, bShaderShouldCache);
+				}
 			}
-		}
 
-		int BreakPoint = 0;
+			int BreakPoint = 0;
 
-		// Assert with detailed information if the shader wasn't found for rendering.  
-		// This is usually the result of an incorrect ShouldCache function.
-		UE_LOG(LogMaterial, Fatal,
-			TEXT("		With VF=%s, Platform=%s\n")
-			TEXT("		MaterialUsageDesc: %s"),
-			VertexFactoryType->GetName(), *LegacyShaderPlatformToShaderFormat(ShaderPlatform).ToString(),
-			*MaterialUsage
-			);
+			// Assert with detailed information if the shader wasn't found for rendering.  
+			// This is usually the result of an incorrect ShouldCache function.
+			UE_LOG(LogMaterial, Fatal,
+				TEXT("		With VF=%s, Platform=%s\n")
+				TEXT("		MaterialUsageDesc: %s"),
+				VertexFactoryType->GetName(), *LegacyShaderPlatformToShaderFormat(ShaderPlatform).ToString(),
+				*MaterialUsage
+				);
+		};
+		noinline_lambda();
 	}
 
-	return ShaderPipeline;
-}
-
-/** Returns the index to the Expression in the Expressions array, or -1 if not found. */
-int32 FMaterial::FindExpression( const TArray<TRefCountPtr<FMaterialUniformExpressionTexture> >&Expressions, const FMaterialUniformExpressionTexture &Expression )
-{
-	for (int32 ExpressionIndex = 0; ExpressionIndex < Expressions.Num(); ++ExpressionIndex)
-	{
-		if ( Expressions[ExpressionIndex]->IsIdentical(&Expression) )
-		{
-			return ExpressionIndex;
-		}
-	}
-	return -1;
+	return FShaderPipelineRef(ShaderPipeline, *RenderingThreadShaderMap);
 }
 
 #if WITH_EDITOR
@@ -2295,13 +2208,11 @@ void FMaterialVirtualTextureStack::GetTextureValues(const FMaterialRenderContext
 	
 	for (uint32 LayerIndex = 0u; LayerIndex < NumLayers; ++LayerIndex)
 	{
-		const int32 ExpressionIndex = LayerUniformExpressionIndices[LayerIndex];
-		if (ExpressionIndex != INDEX_NONE)
+		const int32 ParameterIndex = LayerUniformExpressionIndices[LayerIndex];
+		if (ParameterIndex != INDEX_NONE)
 		{
-			const FMaterialUniformExpressionTexture* UniformExpression = UniformExpressionSet.UniformVirtualTextureExpressions[ExpressionIndex];
-
 			const UTexture* Texture = nullptr;
-			UniformExpression->GetTextureValue(Context, Context.Material, Texture);
+			UniformExpressionSet.GetTextureValue(EMaterialTextureParameterType::Virtual, ParameterIndex, Context, Context.Material, Texture);
 			OutValues[LayerIndex] = Cast<UTexture2D>(Texture);
 		}
 	}
@@ -2309,12 +2220,7 @@ void FMaterialVirtualTextureStack::GetTextureValues(const FMaterialRenderContext
 
 void FMaterialVirtualTextureStack::GetTextureValue(const FMaterialRenderContext& Context, const FUniformExpressionSet& UniformExpressionSet, const URuntimeVirtualTexture*& OutValue) const
 {
-	const int32 ExpressionIndex = LayerUniformExpressionIndices[0];
-	if (ExpressionIndex != INDEX_NONE)
-	{
-		const FMaterialUniformExpressionTexture* UniformExpression = UniformExpressionSet.UniformVirtualTextureExpressions[ExpressionIndex];
-		UniformExpression->GetTextureValue(Context, Context.Material, OutValue);
-	}
+	OutValue = GetIndexedTexture<URuntimeVirtualTexture>(Context.Material, PreallocatedStackTextureIndex);
 }
 
 void FMaterialVirtualTextureStack::Serialize(FArchive& Ar)
@@ -2470,16 +2376,16 @@ void FMaterialRenderProxy::EvaluateUniformExpressions(FUniformExpressionCache& O
 		OutUniformExpressionCache.AllocatedVTs.Add(AllocatedVT);
 	}
 
-	const FShaderParametersMetadata& UniformBufferStruct = UniformExpressionSet.GetUniformBufferStruct();
+	const FRHIUniformBufferLayout& UniformBufferLayout = UniformExpressionSet.GetUniformBufferLayout();
 	FMemMark Mark(FMemStack::Get());
-	uint8* TempBuffer = FMemStack::Get().PushBytes(UniformBufferStruct.GetSize(), SHADER_PARAMETER_STRUCT_ALIGNMENT);
+	uint8* TempBuffer = FMemStack::Get().PushBytes(UniformBufferLayout.ConstantBufferSize, SHADER_PARAMETER_STRUCT_ALIGNMENT);
 
 	check(TempBuffer != nullptr);
-	UniformExpressionSet.FillUniformBuffer(Context, OutUniformExpressionCache, TempBuffer, UniformBufferStruct.GetSize());
+	UniformExpressionSet.FillUniformBuffer(Context, OutUniformExpressionCache, TempBuffer, UniformBufferLayout.ConstantBufferSize);
 
 	if (CommandListIfLocalMode)
 	{
-		OutUniformExpressionCache.LocalUniformBuffer = CommandListIfLocalMode->BuildLocalUniformBuffer(TempBuffer, UniformBufferStruct.GetSize(), UniformBufferStruct.GetLayout());
+		OutUniformExpressionCache.LocalUniformBuffer = CommandListIfLocalMode->BuildLocalUniformBuffer(TempBuffer, UniformBufferLayout.ConstantBufferSize, UniformBufferLayout);
 		check(OutUniformExpressionCache.LocalUniformBuffer.IsValid());
 	}
 	else
@@ -2491,12 +2397,12 @@ void FMaterialRenderProxy::EvaluateUniformExpressions(FUniformExpressionCache& O
 
 		if (IsValidRef(OutUniformExpressionCache.UniformBuffer))
 		{
-			check(OutUniformExpressionCache.UniformBuffer->GetLayout() == UniformBufferStruct.GetLayout());
+			check(OutUniformExpressionCache.UniformBuffer->GetLayout() == UniformBufferLayout);
 			RHIUpdateUniformBuffer(OutUniformExpressionCache.UniformBuffer, TempBuffer);
 		}
 		else
 		{
-			OutUniformExpressionCache.UniformBuffer = RHICreateUniformBuffer(TempBuffer, UniformBufferStruct.GetLayout(), UniformBuffer_MultiFrame);
+			OutUniformExpressionCache.UniformBuffer = RHICreateUniformBuffer(TempBuffer, UniformBufferLayout, UniformBuffer_MultiFrame);
 		}
 	}
 
@@ -2667,6 +2573,13 @@ void FMaterialRenderProxy::UpdateDeferredCachedUniformExpressions()
 			UE_LOG(LogMaterial, Fatal, TEXT("FMaterialRenderProxy deleted and GC mark was: %i"), MaterialProxy->IsMarkedForGarbageCollection());
 		}
 
+		// @todo fnhack
+		if (MaterialProxy->GetFriendlyName() == TEXT("FootSplash_Particle"))
+		{
+			continue;
+		}
+
+
 		UMaterialInterface::IterateOverActiveFeatureLevels([&](ERHIFeatureLevel::Type InFeatureLevel)
 		{
 			const FMaterialRenderProxy* FallbackMaterialProxy = nullptr;
@@ -2692,7 +2605,7 @@ TSet<FMaterialRenderProxy*> FMaterialRenderProxy::DeferredUniformExpressionCache
 	FColoredMaterialRenderProxy
 -----------------------------------------------------------------------------*/
 
-bool FColoredMaterialRenderProxy::GetVectorValue(const FMaterialParameterInfo& ParameterInfo, FLinearColor* OutValue, const FMaterialRenderContext& Context) const
+bool FColoredMaterialRenderProxy::GetVectorValue(const FHashedMaterialParameterInfo& ParameterInfo, FLinearColor* OutValue, const FMaterialRenderContext& Context) const
 {
 	if(ParameterInfo.Name == ColorParamName)
 	{
@@ -2705,17 +2618,17 @@ bool FColoredMaterialRenderProxy::GetVectorValue(const FMaterialParameterInfo& P
 	}
 }
 
-bool FColoredMaterialRenderProxy::GetScalarValue(const FMaterialParameterInfo& ParameterInfo, float* OutValue, const FMaterialRenderContext& Context) const
+bool FColoredMaterialRenderProxy::GetScalarValue(const FHashedMaterialParameterInfo& ParameterInfo, float* OutValue, const FMaterialRenderContext& Context) const
 {
 	return Parent->GetScalarValue(ParameterInfo, OutValue, Context);
 }
 
-bool FColoredMaterialRenderProxy::GetTextureValue(const FMaterialParameterInfo& ParameterInfo,const UTexture** OutValue, const FMaterialRenderContext& Context) const
+bool FColoredMaterialRenderProxy::GetTextureValue(const FHashedMaterialParameterInfo& ParameterInfo,const UTexture** OutValue, const FMaterialRenderContext& Context) const
 {
 	return Parent->GetTextureValue(ParameterInfo,OutValue,Context);
 }
 
-bool FColoredMaterialRenderProxy::GetTextureValue(const FMaterialParameterInfo& ParameterInfo, const URuntimeVirtualTexture** OutValue, const FMaterialRenderContext& Context) const
+bool FColoredMaterialRenderProxy::GetTextureValue(const FHashedMaterialParameterInfo& ParameterInfo, const URuntimeVirtualTexture** OutValue, const FMaterialRenderContext& Context) const
 {
 	return Parent->GetTextureValue(ParameterInfo, OutValue, Context);
 }
@@ -2724,7 +2637,7 @@ bool FColoredMaterialRenderProxy::GetTextureValue(const FMaterialParameterInfo&
 	FColoredTexturedMaterialRenderProxy
 -----------------------------------------------------------------------------*/
 
-bool FColoredTexturedMaterialRenderProxy::GetTextureValue(const FMaterialParameterInfo& ParameterInfo, const UTexture** OutValue, const FMaterialRenderContext& Context) const
+bool FColoredTexturedMaterialRenderProxy::GetTextureValue(const FHashedMaterialParameterInfo& ParameterInfo, const UTexture** OutValue, const FMaterialRenderContext& Context) const
 {
 	if (ParameterInfo.Name == TextureParamName)
 	{
@@ -2745,9 +2658,9 @@ const FMaterial& FOverrideSelectionColorMaterialRenderProxy::GetMaterialWithFall
 	return Parent->GetMaterialWithFallback(InFeatureLevel, OutFallbackMaterialRenderProxy);
 }
 
-bool FOverrideSelectionColorMaterialRenderProxy::GetVectorValue(const FMaterialParameterInfo& ParameterInfo, FLinearColor* OutValue, const FMaterialRenderContext& Context) const
+bool FOverrideSelectionColorMaterialRenderProxy::GetVectorValue(const FHashedMaterialParameterInfo& ParameterInfo, FLinearColor* OutValue, const FMaterialRenderContext& Context) const
 {
-	if (ParameterInfo.Name == NAME_SelectionColor)
+	if (ParameterInfo.Name == FName(NAME_SelectionColor))
 	{
 		*OutValue = SelectionColor;
 		return true;
@@ -2758,17 +2671,17 @@ bool FOverrideSelectionColorMaterialRenderProxy::GetVectorValue(const FMaterialP
 	}
 }
 
-bool FOverrideSelectionColorMaterialRenderProxy::GetScalarValue(const FMaterialParameterInfo& ParameterInfo, float* OutValue, const FMaterialRenderContext& Context) const
+bool FOverrideSelectionColorMaterialRenderProxy::GetScalarValue(const FHashedMaterialParameterInfo& ParameterInfo, float* OutValue, const FMaterialRenderContext& Context) const
 {
 	return Parent->GetScalarValue(ParameterInfo, OutValue, Context);
 }
 
-bool FOverrideSelectionColorMaterialRenderProxy::GetTextureValue(const FMaterialParameterInfo& ParameterInfo, const UTexture** OutValue, const FMaterialRenderContext& Context) const
+bool FOverrideSelectionColorMaterialRenderProxy::GetTextureValue(const FHashedMaterialParameterInfo& ParameterInfo, const UTexture** OutValue, const FMaterialRenderContext& Context) const
 {
 	return Parent->GetTextureValue(ParameterInfo, OutValue, Context);
 }
 
-bool FOverrideSelectionColorMaterialRenderProxy::GetTextureValue(const FMaterialParameterInfo& ParameterInfo, const URuntimeVirtualTexture** OutValue, const FMaterialRenderContext& Context) const
+bool FOverrideSelectionColorMaterialRenderProxy::GetTextureValue(const FHashedMaterialParameterInfo& ParameterInfo, const URuntimeVirtualTexture** OutValue, const FMaterialRenderContext& Context) const
 {
 	return Parent->GetTextureValue(ParameterInfo, OutValue, Context);
 }
@@ -2778,7 +2691,7 @@ bool FOverrideSelectionColorMaterialRenderProxy::GetTextureValue(const FMaterial
 -----------------------------------------------------------------------------*/
 static FName NAME_LightmapRes = FName(TEXT("LightmapRes"));
 
-bool FLightingDensityMaterialRenderProxy::GetVectorValue(const FMaterialParameterInfo& ParameterInfo, FLinearColor* OutValue, const FMaterialRenderContext& Context) const
+bool FLightingDensityMaterialRenderProxy::GetVectorValue(const FHashedMaterialParameterInfo& ParameterInfo, FLinearColor* OutValue, const FMaterialRenderContext& Context) const
 {
 	if (ParameterInfo.Name == NAME_LightmapRes)
 	{
@@ -2874,132 +2787,87 @@ FString FMaterialResource::GetMaterialUsageDescription() const
 	return BaseDescription;
 }
 
-void FMaterial::GetDependentShaderAndVFTypes(EShaderPlatform Platform, TArray<FShaderType*>& OutShaderTypes, TArray<const FShaderPipelineType*>& OutShaderPipelineTypes, TArray<FVertexFactoryType*>& OutVFTypes) const
+static void AddSortedShader(TArray<FShaderType*>& Shaders, FShaderType* Shader)
 {
-	const bool bHasTessellation = GetTessellationMode() != MTM_NoTessellation;
-
-	// Iterate over all vertex factory types.
-	for (TLinkedList<FVertexFactoryType*>::TIterator VertexFactoryTypeIt(FVertexFactoryType::GetTypeList()); VertexFactoryTypeIt; VertexFactoryTypeIt.Next())
+	const int32 SortedIndex = Algo::LowerBoundBy(Shaders, Shader->GetHashedName(), [](const FShaderType* InShaderType) { return InShaderType->GetHashedName(); });
+	if (SortedIndex >= Shaders.Num() || Shaders[SortedIndex] != Shader)
 	{
-		FVertexFactoryType* VertexFactoryType = *VertexFactoryTypeIt;
-		check(VertexFactoryType);
-
-		if (VertexFactoryType->IsUsedWithMaterials())
-		{
-			bool bAddedTypeFromThisVF = false;
-
-			// Iterate over all mesh material shader types.
-			for(TLinkedList<FShaderType*>::TIterator ShaderTypeIt(FShaderType::GetTypeList());ShaderTypeIt;ShaderTypeIt.Next())
-			{
-				FMeshMaterialShaderType* ShaderType = ShaderTypeIt->GetMeshMaterialShaderType();
-				const int32 PermutationCount = ShaderType ? ShaderType->GetPermutationCount() : 0;
-				for (int32 PermutationId = 0; PermutationId < PermutationCount; ++PermutationId)
-				{
-					if (ShaderType->ShouldCompilePermutation(Platform, this, VertexFactoryType, PermutationId) &&
-						ShouldCache(Platform, ShaderType, VertexFactoryType) &&
-						VertexFactoryType->ShouldCache(Platform, this, ShaderType))
-					{
-						bAddedTypeFromThisVF = true;
-						OutShaderTypes.AddUnique(ShaderType);
-					}
-				}
-			}
+		Shaders.Insert(Shader, SortedIndex);
+	}
+}
 
-			for (TLinkedList<FShaderPipelineType*>::TIterator PipelineTypeIt(FShaderPipelineType::GetTypeList()); PipelineTypeIt; PipelineTypeIt.Next())
-			{
-				auto* PipelineType = *PipelineTypeIt;
-				if (PipelineType->IsMeshMaterialTypePipeline() && PipelineType->HasTessellation() == bHasTessellation)
-				{
-					int32 NumShouldCache = 0;
-					auto& ShaderStages = PipelineType->GetStages();
-					for (const FShaderType* Type : ShaderStages)
-					{
-						const FMeshMaterialShaderType* ShaderType = Type->GetMeshMaterialShaderType();
-						if (ShaderType->ShouldCompilePermutation(Platform, this, VertexFactoryType, kUniqueShaderPermutationId) &&
-							ShouldCache(Platform, ShaderType, VertexFactoryType) &&
-							VertexFactoryType->ShouldCache(Platform, this, ShaderType)
-							)
-						{
-							++NumShouldCache;
-						}
-					}
+static void AddSortedShaderPipeline(TArray<const FShaderPipelineType*>& Pipelines, const FShaderPipelineType* Pipeline)
+{
+	const int32 SortedIndex = Algo::LowerBoundBy(Pipelines, Pipeline->GetHashedName(), [](const FShaderPipelineType* InPiplelineType) { return InPiplelineType->GetHashedName(); });
+	if (SortedIndex >= Pipelines.Num() || Pipelines[SortedIndex] != Pipeline)
+	{
+		Pipelines.Insert(Pipeline, SortedIndex);
+	}
+}
 
-					if (NumShouldCache == ShaderStages.Num())
-					{
-						bAddedTypeFromThisVF = true;
-						OutShaderPipelineTypes.AddUnique(PipelineType);
-						for (const FShaderType* Type : ShaderStages)
-						{
-							OutShaderTypes.AddUnique((FShaderType*)Type);
-						}
-					}
-				}
-			}
+void FMaterial::GetDependentShaderAndVFTypes(EShaderPlatform Platform, TArray<FShaderType*>& OutShaderTypes, TArray<const FShaderPipelineType*>& OutShaderPipelineTypes, TArray<FVertexFactoryType*>& OutVFTypes) const
+{
+	const FMaterialShaderParameters MaterialParameters(this);
+	const FMaterialShaderMapLayout& Layout = AcquireMaterialShaderMapLayout(Platform, MaterialParameters);
 
-			if (bAddedTypeFromThisVF)
-			{
-				OutVFTypes.Add(VertexFactoryType);
-			}
+	for (const FShaderLayoutEntry& Shader : Layout.Shaders)
+	{
+		if (ShouldCache(Platform, Shader.ShaderType, nullptr))
+		{
+			AddSortedShader(OutShaderTypes, Shader.ShaderType);
 		}
 	}
 
-	// Iterate over all material shader types.
-	for (TLinkedList<FShaderType*>::TIterator ShaderTypeIt(FShaderType::GetTypeList()); ShaderTypeIt; ShaderTypeIt.Next())
+	for (const FShaderPipelineType* Pipeline : Layout.ShaderPipelines)
 	{
-		FMaterialShaderType* ShaderType = ShaderTypeIt->GetMaterialShaderType();
-		const int32 PermutationCount = ShaderType ? ShaderType->GetPermutationCount() : 0;
-		for (int32 PermutationId = 0; PermutationId < PermutationCount; ++PermutationId)
+		if (ShouldCachePipeline(Platform, Pipeline, nullptr))
 		{
-			if (ShaderType->ShouldCompilePermutation(Platform, this, PermutationId) &&
-				ShouldCache(Platform, ShaderType, nullptr))
+			AddSortedShaderPipeline(OutShaderPipelineTypes, Pipeline);
+			for (const FShaderType* Type : Pipeline->GetStages())
 			{
-				OutShaderTypes.Add(ShaderType);
+				AddSortedShader(OutShaderTypes, (FShaderType*)Type);
 			}
 		}
 	}
 
-	for (TLinkedList<FShaderPipelineType*>::TIterator PipelineTypeIt(FShaderPipelineType::GetTypeList()); PipelineTypeIt; PipelineTypeIt.Next())
+	for (const FMeshMaterialShaderMapLayout& MeshLayout : Layout.MeshShaderMaps)
 	{
-		auto* PipelineType = *PipelineTypeIt;
-		if (PipelineType->IsMaterialTypePipeline() && PipelineType->HasTessellation() == bHasTessellation)
+		bool bIncludeVertexFactory = false;
+		for (const FShaderLayoutEntry& Shader : MeshLayout.Shaders)
 		{
-			int32 NumShouldCache = 0;
-			auto& ShaderStages = PipelineType->GetStages();
-			for (const FShaderType* Type : ShaderStages)
+			if (ShouldCache(Platform, Shader.ShaderType, MeshLayout.VertexFactoryType))
 			{
-				const FMaterialShaderType* ShaderType = Type->GetMaterialShaderType();
-				if (ShaderType &&
-					ShaderType->ShouldCompilePermutation(Platform, this, kUniqueShaderPermutationId) &&
-					ShouldCache(Platform, ShaderType, nullptr)
-					)
-				{
-					++NumShouldCache;
-				}
+				bIncludeVertexFactory = true;
+				AddSortedShader(OutShaderTypes, Shader.ShaderType);
 			}
+		}
 
-			if (NumShouldCache == ShaderStages.Num())
+		for (const FShaderPipelineType* Pipeline : MeshLayout.ShaderPipelines)
+		{
+			if (ShouldCachePipeline(Platform, Pipeline, MeshLayout.VertexFactoryType))
 			{
-				for (const FShaderType* Type : ShaderStages)
+				bIncludeVertexFactory = true;
+				AddSortedShaderPipeline(OutShaderPipelineTypes, Pipeline);
+				for (const FShaderType* Type : Pipeline->GetStages())
 				{
-					const FMaterialShaderType* ShaderType = Type->GetMaterialShaderType();
-					OutShaderPipelineTypes.AddUnique(PipelineType);
-					OutShaderTypes.AddUnique((FShaderType*)Type);
+					AddSortedShader(OutShaderTypes, (FShaderType*)Type);
 				}
 			}
 		}
-	}
 
-	// Sort by name so that we get deterministic keys
-	OutShaderTypes.Sort(FCompareShaderTypes());
-	OutVFTypes.Sort(FCompareVertexFactoryTypes());
-	OutShaderPipelineTypes.Sort(FCompareShaderPipelineNameTypes());
+		if (bIncludeVertexFactory)
+		{
+			// Vertex factories are already sorted
+			OutVFTypes.Add(MeshLayout.VertexFactoryType);
+		}
+	}
 }
 
 void FMaterial::GetReferencedTexturesHash(EShaderPlatform Platform, FSHAHash& OutHash) const
 {
 	FSHA1 HashState;
 
-	const TArray<UObject*>& ReferencedTextures = GetReferencedTextures();
+	const TArrayView<UObject* const> ReferencedTextures = GetReferencedTextures();
 	// Hash the names of the uniform expression textures to capture changes in their order or values resulting from material compiler code changes
 	for (int32 TextureIndex = 0; TextureIndex < ReferencedTextures.Num(); TextureIndex++)
 	{
@@ -3063,20 +2931,12 @@ bool FMaterial::GetMaterialExpressionSource( FString& OutSource )
 
 bool FMaterial::WritesEveryPixel(bool bShadowPass) const
 {
-	bool bStencilDitheredLOD = false;
-	if (FeatureLevel >= ERHIFeatureLevel::SM5)
-	{
-		// this option affects only deferred renderer
-		static TConsoleVariableData<int32>* CVarStencilDitheredLOD =
-			IConsoleManager::Get().FindTConsoleVariableDataInt(TEXT("r.StencilForLODDither"));
-		bStencilDitheredLOD = (CVarStencilDitheredLOD->GetValueOnAnyThread() != 0);
-	}
-
+	bool bLocalStencilDitheredLOD = FeatureLevel >= ERHIFeatureLevel::SM5 && bStencilDitheredLOD;
 	return !IsMasked()
 		// Render dithered material as masked if a stencil prepass is not used (UE-50064, UE-49537)
-		&& !((bShadowPass || !bStencilDitheredLOD) && IsDitheredLODTransition())
+		&& !((bShadowPass || !bLocalStencilDitheredLOD) && IsDitheredLODTransition())
 		&& !IsWireframe()
-		&& !(bStencilDitheredLOD && IsDitheredLODTransition() && IsUsedWithInstancedStaticMeshes())
+		&& !(bLocalStencilDitheredLOD && IsDitheredLODTransition() && IsUsedWithInstancedStaticMeshes())
 		&& !IsStencilTestEnabled();
 }
 
@@ -3558,7 +3418,7 @@ bool FMaterialShaderMapId::ContainsShaderType(const FShaderType* ShaderType, int
 {
 	for (int32 TypeIndex = 0; TypeIndex < ShaderTypeDependencies.Num(); TypeIndex++)
 	{
-		if (ShaderTypeDependencies[TypeIndex].ShaderType == ShaderType &&
+		if (ShaderTypeDependencies[TypeIndex].ShaderTypeName == ShaderType->GetHashedName() &&
 			ShaderTypeDependencies[TypeIndex].PermutationId == PermutationId)
 		{
 			return true;
@@ -3572,7 +3432,7 @@ bool FMaterialShaderMapId::ContainsShaderPipelineType(const FShaderPipelineType*
 {
 	for (int32 TypeIndex = 0; TypeIndex < ShaderPipelineTypeDependencies.Num(); TypeIndex++)
 	{
-		if (ShaderPipelineTypeDependencies[TypeIndex].ShaderPipelineType == ShaderPipelineType)
+		if (ShaderPipelineTypeDependencies[TypeIndex].ShaderPipelineTypeName == ShaderPipelineType->GetHashedName())
 		{
 			return true;
 		}
@@ -3585,7 +3445,7 @@ bool FMaterialShaderMapId::ContainsVertexFactoryType(const FVertexFactoryType* V
 {
 	for (int32 TypeIndex = 0; TypeIndex < VertexFactoryTypeDependencies.Num(); TypeIndex++)
 	{
-		if (VertexFactoryTypeDependencies[TypeIndex].VertexFactoryType == VFType)
+		if (VertexFactoryTypeDependencies[TypeIndex].VertexFactoryTypeName == VFType->GetHashedName())
 		{
 			return true;
 		}
@@ -4093,6 +3953,8 @@ void FMaterialResourceProxyReader::Initialize(
 	EMaterialQualityLevel::Type QualityLevel,
 	bool bSeekToEnd)
 {
+	SCOPED_LOADTIMER(FMaterialResourceProxyReader_Initialize);
+
 	decltype(Names.Num()) NumNames;
 	InnerArchive << NumNames;
 	Names.Empty(NumNames);
@@ -4178,5 +4040,3 @@ void SetShaderMapsOnMaterialResources(const TMap<FMaterial*, FMaterialShaderMap*
 		SetShaderMapsOnMaterialResources_RenderThread(RHICmdList, InMaterialsToUpdate);
 	});
 }
-
-#undef LOCTEXT_NAMESPACE
\ No newline at end of file
diff --git a/Engine/Source/Runtime/Engine/Private/Materials/MaterialStaticParameterValueResolver.h b/Engine/Source/Runtime/Engine/Private/Materials/MaterialStaticParameterValueResolver.h
--- a/Engine/Source/Runtime/Engine/Private/Materials/MaterialStaticParameterValueResolver.h
+++ b/Engine/Source/Runtime/Engine/Private/Materials/MaterialStaticParameterValueResolver.h
@@ -8,6 +8,7 @@
 #include "Materials/MaterialExpressionStaticBoolParameter.h"
 #include "Materials/MaterialExpressionStaticComponentMaskParameter.h"
 
+#if WITH_EDITORONLY_DATA
 template <typename ParameterType, typename OutputType>
 class TMaterialStaticParameterValueResolver
 {
@@ -25,7 +26,7 @@ public:
 	 * @param ParameterInfo The info associated with the parameter to try to resolve.
 	 * @param Function The function from which to try and obtain the value for the parameter.
 	 */
-	void AttemptResolve(int32 ParamIndex, const FMaterialParameterInfo& ParameterInfo, UMaterialFunctionInterface* Function)
+	void AttemptResolve(int32 ParamIndex, const FHashedMaterialParameterInfo& ParameterInfo, UMaterialFunctionInterface* Function)
 	{
 		if (!Function)
 		{
@@ -108,7 +109,7 @@ public:
 						continue;
 
 					int32 ParamIndex = ParameterIndicesForFunction[FunctionParameterIdx];
-					const FMaterialParameterInfo* ParameterInfo = EvalContext.GetParameterInfo(ParamIndex);
+					const FHashedMaterialParameterInfo* ParameterInfo = EvalContext.GetParameterInfo(ParamIndex);
 					bool bIsOverride = ExtractValueFromFunction(ParamIndex, *ParameterInfo, ParameterAndOwner.Owner, ParameterAndOwner.Parameter);
 					EvalContext.MarkParameterResolved(ParamIndex, bIsOverride);
 				}
@@ -126,7 +127,7 @@ private:
 	 * @param Parameter If non-null, the parameter from which to try and obtain the default value if no overrides exist for the parameter on the function.
 	 * @return true if a *non-default* or override value was obtained for the parameter, if no value was obtained, or if we fell back to obtaining the default value, then false will be returned.
 	 */
-	bool ExtractValueFromFunction(int32 ParamIndex, const FMaterialParameterInfo& ParameterInfo, UMaterialFunctionInterface* Function, ParameterType* Parameter)
+	bool ExtractValueFromFunction(int32 ParamIndex, const FHashedMaterialParameterInfo& ParameterInfo, UMaterialFunctionInterface* Function, ParameterType* Parameter)
 	{
 		static_assert(TIsSame<ParameterType,void>::Value, "Must provide explicit specialization for value retrieval method when using TMaterialStaticParameterValueResolver");
 		return false;
@@ -140,7 +141,7 @@ private:
 
 
 template<>
-inline bool TMaterialStaticParameterValueResolver<UMaterialExpressionStaticBoolParameter, TBitArray<>>::ExtractValueFromFunction(int32 ParamIndex, const FMaterialParameterInfo& ParameterInfo, UMaterialFunctionInterface* Function, UMaterialExpressionStaticBoolParameter* Parameter)
+inline bool TMaterialStaticParameterValueResolver<UMaterialExpressionStaticBoolParameter, TBitArray<>>::ExtractValueFromFunction(int32 ParamIndex, const FHashedMaterialParameterInfo& ParameterInfo, UMaterialFunctionInterface* Function, UMaterialExpressionStaticBoolParameter* Parameter)
 {
 	bool bTempVal;
 	if (Function->OverrideNamedStaticSwitchParameter(ParameterInfo, bTempVal, OutExpressionGuids[ParamIndex]))
@@ -160,7 +161,7 @@ inline bool TMaterialStaticParameterValueResolver<UMaterialExpressionStaticBoolP
 }
 
 template<>
-inline bool TMaterialStaticParameterValueResolver<UMaterialExpressionStaticComponentMaskParameter, TBitArray<>>::ExtractValueFromFunction(int32 ParamIndex, const FMaterialParameterInfo& ParameterInfo, UMaterialFunctionInterface* Function, UMaterialExpressionStaticComponentMaskParameter* Parameter)
+inline bool TMaterialStaticParameterValueResolver<UMaterialExpressionStaticComponentMaskParameter, TBitArray<>>::ExtractValueFromFunction(int32 ParamIndex, const FHashedMaterialParameterInfo& ParameterInfo, UMaterialFunctionInterface* Function, UMaterialExpressionStaticComponentMaskParameter* Parameter)
 {
 	bool bTempR, bTempG, bTempB, bTempA;
 	if (Function->OverrideNamedStaticComponentMaskParameter(ParameterInfo, bTempR, bTempG, bTempB, bTempA, OutExpressionGuids[ParamIndex]))
@@ -185,3 +186,5 @@ inline bool TMaterialStaticParameterValueResolver<UMaterialExpressionStaticCompo
 
 	return false;
 }
+
+#endif // WITH_EDITORONLY_DATA
