Copying Private-LoadTimes-4.24 stream to Main. Biggest changes are in Materials/Shader memory freezing.

diff --git a/Engine/Source/Runtime/Engine/Private/Materials/MaterialUniformExpressions.cpp b/Engine/Source/Runtime/Engine/Private/Materials/MaterialUniformExpressions.cpp
--- a/Engine/Source/Runtime/Engine/Private/Materials/MaterialUniformExpressions.cpp
+++ b/Engine/Source/Runtime/Engine/Private/Materials/MaterialUniformExpressions.cpp
@@ -56,165 +56,232 @@ TMap<FName,FMaterialUniformExpressionType*>& FMaterialUniformExpressionType::Get
 	return TypeMap;
 }
 
-FMaterialUniformExpressionType::FMaterialUniformExpressionType(
-	const TCHAR* InName,
-	SerializationConstructorType InSerializationConstructor
-	):
-	Name(InName),
-	SerializationConstructor(InSerializationConstructor)
+static FGuid GetExternalTextureGuid(const FMaterialRenderContext& Context, const FGuid& ExternalTextureGuid, const FHashedName& ParameterName, int32 SourceTextureIndex)
 {
-	// Put the type in the type list until the name subsystem/type map are initialized.
-	(new TLinkedList<FMaterialUniformExpressionType*>(this))->LinkHead(GetTypeList());
+	FGuid GuidToLookup;
+	if (ExternalTextureGuid.IsValid())
+	{
+		// Use the compile-time GUID if it is set
+		GuidToLookup = ExternalTextureGuid;
+	}
+	else
+	{
+		const UTexture* TextureParameterObject = nullptr;
+		if (!ParameterName.IsNone() && Context.MaterialRenderProxy && Context.MaterialRenderProxy->GetTextureValue(ParameterName, &TextureParameterObject, Context) && TextureParameterObject)
+		{
+			GuidToLookup = TextureParameterObject->GetExternalTextureGuid();
+		}
+		else
+		{
+			// Otherwise attempt to use the texture index in the material, if it's valid
+			const UTexture* TextureObject = SourceTextureIndex != INDEX_NONE ? GetIndexedTexture<UTexture>(Context.Material, SourceTextureIndex) : nullptr;
+			if (TextureObject)
+			{
+				GuidToLookup = TextureObject->GetExternalTextureGuid();
+			}
+		}
+	}
+	return GuidToLookup;
 }
 
-FArchive& operator<<(FArchive& Ar,FMaterialUniformExpression*& Ref)
+static void GetTextureParameterValue(const FHashedMaterialParameterInfo& ParameterInfo, int32 TextureIndex, const FMaterialRenderContext& Context, const UTexture*& OutValue)
 {
-	// Serialize the expression type.
-	if(Ar.IsSaving())
+	if (ParameterInfo.Name.IsNone())
 	{
-		// Write the type name.
-		check(Ref);
-		FName TypeName(Ref->GetType()->Name);
-		Ar << TypeName;
+		OutValue = GetIndexedTexture<UTexture>(Context.Material, TextureIndex);
 	}
-	else if(Ar.IsLoading())
+	else if (!Context.MaterialRenderProxy || !Context.MaterialRenderProxy->GetTextureValue(ParameterInfo, &OutValue, Context))
 	{
-		// Read the type name.
-		FName TypeName = NAME_None;
-		Ar << TypeName;
+		UTexture* Value = nullptr;
 
-		// Find the expression type with a matching name.
-		FMaterialUniformExpressionType* Type = FMaterialUniformExpressionType::GetTypeMap().FindRef(TypeName);
-		checkf(Type, TEXT("Unable to find FMaterialUniformExpressionType for TypeName '%s'"), *TypeName.ToString());
+		if (AreExperimentalMaterialLayersEnabled())
+		{
+			UMaterialInterface* Interface = Context.Material.GetMaterialInterface();
+			if (!Interface || !Interface->GetTextureParameterDefaultValue(ParameterInfo, Value))
+			{
+				Value = GetIndexedTexture<UTexture>(Context.Material, TextureIndex);
+			}
+		}
+		else
+		{
+			Value = GetIndexedTexture<UTexture>(Context.Material, TextureIndex);
+		}
 
-		// Construct a new instance of the expression type.
-		Ref = (*Type->SerializationConstructor)();
+		OutValue = Value;
 	}
+}
 
-	// Serialize the expression.
-	Ref->Serialize(Ar);
-
-	return Ar;
+FMaterialUniformExpressionType::FMaterialUniformExpressionType(const TCHAR* InName)
+	: Name(InName)
+{
+	// Put the type in the type list until the name subsystem/type map are initialized.
+	(new TLinkedList<FMaterialUniformExpressionType*>(this))->LinkHead(GetTypeList());
 }
 
-FArchive& operator<<(FArchive& Ar,FMaterialUniformExpressionTexture*& Ref)
+void FMaterialUniformExpression::WriteNumberOpcodes(FMaterialPreshaderData& OutData) const
 {
-	Ar << (FMaterialUniformExpression*&)Ref;
-	return Ar;
+	UE_LOG(LogMaterial, Warning, TEXT("Missing WriteNumberOpcodes impl for %s"), GetType()->GetName());
+	OutData.WriteOpcode(EMaterialPreshaderOpcode::ConstantZero);
 }
 
-void FUniformExpressionSet::Serialize(FArchive& Ar)
+void FUniformParameterOverrides::SetScalarOverride(int32 Index, float Value, bool bOverride)
 {
-	Ar << UniformVectorExpressions;
-	Ar << UniformScalarExpressions;
-	Ar << Uniform2DTextureExpressions;
-	Ar << UniformCubeTextureExpressions;
-	Ar << UniformVolumeTextureExpressions;
-	Ar << UniformVirtualTextureExpressions;
-	Ar << UniformExternalTextureExpressions;
-	Ar << VTStacks;
-	Ar << Uniform2DArrayTextureExpressions;
+	if (Index >= ScalarOverrides.Num()) ScalarOverrides.SetNumZeroed(Index + 1);
+	ScalarOverrides[Index].Value = Value;
+	ScalarOverrides[Index].bOverride = bOverride;
+}
 
+void FUniformParameterOverrides::SetVectorOverride(int32 Index, const FLinearColor& Value, bool bOverride)
+{
+	if (Index >= VectorOverrides.Num()) VectorOverrides.SetNumZeroed(Index + 1);
+	VectorOverrides[Index].Value = Value;
+	VectorOverrides[Index].bOverride = bOverride;
+}
 
-	Ar << ParameterCollections;
+bool FUniformParameterOverrides::GetScalarOverride(int32 Index, float& OutValue) const
+{
+	if (ScalarOverrides.IsValidIndex(Index) && ScalarOverrides[Index].bOverride)
+	{
+		OutValue = ScalarOverrides[Index].Value;
+		return true;
+	}
+	return false;
+}
 
-	// Recreate the uniform buffer struct after loading.
-	if(Ar.IsLoading())
+bool FUniformParameterOverrides::GetVectorOverride(int32 Index, FLinearColor& OutValue) const
+{
+	if (VectorOverrides.IsValidIndex(Index) && VectorOverrides[Index].bOverride)
 	{
-		CreateBufferStruct();
+		OutValue = VectorOverrides[Index].Value;
+		return true;
 	}
+	return false;
 }
 
-bool FUniformExpressionSet::IsEmpty() const
+void FUniformParameterOverrides::SetTextureOverride(EMaterialTextureParameterType Type, int32 Index, UTexture* Texture)
 {
-	return UniformVectorExpressions.Num() == 0
-		&& UniformScalarExpressions.Num() == 0
-		&& Uniform2DTextureExpressions.Num() == 0
-		&& UniformCubeTextureExpressions.Num() == 0
-		&& Uniform2DArrayTextureExpressions.Num() == 0
-		&& UniformVolumeTextureExpressions.Num() == 0
-		&& UniformVirtualTextureExpressions.Num() == 0
-		&& UniformExternalTextureExpressions.Num() == 0
-		&& VTStacks.Num() == 0
-		&& ParameterCollections.Num() == 0;
+	check(IsInGameThread());
+	const uint32 TypeIndex = (uint32)Type;
+	if (Index >= GameThreadTextureOverides[TypeIndex].Num()) GameThreadTextureOverides[TypeIndex].SetNumZeroed(Index + 1);
+	GameThreadTextureOverides[TypeIndex][Index] = Texture;
+
+	FUniformParameterOverrides* Self = this;
+	ENQUEUE_RENDER_COMMAND(SetTextureOverrideCommand)(
+		[Self, TypeIndex, Index, Texture](FRHICommandListImmediate& RHICmdList)
+	{
+		if (Index >= Self->RenderThreadTextureOverrides[TypeIndex].Num()) Self->RenderThreadTextureOverrides[TypeIndex].SetNumZeroed(Index + 1);
+		Self->RenderThreadTextureOverrides[TypeIndex][Index] = Texture;
+	});
 }
 
-bool FUniformExpressionSet::operator==(const FUniformExpressionSet& ReferenceSet) const
+UTexture* FUniformParameterOverrides::GetTextureOverride_GameThread(EMaterialTextureParameterType Type, int32 Index) const
 {
-	if (UniformVectorExpressions.Num() != ReferenceSet.UniformVectorExpressions.Num()
-		|| UniformScalarExpressions.Num() != ReferenceSet.UniformScalarExpressions.Num()
-		|| Uniform2DTextureExpressions.Num() != ReferenceSet.Uniform2DTextureExpressions.Num()
-		|| UniformCubeTextureExpressions.Num() != ReferenceSet.UniformCubeTextureExpressions.Num()
-		|| Uniform2DArrayTextureExpressions.Num() != ReferenceSet.Uniform2DArrayTextureExpressions.Num()
-		|| UniformVolumeTextureExpressions.Num() != ReferenceSet.UniformVolumeTextureExpressions.Num()
-		|| UniformVirtualTextureExpressions.Num() != ReferenceSet.UniformVirtualTextureExpressions.Num()
-		|| UniformExternalTextureExpressions.Num() != ReferenceSet.UniformExternalTextureExpressions.Num()
-		|| VTStacks.Num() != ReferenceSet.VTStacks.Num()
-		|| ParameterCollections.Num() != ReferenceSet.ParameterCollections.Num())
+	check(IsInGameThread());
+	const uint32 TypeIndex = (uint32)Type;
+	if (GameThreadTextureOverides[TypeIndex].IsValidIndex(Index))
 	{
-		return false;
+		return GameThreadTextureOverides[TypeIndex][Index];
+	}
+	return nullptr;
+}
+
+UTexture* FUniformParameterOverrides::GetTextureOverride_RenderThread(EMaterialTextureParameterType Type, int32 Index) const
+{
+	check(IsInParallelRenderingThread());
+	const uint32 TypeIndex = (uint32)Type;
+	if (RenderThreadTextureOverrides[TypeIndex].IsValidIndex(Index))
+	{
+		return RenderThreadTextureOverrides[TypeIndex][Index];
 	}
+	return nullptr;
+}
 
-	for (int32 i = 0; i < UniformVectorExpressions.Num(); i++)
+bool FUniformExpressionSet::IsEmpty() const
+{
+	for (uint32 TypeIndex = 0u; TypeIndex < NumMaterialTextureParameterTypes; ++TypeIndex)
 	{
-		if (!UniformVectorExpressions[i]->IsIdentical(ReferenceSet.UniformVectorExpressions[i]))
+		if (UniformTextureParameters[TypeIndex].Num() != 0)
 		{
 			return false;
 		}
 	}
 
-	for (int32 i = 0; i < UniformScalarExpressions.Num(); i++)
+	return UniformVectorParameters.Num() == 0
+		&& UniformScalarParameters.Num() == 0
+		&& UniformVectorPreshaders.Num() == 0
+		&& UniformScalarPreshaders.Num() == 0
+		&& UniformExternalTextureParameters.Num() == 0
+		&& VTStacks.Num() == 0
+		&& ParameterCollections.Num() == 0;
+}
+
+bool FUniformExpressionSet::operator==(const FUniformExpressionSet& ReferenceSet) const
+{
+	for (uint32 TypeIndex = 0u; TypeIndex < NumMaterialTextureParameterTypes; ++TypeIndex)
 	{
-		if (!UniformScalarExpressions[i]->IsIdentical(ReferenceSet.UniformScalarExpressions[i]))
+		if (UniformTextureParameters[TypeIndex].Num() != ReferenceSet.UniformTextureParameters[TypeIndex].Num())
 		{
 			return false;
 		}
 	}
 
-	for (int32 i = 0; i < Uniform2DTextureExpressions.Num(); i++)
+	if (UniformScalarParameters.Num() != ReferenceSet.UniformScalarParameters.Num()
+		|| UniformVectorParameters.Num() != ReferenceSet.UniformVectorParameters.Num()
+		|| UniformScalarPreshaders.Num() != ReferenceSet.UniformScalarPreshaders.Num()
+		|| UniformVectorPreshaders.Num() != ReferenceSet.UniformVectorPreshaders.Num()
+		|| UniformExternalTextureParameters.Num() != ReferenceSet.UniformExternalTextureParameters.Num()
+		|| VTStacks.Num() != ReferenceSet.VTStacks.Num()
+		|| ParameterCollections.Num() != ReferenceSet.ParameterCollections.Num())
+	{
+		return false;
+	}
+
+	for (int32 i = 0; i < UniformScalarParameters.Num(); i++)
 	{
-		if (!Uniform2DTextureExpressions[i]->IsIdentical(ReferenceSet.Uniform2DTextureExpressions[i]))
+		if (UniformScalarParameters[i] != ReferenceSet.UniformScalarParameters[i])
 		{
 			return false;
 		}
 	}
 
-	for (int32 i = 0; i < UniformCubeTextureExpressions.Num(); i++)
+	for (int32 i = 0; i < UniformVectorParameters.Num(); i++)
 	{
-		if (!UniformCubeTextureExpressions[i]->IsIdentical(ReferenceSet.UniformCubeTextureExpressions[i]))
+		if (UniformVectorParameters[i] != ReferenceSet.UniformVectorParameters[i])
 		{
 			return false;
 		}
 	}
 
-	for (int32 i = 0; i < Uniform2DArrayTextureExpressions.Num(); i++)
+	for (int32 i = 0; i < UniformScalarPreshaders.Num(); i++)
 	{
-		if (!Uniform2DArrayTextureExpressions[i]->IsIdentical(ReferenceSet.Uniform2DArrayTextureExpressions[i]))
+		if (UniformScalarPreshaders[i] != ReferenceSet.UniformScalarPreshaders[i])
 		{
 			return false;
 		}
 	}
 
-
-	for (int32 i = 0; i < UniformVolumeTextureExpressions.Num(); i++)
+	for (int32 i = 0; i < UniformVectorPreshaders.Num(); i++)
 	{
-		if (!UniformVolumeTextureExpressions[i]->IsIdentical(ReferenceSet.UniformVolumeTextureExpressions[i]))
+		if (UniformVectorPreshaders[i] != ReferenceSet.UniformVectorPreshaders[i])
 		{
 			return false;
 		}
 	}
 
-	for (int32 i = 0; i < UniformVirtualTextureExpressions.Num(); i++)
+	for (uint32 TypeIndex = 0u; TypeIndex < NumMaterialTextureParameterTypes; ++TypeIndex)
 	{
-		if (!UniformVirtualTextureExpressions[i]->IsIdentical(ReferenceSet.UniformVirtualTextureExpressions[i]))
+		for (int32 i = 0; i < UniformTextureParameters[TypeIndex].Num(); i++)
 		{
-			return false;
+			if (UniformTextureParameters[TypeIndex][i] != ReferenceSet.UniformTextureParameters[TypeIndex][i])
+			{
+				return false;
+			}
 		}
 	}
 
-	for (int32 i = 0; i < UniformExternalTextureExpressions.Num(); i++)
+	for (int32 i = 0; i < UniformExternalTextureParameters.Num(); i++)
 	{
-		if (!UniformExternalTextureExpressions[i]->IsIdentical(ReferenceSet.UniformExternalTextureExpressions[i]))
+		if (UniformExternalTextureParameters[i] != ReferenceSet.UniformExternalTextureParameters[i])
 		{
 			return false;
 		}
@@ -236,20 +303,25 @@ bool FUniformExpressionSet::operator==(const FUniformExpressionSet& ReferenceSet
 		}
 	}
 
+	if (UniformPreshaderData != ReferenceSet.UniformPreshaderData)
+	{
+		return false;
+	}
+
 	return true;
 }
 
 FString FUniformExpressionSet::GetSummaryString() const
 {
-	return FString::Printf(TEXT("(%u vectors, %u scalars, %u 2d tex, %u cube tex, %u 2d array tex, %u 3d tex, %u virtual tex, %u external tex, %u VT stacks, %u collections)"),
-		UniformVectorExpressions.Num(), 
-		UniformScalarExpressions.Num(),
-		Uniform2DTextureExpressions.Num(),
-		UniformCubeTextureExpressions.Num(),
-		Uniform2DArrayTextureExpressions.Num(),
-		UniformVolumeTextureExpressions.Num(),
-		UniformVirtualTextureExpressions.Num(),
-		UniformExternalTextureExpressions.Num(),
+	return FString::Printf(TEXT("(%u vectors, %u scalars, %u 2d tex, %u cube tex, %u 2darray tex, %u 3d tex, %u virtual tex, %u external tex, %u VT stacks, %u collections)"),
+		UniformVectorPreshaders.Num(), 
+		UniformScalarPreshaders.Num(),
+		UniformTextureParameters[(uint32)EMaterialTextureParameterType::Standard2D].Num(),
+		UniformTextureParameters[(uint32)EMaterialTextureParameterType::Cube].Num(),
+		UniformTextureParameters[(uint32)EMaterialTextureParameterType::Array2D].Num(),
+		UniformTextureParameters[(uint32)EMaterialTextureParameterType::Volume].Num(),
+		UniformTextureParameters[(uint32)EMaterialTextureParameterType::Virtual].Num(),
+		UniformExternalTextureParameters.Num(),
 		VTStacks.Num(),
 		ParameterCollections.Num()
 		);
@@ -265,10 +337,7 @@ void FUniformExpressionSet::SetParameterCollections(const TArray<UMaterialParame
 	}
 }
 
-
-static FName MaterialLayoutName(TEXT("Material"));
-
-void FUniformExpressionSet::CreateBufferStruct()
+FShaderParametersMetadata* FUniformExpressionSet::CreateBufferStruct()
 {
 	// Make sure FUniformExpressionSet::CreateDebugLayout() is in sync
 	TArray<FShaderParametersMetadata::FMember> Members;
@@ -281,24 +350,25 @@ void FUniformExpressionSet::CreateBufferStruct()
 		NextMemberOffset += VTStacks.Num() * sizeof(FUintVector4) * 2;
 	}
 
-	if (UniformVirtualTextureExpressions.Num() > 0)
+	const int32 NumVirtualTextures = UniformTextureParameters[(uint32)EMaterialTextureParameterType::Virtual].Num();
+	if (NumVirtualTextures > 0)
 	{
 		// 1x uint4 per Virtual Texture
-		new(Members) FShaderParametersMetadata::FMember(TEXT("VTPackedUniform"), TEXT(""), NextMemberOffset, UBMT_UINT32, EShaderPrecisionModifier::Float, 1, 4, UniformVirtualTextureExpressions.Num(), NULL);
-		NextMemberOffset += UniformVirtualTextureExpressions.Num() * sizeof(FUintVector4);
+		new(Members) FShaderParametersMetadata::FMember(TEXT("VTPackedUniform"), TEXT(""), NextMemberOffset, UBMT_UINT32, EShaderPrecisionModifier::Float, 1, 4, NumVirtualTextures, NULL);
+		NextMemberOffset += NumVirtualTextures * sizeof(FUintVector4);
 	}
 
-	if (UniformVectorExpressions.Num())
+	if (UniformVectorPreshaders.Num())
 	{
-		new(Members) FShaderParametersMetadata::FMember(TEXT("VectorExpressions"),TEXT(""),NextMemberOffset,UBMT_FLOAT32,EShaderPrecisionModifier::Half,1,4,UniformVectorExpressions.Num(),NULL);
-		const uint32 VectorArraySize = UniformVectorExpressions.Num() * sizeof(FVector4);
+		new(Members) FShaderParametersMetadata::FMember(TEXT("VectorExpressions"),TEXT(""),NextMemberOffset,UBMT_FLOAT32,EShaderPrecisionModifier::Half,1,4, UniformVectorPreshaders.Num(),NULL);
+		const uint32 VectorArraySize = UniformVectorPreshaders.Num() * sizeof(FVector4);
 		NextMemberOffset += VectorArraySize;
 	}
 
-	if (UniformScalarExpressions.Num())
+	if (UniformScalarPreshaders.Num())
 	{
-		new(Members) FShaderParametersMetadata::FMember(TEXT("ScalarExpressions"),TEXT(""),NextMemberOffset,UBMT_FLOAT32,EShaderPrecisionModifier::Half,1,4,(UniformScalarExpressions.Num() + 3) / 4,NULL);
-		const uint32 ScalarArraySize = (UniformScalarExpressions.Num() + 3) / 4 * sizeof(FVector4);
+		new(Members) FShaderParametersMetadata::FMember(TEXT("ScalarExpressions"),TEXT(""),NextMemberOffset,UBMT_FLOAT32,EShaderPrecisionModifier::Half,1,4,(UniformScalarPreshaders.Num() + 3) / 4,NULL);
+		const uint32 ScalarArraySize = (UniformScalarPreshaders.Num() + 3) / 4 * sizeof(FVector4);
 		NextMemberOffset += ScalarArraySize;
 	}
 
@@ -341,14 +411,13 @@ void FUniformExpressionSet::CreateBufferStruct()
 		}
 	}
 
-	check(Uniform2DTextureExpressions.Num() <= 128);
-	check(UniformCubeTextureExpressions.Num() <= 128);
-	check(Uniform2DArrayTextureExpressions.Num() <= 128);
-	check(UniformVolumeTextureExpressions.Num() <= 128);
-	check(UniformVirtualTextureExpressions.Num() <= 128);
+	for (uint32 TypeIndex = 0u; TypeIndex < NumMaterialTextureParameterTypes; ++TypeIndex)
+	{
+		check(UniformTextureParameters[TypeIndex].Num() <= 128);
+	}
 	check(VTStacks.Num() <= 128);
 
-	for (int32 i = 0; i < Uniform2DTextureExpressions.Num(); ++i)
+	for (int32 i = 0; i < UniformTextureParameters[(uint32)EMaterialTextureParameterType::Standard2D].Num(); ++i)
 	{
 		check((NextMemberOffset % SHADER_PARAMETER_POINTER_ALIGNMENT) == 0);
 		new(Members) FShaderParametersMetadata::FMember(*Texture2DNames[i],TEXT("Texture2D"),NextMemberOffset,UBMT_TEXTURE,EShaderPrecisionModifier::Float,1,1,0,NULL);
@@ -357,7 +426,7 @@ void FUniformExpressionSet::CreateBufferStruct()
 		NextMemberOffset += SHADER_PARAMETER_POINTER_ALIGNMENT;
 	}
 
-	for (int32 i = 0; i < UniformCubeTextureExpressions.Num(); ++i)
+	for (int32 i = 0; i < UniformTextureParameters[(uint32)EMaterialTextureParameterType::Cube].Num(); ++i)
 	{
 		check((NextMemberOffset % SHADER_PARAMETER_POINTER_ALIGNMENT) == 0);
 		new(Members) FShaderParametersMetadata::FMember(*TextureCubeNames[i],TEXT("TextureCube"),NextMemberOffset,UBMT_TEXTURE,EShaderPrecisionModifier::Float,1,1,0,NULL);
@@ -366,16 +435,16 @@ void FUniformExpressionSet::CreateBufferStruct()
 		NextMemberOffset += SHADER_PARAMETER_POINTER_ALIGNMENT;
 	}
 
-	for (int32 i = 0; i < Uniform2DArrayTextureExpressions.Num(); ++i)
+	for (int32 i = 0; i < UniformTextureParameters[(uint32)EMaterialTextureParameterType::Array2D].Num(); ++i)
 	{
 		check((NextMemberOffset % SHADER_PARAMETER_POINTER_ALIGNMENT) == 0);
 		new(Members) FShaderParametersMetadata::FMember(*Texture2DArrayNames[i], TEXT("Texture2DArray"), NextMemberOffset, UBMT_TEXTURE, EShaderPrecisionModifier::Float, 1, 1, 0, NULL);
 		NextMemberOffset += SHADER_PARAMETER_POINTER_ALIGNMENT;
-		new(Members) FShaderParametersMetadata::FMember(*Texture2DArraySamplerNames[i], TEXT("SamplerState"), NextMemberOffset, UBMT_SAMPLER, EShaderPrecisionModifier::Float, 1, 1, 0, NULL);
+		new(Members) FShaderParametersMetadata::FMember(*Texture2DSamplerNames[i], TEXT("SamplerState"), NextMemberOffset, UBMT_SAMPLER, EShaderPrecisionModifier::Float, 1, 1, 0, NULL);
 		NextMemberOffset += SHADER_PARAMETER_POINTER_ALIGNMENT;
 	}
 
-	for (int32 i = 0; i < UniformVolumeTextureExpressions.Num(); ++i)
+	for (int32 i = 0; i < UniformTextureParameters[(uint32)EMaterialTextureParameterType::Volume].Num(); ++i)
 	{
 		check((NextMemberOffset % SHADER_PARAMETER_POINTER_ALIGNMENT) == 0);
 		new(Members) FShaderParametersMetadata::FMember(*VolumeTextureNames[i],TEXT("Texture3D"),NextMemberOffset,UBMT_TEXTURE,EShaderPrecisionModifier::Float,1,1,0,NULL);
@@ -384,7 +453,7 @@ void FUniformExpressionSet::CreateBufferStruct()
 		NextMemberOffset += SHADER_PARAMETER_POINTER_ALIGNMENT;
 	}
 
-	for (int32 i = 0; i < UniformExternalTextureExpressions.Num(); ++i)
+	for (int32 i = 0; i < UniformExternalTextureParameters.Num(); ++i)
 	{
 		check((NextMemberOffset % SHADER_PARAMETER_POINTER_ALIGNMENT) == 0);
 		new(Members) FShaderParametersMetadata::FMember(*ExternalTextureNames[i], TEXT("TextureExternal"), NextMemberOffset, UBMT_TEXTURE, EShaderPrecisionModifier::Float, 1, 1, 0, NULL);
@@ -406,7 +475,7 @@ void FUniformExpressionSet::CreateBufferStruct()
 		}
 	}
 
-	for (int32 i = 0; i < UniformVirtualTextureExpressions.Num(); ++i)
+	for (int32 i = 0; i < UniformTextureParameters[(uint32)EMaterialTextureParameterType::Virtual].Num(); ++i)
 	{
 		check((NextMemberOffset % SHADER_PARAMETER_POINTER_ALIGNMENT) == 0);
 
@@ -424,23 +493,19 @@ void FUniformExpressionSet::CreateBufferStruct()
 	NextMemberOffset += SHADER_PARAMETER_POINTER_ALIGNMENT;
 
 	const uint32 StructSize = Align(NextMemberOffset, SHADER_PARAMETER_STRUCT_ALIGNMENT);
-	UniformBufferStruct.Emplace(
+	FShaderParametersMetadata* UniformBufferStruct = new FShaderParametersMetadata(
 		FShaderParametersMetadata::EUseCase::DataDrivenUniformBuffer,
-		MaterialLayoutName,
+		TEXT("Material"),
 		TEXT("MaterialUniforms"),
 		TEXT("Material"),
 		nullptr,
 		StructSize,
 		Members);
-}
 
-const FShaderParametersMetadata& FUniformExpressionSet::GetUniformBufferStruct() const
-{
-	return UniformBufferStruct.GetValue();
+	UniformBufferLayout = UniformBufferStruct->GetLayout();
+	return UniformBufferStruct;
 }
 
-
-
 FUniformExpressionSet::FVTPackedStackAndLayerIndex FUniformExpressionSet::GetVTStackAndLayerIndex(int32 UniformExpressionIndex) const
 {
 	for (int32 VTStackIndex = 0; VTStackIndex < VTStacks.Num(); ++VTStackIndex)
@@ -457,12 +522,479 @@ FUniformExpressionSet::FVTPackedStackAndLayerIndex FUniformExpressionSet::GetVTS
 	return FVTPackedStackAndLayerIndex(0xffff, 0xffff);
 }
 
+void FMaterialPreshaderData::WriteData(const void* Value, uint32 Size)
+{
+	Data.Append((uint8*)Value, Size);
+}
+
+void FMaterialPreshaderData::ReadData(void* Result, uint32 Offset, uint32 Size) const
+{
+	FMemory::Memcpy(Result, &Data[Offset], Size);
+}
+
+static void GetVectorParameter(const FUniformExpressionSet& UniformExpressionSet, uint32 ParameterIndex, const FMaterialRenderContext& Context, FLinearColor& OutValue)
+{
+	const FMaterialVectorParameterInfo& Parameter = UniformExpressionSet.GetVectorParameter(ParameterIndex);
+
+	OutValue.R = OutValue.G = OutValue.B = OutValue.A = 0;
+	bool bNeedsDefaultValue = false;
+	if (!Context.MaterialRenderProxy || !Context.MaterialRenderProxy->GetVectorValue(Parameter.ParameterInfo, &OutValue, Context))
+	{
+		const bool bOveriddenParameterOnly = Parameter.ParameterInfo.Association == EMaterialParameterAssociation::GlobalParameter;
+
+		if (AreExperimentalMaterialLayersEnabled())
+		{
+			UMaterialInterface* Interface = Context.Material.GetMaterialInterface();
+			if (!Interface || !Interface->GetVectorParameterDefaultValue(Parameter.ParameterInfo, OutValue, bOveriddenParameterOnly))
+			{
+				bNeedsDefaultValue = true;
+			}
+		}
+		else
+		{
+			bNeedsDefaultValue = true;
+		}
+	}
+
+	if (bNeedsDefaultValue)
+	{
+#if WITH_EDITOR
+		if (!Context.Material.TransientOverrides.GetVectorOverride(ParameterIndex, OutValue))
+#endif // WITH_EDITOR
+		{
+			Parameter.GetDefaultValue(OutValue);
+		}
+	}
+}
+
+static void GetScalarParameter(const FUniformExpressionSet& UniformExpressionSet, uint32 ParameterIndex, const FMaterialRenderContext& Context, FLinearColor& OutValue)
+{
+	const FMaterialScalarParameterInfo& Parameter = UniformExpressionSet.GetScalarParameter(ParameterIndex);
+
+	OutValue.A = 0;
+
+	bool bNeedsDefaultValue = false;
+	if (!Context.MaterialRenderProxy || !Context.MaterialRenderProxy->GetScalarValue(Parameter.ParameterInfo, &OutValue.A, Context))
+	{
+		const bool bOveriddenParameterOnly = Parameter.ParameterInfo.Association == EMaterialParameterAssociation::GlobalParameter;
+
+		if (AreExperimentalMaterialLayersEnabled())
+		{
+			UMaterialInterface* Interface = Context.Material.GetMaterialInterface();
+			if (!Interface || !Interface->GetScalarParameterDefaultValue(Parameter.ParameterInfo, OutValue.A, bOveriddenParameterOnly))
+			{
+				bNeedsDefaultValue = true;
+			}
+		}
+		else
+		{
+			bNeedsDefaultValue = true;
+		}
+	}
+
+	if (bNeedsDefaultValue)
+	{
+#if WITH_EDITOR
+		if (!Context.Material.TransientOverrides.GetScalarOverride(ParameterIndex, OutValue.A))
+#endif // WITH_EDITOR
+		{
+			Parameter.GetDefaultValue(OutValue.A);
+		}
+	}
+
+	OutValue.R = OutValue.G = OutValue.B = OutValue.A;
+}
+
+using FPreshaderStack = TArray<FLinearColor, TInlineAllocator<64u>>;
+
+template<typename Operation>
+static void EvaluateUnaryOp(FPreshaderStack& Stack, const Operation& Op)
+{
+	const FLinearColor Value = Stack.Pop(false);
+	Stack.Add(FLinearColor(Op(Value.R), Op(Value.G), Op(Value.B), Op(Value.A)));
+}
+
+template<typename Operation>
+static void EvaluateBinaryOp(FPreshaderStack& Stack, const Operation& Op)
+{
+	const FLinearColor Value1 = Stack.Pop(false);
+	const FLinearColor Value0 = Stack.Pop(false);
+	Stack.Add(FLinearColor(Op(Value0.R, Value1.R), Op(Value0.G, Value1.G), Op(Value0.B, Value1.B), Op(Value0.A, Value1.A)));
+}
+
+template<typename Operation>
+static void EvaluateTernaryOp(FPreshaderStack& Stack, const Operation& Op)
+{
+	const FLinearColor Value2 = Stack.Pop(false);
+	const FLinearColor Value1 = Stack.Pop(false);
+	const FLinearColor Value0 = Stack.Pop(false);
+	Stack.Add(FLinearColor(Op(Value0.R, Value1.R, Value2.R), Op(Value0.G, Value1.G, Value2.G), Op(Value0.B, Value1.B, Value2.B), Op(Value0.A, Value1.A, Value2.A)));
+}
+
+static void EvaluateDot(FPreshaderStack& Stack, const FMaterialPreshaderData& Data, uint32& Offset)
+{
+	const uint8 ValueType = Data.Read<uint8>(Offset);
+	const FLinearColor Value1 = Stack.Pop(false);
+	const FLinearColor Value0 = Stack.Pop(false);
+	float Result = Value0.R * Value1.R;
+	Result += (ValueType >= MCT_Float2) ? Value0.G * Value1.G : 0;
+	Result += (ValueType >= MCT_Float3) ? Value0.B * Value1.B : 0;
+	Result += (ValueType >= MCT_Float4) ? Value0.A * Value1.A : 0;
+	Stack.Add(FLinearColor(Result, Result, Result, Result));
+}
+
+static void EvaluateCross(FPreshaderStack& Stack, const FMaterialPreshaderData& Data, uint32& Offset)
+{
+	const uint8 ValueType = Data.Read<uint8>(Offset);
+	FLinearColor ValueB = Stack.Pop(false);
+	FLinearColor ValueA = Stack.Pop(false);
+	
+	// Must be Float3, replicate CoerceParameter behavior
+	switch (ValueType)
+	{
+	case MCT_Float:
+		ValueA.B = ValueA.G = ValueA.R;
+		ValueB.B = ValueB.G = ValueB.R;
+		break;
+	case MCT_Float1:
+		ValueA.B = ValueA.G = 0.f;
+		ValueB.B = ValueB.G = 0.f;
+		break;
+	case MCT_Float2:
+		ValueA.B = 0.f;
+		ValueB.B = 0.f;
+		break;
+	};
+
+	const FVector Cross = FVector::CrossProduct(FVector(ValueA), FVector(ValueB));
+	Stack.Add(FLinearColor(Cross.X, Cross.Y, Cross.Z, 0.0f));
+}
+
+static void EvaluateComponentSwizzle(FPreshaderStack& Stack, const FMaterialPreshaderData& Data, uint32& Offset)
+{
+	const uint8 NumElements = Data.Read<uint8>(Offset);
+	const uint8 IndexR = Data.Read<uint8>(Offset);
+	const uint8 IndexG = Data.Read<uint8>(Offset);
+	const uint8 IndexB = Data.Read<uint8>(Offset);
+	const uint8 IndexA = Data.Read<uint8>(Offset);
+
+	FLinearColor Value = Stack.Pop(false);
+	FLinearColor Result(0.0f, 0.0f, 0.0f, 0.0f);
+	switch (NumElements)
+	{
+	case 1:
+		// Replicate scalar
+		Result.R = Result.G = Result.B = Result.A = Value.Component(IndexR);
+		break;
+	case 4:
+		Result.A = Value.Component(IndexA);
+		// Fallthrough...
+	case 3:
+		Result.B = Value.Component(IndexB);
+		// Fallthrough...
+	case 2:
+		Result.G = Value.Component(IndexG);
+		Result.R = Value.Component(IndexR);
+		break;
+	default:
+		UE_LOG(LogMaterial, Fatal, TEXT("Invalid number of swizzle elements: %d"), NumElements);
+		break;
+	}
+	Stack.Add(Result);
+}
+
+static void EvaluateAppenedVector(FPreshaderStack& Stack, const FMaterialPreshaderData& Data, uint32& Offset)
+{
+	const uint8 NumComponentsA = Data.Read<uint8>(Offset);
+
+	const FLinearColor ValueB = Stack.Pop(false);
+	const FLinearColor ValueA = Stack.Pop(false);
+
+	FLinearColor Result;
+	Result.R = NumComponentsA >= 1 ? ValueA.R : (&ValueB.R)[0 - NumComponentsA];
+	Result.G = NumComponentsA >= 2 ? ValueA.G : (&ValueB.R)[1 - NumComponentsA];
+	Result.B = NumComponentsA >= 3 ? ValueA.B : (&ValueB.R)[2 - NumComponentsA];
+	Result.A = NumComponentsA >= 4 ? ValueA.A : (&ValueB.R)[3 - NumComponentsA];
+	Stack.Add(Result);
+}
+
+static const UTexture* GetTextureParameter(const FMaterialRenderContext& Context, const FMaterialPreshaderData& Data, uint32& Offset)
+{
+	const FHashedMaterialParameterInfo ParameterInfo = Data.Read<FHashedMaterialParameterInfo>(Offset);
+	const int32 TextureIndex = Data.Read<int32>(Offset);
+	
+	const UTexture* Texture = nullptr;
+	GetTextureParameterValue(ParameterInfo, TextureIndex, Context, Texture);
+	return Texture;
+}
+
+static void EvaluateTextureSize(const FMaterialRenderContext& Context, FPreshaderStack& Stack, const FMaterialPreshaderData& Data, uint32& Offset)
+{
+	const UTexture* Texture = GetTextureParameter(Context, Data, Offset);
+	if (Texture)
+	{
+		const uint32 SizeX = Texture->Resource->GetSizeX();
+		const uint32 SizeY = Texture->Resource->GetSizeY();
+		Stack.Add(FLinearColor((float)SizeX, (float)SizeY, 0.0f, 0.0f));
+	}
+}
+
+static void EvaluateTexelSize(const FMaterialRenderContext& Context, FPreshaderStack& Stack, const FMaterialPreshaderData& Data, uint32& Offset)
+{
+	const UTexture* Texture = GetTextureParameter(Context, Data, Offset);
+	if (Texture)
+	{
+		const uint32 SizeX = Texture->Resource->GetSizeX();
+		const uint32 SizeY = Texture->Resource->GetSizeY();
+		Stack.Add(FLinearColor(1.0f / (float)SizeX, 1.0f / (float)SizeY, 0.0f, 0.0f));
+	}
+}
+
+static FGuid GetExternalTextureGuid(const FMaterialRenderContext& Context, const FMaterialPreshaderData& Data, uint32& Offset)
+{
+	const FHashedName ParameterName = Data.Read<FHashedName>(Offset);
+	const FGuid ExternalTextureGuid = Data.Read<FGuid>(Offset);
+	const int32 TextureIndex = Data.Read<int32>(Offset);
+	return GetExternalTextureGuid(Context, ExternalTextureGuid, ParameterName, TextureIndex);
+}
+
+static void EvaluateExternalTextureCoordinateScaleRotation(const FMaterialRenderContext& Context, FPreshaderStack& Stack, const FMaterialPreshaderData& Data, uint32& Offset)
+{
+	const FGuid GuidToLookup = GetExternalTextureGuid(Context, Data, Offset);
+	FLinearColor Result(1.f, 0.f, 0.f, 1.f);
+	if (GuidToLookup.IsValid())
+	{
+		FExternalTextureRegistry::Get().GetExternalTextureCoordinateScaleRotation(GuidToLookup, Result);
+	}
+	Stack.Add(Result);
+}
+
+static void EvaluateExternalTextureCoordinateOffset(const FMaterialRenderContext& Context, FPreshaderStack& Stack, const FMaterialPreshaderData& Data, uint32& Offset)
+{
+	const FGuid GuidToLookup = GetExternalTextureGuid(Context, Data, Offset);
+	FLinearColor Result(0.f, 0.f, 0.f, 0.f);
+	if (GuidToLookup.IsValid())
+	{
+		FExternalTextureRegistry::Get().GetExternalTextureCoordinateOffset(GuidToLookup, Result);
+	}
+	Stack.Add(Result);
+}
+
+static void EvaluateRuntimeVirtualTextureUniform(const FMaterialRenderContext& Context, FPreshaderStack& Stack, const FMaterialPreshaderData& Data, uint32& Offset)
+{
+	const FHashedMaterialParameterInfo ParameterInfo = Data.Read<FHashedMaterialParameterInfo>(Offset);
+	const int32 TextureIndex = Data.Read<int32>(Offset);
+	const int32 VectorIndex = Data.Read<int32>(Offset);
+
+	const URuntimeVirtualTexture* Texture = nullptr;
+	if (ParameterInfo.Name.IsNone() || !Context.MaterialRenderProxy || !Context.MaterialRenderProxy->GetTextureValue(ParameterInfo, &Texture, Context))
+	{
+		Texture = GetIndexedTexture<URuntimeVirtualTexture>(Context.Material, TextureIndex);
+	}
+	if (Texture != nullptr && VectorIndex != INDEX_NONE)
+	{
+		Stack.Add(FLinearColor(Texture->GetUniformParameter(VectorIndex)));
+	}
+	else
+	{
+		Stack.Add(FLinearColor(0.f, 0.f, 0.f, 0.f));
+	}
+}
+
+/** Converts an arbitrary number into a safe divisor. i.e. FMath::Abs(Number) >= DELTA */
+static float GetSafeDivisor(float Number)
+{
+	if (FMath::Abs(Number) < DELTA)
+	{
+		if (Number < 0.0f)
+		{
+			return -DELTA;
+		}
+		else
+		{
+			return +DELTA;
+		}
+	}
+	else
+	{
+		return Number;
+	}
+}
+
+/**
+ * FORCENOINLINE is required to discourage compiler from vectorizing the Div operation, which may tempt it into optimizing divide as A * rcp(B)
+ * This will break shaders that are depending on exact divide results (see SubUV material function)
+ * Technically this could still happen for a scalar divide, but it doesn't seem to occur in practice
+ */
+FORCENOINLINE static float DivideComponent(float A, float B)
+{
+	return A / GetSafeDivisor(B);
+}
+
+static void EvaluatePreshader(const FUniformExpressionSet* UniformExpressionSet, const FMaterialRenderContext& Context, FPreshaderStack& Stack, const FMaterialPreshaderData& Data, const FMaterialUniformPreshaderHeader& Preshader, FLinearColor& OutValue)
+{
+	static const float LogToLog10 = 1.0f / FMath::Loge(10.f);
+
+	uint32 Offset = Preshader.OpcodeOffset;
+	const uint32 OffsetEnd = Offset + Preshader.OpcodeSize;
+
+	Stack.Reset();
+	while (Offset < OffsetEnd)
+	{
+		const EMaterialPreshaderOpcode Opcode = Data.ReadOpcode(Offset);
+		switch (Opcode)
+		{
+		case EMaterialPreshaderOpcode::ConstantZero:
+			Stack.Add(FLinearColor(0.0f, 0.0f, 0.0f, 0.0f));
+			break;
+		case EMaterialPreshaderOpcode::Constant:
+			Stack.Add(Data.Read<FLinearColor>(Offset));
+			break;
+		case EMaterialPreshaderOpcode::VectorParameter:
+			check(UniformExpressionSet);
+			GetVectorParameter(*UniformExpressionSet, Data.Read<uint16>(Offset), Context, Stack.AddDefaulted_GetRef());
+			break;
+		case EMaterialPreshaderOpcode::ScalarParameter:
+			check(UniformExpressionSet);
+			GetScalarParameter(*UniformExpressionSet, Data.Read<uint16>(Offset), Context, Stack.AddDefaulted_GetRef());
+			break;
+		case EMaterialPreshaderOpcode::Add: EvaluateBinaryOp(Stack, [](float Lhs, float Rhs) { return Lhs + Rhs; }); break;
+		case EMaterialPreshaderOpcode::Sub: EvaluateBinaryOp(Stack, [](float Lhs, float Rhs) { return Lhs - Rhs; }); break;
+		case EMaterialPreshaderOpcode::Mul: EvaluateBinaryOp(Stack, [](float Lhs, float Rhs) { return Lhs * Rhs; }); break;
+		case EMaterialPreshaderOpcode::Div: EvaluateBinaryOp(Stack, [](float Lhs, float Rhs) { return DivideComponent(Lhs, Rhs); }); break;
+		case EMaterialPreshaderOpcode::Fmod: EvaluateBinaryOp(Stack, [](float Lhs, float Rhs) { return FMath::Fmod(Lhs, Rhs); }); break;
+		case EMaterialPreshaderOpcode::Min: EvaluateBinaryOp(Stack, [](float Lhs, float Rhs) { return FMath::Min(Lhs, Rhs); }); break;
+		case EMaterialPreshaderOpcode::Max: EvaluateBinaryOp(Stack, [](float Lhs, float Rhs) { return FMath::Max(Lhs, Rhs); }); break;
+		case EMaterialPreshaderOpcode::Clamp: EvaluateTernaryOp(Stack, [](float A, float B, float C) { return FMath::Clamp(A, B, C); }); break;
+		case EMaterialPreshaderOpcode::Dot: EvaluateDot(Stack, Data, Offset); break;
+		case EMaterialPreshaderOpcode::Cross: EvaluateCross(Stack, Data, Offset); break;
+		case EMaterialPreshaderOpcode::Sqrt: EvaluateUnaryOp(Stack, [](float V) { return FMath::Sqrt(V); }); break;
+		case EMaterialPreshaderOpcode::Sin: EvaluateUnaryOp(Stack, [](float V) { return FMath::Sin(V); }); break;
+		case EMaterialPreshaderOpcode::Cos: EvaluateUnaryOp(Stack, [](float V) { return FMath::Cos(V); }); break;
+		case EMaterialPreshaderOpcode::Tan: EvaluateUnaryOp(Stack, [](float V) { return FMath::Tan(V); }); break;
+		case EMaterialPreshaderOpcode::Asin: EvaluateUnaryOp(Stack, [](float V) { return FMath::Asin(V); }); break;
+		case EMaterialPreshaderOpcode::Acos: EvaluateUnaryOp(Stack, [](float V) { return FMath::Acos(V); }); break;
+		case EMaterialPreshaderOpcode::Atan: EvaluateUnaryOp(Stack, [](float V) { return FMath::Atan(V); }); break;
+		case EMaterialPreshaderOpcode::Atan2: EvaluateBinaryOp(Stack, [](float A, float B) { return FMath::Atan2(A, B); }); break;
+		case EMaterialPreshaderOpcode::Abs: EvaluateUnaryOp(Stack, [](float V) { return FMath::Abs(V); }); break;
+		case EMaterialPreshaderOpcode::Saturate: EvaluateUnaryOp(Stack, [](float V) { return FMath::Clamp(V, 0.0f, 1.0f); }); break;
+		case EMaterialPreshaderOpcode::Floor: EvaluateUnaryOp(Stack, [](float V) { return FMath::FloorToFloat(V); }); break;
+		case EMaterialPreshaderOpcode::Ceil: EvaluateUnaryOp(Stack, [](float V) { return FMath::CeilToFloat(V); }); break;
+		case EMaterialPreshaderOpcode::Round: EvaluateUnaryOp(Stack, [](float V) { return FMath::RoundToFloat(V); }); break;
+		case EMaterialPreshaderOpcode::Trunc: EvaluateUnaryOp(Stack, [](float V) { return FMath::TruncToFloat(V); }); break;
+		case EMaterialPreshaderOpcode::Sign: EvaluateUnaryOp(Stack, [](float V) { return FMath::Sign(V); }); break;
+		case EMaterialPreshaderOpcode::Frac: EvaluateUnaryOp(Stack, [](float V) { return FMath::Frac(V); }); break;
+		case EMaterialPreshaderOpcode::Fractional: EvaluateUnaryOp(Stack, [](float V) { return FMath::Fractional(V); }); break;
+		case EMaterialPreshaderOpcode::Log2: EvaluateUnaryOp(Stack, [](float V) { return FMath::Log2(V); }); break;
+		case EMaterialPreshaderOpcode::Log10: EvaluateUnaryOp(Stack, [](float V) { return FMath::Loge(V) * LogToLog10; }); break;
+		case EMaterialPreshaderOpcode::ComponentSwizzle: EvaluateComponentSwizzle(Stack, Data, Offset); break;
+		case EMaterialPreshaderOpcode::AppendVector: EvaluateAppenedVector(Stack, Data, Offset); break;
+		case EMaterialPreshaderOpcode::TextureSize: EvaluateTextureSize(Context, Stack, Data, Offset); break;
+		case EMaterialPreshaderOpcode::TexelSize: EvaluateTexelSize(Context, Stack, Data, Offset); break;
+		case EMaterialPreshaderOpcode::ExternalTextureCoordinateScaleRotation: EvaluateExternalTextureCoordinateScaleRotation(Context, Stack, Data, Offset); break;
+		case EMaterialPreshaderOpcode::ExternalTextureCoordinateOffset: EvaluateExternalTextureCoordinateOffset(Context, Stack, Data, Offset); break;
+		case EMaterialPreshaderOpcode::RuntimeVirtualTextureUniform: EvaluateRuntimeVirtualTextureUniform(Context, Stack, Data, Offset); break;
+		default:
+			UE_LOG(LogMaterial, Fatal, TEXT("Unknown preshader opcode %d"), (uint8)Opcode);
+			break;
+		}
+	}
+	check(Offset == OffsetEnd);
+
+	ensure(Stack.Num() <= 1);
+	if (Stack.Num() > 0)
+	{
+		OutValue = Stack.Last();
+	}
+}
+
+void FMaterialUniformExpression::GetNumberValue(const struct FMaterialRenderContext& Context, FLinearColor& OutValue) const
+{
+	FMaterialPreshaderData PreshaderData;
+	WriteNumberOpcodes(PreshaderData);
+
+	FMaterialUniformPreshaderHeader PreshaderHeader;
+	PreshaderHeader.OpcodeOffset = 0u;
+	PreshaderHeader.OpcodeSize = PreshaderData.Num();
+	
+	FPreshaderStack Stack;
+	EvaluatePreshader(nullptr, Context, Stack, PreshaderData, PreshaderHeader, OutValue);
+}
+
+const FMaterialVectorParameterInfo* FUniformExpressionSet::FindVectorParameter(const FHashedMaterialParameterInfo& ParameterInfo) const
+{
+	for (const FMaterialVectorParameterInfo& Parameter : UniformVectorParameters)
+	{
+		if (Parameter.ParameterInfo == ParameterInfo)
+		{
+			return &Parameter;
+		}
+	}
+	return nullptr;
+}
+
+const FMaterialScalarParameterInfo* FUniformExpressionSet::FindScalarParameter(const FHashedMaterialParameterInfo& ParameterInfo) const
+{
+	for (const FMaterialScalarParameterInfo& Parameter : UniformScalarParameters)
+	{
+		if (Parameter.ParameterInfo == ParameterInfo)
+		{
+			return &Parameter;
+		}
+	}
+	return nullptr;
+}
+
+void FUniformExpressionSet::GetGameThreadTextureValue(EMaterialTextureParameterType Type, int32 Index, const UMaterialInterface* MaterialInterface, const FMaterial& Material, UTexture*& OutValue, bool bAllowOverride) const
+{
+	check(IsInGameThread());
+
+	OutValue = NULL;
+#if WITH_EDITOR
+	if (bAllowOverride)
+	{
+		UTexture* OverrideTexture = Material.TransientOverrides.GetTextureOverride_GameThread(Type, Index);
+		if (OverrideTexture)
+		{
+			OutValue = OverrideTexture;
+			return;
+		}
+	}
+#endif // WITH_EDITOR
+	const FMaterialTextureParameterInfo& Parameter = GetTextureParameter(Type, Index);
+	Parameter.GetGameThreadTextureValue(MaterialInterface, Material, OutValue);
+}
+
+void FUniformExpressionSet::GetTextureValue(EMaterialTextureParameterType Type, int32 Index, const FMaterialRenderContext& Context, const FMaterial& Material, const UTexture*& OutValue) const
+{
+	check(IsInParallelRenderingThread());
+#if WITH_EDITOR
+	{
+		UTexture* OverrideTexture = Material.TransientOverrides.GetTextureOverride_RenderThread(Type, Index);
+		if (OverrideTexture)
+		{
+			OutValue = OverrideTexture;
+			return;
+		}
+	}
+#endif // WITH_EDITOR
+
+	const FMaterialTextureParameterInfo& Parameter = GetTextureParameter(Type, Index);
+	GetTextureParameterValue(Parameter.ParameterInfo, Parameter.TextureIndex, Context, OutValue);
+}
+
+void FUniformExpressionSet::GetTextureValue(int32 Index, const FMaterial& Material, const URuntimeVirtualTexture*& OutValue) const
+{
+	check(IsInParallelRenderingThread());
+	const FMaterialTextureParameterInfo& Parameter = GetTextureParameter(EMaterialTextureParameterType::Virtual, Index);
+	OutValue = GetIndexedTexture<URuntimeVirtualTexture>(Material, Parameter.TextureIndex);
+}
+
 void FUniformExpressionSet::FillUniformBuffer(const FMaterialRenderContext& MaterialRenderContext, const FUniformExpressionCache& UniformExpressionCache, uint8* TempBuffer, int TempBufferSize) const
 {
-	check(UniformBufferStruct);
 	check(IsInParallelRenderingThread());
 
-	if (UniformBufferStruct->GetSize() > 0)
+	if (UniformBufferLayout.ConstantBufferSize > 0)
 	{
 		QUICK_SCOPE_CYCLE_COUNTER(STAT_FUniformExpressionSet_FillUniformBuffer);
 
@@ -488,8 +1020,10 @@ void FUniformExpressionSet::FillUniformBuffer(const FMaterialRenderContext& Mate
 		}
 		
 		// Dump virtual texture per physical texture uniform data
-		for (int32 ExpressionIndex = 0; ExpressionIndex < UniformVirtualTextureExpressions.Num(); ++ExpressionIndex)
+		for (int32 ExpressionIndex = 0; ExpressionIndex < GetNumTextures(EMaterialTextureParameterType::Virtual); ++ExpressionIndex)
 		{
+			const FMaterialTextureParameterInfo& Parameter = GetTextureParameter(EMaterialTextureParameterType::Virtual, ExpressionIndex);
+
 			FUintVector4* VTPackedUniform = (FUintVector4*)BufferCursor;
 			BufferCursor = VTPackedUniform + 1;
 
@@ -499,7 +1033,7 @@ void FUniformExpressionSet::FillUniformBuffer(const FMaterialRenderContext& Mate
 			if (!bFoundTexture)
 			{
 				const UTexture* Texture = nullptr;
-				UniformVirtualTextureExpressions[ExpressionIndex]->GetTextureValue(MaterialRenderContext, MaterialRenderContext.Material, Texture);
+				GetTextureValue(EMaterialTextureParameterType::Virtual, ExpressionIndex, MaterialRenderContext, MaterialRenderContext.Material, Texture);
 				if (Texture != nullptr)
 				{
 					const FVTPackedStackAndLayerIndex StackAndLayerIndex = GetVTStackAndLayerIndex(ExpressionIndex);
@@ -516,25 +1050,27 @@ void FUniformExpressionSet::FillUniformBuffer(const FMaterialRenderContext& Mate
 			if (!bFoundTexture)
 			{
 				const URuntimeVirtualTexture* Texture = nullptr;
-				UniformVirtualTextureExpressions[ExpressionIndex]->GetTextureValue(MaterialRenderContext, MaterialRenderContext.Material, Texture);
+				GetTextureValue(ExpressionIndex, MaterialRenderContext.Material, Texture);
 				if (Texture != nullptr)
 				{
-					int32 LayerIndex = UniformVirtualTextureExpressions[ExpressionIndex]->GetTextureLayerIndex();
 					IAllocatedVirtualTexture const* AllocatedVT = Texture->GetAllocatedVirtualTexture();
 					if (AllocatedVT)
 					{
-						AllocatedVT->GetPackedUniform(VTPackedUniform, LayerIndex);
+						AllocatedVT->GetPackedUniform(VTPackedUniform, Parameter.VirtualTextureLayerIndex);
 					}
 				}
 			}
 		}
 
 		// Dump vector expression into the buffer.
-		for(int32 VectorIndex = 0;VectorIndex < UniformVectorExpressions.Num();++VectorIndex)
+		FPreshaderStack PreshaderStack;
+		for(int32 VectorIndex = 0;VectorIndex < UniformVectorPreshaders.Num();++VectorIndex)
 		{
 			FLinearColor VectorValue(0, 0, 0, 0);
-			UniformVectorExpressions[VectorIndex]->GetNumberValue(MaterialRenderContext, VectorValue);
 
+			const FMaterialUniformPreshaderHeader& Preshader = UniformVectorPreshaders[VectorIndex];
+			EvaluatePreshader(this, MaterialRenderContext, PreshaderStack, UniformPreshaderData, Preshader, VectorValue);
+	
 			FLinearColor* DestAddress = (FLinearColor*)BufferCursor;
 			*DestAddress = VectorValue;
 			BufferCursor = DestAddress + 1;
@@ -542,10 +1078,12 @@ void FUniformExpressionSet::FillUniformBuffer(const FMaterialRenderContext& Mate
 		}
 
 		// Dump scalar expression into the buffer.
-		for(int32 ScalarIndex = 0;ScalarIndex < UniformScalarExpressions.Num();++ScalarIndex)
+		for(int32 ScalarIndex = 0;ScalarIndex < UniformScalarPreshaders.Num();++ScalarIndex)
 		{
 			FLinearColor VectorValue(0,0,0,0);
-			UniformScalarExpressions[ScalarIndex]->GetNumberValue(MaterialRenderContext,VectorValue);
+
+			const FMaterialUniformPreshaderHeader& Preshader = UniformScalarPreshaders[ScalarIndex];
+			EvaluatePreshader(this, MaterialRenderContext, PreshaderStack, UniformPreshaderData, Preshader, VectorValue);
 
 			float* DestAddress = (float*)BufferCursor;
 			*DestAddress = VectorValue.R;
@@ -554,7 +1092,7 @@ void FUniformExpressionSet::FillUniformBuffer(const FMaterialRenderContext& Mate
 		}
 
 		// Offsets the cursor to next first resource.
-		BufferCursor = ((float*)BufferCursor) + ((4 - UniformScalarExpressions.Num() % 4) % 4);
+		BufferCursor = ((float*)BufferCursor) + ((4 - UniformScalarPreshaders.Num() % 4) % 4);
 		check(BufferCursor <= TempBuffer + TempBufferSize);
 
 #if DO_CHECK
@@ -565,33 +1103,35 @@ void FUniformExpressionSet::FillUniformBuffer(const FMaterialRenderContext& Mate
 				NumPageTableTextures += VTStacks[i].GetNumLayers() > 4u ? 2: 1;
 			}
 	
-			check(UniformBufferStruct->GetLayout().Resources.Num() == 
-				Uniform2DTextureExpressions.Num() * 2
-				+ UniformCubeTextureExpressions.Num() * 2 
-				+ Uniform2DArrayTextureExpressions.Num() * 2
-				+ UniformVolumeTextureExpressions.Num() * 2 
-				+ UniformExternalTextureExpressions.Num() * 2
-				+ UniformVirtualTextureExpressions.Num() * 2
+			check(UniformBufferLayout.Resources.Num() == 
+				UniformTextureParameters[(uint32)EMaterialTextureParameterType::Standard2D].Num() * 2
+				+ UniformTextureParameters[(uint32)EMaterialTextureParameterType::Cube].Num() * 2
+				+ UniformTextureParameters[(uint32)EMaterialTextureParameterType::Array2D].Num() * 2
+				+ UniformTextureParameters[(uint32)EMaterialTextureParameterType::Volume].Num() * 2
+				+ UniformExternalTextureParameters.Num() * 2
+				+ UniformTextureParameters[(uint32)EMaterialTextureParameterType::Virtual].Num() * 2
 				+ NumPageTableTextures
 				+ 2);
 		}
 #endif // DO_CHECK
 
 		// Cache 2D texture uniform expressions.
-		for(int32 ExpressionIndex = 0;ExpressionIndex < Uniform2DTextureExpressions.Num();ExpressionIndex++)
+		for(int32 ExpressionIndex = 0;ExpressionIndex < GetNumTextures(EMaterialTextureParameterType::Standard2D);ExpressionIndex++)
 		{
+			const FMaterialTextureParameterInfo& Parameter = GetTextureParameter(EMaterialTextureParameterType::Standard2D, ExpressionIndex);
+
 			const UTexture* Value = nullptr;
-			Uniform2DTextureExpressions[ExpressionIndex]->GetTextureValue(MaterialRenderContext,MaterialRenderContext.Material,Value);
+			GetTextureValue(EMaterialTextureParameterType::Standard2D, ExpressionIndex, MaterialRenderContext,MaterialRenderContext.Material,Value);
 			if (Value)
 			{
 				// Pre-application validity checks (explicit ensures to avoid needless string allocation)
-				const FMaterialUniformExpressionTextureParameter* TextureParameter = (Uniform2DTextureExpressions[ExpressionIndex]->GetType() == &FMaterialUniformExpressionTextureParameter::StaticType) ?
-					&static_cast<const FMaterialUniformExpressionTextureParameter&>(*Uniform2DTextureExpressions[ExpressionIndex]) : nullptr;
+				//const FMaterialUniformExpressionTextureParameter* TextureParameter = (Uniform2DTextureExpressions[ExpressionIndex]->GetType() == &FMaterialUniformExpressionTextureParameter::StaticType) ?
+				//	&static_cast<const FMaterialUniformExpressionTextureParameter&>(*Uniform2DTextureExpressions[ExpressionIndex]) : nullptr;
 
 				// gmartin: Trying to locate UE-23902
 				if (!Value->IsValidLowLevel())
 				{
-					ensureMsgf(false, TEXT("Texture not valid! UE-23902! Parameter (%s)"), TextureParameter ? *TextureParameter->GetParameterName().ToString() : TEXT("non-parameter"));
+					ensureMsgf(false, TEXT("Texture not valid! UE-23902! Parameter (%s)"), *Parameter.ParameterName);
 				}
 
 				// Trying to track down a dangling pointer bug.
@@ -599,7 +1139,7 @@ void FUniformExpressionSet::FillUniformBuffer(const FMaterialRenderContext& Mate
 					Value->IsA<UTexture>(),
 					TEXT("Expecting a UTexture! Name(%s), Type(%s), TextureParameter(%s), Expression(%d), Material(%s)"),
 					*Value->GetName(), *Value->GetClass()->GetName(),
-					*(TextureParameter ? TextureParameter->GetParameterName() : FName()).ToString(),
+					*Parameter.ParameterName,
 					ExpressionIndex,
 					*MaterialRenderContext.Material.GetFriendlyName());
 
@@ -608,7 +1148,7 @@ void FUniformExpressionSet::FillUniformBuffer(const FMaterialRenderContext& Mate
 				{
 					FText MessageText = FText::Format(
 						NSLOCTEXT("MaterialExpressions", "IncompatibleExternalTexture", " applied to a non-external Texture2D sampler. This may work by chance on some platforms but is not portable. Please change sampler type to 'External'. Parameter '{0}' (slot {1}) in material '{2}'"),
-						FText::FromName(TextureParameter ? TextureParameter->GetParameterName() : FName()),
+						FText::FromString(*Parameter.ParameterName),
 						ExpressionIndex,
 						FText::FromString(*MaterialRenderContext.Material.GetFriendlyName()));
 
@@ -635,7 +1175,7 @@ void FUniformExpressionSet::FillUniformBuffer(const FMaterialRenderContext& Mate
 			{
 				FSamplerStateRHIRef* SamplerSource = &Value->Resource->SamplerStateRHI;
 
-				ESamplerSourceMode SourceMode = Uniform2DTextureExpressions[ExpressionIndex]->GetSamplerSource();
+				const ESamplerSourceMode SourceMode = Parameter.SamplerSource;
 				if (SourceMode == SSM_Wrap_WorldGroupSettings)
 				{
 					SamplerSource = &Wrap_WorldGroupSettings->SamplerStateRHI;
@@ -671,10 +1211,12 @@ void FUniformExpressionSet::FillUniformBuffer(const FMaterialRenderContext& Mate
 		}
 
 		// Cache cube texture uniform expressions.
-		for(int32 ExpressionIndex = 0;ExpressionIndex < UniformCubeTextureExpressions.Num();ExpressionIndex++)
+		for(int32 ExpressionIndex = 0;ExpressionIndex < GetNumTextures(EMaterialTextureParameterType::Cube); ExpressionIndex++)
 		{
+			const FMaterialTextureParameterInfo& Parameter = GetTextureParameter(EMaterialTextureParameterType::Cube, ExpressionIndex);
+
 			const UTexture* Value = nullptr;
-			UniformCubeTextureExpressions[ExpressionIndex]->GetTextureValue(MaterialRenderContext,MaterialRenderContext.Material,Value);
+			GetTextureValue(EMaterialTextureParameterType::Cube, ExpressionIndex, MaterialRenderContext,MaterialRenderContext.Material,Value);
 
 			void** ResourceTableTexturePtr = (void**)((uint8*)BufferCursor + 0 * SHADER_PARAMETER_POINTER_ALIGNMENT);
 			void** ResourceTableSamplerPtr = (void**)((uint8*)BufferCursor + 1 * SHADER_PARAMETER_POINTER_ALIGNMENT);
@@ -687,7 +1229,7 @@ void FUniformExpressionSet::FillUniformBuffer(const FMaterialRenderContext& Mate
 				*ResourceTableTexturePtr = Value->TextureReference.TextureReferenceRHI;
 				FSamplerStateRHIRef* SamplerSource = &Value->Resource->SamplerStateRHI;
 
-				ESamplerSourceMode SourceMode = UniformCubeTextureExpressions[ExpressionIndex]->GetSamplerSource();
+				const ESamplerSourceMode SourceMode = Parameter.SamplerSource;
 				if (SourceMode == SSM_Wrap_WorldGroupSettings)
 				{
 					SamplerSource = &Wrap_WorldGroupSettings->SamplerStateRHI;
@@ -710,11 +1252,12 @@ void FUniformExpressionSet::FillUniformBuffer(const FMaterialRenderContext& Mate
 		}
 
 		// Cache 2d array texture uniform expressions.
-		for (int32 ExpressionIndex = 0; ExpressionIndex < Uniform2DArrayTextureExpressions.Num(); ExpressionIndex++)
+		for (int32 ExpressionIndex = 0; ExpressionIndex < GetNumTextures(EMaterialTextureParameterType::Array2D); ExpressionIndex++)
 		{
+			const FMaterialTextureParameterInfo& Parameter = GetTextureParameter(EMaterialTextureParameterType::Array2D, ExpressionIndex);
+
 			const UTexture* Value;
-			
-			Uniform2DArrayTextureExpressions[ExpressionIndex]->GetTextureValue(MaterialRenderContext, MaterialRenderContext.Material, Value);
+			GetTextureValue(EMaterialTextureParameterType::Array2D, ExpressionIndex, MaterialRenderContext, MaterialRenderContext.Material, Value);
 
 			void** ResourceTableTexturePtr = (void**)((uint8*)BufferCursor + 0 * SHADER_PARAMETER_POINTER_ALIGNMENT);
 			void** ResourceTableSamplerPtr = (void**)((uint8*)BufferCursor + 1 * SHADER_PARAMETER_POINTER_ALIGNMENT);
@@ -725,7 +1268,7 @@ void FUniformExpressionSet::FillUniformBuffer(const FMaterialRenderContext& Mate
 				check(Value->TextureReference.TextureReferenceRHI);
 				*ResourceTableTexturePtr = Value->TextureReference.TextureReferenceRHI;
 				FSamplerStateRHIRef* SamplerSource = &Value->Resource->SamplerStateRHI;
-				ESamplerSourceMode SourceMode = Uniform2DArrayTextureExpressions[ExpressionIndex]->GetSamplerSource();
+				ESamplerSourceMode SourceMode = Parameter.SamplerSource;
 				if (SourceMode == SSM_Wrap_WorldGroupSettings)
 				{
 					SamplerSource = &Wrap_WorldGroupSettings->SamplerStateRHI;
@@ -748,10 +1291,12 @@ void FUniformExpressionSet::FillUniformBuffer(const FMaterialRenderContext& Mate
 		}
 
 		// Cache volume texture uniform expressions.
-		for (int32 ExpressionIndex = 0;ExpressionIndex < UniformVolumeTextureExpressions.Num();ExpressionIndex++)
+		for (int32 ExpressionIndex = 0;ExpressionIndex < GetNumTextures(EMaterialTextureParameterType::Volume);ExpressionIndex++)
 		{
+			const FMaterialTextureParameterInfo& Parameter = GetTextureParameter(EMaterialTextureParameterType::Volume, ExpressionIndex);
+
 			const UTexture* Value = nullptr;
-			UniformVolumeTextureExpressions[ExpressionIndex]->GetTextureValue(MaterialRenderContext,MaterialRenderContext.Material,Value);
+			GetTextureValue(EMaterialTextureParameterType::Volume, ExpressionIndex, MaterialRenderContext,MaterialRenderContext.Material,Value);
 
 			void** ResourceTableTexturePtr = (void**)((uint8*)BufferCursor + 0 * SHADER_PARAMETER_POINTER_ALIGNMENT);
 			void** ResourceTableSamplerPtr = (void**)((uint8*)BufferCursor + 1 * SHADER_PARAMETER_POINTER_ALIGNMENT);
@@ -764,7 +1309,7 @@ void FUniformExpressionSet::FillUniformBuffer(const FMaterialRenderContext& Mate
 				*ResourceTableTexturePtr = Value->TextureReference.TextureReferenceRHI;
 				FSamplerStateRHIRef* SamplerSource = &Value->Resource->SamplerStateRHI;
 
-				ESamplerSourceMode SourceMode = UniformVolumeTextureExpressions[ExpressionIndex]->GetSamplerSource();
+				const ESamplerSourceMode SourceMode = Parameter.SamplerSource;
 				if (SourceMode == SSM_Wrap_WorldGroupSettings)
 				{
 					SamplerSource = &Wrap_WorldGroupSettings->SamplerStateRHI;
@@ -790,7 +1335,7 @@ void FUniformExpressionSet::FillUniformBuffer(const FMaterialRenderContext& Mate
 		uint32 ImmutableSamplerIndex = 0;
 		FImmutableSamplerState& ImmutableSamplerState = MaterialRenderContext.MaterialRenderProxy->ImmutableSamplerState;
 		ImmutableSamplerState.Reset();
-		for (int32 ExpressionIndex = 0; ExpressionIndex < UniformExternalTextureExpressions.Num(); ExpressionIndex++)
+		for (int32 ExpressionIndex = 0; ExpressionIndex < UniformExternalTextureParameters.Num(); ExpressionIndex++)
 		{
 			FTextureRHIRef TextureRHI;
 			FSamplerStateRHIRef SamplerStateRHI;
@@ -800,7 +1345,7 @@ void FUniformExpressionSet::FillUniformBuffer(const FMaterialRenderContext& Mate
 			BufferCursor = ((uint8*)BufferCursor) + (SHADER_PARAMETER_POINTER_ALIGNMENT * 2);
 			check(BufferCursor <= TempBuffer + TempBufferSize);
 
-			if (UniformExternalTextureExpressions[ExpressionIndex]->GetExternalTexture(MaterialRenderContext, TextureRHI, SamplerStateRHI))
+			if (UniformExternalTextureParameters[ExpressionIndex].GetExternalTexture(MaterialRenderContext, TextureRHI, SamplerStateRHI))
 			{
 				*ResourceTableTexturePtr = TextureRHI;
 				*ResourceTableSamplerPtr = SamplerStateRHI;
@@ -858,8 +1403,10 @@ void FUniformExpressionSet::FillUniformBuffer(const FMaterialRenderContext& Mate
 		}
 
 		// Cache virtual texture physical uniform expressions.
-		for (int32 ExpressionIndex = 0; ExpressionIndex < UniformVirtualTextureExpressions.Num(); ExpressionIndex++)
+		for (int32 ExpressionIndex = 0; ExpressionIndex < GetNumTextures(EMaterialTextureParameterType::Virtual); ExpressionIndex++)
 		{
+			const FMaterialTextureParameterInfo& Parameter = GetTextureParameter(EMaterialTextureParameterType::Virtual, ExpressionIndex);
+
 			FTextureRHIRef TexturePhysicalRHI;
 			FSamplerStateRHIRef PhysicalSamplerStateRHI;
 
@@ -872,7 +1419,7 @@ void FUniformExpressionSet::FillUniformBuffer(const FMaterialRenderContext& Mate
 			if (!bValidResources)
 			{
 				const UTexture* Texture = nullptr;
-				UniformVirtualTextureExpressions[ExpressionIndex]->GetTextureValue(MaterialRenderContext, MaterialRenderContext.Material, Texture);
+				GetTextureValue(EMaterialTextureParameterType::Virtual, ExpressionIndex, MaterialRenderContext, MaterialRenderContext.Material, Texture);
 				if (Texture && Texture->Resource)
 				{
 					const FVTPackedStackAndLayerIndex StackAndLayerIndex = GetVTStackAndLayerIndex(ExpressionIndex);
@@ -897,13 +1444,13 @@ void FUniformExpressionSet::FillUniformBuffer(const FMaterialRenderContext& Mate
 			if (!bValidResources)
 			{
 				const URuntimeVirtualTexture* Texture = nullptr;
-				UniformVirtualTextureExpressions[ExpressionIndex]->GetTextureValue(MaterialRenderContext, MaterialRenderContext.Material, Texture);
+				GetTextureValue(ExpressionIndex, MaterialRenderContext.Material, Texture);
 				if (Texture != nullptr)
 				{
 					IAllocatedVirtualTexture const* AllocatedVT = Texture->GetAllocatedVirtualTexture();
 					if (AllocatedVT != nullptr)
 					{
-						const int32 LayerIndex = UniformVirtualTextureExpressions[ExpressionIndex]->GetTextureLayerIndex();
+						const int32 LayerIndex = Parameter.VirtualTextureLayerIndex;
 						FRHIShaderResourceView* PhysicalViewRHI = AllocatedVT->GetPhysicalTextureSRV(LayerIndex, Texture->IsLayerSRGB(LayerIndex));
 						if (PhysicalViewRHI != nullptr)
 						{
@@ -941,10 +1488,10 @@ uint32 FUniformExpressionSet::GetReferencedTexture2DRHIHash(const FMaterialRende
 {
 	uint32 BaseHash = 0;
 
-	for (int32 ExpressionIndex = 0; ExpressionIndex < Uniform2DTextureExpressions.Num(); ExpressionIndex++)
+	for (int32 ExpressionIndex = 0; ExpressionIndex < GetNumTextures(EMaterialTextureParameterType::Standard2D); ExpressionIndex++)
 	{
 		const UTexture* Value;
-		Uniform2DTextureExpressions[ExpressionIndex]->GetTextureValue(MaterialRenderContext, MaterialRenderContext.Material, Value);
+		GetTextureValue(EMaterialTextureParameterType::Standard2D, ExpressionIndex, MaterialRenderContext, MaterialRenderContext.Material, Value);
 
 		const uint32 ValidTextureTypes = MCT_Texture2D | MCT_TextureVirtual | MCT_TextureExternal;
 
@@ -959,91 +1506,6 @@ uint32 FUniformExpressionSet::GetReferencedTexture2DRHIHash(const FMaterialRende
 	return BaseHash;
 }
 
-/** Converts an arbitrary number into a safe divisor. i.e. FMath::Abs(Number) >= DELTA */
-static float GetSafeDivisor(float Number)
-{
-	if (FMath::Abs(Number) < DELTA)
-	{
-		if (Number < 0.0f)
-		{
-			return -DELTA;
-		}
-		else
-		{
-			return +DELTA;
-		}
-	}
-	else
-	{
-		return Number;
-	}
-}
-
-/**
- * FORCENOINLINE is required to discourage compiler from vectorizing the Div operation, which may tempt it into optimizing divide as A * rcp(B)
- * This will break shaders that are depending on exact divide results (see SubUV material function)
- * Technically this could still happen for a scalar divide, but it doesn't seem to occur in practice
- */
-FORCENOINLINE static float DivideComponent(float A, float B)
-{
-	return A / GetSafeDivisor(B);
-}
-
-void FMaterialUniformExpressionFoldedMath::GetNumberValue(const FMaterialRenderContext& Context, FLinearColor& OutValue) const
-{
-	FLinearColor ValueA = FLinearColor::Black;
-	FLinearColor ValueB = FLinearColor::Black;
-	A->GetNumberValue(Context, ValueA);
-	B->GetNumberValue(Context, ValueB);
-
-	switch (Op)
-	{
-	case FMO_Add: OutValue = ValueA + ValueB; break;
-	case FMO_Sub: OutValue = ValueA - ValueB; break;
-	case FMO_Mul: OutValue = ValueA * ValueB; break;
-	case FMO_Div:
-		OutValue.R = DivideComponent(ValueA.R, ValueB.R);
-		OutValue.G = DivideComponent(ValueA.G, ValueB.G);
-		OutValue.B = DivideComponent(ValueA.B, ValueB.B);
-		OutValue.A = DivideComponent(ValueA.A, ValueB.A);
-		break;
-	case FMO_Dot:
-	{
-		check(ValueType & MCT_Float);
-		float DotProduct = ValueA.R * ValueB.R;
-		DotProduct += (ValueType >= MCT_Float2) ? ValueA.G * ValueB.G : 0;
-		DotProduct += (ValueType >= MCT_Float3) ? ValueA.B * ValueB.B : 0;
-		DotProduct += (ValueType >= MCT_Float4) ? ValueA.A * ValueB.A : 0;
-		OutValue.R = OutValue.G = OutValue.B = OutValue.A = DotProduct;
-	}
-	break;
-	case FMO_Cross:
-	{
-		// Must be Float3, replicate CoerceParameter behavior
-		switch (ValueType)
-		{
-		case MCT_Float:
-			ValueA.B = ValueA.G = ValueA.R;
-			ValueB.B = ValueB.G = ValueB.R;
-			break;
-		case MCT_Float1:
-			ValueA.B = ValueA.G = 0.f;
-			ValueB.B = ValueB.G = 0.f;
-			break;
-		case MCT_Float2:
-			ValueA.B = 0.f;
-			ValueB.B = 0.f;
-			break;
-		};
-		FVector Cross = FVector::CrossProduct(FVector(ValueA), FVector(ValueB));
-		OutValue.R = Cross.X; OutValue.G = Cross.Y; OutValue.B = Cross.Z;
-		OutValue.A = 0.f;
-	}
-	break;
-	default: UE_LOG(LogMaterial, Fatal, TEXT("Unknown folded math operation: %08x"), (int32)Op);
-	};
-}
-
 FMaterialUniformExpressionTexture::FMaterialUniformExpressionTexture() :
 	TextureIndex(INDEX_NONE),
 	TextureLayerIndex(INDEX_NONE),
@@ -1052,9 +1514,7 @@ FMaterialUniformExpressionTexture::FMaterialUniformExpressionTexture() :
 	SamplerType(SAMPLERTYPE_Color),
 #endif
 	SamplerSource(SSM_FromTextureAsset),
-	bVirtualTexture(false),
-	TransientOverrideValue_GameThread(NULL),
-	TransientOverrideValue_RenderThread(NULL)
+	bVirtualTexture(false)
 {}
 
 FMaterialUniformExpressionTexture::FMaterialUniformExpressionTexture(int32 InTextureIndex, EMaterialSamplerType InSamplerType, ESamplerSourceMode InSamplerSource, bool InVirtualTexture) :
@@ -1065,9 +1525,7 @@ FMaterialUniformExpressionTexture::FMaterialUniformExpressionTexture(int32 InTex
 	SamplerType(InSamplerType),
 #endif
 	SamplerSource(InSamplerSource),
-	bVirtualTexture(InVirtualTexture),
-	TransientOverrideValue_GameThread(NULL),
-	TransientOverrideValue_RenderThread(NULL)
+	bVirtualTexture(InVirtualTexture)
 {
 }
 
@@ -1080,59 +1538,14 @@ FMaterialUniformExpressionTexture::FMaterialUniformExpressionTexture(int32 InTex
 #endif
 	, SamplerSource(SSM_Wrap_WorldGroupSettings)
 	, bVirtualTexture(true)
-	, TransientOverrideValue_GameThread(NULL)
-	, TransientOverrideValue_RenderThread(NULL)
-{
-}
-
-void FMaterialUniformExpressionTexture::Serialize(FArchive& Ar)
-{
-	int32 SamplerSourceInt = (int32)SamplerSource;
-	Ar << TextureIndex << TextureLayerIndex << PageTableLayerIndex << SamplerSourceInt << bVirtualTexture;
-	SamplerSource = (ESamplerSourceMode)SamplerSourceInt;
-}
-
-void FMaterialUniformExpressionTexture::SetTransientOverrideTextureValue( UTexture* InOverrideTexture )
-{
-	TransientOverrideValue_GameThread = InOverrideTexture;
-	FMaterialUniformExpressionTexture* ExpressionTexture = this;
-	ENQUEUE_RENDER_COMMAND(SetTransientOverrideTextureValueCommand)(
-		[ExpressionTexture, InOverrideTexture](FRHICommandListImmediate& RHICmdList)
-		{
-			ExpressionTexture->TransientOverrideValue_RenderThread = InOverrideTexture;
-		});
-}
-
-void FMaterialUniformExpressionTexture::GetTextureValue(const FMaterialRenderContext& Context, const FMaterial& Material, const UTexture*& OutValue) const
-{
-	check(IsInParallelRenderingThread());
-	if (TransientOverrideValue_RenderThread != NULL)
-	{
-		OutValue = TransientOverrideValue_RenderThread;
-	}
-	else
-	{
-		OutValue = GetIndexedTexture<UTexture>(Material, TextureIndex);
-	}
-}
-
-void FMaterialUniformExpressionTexture::GetTextureValue(const FMaterialRenderContext& Context, const FMaterial& Material, const URuntimeVirtualTexture*& OutValue) const
 {
-	check(IsInParallelRenderingThread());
-	OutValue = GetIndexedTexture<URuntimeVirtualTexture>(Material, TextureIndex);
 }
 
-void FMaterialUniformExpressionTexture::GetGameThreadTextureValue(const UMaterialInterface* MaterialInterface, const FMaterial& Material, UTexture*& OutValue, bool bAllowOverride) const
+void FMaterialUniformExpressionTexture::GetTextureParameterInfo(FMaterialTextureParameterInfo& OutParameter) const
 {
-	check(IsInGameThread());
-	if (bAllowOverride && TransientOverrideValue_GameThread)
-	{
-		OutValue = TransientOverrideValue_GameThread;
-	}
-	else
-	{
-		OutValue = GetIndexedTexture<UTexture>(Material, TextureIndex);
-	}
+	OutParameter.TextureIndex = TextureIndex;
+	OutParameter.SamplerSource = SamplerSource;
+	OutParameter.VirtualTextureLayerIndex = TextureLayerIndex;
 }
 
 bool FMaterialUniformExpressionTexture::IsIdentical(const FMaterialUniformExpression* OtherExpression) const
@@ -1149,78 +1562,6 @@ bool FMaterialUniformExpressionTexture::IsIdentical(const FMaterialUniformExpres
 		bVirtualTexture == OtherTextureExpression->bVirtualTexture;
 }
 
-void FMaterialUniformExpressionTextureParameter::GetTextureValue(const FMaterialRenderContext& Context, const FMaterial& Material, const UTexture*& OutValue) const
-{
-	check(IsInParallelRenderingThread());
-	if (TransientOverrideValue_RenderThread != NULL)
-	{
-		OutValue = TransientOverrideValue_RenderThread;
-	}
-	else
-	{
-		if (!Context.MaterialRenderProxy || !Context.MaterialRenderProxy->GetTextureValue(ParameterInfo, &OutValue, Context))
-		{
-			UTexture* Value = nullptr;
-
-			if (AreExperimentalMaterialLayersEnabled())
-			{
-				UMaterialInterface* Interface = Context.Material.GetMaterialInterface();
-				if (!Interface || !Interface->GetTextureParameterDefaultValue(ParameterInfo, Value))
-				{
-					Value = GetIndexedTexture<UTexture>(Material, TextureIndex);
-				}
-			}
-			else
-			{
-				Value = GetIndexedTexture<UTexture>(Material, TextureIndex);
-			}
-
-			OutValue = Value;
-		}
-	}
-}
-
-void FMaterialUniformExpressionTextureParameter::GetTextureValue(const FMaterialRenderContext& Context, const FMaterial& Material, const URuntimeVirtualTexture*& OutValue) const
-{
-	if (!Context.MaterialRenderProxy || !Context.MaterialRenderProxy->GetTextureValue(ParameterInfo, &OutValue, Context))
-	{
-		URuntimeVirtualTexture* Value = nullptr;
-
-		if (AreExperimentalMaterialLayersEnabled())
-		{
-			UMaterialInterface* Interface = Context.Material.GetMaterialInterface();
-			if (!Interface || !Interface->GetRuntimeVirtualTextureParameterDefaultValue(ParameterInfo, Value))
-			{
-				Value = GetIndexedTexture<URuntimeVirtualTexture>(Material, TextureIndex);
-			}
-		}
-		else
-		{
-			Value = GetIndexedTexture<URuntimeVirtualTexture>(Material, TextureIndex);
-		}
-
-		OutValue = Value;
-	}
-}
-
-void FMaterialUniformExpressionTextureParameter::GetGameThreadTextureValue(const UMaterialInterface* MaterialInterface, const FMaterial& Material, UTexture*& OutValue, bool bAllowOverride) const
-{
-	check(IsInGameThread());
-	if (bAllowOverride && TransientOverrideValue_GameThread != NULL)
-	{
-		OutValue = TransientOverrideValue_GameThread;
-	}
-	else
-	{
-		OutValue = NULL;
-		const bool bOverrideValuesOnly = !AreExperimentalMaterialLayersEnabled();
-		if (!MaterialInterface->GetTextureParameterValue(ParameterInfo, OutValue, bOverrideValuesOnly))
-		{
-			OutValue = GetIndexedTexture<UTexture>(Material, TextureIndex);
-		}
-	}
-}
-
 FMaterialUniformExpressionExternalTextureBase::FMaterialUniformExpressionExternalTextureBase(int32 InSourceTextureIndex)
 	: SourceTextureIndex(InSourceTextureIndex)
 {}
@@ -1231,12 +1572,6 @@ FMaterialUniformExpressionExternalTextureBase::FMaterialUniformExpressionExterna
 {
 }
 
-void FMaterialUniformExpressionExternalTextureBase::Serialize(FArchive& Ar)
-{
-	Ar << SourceTextureIndex;
-	Ar << ExternalTextureGuid;
-}
-
 bool FMaterialUniformExpressionExternalTextureBase::IsIdentical(const FMaterialUniformExpression* OtherExpression) const
 {
 	if (GetType() != OtherExpression->GetType())
@@ -1250,34 +1585,14 @@ bool FMaterialUniformExpressionExternalTextureBase::IsIdentical(const FMaterialU
 
 FGuid FMaterialUniformExpressionExternalTextureBase::ResolveExternalTextureGUID(const FMaterialRenderContext& Context, TOptional<FName> ParameterName) const
 {
-	// Use the compile-time GUID if it is set
-	if (ExternalTextureGuid.IsValid())
-	{
-		return ExternalTextureGuid;
-	}
-
-	const UTexture* TextureParameterObject = nullptr;
-	if (ParameterName.IsSet() && Context.MaterialRenderProxy && Context.MaterialRenderProxy->GetTextureValue(ParameterName.GetValue(), &TextureParameterObject, Context) && TextureParameterObject)
-	{
-		return TextureParameterObject->GetExternalTextureGuid();
-	}
-
-	// Otherwise attempt to use the texture index in the material, if it's valid
-	const UTexture* TextureObject = SourceTextureIndex != INDEX_NONE ? GetIndexedTexture<UTexture>(Context.Material, SourceTextureIndex) : nullptr;
-	if (TextureObject)
-	{
-		return TextureObject->GetExternalTextureGuid();
-	}
-
-	return FGuid();
+	const FHashedName HashedParameterName = ParameterName.IsSet() ? FHashedName(ParameterName.GetValue()) : FHashedName();
+	return GetExternalTextureGuid(Context, ExternalTextureGuid, HashedParameterName, SourceTextureIndex);
 }
 
-bool FMaterialUniformExpressionExternalTexture::GetExternalTexture(const FMaterialRenderContext& Context, FTextureRHIRef& OutTextureRHI, FSamplerStateRHIRef& OutSamplerStateRHI) const
+void FMaterialUniformExpressionExternalTexture::GetExternalTextureParameterInfo(FMaterialExternalTextureParameterInfo& OutParameter) const
 {
-	check(IsInParallelRenderingThread());
-
-	FGuid GuidToLookup = ResolveExternalTextureGUID(Context);
-	return FExternalTextureRegistry::Get().GetExternalTexture(Context.MaterialRenderProxy, GuidToLookup, OutTextureRHI, OutSamplerStateRHI);
+	OutParameter.ExternalTextureGuid = ExternalTextureGuid;
+	OutParameter.SourceTextureIndex = SourceTextureIndex;
 }
 
 FMaterialUniformExpressionExternalTextureParameter::FMaterialUniformExpressionExternalTextureParameter()
@@ -1288,18 +1603,10 @@ FMaterialUniformExpressionExternalTextureParameter::FMaterialUniformExpressionEx
 	, ParameterName(InParameterName)
 {}
 
-void FMaterialUniformExpressionExternalTextureParameter::Serialize(FArchive& Ar)
-{
-	Ar << ParameterName;
-	Super::Serialize(Ar);
-}
-
-bool FMaterialUniformExpressionExternalTextureParameter::GetExternalTexture(const FMaterialRenderContext& Context, FTextureRHIRef& OutTextureRHI, FSamplerStateRHIRef& OutSamplerStateRHI) const
+void FMaterialUniformExpressionExternalTextureParameter::GetExternalTextureParameterInfo(FMaterialExternalTextureParameterInfo& OutParameter) const
 {
-	check(IsInParallelRenderingThread());
-
-	FGuid GuidToLookup = ResolveExternalTextureGUID(Context, ParameterName);
-	return FExternalTextureRegistry::Get().GetExternalTexture(Context.MaterialRenderProxy, GuidToLookup, OutTextureRHI, OutSamplerStateRHI);
+	Super::GetExternalTextureParameterInfo(OutParameter);
+	OutParameter.ParameterName = *ParameterName.ToString();
 }
 
 bool FMaterialUniformExpressionExternalTextureParameter::IsIdentical(const FMaterialUniformExpression* OtherExpression) const
@@ -1313,7 +1620,7 @@ bool FMaterialUniformExpressionExternalTextureParameter::IsIdentical(const FMate
 	return ParameterName == Other->ParameterName && Super::IsIdentical(OtherExpression);
 }
 
-void FMaterialUniformExpressionVectorParameter::GetGameThreadNumberValue(const UMaterialInterface* SourceMaterialToCopyFrom, FLinearColor& OutValue) const
+void FMaterialScalarParameterInfo::GetGameThreadNumberValue(const UMaterialInterface* SourceMaterialToCopyFrom, float& OutValue) const
 {
 	check(IsInGameThread());
 	checkSlow(SourceMaterialToCopyFrom);
@@ -1326,8 +1633,8 @@ void FMaterialUniformExpressionVectorParameter::GetGameThreadNumberValue(const U
 
 		if (MatInst)
 		{
-			const FVectorParameterValue* ParameterValue = GameThread_FindParameterByName(MatInst->VectorParameterValues, ParameterInfo);
-			if(ParameterValue)
+			const FScalarParameterValue* ParameterValue = GameThread_FindParameterByName(MatInst->ScalarParameterValues, ParameterInfo);
+			if (ParameterValue)
 			{
 				OutValue = ParameterValue->ParameterValue;
 				break;
@@ -1346,7 +1653,7 @@ void FMaterialUniformExpressionVectorParameter::GetGameThreadNumberValue(const U
 	}
 }
 
-void FMaterialUniformExpressionScalarParameter::GetGameThreadNumberValue(const UMaterialInterface* SourceMaterialToCopyFrom, float& OutValue) const
+void FMaterialVectorParameterInfo::GetGameThreadNumberValue(const UMaterialInterface* SourceMaterialToCopyFrom, FLinearColor& OutValue) const
 {
 	check(IsInGameThread());
 	checkSlow(SourceMaterialToCopyFrom);
@@ -1359,8 +1666,8 @@ void FMaterialUniformExpressionScalarParameter::GetGameThreadNumberValue(const U
 
 		if (MatInst)
 		{
-			const FScalarParameterValue* ParameterValue = GameThread_FindParameterByName(MatInst->ScalarParameterValues, ParameterInfo);
-			if(ParameterValue)
+			const FVectorParameterValue* ParameterValue = GameThread_FindParameterByName(MatInst->VectorParameterValues, ParameterInfo);
+			if (ParameterValue)
 			{
 				OutValue = ParameterValue->ParameterValue;
 				break;
@@ -1379,46 +1686,27 @@ void FMaterialUniformExpressionScalarParameter::GetGameThreadNumberValue(const U
 	}
 }
 
-void FMaterialUniformExpressionScalarParameter::GetGameThreadUsedAsAtlas(const UMaterialInterface* SourceMaterialToCopyFrom, bool& OutValue, TSoftObjectPtr<class UCurveLinearColor>& Curve, TSoftObjectPtr<class UCurveLinearColorAtlas>& Atlas) const
-{
-	check(IsInGameThread());
-	checkSlow(SourceMaterialToCopyFrom);
-
-	const UMaterialInterface* It = SourceMaterialToCopyFrom;
-
-	const UMaterialInstance* MatInst = Cast<UMaterialInstance>(It);
-
-	if (MatInst)
-	{
-		MatInst->IsScalarParameterUsedAsAtlasPosition(ParameterInfo, OutValue, Curve, Atlas);		
-	}
-}
-
-namespace
+void FMaterialTextureParameterInfo::GetGameThreadTextureValue(const UMaterialInterface* MaterialInterface, const FMaterial& Material, UTexture*& OutValue) const
 {
-	void SerializeOptional(FArchive& Ar, TOptional<FName>& OptionalName)
+	if (!ParameterInfo.Name.IsNone())
 	{
-		bool bIsSet = OptionalName.IsSet();
-		Ar << bIsSet;
-
-		if (bIsSet)
+		const bool bOverrideValuesOnly = !AreExperimentalMaterialLayersEnabled();
+		if (!MaterialInterface->GetTextureParameterValue(ParameterInfo, OutValue, bOverrideValuesOnly))
 		{
-			if (!OptionalName.IsSet())
-			{
-				OptionalName.Emplace();
-			}
-
-			Ar << OptionalName.GetValue();
+			OutValue = GetIndexedTexture<UTexture>(Material, TextureIndex);
 		}
 	}
+	else
+	{
+		OutValue = GetIndexedTexture<UTexture>(Material, TextureIndex);
+	}
 }
 
-void FMaterialUniformExpressionExternalTextureCoordinateScaleRotation::Serialize(FArchive& Ar)
+bool FMaterialExternalTextureParameterInfo::GetExternalTexture(const FMaterialRenderContext& Context, FTextureRHIRef& OutTextureRHI, FSamplerStateRHIRef& OutSamplerStateRHI) const
 {
-	// Write out the optional parameter name
-	SerializeOptional(Ar, ParameterName);
-
-	Super::Serialize(Ar);
+	check(IsInParallelRenderingThread());
+	const FGuid GuidToLookup = GetExternalTextureGuid(Context, ExternalTextureGuid, ParameterName, SourceTextureIndex);
+	return FExternalTextureRegistry::Get().GetExternalTexture(Context.MaterialRenderProxy, GuidToLookup, OutTextureRHI, OutSamplerStateRHI);
 }
 
 bool FMaterialUniformExpressionExternalTextureCoordinateScaleRotation::IsIdentical(const FMaterialUniformExpression* OtherExpression) const
@@ -1432,21 +1720,10 @@ bool FMaterialUniformExpressionExternalTextureCoordinateScaleRotation::IsIdentic
 	return ParameterName == Other->ParameterName;
 }
 
-void FMaterialUniformExpressionExternalTextureCoordinateScaleRotation::GetNumberValue(const FMaterialRenderContext& Context, FLinearColor& OutValue) const
-{
-	FGuid GuidToLookup = ResolveExternalTextureGUID(Context, ParameterName);
-	if (!GuidToLookup.IsValid() || !FExternalTextureRegistry::Get().GetExternalTextureCoordinateScaleRotation(GuidToLookup, OutValue))
-	{
-		OutValue = FLinearColor(1.f, 0.f, 0.f, 1.f);
-	}
-}
-
-void FMaterialUniformExpressionExternalTextureCoordinateOffset::Serialize(FArchive& Ar)
+void FMaterialUniformExpressionExternalTextureCoordinateScaleRotation::WriteNumberOpcodes(FMaterialPreshaderData& OutData) const
 {
-	// Write out the optional parameter name
-	SerializeOptional(Ar, ParameterName);
-
-	Super::Serialize(Ar);
+	const FHashedName HashedParameterName = ParameterName.IsSet() ? FHashedName(ParameterName.GetValue()) : FHashedName();
+	OutData.WriteOpcode(EMaterialPreshaderOpcode::ExternalTextureCoordinateScaleRotation).Write(HashedParameterName).Write(ExternalTextureGuid).Write<int32>(SourceTextureIndex);
 }
 
 bool FMaterialUniformExpressionExternalTextureCoordinateOffset::IsIdentical(const FMaterialUniformExpression* OtherExpression) const
@@ -1460,14 +1737,10 @@ bool FMaterialUniformExpressionExternalTextureCoordinateOffset::IsIdentical(cons
 	return ParameterName == Other->ParameterName;
 }
 
-void FMaterialUniformExpressionExternalTextureCoordinateOffset::GetNumberValue(const FMaterialRenderContext& Context, FLinearColor& OutValue) const
+void FMaterialUniformExpressionExternalTextureCoordinateOffset::WriteNumberOpcodes(FMaterialPreshaderData& OutData) const
 {
-	FGuid GuidToLookup = ResolveExternalTextureGUID(Context, ParameterName);
-	if (!GuidToLookup.IsValid() || !FExternalTextureRegistry::Get().GetExternalTextureCoordinateOffset(GuidToLookup, OutValue))
-	{
-		OutValue.R = OutValue.G = 0;
-		OutValue.B = OutValue.A = 0;
-	}
+	const FHashedName HashedParameterName = ParameterName.IsSet() ? FHashedName(ParameterName.GetValue()) : FHashedName();
+	OutData.WriteOpcode(EMaterialPreshaderOpcode::ExternalTextureCoordinateOffset).Write(HashedParameterName).Write(ExternalTextureGuid).Write<int32>(SourceTextureIndex);
 }
 
 FMaterialUniformExpressionRuntimeVirtualTextureUniform::FMaterialUniformExpressionRuntimeVirtualTextureUniform()
@@ -1492,14 +1765,6 @@ FMaterialUniformExpressionRuntimeVirtualTextureUniform::FMaterialUniformExpressi
 {
 }
 
-void FMaterialUniformExpressionRuntimeVirtualTextureUniform::Serialize(FArchive& Ar)
-{
-	Ar << bParameter;
-	Ar << ParameterInfo;
-	Ar << TextureIndex;
-	Ar << VectorIndex;
-}
-
 bool FMaterialUniformExpressionRuntimeVirtualTextureUniform::IsIdentical(const FMaterialUniformExpression* OtherExpression) const
 {
 	if (GetType() != OtherExpression->GetType())
@@ -1511,21 +1776,10 @@ bool FMaterialUniformExpressionRuntimeVirtualTextureUniform::IsIdentical(const F
 	return ParameterInfo == Other->ParameterInfo && TextureIndex == Other->TextureIndex && VectorIndex == Other->VectorIndex;
 }
 
-void FMaterialUniformExpressionRuntimeVirtualTextureUniform::GetNumberValue(const struct FMaterialRenderContext& Context, FLinearColor& OutValue) const
+void FMaterialUniformExpressionRuntimeVirtualTextureUniform::WriteNumberOpcodes(FMaterialPreshaderData& OutData) const
 {
-	const URuntimeVirtualTexture* Texture = nullptr;
-	if (!bParameter || !Context.MaterialRenderProxy || !Context.MaterialRenderProxy->GetTextureValue(ParameterInfo, &Texture, Context))
-	{
-		Texture = GetIndexedTexture<URuntimeVirtualTexture>(Context.Material, TextureIndex);
-	}
-	if (Texture != nullptr && VectorIndex != INDEX_NONE)
-	{
-		OutValue = FLinearColor(Texture->GetUniformParameter(VectorIndex));
-	}
-	else
-	{
-		OutValue = FLinearColor(0.f, 0.f, 0.f, 0.f);
-	}
+	const FHashedMaterialParameterInfo WriteParameterInfo = bParameter ? ParameterInfo : FHashedMaterialParameterInfo();
+	OutData.WriteOpcode(EMaterialPreshaderOpcode::RuntimeVirtualTextureUniform).Write(WriteParameterInfo).Write<int32>(TextureIndex).Write<int32>(VectorIndex);
 }
 
 /**
@@ -1535,17 +1789,7 @@ void FMaterialUniformExpressionRuntimeVirtualTextureUniform::GetNumberValue(cons
 class FMaterialUniformExpressionRuntimeVirtualTextureParameter_DEPRECATED : public FMaterialUniformExpressionRuntimeVirtualTextureUniform
 {
 	DECLARE_MATERIALUNIFORMEXPRESSION_TYPE(FMaterialUniformExpressionRuntimeVirtualTextureParameter_DEPRECATED);
-
-public:
-	virtual void Serialize(FArchive& Ar) override
-	{
-		Ar << TextureIndex;
-		Ar << VectorIndex;
-	}
 };
-// Specific implementation of FMaterialUniformExpressionRuntimeVirtualTextureParameter_DEPRECATED to serialize old class name
-FMaterialUniformExpressionType FMaterialUniformExpressionRuntimeVirtualTextureParameter_DEPRECATED::StaticType(TEXT("FMaterialUniformExpressionRuntimeVirtualTextureParameter"), &FMaterialUniformExpressionRuntimeVirtualTextureParameter_DEPRECATED::SerializationConstructor);
-
 
 IMPLEMENT_MATERIALUNIFORMEXPRESSION_TYPE(FMaterialUniformExpressionTexture);
 IMPLEMENT_MATERIALUNIFORMEXPRESSION_TYPE(FMaterialUniformExpressionConstant);
