Copying Private-LoadTimes-4.24 stream to Main. Biggest changes are in Materials/Shader memory freezing.

diff --git a/Engine/Source/Runtime/Engine/Private/Materials/MaterialUniformExpressions.h b/Engine/Source/Runtime/Engine/Private/Materials/MaterialUniformExpressions.h
--- a/Engine/Source/Runtime/Engine/Private/Materials/MaterialUniformExpressions.h
+++ b/Engine/Source/Runtime/Engine/Private/Materials/MaterialUniformExpressions.h
@@ -12,11 +12,160 @@
 #include "Engine/Texture.h"
 #include "Materials/MaterialExpressionTextureProperty.h"
 #include "Materials/MaterialLayersFunctions.h"
-#include "UObject/RenderingObjectVersion.h"
 
 // Temporary flag for toggling experimental material layers functionality
 bool AreExperimentalMaterialLayersEnabled();
 
+/**
+ * Represents a subclass of FMaterialUniformExpression.
+ */
+class FMaterialUniformExpressionType
+{
+public:
+	/**
+	 * @return The global uniform expression type list.  The list is used to temporarily store the types until
+	 *			the name subsystem has been initialized.
+	 */
+	static TLinkedList<FMaterialUniformExpressionType*>*& GetTypeList();
+
+	/**
+	 * Should not be called until the name subsystem has been initialized.
+	 * @return The global uniform expression type map.
+	 */
+	static TMap<FName, FMaterialUniformExpressionType*>& GetTypeMap();
+
+	/**
+	 * Minimal initialization constructor.
+	 */
+	FMaterialUniformExpressionType(const TCHAR* InName);
+
+	const TCHAR* GetName() const { return Name; }
+
+private:
+	const TCHAR* Name;
+};
+
+#define DECLARE_MATERIALUNIFORMEXPRESSION_TYPE(Name) \
+	public: \
+	static FMaterialUniformExpressionType StaticType; \
+	virtual FMaterialUniformExpressionType* GetType() const { return &StaticType; }
+
+#define IMPLEMENT_MATERIALUNIFORMEXPRESSION_TYPE(Name) \
+	FMaterialUniformExpressionType Name::StaticType(TEXT(#Name));
+
+/**
+ * Represents an expression which only varies with uniform inputs.
+ */
+class FMaterialUniformExpression : public FRefCountedObject
+{
+public:
+	virtual ~FMaterialUniformExpression() {}
+
+	virtual FMaterialUniformExpressionType* GetType() const = 0;
+	virtual class FMaterialUniformExpressionTexture* GetTextureUniformExpression() { return nullptr; }
+	virtual class FMaterialUniformExpressionExternalTexture* GetExternalTextureUniformExpression() { return nullptr; }
+	virtual bool IsConstant() const { return false; }
+	virtual bool IsIdentical(const FMaterialUniformExpression* OtherExpression) const { return false; }
+
+	virtual void WriteNumberOpcodes(FMaterialPreshaderData& OutData) const;
+
+	virtual void GetNumberValue(const struct FMaterialRenderContext& Context, FLinearColor& OutValue) const;
+};
+
+/**
+ * A texture expression.
+ */
+class FMaterialUniformExpressionTexture : public FMaterialUniformExpression
+{
+	DECLARE_MATERIALUNIFORMEXPRESSION_TYPE(FMaterialUniformExpressionTexture);
+
+public:
+	FMaterialUniformExpressionTexture();
+	FMaterialUniformExpressionTexture(int32 InTextureIndex, EMaterialSamplerType InSamplerType, ESamplerSourceMode InSamplerSource, bool InVirtualTexture);
+	FMaterialUniformExpressionTexture(int32 InTextureIndex, int16 InTextureLayerIndex, int16 InPageTableLayerIndex, EMaterialSamplerType InSamplerType);
+
+	//~ Begin FMaterialUniformExpression Interface.
+	virtual class FMaterialUniformExpressionTexture* GetTextureUniformExpression() { return this; }
+	virtual class FMaterialUniformExpressionTextureParameter* GetTextureParameterUniformExpression() { return nullptr; }
+	virtual bool IsIdentical(const FMaterialUniformExpression* OtherExpression) const;
+	//~ End FMaterialUniformExpression Interface.
+
+	/** Gets texture index which is the index in the full set of referenced textures for this material. */
+	int32 GetTextureIndex() const { return TextureIndex; }
+	/** Gets the texture layer index in the virtual texture stack if this is fixed. If we don't have a fixed layer then we will allocate during compilation (and not store here). */
+	int32 GetTextureLayerIndex() const { return TextureLayerIndex; }
+	/** Gets the page table channel index in the virtual texture stack if this is fixed. If we don't have a fixed layer then we will allocate during compilation (and not store here). */
+	int32 GetPageTableLayerIndex() const { return PageTableLayerIndex; }
+
+#if WITH_EDITORONLY_DATA
+	/** Get the sampling/decoding logic to compile in the shader for this texture. */
+	EMaterialSamplerType GetSamplerType() const { return SamplerType; }
+#endif
+
+	/** Get the sampler state object to use (globally shared or from texture asset).  */
+	ESamplerSourceMode GetSamplerSource() const { return SamplerSource; }
+
+	virtual void GetTextureParameterInfo(FMaterialTextureParameterInfo& OutParameter) const;
+
+protected:
+	/** Index into FMaterial::GetReferencedTextures */
+	int32 TextureIndex;
+	/** Texture layer index in virtual texture stack if preallocated */
+	int16 TextureLayerIndex;
+	/** Page table layer index in virtual texture stack if preallocated */
+	int16 PageTableLayerIndex;
+#if WITH_EDITORONLY_DATA
+	/** Sampler logic for this expression */
+	EMaterialSamplerType SamplerType;
+#endif
+	/** Sampler state object source for this expression */
+	ESamplerSourceMode SamplerSource;
+	/** Virtual texture flag used only for unique serialization */
+	bool bVirtualTexture;
+};
+
+class FMaterialUniformExpressionExternalTextureBase : public FMaterialUniformExpression
+{
+	DECLARE_MATERIALUNIFORMEXPRESSION_TYPE(FMaterialUniformExpressionExternalTextureBase);
+public:
+
+	FMaterialUniformExpressionExternalTextureBase(int32 InSourceTextureIndex = INDEX_NONE);
+	FMaterialUniformExpressionExternalTextureBase(const FGuid& InExternalTextureGuid);
+
+	virtual bool IsConstant() const override { return false; }
+	virtual bool IsIdentical(const FMaterialUniformExpression* OtherExpression) const override;
+
+	int32 GetSourceTextureIndex() const { return SourceTextureIndex; }
+
+protected:
+
+	/** Resolve the guid that relates to texture information inside FExternalTexture */
+	FGuid ResolveExternalTextureGUID(const FMaterialRenderContext& Context, TOptional<FName> ParameterName = TOptional<FName>()) const;
+
+	/** Index of the texture in the material that should be used to retrieve the external texture GUID at runtime (or INDEX_NONE) */
+	int32 SourceTextureIndex;
+	/** Optional external texture GUID defined at compile time */
+	FGuid ExternalTextureGuid;
+};
+
+/**
+* An external texture expression.
+*/
+class FMaterialUniformExpressionExternalTexture : public FMaterialUniformExpressionExternalTextureBase
+{
+	DECLARE_MATERIALUNIFORMEXPRESSION_TYPE(FMaterialUniformExpressionExternalTexture);
+public:
+
+	FMaterialUniformExpressionExternalTexture(int32 InSourceTextureIndex = INDEX_NONE) : FMaterialUniformExpressionExternalTextureBase(InSourceTextureIndex) {}
+	FMaterialUniformExpressionExternalTexture(const FGuid& InGuid) : FMaterialUniformExpressionExternalTextureBase(InGuid) {}
+
+	// FMaterialUniformExpression interface.
+	virtual FMaterialUniformExpressionExternalTexture* GetExternalTextureUniformExpression() override { return this; }
+	virtual class FMaterialUniformExpressionExternalTextureParameter* GetExternalTextureParameterUniformExpression() { return nullptr; }
+
+	virtual void GetExternalTextureParameterInfo(FMaterialExternalTextureParameterInfo& OutParameter) const;
+};
+
 /**
  */
 class FMaterialUniformExpressionConstant: public FMaterialUniformExpression
@@ -30,14 +179,6 @@ public:
 	{}
 
 	// FMaterialUniformExpression interface.
-	virtual void Serialize(FArchive& Ar)
-	{
-		Ar << Value << ValueType;
-	}
-	virtual void GetNumberValue(const FMaterialRenderContext& Context,FLinearColor& OutValue) const
-	{
-		OutValue = Value;
-	}
 	virtual bool IsConstant() const
 	{
 		return true;
@@ -52,6 +193,17 @@ public:
 		return OtherConstant->ValueType == ValueType && OtherConstant->Value == Value;
 	}
 
+	virtual void WriteNumberOpcodes(FMaterialPreshaderData& OutData) const override
+	{
+		OutData.WriteOpcode(EMaterialPreshaderOpcode::Constant);
+		OutData.Write(Value);
+	}
+
+	virtual void GetNumberValue(const FMaterialRenderContext& Context, FLinearColor& OutValue) const override
+	{
+		OutValue = Value;
+	}
+
 private:
 	FLinearColor Value;
 	uint8 ValueType;
@@ -64,70 +216,40 @@ class FMaterialUniformExpressionVectorParameter: public FMaterialUniformExpressi
 	DECLARE_MATERIALUNIFORMEXPRESSION_TYPE(FMaterialUniformExpressionVectorParameter);
 public:
 
-	FMaterialUniformExpressionVectorParameter()
-#if WITH_EDITOR
-		: bUseOverriddenDefault(false)
-#endif
-	{}
-
-	FMaterialUniformExpressionVectorParameter(const FMaterialParameterInfo& InParameterInfo,const FLinearColor& InDefaultValue)
+	FMaterialUniformExpressionVectorParameter() {}
+	FMaterialUniformExpressionVectorParameter(const FMaterialParameterInfo& InParameterInfo, int32 InParameterIndex)
 		: ParameterInfo(InParameterInfo)
-		, DefaultValue(InDefaultValue)
-#if WITH_EDITOR
-		, bUseOverriddenDefault(false)
-#endif
-	{}
+		, ParameterName(InParameterInfo.Name)
+		, ParameterIndex(InParameterIndex)
+	{
+		check(InParameterIndex >= 0 && InParameterIndex <= 0xffff);
+	}
 
 	// FMaterialUniformExpression interface.
-	virtual void Serialize(FArchive& Ar)
+	virtual void WriteNumberOpcodes(FMaterialPreshaderData& OutData) const override
 	{
-		Ar << ParameterInfo << DefaultValue;
+		OutData.WriteOpcode(EMaterialPreshaderOpcode::VectorParameter);
+		OutData.Write<uint16>(ParameterIndex);
 	}
 
-	// inefficient compared to GetGameThreadNumberValue(), for editor purpose
-	virtual void GetNumberValue(const FMaterialRenderContext& Context,FLinearColor& OutValue) const
+	virtual bool IsConstant() const
 	{
-		OutValue.R = OutValue.G = OutValue.B = OutValue.A = 0;
-
-		if(!Context.MaterialRenderProxy || !Context.MaterialRenderProxy->GetVectorValue(ParameterInfo, &OutValue, Context))
-		{
-			const bool bOveriddenParameterOnly = ParameterInfo.Association == EMaterialParameterAssociation::GlobalParameter;
-			
-			if (AreExperimentalMaterialLayersEnabled())
-			{
-				UMaterialInterface* Interface = Context.Material.GetMaterialInterface();
-				if (!Interface || !Interface->GetVectorParameterDefaultValue(ParameterInfo, OutValue, bOveriddenParameterOnly))
-				{
-					GetDefaultValue(OutValue);
-				}
-			}
-			else
-			{
-				GetDefaultValue(OutValue);
-			}
-		}
+		return false;
 	}
 
-	void GetDefaultValue(FLinearColor& OutValue) const
+	FMaterialParameterInfo GetParameterInfo() const
 	{
-#if WITH_EDITOR
-		OutValue = bUseOverriddenDefault ? OverriddenDefaultValue : DefaultValue;
-#else
-		OutValue = DefaultValue;
-#endif
+		return FMaterialParameterInfo(ParameterName, ParameterInfo.Association, ParameterInfo.Index);
 	}
 
-	// faster than GetNumberValue(), good for run-time use
-	void GetGameThreadNumberValue(const UMaterialInterface* SourceMaterialToCopyFrom, FLinearColor& OutValue) const;
-
-	virtual bool IsConstant() const
+	const FHashedMaterialParameterInfo& GetHashedParameterInfo() const
 	{
-		return false;
+		return ParameterInfo;
 	}
 
-	const FMaterialParameterInfo& GetParameterInfo() const
+	const FName& GetParameterName() const
 	{
-		return ParameterInfo;
+		return ParameterName;
 	}
 
 	virtual bool IsIdentical(const FMaterialUniformExpression* OtherExpression) const
@@ -137,24 +259,13 @@ public:
 			return false;
 		}
 		FMaterialUniformExpressionVectorParameter* OtherParameter = (FMaterialUniformExpressionVectorParameter*)OtherExpression;
-		return ParameterInfo == OtherParameter->ParameterInfo && DefaultValue == OtherParameter->DefaultValue;
+		return ParameterInfo == OtherParameter->ParameterInfo && ParameterIndex == OtherParameter->ParameterIndex;
 	}
 
-#if WITH_EDITOR
-	void SetTransientOverrideDefaultValue(const FLinearColor& InOverrideDefaultValue, bool bInUseOverriddenDefault)
-	{
-		bUseOverriddenDefault = bInUseOverriddenDefault;
-		OverriddenDefaultValue = InOverrideDefaultValue;
-	}
-#endif
-
 private:
-	FMaterialParameterInfo ParameterInfo;
-	FLinearColor DefaultValue;
-#if WITH_EDITOR
-	bool bUseOverriddenDefault;
-	FLinearColor OverriddenDefaultValue;
-#endif
+	FHashedMaterialParameterInfo ParameterInfo;
+	FName ParameterName;
+	int32 ParameterIndex;
 };
 
 /**
@@ -164,74 +275,40 @@ class FMaterialUniformExpressionScalarParameter: public FMaterialUniformExpressi
 	DECLARE_MATERIALUNIFORMEXPRESSION_TYPE(FMaterialUniformExpressionScalarParameter);
 public:
 
-	FMaterialUniformExpressionScalarParameter()
-#if WITH_EDITOR
-		: bUseOverriddenDefault(false)
-#endif
-	{}
-
-	FMaterialUniformExpressionScalarParameter(const FMaterialParameterInfo& InParameterInfo,float InDefaultValue)
+	FMaterialUniformExpressionScalarParameter() {}
+	FMaterialUniformExpressionScalarParameter(const FMaterialParameterInfo& InParameterInfo, int32 InParameterIndex)
 		: ParameterInfo(InParameterInfo)
-		, DefaultValue(InDefaultValue)
-#if WITH_EDITOR
-		, bUseOverriddenDefault(false)
-#endif
-	{}
+		, ParameterName(InParameterInfo.Name)
+		, ParameterIndex(InParameterIndex)
+	{
+		check(InParameterIndex >= 0 && InParameterIndex <= 0xffff);
+	}
 
 	// FMaterialUniformExpression interface.
-	virtual void Serialize(FArchive& Ar)
+	virtual void WriteNumberOpcodes(FMaterialPreshaderData& OutData) const override
 	{
-		Ar << ParameterInfo << DefaultValue;
+		OutData.WriteOpcode(EMaterialPreshaderOpcode::ScalarParameter);
+		OutData.Write<uint16>(ParameterIndex);
 	}
 
-	// inefficient compared to GetGameThreadNumberValue(), for editor purpose
-	virtual void GetNumberValue(const FMaterialRenderContext& Context,FLinearColor& OutValue) const
-	{	
-		OutValue.A = 0;
-
-		if (!Context.MaterialRenderProxy || !Context.MaterialRenderProxy->GetScalarValue(ParameterInfo, &OutValue.A, Context))
-		{
-			const bool bOveriddenParameterOnly = ParameterInfo.Association == EMaterialParameterAssociation::GlobalParameter;
-			
-			if (AreExperimentalMaterialLayersEnabled())
-			{
-				UMaterialInterface* Interface = Context.Material.GetMaterialInterface();
-				if (!Interface || !Interface->GetScalarParameterDefaultValue(ParameterInfo, OutValue.A, bOveriddenParameterOnly))
-				{
-					GetDefaultValue(OutValue.A);
-				}
-			}
-			else
-			{
-				GetDefaultValue(OutValue.A);
-			}
-		}
-
-		OutValue.R = OutValue.G = OutValue.B = OutValue.A;
+	virtual bool IsConstant() const
+	{
+		return false;
 	}
 
-	void GetDefaultValue(float& OutValue) const
+	FMaterialParameterInfo GetParameterInfo() const
 	{
-#if WITH_EDITOR
-		OutValue = bUseOverriddenDefault ? OverriddenDefaultValue : DefaultValue;
-#else
-		OutValue = DefaultValue;
-#endif
+		return FMaterialParameterInfo(ParameterName, ParameterInfo.Association, ParameterInfo.Index);
 	}
-	
-	// faster than GetNumberValue(), good for run-time use
-	void GetGameThreadNumberValue(const UMaterialInterface* SourceMaterialToCopyFrom, float& OutValue) const;
 
-	void GetGameThreadUsedAsAtlas(const UMaterialInterface* SourceMaterialToCopyFrom, bool& OutValue, TSoftObjectPtr<class UCurveLinearColor>& Curve, TSoftObjectPtr<class UCurveLinearColorAtlas>& Atlas) const;
-	
-	virtual bool IsConstant() const
+	const FHashedMaterialParameterInfo& GetHashedParameterInfo() const
 	{
-		return false;
+		return ParameterInfo;
 	}
 
-	const FMaterialParameterInfo& GetParameterInfo() const
+	const FName& GetParameterName() const
 	{
-		return ParameterInfo;
+		return ParameterName;
 	}
 
 	virtual bool IsIdentical(const FMaterialUniformExpression* OtherExpression) const
@@ -241,24 +318,13 @@ public:
 			return false;
 		}
 		FMaterialUniformExpressionScalarParameter* OtherParameter = (FMaterialUniformExpressionScalarParameter*)OtherExpression;
-		return ParameterInfo == OtherParameter->ParameterInfo && DefaultValue == OtherParameter->DefaultValue;
+		return ParameterInfo == OtherParameter->ParameterInfo && ParameterIndex == OtherParameter->ParameterIndex;
 	}
 
-#if WITH_EDITOR
-	void SetTransientOverrideDefaultValue(float InOverrideDefaultValue, bool bInUseOverriddenDefault)
-	{
-		bUseOverriddenDefault = bInUseOverriddenDefault;
-		OverriddenDefaultValue = InOverrideDefaultValue;
-	}
-#endif
-
 private:
-	FMaterialParameterInfo ParameterInfo;
-	float DefaultValue;
-#if WITH_EDITOR
-	bool bUseOverriddenDefault;
-	float OverriddenDefaultValue;
-#endif
+	FHashedMaterialParameterInfo ParameterInfo;
+	FName ParameterName;
+	int32 ParameterIndex;
 };
 
 /** @return The texture that was associated with the given index when the given material had its uniform expressions/HLSL code generated. */
@@ -266,7 +332,7 @@ template<typename TextureType>
 static TextureType* GetIndexedTexture(const FMaterial& Material, int32 TextureIndex)
 {
 	UObject* IndexedTexture = nullptr;
-	const TArray<UObject*>& ReferencedTextures = Material.GetReferencedTextures();
+	const TArrayView<UObject* const> ReferencedTextures = Material.GetReferencedTextures();
 	if (ReferencedTextures.IsValidIndex(TextureIndex))
 	{
 		IndexedTexture = ReferencedTextures[TextureIndex];
@@ -300,37 +366,41 @@ public:
 	FMaterialUniformExpressionTextureParameter(const FMaterialParameterInfo& InParameterInfo, int32 InTextureIndex, EMaterialSamplerType InSamplerType, ESamplerSourceMode InSourceMode, bool InVirtualTexture)
 		: Super(InTextureIndex, InSamplerType, InSourceMode, InVirtualTexture)
 		, ParameterInfo(InParameterInfo)
+		, ParameterName(InParameterInfo.Name)
 	{}
 
 	FMaterialUniformExpressionTextureParameter(const FMaterialParameterInfo& InParameterInfo, int32 InTextureIndex, int32 InTextureLayerIndex, int32 InPageTableLayerIndex, EMaterialSamplerType InSamplerType)
 		: Super(InTextureIndex, InTextureLayerIndex, InPageTableLayerIndex, InSamplerType)
 		, ParameterInfo(InParameterInfo)
+		, ParameterName(InParameterInfo.Name)
 	{}
 
 	// FMaterialUniformExpression interface.
 	virtual class FMaterialUniformExpressionTextureParameter* GetTextureParameterUniformExpression() override { return this; }
 
-	virtual void Serialize(FArchive& Ar)
+	virtual void GetTextureParameterInfo(FMaterialTextureParameterInfo& OutParameter) const override
 	{
-		Ar << ParameterInfo;
-		Super::Serialize(Ar);
+		Super::GetTextureParameterInfo(OutParameter);
+		OutParameter.ParameterInfo = ParameterInfo;
+		OutParameter.ParameterName = ParameterName.ToString();
 	}
 
-	virtual void GetTextureValue(const FMaterialRenderContext& Context, const FMaterial& Material, const UTexture*& OutValue) const override;
-	virtual void GetTextureValue(const FMaterialRenderContext& Context, const FMaterial& Material, const URuntimeVirtualTexture*& OutValue) const override;
-	virtual void GetGameThreadTextureValue(const UMaterialInterface* MaterialInterface, const FMaterial& Material, UTexture*& OutValue, bool bAllowOverride = true) const override;
-
 	virtual bool IsConstant() const
 	{
 		return false;
 	}
 
-	FName GetParameterName() const
+	FMaterialParameterInfo GetParameterInfo() const
+	{
+		return FMaterialParameterInfo(ParameterName, ParameterInfo.Association, ParameterInfo.Index);
+	}
+
+	const FName& GetParameterName() const
 	{
-		return ParameterInfo.Name;
+		return ParameterName;
 	}
 
-	const FMaterialParameterInfo& GetParameterInfo() const
+	const FHashedMaterialParameterInfo& GetHashedParameterInfo() const
 	{
 		return ParameterInfo;
 	}
@@ -346,7 +416,9 @@ public:
 	}
 
 private:
-	FMaterialParameterInfo ParameterInfo;
+	FHashedMaterialParameterInfo ParameterInfo;
+	FName ParameterName;
+	int32 ParameterIndex;
 };
 
 /**
@@ -360,11 +432,6 @@ public:
 
 	FMaterialUniformExpressionFlipBookTextureParameter() {}
 
-	virtual void GetNumberValue(const FMaterialRenderContext& Context,FLinearColor& OutValue) const
-	{
-		OutValue.R = OutValue.G = OutValue.B = OutValue.A = 0;
-	}
-
 	virtual bool IsConstant() const
 	{
 		return false;
@@ -384,9 +451,8 @@ public:
 	FMaterialUniformExpressionExternalTextureParameter();
 	FMaterialUniformExpressionExternalTextureParameter(FName InParameterName, int32 InTextureIndex);
 
-	virtual void Serialize(FArchive& Ar) override;
 	virtual bool IsIdentical(const FMaterialUniformExpression* OtherExpression) const override;
-	virtual bool GetExternalTexture(const FMaterialRenderContext& Context, FTextureRHIRef& OutTextureRHI, FSamplerStateRHIRef& OutSamplerStateRHI) const override;
+	virtual void GetExternalTextureParameterInfo(FMaterialExternalTextureParameterInfo& OutParameter) const override;
 	virtual FMaterialUniformExpressionExternalTextureParameter* GetExternalTextureParameterUniformExpression() override { return this; }
 
 	FName GetParameterName() const
@@ -412,18 +478,10 @@ public:
 	{}
 
 	// FMaterialUniformExpression interface.
-	virtual void Serialize(FArchive& Ar)
+	virtual void WriteNumberOpcodes(FMaterialPreshaderData& OutData) const override
 	{
-		Ar << X << bIsCosine;
-	}
-	virtual void GetNumberValue(const FMaterialRenderContext& Context,FLinearColor& OutValue) const
-	{
-		FLinearColor ValueX = FLinearColor::Black;
-		X->GetNumberValue(Context,ValueX);
-		OutValue.R = bIsCosine ? FMath::Cos(ValueX.R) : FMath::Sin(ValueX.R);
-		OutValue.G = bIsCosine ? FMath::Cos(ValueX.G) : FMath::Sin(ValueX.G);
-		OutValue.B = bIsCosine ? FMath::Cos(ValueX.B) : FMath::Sin(ValueX.B);
-		OutValue.A = bIsCosine ? FMath::Cos(ValueX.A) : FMath::Sin(ValueX.A);
+		X->WriteNumberOpcodes(OutData);
+		OutData.WriteOpcode(bIsCosine ? EMaterialPreshaderOpcode::Cos : EMaterialPreshaderOpcode::Sin);
 	}
 	virtual bool IsConstant() const
 	{
@@ -478,38 +536,26 @@ public:
 	{}
 
 	// FMaterialUniformExpression interface.
-	virtual void Serialize(FArchive& Ar)
-	{
-		Ar << X << Y << Op;
-	}
-	virtual void GetNumberValue(const FMaterialRenderContext& Context,FLinearColor& OutValue) const
+	virtual void WriteNumberOpcodes(FMaterialPreshaderData& OutData) const override
 	{
-		FLinearColor ValueX = FLinearColor::Black;
-		FLinearColor ValueY = FLinearColor::Black;
-		X->GetNumberValue(Context,ValueX);
-		Y->GetNumberValue(Context,ValueY);
-
+		X->WriteNumberOpcodes(OutData);
+		if (Op == TMO_Atan2)
+		{
+			Y->WriteNumberOpcodes(OutData);
+		}
 		switch (Op)
 		{
-		case TMO_Sin: OutValue.R = FMath::Sin(ValueX.R); OutValue.G = FMath::Sin(ValueX.G); OutValue.B = FMath::Sin(ValueX.B); OutValue.A = FMath::Sin(ValueX.A); break;
-		case TMO_Cos: OutValue.R = FMath::Cos(ValueX.R); OutValue.G = FMath::Cos(ValueX.G); OutValue.B = FMath::Cos(ValueX.B); OutValue.A = FMath::Cos(ValueX.A); break;
-		case TMO_Tan: OutValue.R = FMath::Tan(ValueX.R); OutValue.G = FMath::Tan(ValueX.G); OutValue.B = FMath::Tan(ValueX.B); OutValue.A = FMath::Tan(ValueX.A); break;
-
-		case TMO_Asin: OutValue.R = FMath::Asin(ValueX.R); OutValue.G = FMath::Asin(ValueX.G); OutValue.B = FMath::Asin(ValueX.B); OutValue.A = FMath::Asin(ValueX.A); break;
-		case TMO_Acos: OutValue.R = FMath::Acos(ValueX.R); OutValue.G = FMath::Acos(ValueX.G); OutValue.B = FMath::Acos(ValueX.B); OutValue.A = FMath::Acos(ValueX.A); break;
-		case TMO_Atan: OutValue.R = FMath::Atan(ValueX.R); OutValue.G = FMath::Atan(ValueX.G); OutValue.B = FMath::Atan(ValueX.B); OutValue.A = FMath::Atan(ValueX.A); break;
-
-		case TMO_Atan2:
-			// Note: Param names are reversed here for a trade-off of order consistency vs sharing code
-			OutValue.R = FMath::Atan2(ValueX.R, ValueY.R); OutValue.G = FMath::Atan2(ValueX.G, ValueY.G);
-			OutValue.B = FMath::Atan2(ValueX.B, ValueY.B); OutValue.A = FMath::Atan2(ValueX.A, ValueY.A);
-			break;
-
-		default:
-			checkf(0, TEXT("Invalid trigonometry math operation in uniform expression."));
+		case TMO_Sin: OutData.WriteOpcode(EMaterialPreshaderOpcode::Sin); break;
+		case TMO_Cos: OutData.WriteOpcode(EMaterialPreshaderOpcode::Cos); break;
+		case TMO_Tan: OutData.WriteOpcode(EMaterialPreshaderOpcode::Tan); break;
+		case TMO_Asin: OutData.WriteOpcode(EMaterialPreshaderOpcode::Asin); break;
+		case TMO_Acos: OutData.WriteOpcode(EMaterialPreshaderOpcode::Acos); break;
+		case TMO_Atan: OutData.WriteOpcode(EMaterialPreshaderOpcode::Atan); break;
+		case TMO_Atan2: OutData.WriteOpcode(EMaterialPreshaderOpcode::Atan2); break;
+		default: checkNoEntry(); break;
 		}
-		
 	}
+
 	virtual bool IsConstant() const
 	{
 		return X->IsConstant() && Y->IsConstant();
@@ -543,18 +589,10 @@ public:
 	{}
 
 	// FMaterialUniformExpression interface.
-	virtual void Serialize(FArchive& Ar)
-	{
-		Ar << X;
-	}
-	virtual void GetNumberValue(const FMaterialRenderContext& Context,FLinearColor& OutValue) const
+	virtual void WriteNumberOpcodes(FMaterialPreshaderData& OutData) const override
 	{
-		FLinearColor ValueX = FLinearColor::Black;
-		X->GetNumberValue(Context,ValueX);
-		OutValue.R = FMath::Sqrt(ValueX.R);
-		OutValue.G = FMath::Sqrt(ValueX.G);
-		OutValue.B = FMath::Sqrt(ValueX.B);
-		OutValue.A = FMath::Sqrt(ValueX.A);
+		X->WriteNumberOpcodes(OutData);
+		OutData.WriteOpcode(EMaterialPreshaderOpcode::Sqrt);
 	}
 	virtual bool IsConstant() const
 	{
@@ -588,28 +626,10 @@ public:
 	{}
 
 	// FMaterialUniformExpression interface.
-	virtual void Serialize(FArchive& Ar)
+	virtual void WriteNumberOpcodes(FMaterialPreshaderData& OutData) const override
 	{
-		Ar.UsingCustomVersion(FRenderingObjectVersion::GUID);
-		Ar << X;
-				
-		if (Ar.CustomVer(FRenderingObjectVersion::GUID) >= FRenderingObjectVersion::TypeHandlingForMaterialSqrtNodes)
-		{
-			Ar << ValueType;
-		}	
-	}
-	virtual void GetNumberValue(const FMaterialRenderContext& Context,FLinearColor& OutValue) const
-	{
-		FLinearColor ValueX = FLinearColor::Black;
-		X->GetNumberValue(Context,ValueX);
-
-		check(ValueType & MCT_Float);
-		float LengthSq = ValueX.R * ValueX.R;
-		LengthSq += (ValueType >= MCT_Float2) ? ValueX.G * ValueX.G : 0;
-		LengthSq += (ValueType >= MCT_Float3) ? ValueX.B * ValueX.B : 0;
-		LengthSq += (ValueType >= MCT_Float4) ? ValueX.A * ValueX.A : 0;
-
-		OutValue.R = OutValue.G = OutValue.B = OutValue.A = FMath::Sqrt(LengthSq);
+		X->WriteNumberOpcodes(OutData);
+		OutData.WriteOpcode(EMaterialPreshaderOpcode::Length).Write<uint8>(ValueType);
 	}
 	virtual bool IsConstant() const
 	{
@@ -643,18 +663,10 @@ public:
 	{}
 
 	// FMaterialUniformExpression interface.
-	void Serialize(FArchive& Ar) override
-	{
-		Ar << X;
-	}
-	void GetNumberValue(const FMaterialRenderContext& Context,FLinearColor& OutValue) const override
+	virtual void WriteNumberOpcodes(FMaterialPreshaderData& OutData) const override
 	{
-		FLinearColor ValueX = FLinearColor::Black;
-		X->GetNumberValue(Context,ValueX);
-		OutValue.R = FMath::Log2(ValueX.R);
-		OutValue.G = FMath::Log2(ValueX.G);
-		OutValue.B = FMath::Log2(ValueX.B);
-		OutValue.A = FMath::Log2(ValueX.A);
+		X->WriteNumberOpcodes(OutData);
+		OutData.WriteOpcode(EMaterialPreshaderOpcode::Log2);
 	}
 	bool IsConstant() const override
 	{
@@ -688,20 +700,10 @@ public:
 	{}
 
 	// FMaterialUniformExpression interface.
-	void Serialize(FArchive& Ar) override
+	virtual void WriteNumberOpcodes(FMaterialPreshaderData& OutData) const override
 	{
-		Ar << X;
-	}
-	void GetNumberValue(const FMaterialRenderContext& Context,FLinearColor& OutValue) const override
-	{
-		FLinearColor ValueX = FLinearColor::Black;
-		X->GetNumberValue(Context,ValueX);
-
-		static const float LogToLog10 = 1.0f / FMath::Loge(10.f);
-		OutValue.R = FMath::Loge(ValueX.R) * LogToLog10;
-		OutValue.G = FMath::Loge(ValueX.G) * LogToLog10;
-		OutValue.B = FMath::Loge(ValueX.B) * LogToLog10;
-		OutValue.A = FMath::Loge(ValueX.A) * LogToLog10;
+		X->WriteNumberOpcodes(OutData);
+		OutData.WriteOpcode(EMaterialPreshaderOpcode::Log10);
 	}
 	bool IsConstant() const override
 	{
@@ -748,17 +750,22 @@ public:
 	{}
 
 	// FMaterialUniformExpression interface.
-	virtual void Serialize(FArchive& Ar)
+	virtual void WriteNumberOpcodes(FMaterialPreshaderData& OutData) const override
 	{
-		Ar.UsingCustomVersion(FRenderingObjectVersion::GUID);
-		Ar << A << B << Op;
-		
-		if (Ar.CustomVer(FRenderingObjectVersion::GUID) >= FRenderingObjectVersion::TypeHandlingForMaterialSqrtNodes)
+		A->WriteNumberOpcodes(OutData);
+		B->WriteNumberOpcodes(OutData);
+
+		switch (Op)
 		{
-			Ar << ValueType;
-		}	
+		case FMO_Add: OutData.WriteOpcode(EMaterialPreshaderOpcode::Add); break;
+		case FMO_Sub: OutData.WriteOpcode(EMaterialPreshaderOpcode::Sub); break;
+		case FMO_Mul: OutData.WriteOpcode(EMaterialPreshaderOpcode::Mul); break;
+		case FMO_Div: OutData.WriteOpcode(EMaterialPreshaderOpcode::Div); break;
+		case FMO_Dot: OutData.WriteOpcode(EMaterialPreshaderOpcode::Dot).Write<uint8>(ValueType); break;
+		case FMO_Cross: OutData.WriteOpcode(EMaterialPreshaderOpcode::Cross).Write<uint8>(ValueType); break;
+		default: checkNoEntry(); break;
+		}
 	}
-	virtual void GetNumberValue(const FMaterialRenderContext& Context, FLinearColor& OutValue) const;
 
 	virtual bool IsConstant() const
 	{
@@ -795,19 +802,10 @@ public:
 	{}
 
 	// FMaterialUniformExpression interface.
-	virtual void Serialize(FArchive& Ar)
-	{
-		Ar << X;
-	}
-	virtual void GetNumberValue(const FMaterialRenderContext& Context,FLinearColor& OutValue) const
+	virtual void WriteNumberOpcodes(FMaterialPreshaderData& OutData) const override
 	{
-		FLinearColor TempValue = FLinearColor::Black;
-		X->GetNumberValue(Context,TempValue);
-
-		OutValue.R = FMath::Fractional(TempValue.R);
-		OutValue.G = FMath::Fractional(TempValue.G);
-		OutValue.B = FMath::Fractional(TempValue.B);
-		OutValue.A = FMath::Fractional(TempValue.A);
+		X->WriteNumberOpcodes(OutData);
+		OutData.WriteOpcode(EMaterialPreshaderOpcode::Fractional);
 	}
 	virtual bool IsConstant() const
 	{
@@ -842,21 +840,11 @@ public:
 	{}
 
 	// FMaterialUniformExpression interface.
-	virtual void Serialize(FArchive& Ar)
+	virtual void WriteNumberOpcodes(FMaterialPreshaderData& OutData) const override
 	{
-		Ar << A << B << NumComponentsA;
-	}
-	virtual void GetNumberValue(const FMaterialRenderContext& Context,FLinearColor& OutValue) const
-	{
-		FLinearColor ValueA = FLinearColor::Black;
-		FLinearColor ValueB = FLinearColor::Black;
-		A->GetNumberValue(Context, ValueA);
-		B->GetNumberValue(Context, ValueB);
-
-		OutValue.R = NumComponentsA >= 1 ? ValueA.R : (&ValueB.R)[0 - NumComponentsA];
-		OutValue.G = NumComponentsA >= 2 ? ValueA.G : (&ValueB.R)[1 - NumComponentsA];
-		OutValue.B = NumComponentsA >= 3 ? ValueA.B : (&ValueB.R)[2 - NumComponentsA];
-		OutValue.A = NumComponentsA >= 4 ? ValueA.A : (&ValueB.R)[3 - NumComponentsA];
+		A->WriteNumberOpcodes(OutData);
+		B->WriteNumberOpcodes(OutData);
+		OutData.WriteOpcode(EMaterialPreshaderOpcode::AppendVector).Write<uint8>(NumComponentsA);
 	}
 	virtual bool IsConstant() const
 	{
@@ -892,21 +880,11 @@ public:
 	{}
 
 	// FMaterialUniformExpression interface.
-	virtual void Serialize(FArchive& Ar)
+	virtual void WriteNumberOpcodes(FMaterialPreshaderData& OutData) const override
 	{
-		Ar << A << B;
-	}
-	virtual void GetNumberValue(const FMaterialRenderContext& Context,FLinearColor& OutValue) const
-	{
-		FLinearColor ValueA = FLinearColor::Black;
-		FLinearColor ValueB = FLinearColor::Black;
-		A->GetNumberValue(Context, ValueA);
-		B->GetNumberValue(Context, ValueB);
-
-		OutValue.R = FMath::Min(ValueA.R, ValueB.R);
-		OutValue.G = FMath::Min(ValueA.G, ValueB.G);
-		OutValue.B = FMath::Min(ValueA.B, ValueB.B);
-		OutValue.A = FMath::Min(ValueA.A, ValueB.A);
+		A->WriteNumberOpcodes(OutData);
+		B->WriteNumberOpcodes(OutData);
+		OutData.WriteOpcode(EMaterialPreshaderOpcode::Min);
 	}
 	virtual bool IsConstant() const
 	{
@@ -941,21 +919,11 @@ public:
 	{}
 
 	// FMaterialUniformExpression interface.
-	virtual void Serialize(FArchive& Ar)
+	virtual void WriteNumberOpcodes(FMaterialPreshaderData& OutData) const override
 	{
-		Ar << A << B;
-	}
-	virtual void GetNumberValue(const FMaterialRenderContext& Context,FLinearColor& OutValue) const
-	{
-		FLinearColor ValueA = FLinearColor::Black;
-		FLinearColor ValueB = FLinearColor::Black;
-		A->GetNumberValue(Context, ValueA);
-		B->GetNumberValue(Context, ValueB);
-
-		OutValue.R = FMath::Max(ValueA.R, ValueB.R);
-		OutValue.G = FMath::Max(ValueA.G, ValueB.G);
-		OutValue.B = FMath::Max(ValueA.B, ValueB.B);
-		OutValue.A = FMath::Max(ValueA.A, ValueB.A);
+		A->WriteNumberOpcodes(OutData);
+		B->WriteNumberOpcodes(OutData);
+		OutData.WriteOpcode(EMaterialPreshaderOpcode::Max);
 	}
 	virtual bool IsConstant() const
 	{
@@ -991,23 +959,12 @@ public:
 	{}
 
 	// FMaterialUniformExpression interface.
-	virtual void Serialize(FArchive& Ar)
+	virtual void WriteNumberOpcodes(FMaterialPreshaderData& OutData) const override
 	{
-		Ar << Input << Min << Max;
-	}
-	virtual void GetNumberValue(const FMaterialRenderContext& Context,FLinearColor& OutValue) const
-	{
-		FLinearColor ValueMin = FLinearColor::Black;
-		FLinearColor ValueMax = FLinearColor::Black;
-		FLinearColor ValueInput = FLinearColor::Black;
-		Min->GetNumberValue(Context, ValueMin);
-		Max->GetNumberValue(Context, ValueMax);
-		Input->GetNumberValue(Context, ValueInput);
-
-		OutValue.R = FMath::Clamp(ValueInput.R, ValueMin.R, ValueMax.R);
-		OutValue.G = FMath::Clamp(ValueInput.G, ValueMin.G, ValueMax.G);
-		OutValue.B = FMath::Clamp(ValueInput.B, ValueMin.B, ValueMax.B);
-		OutValue.A = FMath::Clamp(ValueInput.A, ValueMin.A, ValueMax.A);
+		Input->WriteNumberOpcodes(OutData);
+		Min->WriteNumberOpcodes(OutData);
+		Max->WriteNumberOpcodes(OutData);
+		OutData.WriteOpcode(EMaterialPreshaderOpcode::Clamp);
 	}
 	virtual bool IsConstant() const
 	{
@@ -1042,19 +999,10 @@ public:
 	{}
 
 	// FMaterialUniformExpression interface.
-	virtual void Serialize(FArchive& Ar)
+	virtual void WriteNumberOpcodes(FMaterialPreshaderData& OutData) const override
 	{
-		Ar << Input;
-	}
-	virtual void GetNumberValue(const FMaterialRenderContext& Context,FLinearColor& OutValue) const
-	{
-		FLinearColor ValueInput = FLinearColor::Black;
-		Input->GetNumberValue(Context, ValueInput);
-
-		OutValue.R = FMath::Clamp<float>(ValueInput.R, 0, 1);
-		OutValue.G = FMath::Clamp<float>(ValueInput.G, 0, 1);
-		OutValue.B = FMath::Clamp<float>(ValueInput.B, 0, 1);
-		OutValue.A = FMath::Clamp<float>(ValueInput.A, 0, 1);
+		Input->WriteNumberOpcodes(OutData);
+		OutData.WriteOpcode(EMaterialPreshaderOpcode::Saturate);
 	}
 	virtual bool IsConstant() const
 	{
@@ -1114,39 +1062,10 @@ public:
 	}
 
 	// FMaterialUniformExpression interface.
-	virtual void Serialize(FArchive& Ar)
-	{
-		Ar << X;
-		Ar << IndexR;
-		Ar << IndexG;
-		Ar << IndexB;
-		Ar << IndexA;
-		Ar << NumElements;
-	}
-	virtual void GetNumberValue(const FMaterialRenderContext& Context, FLinearColor& OutValue) const
-	{
-		FLinearColor Temp = OutValue;
-		X->GetNumberValue(Context, Temp);
-		// Clear
-		OutValue *= 0;
-		switch (NumElements)
-		{
-		case 1:
-			// Replicate scalar
-			OutValue.R = OutValue.G = OutValue.B = OutValue.A = Temp.Component(IndexR);
-			break;
-		case 4:
-			OutValue.A = Temp.Component(IndexA);
-			// Fallthrough...
-		case 3:
-			OutValue.B = Temp.Component(IndexB);
-			// Fallthrough...
-		case 2:
-			OutValue.G = Temp.Component(IndexG);
-			OutValue.R = Temp.Component(IndexR);
-			break;
-		default: UE_LOG(LogMaterial, Fatal, TEXT("Invalid number of swizzle elements: %d"), NumElements);
-		}
+	virtual void WriteNumberOpcodes(FMaterialPreshaderData& OutData) const override
+	{
+		X->WriteNumberOpcodes(OutData);
+		OutData.WriteOpcode(EMaterialPreshaderOpcode::ComponentSwizzle).Write<uint8>(NumElements).Write<uint8>(IndexR).Write<uint8>(IndexG).Write<uint8>(IndexB).Write<uint8>(IndexA);
 	}
 	virtual bool IsConstant() const
 	{
@@ -1189,18 +1108,10 @@ public:
 	{}
 
 	// FMaterialUniformExpression interface.
-	virtual void Serialize(FArchive& Ar)
-	{
-		Ar << X;
-	}
-	virtual void GetNumberValue(const FMaterialRenderContext& Context,FLinearColor& OutValue) const
+	virtual void WriteNumberOpcodes(FMaterialPreshaderData& OutData) const override
 	{
-		X->GetNumberValue(Context, OutValue);
-
-		OutValue.R = FMath::FloorToInt(OutValue.R);
-		OutValue.G = FMath::FloorToInt(OutValue.G);
-		OutValue.B = FMath::FloorToInt(OutValue.B);
-		OutValue.A = FMath::FloorToInt(OutValue.A);
+		X->WriteNumberOpcodes(OutData);
+		OutData.WriteOpcode(EMaterialPreshaderOpcode::Floor);
 	}
 	virtual bool IsConstant() const
 	{
@@ -1233,18 +1144,10 @@ public:
 	{}
 
 	// FMaterialUniformExpression interface.
-	virtual void Serialize(FArchive& Ar)
-	{
-		Ar << X;
-	}
-	virtual void GetNumberValue(const FMaterialRenderContext& Context,FLinearColor& OutValue) const
+	virtual void WriteNumberOpcodes(FMaterialPreshaderData& OutData) const override
 	{
-		X->GetNumberValue(Context, OutValue);
-
-		OutValue.R = FMath::CeilToInt(OutValue.R);
-		OutValue.G = FMath::CeilToInt(OutValue.G);
-		OutValue.B = FMath::CeilToInt(OutValue.B);
-		OutValue.A = FMath::CeilToInt(OutValue.A);
+		X->WriteNumberOpcodes(OutData);
+		OutData.WriteOpcode(EMaterialPreshaderOpcode::Ceil);
 	}
 	virtual bool IsConstant() const
 	{
@@ -1277,18 +1180,10 @@ public:
 	{}
 
 	// FMaterialUniformExpression interface.
-	virtual void Serialize(FArchive& Ar)
-	{
-		Ar << X;
-	}
-	virtual void GetNumberValue(const FMaterialRenderContext& Context,FLinearColor& OutValue) const
+	virtual void WriteNumberOpcodes(FMaterialPreshaderData& OutData) const override
 	{
-		X->GetNumberValue(Context, OutValue);
-
-		OutValue.R = FMath::RoundToFloat(OutValue.R);
-		OutValue.G = FMath::RoundToFloat(OutValue.G);
-		OutValue.B = FMath::RoundToFloat(OutValue.B);
-		OutValue.A = FMath::RoundToFloat(OutValue.A);
+		X->WriteNumberOpcodes(OutData);
+		OutData.WriteOpcode(EMaterialPreshaderOpcode::Round);
 	}
 	virtual bool IsConstant() const
 	{
@@ -1321,18 +1216,10 @@ public:
 	{}
 
 	// FMaterialUniformExpression interface.
-	virtual void Serialize(FArchive& Ar)
+	virtual void WriteNumberOpcodes(FMaterialPreshaderData& OutData) const override
 	{
-		Ar << X;
-	}
-	virtual void GetNumberValue(const FMaterialRenderContext& Context,FLinearColor& OutValue) const
-	{
-		X->GetNumberValue(Context, OutValue);
-
-		OutValue.R = FMath::TruncToFloat(OutValue.R);
-		OutValue.G = FMath::TruncToFloat(OutValue.G);
-		OutValue.B = FMath::TruncToFloat(OutValue.B);
-		OutValue.A = FMath::TruncToFloat(OutValue.A);
+		X->WriteNumberOpcodes(OutData);
+		OutData.WriteOpcode(EMaterialPreshaderOpcode::Trunc);
 	}
 	virtual bool IsConstant() const
 	{
@@ -1365,18 +1252,10 @@ public:
 	{}
 
 	// FMaterialUniformExpression interface.
-	virtual void Serialize(FArchive& Ar)
+	virtual void WriteNumberOpcodes(FMaterialPreshaderData& OutData) const override
 	{
-		Ar << X;
-	}
-	virtual void GetNumberValue(const FMaterialRenderContext& Context,FLinearColor& OutValue) const
-	{
-		X->GetNumberValue(Context, OutValue);
-
-		OutValue.R = FMath::Sign(OutValue.R);
-		OutValue.G = FMath::Sign(OutValue.G);
-		OutValue.B = FMath::Sign(OutValue.B);
-		OutValue.A = FMath::Sign(OutValue.A);
+		X->WriteNumberOpcodes(OutData);
+		OutData.WriteOpcode(EMaterialPreshaderOpcode::Sign);
 	}
 	virtual bool IsConstant() const
 	{
@@ -1409,18 +1288,10 @@ public:
 	{}
 
 	// FMaterialUniformExpression interface.
-	virtual void Serialize(FArchive& Ar)
+	virtual void WriteNumberOpcodes(FMaterialPreshaderData& OutData) const override
 	{
-		Ar << X;
-	}
-	virtual void GetNumberValue(const FMaterialRenderContext& Context,FLinearColor& OutValue) const
-	{
-		X->GetNumberValue(Context, OutValue);
-
-		OutValue.R = OutValue.R - FMath::FloorToInt(OutValue.R);
-		OutValue.G = OutValue.G - FMath::FloorToInt(OutValue.G);
-		OutValue.B = OutValue.B - FMath::FloorToInt(OutValue.B);
-		OutValue.A = OutValue.A - FMath::FloorToInt(OutValue.A);
+		X->WriteNumberOpcodes(OutData);
+		OutData.WriteOpcode(EMaterialPreshaderOpcode::Frac);
 	}
 	virtual bool IsConstant() const
 	{
@@ -1454,21 +1325,11 @@ public:
 	{}
 
 	// FMaterialUniformExpression interface.
-	virtual void Serialize(FArchive& Ar)
+	virtual void WriteNumberOpcodes(FMaterialPreshaderData& OutData) const override
 	{
-		Ar << A << B;
-	}
-	virtual void GetNumberValue(const FMaterialRenderContext& Context,FLinearColor& OutValue) const
-	{
-		FLinearColor ValueA = FLinearColor::Black;
-		FLinearColor ValueB = FLinearColor::Black;
-		A->GetNumberValue(Context, ValueA);
-		B->GetNumberValue(Context, ValueB);
-
-		OutValue.R = FMath::Fmod(ValueA.R, ValueB.R);
-		OutValue.G = FMath::Fmod(ValueA.G, ValueB.G);
-		OutValue.B = FMath::Fmod(ValueA.B, ValueB.B);
-		OutValue.A = FMath::Fmod(ValueA.A, ValueB.A);
+		A->WriteNumberOpcodes(OutData);
+		B->WriteNumberOpcodes(OutData);
+		OutData.WriteOpcode(EMaterialPreshaderOpcode::Fmod);
 	}
 	virtual bool IsConstant() const
 	{
@@ -1503,17 +1364,10 @@ public:
 	{}
 
 	// FMaterialUniformExpression interface.
-	virtual void Serialize(FArchive& Ar)
-	{
-		Ar << X;
-	}
-	virtual void GetNumberValue(const FMaterialRenderContext& Context,FLinearColor& OutValue) const
+	virtual void WriteNumberOpcodes(FMaterialPreshaderData& OutData) const override
 	{
-		X->GetNumberValue(Context, OutValue);
-		OutValue.R = FMath::Abs<float>(OutValue.R);
-		OutValue.G = FMath::Abs<float>(OutValue.G);
-		OutValue.B = FMath::Abs<float>(OutValue.B);
-		OutValue.A = FMath::Abs<float>(OutValue.A);
+		X->WriteNumberOpcodes(OutData);
+		OutData.WriteOpcode(EMaterialPreshaderOpcode::Abs);
 	}
 	virtual bool IsConstant() const
 	{
@@ -1547,34 +1401,19 @@ public:
 	{}
 
 	// FMaterialUniformExpression interface.
-	virtual void Serialize(FArchive& Ar) override
+	virtual void WriteNumberOpcodes(FMaterialPreshaderData& OutData) const override
 	{
-		Ar << TextureExpression << TextureProperty;
-	}
-	virtual void GetNumberValue(const FMaterialRenderContext& Context,FLinearColor& OutValue) const override
-	{
-		const UTexture* Texture = nullptr;
-		TextureExpression->GetTextureValue(Context, Context.Material, Texture);
+		FMaterialTextureParameterInfo TextureParameter;
+		TextureExpression->GetTextureParameterInfo(TextureParameter);
 
-		if (!Texture || !Texture->Resource)
-		{
-			return;
-		}
-	
-		if (TextureProperty == TMTM_TextureSize)
-		{
-			OutValue.R = Texture->Resource->GetSizeX();
-			OutValue.G = Texture->Resource->GetSizeY();
-		}
-		else if (TextureProperty == TMTM_TexelSize)
-		{
-			OutValue.R = 1.0f / float(Texture->Resource->GetSizeX());
-			OutValue.G = 1.0f / float(Texture->Resource->GetSizeY());
-		}
-		else
+		EMaterialPreshaderOpcode Op = EMaterialPreshaderOpcode::Nop;
+		switch (TextureProperty)
 		{
-			check(0);
+		case TMTM_TextureSize: Op = EMaterialPreshaderOpcode::TextureSize; break;
+		case TMTM_TexelSize: Op = EMaterialPreshaderOpcode::TexelSize; break;
+		default: checkNoEntry(); break;
 		}
+		OutData.WriteOpcode(Op).Write(TextureParameter.ParameterInfo).Write<int32>(TextureParameter.TextureIndex);
 	}
 	virtual bool IsIdentical(const FMaterialUniformExpression* OtherExpression) const override
 	{
@@ -1611,9 +1450,8 @@ public:
 	FMaterialUniformExpressionExternalTextureCoordinateScaleRotation(const FGuid& InGuid) : FMaterialUniformExpressionExternalTextureBase(InGuid) {}
 	FMaterialUniformExpressionExternalTextureCoordinateScaleRotation(int32 InSourceTextureIndex, TOptional<FName> InParameterName) : FMaterialUniformExpressionExternalTextureBase(InSourceTextureIndex), ParameterName(InParameterName) {}
 
-	virtual void Serialize(FArchive& Ar) override;
 	virtual bool IsIdentical(const FMaterialUniformExpression* OtherExpression) const override;
-	virtual void GetNumberValue(const FMaterialRenderContext& Context, FLinearColor& OutValue) const override;
+	virtual void WriteNumberOpcodes(FMaterialPreshaderData& OutData) const override;
 
 protected:
 	typedef FMaterialUniformExpressionExternalTextureBase Super;
@@ -1634,9 +1472,8 @@ public:
 	FMaterialUniformExpressionExternalTextureCoordinateOffset(const FGuid& InGuid) : FMaterialUniformExpressionExternalTextureBase(InGuid) {}
 	FMaterialUniformExpressionExternalTextureCoordinateOffset(int32 InSourceTextureIndex, TOptional<FName> InParameterName) : FMaterialUniformExpressionExternalTextureBase(InSourceTextureIndex), ParameterName(InParameterName) {}
 
-	virtual void Serialize(FArchive& Ar) override;
 	virtual bool IsIdentical(const FMaterialUniformExpression* OtherExpression) const override;
-	virtual void GetNumberValue(const FMaterialRenderContext& Context, FLinearColor& OutValue) const override;
+	virtual void WriteNumberOpcodes(FMaterialPreshaderData& OutData) const override;
 
 protected:
 	typedef FMaterialUniformExpressionExternalTextureBase Super;
@@ -1661,16 +1498,15 @@ public:
 
 	//~ Begin FMaterialUniformExpression Interface.
 	virtual bool IsConstant() const override { return false; }
-	virtual void Serialize(FArchive& Ar) override;
 	virtual bool IsIdentical(const FMaterialUniformExpression* OtherExpression) const override;
-	virtual void GetNumberValue(const struct FMaterialRenderContext& Context, FLinearColor& OutValue) const override;
+	virtual void WriteNumberOpcodes(FMaterialPreshaderData& OutData) const override;
 	//~ End FMaterialUniformExpression Interface.
 
 protected:
 	/** Is this expression using a material instance parameter. */
 	bool bParameter;
 	/** Contains the parameter info used if bParameter is true. */
-	FMaterialParameterInfo ParameterInfo;
+	FHashedMaterialParameterInfo ParameterInfo;
 	/** Index of the associated URuntimeVirtualTexture in the material texture references used if bParameter is false. */
 	int32 TextureIndex;
 	/** Index of the uniform vector to fetch from the URuntimeVirtualTexture. */
diff --git a/Engine/Source/Runtime/Engine/Private/Materials/MeshMaterialShader.cpp b/Engine/Source/Runtime/Engine/Private/Materials/MeshMaterialShader.cpp
--- a/Engine/Source/Runtime/Engine/Private/Materials/MeshMaterialShader.cpp
+++ b/Engine/Source/Runtime/Engine/Private/Materials/MeshMaterialShader.cpp
@@ -5,6 +5,7 @@
 =============================================================================*/
 
 #include "MeshMaterialShader.h"
+#include "MaterialShaderMapLayout.h"
 #include "ShaderCompiler.h"
 #include "ProfilingDebugging/CookStats.h"
 
@@ -21,13 +22,6 @@ namespace MaterialMeshCookStats
 }
 #endif
 
-static inline bool ShouldCacheMeshShader(const FMeshMaterialShaderType* ShaderType, EShaderPlatform Platform, const FMaterial* Material, FVertexFactoryType* InVertexFactoryType, int32 PermutationId)
-{
-	return ShaderType->ShouldCompilePermutation(Platform, Material, InVertexFactoryType, PermutationId) &&
-		Material->ShouldCache(Platform, ShaderType, InVertexFactoryType) &&
-		InVertexFactoryType->ShouldCache(Platform, Material, ShaderType);
-}
-
 /**
  * Enqueues a compilation for a new shader of this type.
  * @param Platform - The platform to compile for.
@@ -52,9 +46,11 @@ FShaderCompileJob* FMeshMaterialShaderType::BeginCompileShader(
 	FShaderCompilerEnvironment& ShaderEnvironment = NewJob->Input.Environment;
 	ShaderEnvironment.TargetPlatform = MaterialEnvironment->TargetPlatform;
 
+	const FMaterialShaderParameters MaterialParameters(Material);
+
 	// apply the vertex factory changes to the compile environment
 	check(VertexFactoryType);
-	VertexFactoryType->ModifyCompilationEnvironment(Platform, Material, ShaderEnvironment);
+	VertexFactoryType->ModifyCompilationEnvironment(FVertexFactoryShaderPermutationParameters(Platform, MaterialParameters, VertexFactoryType), ShaderEnvironment);
 
 	Material->SetupExtaCompilationSettings(Platform, NewJob->Input.ExtraSettings);
 
@@ -65,7 +61,7 @@ FShaderCompileJob* FMeshMaterialShaderType::BeginCompileShader(
 	COOK_STAT(MaterialMeshCookStats::ShadersCompiled++);
 
 	// Allow the shader type to modify the compile environment.
-	SetupCompileEnvironment(Platform, Material, PermutationId, ShaderEnvironment);
+	SetupCompileEnvironment(Platform, MaterialParameters, VertexFactoryType, PermutationId, ShaderEnvironment);
 
 	bool bAllowDevelopmentShaderCompile = Material->GetAllowDevelopmentShaderCompile();
 
@@ -143,14 +139,15 @@ FShader* FMeshMaterialShaderType::FinishCompileShader(
 	const FSHAHash& MaterialShaderMapHash,
 	const FShaderCompileJob& CurrentJob,
 	const FShaderPipelineType* ShaderPipelineType,
-	const FString& InDebugDescription)
+	const FString& InDebugDescription,
+	FShaderMapResourceBuilder& ResourceBuilder)
 {
 	checkf(CurrentJob.bSucceeded, TEXT("Failed MeshMaterialType compilation job: %s"), *GetJobName(&CurrentJob, ShaderPipelineType, InDebugDescription));
 	checkf(CurrentJob.VFType, TEXT("No VF on MeshMaterialType compilation job: %s"), *GetJobName(&CurrentJob, ShaderPipelineType, InDebugDescription));
 
 	// Reuse an existing resource with the same key or create a new one based on the compile output
 	// This allows FShaders to share compiled bytecode and RHI shader references
-	TRefCountPtr<FShaderResource> Resource = FShaderResource::FindOrCreate(CurrentJob.Output, CurrentJob.PermutationId);
+	const int32 ResourceIndex = ResourceBuilder.FindOrAddCode(CurrentJob.Output);
 
 	if (ShaderPipelineType && !ShaderPipelineType->ShouldOptimizeUnusedOutputs(CurrentJob.Input.Target.GetPlatform()))
 	{
@@ -158,18 +155,41 @@ FShader* FMeshMaterialShaderType::FinishCompileShader(
 		ShaderPipelineType = nullptr;
 	}
 
-	// Find a shader with the same key in memory
-	FShader* Shader = CurrentJob.ShaderType->FindShaderByKey(FShaderKey(MaterialShaderMapHash, ShaderPipelineType, CurrentJob.VFType, CurrentJob.PermutationId, CurrentJob.Input.Target.GetPlatform()));
+	FShader* Shader = ConstructCompiled(CompiledShaderInitializerType(this, CurrentJob.PermutationId, CurrentJob.Output, ResourceIndex, UniformExpressionSet, MaterialShaderMapHash, InDebugDescription, ShaderPipelineType, CurrentJob.VFType));
+	ValidateAfterBind(this, (FMeshMaterialShader*)Shader);
+	CurrentJob.Output.ParameterMap.VerifyBindingsAreComplete(GetName(), CurrentJob.Output.Target, CurrentJob.VFType);
+
+	return Shader;
+}
+
+bool FMeshMaterialShaderType::ShouldCompilePermutation(EShaderPlatform Platform, const FMaterialShaderParameters& MaterialParameters, const FVertexFactoryType* VertexFactoryType, int32 PermutationId) const
+{
+	return FShaderType::ShouldCompilePermutation(FMeshMaterialShaderPermutationParameters(Platform, MaterialParameters, VertexFactoryType, PermutationId));
+}
+
+bool FMeshMaterialShaderType::ShouldCompileVertexFactoryPermutation(const FVertexFactoryType* VertexFactoryType, EShaderPlatform Platform, const FMaterialShaderParameters& MaterialParameters)
+{
+	return VertexFactoryType->ShouldCache(FVertexFactoryShaderPermutationParameters(Platform, MaterialParameters, VertexFactoryType));
+}
 
-	// There was no shader with the same key so create a new one with the compile output, which will bind shader parameters
-	if (!Shader)
+bool FMeshMaterialShaderType::ShouldCompilePipeline(const FShaderPipelineType* ShaderPipelineType, EShaderPlatform Platform, const FMaterialShaderParameters& MaterialParameters, const FVertexFactoryType* VertexFactoryType)
+{
+	const FMeshMaterialShaderPermutationParameters Parameters(Platform, MaterialParameters, VertexFactoryType, kUniqueShaderPermutationId);
+	for (const FShaderType* ShaderType : ShaderPipelineType->GetStages())
 	{
-		Shader = (*ConstructCompiledRef)(CompiledShaderInitializerType(this, CurrentJob.PermutationId, CurrentJob.Output, MoveTemp(Resource), UniformExpressionSet, MaterialShaderMapHash, InDebugDescription, ShaderPipelineType, CurrentJob.VFType));
-		ValidateAfterBind((FMeshMaterialShader*)Shader);
-		CurrentJob.Output.ParameterMap.VerifyBindingsAreComplete(GetName(), CurrentJob.Output.Target, CurrentJob.VFType);
+		checkSlow(ShaderType->GetMeshMaterialShaderType());
+		if (!ShaderType->ShouldCompilePermutation(Parameters))
+		{
+			return false;
+		}
 	}
+	return true;
+}
 
-	return Shader;
+void FMeshMaterialShaderType::SetupCompileEnvironment(EShaderPlatform Platform, const FMaterialShaderParameters& MaterialParameters, const FVertexFactoryType* VertexFactoryType, int32 PermutationId, FShaderCompilerEnvironment& Environment)
+{
+	// Allow the shader type to modify its compile environment.
+	FShaderType::ModifyCompilationEnvironment(FMeshMaterialShaderPermutationParameters(Platform, MaterialParameters, VertexFactoryType, PermutationId), Environment);
 }
 
 /**
@@ -182,6 +202,7 @@ uint32 FMeshMaterialShaderMap::BeginCompile(
 	uint32 ShaderMapId,
 	const FMaterialShaderMapId& InShaderMapId, 
 	const FMaterial* Material,
+	const FMeshMaterialShaderMapLayout& MeshLayout,
 	FShaderCompilerEnvironment* MaterialEnvironment,
 	EShaderPlatform InPlatform,
 	TArray<TSharedRef<FShaderCommonCompileJob, ESPMode::ThreadSafe>>& NewJobs,
@@ -189,228 +210,103 @@ uint32 FMeshMaterialShaderMap::BeginCompile(
 	FString DebugExtension
 	)
 {
-	if (!VertexFactoryType)
-	{
-		return 0;
-	}
-
 	uint32 NumShadersPerVF = 0;
 	TSet<FString> ShaderTypeNames;
 
 	// Iterate over all mesh material shader types.
 	TMap<TShaderTypePermutation<const FShaderType>, FShaderCompileJob*> SharedShaderJobs;
-	for (TLinkedList<FShaderType*>::TIterator ShaderTypeIt(FShaderType::GetTypeList());ShaderTypeIt;ShaderTypeIt.Next())
+	for (const FShaderLayoutEntry& Shader : MeshLayout.Shaders)
 	{
-		FMeshMaterialShaderType* ShaderType = ShaderTypeIt->GetMeshMaterialShaderType();
-		if (!ShaderType)
+		FMeshMaterialShaderType* ShaderType = static_cast<FMeshMaterialShaderType*>(Shader.ShaderType);
+		if (!Material->ShouldCache(InPlatform, ShaderType, MeshLayout.VertexFactoryType))
 		{
 			continue;
 		}
 
-		for (int32 PermutationId = 0; PermutationId < ShaderType->GetPermutationCount(); ++PermutationId)
-		{
-			if (ShouldCacheMeshShader(ShaderType, InPlatform, Material, VertexFactoryType, PermutationId))
-			{
-				// Verify that the shader map Id contains inputs for any shaders that will be put into this shader map
+		// Verify that the shader map Id contains inputs for any shaders that will be put into this shader map
 #if WITH_EDITOR
-				check(InShaderMapId.ContainsVertexFactoryType(VertexFactoryType));
-				check(InShaderMapId.ContainsShaderType(ShaderType, kUniqueShaderPermutationId));
+		check(InShaderMapId.ContainsVertexFactoryType(MeshLayout.VertexFactoryType));
+		check(InShaderMapId.ContainsShaderType(ShaderType, kUniqueShaderPermutationId));
 #endif
 
-				NumShadersPerVF++;
-				// only compile the shader if we don't already have it
-				if (!HasShader(ShaderType, PermutationId))
-				{
-					// Compile this mesh material shader for this material and vertex factory type.
-					FShaderCompileJob* Job = ShaderType->BeginCompileShader(
-						ShaderMapId,
-						PermutationId,
-						InPlatform,
-						Material,
-						MaterialEnvironment,
-						VertexFactoryType,
-						nullptr,
-						NewJobs,
-						DebugDescription, 
-						DebugExtension
-						);
-					TShaderTypePermutation<const FShaderType> ShaderTypePermutation(ShaderType, PermutationId);
-					check(!SharedShaderJobs.Find(ShaderTypePermutation));
-					SharedShaderJobs.Add(ShaderTypePermutation, Job);
-				}
-			}
+		NumShadersPerVF++;
+		// only compile the shader if we don't already have it
+		if (!HasShader(ShaderType, Shader.PermutationId))
+		{
+			// Compile this mesh material shader for this material and vertex factory type.
+			FShaderCompileJob* Job = ShaderType->BeginCompileShader(
+				ShaderMapId,
+				Shader.PermutationId,
+				InPlatform,
+				Material,
+				MaterialEnvironment,
+				MeshLayout.VertexFactoryType,
+				nullptr,
+				NewJobs,
+				DebugDescription,
+				DebugExtension
+			);
+			TShaderTypePermutation<const FShaderType> ShaderTypePermutation(ShaderType, Shader.PermutationId);
+			check(!SharedShaderJobs.Find(ShaderTypePermutation));
+			SharedShaderJobs.Add(ShaderTypePermutation, Job);
 		}
 	}
 
 	// Now the pipeline jobs; if it's a shareable pipeline, do not add duplicate jobs
-	const bool bHasTessellation = Material->GetTessellationMode() != MTM_NoTessellation;
-	for (TLinkedList<FShaderPipelineType*>::TIterator ShaderPipelineIt(FShaderPipelineType::GetTypeList());ShaderPipelineIt;ShaderPipelineIt.Next())
+	for (FShaderPipelineType* Pipeline : MeshLayout.ShaderPipelines)
 	{
-		const FShaderPipelineType* Pipeline = *ShaderPipelineIt;
-		if (Pipeline->IsMeshMaterialTypePipeline() && Pipeline->HasTessellation() == bHasTessellation)
+		if (!Material->ShouldCachePipeline(InPlatform, Pipeline, MeshLayout.VertexFactoryType))
+		{
+			continue;
+		}
+
+		auto& StageTypes = Pipeline->GetStages();
+
+		// Verify that the shader map Id contains inputs for any shaders that will be put into this shader map
+#if WITH_EDITOR
+		check(InShaderMapId.ContainsShaderPipelineType(Pipeline));
+#endif
+		if (Pipeline->ShouldOptimizeUnusedOutputs(InPlatform))
 		{
-			auto& StageTypes = Pipeline->GetStages();
-			int32 NumShaderStagesToCompile = 0;
+			NumShadersPerVF += StageTypes.Num();
+			TArray<FMeshMaterialShaderType*> ShaderStagesToCompile;
 			for (auto* Shader : StageTypes)
 			{
 				const FMeshMaterialShaderType* ShaderType = Shader->GetMeshMaterialShaderType();
-				if (ShouldCacheMeshShader(ShaderType, InPlatform, Material, VertexFactoryType, kUniqueShaderPermutationId))
-				{
-					++NumShaderStagesToCompile;
-				}
-				else
-				{
-					break;
-				}
-			}
 
-			if (NumShaderStagesToCompile == StageTypes.Num())
-			{
 				// Verify that the shader map Id contains inputs for any shaders that will be put into this shader map
 #if WITH_EDITOR
-				check(InShaderMapId.ContainsShaderPipelineType(Pipeline));
-#endif
-				if (Pipeline->ShouldOptimizeUnusedOutputs(InPlatform))
-				{
-					NumShadersPerVF += NumShaderStagesToCompile;
-					TArray<FMeshMaterialShaderType*> ShaderStagesToCompile;
-					for (auto* Shader : StageTypes)
-					{
-						const FMeshMaterialShaderType* ShaderType = Shader->GetMeshMaterialShaderType();
-
-						// Verify that the shader map Id contains inputs for any shaders that will be put into this shader map
-#if WITH_EDITOR
-						check(InShaderMapId.ContainsVertexFactoryType(VertexFactoryType));
-						check(InShaderMapId.ContainsShaderType(ShaderType, kUniqueShaderPermutationId));
+				check(InShaderMapId.ContainsVertexFactoryType(MeshLayout.VertexFactoryType));
+				check(InShaderMapId.ContainsShaderType(ShaderType, kUniqueShaderPermutationId));
 #endif
-						ShaderStagesToCompile.Add((FMeshMaterialShaderType*)ShaderType);
-					}
-
-					// Make a pipeline job with all the stages
-					FMeshMaterialShaderType::BeginCompileShaderPipeline(ShaderMapId, kUniqueShaderPermutationId, InPlatform, Material, MaterialEnvironment, VertexFactoryType, Pipeline, ShaderStagesToCompile, NewJobs, DebugDescription, DebugExtension);
-				}
-				else
-				{
-					// If sharing shaders amongst pipelines, add this pipeline as a dependency of an existing job
-					for (const FShaderType* ShaderType : StageTypes)
-					{
-						TShaderTypePermutation<const FShaderType> ShaderTypePermutation(ShaderType, kUniqueShaderPermutationId);
-						FShaderCompileJob** Job = SharedShaderJobs.Find(ShaderTypePermutation);
-						checkf(Job, TEXT("Couldn't find existing shared job for mesh shader %s on pipeline %s!"), ShaderType->GetName(), Pipeline->GetName());
-						auto* SingleJob = (*Job)->GetSingleShaderJob();
-						auto& PipelinesToShare = SingleJob->SharingPipelines.FindOrAdd(VertexFactoryType);
-						check(!PipelinesToShare.Contains(Pipeline));
-						PipelinesToShare.Add(Pipeline);
-					}
-				}
+				ShaderStagesToCompile.Add((FMeshMaterialShaderType*)ShaderType);
 			}
-		}
-	}
-
-	if (NumShadersPerVF > 0)
-	{
-		UE_LOG(LogShaders, Verbose, TEXT("			%s - %u shaders"), VertexFactoryType->GetName(), NumShadersPerVF);
-	}
 
-	return NumShadersPerVF;
-}
-
-inline bool FMeshMaterialShaderMap::IsMeshShaderComplete(
-	const FMeshMaterialShaderMap* MeshShaderMap,
-	EShaderPlatform Platform,
-	const FMaterial* Material,
-	const FMeshMaterialShaderType* ShaderType,
-	const FShaderPipelineType* Pipeline,
-	FVertexFactoryType* InVertexFactoryType,
-	int32 PermutationId,
-	bool bSilent)
-{
-	if (!ShouldCacheMeshShader(ShaderType, Platform, Material, InVertexFactoryType, PermutationId))
-	{
-		return true;
-	}
-
-	if (!MeshShaderMap || (Pipeline && !MeshShaderMap->HasShaderPipeline(Pipeline)) || (!Pipeline && !MeshShaderMap->HasShader((FShaderType*)ShaderType, PermutationId)))
-	{
-		if (!bSilent)
-		{
-			if (Pipeline)
-			{
-				UE_LOG(LogShaders, Warning, TEXT("Incomplete material %s, missing pipeline %s from %s."), *Material->GetFriendlyName(), Pipeline->GetName(), InVertexFactoryType->GetName());
-			}
-			else
-			{
-				UE_LOG(LogShaders, Warning, TEXT("Incomplete material %s, missing (%s, %d) from %s."), *Material->GetFriendlyName(), ShaderType->GetName(), PermutationId, InVertexFactoryType->GetName());
-			}
+			// Make a pipeline job with all the stages
+			FMeshMaterialShaderType::BeginCompileShaderPipeline(ShaderMapId, kUniqueShaderPermutationId, InPlatform, Material, MaterialEnvironment, MeshLayout.VertexFactoryType, Pipeline, ShaderStagesToCompile, NewJobs, DebugDescription, DebugExtension);
 		}
-		return false;
-	}
-
-	return true;
-};
-
-
-bool FMeshMaterialShaderMap::IsComplete(
-	const FMeshMaterialShaderMap* MeshShaderMap,
-	EShaderPlatform Platform,
-	const FMaterial* Material,
-	FVertexFactoryType* InVertexFactoryType,
-	bool bSilent)
-{
-	// Iterate over all mesh material shader types.
-	for (TLinkedList<FShaderType*>::TIterator ShaderTypeIt(FShaderType::GetTypeList());ShaderTypeIt;ShaderTypeIt.Next())
-	{
-		FMeshMaterialShaderType* ShaderType = ShaderTypeIt->GetMeshMaterialShaderType();
-		const int32 PermutationCount = ShaderType ? ShaderType->GetPermutationCount() : 0;
-		for (int32 PermutationId = 0; PermutationId < PermutationCount; ++PermutationId)
+		else
 		{
-			if (!IsMeshShaderComplete(MeshShaderMap, Platform, Material, ShaderType, nullptr, InVertexFactoryType, PermutationId, bSilent))
+			// If sharing shaders amongst pipelines, add this pipeline as a dependency of an existing job
+			for (const FShaderType* ShaderType : StageTypes)
 			{
-				return false;
+				TShaderTypePermutation<const FShaderType> ShaderTypePermutation(ShaderType, kUniqueShaderPermutationId);
+				FShaderCompileJob** Job = SharedShaderJobs.Find(ShaderTypePermutation);
+				checkf(Job, TEXT("Couldn't find existing shared job for mesh shader %s on pipeline %s!"), ShaderType->GetName(), Pipeline->GetName());
+				auto* SingleJob = (*Job)->GetSingleShaderJob();
+				auto& PipelinesToShare = SingleJob->SharingPipelines.FindOrAdd(MeshLayout.VertexFactoryType);
+				check(!PipelinesToShare.Contains(Pipeline));
+				PipelinesToShare.Add(Pipeline);
 			}
 		}
 	}
 
-	// Iterate over all pipeline types
-	const bool bHasTessellation = Material->GetTessellationMode() != MTM_NoTessellation;
-	for (TLinkedList<FShaderPipelineType*>::TIterator ShaderPipelineIt(FShaderPipelineType::GetTypeList());ShaderPipelineIt;ShaderPipelineIt.Next())
+	if (NumShadersPerVF > 0)
 	{
-		const FShaderPipelineType* ShaderPipelineType = *ShaderPipelineIt;
-		if (ShaderPipelineType->IsMeshMaterialTypePipeline() && ShaderPipelineType->HasTessellation() == bHasTessellation)
-		{
-			auto& Stages = ShaderPipelineType->GetStages();
-
-			// Verify all the ShouldCache are in sync
-			int32 NumShouldCache = 0;
-			for (int32 Index = 0; Index < Stages.Num(); ++Index)
-			{
-				auto* ShaderType = Stages[Index]->GetMeshMaterialShaderType();
-				if (ShouldCacheMeshShader(ShaderType, Platform, Material, InVertexFactoryType, kUniqueShaderPermutationId))
-				{
-					++NumShouldCache;
-				}
-				else
-				{
-					break;
-				}
-			}
-
-			if (NumShouldCache == Stages.Num())
-			{
-				// Now check the completeness of the shader map
-				for (int32 Index = 0; Index < Stages.Num(); ++Index)
-				{
-					auto* ShaderType = Stages[Index]->GetMeshMaterialShaderType();
-					if (ShaderType && !IsMeshShaderComplete(MeshShaderMap, Platform, Material, ShaderType, ShaderPipelineType, InVertexFactoryType, kUniqueShaderPermutationId, bSilent))
-					{
-						return false;
-					}
-				}
-			}
-		}
+		UE_LOG(LogShaders, Verbose, TEXT("			%s - %u shaders"), MeshLayout.VertexFactoryType->GetName(), NumShadersPerVF);
 	}
 
-	return true;
+	return NumShadersPerVF;
 }
 
 #if WITH_EDITOR
@@ -419,6 +315,7 @@ void FMeshMaterialShaderMap::LoadMissingShadersFromMemory(
 	const FMaterial* Material, 
 	EShaderPlatform InPlatform)
 {
+#if 0
 	for (TLinkedList<FShaderType*>::TIterator ShaderTypeIt(FShaderType::GetTypeList());ShaderTypeIt;ShaderTypeIt.Next())
 	{
 		FMeshMaterialShaderType* ShaderType = ShaderTypeIt->GetMeshMaterialShaderType();
@@ -480,12 +377,13 @@ void FMeshMaterialShaderMap::LoadMissingShadersFromMemory(
 
 				if (ShadersForPipeline.Num() == NumShaders && !HasShaderPipeline(PipelineType))
 				{
-					auto* Pipeline = new FShaderPipeline(PipelineType, ShadersForPipeline);
+					auto* Pipeline = new FShaderPipeline(PointerTable, PipelineType, ShadersForPipeline);
 					AddShaderPipeline(PipelineType, Pipeline);
 				}
 			}
 		}
 	}
+#endif
 }
 #endif // WITH_EDITOR
 
@@ -493,7 +391,7 @@ void FMeshMaterialShaderMap::LoadMissingShadersFromMemory(
  * Removes all entries in the cache with exceptions based on a shader type
  * @param ShaderType - The shader type to flush
  */
-void FMeshMaterialShaderMap::FlushShadersByShaderType(FShaderType* ShaderType)
+void FMeshMaterialShaderMap::FlushShadersByShaderType(const FShaderType* ShaderType)
 {
 	if (ShaderType->GetMeshMaterialShaderType())
 	{
diff --git a/Engine/Source/Runtime/Engine/Private/Materials/ParameterCollection.cpp b/Engine/Source/Runtime/Engine/Private/Materials/ParameterCollection.cpp
--- a/Engine/Source/Runtime/Engine/Private/Materials/ParameterCollection.cpp
+++ b/Engine/Source/Runtime/Engine/Private/Materials/ParameterCollection.cpp
@@ -218,9 +218,9 @@ void UMaterialParameterCollection::PostEditChangeProperty(FPropertyChangedEvent&
 				}
 				else
 				{
-					for (int32 FunctionIndex = 0; FunctionIndex < CurrentMaterial->MaterialParameterCollectionInfos.Num() && !bRecompile; FunctionIndex++)
+					for (int32 FunctionIndex = 0; FunctionIndex < CurrentMaterial->GetCachedExpressionData().ParameterCollectionInfos.Num() && !bRecompile; FunctionIndex++)
 					{
-						if (CurrentMaterial->MaterialParameterCollectionInfos[FunctionIndex].ParameterCollection == this)
+						if (CurrentMaterial->GetCachedExpressionData().ParameterCollectionInfos[FunctionIndex].ParameterCollection == this)
 						{
 							bRecompile = true;
 							break;
@@ -409,7 +409,6 @@ void UMaterialParameterCollection::CreateBufferStruct()
 	new(Members) FShaderParametersMetadata::FMember(TEXT("Vectors"),TEXT(""),NextMemberOffset,UBMT_FLOAT32,EShaderPrecisionModifier::Half,1,4,NumVectors, nullptr);
 	const uint32 VectorArraySize = NumVectors * sizeof(FVector4);
 	NextMemberOffset += VectorArraySize;
-	static FName LayoutName(TEXT("MaterialCollection"));
 	const uint32 StructSize = Align(NextMemberOffset, SHADER_PARAMETER_STRUCT_ALIGNMENT);
 
 	// If Collections ever get non-numeric resources (eg Textures), OutEnvironment.ResourceTableMap has a map by name
@@ -417,7 +416,7 @@ void UMaterialParameterCollection::CreateBufferStruct()
 	// (and the hlsl cbuffers are named MaterialCollection0, etc, so the names don't match the layout)
 	UniformBufferStruct = MakeUnique<FShaderParametersMetadata>(
 		FShaderParametersMetadata::EUseCase::DataDrivenUniformBuffer,
-		LayoutName,
+		TEXT("MaterialCollection"),
 		TEXT("MaterialCollection"),
 		TEXT("MaterialCollection"),
 		nullptr,
@@ -686,9 +685,8 @@ void FMaterialParameterCollectionInstanceResource::GameThread_Destroy()
 	);
 }
 
-static FName MaterialParameterCollectionInstanceResourceName(TEXT("MaterialParameterCollectionInstanceResource"));
 FMaterialParameterCollectionInstanceResource::FMaterialParameterCollectionInstanceResource() :
-	UniformBufferLayout(MaterialParameterCollectionInstanceResourceName)
+	UniformBufferLayout(TEXT("MaterialParameterCollectionInstanceResource"))
 {
 }
 
diff --git a/Engine/Source/Runtime/Engine/Private/MeshParticleVertexFactory.cpp b/Engine/Source/Runtime/Engine/Private/MeshParticleVertexFactory.cpp
--- a/Engine/Source/Runtime/Engine/Private/MeshParticleVertexFactory.cpp
+++ b/Engine/Source/Runtime/Engine/Private/MeshParticleVertexFactory.cpp
@@ -11,9 +11,9 @@
 
 class FMeshParticleVertexFactoryShaderParameters : public FVertexFactoryShaderParameters
 {
+	DECLARE_INLINE_TYPE_LAYOUT(FMeshParticleVertexFactoryShaderParameters, NonVirtual);
 public:
-
-	virtual void Bind(const FShaderParameterMap& ParameterMap) override
+	void Bind(const FShaderParameterMap& ParameterMap)
 	{
 		Transform1.Bind(ParameterMap,TEXT("Transform1"));
 		Transform2.Bind(ParameterMap,TEXT("Transform2"));
@@ -30,24 +30,7 @@ public:
 		PrevTransformBuffer.Bind(ParameterMap, TEXT("PrevTransformBuffer"));
 	}
 
-	virtual void Serialize(FArchive& Ar) override
-	{
-		Ar << Transform1;
-		Ar << Transform2;
-		Ar << Transform3;
-		Ar << SubUVParams;
-		Ar << SubUVLerp;
-		Ar << ParticleDirection;
-		Ar << RelativeTime;
-		Ar << DynamicParameter;
-		Ar << ParticleColor;
-		Ar << PrevTransform0;
-		Ar << PrevTransform1;
-		Ar << PrevTransform2;
-		Ar << PrevTransformBuffer;
-	}
-
-	virtual void GetElementShaderBindings(
+	void GetElementShaderBindings(
 		const FSceneInterface* Scene,
 		const FSceneView* View,
 		const FMeshMaterialShader* Shader,
@@ -56,7 +39,7 @@ public:
 		const FVertexFactory* VertexFactory,
 		const FMeshBatchElement& BatchElement,
 		class FMeshDrawSingleShaderBindings& ShaderBindings,
-		FVertexInputStreamArray& VertexStreams) const override
+		FVertexInputStreamArray& VertexStreams) const
 	{
 		const bool bInstanced = true;
 		FMeshParticleVertexFactory* MeshParticleVF = (FMeshParticleVertexFactory*)VertexFactory;
@@ -99,19 +82,19 @@ public:
 
 private:
 	// Used only when instancing is off (ES2)
-	FShaderParameter Transform1;
-	FShaderParameter Transform2;
-	FShaderParameter Transform3;
-	FShaderParameter SubUVParams;
-	FShaderParameter SubUVLerp;
-	FShaderParameter ParticleDirection;
-	FShaderParameter RelativeTime;
-	FShaderParameter DynamicParameter;
-	FShaderParameter ParticleColor;
-	FShaderParameter PrevTransform0;
-	FShaderParameter PrevTransform1;
-	FShaderParameter PrevTransform2;
-	FShaderResourceParameter PrevTransformBuffer;
+	LAYOUT_FIELD(FShaderParameter, Transform1);
+	LAYOUT_FIELD(FShaderParameter, Transform2);
+	LAYOUT_FIELD(FShaderParameter, Transform3);
+	LAYOUT_FIELD(FShaderParameter, SubUVParams);
+	LAYOUT_FIELD(FShaderParameter, SubUVLerp);
+	LAYOUT_FIELD(FShaderParameter, ParticleDirection);
+	LAYOUT_FIELD(FShaderParameter, RelativeTime);
+	LAYOUT_FIELD(FShaderParameter, DynamicParameter);
+	LAYOUT_FIELD(FShaderParameter, ParticleColor);
+	LAYOUT_FIELD(FShaderParameter, PrevTransform0);
+	LAYOUT_FIELD(FShaderParameter, PrevTransform1);
+	LAYOUT_FIELD(FShaderParameter, PrevTransform2);
+	LAYOUT_FIELD(FShaderResourceParameter, PrevTransformBuffer);
 };
 
 class FDummyPrevTransformBuffer : public FRenderResource
@@ -321,9 +304,24 @@ FRHIShaderResourceView* FMeshParticleVertexFactory::GetPreviousTransformBufferSR
 	return PrevTransformBuffer.SRV;
 }
 
-bool FMeshParticleVertexFactory::ShouldCompilePermutation(EShaderPlatform Platform, const class FMaterial* Material, const class FShaderType* ShaderType)
+void FMeshParticleVertexFactory::ModifyCompilationEnvironment(const FVertexFactoryShaderPermutationParameters& Parameters, FShaderCompilerEnvironment& OutEnvironment)
+{
+	FParticleVertexFactoryBase::ModifyCompilationEnvironment(Parameters, OutEnvironment);
+
+	// Set a define so we can tell in MaterialTemplate.usf when we are compiling a mesh particle vertex factory
+	OutEnvironment.SetDefine(TEXT("PARTICLE_MESH_FACTORY"), TEXT("1"));
+	OutEnvironment.SetDefine(TEXT("PARTICLE_MESH_INSTANCED"), TEXT("1"));
+
+	const bool ContainsManualVertexFetch = OutEnvironment.GetDefinitions().Contains("MANUAL_VERTEX_FETCH");
+	if (!ContainsManualVertexFetch && RHISupportsManualVertexFetch(Parameters.Platform))
+	{
+		OutEnvironment.SetDefine(TEXT("MANUAL_VERTEX_FETCH"), TEXT("1"));
+	}
+}
+
+bool FMeshParticleVertexFactory::ShouldCompilePermutation(const FVertexFactoryShaderPermutationParameters& Parameters)
 {
-	return (Material->IsUsedWithMeshParticles() || Material->IsSpecialEngineMaterial());
+	return (Parameters.MaterialParameters.bIsUsedWithMeshParticles || Parameters.MaterialParameters.bIsSpecialEngineMaterial);
 }
 
 void FMeshParticleVertexFactory::SetData(const FDataType& InData)
@@ -333,12 +331,21 @@ void FMeshParticleVertexFactory::SetData(const FDataType& InData)
 	UpdateRHI();
 }
 
+bool FMeshParticleVertexFactoryEmulatedInstancing::ShouldCompilePermutation(const FVertexFactoryShaderPermutationParameters& Parameters)
+{
+	return (Parameters.Platform == SP_OPENGL_ES2_ANDROID) // Android platforms that might not support hardware instancing
+		&& FMeshParticleVertexFactory::ShouldCompilePermutation(Parameters);
+}
 
-FVertexFactoryShaderParameters* FMeshParticleVertexFactory::ConstructShaderParameters(EShaderFrequency ShaderFrequency)
+void FMeshParticleVertexFactoryEmulatedInstancing::ModifyCompilationEnvironment(const FVertexFactoryShaderPermutationParameters& Parameters, FShaderCompilerEnvironment& OutEnvironment)
 {
-	return ShaderFrequency == SF_Vertex ? new FMeshParticleVertexFactoryShaderParameters() : NULL;
+	FMeshParticleVertexFactory::ModifyCompilationEnvironment(Parameters, OutEnvironment);
+	OutEnvironment.SetDefine(TEXT("PARTICLE_MESH_INSTANCED"), TEXT("0"));
 }
 
+IMPLEMENT_VERTEX_FACTORY_PARAMETER_TYPE(FMeshParticleVertexFactory, SF_Vertex, FMeshParticleVertexFactoryShaderParameters);
+IMPLEMENT_VERTEX_FACTORY_PARAMETER_TYPE(FMeshParticleVertexFactoryEmulatedInstancing, SF_Vertex, FMeshParticleVertexFactoryShaderParameters);
+
 IMPLEMENT_VERTEX_FACTORY_TYPE(FMeshParticleVertexFactory,"/Engine/Private/MeshParticleVertexFactory.ush",true,false,true,false,false);
 IMPLEMENT_VERTEX_FACTORY_TYPE(FMeshParticleVertexFactoryEmulatedInstancing,"/Engine/Private/MeshParticleVertexFactory.ush",true,false,true,false,false);
 IMPLEMENT_GLOBAL_SHADER_PARAMETER_STRUCT(FMeshParticleUniformParameters,"MeshParticleVF");
diff --git a/Engine/Source/Runtime/Engine/Private/ModelRender.cpp b/Engine/Source/Runtime/Engine/Private/ModelRender.cpp
--- a/Engine/Source/Runtime/Engine/Private/ModelRender.cpp
+++ b/Engine/Source/Runtime/Engine/Private/ModelRender.cpp
@@ -815,7 +815,7 @@ public:
 	friend class UModelComponent;
 };
 
-void UModelComponent::CreateRenderState_Concurrent()
+void UModelComponent::CreateRenderState_Concurrent(FRegisterComponentContext* Context)
 {
 	for (int32 ElementIndex = 0; ElementIndex < Elements.Num(); ElementIndex++)
 	{
@@ -829,7 +829,7 @@ void UModelComponent::CreateRenderState_Concurrent()
 		++GetModel()->VertexBuffer.RefCount;
 	}
 
-	Super::CreateRenderState_Concurrent();
+	Super::CreateRenderState_Concurrent(Context);
 }
 
 void UModelComponent::DestroyRenderState_Concurrent()
diff --git a/Engine/Source/Runtime/Engine/Private/Particles/ParticleBeamTrailVertexFactory.cpp b/Engine/Source/Runtime/Engine/Private/Particles/ParticleBeamTrailVertexFactory.cpp
--- a/Engine/Source/Runtime/Engine/Private/Particles/ParticleBeamTrailVertexFactory.cpp
+++ b/Engine/Source/Runtime/Engine/Private/Particles/ParticleBeamTrailVertexFactory.cpp
@@ -17,16 +17,9 @@ IMPLEMENT_GLOBAL_SHADER_PARAMETER_STRUCT(FParticleBeamTrailUniformParameters,"Be
  */
 class FParticleBeamTrailVertexFactoryShaderParameters : public FVertexFactoryShaderParameters
 {
+	DECLARE_INLINE_TYPE_LAYOUT(FParticleBeamTrailVertexFactoryShaderParameters, NonVirtual);
 public:
-	virtual void Bind(const FShaderParameterMap& ParameterMap) override
-	{
-	}
-
-	virtual void Serialize(FArchive& Ar) override
-	{
-	}
-
-	virtual void GetElementShaderBindings(
+	void GetElementShaderBindings(
 		const FSceneInterface* Scene,
 		const FSceneView* View,
 		const FMeshMaterialShader* Shader,
@@ -35,11 +28,14 @@ public:
 		const FVertexFactory* VertexFactory,
 		const FMeshBatchElement& BatchElement,
 		class FMeshDrawSingleShaderBindings& ShaderBindings,
-		FVertexInputStreamArray& VertexStreams) const override
+		FVertexInputStreamArray& VertexStreams) const
 	{
 		FParticleBeamTrailVertexFactory* BeamTrailVF = (FParticleBeamTrailVertexFactory*)VertexFactory;
 		ShaderBindings.Add(Shader->GetUniformBufferParameter<FParticleBeamTrailUniformParameters>(), BeamTrailVF->GetBeamTrailUniformBuffer() );
 	}
+
+	
+	
 };
 
 ///////////////////////////////////////////////////////////////////////////////
@@ -107,17 +103,17 @@ static TGlobalResource<FParticleBeamTrailVertexDeclaration> GParticleBeamTrailVe
 
 ///////////////////////////////////////////////////////////////////////////////
 
-bool FParticleBeamTrailVertexFactory::ShouldCompilePermutation(EShaderPlatform Platform, const class FMaterial* Material, const class FShaderType* ShaderType)
+bool FParticleBeamTrailVertexFactory::ShouldCompilePermutation(const FVertexFactoryShaderPermutationParameters& Parameters)
 {
-	return Material->IsUsedWithBeamTrails() || Material->IsSpecialEngineMaterial();
+	return Parameters.MaterialParameters.bIsUsedWithBeamTrails || Parameters.MaterialParameters.bIsSpecialEngineMaterial;
 }
 
 /**
  * Can be overridden by FVertexFactory subclasses to modify their compile environment just before compilation occurs.
  */
-void FParticleBeamTrailVertexFactory::ModifyCompilationEnvironment(const FVertexFactoryType* Type, EShaderPlatform Platform, const class FMaterial* Material, FShaderCompilerEnvironment& OutEnvironment)
+void FParticleBeamTrailVertexFactory::ModifyCompilationEnvironment(const FVertexFactoryShaderPermutationParameters& Parameters, FShaderCompilerEnvironment& OutEnvironment)
 {
-	FParticleVertexFactoryBase::ModifyCompilationEnvironment(Type, Platform, Material, OutEnvironment);
+	FParticleVertexFactoryBase::ModifyCompilationEnvironment(Parameters, OutEnvironment);
 	OutEnvironment.SetDefine(TEXT("PARTICLE_BEAMTRAIL_FACTORY"),TEXT("1"));
 }
 
@@ -133,11 +129,6 @@ void FParticleBeamTrailVertexFactory::InitRHI()
 	FVertexStream* DynamicParameterStream = new(Streams) FVertexStream;
 }
 
-FVertexFactoryShaderParameters* FParticleBeamTrailVertexFactory::ConstructShaderParameters(EShaderFrequency ShaderFrequency)
-{
-	return ShaderFrequency == SF_Vertex ? new FParticleBeamTrailVertexFactoryShaderParameters() : NULL;
-}
-
 void FParticleBeamTrailVertexFactory::SetVertexBuffer(const FVertexBuffer* InBuffer, uint32 StreamOffset, uint32 Stride)
 {
 	check(Streams.Num() == 2);
@@ -169,4 +160,6 @@ void FParticleBeamTrailVertexFactory::SetDynamicParameterBuffer(const FVertexBuf
 
 ///////////////////////////////////////////////////////////////////////////////
 
+IMPLEMENT_VERTEX_FACTORY_PARAMETER_TYPE(FParticleBeamTrailVertexFactory, SF_Vertex, FParticleBeamTrailVertexFactoryShaderParameters);
+
 IMPLEMENT_VERTEX_FACTORY_TYPE(FParticleBeamTrailVertexFactory,"/Engine/Private/ParticleBeamTrailVertexFactory.ush",true,false,true,false,false);
diff --git a/Engine/Source/Runtime/Engine/Private/Particles/ParticleComponents.cpp b/Engine/Source/Runtime/Engine/Private/Particles/ParticleComponents.cpp
--- a/Engine/Source/Runtime/Engine/Private/Particles/ParticleComponents.cpp
+++ b/Engine/Source/Runtime/Engine/Private/Particles/ParticleComponents.cpp
@@ -3925,7 +3925,7 @@ void UParticleSystemComponent::OnEndOfFrameUpdateDuringTick()
 	WaitForAsyncAndFinalize(STALL);
 }
 
-void UParticleSystemComponent::CreateRenderState_Concurrent()
+void UParticleSystemComponent::CreateRenderState_Concurrent(FRegisterComponentContext* Context)
 {
 	LLM_SCOPE(ELLMTag::Particles);
 	SCOPE_CYCLE_COUNTER(STAT_ParticleSystemComponent_CreateRenderState_Concurrent);
@@ -3954,7 +3954,7 @@ void UParticleSystemComponent::CreateRenderState_Concurrent()
 		}
 	}
 
-	Super::CreateRenderState_Concurrent();
+	Super::CreateRenderState_Concurrent(Context);
 
 	bJustRegistered = true;
 }
diff --git a/Engine/Source/Runtime/Engine/Private/Particles/ParticleCurveTexture.cpp b/Engine/Source/Runtime/Engine/Private/Particles/ParticleCurveTexture.cpp
--- a/Engine/Source/Runtime/Engine/Private/Particles/ParticleCurveTexture.cpp
+++ b/Engine/Source/Runtime/Engine/Private/Particles/ParticleCurveTexture.cpp
@@ -82,7 +82,7 @@ public:
 		Parameters.PixelScale.Y = 1.0f / GParticleCurveTextureSizeY;
 		Parameters.CurveOffset = CurveOffset;
 		FParticleCurveInjectionBufferRef UniformBuffer = FParticleCurveInjectionBufferRef::CreateUniformBufferImmediate(Parameters, UniformBuffer_SingleDraw);
-		FRHIVertexShader* VertexShader = GetVertexShader();
+		FRHIVertexShader* VertexShader = RHICmdList.GetBoundVertexShader();
 		SetUniformBufferParameter(RHICmdList, VertexShader, GetUniformBufferParameter<FParticleCurveInjectionParameters>(), UniformBuffer);
 	}
 };
@@ -111,13 +111,6 @@ public:
 		: FGlobalShader(Initializer)
 	{
 	}
-
-	/** Serialization. */
-	virtual bool Serialize(FArchive& Ar) override
-	{
-		bool bShaderHasOutdatedParameters = FGlobalShader::Serialize(Ar);
-		return bShaderHasOutdatedParameters;
-	}
 };
 
 /** Implementation for all shaders used for particle injection. */
@@ -213,8 +206,8 @@ static void InjectCurves(
 		TShaderMapRef<FParticleCurveInjectionPS> PixelShader(GetGlobalShaderMap(GMaxRHIFeatureLevel));
 
 		GraphicsPSOInit.BoundShaderState.VertexDeclarationRHI = GParticleCurveInjectionVertexDeclaration.VertexDeclarationRHI;
-		GraphicsPSOInit.BoundShaderState.VertexShaderRHI = GETSAFERHISHADER_VERTEX(*VertexShader);
-		GraphicsPSOInit.BoundShaderState.PixelShaderRHI = GETSAFERHISHADER_PIXEL(*PixelShader);
+		GraphicsPSOInit.BoundShaderState.VertexShaderRHI = VertexShader.GetVertexShader();
+		GraphicsPSOInit.BoundShaderState.PixelShaderRHI = PixelShader.GetPixelShader();
 		GraphicsPSOInit.PrimitiveType = PT_TriangleList;
 
 		SetGraphicsPipelineState(RHICmdList, GraphicsPSOInit);
