Copying Private-LoadTimes-4.24 stream to Main. Biggest changes are in Materials/Shader memory freezing.

diff --git a/Engine/Source/Runtime/Engine/Private/Particles/ParticleGpuSimulation.cpp b/Engine/Source/Runtime/Engine/Private/Particles/ParticleGpuSimulation.cpp
--- a/Engine/Source/Runtime/Engine/Private/Particles/ParticleGpuSimulation.cpp
+++ b/Engine/Source/Runtime/Engine/Private/Particles/ParticleGpuSimulation.cpp
@@ -533,8 +533,9 @@ typedef TUniformBufferRef<FGPUSpriteEmitterDynamicUniformParameters> FGPUSpriteE
  */
 class FGPUSpriteVertexFactoryShaderParametersVS : public FVertexFactoryShaderParameters
 {
+	DECLARE_INLINE_TYPE_LAYOUT(FGPUSpriteVertexFactoryShaderParametersVS, NonVirtual);
 public:
-	virtual void Bind( const FShaderParameterMap& ParameterMap ) override
+	void Bind( const FShaderParameterMap& ParameterMap )
 	{
 		ParticleIndices.Bind(ParameterMap, TEXT("ParticleIndices"));
 		ParticleIndicesOffset.Bind(ParameterMap, TEXT("ParticleIndicesOffset"));
@@ -548,21 +549,7 @@ public:
 		CurveTextureSampler.Bind(ParameterMap, TEXT("CurveTextureSampler"));
 	}
 
-	virtual void Serialize(FArchive& Ar) override
-	{
-		Ar << ParticleIndices;
-		Ar << ParticleIndicesOffset;
-		Ar << PositionTexture;
-		Ar << PositionTextureSampler;
-		Ar << VelocityTexture;
-		Ar << VelocityTextureSampler;
-		Ar << AttributesTexture;
-		Ar << AttributesTextureSampler;
-		Ar << CurveTexture;
-		Ar << CurveTextureSampler;
-	}
-
-	virtual void GetElementShaderBindings(
+	void GetElementShaderBindings(
 		const FSceneInterface* Scene,
 		const FSceneView* View,
 		const FMeshMaterialShader* Shader,
@@ -571,28 +558,27 @@ public:
 		const FVertexFactory* VertexFactory,
 		const FMeshBatchElement& BatchElement,
 		class FMeshDrawSingleShaderBindings& ShaderBindings,
-		FVertexInputStreamArray& VertexStreams) const override;
-
-	virtual uint32 GetSize() const override { return sizeof(*this); }
-
+		FVertexInputStreamArray& VertexStreams) const;
 private:
 
-	/** Buffer containing particle indices. */
-	FShaderResourceParameter ParticleIndices;
-	/** Offset in to the particle indices buffer. */
-	FShaderParameter ParticleIndicesOffset;
-	/** Texture containing positions for all particles. */
-	FShaderResourceParameter PositionTexture;
-	FShaderResourceParameter PositionTextureSampler;
-	/** Texture containing velocities for all particles. */
-	FShaderResourceParameter VelocityTexture;
-	FShaderResourceParameter VelocityTextureSampler;
-	/** Texture containint attributes for all particles. */
-	FShaderResourceParameter AttributesTexture;
-	FShaderResourceParameter AttributesTextureSampler;
-	/** Texture containing curves from which attributes are sampled. */
-	FShaderResourceParameter CurveTexture;
-	FShaderResourceParameter CurveTextureSampler;
+	
+		/** Buffer containing particle indices. */
+		LAYOUT_FIELD(FShaderResourceParameter, ParticleIndices)
+		/** Offset in to the particle indices buffer. */
+		LAYOUT_FIELD(FShaderParameter, ParticleIndicesOffset)
+		/** Texture containing positions for all particles. */
+		LAYOUT_FIELD(FShaderResourceParameter, PositionTexture)
+		LAYOUT_FIELD(FShaderResourceParameter, PositionTextureSampler)
+		/** Texture containing velocities for all particles. */
+		LAYOUT_FIELD(FShaderResourceParameter, VelocityTexture)
+		LAYOUT_FIELD(FShaderResourceParameter, VelocityTextureSampler)
+		/** Texture containint attributes for all particles. */
+		LAYOUT_FIELD(FShaderResourceParameter, AttributesTexture)
+		LAYOUT_FIELD(FShaderResourceParameter, AttributesTextureSampler)
+		/** Texture containing curves from which attributes are sampled. */
+		LAYOUT_FIELD(FShaderResourceParameter, CurveTexture)
+		LAYOUT_FIELD(FShaderResourceParameter, CurveTextureSampler)
+	
 };
 
 /**
@@ -600,12 +586,9 @@ private:
  */
 class FGPUSpriteVertexFactoryShaderParametersPS : public FVertexFactoryShaderParameters
 {
+	DECLARE_INLINE_TYPE_LAYOUT(FGPUSpriteVertexFactoryShaderParametersPS, NonVirtual);
 public:
-	virtual void Bind( const FShaderParameterMap& ParameterMap ) override {}
-
-	virtual void Serialize(FArchive& Ar) override {}
-
-	virtual void GetElementShaderBindings(
+	void GetElementShaderBindings(
 		const FSceneInterface* Scene,
 		const FSceneView* View,
 		const FMeshMaterialShader* Shader,
@@ -614,11 +597,10 @@ public:
 		const FVertexFactory* VertexFactory,
 		const FMeshBatchElement& BatchElement,
 		class FMeshDrawSingleShaderBindings& ShaderBindings,
-		FVertexInputStreamArray& VertexStreams) const override;
-
-	virtual uint32 GetSize() const override { return sizeof(*this); }
+		FVertexInputStreamArray& VertexStreams) const;
 
-private:
+	
+	
 };
 
 /**
@@ -733,43 +715,27 @@ public:
 	/**
 	 * Should we cache the material's shadertype on this platform with this vertex factory? 
 	 */
-	static bool ShouldCompilePermutation(EShaderPlatform Platform, const class FMaterial* Material, const class FShaderType* ShaderType)
+	static bool ShouldCompilePermutation(const FVertexFactoryShaderPermutationParameters& Parameters)
 	{
-		return (Material->IsUsedWithParticleSprites() || Material->IsSpecialEngineMaterial()) && SupportsGPUParticles(Platform);
+		return (Parameters.MaterialParameters.bIsUsedWithParticleSprites || Parameters.MaterialParameters.bIsSpecialEngineMaterial) && SupportsGPUParticles(Parameters.Platform);
 	}
 
 	/**
 	 * Can be overridden by FVertexFactory subclasses to modify their compile environment just before compilation occurs.
 	 */
-	static void ModifyCompilationEnvironment(const FVertexFactoryType* Type, EShaderPlatform Platform, const FMaterial* Material, FShaderCompilerEnvironment& OutEnvironment)
+	static void ModifyCompilationEnvironment(const FVertexFactoryShaderPermutationParameters& Parameters, FShaderCompilerEnvironment& OutEnvironment)
 	{
-		FParticleVertexFactoryBase::ModifyCompilationEnvironment(Type, Platform, Material, OutEnvironment);
+		FParticleVertexFactoryBase::ModifyCompilationEnvironment(Parameters, OutEnvironment);
 		OutEnvironment.SetDefine(TEXT("PARTICLES_PER_INSTANCE"), MAX_PARTICLES_PER_INSTANCE);
 
 		// Set a define so we can tell in MaterialTemplate.usf when we are compiling a sprite vertex factory
 		OutEnvironment.SetDefine(TEXT("PARTICLE_SPRITE_FACTORY"),TEXT("1"));
 
-		if (Platform == SP_OPENGL_ES2_ANDROID)
+		if (Parameters.Platform == SP_OPENGL_ES2_ANDROID)
 		{
 			OutEnvironment.CompilerFlags.Add(CFLAG_FeatureLevelES31);
 		}
 	}
-
-	/**
-	 * Construct shader parameters for this type of vertex factory.
-	 */
-	static FVertexFactoryShaderParameters* ConstructShaderParameters(EShaderFrequency ShaderFrequency)
-	{
-		if (ShaderFrequency == SF_Vertex)
-		{
-			return new FGPUSpriteVertexFactoryShaderParametersVS();
-		}
-		else if (ShaderFrequency == SF_Pixel)
-		{
-			return new FGPUSpriteVertexFactoryShaderParametersPS();
-		}
-		return NULL;
-	}
 };
 
 void FGPUSpriteVertexFactoryShaderParametersVS::GetElementShaderBindings(
@@ -822,6 +788,8 @@ void FGPUSpriteVertexFactoryShaderParametersPS::GetElementShaderBindings(
 	ShaderBindings.Add(Shader->GetUniformBufferParameter<FGPUSpriteEmitterDynamicUniformParameters>(), GPUVF->EmitterDynamicUniformBuffer);
 }
 
+IMPLEMENT_VERTEX_FACTORY_PARAMETER_TYPE(FGPUSpriteVertexFactory, SF_Vertex, FGPUSpriteVertexFactoryShaderParametersVS);
+IMPLEMENT_VERTEX_FACTORY_PARAMETER_TYPE(FGPUSpriteVertexFactory, SF_Pixel, FGPUSpriteVertexFactoryShaderParametersPS);
 IMPLEMENT_VERTEX_FACTORY_TYPE(FGPUSpriteVertexFactory,"/Engine/Private/ParticleGPUSpriteVertexFactory.ush",true,false,true,false,false);
 
 /*-----------------------------------------------------------------------------
@@ -913,12 +881,8 @@ struct FParticlePerFrameSimulationParameters
  */
 struct FParticlePerFrameSimulationShaderParameters
 {
-	FShaderParameter PointAttractor;
-	FShaderParameter PositionOffsetAndAttractorStrength;
-	FShaderParameter LocalToWorldScale;
-	FShaderParameter DeltaSeconds;
-	FShaderParameter NumIterations;
-
+	DECLARE_INLINE_TYPE_LAYOUT(FParticlePerFrameSimulationShaderParameters, NonVirtual);
+public:
 	void Bind(const FShaderParameterMap& ParameterMap)
 	{
 		PointAttractor.Bind(ParameterMap,TEXT("PointAttractor"));
@@ -942,6 +906,13 @@ struct FParticlePerFrameSimulationShaderParameters
 		SetShaderValue(RHICmdList,ShaderRHI,DeltaSeconds, bUseFixDT ? Parameters.DeltaSecondsInFix : Parameters.DeltaSecondsInVar);
 		SetShaderValue(RHICmdList,ShaderRHI,NumIterations, bUseFixDT ? Parameters.NumIterationsInFix : Parameters.NumIterationsInVar);
 	}
+
+	LAYOUT_FIELD(FShaderParameter, PointAttractor);
+	LAYOUT_FIELD(FShaderParameter, PositionOffsetAndAttractorStrength);
+	LAYOUT_FIELD(FShaderParameter, LocalToWorldScale);
+	LAYOUT_FIELD(FShaderParameter, DeltaSeconds);
+	LAYOUT_FIELD(FShaderParameter, NumIterations);
+	
 };
 
 FArchive& operator<<(FArchive& Ar, FParticlePerFrameSimulationShaderParameters& PerFrameParameters)
@@ -1009,21 +980,10 @@ public:
 		TileSizeX.Bind(Initializer.ParameterMap, TEXT("TileSizeX"));
 		TileSizeY.Bind(Initializer.ParameterMap, TEXT("TileSizeY"));
 	}
-
-	/** Serialization. */
-	virtual bool Serialize( FArchive& Ar ) override
-	{
-		bool bShaderHasOutdatedParameters = FGlobalShader::Serialize( Ar );
-		Ar << TileOffsets;
-		Ar << TileSizeX;
-		Ar << TileSizeY;
-		return bShaderHasOutdatedParameters;
-	}
-
 	/** Set parameters. */
 	void SetParameters(FRHICommandList& RHICmdList, FParticleShaderParamRef TileOffsetsRef)
 	{
-		FRHIVertexShader* VertexShaderRHI = GetVertexShader();
+		FRHIVertexShader* VertexShaderRHI = RHICmdList.GetBoundVertexShader();
 		if (TileOffsets.IsBound())
 		{
 			RHICmdList.SetShaderResourceViewParameter(VertexShaderRHI, TileOffsets.GetBaseIndex(), TileOffsetsRef);
@@ -1033,11 +993,10 @@ public:
 	}
 
 private:
-
 	/** Buffer from which to read tile offsets. */
-	FShaderResourceParameter TileOffsets;
-	FShaderParameter TileSizeX;
-	FShaderParameter TileSizeY;
+	LAYOUT_FIELD(FShaderResourceParameter, TileOffsets);
+	LAYOUT_FIELD(FShaderParameter, TileSizeX);
+	LAYOUT_FIELD(FShaderParameter, TileSizeY);
 };
 
 /**
@@ -1099,31 +1058,6 @@ public:
 		GlobalDistanceFieldParameters.Bind(Initializer.ParameterMap);
 	}
 
-	/** Serialization. */
-	virtual bool Serialize(FArchive& Ar) override
-	{
-		bool bShaderHasOutdatedParameters = FGlobalShader::Serialize( Ar );
-		Ar << PositionTexture;
-		Ar << PositionTextureSampler;
-		Ar << VelocityTexture;
-		Ar << VelocityTextureSampler;
-		Ar << AttributesTexture;
-		Ar << AttributesTextureSampler;
-		Ar << RenderAttributesTexture;
-		Ar << RenderAttributesTextureSampler;
-		Ar << CurveTexture;
-		Ar << CurveTextureSampler;
-		for (int32 i = 0; i < MAX_VECTOR_FIELDS; i++)
-		{
-			Ar << VectorFieldTextures[i];
-			Ar << VectorFieldTexturesSamplers[i];
-		}
-		Ar << CollisionDepthBounds;
-		Ar << PerFrameParameters;
-		Ar << GlobalDistanceFieldParameters;
-		return bShaderHasOutdatedParameters;
-	}
-
 	/**
 	 * Set parameters for this shader.
 	 */
@@ -1139,7 +1073,7 @@ public:
 		FRHIUniformBuffer* SceneTexturesUniformBuffer
 		)
 	{
-		FRHIPixelShader* PixelShaderRHI = GetPixelShader();
+		FRHIPixelShader* PixelShaderRHI = RHICmdList.GetBoundPixelShader();
 		FRHISamplerState* SamplerStatePoint = TStaticSamplerState<SF_Point>::GetRHI();
 		FRHISamplerState* SamplerStateLinear = TStaticSamplerState<SF_Bilinear,AM_Clamp,AM_Clamp,AM_Clamp>::GetRHI();
 		SetTextureParameter(RHICmdList, PixelShaderRHI, PositionTexture, PositionTextureSampler, SamplerStatePoint, TextureResources.PositionTextureRHI);
@@ -1184,7 +1118,7 @@ public:
 	 */
 	void SetVectorFieldParameters(FRHICommandList& RHICmdList, const FVectorFieldUniformBufferRef& UniformBuffer, FRHITexture3D* const* VolumeTexturesRHI)
 	{
-		FRHIPixelShader* PixelShaderRHI = GetPixelShader();
+		FRHIPixelShader* PixelShaderRHI = RHICmdList.GetBoundPixelShader();
 			SetUniformBufferParameter(RHICmdList, PixelShaderRHI, GetUniformBufferParameter<FVectorFieldUniformParameters>(), UniformBuffer);
 		
 			FRHISamplerState* SamplerStateLinear = TStaticSamplerState<SF_Bilinear,AM_Clamp,AM_Clamp,AM_Clamp>::GetRHI();
@@ -1201,7 +1135,7 @@ public:
 	 */
 	void SetInstanceParameters(FRHICommandList& RHICmdList, FRHIUniformBuffer* UniformBuffer, const FParticlePerFrameSimulationParameters& InPerFrameParameters, bool bUseFixDT)
 	{
-		FRHIPixelShader* PixelShaderRHI = GetPixelShader();
+		FRHIPixelShader* PixelShaderRHI = RHICmdList.GetBoundPixelShader();
 		SetUniformBufferParameter(RHICmdList, PixelShaderRHI, GetUniformBufferParameter<FParticleSimulationParameters>(), UniformBuffer);
 		PerFrameParameters.Set(RHICmdList, PixelShaderRHI, InPerFrameParameters, bUseFixDT);
 	}
@@ -1211,7 +1145,7 @@ public:
 	 */
 	void UnbindBuffers(FRHICommandList& RHICmdList)
 	{
-		FRHIPixelShader* PixelShaderRHI = GetPixelShader();
+		FRHIPixelShader* PixelShaderRHI = RHICmdList.GetBoundPixelShader();
 		FRHIShaderResourceView* NullSRV = nullptr;
 		for (int32 i = 0; i < MAX_VECTOR_FIELDS; ++i)
 		{
@@ -1223,30 +1157,29 @@ public:
 	}
 
 private:
-
 	/** The position texture parameter. */
-	FShaderResourceParameter PositionTexture;
-	FShaderResourceParameter PositionTextureSampler;
+	LAYOUT_FIELD(FShaderResourceParameter, PositionTexture);
+	LAYOUT_FIELD(FShaderResourceParameter, PositionTextureSampler);
 	/** The velocity texture parameter. */
-	FShaderResourceParameter VelocityTexture;
-	FShaderResourceParameter VelocityTextureSampler;
+	LAYOUT_FIELD(FShaderResourceParameter, VelocityTexture);
+	LAYOUT_FIELD(FShaderResourceParameter, VelocityTextureSampler);
 	/** The simulation attributes texture parameter. */
-	FShaderResourceParameter AttributesTexture;
-	FShaderResourceParameter AttributesTextureSampler;
+	LAYOUT_FIELD(FShaderResourceParameter, AttributesTexture);
+	LAYOUT_FIELD(FShaderResourceParameter, AttributesTextureSampler);
 	/** The render attributes texture parameter. */
-	FShaderResourceParameter RenderAttributesTexture;
-	FShaderResourceParameter RenderAttributesTextureSampler;
+	LAYOUT_FIELD(FShaderResourceParameter, RenderAttributesTexture);
+	LAYOUT_FIELD(FShaderResourceParameter, RenderAttributesTextureSampler);
 	/** The curve texture parameter. */
-	FShaderResourceParameter CurveTexture;
-	FShaderResourceParameter CurveTextureSampler;
+	LAYOUT_FIELD(FShaderResourceParameter, CurveTexture);
+	LAYOUT_FIELD(FShaderResourceParameter, CurveTextureSampler);
 	/** Vector fields. */
-	FShaderResourceParameter VectorFieldTextures[MAX_VECTOR_FIELDS];
-	FShaderResourceParameter VectorFieldTexturesSamplers[MAX_VECTOR_FIELDS];
+	LAYOUT_ARRAY(FShaderResourceParameter, VectorFieldTextures, MAX_VECTOR_FIELDS);
+	LAYOUT_ARRAY(FShaderResourceParameter, VectorFieldTexturesSamplers, MAX_VECTOR_FIELDS);
 	/** Per frame simulation parameters. */
-	FParticlePerFrameSimulationShaderParameters PerFrameParameters;
+	LAYOUT_FIELD(FParticlePerFrameSimulationShaderParameters, PerFrameParameters);
 	/** Collision depth bounds. */
-	FShaderParameter CollisionDepthBounds;
-	FGlobalDistanceFieldParameters GlobalDistanceFieldParameters;
+	LAYOUT_FIELD(FShaderParameter, CollisionDepthBounds);
+	LAYOUT_FIELD(FGlobalDistanceFieldParameters, GlobalDistanceFieldParameters);
 };
 
 /**
@@ -1285,13 +1218,6 @@ public:
 		: FGlobalShader(Initializer)
 	{
 	}
-
-	/** Serialization. */
-	virtual bool Serialize( FArchive& Ar ) override
-	{
-		bool bShaderHasOutdatedParameters = FGlobalShader::Serialize( Ar );
-		return bShaderHasOutdatedParameters;
-	}
 };
 
 /** Implementation for all shaders used for simulation. */
@@ -1533,8 +1459,8 @@ void ExecuteSimulationCommands(
 	TShaderMapRef<TParticleSimulationPS<CollisionMode> > PixelShader(GetGlobalShaderMap(FeatureLevel));
 
 	GraphicsPSOInit.BoundShaderState.VertexDeclarationRHI = GetParticleTileVertexDeclaration(FeatureLevel);
-	GraphicsPSOInit.BoundShaderState.VertexShaderRHI = GETSAFERHISHADER_VERTEX(*VertexShader);
-	GraphicsPSOInit.BoundShaderState.PixelShaderRHI = GETSAFERHISHADER_PIXEL(*PixelShader);
+	GraphicsPSOInit.BoundShaderState.VertexShaderRHI = VertexShader.GetVertexShader();
+	GraphicsPSOInit.BoundShaderState.PixelShaderRHI = PixelShader.GetPixelShader();
 	GraphicsPSOInit.PrimitiveType = PT_TriangleList;
 
 	SetGraphicsPipelineState(RHICmdList, GraphicsPSOInit);
@@ -1651,8 +1577,8 @@ void ClearTiles(FRHICommandList& RHICmdList, FGraphicsPipelineStateInitializer&
 	TShaderMapRef<FParticleSimulationClearPS> PixelShader(GetGlobalShaderMap(FeatureLevel));
 	
 	GraphicsPSOInit.BoundShaderState.VertexDeclarationRHI = GetParticleTileVertexDeclaration(FeatureLevel);
-	GraphicsPSOInit.BoundShaderState.VertexShaderRHI = GETSAFERHISHADER_VERTEX(*VertexShader);
-	GraphicsPSOInit.BoundShaderState.PixelShaderRHI = GETSAFERHISHADER_PIXEL(*PixelShader);
+	GraphicsPSOInit.BoundShaderState.VertexShaderRHI = VertexShader.GetVertexShader();
+	GraphicsPSOInit.BoundShaderState.PixelShaderRHI = PixelShader.GetPixelShader();
 	GraphicsPSOInit.PrimitiveType = PT_TriangleList;
 
 	SetGraphicsPipelineState(RHICmdList, GraphicsPSOInit);
@@ -1742,7 +1668,7 @@ public:
 		Parameters.PixelScale.X = 1.0f / GParticleSimulationTextureSizeX;
 		Parameters.PixelScale.Y = 1.0f / GParticleSimulationTextureSizeY;
 		FParticleInjectionBufferRef UniformBuffer = FParticleInjectionBufferRef::CreateUniformBufferImmediate( Parameters, UniformBuffer_SingleDraw );
-		FRHIVertexShader* VertexShader = GetVertexShader();
+		FRHIVertexShader* VertexShader = RHICmdList.GetBoundVertexShader();
 		SetUniformBufferParameter(RHICmdList, VertexShader, GetUniformBufferParameter<FParticleInjectionParameters>(), UniformBuffer );
 	}
 };
@@ -1788,13 +1714,6 @@ public:
 		: FGlobalShader(Initializer)
 	{
 	}
-
-	/** Serialization. */
-	virtual bool Serialize( FArchive& Ar ) override
-	{
-		bool bShaderHasOutdatedParameters = FGlobalShader::Serialize( Ar );
-		return bShaderHasOutdatedParameters;
-	}
 };
 
 /** Implementation for all shaders used for particle injection. */
@@ -1892,8 +1811,8 @@ void InjectNewParticles(FRHICommandList& RHICmdList, FGraphicsPipelineStateIniti
 		TShaderMapRef<TParticleInjectionPS<StaticPropertiesOnly> > PixelShader(GetGlobalShaderMap(FeatureLevel));
 
 		GraphicsPSOInit.BoundShaderState.VertexDeclarationRHI = GParticleInjectionVertexDeclaration.VertexDeclarationRHI;
-		GraphicsPSOInit.BoundShaderState.VertexShaderRHI = GETSAFERHISHADER_VERTEX(*VertexShader);
-		GraphicsPSOInit.BoundShaderState.PixelShaderRHI = GETSAFERHISHADER_PIXEL(*PixelShader);
+		GraphicsPSOInit.BoundShaderState.VertexShaderRHI = VertexShader.GetVertexShader();
+		GraphicsPSOInit.BoundShaderState.PixelShaderRHI = PixelShader.GetPixelShader();
 		GraphicsPSOInit.PrimitiveType = PT_TriangleList;
 
 		SetGraphicsPipelineState(RHICmdList, GraphicsPSOInit);
@@ -1973,7 +1892,7 @@ public:
 	 */
 	void SetParameters(FRHICommandList& RHICmdList, const FParticleSimVisualizeBufferRef& UniformBuffer )
 	{
-		FRHIVertexShader* VertexShader = GetVertexShader();
+		FRHIVertexShader* VertexShader = RHICmdList.GetBoundVertexShader();
 		SetUniformBufferParameter(RHICmdList, VertexShader, GetUniformBufferParameter<FParticleSimVisualizeParameters>(), UniformBuffer );
 	}
 };
@@ -2008,24 +1927,12 @@ public:
 		CurveTextureSampler.Bind( Initializer.ParameterMap, TEXT("CurveTextureSampler") );
 	}
 
-	/** Serialization. */
-	virtual bool Serialize( FArchive& Ar ) override
-	{
-		bool bShaderHasOutdatedParameters = FGlobalShader::Serialize( Ar );
-		Ar << VisualizationMode;
-		Ar << PositionTexture;
-		Ar << PositionTextureSampler;
-		Ar << CurveTexture;
-		Ar << CurveTextureSampler;
-		return bShaderHasOutdatedParameters;
-	}
-
 	/**
 	 * Set parameters for this shader.
 	 */
 	void SetParameters(FRHICommandList& RHICmdList, int32 InVisualizationMode, FRHITexture2D* PositionTextureRHI, FRHITexture2D* CurveTextureRHI )
 	{
-		FRHIPixelShader* PixelShader = GetPixelShader();
+		FRHIPixelShader* PixelShader = RHICmdList.GetBoundPixelShader();
 		SetShaderValue(RHICmdList, PixelShader, VisualizationMode, InVisualizationMode );
 		FRHISamplerState* SamplerStatePoint = TStaticSamplerState<SF_Point>::GetRHI();
 		SetTextureParameter(RHICmdList, PixelShader, PositionTexture, PositionTextureSampler, SamplerStatePoint, PositionTextureRHI );
@@ -2034,11 +1941,11 @@ public:
 
 private:
 
-	FShaderParameter VisualizationMode;
-	FShaderResourceParameter PositionTexture;
-	FShaderResourceParameter PositionTextureSampler;
-	FShaderResourceParameter CurveTexture;
-	FShaderResourceParameter CurveTextureSampler;
+	LAYOUT_FIELD(FShaderParameter, VisualizationMode);
+	LAYOUT_FIELD(FShaderResourceParameter, PositionTexture);
+	LAYOUT_FIELD(FShaderResourceParameter, PositionTextureSampler);
+	LAYOUT_FIELD(FShaderResourceParameter, CurveTexture);
+	LAYOUT_FIELD(FShaderResourceParameter, CurveTextureSampler);
 };
 
 /** Implementation for all shaders used for visualization. */
@@ -2113,8 +2020,8 @@ static void VisualizeGPUSimulation(
 		TShaderMapRef<FParticleSimVisualizePS> PixelShader(GetGlobalShaderMap(FeatureLevel));
 
 		GraphicsPSOInit.BoundShaderState.VertexDeclarationRHI = GParticleSimVisualizeVertexDeclaration.VertexDeclarationRHI;
-		GraphicsPSOInit.BoundShaderState.VertexShaderRHI = GETSAFERHISHADER_VERTEX(*VertexShader);
-		GraphicsPSOInit.BoundShaderState.PixelShaderRHI = GETSAFERHISHADER_PIXEL(*PixelShader);
+		GraphicsPSOInit.BoundShaderState.VertexShaderRHI = VertexShader.GetVertexShader();
+		GraphicsPSOInit.BoundShaderState.PixelShaderRHI = PixelShader.GetPixelShader();
 		GraphicsPSOInit.PrimitiveType = PT_TriangleList;
 
 		SetGraphicsPipelineState(RHICmdList, GraphicsPSOInit);
@@ -2249,25 +2156,12 @@ public:
 		TextureSizeY.Bind(Initializer.ParameterMap, TEXT("TextureSizeY"));
 	}
 
-	/** Serialization. */
-	virtual bool Serialize( FArchive& Ar ) override
-	{
-		bool bShaderHasOutdatedParameters = FGlobalShader::Serialize( Ar );
-		Ar << InParticleIndices;
-		Ar << PositionTexture;
-		Ar << PositionTextureSampler;
-		Ar << OutBounds;
-		Ar << TextureSizeX;
-		Ar << TextureSizeY;
-		return bShaderHasOutdatedParameters;
-	}
-
 	/**
 	 * Set output buffers for this shader.
 	 */
 	void SetOutput(FRHICommandList& RHICmdList, FRHIUnorderedAccessView* OutBoundsUAV )
 	{
-		FRHIComputeShader* ComputeShaderRHI = GetComputeShader();
+		FRHIComputeShader* ComputeShaderRHI = RHICmdList.GetBoundComputeShader();
 		if ( OutBounds.IsBound() )
 		{
 			RHICmdList.SetUAVParameter(ComputeShaderRHI, OutBounds.GetBaseIndex(), OutBoundsUAV);
@@ -2284,7 +2178,7 @@ public:
 		FRHITexture2D* PositionTextureRHI
 		)
 	{
-		FRHIComputeShader* ComputeShaderRHI = GetComputeShader();
+		FRHIComputeShader* ComputeShaderRHI = RHICmdList.GetBoundComputeShader();
 		SetUniformBufferParameter(RHICmdList, ComputeShaderRHI, GetUniformBufferParameter<FParticleBoundsParameters>(), UniformBuffer );
 		if ( InParticleIndices.IsBound() )
 		{
@@ -2304,7 +2198,7 @@ public:
 	 */
 	void UnbindBuffers(FRHICommandList& RHICmdList)
 	{
-		FRHIComputeShader* ComputeShaderRHI = GetComputeShader();
+		FRHIComputeShader* ComputeShaderRHI = RHICmdList.GetBoundComputeShader();
 		if ( InParticleIndices.IsBound() )
 		{
 			RHICmdList.SetShaderResourceViewParameter(ComputeShaderRHI, InParticleIndices.GetBaseIndex(), nullptr);
@@ -2316,16 +2210,15 @@ public:
 	}
 
 private:
-
 	/** Input buffer containing particle indices. */
-	FShaderResourceParameter InParticleIndices;
+	LAYOUT_FIELD(FShaderResourceParameter, InParticleIndices);
 	/** Texture containing particle positions. */
-	FShaderResourceParameter PositionTexture;
-	FShaderResourceParameter PositionTextureSampler;
+	LAYOUT_FIELD(FShaderResourceParameter, PositionTexture);
+	LAYOUT_FIELD(FShaderResourceParameter, PositionTextureSampler);
 	/** Output key buffer. */
-	FShaderResourceParameter OutBounds;
-	FShaderParameter TextureSizeX;
-	FShaderParameter TextureSizeY;
+	LAYOUT_FIELD(FShaderResourceParameter, OutBounds);
+	LAYOUT_FIELD(FShaderParameter, TextureSizeX);
+	LAYOUT_FIELD(FShaderParameter, TextureSizeY);
 };
 IMPLEMENT_SHADER_TYPE(,FParticleBoundsCS,TEXT("/Engine/Private/ParticleBoundsShader.usf"),TEXT("ComputeParticleBounds"),SF_Compute);
 
@@ -2381,14 +2274,14 @@ static FBox ComputeParticleBounds(
 
 		// Grab the shader.
 		TShaderMapRef<FParticleBoundsCS> ParticleBoundsCS(GetGlobalShaderMap(GMaxRHIFeatureLevel));
-		RHICmdList.SetComputeShader(ParticleBoundsCS->GetComputeShader());
+		RHICmdList.SetComputeShader(ParticleBoundsCS.GetComputeShader());
 
 		// Dispatch shader to compute bounds.
 		ParticleBoundsCS->SetOutput(RHICmdList, BoundsVertexBufferUAV);
 		ParticleBoundsCS->SetParameters(RHICmdList, UniformBuffer, VertexBufferSRV, PositionTextureRHI);
 		DispatchComputeShader(
 			RHICmdList, 
-			*ParticleBoundsCS, 
+			ParticleBoundsCS.GetShader(), 
 			GroupCount,
 			1,
 			1 );
diff --git a/Engine/Source/Runtime/Engine/Private/Particles/ParticleModules.cpp b/Engine/Source/Runtime/Engine/Private/Particles/ParticleModules.cpp
--- a/Engine/Source/Runtime/Engine/Private/Particles/ParticleModules.cpp
+++ b/Engine/Source/Runtime/Engine/Private/Particles/ParticleModules.cpp
@@ -3362,7 +3362,7 @@ int32 UParticleModuleTypeDataMesh::GetCurrentDetailMode()
 
 int32 UParticleModuleTypeDataMesh::GetMeshParticleMotionBlurMinDetailMode()
 {
-	return CVarMinDetailModeForMeshParticleMotionBlur.GetValueOnGameThread();
+	return CVarMinDetailModeForMeshParticleMotionBlur.GetValueOnAnyThread();
 }
 
 UParticleModuleTypeDataMesh::UParticleModuleTypeDataMesh(const FObjectInitializer& ObjectInitializer)
diff --git a/Engine/Source/Runtime/Engine/Private/Particles/ParticleModules_Parameter.cpp b/Engine/Source/Runtime/Engine/Private/Particles/ParticleModules_Parameter.cpp
--- a/Engine/Source/Runtime/Engine/Private/Particles/ParticleModules_Parameter.cpp
+++ b/Engine/Source/Runtime/Engine/Private/Particles/ParticleModules_Parameter.cpp
@@ -427,26 +427,6 @@ UMaterial* UParticleModuleParameterDynamic_RetrieveMaterial(UMaterialInterface*
 	return Material;
 }
 
-/**
- *	Helper function to find the DynamicParameter expression in a material
- */
-UMaterialExpressionDynamicParameter* UParticleModuleParameterDynamic_GetDynamicParameterExpression(UMaterial* InMaterial)
-{
-	UMaterialExpressionDynamicParameter* DynParamExp = NULL;
-	for (int32 ExpIndex = 0; ExpIndex < InMaterial->Expressions.Num(); ExpIndex++)
-	{
-		DynParamExp = Cast<UMaterialExpressionDynamicParameter>(InMaterial->Expressions[ExpIndex]);
-
-		if (DynParamExp != NULL)
-		{
-			break;
-		}
-	}
-
-	return DynParamExp;
-}
-
-
 void UParticleModuleParameterDynamic::UpdateParameterNames(UMaterialInterface* InMaterialInterface)
 {
 	UMaterial* Material = UParticleModuleParameterDynamic_RetrieveMaterial(InMaterialInterface);
@@ -455,16 +435,9 @@ void UParticleModuleParameterDynamic::UpdateParameterNames(UMaterialInterface* I
 		return;
 	}
 
-	// Check the expressions...
-	UMaterialExpressionDynamicParameter* DynParamExp = UParticleModuleParameterDynamic_GetDynamicParameterExpression(Material);
-	if (DynParamExp == NULL)
-	{
-		return;
-	}
-
 	for (int32 ParamIndex = 0; ParamIndex < 4; ParamIndex++)
 	{
-		DynamicParams[ParamIndex].ParamName = FName(*(DynParamExp->ParamNames[ParamIndex]));
+		DynamicParams[ParamIndex].ParamName = Material->GetCachedExpressionData().DynamicParameterNames[ParamIndex];
 	}
 }
 
diff --git a/Engine/Source/Runtime/Engine/Private/Particles/ParticleSortingGPU.cpp b/Engine/Source/Runtime/Engine/Private/Particles/ParticleSortingGPU.cpp
--- a/Engine/Source/Runtime/Engine/Private/Particles/ParticleSortingGPU.cpp
+++ b/Engine/Source/Runtime/Engine/Private/Particles/ParticleSortingGPU.cpp
@@ -77,24 +77,12 @@ public:
 		OutParticleIndices.Bind( Initializer.ParameterMap, TEXT("OutParticleIndices") );
 	}
 
-	/** Serialization. */
-	virtual bool Serialize( FArchive& Ar ) override
-	{
-		bool bShaderHasOutdatedParameters = FGlobalShader::Serialize( Ar );
-		Ar << InParticleIndices;
-		Ar << PositionTexture;
-		Ar << PositionTextureSampler;
-		Ar << OutKeys;
-		Ar << OutParticleIndices;
-		return bShaderHasOutdatedParameters;
-	}
-
 	/**
 	 * Set output buffers for this shader.
 	 */
 	void SetOutput(FRHICommandList& RHICmdList, FRHIUnorderedAccessView* OutKeysUAV, FRHIUnorderedAccessView* OutIndicesUAV )
 	{
-		FRHIComputeShader* ComputeShaderRHI = GetComputeShader();
+		FRHIComputeShader* ComputeShaderRHI = RHICmdList.GetBoundComputeShader();
 		if ( OutKeys.IsBound() )
 		{
 			RHICmdList.SetUAVParameter(ComputeShaderRHI, OutKeys.GetBaseIndex(), OutKeysUAV);
@@ -114,7 +102,7 @@ public:
 		FRHIShaderResourceView* InIndicesSRV
 		)
 	{
-		FRHIComputeShader* ComputeShaderRHI = GetComputeShader();
+		FRHIComputeShader* ComputeShaderRHI = RHICmdList.GetBoundComputeShader();
 		SetUniformBufferParameter(RHICmdList, ComputeShaderRHI, GetUniformBufferParameter<FParticleKeyGenParameters>(), UniformBuffer );
 		if ( InParticleIndices.IsBound() )
 		{
@@ -127,7 +115,7 @@ public:
 	 */
 	void SetPositionTextures(FRHICommandList& RHICmdList, FRHITexture2D* PositionTextureRHI)
 	{
-		FRHIComputeShader* ComputeShaderRHI = GetComputeShader();
+		FRHIComputeShader* ComputeShaderRHI = RHICmdList.GetBoundComputeShader();
 		if (PositionTexture.IsBound())
 		{
 			RHICmdList.SetShaderTexture(ComputeShaderRHI, PositionTexture.GetBaseIndex(), PositionTextureRHI);
@@ -139,7 +127,7 @@ public:
 	 */
 	void UnbindBuffers(FRHICommandList& RHICmdList)
 	{
-		FRHIComputeShader* ComputeShaderRHI = GetComputeShader();
+		FRHIComputeShader* ComputeShaderRHI = RHICmdList.GetBoundComputeShader();
 		if ( InParticleIndices.IsBound() )
 		{
 			RHICmdList.SetShaderResourceViewParameter(ComputeShaderRHI, InParticleIndices.GetBaseIndex(), nullptr);
@@ -155,16 +143,15 @@ public:
 	}
 
 private:
-
 	/** Input buffer containing particle indices. */
-	FShaderResourceParameter InParticleIndices;
+	LAYOUT_FIELD(FShaderResourceParameter, InParticleIndices);
 	/** Texture containing particle positions. */
-	FShaderResourceParameter PositionTexture;
-	FShaderResourceParameter PositionTextureSampler;
+	LAYOUT_FIELD(FShaderResourceParameter, PositionTexture);
+	LAYOUT_FIELD(FShaderResourceParameter, PositionTextureSampler);
 	/** Output key buffer. */
-	FShaderResourceParameter OutKeys;
+	LAYOUT_FIELD(FShaderResourceParameter, OutKeys);
 	/** Output indices buffer. */
-	FShaderResourceParameter OutParticleIndices;
+	LAYOUT_FIELD(FShaderResourceParameter, OutParticleIndices);
 };
 IMPLEMENT_SHADER_TYPE(,FParticleSortKeyGenCS,TEXT("/Engine/Private/ParticleSortKeyGen.usf"),TEXT("GenerateParticleSortKeys"),SF_Compute);
 
@@ -195,7 +182,7 @@ int32 GenerateParticleSortKeys(
 
 	// Grab the shader, set output.
 	TShaderMapRef<FParticleSortKeyGenCS> KeyGenCS(GetGlobalShaderMap(FeatureLevel));
-	RHICmdList.SetComputeShader(KeyGenCS->GetComputeShader());
+	RHICmdList.SetComputeShader(KeyGenCS.GetComputeShader());
 	KeyGenCS->SetOutput(RHICmdList, KeyBufferUAV, SortedVertexBufferUAV);
 	KeyGenCS->SetPositionTextures(RHICmdList, PositionTextureRHI);
 
@@ -220,7 +207,7 @@ int32 GenerateParticleSortKeys(
 
 			// Dispatch.
 			KeyGenCS->SetParameters(RHICmdList, KeyGenUniformBuffer, SortInfo.VertexBufferSRV);
-			DispatchComputeShader(RHICmdList, *KeyGenCS, GroupCount, 1, 1);
+			DispatchComputeShader(RHICmdList, KeyGenCS.GetShader(), GroupCount, 1, 1);
 
 			// TR-KeyGen : No sync needed between tasks since they update different parts of the data (assuming it's ok if cache line overlap).
 			RHICmdList.TransitionResources(EResourceTransitionAccess::ERWNoBarrier, EResourceTransitionPipeline::EComputeToCompute, OutputUAVs, UE_ARRAY_COUNT(OutputUAVs));
diff --git a/Engine/Source/Runtime/Engine/Private/Particles/ParticleVertexFactory.cpp b/Engine/Source/Runtime/Engine/Private/Particles/ParticleVertexFactory.cpp
--- a/Engine/Source/Runtime/Engine/Private/Particles/ParticleVertexFactory.cpp
+++ b/Engine/Source/Runtime/Engine/Private/Particles/ParticleVertexFactory.cpp
@@ -50,34 +50,23 @@ TGlobalResource<FNullSubUVCutoutVertexBuffer> GFNullSubUVCutoutVertexBuffer;
  */
 class FParticleSpriteVertexFactoryShaderParameters : public FVertexFactoryShaderParameters
 {
+	DECLARE_INLINE_TYPE_LAYOUT(FParticleSpriteVertexFactoryShaderParameters, NonVirtual);
 public:
-
-	virtual void Bind(const FShaderParameterMap& ParameterMap) override
-	{
-	}
-
-	virtual void Serialize(FArchive& Ar) override
-	{
-	}
+	
+	
 };
 
 class FParticleSpriteVertexFactoryShaderParametersVS : public FParticleSpriteVertexFactoryShaderParameters
 {
+	DECLARE_INLINE_TYPE_LAYOUT(FParticleSpriteVertexFactoryShaderParametersVS, NonVirtual);
 public:
-
-	virtual void Bind(const FShaderParameterMap& ParameterMap) override
+	void Bind(const FShaderParameterMap& ParameterMap)
 	{
 		NumCutoutVerticesPerFrame.Bind(ParameterMap, TEXT("NumCutoutVerticesPerFrame"));
 		CutoutGeometry.Bind(ParameterMap, TEXT("CutoutGeometry"));
 	}
 
-	virtual void Serialize(FArchive& Ar) override
-	{
-		Ar << NumCutoutVerticesPerFrame;
-		Ar << CutoutGeometry;
-	}
-
-	virtual void GetElementShaderBindings(
+	void GetElementShaderBindings(
 		const FSceneInterface* Scene,
 		const FSceneView* View,
 		const FMeshMaterialShader* Shader,
@@ -86,7 +75,7 @@ public:
 		const FVertexFactory* VertexFactory,
 		const FMeshBatchElement& BatchElement,
 		class FMeshDrawSingleShaderBindings& ShaderBindings,
-		FVertexInputStreamArray& VertexStreams) const override
+		FVertexInputStreamArray& VertexStreams) const
 	{
 		FParticleSpriteVertexFactory* SpriteVF = (FParticleSpriteVertexFactory*)VertexFactory;
 		ShaderBindings.Add(Shader->GetUniformBufferParameter<FParticleSpriteUniformParameters>(), SpriteVF->GetSpriteUniformBuffer() );
@@ -97,15 +86,17 @@ public:
 	}
 
 private:
-	FShaderParameter NumCutoutVerticesPerFrame;
-	FShaderResourceParameter CutoutGeometry;
+	
+		LAYOUT_FIELD(FShaderParameter, NumCutoutVerticesPerFrame)
+		LAYOUT_FIELD(FShaderResourceParameter, CutoutGeometry)
+	
 };
 
 class FParticleSpriteVertexFactoryShaderParametersPS : public FParticleSpriteVertexFactoryShaderParameters
 {
+	DECLARE_INLINE_TYPE_LAYOUT(FParticleSpriteVertexFactoryShaderParametersPS, NonVirtual);
 public:
-
-	virtual void GetElementShaderBindings(
+	void GetElementShaderBindings(
 		const FSceneInterface* Scene,
 		const FSceneView* View,
 		const FMeshMaterialShader* Shader,
@@ -114,11 +105,14 @@ public:
 		const FVertexFactory* VertexFactory,
 		const FMeshBatchElement& BatchElement,
 		class FMeshDrawSingleShaderBindings& ShaderBindings,
-		FVertexInputStreamArray& VertexStreams) const override
+		FVertexInputStreamArray& VertexStreams) const
 	{
 		FParticleSpriteVertexFactory* SpriteVF = (FParticleSpriteVertexFactory*)VertexFactory;
 		ShaderBindings.Add(Shader->GetUniformBufferParameter<FParticleSpriteUniformParameters>(), SpriteVF->GetSpriteUniformBuffer() );
 	}
+
+	
+	
 };
 
 /**
@@ -225,17 +219,17 @@ static inline TGlobalResource<FParticleSpriteVertexDeclaration>& GetParticleSpri
 	}
 }
 
-bool FParticleSpriteVertexFactory::ShouldCompilePermutation(EShaderPlatform Platform, const class FMaterial* Material, const class FShaderType* ShaderType)
+bool FParticleSpriteVertexFactory::ShouldCompilePermutation(const FVertexFactoryShaderPermutationParameters& Parameters)
 {
-	return Material->IsUsedWithParticleSprites() || Material->IsSpecialEngineMaterial();
+	return Parameters.MaterialParameters.bIsUsedWithParticleSprites || Parameters.MaterialParameters.bIsSpecialEngineMaterial;
 }
 
 /**
  * Can be overridden by FVertexFactory subclasses to modify their compile environment just before compilation occurs.
  */
-void FParticleSpriteVertexFactory::ModifyCompilationEnvironment(const FVertexFactoryType* Type, EShaderPlatform Platform, const class FMaterial* Material, FShaderCompilerEnvironment& OutEnvironment)
+void FParticleSpriteVertexFactory::ModifyCompilationEnvironment(const FVertexFactoryShaderPermutationParameters& Parameters, FShaderCompilerEnvironment& OutEnvironment)
 {
-	FParticleVertexFactoryBase::ModifyCompilationEnvironment(Type, Platform, Material, OutEnvironment);
+	FParticleVertexFactoryBase::ModifyCompilationEnvironment(Parameters, OutEnvironment);
 
 	// Set a define so we can tell in MaterialTemplate.usf when we are compiling a sprite vertex factory
 	OutEnvironment.SetDefine(TEXT("PARTICLE_SPRITE_FACTORY"),TEXT("1"));
@@ -297,17 +291,6 @@ void FParticleSpriteVertexFactory::SetDynamicParameterBuffer(const FVertexBuffer
 	}
 }
 
-FVertexFactoryShaderParameters* FParticleSpriteVertexFactory::ConstructShaderParameters(EShaderFrequency ShaderFrequency)
-{
-	if (ShaderFrequency == SF_Vertex)
-	{
-		return new FParticleSpriteVertexFactoryShaderParametersVS();
-	}
-	else if (ShaderFrequency == SF_Pixel)
-	{
-		return new FParticleSpriteVertexFactoryShaderParametersPS();
-	}
-	return NULL;
-}
-
+IMPLEMENT_VERTEX_FACTORY_PARAMETER_TYPE(FParticleSpriteVertexFactory, SF_Vertex, FParticleSpriteVertexFactoryShaderParametersVS);
+IMPLEMENT_VERTEX_FACTORY_PARAMETER_TYPE(FParticleSpriteVertexFactory, SF_Pixel, FParticleSpriteVertexFactoryShaderParametersPS);
 IMPLEMENT_VERTEX_FACTORY_TYPE(FParticleSpriteVertexFactory,"/Engine/Private/ParticleSpriteVertexFactory.ush",true,false,true,false,false);
diff --git a/Engine/Source/Runtime/Engine/Private/PerPlatformProperties.cpp b/Engine/Source/Runtime/Engine/Private/PerPlatformProperties.cpp
--- a/Engine/Source/Runtime/Engine/Private/PerPlatformProperties.cpp
+++ b/Engine/Source/Runtime/Engine/Private/PerPlatformProperties.cpp
@@ -8,6 +8,9 @@
 #include "PlatformInfo.h"
 #endif
 
+IMPLEMENT_TYPE_LAYOUT(FFreezablePerPlatformFloat);
+IMPLEMENT_TYPE_LAYOUT(FFreezablePerPlatformInt);
+
 /** Serializer to cook out the most appropriate platform override */
 template<typename StructType, typename ValueType, EName BasePropertyName>
 ENGINE_API FArchive& operator<<(FArchive& Ar, TPerPlatformProperty<StructType, ValueType, BasePropertyName>& Property)
@@ -74,9 +77,11 @@ ENGINE_API void operator<<(FStructuredArchive::FSlot Slot, TPerPlatformProperty<
 template ENGINE_API FArchive& operator<<(FArchive&, TPerPlatformProperty<FPerPlatformInt, int32, NAME_IntProperty>&);
 template ENGINE_API FArchive& operator<<(FArchive&, TPerPlatformProperty<FPerPlatformFloat, float, NAME_FloatProperty>&);
 template ENGINE_API FArchive& operator<<(FArchive&, TPerPlatformProperty<FPerPlatformBool, bool, NAME_BoolProperty>&);
+template ENGINE_API FArchive& operator<<(FArchive&, TPerPlatformProperty<FFreezablePerPlatformFloat, float, NAME_FloatProperty>&);
 template ENGINE_API void operator<<(FStructuredArchive::FSlot Slot, TPerPlatformProperty<FPerPlatformInt, int32, NAME_IntProperty>&);
 template ENGINE_API void operator<<(FStructuredArchive::FSlot Slot, TPerPlatformProperty<FPerPlatformFloat, float, NAME_FloatProperty>&);
 template ENGINE_API void operator<<(FStructuredArchive::FSlot Slot, TPerPlatformProperty<FPerPlatformBool, bool, NAME_BoolProperty>&);
+template ENGINE_API void operator<<(FStructuredArchive::FSlot Slot, TPerPlatformProperty<FFreezablePerPlatformFloat, float, NAME_FloatProperty>&);
 
 FString FPerPlatformInt::ToString() const
 {
@@ -94,4 +99,9 @@ FString FPerPlatformInt::ToString() const
 #endif
 
 	return Result;
-}
\ No newline at end of file
+}
+
+FString FFreezablePerPlatformInt::ToString() const
+{
+	return FPerPlatformInt(*this).ToString();
+}
diff --git a/Engine/Source/Runtime/Engine/Private/PhysicsEngine/BodyInstance.cpp b/Engine/Source/Runtime/Engine/Private/PhysicsEngine/BodyInstance.cpp
--- a/Engine/Source/Runtime/Engine/Private/PhysicsEngine/BodyInstance.cpp
+++ b/Engine/Source/Runtime/Engine/Private/PhysicsEngine/BodyInstance.cpp
@@ -60,6 +60,8 @@
 DECLARE_CYCLE_STAT(TEXT("Init Body"), STAT_InitBody, STATGROUP_Physics);
 DECLARE_CYCLE_STAT(TEXT("Init Body Debug"), STAT_InitBodyDebug, STATGROUP_Physics);
 DECLARE_CYCLE_STAT(TEXT("Init Body Scene Interaction"), STAT_InitBodySceneInteraction, STATGROUP_Physics);
+DECLARE_CYCLE_STAT(TEXT("Init Body Aggregate"), STAT_InitBodyAggregate, STATGROUP_Physics);
+DECLARE_CYCLE_STAT(TEXT("Init Body Add"), STAT_InitBodyAdd, STATGROUP_Physics);
 DECLARE_CYCLE_STAT(TEXT("Init Body Post Add to Scene"), STAT_InitBodyPostAdd, STATGROUP_Physics);
 DECLARE_CYCLE_STAT(TEXT("Term Body"), STAT_TermBody, STATGROUP_Physics);
 DECLARE_CYCLE_STAT(TEXT("Update Materials"), STAT_UpdatePhysMats, STATGROUP_Physics);
@@ -71,6 +73,8 @@ DECLARE_CYCLE_STAT(TEXT("Bulk Body Scene Add"), STAT_BulkSceneAdd, STATGROUP_Phy
 DECLARE_CYCLE_STAT(TEXT("Static Init Bodies"), STAT_StaticInitBodies, STATGROUP_Physics);
 DECLARE_CYCLE_STAT(TEXT("UpdateBodyScale"), STAT_BodyInstanceUpdateBodyScale, STATGROUP_Physics);
 DECLARE_CYCLE_STAT(TEXT("CreatePhysicsShapesAndActors"), STAT_CreatePhysicsShapesAndActors, STATGROUP_Physics);
+DECLARE_CYCLE_STAT(TEXT("CreatePhysicsShapes"), STAT_CreatePhysicsShapes, STATGROUP_Physics);
+DECLARE_CYCLE_STAT(TEXT("CreatePhysicsActor"), STAT_CreatePhysicsActor, STATGROUP_Physics);
 DECLARE_CYCLE_STAT(TEXT("BodyInstance SetCollisionProfileName"), STAT_BodyInst_SetCollisionProfileName, STATGROUP_Physics);
 DECLARE_CYCLE_STAT(TEXT("Phys SetBodyTransform"), STAT_SetBodyTransform, STATGROUP_Physics);
 
@@ -82,6 +86,13 @@ TAutoConsoleVariable<int32> CVarEnableDynamicPerBodyFilterHacks(
 	ECVF_ReadOnly
 );
 
+TAutoConsoleVariable<int32> CVarEnableDeferredPhysicsCreation(
+	TEXT("p.EnableDeferredPhysicsCreation"), 
+	0, 
+	TEXT("Enables/Disables deferred physics creation."),
+	ECVF_Default
+);
+
 using namespace PhysicsInterfaceTypes;
 
 bool IsRigidBodyKinematic_AssumesLocked(const FPhysicsActorHandle& InActorRef)
@@ -976,224 +987,188 @@ void GetSimulatingAndBlendWeight(const USkeletalMeshComponent* SkelMeshComp, con
 	}
 }
 
-template <bool bCompileStatic>
-struct FInitBodiesHelper
-{
-	FInitBodiesHelper(TArray<FBodyInstance*>& InBodies, TArray<FTransform>& InTransforms, class UBodySetup* InBodySetup, class UPrimitiveComponent* InPrimitiveComp, FPhysScene* InRBScene, const FBodyInstance::FInitBodySpawnParams& InSpawnParams, FPhysicsAggregateHandle InAggregate)
+FInitBodiesHelperBase::FInitBodiesHelperBase(TArray<FBodyInstance*>& InBodies, TArray<FTransform>& InTransforms, class UBodySetup* InBodySetup, class UPrimitiveComponent* InPrimitiveComp, FPhysScene* InRBScene, const FInitBodySpawnParams& InSpawnParams, FPhysicsAggregateHandle InAggregate)
 	: Bodies(InBodies)
 	, Transforms(InTransforms)
 	, BodySetup(InBodySetup)
 	, PrimitiveComp(InPrimitiveComp)
-		, PhysScene(InRBScene)
-		, Aggregate(InAggregate)
-		, bInstanceSimulatePhysics(false)
+	, PhysScene(InRBScene)
+	, Aggregate(InAggregate)
+	, bInstanceSimulatePhysics(false)
 	, InstanceBlendWeight(-1.f)
 	, SkelMeshComp(nullptr)
 	, SpawnParams(InSpawnParams)
-	{
-		//Compute all the needed constants
-
+	, DisableQueryOnlyActors(!!CDisableQueryOnlyActors.GetValueOnGameThread())
+{
 #if USE_BODYINSTANCE_DEBUG_NAMES
-		PhysXName = GetDebugDebugName(PrimitiveComp, BodySetup, DebugName);
+	PhysXName = GetDebugDebugName(PrimitiveComp, BodySetup, DebugName);
 #endif
+}
 
-		bStatic = bCompileStatic || SpawnParams.bStaticPhysics;
-		SkelMeshComp = bCompileStatic ? nullptr : Cast<USkeletalMeshComponent>(PrimitiveComp);
-		if(SpawnParams.bPhysicsTypeDeterminesSimulation)
-		{
-			GetSimulatingAndBlendWeight(SkelMeshComp, BodySetup, InstanceBlendWeight, bInstanceSimulatePhysics);
-		}
-		
-	}
-
-	FORCEINLINE bool IsStatic() const { return bCompileStatic || bStatic; }
-
-	//The arguments passed into InitBodies
-	TArray<FBodyInstance*>& Bodies;   
-	TArray<FTransform>& Transforms;
-	class UBodySetup* BodySetup;
-	class UPrimitiveComponent* PrimitiveComp;
-	FPhysScene* PhysScene;
-	FPhysicsAggregateHandle Aggregate;
+// Return to actor ref
+void FInitBodiesHelperBase::CreateActor_AssumesLocked(FBodyInstance* Instance, const FTransform& Transform) const
+{
+	SCOPE_CYCLE_COUNTER(STAT_CreatePhysicsActor);
+	checkSlow(!FPhysicsInterface::IsValid(Instance->ActorHandle));
+	const ECollisionEnabled::Type CollisionType = Instance->GetCollisionEnabled();
+	const bool bDisableSim = !CollisionEnabledHasPhysics(CollisionType) && DisableQueryOnlyActors;
 
+	FActorCreationParams ActorParams;
+	ActorParams.InitialTM = Transform;
+	ActorParams.BodyInstance = Instance;
 #if USE_BODYINSTANCE_DEBUG_NAMES
-	FString DebugName;
-	TSharedPtr<TArray<ANSICHAR>> PhysXName; // Get rid of ANSICHAR in physics
+	ActorParams.DebugName = Instance->CharDebugName.IsValid() ? Instance->CharDebugName->GetData() : nullptr;
 #endif
+	ActorParams.bEnableGravity = Instance->bEnableGravity;
+	ActorParams.bQueryOnly = bDisableSim;
+	ActorParams.Scene = PhysScene;
 
-	//The constants shared between PhysX and Box2D
-	bool bStatic;
-	bool bInstanceSimulatePhysics;
-	float InstanceBlendWeight;
-
-	const USkeletalMeshComponent* SkelMeshComp;
-
-	const FBodyInstance::FInitBodySpawnParams& SpawnParams;
-
-	// Return to actor ref
-	void CreateActor_AssumesLocked(FBodyInstance* Instance, const FTransform& Transform) const
+	if (IsStatic())
 	{
-		checkSlow(!FPhysicsInterface::IsValid(Instance->ActorHandle));
-		const ECollisionEnabled::Type CollisionType = Instance->GetCollisionEnabled();
-		const bool bDisableSim = !CollisionEnabledHasPhysics(CollisionType) && CDisableQueryOnlyActors.GetValueOnGameThread();
-
-		FActorCreationParams ActorParams;
-		ActorParams.InitialTM = Transform;
-		ActorParams.BodyInstance = Instance;
-#if USE_BODYINSTANCE_DEBUG_NAMES
-		ActorParams.DebugName = Instance->CharDebugName.IsValid() ? Instance->CharDebugName->GetData() : nullptr;
-#endif
-		ActorParams.bEnableGravity = Instance->bEnableGravity;
-		ActorParams.bQueryOnly = bDisableSim;
-		ActorParams.Scene = PhysScene;
-
-		if (IsStatic())
-		{
-			ActorParams.bStatic = true;
+		ActorParams.bStatic = true;
 
-			FPhysicsInterface::CreateActor(ActorParams, Instance->ActorHandle);
-		}
-		else
-		{
-			FPhysicsInterface::CreateActor(ActorParams, Instance->ActorHandle);
+		FPhysicsInterface::CreateActor(ActorParams, Instance->ActorHandle);
+	}
+	else
+	{
+		FPhysicsInterface::CreateActor(ActorParams, Instance->ActorHandle);
 
-			FPhysicsInterface::SetCcdEnabled_AssumesLocked(Instance->ActorHandle, Instance->bUseCCD);
-			FPhysicsInterface::SetIsKinematic_AssumesLocked(Instance->ActorHandle, !Instance->ShouldInstanceSimulatingPhysics());
+		FPhysicsInterface::SetCcdEnabled_AssumesLocked(Instance->ActorHandle, Instance->bUseCCD);
+		FPhysicsInterface::SetIsKinematic_AssumesLocked(Instance->ActorHandle, !Instance->ShouldInstanceSimulatingPhysics());
 
-			// Set sleep even notification
-			FPhysicsInterface::SetSendsSleepNotifies_AssumesLocked(Instance->ActorHandle, Instance->bGenerateWakeEvents);
-		}
+		// Set sleep even notification
+		FPhysicsInterface::SetSendsSleepNotifies_AssumesLocked(Instance->ActorHandle, Instance->bGenerateWakeEvents);
 	}
+}
 
-	bool CreateShapes_AssumesLocked(FBodyInstance* Instance) const
-	{
-		UPhysicalMaterial* SimplePhysMat = Instance->GetSimplePhysicalMaterial();
-
-		TArray<UPhysicalMaterial*> ComplexPhysMats;
-		TArray<FPhysicalMaterialMaskParams> ComplexPhysMatMasks;
+bool FInitBodiesHelperBase::CreateShapes_AssumesLocked(FBodyInstance* Instance) const
+{
+	SCOPE_CYCLE_COUNTER(STAT_CreatePhysicsShapes);
+	UPhysicalMaterial* SimplePhysMat = Instance->GetSimplePhysicalMaterial();
+	TArray<UPhysicalMaterial*> ComplexPhysMats;
+	TArray<FPhysicalMaterialMaskParams> ComplexPhysMatMasks;
 
 #if WITH_CHAOS
-		ComplexPhysMats = Instance->GetComplexPhysicalMaterials(ComplexPhysMatMasks);
+	ComplexPhysMats = Instance->GetComplexPhysicalMaterials(ComplexPhysMatMasks);
 #else
-		ComplexPhysMats = Instance->GetComplexPhysicalMaterials();
+	ComplexPhysMats = Instance->GetComplexPhysicalMaterials();
 #endif
 
-		FBodyCollisionData BodyCollisionData;
-		Instance->BuildBodyFilterData(BodyCollisionData.CollisionFilterData);
-		Instance->BuildBodyCollisionFlags(BodyCollisionData.CollisionFlags, Instance->GetCollisionEnabled(), BodySetup->GetCollisionTraceFlag() == CTF_UseComplexAsSimple);
+	FBodyCollisionData BodyCollisionData;
+	Instance->BuildBodyFilterData(BodyCollisionData.CollisionFilterData);
+	Instance->BuildBodyCollisionFlags(BodyCollisionData.CollisionFlags, Instance->GetCollisionEnabled(), BodySetup->GetCollisionTraceFlag() == CTF_UseComplexAsSimple);
 
-		bool bInitFail = false;
+	bool bInitFail = false;
 
-		// #PHYS2 Call interface AddGeometry
-		BodySetup->AddShapesToRigidActor_AssumesLocked(Instance, Instance->Scale3D, SimplePhysMat, ComplexPhysMats, ComplexPhysMatMasks, BodyCollisionData, FTransform::Identity);
+	// #PHYS2 Call interface AddGeometry
+	BodySetup->AddShapesToRigidActor_AssumesLocked(Instance, Instance->Scale3D, SimplePhysMat, ComplexPhysMats, ComplexPhysMatMasks, BodyCollisionData, FTransform::Identity);
 
-		const int32 NumShapes = FPhysicsInterface::GetNumShapes(Instance->ActorHandle);
-		bInitFail |= NumShapes == 0;
+	const int32 NumShapes = FPhysicsInterface::GetNumShapes(Instance->ActorHandle);
+	bInitFail |= NumShapes == 0;
 
 
-		return bInitFail;
-	}
+	return bInitFail;
+}
 
-	// Takes actor ref arrays.
-	// #PHYS2 this used to return arrays of low-level physics bodies, which would be added to scene in InitBodies. Should it still do that, rather then later iterate over BodyInstances to get phys actor refs?
-	bool CreateShapesAndActors()
-	{
-		SCOPE_CYCLE_COUNTER(STAT_CreatePhysicsShapesAndActors);
+// Takes actor ref arrays.
+// #PHYS2 this used to return arrays of low-level physics bodies, which would be added to scene in InitBodies. Should it still do that, rather then later iterate over BodyInstances to get phys actor refs?
+bool FInitBodiesHelperBase::CreateShapesAndActors()
+{
+	SCOPE_CYCLE_COUNTER(STAT_CreatePhysicsShapesAndActors);
 
-		const int32 NumBodies = Bodies.Num();
+	const int32 NumBodies = Bodies.Num();
 
-		// Ensure we have the AggGeom inside the body setup so we can calculate the number of shapes
-		BodySetup->CreatePhysicsMeshes();
+	// Ensure we have the AggGeom inside the body setup so we can calculate the number of shapes
+	BodySetup->CreatePhysicsMeshes();
 
-		for (int32 BodyIdx = NumBodies - 1; BodyIdx >= 0; BodyIdx--)   // iterate in reverse since list might shrink
-		{
-			FBodyInstance* Instance = Bodies[BodyIdx];
-			const FTransform& Transform = Transforms[BodyIdx];
+	for (int32 BodyIdx = NumBodies - 1; BodyIdx >= 0; BodyIdx--)   // iterate in reverse since list might shrink
+	{
+		FBodyInstance* Instance = Bodies[BodyIdx];
+		const FTransform& Transform = Transforms[BodyIdx];
 
 #if !USE_BODYINSTANCE_DEBUG_NAMES
-			FString DebugName;
+		FString DebugName;
 #endif
-			FBodyInstance::ValidateTransform(Transform, DebugName, BodySetup);
+		FBodyInstance::ValidateTransform(Transform, DebugName, BodySetup);
 
-			Instance->OwnerComponent = PrimitiveComp;
-			Instance->BodySetup = BodySetup;
-			Instance->Scale3D = Transform.GetScale3D();
+		Instance->OwnerComponent = PrimitiveComp;
+		Instance->BodySetup = BodySetup;
+		Instance->Scale3D = Transform.GetScale3D();
 #if USE_BODYINSTANCE_DEBUG_NAMES
-			Instance->CharDebugName = PhysXName;
+		Instance->CharDebugName = PhysXName;
 #endif
-			Instance->bEnableGravity = Instance->bEnableGravity && (SkelMeshComp ? SkelMeshComp->BodyInstance.bEnableGravity : true);	//In the case of skeletal mesh component we AND bodies with the parent body
+		Instance->bEnableGravity = Instance->bEnableGravity && (SkelMeshComp ? SkelMeshComp->BodyInstance.bEnableGravity : true);	//In the case of skeletal mesh component we AND bodies with the parent body
 
-			// Handle autowelding here to avoid extra work
-			if (!IsStatic() && Instance->bAutoWeld)
+		// Handle autowelding here to avoid extra work
+		if (!IsStatic() && Instance->bAutoWeld)
+		{
+			ECollisionEnabled::Type CollisionType = Instance->GetCollisionEnabled();
+			if (CollisionType != ECollisionEnabled::QueryOnly)
 			{
-				ECollisionEnabled::Type CollisionType = Instance->GetCollisionEnabled();
-				if (CollisionType != ECollisionEnabled::QueryOnly)
+				if (UPrimitiveComponent * ParentPrimComponent = PrimitiveComp ? Cast<UPrimitiveComponent>(PrimitiveComp->GetAttachParent()) : NULL)
 				{
-					if (UPrimitiveComponent * ParentPrimComponent = PrimitiveComp ? Cast<UPrimitiveComponent>(PrimitiveComp->GetAttachParent()) : NULL)
+					UWorld* World = PrimitiveComp->GetWorld();
+					if (World && World->IsGameWorld())
 					{
-						UWorld* World = PrimitiveComp->GetWorld();
-						if (World && World->IsGameWorld())
+						//if we have a parent we will now do the weld and exit any further initialization
+						if (PrimitiveComp->WeldToImplementation(ParentPrimComponent, PrimitiveComp->GetAttachSocketName(), false))	//welded new simulated body so initialization is done
 						{
-							//if we have a parent we will now do the weld and exit any further initialization
-							if (PrimitiveComp->WeldToImplementation(ParentPrimComponent, PrimitiveComp->GetAttachSocketName(), false))	//welded new simulated body so initialization is done
-							{
-								return false;
-							}
+							return false;
 						}
 					}
 				}
 			}
+		}
 
-			// Don't process if we've already got a body
-			// Just ask actorref
-			if (FPhysicsInterface::IsValid(Instance->GetPhysicsActorHandle()))
-			{
-				Instance->OwnerComponent = nullptr;
-				Instance->BodySetup      = nullptr;
-				Bodies.RemoveAt(BodyIdx);  // so we wont add it to the physx scene again later.
-				Transforms.RemoveAt(BodyIdx);
-				continue;
-			}
+		// Don't process if we've already got a body
+		// Just ask actorref
+		if (FPhysicsInterface::IsValid(Instance->GetPhysicsActorHandle()))
+		{
+			Instance->OwnerComponent = nullptr;
+			Instance->BodySetup      = nullptr;
+			Bodies.RemoveAt(BodyIdx);  // so we wont add it to the physx scene again later.
+			Transforms.RemoveAt(BodyIdx);
+			continue;
+		}
 
-			// Set sim parameters for bodies from skeletal mesh components
-			if (!IsStatic() && SpawnParams.bPhysicsTypeDeterminesSimulation)
+		// Set sim parameters for bodies from skeletal mesh components
+		if (!IsStatic() && SpawnParams.bPhysicsTypeDeterminesSimulation)
+		{
+			Instance->bSimulatePhysics = bInstanceSimulatePhysics;
+			if (InstanceBlendWeight != -1.0f)
 			{
-				Instance->bSimulatePhysics = bInstanceSimulatePhysics;
-				if (InstanceBlendWeight != -1.0f)
-				{
-					Instance->PhysicsBlendWeight = InstanceBlendWeight;
-				}
+				Instance->PhysicsBlendWeight = InstanceBlendWeight;
 			}
+		}
 
-			// Init user data structure to point back at this instance
-			Instance->PhysicsUserData = FPhysicsUserData(Instance);
+		// Init user data structure to point back at this instance
+		Instance->PhysicsUserData = FPhysicsUserData(Instance);
 
-			CreateActor_AssumesLocked(Instance, Transform);
-			const bool bInitFail = CreateShapes_AssumesLocked(Instance);
-			if (bInitFail)
-			{
+		CreateActor_AssumesLocked(Instance, Transform);
+		const bool bInitFail = CreateShapes_AssumesLocked(Instance);
+		if (bInitFail)
+		{
 #if WITH_EDITOR
-				//In the editor we may have ended up here because of world trace ignoring our EnableCollision. Since we can't get at the data in that function we check for it here
-				if(!PrimitiveComp || PrimitiveComp->IsCollisionEnabled())
+			//In the editor we may have ended up here because of world trace ignoring our EnableCollision. Since we can't get at the data in that function we check for it here
+			if(!PrimitiveComp || PrimitiveComp->IsCollisionEnabled())
 #endif
-				{
-					UE_LOG(LogPhysics, Log, TEXT("Init Instance %d of Primitive Component %s failed. Does it have collision data available?"), BodyIdx, *PrimitiveComp->GetReadableName());
-				}
-
-				FPhysicsInterface::ReleaseActor(Instance->ActorHandle, PhysScene);
+			{
+				UE_LOG(LogPhysics, Log, TEXT("Init Instance %d of Primitive Component %s failed. Does it have collision data available?"), BodyIdx, *PrimitiveComp->GetReadableName());
+			}
 
-				Instance->OwnerComponent = nullptr;
-				Instance->BodySetup = nullptr;
-				Instance->ExternalCollisionProfileBodySetup = nullptr;
+			FPhysicsInterface::ReleaseActor(Instance->ActorHandle, PhysScene);
 
-				continue;
-			}
+			Instance->OwnerComponent = nullptr;
+			Instance->BodySetup = nullptr;
+			Instance->ExternalCollisionProfileBodySetup = nullptr;
 
-			FPhysicsInterface::SetActorUserData_AssumesLocked(Instance->ActorHandle, &Instance->PhysicsUserData);
+			continue;
 		}
 
-		return true;
+		FPhysicsInterface::SetActorUserData_AssumesLocked(Instance->ActorHandle, &Instance->PhysicsUserData);
 	}
+	return true;
+}
 
 #if WITH_CHAOS
 	Chaos::EChaosCollisionTraceFlag ConvertCollisionTraceFlag(ECollisionTraceFlag Flag)
@@ -1213,43 +1188,46 @@ struct FInitBodiesHelper
 	}
 #endif
 
-	void InitBodies()
-	{
+
+void FInitBodiesHelperBase::InitBodies()
+{
 #if WITH_CHAOS
-		LLM_SCOPE(ELLMTag::Chaos);
+	LLM_SCOPE(ELLMTag::Chaos);
 #else
-		LLM_SCOPE(ELLMTag::PhysX);
+	LLM_SCOPE(ELLMTag::PhysX);
 #endif
 
-		check(IsInGameThread());
+	//check(IsInGameThread());
 
-		if (CreateShapesAndActors())
+	if (CreateShapesAndActors())
+	{
+		FPhysicsCommand::ExecuteWrite(PhysScene, [&]()
 		{
-			FPhysicsCommand::ExecuteWrite(PhysScene, [&]()
+			// If an aggregate present, add to that
+			if (Aggregate.IsValid())
 			{
-				// If an aggregate present, add to that
-				if (Aggregate.IsValid())
+				SCOPE_CYCLE_COUNTER(STAT_InitBodyAggregate);
+				for (FBodyInstance* BI : Bodies)
 				{
-					for (FBodyInstance* BI : Bodies)
+					const FPhysicsActorHandle& ActorHandle = BI->GetPhysicsActorHandle();
+					if (FPhysicsInterface::IsValid(ActorHandle))
 					{
-						const FPhysicsActorHandle& ActorHandle = BI->GetPhysicsActorHandle();
-						if (FPhysicsInterface::IsValid(ActorHandle))
-						{
-							FPhysicsInterface::AddActorToAggregate_AssumesLocked(Aggregate, ActorHandle);
-						}
+						FPhysicsInterface::AddActorToAggregate_AssumesLocked(Aggregate, ActorHandle);
 					}
 				}
-				else if (PhysScene)
-				{
-					TArray<FPhysicsActorHandle> ActorHandles;
-					ActorHandles.Reserve(Bodies.Num());
+			}
+			else if (PhysScene)
+			{
+				SCOPE_CYCLE_COUNTER(STAT_InitBodyAdd);
+				TArray<FPhysicsActorHandle> ActorHandles;
+				ActorHandles.Reserve(Bodies.Num());
 
-					for (FBodyInstance* BI : Bodies)
+				for (FBodyInstance* BI : Bodies)
+				{
+					FPhysicsActorHandle& ActorHandle = BI->GetPhysicsActorHandle();
+					if (FPhysicsInterface::IsValid(ActorHandle))
 					{
-						FPhysicsActorHandle& ActorHandle = BI->GetPhysicsActorHandle();
-						if (FPhysicsInterface::IsValid(ActorHandle))
-						{
-							ActorHandles.Add(ActorHandle);
+						ActorHandles.Add(ActorHandle);
 
 #if WITH_CHAOS
 							const int32 NumShapes = FPhysicsInterface::GetNumShapes(ActorHandle);
@@ -1273,57 +1251,56 @@ struct FInitBodiesHelper
 
 #endif
 /*
-							With the implementation of AddActorsToScene_AssumesLocked,
-							this call duplicates the AddToSolver operation and is not necessary
+						With the implementation of AddActorsToScene_AssumesLocked,
+						this call duplicates the AddToSolver operation and is not necessary
 #ifdef WITH_CHAOS
-							if (BI->bSimulatePhysics)
-							{
-								PhysScene->AddActorToSolver(ActorHandle);
-							}
+						if (BI->bSimulatePhysics)
+						{
+							PhysScene->AddActorToSolver(ActorHandle);
+						}
 #endif
 */
-						}
 					}
+				}
 
-					PhysScene->AddActorsToScene_AssumesLocked(ActorHandles);
+				PhysScene->AddActorsToScene_AssumesLocked(ActorHandles);
 #if WITH_CHAOS
-					for (FBodyInstance* BI : Bodies)
+				for (FBodyInstance* BI : Bodies)
+				{
+					FPhysicsActorHandle& ActorHandle = BI->GetPhysicsActorHandle();
+					if (FPhysicsInterface::IsValid(ActorHandle))
 					{
-						FPhysicsActorHandle& ActorHandle = BI->GetPhysicsActorHandle();
-						if (FPhysicsInterface::IsValid(ActorHandle))
-						{
 
-							PhysScene->AddToComponentMaps(BI->OwnerComponent.Get(), ActorHandle->Proxy);
-						}
-						if (BI->bNotifyRigidBodyCollision)
+						PhysScene->AddToComponentMaps(BI->OwnerComponent.Get(), ActorHandle->Proxy);
+					}
+					if (BI->bNotifyRigidBodyCollision)
+					{
+						if (UPrimitiveComponent* PrimComp = BI->OwnerComponent.Get())
 						{
-							if (UPrimitiveComponent* PrimComp = BI->OwnerComponent.Get())
-							{
-								FPhysScene_ChaosInterface* LocalPhysScene = PrimComp->GetWorld()->GetPhysicsScene();
-								FPhysScene_Chaos& Scene = LocalPhysScene->GetScene();
-								Scene.RegisterForCollisionEvents(PrimComp);
-							}
+							FPhysScene_ChaosInterface* LocalPhysScene = PrimComp->GetWorld()->GetPhysicsScene();
+							FPhysScene_Chaos& Scene = LocalPhysScene->GetScene();
+							Scene.RegisterForCollisionEvents(PrimComp);
 						}
 					}
-#endif
 				}
+#endif
+			}
 
-				// Set up dynamic instance data
-				if (!IsStatic())
+			// Set up dynamic instance data
+			if (!IsStatic())
+			{
+				SCOPE_CYCLE_COUNTER(STAT_InitBodyPostAdd);
+				for (int32 BodyIdx = 0, NumBodies = Bodies.Num(); BodyIdx < NumBodies; ++BodyIdx)
 				{
-					SCOPE_CYCLE_COUNTER(STAT_InitBodyPostAdd);
-					for (int32 BodyIdx = 0, NumBodies = Bodies.Num(); BodyIdx < NumBodies; ++BodyIdx)
-					{
-						FBodyInstance* Instance = Bodies[BodyIdx];
-						Instance->InitDynamicProperties_AssumesLocked();
-					}
+					FBodyInstance* Instance = Bodies[BodyIdx];
+					Instance->InitDynamicProperties_AssumesLocked();
 				}
-			});
-		}
+			}
+		});
 	}
-};
+}
 
-FBodyInstance::FInitBodySpawnParams::FInitBodySpawnParams(const UPrimitiveComponent* PrimComp)
+FInitBodySpawnParams::FInitBodySpawnParams(const UPrimitiveComponent* PrimComp)
 {
 	bStaticPhysics = PrimComp == nullptr || PrimComp->Mobility != EComponentMobility::Movable;
 
@@ -1368,13 +1345,27 @@ void FBodyInstance::InitBody(class UBodySetup* Setup, const FTransform& Transfor
 	bool bIsStatic = SpawnParams.bStaticPhysics;
 	if(bIsStatic)
 	{
-		FInitBodiesHelper<true> InitBodiesHelper(Bodies, Transforms, Setup, PrimComp, InRBScene, SpawnParams, SpawnParams.Aggregate);
-		InitBodiesHelper.InitBodies();
+		if (CVarEnableDeferredPhysicsCreation.GetValueOnGameThread())
+		{
+			InitBodiesDeferredListStatic.Add(FInitBodiesHelperWithData<true>(MoveTemp(Bodies), MoveTemp(Transforms), Setup, PrimComp, InRBScene, SpawnParams, SpawnParams.Aggregate));
+		}
+		else
+		{
+			FInitBodiesHelper<true> InitBodiesHelper(Bodies, Transforms, Setup, PrimComp, InRBScene, SpawnParams, SpawnParams.Aggregate);
+			InitBodiesHelper.InitBodies();
+		}
 	}
 	else
 	{
-		FInitBodiesHelper<false> InitBodiesHelper(Bodies, Transforms, Setup, PrimComp, InRBScene, SpawnParams, SpawnParams.Aggregate);
-		InitBodiesHelper.InitBodies();
+		if (CVarEnableDeferredPhysicsCreation.GetValueOnGameThread())
+		{
+			InitBodiesDeferredListDynamic.Add(FInitBodiesHelperWithData<false>(MoveTemp(Bodies), MoveTemp(Transforms), Setup, PrimComp, InRBScene, SpawnParams, SpawnParams.Aggregate));
+		}
+		else
+		{
+			FInitBodiesHelper<false> InitBodiesHelper(Bodies, Transforms, Setup, PrimComp, InRBScene, SpawnParams, SpawnParams.Aggregate);
+			InitBodiesHelper.InitBodies();
+		}
 	}
 
 	Bodies.Reset();
@@ -2668,6 +2659,27 @@ FPhysScene* FBodyInstance::GetPhysicsScene()
 	return FPhysicsInterface::GetCurrentScene(ActorHandle);
 }
 
+void FBodyInstance::InitAllBodies(FPhysScene* PhysScene)
+{
+	for (FInitBodiesHelperWithData<true>& BodyToInit : InitBodiesDeferredListStatic)
+	{
+		// allow this to reset the PhysScene on World cleanup
+		if (PhysScene == nullptr || !ensure(BodyToInit.PhysScene == PhysScene))
+		{
+			BodyToInit.PhysScene = PhysScene;
+		}
+		BodyToInit.InitBodies();
+	}
+	for (FInitBodiesHelperWithData<false>& BodyToInit : InitBodiesDeferredListDynamic)
+	{
+		if (PhysScene == nullptr || !ensure(PhysScene == nullptr || BodyToInit.PhysScene == PhysScene))
+		{
+			BodyToInit.PhysScene = PhysScene;
+		}
+		BodyToInit.InitBodies();
+	}
+}
+
 FPhysicsActorHandle& FBodyInstance::GetPhysicsActorHandle()
 {
 	return ActorHandle;
diff --git a/Engine/Source/Runtime/Engine/Private/PhysicsEngine/BodySetup.cpp b/Engine/Source/Runtime/Engine/Private/PhysicsEngine/BodySetup.cpp
--- a/Engine/Source/Runtime/Engine/Private/PhysicsEngine/BodySetup.cpp
+++ b/Engine/Source/Runtime/Engine/Private/PhysicsEngine/BodySetup.cpp
@@ -813,6 +813,9 @@ void UBodySetup::ClearPhysicsMeshes()
 	AggGeom.FreeRenderInfo();
 }
 
+DECLARE_CYCLE_STAT(TEXT("AddShapesToRigidActor"), STAT_AddShapesToActor, STATGROUP_Physics);
+DECLARE_CYCLE_STAT(TEXT("AddGeomToSolver"), STAT_AddGeomToSolver, STATGROUP_Physics);
+
 void UBodySetup::AddShapesToRigidActor_AssumesLocked(
 	FBodyInstance* OwningInstance, 
 	FVector& Scale3D, 
@@ -823,6 +826,8 @@ void UBodySetup::AddShapesToRigidActor_AssumesLocked(
 	const FTransform& RelativeTM, 
 	TArray<FPhysicsShapeHandle>* NewShapes)
 {
+	SCOPE_CYCLE_COUNTER(STAT_AddShapesToActor);
+
 	check(OwningInstance);
 
 	// in editor, there are a lot of things relying on body setup to create physics meshes
@@ -857,7 +862,10 @@ void UBodySetup::AddShapesToRigidActor_AssumesLocked(
 #if WITH_CHAOS
 	AddParams.ChaosTriMeshes = MakeArrayView(ChaosTriMeshes);
 #endif
-	FPhysicsInterface::AddGeometry(OwningInstance->ActorHandle, AddParams, NewShapes);
+	{
+		SCOPE_CYCLE_COUNTER(STAT_AddGeomToSolver);
+		FPhysicsInterface::AddGeometry(OwningInstance->ActorHandle, AddParams, NewShapes);
+	}
 }
 
 void UBodySetup::RemoveSimpleCollision()
diff --git a/Engine/Source/Runtime/Engine/Private/PositionVertexBuffer.cpp b/Engine/Source/Runtime/Engine/Private/PositionVertexBuffer.cpp
--- a/Engine/Source/Runtime/Engine/Private/PositionVertexBuffer.cpp
+++ b/Engine/Source/Runtime/Engine/Private/PositionVertexBuffer.cpp
@@ -18,8 +18,8 @@ class FPositionVertexData :
 	public TStaticMeshVertexData<FPositionVertex>
 {
 public:
-	FPositionVertexData( bool InNeedsCPUAccess=false )
-		: TStaticMeshVertexData<FPositionVertex>( InNeedsCPUAccess )
+	FPositionVertexData(bool InNeedsCPUAccess = false)
+		: TStaticMeshVertexData<FPositionVertex>(InNeedsCPUAccess)
 	{
 	}
 };
@@ -43,7 +43,7 @@ void FPositionVertexBuffer::CleanUp()
 	if (VertexData)
 	{
 		delete VertexData;
-		VertexData = NULL;
+		VertexData = nullptr;
 	}
 }
 
diff --git a/Engine/Source/Runtime/Engine/Private/PrecomputedVolumetricLightmap.cpp b/Engine/Source/Runtime/Engine/Private/PrecomputedVolumetricLightmap.cpp
--- a/Engine/Source/Runtime/Engine/Private/PrecomputedVolumetricLightmap.cpp
+++ b/Engine/Source/Runtime/Engine/Private/PrecomputedVolumetricLightmap.cpp
@@ -331,7 +331,7 @@ void FPrecomputedVolumetricLightmapData::FinalizeImport()
 
 ENGINE_API void FPrecomputedVolumetricLightmapData::InitRHI()
 {
-	if (FeatureLevel >= ERHIFeatureLevel::SM5)
+	if (GetFeatureLevel() >= ERHIFeatureLevel::SM5)
 	{
 		if (IndirectionTextureDimensions.GetMax() > 0)
 		{
@@ -385,7 +385,7 @@ ENGINE_API void FPrecomputedVolumetricLightmapData::InitRHIForSubLevelResources(
 
 ENGINE_API void FPrecomputedVolumetricLightmapData::ReleaseRHI()
 {
-	if (FeatureLevel >= ERHIFeatureLevel::SM5)
+	if (GetFeatureLevel() >= ERHIFeatureLevel::SM5)
 	{
 		{
 			IndirectionTexture.Texture.SafeRelease();
@@ -421,7 +421,7 @@ ENGINE_API void FPrecomputedVolumetricLightmapData::HandleDataMovementInAtlas(in
 
 		for (FPrecomputedVolumetricLightmapData* SceneData : SceneDataAdded)
 		{
-			TShaderMap<FGlobalShaderType>* GlobalShaderMap = GetGlobalShaderMap(FeatureLevel);
+			FGlobalShaderMap* GlobalShaderMap = GetGlobalShaderMap(GetFeatureLevel());
 
 			TShaderMapRef<FMoveWholeIndirectionTextureCS> ComputeShader(GlobalShaderMap);
 
@@ -436,7 +436,7 @@ ENGINE_API void FPrecomputedVolumetricLightmapData::HandleDataMovementInAtlas(in
 			Parameters.OldIndirectionTexture = SceneData->IndirectionTexture.Texture;
 			Parameters.IndirectionTexture = NewIndirectionTexture.UAV;
 
-			FComputeShaderUtils::Dispatch(RHICmdList, *ComputeShader, Parameters,
+			FComputeShaderUtils::Dispatch(RHICmdList, ComputeShader, Parameters,
 				FIntVector(
 					FMath::DivideAndRoundUp(IndirectionTextureDimensions.X, 4),
 					FMath::DivideAndRoundUp(IndirectionTextureDimensions.Y, 4),
@@ -457,7 +457,7 @@ ENGINE_API void FPrecomputedVolumetricLightmapData::HandleDataMovementInAtlas(in
 		{
 			if (SceneData->IndirectionTexture.Texture.IsValid())
 			{
-				TShaderMap<FGlobalShaderType>* GlobalShaderMap = GetGlobalShaderMap(FeatureLevel);
+				FGlobalShaderMap* GlobalShaderMap = GetGlobalShaderMap(GetFeatureLevel());
 
 				TShaderMapRef<FPatchIndirectionTextureCS> ComputeShader(GlobalShaderMap);
 
@@ -469,7 +469,7 @@ ENGINE_API void FPrecomputedVolumetricLightmapData::HandleDataMovementInAtlas(in
 				Parameters.IndirectionTexture = SceneData->IndirectionTexture.UAV;
 				Parameters.SubLevelBrickPositions = SubLevelBrickPositionsSRV;
 
-				FComputeShaderUtils::Dispatch(RHICmdList, *ComputeShader, Parameters, FIntVector(FMath::DivideAndRoundUp(NumBricks, 64), 1, 1));
+				FComputeShaderUtils::Dispatch(RHICmdList, ComputeShader, Parameters, FIntVector(FMath::DivideAndRoundUp(NumBricks, 64), 1, 1));
 
 				FRHIUnorderedAccessView* UAV = SceneData->IndirectionTexture.UAV;
 				RHICmdList.TransitionResources(EResourceTransitionAccess::EReadable, EResourceTransitionPipeline::EComputeToGfx, &UAV, 1);
@@ -511,14 +511,14 @@ ENGINE_API void FPrecomputedVolumetricLightmapData::AddToSceneData(FPrecomputedV
 		SceneData->IndirectionTexture.Format = IndirectionTexture.Format;
 		SceneData->IndirectionTextureDimensions = IndirectionTextureDimensions;
 
-		if (FeatureLevel >= ERHIFeatureLevel::SM5)
+		if (GetFeatureLevel() >= ERHIFeatureLevel::SM5)
 		{
 			// GPU Path
 
 			const int32 PaddedBrickSize = BrickSize + 1;
 			int32 NumBricks = BrickDataDimensions.X * BrickDataDimensions.Y * BrickDataDimensions.Z / (PaddedBrickSize * PaddedBrickSize * PaddedBrickSize);
 
-			TShaderMap<FGlobalShaderType>* GlobalShaderMap = GetGlobalShaderMap(FeatureLevel);
+			FGlobalShaderMap* GlobalShaderMap = GetGlobalShaderMap(GetFeatureLevel());
 
 			TShaderMapRef<FMoveWholeIndirectionTextureCS> ComputeShader(GlobalShaderMap);
 
@@ -533,7 +533,7 @@ ENGINE_API void FPrecomputedVolumetricLightmapData::AddToSceneData(FPrecomputedV
 			Parameters.OldIndirectionTexture = IndirectionTexture.Texture;
 			Parameters.IndirectionTexture = NewIndirectionTexture.UAV;
 
-			FComputeShaderUtils::Dispatch(RHICmdList, *ComputeShader, Parameters,
+			FComputeShaderUtils::Dispatch(RHICmdList, ComputeShader, Parameters,
 				FIntVector(
 					FMath::DivideAndRoundUp(IndirectionTextureDimensions.X, 4),
 					FMath::DivideAndRoundUp(IndirectionTextureDimensions.Y, 4),
@@ -563,14 +563,14 @@ ENGINE_API void FPrecomputedVolumetricLightmapData::AddToSceneData(FPrecomputedV
 	}
 	else
 	{
-		if (FeatureLevel >= ERHIFeatureLevel::SM5)
+		if (GetFeatureLevel() >= ERHIFeatureLevel::SM5)
 		{
 			// GPU Path
 			if (SceneData->IndirectionTexture.Texture.IsValid())
 			{
 				InitRHIForSubLevelResources();
 
-				TShaderMap<FGlobalShaderType>* GlobalShaderMap = GetGlobalShaderMap(FeatureLevel);
+				FGlobalShaderMap* GlobalShaderMap = GetGlobalShaderMap(GetFeatureLevel());
 
 				TShaderMapRef<FPatchIndirectionTextureCS> ComputeShader(GlobalShaderMap);
 
@@ -582,7 +582,7 @@ ENGINE_API void FPrecomputedVolumetricLightmapData::AddToSceneData(FPrecomputedV
 				Parameters.IndirectionTexture = SceneData->IndirectionTexture.UAV;
 				Parameters.SubLevelBrickPositions = SubLevelBrickPositionsSRV;
 
-				FComputeShaderUtils::Dispatch(RHICmdList, *ComputeShader, Parameters, FIntVector(FMath::DivideAndRoundUp(NumBricks, 64), 1, 1));
+				FComputeShaderUtils::Dispatch(RHICmdList, ComputeShader, Parameters, FIntVector(FMath::DivideAndRoundUp(NumBricks, 64), 1, 1));
 
 				FRHIUnorderedAccessView* UAV = SceneData->IndirectionTexture.UAV;
 				RHICmdList.TransitionResources(EResourceTransitionAccess::EReadable, EResourceTransitionPipeline::EComputeToGfx, &UAV, 1);
@@ -652,7 +652,7 @@ ENGINE_API void FPrecomputedVolumetricLightmapData::RemoveFromSceneData(FPrecomp
 	}
 	else
 	{
-		if (FeatureLevel >= ERHIFeatureLevel::SM5)
+		if (GetFeatureLevel() >= ERHIFeatureLevel::SM5)
 		{
 			if (SceneData->IndirectionTexture.Texture.IsValid())
 			{
@@ -660,7 +660,7 @@ ENGINE_API void FPrecomputedVolumetricLightmapData::RemoveFromSceneData(FPrecomp
 
 				InitRHIForSubLevelResources();
 
-				TShaderMap<FGlobalShaderType>* GlobalShaderMap = GetGlobalShaderMap(FeatureLevel);
+				FGlobalShaderMap* GlobalShaderMap = GetGlobalShaderMap(GetFeatureLevel());
 
 				TShaderMapRef<FRemoveSubLevelBricksCS> ComputeShader(GlobalShaderMap);
 
@@ -671,7 +671,7 @@ ENGINE_API void FPrecomputedVolumetricLightmapData::RemoveFromSceneData(FPrecomp
 				Parameters.IndirectionTextureOriginalValues = IndirectionTextureOriginalValuesSRV;
 				Parameters.PersistentLevelBrickDataBaseOffset = PersistentLevelBrickDataBaseOffset;
 
-				FComputeShaderUtils::Dispatch(RHICmdList, *ComputeShader, Parameters, FIntVector(FMath::DivideAndRoundUp(SubLevelBrickPositions.Num(), 64), 1, 1));
+				FComputeShaderUtils::Dispatch(RHICmdList, ComputeShader, Parameters, FIntVector(FMath::DivideAndRoundUp(SubLevelBrickPositions.Num(), 64), 1, 1));
 
 				FRHIUnorderedAccessView* UAV = SceneData->IndirectionTexture.UAV;
 				RHICmdList.TransitionResources(EResourceTransitionAccess::EReadable, EResourceTransitionPipeline::EComputeToGfx, &UAV, 1);
@@ -809,7 +809,7 @@ void FPrecomputedVolumetricLightmap::SetData(FPrecomputedVolumetricLightmapData*
 
 	if (Data)
 	{
-		Data->FeatureLevel = Scene->GetFeatureLevel();
+		Data->SetFeatureLevel(Scene->GetFeatureLevel());
 		Data->IndirectionTexture.bNeedsCPUAccess = GIsEditor;
 		Data->BrickData.SetNeedsCPUAccess(GIsEditor);
 
@@ -909,7 +909,7 @@ template<class VolumetricLightmapBrickDataType>
 void CopyDataIntoAtlas(FRHICommandList& RHICmdList, ERHIFeatureLevel::Type FeatureLevel, int32 SrcOffset, int32 DestOffset, int32 NumBricks, const VolumetricLightmapBrickDataType& SrcData, FVolumetricLightmapBrickTextureSet DestTextureSet)
 {
 	{
-		TShaderMap<FGlobalShaderType>* GlobalShaderMap = GetGlobalShaderMap(FeatureLevel);
+		FGlobalShaderMap* GlobalShaderMap = GetGlobalShaderMap(FeatureLevel);
 
 		FCopyResidentBricksCS::FPermutationDomain PermutationVector;
 		PermutationVector.Set<FCopyResidentBricksCS::FHasSkyBentNormal>(SrcData.SkyBentNormal.Texture.IsValid());
@@ -929,7 +929,7 @@ void CopyDataIntoAtlas(FRHICommandList& RHICmdList, ERHIFeatureLevel::Type Featu
 		Parameters.OutSkyBentNormal = DestTextureSet.SkyBentNormal.UAV;
 		Parameters.OutDirectionalLightShadowing = DestTextureSet.DirectionalLightShadowing.UAV;
 
-		FComputeShaderUtils::Dispatch(RHICmdList, *ComputeShader, Parameters, FIntVector(NumBricks, 1, 1));
+		FComputeShaderUtils::Dispatch(RHICmdList, ComputeShader, Parameters, FIntVector(NumBricks, 1, 1));
 
 		FRHIUnorderedAccessView* UAVs[3];
 		UAVs[0] = Parameters.OutAmbientVector;
@@ -940,7 +940,7 @@ void CopyDataIntoAtlas(FRHICommandList& RHICmdList, ERHIFeatureLevel::Type Featu
 
 	for (int32 i = 0; i < UE_ARRAY_COUNT(SrcData.SHCoefficients); i++)
 	{
-		TShaderMap<FGlobalShaderType>* GlobalShaderMap = GetGlobalShaderMap(FeatureLevel);
+		FGlobalShaderMap* GlobalShaderMap = GetGlobalShaderMap(FeatureLevel);
 
 		TShaderMapRef<FCopyResidentBrickSHCoefficientsCS> ComputeShader(GlobalShaderMap);
 
@@ -952,7 +952,7 @@ void CopyDataIntoAtlas(FRHICommandList& RHICmdList, ERHIFeatureLevel::Type Featu
 		Parameters.SHCoefficients = SrcData.SHCoefficients[i].Texture;
 		Parameters.OutSHCoefficients = DestTextureSet.SHCoefficients[i].UAV;
 
-		FComputeShaderUtils::Dispatch(RHICmdList, *ComputeShader, Parameters, FIntVector(NumBricks, 1, 1));
+		FComputeShaderUtils::Dispatch(RHICmdList, ComputeShader, Parameters, FIntVector(NumBricks, 1, 1));
 
 		FRHIUnorderedAccessView* UAV = Parameters.OutSHCoefficients;
 		RHICmdList.TransitionResources(EResourceTransitionAccess::EReadable, EResourceTransitionPipeline::EComputeToGfx, &UAV, 1);
@@ -965,7 +965,7 @@ void FVolumetricLightmapBrickAtlas::Insert(int32 Index, FPrecomputedVolumetricLi
 
 	if (!bInitialized)
 	{
-		FeatureLevel = ERHIFeatureLevel::SM5;
+		SetFeatureLevel(ERHIFeatureLevel::SM5);
 		check(Data->BrickSize > 0);
 		PaddedBrickSize = Data->BrickSize + 1;
 		TextureSet.Initialize(Data->BrickDataDimensions, Data->BrickData);
@@ -1067,7 +1067,7 @@ void FVolumetricLightmapBrickAtlas::Insert(int32 Index, FPrecomputedVolumetricLi
 		// Copy old allocations
 		for (int32 AllocationIndex = 0; AllocationIndex < Index && AllocationIndex < Allocations.Num(); AllocationIndex++)
 		{
-			CopyDataIntoAtlas(RHICmdList, FeatureLevel, Allocations[AllocationIndex].StartOffset, BrickStartAllocation, Allocations[AllocationIndex].Size, TextureSet, NewTextureSet);
+			CopyDataIntoAtlas(RHICmdList, GetFeatureLevel(), Allocations[AllocationIndex].StartOffset, BrickStartAllocation, Allocations[AllocationIndex].Size, TextureSet, NewTextureSet);
 
 			NewAllocations.Add(Allocation{ Allocations[AllocationIndex].Data, Allocations[AllocationIndex].Size, BrickStartAllocation });
 			BrickStartAllocation += Allocations[AllocationIndex].Size;
@@ -1077,7 +1077,7 @@ void FVolumetricLightmapBrickAtlas::Insert(int32 Index, FPrecomputedVolumetricLi
 		{
 			int32 NumBricks = Data->BrickDataDimensions.X * Data->BrickDataDimensions.Y * Data->BrickDataDimensions.Z / (PaddedBrickSize * PaddedBrickSize * PaddedBrickSize);
 
-			CopyDataIntoAtlas(RHICmdList, FeatureLevel, 0, BrickStartAllocation, NumBricks, Data->BrickData, NewTextureSet);
+			CopyDataIntoAtlas(RHICmdList, GetFeatureLevel(), 0, BrickStartAllocation, NumBricks, Data->BrickData, NewTextureSet);
 
 			NewAllocations.Add(Allocation{ Data, NumBricks, BrickStartAllocation });
 			Data->BrickDataBaseOffsetInAtlas = BrickStartAllocation;
@@ -1087,7 +1087,7 @@ void FVolumetricLightmapBrickAtlas::Insert(int32 Index, FPrecomputedVolumetricLi
 		// Copy the rest of allocations
 		for (int32 AllocationIndex = Index; AllocationIndex < Allocations.Num(); AllocationIndex++)
 		{
-			CopyDataIntoAtlas(RHICmdList, FeatureLevel, Allocations[AllocationIndex].StartOffset, BrickStartAllocation, Allocations[AllocationIndex].Size, TextureSet, NewTextureSet);
+			CopyDataIntoAtlas(RHICmdList, GetFeatureLevel(), Allocations[AllocationIndex].StartOffset, BrickStartAllocation, Allocations[AllocationIndex].Size, TextureSet, NewTextureSet);
 
 			NewAllocations.Add(Allocation{ Allocations[AllocationIndex].Data, Allocations[AllocationIndex].Size, BrickStartAllocation });
 			// Handle the sub level data movements
@@ -1195,7 +1195,7 @@ void FVolumetricLightmapBrickAtlas::Remove(FPrecomputedVolumetricLightmapData* D
 			// Copy old allocations
 			for (int32 AllocationIndex = 0; AllocationIndex < Index && AllocationIndex < Allocations.Num(); AllocationIndex++)
 			{
-				CopyDataIntoAtlas(RHICmdList, FeatureLevel, Allocations[AllocationIndex].StartOffset, BrickStartAllocation, Allocations[AllocationIndex].Size, TextureSet, NewTextureSet);
+				CopyDataIntoAtlas(RHICmdList, GetFeatureLevel(), Allocations[AllocationIndex].StartOffset, BrickStartAllocation, Allocations[AllocationIndex].Size, TextureSet, NewTextureSet);
 
 				NewAllocations.Add(Allocation{ Allocations[AllocationIndex].Data, Allocations[AllocationIndex].Size, BrickStartAllocation });
 				BrickStartAllocation += Allocations[AllocationIndex].Size;
@@ -1206,7 +1206,7 @@ void FVolumetricLightmapBrickAtlas::Remove(FPrecomputedVolumetricLightmapData* D
 			// Copy the rest of allocations
 			for (int32 AllocationIndex = Index + 1; AllocationIndex < Allocations.Num(); AllocationIndex++)
 			{
-				CopyDataIntoAtlas(RHICmdList, FeatureLevel, Allocations[AllocationIndex].StartOffset, BrickStartAllocation, Allocations[AllocationIndex].Size, TextureSet, NewTextureSet);
+				CopyDataIntoAtlas(RHICmdList, GetFeatureLevel(), Allocations[AllocationIndex].StartOffset, BrickStartAllocation, Allocations[AllocationIndex].Size, TextureSet, NewTextureSet);
 
 				NewAllocations.Add(Allocation{ Allocations[AllocationIndex].Data, Allocations[AllocationIndex].Size, BrickStartAllocation });
 				Allocations[AllocationIndex].Data->HandleDataMovementInAtlas(Allocations[AllocationIndex].StartOffset, BrickStartAllocation);
diff --git a/Engine/Source/Runtime/Engine/Private/PrimitiveSceneProxy.cpp b/Engine/Source/Runtime/Engine/Private/PrimitiveSceneProxy.cpp
--- a/Engine/Source/Runtime/Engine/Private/PrimitiveSceneProxy.cpp
+++ b/Engine/Source/Runtime/Engine/Private/PrimitiveSceneProxy.cpp
@@ -39,22 +39,19 @@ bool CacheShadowDepthsFromPrimitivesUsingWPO()
 	return CVarCacheWPOPrimitives.GetValueOnAnyThread(true) != 0;
 }
 
-bool SupportsCachingMeshDrawCommands(const FPrimitiveSceneProxy* RESTRICT PrimitiveSceneProxy, const FMeshBatch& MeshBatch)
+bool SupportsCachingMeshDrawCommands(const FMeshBatch& MeshBatch)
 {
 	return
 		// Cached mesh commands only allow for a single mesh element per batch.
 		(MeshBatch.Elements.Num() == 1) &&
 
 		// Vertex factory needs to support caching.
-		MeshBatch.VertexFactory->GetType()->SupportsCachingMeshDrawCommands() &&
-
-		// Volumetric self shadow mesh commands need to be generated every frame, as they depend on single frame uniform buffers with self shadow data.
-		!PrimitiveSceneProxy->CastsVolumetricTranslucentShadow();
+		MeshBatch.VertexFactory->GetType()->SupportsCachingMeshDrawCommands();
 }
 
-bool SupportsCachingMeshDrawCommands(const FPrimitiveSceneProxy* RESTRICT PrimitiveSceneProxy, const FMeshBatch& MeshBatch, ERHIFeatureLevel::Type FeatureLevel)
+bool SupportsCachingMeshDrawCommands(const FMeshBatch& MeshBatch, ERHIFeatureLevel::Type FeatureLevel)
 {
-	if (SupportsCachingMeshDrawCommands(PrimitiveSceneProxy, MeshBatch))
+	if (SupportsCachingMeshDrawCommands(MeshBatch))
 	{
 		// External textures get mapped to immutable samplers (which are part of the PSO); the mesh must go through the dynamic path, as the media player might not have
 		// valid textures/samplers the first few calls; once they're available the PSO needs to get invalidated and recreated with the immutable samplers.
diff --git a/Engine/Source/Runtime/Engine/Private/PrimitiveUniformShaderParameters.cpp b/Engine/Source/Runtime/Engine/Private/PrimitiveUniformShaderParameters.cpp
--- a/Engine/Source/Runtime/Engine/Private/PrimitiveUniformShaderParameters.cpp
+++ b/Engine/Source/Runtime/Engine/Private/PrimitiveUniformShaderParameters.cpp
@@ -3,9 +3,12 @@
 #include "PrimitiveUniformShaderParameters.h"
 #include "PrimitiveSceneProxy.h"
 #include "PrimitiveSceneInfo.h"
+#include "ProfilingDebugging/LoadTimeTracker.h"
 
 void FSinglePrimitiveStructured::InitRHI() 
 {
+	SCOPED_LOADTIMER(FSinglePrimitiveStructuredBuffer_InitRHI);
+
 	if (RHISupportsComputeShaders(GMaxRHIShaderPlatform))
 	{
 		FRHIResourceCreateInfo CreateInfo;
@@ -33,7 +36,7 @@ void FSinglePrimitiveStructured::UploadToGPU()
 	{
 		void* LockedData = nullptr;
 
-		if (!GPUSceneUseTexture2D(ShaderPlatform))
+		if (!GPUSceneUseTexture2D(ShaderPlatform != SP_NumPlatforms ? ShaderPlatform : GMaxRHIShaderPlatform))
 		{
 			LockedData = RHILockStructuredBuffer(PrimitiveSceneDataBufferRHI, 0, FPrimitiveSceneShaderData::PrimitiveDataStrideInFloat4s * sizeof(FVector4), RLM_WriteOnly);
 			FPlatformMemory::Memcpy(LockedData, PrimitiveSceneData.Data, FPrimitiveSceneShaderData::PrimitiveDataStrideInFloat4s * sizeof(FVector4));
diff --git a/Engine/Source/Runtime/Engine/Private/RendererSupport.cpp b/Engine/Source/Runtime/Engine/Private/RendererSupport.cpp
--- a/Engine/Source/Runtime/Engine/Private/RendererSupport.cpp
+++ b/Engine/Source/Runtime/Engine/Private/RendererSupport.cpp
@@ -38,9 +38,9 @@ static void ClearReferencesToRendererModuleClasses(
 	TMap<UWorld*, bool>& WorldsToUpdate, 
 	TMap<FMaterialShaderMap*, TUniquePtr<TArray<uint8> > >& ShaderMapToSerializedShaderData,
 	FGlobalShaderBackupData& GlobalShaderBackup,
-	TMap<FShaderType*, FString>& ShaderTypeNames,
-	TMap<const FShaderPipelineType*, FString>& ShaderPipelineTypeNames,
-	TMap<FVertexFactoryType*, FString>& VertexFactoryTypeNames)
+	TMap<FShaderType*, FHashedName>& ShaderTypeNames,
+	TMap<const FShaderPipelineType*, FHashedName>& ShaderPipelineTypeNames,
+	TMap<FVertexFactoryType*, FHashedName>& VertexFactoryTypeNames)
 {
 	// Destroy all renderer scenes
 	for (TObjectIterator<UWorld> WorldIt; WorldIt; ++WorldIt)
@@ -77,19 +77,19 @@ static void ClearReferencesToRendererModuleClasses(
 	{
 		FShaderType* ShaderType = *It;
 		check(ShaderType->GetNumShaders() == 0);
-		ShaderTypeNames.Add(ShaderType, ShaderType->GetName());
+		ShaderTypeNames.Add(ShaderType, ShaderType->GetHashedName());
 	}
 
 	for (TLinkedList<FShaderPipelineType*>::TIterator It(FShaderPipelineType::GetTypeList()); It; It.Next())
 	{
 		const FShaderPipelineType* ShaderPipelineType = *It;
-		ShaderPipelineTypeNames.Add(ShaderPipelineType, ShaderPipelineType->GetName());
+		ShaderPipelineTypeNames.Add(ShaderPipelineType, ShaderPipelineType->GetHashedName());
 	}
 
 	for(TLinkedList<FVertexFactoryType*>::TIterator It(FVertexFactoryType::GetTypeList()); It; It.Next())
 	{
 		FVertexFactoryType* VertexFactoryType = *It;
-		VertexFactoryTypeNames.Add(VertexFactoryType, VertexFactoryType->GetName());
+		VertexFactoryTypeNames.Add(VertexFactoryType, VertexFactoryType->GetHashedName());
 	}
 
 	// Destroy misc renderer module classes and remove references
@@ -142,9 +142,9 @@ static void RestoreReferencesToRendererModuleClasses(
 	const TMap<UWorld*, bool>& WorldsToUpdate, 
 	const TMap<FMaterialShaderMap*, TUniquePtr<TArray<uint8> > >& ShaderMapToSerializedShaderData,
 	const FGlobalShaderBackupData& GlobalShaderBackup,
-	const TMap<FShaderType*, FString>& ShaderTypeNames,
-	const TMap<const FShaderPipelineType*, FString>& ShaderPipelineTypeNames,
-	const TMap<FVertexFactoryType*, FString>& VertexFactoryTypeNames)
+	const TMap<FShaderType*, FHashedName>& ShaderTypeNames,
+	const TMap<const FShaderPipelineType*, FHashedName>& ShaderPipelineTypeNames,
+	const TMap<FVertexFactoryType*, FHashedName>& VertexFactoryTypeNames)
 {
 	FlushShaderFileCache();
 
@@ -180,14 +180,13 @@ static void RestoreReferencesToRendererModuleClasses(
 		{
 			EShaderPlatform ShaderPlatform = GetFeatureLevelShaderPlatform((ERHIFeatureLevel::Type)i);
 			check(ShaderPlatform < EShaderPlatform::SP_NumPlatforms);
-			FMaterialShaderMap::FixupShaderTypes(ShaderPlatform, ShaderTypeNames, ShaderPipelineTypeNames, VertexFactoryTypeNames);
 		}
 	}
 
-	TArray<FShaderType*> OutdatedShaderTypes;
+	TArray<const FShaderType*> OutdatedShaderTypes;
 	TArray<const FVertexFactoryType*> OutdatedFactoryTypes;
 	TArray<const FShaderPipelineType*> OutdatedShaderPipelineTypes;
-	FShaderType::GetOutdatedTypes(OutdatedShaderTypes, OutdatedFactoryTypes);
+	GetOutdatedShaderTypes(OutdatedShaderTypes, OutdatedShaderPipelineTypes, OutdatedFactoryTypes);
 
 	// Recompile any missing shaders
 	UMaterialInterface::IterateOverActiveFeatureLevels([&](ERHIFeatureLevel::Type FeatureLevel) 
@@ -225,9 +224,9 @@ void RecompileRenderer(const TArray<FString>& Args)
 		TMap<UWorld*, bool> WorldsToUpdate;
 		TMap<FMaterialShaderMap*, TUniquePtr<TArray<uint8> > > ShaderMapToSerializedShaderData;
 		FGlobalShaderBackupData GlobalShaderBackup;
-		TMap<FShaderType*, FString> ShaderTypeNames;
-		TMap<const FShaderPipelineType*, FString> ShaderPipelineTypeNames;
-		TMap<FVertexFactoryType*, FString> VertexFactoryTypeNames;
+		TMap<FShaderType*, FHashedName> ShaderTypeNames;
+		TMap<const FShaderPipelineType*, FHashedName> ShaderPipelineTypeNames;
+		TMap<FVertexFactoryType*, FHashedName> VertexFactoryTypeNames;
 
 		ClearReferencesToRendererModuleClasses(WorldsToUpdate, ShaderMapToSerializedShaderData, GlobalShaderBackup, ShaderTypeNames, ShaderPipelineTypeNames, VertexFactoryTypeNames);
 
diff --git a/Engine/Source/Runtime/Engine/Private/Rendering/ColorVertexBuffer.cpp b/Engine/Source/Runtime/Engine/Private/Rendering/ColorVertexBuffer.cpp
--- a/Engine/Source/Runtime/Engine/Private/Rendering/ColorVertexBuffer.cpp
+++ b/Engine/Source/Runtime/Engine/Private/Rendering/ColorVertexBuffer.cpp
@@ -5,7 +5,7 @@
 #include "RHI.h"
 #include "Components.h"
 #include "EngineUtils.h"
-#include "StaticMeshVertexData.h"
+#include "ProfilingDebugging/LoadTimeTracker.h"
 
 /*-----------------------------------------------------------------------------
 FColorVertexBuffer
@@ -16,8 +16,8 @@ class FColorVertexData :
 	public TStaticMeshVertexData<FColor>
 {
 public:
-	FColorVertexData( bool InNeedsCPUAccess=false )
-		: TStaticMeshVertexData<FColor>( InNeedsCPUAccess )
+	FColorVertexData(bool InNeedsCPUAccess = false)
+		: TStaticMeshVertexData<FColor>(InNeedsCPUAccess)
 	{
 	}
 };
@@ -418,6 +418,8 @@ FVertexBufferRHIRef FColorVertexBuffer::CreateRHIBuffer_Async()
 
 void FColorVertexBuffer::InitRHI()
 {
+	SCOPED_LOADTIMER(FColorVertexBuffer_InitRHI);
+
 	VertexBufferRHI = CreateRHIBuffer_RenderThread();
 
 	if (VertexBufferRHI && RHISupportsManualVertexFetch(GMaxRHIShaderPlatform))
diff --git a/Engine/Source/Runtime/Engine/Private/Rendering/MorphTargetVertexInfoBuffers.cpp b/Engine/Source/Runtime/Engine/Private/Rendering/MorphTargetVertexInfoBuffers.cpp
--- a/Engine/Source/Runtime/Engine/Private/Rendering/MorphTargetVertexInfoBuffers.cpp
+++ b/Engine/Source/Runtime/Engine/Private/Rendering/MorphTargetVertexInfoBuffers.cpp
@@ -1,10 +1,12 @@
 // Copyright Epic Games, Inc. All Rights Reserved.
 
 #include "Rendering/MorphTargetVertexInfoBuffers.h"
-
+#include "ProfilingDebugging/LoadTimeTracker.h"
 
 void FMorphTargetVertexInfoBuffers::InitRHI()
 {
+	SCOPED_LOADTIMER(FFMorphTargetVertexInfoBuffers_InitRHI);
+
 	check(NumTotalWorkItems > 0);
 
 	{
diff --git a/Engine/Source/Runtime/Engine/Private/Rendering/SeparableSSS.cpp b/Engine/Source/Runtime/Engine/Private/Rendering/SeparableSSS.cpp
--- a/Engine/Source/Runtime/Engine/Private/Rendering/SeparableSSS.cpp
+++ b/Engine/Source/Runtime/Engine/Private/Rendering/SeparableSSS.cpp
@@ -166,7 +166,7 @@ void ComputeMirroredSSSKernel(FLinearColor* TargetBuffer, uint32 TargetBufferSiz
 
 	// generate output (remove negative samples)
 	{
-		check(kernel[0].A == 0.0f);
+		check(FMath::IsNearlyZero(kernel[0].A));
 
 		// center sample
 		TargetBuffer[0] = kernel[0];
diff --git a/Engine/Source/Runtime/Engine/Private/Rendering/SkeletalMeshVertexClothBuffer.cpp b/Engine/Source/Runtime/Engine/Private/Rendering/SkeletalMeshVertexClothBuffer.cpp
--- a/Engine/Source/Runtime/Engine/Private/Rendering/SkeletalMeshVertexClothBuffer.cpp
+++ b/Engine/Source/Runtime/Engine/Private/Rendering/SkeletalMeshVertexClothBuffer.cpp
@@ -4,6 +4,7 @@
 #include "Rendering/SkeletalMeshVertexBuffer.h"
 #include "EngineUtils.h"
 #include "SkeletalMeshTypes.h"
+#include "ProfilingDebugging/LoadTimeTracker.h"
 
 /**
 * Constructor
@@ -109,6 +110,8 @@ FVertexBufferRHIRef FSkeletalMeshVertexClothBuffer::CreateRHIBuffer_Async()
 */
 void FSkeletalMeshVertexClothBuffer::InitRHI()
 {
+	SCOPED_LOADTIMER(FSkeletalMeshVertexClothBuffer_InitRHI);
+
 	VertexBufferRHI = CreateRHIBuffer_RenderThread();
 
 	if (VertexBufferRHI)
diff --git a/Engine/Source/Runtime/Engine/Private/Rendering/SkinWeightVertexBuffer.cpp b/Engine/Source/Runtime/Engine/Private/Rendering/SkinWeightVertexBuffer.cpp
--- a/Engine/Source/Runtime/Engine/Private/Rendering/SkinWeightVertexBuffer.cpp
+++ b/Engine/Source/Runtime/Engine/Private/Rendering/SkinWeightVertexBuffer.cpp
@@ -6,6 +6,7 @@
 #include "RenderUtils.h"
 #include "SkeletalMeshTypes.h"
 #include "UObject/AnimObjectVersion.h"
+#include "ProfilingDebugging/LoadTimeTracker.h"
 
 /*-----------------------------------------------------------------------------
 FSkinWeightLookupVertexBuffer
@@ -128,15 +129,18 @@ void FSkinWeightLookupVertexBuffer::InitRHI()
 	// BUF_ShaderResource is needed for support of the SkinCache (we could make is dependent on GEnableGPUSkinCacheShaders or are there other users?)
 	VertexBufferRHI = CreateRHIBuffer_RenderThread();
 
-	bool bSRV = VertexBufferRHI && GSupportsResourceView && GPixelFormats[PixelFormat].Supported;
-	// When bAllowCPUAccess is true, the meshes is likely going to be used for Niagara to spawn particles on mesh surface.
-	// And it can be the case for CPU *and* GPU access: no differenciation today. That is why we create a SRV in this case.
-	// This also avoid setting lots of states on all the members of all the different buffers used by meshes. Follow up: https://jira.it.epicgames.net/browse/UE-69376.
-	bSRV |= GetNeedsCPUAccess();
-
-	if (bSRV)
+	if (VertexBufferRHI)
 	{
-		SRVValue = RHICreateShaderResourceView(LookupData ? VertexBufferRHI : nullptr, PixelFormatStride, PixelFormat);
+		bool bSRV = GSupportsResourceView && GPixelFormats[PixelFormat].Supported;
+		// When bAllowCPUAccess is true, the meshes is likely going to be used for Niagara to spawn particles on mesh surface.
+		// And it can be the case for CPU *and* GPU access: no differenciation today. That is why we create a SRV in this case.
+		// This also avoid setting lots of states on all the members of all the different buffers used by meshes. Follow up: https://jira.it.epicgames.net/browse/UE-69376.
+		bSRV |= GetNeedsCPUAccess();
+
+		if (bSRV)
+		{
+			SRVValue = RHICreateShaderResourceView(LookupData ? VertexBufferRHI : nullptr, PixelFormatStride, PixelFormat);
+		}
 	}
 }
 
@@ -419,6 +423,8 @@ FVertexBufferRHIRef FSkinWeightDataVertexBuffer::CreateRHIBuffer_Async()
 
 void FSkinWeightDataVertexBuffer::InitRHI()
 {
+	SCOPED_LOADTIMER(FSkinWeightVertexBuffer_InitRHI);
+
 	// BUF_ShaderResource is needed for support of the SkinCache (we could make is dependent on GEnableGPUSkinCacheShaders or are there other users?)
 	VertexBufferRHI = CreateRHIBuffer_RenderThread();
 
