Copying Private-LoadTimes-4.24 stream to Main. Biggest changes are in Materials/Shader memory freezing.

diff --git a/Engine/Source/Runtime/Engine/Private/Rendering/SkyLightImportanceSampling.cpp b/Engine/Source/Runtime/Engine/Private/Rendering/SkyLightImportanceSampling.cpp
--- a/Engine/Source/Runtime/Engine/Private/Rendering/SkyLightImportanceSampling.cpp
+++ b/Engine/Source/Runtime/Engine/Private/Rendering/SkyLightImportanceSampling.cpp
@@ -52,7 +52,7 @@ public:
 		uint32 MipLevel,
 		FRWBuffer& MipTree)
 	{
-		FRHIComputeShader* ShaderRHI = GetComputeShader();
+		FRHIComputeShader* ShaderRHI = RHICmdList.GetBoundComputeShader();
 
 		SetShaderValue(RHICmdList, ShaderRHI, DimensionsParameter, Dimensions);
 		SetShaderValue(RHICmdList, ShaderRHI, FaceIndexParameter, FaceIndex);
@@ -69,32 +69,21 @@ public:
 		EResourceTransitionPipeline TransitionPipeline,
 		FRWBuffer& MipTree)
 	{
-		FRHIComputeShader* ShaderRHI = GetComputeShader();
+		FRHIComputeShader* ShaderRHI = RHICmdList.GetBoundComputeShader();
 
 		MipTreeParameter.UnsetUAV(RHICmdList, ShaderRHI);
 		RHICmdList.TransitionResource(TransitionAccess, TransitionPipeline, MipTree.UAV);
 	}
 
-	virtual bool Serialize(FArchive& Ar)
-	{
-		bool bShaderHasOutdatedParameters = FGlobalShader::Serialize(Ar);
-		Ar << TextureParameter;
-		Ar << TextureSamplerParameter;
-		Ar << DimensionsParameter;
-		Ar << FaceIndexParameter;
-		Ar << MipLevelParameter;
-		Ar << MipTreeParameter;
-		return bShaderHasOutdatedParameters;
-	}
 
 private:
-	FShaderResourceParameter TextureParameter;
-	FShaderResourceParameter TextureSamplerParameter;
+	LAYOUT_FIELD(FShaderResourceParameter, TextureParameter);
+	LAYOUT_FIELD(FShaderResourceParameter,TextureSamplerParameter);
 
-	FShaderParameter DimensionsParameter;
-	FShaderParameter FaceIndexParameter;
-	FShaderParameter MipLevelParameter;
-	FRWShaderParameter MipTreeParameter;
+	LAYOUT_FIELD(FShaderParameter, DimensionsParameter);
+	LAYOUT_FIELD(FShaderParameter, FaceIndexParameter);
+	LAYOUT_FIELD(FShaderParameter, MipLevelParameter);
+	LAYOUT_FIELD(FRWShaderParameter, MipTreeParameter);
 };
 
 class FBuildSolidAnglePdfCS : public FGlobalShader
@@ -135,7 +124,7 @@ public:
 		FRWBuffer& SolidAnglePdf
 	)
 	{
-		FRHIComputeShader* ShaderRHI = GetComputeShader();
+		FRHIComputeShader* ShaderRHI = RHICmdList.GetBoundComputeShader();
 
 		SetShaderValue(RHICmdList, ShaderRHI, MipLevelParameter, MipLevel);
 		SetShaderValue(RHICmdList, ShaderRHI, DimensionsParameter, Dimensions);
@@ -151,25 +140,16 @@ public:
 		FRHIComputeFence* Fence
 	)
 	{
-		FRHIComputeShader* ShaderRHI = GetComputeShader();
+		FRHIComputeShader* ShaderRHI = RHICmdList.GetBoundComputeShader();
 
 		SolidAnglePdfParameter.UnsetUAV(RHICmdList, ShaderRHI);
 		RHICmdList.TransitionResource(TransitionAccess, TransitionPipeline, MipTreePdf.UAV, Fence);
 	}
 
-	virtual bool Serialize(FArchive& Ar)
-	{
-		bool bShaderHasOutdatedParameters = FGlobalShader::Serialize(Ar);
-		Ar << MipLevelParameter;
-		Ar << DimensionsParameter;
-		Ar << SolidAnglePdfParameter;
-		return bShaderHasOutdatedParameters;
-	}
-
 private:
-	FShaderParameter MipLevelParameter;
-	FShaderParameter DimensionsParameter;
-	FRWShaderParameter SolidAnglePdfParameter;
+	LAYOUT_FIELD(FShaderParameter, MipLevelParameter);
+	LAYOUT_FIELD(FShaderParameter, DimensionsParameter);
+	LAYOUT_FIELD(FRWShaderParameter, SolidAnglePdfParameter);
 };
 
 class FBuildMipTreePdfCS : public FGlobalShader
@@ -211,7 +191,7 @@ public:
 		uint32 MipLevel,
 		FRWBuffer& MipTreePdf)
 	{
-		FRHIComputeShader* ShaderRHI = GetComputeShader();
+		FRHIComputeShader* ShaderRHI = RHICmdList.GetBoundComputeShader();
 
 		SetSRVParameter(RHICmdList, ShaderRHI, MipTreeParameter, MipTree.SRV);
 		SetShaderValue(RHICmdList, ShaderRHI, DimensionsParameter, Dimensions);
@@ -227,28 +207,18 @@ public:
 		EResourceTransitionPipeline TransitionPipeline,
 		FRWBuffer& MipTreePdf)
 	{
-		FRHIComputeShader* ShaderRHI = GetComputeShader();
+		FRHIComputeShader* ShaderRHI = RHICmdList.GetBoundComputeShader();
 
 		MipTreePdfParameter.UnsetUAV(RHICmdList, ShaderRHI);
 		RHICmdList.TransitionResource(TransitionAccess, TransitionPipeline, MipTreePdf.UAV);
 	}
 
-	virtual bool Serialize(FArchive& Ar)
-	{
-		bool bShaderHasOutdatedParameters = FGlobalShader::Serialize(Ar);
-		Ar << MipTreeParameter;
-		Ar << DimensionsParameter;
-		Ar << MipLevelParameter;
-		Ar << MipTreePdfParameter;
-		return bShaderHasOutdatedParameters;
-	}
-
 private:
-	FShaderResourceParameter MipTreeParameter;
-	FShaderParameter DimensionsParameter;
-	FShaderParameter MipLevelParameter;
+	LAYOUT_FIELD(FShaderResourceParameter, MipTreeParameter);
+	LAYOUT_FIELD(FShaderParameter, DimensionsParameter);
+	LAYOUT_FIELD(FShaderParameter, MipLevelParameter);
 
-	FRWShaderParameter MipTreePdfParameter;
+	LAYOUT_FIELD(FRWShaderParameter, MipTreePdfParameter);
 };
 
 IMPLEMENT_SHADER_TYPE(, FBuildMipTreeCS, TEXT("/Engine/Private/Raytracing/BuildMipTreeCS.usf"), TEXT("BuildMipTreeCS"), SF_Compute)
@@ -269,7 +239,7 @@ void BuildSkyLightMipTree(
 {
 	const auto ShaderMap = GetGlobalShaderMap(ERHIFeatureLevel::SM5);
 	TShaderMapRef<FBuildMipTreeCS> BuildSkyLightMipTreeComputeShader(ShaderMap);
-	RHICmdList.SetComputeShader(BuildSkyLightMipTreeComputeShader->GetComputeShader());
+	RHICmdList.SetComputeShader(BuildSkyLightMipTreeComputeShader.GetComputeShader());
 
 	FRWBuffer* MipTrees[] = {
 		&SkyLightMipTreePosX,
@@ -304,7 +274,7 @@ void BuildSkyLightMipTree(
 			BuildSkyLightMipTreeComputeShader->SetParameters(RHICmdList, SkyLightTexture, SkyLightMipTreeDimensions, FaceIndex, MipLevel, *MipTrees[FaceIndex]);
 			FIntVector MipLevelDimensions = FIntVector(SkyLightMipTreeDimensions.X >> MipLevel, SkyLightMipTreeDimensions.Y >> MipLevel, 1);
 			FIntVector NumGroups = FIntVector::DivideAndRoundUp(MipLevelDimensions, FBuildMipTreeCS::GetGroupSize());
-			DispatchComputeShader(RHICmdList, *BuildSkyLightMipTreeComputeShader, NumGroups.X, NumGroups.Y, 1);
+			DispatchComputeShader(RHICmdList, BuildSkyLightMipTreeComputeShader, NumGroups.X, NumGroups.Y, 1);
 			BuildSkyLightMipTreeComputeShader->UnsetParameters(RHICmdList, EResourceTransitionAccess::ERWBarrier, EResourceTransitionPipeline::EComputeToCompute, *MipTrees[FaceIndex]);
 		}
 
@@ -326,7 +296,7 @@ void BuildSolidAnglePdf(
 {
 	const auto ShaderMap = GetGlobalShaderMap(ERHIFeatureLevel::SM5);
 	TShaderMapRef<FBuildSolidAnglePdfCS> BuildSolidAnglePdfComputeShader(ShaderMap);
-	RHICmdList.SetComputeShader(BuildSolidAnglePdfComputeShader->GetComputeShader());
+	RHICmdList.SetComputeShader(BuildSolidAnglePdfComputeShader.GetComputeShader());
 
 	uint32 NumElements = Dimensions.X * Dimensions.Y;
 	uint32 MipLevelCount = FMath::Log2(Dimensions.X);
@@ -341,7 +311,7 @@ void BuildSolidAnglePdf(
 		FComputeFenceRHIRef ComputeFence = RHICmdList.CreateComputeFence(TEXT("SkyLight SolidAnglePdf Build"));
 		BuildSolidAnglePdfComputeShader->SetParameters(RHICmdList, MipLevel, Dimensions, SolidAnglePdf);
 		FIntVector NumGroups = FIntVector::DivideAndRoundUp(Dimensions, FBuildSolidAnglePdfCS::GetGroupSize());
-		DispatchComputeShader(RHICmdList, *BuildSolidAnglePdfComputeShader, NumGroups.X, NumGroups.Y, 1);
+		DispatchComputeShader(RHICmdList, BuildSolidAnglePdfComputeShader, NumGroups.X, NumGroups.Y, 1);
 		BuildSolidAnglePdfComputeShader->UnsetParameters(RHICmdList, EResourceTransitionAccess::ERWBarrier, EResourceTransitionPipeline::EComputeToCompute, SolidAnglePdf, ComputeFence);
 	}
 }
@@ -383,7 +353,7 @@ void BuildSkyLightMipTreePdf(
 
 	const auto ShaderMap = GetGlobalShaderMap(ERHIFeatureLevel::SM5);
 	TShaderMapRef<FBuildMipTreePdfCS> BuildSkyLightMipTreePdfComputeShader(ShaderMap);
-	RHICmdList.SetComputeShader(BuildSkyLightMipTreePdfComputeShader->GetComputeShader());
+	RHICmdList.SetComputeShader(BuildSkyLightMipTreePdfComputeShader.GetComputeShader());
 
 	uint32 NumElements = SkyLightMipTreePosX.NumBytes / sizeof(float);
 	uint32 MipLevelCount = FMath::Log2(SkyLightMipTreeDimensions.X);
@@ -397,7 +367,7 @@ void BuildSkyLightMipTreePdf(
 			BuildSkyLightMipTreePdfComputeShader->SetParameters(RHICmdList, *MipTrees[FaceIndex], SkyLightMipTreeDimensions, MipLevel, *MipTreePdfs[FaceIndex]);
 			FIntVector MipLevelDimensions = FIntVector(SkyLightMipTreeDimensions.X >> MipLevel, SkyLightMipTreeDimensions.Y >> MipLevel, 1);
 			FIntVector NumGroups = FIntVector::DivideAndRoundUp(MipLevelDimensions, FBuildMipTreeCS::GetGroupSize());
-			DispatchComputeShader(RHICmdList, *BuildSkyLightMipTreePdfComputeShader, NumGroups.X, NumGroups.Y, 1);
+			DispatchComputeShader(RHICmdList, BuildSkyLightMipTreePdfComputeShader, NumGroups.X, NumGroups.Y, 1);
 		}
 		BuildSkyLightMipTreePdfComputeShader->UnsetParameters(RHICmdList, EResourceTransitionAccess::ERWBarrier, EResourceTransitionPipeline::EComputeToCompute, *MipTreePdfs[FaceIndex]);
 	}
diff --git a/Engine/Source/Runtime/Engine/Private/Rendering/StaticMeshVertexBuffer.cpp b/Engine/Source/Runtime/Engine/Private/Rendering/StaticMeshVertexBuffer.cpp
--- a/Engine/Source/Runtime/Engine/Private/Rendering/StaticMeshVertexBuffer.cpp
+++ b/Engine/Source/Runtime/Engine/Private/Rendering/StaticMeshVertexBuffer.cpp
@@ -4,6 +4,7 @@
 #include "EngineUtils.h"
 #include "Components.h"
 #include "GPUSkinCache.h"
+#include "ProfilingDebugging/LoadTimeTracker.h"
 
 FStaticMeshVertexBuffer::FStaticMeshVertexBuffer() :
 	TangentsData(nullptr),
@@ -270,7 +271,7 @@ FVertexBufferRHIRef FStaticMeshVertexBuffer::CreateTangentsRHIBuffer_Internal()
 {
 	if (GetNumVertices())
 	{
-		FResourceArrayInterface* RESTRICT ResourceArray = TangentsData ? TangentsData->GetResourceArray() : nullptr;
+ 		FResourceArrayInterface* RESTRICT ResourceArray = TangentsData ? TangentsData->GetResourceArray() : nullptr;
 		const uint32 SizeInBytes = ResourceArray ? ResourceArray->GetResourceDataSize() : 0;
 		FRHIResourceCreateInfo CreateInfo(ResourceArray);
 		CreateInfo.bWithoutNativeResource = !TangentsData;
@@ -329,6 +330,8 @@ FVertexBufferRHIRef FStaticMeshVertexBuffer::CreateTexCoordRHIBuffer_Async()
 
 void FStaticMeshVertexBuffer::InitRHI()
 {
+	SCOPED_LOADTIMER(FStaticMeshVertexBuffer_InitRHI);
+
 	TangentsVertexBuffer.VertexBufferRHI = CreateTangentsRHIBuffer_RenderThread();
 	TexCoordVertexBuffer.VertexBufferRHI = CreateTexCoordRHIBuffer_RenderThread();
 	if (TangentsVertexBuffer.VertexBufferRHI && (RHISupportsManualVertexFetch(GMaxRHIShaderPlatform) || IsGPUSkinCacheAvailable(GMaxRHIShaderPlatform)))
diff --git a/Engine/Source/Runtime/Engine/Private/ScreenRendering.cpp b/Engine/Source/Runtime/Engine/Private/ScreenRendering.cpp
--- a/Engine/Source/Runtime/Engine/Private/ScreenRendering.cpp
+++ b/Engine/Source/Runtime/Engine/Private/ScreenRendering.cpp
@@ -13,6 +13,6 @@ TGlobalResource<FScreenVertexDeclaration> GScreenVertexDeclaration;
 IMPLEMENT_SHADER_TYPE(,FScreenPS,TEXT("/Engine/Private/ScreenPixelShader.usf"),TEXT("Main"),SF_Pixel);
 IMPLEMENT_SHADER_TYPE(,FScreenPSsRGBSource, TEXT("/Engine/Private/ScreenPixelShader.usf"), TEXT("MainsRGBSource"), SF_Pixel);
 IMPLEMENT_SHADER_TYPE(,FScreenVS,TEXT("/Engine/Private/ScreenVertexShader.usf"),TEXT("Main"),SF_Vertex);
-IMPLEMENT_SHADER_TYPE(template<> ENGINE_API,TScreenVSForGS<false>,TEXT("/Engine/Private/ScreenVertexShader.usf"),TEXT("MainForGS"),SF_Vertex);
-IMPLEMENT_SHADER_TYPE(template<> ENGINE_API,TScreenVSForGS<true>,TEXT("/Engine/Private/ScreenVertexShader.usf"),TEXT("MainForGS"),SF_Vertex);
+IMPLEMENT_SHADER_TYPE(template<>,TScreenVSForGS<false>,TEXT("/Engine/Private/ScreenVertexShader.usf"),TEXT("MainForGS"),SF_Vertex);
+IMPLEMENT_SHADER_TYPE(template<>,TScreenVSForGS<true>,TEXT("/Engine/Private/ScreenVertexShader.usf"),TEXT("MainForGS"),SF_Vertex);
 IMPLEMENT_SHADER_TYPE(,FScreenPS_OSE,TEXT("/Engine/Private/ScreenPixelShaderOES.usf"),TEXT("Main"),SF_Pixel);
diff --git a/Engine/Source/Runtime/Engine/Private/ShaderCompiler/ShaderCompiler.cpp b/Engine/Source/Runtime/Engine/Private/ShaderCompiler/ShaderCompiler.cpp
--- a/Engine/Source/Runtime/Engine/Private/ShaderCompiler/ShaderCompiler.cpp
+++ b/Engine/Source/Runtime/Engine/Private/ShaderCompiler/ShaderCompiler.cpp
@@ -47,6 +47,7 @@
 #include "MeshMaterialShaderType.h"
 #include "ShaderParameterMetadata.h"
 #include "ProfilingDebugging/DiagnosticTable.h"
+#include "ProfilingDebugging/LoadTimeTracker.h"
 
 #define LOCTEXT_NAMESPACE "ShaderCompiler"
 
@@ -926,7 +927,13 @@ static bool CheckSingleJob(FShaderCompileJob* SingleJob, const TArray<FMaterial*
 	if (SingleJob->ShaderType)
 	{
 		// Allow the shader validation to fail the compile if it sees any parameters bound that aren't supported.
-		if (FMaterialShaderType* MaterialShaderType = SingleJob->ShaderType->GetMaterialShaderType())
+		const bool bValidationResult = SingleJob->ShaderType->ValidateCompiledResult(
+			(EShaderPlatform)SingleJob->Input.Target.Platform,
+			SingleJob->Output.ParameterMap,
+			Errors);
+		bSucceeded = bValidationResult && bSucceeded;
+
+		/*if (FMaterialShaderType* MaterialShaderType = SingleJob->ShaderType->GetMaterialShaderType())
 		{
 			bool bValidationResult = MaterialShaderType->ValidateCompiledResult(
 				(EShaderPlatform)SingleJob->Input.Target.Platform,
@@ -952,7 +959,7 @@ static bool CheckSingleJob(FShaderCompileJob* SingleJob, const TArray<FMaterial*
 				SingleJob->Output.ParameterMap,
 				Errors);
 			bSucceeded = bValidationResult && bSucceeded;
-		}
+		}*/
 	}
 
 	if (SingleJob->VFType)
@@ -2619,7 +2626,7 @@ bool FShaderCompilingManager::HandlePotentialRetryOnError(TMap<int32, FShaderMap
 					}
 				}
 
-				const TCHAR* MaterialName = ShaderMap ? *ShaderMap->GetFriendlyName() : TEXT("global shaders");
+				const TCHAR* MaterialName = ShaderMap ? ShaderMap->GetFriendlyName() : TEXT("global shaders");
 				FString ErrorString = FString::Printf(TEXT("%i Shader compiler errors compiling %s for platform %s:"), UniqueErrors.Num(), MaterialName, *TargetShaderPlatformString);
 				UE_LOG(LogShaderCompilers, Warning, TEXT("%s"), *ErrorString);
 				ErrorString += TEXT("\n");
@@ -3563,7 +3570,7 @@ void GlobalBeginCompileShader(
 			}
 		}
 	}
-	
+
 	if (IsOpenGLPlatform((EShaderPlatform)Target.Platform))
 	{
 		static const auto CVar = IConsoleManager::Get().FindConsoleVariable(TEXT("r.OpenGL.ForceDXC"));
@@ -3720,7 +3727,7 @@ void GlobalBeginCompileShader(
 	Input.Environment.SetDefine(TEXT("PLATFORM_SUPPORTS_RENDERTARGET_WRITE_MASK"), RHISupportsRenderTargetWriteMask(EShaderPlatform(Target.Platform)) ? 1 : 0);
 	Input.Environment.SetDefine(TEXT("PLATFORM_SUPPORTS_PER_PIXEL_DBUFFER_MASK"), IsUsingPerPixelDBufferMask(EShaderPlatform(Target.Platform)) ? 1 : 0);
 	Input.Environment.SetDefine(TEXT("PLATFORM_SUPPORTS_DISTANCE_FIELDS"), DoesPlatformSupportDistanceFields(EShaderPlatform(Target.Platform)) ? 1 : 0);
-
+	
 	{
 		static const auto CVar = IConsoleManager::Get().FindConsoleVariable(TEXT("r.vt.FeedbackFactor"));
 		Input.Environment.SetDefine(TEXT("VIRTUAL_TEXTURE_FEEDBACK_FACTOR"), CVar ? FMath::Max(CVar->GetInt(), 1) : 1);
@@ -3831,14 +3838,15 @@ public:
 		Payload << MaterialsToLoad;
 		uint32 ShaderPlatform = ( uint32 )GMaxRHIShaderPlatform;
 		Payload << ShaderPlatform;
+		check(false); // TODO
 		// tell the other side the Ids we have so it doesn't send back duplicates
 		// (need to serialize this into a TArray since FShaderResourceId isn't known in the file server)
-		TArray<FShaderResourceId> AllIds;
-		FShaderResource::GetAllShaderResourceId( AllIds );
+		//TArray<FShaderResourceId> AllIds;
+		//FShaderResource_InlineCode::GetAllShaderResourceId( AllIds );
 
 		TArray<uint8> SerializedBytes;
 		FMemoryWriter Ar( SerializedBytes );
-		Ar << AllIds;
+		//Ar << AllIds;
 		Payload << SerializedBytes;
 		Payload << bCompileChangedShaders;
 	}
@@ -3909,6 +3917,35 @@ void RecompileGlobalShaders()
 	}
 }
 
+void GetOutdatedShaderTypes(TArray<const FShaderType*>& OutdatedShaderTypes, TArray<const FShaderPipelineType*>& OutdatedShaderPipelineTypes, TArray<const FVertexFactoryType*>& OutdatedFactoryTypes)
+{
+#if WITH_EDITOR
+	for (int PlatformIndex = 0; PlatformIndex < SP_NumPlatforms; ++PlatformIndex)
+	{
+		const FGlobalShaderMap* ShaderMap = GGlobalShaderMap[PlatformIndex];
+		if (ShaderMap)
+		{
+			ShaderMap->GetOutdatedTypes(OutdatedShaderTypes, OutdatedShaderPipelineTypes, OutdatedFactoryTypes);
+		}
+	}
+
+	FMaterialShaderMap::GetAllOutdatedTypes(OutdatedShaderTypes, OutdatedShaderPipelineTypes, OutdatedFactoryTypes);
+
+	for (int32 TypeIndex = 0; TypeIndex < OutdatedShaderTypes.Num(); TypeIndex++)
+	{
+		UE_LOG(LogShaders, Warning, TEXT("		Recompiling %s"), OutdatedShaderTypes[TypeIndex]->GetName());
+	}
+	for (int32 TypeIndex = 0; TypeIndex < OutdatedShaderPipelineTypes.Num(); TypeIndex++)
+	{
+		UE_LOG(LogShaders, Warning, TEXT("		Recompiling %s"), OutdatedShaderPipelineTypes[TypeIndex]->GetName());
+	}
+	for (int32 TypeIndex = 0; TypeIndex < OutdatedFactoryTypes.Num(); TypeIndex++)
+	{
+		UE_LOG(LogShaders, Warning, TEXT("		Recompiling %s"), OutdatedFactoryTypes[TypeIndex]->GetName());
+	}
+#endif // WITH_EDITOR
+}
+
 bool RecompileShaders(const TCHAR* Cmd, FOutputDevice& Ar)
 {
 	// if this platform can't compile shaders, then we try to send a message to a file/cooker server
@@ -3933,13 +3970,12 @@ bool RecompileShaders(const TCHAR* Cmd, FOutputDevice& Ar)
 
 		if( FCString::Stricmp(*FlagStr,TEXT("Changed"))==0)
 		{
-			TArray<FShaderType*> OutdatedShaderTypes;
+			TArray<const FShaderType*> OutdatedShaderTypes;
 			TArray<const FVertexFactoryType*> OutdatedFactoryTypes;
 			TArray<const FShaderPipelineType*> OutdatedShaderPipelineTypes;
 			{
 				FRecompileShadersTimer SearchTimer(TEXT("Searching for changed files"));
-				FShaderType::GetOutdatedTypes(OutdatedShaderTypes, OutdatedFactoryTypes);
-				FShaderPipelineType::GetOutdatedTypes(OutdatedShaderTypes, OutdatedShaderPipelineTypes, OutdatedFactoryTypes);
+				GetOutdatedShaderTypes(OutdatedShaderTypes, OutdatedShaderPipelineTypes, OutdatedFactoryTypes);
 			}
 
 			if (OutdatedShaderPipelineTypes.Num() > 0 || OutdatedShaderTypes.Num() > 0 || OutdatedFactoryTypes.Num() > 0)
@@ -4023,7 +4059,7 @@ bool RecompileShaders(const TCHAR* Cmd, FOutputDevice& Ar)
 		}
 		else
 		{
-			TArray<FShaderType*> ShaderTypes = FShaderType::GetShaderTypesByFilename(*FlagStr);
+			TArray<const FShaderType*> ShaderTypes = FShaderType::GetShaderTypesByFilename(*FlagStr);
 			TArray<const FShaderPipelineType*> ShaderPipelineTypes = FShaderPipelineType::GetShaderPipelineTypesByFilename(*FlagStr);
 			if (ShaderTypes.Num() > 0 || ShaderPipelineTypes.Num() > 0)
 			{
@@ -4100,10 +4136,12 @@ FShader* FGlobalShaderTypeCompiler::FinishCompileShader(FGlobalShaderType* Shade
 	FShader* Shader = nullptr;
 	if (CurrentJob.bSucceeded)
 	{
+		EShaderPlatform Platform = CurrentJob.Input.Target.GetPlatform();
+		FGlobalShaderMapSection* Section = GGlobalShaderMap[Platform]->FindOrAddSection(ShaderType);
+
 		// Reuse an existing resource with the same key or create a new one based on the compile output
 		// This allows FShaders to share compiled bytecode and RHI shader references
-		TRefCountPtr<FShaderResource> Resource = FShaderResource::FindOrCreate(CurrentJob.Output, 0);
-		check(Resource);
+		const int32 ResourceIndex = Section->GetResourceBuilder().FindOrAddCode(CurrentJob.Output);
 
 		if (ShaderPipelineType && !ShaderPipelineType->ShouldOptimizeUnusedOutputs(CurrentJob.Input.Target.GetPlatform()))
 		{
@@ -4122,15 +4160,8 @@ FShader* FGlobalShaderTypeCompiler::FinishCompileShader(FGlobalShaderType* Shade
 			HashState.GetHash(&GlobalShaderMapHash.Hash[0]);
 		}
 
-		// Find a shader with the same key in memory
-		Shader = CurrentJob.ShaderType->FindShaderByKey(FShaderKey(GlobalShaderMapHash, ShaderPipelineType, nullptr, CurrentJob.PermutationId, CurrentJob.Input.Target.GetPlatform()));
-
-		// There was no shader with the same key so create a new one with the compile output, which will bind shader parameters
-		if (!Shader)
-		{
-			Shader = (*(ShaderType->ConstructCompiledRef))(FGlobalShaderType::CompiledShaderInitializerType(ShaderType, CurrentJob.PermutationId, CurrentJob.Output, MoveTemp(Resource), GlobalShaderMapHash, ShaderPipelineType, nullptr));
-			CurrentJob.Output.ParameterMap.VerifyBindingsAreComplete(ShaderType->GetName(), CurrentJob.Output.Target, CurrentJob.VFType);
-		}
+		Shader = ShaderType->ConstructCompiled(FGlobalShaderType::CompiledShaderInitializerType(ShaderType, CurrentJob.PermutationId, CurrentJob.Output, ResourceIndex, GlobalShaderMapHash, ShaderPipelineType, nullptr));
+		CurrentJob.Output.ParameterMap.VerifyBindingsAreComplete(ShaderType->GetName(), CurrentJob.Output.Target, CurrentJob.VFType);
 	}
 
 	if (CurrentJob.Output.Errors.Num() > 0)
@@ -4164,6 +4195,8 @@ FShader* FGlobalShaderTypeCompiler::FinishCompileShader(FGlobalShaderType* Shade
 */
 void VerifyGlobalShaders(EShaderPlatform Platform, bool bLoadedFromCacheFile)
 {
+	SCOPED_LOADTIMER(VerifyGlobalShaders);
+
 	check(IsInGameThread());
 	check(!FPlatformProperties::IsServerOnly());
 	check(GGlobalShaderMap[Platform]);
@@ -4171,7 +4204,7 @@ void VerifyGlobalShaders(EShaderPlatform Platform, bool bLoadedFromCacheFile)
 	UE_LOG(LogMaterial, Verbose, TEXT("Verifying Global Shaders for %s"), *LegacyShaderPlatformToShaderFormat(Platform).ToString());
 
 	// Ensure that the global shader map contains all global shader types.
-	TShaderMap<FGlobalShaderType>* GlobalShaderMap = GetGlobalShaderMap(Platform);
+	FGlobalShaderMap* GlobalShaderMap = GetGlobalShaderMap(Platform);
 	const bool bEmptyMap = GlobalShaderMap->IsEmpty();
 	if (bEmptyMap)
 	{
@@ -4206,8 +4239,9 @@ void VerifyGlobalShaders(EShaderPlatform Platform, bool bLoadedFromCacheFile)
 			{
 				continue;
 			}
-
-			if (FShader* Shader = GlobalShaderMap->GetShader(GlobalShaderType, PermutationId))
+			
+			TShaderRef<FShader> Shader = GlobalShaderMap->GetShader(GlobalShaderType, PermutationId);
+			if (Shader.IsValid())
 			{
 				// Validate the shader parameter structure early.
 				if (const FShaderParametersMetadata* ParameterStructMetadata = GlobalShaderType->GetRootParametersMetadata())
@@ -4253,7 +4287,7 @@ void VerifyGlobalShaders(EShaderPlatform Platform, bool bLoadedFromCacheFile)
 		const FShaderPipelineType* Pipeline = *ShaderPipelineIt;
 		if (Pipeline->IsGlobalTypePipeline())
 		{
-			if (!GlobalShaderMap->GetShaderPipeline(Pipeline))
+			if (!GlobalShaderMap->HasShaderPipeline(Pipeline))
 			{
 				auto& StageTypes = Pipeline->GetStages();
 				TArray<FGlobalShaderType*> ShaderStages;
@@ -4331,93 +4365,60 @@ void VerifyGlobalShaders(EShaderPlatform Platform, bool bLoadedFromCacheFile)
 	}
 }
 
+#include "Misc/PreLoadFile.h"
+#include "Serialization/LargeMemoryReader.h"
+static FPreLoadFile GGlobalShaderPreLoadFile(*(FString(TEXT("../../../Engine")) / TEXT("GlobalShaderCache-SF_") + FPlatformProperties::IniPlatformName() + TEXT(".bin")));
+
+static const ITargetPlatform* GGlobalShaderTargetPlatform[SP_NumPlatforms] = { nullptr };
+
 static FString GetGlobalShaderCacheFilename(EShaderPlatform Platform)
 {
 	return FString(TEXT("Engine")) / TEXT("GlobalShaderCache-") + LegacyShaderPlatformToShaderFormat(Platform).ToString() + TEXT(".bin");
 }
 
 /** Creates a string key for the derived data cache entry for the global shader map. */
-FString GetGlobalShaderMapKeyString(const FGlobalShaderMapId& ShaderMapId, EShaderPlatform Platform, TArray<FShaderTypeDependency> const& Dependencies)
+static FString GetGlobalShaderMapKeyString(const FGlobalShaderMapId& ShaderMapId, EShaderPlatform Platform, const ITargetPlatform* TargetPlatform, TArray<FShaderTypeDependency> const& Dependencies)
 {
 	FName Format = LegacyShaderPlatformToShaderFormat(Platform);
 	FString ShaderMapKeyString = Format.ToString() + TEXT("_") + FString(FString::FromInt(GetTargetPlatformManagerRef().ShaderFormatVersion(Format))) + TEXT("_");
 	ShaderMapAppendKeyString(Platform, ShaderMapKeyString);
-	ShaderMapId.AppendKeyString(ShaderMapKeyString, Dependencies);
+	ShaderMapId.AppendKeyString(ShaderMapKeyString, Dependencies, TargetPlatform);
 	return FDerivedDataCacheInterface::BuildCacheKey(TEXT("GSM"), GLOBALSHADERMAP_DERIVEDDATA_VER, *ShaderMapKeyString);
 }
 
-/** Serializes the global shader map to an archive. */
-static void SerializeGlobalShaders(FArchive& Ar, TShaderMap<FGlobalShaderType>* GlobalShaderMap, const TArray<TShaderTypePermutation<FShaderType>>* ShaderKeysToSave)
-{
-	check(IsInGameThread());
-
-	// Serialize the global shader map binary file tag.
-	static const uint32 ReferenceTag = 0x47534D42;
-	if (Ar.IsLoading())
-	{
-		// Initialize Tag to 0 as it won't be written to if the serialize fails (ie the global shader cache file is empty)
-		uint32 Tag = 0;
-		Ar << Tag;
-		checkf(Tag == ReferenceTag, TEXT("Global shader map binary file is missing GSMB tag."));
-	}
-	else
-	{
-		uint32 Tag = ReferenceTag;
-		Ar << Tag;
-	}
-
-	// Serialize the global shaders.
-	GlobalShaderMap->SerializeInline(Ar, true, false, false, ShaderKeysToSave);
-	// And now register them.
-	GlobalShaderMap->RegisterSerializedShaders(false);
-}
-
 /** Saves the platform's shader map to the DDC. */
-void SaveGlobalShaderMapToDerivedDataCache(EShaderPlatform Platform)
+static void SaveGlobalShaderMapToDerivedDataCache(EShaderPlatform Platform)
 {
 	// We've finally built the global shader map, so we can count the miss as we put it in the DDC.
 	COOK_STAT(auto Timer = GlobalShaderCookStats::UsageStats.TimeSyncWork());
 
-	FGlobalShaderMapId ShaderMapId(Platform);
+	const ITargetPlatform* TargetPlatform = GGlobalShaderTargetPlatform[Platform];
 
 	TArray<uint8> SaveData;
 
+	FGlobalShaderMapId ShaderMapId(Platform);
 	for (auto const& ShaderFilenameDependencies : ShaderMapId.GetShaderFilenameToDependeciesMap())
 	{
-		SaveData.Reset();
-
-		// Gather shader keys to save for this filename.
-		TArray<TShaderTypePermutation<FShaderType>> ShaderKeysToSave;
-		for (int32 ShaderIndex = 0; ShaderIndex < ShaderFilenameDependencies.Value.Num(); ++ShaderIndex)
+		FGlobalShaderMapSection* Section = GGlobalShaderMap[Platform]->FindSection(ShaderFilenameDependencies.Key);
+		if (Section)
 		{
-			const FShaderTypeDependency& ShaderTypeDependency = ShaderFilenameDependencies.Value[ShaderIndex];
-			FGlobalShaderType* GlobalShaderType = ShaderTypeDependency.ShaderType->GetGlobalShaderType();
-
-			if (GlobalShaderType)
-			{
-				bool bList = false;
-				for (int32 PermutationId = 0; PermutationId < GlobalShaderType->GetPermutationCount(); PermutationId++)
-				{
-					if (GlobalShaderType->ShouldCompilePermutation(Platform, PermutationId))
-					{
-						ShaderKeysToSave.Add(TShaderTypePermutation<FShaderType>(GlobalShaderType, PermutationId));
-					}
-				}
-			}
-		}
+			Section->FinalizeContent();
+			Section->ReleaseResourceBuilder();
 
-		FMemoryWriter Ar(SaveData, true);
-		SerializeGlobalShaders(Ar, GGlobalShaderMap[Platform], &ShaderKeysToSave);
+			SaveData.Reset();
+			FMemoryWriter Ar(SaveData, true);
+			Section->Serialize(Ar);
 
-		GetDerivedDataCacheRef().Put(*GetGlobalShaderMapKeyString(ShaderMapId, Platform, ShaderFilenameDependencies.Value), SaveData, TEXT("GlobalShaderMap"_SV));
-		COOK_STAT(Timer.AddMiss(SaveData.Num()));
+			GetDerivedDataCacheRef().Put(*GetGlobalShaderMapKeyString(ShaderMapId, Platform, TargetPlatform, ShaderFilenameDependencies.Value), SaveData, TEXT("GlobalShaderMap"_SV));
+			COOK_STAT(Timer.AddMiss(SaveData.Num()));
+		}
 	}
 }
 
 /** Saves the global shader map as a file for the target platform. */
 FString SaveGlobalShaderFile(EShaderPlatform Platform, FString SavePath, class ITargetPlatform* TargetPlatform)
 {
-	TShaderMap<FGlobalShaderType>* GlobalShaderMap = GetGlobalShaderMap(Platform);
+	FGlobalShaderMap* GlobalShaderMap = GetGlobalShaderMap(Platform);
 
 	// Wait until all global shaders are compiled
 	if (GShaderCompilingManager)
@@ -4426,12 +4427,14 @@ FString SaveGlobalShaderFile(EShaderPlatform Platform, FString SavePath, class I
 	}
 
 	TArray<uint8> GlobalShaderData;
-	FMemoryWriter MemoryWriter(GlobalShaderData, true);
-	if (TargetPlatform)
 	{
-		MemoryWriter.SetCookingTarget(TargetPlatform);
+		FMemoryWriter MemoryWriter(GlobalShaderData, true);
+		if (TargetPlatform)
+		{
+			MemoryWriter.SetCookingTarget(TargetPlatform);
+		}
+		GlobalShaderMap->SaveToGlobalArchive(MemoryWriter);
 	}
-	SerializeGlobalShaders(MemoryWriter, GlobalShaderMap, nullptr);
 
 	// make the final name
 	FString FullPath = SavePath / GetGlobalShaderCacheFilename(Platform);
@@ -4439,20 +4442,13 @@ FString SaveGlobalShaderFile(EShaderPlatform Platform, FString SavePath, class I
 	{
 		UE_LOG(LogShaders, Fatal, TEXT("Could not save global shader file to '%s'"), *FullPath);
 	}
+
 #if WITH_EDITOR
 	if (FShaderCodeLibrary::NeedsShaderStableKeys(Platform))
 	{
-		for (TLinkedList<FShaderType*>::TIterator ShaderTypeIt(FShaderType::GetTypeList()); ShaderTypeIt; ShaderTypeIt.Next())
-		{
-			FGlobalShaderType* GlobalShaderType = ShaderTypeIt->GetGlobalShaderType();
-			if (!GlobalShaderType)
-			{
-				continue;
-			}
-			GlobalShaderType->SaveShaderStableKeys(Platform);
-		}
+		GlobalShaderMap->SaveShaderStableKeys(Platform);
 	}
-#endif
+#endif // WITH_EDITOR
 	return FullPath;
 }
 
@@ -4471,7 +4467,7 @@ bool IsGlobalShaderMapComplete(const TCHAR* TypeNameSubstring)
 	{
 		EShaderPlatform Platform = (EShaderPlatform)i;
 
-		TShaderMap<FGlobalShaderType>* GlobalShaderMap = GGlobalShaderMap[Platform];
+		FGlobalShaderMap* GlobalShaderMap = GGlobalShaderMap[Platform];
 
 		if (GlobalShaderMap)
 		{
@@ -4515,7 +4511,7 @@ bool IsGlobalShaderMapComplete(const TCHAR* TypeNameSubstring)
 
 					if (NumStagesNeeded == Stages.Num())
 					{
-						if (!GlobalShaderMap->GetShaderPipeline(Pipeline))
+						if (!GlobalShaderMap->HasShaderPipeline(Pipeline))
 						{
 							return false;
 						}
@@ -4528,24 +4524,26 @@ bool IsGlobalShaderMapComplete(const TCHAR* TypeNameSubstring)
 	return true;
 }
 
-void CompileGlobalShaderMap(EShaderPlatform Platform, bool bRefreshShaderMap)
+void CompileGlobalShaderMap(EShaderPlatform Platform, const ITargetPlatform* TargetPlatform, bool bRefreshShaderMap)
 {
 	// No global shaders needed on dedicated server or clients that use NullRHI. Note that cook commandlet needs to have them, even if it is not allowed to render otherwise.
 	if (FPlatformProperties::IsServerOnly() || (!IsRunningCommandlet() && !FApp::CanEverRender()))
 	{
 		if (!GGlobalShaderMap[Platform])
 		{
-			GGlobalShaderMap[Platform] = new TShaderMap<FGlobalShaderType>(Platform);
+			GGlobalShaderMap[Platform] = new FGlobalShaderMap(Platform);
 		}
 		return;
 	}
 
-	if (bRefreshShaderMap)
+	if (bRefreshShaderMap || GGlobalShaderTargetPlatform[Platform] != TargetPlatform)
 	{
 		// delete the current global shader map
 		delete GGlobalShaderMap[Platform];
 		GGlobalShaderMap[Platform] = nullptr;
 
+		GGlobalShaderTargetPlatform[Platform] = TargetPlatform;
+
 		// make sure we look for updated shader source files
 		FlushShaderFileCache();
 	}
@@ -4563,7 +4561,7 @@ void CompileGlobalShaderMap(EShaderPlatform Platform, bool bRefreshShaderMap)
 		SlowTask.EnterProgressFrame(20);
 		VerifyShaderSourceFiles(Platform);
 
-		GGlobalShaderMap[Platform] = new TShaderMap<FGlobalShaderType>(Platform);
+		GGlobalShaderMap[Platform] = new FGlobalShaderMap(Platform);
 
 		bool bLoadedFromCacheFile = false;
 
@@ -4573,32 +4571,44 @@ void CompileGlobalShaderMap(EShaderPlatform Platform, bool bRefreshShaderMap)
 		{
 			SlowTask.EnterProgressFrame(50);
 
-			TArray<uint8> GlobalShaderData;
-			FString GlobalShaderCacheFilename = FPaths::GetRelativePathToRoot() / GetGlobalShaderCacheFilename(Platform);
-			FPaths::MakeStandardFilename(GlobalShaderCacheFilename);
-			bLoadedFromCacheFile = FFileHelper::LoadFileToArray(GlobalShaderData, *GlobalShaderCacheFilename, FILEREAD_Silent);
-
-			if (!bLoadedFromCacheFile)
+			// is the data already loaded?
+			int64 PreloadedSize;
+			void* PreloadedData = GGlobalShaderPreLoadFile.TakeOwnershipOfLoadedData(&PreloadedSize);
+			if (PreloadedData != nullptr)
 			{
-				// Handle this gracefully and exit.
-				FString SandboxPath = IFileManager::Get().ConvertToAbsolutePathForExternalAppForWrite(*GlobalShaderCacheFilename);
-				// This can be too early to localize in some situations.
-				const FText Message = FText::Format(NSLOCTEXT("Engine", "GlobalShaderCacheFileMissing", "The global shader cache file '{0}' is missing.\n\nYour application is built to load COOKED content. No COOKED content was found; This usually means you did not cook content for this build.\nIt also may indicate missing cooked data for a shader platform(e.g., OpenGL under Windows): Make sure your platform's packaging settings include this Targeted RHI.\n\nAlternatively build and run the UNCOOKED version instead."), FText::FromString(SandboxPath));
-				if (FPlatformProperties::SupportsWindowedMode())
-				{
-					UE_LOG(LogShaders, Error, TEXT("%s"), *Message.ToString());
-					FMessageDialog::Open(EAppMsgType::Ok, Message);
-					FPlatformMisc::RequestExit(false);
-					return;
-				}
-				else
+				FLargeMemoryReader MemoryReader((uint8*)PreloadedData, PreloadedSize, ELargeMemoryReaderFlags::TakeOwnership);
+				GGlobalShaderMap[Platform]->LoadFromGlobalArchive(MemoryReader);
+			}
+			else
+			{
+				TArray<uint8> GlobalShaderData;
+
+				FString GlobalShaderCacheFilename = FPaths::GetRelativePathToRoot() / GetGlobalShaderCacheFilename(Platform);
+				FPaths::MakeStandardFilename(GlobalShaderCacheFilename);
+				bLoadedFromCacheFile = FFileHelper::LoadFileToArray(GlobalShaderData, *GlobalShaderCacheFilename, FILEREAD_Silent);
+
+				if (!bLoadedFromCacheFile)
 				{
-					UE_LOG(LogShaders, Fatal, TEXT("%s"), *Message.ToString());
+					// Handle this gracefully and exit.
+					FString SandboxPath = IFileManager::Get().ConvertToAbsolutePathForExternalAppForWrite(*GlobalShaderCacheFilename);
+					// This can be too early to localize in some situations.
+					const FText Message = FText::Format(NSLOCTEXT("Engine", "GlobalShaderCacheFileMissing", "The global shader cache file '{0}' is missing.\n\nYour application is built to load COOKED content. No COOKED content was found; This usually means you did not cook content for this build.\nIt also may indicate missing cooked data for a shader platform(e.g., OpenGL under Windows): Make sure your platform's packaging settings include this Targeted RHI.\n\nAlternatively build and run the UNCOOKED version instead."), FText::FromString(SandboxPath));
+					if (FPlatformProperties::SupportsWindowedMode())
+					{
+						UE_LOG(LogShaders, Error, TEXT("%s"), *Message.ToString());
+						FMessageDialog::Open(EAppMsgType::Ok, Message);
+						FPlatformMisc::RequestExit(false);
+						return;
+					}
+					else
+					{
+						UE_LOG(LogShaders, Fatal, TEXT("%s"), *Message.ToString());
+					}
 				}
-			}
 
-			FMemoryReader MemoryReader(GlobalShaderData);
-			SerializeGlobalShaders(MemoryReader, GGlobalShaderMap[Platform], nullptr);
+				FMemoryReader MemoryReader(GlobalShaderData);
+				GGlobalShaderMap[Platform]->LoadFromGlobalArchive(MemoryReader);
+			}
 		}
 		// Uncooked platform
 		else
@@ -4618,7 +4628,7 @@ void CompileGlobalShaderMap(EShaderPlatform Platform, bool bRefreshShaderMap)
 			{
 				SlowTask.EnterProgressFrame(ProgressStep);
 
-				const FString DataKey = GetGlobalShaderMapKeyString(ShaderMapId, Platform, ShaderFilenameDependencies.Value);
+				const FString DataKey = GetGlobalShaderMapKeyString(ShaderMapId, Platform, TargetPlatform, ShaderFilenameDependencies.Value);
 
 				AsyncDDCRequestHandles[HandleIndex] = GetDerivedDataCacheRef().GetAsynchronous(*DataKey, TEXT("GlobalShaderMap"_SV));
 
@@ -4641,10 +4651,8 @@ void CompileGlobalShaderMap(EShaderPlatform Platform, bool bRefreshShaderMap)
 				if (GetDerivedDataCacheRef().GetAsynchronousResults(AsyncDDCRequestHandles[HandleIndex], CachedData))
 				{
 					COOK_STAT(Timer.AddHit(CachedData.Num()));
-					FMemoryReader Ar(CachedData, true);
-
-					// Deserialize from the cached data
-					SerializeGlobalShaders(Ar, GGlobalShaderMap[Platform], nullptr);
+					FMemoryReader MemoryReader(CachedData);
+					GGlobalShaderMap[Platform]->AddSection(FGlobalShaderMapSection::CreateFromArchive(MemoryReader));
 				}
 				else
 				{
@@ -4661,22 +4669,20 @@ void CompileGlobalShaderMap(EShaderPlatform Platform, bool bRefreshShaderMap)
 
 		if (GCreateShadersOnLoad && Platform == GMaxRHIShaderPlatform)
 		{
-			for (auto& Pair : GGlobalShaderMap[Platform]->GetShaders())
-			{
-				FShader* Shader = Pair.Value;
-				if (Shader)
-				{
-					Shader->BeginInitializeResources();
-				}
-			}
+			GGlobalShaderMap[Platform]->BeginCreateAllShaders();
 		}
 	}
 }
 
+void CompileGlobalShaderMap(EShaderPlatform Platform, bool bRefreshShaderMap)
+{
+	CompileGlobalShaderMap(Platform, nullptr, bRefreshShaderMap);
+}
+
 void CompileGlobalShaderMap(ERHIFeatureLevel::Type InFeatureLevel, bool bRefreshShaderMap)
 {
 	EShaderPlatform Platform = GShaderPlatformForFeatureLevel[InFeatureLevel];
-	CompileGlobalShaderMap(Platform, bRefreshShaderMap);
+	CompileGlobalShaderMap(Platform, nullptr, bRefreshShaderMap);
 }
 
 void CompileGlobalShaderMap(bool bRefreshShaderMap)
@@ -4701,15 +4707,14 @@ bool RecompileChangedShadersForPlatform(const FString& PlatformName)
 
 
 	// figure out which shaders are out of date
-	TArray<FShaderType*> OutdatedShaderTypes;
+	TArray<const FShaderType*> OutdatedShaderTypes;
 	TArray<const FVertexFactoryType*> OutdatedFactoryTypes;
 	TArray<const FShaderPipelineType*> OutdatedShaderPipelineTypes;
 
 	// Pick up new changes to shader files
 	FlushShaderFileCache();
 
-	FShaderType::GetOutdatedTypes(OutdatedShaderTypes, OutdatedFactoryTypes);
-	FShaderPipelineType::GetOutdatedTypes(OutdatedShaderTypes, OutdatedShaderPipelineTypes, OutdatedFactoryTypes);
+	GetOutdatedShaderTypes(OutdatedShaderTypes, OutdatedShaderPipelineTypes, OutdatedFactoryTypes);
 	UE_LOG(LogShaders, Display, TEXT("We found %d out of date shader types, %d outdated pipeline types, and %d out of date VF types!"), OutdatedShaderTypes.Num(), OutdatedShaderPipelineTypes.Num(), OutdatedFactoryTypes.Num());
 
 	for (int32 FormatIndex = 0; FormatIndex < DesiredShaderFormats.Num(); FormatIndex++)
@@ -4778,7 +4783,7 @@ void RecompileShadersForRemote(
 	UE_LOG(LogShaders, Display, TEXT("  Done!"))
 
 	// figure out which shaders are out of date
-	TArray<FShaderType*> OutdatedShaderTypes;
+	TArray<const FShaderType*> OutdatedShaderTypes;
 	TArray<const FVertexFactoryType*> OutdatedFactoryTypes;
 	TArray<const FShaderPipelineType*> OutdatedShaderPipelineTypes;
 
@@ -4787,8 +4792,7 @@ void RecompileShadersForRemote(
 
 	if (bCompileChangedShaders)
 	{
-		FShaderType::GetOutdatedTypes(OutdatedShaderTypes, OutdatedFactoryTypes);
-		FShaderPipelineType::GetOutdatedTypes(OutdatedShaderTypes, OutdatedShaderPipelineTypes, OutdatedFactoryTypes);
+		GetOutdatedShaderTypes(OutdatedShaderTypes, OutdatedShaderPipelineTypes, OutdatedFactoryTypes);
 		UE_LOG(LogShaders, Display, TEXT("We found %d out of date shader types, %d outdated pipeline types, and %d out of date VF types!"), OutdatedShaderTypes.Num(), OutdatedShaderPipelineTypes.Num(), OutdatedFactoryTypes.Num());
 	}
 
@@ -4810,7 +4814,7 @@ void RecompileShadersForRemote(
 				if (bCompileChangedShaders)
 				{
 					// Kick off global shader recompiles
-					BeginRecompileGlobalShaders(OutdatedShaderTypes, OutdatedShaderPipelineTypes, ShaderPlatform);
+					BeginRecompileGlobalShaders(OutdatedShaderTypes, OutdatedShaderPipelineTypes, ShaderPlatform, TargetPlatform);
 
 					// Block on global shaders
 					FinishRecompileGlobalShaders();
@@ -4828,8 +4832,9 @@ void RecompileShadersForRemote(
 
 					// pull the serialized resource ids into an array of resources
 					TArray<FShaderResourceId> ClientResourceIds;
-					FMemoryReader MemReader(SerializedShaderResources, true);
-					MemReader << ClientResourceIds;
+					check(false); // TODO
+					//FMemoryReader MemReader(SerializedShaderResources, true);
+					//MemReader << ClientResourceIds;
 
 					// save out the shader map to the byte array
 					FMaterialShaderMap::SaveForRemoteRecompile(Ar, CompiledShaderMaps, ClientResourceIds);
@@ -4853,7 +4858,7 @@ void RecompileShadersForRemote(
 	}
 }
 
-void BeginRecompileGlobalShaders(const TArray<FShaderType*>& OutdatedShaderTypes, const TArray<const FShaderPipelineType*>& OutdatedShaderPipelineTypes, EShaderPlatform ShaderPlatform)
+void BeginRecompileGlobalShaders(const TArray<const FShaderType*>& OutdatedShaderTypes, const TArray<const FShaderPipelineType*>& OutdatedShaderPipelineTypes, EShaderPlatform ShaderPlatform, const ITargetPlatform* TargetPlatform)
 {
 	if (!FPlatformProperties::RequiresCookedData())
 	{
@@ -4861,15 +4866,15 @@ void BeginRecompileGlobalShaders(const TArray<FShaderType*>& OutdatedShaderTypes
 		FlushRenderingCommands();
 
 		// Calling CompileGlobalShaderMap will force starting the compile jobs if the map is empty (by calling VerifyGlobalShaders)
-		CompileGlobalShaderMap(ShaderPlatform, false);
-		TShaderMap<FGlobalShaderType>* GlobalShaderMap = GetGlobalShaderMap(ShaderPlatform);
+		CompileGlobalShaderMap(ShaderPlatform, TargetPlatform, false);
+		FGlobalShaderMap* GlobalShaderMap = GetGlobalShaderMap(ShaderPlatform);
 
 		// Now check if there is any work to be done wrt outdates types
 		if (OutdatedShaderTypes.Num() > 0 || OutdatedShaderPipelineTypes.Num() > 0)
 		{
 			for (int32 TypeIndex = 0; TypeIndex < OutdatedShaderTypes.Num(); TypeIndex++)
 			{
-				FGlobalShaderType* CurrentGlobalShaderType = OutdatedShaderTypes[TypeIndex]->GetGlobalShaderType();
+				const FGlobalShaderType* CurrentGlobalShaderType = OutdatedShaderTypes[TypeIndex]->GetGlobalShaderType();
 				if (CurrentGlobalShaderType)
 				{
 					UE_LOG(LogShaders, Log, TEXT("Flushing Global Shader %s"), CurrentGlobalShaderType->GetName());
@@ -4912,7 +4917,7 @@ static inline FShader* ProcessCompiledJob(FShaderCompileJob* SingleJob, const FS
 		EShaderPlatform Platform = (EShaderPlatform)SingleJob->Input.Target.Platform;
 		if (!Pipeline || !Pipeline->ShouldOptimizeUnusedOutputs(Platform))
 		{
-			GGlobalShaderMap[Platform]->AddShader(GlobalShaderType, SingleJob->PermutationId, Shader);
+			Shader = GGlobalShaderMap[Platform]->FindOrAddShader(GlobalShaderType, Shader);
 			// Add this shared pipeline to the list
 			if (!Pipeline)
 			{
@@ -4970,8 +4975,8 @@ void ProcessCompiledGlobalShaders(const TArray<TSharedRef<FShaderCommonCompileJo
 			if (ShaderPipeline)
 			{
 				EShaderPlatform Platform = (EShaderPlatform)PipelineJob->StageJobs[0]->GetSingleShaderJob()->Input.Target.Platform;
-				check(ShaderPipeline && !GGlobalShaderMap[Platform]->HasShaderPipeline(ShaderPipeline->PipelineType));
-				GGlobalShaderMap[Platform]->AddShaderPipeline(PipelineJob->ShaderPipeline, ShaderPipeline);
+				check(ShaderPipeline && !GGlobalShaderMap[Platform]->HasShaderPipeline(PipelineJob->ShaderPipeline));
+				GGlobalShaderMap[Platform]->FindOrAddShaderPipeline(PipelineJob->ShaderPipeline, ShaderPipeline);
 			}
 		}
 	}
@@ -4994,9 +4999,9 @@ void ProcessCompiledGlobalShaders(const TArray<TSharedRef<FShaderCommonCompileJo
 						FGlobalShaderType* GlobalShaderType = ((FShaderType*)(StageTypes[Index]))->GetGlobalShaderType();
 						if (GlobalShaderType->ShouldCompilePermutation(Platform, kUniqueShaderPermutationId))
 						{
-							FShader* Shader = GlobalShaderMap->GetShader(GlobalShaderType, kUniqueShaderPermutationId);
-							check(Shader);
-							ShaderStages.Add(Shader);
+							TShaderRef<FShader> Shader = GlobalShaderMap->GetShader(GlobalShaderType, kUniqueShaderPermutationId);
+							check(Shader.IsValid());
+							ShaderStages.Add(Shader.GetShader());
 						}
 						else
 						{
@@ -5006,7 +5011,7 @@ void ProcessCompiledGlobalShaders(const TArray<TSharedRef<FShaderCommonCompileJo
 
 					checkf(StageTypes.Num() == ShaderStages.Num(), TEXT("Internal Error adding Global ShaderPipeline %s"), ShaderPipelineType->GetName());
 					FShaderPipeline* ShaderPipeline = new FShaderPipeline(ShaderPipelineType, ShaderStages);
-					GlobalShaderMap->AddShaderPipeline(ShaderPipelineType, ShaderPipeline);
+					GlobalShaderMap->FindOrAddShaderPipeline(ShaderPipelineType, ShaderPipeline);
 				}
 			}
 		}
diff --git a/Engine/Source/Runtime/Engine/Private/ShaderDerivedDataVersion.h b/Engine/Source/Runtime/Engine/Private/ShaderDerivedDataVersion.h
--- a/Engine/Source/Runtime/Engine/Private/ShaderDerivedDataVersion.h
+++ b/Engine/Source/Runtime/Engine/Private/ShaderDerivedDataVersion.h
@@ -10,5 +10,5 @@
 
 // In case of merge conflicts with DDC versions, you MUST generate a new GUID and set this new
 // guid as version
-#define GLOBALSHADERMAP_DERIVEDDATA_VER			TEXT("A85F9B4CA9264C3CBE27C6356D371896")
-#define MATERIALSHADERMAP_DERIVEDDATA_VER		TEXT("DBB562AD91CE4620B207619E5C81A619")
+#define GLOBALSHADERMAP_DERIVEDDATA_VER			TEXT("{4939F6BE-B9D0-45ED-A8D2-91EB023E7048}")
+#define MATERIALSHADERMAP_DERIVEDDATA_VER		TEXT("{4133F162-9303-450D-8F58-8FA16D3204AB}")
\ No newline at end of file
diff --git a/Engine/Source/Runtime/Engine/Private/SimpleElementShaders.cpp b/Engine/Source/Runtime/Engine/Private/SimpleElementShaders.cpp
--- a/Engine/Source/Runtime/Engine/Private/SimpleElementShaders.cpp
+++ b/Engine/Source/Runtime/Engine/Private/SimpleElementShaders.cpp
@@ -21,16 +21,16 @@ FSimpleElementVS::FSimpleElementVS(const ShaderMetaType::CompiledShaderInitializ
 
 void FSimpleElementVS::SetParameters(FRHICommandList& RHICmdList, const FMatrix& TransformValue, bool bSwitchVerticalAxis)
 {
-	SetShaderValue(RHICmdList, GetVertexShader(),Transform,TransformValue);
-	SetShaderValue(RHICmdList, GetVertexShader(),SwitchVerticalAxis, bSwitchVerticalAxis ? -1.0f : 1.0f);
+	SetShaderValue(RHICmdList, RHICmdList.GetBoundVertexShader(), Transform,TransformValue);
+	SetShaderValue(RHICmdList, RHICmdList.GetBoundVertexShader(), SwitchVerticalAxis, bSwitchVerticalAxis ? -1.0f : 1.0f);
 }
 
-bool FSimpleElementVS::Serialize(FArchive& Ar)
+/*bool FSimpleElementVS::Serialize(FArchive& Ar)
 {
 	bool bShaderHasOutdatedParameters = FGlobalShader::Serialize(Ar);
 	Ar << Transform << SwitchVerticalAxis;
 	return bShaderHasOutdatedParameters;
-}
+}*/
 
 void FSimpleElementVS::ModifyCompilationEnvironment(const FGlobalShaderPermutationParameters& Parameters, FShaderCompilerEnvironment& OutEnvironment)
 {
@@ -55,24 +55,24 @@ void FSimpleElementPS::SetEditorCompositingParameters(FRHICommandList& RHICmdLis
 {
 	if( View )
 	{
-		FGlobalShader::SetParameters<FViewUniformShaderParameters>(RHICmdList, GetPixelShader(), View->ViewUniformBuffer );
+		FGlobalShader::SetParameters<FViewUniformShaderParameters>(RHICmdList, RHICmdList.GetBoundPixelShader(), View->ViewUniformBuffer );
 	}
 	else
 	{
 		// Unset the view uniform buffers since we don't have a view
-		SetUniformBufferParameter(RHICmdList, GetPixelShader(), GetUniformBufferParameter<FViewUniformShaderParameters>(), NULL);
+		SetUniformBufferParameter(RHICmdList, RHICmdList.GetBoundPixelShader(), GetUniformBufferParameter<FViewUniformShaderParameters>(), NULL);
 	}
 }
 
 void FSimpleElementPS::SetParameters(FRHICommandList& RHICmdList, const FTexture* TextureValue)
 {
-	SetTextureParameter(RHICmdList, GetPixelShader(),InTexture,InTextureSampler,TextureValue);
+	SetTextureParameter(RHICmdList, RHICmdList.GetBoundPixelShader(),InTexture,InTextureSampler,TextureValue);
 	
-	SetShaderValue(RHICmdList, GetPixelShader(),TextureComponentReplicate,TextureValue->bGreyScaleFormat ? FLinearColor(1,0,0,0) : FLinearColor(0,0,0,0));
-	SetShaderValue(RHICmdList, GetPixelShader(),TextureComponentReplicateAlpha,TextureValue->bGreyScaleFormat ? FLinearColor(1,0,0,0) : FLinearColor(0,0,0,1));
+	SetShaderValue(RHICmdList, RHICmdList.GetBoundPixelShader(),TextureComponentReplicate,TextureValue->bGreyScaleFormat ? FLinearColor(1,0,0,0) : FLinearColor(0,0,0,0));
+	SetShaderValue(RHICmdList, RHICmdList.GetBoundPixelShader(),TextureComponentReplicateAlpha,TextureValue->bGreyScaleFormat ? FLinearColor(1,0,0,0) : FLinearColor(0,0,0,1));
 }
 
-bool FSimpleElementPS::Serialize(FArchive& Ar)
+/*bool FSimpleElementPS::Serialize(FArchive& Ar)
 {
 	bool bShaderHasOutdatedParameters = FGlobalShader::Serialize(Ar);
 	Ar << InTexture;
@@ -82,7 +82,7 @@ bool FSimpleElementPS::Serialize(FArchive& Ar)
 	Ar << EditorCompositeDepthTestParameter;
 	Ar << ScreenToPixel;
 	return bShaderHasOutdatedParameters;
-}
+}*/
 
 FSimpleElementAlphaOnlyPS::FSimpleElementAlphaOnlyPS(const ShaderMetaType::CompiledShaderInitializerType& Initializer) :
 	FSimpleElementPS(Initializer)
@@ -98,15 +98,15 @@ FSimpleElementGammaBasePS::FSimpleElementGammaBasePS(const ShaderMetaType::Compi
 void FSimpleElementGammaBasePS::SetParameters(FRHICommandList& RHICmdList, const FTexture* Texture, float GammaValue, ESimpleElementBlendMode BlendMode)
 {
 	FSimpleElementPS::SetParameters(RHICmdList, Texture);
-	SetShaderValue(RHICmdList, GetPixelShader(),Gamma,GammaValue);
+	SetShaderValue(RHICmdList, RHICmdList.GetBoundPixelShader(),Gamma,GammaValue);
 }
 
-bool FSimpleElementGammaBasePS::Serialize(FArchive& Ar)
+/*bool FSimpleElementGammaBasePS::Serialize(FArchive& Ar)
 {
 	bool bShaderHasOutdatedParameters = FSimpleElementPS::Serialize(Ar);
 	Ar << Gamma;
 	return bShaderHasOutdatedParameters;
-}
+}*/
 
 FSimpleElementMaskedGammaBasePS::FSimpleElementMaskedGammaBasePS(const ShaderMetaType::CompiledShaderInitializerType& Initializer) :
 	FSimpleElementGammaBasePS(Initializer)
@@ -117,15 +117,15 @@ FSimpleElementMaskedGammaBasePS::FSimpleElementMaskedGammaBasePS(const ShaderMet
 void FSimpleElementMaskedGammaBasePS::SetParameters(FRHICommandList& RHICmdList, const FTexture* Texture, float InGamma, float ClipRefValue, ESimpleElementBlendMode BlendMode)
 {
 	FSimpleElementGammaBasePS::SetParameters(RHICmdList, Texture,InGamma,BlendMode);
-	SetShaderValue(RHICmdList, GetPixelShader(),ClipRef,ClipRefValue);
+	SetShaderValue(RHICmdList, RHICmdList.GetBoundPixelShader(),ClipRef,ClipRefValue);
 }
 
-bool FSimpleElementMaskedGammaBasePS::Serialize(FArchive& Ar)
+/*bool FSimpleElementMaskedGammaBasePS::Serialize(FArchive& Ar)
 {
 	bool bShaderHasOutdatedParameters = FSimpleElementGammaBasePS::Serialize(Ar);
 	Ar << ClipRef;
 	return bShaderHasOutdatedParameters;
-}
+}*/
 
 /**
 * Constructor
@@ -174,21 +174,21 @@ void FSimpleElementDistanceFieldGammaPS::SetParameters(
 	)
 {
 	FSimpleElementMaskedGammaBasePS::SetParameters(RHICmdList, Texture,InGamma,InClipRef,BlendMode);
-	SetShaderValue(RHICmdList, GetPixelShader(),SmoothWidth,SmoothWidthValue);
+	SetShaderValue(RHICmdList, RHICmdList.GetBoundPixelShader(),SmoothWidth,SmoothWidthValue);		
 	const uint32 bEnableShadowValueUInt = (bEnableShadowValue ? 1 : 0);
-	SetShaderValue(RHICmdList, GetPixelShader(),EnableShadow,bEnableShadowValueUInt);
+	SetShaderValue(RHICmdList, RHICmdList.GetBoundPixelShader(),EnableShadow,bEnableShadowValueUInt);
 	if (bEnableShadowValue)
 	{
-		SetShaderValue(RHICmdList, GetPixelShader(),ShadowDirection,ShadowDirectionValue);
-		SetShaderValue(RHICmdList, GetPixelShader(),ShadowColor,ShadowColorValue);
-		SetShaderValue(RHICmdList, GetPixelShader(),ShadowSmoothWidth,ShadowSmoothWidthValue);
+		SetShaderValue(RHICmdList, RHICmdList.GetBoundPixelShader(),ShadowDirection,ShadowDirectionValue);
+		SetShaderValue(RHICmdList, RHICmdList.GetBoundPixelShader(),ShadowColor,ShadowColorValue);
+		SetShaderValue(RHICmdList, RHICmdList.GetBoundPixelShader(),ShadowSmoothWidth,ShadowSmoothWidthValue);
 	}
-	SetShaderValue(RHICmdList, GetPixelShader(),EnableGlow,GlowInfo.bEnableGlow);
+	SetShaderValue(RHICmdList, RHICmdList.GetBoundPixelShader(),EnableGlow,GlowInfo.bEnableGlow);
 	if (GlowInfo.bEnableGlow)
 	{
-		SetShaderValue(RHICmdList, GetPixelShader(),GlowColor,GlowInfo.GlowColor);
-		SetShaderValue(RHICmdList, GetPixelShader(),GlowOuterRadius,GlowInfo.GlowOuterRadius);
-		SetShaderValue(RHICmdList, GetPixelShader(),GlowInnerRadius,GlowInfo.GlowInnerRadius);
+		SetShaderValue(RHICmdList, RHICmdList.GetBoundPixelShader(),GlowColor,GlowInfo.GlowColor);
+		SetShaderValue(RHICmdList, RHICmdList.GetBoundPixelShader(),GlowOuterRadius,GlowInfo.GlowOuterRadius);
+		SetShaderValue(RHICmdList, RHICmdList.GetBoundPixelShader(),GlowInnerRadius,GlowInfo.GlowInnerRadius);
 	}
 
 	// This shader does not use editor compositing
@@ -201,7 +201,7 @@ void FSimpleElementDistanceFieldGammaPS::SetParameters(
 * @param Ar - archive to serialize to
 * @return true if any of the parameters were outdated
 */
-bool FSimpleElementDistanceFieldGammaPS::Serialize(FArchive& Ar)
+/*bool FSimpleElementDistanceFieldGammaPS::Serialize(FArchive& Ar)
 {
 	bool bShaderHasOutdatedParameters = FSimpleElementMaskedGammaBasePS::Serialize(Ar);
 	Ar << SmoothWidth;
@@ -214,7 +214,7 @@ bool FSimpleElementDistanceFieldGammaPS::Serialize(FArchive& Ar)
 	Ar << GlowOuterRadius;
 	Ar << GlowInnerRadius;
 	return bShaderHasOutdatedParameters;
-}
+}*/
 	
 FSimpleElementHitProxyPS::FSimpleElementHitProxyPS(const ShaderMetaType::CompiledShaderInitializerType& Initializer):
 	FGlobalShader(Initializer)
@@ -225,16 +225,16 @@ FSimpleElementHitProxyPS::FSimpleElementHitProxyPS(const ShaderMetaType::Compile
 
 void FSimpleElementHitProxyPS::SetParameters(FRHICommandList& RHICmdList, const FTexture* TextureValue)
 {
-	SetTextureParameter(RHICmdList, GetPixelShader(),InTexture,InTextureSampler,TextureValue);
+	SetTextureParameter(RHICmdList, RHICmdList.GetBoundPixelShader(),InTexture,InTextureSampler,TextureValue);
 }
 
-bool FSimpleElementHitProxyPS::Serialize(FArchive& Ar)
+/*bool FSimpleElementHitProxyPS::Serialize(FArchive& Ar)
 {
 	bool bShaderHasOutdatedParameters = FGlobalShader::Serialize(Ar);
 	Ar << InTexture;
 	Ar << InTextureSampler;
 	return bShaderHasOutdatedParameters;
-}
+}*/
 
 
 FSimpleElementColorChannelMaskPS::FSimpleElementColorChannelMaskPS(const ShaderMetaType::CompiledShaderInitializerType& Initializer)
@@ -255,12 +255,12 @@ FSimpleElementColorChannelMaskPS::FSimpleElementColorChannelMaskPS(const ShaderM
 */
 void FSimpleElementColorChannelMaskPS::SetParameters(FRHICommandList& RHICmdList, const FTexture* TextureValue, const FMatrix& ColorWeightsValue, float GammaValue)
 {
-	SetTextureParameter(RHICmdList, GetPixelShader(),InTexture,InTextureSampler,TextureValue);
-	SetShaderValue(RHICmdList, GetPixelShader(),ColorWeights,ColorWeightsValue);
-	SetShaderValue(RHICmdList, GetPixelShader(),Gamma,GammaValue);
+	SetTextureParameter(RHICmdList, RHICmdList.GetBoundPixelShader(),InTexture,InTextureSampler,TextureValue);
+	SetShaderValue(RHICmdList, RHICmdList.GetBoundPixelShader(),ColorWeights,ColorWeightsValue);
+	SetShaderValue(RHICmdList, RHICmdList.GetBoundPixelShader(),Gamma,GammaValue);
 }
 
-bool FSimpleElementColorChannelMaskPS::Serialize(FArchive& Ar)
+/*bool FSimpleElementColorChannelMaskPS::Serialize(FArchive& Ar)
 {
 	bool bShaderHasOutdatedParameters = FGlobalShader::Serialize(Ar);
 	Ar << InTexture;
@@ -268,12 +268,15 @@ bool FSimpleElementColorChannelMaskPS::Serialize(FArchive& Ar)
 	Ar << ColorWeights;
 	Ar << Gamma;
 	return bShaderHasOutdatedParameters;
-}
+}*/
 
 /*------------------------------------------------------------------------------
 	Shader implementations.
 ------------------------------------------------------------------------------*/
 
+IMPLEMENT_TYPE_LAYOUT(FSimpleElementGammaBasePS);
+IMPLEMENT_TYPE_LAYOUT(FSimpleElementMaskedGammaBasePS);
+
 IMPLEMENT_SHADER_TYPE(,FSimpleElementVS,TEXT("/Engine/Private/SimpleElementVertexShader.usf"),TEXT("Main"),SF_Vertex);
 IMPLEMENT_SHADER_TYPE(,FSimpleElementPS, TEXT("/Engine/Private/SimpleElementPixelShader.usf"), TEXT("Main"), SF_Pixel);
 IMPLEMENT_SHADER_TYPE(,FSimpleElementAlphaOnlyPS, TEXT("/Engine/Private/SimpleElementPixelShader.usf"), TEXT("AlphaOnlyMain"), SF_Pixel);
diff --git a/Engine/Source/Runtime/Engine/Private/SkeletalMesh.cpp b/Engine/Source/Runtime/Engine/Private/SkeletalMesh.cpp
--- a/Engine/Source/Runtime/Engine/Private/SkeletalMesh.cpp
+++ b/Engine/Source/Runtime/Engine/Private/SkeletalMesh.cpp
@@ -5477,7 +5477,7 @@ FSkinnedMeshComponentRecreateRenderStateContext::~FSkinnedMeshComponentRecreateR
 
 		if (Component->IsRegistered() && !Component->IsRenderStateCreated())
 		{
-			Component->CreateRenderState_Concurrent();
+			Component->CreateRenderState_Concurrent(nullptr);
 		}
 	}
 }
diff --git a/Engine/Source/Runtime/Engine/Private/SkeletalMeshComponentPhysics.cpp b/Engine/Source/Runtime/Engine/Private/SkeletalMeshComponentPhysics.cpp
--- a/Engine/Source/Runtime/Engine/Private/SkeletalMeshComponentPhysics.cpp
+++ b/Engine/Source/Runtime/Engine/Private/SkeletalMeshComponentPhysics.cpp
@@ -743,7 +743,7 @@ void USkeletalMeshComponent::InstantiatePhysicsAsset_Internal(const UPhysicsAsse
 			// Create physics body instance.
 			FTransform BoneTransform = BoneTransformGetter(BoneIndex);
 
-			FBodyInstance::FInitBodySpawnParams SpawnParams(OwningComponent);
+			FInitBodySpawnParams SpawnParams(OwningComponent);
 
 			if(OwningComponent == nullptr)
 			{
diff --git a/Engine/Source/Runtime/Engine/Private/SkeletalMeshMerge.cpp b/Engine/Source/Runtime/Engine/Private/SkeletalMeshMerge.cpp
--- a/Engine/Source/Runtime/Engine/Private/SkeletalMeshMerge.cpp
+++ b/Engine/Source/Runtime/Engine/Private/SkeletalMeshMerge.cpp
@@ -526,7 +526,7 @@ void FSkeletalMeshMerge::GenerateLODModel( int32 LODIdx )
 			// keep track of the lowest LOD displayfactor and hysteresis
 			MergeLODInfo.ScreenSize.Default = FMath::Min<float>(MergeLODInfo.ScreenSize.Default, SrcLODInfo.ScreenSize.Default);
 #if WITH_EDITORONLY_DATA
-			for(const TTuple<FName, float>& PerPlatform : SrcLODInfo.ScreenSize.PerPlatform)
+			for(const TPair<FName, float>& PerPlatform : SrcLODInfo.ScreenSize.PerPlatform)
 			{
 				float* Value = MergeLODInfo.ScreenSize.PerPlatform.Find(PerPlatform.Key);
 				if(Value)
diff --git a/Engine/Source/Runtime/Engine/Private/SkeletalRenderGPUSkin.cpp b/Engine/Source/Runtime/Engine/Private/SkeletalRenderGPUSkin.cpp
--- a/Engine/Source/Runtime/Engine/Private/SkeletalRenderGPUSkin.cpp
+++ b/Engine/Source/Runtime/Engine/Private/SkeletalRenderGPUSkin.cpp
@@ -665,8 +665,7 @@ TArray<float> FSkeletalMeshObjectGPUSkin::FSkeletalMeshObjectLOD::MorphAccumulat
 
 void FGPUMorphUpdateCS::SetParameters(FRHICommandList& RHICmdList, const FVector4& LocalScale, const FMorphTargetVertexInfoBuffers& MorphTargetVertexInfoBuffers, FMorphVertexBuffer& MorphVertexBuffer)
 {
-	FRHIComputeShader* CS = GetComputeShader();
-	RHICmdList.SetComputeShader(CS);
+	FRHIComputeShader* CS = RHICmdList.GetBoundComputeShader();
 
 	SetUAVParameter(RHICmdList, CS, MorphVertexBufferParameter, MorphVertexBuffer.GetUAV());
 
@@ -680,7 +679,7 @@ static const uint32 GMorphTargetDispatchBatchSize = 128;
 
 void FGPUMorphUpdateCS::SetOffsetAndSize(FRHICommandList& RHICmdList, uint32 StartIndex, uint32 EndIndexPlusOne, const FMorphTargetVertexInfoBuffers& MorphTargetVertexInfoBuffers, const TArray<float>& MorphTargetWeights)
 {
-	FComputeShaderRHIRef CS = GetComputeShader();
+	FRHIComputeShader* CS = RHICmdList.GetBoundComputeShader();
 
 	uint32 ThreadOffsets[GMorphTargetDispatchBatchSize];
 	float Weights[GMorphTargetDispatchBatchSize];
@@ -718,7 +717,7 @@ void FGPUMorphUpdateCS::Dispatch(FRHICommandList& RHICmdList, uint32 Size)
 
 void FGPUMorphUpdateCS::EndAllDispatches(FRHICommandList& RHICmdList)
 {
-	FComputeShaderRHIRef CS = GetComputeShader();
+	FRHIComputeShader* CS = RHICmdList.GetBoundComputeShader();
 	SetUAVParameter(RHICmdList, CS, MorphVertexBufferParameter, nullptr);
 }
 
@@ -726,9 +725,7 @@ IMPLEMENT_SHADER_TYPE(, FGPUMorphUpdateCS, TEXT("/Engine/Private/MorphTargets.us
 
 void FGPUMorphNormalizeCS::SetParameters(FRHICommandList& RHICmdList, const FVector4& InvLocalScale, const FMorphTargetVertexInfoBuffers& MorphTargetVertexInfoBuffers, FMorphVertexBuffer& MorphVertexBuffer)
 {
-	FComputeShaderRHIRef CS = GetComputeShader();
-	RHICmdList.SetComputeShader(CS);
-
+	FRHIComputeShader* CS = RHICmdList.GetBoundComputeShader();
 	SetUAVParameter(RHICmdList, CS, MorphVertexBufferParameter, MorphVertexBuffer.GetUAV());
 	SetSRVParameter(RHICmdList, CS, MorphPermutationBufferParameter, MorphTargetVertexInfoBuffers.MorphPermutationsSRV);
 	SetShaderValue(RHICmdList, CS, PositionScaleParameter, InvLocalScale);
@@ -736,7 +733,7 @@ void FGPUMorphNormalizeCS::SetParameters(FRHICommandList& RHICmdList, const FVec
 
 void FGPUMorphNormalizeCS::SetOffsetAndSize(FRHICommandList& RHICmdList, uint32 StartIndex, uint32 EndIndexPlusOne, const FMorphTargetVertexInfoBuffers& MorphTargetVertexInfoBuffers, const TArray<float>& InverseAccumulatedWeights)
 {
-	FComputeShaderRHIRef CS = GetComputeShader();
+	FRHIComputeShader* CS = RHICmdList.GetBoundComputeShader();
 
 	uint32 ThreadOffsets[GMorphTargetDispatchBatchSize];
 	float Weights[GMorphTargetDispatchBatchSize];
@@ -774,7 +771,7 @@ void FGPUMorphNormalizeCS::Dispatch(FRHICommandList& RHICmdList, uint32 NumVerti
 
 void FGPUMorphNormalizeCS::EndAllDispatches(FRHICommandList& RHICmdList)
 {
-	FComputeShaderRHIRef CS = GetComputeShader();
+	FRHIComputeShader* CS = RHICmdList.GetBoundComputeShader();
 	SetUAVParameter(RHICmdList, CS, MorphVertexBufferParameter, nullptr);
 }
 
@@ -879,6 +876,7 @@ void FSkeletalMeshObjectGPUSkin::FSkeletalMeshObjectLOD::UpdateMorphVertexBuffer
 
 				if (NumMorphDeltas > 0)
 				{
+					RHICmdList.SetComputeShader(GPUMorphUpdateCS.GetComputeShader());
 					GPUMorphUpdateCS->SetParameters(RHICmdList, MorphScale, MorphTargetVertexInfoBuffers, MorphVertexBuffer);
 					GPUMorphUpdateCS->SetOffsetAndSize(RHICmdList, i, i + j, MorphTargetVertexInfoBuffers, MorphTargetWeights);
 					check(NumMorphDeltas <= FMorphTargetVertexInfoBuffers::GetMaximumThreadGroupSize());
@@ -918,6 +916,7 @@ void FSkeletalMeshObjectGPUSkin::FSkeletalMeshObjectLOD::UpdateMorphVertexBuffer
 
 				if (DispatchSize > 0)
 				{
+					RHICmdList.SetComputeShader(GPUMorphNormalizeCS.GetComputeShader());
 					GPUMorphNormalizeCS->SetParameters(RHICmdList, InvMorphScale, MorphTargetVertexInfoBuffers, MorphVertexBuffer);
 					GPUMorphNormalizeCS->SetOffsetAndSize(RHICmdList, i, i + j, MorphTargetVertexInfoBuffers, InverseAccumulatedWeights);
 					check(DispatchSize <= FMorphTargetVertexInfoBuffers::GetMaximumThreadGroupSize());
diff --git a/Engine/Source/Runtime/Engine/Private/SkeletalRenderGPUSkin.h b/Engine/Source/Runtime/Engine/Private/SkeletalRenderGPUSkin.h
--- a/Engine/Source/Runtime/Engine/Private/SkeletalRenderGPUSkin.h
+++ b/Engine/Source/Runtime/Engine/Private/SkeletalRenderGPUSkin.h
@@ -561,22 +561,6 @@ public:
 		MorphDeltasParameter.Bind(Initializer.ParameterMap, TEXT("MorphDeltas"));
 	}
 
-	// FShader interface.
-	virtual bool Serialize(FArchive& Ar) override
-	{
-		bool bShaderHasOutdatedParameters = FGlobalShader::Serialize(Ar);
-		Ar << MorphVertexBufferParameter;
-
-		Ar << MorphTargetWeightParameter;
-		Ar << ThreadOffsetsParameter;
-		Ar << GlobalDispatchOffsetParameter;
-		Ar << PositionScaleParameter;
-
-		Ar << VertexIndicesParameter;
-		Ar << MorphDeltasParameter;
-		return bShaderHasOutdatedParameters;
-	}
-
 	void SetParameters(FRHICommandList& RHICmdList, const FVector4& LocalScale, const FMorphTargetVertexInfoBuffers& MorphTargetVertexInfoBuffers, FMorphVertexBuffer& MorphVertexBuffer);
 	void SetOffsetAndSize(FRHICommandList& RHICmdList, uint32 StartIndex, uint32 EndIndexPlusOne, const FMorphTargetVertexInfoBuffers& MorphTargetVertexInfoBuffers, const TArray<float>& MorphTargetWeights);
 
@@ -589,16 +573,16 @@ public:
 	}
 
 protected:
-	FShaderResourceParameter MorphVertexBufferParameter;
+	LAYOUT_FIELD(FShaderResourceParameter, MorphVertexBufferParameter);
 
-	FShaderParameter MorphTargetWeightParameter;
-	FShaderParameter OffsetAndSizeParameter;
-	FShaderParameter ThreadOffsetsParameter;
-	FShaderParameter GlobalDispatchOffsetParameter;
-	FShaderParameter PositionScaleParameter;
+	LAYOUT_FIELD(FShaderParameter, MorphTargetWeightParameter);
+	LAYOUT_FIELD(FShaderParameter, OffsetAndSizeParameter);
+	LAYOUT_FIELD(FShaderParameter, ThreadOffsetsParameter);
+	LAYOUT_FIELD(FShaderParameter, GlobalDispatchOffsetParameter);
+	LAYOUT_FIELD(FShaderParameter, PositionScaleParameter);
 
-	FShaderResourceParameter VertexIndicesParameter;
-	FShaderResourceParameter MorphDeltasParameter;
+	LAYOUT_FIELD(FShaderResourceParameter, VertexIndicesParameter);
+	LAYOUT_FIELD(FShaderResourceParameter, MorphDeltasParameter);
 };
 
 class FGPUMorphNormalizeCS : public FGlobalShader
@@ -620,21 +604,6 @@ public:
 		PositionScaleParameter.Bind(Initializer.ParameterMap, TEXT("PositionScale"));
 	}
 
-	// FShader interface.
-	virtual bool Serialize(FArchive& Ar) override
-	{
-		bool bShaderHasOutdatedParameters = FGlobalShader::Serialize(Ar);
-		Ar << MorphVertexBufferParameter;
-		Ar << MorphPermutationBufferParameter;
-
-		Ar << MorphTargetWeightParameter;
-		Ar << ThreadOffsetsParameter;
-		Ar << GlobalDispatchOffsetParameter;
-		Ar << PositionScaleParameter;
-
-		return bShaderHasOutdatedParameters;
-	}
-
 	static bool ShouldCompilePermutation(const FGlobalShaderPermutationParameters& Parameters)
 	{
 		return IsFeatureLevelSupported(Parameters.Platform, ERHIFeatureLevel::SM5);
@@ -647,11 +616,11 @@ public:
 	void EndAllDispatches(FRHICommandList& RHICmdList);
 
 protected:
-	FShaderResourceParameter MorphVertexBufferParameter;
-	FShaderResourceParameter MorphPermutationBufferParameter;
+	LAYOUT_FIELD(FShaderResourceParameter, MorphVertexBufferParameter);
+	LAYOUT_FIELD(FShaderResourceParameter, MorphPermutationBufferParameter);
 
-	FShaderParameter MorphTargetWeightParameter;
-	FShaderParameter ThreadOffsetsParameter;
-	FShaderParameter GlobalDispatchOffsetParameter;
-	FShaderParameter PositionScaleParameter;
+	LAYOUT_FIELD(FShaderParameter, MorphTargetWeightParameter);
+	LAYOUT_FIELD(FShaderParameter, ThreadOffsetsParameter);
+	LAYOUT_FIELD(FShaderParameter, GlobalDispatchOffsetParameter);
+	LAYOUT_FIELD(FShaderParameter, PositionScaleParameter);
 };
diff --git a/Engine/Source/Runtime/Engine/Private/Slate/SlateTextures.cpp b/Engine/Source/Runtime/Engine/Private/Slate/SlateTextures.cpp
--- a/Engine/Source/Runtime/Engine/Private/Slate/SlateTextures.cpp
+++ b/Engine/Source/Runtime/Engine/Private/Slate/SlateTextures.cpp
@@ -4,6 +4,7 @@
 #include "Slate/SlateTextures.h"
 #include "RenderUtils.h"
 #include "ClearQuad.h"
+#include "ProfilingDebugging/LoadTimeTracker.h"
 
 FSlateTexture2DRHIRef::FSlateTexture2DRHIRef( FTexture2DRHIRef InRef, uint32 InWidth, uint32 InHeight )
 	: TSlateTexture( InRef )
@@ -39,6 +40,8 @@ void FSlateTexture2DRHIRef::Cleanup()
 
 void FSlateTexture2DRHIRef::InitDynamicRHI()
 {
+	SCOPED_LOADTIMER(FSlateTexture2DRHIRef_InitDynamicRHI);
+
 	check( IsInRenderingThread() );
 
 	if( Width > 0 && Height > 0 )
@@ -305,6 +308,8 @@ void FSlateTextureRenderTarget2DResource::ClampSize(int32 MaxSizeX,int32 MaxSize
 
 void FSlateTextureRenderTarget2DResource::InitDynamicRHI()
 {
+	SCOPED_LOADTIMER(FSlateTextureRenderTarget2DResource_InitDynamicRHI);
+
 	check(IsInRenderingThread());
 
 	if( TargetSizeX > 0 && TargetSizeY > 0 )
diff --git a/Engine/Source/Runtime/Engine/Private/SoundNode.cpp b/Engine/Source/Runtime/Engine/Private/SoundNode.cpp
--- a/Engine/Source/Runtime/Engine/Private/SoundNode.cpp
+++ b/Engine/Source/Runtime/Engine/Private/SoundNode.cpp
@@ -110,7 +110,7 @@ void USoundNode::PrimeChildWavePlayers(bool bRecurse)
 				if (WavePlayer != nullptr)
 				{
 					USoundWave* SoundWave = WavePlayer->GetSoundWave();
-					if (SoundWave && SoundWave->IsStreaming() && SoundWave->GetNumChunks() > 1)
+					if (SoundWave && SoundWave->IsStreaming(nullptr) && SoundWave->GetNumChunks() > 1)
 					{
 						IStreamingManager::Get().GetAudioStreamingManager().RequestChunk(SoundWave, 1, [](EAudioChunkLoadResult) {});
 					}
diff --git a/Engine/Source/Runtime/Engine/Private/SoundWave.cpp b/Engine/Source/Runtime/Engine/Private/SoundWave.cpp
--- a/Engine/Source/Runtime/Engine/Private/SoundWave.cpp
+++ b/Engine/Source/Runtime/Engine/Private/SoundWave.cpp
@@ -233,7 +233,7 @@ void USoundWave::GetResourceSizeEx(FResourceSizeEx& CumulativeResourceSize)
 				CumulativeResourceSize.AddDedicatedSystemMemoryBytes(MONO_PCM_BUFFER_SIZE * NumChannels);
 			}
 
-			if (!FPlatformProperties::SupportsAudioStreaming() || !IsStreaming())
+			if (!FPlatformProperties::SupportsAudioStreaming() || !IsStreaming(nullptr))
 			{
 				CumulativeResourceSize.AddDedicatedSystemMemoryBytes(GetCompressedDataSize(LocalAudioDevice->GetRuntimeFormat(this)));
 			}
@@ -337,14 +337,13 @@ void USoundWave::Serialize( FArchive& Ar )
 		const ITargetPlatform* CookingTarget = Ar.CookingTarget();
 		if (CookingTarget != nullptr)
 		{
-			const FPlatformAudioCookOverrides* Overrides = CookingTarget->GetAudioCompressionSettings();
-			bShouldStreamSound = IsStreaming(Overrides);
+			bShouldStreamSound = IsStreaming(*CookingTarget->IniPlatformName());
 		}
 #endif
 	}
 	else
 	{
-		bShouldStreamSound = IsStreaming();
+		bShouldStreamSound = IsStreaming(nullptr);
 	}
 
 	bool bSupportsStreaming = false;
@@ -371,7 +370,7 @@ void USoundWave::Serialize( FArchive& Ar )
 				{
 					// for now we only support one format per wav
 					FName Format = CookingTarget->GetWaveFormat(this);
-					const FPlatformAudioCookOverrides* CompressionOverrides = CookingTarget->GetAudioCompressionSettings();
+					const FPlatformAudioCookOverrides* CompressionOverrides = FPlatformCompressionUtilities::GetCookOverrides(*CookingTarget->IniPlatformName());
 
 					GetCompressedData(Format, CompressionOverrides); // Get the data from the DDC or build it
 					if (CompressionOverrides)
@@ -525,12 +524,12 @@ bool USoundWave::HasCompressedData(FName Format, ITargetPlatform* TargetPlatform
 
 	if (GIsEditor)
 	{
-		CompressionOverrides = (TargetPlatform) ? TargetPlatform->GetAudioCompressionSettings() : nullptr;
+		CompressionOverrides = (TargetPlatform) ? FPlatformCompressionUtilities::GetCookOverrides(*TargetPlatform->IniPlatformName()) : nullptr;
 	}
 	else
 	{
 		// TargetPlatform is not available on consoles/mobile, so we have to grab it ourselves:
-		CompressionOverrides = FPlatformCompressionUtilities::GetCookOverridesForCurrentPlatform();
+		CompressionOverrides = CompressionOverrides = FPlatformCompressionUtilities::GetCookOverrides();
 	}
 
 	if (CompressionOverrides)
@@ -564,7 +563,7 @@ bool USoundWave::HasCompressedData(FName Format, ITargetPlatform* TargetPlatform
 
 const FPlatformAudioCookOverrides* USoundWave::GetPlatformCompressionOverridesForCurrentPlatform()
 {
-	return FPlatformCompressionUtilities::GetCookOverridesForCurrentPlatform();
+	return FPlatformCompressionUtilities::GetCookOverrides();
 }
 
 #if WITH_EDITOR
@@ -815,7 +814,7 @@ void USoundWave::PostLoad()
 		}
 
 		// In case any code accesses bStreaming directly, we fix up bStreaming based on the current platform's cook overrides.
-		bStreaming = IsStreaming();
+		bStreaming = IsStreaming(nullptr);
 	}
 
 	// Compress to whatever formats the active target platforms want
@@ -827,7 +826,7 @@ void USoundWave::PostLoad()
 
 		for (int32 Index = 0; Index < Platforms.Num(); Index++)
 		{
-			BeginGetCompressedData(Platforms[Index]->GetWaveFormat(this), Platforms[Index]->GetAudioCompressionSettings());
+			BeginGetCompressedData(Platforms[Index]->GetWaveFormat(this), FPlatformCompressionUtilities::GetCookOverrides(*Platforms[Index]->IniPlatformName()));
 		}
 	}
 
@@ -849,7 +848,7 @@ void USoundWave::PostLoad()
 	}
 
 	// Only add this streaming sound if the platform supports streaming
-	if (IsStreaming() && FPlatformProperties::SupportsAudioStreaming())
+	if (IsStreaming(nullptr) && FPlatformProperties::SupportsAudioStreaming())
 	{
 #if WITH_EDITORONLY_DATA
 		FinishCachePlatformData();
@@ -945,7 +944,7 @@ uint32 USoundWave::GetSizeOfChunk(uint32 ChunkIndex)
 void USoundWave::BeginDestroy()
 {
 	Super::BeginDestroy();
-	
+
 	{
 		FScopeLock Lock(&SourcesPlayingCs);
 		int32 CurrNumSourcesPlaying = SourcesPlaying.Num();
@@ -1016,7 +1015,7 @@ void USoundWave::InitAudioResource( FByteBulkData& CompressedData )
 
 bool USoundWave::InitAudioResource(FName Format)
 {
-	if( !ResourceSize && (!FPlatformProperties::SupportsAudioStreaming() || !IsStreaming()) )
+	if( !ResourceSize && (!FPlatformProperties::SupportsAudioStreaming() || !IsStreaming(nullptr)) )
 	{
 		FByteBulkData* Bulk = GetCompressedData(Format, GetPlatformCompressionOverridesForCurrentPlatform());
 		if (Bulk)
@@ -1072,7 +1071,7 @@ void USoundWave::InvalidateSoundWaveIfNeccessary()
 		bWasStreamCachingEnabledOnLastCook = bIsStreamCachingEnabled;
 		
 		// If stream caching is now turned on, recook the streaming audio if neccessary.
-		if(bIsStreamCachingEnabled && IsStreaming())
+		if(bIsStreamCachingEnabled && IsStreaming(nullptr))
 		{
 			EnsureZerothChunkIsLoaded();
 		}
@@ -1081,7 +1080,7 @@ void USoundWave::InvalidateSoundWaveIfNeccessary()
 
 float USoundWave::GetSampleRateForTargetPlatform(const ITargetPlatform* TargetPlatform)
 {
-	const FPlatformAudioCookOverrides* Overrides = TargetPlatform->GetAudioCompressionSettings();
+	const FPlatformAudioCookOverrides* Overrides = FPlatformCompressionUtilities::GetCookOverrides(*TargetPlatform->IniPlatformName());
 	if (Overrides)
 	{
 		return GetSampleRateForCompressionOverrides(Overrides);
@@ -1979,19 +1978,18 @@ float USoundWave::GetDuration()
 	return (bLooping ? INDEFINITELY_LOOPING_DURATION : Duration);
 }
 
-bool USoundWave::IsStreaming(const FPlatformAudioCookOverrides* Overrides /* = nullptr */) const
+bool USoundWave::IsStreaming(const TCHAR* PlatformName/* = nullptr */) const
 {
-	// TODO: add in check on whether it's part of a streaming SoundGroup.
-	if (!Overrides)
+	if (GIsEditor && ForceNonStreamingInEditorCVar != 0)
 	{
-		if (GIsEditor && ForceNonStreamingInEditorCVar != 0)
-		{
-			return false;
-		}
-
-		Overrides = GetPlatformCompressionOverridesForCurrentPlatform();
+		return false;
 	}
 
+	return IsStreaming(*FPlatformCompressionUtilities::GetCookOverrides(PlatformName));
+}
+
+bool USoundWave::IsStreaming(const FPlatformAudioCookOverrides& Overrides) const
+{
 	// We stream if (A) bStreaming is set to true, (B) bForceInline is false and either bUseLoadOnDemand was set to true in
 	// our cook overrides, or the AutoStreamingThreshold was set and this sound is longer than the auto streaming threshold.
 	
@@ -2007,23 +2005,17 @@ bool USoundWave::IsStreaming(const FPlatformAudioCookOverrides* Overrides /* = n
 	{
 		return false;
 	}
-	else if (Overrides)
-	{
-		// For stream caching, the auto streaming threshold is used to force sounds to be inlined:
-		const bool bUsesStreamCache = (Overrides->bUseStreamCaching && Duration > Overrides->AutoStreamingThreshold);
-		const bool bOverAutoStreamingThreshold = (Overrides->AutoStreamingThreshold > SMALL_NUMBER  && Duration > Overrides->AutoStreamingThreshold);
 
-		return bUsesStreamCache || bOverAutoStreamingThreshold;
-	}
-	else
-	{
-		return false;
-	}
+	// For stream caching, the auto streaming threshold is used to force sounds to be inlined:
+	const bool bUsesStreamCache = (Overrides.bUseStreamCaching && Duration > Overrides.AutoStreamingThreshold);
+	const bool bOverAutoStreamingThreshold = (Overrides.AutoStreamingThreshold > SMALL_NUMBER  && Duration > Overrides.AutoStreamingThreshold);
+
+	return bUsesStreamCache || bOverAutoStreamingThreshold;
 }
 
 bool USoundWave::ShouldUseStreamCaching() const
 {
-	return  FPlatformCompressionUtilities::IsCurrentPlatformUsingStreamCaching() && IsStreaming();
+	return  FPlatformCompressionUtilities::IsCurrentPlatformUsingStreamCaching() && IsStreaming(nullptr);
 }
 
 TArrayView<const uint8> USoundWave::GetZerothChunk(bool bForImmediatePlayback)
@@ -2102,7 +2094,7 @@ void USoundWave::RemovePlayingSource(const FSoundWaveClientPtr& Source)
 
 void USoundWave::UpdatePlatformData()
 {
-	if (IsStreaming())
+	if (IsStreaming(nullptr))
 	{
 		// Make sure there are no pending requests in flight.
 		while (IStreamingManager::Get().GetAudioStreamingManager().IsStreamingInProgress(this))
diff --git a/Engine/Source/Runtime/Engine/Private/SplineMeshSceneProxy.cpp b/Engine/Source/Runtime/Engine/Private/SplineMeshSceneProxy.cpp
--- a/Engine/Source/Runtime/Engine/Private/SplineMeshSceneProxy.cpp
+++ b/Engine/Source/Runtime/Engine/Private/SplineMeshSceneProxy.cpp
@@ -2,6 +2,30 @@
 
 #include "SplineMeshSceneProxy.h"
 #include "Materials/Material.h"
+#include "MeshMaterialShader.h"
+
+IMPLEMENT_TYPE_LAYOUT(FSplineMeshVertexFactoryShaderParameters);
+
+bool FSplineMeshVertexFactory::ShouldCompilePermutation(const FVertexFactoryShaderPermutationParameters& Parameters)
+{
+	return (Parameters.MaterialParameters.bIsUsedWithSplineMeshes || Parameters.MaterialParameters.bIsSpecialEngineMaterial)
+		&& FLocalVertexFactory::ShouldCompilePermutation(Parameters);
+}
+
+/** Modify compile environment to enable spline deformation */
+void FSplineMeshVertexFactory::ModifyCompilationEnvironment(const FVertexFactoryShaderPermutationParameters& Parameters, FShaderCompilerEnvironment& OutEnvironment)
+{
+	const bool ContainsManualVertexFetch = OutEnvironment.GetDefinitions().Contains("MANUAL_VERTEX_FETCH");
+	if (!ContainsManualVertexFetch)
+	{
+		OutEnvironment.SetDefine(TEXT("MANUAL_VERTEX_FETCH"), TEXT("0"));
+	}
+
+	// We don't call this because we don't actually support speed tree wind, and this advertises support for that
+	//FLocalVertexFactory::ModifyCompilationEnvironment(Type, Platform, Material, OutEnvironment);
+
+	OutEnvironment.SetDefine(TEXT("USE_SPLINEDEFORM"), TEXT("1"));
+}
 
 FSplineMeshSceneProxy::FSplineMeshSceneProxy(USplineMeshComponent* InComponent) :
 	FStaticMeshSceneProxy(InComponent, false)
diff --git a/Engine/Source/Runtime/Engine/Private/SplineMeshSceneProxy.h b/Engine/Source/Runtime/Engine/Private/SplineMeshSceneProxy.h
--- a/Engine/Source/Runtime/Engine/Private/SplineMeshSceneProxy.h
+++ b/Engine/Source/Runtime/Engine/Private/SplineMeshSceneProxy.h
@@ -26,26 +26,10 @@ public:
 	}
 
 	/** Should we cache the material's shadertype on this platform with this vertex factory? */
-	static bool ShouldCompilePermutation(EShaderPlatform Platform, const class FMaterial* Material, const class FShaderType* ShaderType)
-	{
-		return (Material->IsUsedWithSplineMeshes() || Material->IsSpecialEngineMaterial())
-			&& FLocalVertexFactory::ShouldCompilePermutation(Platform, Material, ShaderType);
-	}
+	static bool ShouldCompilePermutation(const FVertexFactoryShaderPermutationParameters& Parameters);
 
 	/** Modify compile environment to enable spline deformation */
-	static void ModifyCompilationEnvironment(const FVertexFactoryType* Type, EShaderPlatform Platform, const FMaterial* Material, FShaderCompilerEnvironment& OutEnvironment)
-	{
-		const bool ContainsManualVertexFetch = OutEnvironment.GetDefinitions().Contains("MANUAL_VERTEX_FETCH");
-		if (!ContainsManualVertexFetch)
-		{
-			OutEnvironment.SetDefine(TEXT("MANUAL_VERTEX_FETCH"), TEXT("0"));
-		}
-
-		// We don't call this because we don't actually support speed tree wind, and this advertises support for that
-		//FLocalVertexFactory::ModifyCompilationEnvironment(Type, Platform, Material, OutEnvironment);
-
-		OutEnvironment.SetDefine(TEXT("USE_SPLINEDEFORM"), TEXT("1"));
-	}
+	static void ModifyCompilationEnvironment(const FVertexFactoryShaderPermutationParameters& Parameters, FShaderCompilerEnvironment& OutEnvironment);
 
 	/** Copy the data from another vertex factory */
 	void Copy(const FSplineMeshVertexFactory& Other)
@@ -59,8 +43,6 @@ public:
 			});
 		BeginUpdateResourceRHI(this);
 	}
-
-	static FVertexFactoryShaderParameters* ConstructShaderParameters(EShaderFrequency ShaderFrequency);
 };
 
 //////////////////////////////////////////////////////////////////////////
@@ -69,9 +51,11 @@ public:
 /** Factory specific params */
 class FSplineMeshVertexFactoryShaderParameters : public FVertexFactoryShaderParameters
 {
-	void Bind(const FShaderParameterMap& ParameterMap) override;
+	DECLARE_TYPE_LAYOUT(FSplineMeshVertexFactoryShaderParameters, NonVirtual);
+public:
+	void Bind(const FShaderParameterMap& ParameterMap);
 
-	virtual void GetElementShaderBindings(
+	void GetElementShaderBindings(
 		const class FSceneInterface* Scene,
 		const FSceneView* View,
 		const class FMeshMaterialShader* Shader,
@@ -81,20 +65,10 @@ class FSplineMeshVertexFactoryShaderParameters : public FVertexFactoryShaderPara
 		const FMeshBatchElement& BatchElement,
 		class FMeshDrawSingleShaderBindings& ShaderBindings,
 		FVertexInputStreamArray& VertexStreams
-		) const override;
-
-	void Serialize(FArchive& Ar) override
-	{
-		Ar << SplineMeshParams;
-	}
-
-	virtual uint32 GetSize() const override
-	{
-		return sizeof(*this);
-	}
+		) const;
 
 private:
-	FShaderParameter SplineMeshParams;
+	LAYOUT_FIELD(FShaderParameter, SplineMeshParams);
 };
 
 //////////////////////////////////////////////////////////////////////////
diff --git a/Engine/Source/Runtime/Engine/Private/StaticMesh.cpp b/Engine/Source/Runtime/Engine/Private/StaticMesh.cpp
--- a/Engine/Source/Runtime/Engine/Private/StaticMesh.cpp
+++ b/Engine/Source/Runtime/Engine/Private/StaticMesh.cpp
@@ -89,6 +89,12 @@ DECLARE_MEMORY_STAT( TEXT( "StaticMesh Occluder Memory" ), STAT_StaticMeshOcclud
 
 DECLARE_MEMORY_STAT( TEXT( "StaticMesh Total Memory" ), STAT_StaticMeshTotalMemory, STATGROUP_Memory );
 
+
+DECLARE_FLOAT_ACCUMULATOR_STAT(TEXT("UStaticMesh::Serialize"), STAT_StaticMesh_SerializeFull, STATGROUP_LoadTime);
+DECLARE_FLOAT_ACCUMULATOR_STAT(TEXT("UStaticMesh_SerializeParent"), STAT_StaticMesh_SerializeParent, STATGROUP_LoadTime);
+DECLARE_FLOAT_ACCUMULATOR_STAT(TEXT("UStaticMesh_RenderDataLoad"), STAT_StaticMesh_RenderData, STATGROUP_LoadTime);
+DECLARE_FLOAT_ACCUMULATOR_STAT(TEXT("UStaticMesh_RenderDataFixup"), STAT_StaticMesh_RenderDataFixup, STATGROUP_LoadTime);
+
 /** Package name, that if set will cause only static meshes in that package to be rebuilt based on SM version. */
 ENGINE_API FName GStaticMeshPackageNameToRebuild = NAME_None;
 
@@ -198,8 +204,8 @@ void FStaticMeshSectionAreaWeightedTriangleSamplerBuffer::InitRHI()
 		for (uint32 i = 0; i < AllSectionCount; ++i)
 		{
 			FStaticMeshSectionAreaWeightedTriangleSampler& sampler = (*Samplers)[i];
-			const TArray<float>& ProbTris = sampler.GetProb();
-			const TArray<int32>& AliasTris = sampler.GetAlias();
+			const TArray<float, FMemoryImageAllocator>& ProbTris = sampler.GetProb();
+			const TArray<int32, FMemoryImageAllocator>& AliasTris = sampler.GetAlias();
 			const uint32 NumTriangle = sampler.GetNumEntries();
 
 			for (uint32 t = 0; t < NumTriangle; ++t)
@@ -601,6 +607,7 @@ void FStaticMeshLODResources::ClearAvailabilityInfo()
 	VertexBuffers.StaticMeshVertexBuffer.ClearMetaData();
 	VertexBuffers.PositionVertexBuffer.ClearMetaData();
 	VertexBuffers.ColorVertexBuffer.ClearMetaData();
+
 	delete AdditionalIndexBuffers;
 	AdditionalIndexBuffers = nullptr;
 }
@@ -1083,8 +1090,8 @@ FStaticMeshLODResources::FStaticMeshLODResources()
 
 FStaticMeshLODResources::~FStaticMeshLODResources()
 {
+	// DistanceFieldData.SafeDelete();
 	delete DistanceFieldData;
-	delete AdditionalIndexBuffers;
 }
 
 void FStaticMeshLODResources::ConditionalForce16BitIndexBuffer(EShaderPlatform MaxShaderPlatform, UStaticMesh* Parent)
@@ -1331,6 +1338,7 @@ FStaticMeshRenderData::FStaticMeshRenderData()
 	}
 }
 
+
 void FStaticMeshRenderData::Serialize(FArchive& Ar, UStaticMesh* Owner, bool bCooked)
 {
 	TRACE_CPUPROFILER_EVENT_SCOPE(FStaticMeshRenderData::Serialize);
@@ -1353,7 +1361,32 @@ void FStaticMeshRenderData::Serialize(FArchive& Ar, UStaticMesh* Owner, bool bCo
 
 #endif // #if WITH_EDITORONLY_DATA
 
-	LODResources.Serialize(Ar, Owner);
+	if (Ar.IsLoading())
+	{
+		// Load array.
+		int32 NewNum;
+		Ar << NewNum;
+		LODResources.Empty(NewNum);
+		for (int32 Index = 0; Index < NewNum; Index++)
+		{
+			new (LODResources) FStaticMeshLODResources;
+		}
+		for (int32 Index = 0; Index < NewNum; Index++)
+		{
+			LODResources[Index].Serialize(Ar, Owner, Index);
+		}
+	}
+	else
+	{
+		// Save array.
+		int32 Num = LODResources.Num();
+		Ar << Num;
+		for (int32 Index = 0; Index < Num; Index++)
+		{
+			LODResources[Index].Serialize(Ar, Owner, Index);
+		}
+	}
+
 #if WITH_EDITOR
 	if (Ar.IsSaving())
 	{
@@ -1378,7 +1411,7 @@ void FStaticMeshRenderData::Serialize(FArchive& Ar, UStaticMesh* Owner, bool bCo
 		LODVertexFactories.Empty(LODResources.Num());
 		for (int i = 0; i < LODResources.Num(); i++)
 		{
-			LODVertexFactories.Add(new FStaticMeshVertexFactories(GMaxRHIFeatureLevel));
+			new (LODVertexFactories) FStaticMeshVertexFactories(GMaxRHIFeatureLevel);
 		}
 	}
 
@@ -1544,10 +1577,12 @@ void FStaticMeshRenderData::ReleaseResources()
 void FStaticMeshRenderData::AllocateLODResources(int32 NumLODs)
 {
 	check(LODResources.Num() == 0);
+	LODResources.Reserve(NumLODs);
+	LODVertexFactories.Reserve(NumLODs);
 	while (LODResources.Num() < NumLODs)
 	{
-		LODResources.Add(new FStaticMeshLODResources);
-		LODVertexFactories.Add(new FStaticMeshVertexFactories(GMaxRHIFeatureLevel));
+		new (LODResources) FStaticMeshLODResources;
+		new (LODVertexFactories) FStaticMeshVertexFactories(GMaxRHIFeatureLevel);
 	}
 }
 
@@ -1750,7 +1785,7 @@ void FStaticMeshRenderData::SyncUVChannelData(const TArray<FStaticMaterial>& Obj
 
 	ENQUEUE_RENDER_COMMAND(SyncUVChannelData)([this, UpdateData = MoveTemp(UpdateData)](FRHICommandListImmediate& RHICmdList)
 	{
-		FMemory::Memswap(&UVChannelDataPerMaterial, UpdateData.Get(), sizeof(TArray<FMeshUVChannelInfo>));
+		FMemory::Memswap(&UVChannelDataPerMaterial, UpdateData.Get(), sizeof(UVChannelDataPerMaterial));
 	});
 }
 
@@ -2619,7 +2654,7 @@ void UStaticMesh::InitResources()
 	{
 		LinkStreaming();
 	}
-
+	
 #if	STATS
 	UStaticMesh* This = this;
 	ENQUEUE_RENDER_COMMAND(UpdateMemoryStats)(
@@ -2760,7 +2795,7 @@ int32 UStaticMesh::GetNumSections(int32 InLOD) const
 }
 
 #if WITH_EDITORONLY_DATA
-static float GetUVDensity(const TIndirectArray<FStaticMeshLODResources>& LODResources, int32 UVIndex)
+static float GetUVDensity(const TArrayView<FStaticMeshLODResources>& LODResources, int32 UVIndex)
 {
 	float WeightedUVDensity = 0;
 	float WeightSum = 0;
@@ -4427,7 +4462,12 @@ void UStaticMesh::Serialize(FArchive& Ar)
 
 	TRACE_CPUPROFILER_EVENT_SCOPE(UStaticMesh::Serialize);
 
-	Super::Serialize(Ar);
+	SCOPE_MS_ACCUMULATOR(STAT_StaticMesh_SerializeFull);
+
+	{
+		SCOPE_MS_ACCUMULATOR(STAT_StaticMesh_SerializeParent);
+		Super::Serialize(Ar);
+	}
 
 	Ar.UsingCustomVersion(FReleaseObjectVersion::GUID);
 	Ar.UsingCustomVersion(FEditorObjectVersion::GUID);
@@ -4547,18 +4587,18 @@ void UStaticMesh::Serialize(FArchive& Ar)
 	{	
 		if (Ar.IsLoading())
 		{
+			SCOPE_MS_ACCUMULATOR(STAT_StaticMesh_RenderData);
 			RenderData = MakeUnique<FStaticMeshRenderData>();
 			RenderData->Serialize(Ar, this, bCooked);
-			
+
 			FStaticMeshOccluderData::SerializeCooked(Ar, this);
 		}
-
 #if WITH_EDITOR
 		else if (Ar.IsSaving())
-		{
+		{		
 			FStaticMeshRenderData& PlatformRenderData = GetPlatformStaticMeshRenderData(this, Ar.CookingTarget());
 			PlatformRenderData.Serialize(Ar, this, bCooked);
-			
+
 			FStaticMeshOccluderData::SerializeCooked(Ar, this);
 		}
 #endif
@@ -5106,7 +5146,7 @@ void UStaticMesh::BuildFromMeshDescription(const FMeshDescription& MeshDescripti
 	TArray<uint32> IndexBuffer;
 	IndexBuffer.SetNumZeroed(NumTriangles * 3);
 
-	TArray<FStaticMeshSection>& Sections = LODResources.Sections;
+	FStaticMeshLODResources::FStaticMeshSectionArray& Sections = LODResources.Sections;
 
 	int32 SectionIndex = 0;
 	int32 IndexBufferIndex = 0;
@@ -5383,7 +5423,7 @@ int32 UStaticMesh::CalcNumOptionalMips() const
 	int32 NumOptionalLODs = 0;
 	if (RenderData)
 	{
-		const TIndirectArray<FStaticMeshLODResources>& LODResources = RenderData->LODResources;
+		const FStaticMeshLODResourcesArray& LODResources = RenderData->LODResources;
 		for (int32 Idx = 0; Idx < LODResources.Num(); ++Idx)
 		{
 			const FStaticMeshLODResources& Resource = LODResources[Idx];
diff --git a/Engine/Source/Runtime/Engine/Private/StaticMeshRender.cpp b/Engine/Source/Runtime/Engine/Private/StaticMeshRender.cpp
--- a/Engine/Source/Runtime/Engine/Private/StaticMeshRender.cpp
+++ b/Engine/Source/Runtime/Engine/Private/StaticMeshRender.cpp
@@ -225,8 +225,7 @@ FStaticMeshSceneProxy::FStaticMeshSceneProxy(UStaticMeshComponent* InComponent,
 
 	for(int32 LODIndex = 0;LODIndex < RenderData->LODResources.Num();LODIndex++)
 	{
-		FLODInfo* NewLODInfo = new FLODInfo(InComponent, RenderData->LODVertexFactories, LODIndex, ClampedMinLOD, bLODsShareStaticLighting);
-		LODs.Add(NewLODInfo);
+		FLODInfo* NewLODInfo = new (LODs) FLODInfo(InComponent, RenderData->LODVertexFactories, LODIndex, ClampedMinLOD, bLODsShareStaticLighting);
 
 		// Under certain error conditions an LOD's material will be set to 
 		// DefaultMaterial. Ensure our material view relevance is set properly.
@@ -1776,10 +1775,10 @@ bool FStaticMeshSceneProxy::HasDynamicIndirectShadowCasterRepresentation() const
 }
 
 /** Initialization constructor. */
-FStaticMeshSceneProxy::FLODInfo::FLODInfo(const UStaticMeshComponent* InComponent, const TIndirectArray<FStaticMeshVertexFactories>& InLODVertexFactories, int32 LODIndex, int32 InClampedMinLOD, bool bLODsShareStaticLighting)
+FStaticMeshSceneProxy::FLODInfo::FLODInfo(const UStaticMeshComponent* InComponent, const FStaticMeshVertexFactoriesArray& InLODVertexFactories, int32 LODIndex, int32 InClampedMinLOD, bool bLODsShareStaticLighting)
 	: FLightCacheInterface()
-	, OverrideColorVertexBuffer(0)
-	, PreCulledIndexBuffer(NULL)
+	, OverrideColorVertexBuffer(nullptr)
+	, PreCulledIndexBuffer(nullptr)
 	, bUsesMeshModifyingMaterials(false)
 {
 	const auto FeatureLevel = InComponent->GetWorld()->FeatureLevel;
@@ -1935,8 +1934,7 @@ FStaticMeshSceneProxy::FLODInfo::FLODInfo(const UStaticMeshComponent* InComponen
 		Sections.Add(SectionInfo);
 
 		// Flag the entire LOD if any material modifies its mesh
-		UMaterialInterface::TMicRecursionGuard RecursionGuard;
-		FMaterialResource const* MaterialResource = const_cast<UMaterialInterface const*>(SectionInfo.Material)->GetMaterial_Concurrent(RecursionGuard)->GetMaterialResource(FeatureLevel);
+		FMaterialResource const* MaterialResource = const_cast<UMaterialInterface const*>(SectionInfo.Material)->GetMaterial_Concurrent()->GetMaterialResource(FeatureLevel);
 		if(MaterialResource)
 		{
 			if (IsInGameThread())
@@ -2112,7 +2110,7 @@ FPrimitiveSceneProxy* UStaticMeshComponent::CreateSceneProxy()
 		return nullptr;
 	}
 
-	const TIndirectArray<FStaticMeshLODResources>& LODResources = GetStaticMesh()->RenderData->LODResources;
+	const FStaticMeshLODResourcesArray& LODResources = GetStaticMesh()->RenderData->LODResources;
 	if (LODResources.Num() == 0	|| LODResources[FMath::Clamp<int32>(GetStaticMesh()->MinLOD.Default, 0, LODResources.Num()-1)].VertexBuffers.StaticMeshVertexBuffer.GetNumVertices() == 0)
 	{
 		return nullptr;
diff --git a/Engine/Source/Runtime/Engine/Private/StereoLayerRendering.cpp b/Engine/Source/Runtime/Engine/Private/StereoLayerRendering.cpp
--- a/Engine/Source/Runtime/Engine/Private/StereoLayerRendering.cpp
+++ b/Engine/Source/Runtime/Engine/Private/StereoLayerRendering.cpp
@@ -6,6 +6,7 @@
 
 #include "StereoLayerRendering.h"
 
+IMPLEMENT_TYPE_LAYOUT(FStereoLayerPS_Base);
 IMPLEMENT_SHADER_TYPE(,FStereoLayerVS,TEXT("/Engine/Private/StereoLayerShader.usf"),TEXT("MainVS"),SF_Vertex);
 IMPLEMENT_SHADER_TYPE(,FStereoLayerPS,TEXT("/Engine/Private/StereoLayerShader.usf"),TEXT("MainPS_Texture2D"),SF_Pixel);
 IMPLEMENT_SHADER_TYPE(,FStereoLayerPS_External,TEXT("/Engine/Private/StereoLayerShader.usf"),TEXT("MainPS_TextureExternal"),SF_Pixel);
