Copying Private-LoadTimes-4.24 stream to Main. Biggest changes are in Materials/Shader memory freezing.

diff --git a/Engine/Source/Runtime/Engine/Private/StreamableManager.cpp b/Engine/Source/Runtime/Engine/Private/StreamableManager.cpp
--- a/Engine/Source/Runtime/Engine/Private/StreamableManager.cpp
+++ b/Engine/Source/Runtime/Engine/Private/StreamableManager.cpp
@@ -28,50 +28,76 @@ public:
 	/** Adds a delegate to deferred list */
 	void AddDelegate(const FStreamableDelegate& Delegate, const FStreamableDelegate& CancelDelegate, TSharedPtr<FStreamableHandle> AssociatedHandle)
 	{
-		DataLock.Lock();
+		FPendingDelegate* PendingDelegate = new FPendingDelegate(Delegate, CancelDelegate, AssociatedHandle);
 
-		PendingDelegates.Emplace(Delegate, CancelDelegate, AssociatedHandle);
+		{
+			FScopeLock Lock(&DataLock);
 
-		DataLock.Unlock();
+			FPendingDelegateList& DelegatesForHandle = PendingDelegatesByHandle.FindOrAdd(AssociatedHandle);
+			FPendingDelegateList::AddTail(PendingDelegate, PendingDelegates, DelegatesForHandle);
+		}
 	}
 
 	/** Cancels delegate for handle, this will either delete the delegate or replace with the cancel delegate */
 	void CancelDelegatesForHandle(TSharedPtr<FStreamableHandle> AssociatedHandle)
 	{
-		if (!AssociatedHandle.IsValid())
+		FPendingDelegateList PendingDelegatesToDelete;
+
 		{
-			return;
-		}
+			FScopeLock _(&DataLock);
 
-		DataLock.Lock();
+			FPendingDelegateList* DelegatesForHandle = PendingDelegatesByHandle.Find(AssociatedHandle);
+			if (!DelegatesForHandle)
+			{
+				return;
+			}
 
-		for (int32 DelegateIndex = 0; DelegateIndex < PendingDelegates.Num(); DelegateIndex++)
-		{
-			FPendingDelegate& PendingDelegate = PendingDelegates[DelegateIndex];
-			if (PendingDelegate.RelatedHandle == AssociatedHandle)
+			FPendingDelegate* CurrentNode = DelegatesForHandle->Head;
+			while (CurrentNode)
 			{
-				if (PendingDelegate.CancelDelegate.IsBound())
+				FPendingDelegate* NextNode = CurrentNode->NextByHandle;
+				if (CurrentNode->CancelDelegate.IsBound())
 				{
 					// Replace with cancel delegate
-					PendingDelegate.Delegate = PendingDelegate.CancelDelegate;
-					PendingDelegate.CancelDelegate.Unbind();
+					CurrentNode->Delegate = CurrentNode->CancelDelegate;
+					CurrentNode->CancelDelegate.Unbind();
 				}
 				else
 				{
 					// Remove entirely
-					PendingDelegates.RemoveAt(DelegateIndex);
-					DelegateIndex--;
+					FPendingDelegateList::Unlink(CurrentNode, PendingDelegates, *DelegatesForHandle);
+					if (PendingDelegatesToDelete.Tail)
+					{
+						PendingDelegatesToDelete.Tail->Next = CurrentNode;
+						PendingDelegatesToDelete.Tail = CurrentNode;
+					}
+					else
+					{
+						PendingDelegatesToDelete.Head = PendingDelegatesToDelete.Tail = CurrentNode;
+					}
+					CurrentNode->Next = nullptr;
 				}
+				CurrentNode = NextNode;
+			}
+			if (!DelegatesForHandle->Head)
+			{
+				PendingDelegatesByHandle.Remove(AssociatedHandle);
 			}
 		}
 
-		DataLock.Unlock();
+		FPendingDelegate* NodeToDelete = PendingDelegatesToDelete.Head;
+		while (NodeToDelete)
+		{
+			FPendingDelegate* NextNodeToDelete = NodeToDelete->Next;
+			delete NodeToDelete;
+			NodeToDelete = NextNodeToDelete;
+		}
 	}
 
 	/** Calls all delegates, call from synchronous flushes */
 	void FlushDelegates()
 	{
-		while (PendingDelegates.Num() > 0)
+		while (PendingDelegates.Head)
 		{
 			Tick(0.0f);
 		}
@@ -81,34 +107,51 @@ public:
 
 	void Tick(float DeltaTime) override
 	{
-		if (PendingDelegates.Num() == 0)
+		if (!PendingDelegates.Head)
 		{
 			return;
 		}
 
-		DataLock.Lock();
-
-		TArray<FPendingDelegate> DelegatesToCall;
-		DelegatesToCall.Reserve(PendingDelegates.Num());
-		for (int32 DelegateIndex = 0; DelegateIndex < PendingDelegates.Num(); DelegateIndex++)
+		FPendingDelegateList DelegatesToCall;
 		{
-			if (--PendingDelegates[DelegateIndex].DelayFrames <= 0)
+			FScopeLock Lock(&DataLock);
+
+			FPendingDelegate* CurrentNode = PendingDelegates.Head;
+			while (CurrentNode)
 			{
-				// Add to call array and remove from tracking one
-				DelegatesToCall.Emplace(PendingDelegates[DelegateIndex].Delegate, FStreamableDelegate(), PendingDelegates[DelegateIndex].RelatedHandle);
-				PendingDelegates.RemoveAt(DelegateIndex--);
+				FPendingDelegate* NextNode = CurrentNode->Next;
+				if (--CurrentNode->DelayFrames <= 0)
+				{
+					// Add to call array and remove from tracking one
+					FPendingDelegateList* DelegatesForHandle = PendingDelegatesByHandle.Find(CurrentNode->RelatedHandle);
+					check(DelegatesForHandle);
+					FPendingDelegateList::Unlink(CurrentNode, PendingDelegates, *DelegatesForHandle);
+					if (DelegatesToCall.Tail)
+					{
+						DelegatesToCall.Tail->Next = CurrentNode;
+						DelegatesToCall.Tail = CurrentNode;
+					}
+					else
+					{
+						DelegatesToCall.Head = DelegatesToCall.Tail = CurrentNode;
+					}
+					if (!DelegatesForHandle->Head)
+					{
+						PendingDelegatesByHandle.Remove(CurrentNode->RelatedHandle);
+					}
+				}
+				CurrentNode = NextNode;
 			}
 		}
 
-		DataLock.Unlock();
-
-		for (const FPendingDelegate& PendingDelegate : DelegatesToCall)
+		FPendingDelegate* DelegateToCall = DelegatesToCall.Head;
+		while (DelegateToCall)
 		{
-			// Call delegates, these may add other deferred delegates
-			PendingDelegate.Delegate.ExecuteIfBound();
+			FPendingDelegate* NextDelegateToCall = DelegateToCall->Next;
+			DelegateToCall->Delegate.ExecuteIfBound();
+			delete DelegateToCall;
+			DelegateToCall = NextDelegateToCall;
 		}
-
-		// When DelegatesToCall falls out of scope it may delete the referenced handles
 	}
 
 	virtual ETickableTickType GetTickableTickType() const override
@@ -135,6 +178,12 @@ private:
 
 	struct FPendingDelegate
 	{
+		FPendingDelegate* Prev;
+		FPendingDelegate* Next;
+
+		FPendingDelegate* PrevByHandle;
+		FPendingDelegate* NextByHandle;
+
 		/** Delegate to call when frames are up */
 		FStreamableDelegate Delegate;
 
@@ -155,7 +204,89 @@ private:
 		{}
 	};
 
-	TArray<FPendingDelegate> PendingDelegates;
+	struct FPendingDelegateList
+	{
+		FPendingDelegate* Head = nullptr;
+		FPendingDelegate* Tail = nullptr;
+
+		static void Unlink(FPendingDelegate* PendingDelegate, FPendingDelegateList& PendingDelegates, FPendingDelegateList& PendingDelegatesByHandle)
+		{
+			if (PendingDelegate->Prev)
+			{
+				PendingDelegate->Prev->Next = PendingDelegate->Next;
+			}
+			else
+			{
+				check(PendingDelegate == PendingDelegates.Head);
+				PendingDelegates.Head = PendingDelegate->Next;
+			}
+			if (PendingDelegate->Next)
+			{
+				PendingDelegate->Next->Prev = PendingDelegate->Prev;
+			}
+			else
+			{
+				check(PendingDelegate == PendingDelegates.Tail);
+				PendingDelegates.Tail = PendingDelegate->Prev;
+			}
+			PendingDelegate->Next = PendingDelegate->Prev = nullptr;
+
+			if (PendingDelegate->PrevByHandle)
+			{
+				PendingDelegate->PrevByHandle->NextByHandle = PendingDelegate->NextByHandle;
+			}
+			else
+			{
+				check(PendingDelegate == PendingDelegatesByHandle.Head);
+				PendingDelegatesByHandle.Head = PendingDelegate->NextByHandle;
+			}
+			if (PendingDelegate->NextByHandle)
+			{
+				PendingDelegate->NextByHandle->PrevByHandle = PendingDelegate->PrevByHandle;
+			}
+			else
+			{
+				check(PendingDelegate == PendingDelegatesByHandle.Tail);
+				PendingDelegatesByHandle.Tail = PendingDelegate->PrevByHandle;
+			}
+			PendingDelegate->NextByHandle = PendingDelegate->PrevByHandle = nullptr;
+
+		}
+
+		static void AddTail(FPendingDelegate* PendingDelegate, FPendingDelegateList& PendingDelegates, FPendingDelegateList& PendingDelegatesByHandle)
+		{
+			if (PendingDelegates.Tail)
+			{
+				PendingDelegates.Tail->Next = PendingDelegate;
+				PendingDelegate->Prev = PendingDelegates.Tail;
+				PendingDelegate->Next = nullptr;
+				PendingDelegates.Tail = PendingDelegate;
+			}
+			else
+			{
+				check(!PendingDelegates.Head);
+				PendingDelegates.Head = PendingDelegates.Tail = PendingDelegate;
+				PendingDelegate->Next = PendingDelegate->Prev = nullptr;
+			}
+
+			if (PendingDelegatesByHandle.Tail)
+			{
+				PendingDelegatesByHandle.Tail->NextByHandle = PendingDelegate;
+				PendingDelegate->PrevByHandle = PendingDelegatesByHandle.Tail;
+				PendingDelegate->NextByHandle = nullptr;
+				PendingDelegatesByHandle.Tail = PendingDelegate;
+			}
+			else
+			{
+				check(!PendingDelegatesByHandle.Head);
+				PendingDelegatesByHandle.Head = PendingDelegatesByHandle.Tail = PendingDelegate;
+				PendingDelegate->NextByHandle = PendingDelegate->PrevByHandle = nullptr;
+			}
+		}
+	};
+
+	FPendingDelegateList PendingDelegates;
+	TMap<TSharedPtr<FStreamableHandle>, FPendingDelegateList> PendingDelegatesByHandle;
 
 	FCriticalSection DataLock;
 };
@@ -407,10 +538,12 @@ void FStreamableHandle::CancelHandle()
 		}
 
 		bCanceled = true;
+		NotifyParentsOfCancellation();
 		return;
 	}
 
 	bCanceled = true;
+	NotifyParentsOfCancellation();
 
 	ExecuteDelegate(CancelDelegate, SharedThis);
 	UnbindDelegates();
@@ -435,6 +568,8 @@ void FStreamableHandle::CancelHandle()
 	}
 
 	ChildHandles.Empty();
+	CompletedChildCount = 0;
+	CanceledChildCount = 0;
 
 	OwningManager = nullptr;
 
@@ -488,6 +623,8 @@ void FStreamableHandle::ReleaseHandle()
 		}
 
 		ChildHandles.Empty();
+		CompletedChildCount = 0;
+		CanceledChildCount = 0;
 
 		OwningManager = nullptr;
 	}
@@ -535,6 +672,8 @@ void FStreamableHandle::CompleteLoad()
 		ExecuteDelegate(CompleteDelegate, AsShared(), CancelDelegate);
 		UnbindDelegates();
 
+		NotifyParentsOfCompletion();
+
 		if (ParentHandles.Num() > 0)
 		{
 			// Update any meta handles that are still active. Copy the array first as elements may be removed from original while iterating
@@ -552,6 +691,42 @@ void FStreamableHandle::CompleteLoad()
 	}
 }
 
+void FStreamableHandle::NotifyParentsOfCompletion()
+{
+	if (ParentHandles.Num() > 0)
+	{
+		// Update any meta handles that are still active. Copy the array first as elements may be removed from original while iterating
+		TArray<TWeakPtr<FStreamableHandle>> ParentHandlesCopy = ParentHandles;
+		for (TWeakPtr<FStreamableHandle>& WeakHandle : ParentHandlesCopy)
+		{
+			TSharedPtr<FStreamableHandle> Handle = WeakHandle.Pin();
+
+			if (Handle.IsValid())
+			{
+				++Handle->CompletedChildCount;
+			}
+		}
+	}
+}
+
+void FStreamableHandle::NotifyParentsOfCancellation()
+{
+	if (ParentHandles.Num() > 0)
+	{
+		// Update any meta handles that are still active. Copy the array first as elements may be removed from original while iterating
+		TArray<TWeakPtr<FStreamableHandle>> ParentHandlesCopy = ParentHandles;
+		for (TWeakPtr<FStreamableHandle>& WeakHandle : ParentHandlesCopy)
+		{
+			TSharedPtr<FStreamableHandle> Handle = WeakHandle.Pin();
+
+			if (Handle.IsValid())
+			{
+				++Handle->CanceledChildCount;
+			}
+		}
+	}
+}
+
 void FStreamableHandle::UpdateCombinedHandle()
 {
 	if (!IsActive())
@@ -564,6 +739,11 @@ void FStreamableHandle::UpdateCombinedHandle()
 		return;
 	}
 
+	if (CompletedChildCount + CanceledChildCount < ChildHandles.Num())
+	{
+		return;
+	}
+
 	// Check all our children, complete if done
 	bool bAllCompleted = true;
 	bool bAllCanceled = true;
@@ -718,6 +898,9 @@ struct FStreamable
 					FStreamableHandle::ExecuteDelegate(ActiveHandle->CancelDelegate, ActiveHandle);
 					ActiveHandle->UnbindDelegates();
 				}
+
+				ActiveHandle->NotifyParentsOfCancellation();
+
 			}
 		}
 		ActiveHandles.Empty();
@@ -1361,6 +1544,14 @@ TSharedPtr<FStreamableHandle> FStreamableManager::CreateCombinedHandle(const TAr
 
 		ChildHandle->ParentHandles.Add(NewRequest);
 		NewRequest->ChildHandles.Add(ChildHandle);
+		if (ChildHandle->bLoadCompleted)
+		{
+			++NewRequest->CompletedChildCount;
+		}
+		if (ChildHandle->bCanceled)
+		{
+			++NewRequest->CanceledChildCount;
+		}
 	}
 
 	// Add to pending list so these handles don't free when not referenced
diff --git a/Engine/Source/Runtime/Engine/Private/TessellationRendering.cpp b/Engine/Source/Runtime/Engine/Private/TessellationRendering.cpp
--- a/Engine/Source/Runtime/Engine/Private/TessellationRendering.cpp
+++ b/Engine/Source/Runtime/Engine/Private/TessellationRendering.cpp
@@ -7,7 +7,7 @@
 
 /** Returns true if the Material and Vertex Factory combination require adjacency information.
   * Game thread version that looks at the material settings. Will not change answer during a shader compile */
-bool MaterialSettingsRequireAdjacencyInformation_GameThread(UMaterialInterface* Material, const FVertexFactoryType* VertexFactoryType, ERHIFeatureLevel::Type InFeatureLevel)
+bool MaterialSettingsRequireAdjacencyInformation_GameThread(UMaterialInterface* Material, const FVertexFactoryType* VertexFactoryType, const FStaticFeatureLevel InFeatureLevel)
 {
 	check(IsInGameThread());
 
@@ -27,7 +27,7 @@ bool MaterialSettingsRequireAdjacencyInformation_GameThread(UMaterialInterface*
 
 /** Returns true if the Material and Vertex Factory combination require adjacency information.
   * Rendering thread version that looks at the current shader that will be used. **Will change answer during a shader compile** */
-bool MaterialRenderingRequiresAdjacencyInformation_RenderingThread(UMaterialInterface* Material, const FVertexFactoryType* VertexFactoryType, ERHIFeatureLevel::Type InFeatureLevel)
+bool MaterialRenderingRequiresAdjacencyInformation_RenderingThread(UMaterialInterface* Material, const FVertexFactoryType* VertexFactoryType, const FStaticFeatureLevel InFeatureLevel)
 {
 	check(IsInRenderingThread());
 
@@ -58,8 +58,13 @@ bool MaterialRenderingRequiresAdjacencyInformation_RenderingThread(UMaterialInte
   *
   * WARNING: In single-threaded mode as the game thread will return the rendering thread information
   * Please use the explicit game/render thread functions above instead */
-bool RequiresAdjacencyInformation(UMaterialInterface* Material, const FVertexFactoryType* VertexFactoryType, ERHIFeatureLevel::Type InFeatureLevel)
+bool RequiresAdjacencyInformation(UMaterialInterface* Material, const FVertexFactoryType* VertexFactoryType, const FStaticFeatureLevel InFeatureLevel)
 {
+	if (!RHISupportsTessellation(GShaderPlatformForFeatureLevel[InFeatureLevel]))
+	{
+		return false;
+	}
+
 	if (IsInRenderingThread())
 	{
 		return MaterialRenderingRequiresAdjacencyInformation_RenderingThread(Material, VertexFactoryType, InFeatureLevel);
@@ -74,10 +79,9 @@ bool RequiresAdjacencyInformation(UMaterialInterface* Material, const FVertexFac
 		bool VertexFactorySupportsTessellation = !VertexFactoryType || (VertexFactoryType && VertexFactoryType->SupportsTessellationShaders());
 
 		// Concurrent?
-		if (RHISupportsTessellation(GShaderPlatformForFeatureLevel[InFeatureLevel]) && VertexFactorySupportsTessellation && Material)
+		if (VertexFactorySupportsTessellation && Material)
 		{
-			UMaterialInterface::TMicRecursionGuard RecursionGuard;
-			const UMaterial* BaseMaterial = Material->GetMaterial_Concurrent(RecursionGuard);
+			const UMaterial* BaseMaterial = Material->GetMaterial_Concurrent();
 			check(BaseMaterial);
 			EMaterialTessellationMode TessellationMode = (EMaterialTessellationMode)BaseMaterial->D3D11TessellationMode;
 			bool bEnableCrackFreeDisplacement = BaseMaterial->bEnableCrackFreeDisplacement;
diff --git a/Engine/Source/Runtime/Engine/Private/Texture2D.cpp b/Engine/Source/Runtime/Engine/Private/Texture2D.cpp
--- a/Engine/Source/Runtime/Engine/Private/Texture2D.cpp
+++ b/Engine/Source/Runtime/Engine/Private/Texture2D.cpp
@@ -45,6 +45,7 @@
 #include "Engine/Texture2DArray.h"
 #include "VT/UploadingVirtualTexture.h"
 #include "VT/VirtualTexturePoolConfig.h"
+#include "ProfilingDebugging/LoadTimeTracker.h"
 
 UTexture2D::UTexture2D(const FObjectInitializer& ObjectInitializer)
 	: Super(ObjectInitializer)
@@ -1608,6 +1609,8 @@ FTexture2DResource::~FTexture2DResource()
  */
 void FTexture2DResource::InitRHI()
 {
+	SCOPED_LOADTIMER(FTexture2DResource_InitRHI);
+
 	FTexture2DScopedDebugInfo ScopedDebugInfo(Owner);
 	INC_DWORD_STAT_BY( STAT_TextureMemory, TextureSize );
 	INC_DWORD_STAT_FNAME_BY( LODGroupStatName, TextureSize );
diff --git a/Engine/Source/Runtime/Engine/Private/UnrealEngine.cpp b/Engine/Source/Runtime/Engine/Private/UnrealEngine.cpp
--- a/Engine/Source/Runtime/Engine/Private/UnrealEngine.cpp
+++ b/Engine/Source/Runtime/Engine/Private/UnrealEngine.cpp
@@ -140,6 +140,7 @@ UnrealEngine.cpp: Implements the UEngine class and helpers.
 #include "Streaming/Texture2DUpdate.h"
 #include "Rendering/SkeletalMeshRenderData.h"
 #include "Serialization/LoadTimeTrace.h"
+#include "Async/ParallelFor.h"
 
 #if WITH_EDITOR
 #include "Settings/LevelEditorPlaySettings.h"
@@ -12829,8 +12830,21 @@ bool UEngine::LoadMap( FWorldContext& WorldContext, FURL URL, class UPendingNetG
 	// Note that AI system will be created only if ai-system-creation conditions are met
 	WorldContext.World()->CreateAISystem();
 
+	FRegisterComponentContext Context;
 	// Initialize gameplay for the level.
-	WorldContext.World()->InitializeActorsForPlay(URL);		
+	WorldContext.World()->InitializeActorsForPlay(URL, true, &Context);
+
+	UWorld* ParallelWorld = WorldContext.World();
+	ParallelFor(Context.AddPrimitiveBatches.Num(),
+		[&Context, ParallelWorld](int32 Index)
+		{
+			if (!Context.AddPrimitiveBatches[Index]->IsPendingKill())
+			{
+				ParallelWorld->Scene->AddPrimitive(Context.AddPrimitiveBatches[Index]);
+			}
+		},
+		!FApp::ShouldUseThreadingForPerformance()
+	);
 
 	// calling it after InitializeActorsForPlay has been called to have all potential bounding boxed initialized
 	FNavigationSystem::AddNavigationSystemToWorld(*WorldContext.World(), FNavigationSystemRunMode::GameMode);
@@ -14451,7 +14465,7 @@ static TAutoConsoleVariable<int32> CVarAllowHighQualityLightMaps(
 	ECVF_RenderThreadSafe | ECVF_ReadOnly);
 
 
-bool AllowHighQualityLightmaps(ERHIFeatureLevel::Type FeatureLevel)
+bool AllowHighQualityLightmaps(const FStaticFeatureLevel FeatureLevel)
 {
 	return FPlatformProperties::SupportsHighQualityLightmaps()
 		&& (FeatureLevel > ERHIFeatureLevel::ES3_1)
diff --git a/Engine/Source/Runtime/Engine/Private/VectorField.cpp b/Engine/Source/Runtime/Engine/Private/VectorField.cpp
--- a/Engine/Source/Runtime/Engine/Private/VectorField.cpp
+++ b/Engine/Source/Runtime/Engine/Private/VectorField.cpp
@@ -786,19 +786,6 @@ public:
 		OutVolumeTextureSampler.Bind( Initializer.ParameterMap, TEXT("OutVolumeTextureSampler") );
 	}
 
-	/** Serialization. */
-	virtual bool Serialize( FArchive& Ar ) override
-	{
-		bool bShaderHasOutdatedParameters = FGlobalShader::Serialize( Ar );
-		Ar << AtlasTexture;
-		Ar << AtlasTextureSampler;
-		Ar << NoiseVolumeTexture;
-		Ar << NoiseVolumeTextureSampler;
-		Ar << OutVolumeTexture;
-		Ar << OutVolumeTextureSampler;
-		return bShaderHasOutdatedParameters;
-	}
-
 	/**
 	 * Set parameters for this shader.
 	 * @param UniformBuffer - Uniform buffer containing parameters for compositing vector fields.
@@ -811,7 +798,7 @@ public:
 		FRHITexture* AtlasTextureRHI,
 		FRHITexture* NoiseVolumeTextureRHI )
 	{
-		FRHIComputeShader* ComputeShaderRHI = GetComputeShader();
+		FRHIComputeShader* ComputeShaderRHI = RHICmdList.GetBoundComputeShader();
 		FRHISamplerState* SamplerStateLinear = TStaticSamplerState<SF_Bilinear,AM_Clamp,AM_Clamp,AM_Clamp>::GetRHI();
 		SetUniformBufferParameter(RHICmdList, ComputeShaderRHI, GetUniformBufferParameter<FCompositeAnimatedVectorFieldUniformParameters>(), UniformBuffer );
 		SetTextureParameter(RHICmdList, ComputeShaderRHI, AtlasTexture, AtlasTextureSampler, SamplerStateLinear, AtlasTextureRHI );
@@ -823,7 +810,7 @@ public:
 	 */
 	void SetOutput(FRHICommandList& RHICmdList, FRHIUnorderedAccessView* VolumeTextureUAV)
 	{
-		FRHIComputeShader* ComputeShaderRHI = GetComputeShader();
+		FRHIComputeShader* ComputeShaderRHI = RHICmdList.GetBoundComputeShader();
 		if ( OutVolumeTexture.IsBound() )
 		{
 			RHICmdList.SetUAVParameter(ComputeShaderRHI, OutVolumeTexture.GetBaseIndex(), VolumeTextureUAV);
@@ -835,7 +822,7 @@ public:
 	 */
 	void UnbindBuffers(FRHICommandList& RHICmdList)
 	{
-		FRHIComputeShader* ComputeShaderRHI = GetComputeShader();
+		FRHIComputeShader* ComputeShaderRHI = RHICmdList.GetBoundComputeShader();
 		if ( OutVolumeTexture.IsBound() )
 		{
 			RHICmdList.SetUAVParameter(ComputeShaderRHI, OutVolumeTexture.GetBaseIndex(), nullptr);
@@ -843,16 +830,15 @@ public:
 	}
 
 private:
-
 	/** Vector field volume textures to composite. */
-	FShaderResourceParameter AtlasTexture;
-	FShaderResourceParameter AtlasTextureSampler;
+	LAYOUT_FIELD(FShaderResourceParameter, AtlasTexture);
+	LAYOUT_FIELD(FShaderResourceParameter, AtlasTextureSampler);
 	/** Volume texture to sample as noise. */
-	FShaderResourceParameter NoiseVolumeTexture;
-	FShaderResourceParameter NoiseVolumeTextureSampler;
+	LAYOUT_FIELD(FShaderResourceParameter, NoiseVolumeTexture);
+	LAYOUT_FIELD(FShaderResourceParameter, NoiseVolumeTextureSampler);
 	/** The global vector field volume texture to write to. */
-	FShaderResourceParameter OutVolumeTexture;
-	FShaderResourceParameter OutVolumeTextureSampler;
+	LAYOUT_FIELD(FShaderResourceParameter, OutVolumeTexture);
+	LAYOUT_FIELD(FShaderResourceParameter, OutVolumeTextureSampler);
 };
 IMPLEMENT_SHADER_TYPE(,FCompositeAnimatedVectorFieldCS,TEXT("/Engine/Private/VectorFieldCompositeShaders.usf"),TEXT("CompositeAnimatedVectorField"),SF_Compute);
 
@@ -995,7 +981,7 @@ public:
 			}
 
 			RHICmdList.TransitionResource(EResourceTransitionAccess::ERWBarrier, EResourceTransitionPipeline::EGfxToCompute, VolumeTextureUAV);
-			RHICmdList.SetComputeShader(CompositeCS->GetComputeShader());
+			RHICmdList.SetComputeShader(CompositeCS.GetComputeShader());
 			CompositeCS->SetOutput(RHICmdList, VolumeTextureUAV);
 			/// ?
 			CompositeCS->SetParameters(
@@ -1005,7 +991,7 @@ public:
 				NoiseVolumeTextureRHI );
 			DispatchComputeShader(
 				RHICmdList,
-				*CompositeCS,
+				CompositeCS.GetShader(),
 				SizeX / THREADS_PER_AXIS,
 				SizeY / THREADS_PER_AXIS,
 				SizeZ / THREADS_PER_AXIS );
diff --git a/Engine/Source/Runtime/Engine/Private/VectorFieldVisualization.cpp b/Engine/Source/Runtime/Engine/Private/VectorFieldVisualization.cpp
--- a/Engine/Source/Runtime/Engine/Private/VectorFieldVisualization.cpp
+++ b/Engine/Source/Runtime/Engine/Private/VectorFieldVisualization.cpp
@@ -26,20 +26,15 @@ IMPLEMENT_GLOBAL_SHADER_PARAMETER_STRUCT(FVectorFieldVisualizationParameters,"Ve
  */
 class FVectorFieldVisualizationVertexFactoryShaderParameters : public FVertexFactoryShaderParameters
 {
+	DECLARE_INLINE_TYPE_LAYOUT(FVectorFieldVisualizationVertexFactoryShaderParameters, NonVirtual);
 public:
-	virtual void Bind( const FShaderParameterMap& ParameterMap ) override
+	void Bind( const FShaderParameterMap& ParameterMap )
 	{
 		VectorFieldTexture.Bind(ParameterMap, TEXT("VectorFieldTexture"));
 		VectorFieldTextureSampler.Bind(ParameterMap, TEXT("VectorFieldTextureSampler"));
 	}
 
-	virtual void Serialize(FArchive& Ar) override
-	{
-		Ar << VectorFieldTexture;
-		Ar << VectorFieldTextureSampler;
-	}
-
-	virtual void GetElementShaderBindings(
+	void GetElementShaderBindings(
 		const FSceneInterface* Scene,
 		const FSceneView* View,
 		const FMeshMaterialShader* Shader,
@@ -48,15 +43,14 @@ public:
 		const FVertexFactory* VertexFactory,
 		const FMeshBatchElement& BatchElement,
 		class FMeshDrawSingleShaderBindings& ShaderBindings,
-		FVertexInputStreamArray& VertexStreams) const override;
-
-	virtual uint32 GetSize() const override { return sizeof(*this); }
+		FVertexInputStreamArray& VertexStreams) const;
 
 private:
-
-	/** The vector field texture parameter. */
-	FShaderResourceParameter VectorFieldTexture;
-	FShaderResourceParameter VectorFieldTextureSampler;
+	
+		/** The vector field texture parameter. */
+		LAYOUT_FIELD(FShaderResourceParameter, VectorFieldTexture)
+		LAYOUT_FIELD(FShaderResourceParameter, VectorFieldTextureSampler)
+	
 };
 
 /**
@@ -138,26 +132,20 @@ void FVectorFieldVisualizationVertexFactory::ReleaseRHI()
 /**
  * Should we cache the material's shadertype on this platform with this vertex factory? 
  */
-bool FVectorFieldVisualizationVertexFactory::ShouldCompilePermutation(EShaderPlatform Platform, const class FMaterial* Material, const class FShaderType* ShaderType)
+bool FVectorFieldVisualizationVertexFactory::ShouldCompilePermutation(const FVertexFactoryShaderPermutationParameters& Parameters)
 {
 	// We support vector fields on mobile, but not the visualization, so we limit this VF to Feature Level Preview shader platforms
-	return Material->IsSpecialEngineMaterial() && SupportsGPUParticles(Platform) && (IsFeatureLevelSupported(Platform, ERHIFeatureLevel::SM5) || IsPCPlatform(Platform));
+	return Parameters.MaterialParameters.bIsSpecialEngineMaterial &&
+		SupportsGPUParticles(Parameters.Platform) &&
+		(IsFeatureLevelSupported(Parameters.Platform, ERHIFeatureLevel::SM5) || IsPCPlatform(Parameters.Platform));
 }
 
 /**
  * Can be overridden by FVertexFactory subclasses to modify their compile environment just before compilation occurs.
  */
-void FVectorFieldVisualizationVertexFactory::ModifyCompilationEnvironment(const FVertexFactoryType* Type, EShaderPlatform Platform, const FMaterial* Material, FShaderCompilerEnvironment& OutEnvironment)
+void FVectorFieldVisualizationVertexFactory::ModifyCompilationEnvironment(const FVertexFactoryShaderPermutationParameters& Parameters, FShaderCompilerEnvironment& OutEnvironment)
 {
-	FVertexFactory::ModifyCompilationEnvironment(Type, Platform, Material, OutEnvironment);
-}
-
-/**
- * Construct shader parameters for this type of vertex factory.
- */
-FVertexFactoryShaderParameters* FVectorFieldVisualizationVertexFactory::ConstructShaderParameters(EShaderFrequency ShaderFrequency)
-{
-	return ShaderFrequency == SF_Vertex ? new FVectorFieldVisualizationVertexFactoryShaderParameters() : NULL;
+	FVertexFactory::ModifyCompilationEnvironment(Parameters, OutEnvironment);
 }
 
 /**
@@ -188,6 +176,8 @@ void FVectorFieldVisualizationVertexFactoryShaderParameters::GetElementShaderBin
 	ShaderBindings.AddTexture(VectorFieldTexture, VectorFieldTextureSampler, SamplerStatePoint, VertexFactory->VectorFieldTextureRHI);
 }
 
+IMPLEMENT_VERTEX_FACTORY_PARAMETER_TYPE(FVectorFieldVisualizationVertexFactory, SF_Vertex, FVectorFieldVisualizationVertexFactoryShaderParameters);
+
 IMPLEMENT_VERTEX_FACTORY_TYPE(FVectorFieldVisualizationVertexFactory,"/Engine/Private/VectorFieldVisualizationVertexFactory.ush",true,false,true,false,false);
 
 /*------------------------------------------------------------------------------
diff --git a/Engine/Source/Runtime/Engine/Private/VectorFieldVisualization.h b/Engine/Source/Runtime/Engine/Private/VectorFieldVisualization.h
--- a/Engine/Source/Runtime/Engine/Private/VectorFieldVisualization.h
+++ b/Engine/Source/Runtime/Engine/Private/VectorFieldVisualization.h
@@ -70,17 +70,12 @@ public:
 	/**
 	 * Should we cache the material's shadertype on this platform with this vertex factory? 
 	 */
-	static bool ShouldCompilePermutation(EShaderPlatform Platform, const class FMaterial* Material, const class FShaderType* ShaderType);
+	static bool ShouldCompilePermutation(const FVertexFactoryShaderPermutationParameters& Parameters);
 
 	/**
 	 * Can be overridden by FVertexFactory subclasses to modify their compile environment just before compilation occurs.
 	 */
-	static void ModifyCompilationEnvironment(const FVertexFactoryType* Type, EShaderPlatform Platform, const FMaterial* Material, FShaderCompilerEnvironment& OutEnvironment);
-
-	/**
-	 * Construct shader parameters for this type of vertex factory.
-	 */
-	static FVertexFactoryShaderParameters* ConstructShaderParameters(EShaderFrequency ShaderFrequency);
+	static void ModifyCompilationEnvironment(const FVertexFactoryShaderPermutationParameters& Parameters, FShaderCompilerEnvironment& OutEnvironment);
 
 	/**
 	 * Set parameters for this vertex factory instance.
diff --git a/Engine/Source/Runtime/Engine/Private/VolumeTexture.cpp b/Engine/Source/Runtime/Engine/Private/VolumeTexture.cpp
--- a/Engine/Source/Runtime/Engine/Private/VolumeTexture.cpp
+++ b/Engine/Source/Runtime/Engine/Private/VolumeTexture.cpp
@@ -612,7 +612,7 @@ void UVolumeTexture::UpdateMipGenSettings()
 
 #endif // #if WITH_EDITOR
 
-bool UVolumeTexture::ShaderPlatformSupportsCompression(EShaderPlatform ShaderPlatform)
+bool UVolumeTexture::ShaderPlatformSupportsCompression(FStaticShaderPlatform ShaderPlatform)
 {
 	switch (ShaderPlatform)
 	{
@@ -624,7 +624,7 @@ bool UVolumeTexture::ShaderPlatformSupportsCompression(EShaderPlatform ShaderPla
 		return true;
 
 	default:
-		return FDataDrivenShaderPlatformInfo::GetInfo(ShaderPlatform).bSupportsVolumeTextureCompression;
+		return FDataDrivenShaderPlatformInfo::GetSupportsVolumeTextureCompression(ShaderPlatform);
 	}
 }
 
diff --git a/Engine/Source/Runtime/Engine/Private/WeightedRandomSampler.cpp b/Engine/Source/Runtime/Engine/Private/WeightedRandomSampler.cpp
--- a/Engine/Source/Runtime/Engine/Private/WeightedRandomSampler.cpp
+++ b/Engine/Source/Runtime/Engine/Private/WeightedRandomSampler.cpp
@@ -2,6 +2,8 @@
 
 #include "WeightedRandomSampler.h"
 
+IMPLEMENT_ABSTRACT_TYPE_LAYOUT(FWeightedRandomSampler);
+
 FWeightedRandomSampler::FWeightedRandomSampler()
 	: TotalWeight(0.0f)
 {
@@ -76,4 +78,4 @@ void FWeightedRandomSampler::Serialize(FArchive& Ar)
 	Ar << Prob;
 	Ar << Alias;
 	Ar << TotalWeight;
-}
\ No newline at end of file
+}
diff --git a/Engine/Source/Runtime/Engine/Private/WindDirectionalSource.cpp b/Engine/Source/Runtime/Engine/Private/WindDirectionalSource.cpp
--- a/Engine/Source/Runtime/Engine/Private/WindDirectionalSource.cpp
+++ b/Engine/Source/Runtime/Engine/Private/WindDirectionalSource.cpp
@@ -217,9 +217,9 @@ void UWindDirectionalSourceComponent::SetWindType(EWindSourceType InNewType)
 	MarkRenderDynamicDataDirty();
 }
 
-void UWindDirectionalSourceComponent::CreateRenderState_Concurrent()
+void UWindDirectionalSourceComponent::CreateRenderState_Concurrent(FRegisterComponentContext* Context)
 {
-	Super::CreateRenderState_Concurrent();
+	Super::CreateRenderState_Concurrent(Context);
 	GetWorld()->Scene->AddWindSource(this);
 }
 
diff --git a/Engine/Source/Runtime/Engine/Private/World.cpp b/Engine/Source/Runtime/Engine/Private/World.cpp
--- a/Engine/Source/Runtime/Engine/Private/World.cpp
+++ b/Engine/Source/Runtime/Engine/Private/World.cpp
@@ -1812,7 +1812,7 @@ void UWorld::ClearWorldComponents()
 }
 
 
-void UWorld::UpdateWorldComponents(bool bRerunConstructionScripts, bool bCurrentLevelOnly)
+void UWorld::UpdateWorldComponents(bool bRerunConstructionScripts, bool bCurrentLevelOnly, FRegisterComponentContext* Context)
 {
 	if ( !IsRunningDedicatedServer() )
 	{
@@ -1824,7 +1824,7 @@ void UWorld::UpdateWorldComponents(bool bRerunConstructionScripts, bool bCurrent
 
 		if(!LineBatcher->IsRegistered())
 		{	
-			LineBatcher->RegisterComponentWithWorld(this);
+			LineBatcher->RegisterComponentWithWorld(this, Context);
 		}
 
 		if(!PersistentLineBatcher)
@@ -1835,7 +1835,7 @@ void UWorld::UpdateWorldComponents(bool bRerunConstructionScripts, bool bCurrent
 
 		if(!PersistentLineBatcher->IsRegistered())	
 		{
-			PersistentLineBatcher->RegisterComponentWithWorld(this);
+			PersistentLineBatcher->RegisterComponentWithWorld(this, Context);
 		}
 
 		if(!ForegroundLineBatcher)
@@ -1846,7 +1846,7 @@ void UWorld::UpdateWorldComponents(bool bRerunConstructionScripts, bool bCurrent
 
 		if(!ForegroundLineBatcher->IsRegistered())	
 		{
-			ForegroundLineBatcher->RegisterComponentWithWorld(this);
+			ForegroundLineBatcher->RegisterComponentWithWorld(this, Context);
 		}
 	}
 
@@ -1856,7 +1856,7 @@ void UWorld::UpdateWorldComponents(bool bRerunConstructionScripts, bool bCurrent
 		ULevel* CurrentLevel = PersistentLevel;
 #endif
 		check( CurrentLevel );
-		CurrentLevel->UpdateLevelComponents(bRerunConstructionScripts);
+		CurrentLevel->UpdateLevelComponents(bRerunConstructionScripts, Context);
 	}
 	else
 	{
@@ -1867,7 +1867,7 @@ void UWorld::UpdateWorldComponents(bool bRerunConstructionScripts, bool bCurrent
 			// Update the level only if it is visible (or not a streamed level)
 			if(!StreamingLevel || Level->bIsVisible)
 			{
-				Level->UpdateLevelComponents(bRerunConstructionScripts);
+				Level->UpdateLevelComponents(bRerunConstructionScripts, Context);
 				IStreamingManager::Get().AddLevel(Level);
 			}
 		}
@@ -3940,7 +3940,7 @@ bool UWorld::SetGameMode(const FURL& InURL)
 	return false;
 }
 
-void UWorld::InitializeActorsForPlay(const FURL& InURL, bool bResetTime)
+void UWorld::InitializeActorsForPlay(const FURL& InURL, bool bResetTime, FRegisterComponentContext* Context)
 {
 	TRACE_OBJECT_EVENT(this, InitializeActorsForPlay);
 
@@ -3980,7 +3980,7 @@ void UWorld::InitializeActorsForPlay(const FURL& InURL, bool bResetTime)
 	// We don't need to rerun construction scripts if we have cooked data or we are playing in editor unless the PIE world was loaded
 	// from disk rather than duplicated
 	const bool bRerunConstructionScript = !(FPlatformProperties::RequiresCookedData() || (IsGameWorld() && (PersistentLevel->bHasRerunConstructionScripts || PersistentLevel->bWasDuplicatedForPIE || !bRerunConstructionDuringEditorStreaming)));
-	UpdateWorldComponents( bRerunConstructionScript, true );
+	UpdateWorldComponents( bRerunConstructionScript, true, Context);
 
 	// Init level gameplay info.
 	if( !AreActorsInitialized() )
diff --git a/Engine/Source/Runtime/Engine/Public/AudioCompressionSettingsUtils.h b/Engine/Source/Runtime/Engine/Public/AudioCompressionSettingsUtils.h
--- a/Engine/Source/Runtime/Engine/Public/AudioCompressionSettingsUtils.h
+++ b/Engine/Source/Runtime/Engine/Public/AudioCompressionSettingsUtils.h
@@ -22,10 +22,12 @@ public:
 
 	static void RecacheCookOverrides();
 
-	static const FPlatformAudioCookOverrides* GetCookOverridesForCurrentPlatform(bool bForceRecache = false);
+	// null platformname means to use current platform
+	static const FPlatformAudioCookOverrides* GetCookOverrides(const TCHAR* PlatformName=nullptr, bool bForceRecache = false);
 
 	static bool IsCurrentPlatformUsingStreamCaching();
 
+	// null platformname means to use current platform
 	static const FAudioStreamCachingSettings& GetStreamCachingSettingsForCurrentPlatform();
 
 	/** This is used at runtime to initialize FCachedAudioStreamingManager. */
diff --git a/Engine/Source/Runtime/Engine/Public/ComponentRecreateRenderStateContext.h b/Engine/Source/Runtime/Engine/Public/ComponentRecreateRenderStateContext.h
--- a/Engine/Source/Runtime/Engine/Public/ComponentRecreateRenderStateContext.h
+++ b/Engine/Source/Runtime/Engine/Public/ComponentRecreateRenderStateContext.h
@@ -40,7 +40,7 @@ public:
 	{
 		if (Component && !Component->IsRenderStateCreated() && Component->IsRegistered())
 		{
-			Component->CreateRenderState_Concurrent();
+			Component->CreateRenderState_Concurrent(nullptr);
 
 			UpdateAllPrimitiveSceneInfosForSingleComponent(Component, ScenesToUpdateAllPrimitiveSceneInfos);
 		}
diff --git a/Engine/Source/Runtime/Engine/Public/CubemapUnwrapUtils.h b/Engine/Source/Runtime/Engine/Public/CubemapUnwrapUtils.h
--- a/Engine/Source/Runtime/Engine/Public/CubemapUnwrapUtils.h
+++ b/Engine/Source/Runtime/Engine/Public/CubemapUnwrapUtils.h
@@ -59,16 +59,8 @@ public:
 
 	void SetParameters(FRHICommandList& RHICmdList, const FMatrix& TransformValue);
 
-	virtual bool Serialize(FArchive& Ar) override
-	{
-		bool bShaderHasOutdatedParameters = FGlobalShader::Serialize(Ar);
-		Ar << Transform;
-
-		return bShaderHasOutdatedParameters;
-	}
-
 private:
-	FShaderParameter Transform;
+	LAYOUT_FIELD(FShaderParameter, Transform);
 };
 
 /**
@@ -95,24 +87,17 @@ public:
 
 	void SetParameters(FRHICommandList& RHICmdList, const FTexture* Texture, const FMatrix& ColorWeightsValue, float MipLevel, float GammaValue);
 
-	virtual bool Serialize(FArchive& Ar) override
-	{
-		bool bShaderHasOutdatedParameters = FGlobalShader::Serialize(Ar);
-		Ar << CubeTexture << CubeTextureSampler << PackedProperties0 << ColorWeights << Gamma;
-		return bShaderHasOutdatedParameters;
-	}
-
 	static void ModifyCompilationEnvironment(const FGlobalShaderPermutationParameters& Parameters, FShaderCompilerEnvironment& OutEnvironment)
 	{
 		FGlobalShader::ModifyCompilationEnvironment(Parameters, OutEnvironment);
 		OutEnvironment.SetDefine(TEXT("HDR_OUTPUT"), bHDROutput ? TEXT("1") : TEXT("0"));
 	}
 private:
-	FShaderResourceParameter CubeTexture;
-	FShaderResourceParameter CubeTextureSampler;
-	FShaderParameter PackedProperties0;
-	FShaderParameter ColorWeights;
-	FShaderParameter Gamma;
+	LAYOUT_FIELD(FShaderResourceParameter, CubeTexture);
+	LAYOUT_FIELD(FShaderResourceParameter, CubeTextureSampler);
+	LAYOUT_FIELD(FShaderParameter, PackedProperties0);
+	LAYOUT_FIELD(FShaderParameter, ColorWeights);
+	LAYOUT_FIELD(FShaderParameter, Gamma);
 };
 
 
@@ -162,20 +147,11 @@ public:
 
 	void SetParameters(FRHICommandList& RHICmdList, const FTexture* Texture, float InBrightnessInLumens);
 
-	virtual bool Serialize(FArchive& Ar) override
-	{
-		bool bShaderHasOutdatedParameters = FGlobalShader::Serialize(Ar);
-		Ar << IESTexture;
-		Ar << IESTextureSampler;
-		Ar << BrightnessInLumens;
-		return bShaderHasOutdatedParameters;
-	}
-
 private:
 	/** The texture to sample. */
-	FShaderResourceParameter IESTexture;
-	FShaderResourceParameter IESTextureSampler;
-	FShaderParameter BrightnessInLumens;
+	LAYOUT_FIELD(FShaderResourceParameter, IESTexture);
+	LAYOUT_FIELD(FShaderResourceParameter, IESTextureSampler);
+	LAYOUT_FIELD(FShaderParameter, BrightnessInLumens);
 };
 
 class ENGINE_API FIESLightProfileBatchedElementParameters : public FBatchedElementParameters
diff --git a/Engine/Source/Runtime/Engine/Public/DebugViewModeInterface.h b/Engine/Source/Runtime/Engine/Public/DebugViewModeInterface.h
--- a/Engine/Source/Runtime/Engine/Public/DebugViewModeInterface.h
+++ b/Engine/Source/Runtime/Engine/Public/DebugViewModeInterface.h
@@ -47,9 +47,25 @@ public:
 
 	virtual ~FDebugViewModeInterface() {}
 
-	virtual FDebugViewModePS* GetPixelShader(const FMaterial* InMaterial, FVertexFactoryType* VertexFactoryType) const = 0;
+	virtual TShaderRef<FDebugViewModePS> GetPixelShader(const FMaterial* InMaterial, FVertexFactoryType* VertexFactoryType) const = 0;
 	virtual void SetDrawRenderState(EBlendMode BlendMode, FRenderState& DrawRenderState, bool bHasDepthPrepassForMaskedMaterial) const;
 
+	virtual void GetDebugViewModeShaderBindings(
+		const FDebugViewModePS& Shader,
+		const FPrimitiveSceneProxy* PrimitiveSceneProxy,
+		const FMaterialRenderProxy& MaterialRenderProxy,
+		const FMaterial& Material,
+		EDebugViewShaderMode DebugViewMode,
+		const FVector& ViewOrigin,
+		int32 VisualizeLODIndex,
+		int32 VisualizeElementIndex,
+		int32 NumVSInstructions,
+		int32 NumPSInstructions,
+		int32 ViewModeParam,
+		FName ViewModeParamName,
+		FMeshDrawSingleShaderBindings& ShaderBindings
+	) const {}
+
 	/** The shader class name, used to filter out shaders that need to be compiled. */
 	const TCHAR* PixelShaderName;
 
@@ -73,6 +89,7 @@ public:
 	
 	/** Whether this material can be substituted by the default material. */
 	static bool AllowFallbackToDefaultMaterial(const FMaterial* InMaterial);
+	static bool AllowFallbackToDefaultMaterial(EMaterialTessellationMode TessellationMode, bool bHasVertexPositionOffsetConnected, bool bHasPixelDepthOffsetConnected);
 
 private:
 	
diff --git a/Engine/Source/Runtime/Engine/Public/GPUSkinCache.h b/Engine/Source/Runtime/Engine/Public/GPUSkinCache.h
--- a/Engine/Source/Runtime/Engine/Public/GPUSkinCache.h
+++ b/Engine/Source/Runtime/Engine/Public/GPUSkinCache.h
@@ -112,7 +112,7 @@ public:
 		const FMatrix& ClothLocalToWorld, float ClothBlendWeight, uint32 RevisionNumber, int32 Section, FGPUSkinCacheEntry*& InOutEntry);
 
 	static void SetVertexStreams(FGPUSkinCacheEntry* Entry, int32 Section, FRHICommandList& RHICmdList,
-		class FShader* Shader, const FGPUSkinPassthroughVertexFactory* VertexFactory,
+		class FRHIVertexShader* ShaderRHI, const FGPUSkinPassthroughVertexFactory* VertexFactory,
 		uint32 BaseVertexIndex, FShaderResourceParameter PreviousStreamBuffer);
 
 	static void GetShaderBindings(
diff --git a/Engine/Source/Runtime/Engine/Public/GPUSkinVertexFactory.h b/Engine/Source/Runtime/Engine/Public/GPUSkinVertexFactory.h
--- a/Engine/Source/Runtime/Engine/Public/GPUSkinVertexFactory.h
+++ b/Engine/Source/Runtime/Engine/Public/GPUSkinVertexFactory.h
@@ -396,8 +396,8 @@ public:
 		return Data.bUse16BitBoneIndex;
 	}
 
-	static void ModifyCompilationEnvironment(const FVertexFactoryType* Type, EShaderPlatform Platform, const class ::FMaterial* Material, FShaderCompilerEnvironment& OutEnvironment);
-	static bool ShouldCompilePermutation(EShaderPlatform Platform, const class ::FMaterial* Material, const FShaderType* ShaderType);
+	static void ModifyCompilationEnvironment(const FVertexFactoryShaderPermutationParameters& Parameters, FShaderCompilerEnvironment& OutEnvironment);
+	static bool ShouldCompilePermutation(const FVertexFactoryShaderPermutationParameters& Parameters);
 	
 	/**
 	* An implementation of the interface used by TSynchronizedResource to 
@@ -417,8 +417,6 @@ public:
 	virtual void InitDynamicRHI() override;
 	virtual void ReleaseDynamicRHI() override;
 
-	static FVertexFactoryShaderParameters* ConstructShaderParameters(EShaderFrequency ShaderFrequency);
-
 	void CopyDataTypeForPassthroughFactory(class FGPUSkinPassthroughVertexFactory* PassthroughVertexFactory);
 
 	const FShaderResourceViewRHIRef GetPositionsSRV() const override
@@ -482,8 +480,8 @@ public:
 		bSupportsManualVertexFetch = true;
 	}
 
-	static void ModifyCompilationEnvironment(const FVertexFactoryType* Type, EShaderPlatform Platform, const class ::FMaterial* Material, FShaderCompilerEnvironment& OutEnvironment);
-	static bool ShouldCompilePermutation(EShaderPlatform Platform, const class ::FMaterial* Material, const FShaderType* ShaderType);
+	static void ModifyCompilationEnvironment(const FVertexFactoryShaderPermutationParameters& Parameters, FShaderCompilerEnvironment& OutEnvironment);
+	static bool ShouldCompilePermutation(const FVertexFactoryShaderPermutationParameters& Parameters);
 
 	inline void UpdateVertexDeclaration(FGPUBaseSkinVertexFactory* SourceVertexFactory, struct FRWBuffer* PositionRWBuffer, struct FRWBuffer* TangentRWBuffer)
 	{
@@ -506,9 +504,6 @@ public:
 
 	void SetData(const FDataType& InData);
 
-	// FRenderResource interface.
-	static FVertexFactoryShaderParameters* ConstructShaderParameters(EShaderFrequency ShaderFrequency);
-
 	//TODO should be supported
 	bool SupportsPositionOnlyStream() const override { return false; }
 	bool SupportsPositionAndNormalOnlyStream() const override { return false; }
@@ -564,8 +559,8 @@ public:
 	: TGPUSkinVertexFactory<BoneInfluenceType>(InFeatureLevel, InNumVertices)
 	{}
 
-	static void ModifyCompilationEnvironment(const FVertexFactoryType* Type, EShaderPlatform Platform, const class ::FMaterial* Material, FShaderCompilerEnvironment& OutEnvironment);
-	static bool ShouldCompilePermutation(EShaderPlatform Platform, const class ::FMaterial* Material, const FShaderType* ShaderType);
+	static void ModifyCompilationEnvironment(const FVertexFactoryShaderPermutationParameters& Parameters, FShaderCompilerEnvironment& OutEnvironment);
+	static bool ShouldCompilePermutation(const FVertexFactoryShaderPermutationParameters& Parameters);
 	
 	/**
 	* An implementation of the interface used by TSynchronizedResource to 
@@ -588,7 +583,6 @@ public:
 	*/
 	virtual void InitRHI() override;
 
-	static FVertexFactoryShaderParameters* ConstructShaderParameters(EShaderFrequency ShaderFrequency);
 
 	virtual uint32 GetNumBoneInfluences() const override
 	{
@@ -900,8 +894,8 @@ public:
 		: TGPUSkinVertexFactory<BoneInfluenceType>(InFeatureLevel, InNumVertices)
 	{}
 
-	static void ModifyCompilationEnvironment(const FVertexFactoryType* Type, EShaderPlatform Platform, const class ::FMaterial* Material, FShaderCompilerEnvironment& OutEnvironment);
-	static bool ShouldCompilePermutation(EShaderPlatform Platform, const class ::FMaterial* Material, const FShaderType* ShaderType);
+	static void ModifyCompilationEnvironment(const FVertexFactoryShaderPermutationParameters& Parameters, FShaderCompilerEnvironment& OutEnvironment);
+	static bool ShouldCompilePermutation(const FVertexFactoryShaderPermutationParameters& Parameters);
 
 	/**
 	* An implementation of the interface used by TSynchronizedResource to 
@@ -976,8 +970,6 @@ public:
 	virtual void InitRHI() override;
 	virtual void ReleaseDynamicRHI() override;
 
-	static FVertexFactoryShaderParameters* ConstructShaderParameters(EShaderFrequency ShaderFrequency);
-
 protected:
 	/**
 	* Add the decl elements for the streams
diff --git a/Engine/Source/Runtime/Engine/Public/GeometryCacheVertexFactory.h b/Engine/Source/Runtime/Engine/Public/GeometryCacheVertexFactory.h
--- a/Engine/Source/Runtime/Engine/Public/GeometryCacheVertexFactory.h
+++ b/Engine/Source/Runtime/Engine/Public/GeometryCacheVertexFactory.h
@@ -100,9 +100,8 @@ public:
 	};
 
 
-	static void ModifyCompilationEnvironment(const FVertexFactoryType* Type, EShaderPlatform Platform, const class FMaterial* Material, FShaderCompilerEnvironment& OutEnvironment);
-	static bool ShouldCache(EShaderPlatform Platform, const class FMaterial* Material, const FShaderType* ShaderType);
-	static bool ShouldCompilePermutation(EShaderPlatform Platform, const class FMaterial* Material, const class FShaderType* ShaderType);
+	static void ModifyCompilationEnvironment(const FVertexFactoryShaderPermutationParameters& Parameters, FShaderCompilerEnvironment& OutEnvironment);
+	static bool ShouldCompilePermutation(const FVertexFactoryShaderPermutationParameters& Parameters);
 
 	/**
 	* An implementation of the interface used by TSynchronizedResource to update the resource with new data from the game thread.
@@ -116,9 +115,6 @@ public:
 
 	virtual void InitRHI() override;
 
-	// FRenderResource interface.
-	static FVertexFactoryShaderParameters* ConstructShaderParameters(EShaderFrequency ShaderFrequency);
-
 	friend FGeometryCacheVertexFactoryShaderParameters;
 	
 protected:
diff --git a/Engine/Source/Runtime/Engine/Public/GeometryCollectionRendering.h b/Engine/Source/Runtime/Engine/Public/GeometryCollectionRendering.h
--- a/Engine/Source/Runtime/Engine/Public/GeometryCollectionRendering.h
+++ b/Engine/Source/Runtime/Engine/Public/GeometryCollectionRendering.h
@@ -12,10 +12,11 @@
 
 class FGeometryCollectionVertexFactoryShaderParameters : public FLocalVertexFactoryShaderParameters
 {
+	DECLARE_TYPE_LAYOUT(FGeometryCollectionVertexFactoryShaderParameters, NonVirtual);
 public:
 
 	// #note: We are reusing VertexFetch_InstanceTransformBuffer that was created for instanced static mesh 
-	virtual void Bind(const FShaderParameterMap& ParameterMap) override
+	void Bind(const FShaderParameterMap& ParameterMap)
 	{
 		FLocalVertexFactoryShaderParameters::Bind(ParameterMap);
 
@@ -24,7 +25,7 @@ public:
 		VertexFetch_InstanceBoneMapBufferParameter.Bind(ParameterMap, TEXT("VertexFetch_InstanceBoneMapBuffer"));
 	}	
 
-	virtual void GetElementShaderBindings(
+	void GetElementShaderBindings(
 		const FSceneInterface* Scene,
 		const FSceneView* View,
 		const FMeshMaterialShader* Shader,
@@ -33,22 +34,14 @@ public:
 		const FVertexFactory* VertexFactory,
 		const FMeshBatchElement& BatchElement,
 		class FMeshDrawSingleShaderBindings& ShaderBindings,
-		FVertexInputStreamArray& VertexStreams) const override;
-
-	void Serialize(FArchive& Ar) override
-	{
-		FLocalVertexFactoryShaderParameters::Serialize(Ar);
-		Ar << VertexFetch_InstanceTransformBufferParameter;
-		Ar << VertexFetch_InstancePrevTransformBufferParameter;
-		Ar << VertexFetch_InstanceBoneMapBufferParameter;
-	}
-
-	virtual uint32 GetSize() const override { return sizeof(*this); }
+		FVertexInputStreamArray& VertexStreams) const;
 
 private:
-	FShaderResourceParameter VertexFetch_InstanceTransformBufferParameter;
-	FShaderResourceParameter VertexFetch_InstancePrevTransformBufferParameter;
-	FShaderResourceParameter VertexFetch_InstanceBoneMapBufferParameter;
+	
+		LAYOUT_FIELD(FShaderResourceParameter, VertexFetch_InstanceTransformBufferParameter)
+		LAYOUT_FIELD(FShaderResourceParameter, VertexFetch_InstancePrevTransformBufferParameter)
+		LAYOUT_FIELD(FShaderResourceParameter, VertexFetch_InstanceBoneMapBufferParameter)
+	
 };
 
 /**
@@ -74,37 +67,13 @@ public:
 	//
 	// Permutations are controlled by the material flag
 	//
-	static bool ShouldCompilePermutation(EShaderPlatform Platform, const class FMaterial* Material, const class FShaderType* ShaderType) 
-	{
-		return (Material->IsUsedWithGeometryCollections() || Material->IsSpecialEngineMaterial())
-			&& FLocalVertexFactory::ShouldCompilePermutation(Platform, Material, ShaderType);
-	}
+	static bool ShouldCompilePermutation(const FVertexFactoryShaderPermutationParameters& Parameters);
 
 	//
 	// Modify compile environment to enable instancing
 	// @param OutEnvironment - shader compile environment to modify
 	//
-	static void ModifyCompilationEnvironment(const FVertexFactoryType* Type, EShaderPlatform Platform, const FMaterial* Material, FShaderCompilerEnvironment& OutEnvironment)
-	{
-		const bool ContainsManualVertexFetch = OutEnvironment.GetDefinitions().Contains("MANUAL_VERTEX_FETCH");
-		if (!ContainsManualVertexFetch && RHISupportsManualVertexFetch(Platform))
-		{
-			OutEnvironment.SetDefine(TEXT("MANUAL_VERTEX_FETCH"), TEXT("1"));
-		}
-
-		if (RHISupportsManualVertexFetch(Platform))
-		{
-			OutEnvironment.SetDefine(TEXT("USE_INSTANCING"), TEXT("1"));
-
-			OutEnvironment.SetDefine(TEXT("USE_INSTANCING_BONEMAP"), TEXT("1"));
-			OutEnvironment.SetDefine(TEXT("USE_DITHERED_LOD_TRANSITION_FOR_INSTANCED"), TEXT("0"));
-		}
-
-		// Geometry collections use a custom hit proxy per bone
-		OutEnvironment.SetDefine(TEXT("USE_PER_VERTEX_HITPROXY_ID"), 1);
-
-		FLocalVertexFactory::ModifyCompilationEnvironment(Type, Platform, Material, OutEnvironment);
-	}
+	static void ModifyCompilationEnvironment(const FVertexFactoryShaderPermutationParameters& Parameters, FShaderCompilerEnvironment& OutEnvironment);
 
 	//
 	// Set the data on the vertex factory
@@ -138,11 +107,6 @@ public:
 		FLocalVertexFactory::InitRHI();
 	}
 
-	static FVertexFactoryShaderParameters* ConstructShaderParameters(EShaderFrequency ShaderFrequency)
-	{
-		return ShaderFrequency == SF_Vertex ? new FGeometryCollectionVertexFactoryShaderParameters() : NULL;
-	}
-
 	inline void SetInstanceTransformSRV(FRHIShaderResourceView* InstanceTransformSRV)
 	{
 		Data.InstanceTransformSRV = InstanceTransformSRV;
diff --git a/Engine/Source/Runtime/Engine/Public/LocalVertexFactory.h b/Engine/Source/Runtime/Engine/Public/LocalVertexFactory.h
--- a/Engine/Source/Runtime/Engine/Public/LocalVertexFactory.h
+++ b/Engine/Source/Runtime/Engine/Public/LocalVertexFactory.h
@@ -53,9 +53,9 @@ public:
 	/**
 	 * Should we cache the material's shadertype on this platform with this vertex factory? 
 	 */
-	static bool ShouldCompilePermutation(EShaderPlatform Platform, const class FMaterial* Material, const class FShaderType* ShaderType);
+	static bool ShouldCompilePermutation(const FVertexFactoryShaderPermutationParameters& Parameters);
 
-	static void ModifyCompilationEnvironment(const FVertexFactoryType* Type, EShaderPlatform Platform, const FMaterial* Material, FShaderCompilerEnvironment& OutEnvironment);
+	static void ModifyCompilationEnvironment(const FVertexFactoryShaderPermutationParameters& Parameters, FShaderCompilerEnvironment& OutEnvironment);
 
 	static void ValidateCompiledResult(const FVertexFactoryType* Type, EShaderPlatform Platform, const FShaderParameterMap& ParameterMap, TArray<FString>& OutErrors);
 
@@ -80,8 +80,6 @@ public:
 
 	static bool SupportsTessellationShaders() { return true; }
 
-	static FVertexFactoryShaderParameters* ConstructShaderParameters(EShaderFrequency ShaderFrequency);
-
 	FORCEINLINE_DEBUGGABLE void SetColorOverrideStream(FRHICommandList& RHICmdList, const FVertexBuffer* ColorVertexBuffer) const
 	{
 		checkf(ColorVertexBuffer->IsInitialized(), TEXT("Color Vertex buffer was not initialized! Name %s"), *ColorVertexBuffer->GetFriendlyName());
@@ -164,9 +162,9 @@ protected:
  */
 class FLocalVertexFactoryShaderParametersBase : public FVertexFactoryShaderParameters
 {
+	DECLARE_TYPE_LAYOUT(FLocalVertexFactoryShaderParametersBase, NonVirtual);
 public:
-	virtual void Bind(const FShaderParameterMap& ParameterMap) override;
-	virtual void Serialize(FArchive& Ar) override;
+	void Bind(const FShaderParameterMap& ParameterMap);
 
 	void GetElementShaderBindingsBase(
 		const class FSceneInterface* Scene,
@@ -187,18 +185,18 @@ public:
 	}
 
 	// SpeedTree LOD parameter
-	FShaderParameter LODParameter;
+	LAYOUT_FIELD(FShaderParameter, LODParameter);
 
 	// True if LODParameter is bound, which puts us on the slow path in GetElementShaderBindings
-	bool bAnySpeedTreeParamIsBound;
+	LAYOUT_FIELD(bool, bAnySpeedTreeParamIsBound);
 };
 
 /** Shader parameter class used by FLocalVertexFactory only - no derived classes. */
 class FLocalVertexFactoryShaderParameters : public FLocalVertexFactoryShaderParametersBase
 {
+	DECLARE_TYPE_LAYOUT(FLocalVertexFactoryShaderParameters, NonVirtual);
 public:
-
-	virtual void GetElementShaderBindings(
+	void GetElementShaderBindings(
 		const FSceneInterface* Scene,
 		const FSceneView* View,
 		const FMeshMaterialShader* Shader,
@@ -208,5 +206,5 @@ public:
 		const FMeshBatchElement& BatchElement,
 		FMeshDrawSingleShaderBindings& ShaderBindings,
 		FVertexInputStreamArray& VertexStreams
-	) const override; 
+	) const; 
 };
\ No newline at end of file
diff --git a/Engine/Source/Runtime/Engine/Public/MaterialCachedData.h b/Engine/Source/Runtime/Engine/Public/MaterialCachedData.h
new file mode 100644
--- /dev/null
+++ b/Engine/Source/Runtime/Engine/Public/MaterialCachedData.h
@@ -0,0 +1,193 @@
+#pragma once
+
+#include "CoreMinimal.h"
+#include "Materials/MaterialLayersFunctions.h"
+#include "Misc/Guid.h"
+#include "UObject/ObjectMacros.h"
+#include "MaterialCachedData.generated.h"
+
+class UTexture;
+class UFont;
+class URuntimeVirtualTexture;
+class ULandscapeGrassType;
+
+/** Stores information about a function that this material references, used to know when the material needs to be recompiled. */
+USTRUCT()
+struct FMaterialFunctionInfo
+{
+	GENERATED_USTRUCT_BODY()
+
+	/** Id that the function had when this material was last compiled. */
+	UPROPERTY()
+	FGuid StateId;
+
+	/** The function which this material has a dependency on. */
+	UPROPERTY()
+	class UMaterialFunctionInterface* Function = nullptr;
+};
+
+/** Stores information about a parameter collection that this material references, used to know when the material needs to be recompiled. */
+USTRUCT()
+struct FMaterialParameterCollectionInfo
+{
+	GENERATED_USTRUCT_BODY()
+
+	/** Id that the collection had when this material was last compiled. */
+	UPROPERTY()
+	FGuid StateId;
+
+	/** The collection which this material has a dependency on. */
+	UPROPERTY()
+	class UMaterialParameterCollection* ParameterCollection = nullptr;
+
+	bool operator==(const FMaterialParameterCollectionInfo& Other) const
+	{
+		return StateId == Other.StateId && ParameterCollection == Other.ParameterCollection;
+	}
+};
+
+USTRUCT()
+struct FParameterChannelNames
+{
+	GENERATED_USTRUCT_BODY()
+
+	UPROPERTY(EditAnywhere, Category = MaterialExpressionVectorParameter)
+	FText R;
+
+	UPROPERTY(EditAnywhere, Category = MaterialExpressionVectorParameter)
+	FText G;
+
+	UPROPERTY(EditAnywhere, Category = MaterialExpressionVectorParameter)
+	FText B;
+
+	UPROPERTY(EditAnywhere, Category = MaterialExpressionVectorParameter)
+	FText A;
+};
+
+enum class EMaterialParameterType : int32
+{
+	Scalar,
+	Vector,
+	Texture,
+	Font,
+	RuntimeVirtualTexture,
+
+	Count,
+};
+static const int32 NumMaterialParameterTypes = (int32)EMaterialParameterType::Count;
+
+USTRUCT()
+struct FMaterialCachedParameterEntry
+{
+	GENERATED_USTRUCT_BODY()
+
+	void Reset();
+
+	UPROPERTY()
+	TArray<uint64> NameHashes;
+
+	UPROPERTY()
+	TArray<FMaterialParameterInfo> ParameterInfos;
+
+	UPROPERTY()
+	TArray<FGuid> ExpressionGuids; // editor-only?
+
+	UPROPERTY()
+	TArray<bool> Overrides;
+};
+
+USTRUCT()
+struct FMaterialCachedParameters
+{
+	GENERATED_USTRUCT_BODY()
+
+	inline int32 GetNumParameters(EMaterialParameterType Type) const { return Entries[(int32)Type].ParameterInfos.Num(); }
+	int32 FindParameterIndex(EMaterialParameterType Type, const FHashedMaterialParameterInfo& HashedParameterInfo, bool bOveriddenOnly) const;
+	int32 FindParameterIndex(EMaterialParameterType Type, const FHashedMaterialParameterInfo& HashedParameterInfo) const;
+	bool IsParameterValid(EMaterialParameterType Type, int32 Index, bool bOveriddenOnly) const;
+	void GetAllParameterInfoOfType(EMaterialParameterType Type, bool bEmptyOutput, TArray<FMaterialParameterInfo>& OutParameterInfo, TArray<FGuid>& OutParameterIds) const;
+	void GetAllGlobalParameterInfoOfType(EMaterialParameterType Type, bool bEmptyOutput, TArray<FMaterialParameterInfo>& OutParameterInfo, TArray<FGuid>& OutParameterIds) const;
+	void Reset();
+
+	UPROPERTY()
+	FMaterialCachedParameterEntry Entries[NumMaterialParameterTypes];
+
+	UPROPERTY()
+	TArray<float> ScalarValues;
+
+	UPROPERTY()
+	TArray<FLinearColor> VectorValues;
+
+	UPROPERTY()
+	TArray<UTexture*> TextureValues;
+
+	UPROPERTY()
+	TArray<UFont*> FontValues;
+
+	UPROPERTY()
+	TArray<int32> FontPageValues;
+
+	UPROPERTY()
+	TArray<URuntimeVirtualTexture*> RuntimeVirtualTextureValues;
+
+#if WITH_EDITORONLY_DATA
+	UPROPERTY()
+	TArray<FVector2D> ScalarMinMaxValues;
+
+	UPROPERTY()
+	TArray<UCurveLinearColor*> ScalarCurveValues;
+
+	UPROPERTY()
+	TArray<UCurveLinearColorAtlas*> ScalarCurveAtlasValues;
+
+	UPROPERTY()
+	TArray<FParameterChannelNames> VectorChannelNameValues;
+
+	UPROPERTY()
+	TArray<bool> VectorUsedAsChannelMaskValues;
+
+	UPROPERTY()
+	TArray<FParameterChannelNames> TextureChannelNameValues;
+#endif // WITH_EDITORONLY_DATA
+};
+
+USTRUCT()
+struct FMaterialCachedExpressionData
+{
+	GENERATED_USTRUCT_BODY()
+
+#if WITH_EDITOR
+	void UpdateForExpressions(const TArray<UMaterialExpression*>& Expressions, EMaterialParameterAssociation Association, int32 ParameterIndex);
+	void UpdateForFunction(UMaterialFunctionInterface* Function, EMaterialParameterAssociation Association, int32 ParameterIndex);
+	void UpdateForLayerFunctions(const FMaterialLayersFunctions& LayerFunctions);
+#endif // WITH_EDITOR
+
+	void Reset();
+
+	UPROPERTY()
+	FMaterialCachedParameters Parameters;
+
+	/** Array of all texture referenced by this material */
+	UPROPERTY()
+	TArray<UObject*> ReferencedTextures;
+
+	/** Array of all functions this material depends on. */
+	UPROPERTY()
+	TArray<FMaterialFunctionInfo> FunctionInfos;
+
+	/** Array of all parameter collections this material depends on. */
+	UPROPERTY()
+	TArray<FMaterialParameterCollectionInfo> ParameterCollectionInfos;
+
+	UPROPERTY()
+	TArray<ULandscapeGrassType*> GrassTypes;
+
+	UPROPERTY()
+	TArray<FName> DynamicParameterNames;
+
+	UPROPERTY()
+	TArray<bool> QualityLevelsUsed;
+
+	UPROPERTY()
+	uint32 bHasRuntimeVirtualTextureOutput : 1;
+};
diff --git a/Engine/Source/Runtime/Engine/Public/MaterialShaderType.h b/Engine/Source/Runtime/Engine/Public/MaterialShaderType.h
--- a/Engine/Source/Runtime/Engine/Public/MaterialShaderType.h
+++ b/Engine/Source/Runtime/Engine/Public/MaterialShaderType.h
@@ -51,25 +51,6 @@ extern void UpdateMaterialShaderCompilingStats(const FMaterial* Material);
  */
 extern ENGINE_API void DumpMaterialStats( EShaderPlatform Platform );
 
-struct FMaterialShaderPermutationParameters
-{
-	// Shader platform to compile to.
-	const EShaderPlatform Platform;
-
-	// Material to compile.
-	const FMaterial* Material;
-
-	/** Unique permutation identifier of the material shader type. */
-	const int32 PermutationId;
-
-	FMaterialShaderPermutationParameters(EShaderPlatform InPlatform, const FMaterial* InMaterial, int32 InPermutationId)
-		: Platform(InPlatform)
-		, Material(InMaterial)
-		, PermutationId(InPermutationId)
-	{
-	}
-};
-
 /**
  * A shader meta type for material-linked shaders.
  */
@@ -85,25 +66,21 @@ public:
 			FShaderType* InType,
 			int32 InPermutationId,
 			const FShaderCompilerOutput& CompilerOutput,
-			TRefCountPtr<FShaderResource>&& InResource,
+			int32 InResourceIndex,
 			const FUniformExpressionSet& InUniformExpressionSet,
 			const FSHAHash& InMaterialShaderMapHash,
 			const FShaderPipelineType* InShaderPipeline,
 			FVertexFactoryType* InVertexFactoryType,
 			const FString& InDebugDescription
 			)
-		: FGlobalShaderType::CompiledShaderInitializerType(InType,InPermutationId,CompilerOutput,MoveTemp(InResource),InMaterialShaderMapHash,InShaderPipeline,InVertexFactoryType)
+		: FGlobalShaderType::CompiledShaderInitializerType(InType,InPermutationId,CompilerOutput, InResourceIndex,InMaterialShaderMapHash,InShaderPipeline,InVertexFactoryType)
 		, UniformExpressionSet(InUniformExpressionSet)
 		, DebugDescription(InDebugDescription)
 		{}
 	};
 
-	typedef FShader* (*ConstructCompiledType)(const CompiledShaderInitializerType&);
-	typedef bool (*ShouldCompilePermutationType)(const FMaterialShaderPermutationParameters&);
-	typedef bool(*ValidateCompiledResultType)(EShaderPlatform, const TArray<FMaterial*>&, const FShaderParameterMap&, TArray<FString>&);
-	typedef void (*ModifyCompilationEnvironmentType)(const FMaterialShaderPermutationParameters&, FShaderCompilerEnvironment&);
-
 	FMaterialShaderType(
+		FTypeLayoutDesc& InTypeLayout,
 		const TCHAR* InName,
 		const TCHAR* InSourceFilename,
 		const TCHAR* InFunctionName,
@@ -113,13 +90,18 @@ public:
 		ConstructCompiledType InConstructCompiledRef,
 		ModifyCompilationEnvironmentType InModifyCompilationEnvironmentRef,
 		ShouldCompilePermutationType InShouldCompilePermutationRef,
-		ValidateCompiledResultType InValidateCompiledResultRef
+		ValidateCompiledResultType InValidateCompiledResultRef,
+		uint32 InTypeSize,
+		const FShaderParametersMetadata* InRootParametersMetadata = nullptr
 		):
-		FShaderType(EShaderTypeForDynamicCast::Material, InName, InSourceFilename, InFunctionName, InFrequency, InTotalPermutationCount, InConstructSerializedRef, nullptr),
-		ConstructCompiledRef(InConstructCompiledRef),
-		ShouldCompilePermutationRef(InShouldCompilePermutationRef),
-		ValidateCompiledResultRef(InValidateCompiledResultRef),
-		ModifyCompilationEnvironmentRef(InModifyCompilationEnvironmentRef)
+		FShaderType(EShaderTypeForDynamicCast::Material, InTypeLayout, InName, InSourceFilename, InFunctionName, InFrequency, InTotalPermutationCount,
+			InConstructSerializedRef,
+			InConstructCompiledRef,
+			InModifyCompilationEnvironmentRef,
+			InShouldCompilePermutationRef,
+			InValidateCompiledResultRef,
+			InTypeSize,
+			InRootParametersMetadata)
 	{
 		checkf(FPaths::GetExtension(InSourceFilename) == TEXT("usf"),
 			TEXT("Incorrect virtual shader path extension for material shader '%s': Only .usf files should be compiled."),
@@ -148,7 +130,7 @@ public:
 		const FMaterial* Material,
 		FShaderCompilerEnvironment* MaterialEnvironment,
 		const FShaderPipelineType* ShaderPipeline,
-		const TArray<FMaterialShaderType*>& ShaderStages,
+		const TArray<const FShaderType*>& ShaderStages,
 		TArray<TSharedRef<FShaderCommonCompileJob, ESPMode::ThreadSafe>>& NewJobs,
 		const FString& DebugDescription,
 		const FString& DebugExtension
@@ -164,7 +146,8 @@ public:
 		const FSHAHash& MaterialShaderMapHash,
 		const FShaderCompileJob& CurrentJob,
 		const FShaderPipelineType* ShaderPipeline,
-		const FString& InDebugDescription
+		const FString& InDebugDescription,
+		FShaderMapResourceBuilder& ResourceBuilder
 		);
 
 	/**
@@ -173,102 +156,15 @@ public:
 	 * @param Material - The material to check.
 	 * @return True if this shader type should be cached.
 	 */
-	bool ShouldCompilePermutation(EShaderPlatform Platform, const FMaterial* Material, int32 PermutationId) const
-	{
-		return (*ShouldCompilePermutationRef)(FMaterialShaderPermutationParameters(Platform, Material, PermutationId));
-	}
+	bool ShouldCompilePermutation(EShaderPlatform Platform, const FMaterialShaderParameters& MaterialParameters, int32 PermutationId) const;
 
-	/**
-	* Checks if the shader type should pass compilation for a particular set of parameters.
-	* @param Platform - Platform to validate.
-	* @param Materials - Materials to validate.
-	* @param ParameterMap - Shader parameters to validate.
-	* @param OutError - List for appending validation errors.
-	*/
-	bool ValidateCompiledResult(EShaderPlatform Platform, const TArray<FMaterial*>& Materials, const FShaderParameterMap& ParameterMap, TArray<FString>& OutError) const
-	{
-		return (*ValidateCompiledResultRef)(Platform, Materials, ParameterMap, OutError);
-	}
+	static bool ShouldCompilePipeline(const FShaderPipelineType* ShaderPipelineType, EShaderPlatform Platform, const FMaterialShaderParameters& MaterialParameters);
 
 protected:
-
 	/**
 	 * Sets up the environment used to compile an instance of this shader type.
 	 * @param Platform - Platform to compile for.
 	 * @param Environment - The shader compile environment that the function modifies.
 	 */
-	void SetupCompileEnvironment(EShaderPlatform Platform, const FMaterial* Material, int32 PermutationId, FShaderCompilerEnvironment& Environment)
-	{
-		// Allow the shader type to modify its compile environment.
-		(*ModifyCompilationEnvironmentRef)(FMaterialShaderPermutationParameters(Platform, Material, PermutationId), Environment);
-	}
-
-private:
-	ConstructCompiledType ConstructCompiledRef;
-	ShouldCompilePermutationType ShouldCompilePermutationRef;
-	ValidateCompiledResultType ValidateCompiledResultRef;
-	ModifyCompilationEnvironmentType ModifyCompilationEnvironmentRef;
+	void SetupCompileEnvironment(EShaderPlatform Platform, const FMaterialShaderParameters& MaterialParameters, int32 PermutationId, FShaderCompilerEnvironment& Environment);
 };
-
-/** DECLARE_MATERIAL_SHADER and IMPLEMENT_MATERIAL_SHADER setup a material shader class's boiler plate. They are meant to be used like so:
- *
- * class FMyMaterialShaderPS : public FMaterialShader
- * {
- *		// Setup the shader's boiler plate.
- *		DECLARE_MATERIAL_SHADER(FMyMaterialShaderPS);
- *
- *		// Setup the shader's permutation domain. If no dimensions, can do FPermutationDomain = FShaderPermutationNone.
- *		using FPermutationDomain = TShaderPermutationDomain<DIMENSIONS...>;
- *
- *		// ...
- * };
- *
- * // Instantiates shader's global variable that will take care of compilation process of the shader. This needs imperatively to be
- * done in a .cpp file regardless of whether FMyMaterialShaderPS is in a header or not.
- * IMPLEMENT_MATERIAL_SHADER(FMyMaterialShaderPS, "/Engine/Private/MyShaderFile.usf", "MainPS", SF_Pixel);
- *
- * When the shader class is a public header, let say in Engine module public header, the shader class then should have the ENGINE_API
- * like this:
- *
- * class ENGINE_API FMyMaterialShaderPS : public FMaterialShader
- * {
- *		// Setup the shader's boiler plate.
- *		DECLARE_MATERIAL_SHADER(FMyMaterialShaderPS);
- *
- *		// ...
- * };
- */
-#define DECLARE_MATERIAL_SHADER(ShaderClass) \
-	public: \
-	using ShaderMetaType = FMaterialShaderType; \
-	\
-	static ShaderMetaType StaticType; \
-	\
-	static FShader* ConstructSerializedInstance() { return new ShaderClass(); } \
-	static FShader* ConstructCompiledInstance(const ShaderMetaType::CompiledShaderInitializerType& Initializer) \
-	{ return new ShaderClass(Initializer); } \
-	\
-	virtual uint32 GetTypeSize() const override { return sizeof(*this); } \
-	\
-	static void ModifyCompilationEnvironmentImpl( \
-		const FMaterialShaderPermutationParameters& Parameters, \
-		FShaderCompilerEnvironment& OutEnvironment) \
-	{ \
-		FPermutationDomain PermutationVector(Parameters.PermutationId); \
-		PermutationVector.ModifyCompilationEnvironment(OutEnvironment); \
-		ShaderClass::ModifyCompilationEnvironment(Parameters, OutEnvironment); \
-	}
-
-#define IMPLEMENT_MATERIAL_SHADER(ShaderClass,SourceFilename,FunctionName,Frequency) \
-	ShaderClass::ShaderMetaType ShaderClass::StaticType( \
-		TEXT(#ShaderClass), \
-		TEXT(SourceFilename), \
-		TEXT(FunctionName), \
-		Frequency, \
-		ShaderClass::FPermutationDomain::PermutationCount, \
-		ShaderClass::ConstructSerializedInstance, \
-		ShaderClass::ConstructCompiledInstance, \
-		ShaderClass::ModifyCompilationEnvironmentImpl, \
-		ShaderClass::ShouldCompilePermutation, \
-		ShaderClass::ValidateCompiledResult \
-		)
