Copying Private-LoadTimes-4.24 stream to Main. Biggest changes are in Materials/Shader memory freezing.

diff --git a/Engine/Source/Runtime/Engine/Public/MaterialShared.h b/Engine/Source/Runtime/Engine/Public/MaterialShared.h
--- a/Engine/Source/Runtime/Engine/Public/MaterialShared.h
+++ b/Engine/Source/Runtime/Engine/Public/MaterialShared.h
@@ -8,6 +8,7 @@
 
 #include "CoreMinimal.h"
 #include "Containers/IndirectArray.h"
+#include "Containers/ArrayView.h"
 #include "Misc/Guid.h"
 #include "Engine/EngineTypes.h"
 #include "Templates/RefCounting.h"
@@ -33,6 +34,7 @@ class FMaterialCompiler;
 class FMaterialRenderProxy;
 class FMaterialShaderType;
 class FMaterialUniformExpression;
+class FMaterialUniformExpressionTexture;
 struct FUniformExpressionCache;
 class FUniformExpressionSet;
 class FMeshMaterialShaderType;
@@ -49,6 +51,9 @@ class URuntimeVirtualTexture;
 class USubsurfaceProfile;
 class UTexture;
 class UTexture2D;
+class FMaterialTextureParameterInfo;
+class FMaterialExternalTextureParameterInfo;
+class FMeshMaterialShaderMapLayout;
 
 template <class ElementType> class TLinkedList;
 
@@ -201,196 +206,97 @@ struct ENGINE_API FMaterialRenderContext
 		const FSceneView* InView);
 };
 
-/**
- * Represents a subclass of FMaterialUniformExpression.
- */
-class FMaterialUniformExpressionType
-{
-public:
-
-	typedef class FMaterialUniformExpression* (*SerializationConstructorType)();
-
-	/**
-	 * @return The global uniform expression type list.  The list is used to temporarily store the types until
-	 *			the name subsystem has been initialized.
-	 */
-	static TLinkedList<FMaterialUniformExpressionType*>*& GetTypeList();
-
-	/**
-	 * Should not be called until the name subsystem has been initialized.
-	 * @return The global uniform expression type map.
-	 */
-	static TMap<FName,FMaterialUniformExpressionType*>& GetTypeMap();
-
-	/**
-	 * Minimal initialization constructor.
-	 */
-	FMaterialUniformExpressionType(const TCHAR* InName,SerializationConstructorType InSerializationConstructor);
-
-	/**
-	 * Serializer for references to uniform expressions.
-	 */
-	friend FArchive& operator<<(FArchive& Ar,class FMaterialUniformExpression*& Ref);
-	friend FArchive& operator<<(FArchive& Ar,class FMaterialUniformExpressionTexture*& Ref);
-	friend FArchive& operator<<(FArchive& Ar, class FMaterialUniformExpressionExternalTexture*& Ref);
-
-	const TCHAR* GetName() const { return Name; }
-
-private:
-
-	const TCHAR* Name;
-	SerializationConstructorType SerializationConstructor;
-};
-
-#define DECLARE_MATERIALUNIFORMEXPRESSION_TYPE(Name) \
-	public: \
-	static FMaterialUniformExpressionType StaticType; \
-	static FMaterialUniformExpression* SerializationConstructor() { return new Name(); } \
-	virtual FMaterialUniformExpressionType* GetType() const { return &StaticType; }
-
-#define IMPLEMENT_MATERIALUNIFORMEXPRESSION_TYPE(Name) \
-	FMaterialUniformExpressionType Name::StaticType(TEXT(#Name),&Name::SerializationConstructor);
-
-/**
- * Represents an expression which only varies with uniform inputs.
- */
-class FMaterialUniformExpression : public FRefCountedObject
-{
-public:
-
-	virtual ~FMaterialUniformExpression() {}
-
-	virtual FMaterialUniformExpressionType* GetType() const = 0;
-	virtual void Serialize(FArchive& Ar) = 0;
-	virtual void GetNumberValue(const struct FMaterialRenderContext& Context,FLinearColor& OutValue) const {}
-	virtual class FMaterialUniformExpressionTexture* GetTextureUniformExpression() { return nullptr; }
-	virtual class FMaterialUniformExpressionExternalTexture* GetExternalTextureUniformExpression() { return nullptr; }
-	virtual bool IsConstant() const { return false; }
-	virtual bool IsIdentical(const FMaterialUniformExpression* OtherExpression) const { return false; }
-
-	friend FArchive& operator<<(FArchive& Ar,class FMaterialUniformExpression*& Ref);
-};
-
-/**
- * A texture expression.
- */
-class FMaterialUniformExpressionTexture : public FMaterialUniformExpression
+enum class EMaterialPreshaderOpcode : uint8
 {
-	DECLARE_MATERIALUNIFORMEXPRESSION_TYPE(FMaterialUniformExpressionTexture);
-
-public:
-	FMaterialUniformExpressionTexture();
-	FMaterialUniformExpressionTexture(int32 InTextureIndex, EMaterialSamplerType InSamplerType, ESamplerSourceMode InSamplerSource, bool InVirtualTexture);
-	FMaterialUniformExpressionTexture(int32 InTextureIndex, int16 InTextureLayerIndex, int16 InPageTableLayerIndex, EMaterialSamplerType InSamplerType);
-
-	/** Sets an override texture. */
-	void SetTransientOverrideTextureValue(UTexture* InOverrideTexture);
-
-	//~ Begin FMaterialUniformExpression Interface.
-	virtual void Serialize(FArchive& Ar);
-	virtual class FMaterialUniformExpressionTexture* GetTextureUniformExpression() { return this; }
-	virtual class FMaterialUniformExpressionTextureParameter* GetTextureParameterUniformExpression() { return nullptr; }
-	virtual bool IsIdentical(const FMaterialUniformExpression* OtherExpression) const;
-	//~ End FMaterialUniformExpression Interface.
-
-	friend FArchive& operator<<(FArchive& Ar, class FMaterialUniformExpressionTexture*& Ref);
-
-	/** Gets texture index which is the index in the full set of referenced textures for this material. */
-	int32 GetTextureIndex() const { return TextureIndex; }
-	/** Gets the texture layer index in the virtual texture stack if this is fixed. If we don't have a fixed layer then we will allocate during compilation (and not store here). */
-	int32 GetTextureLayerIndex() const { return TextureLayerIndex; }
-	/** Gets the page table channel index in the virtual texture stack if this is fixed. If we don't have a fixed layer then we will allocate during compilation (and not store here). */
-	int32 GetPageTableLayerIndex() const { return PageTableLayerIndex; }
-
-#if WITH_EDITORONLY_DATA
-	/** Get the sampling/decoding logic to compile in the shader for this texture. */
-	EMaterialSamplerType GetSamplerType() const { return SamplerType; }
-#endif
-
-	/** Get the sampler state object to use (globally shared or from texture asset).  */
-	ESamplerSourceMode GetSamplerSource() const { return SamplerSource; }
-
-	/** Get the UTexture referenced by this expression (can be a nullptr). */
-	virtual void GetTextureValue(const FMaterialRenderContext& Context, const FMaterial& Material, const UTexture*& OutValue) const;
-	/** Get the URuntimeVirtualTexture referenced by this expression (can be a nullptr). */
-	virtual void GetTextureValue(const FMaterialRenderContext& Context, const FMaterial& Material, const URuntimeVirtualTexture*& OutValue) const;
-
-	/** Get the UTexture referenced by this expression including any transient override logic. */
-	virtual void GetGameThreadTextureValue(const UMaterialInterface* MaterialInterface, const FMaterial& Material, UTexture*& OutValue, bool bAllowOverride = true) const;
-
-protected:
-	/** Index into FMaterial::GetReferencedTextures */
-	int32 TextureIndex;
-	/** Texture layer index in virtual texture stack if preallocated */
-	int16 TextureLayerIndex;
-	/** Page table layer index in virtual texture stack if preallocated */
-	int16 PageTableLayerIndex;
-#if WITH_EDITORONLY_DATA
-	/** Sampler logic for this expression */
-	EMaterialSamplerType SamplerType;
-#endif
-	/** Sampler state object source for this expression */
-	ESamplerSourceMode SamplerSource;
-	/** Virtual texture flag used only for unique serialization */
-	bool bVirtualTexture;
-	/** Texture that may be used in the editor for overriding the texture but never saved to disk, accessible only by the game thread! */
-	UTexture* TransientOverrideValue_GameThread;
-	/** Texture that may be used in the editor for overriding the texture but never saved to disk, accessible only by the rendering thread! */
-	UTexture* TransientOverrideValue_RenderThread;
+	Nop,
+	ConstantZero,
+	Constant,
+	ScalarParameter,
+	VectorParameter,
+	Add,
+	Sub,
+	Mul,
+	Div,
+	Fmod,
+	Min,
+	Max,
+	Clamp,
+	Sin,
+	Cos,
+	Tan,
+	Asin,
+	Acos,
+	Atan,
+	Atan2,
+	Dot,
+	Cross,
+	Sqrt,
+	Length,
+	Saturate,
+	Abs,
+	Floor,
+	Ceil,
+	Round,
+	Trunc,
+	Sign,
+	Frac,
+	Fractional,
+	Log2,
+	Log10,
+	ComponentSwizzle,
+	AppendVector,
+	TextureSize,
+	TexelSize,
+	ExternalTextureCoordinateScaleRotation,
+	ExternalTextureCoordinateOffset,
+	RuntimeVirtualTextureUniform,
 };
 
-class FMaterialUniformExpressionExternalTextureBase : public FMaterialUniformExpression
+class FMaterialPreshaderData
 {
-	DECLARE_MATERIALUNIFORMEXPRESSION_TYPE(FMaterialUniformExpressionExternalTextureBase);
+	DECLARE_TYPE_LAYOUT(FMaterialPreshaderData, NonVirtual);
 public:
+	friend inline bool operator==(const FMaterialPreshaderData& Lhs, const FMaterialPreshaderData& Rhs)
+	{
+		return Lhs.Data == Rhs.Data;
+	}
 
-	FMaterialUniformExpressionExternalTextureBase(int32 InSourceTextureIndex = INDEX_NONE);
-	FMaterialUniformExpressionExternalTextureBase(const FGuid& InExternalTextureGuid);
-
-	virtual void Serialize(FArchive& Ar) override;
-	virtual bool IsConstant() const override { return false; }
-	virtual bool IsIdentical(const FMaterialUniformExpression* OtherExpression) const override;
-
-	int32 GetSourceTextureIndex() const { return SourceTextureIndex; }
-
-protected:
-
-	/** Resolve the guid that relates to texture information inside FExternalTexture */
-	FGuid ResolveExternalTextureGUID(const FMaterialRenderContext& Context, TOptional<FName> ParameterName = TOptional<FName>()) const;
-
-	/** Index of the texture in the material that should be used to retrieve the external texture GUID at runtime (or INDEX_NONE) */
-	int32 SourceTextureIndex;
-	/** Optional external texture GUID defined at compile time */
-	FGuid ExternalTextureGuid;
-};
+	friend inline bool operator!=(const FMaterialPreshaderData& Lhs, const FMaterialPreshaderData& Rhs)
+	{
+		return !operator==(Lhs, Rhs);
+	}
 
-/**
-* An external texture expression.
-*/
-class FMaterialUniformExpressionExternalTexture : public FMaterialUniformExpressionExternalTextureBase
-{
-	DECLARE_MATERIALUNIFORMEXPRESSION_TYPE(FMaterialUniformExpressionExternalTexture);
-public:
+	friend inline FArchive& operator<<(FArchive& Ar, FMaterialPreshaderData& Ref)
+	{
+		return Ar << Ref.Data;
+	}
 
-	FMaterialUniformExpressionExternalTexture(int32 InSourceTextureIndex = INDEX_NONE) : FMaterialUniformExpressionExternalTextureBase(InSourceTextureIndex) {}
-	FMaterialUniformExpressionExternalTexture(const FGuid& InGuid) : FMaterialUniformExpressionExternalTextureBase(InGuid) {}
+	inline int32 Num() const { return Data.Num(); }
 
-	// FMaterialUniformExpression interface.
-	virtual FMaterialUniformExpressionExternalTexture* GetExternalTextureUniformExpression() override { return this; }
-	virtual class FMaterialUniformExpressionExternalTextureParameter* GetExternalTextureParameterUniformExpression() { return nullptr; }
+	void WriteData(const void* Value, uint32 Size);
+	void ReadData(void* Result, uint32 Offset, uint32 Size) const;
 
-	// Lookup the external texture if it is set
-	virtual bool GetExternalTexture(const FMaterialRenderContext& Context, FTextureRHIRef& OutTextureRHI, FSamplerStateRHIRef& OutSamplerStateRHI) const;
+	template<typename T>
+	FMaterialPreshaderData& Write(const T& Value) { WriteData(&Value, sizeof(T)); return *this; }
 
-	friend FArchive& operator<<(FArchive& Ar, FMaterialUniformExpressionExternalTexture*& Ref)
+	template<typename T>
+	T Read(uint32& Offset) const
 	{
-		Ar << (FMaterialUniformExpression*&)Ref;
-		return Ar;
+		T Result;
+		ReadData(&Result, Offset, sizeof(T));
+		Offset += sizeof(T);
+		return Result;
 	}
+
+	inline FMaterialPreshaderData& WriteOpcode(EMaterialPreshaderOpcode Op) { return Write<uint8>((uint8)Op); }
+	inline EMaterialPreshaderOpcode ReadOpcode(uint32& Offset) const { return (EMaterialPreshaderOpcode)Read<uint8>(Offset); }
+
+	LAYOUT_FIELD(TMemoryImageArray<uint8>, Data);
 };
 
 class FMaterialVirtualTextureStack
 {
+	DECLARE_TYPE_LAYOUT(FMaterialVirtualTextureStack, NonVirtual);
 public:
 	FMaterialVirtualTextureStack();
 	/** Construct with a texture index when this references a preallocated VT stack (for example when we are using a URuntimeVirtualTexture). */
@@ -438,13 +344,12 @@ public:
 		return true;
 	}
 
-private:
 	/** Number of layers that have been allocated in this stack. */
-	uint32 NumLayers;
+	LAYOUT_FIELD(uint32, NumLayers);
 	/** Indices of the expressions that were set to layers in this stack. */
-	int32 LayerUniformExpressionIndices[VIRTUALTEXTURE_SPACE_MAXLAYERS];
+	LAYOUT_ARRAY(int32, LayerUniformExpressionIndices, VIRTUALTEXTURE_SPACE_MAXLAYERS);
 	/** Index of a texture reference if we create a stack from a single known texture that has it's own layer stack. */
-	int32 PreallocatedStackTextureIndex;
+	LAYOUT_FIELD(int32, PreallocatedStackTextureIndex);
 };
 
 inline bool operator!=(const FMaterialVirtualTextureStack& Lhs, const FMaterialVirtualTextureStack& Rhs)
@@ -452,20 +357,174 @@ inline bool operator!=(const FMaterialVirtualTextureStack& Lhs, const FMaterialV
 	return !operator==(Lhs, Rhs);
 }
 
+class FMaterialUniformPreshaderHeader
+{
+	DECLARE_TYPE_LAYOUT(FMaterialUniformPreshaderHeader, NonVirtual);
+public:
+	friend inline bool operator==(const FMaterialUniformPreshaderHeader& Lhs, const FMaterialUniformPreshaderHeader& Rhs)
+	{
+		return Lhs.OpcodeOffset == Rhs.OpcodeOffset && Lhs.OpcodeSize == Rhs.OpcodeSize;
+	}
+	friend inline bool operator!=(const FMaterialUniformPreshaderHeader& Lhs, const FMaterialUniformPreshaderHeader& Rhs)
+	{
+		return !operator==(Lhs, Rhs);
+	}
+
+	LAYOUT_FIELD(uint32, OpcodeOffset);
+	LAYOUT_FIELD(uint32, OpcodeSize);
+};
+
+class FMaterialScalarParameterInfo
+{
+	DECLARE_TYPE_LAYOUT(FMaterialScalarParameterInfo, NonVirtual);
+public:
+	friend inline bool operator==(const FMaterialScalarParameterInfo& Lhs, const FMaterialScalarParameterInfo& Rhs)
+	{
+		return Lhs.ParameterInfo == Rhs.ParameterInfo && Lhs.ParameterName == Rhs.ParameterName && Lhs.DefaultValue == Rhs.DefaultValue;
+	}
+	friend inline bool operator!=(const FMaterialScalarParameterInfo& Lhs, const FMaterialScalarParameterInfo& Rhs)
+	{
+		return !operator==(Lhs, Rhs);
+	}
+
+	FMaterialParameterInfo GetParameterInfo() const { return FMaterialParameterInfo(*ParameterName, ParameterInfo.Association, ParameterInfo.Index); }
+
+	void GetGameThreadNumberValue(const UMaterialInterface* SourceMaterialToCopyFrom, float& OutValue) const;
+	void GetDefaultValue(float& OutValue) const { OutValue = DefaultValue; }
+	
+	LAYOUT_FIELD(FHashedMaterialParameterInfo, ParameterInfo);
+	LAYOUT_FIELD(FMemoryImageString, ParameterName);
+	LAYOUT_FIELD(float, DefaultValue);
+};
+
+
+class FMaterialVectorParameterInfo
+{
+	DECLARE_TYPE_LAYOUT(FMaterialVectorParameterInfo, NonVirtual);
+public:
+	friend inline bool operator==(const FMaterialVectorParameterInfo& Lhs, const FMaterialVectorParameterInfo& Rhs)
+	{
+		return Lhs.ParameterInfo == Rhs.ParameterInfo && Lhs.ParameterName == Rhs.ParameterName && Lhs.DefaultValue == Rhs.DefaultValue;
+	}
+	friend inline bool operator!=(const FMaterialVectorParameterInfo& Lhs, const FMaterialVectorParameterInfo& Rhs)
+	{
+		return !operator==(Lhs, Rhs);
+	}
+
+	FMaterialParameterInfo GetParameterInfo() const { return FMaterialParameterInfo(*ParameterName, ParameterInfo.Association, ParameterInfo.Index); }
+
+	void GetGameThreadNumberValue(const UMaterialInterface* SourceMaterialToCopyFrom, FLinearColor& OutValue) const;
+	void GetDefaultValue(FLinearColor& OutValue) const { OutValue = DefaultValue; }
+
+	LAYOUT_FIELD(FHashedMaterialParameterInfo, ParameterInfo);
+	LAYOUT_FIELD(FMemoryImageString, ParameterName);
+	LAYOUT_FIELD(FLinearColor, DefaultValue);
+};
+
+/** Must invalidate ShaderVersion.ush when changing */
+enum class EMaterialTextureParameterType : uint32
+{
+	Standard2D,
+	Cube,
+	Array2D,
+	Volume,
+	Virtual,
+
+	Count,
+};
+static const uint32 NumMaterialTextureParameterTypes = (uint32)EMaterialTextureParameterType::Count;
+
+class ENGINE_API FMaterialTextureParameterInfo
+{
+	DECLARE_TYPE_LAYOUT(FMaterialTextureParameterInfo, NonVirtual);
+public:
+	friend inline bool operator==(const FMaterialTextureParameterInfo& Lhs, const FMaterialTextureParameterInfo& Rhs)
+	{
+		return Lhs.ParameterInfo == Rhs.ParameterInfo && Lhs.ParameterName == Rhs.ParameterName && Lhs.TextureIndex == Rhs.TextureIndex && Lhs.SamplerSource == Rhs.SamplerSource && Lhs.VirtualTextureLayerIndex == Rhs.VirtualTextureLayerIndex;
+	}
+	friend inline bool operator!=(const FMaterialTextureParameterInfo& Lhs, const FMaterialTextureParameterInfo& Rhs)
+	{
+		return !operator==(Lhs, Rhs);
+	}
+
+	FMaterialParameterInfo GetParameterInfo() const { return FMaterialParameterInfo(*ParameterName, ParameterInfo.Association, ParameterInfo.Index); }
+
+	void GetGameThreadTextureValue(const UMaterialInterface* MaterialInterface, const FMaterial& Material, UTexture*& OutValue) const;
+
+	LAYOUT_FIELD(FHashedMaterialParameterInfo, ParameterInfo);
+	LAYOUT_FIELD(FMemoryImageString, ParameterName);
+	LAYOUT_FIELD_INITIALIZED(int32, TextureIndex, INDEX_NONE);
+	LAYOUT_FIELD(TEnumAsByte<ESamplerSourceMode>, SamplerSource);
+	LAYOUT_FIELD_INITIALIZED(uint8, VirtualTextureLayerIndex, 0u);
+};
+
+
+class FMaterialExternalTextureParameterInfo
+{
+	DECLARE_TYPE_LAYOUT(FMaterialExternalTextureParameterInfo, NonVirtual);
+public:
+	friend inline bool operator==(const FMaterialExternalTextureParameterInfo& Lhs, const FMaterialExternalTextureParameterInfo& Rhs)
+	{
+		return Lhs.SourceTextureIndex == Rhs.SourceTextureIndex && Lhs.ExternalTextureGuid == Rhs.ExternalTextureGuid && Lhs.ParameterName == Rhs.ParameterName;
+	}
+	friend inline bool operator!=(const FMaterialExternalTextureParameterInfo& Lhs, const FMaterialExternalTextureParameterInfo& Rhs)
+	{
+		return !operator==(Lhs, Rhs);
+	}
+
+	bool GetExternalTexture(const FMaterialRenderContext& Context, FTextureRHIRef& OutTextureRHI, FSamplerStateRHIRef& OutSamplerStateRHI) const;
+
+	LAYOUT_FIELD(FHashedName, ParameterName);
+	LAYOUT_FIELD(FGuid, ExternalTextureGuid);
+	LAYOUT_FIELD(int32, SourceTextureIndex);
+};
+
+class FUniformParameterOverrides
+{
+public:
+	void SetScalarOverride(int32 Index, float Value, bool bOverride);
+	void SetVectorOverride(int32 Index, const FLinearColor& Value, bool bOverride);
+
+	bool GetScalarOverride(int32 Index, float& OutValue) const;
+	bool GetVectorOverride(int32 Index, FLinearColor& OutValue) const;
+
+	void SetTextureOverride(EMaterialTextureParameterType Type, int32 Index, UTexture* Texture);
+	UTexture* GetTextureOverride_GameThread(EMaterialTextureParameterType Type, int32 Index) const;
+	UTexture* GetTextureOverride_RenderThread(EMaterialTextureParameterType Type, int32 Index) const;
+
+private:
+	struct FScalarOverride
+	{
+		float Value;
+		bool bOverride;
+	};
+
+	struct FVectorOverride
+	{
+		FLinearColor Value;
+		bool bOverride;
+	};
+
+	TArray<FScalarOverride> ScalarOverrides;
+	TArray<FVectorOverride> VectorOverrides;
+	TArray<UTexture*> GameThreadTextureOverides[NumMaterialTextureParameterTypes];
+	TArray<UTexture*> RenderThreadTextureOverrides[NumMaterialTextureParameterTypes];
+};
+
 /** Stores all uniform expressions for a material generated from a material translation. */
-class FUniformExpressionSet : public FRefCountedObject
+class FUniformExpressionSet
 {
+	DECLARE_TYPE_LAYOUT(FUniformExpressionSet, NonVirtual);
 public:
-	FUniformExpressionSet() {}
+	FUniformExpressionSet() : UniformBufferLayout(FRHIUniformBufferLayout::Zero) {}
 
-	ENGINE_API void Serialize(FArchive& Ar);
 	bool IsEmpty() const;
 	bool operator==(const FUniformExpressionSet& ReferenceSet) const;
 	FString GetSummaryString() const;
 
+	FShaderParametersMetadata* CreateBufferStruct();
+
 	void SetParameterCollections(const TArray<class UMaterialParameterCollection*>& Collections);
-	void CreateBufferStruct();
-	ENGINE_API const FShaderParametersMetadata& GetUniformBufferStruct() const;
 
 	ENGINE_API void FillUniformBuffer(const FMaterialRenderContext& MaterialRenderContext, const FUniformExpressionCache& UniformExpressionCache, uint8* TempBuffer, int TempBufferSize) const;
 
@@ -474,12 +533,17 @@ public:
 
 	inline bool HasExternalTextureExpressions() const
 	{
-		return UniformExternalTextureExpressions.Num() > 0;
+		return UniformExternalTextureParameters.Num() > 0;
+	}
+
+	const FRHIUniformBufferLayout& GetUniformBufferLayout() const
+	{
+		return UniformBufferLayout;
 	}
 
 	uint32 GetAllocatedSize() const
 	{
-		return UniformVectorExpressions.GetAllocatedSize()
+		/*return UniformVectorExpressions.GetAllocatedSize()
 			+ UniformScalarExpressions.GetAllocatedSize()
 			+ Uniform2DTextureExpressions.GetAllocatedSize()
 			+ UniformCubeTextureExpressions.GetAllocatedSize()
@@ -488,10 +552,22 @@ public:
 			+ UniformVirtualTextureExpressions.GetAllocatedSize()
 			+ UniformExternalTextureExpressions.GetAllocatedSize()
 			+ VTStacks.GetAllocatedSize()
-			+ ParameterCollections.GetAllocatedSize()
-			+ (UniformBufferStruct ? (sizeof(FShaderParametersMetadata) + UniformBufferStruct->GetMembers().GetAllocatedSize()) : 0);
+			+ ParameterCollections.GetAllocatedSize();*/
+		return 0u;
 	}
 
+	inline const FMaterialVectorParameterInfo& GetVectorParameter(uint32 Index) const { return UniformVectorParameters[Index]; }
+	inline const FMaterialScalarParameterInfo& GetScalarParameter(uint32 Index) const { return UniformScalarParameters[Index]; }
+	inline const FMaterialTextureParameterInfo& GetTextureParameter(EMaterialTextureParameterType Type, int32 Index) const { return UniformTextureParameters[(uint32)Type][Index]; }
+
+	const FMaterialVectorParameterInfo* FindVectorParameter(const FHashedMaterialParameterInfo& ParameterInfo) const;
+	const FMaterialScalarParameterInfo* FindScalarParameter(const FHashedMaterialParameterInfo& ParameterInfo) const;
+
+	inline int32 GetNumTextures(EMaterialTextureParameterType Type) const { return UniformTextureParameters[(uint32)Type].Num(); }
+	ENGINE_API void GetGameThreadTextureValue(EMaterialTextureParameterType Type, int32 Index, const UMaterialInterface* MaterialInterface, const FMaterial& Material, UTexture*& OutValue, bool bAllowOverride = true) const;
+	ENGINE_API void GetTextureValue(EMaterialTextureParameterType Type, int32 Index, const FMaterialRenderContext& Context, const FMaterial& Material, const UTexture*& OutValue) const;
+	ENGINE_API void GetTextureValue(int32 Index, const FMaterial& Material, const URuntimeVirtualTexture*& OutValue) const;
+
 protected:
 	union FVTPackedStackAndLayerIndex
 	{
@@ -506,24 +582,6 @@ protected:
 	};
 
 	FVTPackedStackAndLayerIndex GetVTStackAndLayerIndex(int32 UniformExpressionIndex) const;
-	
-	TArray<TRefCountPtr<FMaterialUniformExpression> > UniformVectorExpressions;
-	TArray<TRefCountPtr<FMaterialUniformExpression> > UniformScalarExpressions;
-	TArray<TRefCountPtr<FMaterialUniformExpressionTexture> > Uniform2DTextureExpressions;
-	TArray<TRefCountPtr<FMaterialUniformExpressionTexture> > UniformCubeTextureExpressions;
-	TArray<TRefCountPtr<FMaterialUniformExpressionTexture> > Uniform2DArrayTextureExpressions;
-	TArray<TRefCountPtr<FMaterialUniformExpressionTexture> > UniformVolumeTextureExpressions;
-	TArray<TRefCountPtr<FMaterialUniformExpressionTexture> > UniformVirtualTextureExpressions;
-	TArray<TRefCountPtr<FMaterialUniformExpressionExternalTexture> > UniformExternalTextureExpressions;
-
-	/** Virtual texture stacks found during compilation */
-	TArray<FMaterialVirtualTextureStack> VTStacks;
-
-	/** Ids of parameter collections referenced by the material that was translated. */
-	TArray<FGuid> ParameterCollections;
-
-	/** The structure of a uniform buffer containing values for these uniform expressions. */
-	TOptional<FShaderParametersMetadata> UniformBufferStruct;
 
 	friend class FMaterial;
 	friend class FHLSLMaterialTranslator;
@@ -532,11 +590,29 @@ protected:
 	friend class FMaterialRenderProxy;
 	friend class FMaterialVirtualTextureStack;
 	friend class FDebugUniformExpressionSet;
+
+	LAYOUT_FIELD(TMemoryImageArray<FMaterialUniformPreshaderHeader>, UniformVectorPreshaders);
+	LAYOUT_FIELD(TMemoryImageArray<FMaterialUniformPreshaderHeader>, UniformScalarPreshaders);
+	LAYOUT_FIELD(TMemoryImageArray<FMaterialScalarParameterInfo>, UniformScalarParameters);
+	LAYOUT_FIELD(TMemoryImageArray<FMaterialVectorParameterInfo>, UniformVectorParameters);
+	LAYOUT_ARRAY(TMemoryImageArray<FMaterialTextureParameterInfo>, UniformTextureParameters, NumMaterialTextureParameterTypes);
+	LAYOUT_FIELD(TMemoryImageArray<FMaterialExternalTextureParameterInfo>, UniformExternalTextureParameters);
+
+	LAYOUT_FIELD(FMaterialPreshaderData, UniformPreshaderData);
+
+	/** Virtual texture stacks found during compilation */
+	LAYOUT_FIELD(TMemoryImageArray<FMaterialVirtualTextureStack>, VTStacks);
+
+	/** Ids of parameter collections referenced by the material that was translated. */
+	LAYOUT_FIELD(TMemoryImageArray<FGuid>, ParameterCollections);
+
+	LAYOUT_FIELD(FRHIUniformBufferLayout, UniformBufferLayout);
 };
 
 /** Stores outputs from the material compile that need to be saved. */
 class FMaterialCompilationOutput
 {
+	DECLARE_TYPE_LAYOUT(FMaterialCompilationOutput, NonVirtual);
 public:
 	FMaterialCompilationOutput() :
 		UsedSceneTextures(0),
@@ -558,58 +634,9 @@ public:
 		bHasRuntimeVirtualTextureOutputNode(false)
 	{}
 
-	ENGINE_API void Serialize(FArchive& Ar);
-
 	ENGINE_API bool IsSceneTextureUsed(ESceneTextureId TexId) const { return (UsedSceneTextures & (1 << TexId)) != 0; }
 	ENGINE_API void SetIsSceneTextureUsed(ESceneTextureId TexId) { UsedSceneTextures |= (1 << TexId); }
 
-	FUniformExpressionSet UniformExpressionSet;
-
-	/** Bitfield of the ESceneTextures used */
-	uint32 UsedSceneTextures;
-
-#if WITH_EDITOR
-	/** Number of times SampleTexture is called, excludes custom nodes. */
-	uint16 EstimatedNumTextureSamplesVS;
-	uint16 EstimatedNumTextureSamplesPS;
-
-        /** Number of virtual texture lookups performed, excludes direct invocation in shaders (for example VT lightmaps) */
-	uint16 EstimatedNumVirtualTextureLookups;
-
-	/** Number of used custom UV scalars. */
-	uint8 NumUsedUVScalars;
-
-	/** Number of used custom vertex interpolation scalars. */
-	uint8 NumUsedCustomInterpolatorScalars;
-#endif // WITH_EDITOR
-
-	/** Bitfield of runtime virtual texture output attributes. */
-	uint8 RuntimeVirtualTextureOutputAttributeMask;
-
-	/** true if the material needs the scenetexture lookups. */
-	uint8 bNeedsSceneTextures : 1;
-
-	/** true if the material uses the EyeAdaptationLookup */
-	uint8 bUsesEyeAdaptation : 1;
-
-	/** true if the material modifies the the mesh position. */
-	uint8 bModifiesMeshPosition : 1;
-
-	/** Whether the material uses world position offset. */
-	uint8 bUsesWorldPositionOffset : 1;
-
-	/** true if material uses the global distance field */
-	uint8 bUsesGlobalDistanceField : 1;
-
-	/** true if the material writes a pixel depth offset */
-	uint8 bUsesPixelDepthOffset : 1;
-
-	/** true if the material uses distance cull fade */
-	uint8 bUsesDistanceCullFade : 1;
-
-	/** true if the material writes to a runtime virtual texture custom output node. */
-	uint8 bHasRuntimeVirtualTextureOutputNode : 1;
-
 	/** Indicates whether the material uses scene color. */
 	ENGINE_API bool RequiresSceneColorCopy() const { return IsSceneTextureUsed(PPI_SceneColor); }
 
@@ -640,6 +667,51 @@ public:
 
 	/** true if the material uses the Velocity SceneTexture lookup */
 	ENGINE_API bool UsesVelocitySceneTexture() const { return IsSceneTextureUsed(PPI_Velocity); }
+
+	LAYOUT_FIELD(FUniformExpressionSet, UniformExpressionSet);
+
+	/** Bitfield of the ESceneTextures used */
+	LAYOUT_FIELD(uint32, UsedSceneTextures);
+
+	/** Number of times SampleTexture is called, excludes custom nodes. */
+	LAYOUT_FIELD_EDITORONLY(uint16, EstimatedNumTextureSamplesVS);
+	LAYOUT_FIELD_EDITORONLY(uint16, EstimatedNumTextureSamplesPS);
+
+	/** Number of virtual texture lookups performed, excludes direct invocation in shaders (for example VT lightmaps) */
+	LAYOUT_FIELD_EDITORONLY(uint16, EstimatedNumVirtualTextureLookups);
+
+	/** Number of used custom UV scalars. */
+	LAYOUT_FIELD_EDITORONLY(uint8, NumUsedUVScalars);
+
+	/** Number of used custom vertex interpolation scalars. */
+	LAYOUT_FIELD_EDITORONLY(uint8, NumUsedCustomInterpolatorScalars);
+
+	/** Bitfield of runtime virtual texture output attributes. */
+	LAYOUT_FIELD(uint8, RuntimeVirtualTextureOutputAttributeMask);
+
+	/** true if the material needs the scenetexture lookups. */
+	LAYOUT_BITFIELD(uint8, bNeedsSceneTextures, 1);
+
+	/** true if the material uses the EyeAdaptationLookup */
+	LAYOUT_BITFIELD(uint8, bUsesEyeAdaptation, 1);
+
+	/** true if the material modifies the the mesh position. */
+	LAYOUT_BITFIELD(uint8, bModifiesMeshPosition, 1);
+
+	/** Whether the material uses world position offset. */
+	LAYOUT_BITFIELD(uint8, bUsesWorldPositionOffset, 1);
+
+	/** true if material uses the global distance field */
+	LAYOUT_BITFIELD(uint8, bUsesGlobalDistanceField, 1);
+
+	/** true if the material writes a pixel depth offset */
+	LAYOUT_BITFIELD(uint8, bUsesPixelDepthOffset, 1);
+
+	/** true if the material uses distance cull fade */
+	LAYOUT_BITFIELD(uint8, bUsesDistanceCullFade, 1);
+
+	/** true if the material writes to a runtime virtual texture custom output node. */
+	LAYOUT_BITFIELD(uint8, bHasRuntimeVirtualTextureOutputNode, 1);
 };
 
 /** 
@@ -825,7 +897,7 @@ public:
 	void UpdateFromParameterSet(const FStaticParameterSet& StaticParameters);
 
 	/** Appends string representations of this Id to a key string. */
-	void AppendKeyString(FString& KeyString) const;
+	void AppendKeyString(const FPlatformTypeLayoutParameters& LayoutParameters, FString& KeyString) const;
 
 	const TArray<FStaticSwitchParameter> &GetStaticSwitchParameters() const 					{ return StaticSwitchParameters; }
 	const TArray<FStaticComponentMaskParameter> &GetStaticComponentMaskParameters() const 		{ return StaticComponentMaskParameters; }
@@ -846,13 +918,13 @@ public:
 /**
  * The shaders which the render the material on a mesh generated by a particular vertex factory type.
  */
-class FMeshMaterialShaderMap : public TShaderMap<FMeshMaterialShaderType>
+class FMeshMaterialShaderMap : public FShaderMapContent
 {
+	DECLARE_TYPE_LAYOUT(FMeshMaterialShaderMap, NonVirtual);
 public:
-
 	FMeshMaterialShaderMap(EShaderPlatform InPlatform, FVertexFactoryType* InVFType) 
-		: TShaderMap<FMeshMaterialShaderType>(InPlatform)
-		, VertexFactoryType(InVFType)
+		: FShaderMapContent(InPlatform)
+		, VertexFactoryTypeName(InVFType->GetHashedName())
 	{}
 
 	/**
@@ -865,6 +937,7 @@ public:
 		uint32 ShaderMapId,
 		const FMaterialShaderMapId& InShaderMapId, 
 		const FMaterial* Material,
+		const FMeshMaterialShaderMapLayout& MeshLayout,
 		FShaderCompilerEnvironment* MaterialEnvironment,
 		EShaderPlatform Platform,
 		TArray<TSharedRef<FShaderCommonCompileJob, ESPMode::ThreadSafe>>& NewJobs,
@@ -872,20 +945,6 @@ public:
 		FString DebugExtension
 		);
 
-	/**
-	 * Checks whether a material shader map is missing any shader types necessary for the given material.
-	 * May be called with a NULL FMeshMaterialShaderMap, which is equivalent to a FMeshMaterialShaderMap with no shaders cached.
-	 * @param MeshShaderMap - The FMeshMaterialShaderMap to check for the necessary shaders.
-	 * @param Material - The material which is checked.
-	 * @return True if the shader map has all of the shader types necessary.
-	 */
-	static bool IsComplete(
-		const FMeshMaterialShaderMap* MeshShaderMap,
-		EShaderPlatform Platform,
-		const FMaterial* Material,
-		FVertexFactoryType* InVertexFactoryType,
-		bool bSilent);
-
 #if WITH_EDITOR
 	void LoadMissingShadersFromMemory(
 		const FSHAHash& MaterialShaderMapHash, 
@@ -897,7 +956,7 @@ public:
 	 * Removes all entries in the cache with exceptions based on a shader type
 	 * @param ShaderType - The shader type to flush
 	 */
-	void FlushShadersByShaderType(FShaderType* ShaderType);
+	void FlushShadersByShaderType(const FShaderType* ShaderType);
 
 		/**
 	 * Removes all entries in the cache with exceptions based on a shader type
@@ -906,29 +965,88 @@ public:
 	void FlushShadersByShaderPipelineType(const FShaderPipelineType* ShaderPipelineType);
 
 	// Accessors.
-	inline FVertexFactoryType* GetVertexFactoryType() const { return VertexFactoryType; }
+	inline const FHashedName& GetVertexFactoryTypeName() const { return VertexFactoryTypeName; }
 
 private:
 	/** The vertex factory type these shaders are for. */
-	FVertexFactoryType* VertexFactoryType;
+	LAYOUT_FIELD(FHashedName, VertexFactoryTypeName);
+};
 
-	static bool IsMeshShaderComplete(
-		const FMeshMaterialShaderMap* MeshShaderMap,
-		EShaderPlatform Platform,
-		const FMaterial* Material,
-		const FMeshMaterialShaderType* ShaderType,
-		const FShaderPipelineType* Pipeline,
-		FVertexFactoryType* InVertexFactoryType,
-		int32 PermutationId,
-		bool bSilent);
+struct FMaterialProcessedSource
+{
+	DECLARE_TYPE_LAYOUT(FMaterialProcessedSource, NonVirtual);
+public:
+	FMaterialProcessedSource() {}
+	FMaterialProcessedSource(const FHashedName& InName, const TCHAR* InSource) : Name(InName), Source(InSource) {}
+
+	LAYOUT_FIELD(FHashedName, Name);
+	LAYOUT_FIELD(FMemoryImageString, Source);
+};
+
+class FMaterialShaderMapContent : public FShaderMapContent
+{
+	friend class FMaterialShaderMap;
+	DECLARE_TYPE_LAYOUT(FMaterialShaderMapContent, NonVirtual);
+public:
+	using Super = FShaderMapContent;
+
+	inline explicit FMaterialShaderMapContent(EShaderPlatform InPlatform = EShaderPlatform::SP_NumPlatforms) : FShaderMapContent(InPlatform) {}
+	inline ~FMaterialShaderMapContent() {}
+
+	inline uint32 GetNumShaders() const
+	{
+		uint32 NumShaders = Super::GetNumShaders();
+		for (FMeshMaterialShaderMap* MeshShaderMap : OrderedMeshShaderMaps)
+		{
+			NumShaders += MeshShaderMap->GetNumShaders();
+		}
+		return NumShaders;
+	}
+
+	inline uint32 GetNumShaderPipelines() const
+	{
+		uint32 NumPipelines = Super::GetNumShaderPipelines();
+		for (FMeshMaterialShaderMap* MeshShaderMap : OrderedMeshShaderMaps)
+		{
+			NumPipelines += MeshShaderMap->GetNumShaderPipelines();
+		}
+		return NumPipelines;
+	}
+
+private:
+	struct FProjectMeshShaderMapToKey
+	{
+		inline const FHashedName& operator()(const FMeshMaterialShaderMap* InShaderMap) { return InShaderMap->GetVertexFactoryTypeName(); }
+	};
+
+	//void Serialize(FArchive& Ar, bool bInlineShaderResources, bool bLoadedByCookedMaterial);
+
+	ENGINE_API FMeshMaterialShaderMap* GetMeshShaderMap(const FHashedName& VertexFactoryTypeName) const;
+
+	void AddMeshShaderMap(const FVertexFactoryType* VertexFactoryType, FMeshMaterialShaderMap* MeshShaderMap);
+	void RemoveMeshShaderMap(const FVertexFactoryType* VertexFactoryType);
+
+	/** The material's mesh shader maps, indexed by VFType->GetId(), for fast lookup at runtime. */
+	LAYOUT_FIELD(TMemoryImageArray<TMemoryImagePtr<FMeshMaterialShaderMap>>, OrderedMeshShaderMaps);
+
+	/** Uniform expressions generated from the material compile. */
+	LAYOUT_FIELD(FMaterialCompilationOutput, MaterialCompilationOutput);
+
+	LAYOUT_FIELD(FSHAHash, ShaderContentHash);
+
+	LAYOUT_FIELD_EDITORONLY(TMemoryImageArray<FMaterialProcessedSource>, ShaderProcessedSource);
+	LAYOUT_FIELD_EDITORONLY(FMemoryImageString, FriendlyName);
+	LAYOUT_FIELD_EDITORONLY(FMemoryImageString, DebugDescription);
+	LAYOUT_FIELD_EDITORONLY(FMemoryImageString, MaterialPath);
 };
 
 /**
  * The set of material shaders for a single material.
  */
-class FMaterialShaderMap : public TShaderMap<FMaterialShaderType>, public FDeferredCleanupInterface
+class FMaterialShaderMap : public TShaderMap<FMaterialShaderMapContent, FShaderMapPointerTable>, public FDeferredCleanupInterface
 {
 public:
+	using Super = TShaderMap<FMaterialShaderMapContent, FShaderMapPointerTable>;
 
 	/**
 	 * Finds the shader map for a material.
@@ -940,27 +1058,41 @@ public:
 
 #if ALLOW_SHADERMAP_DEBUG_DATA
 	/** Flushes the given shader types from any loaded FMaterialShaderMap's. */
-	static void FlushShaderTypes(TArray<FShaderType*>& ShaderTypesToFlush, TArray<const FShaderPipelineType*>& ShaderPipelineTypesToFlush, TArray<const FVertexFactoryType*>& VFTypesToFlush);
+	static void FlushShaderTypes(TArray<const FShaderType*>& ShaderTypesToFlush, TArray<const FShaderPipelineType*>& ShaderPipelineTypesToFlush, TArray<const FVertexFactoryType*>& VFTypesToFlush);
 #endif
 
-	static void FixupShaderTypes(EShaderPlatform Platform, 
-		const TMap<FShaderType*, FString>& ShaderTypeNames,
-		const TMap<const FShaderPipelineType*, FString>& ShaderPipelineTypeNames,
-		const TMap<FVertexFactoryType*, FString>& VertexFactoryTypeNames);
-
 #if WITH_EDITOR
+	/** Gets outdated types from all loaded material shader maps */
+	static void GetAllOutdatedTypes(TArray<const FShaderType*>& OutdatedShaderTypes, TArray<const FShaderPipelineType*>& OutdatedShaderPipelineTypes, TArray<const FVertexFactoryType*>& OutdatedFactoryTypes);
+
 	/** 
 	 * Attempts to load the shader map for the given material from the Derived Data Cache.
 	 * If InOutShaderMap is valid, attempts to load the individual missing shaders instead.
 	 */
-	static void LoadFromDerivedDataCache(const FMaterial* Material, const FMaterialShaderMapId& ShaderMapId, EShaderPlatform Platform, TRefCountPtr<FMaterialShaderMap>& InOutShaderMap);
+	static void LoadFromDerivedDataCache(const FMaterial* Material, const FMaterialShaderMapId& ShaderMapId, EShaderPlatform Platform, const ITargetPlatform* TargetPlatform, TRefCountPtr<FMaterialShaderMap>& InOutShaderMap);
 #endif
 
-	inline FMaterialShaderMap() : FMaterialShaderMap(EShaderPlatform::SP_NumPlatforms) {}
-	FMaterialShaderMap(EShaderPlatform InPlatform);
+	FMaterialShaderMap();
+	virtual ~FMaterialShaderMap();
+
+	// ShaderMap interface
+	TShaderRef<FShader> GetShader(FShaderType* ShaderType, int32 PermutationId = 0) const
+	{
+		FShader* Shader = GetContent()->GetShader(ShaderType, PermutationId);
+		return TShaderRef<FShader>(Shader, *this);
+	}
+	template<typename ShaderType> TShaderRef<ShaderType> GetShader(int32 PermutationId = 0) const
+	{
+		return TShaderRef<ShaderType>::Cast(GetShader(&ShaderType::StaticType, PermutationId));
+	}
+	template<typename ShaderType> TShaderRef<ShaderType> GetShader(const typename ShaderType::FPermutationDomain& PermutationVector) const
+	{
+		return TShaderRef<ShaderType>::Cast(GetShader(&ShaderType::StaticType, PermutationVector.ToDimensionValueId()));
+	}
+
+	uint32 GetMaxNumInstructionsForShader(FShaderType* ShaderType) const { return GetContent()->GetMaxNumInstructionsForShader(*this, ShaderType); }
 
-	// Destructor.
-	~FMaterialShaderMap();
+	void FinalizeContent();
 
 	/**
 	 * Compiles the shaders for a material and caches them in this shader map.
@@ -1003,17 +1135,17 @@ public:
 	bool TryToAddToExistingCompilationTask(FMaterial* Material);
 
 #if WITH_EDITOR
-	ENGINE_API const FString *GetShaderSource(const FName ShaderTypeName) const;
+	ENGINE_API const FMemoryImageString *GetShaderSource(const FName ShaderTypeName) const;
 #endif
 
 	/** Builds a list of the shaders in a shader map. */
-	ENGINE_API void GetShaderList(TMap<FShaderId, FShader*>& OutShaders) const;
+	ENGINE_API void GetShaderList(TMap<FShaderId, TShaderRef<FShader>>& OutShaders) const;
 
 	/** Builds a list of the shaders in a shader map. Key is FShaderType::TypeName */
-	ENGINE_API void GetShaderList(TMap<FName, FShader*>& OutShaders) const;
+	ENGINE_API void GetShaderList(TMap<FHashedName, TShaderRef<FShader>>& OutShaders) const;
 
 	/** Builds a list of the shader pipelines in a shader map. */
-	ENGINE_API void GetShaderPipelineList(TArray<FShaderPipeline*>& OutShaderPipelines) const;
+	ENGINE_API void GetShaderPipelineList(TArray<FShaderPipelineRef>& OutShaderPipelines) const;
 
 
 	/** Number of Shaders in Shadermap */
@@ -1030,7 +1162,7 @@ public:
 	 * Removes all entries in the cache with exceptions based on a shader type
 	 * @param ShaderType - The shader type to flush
 	 */
-	void FlushShadersByShaderType(FShaderType* ShaderType);
+	void FlushShadersByShaderType(const FShaderType* ShaderType);
 
 	/**
 	 * Removes all entries in the cache with exceptions based on a shader pipeline type
@@ -1055,13 +1187,9 @@ public:
 
 #if WITH_EDITOR
 	/** Saves this shader map to the derived data cache. */
-	void SaveToDerivedDataCache();
+	void SaveToDerivedDataCache(const ITargetPlatform* TargetPlatform);
 #endif
 
-	/** Registers all shaders that have been loaded in Serialize */
-	virtual void RegisterSerializedShaders(bool bLoadedByCookedMaterial) override;
-	virtual void DiscardSerializedShaders() override;
-
 	/** Backs up any FShaders in this shader map to memory through serialization and clears FShader references. */
 	TArray<uint8>* BackupShadersToMemory();
 	/** Recreates FShaders from the passed in memory, handling shader key changes. */
@@ -1071,61 +1199,50 @@ public:
 	ENGINE_API static void SaveForRemoteRecompile(FArchive& Ar, const TMap<FString, TArray<TRefCountPtr<FMaterialShaderMap> > >& CompiledShaderMaps, const TArray<FShaderResourceId>& ClientResourceIds);
 	ENGINE_API static void LoadForRemoteRecompile(FArchive& Ar, EShaderPlatform ShaderPlatform, const TArray<FString>& MaterialsForShaderMaps);
 
-	/** Computes the memory used by this shader map without counting the shaders themselves. */
-	uint32 GetSizeBytes() const
-	{
-		return sizeof(*this)
-			+ MeshShaderMaps.GetAllocatedSize()
-			+ OrderedMeshShaderMaps.GetAllocatedSize()
-			+ MaterialCompilationOutput.UniformExpressionSet.GetAllocatedSize()
-#if ALLOW_SHADERMAP_DEBUG_DATA
-			+ FriendlyName.GetAllocatedSize()
-			+ DebugDescription.GetAllocatedSize()
-			+ MaterialPath.GetAllocatedSize()
-#endif
-		;
-	}
-
 #if WITH_EDITOR
 	/** Returns the maximum number of texture samplers used by any shader in this shader map. */
 	uint32 GetMaxTextureSamplers() const;
+
+	void SaveShaderStableKeys(EShaderPlatform TargetShaderPlatform, const struct FStableShaderKeyAndValue& SaveKeyVal);
 #endif
 
 	// Accessors.
-	ENGINE_API const FMeshMaterialShaderMap* GetMeshShaderMap(FVertexFactoryType* VertexFactoryType) const;
+	FMeshMaterialShaderMap* GetMeshShaderMap(FVertexFactoryType* VertexFactoryType) const { return GetContent()->GetMeshShaderMap(VertexFactoryType->GetHashedName()); }
+	FMeshMaterialShaderMap* GetMeshShaderMap(const FHashedName& VertexFactoryTypeName) const { return GetContent()->GetMeshShaderMap(VertexFactoryTypeName); }
 	const FMaterialShaderMapId& GetShaderMapId() const { return ShaderMapId; }
 	uint32 GetCompilingId() const { return CompilingId; }
 	bool IsCompilationFinalized() const { return bCompilationFinalized; }
 	bool CompiledSuccessfully() const { return bCompiledSuccessfully; }
-#if ALLOW_SHADERMAP_DEBUG_DATA
-	const FString& GetFriendlyName() const { return FriendlyName; }
-	const FString& GetDebugDescription() const { return DebugDescription; }
-	void SetMaterialPath(FString Path){ MaterialPath = Path;}
+
+#if WITH_EDITORONLY_DATA
+	const TCHAR* GetFriendlyName() const { return *GetContent()->FriendlyName; }
+	const TCHAR* GetDebugDescription() const { return *GetContent()->DebugDescription; }
+	const TCHAR* GetMaterialPath() const { return *GetContent()->MaterialPath; }
 #else
-	const FString& GetFriendlyName() const { static FString T; return T; }
-	const FString& GetDebugDescription() const { static FString T; return T; }
-	void SetMaterialPath(FString Path) { }
+	const TCHAR* GetFriendlyName() const { return TEXT(""); }
+	const TCHAR* GetDebugDescription() const { return TEXT(""); }
+	const TCHAR* GetMaterialPath() const { return TEXT(""); }
 #endif
-	bool RequiresSceneColorCopy() const { return MaterialCompilationOutput.RequiresSceneColorCopy(); }
-	bool NeedsSceneTextures() const { return MaterialCompilationOutput.bNeedsSceneTextures; }
-	bool UsesGlobalDistanceField() const { return MaterialCompilationOutput.bUsesGlobalDistanceField; }
-	bool UsesWorldPositionOffset() const { return MaterialCompilationOutput.bUsesWorldPositionOffset; }
-	bool NeedsGBuffer() const { return MaterialCompilationOutput.NeedsGBuffer(); }
-	bool UsesEyeAdaptation() const { return MaterialCompilationOutput.bUsesEyeAdaptation; }
-	bool ModifiesMeshPosition() const { return MaterialCompilationOutput.bModifiesMeshPosition; }
-	bool UsesPixelDepthOffset() const { return MaterialCompilationOutput.bUsesPixelDepthOffset; }
-	bool UsesSceneDepthLookup() const { return MaterialCompilationOutput.UsesSceneDepthLookup(); }
-	bool UsesVelocitySceneTexture() const { return MaterialCompilationOutput.UsesVelocitySceneTexture(); }
-	bool UsesDistanceCullFade() const { return MaterialCompilationOutput.bUsesDistanceCullFade; }
+	bool RequiresSceneColorCopy() const { return GetContent()->MaterialCompilationOutput.RequiresSceneColorCopy(); }
+	bool NeedsSceneTextures() const { return GetContent()->MaterialCompilationOutput.bNeedsSceneTextures; }
+	bool UsesGlobalDistanceField() const { return GetContent()->MaterialCompilationOutput.bUsesGlobalDistanceField; }
+	bool UsesWorldPositionOffset() const { return GetContent()->MaterialCompilationOutput.bUsesWorldPositionOffset; }
+	bool NeedsGBuffer() const { return GetContent()->MaterialCompilationOutput.NeedsGBuffer(); }
+	bool UsesEyeAdaptation() const { return GetContent()->MaterialCompilationOutput.bUsesEyeAdaptation; }
+	bool ModifiesMeshPosition() const { return GetContent()->MaterialCompilationOutput.bModifiesMeshPosition; }
+	bool UsesPixelDepthOffset() const { return GetContent()->MaterialCompilationOutput.bUsesPixelDepthOffset; }
+	bool UsesSceneDepthLookup() const { return GetContent()->MaterialCompilationOutput.UsesSceneDepthLookup(); }
+	bool UsesVelocitySceneTexture() const { return GetContent()->MaterialCompilationOutput.UsesVelocitySceneTexture(); }
+	bool UsesDistanceCullFade() const { return GetContent()->MaterialCompilationOutput.bUsesDistanceCullFade; }
 #if WITH_EDITOR
-	uint32 GetNumUsedUVScalars() const { return MaterialCompilationOutput.NumUsedUVScalars; }
-	uint32 GetNumUsedCustomInterpolatorScalars() const { return MaterialCompilationOutput.NumUsedCustomInterpolatorScalars; }
-	void GetEstimatedNumTextureSamples(uint32& VSSamples, uint32& PSSamples) const { VSSamples = MaterialCompilationOutput.EstimatedNumTextureSamplesVS; PSSamples = MaterialCompilationOutput.EstimatedNumTextureSamplesPS; }
-	uint32 GetEstimatedNumVirtualTextureLookups() const { return MaterialCompilationOutput.EstimatedNumVirtualTextureLookups; }
+	uint32 GetNumUsedUVScalars() const { return GetContent()->MaterialCompilationOutput.NumUsedUVScalars; }
+	uint32 GetNumUsedCustomInterpolatorScalars() const { return GetContent()->MaterialCompilationOutput.NumUsedCustomInterpolatorScalars; }
+	void GetEstimatedNumTextureSamples(uint32& VSSamples, uint32& PSSamples) const { VSSamples = GetContent()->MaterialCompilationOutput.EstimatedNumTextureSamplesVS; PSSamples = GetContent()->MaterialCompilationOutput.EstimatedNumTextureSamplesPS; }
+	uint32 GetEstimatedNumVirtualTextureLookups() const { return GetContent()->MaterialCompilationOutput.EstimatedNumVirtualTextureLookups; }
 #endif
-	uint32 GetNumVirtualTextureStacks() const { return MaterialCompilationOutput.UniformExpressionSet.VTStacks.Num(); }
-	uint8 GetRuntimeVirtualTextureOutputAttributeMask() const { return MaterialCompilationOutput.RuntimeVirtualTextureOutputAttributeMask; }
-	bool UsesSceneTexture(uint32 TexId) const { return (MaterialCompilationOutput.UsedSceneTextures & (1ull << TexId)) != 0; }
+	uint32 GetNumVirtualTextureStacks() const { return GetContent()->MaterialCompilationOutput.UniformExpressionSet.VTStacks.Num(); }
+	uint8 GetRuntimeVirtualTextureOutputAttributeMask() const { return GetContent()->MaterialCompilationOutput.RuntimeVirtualTextureOutputAttributeMask; }
+	bool UsesSceneTexture(uint32 TexId) const { return (GetContent()->MaterialCompilationOutput.UsedSceneTextures & (1ull << TexId)) != 0; }
 
 	bool IsValidForRendering(bool bFailOnInvalid = false) const
 	{
@@ -1134,16 +1251,17 @@ public:
 		return bValid;
 	}
 
-	const FUniformExpressionSet& GetUniformExpressionSet() const { return MaterialCompilationOutput.UniformExpressionSet; }
+	const FUniformExpressionSet& GetUniformExpressionSet() const { return GetContent()->MaterialCompilationOutput.UniformExpressionSet; }
 
 	int32 GetNumRefs() const { return NumRefs; }
+	int32 GetRefCount() const { return NumRefs; }
 
 	void CountNumShaders(int32& NumShaders, int32& NumPipelines) const
 	{
-		NumShaders = GetNumShaders();
-		NumPipelines = GetNumShaderPipelines();
+		NumShaders = GetContent()->GetNumShaders();
+		NumPipelines = GetContent()->GetNumShaderPipelines();
 
-		for (FMeshMaterialShaderMap* MeshShaderMap : OrderedMeshShaderMaps)
+		for (FMeshMaterialShaderMap* MeshShaderMap : GetContent()->OrderedMeshShaderMaps)
 		{
 			if (MeshShaderMap)
 			{
@@ -1153,10 +1271,8 @@ public:
 		}
 	}
 	void DumpDebugInfo();
-	void SaveShaderStableKeys(EShaderPlatform TargetShaderPlatform, const struct FStableShaderKeyAndValue& SaveKeyVal);
 
 private:
-
 	/** 
 	 * A global map from a material's static parameter set to any shader map cached for that material. 
 	 * Note: this does not necessarily contain all material shader maps in memory.  Shader maps with the same key can evict each other.
@@ -1173,24 +1289,15 @@ private:
 	static TArray<FMaterialShaderMap*> AllMaterialShaderMaps;
 #endif
 
-	/** The material's cached shaders for vertex factory type dependent shaders. */
-	TIndirectArray<FMeshMaterialShaderMap> MeshShaderMaps;
-
-	/** The material's mesh shader maps, indexed by VFType->GetId(), for fast lookup at runtime. */
-	TArray<FMeshMaterialShaderMap*> OrderedMeshShaderMaps;
-
 #if ALLOW_SHADERMAP_DEBUG_DATA
-	/** The material's user friendly name, typically the object name. */
-	FString FriendlyName;
-	FString MaterialPath;
 	float CompileTime;
 #endif
 
 	/** The static parameter set that this shader map was compiled with */
 	FMaterialShaderMapId ShaderMapId;
 
-	/** Uniform expressions generated from the material compile. */
-	FMaterialCompilationOutput MaterialCompilationOutput;
+	/** The platform being compiled, or nullptr for current platform */
+	const ITargetPlatform* CompilingTargetPlatform;
 
 	/** Tracks material resources and their shader maps that need to be compiled but whose compilation is being deferred. */
 	static TMap<TRefCountPtr<FMaterialShaderMap>, TArray<FMaterial*> > ShaderMapsBeingCompiled;
@@ -1217,27 +1324,7 @@ private:
 	/** Indicates whether the shader map should be stored in the shader cache. */
 	uint32 bIsPersistent : 1;
 
-#if ALLOW_SHADERMAP_DEBUG_DATA
-	/** Debug information about how the material shader map was compiled. */
-	FString DebugDescription;
-
-#if WITH_EDITOR
-	/** Map that stores source code used to compile a certain shader present in this FMaterialShaderMap instance. Key is ShaderType::TypeName. */
-	TMap<FName, FString> ShaderProcessedSource;
-#endif
-#endif
-
-	FShader* ProcessCompilationResultsForSingleJob(class FShaderCompileJob* SingleJob, const FShaderPipelineType* ShaderPipeline, const FSHAHash& MaterialShaderMapHash);
-
-	bool IsMaterialShaderComplete(
-		const FMaterial* Material,
-		const FMaterialShaderType* ShaderType,
-		const FShaderPipelineType* Pipeline,
-		int32 PermutationId,
-		bool bSilent) const;
-
-	/** Initializes OrderedMeshShaderMaps from the contents of MeshShaderMaps. */
-	void InitOrderedMeshShaderMaps();
+	FShader* ProcessCompilationResultsForSingleJob(class FShaderCompileJob* SingleJob, const FShaderPipelineType* ShaderPipeline, const FSHAHash& MaterialShaderMapHash, FShaderMapResourceBuilder& ResourceBuilder);
 
 	friend ENGINE_API void DumpMaterialStats( EShaderPlatform Platform );
 	friend class FShaderCompilingManager;
@@ -1381,11 +1468,19 @@ public:
 	FMaterial():
 		RenderingThreadShaderMap(NULL),
 		QualityLevel(EMaterialQualityLevel::High),
-		bHasQualityLevelUsage(false),
 		FeatureLevel(ERHIFeatureLevel::SM5),
+		bHasQualityLevelUsage(false),
 		bContainsInlineShaders(false),
 		bLoadedCookedShaderMapId(false)
-	{}
+	{
+		// this option affects only deferred renderer
+		static TConsoleVariableData<int32>* CVarStencilDitheredLOD;
+		if (CVarStencilDitheredLOD == nullptr)
+		{
+			CVarStencilDitheredLOD = IConsoleManager::Get().FindTConsoleVariableDataInt(TEXT("r.StencilForLODDither"));
+		}
+		bStencilDitheredLOD = (CVarStencilDitheredLOD->GetValueOnAnyThread() != 0);
+	}
 
 	/**
 	 * Destructor
@@ -1402,7 +1497,7 @@ public:
 	 * Caches the material shaders for the given static parameter set and platform.
 	 * This is used by material resources of UMaterialInstances.
 	 */
-	ENGINE_API bool CacheShaders(const FMaterialShaderMapId& ShaderMapId, const FStaticParameterSet *StaticParameterSet, EShaderPlatform Platform, const ITargetPlatform* TargetPlatform = nullptr);
+	ENGINE_API bool CacheShaders(const FMaterialShaderMapId& ShaderMapId, EShaderPlatform Platform, const ITargetPlatform* TargetPlatform = nullptr);
 
 	/**
 	 * Should the shader for this material with the given platform, shader type and vertex 
@@ -1416,6 +1511,8 @@ public:
 	 */
 	ENGINE_API virtual bool ShouldCache(EShaderPlatform Platform, const FShaderType* ShaderType, const FVertexFactoryType* VertexFactoryType) const;
 
+	ENGINE_API bool ShouldCachePipeline(EShaderPlatform Platform, const FShaderPipelineType* PipelineType, const FVertexFactoryType* VertexFactoryType) const;
+
 	/** Serializes the material. */
 	ENGINE_API virtual void LegacySerialize(FArchive& Ar);
 
@@ -1433,7 +1530,9 @@ public:
 
 	// Material properties.
 	ENGINE_API virtual void GetShaderMapId(EShaderPlatform Platform, FMaterialShaderMapId& OutId) const;
+#if WITH_EDITORONLY_DATA
 	ENGINE_API virtual void GetStaticParameterSet(EShaderPlatform Platform, FStaticParameterSet& OutSet) const;
+#endif // WITH_EDITORONLY_DATA
 	virtual EMaterialDomain GetMaterialDomain() const = 0; // See EMaterialDomain.
 	virtual bool IsTwoSided() const = 0;
 	virtual bool IsDitheredLODTransition() const = 0;
@@ -1586,14 +1685,21 @@ public:
 		return QualityLevel;
 	}
 
+#if WITH_EDITOR
+	FUniformParameterOverrides TransientOverrides;
+#endif // WITH_EDITOR
+
 	// Accessors.
-	ENGINE_API const TArray<TRefCountPtr<FMaterialUniformExpressionTexture> >& GetUniform2DTextureExpressions() const;
-	ENGINE_API const TArray<TRefCountPtr<FMaterialUniformExpressionTexture> >& GetUniformCubeTextureExpressions() const;
-	ENGINE_API const TArray<TRefCountPtr<FMaterialUniformExpressionTexture> >& GetUniform2DArrayTextureExpressions() const;
-	ENGINE_API const TArray<TRefCountPtr<FMaterialUniformExpressionTexture> >& GetUniformVolumeTextureExpressions() const;
-	ENGINE_API const TArray<TRefCountPtr<FMaterialUniformExpressionTexture> >& GetUniformVirtualTextureExpressions() const;
-	ENGINE_API const TArray<TRefCountPtr<FMaterialUniformExpression> >& GetUniformVectorParameterExpressions() const;
-	ENGINE_API const TArray<TRefCountPtr<FMaterialUniformExpression> >& GetUniformScalarParameterExpressions() const;
+	ENGINE_API const FUniformExpressionSet& GetUniformExpressions() const;
+	ENGINE_API const TArray<FMaterialTextureParameterInfo, FMemoryImageAllocator>& GetUniformTextureExpressions(EMaterialTextureParameterType Type) const;
+	ENGINE_API const TArray<FMaterialVectorParameterInfo, FMemoryImageAllocator>& GetUniformVectorParameterExpressions() const;
+	ENGINE_API const TArray<FMaterialScalarParameterInfo, FMemoryImageAllocator>& GetUniformScalarParameterExpressions() const;
+
+	inline const TArray<FMaterialTextureParameterInfo, FMemoryImageAllocator>& GetUniform2DTextureExpressions() const { return GetUniformTextureExpressions(EMaterialTextureParameterType::Standard2D); }
+	inline const TArray<FMaterialTextureParameterInfo, FMemoryImageAllocator>& GetUniformCubeTextureExpressions() const { return GetUniformTextureExpressions(EMaterialTextureParameterType::Cube); }
+	inline const TArray<FMaterialTextureParameterInfo, FMemoryImageAllocator>& GetUniform2DArrayTextureExpressions() const { return GetUniformTextureExpressions(EMaterialTextureParameterType::Array2D); }
+	inline const TArray<FMaterialTextureParameterInfo, FMemoryImageAllocator>& GetUniformVolumeTextureExpressions() const { return GetUniformTextureExpressions(EMaterialTextureParameterType::Volume); }
+	inline const TArray<FMaterialTextureParameterInfo, FMemoryImageAllocator>& GetUniformVirtualTextureExpressions() const { return GetUniformTextureExpressions(EMaterialTextureParameterType::Virtual); }
 
 #if WITH_EDITOR
 	const TArray<FString>& GetCompileErrors() const { return CompileErrors; }
@@ -1602,7 +1708,7 @@ public:
 	const FGuid& GetLegacyId() const { return Id_DEPRECATED; }
 #endif // WITH_EDITOR
 
-	ERHIFeatureLevel::Type GetFeatureLevel() const { return FeatureLevel; }
+	const FStaticFeatureLevel GetFeatureLevel() const { return FeatureLevel; }
 	bool GetUsesDynamicParameter() const 
 	{ 
 		//@todo - remove non-dynamic parameter particle VF and always support dynamic parameter
@@ -1685,25 +1791,25 @@ public:
 
 	ENGINE_API virtual void AddReferencedObjects(FReferenceCollector& Collector);
 
-	virtual const TArray<UObject*>& GetReferencedTextures() const = 0;
+	virtual TArrayView<UObject* const> GetReferencedTextures() const = 0;
 
 	/**
 	 * Finds the shader matching the template type and the passed in vertex factory, asserts if not found.
 	 * Note - Only implemented for FMeshMaterialShaderTypes
 	 */
 	template<typename ShaderType>
-	ShaderType* GetShader(FVertexFactoryType* VertexFactoryType, const typename ShaderType::FPermutationDomain& PermutationVector, bool bFatalIfMissing = true) const
+	TShaderRef<ShaderType> GetShader(FVertexFactoryType* VertexFactoryType, const typename ShaderType::FPermutationDomain& PermutationVector, bool bFatalIfMissing = true) const
 	{
 		return GetShader<ShaderType>(VertexFactoryType, PermutationVector.ToDimensionValueId(), bFatalIfMissing);
 	}
 
 	template <typename ShaderType>
-	ShaderType* GetShader(FVertexFactoryType* VertexFactoryType, int32 PermutationId = 0, bool bFatalIfMissing = true) const
+	TShaderRef<ShaderType> GetShader(FVertexFactoryType* VertexFactoryType, int32 PermutationId = 0, bool bFatalIfMissing = true) const
 	{
-		return (ShaderType*)GetShader(&ShaderType::StaticType, VertexFactoryType, PermutationId, bFatalIfMissing);
+		return TShaderRef<ShaderType>::Cast(GetShader(&ShaderType::StaticType, VertexFactoryType, PermutationId, bFatalIfMissing));
 	}
 
-	ENGINE_API FShaderPipeline* GetShaderPipeline(class FShaderPipelineType* ShaderPipelineType, FVertexFactoryType* VertexFactoryType, bool bFatalIfNotFound = true) const;
+	ENGINE_API FShaderPipelineRef GetShaderPipeline(class FShaderPipelineType* ShaderPipelineType, FVertexFactoryType* VertexFactoryType, bool bFatalIfNotFound = true) const;
 
 	/** Returns a string that describes the material's usage for debugging purposes. */
 	virtual FString GetMaterialUsageDescription() const = 0;
@@ -1781,14 +1887,13 @@ protected:
 	virtual int32 CompileCustomAttribute(const FGuid& AttributeID, class FMaterialCompiler* Compiler) const {return INDEX_NONE;}
 #endif
 
+#if WITH_EDITORONLY_DATA
 	/* Gather any UMaterialExpressionCustomOutput expressions they can be compiled in turn */
 	virtual void GatherCustomOutputExpressions(TArray<class UMaterialExpressionCustomOutput*>& OutCustomOutputs) const {}
 
 	/* Gather any UMaterialExpressionCustomOutput expressions in the material and referenced function calls */
 	virtual void GatherExpressionsForCustomInterpolators(TArray<class UMaterialExpression*>& OutExpressions) const {}
-
-	/** Returns the index to the Expression in the Expressions array, or -1 if not found. */
-	int32 FindExpression(const TArray<TRefCountPtr<FMaterialUniformExpressionTexture> >&Expressions, const FMaterialUniformExpressionTexture &Expression);
+#endif // WITH_EDITORONLY_DATA
 
 	/** Useful for debugging. */
 	virtual FString GetBaseMaterialPathName() const { return TEXT(""); }
@@ -1863,12 +1968,15 @@ private:
 	/** Quality level that this material is representing. */
 	EMaterialQualityLevel::Type QualityLevel;
 
-	/** Whether this material has quality level specific nodes. */
-	bool bHasQualityLevelUsage;
-
 	/** Feature level that this material is representing. */
 	ERHIFeatureLevel::Type FeatureLevel;
 
+	/** Whether this material has quality level specific nodes. */
+	uint32 bHasQualityLevelUsage : 1;
+
+	/** Whether tthis project uses stencil dither lod. */
+	uint32 bStencilDitheredLOD : 1;
+
 	/** 
 	 * Whether this material was loaded with shaders inlined. 
 	 * If true, GameThreadShaderMap will contain a reference to the inlined shader map between Serialize and PostLoad.
@@ -1889,6 +1997,7 @@ private:
 	/** Populates OutEnvironment with defines needed to compile shaders for this material. */
 	void SetupMaterialEnvironment(
 		EShaderPlatform Platform,
+		const FShaderParametersMetadata& InUniformBufferStruct,
 		const FUniformExpressionSet& InUniformExpressionSet,
 		FShaderCompilerEnvironment& OutEnvironment
 		) const;
@@ -1896,7 +2005,7 @@ private:
 	/**
 	 * Finds the shader matching the template type and the passed in vertex factory, asserts if not found.
 	 */
-	ENGINE_API FShader* GetShader(class FMeshMaterialShaderType* ShaderType, FVertexFactoryType* VertexFactoryType, int32 PermutationId, bool bFatalIfMissing = true) const;
+	ENGINE_API TShaderRef<FShader> GetShader(class FMeshMaterialShaderType* ShaderType, FVertexFactoryType* VertexFactoryType, int32 PermutationId, bool bFatalIfMissing = true) const;
 
 	void GetReferencedTexturesHash(EShaderPlatform Platform, FSHAHash& OutHash) const;
 
@@ -2006,12 +2115,10 @@ public:
 	/** Returns the FMaterial, without using a fallback if the FMaterial doesn't have a valid shader map. Can return NULL. */
 	virtual FMaterial* GetMaterialNoFallback(ERHIFeatureLevel::Type InFeatureLevel) const { return NULL; }
 	virtual UMaterialInterface* GetMaterialInterface() const { return NULL; }
-
-	virtual bool GetVectorValue(const FMaterialParameterInfo& ParameterInfo, FLinearColor* OutValue, const FMaterialRenderContext& Context) const = 0;
-	virtual bool GetScalarValue(const FMaterialParameterInfo& ParameterInfo, float* OutValue, const FMaterialRenderContext& Context) const = 0;
-	virtual bool GetTextureValue(const FMaterialParameterInfo& ParameterInfo,const UTexture** OutValue, const FMaterialRenderContext& Context) const = 0;
-	virtual bool GetTextureValue(const FMaterialParameterInfo& ParameterInfo, const URuntimeVirtualTexture** OutValue, const FMaterialRenderContext& Context) const = 0;
-
+	virtual bool GetVectorValue(const FHashedMaterialParameterInfo& ParameterInfo, FLinearColor* OutValue, const FMaterialRenderContext& Context) const = 0;
+	virtual bool GetScalarValue(const FHashedMaterialParameterInfo& ParameterInfo, float* OutValue, const FMaterialRenderContext& Context) const = 0;
+	virtual bool GetTextureValue(const FHashedMaterialParameterInfo& ParameterInfo,const UTexture** OutValue, const FMaterialRenderContext& Context) const = 0;
+	virtual bool GetTextureValue(const FHashedMaterialParameterInfo& ParameterInfo, const URuntimeVirtualTexture** OutValue, const FMaterialRenderContext& Context) const = 0;
 	bool IsDeleted() const
 	{
 		return DeletedFlag != 0;
@@ -2094,11 +2201,11 @@ public:
 	{}
 
 	// FMaterialRenderProxy interface.
-	ENGINE_API virtual const FMaterial& GetMaterialWithFallback(ERHIFeatureLevel::Type InFeatureLevel, const FMaterialRenderProxy*& OutFallbackMaterialRenderProxy) const;
-	ENGINE_API virtual bool GetVectorValue(const FMaterialParameterInfo& ParameterInfo, FLinearColor* OutValue, const FMaterialRenderContext& Context) const;
-	ENGINE_API virtual bool GetScalarValue(const FMaterialParameterInfo& ParameterInfo, float* OutValue, const FMaterialRenderContext& Context) const;
-	ENGINE_API virtual bool GetTextureValue(const FMaterialParameterInfo& ParameterInfo,const UTexture** OutValue, const FMaterialRenderContext& Context) const;
-	ENGINE_API virtual bool GetTextureValue(const FMaterialParameterInfo& ParameterInfo, const URuntimeVirtualTexture** OutValue, const FMaterialRenderContext& Context) const;
+	ENGINE_API virtual const FMaterial& GetMaterialWithFallback(ERHIFeatureLevel::Type InFeatureLevel, const FMaterialRenderProxy*& OutFallbackMaterialRenderProxy) const override;
+	ENGINE_API virtual bool GetVectorValue(const FHashedMaterialParameterInfo& ParameterInfo, FLinearColor* OutValue, const FMaterialRenderContext& Context) const override;
+	ENGINE_API virtual bool GetScalarValue(const FHashedMaterialParameterInfo& ParameterInfo, float* OutValue, const FMaterialRenderContext& Context) const override;
+	ENGINE_API virtual bool GetTextureValue(const FHashedMaterialParameterInfo& ParameterInfo,const UTexture** OutValue, const FMaterialRenderContext& Context) const override;
+	ENGINE_API virtual bool GetTextureValue(const FHashedMaterialParameterInfo& ParameterInfo, const URuntimeVirtualTexture** OutValue, const FMaterialRenderContext& Context) const;
 };
 
 
@@ -2119,7 +2226,7 @@ public:
 		TextureParamName(InTextureParamName)
 	{}
 
-	ENGINE_API virtual bool GetTextureValue(const FMaterialParameterInfo& ParameterInfo, const UTexture** OutValue, const FMaterialRenderContext& Context) const;
+	ENGINE_API virtual bool GetTextureValue(const FHashedMaterialParameterInfo& ParameterInfo, const UTexture** OutValue, const FMaterialRenderContext& Context) const override;
 };
 
 
@@ -2141,11 +2248,11 @@ public:
 	}
 
 	// FMaterialRenderProxy interface.
-	ENGINE_API virtual const FMaterial& GetMaterialWithFallback(ERHIFeatureLevel::Type InFeatureLevel, const FMaterialRenderProxy*& OutFallbackMaterialRenderProxy) const;
-	ENGINE_API virtual bool GetVectorValue(const FMaterialParameterInfo& ParameterInfo, FLinearColor* OutValue, const FMaterialRenderContext& Context) const;
-	ENGINE_API virtual bool GetScalarValue(const FMaterialParameterInfo& ParameterInfo, float* OutValue, const FMaterialRenderContext& Context) const;
-	ENGINE_API virtual bool GetTextureValue(const FMaterialParameterInfo& ParameterInfo, const UTexture** OutValue, const FMaterialRenderContext& Context) const;
-	ENGINE_API virtual bool GetTextureValue(const FMaterialParameterInfo& ParameterInfo, const URuntimeVirtualTexture** OutValue, const FMaterialRenderContext& Context) const;
+	ENGINE_API virtual const FMaterial& GetMaterialWithFallback(ERHIFeatureLevel::Type InFeatureLevel, const FMaterialRenderProxy*& OutFallbackMaterialRenderProxy) const override;
+	ENGINE_API virtual bool GetVectorValue(const FHashedMaterialParameterInfo& ParameterInfo, FLinearColor* OutValue, const FMaterialRenderContext& Context) const override;
+	ENGINE_API virtual bool GetScalarValue(const FHashedMaterialParameterInfo& ParameterInfo, float* OutValue, const FMaterialRenderContext& Context) const override;
+	ENGINE_API virtual bool GetTextureValue(const FHashedMaterialParameterInfo& ParameterInfo, const UTexture** OutValue, const FMaterialRenderContext& Context) const override;
+	ENGINE_API virtual bool GetTextureValue(const FHashedMaterialParameterInfo& ParameterInfo, const URuntimeVirtualTexture** OutValue, const FMaterialRenderContext& Context) const;
 };
 
 
@@ -2164,7 +2271,7 @@ public:
 	{}
 
 	// FMaterialRenderProxy interface.
-	virtual bool GetVectorValue(const FMaterialParameterInfo& ParameterInfo, FLinearColor* OutValue, const FMaterialRenderContext& Context) const;
+	virtual bool GetVectorValue(const FHashedMaterialParameterInfo& ParameterInfo, FLinearColor* OutValue, const FMaterialRenderContext& Context) const override;
 };
 
 /**
@@ -2208,7 +2315,9 @@ public:
 
 	// FMaterial interface.
 	ENGINE_API virtual void GetShaderMapId(EShaderPlatform Platform, FMaterialShaderMapId& OutId) const override;
+#if WITH_EDITORONLY_DATA
 	ENGINE_API virtual void GetStaticParameterSet(EShaderPlatform Platform, FStaticParameterSet& OutSet) const override;
+#endif
 #if WITH_EDITOR
 	ENGINE_API virtual void BeginAllowCachingStaticParameterValues() override;
 	ENGINE_API virtual void EndAllowCachingStaticParameterValues() override;
@@ -2322,7 +2431,7 @@ public:
 
 	ENGINE_API virtual void LegacySerialize(FArchive& Ar) override;
 
-	ENGINE_API virtual const TArray<UObject*>& GetReferencedTextures() const override;
+	ENGINE_API virtual TArrayView<UObject* const> GetReferencedTextures() const override;
 
 	ENGINE_API virtual void AddReferencedObjects(FReferenceCollector& Collector) override;
 
@@ -2339,9 +2448,12 @@ protected:
 	ENGINE_API virtual int32 CompileCustomAttribute(const FGuid& AttributeID, FMaterialCompiler* Compiler) const override;
 #endif
 
+#if WITH_EDITORONLY_DATA
 	/* Gives the material a chance to compile any custom output nodes it has added */
 	ENGINE_API virtual void GatherCustomOutputExpressions(TArray<class UMaterialExpressionCustomOutput*>& OutCustomOutputs) const override;
 	ENGINE_API virtual void GatherExpressionsForCustomInterpolators(TArray<class UMaterialExpression*>& OutExpressions) const override;
+#endif // WITH_EDITORONLY_DATA
+
 	ENGINE_API virtual bool HasVertexPositionOffsetConnected() const override;
 	ENGINE_API virtual bool HasPixelDepthOffsetConnected() const override;
 	ENGINE_API virtual bool HasMaterialAttributesConnected() const override;
@@ -2801,9 +2913,135 @@ bool ReloadMaterialResource(
 	EMaterialQualityLevel::Type QualityLevel);
 #endif
 
+//
+struct FMaterialShaderParameters
+{
+	EMaterialDomain MaterialDomain;
+	FMaterialShadingModelField ShadingModels;
+	EBlendMode BlendMode;
+	ERHIFeatureLevel::Type FeatureLevel;
+	EMaterialQualityLevel::Type QualityLevel;
+	EMaterialTessellationMode TessellationMode;
+	int32 BlendableLocation;
+	uint32 DecalBlendMode;
+	int32 NumCustomizedUVs;
+	union
+	{
+		uint64 PackedFlags;
+		struct
+		{
+			uint64 bIsDefaultMaterial : 1;
+			uint64 bIsSpecialEngineMaterial : 1;
+			uint64 bIsMasked : 1;
+			uint64 bIsTwoSided : 1;
+			uint64 bIsDistorted : 1;
+			uint64 bShouldCastDynamicShadows : 1;
+			uint64 bShouldInjectEmissiveIntoLPV : 1;
+			uint64 bShouldBlockGI : 1;
+			uint64 bWritesEveryPixel : 1;
+			uint64 bWritesEveryPixelShadowPass : 1;
+			uint64 bHasNormalConnected : 1;
+			uint64 bHasEmissiveColorConnected : 1;
+			uint64 bHasVertexPositionOffsetConnected : 1;
+			uint64 bHasPixelDepthOffsetConnected : 1;
+			uint64 bMaterialMayModifyMeshPosition : 1;
+			uint64 bIsUsedWithStaticLighting : 1;
+			uint64 bIsUsedWithParticleSprites : 1;
+			uint64 bIsUsedWithMeshParticles : 1;
+			uint64 bIsUsedWithNiagaraSprites : 1;
+			uint64 bIsUsedWithNiagaraMeshParticles : 1;
+			uint64 bIsUsedWithNiagaraRibbons : 1;
+			uint64 bIsUsedWithLandscape : 1;
+			uint64 bIsUsedWithBeamTrails : 1;
+			uint64 bIsUsedWithSplineMeshes : 1;
+			uint64 bIsUsedWithSkeletalMesh : 1;
+			uint64 bIsUsedWithMorphTargets : 1;
+			uint64 bIsUsedWithAPEXCloth : 1;
+			uint64 bIsUsedWithGeometryCache : 1;
+			uint64 bIsUsedWithGeometryCollections : 1;
+			uint64 bIsUsedWithHairStrands : 1;
+			uint64 bIsUsedWithWater : 1;
+			uint64 bIsTranslucencyWritingVelocity : 1;
+			uint64 bIsTranslucencyWritingCustomDepth : 1;
+			uint64 bIsDitheredLODTransition : 1;
+			uint64 bIsUsedWithInstancedStaticMeshes : 1;
+			uint64 bHasRuntimeVirtualTextureOutput : 1;
+			uint64 bIsMaterialTexCoordScale : 1;
+			uint64 bIsMaterialDebugViewMode : 1;
+			uint64 bIsMaterialMeshTexCoordSizeAccuracy : 1;
+			uint64 bMaterialIsPrimitiveDistanceAccuracy : 1;
+			uint64 bMaterialIsRequiredTextureResolution : 1;
+			uint64 bMaterialIsComplexityAccumulate : 1;
+			uint64 bIsUsedWithLidarPointCloud : 1;
+		};
+	};
+
+	FMaterialShaderParameters(const FMaterial* InMaterial)
+	{
+		// Make sure to zero-initialize so we get consistent hashes
+		FMemory::Memzero(*this);
+
+		MaterialDomain = InMaterial->GetMaterialDomain();
+		ShadingModels = InMaterial->GetShadingModels();
+		BlendMode = InMaterial->GetBlendMode();
+		FeatureLevel = InMaterial->GetFeatureLevel();
+		QualityLevel = InMaterial->GetQualityLevel();
+		TessellationMode = InMaterial->GetTessellationMode();
+		BlendableLocation = InMaterial->GetBlendableLocation();
+		DecalBlendMode = InMaterial->GetDecalBlendMode();
+		NumCustomizedUVs = InMaterial->GetNumCustomizedUVs();
+		bIsDefaultMaterial = InMaterial->IsDefaultMaterial();
+		bIsSpecialEngineMaterial = InMaterial->IsSpecialEngineMaterial();
+		bIsMasked = InMaterial->IsMasked();
+		bIsTwoSided = InMaterial->IsTwoSided();
+		bIsDistorted = InMaterial->IsDistorted();
+		bShouldCastDynamicShadows = InMaterial->ShouldCastDynamicShadows();
+		bShouldInjectEmissiveIntoLPV = InMaterial->ShouldInjectEmissiveIntoLPV();
+		bShouldBlockGI = InMaterial->ShouldBlockGI();
+		bWritesEveryPixel = InMaterial->WritesEveryPixel(false);
+		bWritesEveryPixelShadowPass = InMaterial->WritesEveryPixel(true);
+		bHasNormalConnected = InMaterial->HasNormalConnected();
+		bHasEmissiveColorConnected = InMaterial->HasEmissiveColorConnected();
+		bHasVertexPositionOffsetConnected = InMaterial->HasVertexPositionOffsetConnected();
+		bHasPixelDepthOffsetConnected = InMaterial->HasPixelDepthOffsetConnected();
+		bMaterialMayModifyMeshPosition = InMaterial->MaterialMayModifyMeshPosition();
+		bIsUsedWithStaticLighting = InMaterial->IsUsedWithStaticLighting();
+		bIsUsedWithParticleSprites = InMaterial->IsUsedWithParticleSprites();
+		bIsUsedWithMeshParticles = InMaterial->IsUsedWithMeshParticles();
+		bIsUsedWithNiagaraSprites = InMaterial->IsUsedWithNiagaraSprites();
+		bIsUsedWithNiagaraMeshParticles = InMaterial->IsUsedWithNiagaraMeshParticles();
+		bIsUsedWithNiagaraRibbons = InMaterial->IsUsedWithNiagaraRibbons();
+		bIsUsedWithLandscape = InMaterial->IsUsedWithLandscape();
+		bIsUsedWithBeamTrails = InMaterial->IsUsedWithBeamTrails();
+		bIsUsedWithSplineMeshes = InMaterial->IsUsedWithSplineMeshes();
+		bIsUsedWithSkeletalMesh = InMaterial->IsUsedWithSkeletalMesh();
+		bIsUsedWithMorphTargets = InMaterial->IsUsedWithMorphTargets();
+		bIsUsedWithAPEXCloth = InMaterial->IsUsedWithAPEXCloth();
+		bIsUsedWithGeometryCache = InMaterial->IsUsedWithGeometryCache();
+		bIsUsedWithGeometryCollections = InMaterial->IsUsedWithGeometryCollections();
+		bIsUsedWithHairStrands = InMaterial->IsUsedWithHairStrands();
+		bIsUsedWithWater = InMaterial->IsUsedWithWater();
+		bIsTranslucencyWritingVelocity = InMaterial->IsTranslucencyWritingVelocity();
+		bIsTranslucencyWritingCustomDepth = InMaterial->IsTranslucencyWritingCustomDepth();
+		bIsDitheredLODTransition = InMaterial->IsDitheredLODTransition();
+		bIsUsedWithInstancedStaticMeshes = InMaterial->IsUsedWithInstancedStaticMeshes();
+		bHasRuntimeVirtualTextureOutput = InMaterial->HasRuntimeVirtualTextureOutput();
+		bIsUsedWithLidarPointCloud = InMaterial->IsUsedWithLidarPointCloud();
+
+		// See FDebugViewModeMaterialProxy::GetFriendlyName()
+		// TODO seems horrible that friendly name controls which shaders get compiled, should refactor this to use regular accessors
+		const FString FriendlyName = InMaterial->GetFriendlyName();
+		bIsMaterialTexCoordScale = FriendlyName.Contains(TEXT("MaterialTexCoordScale"));
+		bIsMaterialDebugViewMode = FriendlyName.Contains(TEXT("DebugViewMode"));
+		bIsMaterialMeshTexCoordSizeAccuracy = FriendlyName.Contains(TEXT("MeshTexCoordSizeAccuracy"));
+		bMaterialIsPrimitiveDistanceAccuracy = FriendlyName.Contains(TEXT("PrimitiveDistanceAccuracy"));
+		bMaterialIsRequiredTextureResolution = FriendlyName.Contains(TEXT("RequiredTextureResolution"));
+		bMaterialIsComplexityAccumulate = FriendlyName.Contains(TEXT("ComplexityAccumulate"));
+	}
+};
+
 inline bool ShouldIncludeMaterialInDefaultOpaquePass(const FMaterial& Material)
 {
 	return !Material.IsSky()
 		&& !Material.GetShadingModels().HasShadingModel(MSM_SingleLayerWater);
 }
-
diff --git a/Engine/Source/Runtime/Engine/Public/MeshBatch.h b/Engine/Source/Runtime/Engine/Public/MeshBatch.h
--- a/Engine/Source/Runtime/Engine/Public/MeshBatch.h
+++ b/Engine/Source/Runtime/Engine/Public/MeshBatch.h
@@ -53,12 +53,6 @@
 	 */
 	const TUniformBuffer<FPrimitiveUniformShaderParameters>* PrimitiveUniformBufferResource;
 
-	/** Assigned by renderer */
-	EPrimitiveIdMode PrimitiveIdMode : PrimID_NumBits + 1;
-
-	/** Assigned by renderer */
-	uint32 DynamicPrimitiveShaderDataIndex : 24;
-
 	const FIndexBuffer* IndexBuffer;
 
 	union 
@@ -70,6 +64,18 @@
 	};
 	const void* UserData;
 
+	// Meaning depends on the vertex factory, e.g. FGPUSkinPassthroughVertexFactory: element index in FGPUSkinCache::CachedElements
+	void* VertexFactoryUserData;
+
+	FRHIVertexBuffer* IndirectArgsBuffer;
+	uint32 IndirectArgsOffset;
+
+	/** Assigned by renderer */
+	EPrimitiveIdMode PrimitiveIdMode : PrimID_NumBits + 1;
+
+	/** Assigned by renderer */
+	uint32 DynamicPrimitiveShaderDataIndex : 24;
+
 	uint32 FirstIndex;
 	/** When 0, IndirectArgsBuffer will be used. */
 	uint32 NumPrimitives;
@@ -79,24 +85,20 @@
 	uint32 BaseVertexIndex;
 	uint32 MinVertexIndex;
 	uint32 MaxVertexIndex;
-	// Meaning depends on the vertex factory, e.g. FGPUSkinPassthroughVertexFactory: element index in FGPUSkinCache::CachedElements
-	void* VertexFactoryUserData;
 	int32 UserIndex;
 	float MinScreenSize;
 	float MaxScreenSize;
 
-	uint8 InstancedLODIndex : 4;
-	uint8 InstancedLODRange : 4;
-	uint8 bUserDataIsColorVertexBuffer : 1;
-	uint8 bIsSplineProxy : 1;
-	uint8 bIsInstanceRuns : 1;
+	uint32 InstancedLODIndex : 4;
+	uint32 InstancedLODRange : 4;
+	uint32 bUserDataIsColorVertexBuffer : 1;
+	uint32 bIsSplineProxy : 1;
+	uint32 bIsInstanceRuns : 1;
 
 #if !(UE_BUILD_SHIPPING || UE_BUILD_TEST)
 	/** Conceptual element index used for debug viewmodes. */
-	int8 VisualizeElementIndex;
+	int32 VisualizeElementIndex : 8;
 #endif
-	FRHIVertexBuffer* IndirectArgsBuffer;
-	uint32 IndirectArgsOffset;
 
 	FORCEINLINE int32 GetNumPrimitives() const
 	{
@@ -118,14 +120,16 @@
 	FMeshBatchElement()
 	:	PrimitiveUniformBuffer(nullptr)
 	,	PrimitiveUniformBufferResource(nullptr)
-	,	PrimitiveIdMode(PrimID_FromPrimitiveSceneInfo)
-	,	DynamicPrimitiveShaderDataIndex(0)
 	,	IndexBuffer(nullptr)
 	,	InstanceRuns(nullptr)
 	,	UserData(nullptr)
+	,	VertexFactoryUserData(nullptr)
+	,	IndirectArgsBuffer(nullptr)
+	,	IndirectArgsOffset(0)
+	,	PrimitiveIdMode(PrimID_FromPrimitiveSceneInfo)
+	,	DynamicPrimitiveShaderDataIndex(0)
 	,	NumInstances(1)
 	,	BaseVertexIndex(0)
-	,	VertexFactoryUserData(nullptr)
 	,	UserIndex(-1)
 	,	MinScreenSize(0.0f)
 	,	MaxScreenSize(1.0f)
@@ -137,8 +141,6 @@
 #if !(UE_BUILD_SHIPPING || UE_BUILD_TEST)
 	,	VisualizeElementIndex(INDEX_NONE)
 #endif
-	,	IndirectArgsBuffer(nullptr)
-	,	IndirectArgsOffset(0)
 	{
 	}
 };
@@ -151,12 +153,12 @@
 		(Material->GetBlendMode() == BLEND_Opaque || Material->GetBlendMode() == BLEND_Masked);
 }
 
-FORCEINLINE bool IsCompatibleWithHairStrands(const FMaterial* Material, const EShaderPlatform Platform)
+FORCEINLINE bool IsCompatibleWithHairStrands(EShaderPlatform Platform, const FMaterialShaderParameters& Parameters)
 {
 	return
 		IsPCPlatform(Platform) && GetMaxSupportedFeatureLevel(Platform) == ERHIFeatureLevel::SM5 &&
-		Material && Material->IsUsedWithHairStrands() && Material->GetShadingModels().HasShadingModel(MSM_Hair) &&
-		(Material->GetBlendMode() == BLEND_Opaque || Material->GetBlendMode() == BLEND_Masked);
+		Parameters.bIsUsedWithHairStrands && Parameters.ShadingModels.HasShadingModel(MSM_Hair) &&
+		(Parameters.BlendMode == BLEND_Opaque || Parameters.BlendMode == BLEND_Masked);
 }
 
 /**
@@ -166,6 +168,21 @@
 {
 	TArray<FMeshBatchElement,TInlineAllocator<1> > Elements;
 
+	/** Vertex factory for rendering, required. */
+	const FVertexFactory* VertexFactory;
+
+	/** Material proxy for rendering, required. */
+	const FMaterialRenderProxy* MaterialRenderProxy;
+
+	// can be NULL
+	const FLightCacheInterface* LCI;
+
+	/** The current hit proxy ID being rendered. */
+	FHitProxyId BatchHitProxyId;
+
+	/** This is the threshold that will be used to know if we should use this mesh batch or use one with no tessellation enabled */
+	float TessellationDisablingShadowMapMeshSize;
+
 	/* Mesh Id in a primitive. Used for stable sorting of draws belonging to the same primitive. **/
 	uint16 MeshIdInPrimitive;
 
@@ -173,14 +190,6 @@
 	int8 LODIndex;
 	uint8 SegmentIndex;
 
-#if !(UE_BUILD_SHIPPING || UE_BUILD_TEST)
-	/** Conceptual LOD index used for the LOD Coloration visualization. */
-	int8 VisualizeLODIndex;
-#endif
-
-	/** Conceptual HLOD index used for the HLOD Coloration visualization. */
-	int8 VisualizeHLODIndex;
-
 	uint32 ReverseCulling : 1;
 	uint32 bDisableBackfaceCulling : 1;
 
@@ -188,9 +197,6 @@
 	 * Pass feature relevance flags.  Allows a proxy to submit fast representations for passes which can take advantage of it, 
 	 * for example separate index buffer for depth-only rendering since vertices can be merged based on position and ignore UV differences.
 	 */
-#if RHI_RAYTRACING
-	uint32 CastRayTracedShadow : 1;	// Whether it casts ray traced shadow.
-#endif
 	uint32 CastShadow		: 1;	// Whether it can be used in shadow renderpasses.
 	uint32 bUseForMaterial	: 1;	// Whether it can be used in renderpasses requiring material outputs.
 	uint32 bUseForDepthPass : 1;	// Whether it can be used in depth pass.
@@ -231,17 +237,19 @@
 	/** What virtual texture material type this mesh batch should be rendered with. */
 	uint32 RuntimeVirtualTextureMaterialType : RuntimeVirtualTexture::MaterialType_NumBits;
 
-	// can be NULL
-	const FLightCacheInterface* LCI;
+#if RHI_RAYTRACING
+	uint32 CastRayTracedShadow : 1;	// Whether it casts ray traced shadow.
+#endif
 
-	/** Vertex factory for rendering, required. */
-	const FVertexFactory* VertexFactory;
+#if (!(UE_BUILD_SHIPPING || UE_BUILD_TEST) || WITH_EDITOR)
+	/** Conceptual HLOD index used for the HLOD Coloration visualization. */
+	int8 VisualizeHLODIndex;
+#endif
 
-	/** Material proxy for rendering, required. */
-	const FMaterialRenderProxy* MaterialRenderProxy;
+#if !(UE_BUILD_SHIPPING || UE_BUILD_TEST)
+	/** Conceptual LOD index used for the LOD Coloration visualization. */
+	int8 VisualizeLODIndex;
+#endif
 
-	/** The current hit proxy ID being rendered. */
-	FHitProxyId BatchHitProxyId;
-
 	FORCEINLINE bool IsTranslucent(ERHIFeatureLevel::Type InFeatureLevel) const
 	{
@@ -306,18 +314,15 @@
 
 	/** Default constructor. */
 	FMeshBatch()
-	:	MeshIdInPrimitive(0)
+	:	VertexFactory(nullptr)
+	,	MaterialRenderProxy(nullptr)
+	,	LCI(nullptr)
+	,	TessellationDisablingShadowMapMeshSize(0.0f)
+	,	MeshIdInPrimitive(0)
 	,	LODIndex(INDEX_NONE)
 	,	SegmentIndex(0xFF)
-#if !(UE_BUILD_SHIPPING || UE_BUILD_TEST)
-	,	VisualizeLODIndex(INDEX_NONE)
-#endif
-	,	VisualizeHLODIndex(INDEX_NONE)
 	,	ReverseCulling(false)
 	,	bDisableBackfaceCulling(false)
-#if RHI_RAYTRACING
-	,	CastRayTracedShadow(true)
-#endif
 	,	CastShadow(true)
 	,   bUseForMaterial(true)
 	,	bUseForDepthPass(true)
@@ -333,9 +338,15 @@
 	,	bDitheredLODTransition(false)
 	,	bRenderToVirtualTexture(false)
 	,	RuntimeVirtualTextureMaterialType(0)
-	,	LCI(NULL)
-	,	VertexFactory(NULL)
-	,	MaterialRenderProxy(NULL)
+#if RHI_RAYTRACING
+	,	CastRayTracedShadow(true)
+#endif
+#if (!(UE_BUILD_SHIPPING || UE_BUILD_TEST) || WITH_EDITOR)
+	,	VisualizeHLODIndex(INDEX_NONE)
+#endif
+#if !(UE_BUILD_SHIPPING || UE_BUILD_TEST)
+	,	VisualizeLODIndex(INDEX_NONE)
+#endif
 	{
 		// By default always add the first element.
 		new(Elements) FMeshBatchElement;
diff --git a/Engine/Source/Runtime/Engine/Public/MeshMaterialShaderType.h b/Engine/Source/Runtime/Engine/Public/MeshMaterialShaderType.h
--- a/Engine/Source/Runtime/Engine/Public/MeshMaterialShaderType.h
+++ b/Engine/Source/Runtime/Engine/Public/MeshMaterialShaderType.h
@@ -16,29 +16,6 @@ class FShaderCompileJob;
 class FUniformExpressionSet;
 class FVertexFactoryType;
 
-struct FMeshMaterialShaderPermutationParameters
-{
-	// Shader platform to compile to.
-	const EShaderPlatform Platform;
-
-	// Material to compile.
-	const FMaterial* Material;
-
-	// Type of vertex factory to compile.
-	const FVertexFactoryType* VertexFactoryType;
-
-	// Unique permutation identifier of the mesh material shader type.
-	const int32 PermutationId;
-
-	FMeshMaterialShaderPermutationParameters(EShaderPlatform InPlatform, const FMaterial* InMaterial, const FVertexFactoryType* InVertexFactoryType, const int32 InPermutationId)
-		: Platform(InPlatform)
-		, Material(InMaterial)
-		, VertexFactoryType(InVertexFactoryType)
-		, PermutationId(InPermutationId)
-	{
-	}
-};
-
 /**
  * A shader meta type for material-linked shaders which use a vertex factory.
  */
@@ -52,23 +29,20 @@ public:
 			FShaderType* InType,
 			int32 PermutationId,
 			const FShaderCompilerOutput& CompilerOutput,
-			FShaderResource* InResource,
+			int32 InResourceIndex,
 			const FUniformExpressionSet& InUniformExpressionSet,
 			const FSHAHash& InMaterialShaderMapHash,
 			const FString& InDebugDescription,
 			const FShaderPipelineType* InShaderPipeline,
 			FVertexFactoryType* InVertexFactoryType
 			):
-			FMaterialShaderType::CompiledShaderInitializerType(InType,PermutationId,CompilerOutput,InResource,InUniformExpressionSet,InMaterialShaderMapHash,InShaderPipeline,InVertexFactoryType,InDebugDescription),
+			FMaterialShaderType::CompiledShaderInitializerType(InType,PermutationId,CompilerOutput, InResourceIndex,InUniformExpressionSet,InMaterialShaderMapHash,InShaderPipeline,InVertexFactoryType,InDebugDescription),
 			VertexFactoryType(InVertexFactoryType)
 		{}
 	};
-	typedef FShader* (*ConstructCompiledType)(const CompiledShaderInitializerType&);
-	typedef bool (*ShouldCompilePermutationType)(const FMeshMaterialShaderPermutationParameters&);
-	typedef bool(*ValidateCompiledResultType)(EShaderPlatform, const TArray<FMaterial*>&, const FVertexFactoryType*, const FShaderParameterMap&, TArray<FString>&);
-	typedef void (*ModifyCompilationEnvironmentType)(const FMaterialShaderPermutationParameters&, FShaderCompilerEnvironment&);
 
 	FMeshMaterialShaderType(
+		FTypeLayoutDesc& InTypeLayout,
 		const TCHAR* InName,
 		const TCHAR* InSourceFilename,
 		const TCHAR* InFunctionName,
@@ -78,13 +52,18 @@ public:
 		ConstructCompiledType InConstructCompiledRef,
 		ModifyCompilationEnvironmentType InModifyCompilationEnvironmentRef,
 		ShouldCompilePermutationType InShouldCompilePermutationRef,
-		ValidateCompiledResultType InValidateCompiledResultRef
+		ValidateCompiledResultType InValidateCompiledResultRef,
+		uint32 InTypeSize,
+		const FShaderParametersMetadata* InRootParametersMetadata = nullptr
 		):
-		FShaderType(EShaderTypeForDynamicCast::MeshMaterial, InName,InSourceFilename,InFunctionName,InFrequency,InTotalPermutationCount,InConstructSerializedRef, nullptr),
-		ConstructCompiledRef(InConstructCompiledRef),
-		ShouldCompilePermutationRef(InShouldCompilePermutationRef),
-		ValidateCompiledResultRef(InValidateCompiledResultRef),
-		ModifyCompilationEnvironmentRef(InModifyCompilationEnvironmentRef)
+		FShaderType(EShaderTypeForDynamicCast::MeshMaterial, InTypeLayout, InName,InSourceFilename,InFunctionName,InFrequency,InTotalPermutationCount,
+			InConstructSerializedRef,
+			InConstructCompiledRef,
+			InModifyCompilationEnvironmentRef,
+			InShouldCompilePermutationRef,
+			InValidateCompiledResultRef,
+			InTypeSize,
+			InRootParametersMetadata)
 	{
 		checkf(FPaths::GetExtension(InSourceFilename) == TEXT("usf"),
 			TEXT("Incorrect virtual shader path extension for mesh material shader '%s': Only .usf files should be compiled."),
@@ -134,7 +113,8 @@ public:
 		const FSHAHash& MaterialShaderMapHash,
 		const FShaderCompileJob& CurrentJob,
 		const FShaderPipelineType* ShaderPipeline,
-		const FString& InDebugDescription
+		const FString& InDebugDescription,
+		FShaderMapResourceBuilder& ResourceBuilder
 		);
 
 	/**
@@ -144,103 +124,17 @@ public:
 	 * @param VertexFactoryType - The vertex factory type to check.
 	 * @return True if this shader type should be cached.
 	 */
-	bool ShouldCompilePermutation(EShaderPlatform Platform,const FMaterial* Material,const FVertexFactoryType* VertexFactoryType, int32 PermutationId) const
-	{
-		return (*ShouldCompilePermutationRef)(FMeshMaterialShaderPermutationParameters(Platform,Material,VertexFactoryType,PermutationId));
-	}
+	bool ShouldCompilePermutation(EShaderPlatform Platform, const FMaterialShaderParameters& MaterialParameters, const FVertexFactoryType* VertexFactoryType, int32 PermutationId) const;
 
-	/**
-	* Checks if the shader type should pass compilation for a particular set of parameters.
-	* @param Platform - Platform to validate.
-	* @param Materials - Materials to validate.
-	* @param VertexFactoryType - Vertex factory to validate.
-	* @param ParameterMap - Shader parameters to validate.
-	* @param OutError - List for appending validation errors.
-	*/
-	bool ValidateCompiledResult(EShaderPlatform Platform, const TArray<FMaterial*>& Materials, const FVertexFactoryType* VertexFactoryType, const FShaderParameterMap& ParameterMap, TArray<FString>& OutError) const
-	{
-		return (*ValidateCompiledResultRef)(Platform, Materials, VertexFactoryType, ParameterMap, OutError);
-	}
+	static bool ShouldCompileVertexFactoryPermutation(const FVertexFactoryType* VertexFactoryType, EShaderPlatform Platform, const FMaterialShaderParameters& MaterialParameters);
 
-protected:
+	static bool ShouldCompilePipeline(const FShaderPipelineType* ShaderPipelineType, EShaderPlatform Platform, const FMaterialShaderParameters& MaterialParameters, const FVertexFactoryType* VertexFactoryType);
 
+protected:
 	/**
 	 * Sets up the environment used to compile an instance of this shader type.
 	 * @param Platform - Platform to compile for.
 	 * @param Environment - The shader compile environment that the function modifies.
 	 */
-	void SetupCompileEnvironment(EShaderPlatform Platform, const FMaterial* Material, int32 PermutationId, FShaderCompilerEnvironment& Environment)
-	{
-		// Allow the shader type to modify its compile environment.
-		(*ModifyCompilationEnvironmentRef)(FMaterialShaderPermutationParameters(Platform, Material, PermutationId), Environment);
-	}
-
-private:
-	ConstructCompiledType ConstructCompiledRef;
-	ShouldCompilePermutationType ShouldCompilePermutationRef;
-	ValidateCompiledResultType ValidateCompiledResultRef;
-	ModifyCompilationEnvironmentType ModifyCompilationEnvironmentRef;
+	void SetupCompileEnvironment(EShaderPlatform Platform, const FMaterialShaderParameters& MaterialParameters, const FVertexFactoryType* VertexFactoryType, int32 PermutationId, FShaderCompilerEnvironment& Environment);
 };
-
-/** DECLARE_MESH_MATERIAL_SHADER and IMPLEMENT_MESH_MATERIAL_SHADER setup a mesh material shader class's boiler plate. They are meant to be used like so:
- *
- * class FMyMeshMaterialShaderPS : public FMeshMaterialShader
- * {
- *		// Setup the shader's boiler plate.
- *		DECLARE_MESH_MATERIAL_SHADER(FMyMeshMaterialShaderPS);
- *
- *		// Setup the shader's permutation domain. If no dimensions, can do FPermutationDomain = FShaderPermutationNone.
- *		using FPermutationDomain = TShaderPermutationDomain<DIMENSIONS...>;
- *
- *		// ...
- * };
- *
- * // Instantiates shader's global variable that will take care of compilation process of the shader. This needs imperatively to be
- * done in a .cpp file regardless of whether FMyMeshMaterialShaderPS is in a header or not.
- * IMPLEMENT_MESH_MATERIAL_SHADER(FMyMaterialShaderPS, "/Engine/Private/MyShaderFile.usf", "MainPS", SF_Pixel);
- *
- * When the shader class is a public header, let say in Engine module public header, the shader class then should have the ENGINE_API
- * like this:
- *
- * class ENGINE_API FMyMeshMaterialShaderPS : public FMeshMaterialShader
- * {
- *		// Setup the shader's boiler plate.
- *		DECLARE_MESH_MATERIAL_SHADER(FMyMeshMaterialShaderPS);
- *
- *		// ...
- * };
- */
-#define DECLARE_MESH_MATERIAL_SHADER(ShaderClass) \
-	public: \
-	using ShaderMetaType = FMeshMaterialShaderType; \
-	\
-	static ShaderMetaType StaticType; \
-	\
-	static FShader* ConstructSerializedInstance() { return new ShaderClass(); } \
-	static FShader* ConstructCompiledInstance(const ShaderMetaType::CompiledShaderInitializerType& Initializer) \
-	{ return new ShaderClass(Initializer); } \
-	\
-	virtual uint32 GetTypeSize() const override { return sizeof(*this); } \
-	\
-	static void ModifyCompilationEnvironmentImpl( \
-		const FMaterialShaderPermutationParameters& Parameters, \
-		FShaderCompilerEnvironment& OutEnvironment) \
-	{ \
-		FPermutationDomain PermutationVector(Parameters.PermutationId); \
-		PermutationVector.ModifyCompilationEnvironment(OutEnvironment); \
-		ShaderClass::ModifyCompilationEnvironment(Parameters, OutEnvironment); \
-	}
-
-#define IMPLEMENT_MESH_MATERIAL_SHADER(ShaderClass,SourceFilename,FunctionName,Frequency) \
-	ShaderClass::ShaderMetaType ShaderClass::StaticType( \
-		TEXT(#ShaderClass), \
-		TEXT(SourceFilename), \
-		TEXT(FunctionName), \
-		Frequency, \
-		ShaderClass::FPermutationDomain::PermutationCount, \
-		ShaderClass::ConstructSerializedInstance, \
-		ShaderClass::ConstructCompiledInstance, \
-		ShaderClass::ModifyCompilationEnvironmentImpl, \
-		ShaderClass::ShouldCompilePermutation, \
-		ShaderClass::ValidateCompiledResult \
-		)
diff --git a/Engine/Source/Runtime/Engine/Public/MeshParticleVertexFactory.h b/Engine/Source/Runtime/Engine/Public/MeshParticleVertexFactory.h
--- a/Engine/Source/Runtime/Engine/Public/MeshParticleVertexFactory.h
+++ b/Engine/Source/Runtime/Engine/Public/MeshParticleVertexFactory.h
@@ -101,27 +101,14 @@ public:
 	/**
 	 * Should we cache the material's shadertype on this platform with this vertex factory? 
 	 */
-	static bool ShouldCompilePermutation(EShaderPlatform Platform, const class FMaterial* Material, const class FShaderType* ShaderType);
+	static bool ShouldCompilePermutation(const FVertexFactoryShaderPermutationParameters& Parameters);
 
 
 	/**
 	 * Modify compile environment to enable instancing
 	 * @param OutEnvironment - shader compile environment to modify
 	 */
-	static void ModifyCompilationEnvironment(const FVertexFactoryType* Type, EShaderPlatform Platform, const FMaterial* Material, FShaderCompilerEnvironment& OutEnvironment)
-	{
-		FParticleVertexFactoryBase::ModifyCompilationEnvironment(Type, Platform, Material, OutEnvironment);
-
-		// Set a define so we can tell in MaterialTemplate.usf when we are compiling a mesh particle vertex factory
-		OutEnvironment.SetDefine(TEXT("PARTICLE_MESH_FACTORY"),TEXT("1"));
-		OutEnvironment.SetDefine(TEXT("PARTICLE_MESH_INSTANCED"),TEXT("1"));
-
-		const bool ContainsManualVertexFetch = OutEnvironment.GetDefinitions().Contains("MANUAL_VERTEX_FETCH");
-		if (!ContainsManualVertexFetch && RHISupportsManualVertexFetch(Platform))
-		{
-			OutEnvironment.SetDefine(TEXT("MANUAL_VERTEX_FETCH"), TEXT("1"));
-		}
-	}
+	static void ModifyCompilationEnvironment(const FVertexFactoryShaderPermutationParameters& Parameters, FShaderCompilerEnvironment& OutEnvironment);
 	
 	/**
 	 * An implementation of the interface used by TSynchronizedResource to update the resource with new data from the game thread.
@@ -179,8 +166,6 @@ public:
 
 	static bool SupportsTessellationShaders() { return true; }
 
-	static FVertexFactoryShaderParameters* ConstructShaderParameters(EShaderFrequency ShaderFrequency);
-
 	FMeshParticleInstanceVertices*& GetInstanceVerticesCPU()
 	{
 		return InstanceVerticesCPU;
@@ -228,18 +213,9 @@ public:
 		: FMeshParticleVertexFactory(InFeatureLevel)
 	{}
 
-	static bool ShouldCompilePermutation(EShaderPlatform Platform, const class FMaterial* Material, const class FShaderType* ShaderType)
-	{
-		return (Platform == SP_OPENGL_ES2_ANDROID) // Android platforms that might not support hardware instancing
-			&& FMeshParticleVertexFactory::ShouldCompilePermutation(Platform, Material, ShaderType);
-	}
-
-	static void ModifyCompilationEnvironment(const FVertexFactoryType* Type, EShaderPlatform Platform, const FMaterial* Material, FShaderCompilerEnvironment& OutEnvironment)
-	{
-		FMeshParticleVertexFactory::ModifyCompilationEnvironment(Type, Platform, Material, OutEnvironment);
+	static bool ShouldCompilePermutation(const FVertexFactoryShaderPermutationParameters& Parameters);
 
-		OutEnvironment.SetDefine(TEXT("PARTICLE_MESH_INSTANCED"),TEXT("0"));
-	}
+	static void ModifyCompilationEnvironment(const FVertexFactoryShaderPermutationParameters& Parameters, FShaderCompilerEnvironment& OutEnvironment);
 };
 
 inline FMeshParticleVertexFactory* ConstructMeshParticleVertexFactory(ERHIFeatureLevel::Type InFeatureLevel)
diff --git a/Engine/Source/Runtime/Engine/Public/ParticleVertexFactory.h b/Engine/Source/Runtime/Engine/Public/ParticleVertexFactory.h
--- a/Engine/Source/Runtime/Engine/Public/ParticleVertexFactory.h
+++ b/Engine/Source/Runtime/Engine/Public/ParticleVertexFactory.h
@@ -48,9 +48,9 @@ public:
 	{
 	}
 
-	static void ModifyCompilationEnvironment(const FVertexFactoryType* Type, EShaderPlatform Platform, const FMaterial* Material, FShaderCompilerEnvironment& OutEnvironment) 
+	static void ModifyCompilationEnvironment(const FVertexFactoryShaderPermutationParameters& Parameters, FShaderCompilerEnvironment& OutEnvironment) 
 	{
-		FVertexFactory::ModifyCompilationEnvironment(Type, Platform, Material, OutEnvironment);
+		FVertexFactory::ModifyCompilationEnvironment(Parameters, OutEnvironment);
 		OutEnvironment.SetDefine(TEXT("PARTICLE_FACTORY"),TEXT("1"));
 	}
 
@@ -169,12 +169,12 @@ public:
 	/**
 	 * Should we cache the material's shadertype on this platform with this vertex factory? 
 	 */
-	static bool ShouldCompilePermutation(EShaderPlatform Platform, const class FMaterial* Material, const class FShaderType* ShaderType);
+	static bool ShouldCompilePermutation(const FVertexFactoryShaderPermutationParameters& Parameters);
 
 	/**
 	 * Can be overridden by FVertexFactory subclasses to modify their compile environment just before compilation occurs.
 	 */
-	static void ModifyCompilationEnvironment(const FVertexFactoryType* Type, EShaderPlatform Platform, const FMaterial* Material, FShaderCompilerEnvironment& OutEnvironment);
+	static void ModifyCompilationEnvironment(const FVertexFactoryShaderPermutationParameters& Parameters, FShaderCompilerEnvironment& OutEnvironment);
 
 	/**
 	 * Set the source vertex buffer that contains particle instance data.
@@ -233,11 +233,6 @@ public:
 		return bCustomAlignment;
 	}
 
-	/**
-	 * Construct shader parameters for this type of vertex factory.
-	 */
-	static FVertexFactoryShaderParameters* ConstructShaderParameters(EShaderFrequency ShaderFrequency);
-
 protected:
 	/** Initialize streams for this vertex factory. */
 	void InitStreams();
