Added GRHISupportsDirectGPUMemoryLock - true if the pointer returned by Lock is a persistent direct pointer to gpu memory

diff --git a/Engine/Source/Runtime/Engine/Private/VT/VirtualTextureUploadCache.cpp b/Engine/Source/Runtime/Engine/Private/VT/VirtualTextureUploadCache.cpp
--- a/Engine/Source/Runtime/Engine/Private/VT/VirtualTextureUploadCache.cpp
+++ b/Engine/Source/Runtime/Engine/Private/VT/VirtualTextureUploadCache.cpp
@@ -4,14 +4,6 @@
 #include "VirtualTextureChunkManager.h"
 #include "RHI.h"
 
-// Stage to persist mapped GPU buffer then GPU copy into texture
-// this is fast where supported
-#if PLATFORM_PS4
-static const bool ALLOW_COPY_FROM_BUFFER = true;
-#else
-static const bool ALLOW_COPY_FROM_BUFFER = false;
-#endif
-
 // allow uploading CPU buffer directly to GPU texture
 // this is slow under D3D11
 // Should be pretty decent on D3D12X...UpdateTexture does make an extra copy of the data, but Lock/Unlock texture also buffers an extra copy of texture on this platform
@@ -215,7 +207,10 @@ FVTUploadTileHandle FVirtualTextureUploadCache::PrepareTileForUpload(FVTUploadTi
 		// Can potentially write each tile to a separate staging texture, but this has too much lock/unlock overhead
 		NewEntry.Stride = Stride;
 		NewEntry.MemorySize = MemorySize;
-		if (ALLOW_COPY_FROM_BUFFER)
+
+		// Stage to persist mapped GPU buffer then GPU copy into texture
+		// this is fast where supported
+		if (GRHISupportsDirectGPUMemoryLock)
 		{
 			FRHIResourceCreateInfo CreateInfo;
 			NewEntry.RHIStagingBuffer = RHICreateStructuredBuffer(FormatInfo.BlockBytes, MemorySize, BUF_ShaderResource | BUF_Static | BUF_KeepCPUAccessible, CreateInfo);
diff --git a/Engine/Source/Runtime/RHI/Private/RHI.cpp b/Engine/Source/Runtime/RHI/Private/RHI.cpp
--- a/Engine/Source/Runtime/RHI/Private/RHI.cpp
+++ b/Engine/Source/Runtime/RHI/Private/RHI.cpp
@@ -561,6 +561,7 @@ bool GSupportsParallelOcclusionQueries = false;
 bool GSupportsTransientResourceAliasing = false;
 bool GRHIRequiresRenderTargetForPixelShaderUAVs = false;
 bool GRHISupportsUAVFormatAliasing = false;
+bool GRHISupportsDirectGPUMemoryLock = false;
 
 bool GRHISupportsMSAADepthSampleAccess = false;
 bool GRHISupportsResolveCubemapFaces = false;
diff --git a/Engine/Source/Runtime/RHI/Public/RHI.h b/Engine/Source/Runtime/RHI/Public/RHI.h
--- a/Engine/Source/Runtime/RHI/Public/RHI.h
+++ b/Engine/Source/Runtime/RHI/Public/RHI.h
@@ -389,6 +389,9 @@ extern RHI_API bool GRHIRequiresRenderTargetForPixelShaderUAVs;
 /** true if the RHI supports unordered access view format aliasing */
 extern RHI_API bool GRHISupportsUAVFormatAliasing;
 
+/** true if the pointer returned by Lock is a persistent direct pointer to gpu memory */
+extern RHI_API bool GRHISupportsDirectGPUMemoryLock;
+
 /** The minimum Z value in clip space for the RHI. */
 extern RHI_API float GMinClipZ;
 
