Split screen parties cannot join and spectate ONE player in a match. Two viewports will track the same player.

diff --git a/Engine/Source/Runtime/Engine/Classes/Engine/LocalPlayer.h b/Engine/Source/Runtime/Engine/Classes/Engine/LocalPlayer.h
--- a/Engine/Source/Runtime/Engine/Classes/Engine/LocalPlayer.h
+++ b/Engine/Source/Runtime/Engine/Classes/Engine/LocalPlayer.h
@@ -386,8 +386,10 @@ public:
 	 * @note this happens automatically for all viewports that exist during the initial server connect
 	 * 	so it's only necessary to manually call this for viewports created after that
 	 * if the join fails (because the server was full, for example) all viewports on this client will be disconnected
+	 *
+	 * @param	Options		array of URL options to append.
 	 */
-	virtual void SendSplitJoin();
+	virtual void SendSplitJoin(TArray<FString>& Options);
 	
 	/**
 	 * Change the ControllerId for this player; if the specified ControllerId is already taken by another player, changes the ControllerId
diff --git a/Engine/Source/Runtime/Engine/Private/GameInstance.cpp b/Engine/Source/Runtime/Engine/Private/GameInstance.cpp
--- a/Engine/Source/Runtime/Engine/Private/GameInstance.cpp
+++ b/Engine/Source/Runtime/Engine/Private/GameInstance.cpp
@@ -728,7 +728,8 @@ ULocalPlayer* UGameInstance::CreateLocalPlayer(int32 ControllerId, FString& OutE
 			else
 			{
 				// client; ask the server to let the new player join
-				NewPlayer->SendSplitJoin();
+				TArray<FString> Options;
+				NewPlayer->SendSplitJoin(Options);
 			}
 		}
 	}
diff --git a/Engine/Source/Runtime/Engine/Private/LocalPlayer.cpp b/Engine/Source/Runtime/Engine/Private/LocalPlayer.cpp
--- a/Engine/Source/Runtime/Engine/Private/LocalPlayer.cpp
+++ b/Engine/Source/Runtime/Engine/Private/LocalPlayer.cpp
@@ -319,7 +319,7 @@ bool ULocalPlayer::SpawnPlayActor(const FString& URL,FString& OutError, UWorld*
 	return PlayerController != NULL;
 }
 
-void ULocalPlayer::SendSplitJoin()
+void ULocalPlayer::SendSplitJoin(TArray<FString>& Options)
 {
 	UNetDriver* NetDriver = NULL;
 
@@ -374,6 +374,11 @@ void ULocalPlayer::SendSplitJoin()
 				URL.AddOption(*FString::Printf(TEXT("%s"), *GameUrlOptions));
 			}
 
+			for (FString& Option : Options)
+			{
+				URL.AddOption(*FString::Printf(TEXT("%s"), *Option));
+			}
+
 			// Send the player unique Id at login
 			FUniqueNetIdRepl UniqueIdRepl(GetPreferredUniqueNetId());
 
diff --git a/Engine/Source/Runtime/Engine/Private/NetConnection.cpp b/Engine/Source/Runtime/Engine/Private/NetConnection.cpp
--- a/Engine/Source/Runtime/Engine/Private/NetConnection.cpp
+++ b/Engine/Source/Runtime/Engine/Private/NetConnection.cpp
@@ -3597,7 +3597,8 @@ void UNetConnection::HandleClientPlayer( APlayerController *PC, UNetConnection*
 		if (*It != LocalPlayer)
 		{
 			// send server command for new child connection
-			It->SendSplitJoin();
+			TArray<FString> Options;
+			It->SendSplitJoin(Options);
 		}
 	}
 }
