Move PushToPhysicsState to after pretick/step broadcasts.

diff --git a/Engine/Source/Runtime/Engine/Private/PhysicsEngine/Experimental/PhysScene_Chaos.cpp b/Engine/Source/Runtime/Engine/Private/PhysicsEngine/Experimental/PhysScene_Chaos.cpp
--- a/Engine/Source/Runtime/Engine/Private/PhysicsEngine/Experimental/PhysScene_Chaos.cpp
+++ b/Engine/Source/Runtime/Engine/Private/PhysicsEngine/Experimental/PhysScene_Chaos.cpp
@@ -1834,11 +1834,6 @@ void FPhysScene_ChaosInterface::StartFrame()
 
 	if (Chaos::IDispatcher* Dispatcher = SolverModule->GetDispatcher())
 	{
-		if (FPhysicsSolver* Solver = GetSolver())
-		{
-			Solver->PushPhysicsState(Dispatcher);
-		}
-
 		switch (Dispatcher->GetMode())
 		{
 		case EChaosThreadingMode::SingleThread:
@@ -1846,6 +1841,10 @@ void FPhysScene_ChaosInterface::StartFrame()
 			OnPhysScenePreTick.Broadcast(this, Dt);
 			OnPhysSceneStep.Broadcast(this, Dt);
 
+			if (FPhysicsSolver* Solver = GetSolver())
+			{
+				Solver->PushPhysicsState(Dispatcher);
+			}
 			// Here we can directly tick the scene. Single threaded mode doesn't buffer any commands
 			// that would require pumping here - everything is done on demand.
 			Scene.Tick(Dt);
@@ -1866,6 +1865,11 @@ void FPhysScene_ChaosInterface::StartFrame()
 			OnPhysScenePreTick.Broadcast(this, Dt);
 			OnPhysSceneStep.Broadcast(this, Dt);
 
+			if (FPhysicsSolver* Solver = GetSolver())
+			{
+				Solver->PushPhysicsState(Dispatcher);
+			}
+
 			FGraphEventRef SimulationCompleteEvent = FGraphEvent::CreateGraphEvent();
 
 			// Need to fire off a parallel task to handle running physics commands and
@@ -1891,6 +1895,11 @@ void FPhysScene_ChaosInterface::StartFrame()
 		// this case. (See FChaosSolversModule::SyncTask and FPhysicsThreadSyncCaller)
 		case EChaosThreadingMode::DedicatedThread:
 		default:
+			if (FPhysicsSolver* Solver = GetSolver())
+			{
+				Solver->PushPhysicsState(Dispatcher);
+			}
+			break;
 			break;
 		}
 	}
