Use impulse properties rather than velocity sets for AddImpulse

diff --git a/Engine/Source/Runtime/Engine/Private/PhysicsEngine/Experimental/PhysInterface_Chaos.cpp b/Engine/Source/Runtime/Engine/Private/PhysicsEngine/Experimental/PhysInterface_Chaos.cpp
--- a/Engine/Source/Runtime/Engine/Private/PhysicsEngine/Experimental/PhysInterface_Chaos.cpp
+++ b/Engine/Source/Runtime/Engine/Private/PhysicsEngine/Experimental/PhysInterface_Chaos.cpp
@@ -782,7 +782,7 @@ void FPhysInterface_Chaos::AddImpulse_AssumesLocked(const FPhysicsActorHandle& I
 		Chaos::TPBDRigidParticle<float, 3>* Rigid = InActorReference->CastToRigidParticle();
 		if (ensure(Rigid))
 		{
-			AddVelocity_AssumesLocked(InActorReference, Rigid->InvM() * InForce);
+			Rigid->SetLinearImpulse(Rigid->LinearImpulse() + InForce);
 		}
 	}
 }
@@ -794,7 +794,7 @@ void FPhysInterface_Chaos::AddAngularImpulseInRadians_AssumesLocked(const FPhysi
 		Chaos::TPBDRigidParticle<float, 3>* Rigid = InActorReference->CastToRigidParticle();
 		if (ensure(Rigid))
 		{
-			AddAngularVelocityInRadians_AssumesLocked(InActorReference, Chaos::Utilities::ComputeWorldSpaceInertia(Rigid->R(), Rigid->InvI()) * InTorque);
+			Rigid->SetAngularImpulse(Rigid->AngularImpulse() + InTorque);
 		}
 	}
 }
diff --git a/Engine/Source/Runtime/Experimental/Chaos/Public/Chaos/PerParticleAddImpulses.h b/Engine/Source/Runtime/Experimental/Chaos/Public/Chaos/PerParticleAddImpulses.h
--- a/Engine/Source/Runtime/Experimental/Chaos/Public/Chaos/PerParticleAddImpulses.h
+++ b/Engine/Source/Runtime/Experimental/Chaos/Public/Chaos/PerParticleAddImpulses.h
@@ -40,6 +40,8 @@ namespace Chaos
 			//
 			const FMatrix33 WorldInvI = Utilities::ComputeWorldSpaceInertia(InParticles.R(Index), InParticles.InvI(Index));
 			InParticles.W(Index) += WorldInvI * InParticles.AngularImpulse(Index);
+			InParticles.LinearImpulse(Index) = TVector<T, 3>(0);
+			InParticles.AngularImpulse(Index) = TVector<T, 3>(0);
 		}
 
 		inline void Apply(TTransientPBDRigidParticleHandle<T, d>& Particle, const T Dt) const override //-V762
@@ -47,6 +49,8 @@ namespace Chaos
 			Particle.V() += Particle.LinearImpulse() * Particle.InvM();
 			const FMatrix33 WorldInvI = Utilities::ComputeWorldSpaceInertia(Particle.R(), Particle.InvI());
 			Particle.W() += WorldInvI * Particle.AngularImpulse();
+			Particle.LinearImpulse() = TVector<T, 3>(0);
+			Particle.AngularImpulse() = TVector<T, 3>(0);
 		}
 	};
 }
diff --git a/Engine/Source/Runtime/Experimental/ChaosSolvers/Private/PhysicsProxy/SingleParticlePhysicsProxy.cpp b/Engine/Source/Runtime/Experimental/ChaosSolvers/Private/PhysicsProxy/SingleParticlePhysicsProxy.cpp
--- a/Engine/Source/Runtime/Experimental/ChaosSolvers/Private/PhysicsProxy/SingleParticlePhysicsProxy.cpp
+++ b/Engine/Source/Runtime/Experimental/ChaosSolvers/Private/PhysicsProxy/SingleParticlePhysicsProxy.cpp
@@ -333,6 +333,19 @@ void FSingleParticlePhysicsProxy<Chaos::TPBDRigidParticle<float, 3>>::PushToPhys
 			RigidHandle->SetTorque(Data->MTorque);
 			bDynamicPropertyUpdated = true;
 		}
+
+		if (Data->DirtyFlags.IsDirty(Chaos::EParticleFlags::LinearImpulse))
+		{
+			RigidHandle->SetLinearImpulse(Data->MLinearImpulse);
+			bDynamicPropertyUpdated = true;
+		}
+
+		if (Data->DirtyFlags.IsDirty(Chaos::EParticleFlags::AngularImpulse))
+		{
+			RigidHandle->SetAngularImpulse(Data->MAngularImpulse);
+			bDynamicPropertyUpdated = true;
+		}
+
 		if (Data->DirtyFlags.IsDirty(Chaos::EParticleFlags::ObjectState))
 		{
 			GetSolver()->GetEvolution()->SetParticleObjectState(RigidHandle, Data->MObjectState);
