- Build system fixes (GoogleAR, BodyInstance)

diff --git a/Engine/Plugins/Runtime/AR/Google/GoogleARCore/Source/GoogleARCoreRendering/Private/GoogleARCorePassthroughCameraRenderer.cpp b/Engine/Plugins/Runtime/AR/Google/GoogleARCore/Source/GoogleARCoreRendering/Private/GoogleARCorePassthroughCameraRenderer.cpp
--- a/Engine/Plugins/Runtime/AR/Google/GoogleARCore/Source/GoogleARCoreRendering/Private/GoogleARCorePassthroughCameraRenderer.cpp
+++ b/Engine/Plugins/Runtime/AR/Google/GoogleARCore/Source/GoogleARCoreRendering/Private/GoogleARCorePassthroughCameraRenderer.cpp
@@ -242,8 +242,8 @@ void FGoogleARCorePassthroughCameraRenderer::RenderVideoOverlay_RenderThread(FRH
 		const FMaterial* CameraMaterial = RenderingOverlayMaterial->GetRenderProxy()->GetMaterial(FeatureLevel);
 		const FMaterialShaderMap* MaterialShaderMap = CameraMaterial->GetRenderingThreadShaderMap();
 
-		FGoogleARCoreCameraOverlayPS* PixelShader = MaterialShaderMap->GetShader<FGoogleARCoreCameraOverlayPS>();
-		FGoogleARCoreCameraOverlayVS* VertexShader = MaterialShaderMap->GetShader<FGoogleARCoreCameraOverlayVS>();
+		TShaderRef<FGoogleARCoreCameraOverlayPS> PixelShader = MaterialShaderMap->GetShader<FGoogleARCoreCameraOverlayPS>();
+		TShaderRef<FGoogleARCoreCameraOverlayVS> VertexShader = MaterialShaderMap->GetShader<FGoogleARCoreCameraOverlayVS>();
 
 		FGraphicsPipelineStateInitializer GraphicsPSOInit;
 
diff --git a/Engine/Source/Runtime/Engine/Classes/PhysicsEngine/BodyInstance.h b/Engine/Source/Runtime/Engine/Classes/PhysicsEngine/BodyInstance.h
--- a/Engine/Source/Runtime/Engine/Classes/PhysicsEngine/BodyInstance.h
+++ b/Engine/Source/Runtime/Engine/Classes/PhysicsEngine/BodyInstance.h
@@ -114,6 +114,10 @@ struct FInitBodiesHelperBase
 	// #PHYS2 this used to return arrays of low-level physics bodies, which would be added to scene in InitBodies. Should it still do that, rather then later iterate over BodyInstances to get phys actor refs?
 	bool CreateShapesAndActors();
 	void InitBodies();
+
+protected:
+	void UpdateSimulatingAndBlendWeight();
+
 };
 
 template <bool bCompileStatic>
@@ -127,7 +131,7 @@ struct FInitBodiesHelper : public FInitBodiesHelperBase
 		SkelMeshComp = bCompileStatic ? nullptr : Cast<USkeletalMeshComponent>(PrimitiveComp);
 		if(SpawnParams.bPhysicsTypeDeterminesSimulation)
 		{
-			GetSimulatingAndBlendWeight(SkelMeshComp, BodySetup, InstanceBlendWeight, bInstanceSimulatePhysics);
+			this->UpdateSimulatingAndBlendWeight();
 		}
 	}
 };
@@ -144,7 +148,7 @@ struct FInitBodiesHelperWithData : public FInitBodiesHelperBase
 		SkelMeshComp = bCompileStatic ? nullptr : Cast<USkeletalMeshComponent>(PrimitiveComp);
 		if(SpawnParams.bPhysicsTypeDeterminesSimulation)
 		{
-			GetSimulatingAndBlendWeight(SkelMeshComp, BodySetup, InstanceBlendWeight, bInstanceSimulatePhysics);
+			this->UpdateSimulatingAndBlendWeight();
 		}
 	}
 
@@ -162,7 +166,7 @@ struct FInitBodiesHelperWithData : public FInitBodiesHelperBase
 		SkelMeshComp = bCompileStatic ? nullptr : Cast<USkeletalMeshComponent>(PrimitiveComp);
 		if(SpawnParams.bPhysicsTypeDeterminesSimulation)
 		{
-			GetSimulatingAndBlendWeight(SkelMeshComp, BodySetup, InstanceBlendWeight, bInstanceSimulatePhysics);
+			this->UpdateSimulatingAndBlendWeight();
 		}
 	}
 
diff --git a/Engine/Source/Runtime/Engine/Private/PhysicsEngine/BodyInstance.cpp b/Engine/Source/Runtime/Engine/Private/PhysicsEngine/BodyInstance.cpp
--- a/Engine/Source/Runtime/Engine/Private/PhysicsEngine/BodyInstance.cpp
+++ b/Engine/Source/Runtime/Engine/Private/PhysicsEngine/BodyInstance.cpp
@@ -948,7 +948,7 @@ TSharedPtr<TArray<ANSICHAR>> GetDebugDebugName(const UPrimitiveComponent* Primit
 }
 #endif
 
-void GetSimulatingAndBlendWeight(const USkeletalMeshComponent* SkelMeshComp, const UBodySetup* BodySetup, float& InstanceBlendWeight, bool& bInstanceSimulatePhysics)
+static void GetSimulatingAndBlendWeight(const USkeletalMeshComponent* SkelMeshComp, const UBodySetup* BodySetup, float& InstanceBlendWeight, bool& bInstanceSimulatePhysics)
 {
 	bool bEnableSim = false;
 	if (SkelMeshComp)
@@ -987,6 +987,12 @@ void GetSimulatingAndBlendWeight(const USkeletalMeshComponent* SkelMeshComp, con
 	}
 }
 
+void FInitBodiesHelperBase::UpdateSimulatingAndBlendWeight()
+{
+	GetSimulatingAndBlendWeight(SkelMeshComp, BodySetup, InstanceBlendWeight, bInstanceSimulatePhysics);
+}
+
+
 FInitBodiesHelperBase::FInitBodiesHelperBase(TArray<FBodyInstance*>& InBodies, TArray<FTransform>& InTransforms, class UBodySetup* InBodySetup, class UPrimitiveComponent* InPrimitiveComp, FPhysScene* InRBScene, const FInitBodySpawnParams& InSpawnParams, FPhysicsAggregateHandle InAggregate)
 	: Bodies(InBodies)
 	, Transforms(InTransforms)
