* Made IsSplitScreen blueprintable
* Fixed Splitscreen detection for triggering underwater sound effects

diff --git a/Engine/Source/Runtime/Engine/Classes/Kismet/KismetSystemLibrary.h b/Engine/Source/Runtime/Engine/Classes/Kismet/KismetSystemLibrary.h
--- a/Engine/Source/Runtime/Engine/Classes/Kismet/KismetSystemLibrary.h
+++ b/Engine/Source/Runtime/Engine/Classes/Kismet/KismetSystemLibrary.h
@@ -189,6 +189,10 @@ class ENGINE_API UKismetSystemLibrary : public UBlueprintFunctionLibrary
 	UFUNCTION(BlueprintPure, Category="Networking", meta=(WorldContext="WorldContextObject"))
 	static bool IsStandalone(UObject* WorldContextObject);
 	
+	/** Returns whether we're currently running in split screen (more than one local player). */
+	UFUNCTION(BlueprintPure, Category = "Utilities", meta = (WorldContext = "WorldContextObject"))
+	static bool IsSplitScreen(UObject* WorldContextObject);
+
 	/** Returns whether this is a build that is packaged for distribution */
 	UFUNCTION(BlueprintPure, Category="Development", meta=(BlueprintThreadSafe))
 	static bool IsPackagedForDistribution();
diff --git a/Engine/Source/Runtime/Engine/Private/KismetSystemLibrary.cpp b/Engine/Source/Runtime/Engine/Private/KismetSystemLibrary.cpp
--- a/Engine/Source/Runtime/Engine/Private/KismetSystemLibrary.cpp
+++ b/Engine/Source/Runtime/Engine/Private/KismetSystemLibrary.cpp
@@ -203,6 +203,12 @@ bool UKismetSystemLibrary::IsStandalone(UObject* WorldContextObject)
 	return World ? (World->GetNetMode() == NM_Standalone) : false;
 }
 
+bool UKismetSystemLibrary::IsSplitScreen(UObject* WorldContextObject)
+{
+	UWorld* World = GEngine->GetWorldFromContextObject(WorldContextObject, EGetWorldErrorMode::LogAndReturnNull);
+	return World ? GEngine->IsSplitScreen(World) : false;
+}
+
 bool UKismetSystemLibrary::IsPackagedForDistribution()
 {
 	return FPlatformMisc::IsPackagedForDistribution();
