Update AddVelocity to only work on Rigids (like PhysX), and go through the impulse path rather than directly setting velocity.

diff --git a/Engine/Source/Runtime/Engine/Private/PhysicsEngine/Experimental/PhysInterface_Chaos.cpp b/Engine/Source/Runtime/Engine/Private/PhysicsEngine/Experimental/PhysInterface_Chaos.cpp
--- a/Engine/Source/Runtime/Engine/Private/PhysicsEngine/Experimental/PhysInterface_Chaos.cpp
+++ b/Engine/Source/Runtime/Engine/Private/PhysicsEngine/Experimental/PhysInterface_Chaos.cpp
@@ -803,10 +803,10 @@ void FPhysInterface_Chaos::AddVelocity_AssumesLocked(const FPhysicsActorHandle&
 {
 	if (ensure(FPhysicsInterface::IsValid(InActorReference)))
 	{
-		Chaos::TKinematicGeometryParticle<float, 3>* Kinematic = InActorReference->CastToKinematicParticle();
-		if (ensure(Kinematic))
+		Chaos::TPBDRigidParticle<float, 3>* Rigid = InActorReference->CastToRigidParticle();
+		if (ensure(Rigid))
 		{
-			Kinematic->SetV(Kinematic->V() + InVelocityDelta);
+			AddImpulse_AssumesLocked(InActorReference, Rigid->M() * InVelocityDelta);
 		}
 	}
 }
@@ -815,11 +815,11 @@ void FPhysInterface_Chaos::AddAngularVelocityInRadians_AssumesLocked(const FPhys
 {
 	if (ensure(FPhysicsInterface::IsValid(InActorReference)))
 	{
-		Chaos::TKinematicGeometryParticle<float, 3>* Kinematic = InActorReference->CastToKinematicParticle();
-		if (ensure(Kinematic))
+		Chaos::TPBDRigidParticle<float, 3>* Rigid = InActorReference->CastToRigidParticle();
+		if (ensure(Rigid))
 		{
-			// NOTE: Need to convert W to rad/sec?
-			Kinematic->SetW(Kinematic->W() + InAngularVelocityDeltaRad);
+			const Chaos::FMatrix33 WorldI = Chaos::Utilities::ComputeWorldSpaceInertia(Rigid->R(), Rigid->I());
+			AddAngularImpulseInRadians_AssumesLocked(InActorReference, WorldI * InAngularVelocityDeltaRad);
 		}
 	}
 }
