Disable Property Comparison tracking by default in the NetworkProfiler.
It can be re-enabled using Net.ProfilerUseComparisonTracking

diff --git a/Engine/Source/Runtime/Engine/Private/DataReplication.cpp b/Engine/Source/Runtime/Engine/Private/DataReplication.cpp
--- a/Engine/Source/Runtime/Engine/Private/DataReplication.cpp
+++ b/Engine/Source/Runtime/Engine/Private/DataReplication.cpp
@@ -1561,7 +1561,7 @@ bool FObjectReplicator::ReplicateProperties( FOutBunch & Bunch, FReplicationFlag
 	const bool bHasRepLayout = RepLayout->ReplicateProperties(SendingRepState, ChangelistMgr->GetRepChangelistState(), (uint8*)Object, ObjectClass, OwningChannel, Writer, RepFlags);
 
 	NETWORK_PROFILER(
-		if (bHasRepLayout)
+		if (GNetworkProfiler.IsComparisonTrackingEnabled() && bHasRepLayout)
 		{
 			GNetworkProfiler.TrackReplicatePropertiesMetadata(RepLayout->GetOwner(), SendingRepState->InactiveParents, (Connection->ResendAllDataState != EResendAllDataState::None), Connection);
 		}
diff --git a/Engine/Source/Runtime/Engine/Private/NetworkProfiler.cpp b/Engine/Source/Runtime/Engine/Private/NetworkProfiler.cpp
--- a/Engine/Source/Runtime/Engine/Private/NetworkProfiler.cpp
+++ b/Engine/Source/Runtime/Engine/Private/NetworkProfiler.cpp
@@ -12,11 +12,27 @@
 #include "Engine/EngineBaseTypes.h"
 #include "Engine/World.h"
 #include "Serialization/MemoryWriter.h"
+#include "HAL/IConsoleManager.h"
 
 #if USE_NETWORK_PROFILER
 
 #define SCOPE_LOCK_REF(X) FScopeLock ScopeLock(&X);
 
+struct FNetworkProfilerCVarHelper
+{
+	static bool& GetNetProfilerIsComparisonTrackingEnabled()
+	{
+		return FNetworkProfiler::bIsComparisonTrackingEnabled;
+	}
+};
+
+bool FNetworkProfiler::bIsComparisonTrackingEnabled = false;
+static FAutoConsoleVariableRef CVarProfilerUseComparisonTracking(
+	TEXT("Net.ProfilerUseComparisonTracking"),
+	FNetworkProfilerCVarHelper::GetNetProfilerIsComparisonTrackingEnabled(),
+	TEXT("")
+);
+
 
 /**
  * Whether to track the raw network data or not.
@@ -559,7 +575,7 @@ void FNetworkProfiler::TrackCompareProperties_Unsafe(const FString& ObjectName,
 
 void FNetworkProfiler::TrackReplicatePropertiesMetadata(const UObject* Object, TBitArray<>& InactiveProperties, bool bSentAllProperties, UNetConnection* Connection)
 {
-	if (bIsTrackingEnabled)
+	if (IsComparisonTrackingEnabled())
 	{
 		SCOPE_LOCK_REF(CriticalSection);
 
diff --git a/Engine/Source/Runtime/Engine/Private/RepLayout.cpp b/Engine/Source/Runtime/Engine/Private/RepLayout.cpp
--- a/Engine/Source/Runtime/Engine/Private/RepLayout.cpp
+++ b/Engine/Source/Runtime/Engine/Private/RepLayout.cpp
@@ -276,6 +276,15 @@ namespace UE4_RepLayout_Private
 		NETWORK_PROFILER(return true;);
 		return false;
 	}
+
+	static bool IsNetworkProfilerComparisonTrackingEnabled()
+	{
+#if USE_NETWORK_PROFILER
+		return GNetworkProfiler.IsComparisonTrackingEnabled();
+#else
+		return false;
+#endif
+	}
 }
 
 //~ TODO: Consider moving the FastArray members into their own sub-struct to save memory for non fast array
@@ -1501,7 +1510,7 @@ bool FRepLayout::CompareProperties(
 		/*PushModelState=*/UE4_RepLayout_Private::GetPerNetDriverState(RepChangelistState),
 		/*PushModelProperties=*/ LocalPushModelProperties,	
 		/*bValidateProperties=*/GbPushModelValidateProperties,
-		/*bIsNetworkProfilerActive=*/UE4_RepLayout_Private::IsNetworkProfilerEnabled()
+		/*bIsNetworkProfilerActive=*/UE4_RepLayout_Private::IsNetworkProfilerComparisonTrackingEnabled()
 	};
 
 	FComparePropertiesStackParams StackParams{
@@ -1540,7 +1549,10 @@ bool FRepLayout::CompareProperties(
 	
 		CompareParentProperties(SharedParams, StackParams);
 
-		NETWORK_PROFILER(GNetworkProfiler.TrackCompareProperties(Owner, FPlatformTime::Cycles() - ReplicateParentPropertiesStartTime, SharedParams.PropertiesCompared, SharedParams.PropertiesChanged, Parents, &FPropertyNameHelper::ConvertParentCmdToPropertyName););
+		if (SharedParams.bIsNetworkProfilerActive)
+		{
+			NETWORK_PROFILER(GNetworkProfiler.TrackCompareProperties(Owner, FPlatformTime::Cycles() - ReplicateParentPropertiesStartTime, SharedParams.PropertiesCompared, SharedParams.PropertiesChanged, Parents, &FPropertyNameHelper::ConvertParentCmdToPropertyName););
+		}
 	}
 
 	if (Changed.Num() == 0)
diff --git a/Engine/Source/Runtime/Engine/Public/Net/NetworkProfiler.h b/Engine/Source/Runtime/Engine/Public/Net/NetworkProfiler.h
--- a/Engine/Source/Runtime/Engine/Public/Net/NetworkProfiler.h
+++ b/Engine/Source/Runtime/Engine/Public/Net/NetworkProfiler.h
@@ -77,6 +77,10 @@ class FNetworkProfiler
 private:
 
 	friend struct FNetworkProfilerScopedIgnoreReplicateProperties;
+	friend struct FNetworkProfilerCVarHelper;
+
+	/** Whether or not want to track granular information about comparisons. This can be very expensive. */
+	static bool bIsComparisonTrackingEnabled;
 
 	/** File writer used to serialize data.															*/
 	FArchive*								FileWriter;
@@ -348,7 +352,7 @@ public:
 	template<typename T, typename ProjectionType>
 	void TrackCompareProperties(const UObject* Object, uint32 Cycles, TBitArray<>& PropertiesCompared, TBitArray<>& PropertiesThatChanged, const TArray<T>& PropertyNameContainers, ProjectionType PropertyNameProjection)
 	{
-		if (bIsTrackingEnabled)
+		if (IsComparisonTrackingEnabled())
 		{
             FScopeLock ScopeLock(&CriticalSection);
 
@@ -463,6 +467,7 @@ public:
 	bool Exec( UWorld * InWorld, const TCHAR* Cmd, FOutputDevice & Ar );
 
 	bool FORCEINLINE IsTrackingEnabled() const { return bIsTrackingEnabled; }
+	bool IsComparisonTrackingEnabled() const { return bIsTrackingEnabled && bIsComparisonTrackingEnabled;  }
 
 	/** Return the network profile finished delegate */
 	FOnNetworkProfileFinished& OnNetworkProfileFinished() { return OnNetworkProfileFinishedDelegate; }
