Submit missing includes
#rnx

diff --git a/Engine/Source/Developer/DesktopPlatform/Private/PlatformInfo.cpp b/Engine/Source/Developer/DesktopPlatform/Private/PlatformInfo.cpp
--- a/Engine/Source/Developer/DesktopPlatform/Private/PlatformInfo.cpp
+++ b/Engine/Source/Developer/DesktopPlatform/Private/PlatformInfo.cpp
@@ -4,6 +4,7 @@
 #include "DesktopPlatformPrivate.h"
 #include "Misc/DataDrivenPlatformInfoRegistry.h"
 #include "HAL/FileManager.h"
+#include "Misc/ConfigCacheIni.h"
 #include "Misc/Paths.h"
 
 #define LOCTEXT_NAMESPACE "PlatformInfo"
diff --git a/Engine/Source/Runtime/Core/Private/Misc/PreLoadFile.cpp b/Engine/Source/Runtime/Core/Private/Misc/PreLoadFile.cpp
--- a/Engine/Source/Runtime/Core/Private/Misc/PreLoadFile.cpp
+++ b/Engine/Source/Runtime/Core/Private/Misc/PreLoadFile.cpp
@@ -1,8 +1,12 @@
 // Copyright 1998-2019 Epic Games, Inc. All Rights Reserved.
 
 #include "Misc/PreLoadFile.h"
+#include "HAL/Event.h"
+#include "HAL/FileManager.h"
 #include "HAL/PlatformFilemanager.h"
 #include "Async/AsyncFileHandle.h"
+#include "Misc/Paths.h"
+#include "Misc/ScopeLock.h"
 
 bool FPreLoadFile::bSystemNoLongerTakingRequests = false;
 
diff --git a/Engine/Source/Runtime/Core/Public/Misc/PreLoadFile.h b/Engine/Source/Runtime/Core/Public/Misc/PreLoadFile.h
--- a/Engine/Source/Runtime/Core/Public/Misc/PreLoadFile.h
+++ b/Engine/Source/Runtime/Core/Public/Misc/PreLoadFile.h
@@ -2,6 +2,7 @@
 
 #pragma once
 
+#include "CoreMinimal.h"
 #include "Misc/DelayedAutoRegister.h"
 
 struct CORE_API FPreLoadFile : public FDelayedAutoRegisterHelper
diff --git a/Engine/Source/Runtime/Core/Public/Serialization/MemoryLayout.h b/Engine/Source/Runtime/Core/Public/Serialization/MemoryLayout.h
--- a/Engine/Source/Runtime/Core/Public/Serialization/MemoryLayout.h
+++ b/Engine/Source/Runtime/Core/Public/Serialization/MemoryLayout.h
@@ -2,6 +2,7 @@
 
 #pragma once
 
+#include "Containers/EnumAsByte.h"
 #include "Templates/IsAbstract.h"
 #include "Templates/IsPolymorphic.h"
 #include "Templates/IsTriviallyDestructible.h"
diff --git a/Engine/Source/Runtime/Engine/Public/DebugViewModeInterface.h b/Engine/Source/Runtime/Engine/Public/DebugViewModeInterface.h
--- a/Engine/Source/Runtime/Engine/Public/DebugViewModeInterface.h
+++ b/Engine/Source/Runtime/Engine/Public/DebugViewModeInterface.h
@@ -11,11 +11,15 @@ DebugViewModeInterface.h: Contains definitions for rendering debug viewmodes.
 #include "Engine/EngineTypes.h"
 #include "RHI.h"
 #include "RHIResources.h"
+#include "Shader.h"
 
 #if !(UE_BUILD_SHIPPING || UE_BUILD_TEST)
 
 class FDebugViewModePS;
 class FMaterial;
+class FMaterialRenderProxy;
+class FMeshDrawSingleShaderBindings;
+class FPrimitiveSceneProxy;
 class FVertexFactoryType;
 struct FMeshPassProcessorRenderState;
 
diff --git a/Engine/Source/Runtime/Engine/Public/MaterialCachedData.h b/Engine/Source/Runtime/Engine/Public/MaterialCachedData.h
--- a/Engine/Source/Runtime/Engine/Public/MaterialCachedData.h
+++ b/Engine/Source/Runtime/Engine/Public/MaterialCachedData.h
@@ -7,7 +7,10 @@
 #include "MaterialCachedData.generated.h"
 
 class UTexture;
+class UCurveLinearColor;
+class UCurveLinearColorAtlas;
 class UFont;
+class UMaterialExpression;
 class URuntimeVirtualTexture;
 class ULandscapeGrassType;
 
diff --git a/Engine/Source/Runtime/Engine/Public/MaterialShaderType.h b/Engine/Source/Runtime/Engine/Public/MaterialShaderType.h
--- a/Engine/Source/Runtime/Engine/Public/MaterialShaderType.h
+++ b/Engine/Source/Runtime/Engine/Public/MaterialShaderType.h
@@ -26,6 +26,7 @@ class FShaderCommonCompileJob;
 class FShaderCompileJob;
 class FUniformExpressionSet;
 class FVertexFactoryType;
+struct FMaterialShaderParameters;
 
 DECLARE_DELEGATE_RetVal_OneParam(FString, FShadingModelToStringDelegate, EMaterialShadingModel)
 
diff --git a/Engine/Source/Runtime/Engine/Public/WeightedRandomSampler.h b/Engine/Source/Runtime/Engine/Public/WeightedRandomSampler.h
--- a/Engine/Source/Runtime/Engine/Public/WeightedRandomSampler.h
+++ b/Engine/Source/Runtime/Engine/Public/WeightedRandomSampler.h
@@ -23,6 +23,7 @@ GetEntryIndex takes two random floating values. This does not generate randoms f
 #pragma once
 
 #include "CoreMinimal.h"
+#include "Serialization/MemoryImage.h"
 #include "Serialization/MemoryLayout.h"
 
 struct ENGINE_API FWeightedRandomSampler
