Fix for soundwaves getting released prematurely: if multiple soundcues referenced a soundwave we would release that soundwave once the first soundcue GC happened.

diff --git a/Engine/Source/Runtime/Engine/Private/SoundCue.cpp b/Engine/Source/Runtime/Engine/Private/SoundCue.cpp
--- a/Engine/Source/Runtime/Engine/Private/SoundCue.cpp
+++ b/Engine/Source/Runtime/Engine/Private/SoundCue.cpp
@@ -436,11 +436,6 @@ void USoundCue::AudioQualityChanged()
 void USoundCue::BeginDestroy()
 {
 	Super::BeginDestroy();
-	
-	if (bIsRetainingAudio)
-	{
-		ReleaseRetainedAudio();
-	}
 }
 
 FString USoundCue::GetDesc()
diff --git a/Engine/Source/Runtime/Engine/Private/SoundNode.cpp b/Engine/Source/Runtime/Engine/Private/SoundNode.cpp
--- a/Engine/Source/Runtime/Engine/Private/SoundNode.cpp
+++ b/Engine/Source/Runtime/Engine/Private/SoundNode.cpp
@@ -186,10 +186,6 @@ void USoundNode::ReleaseRetainerOnChildWavePlayers(bool bRecurse)
 void USoundNode::BeginDestroy()
 {
 	Super::BeginDestroy();
-	if (bIsRetainingAudio)
-	{
-		ReleaseRetainerOnChildWavePlayers(true);
-	}
 }
 
 void USoundNode::ParseNodes( FAudioDevice* AudioDevice, const UPTRINT NodeWaveInstanceHash, FActiveSound& ActiveSound, const FSoundParseParameters& ParseParams, TArray<FWaveInstance*>& WaveInstances )
diff --git a/Engine/Source/Runtime/Engine/Private/SoundNodeSoundClass.cpp b/Engine/Source/Runtime/Engine/Private/SoundNodeSoundClass.cpp
--- a/Engine/Source/Runtime/Engine/Private/SoundNodeSoundClass.cpp
+++ b/Engine/Source/Runtime/Engine/Private/SoundNodeSoundClass.cpp
@@ -55,8 +55,4 @@ void USoundNodeSoundClass::PostLoad()
 void USoundNodeSoundClass::BeginDestroy()
 {
 	Super::BeginDestroy();
-	if (bRetainingAudioDueToSoundClass)
-	{
-		ReleaseRetainerOnChildWavePlayers(true);
-	}
 }
