Deprecate public replicated members of APlayerState and expose Getters and Setters instead.

diff --git a/Engine/Plugins/Online/OnlineFramework/Source/Lobby/Private/LobbyBeaconClient.cpp b/Engine/Plugins/Online/OnlineFramework/Source/Lobby/Private/LobbyBeaconClient.cpp
--- a/Engine/Plugins/Online/OnlineFramework/Source/Lobby/Private/LobbyBeaconClient.cpp
+++ b/Engine/Plugins/Online/OnlineFramework/Source/Lobby/Private/LobbyBeaconClient.cpp
@@ -163,7 +163,7 @@ void ALobbyBeaconClient::LoginLocalPlayers()
 	for (FConstPlayerControllerIterator Iterator = World->GetPlayerControllerIterator(); Iterator; ++Iterator)
 	{
 		APlayerController* PC = Iterator->Get();
-		if (PC && PC->PlayerState && PC->PlayerState->UniqueId.IsValid())
+		if (PC && PC->PlayerState && PC->PlayerState->GetUniqueId().IsValid())
 		{
 			ULocalPlayer* LP = Cast<ULocalPlayer>(PC->Player);
 			if (LP)
diff --git a/Engine/Plugins/Online/OnlineSubsystemSteam/Source/Private/OnlineAuthInterfaceSteam.cpp b/Engine/Plugins/Online/OnlineSubsystemSteam/Source/Private/OnlineAuthInterfaceSteam.cpp
--- a/Engine/Plugins/Online/OnlineSubsystemSteam/Source/Private/OnlineAuthInterfaceSteam.cpp
+++ b/Engine/Plugins/Online/OnlineSubsystemSteam/Source/Private/OnlineAuthInterfaceSteam.cpp
@@ -341,8 +341,8 @@ bool FOnlineAuthSteam::KickPlayer(const FUniqueNetId& InUserId, bool bSuppressFa
 		for (FConstPlayerControllerIterator Itr = World->GetPlayerControllerIterator(); Itr; ++Itr)
 		{
 			APlayerController* PC = Itr->Get();
-			if (PC && PC->PlayerState != nullptr && PC->PlayerState->UniqueId.IsValid() &&
-				*(PC->PlayerState->UniqueId.GetUniqueNetId()) == InUserId)
+			if (PC && PC->PlayerState != nullptr && PC->PlayerState->GetUniqueId().IsValid() &&
+				*(PC->PlayerState->GetUniqueId().GetUniqueNetId()) == InUserId)
 			{
 				const FText AuthKickReason = NSLOCTEXT("NetworkErrors", "HostClosedConnection", "Host closed the connection.");
 				bKickSuccess = GameMode->GameSession->KickPlayer(PC, AuthKickReason);
diff --git a/Engine/Plugins/Online/OnlineSubsystemSteam/Source/Private/OnlineSessionAsyncServerSteam.cpp b/Engine/Plugins/Online/OnlineSubsystemSteam/Source/Private/OnlineSessionAsyncServerSteam.cpp
--- a/Engine/Plugins/Online/OnlineSubsystemSteam/Source/Private/OnlineSessionAsyncServerSteam.cpp
+++ b/Engine/Plugins/Online/OnlineSubsystemSteam/Source/Private/OnlineSessionAsyncServerSteam.cpp
@@ -256,10 +256,10 @@ void UpdatePublishedSettings(UWorld* World, FNamedOnlineSession* Session)
 			for (int32 PlayerIdx=0; PlayerIdx < GameState->PlayerArray.Num(); PlayerIdx++)
 			{
 				APlayerState const* const PlayerState = GameState->PlayerArray[PlayerIdx];
-				if (PlayerState && PlayerState->UniqueId.IsValid())
+				if (PlayerState && PlayerState->GetUniqueId().IsValid())
 				{
-					CSteamID SteamId(*(uint64*)PlayerState->UniqueId->GetBytes());
-					SteamGameServerPtr->BUpdateUserData(SteamId, TCHAR_TO_UTF8(*PlayerState->GetPlayerName()), PlayerState->Score);
+					CSteamID SteamId(*(uint64*)PlayerState->GetUniqueId()->GetBytes());
+					SteamGameServerPtr->BUpdateUserData(SteamId, TCHAR_TO_UTF8(*PlayerState->GetPlayerName()), PlayerState->GetScore());
 				}
 			}
 		}
diff --git a/Engine/Plugins/Online/OnlineSubsystemUtils/Source/OnlineSubsystemUtils/Private/LeaderboardBlueprintLibrary.cpp b/Engine/Plugins/Online/OnlineSubsystemUtils/Source/OnlineSubsystemUtils/Private/LeaderboardBlueprintLibrary.cpp
--- a/Engine/Plugins/Online/OnlineSubsystemUtils/Source/OnlineSubsystemUtils/Private/LeaderboardBlueprintLibrary.cpp
+++ b/Engine/Plugins/Online/OnlineSubsystemUtils/Source/OnlineSubsystemUtils/Private/LeaderboardBlueprintLibrary.cpp
@@ -20,7 +20,7 @@ bool ULeaderboardBlueprintLibrary::WriteLeaderboardObject(APlayerController* Pla
 {
 	if (APlayerState* PlayerState = (PlayerController != NULL) ? PlayerController->PlayerState : NULL)
 	{
-		TSharedPtr<const FUniqueNetId> UserId = PlayerState->UniqueId.GetUniqueNetId();
+		TSharedPtr<const FUniqueNetId> UserId = PlayerState->GetUniqueId().GetUniqueNetId();
 		if (UserId.IsValid())
 		{
 			if (IOnlineSubsystem* const OnlineSub = IOnlineSubsystem::IsLoaded() ? IOnlineSubsystem::Get() : nullptr)
diff --git a/Engine/Plugins/Online/OnlineSubsystemUtils/Source/OnlineSubsystemUtils/Private/LeaderboardFlushCallbackProxy.cpp b/Engine/Plugins/Online/OnlineSubsystemUtils/Source/OnlineSubsystemUtils/Private/LeaderboardFlushCallbackProxy.cpp
--- a/Engine/Plugins/Online/OnlineSubsystemUtils/Source/OnlineSubsystemUtils/Private/LeaderboardFlushCallbackProxy.cpp
+++ b/Engine/Plugins/Online/OnlineSubsystemUtils/Source/OnlineSubsystemUtils/Private/LeaderboardFlushCallbackProxy.cpp
@@ -19,7 +19,7 @@ void ULeaderboardFlushCallbackProxy::TriggerFlush(APlayerController* PlayerContr
 
 	if (APlayerState* PlayerState = (PlayerController != NULL) ? PlayerController->PlayerState : NULL)
 	{
-		TSharedPtr<const FUniqueNetId> UserID = PlayerState->UniqueId.GetUniqueNetId();
+		TSharedPtr<const FUniqueNetId> UserID = PlayerState->GetUniqueId().GetUniqueNetId();
 		if (UserID.IsValid())
 		{
 			if (IOnlineSubsystem* const OnlineSub = IOnlineSubsystem::IsLoaded() ? IOnlineSubsystem::Get() : nullptr)
diff --git a/Engine/Plugins/Online/OnlineSubsystemUtils/Source/OnlineSubsystemUtils/Private/LeaderboardQueryCallbackProxy.cpp b/Engine/Plugins/Online/OnlineSubsystemUtils/Source/OnlineSubsystemUtils/Private/LeaderboardQueryCallbackProxy.cpp
--- a/Engine/Plugins/Online/OnlineSubsystemUtils/Source/OnlineSubsystemUtils/Private/LeaderboardQueryCallbackProxy.cpp
+++ b/Engine/Plugins/Online/OnlineSubsystemUtils/Source/OnlineSubsystemUtils/Private/LeaderboardQueryCallbackProxy.cpp
@@ -22,7 +22,7 @@ void ULeaderboardQueryCallbackProxy::TriggerQuery(APlayerController* PlayerContr
 	WorldPtr = (PlayerController != NULL) ? PlayerController->GetWorld() : NULL;
 	if (APlayerState* PlayerState = (PlayerController != NULL) ? PlayerController->PlayerState : NULL)
 	{
-		TSharedPtr<const FUniqueNetId> UserID = PlayerState->UniqueId.GetUniqueNetId();
+		TSharedPtr<const FUniqueNetId> UserID = PlayerState->GetUniqueId().GetUniqueNetId();
 		if (UserID.IsValid())
 		{
 			if (IOnlineSubsystem* const OnlineSub = IOnlineSubsystem::IsLoaded() ? IOnlineSubsystem::Get() : nullptr)
diff --git a/Engine/Plugins/Online/OnlineSubsystemUtils/Source/OnlineSubsystemUtils/Private/OnlineSubsystemUtils.cpp b/Engine/Plugins/Online/OnlineSubsystemUtils/Source/OnlineSubsystemUtils/Private/OnlineSubsystemUtils.cpp
--- a/Engine/Plugins/Online/OnlineSubsystemUtils/Source/OnlineSubsystemUtils/Private/OnlineSubsystemUtils.cpp
+++ b/Engine/Plugins/Online/OnlineSubsystemUtils/Source/OnlineSubsystemUtils/Private/OnlineSubsystemUtils.cpp
@@ -763,7 +763,7 @@ void FOnlineSubsystemBPCallHelper::QueryIDFromPlayerController(APlayerController
 
 	if (APlayerState* PlayerState = (PlayerController != NULL) ? PlayerController->PlayerState : NULL)
 	{
-		UserID = PlayerState->UniqueId.GetUniqueNetId();
+		UserID = PlayerState->GetUniqueId().GetUniqueNetId();
 		if (!UserID.IsValid())
 		{
 			FFrame::KismetExecutionMessage(*FString::Printf(TEXT("%s - Cannot map local player to unique net ID"), FunctionContext), ELogVerbosity::Warning);
diff --git a/Engine/Source/Runtime/Engine/Classes/GameFramework/PlayerState.h b/Engine/Source/Runtime/Engine/Classes/GameFramework/PlayerState.h
--- a/Engine/Source/Runtime/Engine/Classes/GameFramework/PlayerState.h
+++ b/Engine/Source/Runtime/Engine/Classes/GameFramework/PlayerState.h
@@ -69,24 +69,18 @@ class ENGINE_API APlayerState : public AInfo
 	virtual ~APlayerState();
 
 	/** Player's current score. */
-	UPROPERTY(replicatedUsing=OnRep_Score, BlueprintReadOnly, Category=PlayerState)
+	UE_DEPRECATED(4.25, "This member will be made private. Use GetScore or SetScore instead.")
+	UPROPERTY(ReplicatedUsing=OnRep_Score, BlueprintReadOnly, Category=PlayerState)
 	float Score;
 
-	/** Player name, or blank if none. */
-	UE_DEPRECATED(4.19, "This property is now deprecated, please use GetPlayerName and SetPlayerName/SetPlayerNameInternal functions instead.")
-	UPROPERTY(BlueprintReadOnly, Category = PlayerState)
-	FString PlayerName;
-
-	/** Previous player name.  Saved on client-side to detect player name changes. */
-	UE_DEPRECATED(4.19, "This property is now deprecated, please use GetOldPlayerName() function instead.")
-	FString OldName;
-
 	/** Unique net id number. Actual value varies based on current online subsystem, use it only as a guaranteed unique number per player. */
-	UPROPERTY(replicatedUsing=OnRep_PlayerId, BlueprintReadOnly, Category=PlayerState)
+	UE_DEPRECATED(4.25, "This member will be made private. Use GetPlayerId or SetPlayerId instead.")
+	UPROPERTY(ReplicatedUsing=OnRep_PlayerId, BlueprintReadOnly, Category=PlayerState)
 	int32 PlayerId;
 
 	/** Replicated compressed ping for this player (holds ping in msec divided by 4) */
-	UPROPERTY(replicated, BlueprintReadOnly, Category=PlayerState)
+	UE_DEPRECATED(4.25, "This member will be made private. Use GetPing or SetPing instead.")
+	UPROPERTY(Replicated, BlueprintReadOnly, Category=PlayerState)
 	uint8 Ping;
 
 private:
@@ -104,36 +98,42 @@ private:
 
 public:
 	/** Whether this player is currently a spectator */
-	UPROPERTY(replicated, BlueprintReadOnly, Category=PlayerState)
+	UE_DEPRECATED(4.25, "This member will be made private. Use IsSpectator or SetIsSpectator instead.")
+	UPROPERTY(Replicated, BlueprintReadOnly, Category=PlayerState)
 	uint8 bIsSpectator:1;
 
 	/** Whether this player can only ever be a spectator */
-	UPROPERTY(replicated)
+	UE_DEPRECATED(4.25, "This member will be made private. Use IsOnlyASpectator or SetIsOnlyASpectator instead.")
+	UPROPERTY(Replicated)
 	uint8 bOnlySpectator:1;
 
 	/** True if this PlayerState is associated with an AIController */
-	UPROPERTY(replicated, BlueprintReadOnly, Category=PlayerState)
+	UE_DEPRECATED(4.25, "This member will be made private. Use IsABot or SetIsABot instead.")
+	UPROPERTY(Replicated, BlueprintReadOnly, Category=PlayerState)
 	uint8 bIsABot:1;
 
 	/** client side flag - whether this player has been welcomed or not (player entered message) */
 	uint8 bHasBeenWelcomed:1;
 
 	/** Means this PlayerState came from the GameMode's InactivePlayerArray */
-	UPROPERTY(replicatedUsing=OnRep_bIsInactive)
+	UE_DEPRECATED(4.25, "This member will be made private. Use IsInative or SetIsInactive instead.")
+	UPROPERTY(ReplicatedUsing=OnRep_bIsInactive)
 	uint8 bIsInactive:1;
 
 	/** indicates this is a PlayerState from the previous level of a seamless travel,
 	 * waiting for the player to finish the transition before creating a new one
 	 * this is used to avoid preserving the PlayerState in the InactivePlayerArray if the player leaves
 	 */
-	UPROPERTY(replicated)
+	UE_DEPRECATED(4.25, "This member will be made private. Use IsFromPreviousLevel or SetIsFromPreviousLevel instead.")
+	UPROPERTY(Replicated)
 	uint8 bFromPreviousLevel:1;
 
 	/** if set, GetPlayerName() will call virtual GetPlayerNameCustom() to allow custom access */
 	uint8 bUseCustomPlayerNames : 1;
 
 	/** Elapsed time on server when this PlayerState was first created.  */
-	UPROPERTY(replicated)
+	UE_DEPRECATED(4.25, "This member will be made private. Use GetStartTime or SetStartTime instead.")
+	UPROPERTY(Replicated)
 	int32 StartTime;
 
 	/** This is used for sending game agnostic messages that can be localized */
@@ -152,7 +152,8 @@ public:
 	 * NOTE: the internals of this property should *never* be exposed to the player as it's transient
 	 * and opaque in meaning (ie it might mean date/time followed by something else).
 	 * It is OK to use and pass around this property, though. */
-	UPROPERTY(replicatedUsing=OnRep_UniqueId)
+	UE_DEPRECATED(4.25, "This member will be made private. Use GetUniqueId or SetUniqueId instead.")
+	UPROPERTY(ReplicatedUsing=OnRep_UniqueId)
 	FUniqueNetIdRepl UniqueId; 
 
 	/** The session that the player needs to join/remove from as it is created/leaves */
@@ -178,7 +179,7 @@ private:
 	float			CurPingBucketTimestamp;
 
 	/** Player name, or blank if none. */
-	UPROPERTY(replicatedUsing = OnRep_PlayerName)
+	UPROPERTY(ReplicatedUsing = OnRep_PlayerName)
 	FString PlayerNamePrivate;
 
 	/** Previous player name.  Saved on client-side to detect player name changes. */
@@ -321,6 +322,103 @@ protected:
 private:
 	// Hidden functions that don't make sense to use on this class.
 	HIDE_ACTOR_TRANSFORM_FUNCTIONS();
+
+PRAGMA_DISABLE_DEPRECATION_WARNINGS
+	//~ Begin Methods for Replicated Members.
+public:
+
+	/** Gets the literal value of Score. */
+	float GetScore() const
+	{
+		return Score;
+	}
+
+	/** Sets the value of Score without causing other side effects to this instance. */
+	void SetScore(const float NewScore);
+
+	/** Gets the literal value of PlayerId. */
+	int32 GetPlayerId() const
+	{
+		return PlayerId;
+	}
+
+	/** Sets the value of PlayerId without causing other side effects to this instance. */
+	void SetPlayerId(const int32 NewId);
+
+	/** Gets the literal value of Ping. */
+	uint8 GetPing() const
+	{
+		return Ping;
+	}
+
+	/** Sets the value of Ping without causing other side effects to this instance. */
+	void SetPing(const uint8 NewPing);
+
+	/** Gets the literal value of bIsSpectator. */
+	bool IsSpectator() const
+	{
+		return bIsSpectator;
+	}
+
+	/** Sets the value of bIsSpectator without causing other side effects to this instance. */
+	void SetIsSpectator(const bool bNewSpectator);
+
+	/** Gets the literal value of bOnlySpectator. */
+	bool IsOnlyASpectator() const
+	{
+		return bOnlySpectator;
+	}
+
+	/** Sets the value of bOnlySpectator without causing other side effects to this instance. */
+	void SetIsOnlyASpectator(const bool bNewSpectator);
+
+	/** Gets the literal value of bIsABot. */
+	bool IsABot() const
+	{
+		return bIsABot;
+	}
+
+	/** Sets the value of bIsABot without causing other side effects to this instance. */
+	void SetIsABot(const bool bNewIsABot);
+
+	/** Gets the literal value of bIsInactive. */
+	bool IsInactive() const
+	{
+		return bIsInactive;
+	}
+
+	/** Sets the value of bIsInactive without causing other side effects to this instance. */
+	void SetIsInactive(const bool bNewInactive);
+
+	/** Gets the literal value of bFromPreviousLevel. */
+	bool IsFromPreviousLevel() const
+	{
+		return bFromPreviousLevel;
+	}
+
+	/** Sets the value of bFromPreviousLevel without causing other side effects to this instance. */
+	void SetIsFromPreviousLevel(const bool bNewFromPreviousLevel);
+
+	/** Gets the literal value of StartTime. */
+	int32 GetStartTime() const
+	{
+		return StartTime;
+	}
+
+	/** Sets the value of StartTime without causing other side effects to this instance. */
+	void SetStartTime(const int32 NewStartTime);
+
+	/** Gets the literal value of UniqueId. */
+	const FUniqueNetIdRepl& GetUniqueId() const
+	{
+		return UniqueId;
+	}
+
+	/** Sets the value of UniqueId without causing other side effects to this instance. */
+	void SetUniqueId(const FUniqueNetIdRepl& NewUniqueId);
+	void SetUniqueId(FUniqueNetIdRepl&& NewUniqueId);
+	//~ End Methods for Replicated Members.
+PRAGMA_ENABLE_DEPRECATION_WARNINGS
 };
 
 struct FSetPlayerStatePawn
diff --git a/Engine/Source/Runtime/Engine/Private/DemoNetDriver.cpp b/Engine/Source/Runtime/Engine/Private/DemoNetDriver.cpp
--- a/Engine/Source/Runtime/Engine/Private/DemoNetDriver.cpp
+++ b/Engine/Source/Runtime/Engine/Private/DemoNetDriver.cpp
@@ -4382,7 +4382,7 @@ APlayerController* UDemoNetDriver::CreateDemoPlayerController(UNetConnection* Co
 	// Tell the game that we're spectator and not a normal player
 	if (NewDemoController->PlayerState)
 	{
-		NewDemoController->PlayerState->bOnlySpectator = true;
+		NewDemoController->PlayerState->SetIsOnlyASpectator(true);
 	}
 
 	for (FActorIterator It(World); It; ++It)
diff --git a/Engine/Source/Runtime/Engine/Private/GameInstance.cpp b/Engine/Source/Runtime/Engine/Private/GameInstance.cpp
--- a/Engine/Source/Runtime/Engine/Private/GameInstance.cpp
+++ b/Engine/Source/Runtime/Engine/Private/GameInstance.cpp
@@ -908,7 +908,7 @@ APlayerController* UGameInstance::GetPrimaryPlayerController() const
 	for (FConstPlayerControllerIterator Iterator = World->GetPlayerControllerIterator(); Iterator; ++Iterator)
 	{
 		APlayerController* NextPlayer = Iterator->Get();
-		if (NextPlayer && NextPlayer->PlayerState && NextPlayer->PlayerState->UniqueId.IsValid() && NextPlayer->IsPrimaryPlayer())
+		if (NextPlayer && NextPlayer->PlayerState && NextPlayer->PlayerState->GetUniqueId().IsValid() && NextPlayer->IsPrimaryPlayer())
 		{
 			PrimaryController = NextPlayer;
 			break;
diff --git a/Engine/Source/Runtime/Engine/Private/GameMode.cpp b/Engine/Source/Runtime/Engine/Private/GameMode.cpp
--- a/Engine/Source/Runtime/Engine/Private/GameMode.cpp
+++ b/Engine/Source/Runtime/Engine/Private/GameMode.cpp
@@ -627,7 +627,7 @@ void AGameMode::AddInactivePlayer(APlayerState* PlayerState, APlayerController*
 	check(PlayerState)
 	UWorld* LocalWorld = GetWorld();
 	// don't store if it's an old PlayerState from the previous level or if it's a spectator... or if we are shutting down
-	if (!PlayerState->bFromPreviousLevel && !MustSpectate(PC) && !LocalWorld->bIsTearingDown)
+	if (!PlayerState->IsFromPreviousLevel() && !MustSpectate(PC) && !LocalWorld->bIsTearingDown)
 	{
 		APlayerState* const NewPlayerState = PlayerState->Duplicate();
 		if (NewPlayerState)
@@ -644,7 +644,7 @@ void AGameMode::AddInactivePlayer(APlayerState* PlayerState, APlayerController*
 			// On console, we have to check the unique net id as network address isn't valid
 			const bool bIsConsole = !PLATFORM_DESKTOP;
 			// Assume valid unique ids means comparison should be via this method
-			const bool bHasValidUniqueId = NewPlayerState->UniqueId.IsValid();
+			const bool bHasValidUniqueId = NewPlayerState->GetUniqueId().IsValid();
 			// Don't accidentally compare empty network addresses (already issue with two clients on same machine during development)
 			const bool bHasValidNetworkAddress = !NewPlayerState->SavedNetworkAddress.IsEmpty();
 			const bool bUseUniqueIdCheck = bIsConsole || bHasValidUniqueId;
@@ -660,7 +660,7 @@ void AGameMode::AddInactivePlayer(APlayerState* PlayerState, APlayerController*
 					Idx--;
 				}
 				else if ( (!bUseUniqueIdCheck && bHasValidNetworkAddress && (CurrentPlayerState->SavedNetworkAddress == NewPlayerState->SavedNetworkAddress))
-						|| (bUseUniqueIdCheck && (CurrentPlayerState->UniqueId == NewPlayerState->UniqueId)) )
+						|| (bUseUniqueIdCheck && (CurrentPlayerState->GetUniqueId() == NewPlayerState->GetUniqueId())) )
 				{
 					// destroy the playerstate, then remove it from the tracking
 					CurrentPlayerState->Destroy();
@@ -704,7 +704,7 @@ bool AGameMode::FindInactivePlayer(APlayerController* PC)
 	// On console, we have to check the unique net id as network address isn't valid
 	const bool bIsConsole = !PLATFORM_DESKTOP;
 	// Assume valid unique ids means comparison should be via this method
-	const bool bHasValidUniqueId = PC->PlayerState->UniqueId.IsValid();
+	const bool bHasValidUniqueId = PC->PlayerState->GetUniqueId().IsValid();
 	// Don't accidentally compare empty network addresses (already issue with two clients on same machine during development)
 	const bool bHasValidNetworkAddress = !PC->PlayerState->SavedNetworkAddress.IsEmpty();
 	const bool bUseUniqueIdCheck = bIsConsole || bHasValidUniqueId;
@@ -719,7 +719,7 @@ bool AGameMode::FindInactivePlayer(APlayerController* PC)
 			InactivePlayerArray.RemoveAt(i,1);
 			i--;
 		}
-		else if ((bUseUniqueIdCheck && (CurrentPlayerState->UniqueId == PC->PlayerState->UniqueId)) ||
+		else if ((bUseUniqueIdCheck && (CurrentPlayerState->GetUniqueId() == PC->PlayerState->GetUniqueId())) ||
 				 (!bUseUniqueIdCheck && bHasValidNetworkAddress && (FCString::Stricmp(*CurrentPlayerState->SavedNetworkAddress, *NewNetworkAddress) == 0) && (FCString::Stricmp(*CurrentPlayerState->GetPlayerName(), *NewName) == 0)))
 		{
 			// found it!
@@ -731,7 +731,7 @@ bool AGameMode::FindInactivePlayer(APlayerController* PC)
 			OverridePlayerState(PC, OldPlayerState);
 			GameState->AddPlayerState(PC->PlayerState);
 			InactivePlayerArray.RemoveAt(i, 1);
-			OldPlayerState->bIsInactive = true;
+			OldPlayerState->SetIsInactive(true);
 			// Set the uniqueId to nullptr so it will not kill the player's registration 
 			// in UnregisterPlayerWithSession()
 			OldPlayerState->SetUniqueId(nullptr);
diff --git a/Engine/Source/Runtime/Engine/Private/GameModeBase.cpp b/Engine/Source/Runtime/Engine/Private/GameModeBase.cpp
--- a/Engine/Source/Runtime/Engine/Private/GameModeBase.cpp
+++ b/Engine/Source/Runtime/Engine/Private/GameModeBase.cpp
@@ -537,7 +537,7 @@ TSubclassOf<APlayerController> AGameModeBase::GetPlayerControllerClassToSpawnFor
 {
 	UClass* PCClassToSpawn = PlayerControllerClass;
 
-	if (PreviousPC && ReplaySpectatorPlayerControllerClass && PreviousPC->PlayerState && PreviousPC->PlayerState->bOnlySpectator)
+	if (PreviousPC && ReplaySpectatorPlayerControllerClass && PreviousPC->PlayerState && PreviousPC->PlayerState->IsOnlyASpectator())
 	{
 		PCClassToSpawn = ReplaySpectatorPlayerControllerClass;
 	}
@@ -759,7 +759,7 @@ FString AGameModeBase::InitNewPlayer(APlayerController* NewPlayerController, con
 	FString InName = UGameplayStatics::ParseOption(Options, TEXT("Name")).Left(20);
 	if (InName.IsEmpty())
 	{
-		InName = FString::Printf(TEXT("%s%i"), *DefaultPlayerName.ToString(), NewPlayerController->PlayerState->PlayerId);
+		InName = FString::Printf(TEXT("%s%i"), *DefaultPlayerName.ToString(), NewPlayerController->PlayerState->GetPlayerId());
 	}
 
 	ChangeName(NewPlayerController, InName, false);
@@ -972,9 +972,10 @@ void AGameModeBase::PostLogin(APlayerController* NewPlayer)
 	else
 	{
 		// If NewPlayer is not only a spectator and has a valid ID, add him as a user to the replay.
-		if (NewPlayer->PlayerState->UniqueId.IsValid())
+		const FUniqueNetIdRepl& NewPlayerStateUniqueId = NewPlayer->PlayerState->GetUniqueId();
+		if (NewPlayerStateUniqueId.IsValid())
 		{
-			GetGameInstance()->AddUserToReplay(NewPlayer->PlayerState->UniqueId.ToString());
+			GetGameInstance()->AddUserToReplay(NewPlayerStateUniqueId.ToString());
 		}
 	}
 
@@ -1025,7 +1026,7 @@ bool AGameModeBase::MustSpectate_Implementation(APlayerController* NewPlayerCont
 		return false;
 	}
 
-	return NewPlayerController->PlayerState->bOnlySpectator;
+	return NewPlayerController->PlayerState->IsOnlyASpectator();
 }
 
 bool AGameModeBase::CanSpectate_Implementation(APlayerController* Viewer, APlayerState* ViewTarget)
@@ -1354,7 +1355,7 @@ bool AGameModeBase::SpawnPlayerFromSimulate(const FVector& NewLocation, const FR
 	APlayerController* PC = GetGameInstance()->GetFirstLocalPlayerController();
 	if (PC != nullptr)
 	{
-		PC->PlayerState->bOnlySpectator = false;
+		PC->PlayerState->SetIsOnlyASpectator(false);
 
 		bool bNeedsRestart = true;
 		if (PC->GetPawn() == nullptr)
diff --git a/Engine/Source/Runtime/Engine/Private/GameSession.cpp b/Engine/Source/Runtime/Engine/Private/GameSession.cpp
--- a/Engine/Source/Runtime/Engine/Private/GameSession.cpp
+++ b/Engine/Source/Runtime/Engine/Private/GameSession.cpp
@@ -32,10 +32,10 @@ APlayerController* GetPlayerControllerFromNetId(UWorld* World, const FUniqueNetI
 		{
 			APlayerController* PlayerController = Iterator->Get();
 			// Determine if this is a player with replication
-			if (PlayerController && PlayerController->PlayerState != NULL && PlayerController->PlayerState->UniqueId.IsValid())
+			if (PlayerController && PlayerController->PlayerState != NULL && PlayerController->PlayerState->GetUniqueId().IsValid())
 			{
 				// If the ids match, then this is the right player.
-				if (*PlayerController->PlayerState->UniqueId == PlayerNetId)
+				if (*PlayerController->PlayerState->GetUniqueId() == PlayerNetId)
 				{
 					return PlayerController;
 				}
@@ -233,7 +233,7 @@ void AGameSession::RegisterPlayer(APlayerController* NewPlayer, const TSharedPtr
 	{
 		// Set the player's ID.
 		check(NewPlayer->PlayerState);
-		NewPlayer->PlayerState->PlayerId = GetNextPlayerID();
+		NewPlayer->PlayerState->SetPlayerId(GetNextPlayerID());
 		NewPlayer->PlayerState->SetUniqueId(UniqueId);
 		NewPlayer->PlayerState->RegisterPlayerWithSession(bWasFromInvite);
 	}
@@ -267,9 +267,9 @@ void AGameSession::UnregisterPlayer(const APlayerController* ExitingPlayer)
 	if (GetNetMode() != NM_Standalone &&
 		ExitingPlayer != NULL &&
 		ExitingPlayer->PlayerState &&
-		ExitingPlayer->PlayerState->UniqueId.IsValid())
+		ExitingPlayer->PlayerState->GetUniqueId().IsValid())
 	{
-		UnregisterPlayer(ExitingPlayer->PlayerState->SessionName, ExitingPlayer->PlayerState->UniqueId);
+		UnregisterPlayer(ExitingPlayer->PlayerState->SessionName, ExitingPlayer->PlayerState->GetUniqueId());
 	}
 }
 
diff --git a/Engine/Source/Runtime/Engine/Private/GameStateBase.cpp b/Engine/Source/Runtime/Engine/Private/GameStateBase.cpp
--- a/Engine/Source/Runtime/Engine/Private/GameStateBase.cpp
+++ b/Engine/Source/Runtime/Engine/Private/GameStateBase.cpp
@@ -110,14 +110,14 @@ void AGameStateBase::SeamlessTravelTransitionCheckpoint(bool bToTransitionMap)
 	// mark all existing player states as from previous level for various bookkeeping
 	for (int32 i = 0; i < PlayerArray.Num(); i++)
 	{
-		PlayerArray[i]->bFromPreviousLevel = true;
+		PlayerArray[i]->SetIsFromPreviousLevel(true);
 	}
 }
 
 void AGameStateBase::AddPlayerState(APlayerState* PlayerState)
 {
 	// Determine whether it should go in the active or inactive list
-	if (!PlayerState->bIsInactive)
+	if (!PlayerState->IsInactive())
 	{
 		// make sure no duplicates
 		PlayerArray.AddUnique(PlayerState);
@@ -255,7 +255,7 @@ APlayerState* AGameStateBase::GetPlayerStateFromUniqueNetId(const FUniqueNetIdWr
 	const TArray<APlayerState*>& Players = PlayerArray;
 	for (APlayerState* Player : Players)
 	{
-		if (Player && Player->UniqueId == InPlayerId)
+		if (Player && Player->GetUniqueId() == InPlayerId)
 		{
 			return Player;
 		}
diff --git a/Engine/Source/Runtime/Engine/Private/Pawn.cpp b/Engine/Source/Runtime/Engine/Private/Pawn.cpp
--- a/Engine/Source/Runtime/Engine/Private/Pawn.cpp
+++ b/Engine/Source/Runtime/Engine/Private/Pawn.cpp
@@ -213,12 +213,12 @@ bool APawn::IsLocallyControlled() const
 }
 bool APawn::IsPlayerControlled() const
 {
-	return PlayerState && !PlayerState->bIsABot;
+	return PlayerState && !PlayerState->IsABot();
 }
 
 bool APawn::IsBotControlled() const
 {
-	return PlayerState && PlayerState->bIsABot;
+	return PlayerState && PlayerState->IsABot();
 }
 
 bool APawn::ReachedDesiredRotation()
diff --git a/Engine/Source/Runtime/Engine/Private/PlayerController.cpp b/Engine/Source/Runtime/Engine/Private/PlayerController.cpp
--- a/Engine/Source/Runtime/Engine/Private/PlayerController.cpp
+++ b/Engine/Source/Runtime/Engine/Private/PlayerController.cpp
@@ -811,7 +811,7 @@ void APlayerController::ClientRestart_Implementation(APawn* NewPawn)
 void APlayerController::OnPossess(APawn* PawnToPossess)
 {
 	if ( PawnToPossess != NULL && 
-		(PlayerState == NULL || !PlayerState->bOnlySpectator) )
+		(PlayerState == NULL || !PlayerState->IsOnlyASpectator()) )
 	{
 		const bool bNewPawn = (GetPawn() != PawnToPossess);
 
@@ -1221,7 +1221,7 @@ void APlayerController::Reset()
 	SetViewTarget(this);
 	ResetCameraMode();
 
-	bPlayerIsWaiting = !PlayerState->bOnlySpectator;
+	bPlayerIsWaiting = !PlayerState->IsOnlyASpectator();
 	ChangeState(NAME_Spectating);
 }
 
@@ -1232,7 +1232,7 @@ void APlayerController::ClientReset_Implementation()
 	ResetCameraMode();
 	SetViewTarget(this);
 
-	bPlayerIsWaiting = (PlayerState == nullptr) || !PlayerState->bOnlySpectator;
+	bPlayerIsWaiting = (PlayerState == nullptr) || !PlayerState->IsOnlyASpectator();
 	ChangeState(NAME_Spectating);
 }
 
@@ -3048,7 +3048,7 @@ void APlayerController::ServerRestartPlayer_Implementation()
 
 bool APlayerController::CanRestartPlayer()
 {
-	return PlayerState && !PlayerState->bOnlySpectator && HasClientLoadedCurrentWorld() && PendingSwapConnection == NULL;
+	return PlayerState && !PlayerState->IsOnlyASpectator() && HasClientLoadedCurrentWorld() && PendingSwapConnection == NULL;
 }
 
 /// @cond DOXYGEN_WARNINGS
@@ -4845,8 +4845,8 @@ bool APlayerController::DefaultCanUnpause()
 void APlayerController::StartSpectatingOnly()
 {
 	ChangeState(NAME_Spectating);
-	PlayerState->bIsSpectator = true;
-	PlayerState->bOnlySpectator = true;
+	PlayerState->SetIsSpectator(true);
+	PlayerState->SetIsOnlyASpectator(true);
 	bPlayerIsWaiting = false; // Can't spawn, we are only allowed to be a spectator.
 }
 
@@ -5030,11 +5030,11 @@ void APlayerController::EndSpectatingState()
 {
 	if ( PlayerState != NULL )
 	{
-		if ( PlayerState->bOnlySpectator )
+		if ( PlayerState->IsOnlyASpectator() )
 		{
 			UE_LOG(LogPlayerController, Warning, TEXT("Spectator only UPlayer* leaving spectating state"));
 		}
-		PlayerState->bIsSpectator = false;
+		PlayerState->SetIsSpectator(false);
 	}
 
 	bPlayerIsWaiting = false;
diff --git a/Engine/Source/Runtime/Engine/Private/PlayerMuteList.cpp b/Engine/Source/Runtime/Engine/Private/PlayerMuteList.cpp
--- a/Engine/Source/Runtime/Engine/Private/PlayerMuteList.cpp
+++ b/Engine/Source/Runtime/Engine/Private/PlayerMuteList.cpp
@@ -45,10 +45,10 @@ void FPlayerMuteList::ServerMutePlayer(APlayerController* OwningPC, const FUniqu
 	if (OtherPC != NULL)
 	{
 		// Update their packet filter list too
-		AddIdToMuteList(OtherPC->MuteList.VoicePacketFilter, OwningPC->PlayerState->UniqueId.GetUniqueNetId());
+		AddIdToMuteList(OtherPC->MuteList.VoicePacketFilter, OwningPC->PlayerState->GetUniqueId().GetUniqueNetId());
 
 		// Tell the other PC to mute this one
-		OtherPC->ClientMutePlayer(OwningPC->PlayerState->UniqueId);
+		OtherPC->ClientMutePlayer(OwningPC->PlayerState->GetUniqueId());
 	}
 }
 
@@ -63,7 +63,7 @@ void FPlayerMuteList::ServerUnmutePlayer(APlayerController* OwningPC, const FUni
 	APlayerController* OtherPC = OwningPC->GetPlayerControllerForMuting(UnmuteId);
 	if (OtherPC != NULL)
 	{
-		FUniqueNetIdRepl& OwningPlayerId = OwningPC->PlayerState->UniqueId;
+		const FUniqueNetIdRepl& OwningPlayerId = OwningPC->PlayerState->GetUniqueId();
 		auto PlayerIdToUnmutePred = [&PlayerIdToUnmute](TSharedRef<const FUniqueNetId> Other) { return PlayerIdToUnmute.IsValid() && *PlayerIdToUnmute == *Other; };
 		auto OwningPlayerIdPred = [&OwningPlayerId](TSharedRef<const FUniqueNetId> Other) { return *OwningPlayerId == *Other; };
 
@@ -155,7 +155,7 @@ void FPlayerMuteList::GameplayUnmutePlayer(APlayerController* OwningPC, const FU
 
 	if (OtherPC != nullptr)
 	{
-		FUniqueNetIdRepl& OwningPlayerId = OwningPC->PlayerState->UniqueId;
+		const FUniqueNetIdRepl& OwningPlayerId = OwningPC->PlayerState->GetUniqueId();
 		auto PlayerIdToUnmutePred = [&PlayerIdToUnmute](TSharedRef<const FUniqueNetId> Other) { return PlayerIdToUnmute.IsValid() && *PlayerIdToUnmute == *Other; };
 		auto OwningPlayerIdPred = [&OwningPlayerId](TSharedRef<const FUniqueNetId> Other) { return OwningPlayerId.IsValid() && *OwningPlayerId == *Other; };
 
diff --git a/Engine/Source/Runtime/Engine/Private/PlayerState.cpp b/Engine/Source/Runtime/Engine/Private/PlayerState.cpp
--- a/Engine/Source/Runtime/Engine/Private/PlayerState.cpp
+++ b/Engine/Source/Runtime/Engine/Private/PlayerState.cpp
@@ -12,6 +12,7 @@
 #include "Net/UnrealNetwork.h"
 #include "Net/OnlineEngineInterface.h"
 #include "GameFramework/GameStateBase.h"
+#include "Net/Core/PushModel/PushModel.h"
 
 APlayerState::APlayerState(const FObjectInitializer& ObjectInitializer)
 	: Super(ObjectInitializer
@@ -132,7 +133,7 @@ void APlayerState::RecalculateAvgPing()
 
 	if (bShouldUpdateReplicatedPing || !HasAuthority())
 	{
-		Ping = FMath::Min(255, (int32)(ExactPing * 0.25f));
+		SetPing(FMath::Min(255, (int32)(ExactPing * 0.25f)));
 	}
 }
 
@@ -150,22 +151,22 @@ void APlayerState::DispatchCopyProperties(APlayerState* PlayerState)
 
 void APlayerState::OverrideWith(APlayerState* PlayerState)
 {
-	bIsSpectator = PlayerState->bIsSpectator;
-	bOnlySpectator = PlayerState->bOnlySpectator;
-	SetUniqueId(PlayerState->UniqueId.GetUniqueNetId());
+	SetIsSpectator(PlayerState->IsSpectator());
+	SetIsOnlyASpectator(PlayerState->IsOnlyASpectator());
+	SetUniqueId(PlayerState->GetUniqueId().GetUniqueNetId());
 	SetPlayerNameInternal(PlayerState->GetPlayerName());
 }
 
 
 void APlayerState::CopyProperties(APlayerState* PlayerState)
 {
-	PlayerState->Score = Score;
-	PlayerState->Ping = Ping;
+	PlayerState->SetScore(GetScore());
+	PlayerState->SetPing(GetPing());
 	PlayerState->ExactPing = ExactPing;
-	PlayerState->PlayerId = PlayerId;
-	PlayerState->SetUniqueId(UniqueId.GetUniqueNetId());
+	PlayerState->SetPlayerId(GetPlayerId());
+	PlayerState->SetUniqueId(GetUniqueId().GetUniqueNetId());
 	PlayerState->SetPlayerNameInternal(GetPlayerName());
-	PlayerState->StartTime = StartTime;
+	PlayerState->SetStartTime(GetStartTime());
 	PlayerState->SavedNetworkAddress = SavedNetworkAddress;
 }
 
@@ -201,12 +202,12 @@ void APlayerState::PostInitializeComponents()
 	AController* OwningController = Cast<AController>(GetOwner());
 	if (OwningController != NULL)
 	{
-		bIsABot = (Cast<APlayerController>(OwningController) == nullptr);
+		SetIsABot(Cast<APlayerController>(OwningController) == nullptr);
 	}
 
 	if (GameStateBase)
 	{
-		StartTime = GameStateBase->GetPlayerStartTime(OwningController);
+		SetStartTime(GameStateBase->GetPlayerStartTime(OwningController));
 	}
 }
 
@@ -232,7 +233,7 @@ void APlayerState::OnRep_bIsInactive()
 
 bool APlayerState::ShouldBroadCastWelcomeMessage(bool bExiting)
 {
-	return (!bIsInactive && GetNetMode() != NM_Standalone);
+	return (!IsInactive() && GetNetMode() != NM_Standalone);
 }
 
 void APlayerState::Destroyed()
@@ -264,7 +265,7 @@ void APlayerState::Destroyed()
 void APlayerState::Reset()
 {
 	Super::Reset();
-	Score = 0;
+	SetScore(0);
 	ForceNetUpdate();
 }
 
@@ -276,20 +277,13 @@ FString APlayerState::GetHumanReadableName() const
 void APlayerState::OnRep_PlayerName()
 {
 	OldNamePrivate = GetPlayerName();
-	PRAGMA_DISABLE_DEPRECATION_WARNINGS;
-	OldName = OldNamePrivate;
-	PRAGMA_ENABLE_DEPRECATION_WARNINGS;
-
 	HandleWelcomeMessage();
 }
 
 void APlayerState::SetPlayerNameInternal(const FString& S)
 {
+	MARK_PROPERTY_DIRTY_FROM_NAME(APlayerState, PlayerNamePrivate, this);
 	PlayerNamePrivate = S;
-
-	PRAGMA_DISABLE_DEPRECATION_WARNINGS;
-	PlayerName = PlayerNamePrivate;
-	PRAGMA_ENABLE_DEPRECATION_WARNINGS;
 }
 
 void APlayerState::SetPlayerName(const FString& S)
@@ -304,10 +298,6 @@ void APlayerState::SetPlayerName(const FString& S)
 	}
 
 	OldNamePrivate = GetPlayerName();
-	PRAGMA_DISABLE_DEPRECATION_WARNINGS;
-	OldName = OldNamePrivate;
-	PRAGMA_ENABLE_DEPRECATION_WARNINGS;
-
 	ForceNetUpdate();
 }
 
@@ -329,10 +319,6 @@ FString APlayerState::GetOldPlayerName() const
 void APlayerState::SetOldPlayerName(const FString& S)
 {
 	OldNamePrivate = S;
-
-	PRAGMA_DISABLE_DEPRECATION_WARNINGS
-	OldName = S;
-	PRAGMA_ENABLE_DEPRECATION_WARNINGS
 }
 
 void APlayerState::HandleWelcomeMessage()
@@ -361,7 +347,7 @@ void APlayerState::HandleWelcomeMessage()
 	}
 	else
 	{
-		int32 WelcomeMessageNum = bOnlySpectator ? 16 : 1;
+		int32 WelcomeMessageNum = IsOnlyASpectator() ? 16 : 1;
 		bHasBeenWelcomed = true;
 
 		if (ShouldBroadCastWelcomeMessage())
@@ -388,32 +374,35 @@ void APlayerState::OnRep_UniqueId()
 	RegisterPlayerWithSession(false);
 }
 
+PRAGMA_DISABLE_DEPRECATION_WARNINGS
 void APlayerState::SetUniqueId(const TSharedPtr<const FUniqueNetId>& InUniqueId)
 {
+	MARK_PROPERTY_DIRTY_FROM_NAME(APlayerState, UniqueId, this);
 	UniqueId.SetUniqueNetId(InUniqueId);
 }
+PRAGMA_ENABLE_DEPRECATION_WARNINGS
 
 void APlayerState::RegisterPlayerWithSession(bool bWasFromInvite)
 {
 	if (GetNetMode() != NM_Standalone)
 	{
-		if (UniqueId.IsValid()) // May not be valid if this is was created via DebugCreatePlayer
+		if (GetUniqueId().IsValid()) // May not be valid if this is was created via DebugCreatePlayer
 		{
 			// Register the player as part of the session
 			const APlayerState* PlayerState = GetDefault<APlayerState>();
-			UOnlineEngineInterface::Get()->RegisterPlayer(GetWorld(), PlayerState->SessionName, *UniqueId, bWasFromInvite);
+			UOnlineEngineInterface::Get()->RegisterPlayer(GetWorld(), PlayerState->SessionName, *GetUniqueId(), bWasFromInvite);
 		}
 	}
 }
 
 void APlayerState::UnregisterPlayerWithSession()
 {
-	if (GetNetMode() == NM_Client && UniqueId.IsValid())
+	if (GetNetMode() == NM_Client && GetUniqueId().IsValid())
 	{
 		const APlayerState* PlayerState = GetDefault<APlayerState>();
 		if (PlayerState->SessionName != NAME_None)
 		{
-			UOnlineEngineInterface::Get()->UnregisterPlayer(GetWorld(), PlayerState->SessionName, *UniqueId);
+			UOnlineEngineInterface::Get()->UnregisterPlayer(GetWorld(), PlayerState->SessionName, *GetUniqueId());
 		}
 	}
 }
@@ -436,7 +425,7 @@ APlayerState* APlayerState::Duplicate()
 void APlayerState::SeamlessTravelTo(APlayerState* NewPlayerState)
 {
 	DispatchCopyProperties(NewPlayerState);
-	NewPlayerState->bOnlySpectator = bOnlySpectator;
+	NewPlayerState->SetIsOnlyASpectator(IsOnlyASpectator());
 }
 
 
@@ -447,23 +436,93 @@ bool APlayerState::IsPrimaryPlayer() const
 
 void APlayerState::GetLifetimeReplicatedProps(TArray< FLifetimeProperty > & OutLifetimeProps) const
 {
-	Super::GetLifetimeReplicatedProps( OutLifetimeProps );
+	Super::GetLifetimeReplicatedProps(OutLifetimeProps);
+
+	FDoRepLifetimeParams SharedParams;
+	SharedParams.bIsPushBased = true;
+
+	DOREPLIFETIME_WITH_PARAMS_FAST(APlayerState, Score, SharedParams);
+	DOREPLIFETIME_WITH_PARAMS_FAST(APlayerState, bIsSpectator, SharedParams);
+	DOREPLIFETIME_WITH_PARAMS_FAST(APlayerState, bOnlySpectator, SharedParams);
+	DOREPLIFETIME_WITH_PARAMS_FAST(APlayerState, bFromPreviousLevel, SharedParams);
+	DOREPLIFETIME_WITH_PARAMS_FAST(APlayerState, StartTime, SharedParams);
+	DOREPLIFETIME_WITH_PARAMS_FAST(APlayerState, PlayerNamePrivate, SharedParams);
+
+	SharedParams.Condition = COND_SkipOwner;
+	DOREPLIFETIME_WITH_PARAMS_FAST(APlayerState, Ping, SharedParams);
+
+	SharedParams.Condition = COND_InitialOnly;
+	DOREPLIFETIME_WITH_PARAMS_FAST(APlayerState, PlayerId, SharedParams);
+	DOREPLIFETIME_WITH_PARAMS_FAST(APlayerState, bIsABot, SharedParams);
+	DOREPLIFETIME_WITH_PARAMS_FAST(APlayerState, bIsInactive, SharedParams);
+	DOREPLIFETIME_WITH_PARAMS_FAST(APlayerState, UniqueId, SharedParams);
+}
 
-	DOREPLIFETIME( APlayerState, Score );
 
-	DOREPLIFETIME( APlayerState, bIsSpectator );
-	DOREPLIFETIME( APlayerState, bOnlySpectator );
-	DOREPLIFETIME( APlayerState, bFromPreviousLevel );
 PRAGMA_DISABLE_DEPRECATION_WARNINGS
-	DOREPLIFETIME( APlayerState, StartTime );
-PRAGMA_ENABLE_DEPRECATION_WARNINGS
+void APlayerState::SetScore(const float NewScore)
+{
+	MARK_PROPERTY_DIRTY_FROM_NAME(APlayerState, Score, this);
+	Score = NewScore;
+}
+
+void APlayerState::SetPlayerId(const int32 NewId)
+{
+	MARK_PROPERTY_DIRTY_FROM_NAME(APlayerState, PlayerId, this);
+	PlayerId = NewId;
+}
+
+void APlayerState::SetPing(const uint8 NewPing)
+{
+	MARK_PROPERTY_DIRTY_FROM_NAME(APlayerState, Ping, this);
+	Ping = NewPing;
+}
 
-	DOREPLIFETIME_CONDITION( APlayerState, Ping,		COND_SkipOwner );
+void APlayerState::SetIsSpectator(const bool bNewSpectator)
+{
+	MARK_PROPERTY_DIRTY_FROM_NAME(APlayerState, bIsSpectator, this);
+	bIsSpectator = bNewSpectator;
+}
+
+void APlayerState::SetIsOnlyASpectator(const bool bNewSpectator)
+{
+	MARK_PROPERTY_DIRTY_FROM_NAME(APlayerState, bOnlySpectator, this);
+	bOnlySpectator = bNewSpectator;
+}
+
+void APlayerState::SetIsABot(const bool bNewIsABot)
+{
+	MARK_PROPERTY_DIRTY_FROM_NAME(APlayerState, bIsABot, this);
+	bIsABot = bNewIsABot;
+}
+
+void APlayerState::SetIsInactive(const bool bNewInactive)
+{
+	MARK_PROPERTY_DIRTY_FROM_NAME(APlayerState, bIsInactive, this);
+	bIsInactive = bNewInactive;
+}
+
+void APlayerState::SetIsFromPreviousLevel(const bool bNewFromPreviousLevel)
+{
+	MARK_PROPERTY_DIRTY_FROM_NAME(APlayerState, bFromPreviousLevel, this);
+	bFromPreviousLevel = bNewFromPreviousLevel;
+}
+
+void APlayerState::SetStartTime(const int32 NewStartTime)
+{
+	MARK_PROPERTY_DIRTY_FROM_NAME(APlayerState, StartTime, this);
+	StartTime = NewStartTime;
+}
 
-	DOREPLIFETIME_CONDITION( APlayerState, PlayerId,	COND_InitialOnly );
-	DOREPLIFETIME_CONDITION( APlayerState, bIsABot,		COND_InitialOnly );
-	DOREPLIFETIME_CONDITION( APlayerState, bIsInactive, COND_InitialOnly );
-	DOREPLIFETIME_CONDITION( APlayerState, UniqueId,	COND_InitialOnly );
+void APlayerState::SetUniqueId(const FUniqueNetIdRepl& NewUniqueId)
+{
+	MARK_PROPERTY_DIRTY_FROM_NAME(APlayerState, UniqueId, this);
+	UniqueId = NewUniqueId;
+}
 
-	DOREPLIFETIME(APlayerState, PlayerNamePrivate);
+void APlayerState::SetUniqueId(FUniqueNetIdRepl&& NewUniqueId)
+{
+	MARK_PROPERTY_DIRTY_FROM_NAME(APlayerState, UniqueId, this);
+	UniqueId = MoveTemp(NewUniqueId);
 }
+PRAGMA_ENABLE_DEPRECATION_WARNINGS
\ No newline at end of file
diff --git a/Engine/Source/Runtime/Engine/Private/VoiceConfig.cpp b/Engine/Source/Runtime/Engine/Private/VoiceConfig.cpp
--- a/Engine/Source/Runtime/Engine/Private/VoiceConfig.cpp
+++ b/Engine/Source/Runtime/Engine/Private/VoiceConfig.cpp
@@ -281,9 +281,9 @@ bool UVOIPStatics::IsVOIPTalkerStillAlive(UVOIPTalker* InTalker)
 
 void UVOIPStatics::ResetPlayerVoiceTalker(APlayerState* InPlayerState)
 {
-	if (InPlayerState && InPlayerState->UniqueId.IsValid())
+	if (InPlayerState && InPlayerState->GetUniqueId().IsValid())
 	{
-		VoiceTalkerMap.Remove(InPlayerState->UniqueId);
+		VoiceTalkerMap.Remove(InPlayerState->GetUniqueId());
 	}
 }
 
@@ -331,10 +331,10 @@ void UVOIPTalker::RegisterWithPlayerState(APlayerState* OwningState)
 		UnregisterFromVoiceTalkerMap();
 	}
 
-	if (OwningState->UniqueId.IsValid())
+	if (OwningState->GetUniqueId().IsValid())
 	{
-		UVOIPStatics::SetVOIPTalkerForPlayer(OwningState->UniqueId, this);
-		PlayerId = OwningState->UniqueId;
+		UVOIPStatics::SetVOIPTalkerForPlayer(OwningState->GetUniqueId(), this);
+		PlayerId = OwningState->GetUniqueId();
 		bIsRegistered = true;
 	}
 }
diff --git a/Engine/Source/Runtime/Engine/Private/World.cpp b/Engine/Source/Runtime/Engine/Private/World.cpp
--- a/Engine/Source/Runtime/Engine/Private/World.cpp
+++ b/Engine/Source/Runtime/Engine/Private/World.cpp
@@ -5241,7 +5241,7 @@ void UWorld::NotifyControlMessage(UNetConnection* Connection, uint8 MessageType,
 							// Successfully spawned in game.
 							UE_LOG(LogNet, Log, TEXT("JOINSPLIT: Succeeded: %s PlayerId: %s"),
 								*ChildConn->PlayerController->PlayerState->GetPlayerName(),
-								*ChildConn->PlayerController->PlayerState->UniqueId.ToDebugString());
+								*ChildConn->PlayerController->PlayerState->GetUniqueId().ToDebugString());
 						}
 					}
 				}
