Fixing audio device ID, which fixes occlusion and a number of other systems.

diff --git a/Engine/Source/Editor/UnrealEd/Private/Commandlets/AudioMixerCommandlet.cpp b/Engine/Source/Editor/UnrealEd/Private/Commandlets/AudioMixerCommandlet.cpp
--- a/Engine/Source/Editor/UnrealEd/Private/Commandlets/AudioMixerCommandlet.cpp
+++ b/Engine/Source/Editor/UnrealEd/Private/Commandlets/AudioMixerCommandlet.cpp
@@ -278,7 +278,7 @@ public:
 			FAudioQualitySettings QualitySettings = AudioDevice->GetQualityLevelSettings();
 
 			// Initialize the audio device
-			bSuccess = AudioDevice->Init(QualitySettings.MaxChannels);
+			bSuccess = AudioDevice->Init(AudioDevice.GetDeviceID(), QualitySettings.MaxChannels);
 
 			if (bSuccess)
 			{
@@ -348,7 +348,7 @@ public:
 			FAudioQualitySettings QualitySettings = AudioDevice->GetQualityLevelSettings();
 
 			// Initialize the audio device
-			bSuccess = AudioDevice->Init(QualitySettings.MaxChannels);
+			bSuccess = AudioDevice->Init(AudioDevice.GetDeviceID(), QualitySettings.MaxChannels);
 
 			// Wait a few seconds to give the editor a chance to load everything... you get hitches in the beginning otherwise
 			FPlatformProcess::Sleep(1);
@@ -432,7 +432,7 @@ public:
 			FAudioQualitySettings QualitySettings = AudioDevice->GetQualityLevelSettings();
 
 			// Initialize the audio device
-			bSuccess = AudioDevice->Init(QualitySettings.MaxChannels);
+			bSuccess = AudioDevice->Init(AudioDevice.GetDeviceID(), QualitySettings.MaxChannels);
 
 			// Wait a few seconds to give the editor a chance to load everything... you get hitches in the beginning otherwise
 			FPlatformProcess::Sleep(1);
@@ -513,7 +513,7 @@ public:
 			FAudioQualitySettings QualitySettings = AudioDevice->GetQualityLevelSettings();
 
 			// Initialize the audio device
-			bSuccess = AudioDevice->Init(QualitySettings.MaxChannels);
+			bSuccess = AudioDevice->Init(AudioDevice.GetDeviceID(), QualitySettings.MaxChannels);
 
 			// Wait a few seconds to give the editor a chance to load everything... you get hitches in the beginning otherwise
 			FPlatformProcess::Sleep(1);
diff --git a/Engine/Source/Runtime/Engine/Private/AudioDevice.cpp b/Engine/Source/Runtime/Engine/Private/AudioDevice.cpp
--- a/Engine/Source/Runtime/Engine/Private/AudioDevice.cpp
+++ b/Engine/Source/Runtime/Engine/Private/AudioDevice.cpp
@@ -322,7 +322,7 @@ const FAudioQualitySettings& FAudioDevice::GetQualityLevelSettings()
 	return AudioSettings->GetQualityLevelSettings(QualityLevel);
 }
 
-bool FAudioDevice::Init(int32 InMaxSources)
+bool FAudioDevice::Init(Audio::FDeviceId InDeviceID, int32 InMaxSources)
 {
 	SCOPED_BOOT_TIMING("FAudioDevice::Init");
 
@@ -333,6 +333,13 @@ bool FAudioDevice::Init(int32 InMaxSources)
 		return true;
 	}
 
+	if (InDeviceID == INDEX_NONE)
+	{
+		return false;
+	}
+
+	DeviceID = InDeviceID;
+
 	bool bDeferStartupPrecache = false;
 
 	PluginListeners.Reset();
diff --git a/Engine/Source/Runtime/Engine/Private/AudioDeviceManager.cpp b/Engine/Source/Runtime/Engine/Private/AudioDeviceManager.cpp
--- a/Engine/Source/Runtime/Engine/Private/AudioDeviceManager.cpp
+++ b/Engine/Source/Runtime/Engine/Private/AudioDeviceManager.cpp
@@ -243,11 +243,8 @@ void FAudioDeviceManager::ToggleAudioMixer()
 
 				check(AudioDevice);
 
-				// Set the new audio device handle to the old audio device handle
-				AudioDevice->DeviceID = DeviceID;
-
 				// Re-init the new audio device using appropriate settings so it behaves the same
-				if (AudioDevice->Init(AudioSettings->GetHighestMaxChannels()))
+				if (AudioDevice->Init(DeviceID, AudioSettings->GetHighestMaxChannels()))
 				{
 					AudioDevice->SetMaxChannels(QualityLevelMaxChannels);
 				}
@@ -465,7 +462,7 @@ bool FAudioDeviceManager::LoadDefaultAudioDeviceModule()
 FAudioDeviceHandle FAudioDeviceManager::CreateNewDevice(const FAudioDeviceParams& InParams)
 {
 	Audio::FDeviceId DeviceID = GetNewDeviceID();
-	Devices.Emplace(DeviceID, FAudioDeviceContainer(InParams, this));
+	Devices.Emplace(DeviceID, FAudioDeviceContainer(InParams, DeviceID, this));
 	FAudioDeviceContainer* ContainerPtr = Devices.Find(DeviceID);
 	check(ContainerPtr);
 	if (!ContainerPtr->Device)
@@ -1291,7 +1288,7 @@ FAudioDeviceHandle& FAudioDeviceHandle::operator=(FAudioDeviceHandle&& Other)
 	return *this;
 }
 
-FAudioDeviceManager::FAudioDeviceContainer::FAudioDeviceContainer(const FAudioDeviceParams& InParams, FAudioDeviceManager* DeviceManager)
+FAudioDeviceManager::FAudioDeviceContainer::FAudioDeviceContainer(const FAudioDeviceParams& InParams, Audio::FDeviceId InDeviceID, FAudioDeviceManager* DeviceManager)
 	: NumberOfHandlesToThisDevice(0)
 	, Scope(InParams.Scope)
 	, bIsNonRealtime(InParams.bIsNonRealtime)
@@ -1326,7 +1323,7 @@ FAudioDeviceManager::FAudioDeviceContainer::FAudioDeviceContainer(const FAudioDe
 	// runtime is supported.
 	const UAudioSettings* AudioSettings = GetDefault<UAudioSettings>();
 	const int32 HighestMaxChannels = AudioSettings ? AudioSettings->GetHighestMaxChannels() : 0;
-	if (Device->Init(HighestMaxChannels))
+	if (Device->Init(InDeviceID, HighestMaxChannels))
 	{
 		const FAudioQualitySettings& QualitySettings = Device->GetQualityLevelSettings();
 		Device->SetMaxChannels(QualitySettings.MaxChannels);
diff --git a/Engine/Source/Runtime/Engine/Public/AudioDevice.h b/Engine/Source/Runtime/Engine/Public/AudioDevice.h
--- a/Engine/Source/Runtime/Engine/Public/AudioDevice.h
+++ b/Engine/Source/Runtime/Engine/Public/AudioDevice.h
@@ -503,7 +503,7 @@ public:
 	/**
 	 * Basic initialization of the platform agnostic layer of the audio system
 	 */
-	bool Init(int32 InMaxSources);
+	bool Init(Audio::FDeviceId InDeviceID, int32 InMaxSources);
 
 	/**
 	 * Tears down the audio device
diff --git a/Engine/Source/Runtime/Engine/Public/AudioDeviceManager.h b/Engine/Source/Runtime/Engine/Public/AudioDeviceManager.h
--- a/Engine/Source/Runtime/Engine/Public/AudioDeviceManager.h
+++ b/Engine/Source/Runtime/Engine/Public/AudioDeviceManager.h
@@ -459,7 +459,7 @@ private:
 		/** Module this was created with. If nullptr, this device was created with the default module. */
 		IAudioDeviceModule* SpecifiedModule;
 
-		FAudioDeviceContainer(const FAudioDeviceParams& InParams, FAudioDeviceManager* DeviceManager);
+		FAudioDeviceContainer(const FAudioDeviceParams& InParams, Audio::FDeviceId InDeviceID, FAudioDeviceManager* DeviceManager);
 		~FAudioDeviceContainer();
 
 		FAudioDeviceContainer(const FAudioDeviceContainer& Other)
