Fix for nondeterministic cooking issues with USkeletons

diff --git a/Engine/Source/Runtime/Engine/Classes/Animation/SmartName.h b/Engine/Source/Runtime/Engine/Classes/Animation/SmartName.h
--- a/Engine/Source/Runtime/Engine/Classes/Animation/SmartName.h
+++ b/Engine/Source/Runtime/Engine/Classes/Animation/SmartName.h
@@ -170,11 +170,6 @@ private:
 #endif
 
 	TMap<FName, FCurveMetaData> CurveMetaDataMap;
-
-#if WITH_EDITORONLY_DATA
-	// Editor copy of the data we loaded, used to preserve determinism during cooking
-	TMap<FName, FCurveMetaData> LoadDataCurveMetaDataMap;
-#endif
 };
 
 USTRUCT()
@@ -201,6 +196,11 @@ protected:
 
 private:
 	TMap<FName, FSmartNameMapping> NameMappings;	// List of smartname mappings
+
+#if WITH_EDITORONLY_DATA
+	// Editor copy of the data we loaded, used to preserve determinism during cooking
+	TMap<FName, FSmartNameMapping> LoadedNameMappings;
+#endif
 };
 
 USTRUCT()
diff --git a/Engine/Source/Runtime/Engine/Private/Animation/SmartName.cpp b/Engine/Source/Runtime/Engine/Private/Animation/SmartName.cpp
--- a/Engine/Source/Runtime/Engine/Private/Animation/SmartName.cpp
+++ b/Engine/Source/Runtime/Engine/Private/Animation/SmartName.cpp
@@ -111,16 +111,7 @@ void FSmartNameMapping::Serialize(FArchive& Ar)
 
 	if (Ar.CustomVer(FFrameworkObjectVersion::GUID) >= FFrameworkObjectVersion::MoveCurveTypesToSkeleton)
 	{
-#if WITH_EDITORONLY_DATA
-		if (Ar.IsCooking())
-		{
-			Ar << LoadDataCurveMetaDataMap; // Cook out loaded data for determinism
-		}
-		else
-#endif
-		{
-			Ar << CurveMetaDataMap;
-		}
+		Ar << CurveMetaDataMap;
 	}
 
 	if (Ar.IsLoading())
@@ -129,10 +120,6 @@ void FSmartNameMapping::Serialize(FArchive& Ar)
 #if !WITH_EDITOR
 		CurveMetaDataMap.GenerateValueArray(CurveMetaDataList);
 #endif
-
-#if WITH_EDITORONLY_DATA
-		LoadDataCurveMetaDataMap = CurveMetaDataMap;
-#endif
 	}
 }
 
@@ -306,7 +293,24 @@ const FSmartNameMapping* FSmartNameContainer::GetContainer(FName ContainerName)
 
 void FSmartNameContainer::Serialize(FArchive& Ar)
 {
-	Ar << NameMappings;
+#if WITH_EDITORONLY_DATA
+	if (Ar.IsCooking())
+	{
+		Ar << LoadedNameMappings;
+	}
+	else
+#endif
+	{
+		Ar << NameMappings;
+	}
+
+#if WITH_EDITORONLY_DATA
+	if (Ar.IsLoading())
+	{
+		//To preserve 
+		LoadedNameMappings = NameMappings;
+	}
+#endif
 }
 
 FSmartNameMapping* FSmartNameContainer::GetContainerInternal(const FName& ContainerName)
