- Some StaticAnalysis fixes

diff --git a/Engine/Source/Runtime/CoreUObject/Public/UObject/FastReferenceCollector.h b/Engine/Source/Runtime/CoreUObject/Public/UObject/FastReferenceCollector.h
--- a/Engine/Source/Runtime/CoreUObject/Public/UObject/FastReferenceCollector.h
+++ b/Engine/Source/Runtime/CoreUObject/Public/UObject/FastReferenceCollector.h
@@ -614,6 +614,7 @@ private:
 				uint32 StartCycles = FPlatformTime::Cycles();
 #endif
 				CurrentObject = ObjectsToSerialize[CurrentIndex++];
+				checkSlow(CurrentObject);
 
 				// GetData() used to avoiding bounds checking (min and max)
 				// FMath::Min used to avoid out of bounds (without branching) on last iteration. Though anything can be passed into PrefetchBlock, 
diff --git a/Engine/Source/Runtime/Engine/Private/ShaderCompiler/ShaderCompiler.cpp b/Engine/Source/Runtime/Engine/Private/ShaderCompiler/ShaderCompiler.cpp
--- a/Engine/Source/Runtime/Engine/Private/ShaderCompiler/ShaderCompiler.cpp
+++ b/Engine/Source/Runtime/Engine/Private/ShaderCompiler/ShaderCompiler.cpp
@@ -1530,7 +1530,7 @@ void FShaderCompileUtilities::ExecuteShaderCompileJob(FShaderCommonCompileJob& J
 			{
 				UE_LOG(LogShaderCompilers, Fatal, TEXT("Can't nest Shader Pipelines inside Shader Pipeline '%s'!"), PipelineJob->ShaderPipeline->GetName());
 			}
-			if (Platform != SingleStage->Input.Target.Platform)
+			else if (Platform != SingleStage->Input.Target.Platform)
 			{
 				UE_LOG(LogShaderCompilers, Fatal, TEXT("Mismatched Target Platform %s while compiling Shader Pipeline '%s'."), *Format.GetPlainNameString(), PipelineJob->ShaderPipeline->GetName());
 			}
