Fix the specular highlight in lighting only mode that was accidently introduced by CL9489741

diff --git a/Engine/Plugins/Runtime/nDisplay/Source/DisplayCluster/Private/Game/Classes/Basics/DisplayClusterViewportClient.cpp b/Engine/Plugins/Runtime/nDisplay/Source/DisplayCluster/Private/Game/Classes/Basics/DisplayClusterViewportClient.cpp
--- a/Engine/Plugins/Runtime/nDisplay/Source/DisplayCluster/Private/Game/Classes/Basics/DisplayClusterViewportClient.cpp
+++ b/Engine/Plugins/Runtime/nDisplay/Source/DisplayCluster/Private/Game/Classes/Basics/DisplayClusterViewportClient.cpp
@@ -277,6 +277,11 @@ void UDisplayClusterViewportClient::Draw(FViewport* InViewport, FCanvas* SceneCa
 							View->DiffuseOverrideParameter = FVector4(GEngine->LightingOnlyBrightness.R, GEngine->LightingOnlyBrightness.G, GEngine->LightingOnlyBrightness.B, 0.0f);
 							View->SpecularOverrideParameter = FVector4(.1f, .1f, .1f, 0.0f);
 						}
+						else if (View->Family->EngineShowFlags.LightingOnlyOverride)
+						{
+							View->DiffuseOverrideParameter = FVector4(GEngine->LightingOnlyBrightness.R, GEngine->LightingOnlyBrightness.G, GEngine->LightingOnlyBrightness.B, 0.0f);
+							View->SpecularOverrideParameter = FVector4(0.f, 0.f, 0.f, 0.f);
+						}
 						else if (View->Family->EngineShowFlags.ReflectionOverride)
 						{
 							View->DiffuseOverrideParameter = FVector4(0.f, 0.f, 0.f, 0.f);
diff --git a/Engine/Source/Editor/UnrealEd/Private/EditorViewportClient.cpp b/Engine/Source/Editor/UnrealEd/Private/EditorViewportClient.cpp
--- a/Engine/Source/Editor/UnrealEd/Private/EditorViewportClient.cpp
+++ b/Engine/Source/Editor/UnrealEd/Private/EditorViewportClient.cpp
@@ -3613,6 +3613,11 @@ void FEditorViewportClient::SetupViewForRendering(FSceneViewFamily& ViewFamily,
 		View.DiffuseOverrideParameter = FVector4(GEngine->LightingOnlyBrightness.R, GEngine->LightingOnlyBrightness.G, GEngine->LightingOnlyBrightness.B, 0.0f);
 		View.SpecularOverrideParameter = FVector4(.1f, .1f, .1f, 0.0f);
 	}
+	else if (ViewFamily.EngineShowFlags.LightingOnlyOverride)
+	{
+		View.DiffuseOverrideParameter = FVector4(GEngine->LightingOnlyBrightness.R, GEngine->LightingOnlyBrightness.G, GEngine->LightingOnlyBrightness.B, 0.0f);
+		View.SpecularOverrideParameter = FVector4(0.f, 0.f, 0.f, 0.f);
+	}
 	else if (ViewFamily.EngineShowFlags.ReflectionOverride)
 	{
 		View.DiffuseOverrideParameter = FVector4(0.f, 0.f, 0.f, 0.f);
diff --git a/Engine/Source/Runtime/Engine/Private/GameViewportClient.cpp b/Engine/Source/Runtime/Engine/Private/GameViewportClient.cpp
--- a/Engine/Source/Runtime/Engine/Private/GameViewportClient.cpp
+++ b/Engine/Source/Runtime/Engine/Private/GameViewportClient.cpp
@@ -1268,6 +1268,11 @@ void UGameViewportClient::Draw(FViewport* InViewport, FCanvas* SceneCanvas)
 						View->DiffuseOverrideParameter = FVector4(GEngine->LightingOnlyBrightness.R, GEngine->LightingOnlyBrightness.G, GEngine->LightingOnlyBrightness.B, 0.0f);
 						View->SpecularOverrideParameter = FVector4(.1f, .1f, .1f, 0.0f);
 					}
+					else if (View->Family->EngineShowFlags.LightingOnlyOverride)
+					{
+						View->DiffuseOverrideParameter = FVector4(GEngine->LightingOnlyBrightness.R, GEngine->LightingOnlyBrightness.G, GEngine->LightingOnlyBrightness.B, 0.0f);
+						View->SpecularOverrideParameter = FVector4(0.f, 0.f, 0.f, 0.f);
+					}
 					else if (View->Family->EngineShowFlags.ReflectionOverride)
 					{
 						View->DiffuseOverrideParameter = FVector4(0.f, 0.f, 0.f, 0.f);
diff --git a/Engine/Source/Runtime/Engine/Private/ShowFlags.cpp b/Engine/Source/Runtime/Engine/Private/ShowFlags.cpp
--- a/Engine/Source/Runtime/Engine/Private/ShowFlags.cpp
+++ b/Engine/Source/Runtime/Engine/Private/ShowFlags.cpp
@@ -307,7 +307,8 @@ void ApplyViewMode(EViewModeIndex ViewModeIndex, bool bPerspective, FEngineShowF
 	// Assigning the new state like this ensures we always set the same variables (they depend on the view mode)
 	// This is affecting the state of showflags - if the state can be changed by the user as well it should better be done in EngineShowFlagOverride
 
-	EngineShowFlags.SetOverrideDiffuseAndSpecular(ViewModeIndex == VMI_Lit_DetailLighting || ViewModeIndex == VMI_LightingOnly);
+	EngineShowFlags.SetOverrideDiffuseAndSpecular(ViewModeIndex == VMI_Lit_DetailLighting);
+	EngineShowFlags.SetLightingOnlyOverride(ViewModeIndex == VMI_LightingOnly);
 	EngineShowFlags.SetReflectionOverride(ViewModeIndex == VMI_ReflectionOverride);
 	EngineShowFlags.SetVisualizeBuffer(ViewModeIndex == VMI_VisualizeBuffer);
 	EngineShowFlags.SetVisualizeLightCulling(ViewModeIndex == VMI_LightComplexity);
@@ -641,6 +642,10 @@ EViewModeIndex FindViewMode(const FEngineShowFlags& EngineShowFlags)
 	{
 		return VMI_Lit_DetailLighting;
 	}
+	else if (EngineShowFlags.LightingOnlyOverride)
+	{
+		return VMI_LightingOnly;
+	}
 	else if (EngineShowFlags.ReflectionOverride)
 	{
 		return VMI_ReflectionOverride;
diff --git a/Engine/Source/Runtime/Engine/Public/ShowFlags.h b/Engine/Source/Runtime/Engine/Public/ShowFlags.h
--- a/Engine/Source/Runtime/Engine/Public/ShowFlags.h
+++ b/Engine/Source/Runtime/Engine/Public/ShowFlags.h
@@ -313,6 +313,7 @@ private:
 		SetVisualizeHDR(false);
 		SetVisualizeShadingModels(false);
 		SetOverrideDiffuseAndSpecular(false);
+		SetLightingOnlyOverride(false);
 		SetReflectionOverride(false);
 		SetVisualizeBuffer(false);
 		SetVectorFields(false);
diff --git a/Engine/Source/Runtime/Engine/Public/ShowFlagsValues.inl b/Engine/Source/Runtime/Engine/Public/ShowFlagsValues.inl
--- a/Engine/Source/Runtime/Engine/Public/ShowFlagsValues.inl
+++ b/Engine/Source/Runtime/Engine/Public/ShowFlagsValues.inl
@@ -52,6 +52,8 @@ SHOWFLAG_FIXED_IN_SHIPPING(1, CameraImperfections, SFG_PostProcess, NSLOCTEXT("U
 SHOWFLAG_ALWAYS_ACCESSIBLE(OnScreenDebug, SFG_Developer, NSLOCTEXT("UnrealEd", "OnScreenDebugSF", "On Screen Debug"))
 /** needed for VMI_Lit_DetailLighting, Whether to override material diffuse and specular with constants, used by the Detail Lighting viewmode. */
 SHOWFLAG_FIXED_IN_SHIPPING(0, OverrideDiffuseAndSpecular, SFG_Hidden, NSLOCTEXT("UnrealEd", "OverrideDiffuseAndSpecularSF", "Override Diffuse And Specular"))
+/** needed for VMI_LightingOnly, Whether to override material diffuse with constants, used by the Lighting Only viewmode. */
+SHOWFLAG_FIXED_IN_SHIPPING(0, LightingOnlyOverride, SFG_Hidden, NSLOCTEXT("UnrealEd", "LightingOnlyOverrideSF", "Lighting Only"))
 /** needed for VMI_ReflectionOverride, Whether to override all materials to be smooth, mirror reflections. */
 SHOWFLAG_FIXED_IN_SHIPPING(0, ReflectionOverride, SFG_Hidden, NSLOCTEXT("UnrealEd", "ReflectionOverrideSF", "Reflections"))
 /** needed for VMI_VisualizeBuffer, Whether to enable the buffer visualization mode. */
