- Fixed a LOCTEXT_NAMESPACE define

diff --git a/Engine/Source/Runtime/Engine/Private/Materials/MaterialShared.cpp b/Engine/Source/Runtime/Engine/Private/Materials/MaterialShared.cpp
--- a/Engine/Source/Runtime/Engine/Private/Materials/MaterialShared.cpp
+++ b/Engine/Source/Runtime/Engine/Private/Materials/MaterialShared.cpp
@@ -4040,3 +4040,5 @@ void SetShaderMapsOnMaterialResources(const TMap<FMaterial*, FMaterialShaderMap*
 		SetShaderMapsOnMaterialResources_RenderThread(RHICmdList, InMaterialsToUpdate);
 	});
 }
+
+#undef LOCTEXT_NAMESPACE
\ No newline at end of file
