UE-88269
Fix non-editor use of UMaterial::GetAssetRegistryTags()

diff --git a/Engine/Source/Runtime/Engine/Private/Materials/Material.cpp b/Engine/Source/Runtime/Engine/Private/Materials/Material.cpp
--- a/Engine/Source/Runtime/Engine/Private/Materials/Material.cpp
+++ b/Engine/Source/Runtime/Engine/Private/Materials/Material.cpp
@@ -4484,13 +4484,7 @@ bool UMaterial::CanBeClusterRoot() const
 
 void UMaterial::GetAssetRegistryTags(TArray<FAssetRegistryTag>& OutTags) const
 {
-#if WITH_EDITORONLY_DATA
-	// Search for any scene color nodes
-	bool bHasSceneColor = HasAnyExpressionsInMaterialAndFunctionsOfType<UMaterialExpressionSceneColor>();
-	OutTags.Add(FAssetRegistryTag("HasSceneColor", bHasSceneColor ? TEXT("True") : TEXT("False"), FAssetRegistryTag::TT_Alphabetical));
-#else
-	check(false); // TODO Do we need to handle this??
-#endif
+	OutTags.Add(FAssetRegistryTag("HasSceneColor", CachedExpressionData.bHasSceneColor ? TEXT("True") : TEXT("False"), FAssetRegistryTag::TT_Alphabetical));
 	Super::GetAssetRegistryTags(OutTags);
 }
 
diff --git a/Engine/Source/Runtime/Engine/Private/Materials/MaterialCachedData.cpp b/Engine/Source/Runtime/Engine/Private/Materials/MaterialCachedData.cpp
--- a/Engine/Source/Runtime/Engine/Private/Materials/MaterialCachedData.cpp
+++ b/Engine/Source/Runtime/Engine/Private/Materials/MaterialCachedData.cpp
@@ -47,6 +47,7 @@ void FMaterialCachedExpressionData::Reset()
 	QualityLevelsUsed.Reset();
 	QualityLevelsUsed.AddDefaulted(EMaterialQualityLevel::Num);
 	bHasRuntimeVirtualTextureOutput = false;
+	bHasSceneColor = false;
 }
 
 static int32 FindParameterLowerBoundIndex(const FMaterialCachedParameterEntry& Entry, const FHashedMaterialParameterInfo& HashedParameterInfo)
@@ -299,6 +300,10 @@ void FMaterialCachedExpressionData::UpdateForExpressions(const TArray<UMaterialE
 		{
 			bHasRuntimeVirtualTextureOutput = true;
 		}
+		else if (Expression->IsA(UMaterialExpressionSceneColor::StaticClass()))
+		{
+			bHasSceneColor = true;
+		}
 		else if (UMaterialExpressionMaterialFunctionCall* FunctionCall = Cast<UMaterialExpressionMaterialFunctionCall>(Expression))
 		{
 			UpdateForFunction(FunctionCall->MaterialFunction, GlobalParameter, -1);
diff --git a/Engine/Source/Runtime/Engine/Public/MaterialCachedData.h b/Engine/Source/Runtime/Engine/Public/MaterialCachedData.h
--- a/Engine/Source/Runtime/Engine/Public/MaterialCachedData.h
+++ b/Engine/Source/Runtime/Engine/Public/MaterialCachedData.h
@@ -195,4 +195,7 @@ struct FMaterialCachedExpressionData
 
 	UPROPERTY()
 	uint32 bHasRuntimeVirtualTextureOutput : 1;
+
+	UPROPERTY()
+	uint32 bHasSceneColor : 1;
 };
