Factor platform into material shader layout hash, fixes problem when building shaders for multiple platforms

diff --git a/Engine/Source/Runtime/Engine/Private/Materials/MaterialShader.cpp b/Engine/Source/Runtime/Engine/Private/Materials/MaterialShader.cpp
--- a/Engine/Source/Runtime/Engine/Private/Materials/MaterialShader.cpp
+++ b/Engine/Source/Runtime/Engine/Private/Materials/MaterialShader.cpp
@@ -1867,7 +1867,7 @@ public:
 
 	const FMaterialShaderMapLayout& AcquireLayout(EShaderPlatform Platform, const FMaterialShaderParameters& MaterialParameters)
 	{
-		const uint64 ParameterHash = CityHash64((char*)&MaterialParameters, sizeof(MaterialParameters));
+		const uint64 ParameterHash = CityHash64WithSeed((char*)&MaterialParameters, sizeof(MaterialParameters), (uint64)Platform);
 		
 		int32 Index = INDEX_NONE;
 		for (Index = ShaderMapHashTable.First(ParameterHash); ShaderMapHashTable.IsValid(Index); Index = ShaderMapHashTable.Next(Index))
@@ -1885,9 +1885,11 @@ public:
 			FMaterialShaderMapLayout& Layout = ShaderMapLayouts.AddDefaulted_GetRef();
 			check(MaterialParameterHashes.Num() == ShaderMapLayouts.Num());
 			check(MaterialShaderParameters.Num() == ShaderMapLayouts.Num());
+			Layout.Platform = Platform;
 			CreateLayout(Layout, Platform, MaterialParameters);
 		}
 
+		check(ShaderMapLayouts[Index].Platform == Platform);
 		return ShaderMapLayouts[Index];
 	}
 
diff --git a/Engine/Source/Runtime/Engine/Private/Materials/MaterialShaderMapLayout.h b/Engine/Source/Runtime/Engine/Private/Materials/MaterialShaderMapLayout.h
--- a/Engine/Source/Runtime/Engine/Private/Materials/MaterialShaderMapLayout.h
+++ b/Engine/Source/Runtime/Engine/Private/Materials/MaterialShaderMapLayout.h
@@ -24,6 +24,7 @@ public:
 class FShaderMapLayout
 {
 public:
+	EShaderPlatform Platform;
 	TArray<FShaderLayoutEntry> Shaders;
 	TArray<FShaderPipelineType*> ShaderPipelines;
 };
