Fix non-deterministic texture updates in CurveLinearColorAtlas

The atlas updates its texture data in PostLoad which causes the ID to change every time it is loaded. Use a hash of the data as the ID to avoid the problem.

diff --git a/Engine/Source/Runtime/Engine/Private/Curves/CurveLinearColorAtlas.cpp b/Engine/Source/Runtime/Engine/Private/Curves/CurveLinearColorAtlas.cpp
--- a/Engine/Source/Runtime/Engine/Private/Curves/CurveLinearColorAtlas.cpp
+++ b/Engine/Source/Runtime/Engine/Private/Curves/CurveLinearColorAtlas.cpp
@@ -122,6 +122,23 @@ static void RenderGradient(TArray<FFloat16Color>& InSrcData, UObject* Gradient,
 	}
 }
 
+static void UpdateTexture(UCurveLinearColorAtlas& Atlas)
+{
+	const int32 TextureDataSize = Atlas.Source.CalcMipSize(0);
+
+	FGuid MD5Guid;
+	FMD5 MD5;
+	MD5.Update(reinterpret_cast<const uint8*>(Atlas.SrcData.GetData()), TextureDataSize);
+	MD5.Final(reinterpret_cast<uint8*>(&MD5Guid));
+
+	uint32* TextureData = reinterpret_cast<uint32*>(Atlas.Source.LockMip(0));
+	FMemory::Memcpy(TextureData, Atlas.SrcData.GetData(), TextureDataSize);
+	Atlas.Source.UnlockMip(0);
+
+	Atlas.Source.SetId(MD5Guid, /*bInGuidIsHash*/ true);
+	Atlas.UpdateResource();
+}
+
 // Immediately render a new material to the specified slot index (SlotIndex must be within this section's range)
 void UCurveLinearColorAtlas::UpdateGradientSlot(UCurveLinearColor* Gradient)
 {
@@ -137,14 +154,7 @@ void UCurveLinearColorAtlas::UpdateGradientSlot(UCurveLinearColor* Gradient)
 		// Render the single gradient to the render target
 		RenderGradient(SrcData, Gradient, StartXY, SizeXY);
 
-		uint32* TextureData = (uint32*)Source.LockMip(0);
-		const int32 TextureDataSize = Source.CalcMipSize(0);
-		FMemory::Memcpy(TextureData, SrcData.GetData(), TextureDataSize);
-
-		Source.UnlockMip(0);
-
-		// Immediately update the texture
-		UpdateResource();
+		UpdateTexture(*this);
 	}
 	
 }
@@ -178,10 +188,7 @@ void UCurveLinearColorAtlas::UpdateTextures()
 		}
 	}
 
-	uint32* TextureData = (uint32*)Source.LockMip(0);
-	FMemory::Memcpy(TextureData, SrcData.GetData(), TextureDataSize);
-	Source.UnlockMip(0);
-	UpdateResource();
+	UpdateTexture(*this);
 
 	bIsDirty = false;
 }
