#UE4 Added more info to an ensure

diff --git a/Engine/Source/Runtime/Engine/Private/AssetManager.cpp b/Engine/Source/Runtime/Engine/Private/AssetManager.cpp
--- a/Engine/Source/Runtime/Engine/Private/AssetManager.cpp
+++ b/Engine/Source/Runtime/Engine/Private/AssetManager.cpp
@@ -403,7 +403,7 @@ int32 UAssetManager::ScanPathsForPrimaryAssets(FPrimaryAssetType PrimaryAssetTyp
 	FPrimaryAssetTypeData& TypeData = FoundType->Get();
 
 	// Make sure types match
-	if (!ensure(TypeData.Info.AssetBaseClassLoaded == BaseClass && TypeData.Info.bHasBlueprintClasses == bHasBlueprintClasses && TypeData.Info.bIsEditorOnly == bIsEditorOnly))
+	if (!ensureMsgf(TypeData.Info.AssetBaseClassLoaded == BaseClass && TypeData.Info.bHasBlueprintClasses == bHasBlueprintClasses && TypeData.Info.bIsEditorOnly == bIsEditorOnly, TEXT("UAssetManager::ScanPathsForPrimaryAssets TypeData parameters did not match for type '%s'"), *TypeData.Info.PrimaryAssetType.ToString()))
 	{
 		return 0;
 	}
