The Particle System Component Manager will now account for the owner's time dilation when ticking particle system components.

diff --git a/Engine/Source/Runtime/Engine/Private/Particles/ParticleSystemManager.cpp b/Engine/Source/Runtime/Engine/Private/Particles/ParticleSystemManager.cpp
--- a/Engine/Source/Runtime/Engine/Private/Particles/ParticleSystemManager.cpp
+++ b/Engine/Source/Runtime/Engine/Private/Particles/ParticleSystemManager.cpp
@@ -566,17 +566,20 @@ void FParticleSystemWorldManager::ProcessTickList(float DeltaTime, ELevelTick Ti
 				{
 					continue;
 				}
+				
+				AActor* PSCOwner = PSC->GetOwner();
+				const float TimeDilation = (PSCOwner ? PSCOwner->CustomTimeDilation : 1.f);
 
 				//TODO: Replace call to TickComponent with new call that allows us to pull duplicated work up to share across all ticks.
 				if (bAsync)
 				{
-					PSC->TickComponent(DeltaTime, TickType, nullptr);
+					PSC->TickComponent(DeltaTime * TimeDilation, TickType, nullptr);
 					PSC->MarshalParamsForAsyncTick();
 					QueueAsyncTick(Handle, TickGroupCompletionGraphEvent);
 				}
 				else
 				{
-					PSC->TickComponent(DeltaTime, TickType, nullptr);
+					PSC->TickComponent(DeltaTime * TimeDilation, TickType, nullptr);
 					PSC->ComputeTickComponent_Concurrent();
 					PSC->FinalizeTickComponent();
 				}
