Fix for crash when USoundWave::GetHandleForChunkOfAudio- issue was that callback here could complete after the USoundWave is GC'd. Fix is capturing a TWeakObjPtr instead.

diff --git a/Engine/Source/Runtime/Engine/Private/SoundWave.cpp b/Engine/Source/Runtime/Engine/Private/SoundWave.cpp
--- a/Engine/Source/Runtime/Engine/Private/SoundWave.cpp
+++ b/Engine/Source/Runtime/Engine/Private/SoundWave.cpp
@@ -2425,21 +2425,36 @@ void USoundWave::GetHandleForChunkOfAudio(TFunction<void(FAudioChunkHandle)> OnL
 	}
 	else
 	{
+		TWeakObjectPtr<USoundWave> WeakThis = MakeWeakObjectPtr(this);
+
 		// For async cases, we call RequestChunk and request the loaded chunk in the completion callback.
-		IStreamingManager::Get().GetAudioStreamingManager().RequestChunk(this, ChunkIndex, [this, OnLoadCompleted, ChunkIndex](EAudioChunkLoadResult LoadResult)
+		IStreamingManager::Get().GetAudioStreamingManager().RequestChunk(this, ChunkIndex, [WeakThis, OnLoadCompleted, ChunkIndex, CallbackThread](EAudioChunkLoadResult LoadResult)
 		{
-			if (LoadResult == EAudioChunkLoadResult::Completed || LoadResult == EAudioChunkLoadResult::AlreadyLoaded)
+			auto DispatchOnLoadCompletedCallback = [OnLoadCompleted, CallbackThread](FAudioChunkHandle InHandle)
+			{
+				// If the callback was requested on a non-game thread, dispatch the callback to that thread.
+				AsyncTask(CallbackThread, [OnLoadCompleted, InHandle]()
+				{
+					OnLoadCompleted(InHandle);
+				});
+			};
+
+			// If the USoundWave has been GC'd by the time this chunk finishes loading, abandon ship.
+			if (WeakThis.IsValid() && (LoadResult == EAudioChunkLoadResult::Completed || LoadResult == EAudioChunkLoadResult::AlreadyLoaded))
 			{
-				FAudioChunkHandle ChunkHandle = IStreamingManager::Get().GetAudioStreamingManager().GetLoadedChunk(this, ChunkIndex, true);
-				OnLoadCompleted(ChunkHandle);
+				USoundWave* ThisSoundWave = WeakThis.Get();
+				check(ThisSoundWave);
+
+				FAudioChunkHandle ChunkHandle = IStreamingManager::Get().GetAudioStreamingManager().GetLoadedChunk(ThisSoundWave, ChunkIndex, true);
+				DispatchOnLoadCompletedCallback(ChunkHandle);
 			}
 			else
 			{
 				// Load failed. Return an invalid chunk handle.
 				FAudioChunkHandle ChunkHandle;
-				OnLoadCompleted(ChunkHandle);
+				DispatchOnLoadCompletedCallback(ChunkHandle);
 			}
-		}, CallbackThread);
+		}, ENamedThreads::GameThread);
 	}
 }
 
