Make sure we have a cvar that will outright disable retaining audio if necessary.

diff --git a/Engine/Source/Runtime/Engine/Private/SoundWave.cpp b/Engine/Source/Runtime/Engine/Private/SoundWave.cpp
--- a/Engine/Source/Runtime/Engine/Private/SoundWave.cpp
+++ b/Engine/Source/Runtime/Engine/Private/SoundWave.cpp
@@ -52,6 +52,14 @@ FAutoConsoleVariableRef CVarForceNonStreamingInEditor(
 	TEXT("0: Honor the Play When Silent flag, 1: stop all silent non-procedural sources."),
 	ECVF_Default);
 
+static int32 DisableRetainingCVar = 0;
+FAutoConsoleVariableRef CVarDisableRetaining(
+	TEXT("au.streamcache.DisableRetaining"),
+	DisableRetainingCVar,
+	TEXT("When set to 1, USoundWaves will not retain chunks of their own audio.\n")
+	TEXT("0: Don't disable retaining, 1: retaining."),
+	ECVF_Default);
+
 #if !UE_BUILD_SHIPPING
 static void DumpBakedAnalysisData(const TArray<FString>& Args)
 {
@@ -2462,7 +2470,7 @@ void USoundWave::RetainCompressedAudio(bool bForceSync /*= false*/)
 {
 	// Since the zeroth chunk is always inlined and stored in memory,
 	// early exit if we only have one chunk.
-	if (GetNumChunks() <= 1)
+	if (DisableRetainingCVar || GetNumChunks() <= 1)
 	{
 		return;
 	}
