Better fix for split screen detection

diff --git a/Engine/Source/Runtime/Engine/Classes/Engine/Engine.h b/Engine/Source/Runtime/Engine/Classes/Engine/Engine.h
--- a/Engine/Source/Runtime/Engine/Classes/Engine/Engine.h
+++ b/Engine/Source/Runtime/Engine/Classes/Engine/Engine.h
@@ -2472,7 +2472,7 @@ public:
 	FAudioDeviceHandle GetActiveAudioDevice();
 
 	/** @return whether we're currently running in split screen (more than one local player) */
-	bool IsSplitScreen(UWorld *InWorld);
+	virtual bool IsSplitScreen(UWorld *InWorld);
 
 	/** @return whether we're currently running with stereoscopic 3D enabled for the specified viewport (or globally, if viewport is nullptr) */
 	bool IsStereoscopic3D(FViewport* InViewport = nullptr);
