Fix CIS Issues in DataReplication.cpp

#rnx

diff --git a/Engine/Source/Runtime/Engine/Private/DataReplication.cpp b/Engine/Source/Runtime/Engine/Private/DataReplication.cpp
--- a/Engine/Source/Runtime/Engine/Private/DataReplication.cpp
+++ b/Engine/Source/Runtime/Engine/Private/DataReplication.cpp
@@ -1422,6 +1422,8 @@ void FObjectReplicator::ReplicateCustomDeltaProperties( FNetBitWriter & Bunch, F
 
 	FNetBitWriter TempBitWriter( Connection->PackageMap, 1024 );
 
+	const bool bIsConnectionInternalAck = Connection->IsInternalAck();
+
 	// Replicate those properties.
 	for (uint16 CustomDeltaProperty = 0; CustomDeltaProperty < NumLifetimeCustomDeltaProperties; ++CustomDeltaProperty)
 	{
@@ -1474,7 +1476,7 @@ void FObjectReplicator::ReplicateCustomDeltaProperties( FNetBitWriter & Bunch, F
 
 		FPropertyRetirement** LastNext = nullptr;
 
-		if (!Connection->IsInternalAck())
+		if (!bIsConnectionInternalAck)
 		{
 			// Get info.
 			FPropertyRetirement& Retire = SendingRepState->Retirement[CustomDeltaProperty];
@@ -1499,7 +1501,7 @@ void FObjectReplicator::ReplicateCustomDeltaProperties( FNetBitWriter & Bunch, F
 			continue;
 		}
 
-		if (!Connection->IsInternalAck())
+		if (!bIsConnectionInternalAck)
 		{
 			*LastNext = new FPropertyRetirement();
 
