Blueprints: Added DeterminesOutputType+DynamicOutputParam metadata to GetAllActorsOfClassWithTag so the output array will be typed if a class literal is used

diff --git a/Engine/Source/Runtime/Engine/Classes/Kismet/GameplayStatics.h b/Engine/Source/Runtime/Engine/Classes/Kismet/GameplayStatics.h
--- a/Engine/Source/Runtime/Engine/Classes/Kismet/GameplayStatics.h
+++ b/Engine/Source/Runtime/Engine/Classes/Kismet/GameplayStatics.h
@@ -123,7 +123,7 @@ class ENGINE_API UGameplayStatics : public UBlueprintFunctionLibrary
 	 *	@param	ActorClass	Class of Actor to find. Must be specified or result array will be empty.
 	 *	@param	OutActors	Output array of Actors of the specified tag.
 	 */
-	UFUNCTION(BlueprintCallable, Category = "Utilities", meta = (WorldContext = "WorldContextObject"))
+	UFUNCTION(BlueprintCallable, Category = "Utilities", meta = (WorldContext="WorldContextObject", DeterminesOutputType="ActorClass", DynamicOutputParam="OutActors"))
 	static void GetAllActorsOfClassWithTag(const UObject* WorldContextObject, TSubclassOf<AActor> ActorClass, FName Tag, TArray<AActor*>& OutActors);
 
 	// --- Player functions ------------------------------
