Adding bStencilTestEnabled to FMaterialShaderParameters needed for a merge conflict

diff --git a/Engine/Source/Runtime/Engine/Public/MaterialShared.h b/Engine/Source/Runtime/Engine/Public/MaterialShared.h
--- a/Engine/Source/Runtime/Engine/Public/MaterialShared.h
+++ b/Engine/Source/Runtime/Engine/Public/MaterialShared.h
@@ -2925,6 +2925,7 @@ struct FMaterialShaderParameters
 	int32 BlendableLocation;
 	uint32 DecalBlendMode;
 	int32 NumCustomizedUVs;
+	uint32 StencilCompare;
 	union
 	{
 		uint64 PackedFlags;
@@ -2973,6 +2974,7 @@ struct FMaterialShaderParameters
 			uint64 bMaterialIsRequiredTextureResolution : 1;
 			uint64 bMaterialIsComplexityAccumulate : 1;
 			uint64 bIsUsedWithLidarPointCloud : 1;
+			uint64 bStencilTestEnabled : 1;
 		};
 	};
 
@@ -2990,6 +2992,7 @@ struct FMaterialShaderParameters
 		BlendableLocation = InMaterial->GetBlendableLocation();
 		DecalBlendMode = InMaterial->GetDecalBlendMode();
 		NumCustomizedUVs = InMaterial->GetNumCustomizedUVs();
+		StencilCompare = InMaterial->GetStencilCompare();
 		bIsDefaultMaterial = InMaterial->IsDefaultMaterial();
 		bIsSpecialEngineMaterial = InMaterial->IsSpecialEngineMaterial();
 		bIsMasked = InMaterial->IsMasked();
@@ -3027,6 +3030,7 @@ struct FMaterialShaderParameters
 		bIsUsedWithInstancedStaticMeshes = InMaterial->IsUsedWithInstancedStaticMeshes();
 		bHasRuntimeVirtualTextureOutput = InMaterial->HasRuntimeVirtualTextureOutput();
 		bIsUsedWithLidarPointCloud = InMaterial->IsUsedWithLidarPointCloud();
+		bStencilTestEnabled = InMaterial->IsStencilTestEnabled();
 
 		// See FDebugViewModeMaterialProxy::GetFriendlyName()
 		// TODO seems horrible that friendly name controls which shaders get compiled, should refactor this to use regular accessors
