Correctly handle streaming levels being removed as a side-effect of updating the streaming state of another streaming level

diff --git a/Engine/Source/Runtime/Engine/Classes/Engine/World.h b/Engine/Source/Runtime/Engine/Classes/Engine/World.h
--- a/Engine/Source/Runtime/Engine/Classes/Engine/World.h
+++ b/Engine/Source/Runtime/Engine/Classes/Engine/World.h
@@ -790,13 +790,13 @@ struct FStreamingLevelsToConsider
 		: bStreamingLevelsBeingConsidered(false)
 	{}
 
+private:
+
 	/** Priority sorted array of streaming levels actively being considered. */
 	UPROPERTY()
 	TArray<FLevelStreamingWrapper> StreamingLevels;
 
-private:
-
-	enum class EProcessReason
+	enum class EProcessReason : uint8
 	{
 		Add,
 		Reevaluate
@@ -816,6 +816,8 @@ private:
 
 public:
 
+	const TArray<FLevelStreamingWrapper>& GetStreamingLevels() const { return StreamingLevels; }
+
 	void AddReferencedObjects(UObject* InThis, FReferenceCollector& Collector);
 
 	void BeginConsideration();
@@ -827,6 +829,9 @@ public:
 	/* Remove an element from the container. */
 	bool Remove(ULevelStreaming* StreamingLevel);
 
+	/* Remove the element at a given index from the container. */
+	void RemoveAt(int32 Index);
+
 	/* Returns if an element is in the container. */
 	bool Contains(ULevelStreaming* StreamingLevel) const;
 
diff --git a/Engine/Source/Runtime/Engine/Private/World.cpp b/Engine/Source/Runtime/Engine/Private/World.cpp
--- a/Engine/Source/Runtime/Engine/Private/World.cpp
+++ b/Engine/Source/Runtime/Engine/Private/World.cpp
@@ -3132,12 +3132,37 @@ bool FStreamingLevelsToConsider::Remove(ULevelStreaming* StreamingLevel)
 	if (StreamingLevel)
 	{
 		SCOPE_CYCLE_COUNTER(STAT_ManageLevelsToConsider);
-		bRemoved = (StreamingLevels.Remove(StreamingLevel) > 0);
-		bRemoved |= (LevelsToProcess.Remove(StreamingLevel) > 0);
+		if (bStreamingLevelsBeingConsidered)
+		{
+			int32 Index;
+			if (StreamingLevels.Find(StreamingLevel, Index))
+			{
+				// While we are considering we must null here because we are iterating the array and changing the size would be undesirable
+				StreamingLevels[Index] = nullptr;
+				bRemoved = true;
+			}
+			bRemoved |= (LevelsToProcess.Remove(StreamingLevel) > 0);
+		}
+		else
+		{
+			bRemoved = (StreamingLevels.Remove(StreamingLevel) > 0);
+		}
 	}
 	return bRemoved;
 }
 
+void FStreamingLevelsToConsider::RemoveAt(const int32 Index)
+{
+	if (bStreamingLevelsBeingConsidered)
+	{
+		if (ULevelStreaming* StreamingLevel = StreamingLevels[Index].Get())
+		{
+			LevelsToProcess.Remove(StreamingLevel);
+		}
+	}
+	StreamingLevels.RemoveAt(Index, 1, false);
+}
+
 void FStreamingLevelsToConsider::Reevaluate(ULevelStreaming* StreamingLevel)
 {
 	if (StreamingLevel)
@@ -3233,11 +3258,10 @@ void UWorld::UpdateLevelStreaming()
 	const int32 NumLevelsPendingPurge = FLevelStreamingGCHelper::GetNumLevelsPendingPurge();
 
 	StreamingLevelsToConsider.BeginConsideration();
-	TArray<FLevelStreamingWrapper> StreamingLevelsBeingConsidered = MoveTemp(StreamingLevelsToConsider.StreamingLevels);
 
-	for (int32 Index = StreamingLevelsBeingConsidered.Num() - 1; Index >= 0; --Index)
+	for (int32 Index = StreamingLevelsToConsider.GetStreamingLevels().Num() - 1; Index >= 0; --Index)
 	{
-		if (ULevelStreaming* StreamingLevel = StreamingLevelsBeingConsidered[Index].Get())
+		if (ULevelStreaming* StreamingLevel = StreamingLevelsToConsider.GetStreamingLevels()[Index].Get())
 		{
 			bool bUpdateAgain = true;
 			bool bShouldContinueToConsider = true;
@@ -3254,17 +3278,15 @@ void UWorld::UpdateLevelStreaming()
 
 			if (!bShouldContinueToConsider)
 			{
-				StreamingLevelsBeingConsidered.RemoveAt(Index, 1, false);
-				StreamingLevelsToConsider.Remove(StreamingLevel); // In case something had added it to the list while we're in this loop
+				StreamingLevelsToConsider.RemoveAt(Index);
 			}
 		}
 		else
 		{
-			StreamingLevelsBeingConsidered.RemoveAt(Index, 1, false);
+			StreamingLevelsToConsider.RemoveAt(Index);
 		}
 	}
 
-	StreamingLevelsToConsider.StreamingLevels = MoveTemp(StreamingLevelsBeingConsidered);
 	StreamingLevelsToConsider.EndConsideration();
 
 
@@ -3355,7 +3377,7 @@ void UWorld::SetStreamingLevels(TArrayView<ULevelStreaming* const> InStreamingLe
 
 void UWorld::SetStreamingLevels(TArray<ULevelStreaming*>&& InStreamingLevels)
 {
-	StreamingLevels = MoveTempIfPossible(InStreamingLevels);
+	StreamingLevels = MoveTemp(InStreamingLevels);
 
 	PopulateStreamingLevelsToConsider();
 }
