Fix typo in merge conflict fix.

diff --git a/Engine/Source/Runtime/Engine/Public/MaterialShared.h b/Engine/Source/Runtime/Engine/Public/MaterialShared.h
--- a/Engine/Source/Runtime/Engine/Public/MaterialShared.h
+++ b/Engine/Source/Runtime/Engine/Public/MaterialShared.h
@@ -2974,7 +2974,7 @@ struct FMaterialShaderParameters
 			uint64 bMaterialIsRequiredTextureResolution : 1;
 			uint64 bMaterialIsComplexityAccumulate : 1;
 			uint64 bIsUsedWithLidarPointCloud : 1;
-			uint64 bStencilTestEnabled : 1;
+			uint64 bIsStencilTestEnabled : 1;
 		};
 	};
 
@@ -3030,7 +3030,7 @@ struct FMaterialShaderParameters
 		bIsUsedWithInstancedStaticMeshes = InMaterial->IsUsedWithInstancedStaticMeshes();
 		bHasRuntimeVirtualTextureOutput = InMaterial->HasRuntimeVirtualTextureOutput();
 		bIsUsedWithLidarPointCloud = InMaterial->IsUsedWithLidarPointCloud();
-		bStencilTestEnabled = InMaterial->IsStencilTestEnabled();
+		bIsStencilTestEnabled = InMaterial->IsStencilTestEnabled();
 
 		// See FDebugViewModeMaterialProxy::GetFriendlyName()
 		// TODO seems horrible that friendly name controls which shaders get compiled, should refactor this to use regular accessors
