- Fixing some SA issues

diff --git a/Engine/Source/Runtime/Core/Public/Serialization/MemoryImage.h b/Engine/Source/Runtime/Core/Public/Serialization/MemoryImage.h
--- a/Engine/Source/Runtime/Core/Public/Serialization/MemoryImage.h
+++ b/Engine/Source/Runtime/Core/Public/Serialization/MemoryImage.h
@@ -867,7 +867,7 @@ namespace Freeze
 		{
 			const uint32 Index = Writer.GetPointerTable().AddIndexedPointer(TypeDesc, RawPtr);
 			check(Index != (uint32)INDEX_NONE);
-			const uint64 FrozenPackedIndex = (Index << 1u) | 1u;
+			const uint64 FrozenPackedIndex = (Index << 1llu) | 1u;
 			Writer.WriteMemoryImagePointerSizedBytes(FrozenPackedIndex);
 		}
 		else
diff --git a/Engine/Source/Runtime/Engine/Private/Rendering/ColorVertexBuffer.cpp b/Engine/Source/Runtime/Engine/Private/Rendering/ColorVertexBuffer.cpp
--- a/Engine/Source/Runtime/Engine/Private/Rendering/ColorVertexBuffer.cpp
+++ b/Engine/Source/Runtime/Engine/Private/Rendering/ColorVertexBuffer.cpp
@@ -330,6 +330,7 @@ void FColorVertexBuffer::ImportText(const TCHAR* SourceText)
 void FColorVertexBuffer::operator=(const FColorVertexBuffer &Other)
 {
 	//VertexData doesn't need to be allocated here because Build will be called next,
+	delete VertexData;
 	VertexData = NULL;
 }
 
diff --git a/Engine/Source/Runtime/Engine/Private/SoundWave.cpp b/Engine/Source/Runtime/Engine/Private/SoundWave.cpp
--- a/Engine/Source/Runtime/Engine/Private/SoundWave.cpp
+++ b/Engine/Source/Runtime/Engine/Private/SoundWave.cpp
@@ -537,7 +537,7 @@ bool USoundWave::HasCompressedData(FName Format, ITargetPlatform* TargetPlatform
 	else
 	{
 		// TargetPlatform is not available on consoles/mobile, so we have to grab it ourselves:
-		CompressionOverrides = CompressionOverrides = FPlatformCompressionUtilities::GetCookOverrides();
+		CompressionOverrides = FPlatformCompressionUtilities::GetCookOverrides();
 	}
 
 	if (CompressionOverrides)
