Fix vehicle seat capsules not colliding. * Implement FPhysInterface_Chaos::GetSimulationFilter and FPhysInterface_Chaos::GetQueryFilter functions.
* Don't set query filters AND disabled flags for chaos objects.

diff --git a/Engine/Source/Runtime/Engine/Private/PhysicsEngine/Experimental/PhysInterface_Chaos.cpp b/Engine/Source/Runtime/Engine/Private/PhysicsEngine/Experimental/PhysInterface_Chaos.cpp
--- a/Engine/Source/Runtime/Engine/Private/PhysicsEngine/Experimental/PhysInterface_Chaos.cpp
+++ b/Engine/Source/Runtime/Engine/Private/PhysicsEngine/Experimental/PhysInterface_Chaos.cpp
@@ -1559,14 +1559,28 @@ FPhysicsGeometryCollection_Chaos::FPhysicsGeometryCollection_Chaos(const FPhysic
 {
 }
 
-FCollisionFilterData FPhysInterface_Chaos::GetSimulationFilter(const FPhysicsShapeHandle& InShape)
+FCollisionFilterData FPhysInterface_Chaos::GetSimulationFilter(const FPhysicsShapeReference_Chaos& InShape)
 {
-	return FCollisionFilterData();
+	if (ensure(InShape.Shape))
+	{
+		return InShape.Shape->SimData;
+	}
+	else
+	{
+		return FCollisionFilterData();
+	}
 }
 
-FCollisionFilterData FPhysInterface_Chaos::GetQueryFilter(const FPhysicsShapeHandle& InShape)
+FCollisionFilterData FPhysInterface_Chaos::GetQueryFilter(const FPhysicsShapeReference_Chaos& InShape)
 {
-	return FCollisionFilterData();
+	if (ensure(InShape.Shape))
+	{
+		return InShape.Shape->QueryData;
+	}
+	else
+	{
+		return FCollisionFilterData();
+	}
 }
 
 void FPhysInterface_Chaos::SetQueryFilter(const FPhysicsShapeReference_Chaos& InShapeRef, const FCollisionFilterData& InFilter)
