Fix for warning after upgrading to PVS-Studio 7.0.4:

V1043: A global object variable 'NAME_RHI' is declared in the header. Multiple copies of it will be created in all translation units that include this header file.

#rnx

diff --git a/Engine/Source/Runtime/Engine/Private/HardwareInfo.cpp b/Engine/Source/Runtime/Engine/Private/HardwareInfo.cpp
--- a/Engine/Source/Runtime/Engine/Private/HardwareInfo.cpp
+++ b/Engine/Source/Runtime/Engine/Private/HardwareInfo.cpp
@@ -11,6 +11,9 @@
 
 static TMap< FName, FString > HardwareDetailsMap;
 
+ENGINE_API const FName NAME_RHI( "RHI" );
+ENGINE_API const FName NAME_TextureFormat( "TextureFormat" );
+ENGINE_API const FName NAME_DeviceType( "DeviceType" );
 
 void FHardwareInfo::RegisterHardwareInfo( const FName SpecIdentifier, const FString& HardwareInfo )
 {
diff --git a/Engine/Source/Runtime/Engine/Public/HardwareInfo.h b/Engine/Source/Runtime/Engine/Public/HardwareInfo.h
--- a/Engine/Source/Runtime/Engine/Public/HardwareInfo.h
+++ b/Engine/Source/Runtime/Engine/Public/HardwareInfo.h
@@ -9,9 +9,9 @@
 #include "CoreMinimal.h"
 
 /** Hardware entry lookups */
-static const FName NAME_RHI( "RHI" );
-static const FName NAME_TextureFormat( "TextureFormat" );
-static const FName NAME_DeviceType( "DeviceType" );
+extern ENGINE_API const FName NAME_RHI;
+extern ENGINE_API const FName NAME_TextureFormat;
+extern ENGINE_API const FName NAME_DeviceType;
 
 
 struct ENGINE_API FHardwareInfo
