Temporarily change the child actor root component's mobility to allow attachment to work when spawning and attaching to the child actor component.

diff --git a/Engine/Source/Runtime/Engine/Private/Components/ChildActorComponent.cpp b/Engine/Source/Runtime/Engine/Private/Components/ChildActorComponent.cpp
--- a/Engine/Source/Runtime/Engine/Private/Components/ChildActorComponent.cpp
+++ b/Engine/Source/Runtime/Engine/Private/Components/ChildActorComponent.cpp
@@ -594,7 +594,11 @@ void UChildActorComponent::CreateChildActor()
 					const FComponentInstanceDataCache* ComponentInstanceData = (CachedInstanceData ? CachedInstanceData->ComponentInstanceData.Get() : nullptr);
 					ChildActor->FinishSpawning(GetComponentTransform(), false, ComponentInstanceData);
 
-					ChildActor->AttachToComponent(this, FAttachmentTransformRules::SnapToTargetNotIncludingScale);
+					if (USceneComponent* ChildRoot = ChildActor->GetRootComponent())
+					{
+						TGuardValue<TEnumAsByte<EComponentMobility::Type>> MobilityGuard(ChildRoot->Mobility, Mobility);
+						ChildRoot->AttachToComponent(this, FAttachmentTransformRules::SnapToTargetNotIncludingScale);
+					}
 
 					SetIsReplicated(ChildActor->GetIsReplicated());
 
