- Made a UProperty back to an FProperty (change got stomped)
#rnx

diff --git a/Engine/Source/Runtime/Engine/Private/SoundWave.cpp b/Engine/Source/Runtime/Engine/Private/SoundWave.cpp
--- a/Engine/Source/Runtime/Engine/Private/SoundWave.cpp
+++ b/Engine/Source/Runtime/Engine/Private/SoundWave.cpp
@@ -1496,7 +1496,7 @@ void USoundWave::PostEditChangeProperty(FPropertyChangedEvent& PropertyChangedEv
 	if (PropertyChangedEvent.ChangeType != EPropertyChangeType::Interactive)
 	{
 		// Regenerate on save any compressed sound formats or if analysis needs to be re-done
-		if (UProperty* PropertyThatChanged = PropertyChangedEvent.Property)
+		if (FProperty* PropertyThatChanged = PropertyChangedEvent.Property)
 		{
 			const FName& Name = PropertyThatChanged->GetFName();
 			if (Name == CompressionQualityFName || Name == StreamingFName || Name == SeekableStreamingFName)
