Fix assert firing when there are zero guids to serialize for the checkpoint (if using delta)

diff --git a/Engine/Source/Runtime/Engine/Private/DemoNetDriver.cpp b/Engine/Source/Runtime/Engine/Private/DemoNetDriver.cpp
--- a/Engine/Source/Runtime/Engine/Private/DemoNetDriver.cpp
+++ b/Engine/Source/Runtime/Engine/Private/DemoNetDriver.cpp
@@ -1916,7 +1916,7 @@ bool UDemoNetDriver::SerializeGuidCache(const FRepActorsCheckpointParams& Params
 	const double StartTime = FPlatformTime::Seconds();
 	const double Deadline = Params.StartCheckpointTime + Params.CheckpointMaxUploadTimePerFrame;
 
-	check(NextNetGuidForRecording >= 0 && NextNetGuidForRecording < NetGuidCacheSnapshot.Num());
+	check(NetGuidCacheSnapshot.Num() == 0 || NetGuidCacheSnapshot.IsValidIndex(NextNetGuidForRecording));
 
 	for (; NextNetGuidForRecording != NetGuidCacheSnapshot.Num(); ++NextNetGuidForRecording)
 	{
