Fix static analysis warnings in CharacterMovementComponent.cpp

#rnx

diff --git a/Engine/Source/Runtime/Engine/Private/Components/CharacterMovementComponent.cpp b/Engine/Source/Runtime/Engine/Private/Components/CharacterMovementComponent.cpp
--- a/Engine/Source/Runtime/Engine/Private/Components/CharacterMovementComponent.cpp
+++ b/Engine/Source/Runtime/Engine/Private/Components/CharacterMovementComponent.cpp
@@ -1504,13 +1504,14 @@ void UCharacterMovementComponent::SimulatedTick(float DeltaSeconds)
 			// debug
 			if (false)
 			{
+				checkSlow(CharacterOwner != nullptr);
 				const FRotator OldRotation = OldRotationQuat.Rotator();
 				const FRotator NewRotation = UpdatedComponent->GetComponentRotation();
 				const FVector NewLocation = UpdatedComponent->GetComponentLocation();
-				DrawDebugCoordinateSystem(GetWorld(), CharacterOwner->GetMesh()->GetComponentLocation() + FVector(0,0,1), NewRotation, 50.f, false); //-V595
+				DrawDebugCoordinateSystem(GetWorld(), CharacterOwner->GetMesh()->GetComponentLocation() + FVector(0,0,1), NewRotation, 50.f, false);
 				DrawDebugLine(GetWorld(), OldLocation, NewLocation, FColor::Red, false, 10.f);
 
-				UE_LOG(LogRootMotion, Log,  TEXT("UCharacterMovementComponent::SimulatedTick DeltaMovement Translation: %s, Rotation: %s, MovementBase: %s"), //-V595
+				UE_LOG(LogRootMotion, Log,  TEXT("UCharacterMovementComponent::SimulatedTick DeltaMovement Translation: %s, Rotation: %s, MovementBase: %s"),
 					*(NewLocation - OldLocation).ToCompactString(), *(NewRotation - OldRotation).GetNormalized().ToCompactString(), *GetNameSafe(CharacterOwner->GetMovementBase()) );
 			}
 #endif // !(UE_BUILD_SHIPPING)
@@ -9314,7 +9315,8 @@ void UCharacterMovementComponent::ClientAdjustPosition_Implementation
 		// If walking, we'd like to continue walking if possible, to avoid falling for a frame, so try to find a base where we moved to.
 		if (PreviousBase)
 		{
-			FindFloor(UpdatedComponent->GetComponentLocation(), CurrentFloor, false); //-V595
+			checkSlow(UpdatedComponent != nullptr);
+			FindFloor(UpdatedComponent->GetComponentLocation(), CurrentFloor, false);
 			if (CurrentFloor.IsWalkableFloor())
 			{
 				FinalBase = CurrentFloor.HitResult.Component.Get();
