Add option to create standalone game instance's world in a transient package

diff --git a/Engine/Source/Runtime/Engine/Classes/Engine/GameInstance.h b/Engine/Source/Runtime/Engine/Classes/Engine/GameInstance.h
--- a/Engine/Source/Runtime/Engine/Classes/Engine/GameInstance.h
+++ b/Engine/Source/Runtime/Engine/Classes/Engine/GameInstance.h
@@ -222,7 +222,7 @@ public:
 	void HandleTravelError(ETravelFailure::Type FailureType);
 
 	/* Called to initialize the game instance for standalone instances of the game */
-	void InitializeStandalone();
+	void InitializeStandalone(const FName InWorldName = NAME_None, UPackage* InWorldPackage = nullptr);
 
 	/* Called to initialize the game instance with a minimal world suitable for basic network RPC */
 	void InitializeForMinimalNetRPC(const FName InPackageName);
diff --git a/Engine/Source/Runtime/Engine/Private/GameInstance.cpp b/Engine/Source/Runtime/Engine/Private/GameInstance.cpp
--- a/Engine/Source/Runtime/Engine/Private/GameInstance.cpp
+++ b/Engine/Source/Runtime/Engine/Private/GameInstance.cpp
@@ -143,14 +143,14 @@ void UGameInstance::Shutdown()
 	WorldContext = nullptr;
 }
 
-void UGameInstance::InitializeStandalone()
+void UGameInstance::InitializeStandalone(const FName InPackageName, UPackage* InWorldPackage)
 {
 	// Creates the world context. This should be the only WorldContext that ever gets created for this GameInstance.
 	WorldContext = &GetEngine()->CreateNewWorldContext(EWorldType::Game);
 	WorldContext->OwningGameInstance = this;
 
 	// In standalone create a dummy world from the beginning to avoid issues of not having a world until LoadMap gets us our real world
-	UWorld* DummyWorld = UWorld::CreateWorld(EWorldType::Game, false);
+	UWorld* DummyWorld = UWorld::CreateWorld(EWorldType::Game, false, InPackageName, InWorldPackage);
 	DummyWorld->SetGameInstance(this);
 	WorldContext->SetCurrentWorld(DummyWorld);
 
