GetRelevance -> GetRelevance_Concurrent, fixes crash during async AddPrimitive

diff --git a/Engine/Source/Runtime/Engine/Private/Particles/ParticleEmitterInstances.cpp b/Engine/Source/Runtime/Engine/Private/Particles/ParticleEmitterInstances.cpp
--- a/Engine/Source/Runtime/Engine/Private/Particles/ParticleEmitterInstances.cpp
+++ b/Engine/Source/Runtime/Engine/Private/Particles/ParticleEmitterInstances.cpp
@@ -2914,16 +2914,16 @@ void FParticleEmitterInstance::GatherMaterialRelevance(FMaterialRelevance* OutMa
 	// These will catch the sprite cases...
 	if (CurrentMaterial)
 	{
-		(*OutMaterialRelevance) |= CurrentMaterial->GetRelevance(InFeatureLevel);
+		(*OutMaterialRelevance) |= CurrentMaterial->GetRelevance_Concurrent(InFeatureLevel);
 	}
 	else if (LODLevel->RequiredModule->Material)
 	{
-		(*OutMaterialRelevance) |= LODLevel->RequiredModule->Material->GetRelevance(InFeatureLevel);
+		(*OutMaterialRelevance) |= LODLevel->RequiredModule->Material->GetRelevance_Concurrent(InFeatureLevel);
 	}
 	else
 	{
 		check( UMaterial::GetDefaultMaterial(MD_Surface) );
-		(*OutMaterialRelevance) |= UMaterial::GetDefaultMaterial(MD_Surface)->GetRelevance(InFeatureLevel);
+		(*OutMaterialRelevance) |= UMaterial::GetDefaultMaterial(MD_Surface)->GetRelevance_Concurrent(InFeatureLevel);
 	}
 }
 
@@ -3800,7 +3800,7 @@ void FParticleMeshEmitterInstance::GatherMaterialRelevance( FMaterialRelevance*
 	GetMeshMaterials(Materials, LODLevel, InFeatureLevel, true); // Allow log issues since GatherMaterialRelevance is only called when the proxy is created.
 	for (int32 MaterialIndex = 0; MaterialIndex < Materials.Num(); ++MaterialIndex)
 	{
-		(*OutMaterialRelevance) |= Materials[MaterialIndex]->GetRelevance(InFeatureLevel);
+		(*OutMaterialRelevance) |= Materials[MaterialIndex]->GetRelevance_Concurrent(InFeatureLevel);
 	}
 }
 
