MPCs - now use the TitleProperty metadata to show the parameter name in the header for all those items in an array.

diff --git a/Engine/Source/Runtime/Engine/Classes/Materials/MaterialParameterCollection.h b/Engine/Source/Runtime/Engine/Classes/Materials/MaterialParameterCollection.h
--- a/Engine/Source/Runtime/Engine/Classes/Materials/MaterialParameterCollection.h
+++ b/Engine/Source/Runtime/Engine/Classes/Materials/MaterialParameterCollection.h
@@ -82,10 +82,10 @@ class UMaterialParameterCollection : public UObject
 	UPROPERTY(duplicatetransient)
 	FGuid StateId;
 
-	UPROPERTY(EditAnywhere, Category=Material)
+	UPROPERTY(EditAnywhere, Category=Material, Meta = (TitleProperty = "ParameterName"))
 	TArray<FCollectionScalarParameter> ScalarParameters;
 
-	UPROPERTY(EditAnywhere, Category=Material)
+	UPROPERTY(EditAnywhere, Category=Material, Meta = (TitleProperty = "ParameterName"))
 	TArray<FCollectionVectorParameter> VectorParameters;
 
 	//~ Begin UObject Interface.
