Move Abs to correct place in PhysicsReplication. AngDiff needs to be signed, but its contribution to Error needs to be positive.

diff --git a/Engine/Source/Runtime/Engine/Private/PhysicsEngine/PhysicsReplication.cpp b/Engine/Source/Runtime/Engine/Private/PhysicsEngine/PhysicsReplication.cpp
--- a/Engine/Source/Runtime/Engine/Private/PhysicsEngine/PhysicsReplication.cpp
+++ b/Engine/Source/Runtime/Engine/Private/PhysicsEngine/PhysicsReplication.cpp
@@ -191,7 +191,8 @@ bool FPhysicsReplication::ApplyRigidBodyState(float DeltaSeconds, FBodyInstance*
 	float AngDiff;
 	const FQuat DeltaQuat = InvCurrentQuat * TargetQuat;
 	DeltaQuat.ToAxisAndAngle(AngDiffAxis, AngDiff);
-	AngDiff = FMath::Abs(FMath::RadiansToDegrees(FMath::UnwindRadians(AngDiff)));
+	AngDiff = FMath::RadiansToDegrees(FMath::UnwindRadians(AngDiff));
+	const float AngDiffSize = FMath::Abs(AngDiff);
 
 	/////// ACCUMULATE ERROR IF NOT APPROACHING SOLUTION ///////
 
@@ -200,7 +201,7 @@ bool FPhysicsReplication::ApplyRigidBodyState(float DeltaSeconds, FBodyInstance*
 	const bool bWasAwake = BI->IsInstanceAwake();
 	const bool bAutoWake = false;
 
-	const float Error = (LinDiffSize * ErrorPerLinearDiff) + (AngDiff * ErrorPerAngularDiff);
+	const float Error = (LinDiffSize * ErrorPerLinearDiff) + (AngDiffSize * ErrorPerAngularDiff);
 	bRestoredState = Error < MaxRestoredStateError;
 	if (bRestoredState)
 	{
