Make the ReplicateActor csv stat in UActorChannel::ReplicateActor conditional on whether we're recording a replay or not - this gives a better understanding of replication time spent recording a replay compared with replicating to players

diff --git a/Engine/Source/Runtime/Engine/Private/DataChannel.cpp b/Engine/Source/Runtime/Engine/Private/DataChannel.cpp
--- a/Engine/Source/Runtime/Engine/Private/DataChannel.cpp
+++ b/Engine/Source/Runtime/Engine/Private/DataChannel.cpp
@@ -2829,7 +2829,6 @@ int32 GNumReplicateActorCalls = 0;
 int64 UActorChannel::ReplicateActor()
 {
 	SCOPE_CYCLE_COUNTER(STAT_NetReplicateActorTime);
-	CSV_SCOPED_TIMING_STAT_EXCLUSIVE(ReplicateActor);
 
 	check(Actor);
 	check(!Closing);
@@ -2845,6 +2844,8 @@ int64 UActorChannel::ReplicateActor()
 #endif
 
 	const bool bReplay = ActorWorld->DemoNetDriver == Connection->GetDriver();
+	CSV_SCOPED_TIMING_STAT_EXCLUSIVE_CONDITIONAL(ReplicateActor, !bReplay);
+
 	const bool bEnableScopedCycleCounter = !bReplay && GReplicateActorTimingEnabled;
 	FSimpleScopeSecondsCounter ScopedSecondsCounter(GReplicateActorTimeSeconds, bEnableScopedCycleCounter);
 
