Fix for static analysis error

diff --git a/Engine/Source/Runtime/Engine/Private/Animation/AnimationUtils.cpp b/Engine/Source/Runtime/Engine/Private/Animation/AnimationUtils.cpp
--- a/Engine/Source/Runtime/Engine/Private/Animation/AnimationUtils.cpp
+++ b/Engine/Source/Runtime/Engine/Private/Animation/AnimationUtils.cpp
@@ -778,6 +778,7 @@ UObject* GetDefaultAnimationCompressionSettings(const TCHAR* IniValueName)
 	if (Value == nullptr)
 	{
 		UE_LOG(LogAnimationCompression, Fatal, TEXT("Couldn't find default compression setting for '%s' under '[Animation.DefaultObjectSettings]'"), IniValueName);
+		return nullptr;
 	}
 
 	const FString& CompressionSettingsName = Value->GetValue();
