Don't include multicast delegates in the UCSModifiedProperties

diff --git a/Engine/Source/Runtime/Engine/Private/Components/ActorComponent.cpp b/Engine/Source/Runtime/Engine/Private/Components/ActorComponent.cpp
--- a/Engine/Source/Runtime/Engine/Private/Components/ActorComponent.cpp
+++ b/Engine/Source/Runtime/Engine/Private/Components/ActorComponent.cpp
@@ -1932,6 +1932,7 @@ void UActorComponent::DetermineUCSModifiedProperties()
 				static const FName MD_SkipUCSModifiedProperties(TEXT("SkipUCSModifiedProperties"));
 				return (InProperty->HasAnyPropertyFlags(CPF_Transient)
 					|| !InProperty->HasAnyPropertyFlags(CPF_Edit | CPF_Interp)
+					|| InProperty->IsA<FMulticastDelegateProperty>()
 #if WITH_EDITOR
 					|| InProperty->HasMetaData(MD_SkipUCSModifiedProperties)
 #endif
