DRPS - Copy from 9388848
#rnx

diff --git a/Engine/Source/Runtime/Engine/Classes/Engine/Texture.h b/Engine/Source/Runtime/Engine/Classes/Engine/Texture.h
--- a/Engine/Source/Runtime/Engine/Classes/Engine/Texture.h
+++ b/Engine/Source/Runtime/Engine/Classes/Engine/Texture.h
@@ -435,6 +435,8 @@ struct FTexturePlatformData
 	uint32 ExtData;
 	/** Number of mips making up the mip tail, which must always be resident */
 	uint32 NumMipsInTail;
+	/** Is Cubemap texture */
+	bool bCubemap;
 	/** Mip data or VT data. one or the other. */
 	TIndirectArray<struct FTexture2DMipMap> Mips;
 	struct FVirtualTextureBuiltData* VTData;
diff --git a/Engine/Source/Runtime/Engine/Private/TextureDerivedData.cpp b/Engine/Source/Runtime/Engine/Private/TextureDerivedData.cpp
--- a/Engine/Source/Runtime/Engine/Private/TextureDerivedData.cpp
+++ b/Engine/Source/Runtime/Engine/Private/TextureDerivedData.cpp
@@ -47,11 +47,11 @@
 // In case of merge conflicts with DDC versions, you MUST generate a new GUID and set this new
 // guid as version
 
-#define TEXTURE_DERIVEDDATA_VER		TEXT("EA79170AE1CF49D5BB1C1CE2FC472B16")
+#define TEXTURE_DERIVEDDATA_VER		TEXT("E2FA6204B41A4C70A1E43E459EEAA510")
 
 // This GUID is mixed into DDC version for virtual textures only, this allows updating DDC version for VT without invalidating DDC for all textures
 // This is useful during development, but once large numbers of VT are present in shipped content, it will have the same problem as TEXTURE_DERIVEDDATA_VER
-#define TEXTURE_VT_DERIVEDDATA_VER	TEXT("D778766A0580478F9AA19D184C0BDB27")
+#define TEXTURE_VT_DERIVEDDATA_VER	TEXT("3CD28A010A9447808ECC76DF7759FBBE")
 
 #if ENABLE_COOK_STATS
 namespace TextureCookStats
@@ -571,7 +571,7 @@ uint32 PutDerivedDataInCache(FTexturePlatformData* DerivedData, const FString& D
 
 	// Write out individual mips to the derived data cache.
 	const int32 MipCount = DerivedData->Mips.Num();
-	const bool bIsCubemap = (DerivedData->NumSlices > 0);
+	const bool bIsCubemap = DerivedData->bCubemap;
 	const int32 FirstInlineMip = bIsCubemap ? 0 : FMath::Max(0, MipCount - FMath::Max((int32)NUM_INLINE_DERIVED_MIPS, (int32)DerivedData->NumMipsInTail));
 	const int32 WritableMipCount = MipCount - ((DerivedData->NumMipsInTail > 0) ? (DerivedData->NumMipsInTail - 1) : 0);
 	for (int32 MipIndex = 0; MipIndex < WritableMipCount; ++MipIndex)
@@ -1014,6 +1014,7 @@ FTexturePlatformData::FTexturePlatformData()
 	, PixelFormat(PF_Unknown)
 	, ExtData(0)
 	, NumMipsInTail(0)
+	, bCubemap(false)
 	, VTData(nullptr)
 #if WITH_EDITORONLY_DATA
 	, AsyncTask(NULL)
diff --git a/Engine/Source/Runtime/Engine/Private/TextureDerivedDataTask.cpp b/Engine/Source/Runtime/Engine/Private/TextureDerivedDataTask.cpp
--- a/Engine/Source/Runtime/Engine/Private/TextureDerivedDataTask.cpp
+++ b/Engine/Source/Runtime/Engine/Private/TextureDerivedDataTask.cpp
@@ -252,6 +252,7 @@ void FTextureCacheDerivedDataWorker::BuildTexture()
 		DerivedData->SizeX = 0;
 		DerivedData->SizeY = 0;
 		DerivedData->PixelFormat = PF_Unknown;
+		DerivedData->bCubemap = false;
 		DerivedData->VTData = nullptr;
 
 		// Compress the texture.
@@ -288,6 +289,7 @@ void FTextureCacheDerivedDataWorker::BuildTexture()
 					DerivedData->SizeY = CompressedImage.SizeY;
 					DerivedData->PixelFormat = (EPixelFormat)CompressedImage.PixelFormat;
 					DerivedData->NumSlices = BuildSettingsPerLayer[0].bCubemap ? 6 : (BuildSettingsPerLayer[0].bVolume || BuildSettingsPerLayer[0].bTextureArray) ? CompressedImage.SizeZ : 1;
+					DerivedData->bCubemap = BuildSettingsPerLayer[0].bCubemap;
 				}
 				else
 				{
