Stopgap fix for Oculus Stereo Layers: FEquirectProps should not be an USTRUCT

diff --git a/Engine/Source/Runtime/Engine/Classes/Components/StereoLayerComponent.h b/Engine/Source/Runtime/Engine/Classes/Components/StereoLayerComponent.h
--- a/Engine/Source/Runtime/Engine/Classes/Components/StereoLayerComponent.h
+++ b/Engine/Source/Runtime/Engine/Classes/Components/StereoLayerComponent.h
@@ -45,6 +45,69 @@ enum EStereoLayerShape
 	SLSH_MAX,
 };
 
+/** Additional properties for equirect layers */
+USTRUCT(BlueprintType)
+struct FEquirectProps
+{
+	GENERATED_USTRUCT_BODY()
+public:
+
+	UPROPERTY(EditAnywhere, BlueprintReadOnly, export, Category = "StereoLayer | Equirect Layer Properties")
+		/** Left source texture UVRect, specifying portion of input texture corresponding to left eye. */
+		FBox2D LeftUVRect;
+
+	UPROPERTY(EditAnywhere, BlueprintReadOnly, export, Category = "StereoLayer | Equirect Layer Properties")
+		/** Right source texture UVRect, specifying portion of input texture corresponding to right eye. */
+		FBox2D RightUVRect;
+
+	UPROPERTY(EditAnywhere, BlueprintReadOnly, export, Category = "StereoLayer | Equirect Layer Properties")
+		/** Left eye's texture coordinate scale after mapping to 2D. */
+		FVector2D LeftScale;
+
+	UPROPERTY(EditAnywhere, BlueprintReadOnly, export, Category = "StereoLayer | Equirect Layer Properties")
+		/** Right eye's texture coordinate scale after mapping to 2D. */
+		FVector2D RightScale;
+
+	UPROPERTY(EditAnywhere, BlueprintReadOnly, export, Category = "StereoLayer | Equirect Layer Properties")
+		/** Left eye's texture coordinate bias after mapping to 2D. */
+		FVector2D LeftBias;
+
+	UPROPERTY(EditAnywhere, BlueprintReadOnly, export, Category = "StereoLayer | Equirect Layer Properties")
+		/** Right eye's texture coordinate bias after mapping to 2D. */
+		FVector2D RightBias;
+
+public:
+
+	FEquirectProps()
+		: LeftUVRect(FBox2D(FVector2D(0.0f, 0.0f), FVector2D(1.0f, 1.0f)))
+		, RightUVRect(FBox2D(FVector2D(0.0f, 0.0f), FVector2D(1.0f, 1.0f)))
+		, LeftScale(FVector2D(1.0f, 1.0f))
+		, RightScale(FVector2D(1.0f, 1.0f))
+		, LeftBias(FVector2D(0.0f, 0.0f))
+		, RightBias(FVector2D(0.0f, 0.0f))
+	{}
+
+	FEquirectProps(FBox2D InLeftUVRect, FBox2D InRightUVRect, FVector2D InLeftScale, FVector2D InRightScale, FVector2D InLeftBias, FVector2D InRightBias)
+		: LeftUVRect(InLeftUVRect)
+		, RightUVRect(InRightUVRect)
+		, LeftScale(InLeftScale)
+		, RightScale(InRightScale)
+		, LeftBias(InLeftBias)
+		, RightBias(InRightBias)
+	{ }
+
+	/**
+	 * Compares two FEquirectProps for equality.
+	 *
+	 * @param Other The other FEquirectProps to compare with.
+	 * @return true if the are equal, false otherwise.
+	 */
+	bool operator==(const FEquirectProps& Other) const
+	{
+		return (LeftUVRect == Other.LeftUVRect) && (RightUVRect == Other.RightUVRect) && (LeftScale == Other.LeftScale) && (RightScale == Other.RightScale) && (LeftBias == Other.LeftBias) && (RightBias == Other.RightBias);
+	}
+};
+
 /** 
  * A geometry layer within the stereo rendered viewport.
  */
diff --git a/Engine/Source/Runtime/Engine/Private/Components/StereoLayerComponent.cpp b/Engine/Source/Runtime/Engine/Private/Components/StereoLayerComponent.cpp
--- a/Engine/Source/Runtime/Engine/Private/Components/StereoLayerComponent.cpp
+++ b/Engine/Source/Runtime/Engine/Private/Components/StereoLayerComponent.cpp
@@ -85,7 +85,6 @@ void UStereoLayerComponent::TickComponent(float DeltaTime, enum ELevelTick TickT
 		LayerDsec.Priority = Priority;
 		LayerDsec.QuadSize = QuadSize;
 		LayerDsec.UVRect = UVRect;
-		LayerDsec.EquirectProps = EquirectProps;
 		LayerDsec.Transform = Transform;
 
 		if (Texture)
@@ -97,18 +96,13 @@ void UStereoLayerComponent::TickComponent(float DeltaTime, enum ELevelTick TickT
 		{
 			LayerDsec.LeftTexture = LeftTexture->Resource->TextureRHI;
 		}
-		LayerDsec.CylinderRadius = CylinderRadius; 
-		LayerDsec.CylinderOverlayArc = CylinderOverlayArc;
-		LayerDsec.CylinderHeight = CylinderHeight;
-			
-			
+				
 		LayerDsec.Flags |= (bLiveTexture) ? IStereoLayers::LAYER_FLAG_TEX_CONTINUOUS_UPDATE : 0;
 		LayerDsec.Flags |= (bNoAlphaChannel) ? IStereoLayers::LAYER_FLAG_TEX_NO_ALPHA_CHANNEL : 0;
 		LayerDsec.Flags |= (bQuadPreserveTextureRatio) ? IStereoLayers::LAYER_FLAG_QUAD_PRESERVE_TEX_RATIO : 0;
 		LayerDsec.Flags |= (bSupportsDepth) ? IStereoLayers::LAYER_FLAG_SUPPORT_DEPTH : 0;
 		LayerDsec.Flags |= (!bCurrVisible) ? IStereoLayers::LAYER_FLAG_HIDDEN : 0;
 
-
 		switch (StereoLayerType)
 		{
 		case SLT_WorldLocked:
@@ -130,6 +124,9 @@ void UStereoLayerComponent::TickComponent(float DeltaTime, enum ELevelTick TickT
 
 		case SLSH_CylinderLayer:
 			LayerDsec.ShapeType = IStereoLayers::CylinderLayer;
+			LayerDsec.CylinderRadius = CylinderRadius;
+			LayerDsec.CylinderOverlayArc = CylinderOverlayArc;
+			LayerDsec.CylinderHeight = CylinderHeight;
 			break;
 
 		case SLSH_CubemapLayer:
@@ -138,6 +135,7 @@ void UStereoLayerComponent::TickComponent(float DeltaTime, enum ELevelTick TickT
 
 		case SLSH_EquirectLayer:
 			LayerDsec.ShapeType = IStereoLayers::EquirectLayer;
+			LayerDsec.EquirectProps = { EquirectProps.LeftUVRect, EquirectProps.RightUVRect, EquirectProps.LeftScale, EquirectProps.RightScale, EquirectProps.LeftBias, EquirectProps.RightBias };
 			break;
 
 		default:
diff --git a/Engine/Source/Runtime/Engine/Public/IStereoLayers.h b/Engine/Source/Runtime/Engine/Public/IStereoLayers.h
--- a/Engine/Source/Runtime/Engine/Public/IStereoLayers.h
+++ b/Engine/Source/Runtime/Engine/Public/IStereoLayers.h
@@ -9,69 +9,6 @@
 
 #include "CoreMinimal.h"
 #include "RHI.h"
-#include "IStereoLayers.generated.h"
-
-USTRUCT(BlueprintType)
-struct FEquirectProps
-{
-	GENERATED_BODY()
-public:
-
-	UPROPERTY(EditAnywhere, BlueprintReadOnly, export, Category = "StereoLayer | Equirect Layer Properties")
-	/** Left source texture UVRect, specifying porion of input texture corresponding to left eye. */
-	FBox2D LeftUVRect;
-
-	UPROPERTY(EditAnywhere, BlueprintReadOnly, export, Category = "StereoLayer | Equirect Layer Properties")
-	/** Right source texture UVRect, specifying porion of input texture corresponding to right eye. */
-	FBox2D RightUVRect;
-
-	UPROPERTY(EditAnywhere, BlueprintReadOnly, export, Category = "StereoLayer | Equirect Layer Properties")
-	/** Left eye's texture coordinate scale after mapping to 2D. */
-	FVector2D LeftScale;
-
-	UPROPERTY(EditAnywhere, BlueprintReadOnly, export, Category = "StereoLayer | Equirect Layer Properties")
-	/** Right eye's texture coordinate scale after mapping to 2D. */
-	FVector2D RightScale;
-
-	UPROPERTY(EditAnywhere, BlueprintReadOnly, export, Category = "StereoLayer | Equirect Layer Properties")
-	/** Left eye's texture coordinate bias after mapping to 2D. */
-	FVector2D LeftBias;
-
-	UPROPERTY(EditAnywhere, BlueprintReadOnly, export, Category = "StereoLayer | Equirect Layer Properties")
-	/** Right eye's texture coordinate bias after mapping to 2D. */
-	FVector2D RightBias;
-
-public:
-
-	FEquirectProps()
-		:	LeftUVRect(FBox2D(FVector2D(0.0f, 0.0f), FVector2D(1.0f, 1.0f)))
-		,	RightUVRect(FBox2D(FVector2D(0.0f, 0.0f), FVector2D(1.0f, 1.0f)))
-		,	LeftScale(FVector2D(1.0f, 1.0f))
-		,	RightScale(FVector2D(1.0f, 1.0f))
-		,	LeftBias(FVector2D(0.0f, 0.0f))
-		,	RightBias(FVector2D(0.0f, 0.0f))
-	{}
-
-	FEquirectProps(FBox2D InLeftUVRect, FBox2D InRightUVRect, FVector2D InLeftScale, FVector2D InRightScale, FVector2D InLeftBias, FVector2D InRightBias)
-		:	LeftUVRect(InLeftUVRect)
-		,	RightUVRect(InRightUVRect)
-		,	LeftScale(InLeftScale)
-		,	RightScale(InRightScale)
-		,	LeftBias(InLeftBias)
-		,	RightBias(InRightBias)
-	{ }
-
-	/**
-	 * Compares two FEquirectProps for equality.
-	 *
-	 * @param Other The other FEquirectProps to compare with.
-	 * @return true if the are equal, false otherwise.
-	 */
-	bool operator==(const FEquirectProps& Other) const
-	{
-		return (LeftUVRect == Other.LeftUVRect) && (RightUVRect == Other.RightUVRect) && (LeftScale == Other.LeftScale) && (RightScale == Other.RightScale) && (LeftBias == Other.LeftBias) && (RightBias == Other.RightBias);
-	}
-};
 
 class IStereoLayers
 {
@@ -108,6 +45,33 @@ public:
 		LAYER_FLAG_HIDDEN = 0x00000020,
 	};
 
+	/** Structure describing additional settings for equirect layers */
+	struct FEquirectProps
+	{
+		/** Left source texture UVRect, specifying portion of input texture corresponding to left eye. */
+		FBox2D LeftUVRect;
+
+		/** Right source texture UVRect, specifying portion of input texture corresponding to right eye. */
+		FBox2D RightUVRect;
+
+		/** Left eye's texture coordinate scale after mapping to 2D. */
+		FVector2D LeftScale;
+
+		/** Right eye's texture coordinate scale after mapping to 2D. */
+		FVector2D RightScale;
+
+		/** Left eye's texture coordinate bias after mapping to 2D. */
+		FVector2D LeftBias;
+
+		/** Right eye's texture coordinate bias after mapping to 2D. */
+		FVector2D RightBias;
+	};
+
+	/** Structure describing additional settings for cylindrical layers */
+	struct FCylinderProps
+	{
+	};
+
 	/**
 	 * Structure describing the visual appearance of a single stereo layer
 	 */
@@ -127,15 +91,8 @@ public:
 		FVector2D			QuadSize	 = FVector2D(1.0f, 1.0f);
 		// UVs of rendered quad in UE units
 		FBox2D				UVRect		 = FBox2D(FVector2D(0.0f, 0.0f), FVector2D(1.0f, 1.0f));
-		//UVs and Scale/Bias of Equirect Layers
-		FEquirectProps		EquirectProps = FEquirectProps(FBox2D(FVector2D(0.0f, 0.0f), FVector2D(1.0f, 1.0f)), FBox2D(FVector2D(0.0f, 0.0f), FVector2D(1.0f, 1.0f)), FVector2D(1.0f, 1.0f), FVector2D(1.0f, 1.0f), FVector2D(0.0f, 0.0f), FVector2D(0.0f, 0.0f));
-
-		FVector2D			LeftScale = FVector2D(1.0f, 1.0f);
-		FVector2D			LeftBias = FVector2D(0.0f, 0.0f);
-		FVector2D			RightScale = FVector2D(1.0f, 1.0f);
-		FVector2D			RightBias = FVector2D(0.0f, 0.0f);
 
-		// Size of texture that the compositor should allocate. Un-necessary if Texture is provided. The compositor will allocate a cubemap whose faces are of LayerSize if ShapeType is CubemapLayer.
+		// Size of texture that the compositor should allocate. Unnecessary if Texture is provided. The compositor will allocate a cubemap whose faces are of LayerSize if ShapeType is CubemapLayer.
 		FIntPoint			LayerSize = FIntPoint(0, 0);
 		// Render order priority, higher priority render on top of lower priority. Face-Locked layers are rendered on top of other layer types regardless of priority. 
 		int32				Priority	 = 0;
@@ -143,9 +100,14 @@ public:
 		ELayerType			PositionType = ELayerType::FaceLocked;
 		// which shape of layer it is. ELayerShape::QuadLayer is the only shape supported by all VR platforms.
 		ELayerShape			ShapeType	 = ELayerShape::QuadLayer;
-		float				CylinderRadius		= 1.0f;
-		float				CylinderOverlayArc	= 1.0f;
-		float				CylinderHeight = 1.0f;
+
+		// UVs and Scale/Bias of Equirect Layers.
+		FEquirectProps		EquirectProps;
+		// Cylinder layer settings.
+		float				CylinderRadius;
+		float				CylinderOverlayArc;
+		float				CylinderHeight;
+
 		// Texture mapped for right eye (if one texture provided, mono assumed)
 		FTextureRHIRef		Texture		 = nullptr;	
 		// Texture mapped for left eye (if one texture provided, mono assumed)
