Removing empty struct accidentally added in CL#9470776

diff --git a/Engine/Source/Runtime/Engine/Public/IStereoLayers.h b/Engine/Source/Runtime/Engine/Public/IStereoLayers.h
--- a/Engine/Source/Runtime/Engine/Public/IStereoLayers.h
+++ b/Engine/Source/Runtime/Engine/Public/IStereoLayers.h
@@ -67,11 +67,6 @@ public:
 		FVector2D RightBias;
 	};
 
-	/** Structure describing additional settings for cylindrical layers */
-	struct FCylinderProps
-	{
-	};
-
 	/**
 	 * Structure describing the visual appearance of a single stereo layer
 	 */
