Culling enabled by default for ISM/HISM in ray tracing

diff --git a/Engine/Source/Runtime/Engine/Private/InstancedStaticMesh.cpp b/Engine/Source/Runtime/Engine/Private/InstancedStaticMesh.cpp
--- a/Engine/Source/Runtime/Engine/Private/InstancedStaticMesh.cpp
+++ b/Engine/Source/Runtime/Engine/Private/InstancedStaticMesh.cpp
@@ -60,8 +60,8 @@ static TAutoConsoleVariable<int32> CVarRayTracingRenderInstances(
 
 static TAutoConsoleVariable<int32> CVarRayTracingRenderInstancesCulling(
 	TEXT("r.RayTracing.InstancedStaticMeshes.Culling"),
-	0,
-	TEXT("Enable culling for instances in ray tracing (default = 0 (Culling disabled))"));
+	1,
+	TEXT("Enable culling for instances in ray tracing (default = 1 (Culling enabled))"));
 
 static TAutoConsoleVariable<float> CVarRayTracingInstancesCullClusterMaxRadiusMultiplier(
 	TEXT("r.RayTracing.InstancedStaticMeshes.CullClusterMaxRadiusMultiplier"),
