Using a flag instead of the element count in the ray tracing geometry array to deduce whether it is dynamic.

(ushell-p4-cherrypick of 9481162 by [REDACTED_PERSON])

diff --git a/Engine/Source/Runtime/Engine/Private/StaticMeshRender.cpp b/Engine/Source/Runtime/Engine/Private/StaticMeshRender.cpp
--- a/Engine/Source/Runtime/Engine/Private/StaticMeshRender.cpp
+++ b/Engine/Source/Runtime/Engine/Private/StaticMeshRender.cpp
@@ -194,9 +194,10 @@ FStaticMeshSceneProxy::FStaticMeshSceneProxy(UStaticMeshComponent* InComponent,
 	const bool bLODsShareStaticLighting = RenderData->bLODsShareStaticLighting || bForceLODsShareStaticLighting;
 
 #if RHI_RAYTRACING
+	bDynamicRayTracingGeometry = InComponent->bEvaluateWorldPositionOffset && MaterialRelevance.bUsesWorldPositionOffset;
 	if (IsRayTracingEnabled())
 	{
-		if(InComponent->bEvaluateWorldPositionOffset && MaterialRelevance.bUsesWorldPositionOffset)
+		if(bDynamicRayTracingGeometry)
 		{
 			DynamicRayTracingGeometries.AddDefaulted(RenderData->LODResources.Num());
 			for (int32 LODIndex = 0; LODIndex < RenderData->LODResources.Num(); LODIndex++)
diff --git a/Engine/Source/Runtime/Engine/Public/StaticMeshResources.h b/Engine/Source/Runtime/Engine/Public/StaticMeshResources.h
--- a/Engine/Source/Runtime/Engine/Public/StaticMeshResources.h
+++ b/Engine/Source/Runtime/Engine/Public/StaticMeshResources.h
@@ -865,7 +865,7 @@ public:
 	virtual bool IsRayTracingStaticRelevant() const override 
 	{ 
 		const bool bAllowStaticLighting = FReadOnlyCVARCache::Get().bAllowStaticLighting;
-		const bool bIsStaticInstance = RayTracingGeometries.Num() > 0;
+		const bool bIsStaticInstance = !bDynamicRayTracingGeometry;
 		return bIsStaticInstance && IsStaticPathAvailable() && !HasViewDependentDPG() && !(bAllowStaticLighting && HasStaticLighting() && !HasValidSettingsForStaticLighting());
 	}
 #endif // RHI_RAYTRACING
@@ -956,6 +956,7 @@ protected:
 	const FDistanceFieldVolumeData* DistanceFieldData;	
 
 #if RHI_RAYTRACING
+	bool bDynamicRayTracingGeometry;
 	TArray<FRayTracingGeometry, TInlineAllocator<MAX_MESH_LOD_COUNT>> DynamicRayTracingGeometries;
 	TArray<FRWBuffer, TInlineAllocator<MAX_MESH_LOD_COUNT>> DynamicRayTracingGeometryVertexBuffers;
 #endif
