Extend FOverrideSelectionColorMaterialRenderProxy to support overriding BaseColor as well as EmissiveColor, fixes Lighting mode crashes on tessellated geometry

diff --git a/Engine/Source/Runtime/Engine/Private/Materials/Material.cpp b/Engine/Source/Runtime/Engine/Private/Materials/Material.cpp
--- a/Engine/Source/Runtime/Engine/Private/Materials/Material.cpp
+++ b/Engine/Source/Runtime/Engine/Private/Materials/Material.cpp
@@ -120,12 +120,16 @@ int32 FMaterialResource::CompilePropertyAndSetMaterialProperty(EMaterialProperty
 	EShaderFrequency ShaderFrequency = Compiler->GetCurrentShaderFrequency();
 	
 	int32 SelectionColorIndex = INDEX_NONE;
+	int32 SelectionColorToggle = INDEX_NONE;
 
 	if (ShaderFrequency == SF_Pixel &&
 		GetMaterialDomain() != MD_Volume &&
 		Compiler->IsDevelopmentFeatureEnabled(NAME_SelectionColor))
 	{
-		SelectionColorIndex = Compiler->ComponentMask(Compiler->VectorParameter(NAME_SelectionColor, FLinearColor::Black), 1, 1, 1, 0);
+		// RGB stores SelectionColor value, A is toggle on/off switch for SelectionColor
+		int32 SelectionColorVector = Compiler->VectorParameter(NAME_SelectionColor, FLinearColor::Transparent);
+		SelectionColorIndex = Compiler->ComponentMask(SelectionColorVector, 1, 1, 1, 0);
+		SelectionColorToggle = Compiler->ComponentMask(SelectionColorVector, 0, 0, 0, 1);
 	}
 
 	//Compile the material instance if we have one.
@@ -138,7 +142,8 @@ int32 FMaterialResource::CompilePropertyAndSetMaterialProperty(EMaterialProperty
 		case MP_EmissiveColor:
 			if (SelectionColorIndex != INDEX_NONE)
 			{
-				Ret = Compiler->Add(MaterialInterface->CompileProperty(Compiler, MP_EmissiveColor, MFCF_ForceCast), SelectionColorIndex);
+				// Alpha channel is used to as toggle between EmissiveColor and SelectionColor
+				Ret = Compiler->Lerp(MaterialInterface->CompileProperty(Compiler, MP_EmissiveColor, MFCF_ForceCast), SelectionColorIndex, SelectionColorToggle);
 			}
 			else
 			{
diff --git a/Engine/Source/Runtime/Engine/Private/PrimitiveDrawingUtils.cpp b/Engine/Source/Runtime/Engine/Private/PrimitiveDrawingUtils.cpp
--- a/Engine/Source/Runtime/Engine/Private/PrimitiveDrawingUtils.cpp
+++ b/Engine/Source/Runtime/Engine/Private/PrimitiveDrawingUtils.cpp
@@ -1495,14 +1495,18 @@ void ApplyViewModeOverrides(
 
 			if (bTextureMapped == false)
 			{
-				FMaterialRenderProxy* RenderProxy = GEngine->LevelColorationLitMaterial->GetRenderProxy();
-				auto LightingOnlyMaterialInstance = new FColoredMaterialRenderProxy(
-					RenderProxy,
-					GEngine->LightingOnlyBrightness
+				// Tessellated geometry cannot use engine debug LevelColorationLitMaterial as it doesn't support tessellation
+				if (!bMaterialModifiesMeshPosition)
+				{
+					FMaterialRenderProxy* RenderProxy = GEngine->LevelColorationLitMaterial->GetRenderProxy();
+					auto LightingOnlyMaterialInstance = new FColoredMaterialRenderProxy(
+						RenderProxy,
+						GEngine->LightingOnlyBrightness
 					);
 
-				Mesh.MaterialRenderProxy = LightingOnlyMaterialInstance;
-				Collector.RegisterOneFrameMaterialProxy(LightingOnlyMaterialInstance);
+					Mesh.MaterialRenderProxy = LightingOnlyMaterialInstance;
+					Collector.RegisterOneFrameMaterialProxy(LightingOnlyMaterialInstance);
+				}
 			}
 			else
 			{
diff --git a/Engine/Source/Runtime/Engine/Private/ShowFlags.cpp b/Engine/Source/Runtime/Engine/Private/ShowFlags.cpp
--- a/Engine/Source/Runtime/Engine/Private/ShowFlags.cpp
+++ b/Engine/Source/Runtime/Engine/Private/ShowFlags.cpp
@@ -307,7 +307,7 @@ void ApplyViewMode(EViewModeIndex ViewModeIndex, bool bPerspective, FEngineShowF
 	// Assigning the new state like this ensures we always set the same variables (they depend on the view mode)
 	// This is affecting the state of showflags - if the state can be changed by the user as well it should better be done in EngineShowFlagOverride
 
-	EngineShowFlags.SetOverrideDiffuseAndSpecular(ViewModeIndex == VMI_Lit_DetailLighting);
+	EngineShowFlags.SetOverrideDiffuseAndSpecular(ViewModeIndex == VMI_Lit_DetailLighting || ViewModeIndex == VMI_LightingOnly);
 	EngineShowFlags.SetReflectionOverride(ViewModeIndex == VMI_ReflectionOverride);
 	EngineShowFlags.SetVisualizeBuffer(ViewModeIndex == VMI_VisualizeBuffer);
 	EngineShowFlags.SetVisualizeLightCulling(ViewModeIndex == VMI_LightComplexity);
diff --git a/Engine/Source/Runtime/Engine/Public/MaterialShared.h b/Engine/Source/Runtime/Engine/Public/MaterialShared.h
--- a/Engine/Source/Runtime/Engine/Public/MaterialShared.h
+++ b/Engine/Source/Runtime/Engine/Public/MaterialShared.h
@@ -2113,7 +2113,7 @@ public:
 	/** Initialization constructor. */
 	FOverrideSelectionColorMaterialRenderProxy(const FMaterialRenderProxy* InParent, const FLinearColor& InSelectionColor) :
 		Parent(InParent),
-		SelectionColor(InSelectionColor)
+		SelectionColor(FLinearColor(InSelectionColor.R, InSelectionColor.G, InSelectionColor.B, 1))
 	{
 	}
 
