Backout CL 8327634 and set INVALID_CONTROLLERID to -1. CL 8327634 was a workaround for a startup crash which is now resolved. ShooterGame boots properly with this change.

diff --git a/Engine/Source/Runtime/Engine/Classes/Engine/LocalPlayer.h b/Engine/Source/Runtime/Engine/Classes/Engine/LocalPlayer.h
--- a/Engine/Source/Runtime/Engine/Classes/Engine/LocalPlayer.h
+++ b/Engine/Source/Runtime/Engine/Classes/Engine/LocalPlayer.h
@@ -23,7 +23,7 @@
 #include "LocalPlayer.generated.h"
 
 
-#define INVALID_CONTROLLERID 255
+#define INVALID_CONTROLLERID (-1)
 
 class AActor;
 class FSceneView;
diff --git a/Engine/Source/Runtime/Engine/Private/LocalPlayer.cpp b/Engine/Source/Runtime/Engine/Private/LocalPlayer.cpp
--- a/Engine/Source/Runtime/Engine/Private/LocalPlayer.cpp
+++ b/Engine/Source/Runtime/Engine/Private/LocalPlayer.cpp
@@ -255,7 +255,7 @@ void ULocalPlayer::PostInitProperties()
 void ULocalPlayer::PlayerAdded(UGameViewportClient* InViewportClient, int32 InControllerID)
 {
 	ViewportClient = InViewportClient;
-	ControllerId = InControllerID;
+	SetControllerId(InControllerID);
 
 	SubsystemCollection.Initialize(this);
 }
