Fix static analysis warnings. Revert CL 9421227 and remove the null checks on Owner->PlatformData. These checks were superfluous as the original code paths assumed that PlatformData was never nullptr. Adding the checks makes static analysis believe that the existing code is now wrong, since it wasn't checking for nullptr.

diff --git a/Engine/Source/Runtime/Engine/Private/Texture2D.cpp b/Engine/Source/Runtime/Engine/Private/Texture2D.cpp
--- a/Engine/Source/Runtime/Engine/Private/Texture2D.cpp
+++ b/Engine/Source/Runtime/Engine/Private/Texture2D.cpp
@@ -1602,7 +1602,7 @@ void FTexture2DResource::InitRHI()
 				TexCreateFlags |= TexCreate_Virtual;
 
 				FRHIResourceCreateInfo CreateInfo(ResourceMem);
-				CreateInfo.ExtData = Owner->PlatformData ? Owner->PlatformData->ExtData : 0;
+				CreateInfo.ExtData = Owner->PlatformData->ExtData;
 				Texture2DRHI = RHICreateTexture2D( OwnerMips[0].SizeX, OwnerMips[0].SizeY, EffectiveFormat, OwnerMips.Num(), 1, TexCreateFlags, CreateInfo);
 				RHIVirtualTextureSetFirstMipInMemory(Texture2DRHI, CurrentFirstMip);
 				RHIVirtualTextureSetFirstMipVisible(Texture2DRHI, CurrentFirstMip);
@@ -1638,14 +1638,14 @@ void FTexture2DResource::InitRHI()
 
 			// create texture with ResourceMem data when available
 			FRHIResourceCreateInfo CreateInfo(ResourceMem);
-			CreateInfo.ExtData = Owner->PlatformData ? Owner->PlatformData->ExtData : 0;
+			CreateInfo.ExtData = Owner->PlatformData->ExtData;
 			Texture2DRHI	= RHICreateTexture2D( SizeX, SizeY, EffectiveFormat, RequestedMips, 1, TexCreateFlags, CreateInfo);
 			TextureRHI		= Texture2DRHI;
 			TextureRHI->SetName(Owner->GetFName());
 			RHIBindDebugLabelName(TextureRHI, *Owner->GetName());
 			RHIUpdateTextureReference(Owner->TextureReference.TextureReferenceRHI,TextureRHI);
 
-			check(!Owner->PlatformData || (Owner->PlatformData->Mips[CurrentFirstMip].SizeX == Texture2DRHI->GetSizeX() && Owner->PlatformData->Mips[CurrentFirstMip].SizeY == Texture2DRHI->GetSizeY()));
+			check(Owner->PlatformData->Mips[CurrentFirstMip].SizeX == Texture2DRHI->GetSizeX() && Owner->PlatformData->Mips[CurrentFirstMip].SizeY == Texture2DRHI->GetSizeY());
 
 			if( ResourceMem )
 			{
@@ -1689,7 +1689,7 @@ void FTexture2DResource::InitRHI()
 		if (GIsEditor || (!bSkipRHITextureCreation)) //-V560
 		{
 			FRHIResourceCreateInfo CreateInfo;
-			CreateInfo.ExtData = Owner->PlatformData ? Owner->PlatformData->ExtData : 0;
+			CreateInfo.ExtData = Owner->PlatformData->ExtData;
 			Texture2DRHI	= RHICreateTexture2D( SizeX, SizeY, EffectiveFormat, RequestedMips, 1, TexCreateFlags, CreateInfo );
 			TextureRHI		= Texture2DRHI;
 			TextureRHI->SetName(Owner->GetFName());
