DRPS - Missing data in texture DDC
#rnx

diff --git a/Engine/Source/Runtime/Engine/Private/TextureDerivedData.cpp b/Engine/Source/Runtime/Engine/Private/TextureDerivedData.cpp
--- a/Engine/Source/Runtime/Engine/Private/TextureDerivedData.cpp
+++ b/Engine/Source/Runtime/Engine/Private/TextureDerivedData.cpp
@@ -47,11 +47,11 @@
 // In case of merge conflicts with DDC versions, you MUST generate a new GUID and set this new
 // guid as version
 
-#define TEXTURE_DERIVEDDATA_VER		TEXT("E2FA6204B41A4C70A1E43E459EEAA510")
+#define TEXTURE_DERIVEDDATA_VER		TEXT("51C887FFAD4D4179AECB2FA5D9D6A0F8")
 
 // This GUID is mixed into DDC version for virtual textures only, this allows updating DDC version for VT without invalidating DDC for all textures
 // This is useful during development, but once large numbers of VT are present in shipped content, it will have the same problem as TEXTURE_DERIVEDDATA_VER
-#define TEXTURE_VT_DERIVEDDATA_VER	TEXT("3CD28A010A9447808ECC76DF7759FBBE")
+#define TEXTURE_VT_DERIVEDDATA_VER	TEXT("1F1EB5C2A6054EAEB31E38E85AAD1A33")
 
 #if ENABLE_COOK_STATS
 namespace TextureCookStats
@@ -1265,6 +1265,7 @@ static void SerializePlatformData(
 	}
  	Ar << PlatformData->ExtData;
  	Ar << PlatformData->NumMipsInTail;
+	Ar << PlatformData->bCubemap;
 
 	int32 NumMips = PlatformData->Mips.Num();
 	int32 FirstMipToSerialize = 0;
