Fix FSkeletalMeshObjectGPUSkin::RayTracingDynamicVertexBuffer not being released properly before destruction

diff --git a/Engine/Source/Runtime/Engine/Private/SkeletalRenderGPUSkin.cpp b/Engine/Source/Runtime/Engine/Private/SkeletalRenderGPUSkin.cpp
--- a/Engine/Source/Runtime/Engine/Private/SkeletalRenderGPUSkin.cpp
+++ b/Engine/Source/Runtime/Engine/Private/SkeletalRenderGPUSkin.cpp
@@ -225,6 +225,11 @@ void FSkeletalMeshObjectGPUSkin::ReleaseResources()
 
 #if RHI_RAYTRACING
 	BeginReleaseResource(&RayTracingGeometry);
+	ENQUEUE_RENDER_COMMAND(ReleaseRayTracingDynamicVertexBuffer)(
+		[RayTracingDynamicVertexBuffer = MoveTemp(RayTracingDynamicVertexBuffer)](FRHICommandListImmediate& RHICmdList) mutable
+	{
+		RayTracingDynamicVertexBuffer.Release();
+	});
 #endif
 }
 
