Refactor stereo layer shapes.
Part2, make the StereoLayer component design pluggable as the stereo layers themselves.
Includes serialization changes to upgrade existing levels when loaded in the editor.
Also moves visualization rendering into the pluggable shape objects to allow customizations of those.

diff --git a/Engine/Source/Editor/ComponentVisualizers/Private/StereoLayerComponentVisualizer.cpp b/Engine/Source/Editor/ComponentVisualizers/Private/StereoLayerComponentVisualizer.cpp
--- a/Engine/Source/Editor/ComponentVisualizers/Private/StereoLayerComponentVisualizer.cpp
+++ b/Engine/Source/Editor/ComponentVisualizers/Private/StereoLayerComponentVisualizer.cpp
@@ -9,36 +9,9 @@
 void FStereoLayerComponentVisualizer::DrawVisualization( const UActorComponent* Component, const FSceneView* View, FPrimitiveDrawInterface* PDI )
 {
 	const UStereoLayerComponent* StereoLayerComp = Cast<const UStereoLayerComponent>(Component);
-	if(StereoLayerComp != NULL)
+	if(StereoLayerComp != nullptr)
 	{
-        FLinearColor YellowColor = FColor(231, 239, 0, 255);
-        if(StereoLayerComp->StereoLayerShape == EStereoLayerShape::SLSH_QuadLayer)
-        {
-            const FVector2D QuadSize = StereoLayerComp->GetQuadSize() / 2.0f;
-            const FBox QuadBox(FVector(0.0f, -QuadSize.X, -QuadSize.Y), FVector(0.0f, QuadSize.X, QuadSize.Y));
-
-            DrawWireBox(PDI, StereoLayerComp->GetComponentTransform().ToMatrixWithScale(), QuadBox, YellowColor, 0);
-        }
-        else if(StereoLayerComp->StereoLayerShape == EStereoLayerShape::SLSH_CylinderLayer)
-        {
-            float ArcAngle = StereoLayerComp->CylinderOverlayArc * 180 / (StereoLayerComp->CylinderRadius * PI);
-            
-            FVector X = StereoLayerComp->GetComponentTransform().GetUnitAxis(EAxis::Type::X);
-            FVector Y = StereoLayerComp->GetComponentTransform().GetUnitAxis(EAxis::Type::Y);
-            FVector Base = StereoLayerComp->GetComponentTransform().GetLocation();
-            FVector HalfHeight = FVector(0, 0, StereoLayerComp->CylinderHeight/2);
-            
-            FVector LeftVertex = Base + StereoLayerComp->CylinderRadius * ( FMath::Cos(ArcAngle/2 * (PI/180.0f)) * X + FMath::Sin(ArcAngle/2 * (PI/180.0f)) * Y );
-            FVector RightVertex = Base + StereoLayerComp->CylinderRadius * ( FMath::Cos(-ArcAngle/2 * (PI/180.0f)) * X + FMath::Sin(-ArcAngle/2 * (PI/180.0f)) * Y );
-
-            DrawArc(PDI, Base + HalfHeight, X, Y, -ArcAngle/2, ArcAngle/2, StereoLayerComp->CylinderRadius, 10, YellowColor, 0);
-            
-            DrawArc(PDI, Base - HalfHeight, X, Y, -ArcAngle/2, ArcAngle/2, StereoLayerComp->CylinderRadius, 10, YellowColor, 0);
-
-            PDI->DrawLine( LeftVertex - HalfHeight, LeftVertex + HalfHeight, YellowColor, 0 );
-            
-            PDI->DrawLine( RightVertex - HalfHeight, RightVertex + HalfHeight, YellowColor, 0 );
-
-        }
+		check(StereoLayerComp->Shape != nullptr);
+		StereoLayerComp->Shape->DrawShapeVisualization(View, PDI);
     }
 }
diff --git a/Engine/Source/Runtime/Core/Public/UObject/VRObjectVersion.h b/Engine/Source/Runtime/Core/Public/UObject/VRObjectVersion.h
--- a/Engine/Source/Runtime/Core/Public/UObject/VRObjectVersion.h
+++ b/Engine/Source/Runtime/Core/Public/UObject/VRObjectVersion.h
@@ -18,6 +18,9 @@ struct CORE_API FVRObjectVersion
 		// Change how ARSessionConfig stores plane detection configuration from bitmask to bools
 		UseBoolsForARSessionConfigPlaneDetectionConfiguration,
 
+		// Change how UStereoLayerComponent stores additional properties for non-quad layer types
+		UseSubobjectForStereoLayerShapeProperties,
+
 		// -----<new versions can be added above this line>-------------------------------------------------
 		VersionPlusOne,
 		LatestVersion = VersionPlusOne - 1
diff --git a/Engine/Source/Runtime/Engine/Classes/Components/StereoLayerComponent.h b/Engine/Source/Runtime/Engine/Classes/Components/StereoLayerComponent.h
--- a/Engine/Source/Runtime/Engine/Classes/Components/StereoLayerComponent.h
+++ b/Engine/Source/Runtime/Engine/Classes/Components/StereoLayerComponent.h
@@ -45,7 +45,78 @@ enum EStereoLayerShape
 	SLSH_MAX,
 };
 
-/** Additional properties for equirect layers */
+UCLASS(EditInlineNew, Abstract, BlueprintType, CollapseCategories, Within = StereoLayerComponent)
+class UStereoLayerShape : public UObject
+{
+	GENERATED_BODY()
+
+public:
+	virtual void ApplyShape(IStereoLayers::FLayerDesc& LayerDesc);
+#if WITH_EDITOR
+	virtual void DrawShapeVisualization(const class FSceneView* View, class FPrimitiveDrawInterface* PDI);
+#endif
+
+protected:
+	void MarkStereoLayerDirty();
+};
+
+UCLASS(meta = (DisplayName = "Quad Layer"))
+class UStereoLayerShapeQuad : public UStereoLayerShape
+{
+	GENERATED_BODY()
+public:
+	virtual void ApplyShape(IStereoLayers::FLayerDesc& LayerDesc) override;
+#if WITH_EDITOR
+	virtual void DrawShapeVisualization(const class FSceneView* View, class FPrimitiveDrawInterface* PDI) override;
+#endif
+};
+
+UCLASS(meta = (DisplayName = "Cylinder Layer"))
+class UStereoLayerShapeCylinder : public UStereoLayerShape
+{
+	GENERATED_BODY()
+
+public:
+	UStereoLayerShapeCylinder()
+		: Radius(100)
+		, OverlayArc(100)
+		, Height(50)
+	{}
+
+	/** Radial size of the rendered stereo layer cylinder **/
+	UPROPERTY(EditAnywhere, BlueprintReadOnly, export, Category = "Cylinder Properties")
+	float Radius;
+
+	/** Arc angle for the stereo layer cylinder **/
+	UPROPERTY(EditAnywhere, BlueprintReadOnly, export, Category = "Cylinder Properties")
+	float OverlayArc;
+
+	/** Height of the stereo layer cylinder **/
+	UPROPERTY(EditAnywhere, BlueprintReadOnly, export, Category = "Cylinder Properties")
+	int Height;
+
+	UFUNCTION(BlueprintCallable, Category = "Components|Stereo Layer")
+	void SetRadius(float InRadius);
+	UFUNCTION(BlueprintCallable, Category = "Components|Stereo Layer")
+	void SetOverlayArc(float InOverlayArc);
+	UFUNCTION(BlueprintCallable, Category = "Components|Stereo Layer")
+	void SetHeight(int InHeight);
+
+	virtual void ApplyShape(IStereoLayers::FLayerDesc& LayerDesc) override;
+#if WITH_EDITOR
+	virtual void DrawShapeVisualization(const class FSceneView* View, class FPrimitiveDrawInterface* PDI) override;
+#endif
+};
+
+UCLASS(meta = (DisplayName = "Cubemap Layer"))
+class UStereoLayerShapeCubemap : public UStereoLayerShape
+{
+	GENERATED_BODY()
+public:
+	virtual void ApplyShape(IStereoLayers::FLayerDesc& LayerDesc) override;
+};
+
+/** Properties for equirect layers */
 USTRUCT(BlueprintType)
 struct FEquirectProps
 {
@@ -53,28 +124,28 @@ struct FEquirectProps
 public:
 
 	UPROPERTY(EditAnywhere, BlueprintReadOnly, export, Category = "StereoLayer | Equirect Layer Properties")
-		/** Left source texture UVRect, specifying portion of input texture corresponding to left eye. */
-		FBox2D LeftUVRect;
+	/** Left source texture UVRect, specifying portion of input texture corresponding to left eye. */
+	FBox2D LeftUVRect;
 
 	UPROPERTY(EditAnywhere, BlueprintReadOnly, export, Category = "StereoLayer | Equirect Layer Properties")
-		/** Right source texture UVRect, specifying portion of input texture corresponding to right eye. */
-		FBox2D RightUVRect;
+	/** Right source texture UVRect, specifying portion of input texture corresponding to right eye. */
+	FBox2D RightUVRect;
 
 	UPROPERTY(EditAnywhere, BlueprintReadOnly, export, Category = "StereoLayer | Equirect Layer Properties")
-		/** Left eye's texture coordinate scale after mapping to 2D. */
-		FVector2D LeftScale;
+	/** Left eye's texture coordinate scale after mapping to 2D. */
+	FVector2D LeftScale;
 
 	UPROPERTY(EditAnywhere, BlueprintReadOnly, export, Category = "StereoLayer | Equirect Layer Properties")
-		/** Right eye's texture coordinate scale after mapping to 2D. */
-		FVector2D RightScale;
+	/** Right eye's texture coordinate scale after mapping to 2D. */
+	FVector2D RightScale;
 
 	UPROPERTY(EditAnywhere, BlueprintReadOnly, export, Category = "StereoLayer | Equirect Layer Properties")
-		/** Left eye's texture coordinate bias after mapping to 2D. */
-		FVector2D LeftBias;
+	/** Left eye's texture coordinate bias after mapping to 2D. */
+	FVector2D LeftBias;
 
 	UPROPERTY(EditAnywhere, BlueprintReadOnly, export, Category = "StereoLayer | Equirect Layer Properties")
-		/** Right eye's texture coordinate bias after mapping to 2D. */
-		FVector2D RightBias;
+	/** Right eye's texture coordinate bias after mapping to 2D. */
+	FVector2D RightBias;
 
 public:
 
@@ -102,10 +173,68 @@ public:
 	 * @param Other The other FEquirectProps to compare with.
 	 * @return true if the are equal, false otherwise.
 	 */
-	bool operator==(const FEquirectProps& Other) const
-	{
-		return (LeftUVRect == Other.LeftUVRect) && (RightUVRect == Other.RightUVRect) && (LeftScale == Other.LeftScale) && (RightScale == Other.RightScale) && (LeftBias == Other.LeftBias) && (RightBias == Other.RightBias);
-	}
+	bool operator==(const FEquirectProps& Other) const;
+
+	/**
+	 * Compares FEquirectProps with an UStereoLayerShapeEquirect
+	 *
+	 * @param Other The UStereoLayerShapeEquirect to compare with.
+	 * @return true if the are equal, false otherwise.
+	 */
+	bool operator==(const class UStereoLayerShapeEquirect& Other) const;
+};
+
+
+UCLASS(meta = (DisplayName = "Equirect Layer"))
+class UStereoLayerShapeEquirect : public UStereoLayerShape
+{
+	GENERATED_BODY()
+
+public:
+	UStereoLayerShapeEquirect()
+		: LeftUVRect(FBox2D(FVector2D(0.0f, 0.0f), FVector2D(1.0f, 1.0f)))
+		, RightUVRect(FBox2D(FVector2D(0.0f, 0.0f), FVector2D(1.0f, 1.0f)))
+		, LeftScale(FVector2D(1.0f, 1.0f))
+		, RightScale(FVector2D(1.0f, 1.0f))
+		, LeftBias(FVector2D(0.0f, 0.0f))
+		, RightBias(FVector2D(0.0f, 0.0f))
+	{}
+
+	UPROPERTY(EditAnywhere, BlueprintReadOnly, Category = "Equirect Properties")
+	/** Left source texture UVRect, specifying portion of input texture corresponding to left eye. */
+	FBox2D LeftUVRect;
+
+	UPROPERTY(EditAnywhere, BlueprintReadOnly, Category = "Equirect Properties")
+	/** Right source texture UVRect, specifying portion of input texture corresponding to right eye. */
+	FBox2D RightUVRect;
+
+	UPROPERTY(EditAnywhere, BlueprintReadOnly, Category = "Equirect Properties")
+	/** Left eye's texture coordinate scale after mapping to 2D. */
+	FVector2D LeftScale;
+
+	UPROPERTY(EditAnywhere, BlueprintReadOnly, Category = "Equirect Properties")
+	/** Right eye's texture coordinate scale after mapping to 2D. */
+	FVector2D RightScale;
+
+	UPROPERTY(EditAnywhere, BlueprintReadOnly, Category = "Equirect Properties")
+	/** Left eye's texture coordinate bias after mapping to 2D. */
+	FVector2D LeftBias;
+
+	UPROPERTY(EditAnywhere, BlueprintReadOnly, Category = "Equirect Properties")
+	/** Right eye's texture coordinate bias after mapping to 2D. */
+	FVector2D RightBias;
+
+	/**
+	 * Set Equirect layer properties: UVRect, Scale, and Bias
+	 * @param	LeftScale: Scale for left eye
+	 * @param	LeftBias: Bias for left eye
+	 * @param	RightScale: Scale for right eye
+	 * @param	RightBias: Bias for right eye
+	 */
+	UFUNCTION(BlueprintCallable, Category = "Components|Stereo Layer")
+	void SetEquirectProps(FEquirectProps InScaleBiases);
+
+	virtual void ApplyShape(IStereoLayers::FLayerDesc& LayerDesc) override;
 };
 
 /** 
@@ -121,10 +250,14 @@ public:
 
 	//~ Begin UObject Interface
 	void BeginDestroy() override;
+	virtual void Serialize(FArchive& Ar) override;
+	virtual void PostLoad() override;
+	//~ End UObject Interface
 
 	//~ Begin UActorComponent Interface
 	virtual void TickComponent(float DeltaTime, enum ELevelTick TickType, FActorComponentTickFunction *ThisTickFunction) override;
-	
+	//~ End UActorComponent Interface
+
 	/** 
 	 * Change the texture displayed on the stereo layer. 
 	 *
@@ -178,7 +311,7 @@ public:
 	 * @param	RightScale: Scale for right eye
 	 * @param	RightBias: Bias for right eye
 	 */
-	UFUNCTION(BlueprintCallable, Category="Components|Stereo Layer")
+	UFUNCTION(BlueprintCallable, Category = "Components|Stereo Layer", meta = (DeprecatedFunction, DeprecationMessage = "Use UStereoLayerShapeEquirect::SetEquirectProps() instead."))
 	void SetEquirectProps(FEquirectProps InScaleBiases);
 
 	/** 
@@ -208,13 +341,14 @@ public:
 	UPROPERTY(EditAnywhere, BlueprintReadWrite, Category= "StereoLayer")
 	uint32 bNoAlphaChannel:1;
 
+	void MarkStereoLayerDirty();
 protected:
 	/** Texture displayed on the stereo layer (is stereocopic textures are supported on the platfrom and more than one texture is provided, this will be the right eye) **/
 	UPROPERTY(EditAnywhere, BlueprintReadOnly, Category= "StereoLayer")
 	class UTexture* Texture;
 
-	/** Texture displayed on the stereo layer for left eye, if stereoscopic textures are supported on the platform **/
-	UPROPERTY(EditAnywhere, BlueprintReadOnly, Category = "StereoLayer | Cubemap Overlay Properties")
+	/** Texture displayed on the stereo layer for left eye, if stereoscopic textures are supported on the platform and by the layer shape **/
+	UPROPERTY(EditAnywhere, BlueprintReadOnly, Category = "Stereoscopic Properties")
 	class UTexture* LeftTexture;
 
 public:
@@ -224,36 +358,41 @@ public:
 
 protected:
 	/** Size of the rendered stereo layer quad **/
-	UPROPERTY(EditAnywhere, BlueprintReadOnly, export, Category="StereoLayer | Quad Overlay Properties")
+	UPROPERTY(EditAnywhere, BlueprintReadOnly, export, Category="StereoLayer")
 	FVector2D QuadSize;
 
 	/** UV coordinates mapped to the quad face **/
-	UPROPERTY(EditAnywhere, BlueprintReadOnly, export, Category="StereoLayer | Quad Overlay Properties")
+	UPROPERTY(EditAnywhere, BlueprintReadOnly, export, Category="StereoLayer")
 	FBox2D UVRect;
 
-	/** Radial size of the rendered stereo layer cylinder **/
-	UPROPERTY(EditAnywhere, BlueprintReadOnly, export, Category = "StereoLayer | Cylinder Overlay Properties")
-	float CylinderRadius;
+	UE_DEPRECATED(4.25, "float CylinderRadius is deprecated. Use the corresponding property of the UStereoLayerShapeCylinder subobject instead.")
+	UPROPERTY()
+	float CylinderRadius_DEPRECATED;
 
-	/** Arc angle for the stereo layer cylinder **/
-	UPROPERTY(EditAnywhere, BlueprintReadOnly, export, Category = "StereoLayer | Cylinder Overlay Properties")
-	float CylinderOverlayArc;
+	UE_DEPRECATED(4.25, "float CylinderOverlayArc is deprecated. Use the corresponding property of the UStereoLayerShapeCylinder subobject instead.")
+	UPROPERTY()
+	float CylinderOverlayArc_DEPRECATED;
 
-	/** Height of the stereo layer cylinder **/
-	UPROPERTY(EditAnywhere, BlueprintReadOnly, export, Category = "StereoLayer | Cylinder Overlay Properties")
-	int CylinderHeight;
+	UE_DEPRECATED(4.25, "int CylinderHeight is deprecated. Use the corresponding property of the UStereoLayerShapeCylinder subobject instead.")
+	UPROPERTY()
+	int CylinderHeight_DEPRECATED;
 
-	/** UVRects, scale, and biases per eye for equirect layers. UVRects are for input texture UVs per eye, and the latter 2 control scale and bias of texture coordinates after mapping to 2D.  **/
-	UPROPERTY(EditAnywhere, BlueprintReadOnly, export, Category = "StereoLayer | Equirect Layer Properties")
-	FEquirectProps EquirectProps;
+	UE_DEPRECATED(4.25, "FEquirectProps EquirectProps is deprecated. Use the corresponding properties of the UStereoLayerShapeEquirect subobject instead.")
+	UPROPERTY()
+	FEquirectProps EquirectProps_DEPRECATED;
 
 	/** Specifies how and where the quad is rendered to the screen **/
 	UPROPERTY(EditAnywhere, BlueprintReadOnly, export, Category="StereoLayer")
     TEnumAsByte<enum EStereoLayerType> StereoLayerType;
 
-	/** Specifies which type of layer it is.  Note that some shapes will be supported only on certain platforms! **/
-	UPROPERTY(EditAnywhere, BlueprintReadOnly, export, Category = "StereoLayer")
-	TEnumAsByte<enum EStereoLayerShape> StereoLayerShape;
+	UE_DEPRECATED(4.25, "TEnumAsByte<enum EStereoLayerShape> StereoLayerShape is deprecated. Use Shape instead.")
+	UPROPERTY()
+	TEnumAsByte<enum EStereoLayerShape> StereoLayerShape_DEPRECATED;
+
+	/** Specifies which shape of layer it is.  Note that some shapes will be supported only on certain platforms! **/
+	UPROPERTY(EditAnywhere, BlueprintReadOnly, NoClear, Instanced, Category = "StereoLayer", DisplayName="Stereo Layer Shape")
+	UStereoLayerShape* Shape;
+
 
 	/** Render priority among all stereo layers, higher priority render on top of lower priority **/
 	UPROPERTY(EditAnywhere, BlueprintReadOnly, export, Category= "StereoLayer")
@@ -274,5 +413,8 @@ private:
 
 	/** Last frames visiblity state **/
 	bool bLastVisible;
+
+	/** Set if the component was loaded from an old version */
+	bool bNeedsPostLoadFixup;
 };
 
diff --git a/Engine/Source/Runtime/Engine/Private/Components/StereoLayerComponent.cpp b/Engine/Source/Runtime/Engine/Private/Components/StereoLayerComponent.cpp
--- a/Engine/Source/Runtime/Engine/Private/Components/StereoLayerComponent.cpp
+++ b/Engine/Source/Runtime/Engine/Private/Components/StereoLayerComponent.cpp
@@ -1,6 +1,7 @@
 // Copyright 1998-2019 Epic Games, Inc. All Rights Reserved.
 
 #include "Components/StereoLayerComponent.h"
+#include "UObject/VRObjectVersion.h"
 #include "EngineGlobals.h"
 #include "Engine/Engine.h"
 #include "TextureResource.h"
@@ -8,6 +9,9 @@
 #include "IStereoLayers.h"
 #include "StereoLayerShapes.h"
 #include "StereoRendering.h"
+#if WITH_EDITOR
+#include "SceneManagement.h"
+#endif
 
 UStereoLayerComponent::UStereoLayerComponent(const FObjectInitializer& ObjectInitializer)
 	: Super(ObjectInitializer)
@@ -19,21 +23,18 @@ UStereoLayerComponent::UStereoLayerComponent(const FObjectInitializer& ObjectIni
 	, bQuadPreserveTextureRatio(false)
 	, QuadSize(FVector2D(100.0f, 100.0f))
 	, UVRect(FBox2D(FVector2D(0.0f, 0.0f), FVector2D(1.0f, 1.0f)))
-	, CylinderRadius(100)
-	, CylinderOverlayArc(100)
-	, CylinderHeight(50)
-	, EquirectProps()
 	, StereoLayerType(SLT_FaceLocked)
-	, StereoLayerShape(SLSH_QuadLayer)
 	, Priority(0)
 	, bIsDirty(true)
 	, bTextureNeedsUpdate(false)
 	, LayerId(0)
 	, LastTransform(FTransform::Identity)
 	, bLastVisible(false)
+	, bNeedsPostLoadFixup(false)
 {
 	PrimaryComponentTick.bCanEverTick = true;
 	PrimaryComponentTick.TickGroup = TG_PrePhysics;
+	Shape = ObjectInitializer.CreateDefaultSubobject<UStereoLayerShapeQuad>(this, TEXT("Shape"));
 }
 
 void UStereoLayerComponent::BeginDestroy()
@@ -48,6 +49,56 @@ void UStereoLayerComponent::BeginDestroy()
 	}
 }
 
+void UStereoLayerComponent::Serialize(FArchive& Ar)
+{
+	Ar.UsingCustomVersion(FVRObjectVersion::GUID);
+	Super::Serialize(Ar);
+	if (Ar.IsLoading() && Ar.CustomVer(FVRObjectVersion::GUID) < FVRObjectVersion::UseSubobjectForStereoLayerShapeProperties)
+	{
+		bNeedsPostLoadFixup = true; // Postponing fixups until after load, as we need to modify subobjects.
+	}
+}
+
+void UStereoLayerComponent::PostLoad()
+{
+	Super::PostLoad();
+	PRAGMA_DISABLE_DEPRECATION_WARNINGS
+	if (bNeedsPostLoadFixup)
+	{
+		switch (StereoLayerShape_DEPRECATED)
+		{
+			case SLSH_QuadLayer:
+			{
+				Shape = NewObject<UStereoLayerShapeQuad>(this);
+				break;
+			}
+			case SLSH_CubemapLayer:
+			{
+				Shape = NewObject<UStereoLayerShapeCubemap>(this);
+				break;
+			}
+			case SLSH_CylinderLayer:
+			{
+				auto Cylinder = NewObject<UStereoLayerShapeCylinder>(this);
+				Shape = Cylinder;
+				Cylinder->Height = CylinderHeight_DEPRECATED;
+				Cylinder->OverlayArc = CylinderOverlayArc_DEPRECATED;
+				Cylinder->Radius = CylinderRadius_DEPRECATED;
+				break;
+			}
+			case SLSH_EquirectLayer:
+			{
+				auto Equirect = NewObject<UStereoLayerShapeEquirect>(this);
+				Shape = Equirect;
+				Equirect->SetEquirectProps(EquirectProps_DEPRECATED);
+				break;
+			}
+		}
+		bNeedsPostLoadFixup = false;
+	}
+	PRAGMA_ENABLE_DEPRECATION_WARNINGS
+}
+
 void UStereoLayerComponent::TickComponent(float DeltaTime, enum ELevelTick TickType, FActorComponentTickFunction *ThisTickFunction)
 {
 	Super::TickComponent( DeltaTime, TickType, ThisTickFunction );
@@ -82,70 +133,52 @@ void UStereoLayerComponent::TickComponent(float DeltaTime, enum ELevelTick TickT
 
 	if (bIsDirty)
 	{
-		IStereoLayers::FLayerDesc LayerDsec;
-		LayerDsec.Priority = Priority;
-		LayerDsec.QuadSize = QuadSize;
-		LayerDsec.UVRect = UVRect;
-		LayerDsec.Transform = Transform;
+		IStereoLayers::FLayerDesc LayerDesc;
+		LayerDesc.Priority = Priority;
+		LayerDesc.QuadSize = QuadSize;
+		LayerDesc.UVRect = UVRect;
+		LayerDesc.Transform = Transform;
 
 		if (Texture)
 		{
-			LayerDsec.Texture = Texture->Resource->TextureRHI;
-			LayerDsec.Flags |= (Texture->GetMaterialType() == MCT_TextureExternal) ? IStereoLayers::LAYER_FLAG_TEX_EXTERNAL : 0;
+			LayerDesc.Texture = Texture->Resource->TextureRHI;
+			LayerDesc.Flags |= (Texture->GetMaterialType() == MCT_TextureExternal) ? IStereoLayers::LAYER_FLAG_TEX_EXTERNAL : 0;
 		}
 		if (LeftTexture)
 		{
-			LayerDsec.LeftTexture = LeftTexture->Resource->TextureRHI;
+			LayerDesc.LeftTexture = LeftTexture->Resource->TextureRHI;
 		}
 				
-		LayerDsec.Flags |= (bLiveTexture) ? IStereoLayers::LAYER_FLAG_TEX_CONTINUOUS_UPDATE : 0;
-		LayerDsec.Flags |= (bNoAlphaChannel) ? IStereoLayers::LAYER_FLAG_TEX_NO_ALPHA_CHANNEL : 0;
-		LayerDsec.Flags |= (bQuadPreserveTextureRatio) ? IStereoLayers::LAYER_FLAG_QUAD_PRESERVE_TEX_RATIO : 0;
-		LayerDsec.Flags |= (bSupportsDepth) ? IStereoLayers::LAYER_FLAG_SUPPORT_DEPTH : 0;
-		LayerDsec.Flags |= (!bCurrVisible) ? IStereoLayers::LAYER_FLAG_HIDDEN : 0;
+		LayerDesc.Flags |= (bLiveTexture) ? IStereoLayers::LAYER_FLAG_TEX_CONTINUOUS_UPDATE : 0;
+		LayerDesc.Flags |= (bNoAlphaChannel) ? IStereoLayers::LAYER_FLAG_TEX_NO_ALPHA_CHANNEL : 0;
+		LayerDesc.Flags |= (bQuadPreserveTextureRatio) ? IStereoLayers::LAYER_FLAG_QUAD_PRESERVE_TEX_RATIO : 0;
+		LayerDesc.Flags |= (bSupportsDepth) ? IStereoLayers::LAYER_FLAG_SUPPORT_DEPTH : 0;
+		LayerDesc.Flags |= (!bCurrVisible) ? IStereoLayers::LAYER_FLAG_HIDDEN : 0;
 
 		switch (StereoLayerType)
 		{
 		case SLT_WorldLocked:
-			LayerDsec.PositionType = IStereoLayers::WorldLocked;
+			LayerDesc.PositionType = IStereoLayers::WorldLocked;
 			break;
 		case SLT_TrackerLocked:
-			LayerDsec.PositionType = IStereoLayers::TrackerLocked;
+			LayerDesc.PositionType = IStereoLayers::TrackerLocked;
 			break;
 		case SLT_FaceLocked:
-			LayerDsec.PositionType = IStereoLayers::FaceLocked;
-			break;
-		}
-
-		switch (StereoLayerShape)
-		{
-		case SLSH_QuadLayer:
-			break;
-
-		case SLSH_CylinderLayer:
-			LayerDsec.SetShape<FCylinderLayer>(CylinderRadius, CylinderOverlayArc, CylinderHeight);
-			break;
-
-		case SLSH_CubemapLayer:
-			LayerDsec.SetShape<FCubemapLayer>();
-			break;
-
-		case SLSH_EquirectLayer:
-			LayerDsec.SetShape<FEquirectLayer>(EquirectProps.LeftUVRect, EquirectProps.RightUVRect, EquirectProps.LeftScale, EquirectProps.RightScale, EquirectProps.LeftBias, EquirectProps.RightBias );
-			break;
-
-		default:
+			LayerDesc.PositionType = IStereoLayers::FaceLocked;
 			break;
 		}
 
+		// Set the correct layer shape and apply any shape-specific properties
+		check(Shape);
+		Shape->ApplyShape(LayerDesc);
 
 		if (LayerId)
 		{
-			StereoLayers->SetLayerDesc(LayerId, LayerDsec);
+			StereoLayers->SetLayerDesc(LayerId, LayerDesc);
 		}
 		else
 		{
-			LayerId = StereoLayers->CreateLayer(LayerDsec);
+			LayerId = StereoLayers->CreateLayer(LayerDesc);
 		}
 		
 		LastTransform = Transform;
@@ -168,7 +201,7 @@ void UStereoLayerComponent::SetTexture(UTexture* InTexture)
 	}
 
 	Texture = InTexture;
-	bIsDirty = true;
+	MarkStereoLayerDirty();
 }
 
 void UStereoLayerComponent::SetLeftTexture(UTexture* InTexture)
@@ -179,7 +212,7 @@ void UStereoLayerComponent::SetLeftTexture(UTexture* InTexture)
 	}
 
 	LeftTexture = InTexture;
-	bIsDirty = true;
+	MarkStereoLayerDirty();
 }
 
 void UStereoLayerComponent::SetQuadSize(FVector2D InQuadSize)
@@ -190,7 +223,7 @@ void UStereoLayerComponent::SetQuadSize(FVector2D InQuadSize)
 	}
 
 	QuadSize = InQuadSize;
-	bIsDirty = true;
+	MarkStereoLayerDirty();
 }
 
 void UStereoLayerComponent::SetUVRect(FBox2D InUVRect)
@@ -201,17 +234,31 @@ void UStereoLayerComponent::SetUVRect(FBox2D InUVRect)
 	}
 
 	UVRect = InUVRect;
-	bIsDirty = true;
+	MarkStereoLayerDirty();
 }
 
 void UStereoLayerComponent::SetEquirectProps(FEquirectProps InEquirectProps)
 {
-	if (EquirectProps == InEquirectProps)
+	if (Shape->IsA<UStereoLayerShapeEquirect>())
+	{
+		Cast<UStereoLayerShapeEquirect>(Shape)->SetEquirectProps(InEquirectProps);
+	}
+}
+
+void UStereoLayerShapeEquirect::SetEquirectProps(FEquirectProps InEquirectProps)
+{
+	if (InEquirectProps == *this)
 	{
 		return;
 	}
-	EquirectProps = InEquirectProps;
-	bIsDirty = true;
+	LeftUVRect = InEquirectProps.LeftUVRect;
+	RightUVRect = InEquirectProps.RightUVRect;
+	LeftScale = InEquirectProps.LeftScale;
+	RightScale = InEquirectProps.RightScale;
+	LeftBias = InEquirectProps.LeftBias;
+	RightBias = InEquirectProps.RightBias;
+
+	MarkStereoLayerDirty();
 }
 
 void UStereoLayerComponent::SetPriority(int32 InPriority)
@@ -222,7 +269,40 @@ void UStereoLayerComponent::SetPriority(int32 InPriority)
 	}
 
 	Priority = InPriority;
-	bIsDirty = true;
+	MarkStereoLayerDirty();
+}
+
+void UStereoLayerShapeCylinder::SetRadius(float InRadius)
+{
+	if (Radius == InRadius)
+	{
+		return;
+	}
+
+	Radius = InRadius;
+	MarkStereoLayerDirty();
+}
+
+void UStereoLayerShapeCylinder::SetOverlayArc(float InOverlayArc)
+{
+	if (OverlayArc == InOverlayArc)
+	{
+		return;
+	}
+
+	OverlayArc = InOverlayArc;
+	MarkStereoLayerDirty();
+}
+
+void UStereoLayerShapeCylinder::SetHeight(int InHeight)
+{
+	if (Height == InHeight)
+	{
+		return;
+	}
+
+	Height = InHeight;
+	MarkStereoLayerDirty();
 }
 
 void UStereoLayerComponent::MarkTextureForUpdate()
@@ -230,3 +310,91 @@ void UStereoLayerComponent::MarkTextureForUpdate()
 	bTextureNeedsUpdate = true;
 }
 
+void UStereoLayerComponent::MarkStereoLayerDirty()
+{
+	bIsDirty = true;
+}
+
+void UStereoLayerShape::MarkStereoLayerDirty()
+{
+	check(GetOuter()->IsA<UStereoLayerComponent>());
+	Cast<UStereoLayerComponent>(GetOuter())->MarkStereoLayerDirty();
+}
+
+void UStereoLayerShapeCylinder::ApplyShape(IStereoLayers::FLayerDesc& LayerDesc)
+{
+	LayerDesc.SetShape<FCylinderLayer>(Radius, OverlayArc, Height);
+}
+
+void UStereoLayerShapeCubemap::ApplyShape(IStereoLayers::FLayerDesc& LayerDesc)
+{
+	LayerDesc.SetShape<FCubemapLayer>();
+}
+
+void UStereoLayerShapeEquirect::ApplyShape(IStereoLayers::FLayerDesc& LayerDesc)
+{
+	LayerDesc.SetShape<FEquirectLayer>(LeftUVRect, RightUVRect, LeftScale, RightScale, LeftBias, RightBias);
+}
+
+void UStereoLayerShapeQuad::ApplyShape(IStereoLayers::FLayerDesc& LayerDesc)
+{
+	LayerDesc.SetShape<FQuadLayer>();
+}
+
+void UStereoLayerShape::ApplyShape(IStereoLayers::FLayerDesc& LayerDesc)
+{
+}
+
+#if WITH_EDITOR
+void UStereoLayerShape::DrawShapeVisualization(const class FSceneView* View, class FPrimitiveDrawInterface* PDI)
+{
+}
+
+void UStereoLayerShapeQuad::DrawShapeVisualization(const class FSceneView* View, class FPrimitiveDrawInterface* PDI)
+{
+	FLinearColor YellowColor = FColor(231, 239, 0, 255);
+	check(GetOuter()->IsA<UStereoLayerComponent>());
+
+	auto StereoLayerComp = Cast<UStereoLayerComponent>(GetOuter());
+	const FVector2D QuadSize = StereoLayerComp->GetQuadSize() / 2.0f;
+	const FBox QuadBox(FVector(0.0f, -QuadSize.X, -QuadSize.Y), FVector(0.0f, QuadSize.X, QuadSize.Y));
+
+	DrawWireBox(PDI, StereoLayerComp->GetComponentTransform().ToMatrixWithScale(), QuadBox, YellowColor, 0);
+}
+
+void UStereoLayerShapeCylinder::DrawShapeVisualization(const class FSceneView* View, class FPrimitiveDrawInterface* PDI)
+{
+	FLinearColor YellowColor = FColor(231, 239, 0, 255);
+	check(GetOuter()->IsA<UStereoLayerComponent>());
+
+	auto StereoLayerComp = Cast<UStereoLayerComponent>(GetOuter());
+	float ArcAngle = OverlayArc * 180 / (Radius * PI);
+
+	FVector X = StereoLayerComp->GetComponentTransform().GetUnitAxis(EAxis::Type::X);
+	FVector Y = StereoLayerComp->GetComponentTransform().GetUnitAxis(EAxis::Type::Y);
+	FVector Base = StereoLayerComp->GetComponentTransform().GetLocation();
+	FVector HalfHeight = FVector(0, 0, Height / 2);
+
+	FVector LeftVertex = Base + Radius * (FMath::Cos(ArcAngle / 2 * (PI / 180.0f)) * X + FMath::Sin(ArcAngle / 2 * (PI / 180.0f)) * Y);
+	FVector RightVertex = Base + Radius * (FMath::Cos(-ArcAngle / 2 * (PI / 180.0f)) * X + FMath::Sin(-ArcAngle / 2 * (PI / 180.0f)) * Y);
+
+	DrawArc(PDI, Base + HalfHeight, X, Y, -ArcAngle / 2, ArcAngle / 2, Radius, 10, YellowColor, 0);
+
+	DrawArc(PDI, Base - HalfHeight, X, Y, -ArcAngle / 2, ArcAngle / 2, Radius, 10, YellowColor, 0);
+
+	PDI->DrawLine(LeftVertex - HalfHeight, LeftVertex + HalfHeight, YellowColor, 0);
+
+	PDI->DrawLine(RightVertex - HalfHeight, RightVertex + HalfHeight, YellowColor, 0);
+}
+#endif
+
+
+bool FEquirectProps::operator==(const class UStereoLayerShapeEquirect& Other) const
+{
+	return (LeftUVRect == Other.LeftUVRect) && (RightUVRect == Other.RightUVRect) && (LeftScale == Other.LeftScale) && (RightScale == Other.RightScale) && (LeftBias == Other.LeftBias) && (RightBias == Other.RightBias);
+}
+
+bool FEquirectProps::operator==(const FEquirectProps& Other) const
+{
+	return (LeftUVRect == Other.LeftUVRect) && (RightUVRect == Other.RightUVRect) && (LeftScale == Other.LeftScale) && (RightScale == Other.RightScale) && (LeftBias == Other.LeftBias) && (RightBias == Other.RightBias);
+}
