Final StereoLayer Shape refactoing changes and fixes

diff --git a/Engine/Source/Editor/ComponentVisualizers/Private/StereoLayerComponentVisualizer.cpp b/Engine/Source/Editor/ComponentVisualizers/Private/StereoLayerComponentVisualizer.cpp
--- a/Engine/Source/Editor/ComponentVisualizers/Private/StereoLayerComponentVisualizer.cpp
+++ b/Engine/Source/Editor/ComponentVisualizers/Private/StereoLayerComponentVisualizer.cpp
@@ -9,9 +9,8 @@
 void FStereoLayerComponentVisualizer::DrawVisualization( const UActorComponent* Component, const FSceneView* View, FPrimitiveDrawInterface* PDI )
 {
 	const UStereoLayerComponent* StereoLayerComp = Cast<const UStereoLayerComponent>(Component);
-	if(StereoLayerComp != nullptr)
+	if(StereoLayerComp != nullptr && StereoLayerComp->Shape != nullptr)
 	{
-		check(StereoLayerComp->Shape != nullptr);
 		StereoLayerComp->Shape->DrawShapeVisualization(View, PDI);
     }
 }
diff --git a/Engine/Source/Runtime/Engine/Private/Components/StereoLayerComponent.cpp b/Engine/Source/Runtime/Engine/Private/Components/StereoLayerComponent.cpp
--- a/Engine/Source/Runtime/Engine/Private/Components/StereoLayerComponent.cpp
+++ b/Engine/Source/Runtime/Engine/Private/Components/StereoLayerComponent.cpp
@@ -34,7 +34,8 @@ UStereoLayerComponent::UStereoLayerComponent(const FObjectInitializer& ObjectIni
 {
 	PrimaryComponentTick.bCanEverTick = true;
 	PrimaryComponentTick.TickGroup = TG_PrePhysics;
-	Shape = ObjectInitializer.CreateDefaultSubobject<UStereoLayerShapeQuad>(this, TEXT("Shape"));
+	// The following somehow overrides subobjects read through serialization.
+	//Shape = ObjectInitializer.CreateDefaultSubobject<UStereoLayerShapeQuad>(this, TEXT("Shape"));
 }
 
 void UStereoLayerComponent::BeginDestroy()
@@ -97,6 +98,11 @@ void UStereoLayerComponent::PostLoad()
 		bNeedsPostLoadFixup = false;
 	}
 	PRAGMA_ENABLE_DEPRECATION_WARNINGS
+
+	if (Shape == nullptr)
+	{
+		Shape = NewObject<UStereoLayerShapeQuad>(this);
+	}
 }
 
 void UStereoLayerComponent::TickComponent(float DeltaTime, enum ELevelTick TickType, FActorComponentTickFunction *ThisTickFunction)
@@ -169,7 +175,10 @@ void UStereoLayerComponent::TickComponent(float DeltaTime, enum ELevelTick TickT
 		}
 
 		// Set the correct layer shape and apply any shape-specific properties
-		check(Shape);
+		if (Shape == nullptr)
+		{
+			Shape = NewObject<UStereoLayerShapeQuad>(this);
+		}
 		Shape->ApplyShape(LayerDesc);
 
 		if (LayerId)
