Merge DevVR to Dev-RenderPlat-Staging

diff --git a/Engine/Plugins/Lumin/MagicLeap/Source/MagicLeap/Private/MagicLeapStereoLayers.cpp b/Engine/Plugins/Lumin/MagicLeap/Source/MagicLeap/Private/MagicLeapStereoLayers.cpp
--- a/Engine/Plugins/Lumin/MagicLeap/Source/MagicLeap/Private/MagicLeapStereoLayers.cpp
+++ b/Engine/Plugins/Lumin/MagicLeap/Source/MagicLeap/Private/MagicLeapStereoLayers.cpp
@@ -80,7 +80,6 @@ IStereoLayers::FLayerDesc FMagicLeapStereoLayers::GetDebugCanvasLayerDesc(FTextu
 	const float DebugCanvasHeight = CVarHeight ? CVarHeight->GetFloat() : DefaultHeight;
 	StereoLayerDesc.QuadSize = FVector2D(DebugCanvasWidth, DebugCanvasHeight);
 	StereoLayerDesc.PositionType = IStereoLayers::ELayerType::FaceLocked;
-	StereoLayerDesc.ShapeType = IStereoLayers::ELayerShape::QuadLayer;
 	StereoLayerDesc.Texture = Texture;
 	StereoLayerDesc.Flags = IStereoLayers::ELayerFlags::LAYER_FLAG_TEX_CONTINUOUS_UPDATE;
 	StereoLayerDesc.Flags |= IStereoLayers::ELayerFlags::LAYER_FLAG_QUAD_PRESERVE_TEX_RATIO;
diff --git a/Engine/Plugins/Lumin/MagicLeap/Source/MagicLeapHelperVulkan/Private/MagicLeapHelperVulkan.cpp b/Engine/Plugins/Lumin/MagicLeap/Source/MagicLeapHelperVulkan/Private/MagicLeapHelperVulkan.cpp
--- a/Engine/Plugins/Lumin/MagicLeap/Source/MagicLeapHelperVulkan/Private/MagicLeapHelperVulkan.cpp
+++ b/Engine/Plugins/Lumin/MagicLeap/Source/MagicLeapHelperVulkan/Private/MagicLeapHelperVulkan.cpp
@@ -320,9 +320,9 @@ bool FMagicLeapHelperVulkan::GetMediaTexture(FTextureRHIRef& ResultTexture, FSam
 	return false;
 }
 
-void FMagicLeapHelperVulkan::AliasMediaTexture(FRHITexture* DestTexture, FRHITexture* SrcTexture)
+void FMagicLeapHelperVulkan::AliasMediaTexture(FTextureRHIRef& DestTexture, FTextureRHIRef& SrcTexture)
 {
 #if PLATFORM_LUMIN
-	GDynamicRHI->RHIAliasTextureResources(DestTexture, SrcTexture);
+	GDynamicRHI->RHIAliasTextureResources((FTextureRHIRef&)DestTexture, (FTextureRHIRef&)SrcTexture);
 #endif
 }
diff --git a/Engine/Plugins/Lumin/MagicLeap/Source/MagicLeapHelperVulkan/Public/MagicLeapHelperVulkan.h b/Engine/Plugins/Lumin/MagicLeap/Source/MagicLeapHelperVulkan/Public/MagicLeapHelperVulkan.h
--- a/Engine/Plugins/Lumin/MagicLeap/Source/MagicLeapHelperVulkan/Public/MagicLeapHelperVulkan.h
+++ b/Engine/Plugins/Lumin/MagicLeap/Source/MagicLeapHelperVulkan/Public/MagicLeapHelperVulkan.h
@@ -16,5 +16,5 @@ public:
 	static bool GetVulkanInstanceExtensionsRequired(TArray<const ANSICHAR*>& Out);
 	static bool GetVulkanDeviceExtensionsRequired(struct VkPhysicalDevice_T* pPhysicalDevice, TArray<const ANSICHAR*>& Out);
 	static bool GetMediaTexture(FTextureRHIRef& Result, FSamplerStateRHIRef& SamplerResult, const uint64 MediaTextureHandle);
-	static void AliasMediaTexture(FRHITexture* DestTexture, FRHITexture* SrcTexture);
+	static void AliasMediaTexture(FTextureRHIRef& DestTexture, FTextureRHIRef& SrcTexture);
 };
diff --git a/Engine/Plugins/Runtime/Oculus/OculusVR/Source/OculusHMD/Private/OculusHMD.cpp b/Engine/Plugins/Runtime/Oculus/OculusVR/Source/OculusHMD/Private/OculusHMD.cpp
--- a/Engine/Plugins/Runtime/Oculus/OculusVR/Source/OculusHMD/Private/OculusHMD.cpp
+++ b/Engine/Plugins/Runtime/Oculus/OculusVR/Source/OculusHMD/Private/OculusHMD.cpp
@@ -1758,9 +1758,10 @@ namespace OculusHMD
 		if (LayerFound && (*LayerFound)->GetSwapChain().IsValid())
 		{
 			bool bRightTexture = (*LayerFound)->GetRightSwapChain().IsValid();
-			switch ((*LayerFound)->GetDesc().ShapeType)
+			const IStereoLayers::FLayerDesc& Desc = (*LayerFound)->GetDesc();
+			
+			if (Desc.HasShape<FCubemapLayer>())
 			{
-			case IStereoLayers::CubemapLayer:
 				if (bRightTexture)
 				{
 					Texture = (*LayerFound)->GetRightSwapChain()->GetTextureCube();
@@ -1769,10 +1770,10 @@ namespace OculusHMD
 				else
 				{
 					Texture = LeftTexture = (*LayerFound)->GetSwapChain()->GetTextureCube();
-				}				break;
-
-			case IStereoLayers::CylinderLayer:
-			case IStereoLayers::QuadLayer:
+				}
+			}
+			else if (Desc.HasShape<FCylinderLayer>() || Desc.HasShape<FQuadLayer>())
+			{
 				if (bRightTexture)
 				{
 					Texture = (*LayerFound)->GetRightSwapChain()->GetTexture2D();
@@ -1782,24 +1783,16 @@ namespace OculusHMD
 				{
 					Texture = LeftTexture = (*LayerFound)->GetSwapChain()->GetTexture2D();
 				}
-				break;
-
-			default:
-				break;
 			}
 		}
 	}
 
 	IStereoLayers::FLayerDesc FOculusHMD::GetDebugCanvasLayerDesc(FTextureRHIRef Texture)
 	{
-		IStereoLayers::FLayerDesc StereoLayerDesc;
+		IStereoLayers::FLayerDesc StereoLayerDesc(FCylinderLayer(180.f, 488.f / 4, 100.f));
 		StereoLayerDesc.Transform = FTransform(FVector(0.f, 0, 0)); //100/0/0 for quads
-		StereoLayerDesc.CylinderHeight = 180.f;
-		StereoLayerDesc.CylinderOverlayArc = 488.f/4;
-		StereoLayerDesc.CylinderRadius = 100.f;
 		StereoLayerDesc.QuadSize = FVector2D(180.f, 180.f);
 		StereoLayerDesc.PositionType = IStereoLayers::ELayerType::FaceLocked;
-		StereoLayerDesc.ShapeType = IStereoLayers::ELayerShape::CylinderLayer;
 		StereoLayerDesc.LayerSize = Texture->GetTexture2D()->GetSizeXY();
 		StereoLayerDesc.Flags = IStereoLayers::ELayerFlags::LAYER_FLAG_TEX_CONTINUOUS_UPDATE;
 		StereoLayerDesc.Flags |= IStereoLayers::ELayerFlags::LAYER_FLAG_QUAD_PRESERVE_TEX_RATIO;
diff --git a/Engine/Plugins/Runtime/Oculus/OculusVR/Source/OculusHMD/Private/OculusHMD_Layer.cpp b/Engine/Plugins/Runtime/Oculus/OculusVR/Source/OculusHMD/Private/OculusHMD_Layer.cpp
--- a/Engine/Plugins/Runtime/Oculus/OculusVR/Source/OculusHMD/Private/OculusHMD_Layer.cpp
+++ b/Engine/Plugins/Runtime/Oculus/OculusVR/Source/OculusHMD/Private/OculusHMD_Layer.cpp
@@ -184,7 +184,7 @@ static void AppendFaceIndices(const int v0, const int v1, const int v2, const in
 
 void FLayer::BuildPokeAHoleMesh(TArray<FVector>& Vertices, TArray<int32>& Triangles, TArray<FVector2D>& UV0)
 {
-	if (Desc.ShapeType == IStereoLayers::QuadLayer)
+	if (Desc.HasShape<FQuadLayer>())
 	{
 		const float QuadScale = 0.99;
 
@@ -209,20 +209,21 @@ void FLayer::BuildPokeAHoleMesh(TArray<FVector>& Vertices, TArray<int32>& Triang
 		Triangles.Reserve(6);
 		AppendFaceIndices(0, 1, 2, 3, Triangles, false);
 	}
-	else if (Desc.ShapeType == IStereoLayers::CylinderLayer)
+	else if (Desc.HasShape<FCylinderLayer>())
 	{
+		const FCylinderLayer& CylinderProps = Desc.GetShape<FCylinderLayer>();
 		const float CylinderScale = 0.99;
 
 		FIntPoint TexSize = Desc.Texture.IsValid() ? Desc.Texture->GetTexture2D()->GetSizeXY() : Desc.LayerSize;
 		float AspectRatio = TexSize.X ? (float)TexSize.Y / (float)TexSize.X : 3.0f / 4.0f;
 
-		float CylinderHeight = (Desc.Flags & IStereoLayers::LAYER_FLAG_QUAD_PRESERVE_TEX_RATIO) ? Desc.CylinderOverlayArc * AspectRatio : Desc.CylinderHeight;
+		float CylinderHeight = (Desc.Flags & IStereoLayers::LAYER_FLAG_QUAD_PRESERVE_TEX_RATIO) ? CylinderProps.OverlayArc * AspectRatio : CylinderProps.Height;
 
 		const FVector XAxis = FVector(1, 0, 0);
 		const FVector YAxis = FVector(0, 1, 0);
 		const FVector HalfHeight = FVector(0, 0, CylinderHeight / 2);
 
-		const float ArcAngle = Desc.CylinderOverlayArc / Desc.CylinderRadius;
+		const float ArcAngle = CylinderProps.OverlayArc / CylinderProps.Radius;
 		const int Sides = (int)( (ArcAngle * 180) / (PI * 5) ); // one triangle every 10 degrees of cylinder for a good-cheap approximation
 		Vertices.Init(FVector::ZeroVector, 2 * (Sides + 1));
 		UV0.Init(FVector2D::ZeroVector, 2 * (Sides + 1));
@@ -234,7 +235,7 @@ void FLayer::BuildPokeAHoleMesh(TArray<FVector>& Vertices, TArray<int32>& Triang
 
 		for (int Side = 0; Side < Sides + 1; Side++)
 		{
-			FVector MidVertex = Desc.CylinderRadius * (FMath::Cos(CurrentAngle) * XAxis + FMath::Sin(CurrentAngle) * YAxis);
+			FVector MidVertex = CylinderProps.Radius * (FMath::Cos(CurrentAngle) * XAxis + FMath::Sin(CurrentAngle) * YAxis);
 			Vertices[2 * Side] = (MidVertex - HalfHeight) * CylinderScale;
 			Vertices[(2 * Side) + 1] = (MidVertex + HalfHeight) * CylinderScale;
 
@@ -254,7 +255,7 @@ void FLayer::BuildPokeAHoleMesh(TArray<FVector>& Vertices, TArray<int32>& Triang
 			}
 		}
 	}
-	else if (Desc.ShapeType == IStereoLayers::CubemapLayer)
+	else if (Desc.HasShape<FCubemapLayer>())
 	{
 		const float CubemapScale = 1000;
 		Vertices.Init(FVector::ZeroVector, 8);
@@ -371,25 +372,25 @@ void FLayer::Initialize_RenderThread(const FSettings* Settings, FCustomPresent*
 
 		ovrpShape Shape;
 
-		switch (Desc.ShapeType)
+		if (Desc.HasShape<FQuadLayer>())
 		{
-		case IStereoLayers::QuadLayer:
 			Shape = ovrpShape_Quad;
-			break;
-
-		case IStereoLayers::CylinderLayer:
+		}
+		else if (Desc.HasShape<FCylinderLayer>())
+		{
 			Shape = ovrpShape_Cylinder;
-			break;
-
-		case IStereoLayers::CubemapLayer:
+		}
+		else if (Desc.HasShape<FCubemapLayer>())
+		{
 			Shape = ovrpShape_Cubemap;
-			break;
-
-		case IStereoLayers::EquirectLayer:
+		}
+		else if (Desc.HasShape<FEquirectLayer>())
+		{
 			Shape = ovrpShape_Equirect;
-			break;
+		}
+		else
+		{
 
-		default:
 			return;
 		}
 
@@ -609,7 +610,6 @@ void FLayer::UpdateTexture_RenderThread(FCustomPresent* CustomPresent, FRHIComma
 		{
 			bool bAlphaPremultiply = true;
 			bool bNoAlphaWrite = (Desc.Flags & IStereoLayers::LAYER_FLAG_TEX_NO_ALPHA_CHANNEL) != 0;
-			bool bIsCubemap = (Desc.ShapeType == IStereoLayers::ELayerShape::CubemapLayer);
 
 			// Left
 			{
@@ -659,28 +659,30 @@ const ovrpLayerSubmit* FLayer::UpdateLayer_RHIThread(const FSettings* Settings,
 	OvrpLayerSubmit.ColorScale = injectColorScale ? Settings->ColorScale : ovrpVector4f{ 1, 1, 1, 1};
 
 	if (OvrpLayerDesc.Shape == ovrpShape_Equirect) {
+		const FEquirectLayer& EquirectProps = Desc.GetShape<FEquirectLayer>();
+
 		ovrpTextureRectMatrixf& RectMatrix = OvrpLayerSubmit.TextureRectMatrix;
 		ovrpRectf& LeftUVRect = RectMatrix.LeftRect;
 		ovrpRectf& RightUVRect = RectMatrix.RightRect;
-		LeftUVRect.Pos.x = Desc.EquirectProps.LeftUVRect.Min.X;
-		LeftUVRect.Pos.y = Desc.EquirectProps.LeftUVRect.Min.Y;
-		LeftUVRect.Size.w = Desc.EquirectProps.LeftUVRect.Max.X - Desc.EquirectProps.LeftUVRect.Min.X;
-		LeftUVRect.Size.h = Desc.EquirectProps.LeftUVRect.Max.Y - Desc.EquirectProps.LeftUVRect.Min.Y;
-		RightUVRect.Pos.x = Desc.EquirectProps.RightUVRect.Min.X;
-		RightUVRect.Pos.y = Desc.EquirectProps.RightUVRect.Min.Y;
-		RightUVRect.Size.w = Desc.EquirectProps.RightUVRect.Max.X - Desc.EquirectProps.RightUVRect.Min.X;
-		RightUVRect.Size.h = Desc.EquirectProps.RightUVRect.Max.Y - Desc.EquirectProps.RightUVRect.Min.Y;
+		LeftUVRect.Pos.x = EquirectProps.LeftUVRect.Min.X;
+		LeftUVRect.Pos.y = EquirectProps.LeftUVRect.Min.Y;
+		LeftUVRect.Size.w = EquirectProps.LeftUVRect.Max.X - EquirectProps.LeftUVRect.Min.X;
+		LeftUVRect.Size.h = EquirectProps.LeftUVRect.Max.Y - EquirectProps.LeftUVRect.Min.Y;
+		RightUVRect.Pos.x = EquirectProps.RightUVRect.Min.X;
+		RightUVRect.Pos.y = EquirectProps.RightUVRect.Min.Y;
+		RightUVRect.Size.w = EquirectProps.RightUVRect.Max.X - EquirectProps.RightUVRect.Min.X;
+		RightUVRect.Size.h = EquirectProps.RightUVRect.Max.Y - EquirectProps.RightUVRect.Min.Y;
 
 		ovrpVector4f& LeftScaleBias = RectMatrix.LeftScaleBias;
-		LeftScaleBias.x = Desc.EquirectProps.LeftScale.X;
-		LeftScaleBias.y = Desc.EquirectProps.LeftScale.Y;
-		LeftScaleBias.z = Desc.EquirectProps.LeftBias.X;
-		LeftScaleBias.w = Desc.EquirectProps.LeftBias.Y;
+		LeftScaleBias.x = EquirectProps.LeftScale.X;
+		LeftScaleBias.y = EquirectProps.LeftScale.Y;
+		LeftScaleBias.z = EquirectProps.LeftBias.X;
+		LeftScaleBias.w = EquirectProps.LeftBias.Y;
 		ovrpVector4f& RightScaleBias = RectMatrix.RightScaleBias;
-		RightScaleBias.x = Desc.EquirectProps.RightScale.X;
-		RightScaleBias.y = Desc.EquirectProps.RightScale.Y;
-		RightScaleBias.z = Desc.EquirectProps.RightBias.X;
-		RightScaleBias.w = Desc.EquirectProps.RightBias.Y;
+		RightScaleBias.x = EquirectProps.RightScale.X;
+		RightScaleBias.y = EquirectProps.RightScale.Y;
+		RightScaleBias.z = EquirectProps.RightBias.X;
+		RightScaleBias.w = EquirectProps.RightBias.Y;
 
 		OvrpLayerSubmit.OverrideTextureRectMatrix = ovrpBool_True;
 	}
@@ -705,10 +707,11 @@ const ovrpLayerSubmit* FLayer::UpdateLayer_RHIThread(const FSettings* Settings,
 			break;
 		case ovrpShape_Cylinder:
 			{
-				float CylinderHeight = (Desc.Flags & IStereoLayers::LAYER_FLAG_QUAD_PRESERVE_TEX_RATIO) ? Desc.CylinderOverlayArc * AspectRatio : Desc.CylinderHeight;
-				OvrpLayerSubmit.Cylinder.ArcWidth = Desc.CylinderOverlayArc * Scale.x;
+				const FCylinderLayer& CylinderProps = Desc.GetShape<FCylinderLayer>();
+				float CylinderHeight = (Desc.Flags & IStereoLayers::LAYER_FLAG_QUAD_PRESERVE_TEX_RATIO) ? CylinderProps.OverlayArc * AspectRatio : CylinderProps.Height;
+				OvrpLayerSubmit.Cylinder.ArcWidth = CylinderProps.OverlayArc * Scale.x;
 				OvrpLayerSubmit.Cylinder.Height = CylinderHeight * Scale.x;
-				OvrpLayerSubmit.Cylinder.Radius = Desc.CylinderRadius * Scale.x;
+				OvrpLayerSubmit.Cylinder.Radius = CylinderProps.Radius * Scale.x;
 			}
 			break;
 		}
diff --git a/Engine/Plugins/Runtime/Oculus/OculusVR/Source/OculusHMD/Private/OculusHMD_Splash.cpp b/Engine/Plugins/Runtime/Oculus/OculusVR/Source/OculusHMD/Private/OculusHMD_Splash.cpp
--- a/Engine/Plugins/Runtime/Oculus/OculusVR/Source/OculusHMD/Private/OculusHMD_Splash.cpp
+++ b/Engine/Plugins/Runtime/Oculus/OculusVR/Source/OculusHMD/Private/OculusHMD_Splash.cpp
@@ -39,7 +39,6 @@ FSplash::FSplash(FOculusHMD* InOculusHMD) :
 		LayerDesc.QuadSize = FVector2D(0.01f, 0.01f);
 		LayerDesc.Priority = 0;
 		LayerDesc.PositionType = IStereoLayers::TrackerLocked;
-		LayerDesc.ShapeType = IStereoLayers::QuadLayer;
 		LayerDesc.Texture = GBlackTexture->TextureRHI;
 		BlackLayer = MakeShareable(new FLayer(NextLayerId++, LayerDesc));
 	}
@@ -50,7 +49,6 @@ FSplash::FSplash(FOculusHMD* InOculusHMD) :
 		LayerDesc.QuadSize = FVector2D(0.01f, 0.01f);
 		LayerDesc.Priority = 0;
 		LayerDesc.PositionType = IStereoLayers::TrackerLocked;
-		LayerDesc.ShapeType = IStereoLayers::QuadLayer;
 		LayerDesc.Texture = nullptr;
 		UELayer = MakeShareable(new FLayer(NextLayerId++, LayerDesc));
 	}
@@ -422,15 +420,18 @@ bool FSplash::GetSplash(unsigned InSplashLayerIndex, FOculusSplashDesc& OutDesc)
 
 IStereoLayers::FLayerDesc FSplash::StereoLayerDescFromOculusSplashDesc(FOculusSplashDesc OculusDesc)
 {
-	bool bIsCubemap = OculusDesc.LoadedTexture->GetTextureCube() != nullptr;
-
 	IStereoLayers::FLayerDesc LayerDesc;
+	if (OculusDesc.LoadedTexture->GetTextureCube() != nullptr)
+	{
+		LayerDesc.SetShape<FCubemapLayer>();
+	}
+	// else LayerDesc.Shape defaults to FQuadLayer
+
 	LayerDesc.Transform = OculusDesc.TransformInMeters * FTransform(OculusHMD->GetSplashRotation().Quaternion());
 	LayerDesc.QuadSize = OculusDesc.QuadSizeInMeters;
 	LayerDesc.UVRect = FBox2D(OculusDesc.TextureOffset, OculusDesc.TextureOffset + OculusDesc.TextureScale);
 	LayerDesc.Priority = INT32_MAX - (int32)(OculusDesc.TransformInMeters.GetTranslation().X * 1000.f);
 	LayerDesc.PositionType = IStereoLayers::TrackerLocked;
-	LayerDesc.ShapeType = bIsCubemap ? IStereoLayers::CubemapLayer : IStereoLayers::QuadLayer;
 	LayerDesc.Texture = OculusDesc.LoadedTexture;
 	LayerDesc.Flags = IStereoLayers::LAYER_FLAG_QUAD_PRESERVE_TEX_RATIO | (OculusDesc.bNoAlphaChannel ? IStereoLayers::LAYER_FLAG_TEX_NO_ALPHA_CHANNEL : 0) | (OculusDesc.bIsDynamic ? IStereoLayers::LAYER_FLAG_TEX_CONTINUOUS_UPDATE : 0);
 
diff --git a/Engine/Plugins/Runtime/OpenXR/Source/OpenXRHMD/Private/OpenXRHMD.cpp b/Engine/Plugins/Runtime/OpenXR/Source/OpenXRHMD/Private/OpenXRHMD.cpp
--- a/Engine/Plugins/Runtime/OpenXR/Source/OpenXRHMD/Private/OpenXRHMD.cpp
+++ b/Engine/Plugins/Runtime/OpenXR/Source/OpenXRHMD/Private/OpenXRHMD.cpp
@@ -833,6 +833,8 @@ FOpenXRHMD::FOpenXRHMD(const FAutoRegister& AutoRegister, XrInstance InInstance,
 	, SelectedViewConfigurationType(XR_VIEW_CONFIGURATION_TYPE_PRIMARY_STEREO)
 	, RenderBridge(InRenderBridge)
 	, RendererModule(nullptr)
+	, LastRequestedSwapchainFormat(0)
+	, LastRequestedDepthSwapchainFormat(0)
 {
 	FrameState.type = XR_TYPE_FRAME_STATE;
 	FrameState.next = nullptr;
@@ -1206,14 +1208,18 @@ bool FOpenXRHMD::AllocateRenderTargetTexture(uint32 Index, uint32 SizeX, uint32
 	// We need to ensure we can sample from the texture in CopyTexture
 	Flags |= TexCreate_ShaderResource;
 
-	Swapchain = RenderBridge->CreateSwapchain(Session, Format, SizeX, SizeY, NumMips, NumSamples, Flags, TargetableTextureFlags);
-	if (!Swapchain)
+	if (Swapchain == nullptr || Format != LastRequestedSwapchainFormat || Swapchain->GetTexture2D()->GetSizeX() != SizeX || Swapchain->GetTexture2D()->GetSizeY() != SizeY)
 	{
-		return false;
+		Swapchain = RenderBridge->CreateSwapchain(Session, Format, SizeX, SizeY, NumMips, NumSamples, Flags, TargetableTextureFlags);
+		if (!Swapchain)
+		{
+			return false;
+		}
 	}
-	
+
 	// Grab the presentation texture out of the swapchain.
 	OutTargetableTexture = OutShaderResourceTexture = (FTexture2DRHIRef&)Swapchain->GetTextureRef();
+	LastRequestedSwapchainFormat = Format;
 
 	bNeedReAllocatedDepth = bDepthExtensionSupported;
 
@@ -1228,15 +1234,19 @@ bool FOpenXRHMD::AllocateDepthTexture(uint32 Index, uint32 SizeX, uint32 SizeY,
 		return false;
 	}
 
-	DepthSwapchain = RenderBridge->CreateSwapchain(Session, PF_DepthStencil, SizeX, SizeY, FMath::Max(NumMips, 1u), NumSamples, 0, TexCreate_DepthStencilTargetable);
-	if (!DepthSwapchain)
+	if (DepthSwapchain == nullptr || Format != LastRequestedDepthSwapchainFormat || DepthSwapchain->GetTexture2D()->GetSizeX() != SizeX || DepthSwapchain->GetTexture2D()->GetSizeY() != SizeY)
 	{
-		return false;
+		DepthSwapchain = RenderBridge->CreateSwapchain(Session, PF_DepthStencil, SizeX, SizeY, FMath::Max(NumMips, 1u), NumSamples, 0, TexCreate_DepthStencilTargetable);
+		if (!DepthSwapchain)
+		{
+			return false;
+		}
 	}
 
 	bNeedReAllocatedDepth = false;
 
 	OutTargetableTexture = OutShaderResourceTexture = (FTexture2DRHIRef&)DepthSwapchain->GetTextureRef();
+	LastRequestedDepthSwapchainFormat = Format;
 
 	return true;
 }
diff --git a/Engine/Plugins/Runtime/OpenXR/Source/OpenXRHMD/Private/OpenXRHMD.h b/Engine/Plugins/Runtime/OpenXR/Source/OpenXRHMD/Private/OpenXRHMD.h
--- a/Engine/Plugins/Runtime/OpenXR/Source/OpenXRHMD/Private/OpenXRHMD.h
+++ b/Engine/Plugins/Runtime/OpenXR/Source/OpenXRHMD/Private/OpenXRHMD.h
@@ -232,6 +232,8 @@ private:
 
 	FXRSwapChainPtr			Swapchain;
 	FXRSwapChainPtr			DepthSwapchain;
+	uint8					LastRequestedSwapchainFormat;
+	uint8					LastRequestedDepthSwapchainFormat;
 
 	TArray<FHMDViewMesh>	HiddenAreaMeshes;
 	TArray<FHMDViewMesh>	VisibleAreaMeshes;
diff --git a/Engine/Plugins/Runtime/OpenXR/Source/OpenXRHMD/Private/OpenXRHMD_RenderBridge.cpp b/Engine/Plugins/Runtime/OpenXR/Source/OpenXRHMD/Private/OpenXRHMD_RenderBridge.cpp
--- a/Engine/Plugins/Runtime/OpenXR/Source/OpenXRHMD/Private/OpenXRHMD_RenderBridge.cpp
+++ b/Engine/Plugins/Runtime/OpenXR/Source/OpenXRHMD/Private/OpenXRHMD_RenderBridge.cpp
@@ -42,7 +42,7 @@ public:
 		return &Binding;
 	}
 
-	virtual FXRSwapChainPtr CreateSwapchain(XrSession InSession, uint8 Format, uint32 SizeX, uint32 SizeY, uint32 NumMips, uint32 NumSamples, uint32 Flags, uint32 TargetableTextureFlags)
+	virtual FXRSwapChainPtr CreateSwapchain(XrSession InSession, uint8 Format, uint32 SizeX, uint32 SizeY, uint32 NumMips, uint32 NumSamples, uint32 Flags, uint32 TargetableTextureFlags) override final
 	{
 		return CreateSwapchain_D3D11(InSession, Format, SizeX, SizeY, NumMips, NumSamples, Flags, TargetableTextureFlags);
 	}
@@ -78,7 +78,7 @@ public:
 		return &Binding;
 	}
 
-	virtual FXRSwapChainPtr CreateSwapchain(XrSession InSession, uint8 Format, uint32 SizeX, uint32 SizeY, uint32 NumMips, uint32 NumSamples, uint32 Flags, uint32 TargetableTextureFlags)
+	virtual FXRSwapChainPtr CreateSwapchain(XrSession InSession, uint8 Format, uint32 SizeX, uint32 SizeY, uint32 NumMips, uint32 NumSamples, uint32 Flags, uint32 TargetableTextureFlags) override final
 	{
 		return CreateSwapchain_D3D12(InSession, Format, SizeX, SizeY, NumMips, NumSamples, Flags, TargetableTextureFlags);
 	}
@@ -128,7 +128,7 @@ public:
 		return nullptr;
 	}
 
-	virtual FXRSwapChainPtr CreateSwapchain(XrSession InSession, uint8 Format, uint32 SizeX, uint32 SizeY, uint32 NumMips, uint32 NumSamples, uint32 Flags, uint32 TargetableTextureFlags)
+	virtual FXRSwapChainPtr CreateSwapchain(XrSession InSession, uint8 Format, uint32 SizeX, uint32 SizeY, uint32 NumMips, uint32 NumSamples, uint32 Flags, uint32 TargetableTextureFlags) override final
 	{
 		return CreateSwapchain_OpenGL(InSession, Format, SizeX, SizeY, NumMips, NumSamples, Flags, TargetableTextureFlags);
 	}
@@ -210,7 +210,7 @@ public:
 		return AdapterLuid;
 	}
 
-	virtual FXRSwapChainPtr CreateSwapchain(XrSession InSession, uint8 Format, uint32 SizeX, uint32 SizeY, uint32 NumMips, uint32 NumSamples, uint32 Flags, uint32 TargetableTextureFlags)
+	virtual FXRSwapChainPtr CreateSwapchain(XrSession InSession, uint8 Format, uint32 SizeX, uint32 SizeY, uint32 NumMips, uint32 NumSamples, uint32 Flags, uint32 TargetableTextureFlags) override final
 	{
 		return CreateSwapchain_Vulkan(InSession, Format, SizeX, SizeY, NumMips, NumSamples, Flags, TargetableTextureFlags);
 	}
diff --git a/Engine/Plugins/Runtime/OpenXR/Source/OpenXRHMD/Private/OpenXRHMD_Swapchain.cpp b/Engine/Plugins/Runtime/OpenXR/Source/OpenXRHMD/Private/OpenXRHMD_Swapchain.cpp
--- a/Engine/Plugins/Runtime/OpenXR/Source/OpenXRHMD/Private/OpenXRHMD_Swapchain.cpp
+++ b/Engine/Plugins/Runtime/OpenXR/Source/OpenXRHMD/Private/OpenXRHMD_Swapchain.cpp
@@ -33,7 +33,7 @@ void FOpenXRSwapchain::IncrementSwapChainIndex_RHIThread(int64 Timeout)
 	WaitInfo.timeout = Timeout;
 	XR_ENSURE(xrWaitSwapchainImage(Handle, &WaitInfo));
 
-	GDynamicRHI->RHIAliasTextureResources(RHITexture, RHITextureSwapChain[SwapChainIndex_RHIThread]);
+	GDynamicRHI->RHIAliasTextureResources((FTextureRHIRef&)RHITexture, (FTextureRHIRef&)RHITextureSwapChain[SwapChainIndex_RHIThread]);
 }
 
 void FOpenXRSwapchain::ReleaseCurrentImage_RHIThread()
@@ -194,13 +194,12 @@ FXRSwapChainPtr CreateSwapchain_D3D11(XrSession InSession, uint8 Format, uint32
 
 	FD3D11DynamicRHI* DynamicRHI = static_cast<FD3D11DynamicRHI*>(GDynamicRHI);
 	TArray<FTextureRHIRef> TextureChain;
-	// @todo: Once things settle down, the chain target will be created below in the CreateXRSwapChain call, via an RHI "CreateAliasedTexture" call.
 	TArray<XrSwapchainImageD3D11KHR> Images = EnumerateImages<XrSwapchainImageD3D11KHR>(Swapchain, XR_TYPE_SWAPCHAIN_IMAGE_D3D11_KHR);
-	FTextureRHIRef ChainTarget = static_cast<FTextureRHIRef>(DynamicRHI->RHICreateTexture2DFromResource(GPixelFormats[Format].UnrealFormat, TargetableTextureFlags, bDepthStencil ? FClearValueBinding::DepthFar : FClearValueBinding::Black, Images[0].texture));
 	for (const auto& Image : Images)
 	{
 		TextureChain.Add(static_cast<FTextureRHIRef>(DynamicRHI->RHICreateTexture2DFromResource(GPixelFormats[Format].UnrealFormat, TargetableTextureFlags, bDepthStencil ? FClearValueBinding::DepthFar : FClearValueBinding::Black, Image.texture)));
 	}
+	FTextureRHIRef ChainTarget = static_cast<FTextureRHIRef>(GDynamicRHI->RHICreateAliasedTexture((FTextureRHIRef&)TextureChain[0]));
 
 	return CreateXRSwapChain<FOpenXRSwapchain>(MoveTemp(TextureChain), ChainTarget, Swapchain);
 }
@@ -234,13 +233,12 @@ FXRSwapChainPtr CreateSwapchain_D3D12(XrSession InSession, uint8 Format, uint32
 
 	FD3D12DynamicRHI* DynamicRHI = static_cast<FD3D12DynamicRHI*>(GDynamicRHI);
 	TArray<FTextureRHIRef> TextureChain;
-	// @todo: Once things settle down, the chain target will be created below in the CreateXRSwapChain call, via an RHI "CreateAliasedTexture" call.
 	TArray<XrSwapchainImageD3D12KHR> Images = EnumerateImages<XrSwapchainImageD3D12KHR>(Swapchain, XR_TYPE_SWAPCHAIN_IMAGE_D3D12_KHR);
-	FTextureRHIRef ChainTarget = static_cast<FTextureRHIRef>(DynamicRHI->RHICreateTexture2DFromResource(GPixelFormats[Format].UnrealFormat, TargetableTextureFlags, FClearValueBinding::Black, Images[0].texture));
 	for (const auto& Image : Images)
 	{
 		TextureChain.Add(static_cast<FTextureRHIRef>(DynamicRHI->RHICreateTexture2DFromResource(GPixelFormats[Format].UnrealFormat, TargetableTextureFlags, FClearValueBinding::Black, Image.texture)));
 	}
+	FTextureRHIRef ChainTarget = static_cast<FTextureRHIRef>(GDynamicRHI->RHICreateAliasedTexture((FTextureRHIRef&)TextureChain[0]));
 
 	return CreateXRSwapChain<FOpenXRSwapchain>(MoveTemp(TextureChain), ChainTarget, Swapchain);
 }
@@ -263,13 +261,12 @@ FXRSwapChainPtr CreateSwapchain_OpenGL(XrSession InSession, uint8 Format, uint32
 
 	TArray<FTextureRHIRef> TextureChain;
 	FOpenGLDynamicRHI* DynamicRHI = static_cast<FOpenGLDynamicRHI*>(GDynamicRHI);
-	// @todo: Once things settle down, the chain target will be created below in the CreateXRSwapChain call, via an RHI "CreateAliasedTexture" call.
 	TArray<XrSwapchainImageOpenGLKHR> Images = EnumerateImages<XrSwapchainImageOpenGLKHR>(Swapchain, XR_TYPE_SWAPCHAIN_IMAGE_OPENGL_KHR);
-	FTextureRHIRef ChainTarget = static_cast<FTextureRHIRef>(DynamicRHI->RHICreateTexture2DFromResource(GPixelFormats[Format].UnrealFormat, SizeX, SizeY, NumMips, NumSamples, 1, FClearValueBinding::Black, Images[0].image, Flags));
 	for (const auto& Image : Images)
 	{
 		TextureChain.Add(static_cast<FTextureRHIRef>(DynamicRHI->RHICreateTexture2DFromResource(GPixelFormats[Format].UnrealFormat, SizeX, SizeY, NumMips, NumSamples, 1, FClearValueBinding::Black, Image.image, Flags)));
 	}
+	FTextureRHIRef ChainTarget = static_cast<FTextureRHIRef>(GDynamicRHI->RHICreateAliasedTexture((FTextureRHIRef&)TextureChain[0]));
 
 	return CreateXRSwapChain<FOpenXRSwapchain>(MoveTemp(TextureChain), ChainTarget, Swapchain);
 }
@@ -297,13 +294,12 @@ FXRSwapChainPtr CreateSwapchain_Vulkan(XrSession InSession, uint8 Format, uint32
 
 	TArray<FTextureRHIRef> TextureChain;
 	FVulkanDynamicRHI* DynamicRHI = static_cast<FVulkanDynamicRHI*>(GDynamicRHI);
-	// @todo: Once things settle down, the chain target will be created below in the CreateXRSwapChain call, via an RHI "CreateAliasedTexture" call.
 	TArray<XrSwapchainImageVulkanKHR> Images = EnumerateImages<XrSwapchainImageVulkanKHR>(Swapchain, XR_TYPE_SWAPCHAIN_IMAGE_VULKAN_KHR);
-	FTextureRHIRef ChainTarget = static_cast<FTextureRHIRef>(DynamicRHI->RHICreateTexture2DFromResource(GPixelFormats[Format].UnrealFormat, SizeX, SizeY, NumMips, NumSamples, Images[0].image, Flags));
 	for (const auto& Image : Images)
 	{
 		TextureChain.Add(static_cast<FTextureRHIRef>(DynamicRHI->RHICreateTexture2DFromResource(GPixelFormats[Format].UnrealFormat, SizeX, SizeY, NumMips, NumSamples, Image.image, Flags)));
 	}
+	FTextureRHIRef ChainTarget = static_cast<FTextureRHIRef>(GDynamicRHI->RHICreateAliasedTexture((FTextureRHIRef&)TextureChain[0]));
 
 	return CreateXRSwapChain<FOpenXRSwapchain>(MoveTemp(TextureChain), ChainTarget, Swapchain);
 }
diff --git a/Engine/Plugins/Runtime/Steam/SteamVR/Source/SteamVR/Classes/SteamVRHQLayerShape.h b/Engine/Plugins/Runtime/Steam/SteamVR/Source/SteamVR/Classes/SteamVRHQLayerShape.h
new file mode 100644
--- /dev/null
+++ b/Engine/Plugins/Runtime/Steam/SteamVR/Source/SteamVR/Classes/SteamVRHQLayerShape.h
@@ -0,0 +1,38 @@
+// Copyright 1998-2019 Epic Games, Inc. All Rights Reserved.
+
+#pragma once
+#include "Components/StereoLayerComponent.h"
+#include "SteamVRHQLayerShape.generated.h"
+
+UCLASS(meta = (DisplayName = "High Quality Layer (SteamVR)"))
+class STEAMVR_API USteamVRHQStereoLayerShape : public UStereoLayerShapeQuad
+{
+	GENERATED_BODY()
+public:
+	virtual void ApplyShape(IStereoLayers::FLayerDesc& LayerDesc) override;
+
+	UPROPERTY(EditAnywhere, BlueprintReadOnly, export, Category = "HQ Layer Properties")
+	bool bCurved;
+	
+	UPROPERTY(EditAnywhere, BlueprintReadOnly, export, Category = "HQ Layer Properties")
+	bool bAntiAlias;
+	
+	/** For curved layers, sets the minimum distance from the layer used to automatically curve
+      * the surface around the viewer.  The minimum distance is when the layer is most curved. */
+	UPROPERTY(EditAnywhere, BlueprintReadOnly, export, Category = "HQ Layer Properties")
+	float AutoCurveMinDistance;
+	
+	/** For curved layers, sets the maximum distance from the layer used to automatically curve
+	  * the surface around the viewer.  The maximum distance is when the layer is the least curved. */
+	UPROPERTY(EditAnywhere, BlueprintReadOnly, export, Category = "HQ Layer Properties")
+	float AutoCurveMaxDistance;
+
+	UFUNCTION(BlueprintCallable, Category = "Components|Stereo Layer")
+	void SetCurved(bool InCurved);
+	UFUNCTION(BlueprintCallable, Category = "Components|Stereo Layer")
+	void SetAntiAlias(bool InAntiAlias);
+	UFUNCTION(BlueprintCallable, Category = "Components|Stereo Layer")
+	void SetAutoCurveMinDistance(float InDistance);
+	UFUNCTION(BlueprintCallable, Category = "Components|Stereo Layer")
+	void SetAutoCurveMaxDistance(float InDistance);
+};
\ No newline at end of file
diff --git a/Engine/Plugins/Runtime/Steam/SteamVR/Source/SteamVR/Private/SteamVRHMD.cpp b/Engine/Plugins/Runtime/Steam/SteamVR/Source/SteamVR/Private/SteamVRHMD.cpp
--- a/Engine/Plugins/Runtime/Steam/SteamVR/Source/SteamVR/Private/SteamVRHMD.cpp
+++ b/Engine/Plugins/Runtime/Steam/SteamVR/Source/SteamVR/Private/SteamVRHMD.cpp
@@ -1608,6 +1608,13 @@ bool FSteamVRHMD::AllocateRenderTargetTexture(uint32 Index, uint32 SizeX, uint32
 	MetalBridge* MetalBridgePtr = (MetalBridge*)pBridge.GetReference();
 #endif
 
+	if (pBridge != nullptr && pBridge->GetSwapChain() != nullptr && pBridge->GetSwapChain()->GetTexture2D() != nullptr && pBridge->GetSwapChain()->GetTexture2D()->GetSizeX() == SizeX && pBridge->GetSwapChain()->GetTexture2D()->GetSizeY() == SizeY)
+	{
+		OutTargetableTexture = (FTexture2DRHIRef&)pBridge->GetSwapChain()->GetTextureRef();
+		OutShaderResourceTexture = OutTargetableTexture;
+		return true;
+	}
+
 	for (uint32 SwapChainIter = 0; SwapChainIter < SteamVRSwapChainLength; ++SwapChainIter)
 	{
 #if PLATFORM_MAC
@@ -1630,16 +1637,7 @@ bool FSteamVRHMD::AllocateRenderTargetTexture(uint32 Index, uint32 SizeX, uint32
 
 		if (BindingTexture == nullptr)
 		{
-			// For this iteration, create a temporary texture resource that will get stomped the first time we do our bind/alias
-			// pass on this object.
-			// Note, even though this is a wasted allocation, it never gets bound or used in any way, so should be minimally impactful.
-			FTexture2DRHIRef TempTargetableTexture, TempShaderResourceTexture;
-			RHICreateTargetableShaderResource2D(SizeX, SizeY, PF_B8G8R8A8, 1, TexCreate_None, TexCreate_RenderTargetable, false, CreateInfo, TempTargetableTexture, TempShaderResourceTexture, NumSamples);
-			check(TempTargetableTexture == TempShaderResourceTexture);
-			BindingTexture = TempTargetableTexture;
-
-			// Alias immediately.
-			GDynamicRHI->RHIAliasTextureResources(BindingTexture, TargetableTexture);
+			BindingTexture = GDynamicRHI->RHICreateAliasedTexture((FTextureRHIRef&)TargetableTexture);
 		}
 	}
 
@@ -1705,16 +1703,9 @@ bool FSteamVRHMD::AllocateDepthTexture(uint32 Index, uint32 SizeX, uint32 SizeY,
 
 		SwapChainTextures.Add((FTextureRHIRef&)TargetableTexture);
 
-		// As above, create a temp, stompable texture to alias over.
 		if (BindingTexture == nullptr)
 		{
-			FTexture2DRHIRef TempTargetableTexture, TempShaderResourceTexture;
-			RHICreateTargetableShaderResource2D(SizeX, SizeY, PF_DepthStencil, 1, Flags, TargetableTextureFlags, false, CreateInfo, TempTargetableTexture, TempShaderResourceTexture, 1);
-			check(TempTargetableTexture == TempShaderResourceTexture);
-			BindingTexture = TempTargetableTexture;
-
-			// Alias immediately.
-			GDynamicRHI->RHIAliasTextureResources(BindingTexture, TargetableTexture);
+			BindingTexture = GDynamicRHI->RHICreateAliasedTexture((FTextureRHIRef&)TargetableTexture);
 		}
 	}
 
diff --git a/Engine/Plugins/Runtime/Steam/SteamVR/Source/SteamVR/Private/SteamVRHMD.h b/Engine/Plugins/Runtime/Steam/SteamVR/Source/SteamVR/Private/SteamVRHMD.h
--- a/Engine/Plugins/Runtime/Steam/SteamVR/Source/SteamVR/Private/SteamVRHMD.h
+++ b/Engine/Plugins/Runtime/Steam/SteamVR/Source/SteamVR/Private/SteamVRHMD.h
@@ -178,8 +178,7 @@ public:
 	virtual bool AllocateDepthTexture(uint32 Index, uint32 SizeX, uint32 SizeY, uint8 Format, uint32 NumMips, uint32 Flags, uint32 TargetableTextureFlags, FTexture2DRHIRef& OutTargetableTexture, FTexture2DRHIRef& OutShaderResourceTexture, uint32 NumSamples = 1) override;
 
 	// IStereoLayers interface
-	// Create/Set/Get/Destroy inherited from TStereoLayerManager
-	virtual void GetAllocatedTexture(uint32 LayerId, FTextureRHIRef &Texture, FTextureRHIRef &LeftTexture) override;
+	// Create/Set/Get/Destroy and GetAllocatedTexture inherited from TStereoLayerManager
 	virtual bool ShouldCopyDebugLayersToSpectatorScreen() const override { return true; }
 
 	// ISceneViewExtension interface
diff --git a/Engine/Plugins/Runtime/Steam/SteamVR/Source/SteamVR/Private/SteamVRHQLayerShape.cpp b/Engine/Plugins/Runtime/Steam/SteamVR/Source/SteamVR/Private/SteamVRHQLayerShape.cpp
new file mode 100644
--- /dev/null
+++ b/Engine/Plugins/Runtime/Steam/SteamVR/Source/SteamVR/Private/SteamVRHQLayerShape.cpp
@@ -0,0 +1,51 @@
+// Copyright 1998-2019 Epic Games, Inc. All Rights Reserved.
+
+#include "SteamVRHQLayerShape.h"
+#include "SteamVRStereoLayers.h"
+#include "IStereoLayers.h"
+
+void USteamVRHQStereoLayerShape::ApplyShape(IStereoLayers::FLayerDesc& LayerDesc)
+{
+	LayerDesc.SetShape<FSteamVRHQLayer>(bCurved, bAntiAlias, AutoCurveMinDistance, AutoCurveMaxDistance);
+}
+
+void USteamVRHQStereoLayerShape::SetCurved(bool InCurved)
+{
+	if (InCurved == bCurved)
+	{
+		return;
+	}
+	bCurved = InCurved;
+	MarkStereoLayerDirty();
+}
+
+void USteamVRHQStereoLayerShape::SetAntiAlias(bool InAntiAlias)
+{
+	if (InAntiAlias == bAntiAlias)
+	{
+		return;
+	}
+	bAntiAlias = InAntiAlias;
+	MarkStereoLayerDirty();
+}
+
+void USteamVRHQStereoLayerShape::SetAutoCurveMinDistance(float InDistance)
+{
+	if (InDistance == AutoCurveMinDistance)
+	{
+		return;
+	}
+	AutoCurveMinDistance = InDistance;
+	MarkStereoLayerDirty();
+}
+
+void USteamVRHQStereoLayerShape::SetAutoCurveMaxDistance(float InDistance)
+{
+	if (InDistance == AutoCurveMaxDistance)
+	{
+		return;
+	}
+	AutoCurveMaxDistance = InDistance;
+	MarkStereoLayerDirty();
+}
+
diff --git a/Engine/Plugins/Runtime/Steam/SteamVR/Source/SteamVR/Private/SteamVRStereoLayers.cpp b/Engine/Plugins/Runtime/Steam/SteamVR/Source/SteamVR/Private/SteamVRStereoLayers.cpp
--- a/Engine/Plugins/Runtime/Steam/SteamVR/Source/SteamVR/Private/SteamVRStereoLayers.cpp
+++ b/Engine/Plugins/Runtime/Steam/SteamVR/Source/SteamVR/Private/SteamVRStereoLayers.cpp
@@ -1,5 +1,6 @@
 // Copyright 1998-2019 Epic Games, Inc. All Rights Reserved.
 //
+#include "SteamVRStereoLayers.h"
 #include "CoreMinimal.h"
 #include "SteamVRPrivate.h"
 #include "StereoLayerManager.h"
@@ -7,6 +8,8 @@
 #include "Misc/ScopeLock.h"
 #include "DefaultXRCamera.h"
 
+const FName FSteamVRHQLayer::ShapeName = FName("SteamVRHQLayer");
+
 #if STEAMVR_SUPPORTED_PLATFORMS
 
 // enable to ensure all openvr calls succeed
@@ -95,9 +98,33 @@ void FSteamVRHMD::UpdateLayer(struct FSteamVRLayer& Layer, uint32 LayerId, bool
 	{
 		UE_LOG(LogHMD, Warning, TEXT("Unsupported StereoLayer flag. SteamVR StereoLayers do not support disabling alpha renderding. Make the texture opaque instead."));
 	}
-	if (Layer.LayerDesc.ShapeType != IStereoLayers::QuadLayer)
+
+	const float WorldToMeterScale = GetWorldToMetersScale();
+	check(WorldToMeterScale > 0.f);
+
+	if (Layer.LayerDesc.HasShape<FSteamVRHQLayer>())
+	{
+		vr::VROverlayHandle_t ExistingHQLayer = VROverlay->GetHighQualityOverlay();
+		if (ExistingHQLayer != vr::k_ulOverlayHandleInvalid && ExistingHQLayer != Layer.OverlayHandle)
+		{
+			UE_LOG(LogHMD, Warning, TEXT("There can only be one high-quality stereo layer active at a time. Please disable or change one or more layers to non-high-quality."));
+		}
+		else
+		{
+			VROverlay->SetHighQualityOverlay(Layer.OverlayHandle);
+
+			const auto& HQSettings = Layer.LayerDesc.GetShape<FSteamVRHQLayer>();
+			VROverlay->SetOverlayFlag(Layer.OverlayHandle, vr::VROverlayFlags_Curved, HQSettings.bCurved);
+			VROverlay->SetOverlayFlag(Layer.OverlayHandle, vr::VROverlayFlags_RGSS4X, HQSettings.bAntiAlias);
+			if (HQSettings.bCurved)
+			{
+				VROverlay->SetOverlayAutoCurveDistanceRangeInMeters(Layer.OverlayHandle, HQSettings.AutoCurveMinDistance / WorldToMeterScale, HQSettings.AutoCurveMaxDistance / WorldToMeterScale);
+			}
+		}
+	}
+	else if (!Layer.LayerDesc.HasShape<FQuadLayer>())
 	{
-		UE_LOG(LogHMD, Warning, TEXT("Unsupported StereoLayer shape. SteamVR StereoLayers can only be Quads."));
+		UE_LOG(LogHMD, Warning, TEXT("Unsupported StereoLayer shape. SteamVR StereoLayers can only be Quads or High Quality Quads."));
 	}
 
 	// UVs
@@ -107,8 +134,6 @@ void FSteamVRHMD::UpdateLayer(struct FSteamVRLayer& Layer, uint32 LayerId, bool
     TextureBounds.vMin = Layer.LayerDesc.UVRect.Min.Y;
     TextureBounds.vMax = Layer.LayerDesc.UVRect.Max.Y;
 	OVR_VERIFY(VROverlay->SetOverlayTextureBounds(Layer.OverlayHandle, &TextureBounds));
-	const float WorldToMeterScale = GetWorldToMetersScale();
-	check(WorldToMeterScale > 0.f);
 	OVR_VERIFY(VROverlay->SetOverlayWidthInMeters(Layer.OverlayHandle, Layer.LayerDesc.QuadSize.X / WorldToMeterScale));
 	
 	float TexelAspect = 1.0f;
@@ -280,36 +305,6 @@ void FSteamVRHMD::UpdateStereoLayers_RenderThread()
 	}
 }
 
-void FSteamVRHMD::GetAllocatedTexture(uint32 LayerId, FTextureRHIRef &Texture, FTextureRHIRef &LeftTexture)
-{
-	Texture = LeftTexture = nullptr;
-	FSteamVRLayer* LayerFound = nullptr;
-	check(IsInRenderingThread()); // Not strictly necessary, as WithLayer uses a scope lock
-
-	WithLayer(LayerId, [&](FSteamVRLayer* LayerFound)
-	{
-		if (LayerFound && LayerFound->LayerDesc.Texture)
-		{
-			switch (LayerFound->LayerDesc.ShapeType)
-			{
-			case IStereoLayers::CubemapLayer:
-				Texture = LayerFound->LayerDesc.Texture->GetTextureCube();
-				LeftTexture = LayerFound->LayerDesc.LeftTexture ? LayerFound->LayerDesc.LeftTexture->GetTextureCube() : nullptr;
-				break;
-
-			case IStereoLayers::CylinderLayer:
-			case IStereoLayers::QuadLayer:
-				Texture = LayerFound->LayerDesc.Texture->GetTexture2D();
-				LeftTexture = LayerFound->LayerDesc.LeftTexture ? LayerFound->LayerDesc.LeftTexture->GetTexture2D() : nullptr;
-				break;
-
-			default:
-				break;
-			}
-		}
-	});
-}
-
 //=============================================================================
 IStereoLayers* FSteamVRHMD::GetStereoLayers()
 {
diff --git a/Engine/Plugins/Runtime/Steam/SteamVR/Source/SteamVR/Public/SteamVRStereoLayers.h b/Engine/Plugins/Runtime/Steam/SteamVR/Source/SteamVR/Public/SteamVRStereoLayers.h
new file mode 100644
--- /dev/null
+++ b/Engine/Plugins/Runtime/Steam/SteamVR/Source/SteamVR/Public/SteamVRStereoLayers.h
@@ -0,0 +1,23 @@
+// Copyright 1998-2019 Epic Games, Inc. All Rights Reserved.
+#pragma once
+#include "StereoLayerShapes.h"
+
+/** Class describing additional settings for high-quality layers. Note there can only be one HQ layer active at any given time. */
+class STEAMVR_API FSteamVRHQLayer : public IStereoLayerShape
+{
+	STEREO_LAYER_SHAPE_BOILERPLATE(FSteamVRHQLayer)
+public:
+
+	FSteamVRHQLayer() {}
+	FSteamVRHQLayer(bool bInCurved, bool bInAntiAlias, float InAutoCurveMinDistance, float InAutoCurveMaxDistance) :
+		bCurved(bInCurved),
+		bAntiAlias(bInAntiAlias),
+		AutoCurveMinDistance(InAutoCurveMinDistance),
+		AutoCurveMaxDistance(InAutoCurveMaxDistance)
+	{}
+
+	bool bCurved;
+	bool bAntiAlias;
+	float AutoCurveMinDistance;
+	float AutoCurveMaxDistance;
+};
\ No newline at end of file
diff --git a/Engine/Plugins/Runtime/WindowsMixedReality/Source/WindowsMixedRealityHMD/Private/WindowsMixedRealityStereoLayers.cpp b/Engine/Plugins/Runtime/WindowsMixedReality/Source/WindowsMixedRealityHMD/Private/WindowsMixedRealityStereoLayers.cpp
--- a/Engine/Plugins/Runtime/WindowsMixedReality/Source/WindowsMixedRealityHMD/Private/WindowsMixedRealityStereoLayers.cpp
+++ b/Engine/Plugins/Runtime/WindowsMixedReality/Source/WindowsMixedRealityHMD/Private/WindowsMixedRealityStereoLayers.cpp
@@ -26,7 +26,6 @@ IStereoLayers::FLayerDesc FWindowsMixedRealityStereoLayers::GetDebugCanvasLayerD
 	StereoLayerDesc.QuadSize = FVector2D(120.f, 120.f);
 #endif
 	StereoLayerDesc.PositionType = IStereoLayers::ELayerType::FaceLocked;
-	StereoLayerDesc.ShapeType = IStereoLayers::ELayerShape::QuadLayer;
 	StereoLayerDesc.Texture = Texture;
 	StereoLayerDesc.Flags = IStereoLayers::ELayerFlags::LAYER_FLAG_TEX_CONTINUOUS_UPDATE;
 	StereoLayerDesc.Flags |= IStereoLayers::ELayerFlags::LAYER_FLAG_QUAD_PRESERVE_TEX_RATIO;
diff --git a/Engine/Source/Editor/ComponentVisualizers/Private/StereoLayerComponentVisualizer.cpp b/Engine/Source/Editor/ComponentVisualizers/Private/StereoLayerComponentVisualizer.cpp
--- a/Engine/Source/Editor/ComponentVisualizers/Private/StereoLayerComponentVisualizer.cpp
+++ b/Engine/Source/Editor/ComponentVisualizers/Private/StereoLayerComponentVisualizer.cpp
@@ -9,36 +9,8 @@
 void FStereoLayerComponentVisualizer::DrawVisualization( const UActorComponent* Component, const FSceneView* View, FPrimitiveDrawInterface* PDI )
 {
 	const UStereoLayerComponent* StereoLayerComp = Cast<const UStereoLayerComponent>(Component);
-	if(StereoLayerComp != NULL)
+	if(StereoLayerComp != nullptr && StereoLayerComp->Shape != nullptr)
 	{
-        FLinearColor YellowColor = FColor(231, 239, 0, 255);
-        if(StereoLayerComp->StereoLayerShape == EStereoLayerShape::SLSH_QuadLayer)
-        {
-            const FVector2D QuadSize = StereoLayerComp->GetQuadSize() / 2.0f;
-            const FBox QuadBox(FVector(0.0f, -QuadSize.X, -QuadSize.Y), FVector(0.0f, QuadSize.X, QuadSize.Y));
-
-            DrawWireBox(PDI, StereoLayerComp->GetComponentTransform().ToMatrixWithScale(), QuadBox, YellowColor, 0);
-        }
-        else if(StereoLayerComp->StereoLayerShape == EStereoLayerShape::SLSH_CylinderLayer)
-        {
-            float ArcAngle = StereoLayerComp->CylinderOverlayArc * 180 / (StereoLayerComp->CylinderRadius * PI);
-            
-            FVector X = StereoLayerComp->GetComponentTransform().GetUnitAxis(EAxis::Type::X);
-            FVector Y = StereoLayerComp->GetComponentTransform().GetUnitAxis(EAxis::Type::Y);
-            FVector Base = StereoLayerComp->GetComponentTransform().GetLocation();
-            FVector HalfHeight = FVector(0, 0, StereoLayerComp->CylinderHeight/2);
-            
-            FVector LeftVertex = Base + StereoLayerComp->CylinderRadius * ( FMath::Cos(ArcAngle/2 * (PI/180.0f)) * X + FMath::Sin(ArcAngle/2 * (PI/180.0f)) * Y );
-            FVector RightVertex = Base + StereoLayerComp->CylinderRadius * ( FMath::Cos(-ArcAngle/2 * (PI/180.0f)) * X + FMath::Sin(-ArcAngle/2 * (PI/180.0f)) * Y );
-
-            DrawArc(PDI, Base + HalfHeight, X, Y, -ArcAngle/2, ArcAngle/2, StereoLayerComp->CylinderRadius, 10, YellowColor, 0);
-            
-            DrawArc(PDI, Base - HalfHeight, X, Y, -ArcAngle/2, ArcAngle/2, StereoLayerComp->CylinderRadius, 10, YellowColor, 0);
-
-            PDI->DrawLine( LeftVertex - HalfHeight, LeftVertex + HalfHeight, YellowColor, 0 );
-            
-            PDI->DrawLine( RightVertex - HalfHeight, RightVertex + HalfHeight, YellowColor, 0 );
-
-        }
+		StereoLayerComp->Shape->DrawShapeVisualization(View, PDI);
     }
 }
diff --git a/Engine/Source/Runtime/Core/Public/UObject/VRObjectVersion.h b/Engine/Source/Runtime/Core/Public/UObject/VRObjectVersion.h
--- a/Engine/Source/Runtime/Core/Public/UObject/VRObjectVersion.h
+++ b/Engine/Source/Runtime/Core/Public/UObject/VRObjectVersion.h
@@ -18,6 +18,9 @@ struct CORE_API FVRObjectVersion
 		// Change how ARSessionConfig stores plane detection configuration from bitmask to bools
 		UseBoolsForARSessionConfigPlaneDetectionConfiguration,
 
+		// Change how UStereoLayerComponent stores additional properties for non-quad layer types
+		UseSubobjectForStereoLayerShapeProperties,
+
 		// -----<new versions can be added above this line>-------------------------------------------------
 		VersionPlusOne,
 		LatestVersion = VersionPlusOne - 1
diff --git a/Engine/Source/Runtime/D3D12RHI/Private/D3D12RHIPrivate.h b/Engine/Source/Runtime/D3D12RHI/Private/D3D12RHIPrivate.h
--- a/Engine/Source/Runtime/D3D12RHI/Private/D3D12RHIPrivate.h
+++ b/Engine/Source/Runtime/D3D12RHI/Private/D3D12RHIPrivate.h
@@ -342,8 +342,10 @@ public:
 	virtual uint32 RHIGetViewportNextPresentGPUIndex(FRHIViewport* Viewport) final override;
 	virtual FTexture2DRHIRef RHIGetViewportBackBuffer(FRHIViewport* Viewport) final override;
 	virtual void RHIAliasTextureResources(FRHITexture* DestTexture, FRHITexture* SrcTexture) final override;
-	virtual void RHIAdvanceFrameFence() final override;
 	virtual FTextureRHIRef RHICreateAliasedTexture(FRHITexture* SourceTexture) final override;
+	virtual void RHIAliasTextureResources(FTextureRHIRef& DestTexture, FTextureRHIRef& SrcTexture) final override;
+	virtual FTextureRHIRef RHICreateAliasedTexture(FTextureRHIRef& SourceTexture) final override;
+	virtual void RHIAdvanceFrameFence() final override;
 	virtual void RHIAdvanceFrameForGetViewportBackBuffer(FRHIViewport* Viewport) final override;
 	virtual void RHIAcquireThreadOwnership() final override;
 	virtual void RHIReleaseThreadOwnership() final override;
diff --git a/Engine/Source/Runtime/D3D12RHI/Private/D3D12Texture.cpp b/Engine/Source/Runtime/D3D12RHI/Private/D3D12Texture.cpp
--- a/Engine/Source/Runtime/D3D12RHI/Private/D3D12Texture.cpp
+++ b/Engine/Source/Runtime/D3D12RHI/Private/D3D12Texture.cpp
@@ -364,7 +364,7 @@ void FD3D12TextureStats::D3D12TextureDeleted(TD3D12Texture2D<BaseResourceType>&
 	{
 		const D3D12_RESOURCE_DESC& Desc = D3D12Texture2D->GetDesc();
 		const int64 TextureSize = Texture.GetMemorySize();
-		ensure(TextureSize > 0 || (Texture.Flags & TexCreate_Virtual));
+		ensure(TextureSize > 0 || (Texture.Flags & TexCreate_Virtual) || Texture.GetAliasingSourceTexture() != nullptr);
 
 		UpdateD3D12TextureStats(Desc, -TextureSize, false, Texture.IsCubemap());
 
@@ -2819,10 +2819,32 @@ FTextureCubeRHIRef FD3D12DynamicRHI::RHICreateTextureCubeFromResource(EPixelForm
 }
 
 void FD3D12DynamicRHI::RHIAliasTextureResources(FRHITexture* DestTextureRHI, FRHITexture* SrcTextureRHI)
-{	
+{
 	FD3D12TextureBase* DestTexture = GetD3D12TextureFromRHITexture(DestTextureRHI);
 	FD3D12TextureBase* SrcTexture = GetD3D12TextureFromRHITexture(SrcTextureRHI);
 
+	// This path will potentially cause crashes, if the source texture is destroyed and we're still being used. This
+	// API path will be deprecated post 4.25. To avoid issues, use the version that takes FTextureRHIRef references instead.
+	check(false);
+
+	while (DestTexture && SrcTexture)
+	{
+		DestTexture->AliasResources(SrcTexture);
+
+		DestTexture = DestTexture->GetNextObject();
+		SrcTexture = SrcTexture->GetNextObject();
+	}
+}
+
+void FD3D12DynamicRHI::RHIAliasTextureResources(FTextureRHIRef& DestTextureRHI, FTextureRHIRef& SrcTextureRHI)
+{
+	FD3D12TextureBase* DestTexture = GetD3D12TextureFromRHITexture(DestTextureRHI);
+	FD3D12TextureBase* SrcTexture = GetD3D12TextureFromRHITexture(SrcTextureRHI);
+
+	// Make sure we keep a reference to the source texture we're aliasing, so we don't lose it if all other references
+	// go away but we're kept around.
+	DestTexture->SetAliasingSource(SrcTextureRHI);
+
 	while (DestTexture && SrcTexture)
 	{
 		DestTexture->AliasResources(SrcTexture);
@@ -2900,7 +2922,7 @@ TD3D12Texture2D<BaseResourceType>* FD3D12DynamicRHI::CreateAliasedD3D12Texture2D
 		}
 	}
 
-	RHIAliasTextureResources(Texture2D, SourceTexture);
+	RHIAliasTextureResources((FTextureRHIRef&)Texture2D, (FTextureRHIRef&)SourceTexture);
 
 	return Texture2D;
 }
@@ -2925,6 +2947,35 @@ FTextureRHIRef FD3D12DynamicRHI::RHICreateAliasedTexture(FRHITexture* SourceText
 	return nullptr;
 }
 
+FTextureRHIRef FD3D12DynamicRHI::RHICreateAliasedTexture(FTextureRHIRef& SourceTextureRHI)
+{
+	FD3D12TextureBase* SourceTexture = GetD3D12TextureFromRHITexture(SourceTextureRHI);
+	FTextureRHIRef ReturnTexture;
+	if (SourceTextureRHI->GetTexture2D() != nullptr)
+	{
+		ReturnTexture = CreateAliasedD3D12Texture2D<FD3D12BaseTexture2D>(static_cast<FD3D12Texture2D*>(SourceTextureRHI->GetTexture2D()));
+	}
+	else if (SourceTextureRHI->GetTexture2DArray() != nullptr)
+	{
+		ReturnTexture = CreateAliasedD3D12Texture2D<FD3D12BaseTexture2DArray>(static_cast<FD3D12Texture2DArray*>(SourceTextureRHI->GetTexture2DArray()));
+	}
+	else if (SourceTextureRHI->GetTextureCube() != nullptr)
+	{
+		ReturnTexture = CreateAliasedD3D12Texture2D<FD3D12BaseTextureCube>(static_cast<FD3D12TextureCube*>(SourceTextureRHI->GetTextureCube()));
+	}
+
+	if (ReturnTexture == nullptr)
+	{
+		UE_LOG(LogD3D12RHI, Error, TEXT("Currently FD3D12DynamicRHI::RHICreateAliasedTexture only supports 2D, 2D Array and Cube textures."));
+		return nullptr;
+	}
+
+	FD3D12TextureBase* DestTexture = GetD3D12TextureFromRHITexture(ReturnTexture);
+	DestTexture->SetAliasingSource(SourceTextureRHI);
+
+	return ReturnTexture;
+}
+
 void FD3D12DynamicRHI::RHICopySubTextureRegion(FRHITexture2D* SourceTextureRHI, FRHITexture2D* DestTextureRHI, FBox2D SourceBox, FBox2D DestinationBox)
 {
 	FD3D12TextureBase* SourceTexture = GetD3D12TextureFromRHITexture(SourceTextureRHI);
diff --git a/Engine/Source/Runtime/D3D12RHI/Private/D3D12Texture.h b/Engine/Source/Runtime/D3D12RHI/Private/D3D12Texture.h
--- a/Engine/Source/Runtime/D3D12RHI/Private/D3D12Texture.h
+++ b/Engine/Source/Runtime/D3D12RHI/Private/D3D12Texture.h
@@ -85,11 +85,18 @@ public:
 		MemorySize = InMemorySize;
 	}
 
+	void SetAliasingSource(FTextureRHIRef& SourceTextureRHI)
+	{
+		AliasingSourceTexture = SourceTextureRHI;
+	}
+
 	// Accessors.
 	FD3D12Resource* GetResource() const { return ResourceLocation.GetResource(); }
 	uint64 GetOffset() const { return ResourceLocation.GetOffsetFromBaseOfResource(); }
 	FD3D12ShaderResourceView* GetShaderResourceView() const { return ShaderResourceView; }
 	FD3D12BaseShaderResource* GetBaseShaderResource() const { return BaseShaderResource; }
+	inline const FTextureRHIRef& GetAliasingSourceTexture() const { return AliasingSourceTexture; }
+
 	void SetShaderResourceView(FD3D12ShaderResourceView* InShaderResourceView) { ShaderResourceView = InShaderResourceView; }
 
 	static inline bool ShouldDeferCmdListOperation(FRHICommandList* RHICmdList)
@@ -199,6 +206,8 @@ protected:
 	uint32	NumDepthStencilViews;
 
 	TMap<uint32, FD3D12LockedResource*> LockedMap;
+
+	FTextureRHIRef AliasingSourceTexture;
 };
 
 #if !PLATFORM_SUPPORTS_VIRTUAL_TEXTURES
diff --git a/Engine/Source/Runtime/Engine/Classes/Components/StereoLayerComponent.h b/Engine/Source/Runtime/Engine/Classes/Components/StereoLayerComponent.h
--- a/Engine/Source/Runtime/Engine/Classes/Components/StereoLayerComponent.h
+++ b/Engine/Source/Runtime/Engine/Classes/Components/StereoLayerComponent.h
@@ -45,7 +45,78 @@ enum EStereoLayerShape
 	SLSH_MAX,
 };
 
-/** Additional properties for equirect layers */
+UCLASS(EditInlineNew, Abstract, BlueprintType, CollapseCategories, Within = StereoLayerComponent)
+class ENGINE_API UStereoLayerShape : public UObject
+{
+	GENERATED_BODY()
+
+public:
+	virtual void ApplyShape(IStereoLayers::FLayerDesc& LayerDesc);
+#if WITH_EDITOR
+	virtual void DrawShapeVisualization(const class FSceneView* View, class FPrimitiveDrawInterface* PDI);
+#endif
+
+protected:
+	void MarkStereoLayerDirty();
+};
+
+UCLASS(meta = (DisplayName = "Quad Layer"))
+class ENGINE_API UStereoLayerShapeQuad : public UStereoLayerShape
+{
+	GENERATED_BODY()
+public:
+	virtual void ApplyShape(IStereoLayers::FLayerDesc& LayerDesc) override;
+#if WITH_EDITOR
+	virtual void DrawShapeVisualization(const class FSceneView* View, class FPrimitiveDrawInterface* PDI) override;
+#endif
+};
+
+UCLASS(meta = (DisplayName = "Cylinder Layer"))
+class ENGINE_API UStereoLayerShapeCylinder : public UStereoLayerShape
+{
+	GENERATED_BODY()
+
+public:
+	UStereoLayerShapeCylinder()
+		: Radius(100)
+		, OverlayArc(100)
+		, Height(50)
+	{}
+
+	/** Radial size of the rendered stereo layer cylinder **/
+	UPROPERTY(EditAnywhere, BlueprintReadOnly, export, Category = "Cylinder Properties")
+	float Radius;
+
+	/** Arc angle for the stereo layer cylinder **/
+	UPROPERTY(EditAnywhere, BlueprintReadOnly, export, Category = "Cylinder Properties")
+	float OverlayArc;
+
+	/** Height of the stereo layer cylinder **/
+	UPROPERTY(EditAnywhere, BlueprintReadOnly, export, Category = "Cylinder Properties")
+	int Height;
+
+	UFUNCTION(BlueprintCallable, Category = "Components|Stereo Layer")
+	void SetRadius(float InRadius);
+	UFUNCTION(BlueprintCallable, Category = "Components|Stereo Layer")
+	void SetOverlayArc(float InOverlayArc);
+	UFUNCTION(BlueprintCallable, Category = "Components|Stereo Layer")
+	void SetHeight(int InHeight);
+
+	virtual void ApplyShape(IStereoLayers::FLayerDesc& LayerDesc) override;
+#if WITH_EDITOR
+	virtual void DrawShapeVisualization(const class FSceneView* View, class FPrimitiveDrawInterface* PDI) override;
+#endif
+};
+
+UCLASS(meta = (DisplayName = "Cubemap Layer"))
+class ENGINE_API UStereoLayerShapeCubemap : public UStereoLayerShape
+{
+	GENERATED_BODY()
+public:
+	virtual void ApplyShape(IStereoLayers::FLayerDesc& LayerDesc) override;
+};
+
+/** Properties for equirect layers */
 USTRUCT(BlueprintType)
 struct FEquirectProps
 {
@@ -53,28 +124,28 @@ struct FEquirectProps
 public:
 
 	UPROPERTY(EditAnywhere, BlueprintReadOnly, export, Category = "StereoLayer | Equirect Layer Properties")
-		/** Left source texture UVRect, specifying portion of input texture corresponding to left eye. */
-		FBox2D LeftUVRect;
+	/** Left source texture UVRect, specifying portion of input texture corresponding to left eye. */
+	FBox2D LeftUVRect;
 
 	UPROPERTY(EditAnywhere, BlueprintReadOnly, export, Category = "StereoLayer | Equirect Layer Properties")
-		/** Right source texture UVRect, specifying portion of input texture corresponding to right eye. */
-		FBox2D RightUVRect;
+	/** Right source texture UVRect, specifying portion of input texture corresponding to right eye. */
+	FBox2D RightUVRect;
 
 	UPROPERTY(EditAnywhere, BlueprintReadOnly, export, Category = "StereoLayer | Equirect Layer Properties")
-		/** Left eye's texture coordinate scale after mapping to 2D. */
-		FVector2D LeftScale;
+	/** Left eye's texture coordinate scale after mapping to 2D. */
+	FVector2D LeftScale;
 
 	UPROPERTY(EditAnywhere, BlueprintReadOnly, export, Category = "StereoLayer | Equirect Layer Properties")
-		/** Right eye's texture coordinate scale after mapping to 2D. */
-		FVector2D RightScale;
+	/** Right eye's texture coordinate scale after mapping to 2D. */
+	FVector2D RightScale;
 
 	UPROPERTY(EditAnywhere, BlueprintReadOnly, export, Category = "StereoLayer | Equirect Layer Properties")
-		/** Left eye's texture coordinate bias after mapping to 2D. */
-		FVector2D LeftBias;
+	/** Left eye's texture coordinate bias after mapping to 2D. */
+	FVector2D LeftBias;
 
 	UPROPERTY(EditAnywhere, BlueprintReadOnly, export, Category = "StereoLayer | Equirect Layer Properties")
-		/** Right eye's texture coordinate bias after mapping to 2D. */
-		FVector2D RightBias;
+	/** Right eye's texture coordinate bias after mapping to 2D. */
+	FVector2D RightBias;
 
 public:
 
@@ -102,10 +173,68 @@ public:
 	 * @param Other The other FEquirectProps to compare with.
 	 * @return true if the are equal, false otherwise.
 	 */
-	bool operator==(const FEquirectProps& Other) const
-	{
-		return (LeftUVRect == Other.LeftUVRect) && (RightUVRect == Other.RightUVRect) && (LeftScale == Other.LeftScale) && (RightScale == Other.RightScale) && (LeftBias == Other.LeftBias) && (RightBias == Other.RightBias);
-	}
+	bool operator==(const FEquirectProps& Other) const;
+
+	/**
+	 * Compares FEquirectProps with an UStereoLayerShapeEquirect
+	 *
+	 * @param Other The UStereoLayerShapeEquirect to compare with.
+	 * @return true if the are equal, false otherwise.
+	 */
+	bool operator==(const class UStereoLayerShapeEquirect& Other) const;
+};
+
+
+UCLASS(meta = (DisplayName = "Equirect Layer"))
+class ENGINE_API UStereoLayerShapeEquirect : public UStereoLayerShape
+{
+	GENERATED_BODY()
+
+public:
+	UStereoLayerShapeEquirect()
+		: LeftUVRect(FBox2D(FVector2D(0.0f, 0.0f), FVector2D(1.0f, 1.0f)))
+		, RightUVRect(FBox2D(FVector2D(0.0f, 0.0f), FVector2D(1.0f, 1.0f)))
+		, LeftScale(FVector2D(1.0f, 1.0f))
+		, RightScale(FVector2D(1.0f, 1.0f))
+		, LeftBias(FVector2D(0.0f, 0.0f))
+		, RightBias(FVector2D(0.0f, 0.0f))
+	{}
+
+	UPROPERTY(EditAnywhere, BlueprintReadOnly, Category = "Equirect Properties")
+	/** Left source texture UVRect, specifying portion of input texture corresponding to left eye. */
+	FBox2D LeftUVRect;
+
+	UPROPERTY(EditAnywhere, BlueprintReadOnly, Category = "Equirect Properties")
+	/** Right source texture UVRect, specifying portion of input texture corresponding to right eye. */
+	FBox2D RightUVRect;
+
+	UPROPERTY(EditAnywhere, BlueprintReadOnly, Category = "Equirect Properties")
+	/** Left eye's texture coordinate scale after mapping to 2D. */
+	FVector2D LeftScale;
+
+	UPROPERTY(EditAnywhere, BlueprintReadOnly, Category = "Equirect Properties")
+	/** Right eye's texture coordinate scale after mapping to 2D. */
+	FVector2D RightScale;
+
+	UPROPERTY(EditAnywhere, BlueprintReadOnly, Category = "Equirect Properties")
+	/** Left eye's texture coordinate bias after mapping to 2D. */
+	FVector2D LeftBias;
+
+	UPROPERTY(EditAnywhere, BlueprintReadOnly, Category = "Equirect Properties")
+	/** Right eye's texture coordinate bias after mapping to 2D. */
+	FVector2D RightBias;
+
+	/**
+	 * Set Equirect layer properties: UVRect, Scale, and Bias
+	 * @param	LeftScale: Scale for left eye
+	 * @param	LeftBias: Bias for left eye
+	 * @param	RightScale: Scale for right eye
+	 * @param	RightBias: Bias for right eye
+	 */
+	UFUNCTION(BlueprintCallable, Category = "Components|Stereo Layer")
+	void SetEquirectProps(FEquirectProps InScaleBiases);
+
+	virtual void ApplyShape(IStereoLayers::FLayerDesc& LayerDesc) override;
 };
 
 /** 
@@ -121,10 +250,14 @@ public:
 
 	//~ Begin UObject Interface
 	void BeginDestroy() override;
+	virtual void Serialize(FArchive& Ar) override;
+	virtual void PostLoad() override;
+	//~ End UObject Interface
 
 	//~ Begin UActorComponent Interface
 	virtual void TickComponent(float DeltaTime, enum ELevelTick TickType, FActorComponentTickFunction *ThisTickFunction) override;
-	
+	//~ End UActorComponent Interface
+
 	/** 
 	 * Change the texture displayed on the stereo layer. 
 	 *
@@ -178,7 +311,7 @@ public:
 	 * @param	RightScale: Scale for right eye
 	 * @param	RightBias: Bias for right eye
 	 */
-	UFUNCTION(BlueprintCallable, Category="Components|Stereo Layer")
+	UFUNCTION(BlueprintCallable, Category = "Components|Stereo Layer", meta = (DeprecatedFunction, DeprecationMessage = "Use UStereoLayerShapeEquirect::SetEquirectProps() instead."))
 	void SetEquirectProps(FEquirectProps InScaleBiases);
 
 	/** 
@@ -208,13 +341,14 @@ public:
 	UPROPERTY(EditAnywhere, BlueprintReadWrite, Category= "StereoLayer")
 	uint32 bNoAlphaChannel:1;
 
+	void MarkStereoLayerDirty();
 protected:
 	/** Texture displayed on the stereo layer (is stereocopic textures are supported on the platfrom and more than one texture is provided, this will be the right eye) **/
 	UPROPERTY(EditAnywhere, BlueprintReadOnly, Category= "StereoLayer")
 	class UTexture* Texture;
 
-	/** Texture displayed on the stereo layer for left eye, if stereoscopic textures are supported on the platform **/
-	UPROPERTY(EditAnywhere, BlueprintReadOnly, Category = "StereoLayer | Cubemap Overlay Properties")
+	/** Texture displayed on the stereo layer for left eye, if stereoscopic textures are supported on the platform and by the layer shape **/
+	UPROPERTY(EditAnywhere, BlueprintReadOnly, Category = "Stereoscopic Properties")
 	class UTexture* LeftTexture;
 
 public:
@@ -224,36 +358,41 @@ public:
 
 protected:
 	/** Size of the rendered stereo layer quad **/
-	UPROPERTY(EditAnywhere, BlueprintReadOnly, export, Category="StereoLayer | Quad Overlay Properties")
+	UPROPERTY(EditAnywhere, BlueprintReadOnly, export, Category="StereoLayer")
 	FVector2D QuadSize;
 
 	/** UV coordinates mapped to the quad face **/
-	UPROPERTY(EditAnywhere, BlueprintReadOnly, export, Category="StereoLayer | Quad Overlay Properties")
+	UPROPERTY(EditAnywhere, BlueprintReadOnly, export, Category="StereoLayer")
 	FBox2D UVRect;
 
-	/** Radial size of the rendered stereo layer cylinder **/
-	UPROPERTY(EditAnywhere, BlueprintReadOnly, export, Category = "StereoLayer | Cylinder Overlay Properties")
-	float CylinderRadius;
+	UE_DEPRECATED(4.25, "float CylinderRadius is deprecated. Use the corresponding property of the UStereoLayerShapeCylinder subobject instead.")
+	UPROPERTY()
+	float CylinderRadius_DEPRECATED;
 
-	/** Arc angle for the stereo layer cylinder **/
-	UPROPERTY(EditAnywhere, BlueprintReadOnly, export, Category = "StereoLayer | Cylinder Overlay Properties")
-	float CylinderOverlayArc;
+	UE_DEPRECATED(4.25, "float CylinderOverlayArc is deprecated. Use the corresponding property of the UStereoLayerShapeCylinder subobject instead.")
+	UPROPERTY()
+	float CylinderOverlayArc_DEPRECATED;
 
-	/** Height of the stereo layer cylinder **/
-	UPROPERTY(EditAnywhere, BlueprintReadOnly, export, Category = "StereoLayer | Cylinder Overlay Properties")
-	int CylinderHeight;
+	UE_DEPRECATED(4.25, "int CylinderHeight is deprecated. Use the corresponding property of the UStereoLayerShapeCylinder subobject instead.")
+	UPROPERTY()
+	int CylinderHeight_DEPRECATED;
 
-	/** UVRects, scale, and biases per eye for equirect layers. UVRects are for input texture UVs per eye, and the latter 2 control scale and bias of texture coordinates after mapping to 2D.  **/
-	UPROPERTY(EditAnywhere, BlueprintReadOnly, export, Category = "StereoLayer | Equirect Layer Properties")
-	FEquirectProps EquirectProps;
+	UE_DEPRECATED(4.25, "FEquirectProps EquirectProps is deprecated. Use the corresponding properties of the UStereoLayerShapeEquirect subobject instead.")
+	UPROPERTY()
+	FEquirectProps EquirectProps_DEPRECATED;
 
 	/** Specifies how and where the quad is rendered to the screen **/
 	UPROPERTY(EditAnywhere, BlueprintReadOnly, export, Category="StereoLayer")
     TEnumAsByte<enum EStereoLayerType> StereoLayerType;
 
-	/** Specifies which type of layer it is.  Note that some shapes will be supported only on certain platforms! **/
-	UPROPERTY(EditAnywhere, BlueprintReadOnly, export, Category = "StereoLayer")
-	TEnumAsByte<enum EStereoLayerShape> StereoLayerShape;
+	UE_DEPRECATED(4.25, "TEnumAsByte<enum EStereoLayerShape> StereoLayerShape is deprecated. Use Shape instead.")
+	UPROPERTY()
+	TEnumAsByte<enum EStereoLayerShape> StereoLayerShape_DEPRECATED;
+
+	/** Specifies which shape of layer it is.  Note that some shapes will be supported only on certain platforms! **/
+	UPROPERTY(EditAnywhere, BlueprintReadOnly, NoClear, Instanced, Category = "StereoLayer", DisplayName="Stereo Layer Shape")
+	UStereoLayerShape* Shape;
+
 
 	/** Render priority among all stereo layers, higher priority render on top of lower priority **/
 	UPROPERTY(EditAnywhere, BlueprintReadOnly, export, Category= "StereoLayer")
@@ -274,5 +413,8 @@ private:
 
 	/** Last frames visiblity state **/
 	bool bLastVisible;
+
+	/** Set if the component was loaded from an old version */
+	bool bNeedsPostLoadFixup;
 };
 
diff --git a/Engine/Source/Runtime/Engine/Private/Components/StereoLayerComponent.cpp b/Engine/Source/Runtime/Engine/Private/Components/StereoLayerComponent.cpp
--- a/Engine/Source/Runtime/Engine/Private/Components/StereoLayerComponent.cpp
+++ b/Engine/Source/Runtime/Engine/Private/Components/StereoLayerComponent.cpp
@@ -1,12 +1,17 @@
 // Copyright 1998-2019 Epic Games, Inc. All Rights Reserved.
 
 #include "Components/StereoLayerComponent.h"
+#include "UObject/VRObjectVersion.h"
 #include "EngineGlobals.h"
 #include "Engine/Engine.h"
 #include "TextureResource.h"
 #include "Engine/Texture.h"
 #include "IStereoLayers.h"
+#include "StereoLayerShapes.h"
 #include "StereoRendering.h"
+#if WITH_EDITOR
+#include "SceneManagement.h"
+#endif
 
 UStereoLayerComponent::UStereoLayerComponent(const FObjectInitializer& ObjectInitializer)
 	: Super(ObjectInitializer)
@@ -18,21 +23,19 @@ UStereoLayerComponent::UStereoLayerComponent(const FObjectInitializer& ObjectIni
 	, bQuadPreserveTextureRatio(false)
 	, QuadSize(FVector2D(100.0f, 100.0f))
 	, UVRect(FBox2D(FVector2D(0.0f, 0.0f), FVector2D(1.0f, 1.0f)))
-	, CylinderRadius(100)
-	, CylinderOverlayArc(100)
-	, CylinderHeight(50)
-	, EquirectProps()
 	, StereoLayerType(SLT_FaceLocked)
-	, StereoLayerShape(SLSH_QuadLayer)
 	, Priority(0)
 	, bIsDirty(true)
 	, bTextureNeedsUpdate(false)
 	, LayerId(0)
 	, LastTransform(FTransform::Identity)
 	, bLastVisible(false)
+	, bNeedsPostLoadFixup(false)
 {
 	PrimaryComponentTick.bCanEverTick = true;
 	PrimaryComponentTick.TickGroup = TG_PrePhysics;
+	// The following somehow overrides subobjects read through serialization.
+	//Shape = ObjectInitializer.CreateDefaultSubobject<UStereoLayerShapeQuad>(this, TEXT("Shape"));
 }
 
 void UStereoLayerComponent::BeginDestroy()
@@ -47,6 +50,61 @@ void UStereoLayerComponent::BeginDestroy()
 	}
 }
 
+void UStereoLayerComponent::Serialize(FArchive& Ar)
+{
+	Ar.UsingCustomVersion(FVRObjectVersion::GUID);
+	Super::Serialize(Ar);
+	if (Ar.IsLoading() && Ar.CustomVer(FVRObjectVersion::GUID) < FVRObjectVersion::UseSubobjectForStereoLayerShapeProperties)
+	{
+		bNeedsPostLoadFixup = true; // Postponing fixups until after load, as we need to modify subobjects.
+	}
+}
+
+void UStereoLayerComponent::PostLoad()
+{
+	Super::PostLoad();
+	PRAGMA_DISABLE_DEPRECATION_WARNINGS
+	if (bNeedsPostLoadFixup)
+	{
+		switch (StereoLayerShape_DEPRECATED)
+		{
+			case SLSH_QuadLayer:
+			{
+				Shape = NewObject<UStereoLayerShapeQuad>(this);
+				break;
+			}
+			case SLSH_CubemapLayer:
+			{
+				Shape = NewObject<UStereoLayerShapeCubemap>(this);
+				break;
+			}
+			case SLSH_CylinderLayer:
+			{
+				auto Cylinder = NewObject<UStereoLayerShapeCylinder>(this);
+				Shape = Cylinder;
+				Cylinder->Height = CylinderHeight_DEPRECATED;
+				Cylinder->OverlayArc = CylinderOverlayArc_DEPRECATED;
+				Cylinder->Radius = CylinderRadius_DEPRECATED;
+				break;
+			}
+			case SLSH_EquirectLayer:
+			{
+				auto Equirect = NewObject<UStereoLayerShapeEquirect>(this);
+				Shape = Equirect;
+				Equirect->SetEquirectProps(EquirectProps_DEPRECATED);
+				break;
+			}
+		}
+		bNeedsPostLoadFixup = false;
+	}
+	PRAGMA_ENABLE_DEPRECATION_WARNINGS
+
+	if (Shape == nullptr)
+	{
+		Shape = NewObject<UStereoLayerShapeQuad>(this);
+	}
+}
+
 void UStereoLayerComponent::TickComponent(float DeltaTime, enum ELevelTick TickType, FActorComponentTickFunction *ThisTickFunction)
 {
 	Super::TickComponent( DeltaTime, TickType, ThisTickFunction );
@@ -81,75 +139,55 @@ void UStereoLayerComponent::TickComponent(float DeltaTime, enum ELevelTick TickT
 
 	if (bIsDirty)
 	{
-		IStereoLayers::FLayerDesc LayerDsec;
-		LayerDsec.Priority = Priority;
-		LayerDsec.QuadSize = QuadSize;
-		LayerDsec.UVRect = UVRect;
-		LayerDsec.Transform = Transform;
+		IStereoLayers::FLayerDesc LayerDesc;
+		LayerDesc.Priority = Priority;
+		LayerDesc.QuadSize = QuadSize;
+		LayerDesc.UVRect = UVRect;
+		LayerDesc.Transform = Transform;
 
 		if (Texture)
 		{
-			LayerDsec.Texture = Texture->Resource->TextureRHI;
-			LayerDsec.Flags |= (Texture->GetMaterialType() == MCT_TextureExternal) ? IStereoLayers::LAYER_FLAG_TEX_EXTERNAL : 0;
+			LayerDesc.Texture = Texture->Resource->TextureRHI;
+			LayerDesc.Flags |= (Texture->GetMaterialType() == MCT_TextureExternal) ? IStereoLayers::LAYER_FLAG_TEX_EXTERNAL : 0;
 		}
 		if (LeftTexture)
 		{
-			LayerDsec.LeftTexture = LeftTexture->Resource->TextureRHI;
+			LayerDesc.LeftTexture = LeftTexture->Resource->TextureRHI;
 		}
 				
-		LayerDsec.Flags |= (bLiveTexture) ? IStereoLayers::LAYER_FLAG_TEX_CONTINUOUS_UPDATE : 0;
-		LayerDsec.Flags |= (bNoAlphaChannel) ? IStereoLayers::LAYER_FLAG_TEX_NO_ALPHA_CHANNEL : 0;
-		LayerDsec.Flags |= (bQuadPreserveTextureRatio) ? IStereoLayers::LAYER_FLAG_QUAD_PRESERVE_TEX_RATIO : 0;
-		LayerDsec.Flags |= (bSupportsDepth) ? IStereoLayers::LAYER_FLAG_SUPPORT_DEPTH : 0;
-		LayerDsec.Flags |= (!bCurrVisible) ? IStereoLayers::LAYER_FLAG_HIDDEN : 0;
+		LayerDesc.Flags |= (bLiveTexture) ? IStereoLayers::LAYER_FLAG_TEX_CONTINUOUS_UPDATE : 0;
+		LayerDesc.Flags |= (bNoAlphaChannel) ? IStereoLayers::LAYER_FLAG_TEX_NO_ALPHA_CHANNEL : 0;
+		LayerDesc.Flags |= (bQuadPreserveTextureRatio) ? IStereoLayers::LAYER_FLAG_QUAD_PRESERVE_TEX_RATIO : 0;
+		LayerDesc.Flags |= (bSupportsDepth) ? IStereoLayers::LAYER_FLAG_SUPPORT_DEPTH : 0;
+		LayerDesc.Flags |= (!bCurrVisible) ? IStereoLayers::LAYER_FLAG_HIDDEN : 0;
 
 		switch (StereoLayerType)
 		{
 		case SLT_WorldLocked:
-			LayerDsec.PositionType = IStereoLayers::WorldLocked;
+			LayerDesc.PositionType = IStereoLayers::WorldLocked;
 			break;
 		case SLT_TrackerLocked:
-			LayerDsec.PositionType = IStereoLayers::TrackerLocked;
+			LayerDesc.PositionType = IStereoLayers::TrackerLocked;
 			break;
 		case SLT_FaceLocked:
-			LayerDsec.PositionType = IStereoLayers::FaceLocked;
+			LayerDesc.PositionType = IStereoLayers::FaceLocked;
 			break;
 		}
 
-		switch (StereoLayerShape)
+		// Set the correct layer shape and apply any shape-specific properties
+		if (Shape == nullptr)
 		{
-		case SLSH_QuadLayer:
-			LayerDsec.ShapeType = IStereoLayers::QuadLayer;
-			break;
-
-		case SLSH_CylinderLayer:
-			LayerDsec.ShapeType = IStereoLayers::CylinderLayer;
-			LayerDsec.CylinderRadius = CylinderRadius;
-			LayerDsec.CylinderOverlayArc = CylinderOverlayArc;
-			LayerDsec.CylinderHeight = CylinderHeight;
-			break;
-
-		case SLSH_CubemapLayer:
-			LayerDsec.ShapeType = IStereoLayers::CubemapLayer;
-			break;
-
-		case SLSH_EquirectLayer:
-			LayerDsec.ShapeType = IStereoLayers::EquirectLayer;
-			LayerDsec.EquirectProps = { EquirectProps.LeftUVRect, EquirectProps.RightUVRect, EquirectProps.LeftScale, EquirectProps.RightScale, EquirectProps.LeftBias, EquirectProps.RightBias };
-			break;
-
-		default:
-			break;
+			Shape = NewObject<UStereoLayerShapeQuad>(this);
 		}
-
+		Shape->ApplyShape(LayerDesc);
 
 		if (LayerId)
 		{
-			StereoLayers->SetLayerDesc(LayerId, LayerDsec);
+			StereoLayers->SetLayerDesc(LayerId, LayerDesc);
 		}
 		else
 		{
-			LayerId = StereoLayers->CreateLayer(LayerDsec);
+			LayerId = StereoLayers->CreateLayer(LayerDesc);
 		}
 		
 		LastTransform = Transform;
@@ -172,7 +210,7 @@ void UStereoLayerComponent::SetTexture(UTexture* InTexture)
 	}
 
 	Texture = InTexture;
-	bIsDirty = true;
+	MarkStereoLayerDirty();
 }
 
 void UStereoLayerComponent::SetLeftTexture(UTexture* InTexture)
@@ -183,7 +221,7 @@ void UStereoLayerComponent::SetLeftTexture(UTexture* InTexture)
 	}
 
 	LeftTexture = InTexture;
-	bIsDirty = true;
+	MarkStereoLayerDirty();
 }
 
 void UStereoLayerComponent::SetQuadSize(FVector2D InQuadSize)
@@ -194,7 +232,7 @@ void UStereoLayerComponent::SetQuadSize(FVector2D InQuadSize)
 	}
 
 	QuadSize = InQuadSize;
-	bIsDirty = true;
+	MarkStereoLayerDirty();
 }
 
 void UStereoLayerComponent::SetUVRect(FBox2D InUVRect)
@@ -205,17 +243,31 @@ void UStereoLayerComponent::SetUVRect(FBox2D InUVRect)
 	}
 
 	UVRect = InUVRect;
-	bIsDirty = true;
+	MarkStereoLayerDirty();
 }
 
 void UStereoLayerComponent::SetEquirectProps(FEquirectProps InEquirectProps)
 {
-	if (EquirectProps == InEquirectProps)
+	if (Shape->IsA<UStereoLayerShapeEquirect>())
+	{
+		Cast<UStereoLayerShapeEquirect>(Shape)->SetEquirectProps(InEquirectProps);
+	}
+}
+
+void UStereoLayerShapeEquirect::SetEquirectProps(FEquirectProps InEquirectProps)
+{
+	if (InEquirectProps == *this)
 	{
 		return;
 	}
-	EquirectProps = InEquirectProps;
-	bIsDirty = true;
+	LeftUVRect = InEquirectProps.LeftUVRect;
+	RightUVRect = InEquirectProps.RightUVRect;
+	LeftScale = InEquirectProps.LeftScale;
+	RightScale = InEquirectProps.RightScale;
+	LeftBias = InEquirectProps.LeftBias;
+	RightBias = InEquirectProps.RightBias;
+
+	MarkStereoLayerDirty();
 }
 
 void UStereoLayerComponent::SetPriority(int32 InPriority)
@@ -226,7 +278,40 @@ void UStereoLayerComponent::SetPriority(int32 InPriority)
 	}
 
 	Priority = InPriority;
-	bIsDirty = true;
+	MarkStereoLayerDirty();
+}
+
+void UStereoLayerShapeCylinder::SetRadius(float InRadius)
+{
+	if (Radius == InRadius)
+	{
+		return;
+	}
+
+	Radius = InRadius;
+	MarkStereoLayerDirty();
+}
+
+void UStereoLayerShapeCylinder::SetOverlayArc(float InOverlayArc)
+{
+	if (OverlayArc == InOverlayArc)
+	{
+		return;
+	}
+
+	OverlayArc = InOverlayArc;
+	MarkStereoLayerDirty();
+}
+
+void UStereoLayerShapeCylinder::SetHeight(int InHeight)
+{
+	if (Height == InHeight)
+	{
+		return;
+	}
+
+	Height = InHeight;
+	MarkStereoLayerDirty();
 }
 
 void UStereoLayerComponent::MarkTextureForUpdate()
@@ -234,3 +319,91 @@ void UStereoLayerComponent::MarkTextureForUpdate()
 	bTextureNeedsUpdate = true;
 }
 
+void UStereoLayerComponent::MarkStereoLayerDirty()
+{
+	bIsDirty = true;
+}
+
+void UStereoLayerShape::MarkStereoLayerDirty()
+{
+	check(GetOuter()->IsA<UStereoLayerComponent>());
+	Cast<UStereoLayerComponent>(GetOuter())->MarkStereoLayerDirty();
+}
+
+void UStereoLayerShapeCylinder::ApplyShape(IStereoLayers::FLayerDesc& LayerDesc)
+{
+	LayerDesc.SetShape<FCylinderLayer>(Radius, OverlayArc, Height);
+}
+
+void UStereoLayerShapeCubemap::ApplyShape(IStereoLayers::FLayerDesc& LayerDesc)
+{
+	LayerDesc.SetShape<FCubemapLayer>();
+}
+
+void UStereoLayerShapeEquirect::ApplyShape(IStereoLayers::FLayerDesc& LayerDesc)
+{
+	LayerDesc.SetShape<FEquirectLayer>(LeftUVRect, RightUVRect, LeftScale, RightScale, LeftBias, RightBias);
+}
+
+void UStereoLayerShapeQuad::ApplyShape(IStereoLayers::FLayerDesc& LayerDesc)
+{
+	LayerDesc.SetShape<FQuadLayer>();
+}
+
+void UStereoLayerShape::ApplyShape(IStereoLayers::FLayerDesc& LayerDesc)
+{
+}
+
+#if WITH_EDITOR
+void UStereoLayerShape::DrawShapeVisualization(const class FSceneView* View, class FPrimitiveDrawInterface* PDI)
+{
+}
+
+void UStereoLayerShapeQuad::DrawShapeVisualization(const class FSceneView* View, class FPrimitiveDrawInterface* PDI)
+{
+	FLinearColor YellowColor = FColor(231, 239, 0, 255);
+	check(GetOuter()->IsA<UStereoLayerComponent>());
+
+	auto StereoLayerComp = Cast<UStereoLayerComponent>(GetOuter());
+	const FVector2D QuadSize = StereoLayerComp->GetQuadSize() / 2.0f;
+	const FBox QuadBox(FVector(0.0f, -QuadSize.X, -QuadSize.Y), FVector(0.0f, QuadSize.X, QuadSize.Y));
+
+	DrawWireBox(PDI, StereoLayerComp->GetComponentTransform().ToMatrixWithScale(), QuadBox, YellowColor, 0);
+}
+
+void UStereoLayerShapeCylinder::DrawShapeVisualization(const class FSceneView* View, class FPrimitiveDrawInterface* PDI)
+{
+	FLinearColor YellowColor = FColor(231, 239, 0, 255);
+	check(GetOuter()->IsA<UStereoLayerComponent>());
+
+	auto StereoLayerComp = Cast<UStereoLayerComponent>(GetOuter());
+	float ArcAngle = OverlayArc * 180 / (Radius * PI);
+
+	FVector X = StereoLayerComp->GetComponentTransform().GetUnitAxis(EAxis::Type::X);
+	FVector Y = StereoLayerComp->GetComponentTransform().GetUnitAxis(EAxis::Type::Y);
+	FVector Base = StereoLayerComp->GetComponentTransform().GetLocation();
+	FVector HalfHeight = FVector(0, 0, Height / 2);
+
+	FVector LeftVertex = Base + Radius * (FMath::Cos(ArcAngle / 2 * (PI / 180.0f)) * X + FMath::Sin(ArcAngle / 2 * (PI / 180.0f)) * Y);
+	FVector RightVertex = Base + Radius * (FMath::Cos(-ArcAngle / 2 * (PI / 180.0f)) * X + FMath::Sin(-ArcAngle / 2 * (PI / 180.0f)) * Y);
+
+	DrawArc(PDI, Base + HalfHeight, X, Y, -ArcAngle / 2, ArcAngle / 2, Radius, 10, YellowColor, 0);
+
+	DrawArc(PDI, Base - HalfHeight, X, Y, -ArcAngle / 2, ArcAngle / 2, Radius, 10, YellowColor, 0);
+
+	PDI->DrawLine(LeftVertex - HalfHeight, LeftVertex + HalfHeight, YellowColor, 0);
+
+	PDI->DrawLine(RightVertex - HalfHeight, RightVertex + HalfHeight, YellowColor, 0);
+}
+#endif
+
+
+bool FEquirectProps::operator==(const class UStereoLayerShapeEquirect& Other) const
+{
+	return (LeftUVRect == Other.LeftUVRect) && (RightUVRect == Other.RightUVRect) && (LeftScale == Other.LeftScale) && (RightScale == Other.RightScale) && (LeftBias == Other.LeftBias) && (RightBias == Other.RightBias);
+}
+
+bool FEquirectProps::operator==(const FEquirectProps& Other) const
+{
+	return (LeftUVRect == Other.LeftUVRect) && (RightUVRect == Other.RightUVRect) && (LeftScale == Other.LeftScale) && (RightScale == Other.RightScale) && (LeftBias == Other.LeftBias) && (RightBias == Other.RightBias);
+}
diff --git a/Engine/Source/Runtime/Engine/Private/StereoLayerShapes.cpp b/Engine/Source/Runtime/Engine/Private/StereoLayerShapes.cpp
new file mode 100644
--- /dev/null
+++ b/Engine/Source/Runtime/Engine/Private/StereoLayerShapes.cpp
@@ -0,0 +1,9 @@
+// Copyright 1998-2019 Epic Games, Inc. All Rights Reserved.
+
+#include "StereoLayerShapes.h"
+#include "RHI.h"
+
+const FName FQuadLayer::ShapeName = FName("QuadLayer");
+const FName FCylinderLayer::ShapeName = FName("CylinderLayer");
+const FName FCubemapLayer::ShapeName = FName("CubemapLayer");
+const FName FEquirectLayer::ShapeName = FName("EquirectLayer");
diff --git a/Engine/Source/Runtime/Engine/Public/IStereoLayers.h b/Engine/Source/Runtime/Engine/Public/IStereoLayers.h
--- a/Engine/Source/Runtime/Engine/Public/IStereoLayers.h
+++ b/Engine/Source/Runtime/Engine/Public/IStereoLayers.h
@@ -8,10 +8,50 @@
 #pragma once
 
 #include "CoreMinimal.h"
+#include "StereoLayerShapes.h"
 #include "RHI.h"
 
 class IStereoLayers
 {
+
+	// The pointer to the shape of a stereo layer is wrapped in this simple class that clones and destroys the shape.
+	// This is to avoid having to implement non-default copy and assignment operations for the entire FLayerDesc struct.
+	struct FShapeWrapper
+	{
+		TUniquePtr<IStereoLayerShape> Wrapped;
+		FShapeWrapper(IStereoLayerShape* InWrapped)
+			: Wrapped(InWrapped)
+		{}
+
+		FShapeWrapper(const FShapeWrapper& In) :
+			Wrapped(In.IsValid()?In.Wrapped->Clone():nullptr)
+		{}
+
+		FShapeWrapper(FShapeWrapper&& In)
+		{
+			Wrapped = MoveTemp(In.Wrapped);
+		}
+
+		FShapeWrapper& operator= (IStereoLayerShape* InWrapped)
+		{
+			Wrapped = TUniquePtr<IStereoLayerShape>(InWrapped);
+			return *this;
+		}
+
+		FShapeWrapper& operator= (const FShapeWrapper& In)
+		{
+			return operator=(In.IsValid() ? In.Wrapped->Clone() : nullptr);
+		}
+
+		FShapeWrapper& operator= (FShapeWrapper&& In)
+		{
+			Wrapped = MoveTemp(In.Wrapped);
+			return *this;
+		}
+
+		bool IsValid() const { return Wrapped.IsValid(); }
+	};
+
 public:
 
 	enum ELayerType
@@ -21,14 +61,6 @@ public:
 		FaceLocked
 	};
 
-	enum ELayerShape
-	{
-		QuadLayer,
-		CylinderLayer,
-		CubemapLayer,
-		EquirectLayer
-	};
-
 	enum ELayerFlags
 	{
 		// Internally copies the texture on every frame for video, etc.
@@ -45,33 +77,23 @@ public:
 		LAYER_FLAG_HIDDEN = 0x00000020,
 	};
 
-	/** Structure describing additional settings for equirect layers */
-	struct FEquirectProps
-	{
-		/** Left source texture UVRect, specifying portion of input texture corresponding to left eye. */
-		FBox2D LeftUVRect;
-
-		/** Right source texture UVRect, specifying portion of input texture corresponding to right eye. */
-		FBox2D RightUVRect;
-
-		/** Left eye's texture coordinate scale after mapping to 2D. */
-		FVector2D LeftScale;
-
-		/** Right eye's texture coordinate scale after mapping to 2D. */
-		FVector2D RightScale;
-
-		/** Left eye's texture coordinate bias after mapping to 2D. */
-		FVector2D LeftBias;
-
-		/** Right eye's texture coordinate bias after mapping to 2D. */
-		FVector2D RightBias;
-	};
 
 	/**
 	 * Structure describing the visual appearance of a single stereo layer
 	 */
 	struct FLayerDesc
 	{
+
+		FLayerDesc()
+			: Shape(new FQuadLayer())
+		{
+		}
+
+		FLayerDesc(const IStereoLayerShape& InShape)
+			: Shape(InShape.Clone())
+		{
+		}
+
 		void SetLayerId(uint32 InId) { Id = InId; }
 		uint32 GetLayerId() const { return Id; }
 		bool IsVisible() const { return Texture != nullptr && !(Flags & LAYER_FLAG_HIDDEN); }
@@ -93,15 +115,14 @@ public:
 		int32				Priority	 = 0;
 		// Which space the layer is locked within
 		ELayerType			PositionType = ELayerType::FaceLocked;
-		// which shape of layer it is. ELayerShape::QuadLayer is the only shape supported by all VR platforms.
-		ELayerShape			ShapeType	 = ELayerShape::QuadLayer;
+		// which shape of layer it is. FQuadLayer is the only shape supported by all VR platforms.
+		// queries the shape of the layer at run time
+		template <typename T> bool HasShape() const { check(Shape.IsValid()); return Shape.Wrapped->GetShapeName() == T::ShapeName; }
+		// returns Shape cast to the supplied type. It's up to the caller to have ensured the cast is valid before calling this method.
+		template <typename T> T& GetShape() { check(HasShape<T>()); return *static_cast<T*>(Shape.Wrapped.Get()); }
+		template <typename T> const T& GetShape() const { check(HasShape<T>()); return *static_cast<const T*>(Shape.Wrapped.Get()); }
+		template <typename T, typename... InArgTypes> void SetShape(InArgTypes&&... Args) { Shape = FShapeWrapper(new T(Forward<InArgTypes>(Args)...)); }
 
-		// UVs and Scale/Bias of Equirect Layers.
-		FEquirectProps		EquirectProps;
-		// Cylinder layer settings.
-		float				CylinderRadius;
-		float				CylinderOverlayArc;
-		float				CylinderHeight;
 
 		// Texture mapped for right eye (if one texture provided, mono assumed)
 		FTextureRHIRef		Texture		 = nullptr;	
@@ -109,6 +130,8 @@ public:
 		FTextureRHIRef		LeftTexture  = nullptr;
 		// Uses LAYER_FLAG_... -- See: ELayerFlags
 		uint32				Flags		 = 0;
+	private: 
+		FShapeWrapper		Shape;
 	};
 
 	virtual ~IStereoLayers() { }
@@ -281,7 +304,6 @@ public:
 		StereoLayerDesc.Transform = FTransform(FVector(100.f, 0, 0));
 		StereoLayerDesc.QuadSize = FVector2D(120.f, 120.f);
 		StereoLayerDesc.PositionType = IStereoLayers::ELayerType::FaceLocked;
-		StereoLayerDesc.ShapeType = IStereoLayers::ELayerShape::QuadLayer;
 		StereoLayerDesc.Texture = Texture;
 		StereoLayerDesc.Flags = IStereoLayers::ELayerFlags::LAYER_FLAG_TEX_CONTINUOUS_UPDATE;
 		StereoLayerDesc.Flags |= IStereoLayers::ELayerFlags::LAYER_FLAG_QUAD_PRESERVE_TEX_RATIO;
diff --git a/Engine/Source/Runtime/Engine/Public/StereoLayerShapes.h b/Engine/Source/Runtime/Engine/Public/StereoLayerShapes.h
new file mode 100644
--- /dev/null
+++ b/Engine/Source/Runtime/Engine/Public/StereoLayerShapes.h
@@ -0,0 +1,89 @@
+// Copyright 1998-2019 Epic Games, Inc. All Rights Reserved.
+
+#pragma once
+
+#include "CoreMinimal.h"
+/** Different layer shapes are defined as subclass of ILayerShape.
+  * Currently, only FQuadLayer is supported by all platforms, but plugins may define their own custom shapes.
+  * When encountering an unsupported shape type, implementations should treat the layer as a quad layer.
+  */
+
+#define STEREO_LAYER_SHAPE_BOILERPLATE(ClassName) \
+public: \
+	static const FName ShapeName; \
+	virtual FName GetShapeName() override { return ShapeName; } \
+	virtual IStereoLayerShape* Clone() const override { return new ClassName(*this); }
+
+class ENGINE_API IStereoLayerShape
+{
+public:
+	/** Shape name is used to identify the shape type. */
+	virtual FName GetShapeName() = 0;
+	virtual IStereoLayerShape* Clone() const = 0;
+	virtual ~IStereoLayerShape() {}
+};
+
+/** Quad layer is the default layer shape and contains no additional settings */
+class ENGINE_API FQuadLayer : public IStereoLayerShape
+{
+	STEREO_LAYER_SHAPE_BOILERPLATE(FQuadLayer)
+};
+
+/** Class describing additional settings for cylinder layers. Currently only supported by Oculus. */
+class ENGINE_API FCylinderLayer : public IStereoLayerShape
+{
+	STEREO_LAYER_SHAPE_BOILERPLATE(FCylinderLayer)
+public:
+
+	FCylinderLayer() {}
+	FCylinderLayer(float InRadius, float InOverlayArc, float InHeight) :
+		Radius(InRadius),
+		OverlayArc(InOverlayArc),
+		Height(InHeight)
+	{}	
+
+	float	Radius;
+	float	OverlayArc;
+	float	Height;
+};
+
+/** Class describing additional settings for cube map layers. Currently only supported by Oculus. */
+class ENGINE_API FCubemapLayer : public IStereoLayerShape
+{
+	STEREO_LAYER_SHAPE_BOILERPLATE(FCubemapLayer)
+};
+
+/** Class describing additional settings for equirect layers. Currently only supported by Oculus. */
+class ENGINE_API FEquirectLayer : public IStereoLayerShape
+{
+	STEREO_LAYER_SHAPE_BOILERPLATE(FEquirectLayer)
+public:
+
+	FEquirectLayer() {}
+	FEquirectLayer(FBox2D InLeftUVRect, FBox2D InRightUVRect, FVector2D InLeftScale, FVector2D InRightScale, FVector2D InLeftBias, FVector2D InRightBias) :
+		LeftUVRect(InLeftUVRect),
+		RightUVRect(InRightUVRect),
+		LeftScale(InLeftScale),
+		RightScale(InRightScale),
+		LeftBias(InLeftBias),
+		RightBias(InRightBias)
+	{}
+
+	/** Left source texture UVRect, specifying portion of input texture corresponding to left eye. */
+	FBox2D LeftUVRect;
+
+	/** Right source texture UVRect, specifying portion of input texture corresponding to right eye. */
+	FBox2D RightUVRect;
+
+	/** Left eye's texture coordinate scale after mapping to 2D. */
+	FVector2D LeftScale;
+
+	/** Right eye's texture coordinate scale after mapping to 2D. */
+	FVector2D RightScale;
+
+	/** Left eye's texture coordinate bias after mapping to 2D. */
+	FVector2D LeftBias;
+
+	/** Right eye's texture coordinate bias after mapping to 2D. */
+	FVector2D RightBias;
+};
diff --git a/Engine/Source/Runtime/HeadMountedDisplay/Private/DefaultStereoLayers.cpp b/Engine/Source/Runtime/HeadMountedDisplay/Private/DefaultStereoLayers.cpp
--- a/Engine/Source/Runtime/HeadMountedDisplay/Private/DefaultStereoLayers.cpp
+++ b/Engine/Source/Runtime/HeadMountedDisplay/Private/DefaultStereoLayers.cpp
@@ -290,48 +290,4 @@ void FDefaultStereoLayers::SetupViewFamily(FSceneViewFamily& InViewFamily)
 	FVector HmdPosition = FVector::ZeroVector;
 	HMDDevice->GetCurrentPose(IXRTrackingSystem::HMDDeviceId, HmdOrientation, HmdPosition);
 	HmdTransform = FTransform(HmdOrientation, HmdPosition);
-}
-
-void FDefaultStereoLayers::GetAllocatedTexture(uint32 LayerId, FTextureRHIRef &Texture, FTextureRHIRef &LeftTexture)
-{
-	Texture = LeftTexture = nullptr;
-	FLayerDesc* LayerFound = nullptr;
-
-	if (IsInRenderingThread())
-	{
-		for (int32 LayerIndex = 0; LayerIndex < RenderThreadLayers.Num(); LayerIndex++)
-		{
-			if (RenderThreadLayers[LayerIndex].GetLayerId() == LayerId)
-			{
-				LayerFound = &RenderThreadLayers[LayerIndex];
-			}
-		}
-	}
-
-	else
-	{
-		// Only supporting the use of this function on RenderingThread.
-		check(false);
-		return;
-	}
-
-	if (LayerFound && LayerFound->Texture)
-	{
-		switch (LayerFound->ShapeType)
-		{
-		case IStereoLayers::CubemapLayer:
-			Texture = LayerFound->Texture->GetTextureCube();
-			LeftTexture = LayerFound->LeftTexture ? LayerFound->LeftTexture->GetTextureCube() : nullptr;			
-			break;
-
-		case IStereoLayers::CylinderLayer:
-		case IStereoLayers::QuadLayer:
-			Texture = LayerFound->Texture->GetTexture2D();
-			LeftTexture = LayerFound->LeftTexture ? LayerFound->LeftTexture->GetTexture2D() : nullptr;
-			break;
-
-		default:
-			break;
-		}
-	}
-}
+}
\ No newline at end of file
diff --git a/Engine/Source/Runtime/HeadMountedDisplay/Private/DefaultXRLoadingScreen.cpp b/Engine/Source/Runtime/HeadMountedDisplay/Private/DefaultXRLoadingScreen.cpp
--- a/Engine/Source/Runtime/HeadMountedDisplay/Private/DefaultXRLoadingScreen.cpp
+++ b/Engine/Source/Runtime/HeadMountedDisplay/Private/DefaultXRLoadingScreen.cpp
@@ -86,7 +86,22 @@ void FDefaultXRLoadingScreen::DoShowSplash(FSplashData& Splash)
 	// Sort layers by distance from the viewer
 	LayerDesc.Priority = INT32_MAX - static_cast<int32>(Splash.Desc.Transform.GetTranslation().X * 1000.f);
 	LayerDesc.PositionType = IStereoLayers::TrackerLocked;
-	LayerDesc.ShapeType = Splash.Desc.Texture->GetTextureCube() != nullptr ? IStereoLayers::CubemapLayer : IStereoLayers::QuadLayer;
+	
+	if (Splash.Desc.Texture->GetTextureCube() != nullptr)
+	{
+		if (!LayerDesc.HasShape<FCubemapLayer>())
+		{
+			LayerDesc.SetShape<FCubemapLayer>();
+		}
+	}
+	else
+	{
+		if (!LayerDesc.HasShape<FQuadLayer>())
+		{
+			LayerDesc.SetShape<FQuadLayer>();
+		}
+	}
+
 	LayerDesc.Texture = Splash.Desc.Texture;
 	LayerDesc.LeftTexture = Splash.Desc.LeftTexture;
 	LayerDesc.Flags = IStereoLayers::LAYER_FLAG_QUAD_PRESERVE_TEX_RATIO
diff --git a/Engine/Source/Runtime/HeadMountedDisplay/Private/XRSwapChain.cpp b/Engine/Source/Runtime/HeadMountedDisplay/Private/XRSwapChain.cpp
--- a/Engine/Source/Runtime/HeadMountedDisplay/Private/XRSwapChain.cpp
+++ b/Engine/Source/Runtime/HeadMountedDisplay/Private/XRSwapChain.cpp
@@ -117,7 +117,7 @@ void FXRSwapChain::IncrementSwapChainIndex_RHIThread(int64 /* TimeoutNanoseconds
 	CheckInRHIThread();
 
 	SwapChainIndex_RHIThread = (SwapChainIndex_RHIThread + 1) % GetSwapChainLength();
-	GDynamicRHI->RHIAliasTextureResources(RHITexture, RHITextureSwapChain[SwapChainIndex_RHIThread]);
+	GDynamicRHI->RHIAliasTextureResources((FTextureRHIRef&)RHITexture, (FTextureRHIRef&)RHITextureSwapChain[SwapChainIndex_RHIThread]);
 }
 
 
diff --git a/Engine/Source/Runtime/HeadMountedDisplay/Public/DefaultStereoLayers.h b/Engine/Source/Runtime/HeadMountedDisplay/Public/DefaultStereoLayers.h
--- a/Engine/Source/Runtime/HeadMountedDisplay/Public/DefaultStereoLayers.h
+++ b/Engine/Source/Runtime/HeadMountedDisplay/Public/DefaultStereoLayers.h
@@ -24,8 +24,6 @@ public:
 	virtual void PostRenderView_RenderThread(FRHICommandListImmediate& RHICmdList, FSceneView& InView) override;
 	virtual bool IsActiveThisFrame(class FViewport* InViewport) const override;
 
-	void GetAllocatedTexture(uint32 LayerId, FTextureRHIRef &Texture, FTextureRHIRef &LeftTexture);
-
 	virtual bool ShouldCopyDebugLayersToSpectatorScreen() const override
 	{
 		// Emulated layer support means that the debug layer will be in the 3d scene render that the spectator screen displays.
diff --git a/Engine/Source/Runtime/HeadMountedDisplay/Public/StereoLayerManager.h b/Engine/Source/Runtime/HeadMountedDisplay/Public/StereoLayerManager.h
--- a/Engine/Source/Runtime/HeadMountedDisplay/Public/StereoLayerManager.h
+++ b/Engine/Source/Runtime/HeadMountedDisplay/Public/StereoLayerManager.h
@@ -296,6 +296,37 @@ public:
 		bStereoLayersDirty = true;
 	}
 
+	virtual void GetAllocatedTexture(uint32 LayerId, FTextureRHIRef &Texture, FTextureRHIRef &LeftTexture) override
+	{
+		Texture = LeftTexture = nullptr;
+		if (LayerId == FLayerDesc::INVALID_LAYER_ID)
+		{
+			return;
+		}
+		FScopeLock LockLayers(&LayerCritSect);
+
+		int32 FoundLevel;
+		LayerType* Found = FindLayerById(LayerId, FoundLevel);
+		if (Found)
+		{
+			FLayerDesc LayerDesc;
+			GetLayerDescMember(*Found, LayerDesc);
+			if (LayerDesc.Texture)
+			{
+				if (LayerDesc.HasShape<FCubemapLayer>())
+				{
+					Texture = LayerDesc.Texture->GetTextureCube();
+					LeftTexture = LayerDesc.LeftTexture ? LayerDesc.LeftTexture->GetTextureCube() : nullptr;
+				}
+				else
+				{
+					Texture = LayerDesc.Texture->GetTexture2D();
+					LeftTexture = LayerDesc.LeftTexture ? LayerDesc.LeftTexture->GetTexture2D() : nullptr;
+				}
+			}
+		}
+	}
+
 	virtual bool SupportsLayerState() override { return true; }
 
 	virtual void HideBackgroundLayer() { LayerState().bShowBackground = false; }
diff --git a/Engine/Source/Runtime/NullDrv/Public/NullRHI.h b/Engine/Source/Runtime/NullDrv/Public/NullRHI.h
--- a/Engine/Source/Runtime/NullDrv/Public/NullRHI.h
+++ b/Engine/Source/Runtime/NullDrv/Public/NullRHI.h
@@ -471,6 +471,10 @@ public:
 	virtual void RHIAliasTextureResources(FRHITexture* DestTexture, FRHITexture* SrcTexture) final override
 	{
 
+	}
+	virtual void RHIAliasTextureResources(FTextureRHIRef& DestTexture, FTextureRHIRef& SrcTexture) final override
+	{
+
 	}
 	virtual void RHIAdvanceFrameForGetViewportBackBuffer(FRHIViewport* Viewport) final override
 	{
diff --git a/Engine/Source/Runtime/RHI/Public/DynamicRHI.h b/Engine/Source/Runtime/RHI/Public/DynamicRHI.h
--- a/Engine/Source/Runtime/RHI/Public/DynamicRHI.h
+++ b/Engine/Source/Runtime/RHI/Public/DynamicRHI.h
@@ -1004,18 +1004,32 @@ public:
 		return nullptr;
 	}
 
+	UE_DEPRECATED(4.25, "RHIAliasTextureResources now takes references to FTextureRHIRef objects as parameters")
 	virtual void RHIAliasTextureResources(FRHITexture* DestTexture, FRHITexture* SrcTexture)
 	{
 		checkNoEntry();
 	}
 
-	virtual void RHIAdvanceFrameFence(){};
+	UE_DEPRECATED(4.25, "RHICreateAliasedTexture now takes a reference to an FTextureRHIRef object")
 	virtual FTextureRHIRef RHICreateAliasedTexture(FRHITexture* SourceTexture)
 	{
 		checkNoEntry();
 		return nullptr;
 	}
 
+	virtual void RHIAliasTextureResources(FTextureRHIRef& DestTexture, FTextureRHIRef& SrcTexture)
+	{
+		checkNoEntry();
+	}
+
+	virtual FTextureRHIRef RHICreateAliasedTexture(FTextureRHIRef& SourceTexture)
+	{
+		checkNoEntry();
+		return nullptr;
+	}
+
+	virtual void RHIAdvanceFrameFence() {};
+
 	// Only relevant with an RHI thread, this advances the backbuffer for the purpose of GetViewportBackBuffer
 	// FlushType: Thread safe
 	virtual void RHIAdvanceFrameForGetViewportBackBuffer(FRHIViewport* Viewport) = 0;
diff --git a/Engine/Source/Runtime/RHI/Public/RHIValidation.h b/Engine/Source/Runtime/RHI/Public/RHIValidation.h
--- a/Engine/Source/Runtime/RHI/Public/RHIValidation.h
+++ b/Engine/Source/Runtime/RHI/Public/RHIValidation.h
@@ -999,6 +999,21 @@ public:
 		return RHI->RHICreateAliasedTexture(SourceTexture);
 	}
 
+	virtual void RHIAliasTextureResources(FTextureRHIRef& DestTexture, FTextureRHIRef& SourceTexture) override final
+	{
+		// Source and target need to be valid objects.
+		check(DestTexture && SourceTexture);
+		// Source texture must have been created (i.e. have a native resource backing).
+		check(SourceTexture->GetNativeResource() != nullptr);
+		RHI->RHIAliasTextureResources(DestTexture, SourceTexture);
+	}
+
+	virtual FTextureRHIRef RHICreateAliasedTexture(FTextureRHIRef& SourceTexture) override final
+	{
+		check(SourceTexture);
+		return RHI->RHICreateAliasedTexture(SourceTexture);
+	}
+
 	// Only relevant with an RHI thread, this advances the backbuffer for the purpose of GetViewportBackBuffer
 	// FlushType: Thread safe
 	virtual void RHIAdvanceFrameForGetViewportBackBuffer(FRHIViewport* Viewport) override final
diff --git a/Engine/Source/Runtime/VulkanRHI/Private/VulkanRHI.cpp b/Engine/Source/Runtime/VulkanRHI/Private/VulkanRHI.cpp
--- a/Engine/Source/Runtime/VulkanRHI/Private/VulkanRHI.cpp
+++ b/Engine/Source/Runtime/VulkanRHI/Private/VulkanRHI.cpp
@@ -1044,37 +1044,49 @@ FTextureCubeRHIRef FVulkanDynamicRHI::RHICreateTextureCubeFromResource(EPixelFor
 }
 
 void FVulkanDynamicRHI::RHIAliasTextureResources(FRHITexture* DestTextureRHI, FRHITexture* SrcTextureRHI)
+{
+	check(false);
+}
+
+FTextureRHIRef FVulkanDynamicRHI::RHICreateAliasedTexture(FRHITexture* SourceTexture)
+{
+	check(false);
+	return nullptr;
+}
+
+void FVulkanDynamicRHI::RHIAliasTextureResources(FTextureRHIRef& DestTextureRHI, FTextureRHIRef& SrcTextureRHI)
 {
 	if (DestTextureRHI && SrcTextureRHI)
 	{
-		FVulkanTextureBase* DestTextureBase = (FVulkanTextureBase*) DestTextureRHI->GetTextureBaseRHI();
-		FVulkanTextureBase* SrcTextureBase = (FVulkanTextureBase*) SrcTextureRHI->GetTextureBaseRHI();
+		FVulkanTextureBase* DestTextureBase = (FVulkanTextureBase*)DestTextureRHI->GetTextureBaseRHI();
+		FVulkanTextureBase* SrcTextureBase = (FVulkanTextureBase*)SrcTextureRHI->GetTextureBaseRHI();
 
 		if (DestTextureBase && SrcTextureBase)
 		{
-			DestTextureBase->AliasTextureResources(SrcTextureBase);
+			DestTextureBase->AliasTextureResources(SrcTextureRHI);
 		}
 	}
 }
 
-FTextureRHIRef FVulkanDynamicRHI::RHICreateAliasedTexture(FRHITexture* SourceTexture)
+FTextureRHIRef FVulkanDynamicRHI::RHICreateAliasedTexture(FTextureRHIRef& SourceTextureRHI)
 {
+	FVulkanTextureBase* SourceTexture = (FVulkanTextureBase*)SourceTextureRHI->GetTextureBaseRHI();
 	FTextureRHIRef AliasedTexture;
-	if (SourceTexture->GetTexture2D() != nullptr)
+	if (SourceTextureRHI->GetTexture2D() != nullptr)
 	{
-		AliasedTexture = new FVulkanTexture2D(static_cast<FVulkanTexture2D*>(SourceTexture));
+		AliasedTexture = new FVulkanTexture2D(SourceTextureRHI, (FVulkanTexture2D*)SourceTexture);
 	}
-	else if (SourceTexture->GetTexture2DArray() != nullptr)
+	else if (SourceTextureRHI->GetTexture2DArray() != nullptr)
 	{
-		AliasedTexture = new FVulkanTexture2DArray(static_cast<FVulkanTexture2DArray*>(SourceTexture));
+		AliasedTexture = new FVulkanTexture2DArray(SourceTextureRHI, (FVulkanTexture2DArray*)SourceTexture);
 	}
-	else if (SourceTexture->GetTextureCube() != nullptr)
+	else if (SourceTextureRHI->GetTextureCube() != nullptr)
 	{
-		AliasedTexture = new FVulkanTextureCube(static_cast<FVulkanTextureCube*>(SourceTexture));
+		AliasedTexture = new FVulkanTextureCube(SourceTextureRHI, (FVulkanTextureCube*)SourceTexture);
 	}
 	else
 	{
-		UE_LOG(LogRHI, Error, TEXT("Currently FOpenGLDynamicRHI::RHICreateAliasedTexture only supports 2D, 2D Array and Cube textures."));
+		UE_LOG(LogRHI, Error, TEXT("Currently FVulkanDynamicRHI::RHICreateAliasedTexture only supports 2D, 2D Array and Cube textures."));
 	}
 
 	return AliasedTexture;
diff --git a/Engine/Source/Runtime/VulkanRHI/Private/VulkanResources.h b/Engine/Source/Runtime/VulkanRHI/Private/VulkanResources.h
--- a/Engine/Source/Runtime/VulkanRHI/Private/VulkanResources.h
+++ b/Engine/Source/Runtime/VulkanRHI/Private/VulkanResources.h
@@ -486,11 +486,11 @@ struct FVulkanTextureBase : public FVulkanBaseShaderResource
 	FVulkanTextureBase(FVulkanDevice& Device, VkImageViewType ResourceType, EPixelFormat Format, uint32 SizeX, uint32 SizeY, uint32 SizeZ, uint32 ArraySize, uint32 NumMips, uint32 NumSamples, VkImage InImage, VkDeviceMemory InMem, FSamplerYcbcrConversionInitializer& ConversionInitializer, uint32 UEFlags, const FRHIResourceCreateInfo& CreateInfo = FRHIResourceCreateInfo());
 
 	// Aliasing constructor.
-	FVulkanTextureBase(const FVulkanTextureBase* SrcTexture, VkImageViewType ResourceType, uint32 SizeX, uint32 SizeY, uint32 sizeZ);
+	FVulkanTextureBase(FTextureRHIRef& SrcTextureRHI, const FVulkanTextureBase* SrcTexture, VkImageViewType ResourceType, uint32 SizeX, uint32 SizeY, uint32 sizeZ);
 
 	virtual ~FVulkanTextureBase();
 
-	void AliasTextureResources(const FVulkanTextureBase* SrcTexture);
+	void AliasTextureResources(FTextureRHIRef& SrcTexture);
 
 	FVulkanSurface Surface;
 
@@ -499,7 +499,7 @@ struct FVulkanTextureBase : public FVulkanBaseShaderResource
 	// View with all mips/layers, but if it's a Depth/Stencil, only the Depth view
 	FVulkanTextureView* PartialView;
 
-	bool bIsAliased;
+	FTextureRHIRef AliasedTexture;
 
 	virtual void OnTransitionResource(FVulkanCommandListContext& Context, EResourceTransitionAccess TransitionType) {};
 
@@ -515,7 +515,7 @@ public:
 	FVulkanTexture2D(FVulkanDevice& Device, EPixelFormat Format, uint32 SizeX, uint32 SizeY, uint32 NumMips, uint32 NumSamples, VkImage Image, struct FSamplerYcbcrConversionInitializer& ConversionInitializer, uint32 UEFlags, const FRHIResourceCreateInfo& CreateInfo);
 
 	// Aliasing constructor
-	FVulkanTexture2D(const FVulkanTexture2D* SrcTexture);
+	FVulkanTexture2D(FTextureRHIRef& SrcTextureRHI, const FVulkanTexture2D* SrcTexture);
 
 	virtual ~FVulkanTexture2D();
 
@@ -553,7 +553,7 @@ public:
 	FVulkanTexture2DArray(FVulkanDevice& Device, EPixelFormat Format, uint32 SizeX, uint32 SizeY, uint32 ArraySize, uint32 NumMips, uint32 NumSamples, VkImage Image, uint32 Flags, FResourceBulkDataInterface* BulkData, const FClearValueBinding& InClearValue);
 
 	// Aliasing constructor
-	FVulkanTexture2DArray(const FVulkanTexture2DArray* SrcTexture);
+	FVulkanTexture2DArray(FTextureRHIRef& SrcTextureRHI, const FVulkanTexture2DArray* SrcTexture);
 
 		
 	// IRefCountedObject interface.
@@ -621,7 +621,7 @@ public:
 	FVulkanTextureCube(FVulkanDevice& Device, EPixelFormat Format, uint32 Size, bool bArray, uint32 ArraySize, uint32 NumMips, VkImage Image, uint32 Flags, FResourceBulkDataInterface* BulkData, const FClearValueBinding& InClearValue);
 
 	// Aliasing constructor
-	FVulkanTextureCube(const FVulkanTextureCube* SrcTexture);
+	FVulkanTextureCube(FTextureRHIRef& SrcTextureRHI, const FVulkanTextureCube* SrcTexture);
 
 	virtual ~FVulkanTextureCube();
 
diff --git a/Engine/Source/Runtime/VulkanRHI/Private/VulkanTexture.cpp b/Engine/Source/Runtime/VulkanRHI/Private/VulkanTexture.cpp
--- a/Engine/Source/Runtime/VulkanRHI/Private/VulkanTexture.cpp
+++ b/Engine/Source/Runtime/VulkanRHI/Private/VulkanTexture.cpp
@@ -1643,7 +1643,6 @@ void FVulkanTextureView::Destroy(FVulkanDevice& Device)
 FVulkanTextureBase::FVulkanTextureBase(FVulkanDevice& Device, VkImageViewType ResourceType, EPixelFormat InFormat, uint32 SizeX, uint32 SizeY, uint32 SizeZ, uint32 ArraySize, uint32 NumMips, uint32 NumSamples, uint32 UEFlags, const FRHIResourceCreateInfo& CreateInfo)
 	: Surface(Device, ResourceType, InFormat, SizeX, SizeY, SizeZ, ArraySize, NumMips, NumSamples, UEFlags, CreateInfo)
 	, PartialView(nullptr)
-	, bIsAliased(false)
 {
 	LLM_SCOPE_VULKAN(ELLMTagVulkan::VulkanTextures);
 	const bool bArray = ResourceType == VK_IMAGE_VIEW_TYPE_1D_ARRAY || ResourceType == VK_IMAGE_VIEW_TYPE_2D_ARRAY || ResourceType == VK_IMAGE_VIEW_TYPE_CUBE_ARRAY;
@@ -1727,7 +1726,6 @@ FVulkanTextureBase::FVulkanTextureBase(FVulkanDevice& Device, VkImageViewType Re
 FVulkanTextureBase::FVulkanTextureBase(FVulkanDevice& Device, VkImageViewType ResourceType, EPixelFormat Format, uint32 SizeX, uint32 SizeY, uint32 SizeZ, uint32 ArraySize, uint32 InNumMips, uint32 InNumSamples, VkImage InImage, VkDeviceMemory InMem, uint32 UEFlags, const FRHIResourceCreateInfo& CreateInfo)
 	: Surface(Device, ResourceType, Format, SizeX, SizeY, SizeZ, ArraySize, InNumMips, InNumSamples, InImage, UEFlags, CreateInfo)
 	, PartialView(nullptr)
-	, bIsAliased(false)
 {
 	check(InMem == VK_NULL_HANDLE);
 	const bool bArray = ResourceType == VK_IMAGE_VIEW_TYPE_1D_ARRAY || ResourceType == VK_IMAGE_VIEW_TYPE_2D_ARRAY || ResourceType == VK_IMAGE_VIEW_TYPE_CUBE_ARRAY;
@@ -1751,7 +1749,6 @@ FVulkanTextureBase::FVulkanTextureBase(FVulkanDevice& Device, VkImageViewType Re
 FVulkanTextureBase::FVulkanTextureBase(FVulkanDevice& Device, VkImageViewType ResourceType, EPixelFormat Format, uint32 SizeX, uint32 SizeY, uint32 SizeZ, uint32 ArraySize, uint32 NumMips, uint32 NumSamples, VkImage InImage, VkDeviceMemory InMem, FSamplerYcbcrConversionInitializer& ConversionInitializer, uint32 UEFlags, const FRHIResourceCreateInfo& CreateInfo)
 	: Surface(Device, ResourceType, Format, SizeX, SizeY, SizeZ, ArraySize, NumMips, NumSamples, InImage, UEFlags, CreateInfo)
 	, PartialView(nullptr)
-	, bIsAliased(false)
 {
 	check(InMem == VK_NULL_HANDLE);
 	const bool bArray = ResourceType == VK_IMAGE_VIEW_TYPE_1D_ARRAY || ResourceType == VK_IMAGE_VIEW_TYPE_2D_ARRAY || ResourceType == VK_IMAGE_VIEW_TYPE_CUBE_ARRAY;
@@ -1782,10 +1779,10 @@ FVulkanTextureBase::FVulkanTextureBase(FVulkanDevice& Device, VkImageViewType Re
 	InsertInitialImageLayout(Device, InImage, VK_IMAGE_LAYOUT_SHADER_READ_ONLY_OPTIMAL);
 }
 
-FVulkanTextureBase::FVulkanTextureBase(const FVulkanTextureBase* SrcTexture, VkImageViewType ResourceType, uint32 SizeX, uint32 SizeY, uint32 SizeZ)
+FVulkanTextureBase::FVulkanTextureBase(FTextureRHIRef& SrcTextureRHI, const FVulkanTextureBase* SrcTexture, VkImageViewType ResourceType, uint32 SizeX, uint32 SizeY, uint32 SizeZ)
 	: Surface(*SrcTexture->Surface.Device, ResourceType, SrcTexture->Surface.PixelFormat, SizeX, SizeY, SizeZ, SrcTexture->Surface.NumArrayLevels, SrcTexture->Surface.NumMips, SrcTexture->Surface.NumSamples, SrcTexture->Surface.Image, SrcTexture->Surface.UEFlags, FRHIResourceCreateInfo())
 	, PartialView(nullptr)
-	, bIsAliased(true)
+	, AliasedTexture(SrcTextureRHI)
 {
 	if (Surface.FullAspectMask == Surface.PartialAspectMask)
 	{
@@ -1796,7 +1793,8 @@ FVulkanTextureBase::FVulkanTextureBase(const FVulkanTextureBase* SrcTexture, VkI
 		PartialView = new FVulkanTextureView;
 		// Skip create, since we're aliasing.
 	}
-	AliasTextureResources(SrcTexture);
+
+	AliasTextureResources(SrcTextureRHI);
 }
 
 FVulkanTextureBase::~FVulkanTextureBase()
@@ -1809,10 +1807,12 @@ FVulkanTextureBase::~FVulkanTextureBase()
 	}
 }
 
-void FVulkanTextureBase::AliasTextureResources(const FVulkanTextureBase* SrcTexture)
+void FVulkanTextureBase::AliasTextureResources(FTextureRHIRef& SrcTextureRHI)
 {
 	DestroyViews();
 
+	FVulkanTextureBase* SrcTexture = (FVulkanTextureBase*)SrcTextureRHI->GetTextureBaseRHI();
+
 	Surface.Destroy();
 	Surface.Image = SrcTexture->Surface.Image;
 	DefaultView.View = SrcTexture->DefaultView.View;
@@ -1825,13 +1825,11 @@ void FVulkanTextureBase::AliasTextureResources(const FVulkanTextureBase* SrcText
 		PartialView->Image = SrcTexture->PartialView->Image;
 		PartialView->ViewId = SrcTexture->PartialView->ViewId;
 	}
-
-	bIsAliased = true;
 }
 
 void FVulkanTextureBase::DestroyViews()
 {
-	if (!bIsAliased)
+	if (AliasedTexture == nullptr)
 	{
 		DefaultView.Destroy(*Surface.Device);
 
@@ -1861,9 +1859,9 @@ FVulkanTexture2D::FVulkanTexture2D(FVulkanDevice& Device, EPixelFormat Format, u
 {
 }
 
-FVulkanTexture2D::FVulkanTexture2D(const FVulkanTexture2D* SrcTexture)
+FVulkanTexture2D::FVulkanTexture2D(FTextureRHIRef& SrcTextureRHI, const FVulkanTexture2D* SrcTexture)
 	: FRHITexture2D(SrcTexture->GetSizeX(), SrcTexture->GetSizeY(), SrcTexture->GetNumMips(), SrcTexture->GetNumSamples(), SrcTexture->GetFormat(), SrcTexture->GetFlags(), SrcTexture->GetClearBinding())
-	, FVulkanTextureBase(static_cast<const FVulkanTextureBase*>(SrcTexture), VK_IMAGE_VIEW_TYPE_2D, SrcTexture->GetSizeX(), SrcTexture->GetSizeY(), 1)
+	, FVulkanTextureBase(SrcTextureRHI, static_cast<const FVulkanTextureBase*>(SrcTexture), VK_IMAGE_VIEW_TYPE_2D, SrcTexture->GetSizeX(), SrcTexture->GetSizeY(), 1)
 {
 }
 
@@ -1887,9 +1885,9 @@ FVulkanTexture2DArray::FVulkanTexture2DArray(FVulkanDevice& Device, EPixelFormat
 {
 }
 
-FVulkanTexture2DArray::FVulkanTexture2DArray(const FVulkanTexture2DArray* SrcTexture)
+FVulkanTexture2DArray::FVulkanTexture2DArray(FTextureRHIRef& SrcTextureRHI, const FVulkanTexture2DArray* SrcTexture)
 	: FRHITexture2DArray(SrcTexture->GetSizeX(), SrcTexture->GetSizeY(), SrcTexture->GetSizeZ(), SrcTexture->GetNumMips(), SrcTexture->GetNumSamples(), SrcTexture->GetFormat(), SrcTexture->GetFlags(), SrcTexture->GetClearBinding())
-	, FVulkanTextureBase(static_cast<const FVulkanTextureBase*>(SrcTexture), VK_IMAGE_VIEW_TYPE_2D_ARRAY, SrcTexture->GetSizeX(), SrcTexture->GetSizeY(), SrcTexture->GetSizeZ())
+	, FVulkanTextureBase(SrcTextureRHI, static_cast<const FVulkanTextureBase*>(SrcTexture), VK_IMAGE_VIEW_TYPE_2D_ARRAY, SrcTexture->GetSizeX(), SrcTexture->GetSizeY(), SrcTexture->GetSizeZ())
 {
 }
 
@@ -1913,9 +1911,9 @@ FVulkanTextureCube::FVulkanTextureCube(FVulkanDevice& Device, EPixelFormat Forma
 {
 }
 
-FVulkanTextureCube::FVulkanTextureCube(const FVulkanTextureCube* SrcTexture)
+FVulkanTextureCube::FVulkanTextureCube(FTextureRHIRef& SrcTextureRHI, const FVulkanTextureCube* SrcTexture)
 	: FRHITextureCube(SrcTexture->GetSize(), SrcTexture->GetNumMips(), SrcTexture->GetFormat(), SrcTexture->GetFlags(), SrcTexture->GetClearBinding())
-	, FVulkanTextureBase(static_cast<const FVulkanTextureBase*>(SrcTexture), VK_IMAGE_VIEW_TYPE_CUBE, SrcTexture->GetSize(), SrcTexture->GetSize(), 1)
+	, FVulkanTextureBase(SrcTextureRHI, static_cast<const FVulkanTextureBase*>(SrcTexture), VK_IMAGE_VIEW_TYPE_CUBE, SrcTexture->GetSize(), SrcTexture->GetSize(), 1)
 {
 }
 
diff --git a/Engine/Source/Runtime/VulkanRHI/Public/VulkanDynamicRHI.h b/Engine/Source/Runtime/VulkanRHI/Public/VulkanDynamicRHI.h
--- a/Engine/Source/Runtime/VulkanRHI/Public/VulkanDynamicRHI.h
+++ b/Engine/Source/Runtime/VulkanRHI/Public/VulkanDynamicRHI.h
@@ -158,6 +158,8 @@ public:
 	virtual FTexture2DRHIRef RHIGetViewportBackBuffer(FRHIViewport* Viewport) final override;
 	virtual void RHIAliasTextureResources(FRHITexture* DestTexture, FRHITexture* SrcTexture) final override;
 	virtual FTextureRHIRef RHICreateAliasedTexture(FRHITexture* SourceTexture) final override;
+	virtual void RHIAliasTextureResources(FTextureRHIRef& DestTexture, FTextureRHIRef& SrcTexture) final override;
+	virtual FTextureRHIRef RHICreateAliasedTexture(FTextureRHIRef& SourceTexture) final override;
 	virtual void RHIAdvanceFrameForGetViewportBackBuffer(FRHIViewport* Viewport) final override;
 	virtual void RHIAcquireThreadOwnership() final override;
 	virtual void RHIReleaseThreadOwnership() final override;
diff --git a/Engine/Source/Runtime/Windows/D3D11RHI/Private/D3D11RHIPrivate.h b/Engine/Source/Runtime/Windows/D3D11RHI/Private/D3D11RHIPrivate.h
--- a/Engine/Source/Runtime/Windows/D3D11RHI/Private/D3D11RHIPrivate.h
+++ b/Engine/Source/Runtime/Windows/D3D11RHI/Private/D3D11RHIPrivate.h
@@ -478,6 +478,8 @@ public:
 	virtual FTexture2DRHIRef RHIGetViewportBackBuffer(FRHIViewport* Viewport) final override;
 	virtual void RHIAliasTextureResources(FRHITexture* DestTexture, FRHITexture* SrcTexture) final override;
 	virtual FTextureRHIRef RHICreateAliasedTexture(FRHITexture* SourceTexture) final override;
+	virtual void RHIAliasTextureResources(FTextureRHIRef& DestTexture, FTextureRHIRef& SrcTexture) final override;
+	virtual FTextureRHIRef RHICreateAliasedTexture(FTextureRHIRef& SourceTexture) final override;
 	virtual void RHIAdvanceFrameForGetViewportBackBuffer(FRHIViewport* Viewport) final override;
 	virtual void RHIAcquireThreadOwnership() final override;
 	virtual void RHIReleaseThreadOwnership() final override;
diff --git a/Engine/Source/Runtime/Windows/D3D11RHI/Private/D3D11Texture.cpp b/Engine/Source/Runtime/Windows/D3D11RHI/Private/D3D11Texture.cpp
--- a/Engine/Source/Runtime/Windows/D3D11RHI/Private/D3D11Texture.cpp
+++ b/Engine/Source/Runtime/Windows/D3D11RHI/Private/D3D11Texture.cpp
@@ -2483,6 +2483,18 @@ FTextureRHIRef FD3D11DynamicRHI::RHICreateAliasedTexture(FRHITexture* SourceText
 	return nullptr;
 }
 
+void FD3D11DynamicRHI::RHIAliasTextureResources(FTextureRHIRef& DestTextureRHI, FTextureRHIRef& SrcTextureRHI)
+{
+	// @todo: Move the raw-pointer implementation down here when it's deprecation is completed.
+	RHIAliasTextureResources((FRHITexture*)DestTextureRHI, (FRHITexture*)SrcTextureRHI);
+}
+
+FTextureRHIRef FD3D11DynamicRHI::RHICreateAliasedTexture(FTextureRHIRef& SourceTexture)
+{
+	// @todo: Move the raw-pointer implementation down here when it's deprecation is completed.
+	return RHICreateAliasedTexture((FRHITexture*)SourceTexture);
+}
+
 void FD3D11DynamicRHI::RHICopyTexture(FRHITexture* SourceTextureRHI, FRHITexture* DestTextureRHI, const FRHICopyTextureInfo& CopyInfo)
 {
 	if (!SourceTextureRHI || !DestTextureRHI || SourceTextureRHI == DestTextureRHI)
