Fix a problem that movable objects never get reflection captures updated in forward renderer.

diff --git a/Engine/Source/Runtime/Engine/Private/PrimitiveSceneProxy.cpp b/Engine/Source/Runtime/Engine/Private/PrimitiveSceneProxy.cpp
--- a/Engine/Source/Runtime/Engine/Private/PrimitiveSceneProxy.cpp
+++ b/Engine/Source/Runtime/Engine/Private/PrimitiveSceneProxy.cpp
@@ -403,6 +403,12 @@ void FPrimitiveSceneProxy::SetTransform(const FMatrix& InLocalToWorld, const FBo
 	LocalBounds = InLocalBounds;
 	ActorPosition = InActorPosition;
 	
+	// Update cached reflection capture.
+	if (PrimitiveSceneInfo)
+	{
+		PrimitiveSceneInfo->bNeedsCachedReflectionCaptureUpdate = true;
+	}
+	
 	UpdateUniformBuffer();
 	
 	// Notify the proxy's implementation of the change.
