Make sure render resource is released when destroying VT streaming manager

diff --git a/Engine/Source/Runtime/Engine/Private/VT/VirtualTextureChunkManager.cpp b/Engine/Source/Runtime/Engine/Private/VT/VirtualTextureChunkManager.cpp
--- a/Engine/Source/Runtime/Engine/Private/VT/VirtualTextureChunkManager.cpp
+++ b/Engine/Source/Runtime/Engine/Private/VT/VirtualTextureChunkManager.cpp
@@ -53,6 +53,7 @@ FVirtualTextureChunkStreamingManager::FVirtualTextureChunkStreamingManager()
 
 FVirtualTextureChunkStreamingManager::~FVirtualTextureChunkStreamingManager()
 {
+	UploadCache.ReleaseResource(); // Must be called from rendering thread
 #if WITH_EDITOR
 	GetVirtualTextureChunkDDCCache()->ShutDown();
 #endif
