Changed kinematic state ensure to CHAOS_ENSURE

diff --git a/Engine/Source/Runtime/Engine/Private/PhysicsEngine/Experimental/PhysInterface_Chaos.cpp b/Engine/Source/Runtime/Engine/Private/PhysicsEngine/Experimental/PhysInterface_Chaos.cpp
--- a/Engine/Source/Runtime/Engine/Private/PhysicsEngine/Experimental/PhysInterface_Chaos.cpp
+++ b/Engine/Source/Runtime/Engine/Private/PhysicsEngine/Experimental/PhysInterface_Chaos.cpp
@@ -551,7 +551,7 @@ void FPhysInterface_Chaos::SetIsKinematic_AssumesLocked(const FPhysicsActorHandl
 	}
 	else
 	{
-		ensureMsgf(false, TEXT("Can only set kinematic state of underlying dynamic particles"));
+		CHAOS_ENSURE_MSG(false, TEXT("Can only set kinematic state of underlying dynamic particles"));
 	}
 }
 
