Chaos: fix Navmesh generation issue, recompute indices on old assets, recompute indices on update

diff --git a/Engine/Source/Runtime/Core/Public/UObject/ExternalPhysicsCustomObjectVersion.h b/Engine/Source/Runtime/Core/Public/UObject/ExternalPhysicsCustomObjectVersion.h
--- a/Engine/Source/Runtime/Core/Public/UObject/ExternalPhysicsCustomObjectVersion.h
+++ b/Engine/Source/Runtime/Core/Public/UObject/ExternalPhysicsCustomObjectVersion.h
@@ -117,6 +117,9 @@ struct CORE_API FExternalPhysicsCustomObjectVersion
 		// Added serialization for the collision type in the shape
 		SerializeCollisionTraceType,
 
+		// Force rebuild of indices in BodySetup
+		ForceRebuildBodySetupIndices,
+
 		// -----<new versions can be added above this line>-------------------------------------------------
 		VersionPlusOne,
 		LatestVersion = VersionPlusOne - 1
diff --git a/Engine/Source/Runtime/Engine/Private/Collision/KAggregateGeom.cpp b/Engine/Source/Runtime/Engine/Private/Collision/KAggregateGeom.cpp
--- a/Engine/Source/Runtime/Engine/Private/Collision/KAggregateGeom.cpp
+++ b/Engine/Source/Runtime/Engine/Private/Collision/KAggregateGeom.cpp
@@ -539,6 +539,14 @@ static void AddEdgeIfNotPresent(TArray<int32>& Edges, int32 Edge0, int32 Edge1)
 
 void FKConvexElem::UpdateElemBox()
 {
+#if WITH_CHAOS
+	// Fixup indices in case an operation has invalidated them
+	{
+		IndexData.Reset();
+		ComputeChaosConvexIndices();
+	}
+#endif
+
 	ElemBox.Init();
 	for(int32 j=0; j<VertexData.Num(); j++)
 	{
diff --git a/Engine/Source/Runtime/Engine/Private/PhysicsEngine/BodySetup.cpp b/Engine/Source/Runtime/Engine/Private/PhysicsEngine/BodySetup.cpp
--- a/Engine/Source/Runtime/Engine/Private/PhysicsEngine/BodySetup.cpp
+++ b/Engine/Source/Runtime/Engine/Private/PhysicsEngine/BodySetup.cpp
@@ -966,6 +966,7 @@ void UBodySetup::FinishDestroy()
 void UBodySetup::Serialize(FArchive& Ar)
 {
 	Super::Serialize(Ar);
+	Ar.UsingCustomVersion(FExternalPhysicsCustomObjectVersion::GUID);
 
 	// Load GUID (or create one for older versions)
 	Ar << BodySetupGuid;
@@ -1041,6 +1042,29 @@ void UBodySetup::Serialize(FArchive& Ar)
 #if WITH_EDITOR
 	AggGeom.FixupDeprecated( Ar );
 #endif
+
+#if WITH_CHAOS && WITH_EDITOR
+
+	if (Ar.IsLoading())
+	{
+		const bool bForceIndexRebuild = Ar.CustomVer(FExternalPhysicsCustomObjectVersion::GUID) < FExternalPhysicsCustomObjectVersion::ForceRebuildBodySetupIndices;
+		for (FKConvexElem& Convex : AggGeom.ConvexElems)
+		{
+			// Reset potentially corrupted index data to correctly rebuild below
+			if (bForceIndexRebuild)
+			{
+				Convex.IndexData.Reset();
+			}
+			// Build an index buffer if we don't have one, either as a consequence of the check above or loading in a mesh that has never been
+			// processed with Chaos previously
+			Convex.ComputeChaosConvexIndices();
+		}
+		if (bForceIndexRebuild)
+		{
+			BodySetupGuid = FGuid::NewGuid(); // change the guid for dependents (navigation)
+		}
+	}
+#endif
 }
 
 void UBodySetup::PostLoad()
