- Remove various UMaterialExpression::NeedsLoadForClient() impl, and instead always return true for UMaterialExpression::IsEditorOnly(), as we no longer want any expressions in non-editor build
- MaterialCachedData checks for nullptr expressions, and avoids updating if any are encountered...this fixes crash with editor loading cooked/non-editor data

diff --git a/Engine/Source/Runtime/Engine/Classes/Materials/MaterialExpression.h b/Engine/Source/Runtime/Engine/Classes/Materials/MaterialExpression.h
--- a/Engine/Source/Runtime/Engine/Classes/Materials/MaterialExpression.h
+++ b/Engine/Source/Runtime/Engine/Classes/Materials/MaterialExpression.h
@@ -207,8 +207,7 @@ class ENGINE_API UMaterialExpression : public UObject
 	virtual bool Modify( bool bAlwaysMarkDirty=true ) override;
 #endif // WITH_EDITOR
 	virtual void Serialize( FStructuredArchive::FRecord Record ) override;
-	virtual bool NeedsLoadForClient() const override;
-	virtual bool NeedsLoadForEditorGame() const override
+	virtual bool IsEditorOnly() const
 	{
 		return true;
 	}
diff --git a/Engine/Source/Runtime/Engine/Classes/Materials/MaterialExpressionCollectionParameter.h b/Engine/Source/Runtime/Engine/Classes/Materials/MaterialExpressionCollectionParameter.h
--- a/Engine/Source/Runtime/Engine/Classes/Materials/MaterialExpressionCollectionParameter.h
+++ b/Engine/Source/Runtime/Engine/Classes/Materials/MaterialExpressionCollectionParameter.h
@@ -36,7 +36,6 @@ class UMaterialExpressionCollectionParameter : public UMaterialExpression
 	virtual void PostEditChangeProperty(FPropertyChangedEvent& PropertyChangedEvent) override;
 #endif // WITH_EDITOR
 	virtual void PostLoad() override;
-	virtual bool NeedsLoadForClient() const override;
 	//~ End UObject Interface.
 
 	//~ Begin UMaterialExpression Interface
diff --git a/Engine/Source/Runtime/Engine/Classes/Materials/MaterialExpressionDynamicParameter.h b/Engine/Source/Runtime/Engine/Classes/Materials/MaterialExpressionDynamicParameter.h
--- a/Engine/Source/Runtime/Engine/Classes/Materials/MaterialExpressionDynamicParameter.h
+++ b/Engine/Source/Runtime/Engine/Classes/Materials/MaterialExpressionDynamicParameter.h
@@ -42,7 +42,6 @@ class UMaterialExpressionDynamicParameter : public UMaterialExpression
 #endif // WITH_EDITOR
 
 	virtual void PostLoad() override;
-	virtual bool NeedsLoadForClient() const override;
 	//~ End UObject Interface
 
 	//~ Begin UMaterialExpression Interface
diff --git a/Engine/Source/Runtime/Engine/Classes/Materials/MaterialExpressionFeatureLevelSwitch.h b/Engine/Source/Runtime/Engine/Classes/Materials/MaterialExpressionFeatureLevelSwitch.h
--- a/Engine/Source/Runtime/Engine/Classes/Materials/MaterialExpressionFeatureLevelSwitch.h
+++ b/Engine/Source/Runtime/Engine/Classes/Materials/MaterialExpressionFeatureLevelSwitch.h
@@ -23,7 +23,6 @@ class UMaterialExpressionFeatureLevelSwitch : public UMaterialExpression
 
 	//~ Begin UObject Interface.
 	virtual void Serialize(FStructuredArchive::FRecord Record) override;
-	virtual bool NeedsLoadForClient() const override;
 	//~ End UObject Interface.
 
 	//~ Begin UMaterialExpression Interface
diff --git a/Engine/Source/Runtime/Engine/Classes/Materials/MaterialExpressionMaterialAttributeLayers.h b/Engine/Source/Runtime/Engine/Classes/Materials/MaterialExpressionMaterialAttributeLayers.h
--- a/Engine/Source/Runtime/Engine/Classes/Materials/MaterialExpressionMaterialAttributeLayers.h
+++ b/Engine/Source/Runtime/Engine/Classes/Materials/MaterialExpressionMaterialAttributeLayers.h
@@ -86,7 +86,6 @@ class UMaterialExpressionMaterialAttributeLayers : public UMaterialExpression
 #if WITH_EDITOR
 	virtual void PostEditChangeProperty(FPropertyChangedEvent& PropertyChangedEvent) override;
 #endif
-	virtual bool NeedsLoadForClient() const override { return true; };
 	//~ Begin UObject Interface
 
 #if WITH_EDITORONLY_DATA
diff --git a/Engine/Source/Runtime/Engine/Classes/Materials/MaterialExpressionMaterialFunctionCall.h b/Engine/Source/Runtime/Engine/Classes/Materials/MaterialExpressionMaterialFunctionCall.h
--- a/Engine/Source/Runtime/Engine/Classes/Materials/MaterialExpressionMaterialFunctionCall.h
+++ b/Engine/Source/Runtime/Engine/Classes/Materials/MaterialExpressionMaterialFunctionCall.h
@@ -103,7 +103,6 @@ class UMaterialExpressionMaterialFunctionCall : public UMaterialExpression
 	virtual void PostEditChangeProperty(FPropertyChangedEvent& PropertyChangedEvent) override;
 #endif // WITH_EDITOR
 	virtual void PostLoad() override;
-	virtual bool NeedsLoadForClient() const override;
 	//~ End UObject Interface.
 
 #if WITH_EDITORONLY_DATA
diff --git a/Engine/Source/Runtime/Engine/Classes/Materials/MaterialExpressionParameter.h b/Engine/Source/Runtime/Engine/Classes/Materials/MaterialExpressionParameter.h
--- a/Engine/Source/Runtime/Engine/Classes/Materials/MaterialExpressionParameter.h
+++ b/Engine/Source/Runtime/Engine/Classes/Materials/MaterialExpressionParameter.h
@@ -51,10 +51,6 @@ class UMaterialExpressionParameter : public UMaterialExpression
 #endif
 	//~ End UMaterialExpression Interface
 
-	//~ Begin UObject Interface
-	virtual bool NeedsLoadForClient() const override;
-	//~ End UObject Interface
-
 	ENGINE_API virtual FGuid& GetParameterExpressionId() override
 	{
 		return ExpressionGUID;
diff --git a/Engine/Source/Runtime/Engine/Classes/Materials/MaterialExpressionQualitySwitch.h b/Engine/Source/Runtime/Engine/Classes/Materials/MaterialExpressionQualitySwitch.h
--- a/Engine/Source/Runtime/Engine/Classes/Materials/MaterialExpressionQualitySwitch.h
+++ b/Engine/Source/Runtime/Engine/Classes/Materials/MaterialExpressionQualitySwitch.h
@@ -35,8 +35,4 @@ class UMaterialExpressionQualitySwitch : public UMaterialExpression
 	virtual uint32 GetOutputType(int32 InputIndex) override {return MCT_Unknown;}
 #endif // WITH_EDITOR
 	//~ End UMaterialExpression Interface
-
-	//~ Begin UObject Interface
-	virtual bool NeedsLoadForClient() const override;
-	//~ End UObject Interface
 };
diff --git a/Engine/Source/Runtime/Engine/Classes/Materials/MaterialExpressionShadingPathSwitch.h b/Engine/Source/Runtime/Engine/Classes/Materials/MaterialExpressionShadingPathSwitch.h
--- a/Engine/Source/Runtime/Engine/Classes/Materials/MaterialExpressionShadingPathSwitch.h
+++ b/Engine/Source/Runtime/Engine/Classes/Materials/MaterialExpressionShadingPathSwitch.h
@@ -35,8 +35,4 @@ class UMaterialExpressionShadingPathSwitch : public UMaterialExpression
 	virtual uint32 GetOutputType(int32 InputIndex) override {return MCT_Unknown;}
 #endif // WITH_EDITOR
 	//~ End UMaterialExpression Interface
-
-	//~ Begin UObject Interface
-	virtual bool NeedsLoadForClient() const override;
-	//~ End UObject Interface
 };
diff --git a/Engine/Source/Runtime/Engine/Private/Materials/Material.cpp b/Engine/Source/Runtime/Engine/Private/Materials/Material.cpp
--- a/Engine/Source/Runtime/Engine/Private/Materials/Material.cpp
+++ b/Engine/Source/Runtime/Engine/Private/Materials/Material.cpp
@@ -2039,8 +2039,17 @@ bool UMaterial::GetParameterDesc(const FHashedMaterialParameterInfo& ParameterIn
 }
 void UMaterial::UpdateCachedExpressionData()
 {
-	CachedExpressionData.Reset();
-	CachedExpressionData.UpdateForExpressions(Expressions, EMaterialParameterAssociation::GlobalParameter, -1);
+	{
+		FMaterialCachedExpressionData UpdatedCachedExpressionData;
+		UpdatedCachedExpressionData.Reset();
+		if (UpdatedCachedExpressionData.UpdateForExpressions(Expressions, EMaterialParameterAssociation::GlobalParameter, -1))
+		{
+			// Only update our cached data if the update succeeded
+			// It's possible we could have some nullptr UMaterialExpressions, if we're loading cooked data here
+			// In that case we simply keep existing cached data
+			CachedExpressionData = MoveTemp(UpdatedCachedExpressionData);
+		}
+	}
 
 	// Always append a default high quality level if nothing else is set
 	bool bAnyQualityLevelsSet = false;
diff --git a/Engine/Source/Runtime/Engine/Private/Materials/MaterialCachedData.cpp b/Engine/Source/Runtime/Engine/Private/Materials/MaterialCachedData.cpp
--- a/Engine/Source/Runtime/Engine/Private/Materials/MaterialCachedData.cpp
+++ b/Engine/Source/Runtime/Engine/Private/Materials/MaterialCachedData.cpp
@@ -96,11 +96,21 @@ static int32 TryAddParameter(FMaterialCachedParameters& CachedParameters, EMater
 	return INDEX_NONE;
 }
 
-void FMaterialCachedExpressionData::UpdateForFunction(UMaterialFunctionInterface* Function, EMaterialParameterAssociation Association, int32 ParameterIndex)
+bool FMaterialCachedExpressionData::UpdateForFunction(UMaterialFunctionInterface* Function, EMaterialParameterAssociation Association, int32 ParameterIndex)
 {
 	if (!Function)
 	{
-		return;
+		return true;
+	}
+
+	bool bResult = true;
+	const TArray<UMaterialExpression*>* FunctionExpressions = Function->GetFunctionExpressions();
+	if (FunctionExpressions)
+	{
+		if (!UpdateForExpressions(*FunctionExpressions, Association, ParameterIndex))
+		{
+			bResult = false;
+		}
 	}
 
 	FMaterialFunctionInfo NewFunctionInfo;
@@ -163,30 +173,41 @@ void FMaterialCachedExpressionData::UpdateForFunction(UMaterialFunctionInterface
 		}
 	}
 
-	const TArray<UMaterialExpression*>* FunctionExpressions = Function->GetFunctionExpressions();
-	if (FunctionExpressions)
-	{
-		UpdateForExpressions(*FunctionExpressions, Association, ParameterIndex);
-	}
+	return bResult;
 }
 
-void FMaterialCachedExpressionData::UpdateForLayerFunctions(const FMaterialLayersFunctions& LayerFunctions)
+bool FMaterialCachedExpressionData::UpdateForLayerFunctions(const FMaterialLayersFunctions& LayerFunctions)
 {
+	bool bResult = true;
 	for (int32 LayerIndex = 0; LayerIndex < LayerFunctions.Layers.Num(); ++LayerIndex)
 	{
-		UpdateForFunction(LayerFunctions.Layers[LayerIndex], LayerParameter, LayerIndex);
+		if (!UpdateForFunction(LayerFunctions.Layers[LayerIndex], LayerParameter, LayerIndex))
+		{
+			bResult = false;
+		}
 	}
 
 	for (int32 BlendIndex = 0; BlendIndex < LayerFunctions.Blends.Num(); ++BlendIndex)
 	{
-		UpdateForFunction(LayerFunctions.Blends[BlendIndex], BlendParameter, BlendIndex);
+		if (!UpdateForFunction(LayerFunctions.Blends[BlendIndex], BlendParameter, BlendIndex))
+		{
+			bResult = false;
+		}
 	}
+	return bResult;
 }
 
-void FMaterialCachedExpressionData::UpdateForExpressions(const TArray<UMaterialExpression*>& Expressions, EMaterialParameterAssociation Association, int32 ParameterIndex)
+bool FMaterialCachedExpressionData::UpdateForExpressions(const TArray<UMaterialExpression*>& Expressions, EMaterialParameterAssociation Association, int32 ParameterIndex)
 {
+	bool bResult = true;
 	for (UMaterialExpression* Expression : Expressions)
 	{
+		if (!Expression)
+		{
+			bResult = false;
+			continue;
+		}
+
 		UObject* ReferencedTexture = Expression->GetReferencedTexture();
 		checkf(!ReferencedTexture || Expression->CanReferenceTexture(), TEXT("This expression type missing an override for CanReferenceTexture?"));
 		if (Expression->CanReferenceTexture())
@@ -306,7 +327,10 @@ void FMaterialCachedExpressionData::UpdateForExpressions(const TArray<UMaterialE
 		}
 		else if (UMaterialExpressionMaterialFunctionCall* FunctionCall = Cast<UMaterialExpressionMaterialFunctionCall>(Expression))
 		{
-			UpdateForFunction(FunctionCall->MaterialFunction, GlobalParameter, -1);
+			if (!UpdateForFunction(FunctionCall->MaterialFunction, GlobalParameter, -1))
+			{
+				bResult = false;
+			}
 
 			// Update the function call node, so it can relink inputs and outputs as needed
 			// Update even if MaterialFunctionNode->MaterialFunction is NULL, because we need to remove the invalid inputs in that case
@@ -315,11 +339,16 @@ void FMaterialCachedExpressionData::UpdateForExpressions(const TArray<UMaterialE
 		else if (UMaterialExpressionMaterialAttributeLayers* LayersExpression = Cast<UMaterialExpressionMaterialAttributeLayers>(Expression))
 		{
 			checkf(Association == GlobalParameter, TEXT("UMaterialExpressionMaterialAttributeLayers can't be nested"));
-			UpdateForLayerFunctions(LayersExpression->DefaultLayers);
+			if (!UpdateForLayerFunctions(LayersExpression->DefaultLayers))
+			{
+				bResult = false;
+			}
 
 			LayersExpression->RebuildLayerGraph(false);
 		}
 	}
+
+	return bResult;
 }
 #endif // WITH_EDITOR
 
diff --git a/Engine/Source/Runtime/Engine/Private/Materials/MaterialExpressions.cpp b/Engine/Source/Runtime/Engine/Private/Materials/MaterialExpressions.cpp
--- a/Engine/Source/Runtime/Engine/Private/Materials/MaterialExpressions.cpp
+++ b/Engine/Source/Runtime/Engine/Private/Materials/MaterialExpressions.cpp
@@ -727,12 +727,6 @@ void UMaterialExpression::Serialize(FStructuredArchive::FRecord Record)
 #endif // WITH_EDITORONLY_DATA
 }
 
-bool UMaterialExpression::NeedsLoadForClient() const
-{
-	// Expressions that reference texture objects need to be cooked
-	return CanReferenceTexture() || GetReferencedTexture() != nullptr;
-}
-
 void UMaterialExpression::PostInitProperties()
 {
 	Super::PostInitProperties();
@@ -7199,12 +7193,6 @@ void UMaterialExpressionParameter::ValidateParameterName(const bool bAllowDuplic
 }
 #endif
 
-bool UMaterialExpressionParameter::NeedsLoadForClient() const
-{
-	// Keep named parameters
-	return ParameterName != NAME_None;
-}
-
 //
 //	UMaterialExpressionVectorParameter
 //
@@ -8165,11 +8153,6 @@ bool UMaterialExpressionQualitySwitch::IsResultMaterialAttributes(int32 OutputIn
 }
 #endif // WITH_EDITOR
 
-bool UMaterialExpressionQualitySwitch::NeedsLoadForClient() const
-{
-	return true;
-}
-
 //
 //	UMaterialExpressionFeatureLevelSwitch
 //
@@ -8296,11 +8279,6 @@ void UMaterialExpressionFeatureLevelSwitch::Serialize(FStructuredArchive::FRecor
 	}
 }
 
-bool UMaterialExpressionFeatureLevelSwitch::NeedsLoadForClient() const
-{
-	return true;
-}
-
 //
 //	UMaterialExpressionShadingPathSwitch
 //
@@ -8420,11 +8398,6 @@ bool UMaterialExpressionShadingPathSwitch::IsResultMaterialAttributes(int32 Outp
 }
 #endif // WITH_EDITOR
 
-bool UMaterialExpressionShadingPathSwitch::NeedsLoadForClient() const
-{
-	return true;
-}
-
 //
 //	UMaterialExpressionNormalize
 //
@@ -8733,11 +8706,6 @@ void UMaterialExpressionDynamicParameter::PostLoad()
 	}
 }
 
-bool UMaterialExpressionDynamicParameter::NeedsLoadForClient() const
-{
-	return true;
-}
-
 #if WITH_EDITORONLY_DATA
 void UMaterialExpressionDynamicParameter::UpdateDynamicParameterProperties()
 {
@@ -12613,11 +12581,6 @@ void UMaterialExpressionMaterialFunctionCall::PostLoad()
 	Super::PostLoad();
 }
 
-bool UMaterialExpressionMaterialFunctionCall::NeedsLoadForClient() const
-{
-	return true;
-}
-
 #if WITH_EDITORONLY_DATA
 bool UMaterialExpressionMaterialFunctionCall::IterateDependentFunctions(TFunctionRef<bool(UMaterialFunctionInterface*)> Predicate) const
 {
@@ -13722,11 +13685,6 @@ void UMaterialExpressionCollectionParameter::PostLoad()
 	Super::PostLoad();
 }
 
-bool UMaterialExpressionCollectionParameter::NeedsLoadForClient() const
-{
-	return true;
-}
-
 #if WITH_EDITOR
 void UMaterialExpressionCollectionParameter::PostEditChangeProperty(FPropertyChangedEvent& PropertyChangedEvent)
 {
diff --git a/Engine/Source/Runtime/Engine/Private/Materials/MaterialInstance.cpp b/Engine/Source/Runtime/Engine/Private/Materials/MaterialInstance.cpp
--- a/Engine/Source/Runtime/Engine/Private/Materials/MaterialInstance.cpp
+++ b/Engine/Source/Runtime/Engine/Private/Materials/MaterialInstance.cpp
@@ -4117,13 +4117,20 @@ void UMaterialInstance::UpdateCachedLayerParameters()
 		CachedExpressionData.ReferencedTextures = Parent->GetReferencedTextures();
 	}
 	
+	bool bCachedDataValid = true;
 	for (FStaticMaterialLayersParameter& LayerParameters : StaticParameters.MaterialLayersParameters)
 	{
-		CachedExpressionData.UpdateForLayerFunctions(LayerParameters.Value);
+		if (!CachedExpressionData.UpdateForLayerFunctions(LayerParameters.Value))
+		{
+			bCachedDataValid = false;
+		}
 	}
 
-	CachedLayerParameters = MoveTemp(CachedExpressionData.Parameters);
-	CachedReferencedTextures = MoveTemp(CachedExpressionData.ReferencedTextures);
+	if (bCachedDataValid)
+	{
+		CachedLayerParameters = MoveTemp(CachedExpressionData.Parameters);
+		CachedReferencedTextures = MoveTemp(CachedExpressionData.ReferencedTextures);
+	}
 }
 
 void UMaterialInstance::UpdateStaticPermutation(const FStaticParameterSet& NewParameters, FMaterialUpdateContext* MaterialUpdateContext)
diff --git a/Engine/Source/Runtime/Engine/Public/MaterialCachedData.h b/Engine/Source/Runtime/Engine/Public/MaterialCachedData.h
--- a/Engine/Source/Runtime/Engine/Public/MaterialCachedData.h
+++ b/Engine/Source/Runtime/Engine/Public/MaterialCachedData.h
@@ -162,9 +162,10 @@ struct FMaterialCachedExpressionData
 	GENERATED_USTRUCT_BODY()
 
 #if WITH_EDITOR
-	void UpdateForExpressions(const TArray<UMaterialExpression*>& Expressions, EMaterialParameterAssociation Association, int32 ParameterIndex);
-	void UpdateForFunction(UMaterialFunctionInterface* Function, EMaterialParameterAssociation Association, int32 ParameterIndex);
-	void UpdateForLayerFunctions(const FMaterialLayersFunctions& LayerFunctions);
+	/** Returns 'false' if update is incomplete, due to missing expression data (stripped from non-editor build) */
+	bool UpdateForExpressions(const TArray<UMaterialExpression*>& Expressions, EMaterialParameterAssociation Association, int32 ParameterIndex);
+	bool UpdateForFunction(UMaterialFunctionInterface* Function, EMaterialParameterAssociation Association, int32 ParameterIndex);
+	bool UpdateForLayerFunctions(const FMaterialLayersFunctions& LayerFunctions);
 #endif // WITH_EDITOR
 
 	void Reset();
diff --git a/Engine/Source/Runtime/Landscape/Classes/Materials/MaterialExpressionLandscapeGrassOutput.h b/Engine/Source/Runtime/Landscape/Classes/Materials/MaterialExpressionLandscapeGrassOutput.h
--- a/Engine/Source/Runtime/Landscape/Classes/Materials/MaterialExpressionLandscapeGrassOutput.h
+++ b/Engine/Source/Runtime/Landscape/Classes/Materials/MaterialExpressionLandscapeGrassOutput.h
@@ -52,7 +52,6 @@ class LANDSCAPE_API UMaterialExpressionLandscapeGrassOutput : public UMaterialEx
 #endif
 
 	//~ Begin UObject Interface
-	virtual bool NeedsLoadForClient() const override;
 #if WITH_EDITOR
 	virtual void PostEditChangeProperty(FPropertyChangedEvent& PropertyChangedEvent) override;
 	//~ End UObject Interface
diff --git a/Engine/Source/Runtime/Landscape/Classes/Materials/MaterialExpressionLandscapeLayerBlend.h b/Engine/Source/Runtime/Landscape/Classes/Materials/MaterialExpressionLandscapeLayerBlend.h
--- a/Engine/Source/Runtime/Landscape/Classes/Materials/MaterialExpressionLandscapeLayerBlend.h
+++ b/Engine/Source/Runtime/Landscape/Classes/Materials/MaterialExpressionLandscapeLayerBlend.h
@@ -74,7 +74,6 @@ class LANDSCAPE_API UMaterialExpressionLandscapeLayerBlend : public UMaterialExp
 
 	//~ Begin UObject Interface
 	virtual void Serialize(FStructuredArchive::FRecord Record) override;
-	virtual bool NeedsLoadForClient() const override;
 #if WITH_EDITOR
 	virtual void PostEditChangeProperty(FPropertyChangedEvent& PropertyChangedEvent) override;
 	//~ End UObject Interface
diff --git a/Engine/Source/Runtime/Landscape/Classes/Materials/MaterialExpressionLandscapeLayerCoords.h b/Engine/Source/Runtime/Landscape/Classes/Materials/MaterialExpressionLandscapeLayerCoords.h
--- a/Engine/Source/Runtime/Landscape/Classes/Materials/MaterialExpressionLandscapeLayerCoords.h
+++ b/Engine/Source/Runtime/Landscape/Classes/Materials/MaterialExpressionLandscapeLayerCoords.h
@@ -66,10 +66,6 @@ class LANDSCAPE_API UMaterialExpressionLandscapeLayerCoords : public UMaterialEx
 	virtual void GetCaption(TArray<FString>& OutCaptions) const override;
 #endif
 	//~ End UMaterialExpression Interface
-
-	//~ Begin UObject Interface
-	virtual bool NeedsLoadForClient() const override;
-	//~ End UObject Interface
 };
 
 
diff --git a/Engine/Source/Runtime/Landscape/Classes/Materials/MaterialExpressionLandscapeLayerSample.h b/Engine/Source/Runtime/Landscape/Classes/Materials/MaterialExpressionLandscapeLayerSample.h
--- a/Engine/Source/Runtime/Landscape/Classes/Materials/MaterialExpressionLandscapeLayerSample.h
+++ b/Engine/Source/Runtime/Landscape/Classes/Materials/MaterialExpressionLandscapeLayerSample.h
@@ -45,10 +45,6 @@ public:
 	 * Called to get list of parameter names for static parameter sets
 	 */
 	void GetAllParameterInfo(TArray<FMaterialParameterInfo> &OutParameterInfo, TArray<FGuid> &OutParameterIds, const FMaterialParameterInfo& InBaseParameterInfo) const;
-
-	//~ Begin UObject Interface
-	virtual bool NeedsLoadForClient() const override;
-	//~ End UObject Interface
 };
 
 
diff --git a/Engine/Source/Runtime/Landscape/Classes/Materials/MaterialExpressionLandscapeLayerSwitch.h b/Engine/Source/Runtime/Landscape/Classes/Materials/MaterialExpressionLandscapeLayerSwitch.h
--- a/Engine/Source/Runtime/Landscape/Classes/Materials/MaterialExpressionLandscapeLayerSwitch.h
+++ b/Engine/Source/Runtime/Landscape/Classes/Materials/MaterialExpressionLandscapeLayerSwitch.h
@@ -66,7 +66,6 @@ public:
 	 * and immediately after any undo/redo.
 	 */
 	virtual void PostLoad() override;
-	virtual bool NeedsLoadForClient() const override;
 	//~ End UObject Interface
 };
 
diff --git a/Engine/Source/Runtime/Landscape/Classes/Materials/MaterialExpressionLandscapeLayerWeight.h b/Engine/Source/Runtime/Landscape/Classes/Materials/MaterialExpressionLandscapeLayerWeight.h
--- a/Engine/Source/Runtime/Landscape/Classes/Materials/MaterialExpressionLandscapeLayerWeight.h
+++ b/Engine/Source/Runtime/Landscape/Classes/Materials/MaterialExpressionLandscapeLayerWeight.h
@@ -67,7 +67,6 @@ public:
 	 * and immediately after any undo/redo.
 	 */
 	virtual void PostLoad() override;
-	virtual bool NeedsLoadForClient() const override;
 	//~ End UObject Interface
 };
 
diff --git a/Engine/Source/Runtime/Landscape/Classes/Materials/MaterialExpressionLandscapeVisibilityMask.h b/Engine/Source/Runtime/Landscape/Classes/Materials/MaterialExpressionLandscapeVisibilityMask.h
--- a/Engine/Source/Runtime/Landscape/Classes/Materials/MaterialExpressionLandscapeVisibilityMask.h
+++ b/Engine/Source/Runtime/Landscape/Classes/Materials/MaterialExpressionLandscapeVisibilityMask.h
@@ -40,10 +40,6 @@ public:
 	 * Called to get list of parameter names for static parameter sets
 	 */
 	void GetAllParameterInfo(TArray<FMaterialParameterInfo> &OutParameterInfo, TArray<FGuid> &OutParameterIds, const FMaterialParameterInfo& InBaseParameterInfo) const;
-
-	//~ Begin UObject Interface
-	virtual bool NeedsLoadForClient() const override;
-	//~ End UObject Interface
 };
 
 
diff --git a/Engine/Source/Runtime/Landscape/Private/LandscapeGrass.cpp b/Engine/Source/Runtime/Landscape/Private/LandscapeGrass.cpp
--- a/Engine/Source/Runtime/Landscape/Private/LandscapeGrass.cpp
+++ b/Engine/Source/Runtime/Landscape/Private/LandscapeGrass.cpp
@@ -1121,10 +1121,6 @@ FName UMaterialExpressionLandscapeGrassOutput::GetInputName(int32 InputIndex) co
 }
 #endif // WITH_EDITOR
 
-bool UMaterialExpressionLandscapeGrassOutput::NeedsLoadForClient() const
-{
-	return true;
-}
 
 #if WITH_EDITOR
 void UMaterialExpressionLandscapeGrassOutput::PostEditChangeProperty(FPropertyChangedEvent& PropertyChangedEvent)
diff --git a/Engine/Source/Runtime/Landscape/Private/Materials/MaterialExpressionLandscapeLayerBlend.cpp b/Engine/Source/Runtime/Landscape/Private/Materials/MaterialExpressionLandscapeLayerBlend.cpp
--- a/Engine/Source/Runtime/Landscape/Private/Materials/MaterialExpressionLandscapeLayerBlend.cpp
+++ b/Engine/Source/Runtime/Landscape/Private/Materials/MaterialExpressionLandscapeLayerBlend.cpp
@@ -341,10 +341,4 @@ void UMaterialExpressionLandscapeLayerBlend::GetAllParameterInfo(TArray<FMateria
 	}
 }
 
-bool UMaterialExpressionLandscapeLayerBlend::NeedsLoadForClient() const
-{
-	return true;
-}
-
-
 #undef LOCTEXT_NAMESPACE
diff --git a/Engine/Source/Runtime/Landscape/Private/Materials/MaterialExpressionLandscapeLayerCoords.cpp b/Engine/Source/Runtime/Landscape/Private/Materials/MaterialExpressionLandscapeLayerCoords.cpp
--- a/Engine/Source/Runtime/Landscape/Private/Materials/MaterialExpressionLandscapeLayerCoords.cpp
+++ b/Engine/Source/Runtime/Landscape/Private/Materials/MaterialExpressionLandscapeLayerCoords.cpp
@@ -89,9 +89,4 @@ void UMaterialExpressionLandscapeLayerCoords::GetCaption(TArray<FString>& OutCap
 }
 #endif // WITH_EDITOR
 
-bool UMaterialExpressionLandscapeLayerCoords::NeedsLoadForClient() const
-{
-	return true;
-}
-
 #undef LOCTEXT_NAMESPACE
diff --git a/Engine/Source/Runtime/Landscape/Private/Materials/MaterialExpressionLandscapeLayerSample.cpp b/Engine/Source/Runtime/Landscape/Private/Materials/MaterialExpressionLandscapeLayerSample.cpp
--- a/Engine/Source/Runtime/Landscape/Private/Materials/MaterialExpressionLandscapeLayerSample.cpp
+++ b/Engine/Source/Runtime/Landscape/Private/Materials/MaterialExpressionLandscapeLayerSample.cpp
@@ -97,9 +97,4 @@ void UMaterialExpressionLandscapeLayerSample::GetAllParameterInfo(TArray<FMateri
 	}
 }
 
-bool UMaterialExpressionLandscapeLayerSample::NeedsLoadForClient() const
-{
-	return ParameterName != NAME_None;
-}
-
 #undef LOCTEXT_NAMESPACE
diff --git a/Engine/Source/Runtime/Landscape/Private/Materials/MaterialExpressionLandscapeLayerSwitch.cpp b/Engine/Source/Runtime/Landscape/Private/Materials/MaterialExpressionLandscapeLayerSwitch.cpp
--- a/Engine/Source/Runtime/Landscape/Private/Materials/MaterialExpressionLandscapeLayerSwitch.cpp
+++ b/Engine/Source/Runtime/Landscape/Private/Materials/MaterialExpressionLandscapeLayerSwitch.cpp
@@ -149,9 +149,4 @@ void UMaterialExpressionLandscapeLayerSwitch::GetAllParameterInfo(TArray<FMateri
 	}
 }
 
-bool UMaterialExpressionLandscapeLayerSwitch::NeedsLoadForClient() const
-{
-	return ParameterName != NAME_None;
-}
-
 #undef LOCTEXT_NAMESPACE
diff --git a/Engine/Source/Runtime/Landscape/Private/Materials/MaterialExpressionLandscapeLayerWeight.cpp b/Engine/Source/Runtime/Landscape/Private/Materials/MaterialExpressionLandscapeLayerWeight.cpp
--- a/Engine/Source/Runtime/Landscape/Private/Materials/MaterialExpressionLandscapeLayerWeight.cpp
+++ b/Engine/Source/Runtime/Landscape/Private/Materials/MaterialExpressionLandscapeLayerWeight.cpp
@@ -135,9 +135,4 @@ void UMaterialExpressionLandscapeLayerWeight::GetAllParameterInfo(TArray<FMateri
 	}
 }
 
-bool UMaterialExpressionLandscapeLayerWeight::NeedsLoadForClient() const
-{
-	return ParameterName != NAME_None;
-}
-
 #undef LOCTEXT_NAMESPACE
diff --git a/Engine/Source/Runtime/Landscape/Private/Materials/MaterialExpressionLandscapeVisibilityMask.cpp b/Engine/Source/Runtime/Landscape/Private/Materials/MaterialExpressionLandscapeVisibilityMask.cpp
--- a/Engine/Source/Runtime/Landscape/Private/Materials/MaterialExpressionLandscapeVisibilityMask.cpp
+++ b/Engine/Source/Runtime/Landscape/Private/Materials/MaterialExpressionLandscapeVisibilityMask.cpp
@@ -74,9 +74,5 @@ void UMaterialExpressionLandscapeVisibilityMask::GetCaption(TArray<FString>& Out
 }
 #endif // WITH_EDITOR
 
-bool UMaterialExpressionLandscapeVisibilityMask::NeedsLoadForClient() const
-{
-	return ParameterName != NAME_None;
-}
 
 #undef LOCTEXT_NAMESPACE
