Mitigated check failure when getting bad data for post process volumes. Added ensures to get better information in case this issue returns.

diff --git a/Engine/Source/Runtime/Engine/Private/World.cpp b/Engine/Source/Runtime/Engine/Private/World.cpp
--- a/Engine/Source/Runtime/Engine/Private/World.cpp
+++ b/Engine/Source/Runtime/Engine/Private/World.cpp
@@ -7456,6 +7456,8 @@ static void DoPostProcessVolume(IInterface_PostProcessVolume* Volume, FVector Vi
 	float DistanceToPoint = 0.0f;
 	float LocalWeight = FMath::Clamp(VolumeProperties.BlendWeight, 0.0f, 1.0f);
 
+	ensureMsgf((LocalWeight >= 0 && LocalWeight <= 1.0f), TEXT("Invalid post process blend weight retrieved from volume (%f)"), LocalWeight);
+
 	if (!VolumeProperties.bIsUnbound)
 	{
 		float SquaredBlendRadius = VolumeProperties.BlendRadius * VolumeProperties.BlendRadius;
@@ -7475,7 +7477,12 @@ static void DoPostProcessVolume(IInterface_PostProcessVolume* Volume, FVector Vi
 				{
 					LocalWeight *= 1.0f - DistanceToPoint / VolumeProperties.BlendRadius;
 
-					check(LocalWeight >= 0 && LocalWeight <= 1.0f);
+					if(!(LocalWeight >= 0 && LocalWeight <= 1.0f))
+					{
+						// Mitigate crash here by disabling this volume and generating info regarding the calculation that went wrong.
+						ensureMsgf(false, TEXT("Invalid LocalWeight after post process volume weight calculation (Local: %f, DtP: %f, Radius: %f, SettingsWeight: %f)"), LocalWeight, DistanceToPoint, VolumeProperties.BlendRadius, VolumeProperties.BlendWeight);
+						LocalWeight = 0.0f;
+					}
 				}
 			}
 		}
