Add more information to warning message in UMaterialExpressionGetMaterialAttributes::PostLoad()

diff --git a/Engine/Source/Runtime/Engine/Private/Materials/MaterialExpressions.cpp b/Engine/Source/Runtime/Engine/Private/Materials/MaterialExpressions.cpp
--- a/Engine/Source/Runtime/Engine/Private/Materials/MaterialExpressions.cpp
+++ b/Engine/Source/Runtime/Engine/Private/Materials/MaterialExpressions.cpp
@@ -5887,7 +5887,9 @@ void UMaterialExpressionGetMaterialAttributes::PostLoad()
 				Function->GetName(MaterialName);
 			}
 
-			UE_LOG(LogMaterial, Warning, TEXT("Serialized attribute that no longer exists (%s) for material \"%s\"."), *(Outputs[i].OutputName.ToString()), *MaterialName);
+			const FString AttributeName = FMaterialAttributeDefinitionMap::GetAttributeName(AttributeGetTypes[i - 1]);
+			UE_LOG(LogMaterial, Warning, TEXT("Material \"%s\", serialized attribute '%s' with OutputName '%s', expected '%s'"),
+				*MaterialName, *AttributeName, *(Outputs[i].OutputName.ToString()), *DisplayName);
 			Outputs[i].OutputName = *DisplayName;
 		}
 	}
