Fix for another NDC - WeightedRandomSampler can have some uninitalized values at the end

diff --git a/Engine/Source/Runtime/Engine/Private/WeightedRandomSampler.cpp b/Engine/Source/Runtime/Engine/Private/WeightedRandomSampler.cpp
--- a/Engine/Source/Runtime/Engine/Private/WeightedRandomSampler.cpp
+++ b/Engine/Source/Runtime/Engine/Private/WeightedRandomSampler.cpp
@@ -64,12 +64,14 @@ void FWeightedRandomSampler::Initialize()
 	for (int32& LargeIdx : Large)
 	{
 		Prob[LargeIdx] = 1;
+		Alias[LargeIdx] = LargeIdx;
 	}
 
 	//FP inaccuracies can lead S to still have entries on occasion, these should also be now 1.
 	for (int32& SmallIdx : Small)
 	{
 		Prob[SmallIdx] = 1;
+		Alias[SmallIdx] = SmallIdx;
 	}
 }
 
