CSV added:
WorldCollision/StatWaitForAllAsyncTraceTasks
ChaosPhysics/AABBTreeDirtyElementCount

diff --git a/Engine/Source/Runtime/Engine/Private/Collision/WorldCollisionAsync.cpp b/Engine/Source/Runtime/Engine/Private/Collision/WorldCollisionAsync.cpp
--- a/Engine/Source/Runtime/Engine/Private/Collision/WorldCollisionAsync.cpp
+++ b/Engine/Source/Runtime/Engine/Private/Collision/WorldCollisionAsync.cpp
@@ -14,6 +14,9 @@
 #include "Engine/World.h"
 #include "PhysicsEngine/BodyInstance.h"
 #include "Physics/PhysicsInterfaceCore.h"
+#include "ProfilingDebugging/CsvProfiler.h"
+
+CSV_DEFINE_CATEGORY(WorldCollision, true);
 
 /**
  * Async trace functions
@@ -458,6 +461,7 @@ void UWorld::WaitForAllAsyncTraceTasks()
 		if (DataBufferExecuted.AsyncTraceCompletionEvent.Num() > 0)
 		{
 			QUICK_SCOPE_CYCLE_COUNTER(STAT_WaitForAllAsyncTraceTasks);
+			CSV_SCOPED_TIMING_STAT(WorldCollision, StatWaitForAllAsyncTraceTasks);
 			FTaskGraphInterface::Get().WaitUntilTasksComplete(DataBufferExecuted.AsyncTraceCompletionEvent, ENamedThreads::GameThread);
 			DataBufferExecuted.AsyncTraceCompletionEvent.Reset();
 		}
diff --git a/Engine/Source/Runtime/Engine/Private/PhysicsEngine/Experimental/PhysScene_Chaos.cpp b/Engine/Source/Runtime/Engine/Private/PhysicsEngine/Experimental/PhysScene_Chaos.cpp
--- a/Engine/Source/Runtime/Engine/Private/PhysicsEngine/Experimental/PhysScene_Chaos.cpp
+++ b/Engine/Source/Runtime/Engine/Private/PhysicsEngine/Experimental/PhysScene_Chaos.cpp
@@ -41,6 +41,7 @@
 #include "ChaosSolvers/Public/EventsData.h"
 #include "ChaosSolvers/Public/EventManager.h"
 #include "Physics/Experimental/PhysicsUserData_Chaos.h"
+#include "ProfilingDebugging/CsvProfiler.h"
 
 #if !UE_BUILD_SHIPPING
 #include "Engine/World.h"
@@ -60,6 +61,10 @@ TAutoConsoleVariable<int32> CVar_ChaosSimulationEnable(TEXT("P.Chaos.Simulation.
 
 DECLARE_CYCLE_STAT(TEXT("Update Kinematics On Deferred SkelMeshes"), STAT_UpdateKinematicsOnDeferredSkelMeshesChaos, STATGROUP_Physics);
 
+#if WITH_CHAOS
+CSV_DEFINE_CATEGORY(ChaosPhysics, true);
+#endif // WITH_CHAOS
+
 #if WITH_EDITOR
 #include "Editor.h"
 #endif
@@ -1905,9 +1910,34 @@ void FPhysScene_ChaosInterface::StartFrame()
 	}
 }
 
+// Find the number of dirty elements in all substructures that has dirty elements that we know of
+// This is non recursive for now
+// Todo: consider making DirtyElementsCount a method on ISpatialAcceleration instead
+int32 FPhysScene_ChaosInterface::DirtyElementCount(Chaos::ISpatialAccelerationCollection<Chaos::TAccelerationStructureHandle<Chaos::FReal, 3>, Chaos::FReal, 3>& Collection)
+{
+	using namespace Chaos;
+	int32 DirtyElements = 0;
+	TArray<FSpatialAccelerationIdx> SpatialIndices = Collection.GetAllSpatialIndices();
+	for (const FSpatialAccelerationIdx SpatialIndex : SpatialIndices)
+	{
+		auto SubStructure = Collection.GetSubstructure(SpatialIndex);
+		if (const auto AABBTree = SubStructure->template As<TAABBTree<TAccelerationStructureHandle<FReal, 3>, TAABBTreeLeafArray<TAccelerationStructureHandle<FReal, 3>, FReal>, FReal>>())
+		{
+			DirtyElements += AABBTree->NumDirtyElements();
+		}
+		else if (const auto AABBTreeBV = SubStructure->template As<TAABBTree<TAccelerationStructureHandle<FReal, 3>, TBoundingVolume<TAccelerationStructureHandle<FReal, 3>, FReal, 3>, FReal>>())
+		{
+			DirtyElements += AABBTree->NumDirtyElements();
+		}
+	}
+	return DirtyElements;
+}
+
+
 void FPhysScene_ChaosInterface::EndFrame(ULineBatchComponent* InLineBatcher)
 {
 	using namespace Chaos;
+	using SpatialAccelerationCollection = Chaos::ISpatialAccelerationCollection<Chaos::TAccelerationStructureHandle<FReal, 3>, FReal, 3>;
 
 	SCOPE_CYCLE_COUNTER(STAT_Scene_EndFrame);
 
@@ -1921,6 +1951,9 @@ void FPhysScene_ChaosInterface::EndFrame(ULineBatchComponent* InLineBatcher)
 
 	Chaos::IDispatcher* Dispatcher = SolverModule->GetDispatcher();
 
+	int32 DirtyElements = DirtyElementCount(Scene.GetSpacialAcceleration()->AsChecked<SpatialAccelerationCollection>());
+	CSV_CUSTOM_STAT(ChaosPhysics, AABBTreeDirtyElementCount, DirtyElements, ECsvCustomStatOp::Set);
+
 	switch(Dispatcher->GetMode())
 	{
 	case EChaosThreadingMode::SingleThread:
diff --git a/Engine/Source/Runtime/Engine/Public/Physics/Experimental/PhysScene_Chaos.h b/Engine/Source/Runtime/Engine/Public/Physics/Experimental/PhysScene_Chaos.h
--- a/Engine/Source/Runtime/Engine/Public/Physics/Experimental/PhysScene_Chaos.h
+++ b/Engine/Source/Runtime/Engine/Public/Physics/Experimental/PhysScene_Chaos.h
@@ -383,6 +383,7 @@ public:
 
 	void AddPendingOnConstraintBreak(FConstraintInstance* ConstraintInstance, int32 SceneType);
 	void AddPendingSleepingEvent(FBodyInstance* BI, ESleepEvent SleepEventType, int32 SceneType);
+	int32 DirtyElementCount(Chaos::ISpatialAccelerationCollection<Chaos::TAccelerationStructureHandle<Chaos::FReal, 3>, Chaos::FReal, 3>& Collection);
 
 	TArray<FCollisionNotifyInfo>& GetPendingCollisionNotifies(int32 SceneType);
 
diff --git a/Engine/Source/Runtime/Experimental/Chaos/Private/Chaos/PBDRigidsEvolution.cpp b/Engine/Source/Runtime/Experimental/Chaos/Private/Chaos/PBDRigidsEvolution.cpp
--- a/Engine/Source/Runtime/Experimental/Chaos/Private/Chaos/PBDRigidsEvolution.cpp
+++ b/Engine/Source/Runtime/Experimental/Chaos/Private/Chaos/PBDRigidsEvolution.cpp
@@ -287,12 +287,12 @@ namespace Chaos
 		{
 			if (!bIsSingleThreaded)
 			{
-			// This operation is slow!
-			SCOPE_CYCLE_COUNTER(STAT_CopyAccelerationStructure);
-			AccelerationStructureCopy = AsUniqueSpatialAccelerationChecked<FAccelerationStructure>(AccelerationStructure->Copy());
+				// This operation is slow!
+				SCOPE_CYCLE_COUNTER(STAT_CopyAccelerationStructure);
+				AccelerationStructureCopy = AsUniqueSpatialAccelerationChecked<FAccelerationStructure>(AccelerationStructure->Copy());
+			}
 		}
 	}
-	}
 
 	template<class FPBDRigidsEvolution, class FPBDCollisionConstraint, class T, int d>
 	void TPBDRigidsEvolutionBase<FPBDRigidsEvolution, FPBDCollisionConstraint, T, d>::ApplyParticlePendingData(const FPendingSpatialData& SpatialData, FAccelerationStructure& AccelerationStructure, bool bUpdateCache)
@@ -471,8 +471,7 @@ namespace Chaos
 				}
 				bExternalReady = true;
 			}
-
-			// we run the task for both starting a new accel structure as well as for the timeslicing
+			
 			// we run the task for both starting a new accel structure as well as for the timeslicing
 			AccelerationStructureTaskComplete = TGraphTask<FChaosAccelerationStructureTask>::CreateTask().ConstructAndDispatchWhenReady(*SpatialCollectionFactory, SpatialAccelerationCache, AsyncInternalAcceleration, AsyncExternalAcceleration, ForceFullBuild, bIsSingleThreaded);
 
