Added "NetProfile Autostop" console command
* Starts a nprof recording but stops it after 30 secs.
* Time value can be changed by adding Time=Seconds in the console command.
ex: "NetProfile Autostop Time=60" to capture for one minute

diff --git a/Engine/Source/Runtime/Engine/Private/NetworkProfiler.cpp b/Engine/Source/Runtime/Engine/Private/NetworkProfiler.cpp
--- a/Engine/Source/Runtime/Engine/Private/NetworkProfiler.cpp
+++ b/Engine/Source/Runtime/Engine/Private/NetworkProfiler.cpp
@@ -840,6 +840,16 @@ bool FNetworkProfiler::Exec( UWorld * InWorld, const TCHAR* Cmd, FOutputDevice &
 	{
 		EnableTracking( false );
 	} 
+	else if (FParse::Command(&Cmd, TEXT("AUTOSTOP")))
+	{
+		EnableTracking(true);
+
+		// Default to 30secs
+		float AutoStopTime = 30.f;
+		FParse::Value(Cmd, TEXT("TIME="), AutoStopTime);
+
+		InWorld->GetTimerManager().SetTimer(AutoStopTimerHandle, FTimerDelegate::CreateRaw(this, &FNetworkProfiler::AutoStopTracking), AutoStopTime, false);
+	}
 	else 
 	{
 		// Default to toggle
@@ -860,5 +870,13 @@ bool FNetworkProfiler::Exec( UWorld * InWorld, const TCHAR* Cmd, FOutputDevice &
 	return true;
 }
 
+void FNetworkProfiler::AutoStopTracking()
+{
+	if (bIsTrackingEnabled)
+	{
+		EnableTracking(false);
+	}
+}
+
 #endif	//#if USE_NETWORK_PROFILER
 
diff --git a/Engine/Source/Runtime/Engine/Public/Net/NetworkProfiler.h b/Engine/Source/Runtime/Engine/Public/Net/NetworkProfiler.h
--- a/Engine/Source/Runtime/Engine/Public/Net/NetworkProfiler.h
+++ b/Engine/Source/Runtime/Engine/Public/Net/NetworkProfiler.h
@@ -111,6 +111,8 @@ private:
 	/** Delegate that's fired when tracking stops on the current network profile */
 	FOnNetworkProfileFinished				OnNetworkProfileFinishedDelegate;
 
+	FTimerHandle							AutoStopTimerHandle;
+
 	/** All the data required for writing sent bunches to the profiler stream						*/
 	struct FSendBunchInfo
 	{
@@ -182,6 +184,8 @@ private:
 	// Set of names that correspond to Object's whose top level property names have been exported.
 	TSet<FString> ExportedObjects;
 
+	void AutoStopTracking();
+
 public:
 	/**
 	 * Constructor, initializing members.
