Root Motion Sources: Added rotation support for override root motion sources. Added the ability for SimulatedProxy root motion sources to trigger smoothing code in CharacterMovementComponent. PendingAddRootMotionSources are now replicated.

Zak.Middleton, Frank.Gigliotti
#rnx

diff --git a/Engine/Source/Runtime/Engine/Classes/GameFramework/RootMotionSource.h b/Engine/Source/Runtime/Engine/Classes/GameFramework/RootMotionSource.h
--- a/Engine/Source/Runtime/Engine/Classes/GameFramework/RootMotionSource.h
+++ b/Engine/Source/Runtime/Engine/Classes/GameFramework/RootMotionSource.h
@@ -276,6 +276,15 @@ struct ENGINE_API FRootMotionSource
 	 */
 	bool bNeedsSimulatedCatchup;
 
+	/** 
+	 *  True when this RootMotionSource is running on a SimulatedProxy and has marked itself as needing
+	 *  location/rotation smoothing. This smoothing logic is what runs whenever the SimulatedProxy receives
+	 *  a replication update from the server (SimulatedTick in CharacterMovementComponent, see bCorrectedToServer 
+	 *  and SmoothCorrection), this flag allows the RootMotionSource to signal to the CMC that a change
+	 *  was done significant enough that needs smoothing. 
+	 */
+	bool bSimulatedNeedsSmoothing;
+
 	/** Root Motion generated by this Source */
 	UPROPERTY()
 	FRootMotionMovementParams RootMotionParams;
@@ -794,6 +803,12 @@ struct ENGINE_API FRootMotionSourceGroup
 	/**  Helper function for accumulating additive velocity into InOutVelocity */
 	void AccumulateAdditiveRootMotionVelocity(float DeltaTime, const ACharacter& Character, const UCharacterMovementComponent& MoveComponent, FVector& InOutVelocity) const;
 
+	/** Get rotation output of current override root motion source, returns true if OutRotation was filled */
+	bool GetOverrideRootMotionRotation(float DeltaTime, const ACharacter& Character, const UCharacterMovementComponent& MoveComponent, FQuat& OutRotation) const;
+
+	/** Gets whether any active root motion source has been marked to need NetworkSmoothing as a SimulatedProxy */
+	bool NeedsSimulatedSmoothing() const;
+
 	/** Sets the StartTime of all pending root motion sources to be at least this time, can be used on servers to match client-side start times */
 	void SetPendingRootMotionSourceMinStartTimes(float NewStartTime);
 
@@ -873,6 +888,9 @@ protected:
 	/** Accumulates contributions for velocity into InOutVelocity for a given type of root motion from this group */
 	void AccumulateRootMotionVelocityFromSource(const FRootMotionSource& RootMotionSource, float DeltaTime, const ACharacter& Character, const UCharacterMovementComponent& MoveComponent, FVector& InOutVelocity) const;
 
+	/** Helper function for serializing array of root motion sources */
+	static void NetSerializeRMSArray(FArchive& Ar, class UPackageMap* Map, bool& bOutSuccess, TArray< TSharedPtr<FRootMotionSource> >& RootMotionSourceArray);
+
 };
 
 template<>
diff --git a/Engine/Source/Runtime/Engine/Private/Components/CharacterMovementComponent.cpp b/Engine/Source/Runtime/Engine/Private/Components/CharacterMovementComponent.cpp
--- a/Engine/Source/Runtime/Engine/Private/Components/CharacterMovementComponent.cpp
+++ b/Engine/Source/Runtime/Engine/Private/Components/CharacterMovementComponent.cpp
@@ -1564,7 +1564,7 @@ void UCharacterMovementComponent::SimulatedTick(float DeltaSeconds)
 		PerformMovement(DeltaSeconds);
 
 		// After movement correction, smooth out error in position if any.
-		if( bCorrectedToServer )
+		if( bCorrectedToServer || CurrentRootMotion.NeedsSimulatedSmoothing() )
 		{
 			SmoothCorrection(OldLocation, OldRotation, UpdatedComponent->GetComponentLocation(), UpdatedComponent->GetComponentQuat());
 		}
@@ -2444,6 +2444,16 @@ void UCharacterMovementComponent::PerformMovement(float DeltaSeconds)
 			// Root Motion has been used, clear
 			RootMotionParams.Clear();
 		}
+		else if (CurrentRootMotion.HasActiveRootMotionSources())
+		{
+			FQuat RootMotionRotationQuat;
+			if (CurrentRootMotion.GetOverrideRootMotionRotation(DeltaSeconds, *CharacterOwner, *this, RootMotionRotationQuat))
+			{
+				const FQuat OldActorRotationQuat = UpdatedComponent->GetComponentQuat();
+				const FQuat NewActorRotationQuat = RootMotionRotationQuat * OldActorRotationQuat;
+				MoveUpdatedComponent(FVector::ZeroVector, NewActorRotationQuat, true);
+			}
+		}
 
 		// consume path following requested velocity
 		bHasRequestedVelocity = false;
diff --git a/Engine/Source/Runtime/Engine/Private/GameFramework/RootMotionSource.cpp b/Engine/Source/Runtime/Engine/Private/GameFramework/RootMotionSource.cpp
--- a/Engine/Source/Runtime/Engine/Private/GameFramework/RootMotionSource.cpp
+++ b/Engine/Source/Runtime/Engine/Private/GameFramework/RootMotionSource.cpp
@@ -163,6 +163,7 @@ FRootMotionSource::FRootMotionSource()
 	, Duration(-1.0f)
 	, bInLocalSpace(false)
 	, bNeedsSimulatedCatchup(false)
+	, bSimulatedNeedsSmoothing(false)
 {
 }
 
@@ -1562,6 +1563,7 @@ void FRootMotionSourceGroup::PrepareRootMotion(float DeltaTime, const ACharacter
 					SimulationTime = FMath::Max(SimulationTime, 0.f);
 
 					// Do the Preparation (calculates root motion transforms to be applied)
+					RootMotionSource->bSimulatedNeedsSmoothing = false;
 					RootMotionSource->PrepareRootMotion(SimulationTime, DeltaTime, Character, MoveComponent);
 					LastAccumulatedSettings += RootMotionSource->Settings;
 					RootMotionSource->Status.SetFlag(ERootMotionSourceStatusFlags::Prepared);
@@ -1690,6 +1692,37 @@ void FRootMotionSourceGroup::AccumulateRootMotionVelocityFromSource
 	}
 }
 
+bool FRootMotionSourceGroup::GetOverrideRootMotionRotation
+	(
+		float DeltaTime, 
+		const ACharacter& Character, 
+		const UCharacterMovementComponent& MoveComponent, 
+		FQuat& OutRotation
+	) const
+{
+	for (const auto& RootMotionSource : RootMotionSources)
+	{
+		if (RootMotionSource.IsValid() && RootMotionSource->AccumulateMode == ERootMotionAccumulateMode::Override)
+		{
+			OutRotation = RootMotionSource->RootMotionParams.GetRootMotionTransform().GetRotation();
+			return !OutRotation.IsIdentity();
+		}
+	}
+	return false;
+}
+
+bool FRootMotionSourceGroup::NeedsSimulatedSmoothing() const
+{
+	for (const auto& RootMotionSource : RootMotionSources)
+	{
+		if (RootMotionSource.IsValid() && RootMotionSource->bSimulatedNeedsSmoothing)
+		{
+			return true;
+		}
+	}
+	return false;
+}
+
 void FRootMotionSourceGroup::SetPendingRootMotionSourceMinStartTimes(float NewStartTime)
 {
 	for (auto& RootMotionSource : PendingAddRootMotionSources)
@@ -1941,30 +1974,24 @@ void FRootMotionSourceGroup::UpdateStateFrom(const FRootMotionSourceGroup& Group
 	}
 }
 
-bool FRootMotionSourceGroup::NetSerialize(FArchive& Ar, class UPackageMap* Map, bool& bOutSuccess)
+void FRootMotionSourceGroup::NetSerializeRMSArray(FArchive& Ar, class UPackageMap* Map, bool& bOutSuccess, TArray< TSharedPtr<FRootMotionSource> >& RootMotionSourceArray)
 {
 	uint8 SourcesNum;
 	if (Ar.IsSaving())
 	{
-		UE_CLOG(RootMotionSources.Num() > MAX_uint8, LogRootMotion, Warning, TEXT("Too many root motion sources (%d!) to net serialize. Clamping to %d"), RootMotionSources.Num(), MAX_uint8);
-		SourcesNum = FMath::Min<int32>(RootMotionSources.Num(), MAX_uint8);
+		UE_CLOG(RootMotionSourceArray.Num() > MAX_uint8, LogRootMotion, Warning, TEXT("Too many root motion sources (%d!) to net serialize. Clamping to %d"), 
+			RootMotionSourceArray.Num(), MAX_uint8);
+		SourcesNum = FMath::Min<int32>(RootMotionSourceArray.Num(), MAX_uint8);
 	}
 	Ar << SourcesNum;
 	if (Ar.IsLoading())
 	{
-		RootMotionSources.SetNumZeroed(SourcesNum);
+		RootMotionSourceArray.SetNumZeroed(SourcesNum);
 	}
 
-	FArchive_Serialize_BitfieldBool(Ar, bHasAdditiveSources);
-	FArchive_Serialize_BitfieldBool(Ar, bHasOverrideSources);
-	FArchive_Serialize_BitfieldBool(Ar, bHasOverrideSourcesWithIgnoreZAccumulate);
-	LastPreAdditiveVelocity.NetSerialize(Ar, Map, bOutSuccess);
-	FArchive_Serialize_BitfieldBool(Ar, bIsAdditiveVelocityApplied);
-	Ar << LastAccumulatedSettings.Flags;
-
 	for (int32 i = 0; i < SourcesNum && !Ar.IsError(); ++i)
 	{
-		TCheckedObjPtr<UScriptStruct> ScriptStruct = RootMotionSources[i].IsValid() ? RootMotionSources[i]->GetScriptStruct() : nullptr;
+		TCheckedObjPtr<UScriptStruct> ScriptStruct = RootMotionSourceArray[i].IsValid() ? RootMotionSourceArray[i]->GetScriptStruct() : nullptr;
 		UScriptStruct* ScriptStructLocal = ScriptStruct.Get();
 		Ar << ScriptStruct;
 
@@ -1991,7 +2018,7 @@ bool FRootMotionSourceGroup::NetSerialize(FArchive& Ar, class UPackageMap* Map,
 			{
 				if (Ar.IsLoading())
 				{
-					if (RootMotionSources[i].IsValid() && ScriptStructLocal == ScriptStruct.Get())
+					if (RootMotionSourceArray[i].IsValid() && ScriptStructLocal == ScriptStruct.Get())
 					{
 						// What we have locally is the same type as we're being serialized into, so we don't need to
 						// reallocate - just use existing structure
@@ -2004,15 +2031,15 @@ bool FRootMotionSourceGroup::NetSerialize(FArchive& Ar, class UPackageMap* Map,
 						FRootMotionSource* NewSource = (FRootMotionSource*)FMemory::Malloc(ScriptStruct->GetCppStructOps()->GetSize());
 						ScriptStruct->InitializeStruct(NewSource);
 
-						RootMotionSources[i] = TSharedPtr<FRootMotionSource>(NewSource);
+						RootMotionSourceArray[i] = TSharedPtr<FRootMotionSource>(NewSource);
 					}
 				}
 
-				void* ContainerPtr = RootMotionSources[i].Get();
+				void* ContainerPtr = RootMotionSourceArray[i].Get();
 
 				if (ScriptStruct->StructFlags & STRUCT_NetSerializeNative)
 				{
-					ScriptStruct->GetCppStructOps()->NetSerialize(Ar, Map, bOutSuccess, RootMotionSources[i].Get());
+					ScriptStruct->GetCppStructOps()->NetSerialize(Ar, Map, bOutSuccess, RootMotionSourceArray[i].Get());
 				}
 				else
 				{
@@ -2034,6 +2061,20 @@ bool FRootMotionSourceGroup::NetSerialize(FArchive& Ar, class UPackageMap* Map,
 		}
 	}
 
+}
+
+bool FRootMotionSourceGroup::NetSerialize(FArchive& Ar, class UPackageMap* Map, bool& bOutSuccess)
+{
+	FArchive_Serialize_BitfieldBool(Ar, bHasAdditiveSources);
+	FArchive_Serialize_BitfieldBool(Ar, bHasOverrideSources);
+	FArchive_Serialize_BitfieldBool(Ar, bHasOverrideSourcesWithIgnoreZAccumulate);
+	LastPreAdditiveVelocity.NetSerialize(Ar, Map, bOutSuccess);
+	FArchive_Serialize_BitfieldBool(Ar, bIsAdditiveVelocityApplied);
+	Ar << LastAccumulatedSettings.Flags;
+
+	NetSerializeRMSArray(Ar, Map, bOutSuccess, RootMotionSources);
+	NetSerializeRMSArray(Ar, Map, bOutSuccess, PendingAddRootMotionSources);
+
 	if (Ar.IsError())
 	{
 		// Something bad happened, make sure to not return invalid shared ptrs
@@ -2044,6 +2085,13 @@ bool FRootMotionSourceGroup::NetSerialize(FArchive& Ar, class UPackageMap* Map,
 				RootMotionSources.RemoveAt(i);
 			}
 		}
+		for (int32 i = PendingAddRootMotionSources.Num() - 1; i >= 0; --i)
+		{
+			if (PendingAddRootMotionSources[i].IsValid() == false)
+			{
+				PendingAddRootMotionSources.RemoveAt(i);
+			}
+		}
 		bOutSuccess = false;
 		return false;
 	}
@@ -2140,6 +2188,10 @@ bool FRootMotionSourceGroup::operator==(const FRootMotionSourceGroup& Other) con
 	{
 		return false;
 	}
+	if (PendingAddRootMotionSources.Num() != Other.PendingAddRootMotionSources.Num())
+	{
+		return false;
+	}
 	for (int32 i = 0; i < RootMotionSources.Num(); ++i)
 	{
 		if (RootMotionSources[i].IsValid() == Other.RootMotionSources[i].IsValid())
@@ -2157,6 +2209,23 @@ bool FRootMotionSourceGroup::operator==(const FRootMotionSourceGroup& Other) con
 			return false; // Mismatch in validity
 		}
 	}
+	for (int32 i = 0; i < PendingAddRootMotionSources.Num(); ++i)
+	{
+		if (PendingAddRootMotionSources[i].IsValid() == Other.PendingAddRootMotionSources[i].IsValid())
+		{
+			if (PendingAddRootMotionSources[i].IsValid())
+			{
+				if (!PendingAddRootMotionSources[i]->MatchesAndHasSameState(Other.PendingAddRootMotionSources[i].Get()))
+				{
+					return false; // They're valid and don't match/have same state
+				}
+			}
+		}
+		else
+		{
+			return false; // Mismatch in validity
+		}
+	}
 	return true;
 }
 
