Fix false CharacterMovementComponent.cpp static analysis warnings

#rnx

diff --git a/Engine/Source/Runtime/Engine/Private/Components/CharacterMovementComponent.cpp b/Engine/Source/Runtime/Engine/Private/Components/CharacterMovementComponent.cpp
--- a/Engine/Source/Runtime/Engine/Private/Components/CharacterMovementComponent.cpp
+++ b/Engine/Source/Runtime/Engine/Private/Components/CharacterMovementComponent.cpp
@@ -1502,9 +1502,8 @@ void UCharacterMovementComponent::SimulatedTick(float DeltaSeconds)
 
 #if !(UE_BUILD_SHIPPING)
 			// debug
-			if (false)
+			if (CharacterOwner && false)
 			{
-				checkSlow(CharacterOwner != nullptr);
 				const FRotator OldRotation = OldRotationQuat.Rotator();
 				const FRotator NewRotation = UpdatedComponent->GetComponentRotation();
 				const FVector NewLocation = UpdatedComponent->GetComponentLocation();
@@ -9323,9 +9322,8 @@ void UCharacterMovementComponent::ClientAdjustPosition_Implementation
 		
 		// We had an unresolved base from the server
 		// If walking, we'd like to continue walking if possible, to avoid falling for a frame, so try to find a base where we moved to.
-		if (PreviousBase)
+		if (PreviousBase && UpdatedComponent)
 		{
-			checkSlow(UpdatedComponent != nullptr);
 			FindFloor(UpdatedComponent->GetComponentLocation(), CurrentFloor, false);
 			if (CurrentFloor.IsWalkableFloor())
 			{
