Fix cloth authoring when applying a cloth made from a section to a different section (and set to be removed from mesh).

diff --git a/Engine/Source/Editor/ClothingSystemEditor/Private/ClothingAssetFactory.cpp b/Engine/Source/Editor/ClothingSystemEditor/Private/ClothingAssetFactory.cpp
--- a/Engine/Source/Editor/ClothingSystemEditor/Private/ClothingAssetFactory.cpp
+++ b/Engine/Source/Editor/ClothingSystemEditor/Private/ClothingAssetFactory.cpp
@@ -248,6 +248,7 @@ UClothingAssetBase* UClothingAssetFactory::CreateFromSkeletalMesh(USkeletalMesh*
 
 	if(ImportToLodInternal(TargetMesh, Params.LodIndex, Params.SourceSection, NewAsset, LodData))
 	{
+		FScopedSkeletalMeshPostEditChange ScopedSkeletalMeshPostEditChange(TargetMesh);
 		if(Params.bRemoveFromMesh)
 		{
 			// User doesn't want the section anymore as a renderable, get rid of it
diff --git a/Engine/Source/Editor/SkeletalMeshEditor/Private/SkeletalMeshEditor.cpp b/Engine/Source/Editor/SkeletalMeshEditor/Private/SkeletalMeshEditor.cpp
--- a/Engine/Source/Editor/SkeletalMeshEditor/Private/SkeletalMeshEditor.cpp
+++ b/Engine/Source/Editor/SkeletalMeshEditor/Private/SkeletalMeshEditor.cpp
@@ -821,6 +821,9 @@ void FSkeletalMeshEditor::OnCreateClothingAssetMenuItemClicked(FSkeletalMeshClot
 
 	if(Mesh)
 	{
+		FScopedSuspendAlternateSkinWeightPreview ScopedSuspendAlternateSkinnWeightPreview(Mesh);
+		FScopedSkeletalMeshPostEditChange ScopedSkeletalMeshPostEditChange(Mesh);
+
 		// Handle the creation through the clothing asset factory
 		FClothingSystemEditorInterfaceModule& ClothingEditorModule = FModuleManager::LoadModuleChecked<FClothingSystemEditorInterfaceModule>("ClothingSystemEditorInterface");
 		UClothingAssetFactoryBase* AssetFactory = ClothingEditorModule.GetClothingAssetFactory();
diff --git a/Engine/Source/Editor/UnrealEd/Private/Animation/DebugSkelMeshComponent.cpp b/Engine/Source/Editor/UnrealEd/Private/Animation/DebugSkelMeshComponent.cpp
--- a/Engine/Source/Editor/UnrealEd/Private/Animation/DebugSkelMeshComponent.cpp
+++ b/Engine/Source/Editor/UnrealEd/Private/Animation/DebugSkelMeshComponent.cpp
@@ -726,6 +726,7 @@ void UDebugSkelMeshComponent::ResetMeshSectionVisibility()
 void UDebugSkelMeshComponent::RebuildClothingSectionsFixedVerts()
 {
 	FSkeletalMeshModel* Resource = SkeletalMesh->GetImportedModel();
+	FScopedSkeletalMeshPostEditChange ScopedSkeletalMeshPostEditChange(SkeletalMesh);
 
 	const int32 NumLods = Resource->LODModels.Num();
 	for (FSkeletalMeshLODModel& LodModel : Resource->LODModels)
@@ -758,11 +759,11 @@ void UDebugSkelMeshComponent::RebuildClothingSectionsFixedVerts()
 							VertData.SourceMeshVertIndices[3] = 0;
 						}
 					}
+					//We must dirty the DDC key
+					SkeletalMesh->InvalidateDeriveDataCacheGUID();
 				}
 			}
 		}
-
-		SkeletalMesh->PostEditChange();
 	}
 
 	ReregisterComponent();
diff --git a/Engine/Source/Runtime/Engine/Private/SkeletalMesh.cpp b/Engine/Source/Runtime/Engine/Private/SkeletalMesh.cpp
--- a/Engine/Source/Runtime/Engine/Private/SkeletalMesh.cpp
+++ b/Engine/Source/Runtime/Engine/Private/SkeletalMesh.cpp
@@ -3378,22 +3378,28 @@ void USkeletalMesh::RemoveMeshSection(int32 InLodIndex, int32 InSectionIndex)
 		return;
 	}
 
-	FSkelMeshSection& SectionToRemove = LodModel.Sections[InSectionIndex];
+	FSkelMeshSection& SectionToDisable = LodModel.Sections[InSectionIndex];
+	
+	//Get the UserSectionData
+	FSkelMeshSourceSectionUserData& UserSectionToDisableData = LodModel.UserSectionsData.FindChecked(SectionToDisable.OriginalDataSectionIndex);
 
-	if(SectionToRemove.HasClothingData())
+	if(UserSectionToDisableData.HasClothingData())
 	{
 		// Can't remove this, clothing currently relies on it
 		UE_LOG(LogSkeletalMesh, Warning, TEXT("Failed to remove skeletal mesh section, clothing is currently bound to Lod%d Section %d, unbind clothing before removal."), InLodIndex, InSectionIndex);
 		return;
 	}
 
-	// Valid to remove, dirty the mesh
-	Modify();
-	PreEditChange(nullptr);
-
-	SectionToRemove.bDisabled = true;
-
-	PostEditChange();
+	{
+		//Scope a post edit change
+		FScopedSkeletalMeshPostEditChange ScopedPostEditChange(this);
+		// Valid to disable, dirty the mesh
+		Modify();
+		PreEditChange(nullptr);
+		//Disable the section
+		UserSectionToDisableData.bDisabled = true;
+		SectionToDisable.bDisabled = true;
+	}
 }
 
 #endif // #if WITH_EDITOR
