Removing the warning around local XGE rescheduling for shader compilation.

#rnx

diff --git a/Engine/Source/Runtime/Engine/Private/ShaderCompiler/ShaderCompilerXGE.cpp b/Engine/Source/Runtime/Engine/Private/ShaderCompiler/ShaderCompilerXGE.cpp
--- a/Engine/Source/Runtime/Engine/Private/ShaderCompiler/ShaderCompilerXGE.cpp
+++ b/Engine/Source/Runtime/Engine/Private/ShaderCompiler/ShaderCompilerXGE.cpp
@@ -858,7 +858,7 @@ int32 FShaderCompileXGEThreadRunnable_InterceptionInterface::CompilingLoop()
 
 	for (auto Iter = DispatchedTasks.CreateIterator(); Iter; ++Iter)
 	{
-		bool bOutputFileReadFailed = false;
+		bool bOutputFileReadFailed = true;
 
 		FXGEShaderCompilerTask* Task = *Iter;
 		if (!Task->Future.IsReady())
@@ -877,31 +877,21 @@ int32 FShaderCompileXGEThreadRunnable_InterceptionInterface::CompilingLoop()
 		if (Result.bCompleted)
 		{
 			// Check the output file exists. If it does, attempt to open it and serialize in the completed jobs.
-			FArchive* OutputFileAr = nullptr;
 			if (IFileManager::Get().FileExists(*Task->OutputFilePath))
 			{
-				float TimeWaited = 0.0f;
-				OutputFileAr = IFileManager::Get().CreateFileReader(*Task->OutputFilePath, FILEREAD_Silent);
-				while (OutputFileAr == nullptr && TimeWaited < 5.0f)
+				FArchive* OutputFileAr = IFileManager::Get().CreateFileReader(*Task->OutputFilePath, FILEREAD_Silent);
+				if (OutputFileAr)
 				{
-					UE_LOG(LogShaderCompilers, Warning, TEXT("Expected XGE output file '%s' exists but can't be opened for read, waiting 1 second to try again.."), *Task->OutputFilePath);
-					FPlatformProcess::Sleep(1.0f);
-					TimeWaited += 1.0f;
-					OutputFileAr = IFileManager::Get().CreateFileReader(*Task->OutputFilePath, FILEREAD_Silent);
+					bOutputFileReadFailed = false;
+					FShaderCompileUtilities::DoReadTaskResults(Task->ShaderJobs, *OutputFileAr);
+					delete OutputFileAr;
 				}
 			}
 
-			if (OutputFileAr)
-			{
-				FShaderCompileUtilities::DoReadTaskResults(Task->ShaderJobs, *OutputFileAr);
-				delete OutputFileAr;
-			}
-			else
+			if (bOutputFileReadFailed)
 			{
 				// Reading result from XGE job failed, so recompile shaders in current job batch locally
-				bOutputFileReadFailed = true;
-
-				UE_LOG(LogShaderCompilers, Warning, TEXT("Rescheduling shader compilation to run locally after XGE job failed: %s"), *Task->OutputFilePath);
+				UE_LOG(LogShaderCompilers, Log, TEXT("Rescheduling shader compilation to run locally after XGE job failed: %s"), *Task->OutputFilePath);
 
 				for (TSharedRef<FShaderCommonCompileJob, ESPMode::ThreadSafe> Job : Task->ShaderJobs)
 				{
