Fix missing include in NetworkProfiler

trivial

diff --git a/Engine/Source/Runtime/Engine/Private/NetworkProfiler.cpp b/Engine/Source/Runtime/Engine/Private/NetworkProfiler.cpp
--- a/Engine/Source/Runtime/Engine/Private/NetworkProfiler.cpp
+++ b/Engine/Source/Runtime/Engine/Private/NetworkProfiler.cpp
@@ -13,6 +13,7 @@
 #include "Engine/World.h"
 #include "Serialization/MemoryWriter.h"
 #include "HAL/IConsoleManager.h"
+#include "Engine/Public/TimerManager.h"
 
 #if USE_NETWORK_PROFILER
 
