Adding more scopes to track down player controller tick performance.
#RNX

diff --git a/Engine/Source/Runtime/Engine/Private/PlayerController.cpp b/Engine/Source/Runtime/Engine/Private/PlayerController.cpp
--- a/Engine/Source/Runtime/Engine/Private/PlayerController.cpp
+++ b/Engine/Source/Runtime/Engine/Private/PlayerController.cpp
@@ -4715,11 +4715,13 @@ void APlayerController::TickActor( float DeltaSeconds, ELevelTick TickType, FAct
 		// Process PlayerTick with input.
 		if (!PlayerInput && (Player == nullptr || Cast<ULocalPlayer>( Player ) != nullptr))
 		{
+			QUICK_SCOPE_CYCLE_COUNTER(HVS_InitInputSystem); // HVS_BCF - Temp.  Just trying to track down where time is being spent.
 			InitInputSystem();
 		}
 
 		if (PlayerInput)
 		{
+			QUICK_SCOPE_CYCLE_COUNTER(HVS_PlayerTick); // HVS_BCF - Temp.  Just trying to track down where time is being spent.
 			PlayerTick(DeltaSeconds);
 		}
 
@@ -4731,6 +4733,7 @@ void APlayerController::TickActor( float DeltaSeconds, ELevelTick TickType, FAct
 		// update viewtarget replicated info
 		if (PlayerCameraManager != nullptr)
 		{
+			QUICK_SCOPE_CYCLE_COUNTER(HVS_UpdateViewtargetReplicatedInfo); // HVS_BCF - Temp.  Just trying to track down where time is being spent.
 			APawn* TargetPawn = PlayerCameraManager->GetViewTargetPawn();
 			if ((TargetPawn != GetPawn()) && (TargetPawn != nullptr))
 			{
@@ -4756,6 +4759,7 @@ void APlayerController::TickActor( float DeltaSeconds, ELevelTick TickType, FAct
 
 	if (!IsPendingKill())
 	{
+		QUICK_SCOPE_CYCLE_COUNTER(HVS_Tick); // HVS_BCF - Temp.  Just trying to track down where time is being spent.
 		Tick(DeltaSeconds);	// perform any tick functions unique to an actor subclass
 	}
 
