Fix variables on spawned Blueprint instances not being initialized to modified CDO values for native struct types with a details customization and no UPROPERTY fields.

diff --git a/Engine/Source/Runtime/Engine/Private/BlueprintGeneratedClass.cpp b/Engine/Source/Runtime/Engine/Private/BlueprintGeneratedClass.cpp
--- a/Engine/Source/Runtime/Engine/Private/BlueprintGeneratedClass.cpp
+++ b/Engine/Source/Runtime/Engine/Private/BlueprintGeneratedClass.cpp
@@ -437,8 +437,35 @@ bool UBlueprintGeneratedClass::BuildCustomPropertyListForPostConstruction(FCusto
 					*CurrentNodePtr = new FCustomPropertyListNode(Property, Idx);
 					CustomPropertyListForPostConstruction.Add(*CurrentNodePtr);
 
-					// Recursively gather up all struct fields that differ and assign to the current node's sub property list.
-					if (BuildCustomPropertyListForPostConstruction((*CurrentNodePtr)->SubPropertyList, StructProperty->Struct, PropertyValue, DefaultPropertyValue))
+					UScriptStruct::ICppStructOps* CppStructOps = nullptr;
+					if (StructProperty->Struct)
+					{
+						CppStructOps = StructProperty->Struct->GetCppStructOps();
+					}
+
+					bool bIsIdentical = false;
+					const bool bHasIdentical = CppStructOps && CppStructOps->HasIdentical();
+					const bool bIsPlainOldData = CppStructOps && CppStructOps->IsPlainOldData();
+					
+					if (bHasIdentical)
+					{
+						// Use the explicit comparison function that's associated with this struct.
+						const uint32 PortFlags = 0;
+						bIsIdentical = StructProperty->Identical(PropertyValue, DefaultPropertyValue, PortFlags);
+					}
+					else if (bIsPlainOldData)
+					{
+						// Use a byte-for-byte comparison for simple POD types. We'll initialize the whole value in this case.
+						const int32 NumBytes = StructProperty->Struct->GetStructureSize();
+						bIsIdentical = !FMemory::Memcmp(PropertyValue, DefaultPropertyValue, NumBytes);
+					}
+					else
+					{
+						// Recursively gather up all struct fields that differ and assign to the current node's sub property list.
+						bIsIdentical = !BuildCustomPropertyListForPostConstruction((*CurrentNodePtr)->SubPropertyList, StructProperty->Struct, PropertyValue, DefaultPropertyValue);
+					}
+
+					if (!bIsIdentical)
 					{
 						// Advance to the next node in the list.
 						CurrentNodePtr = &(*CurrentNodePtr)->PropertyListNext;
@@ -643,10 +670,28 @@ void UBlueprintGeneratedClass::InitPropertiesFromCustomList(const FCustomPropert
 
 		if (const FStructProperty* StructProperty = CastField<FStructProperty>(CustomPropertyListNode->Property))
 		{
-			// This should never be NULL; we should not be recording the StructProperty without at least one sub property, but we'll verify just to be sure.
-			if (ensure(CustomPropertyListNode->SubPropertyList != nullptr))
+			UScriptStruct::ICppStructOps* CppStructOps = nullptr;
+			if (StructProperty->Struct)
+			{
+				CppStructOps = StructProperty->Struct->GetCppStructOps();
+			}
+
+			const bool bHasCopy = CppStructOps && CppStructOps->HasCopy();
+			const bool bHasIdentical = CppStructOps && CppStructOps->HasIdentical();
+			const bool bIsPlainOldData = CppStructOps && CppStructOps->IsPlainOldData();
+
+			// Copy the full value if this struct has an explicit copy functor, or if we did not determine the delta for this struct ourselves.
+			if (bHasCopy || bHasIdentical || bIsPlainOldData)
 			{
-				InitPropertiesFromCustomList(CustomPropertyListNode->SubPropertyList, StructProperty->Struct, PropertyValue, DefaultPropertyValue);
+				StructProperty->CopySingleValue(PropertyValue, DefaultPropertyValue);
+			}
+			else
+			{
+				// This should never be NULL; we should not be recording the StructProperty without at least one sub property, but we'll verify just to be sure.
+				if (ensure(CustomPropertyListNode->SubPropertyList != nullptr))
+				{
+					InitPropertiesFromCustomList(CustomPropertyListNode->SubPropertyList, StructProperty->Struct, PropertyValue, DefaultPropertyValue);
+				}
 			}
 		}
 		else if (const FArrayProperty* ArrayProperty = CastField<FArrayProperty>(CustomPropertyListNode->Property))
