TActorIterator/GetAllActorsOfClass will no longer return actors from a level that is in the process of being incrementally removed from world, unless you are in the RemoveFromWorld scope for that level. This is equivalent to the behavior for a level being incrementally added to the world

diff --git a/Engine/Source/Runtime/Engine/Classes/Engine/Level.h b/Engine/Source/Runtime/Engine/Classes/Engine/Level.h
--- a/Engine/Source/Runtime/Engine/Classes/Engine/Level.h
+++ b/Engine/Source/Runtime/Engine/Classes/Engine/Level.h
@@ -589,8 +589,10 @@ public:
 	uint8										bAlreadyRoutedActorInitialize:1;
 	/** Whether we already sorted the actor list.											*/
 	uint8										bAlreadySortedActorList:1;
-	/** Whether this level is in the process of being associated with its world				*/
+	/** Whether this level is in the process of being associated with its world	(i.e. we are within AddToWorld for this level */
 	uint8										bIsAssociatingLevel:1;
+	/** Whether this level is in the process of being disassociated with its world (i.e. we are within RemoveFromWorld for this level */
+	uint8										bIsDisassociatingLevel : 1;
 	/** Whether this level should be fully added to the world before rendering his components	*/
 	uint8										bRequireFullVisibilityToRender:1;
 	/** Whether this level is specific to client, visibility state will not be replicated to server	*/
diff --git a/Engine/Source/Runtime/Engine/Private/World.cpp b/Engine/Source/Runtime/Engine/Private/World.cpp
--- a/Engine/Source/Runtime/Engine/Private/World.cpp
+++ b/Engine/Source/Runtime/Engine/Private/World.cpp
@@ -2568,6 +2568,8 @@ void UWorld::RemoveFromWorld( ULevel* Level, bool bAllowIncrementalRemoval )
 	check(!Level->IsPendingKill());
 	check(!Level->IsUnreachable());
 
+	Level->bIsDisassociatingLevel = true;
+
 	// To be removed from the world a world must be visible and not pending being made visible (this may be redundent, but for safety)
 	// If the level may be removed incrementally then there must also be no level pending visibility
 	if ( ((CurrentLevelPendingVisibility == nullptr) || (!bAllowIncrementalRemoval && (CurrentLevelPendingVisibility != Level))) && Level->bIsVisible )
@@ -2687,6 +2689,8 @@ void UWorld::RemoveFromWorld( ULevel* Level, bool bAllowIncrementalRemoval )
 			Level->bIsBeingRemoved = false;
 		} // if ( bFinishRemovingLevel )
 
+		Level->bIsDisassociatingLevel = false;
+
 #if PERF_TRACK_DETAILED_ASYNC_STATS
 		UE_LOG(LogStreaming, Display, TEXT("UWorld::RemoveFromWorld for %s took %5.2f ms"), *Level->GetOutermost()->GetName(), (FPlatformTime::Seconds() - StartTime) * 1000.0);
 #endif // PERF_TRACK_DETAILED_ASYNC_STATS
diff --git a/Engine/Source/Runtime/Engine/Public/EngineUtils.h b/Engine/Source/Runtime/Engine/Public/EngineUtils.h
--- a/Engine/Source/Runtime/Engine/Public/EngineUtils.h
+++ b/Engine/Source/Runtime/Engine/Public/EngineUtils.h
@@ -422,7 +422,7 @@ protected:
 	{
 		if (EnumHasAnyFlags(Flags, EActorIteratorFlags::OnlyActiveLevels))
 		{
-			const bool bIsLevelVisibleOrAssociating = Level->bIsVisible || Level->bIsAssociatingLevel;
+			const bool bIsLevelVisibleOrAssociating = (Level->bIsVisible && !Level->bIsBeingRemoved) || Level->bIsAssociatingLevel || Level->bIsDisassociatingLevel;
 
 			// Only allow iteration of Level if it's in the currently active level collection of the world, or is a static level.
 			const FLevelCollection* const ActorLevelCollection = Level->GetCachedLevelCollection();
