Fix nullptr static analysis warning

diff --git a/Engine/Source/Runtime/Engine/Private/NetworkProfiler.cpp b/Engine/Source/Runtime/Engine/Private/NetworkProfiler.cpp
--- a/Engine/Source/Runtime/Engine/Private/NetworkProfiler.cpp
+++ b/Engine/Source/Runtime/Engine/Private/NetworkProfiler.cpp
@@ -849,7 +849,10 @@ bool FNetworkProfiler::Exec( UWorld * InWorld, const TCHAR* Cmd, FOutputDevice &
 		float AutoStopTime = 30.f;
 		FParse::Value(Cmd, TEXT("TIME="), AutoStopTime);
 
-		InWorld->GetTimerManager().SetTimer(AutoStopTimerHandle, FTimerDelegate::CreateRaw(this, &FNetworkProfiler::AutoStopTracking), AutoStopTime, false);
+		if( InWorld )
+		{
+			InWorld->GetTimerManager().SetTimer(AutoStopTimerHandle, FTimerDelegate::CreateRaw(this, &FNetworkProfiler::AutoStopTracking), AutoStopTime, false);
+		}
 	}
 	else 
 	{
