Fix compil error with non-exported function

diff --git a/Engine/Source/Runtime/Engine/Public/Net/NetworkProfiler.h b/Engine/Source/Runtime/Engine/Public/Net/NetworkProfiler.h
--- a/Engine/Source/Runtime/Engine/Public/Net/NetworkProfiler.h
+++ b/Engine/Source/Runtime/Engine/Public/Net/NetworkProfiler.h
@@ -468,7 +468,7 @@ public:
 	 *
 	 * @return			True if processed, false otherwise
 	 */
-	bool Exec( UWorld * InWorld, const TCHAR* Cmd, FOutputDevice & Ar );
+	ENGINE_API bool Exec( UWorld * InWorld, const TCHAR* Cmd, FOutputDevice & Ar );
 
 	bool FORCEINLINE IsTrackingEnabled() const { return bIsTrackingEnabled; }
 	bool IsComparisonTrackingEnabled() const { return bIsTrackingEnabled && bIsComparisonTrackingEnabled;  }
