Sequencer:  Skel Mesh Animation plays in sequencer with physics bodies set up.  Physics tick group doesn't run in editor so added check.

diff --git a/Engine/Source/Runtime/Engine/Private/Components/SkeletalMeshComponent.cpp b/Engine/Source/Runtime/Engine/Private/Components/SkeletalMeshComponent.cpp
--- a/Engine/Source/Runtime/Engine/Private/Components/SkeletalMeshComponent.cpp
+++ b/Engine/Source/Runtime/Engine/Private/Components/SkeletalMeshComponent.cpp
@@ -2453,12 +2453,20 @@ void USkeletalMeshComponent::PostAnimEvaluation(FAnimationEvaluationContext& Eva
 			}
 		}
 
-		// If we have no physics to blend, we are done
+#if WITH_EDITOR	
+		// If we have no physics to blend or in editor since there is no physics tick group, we are done
+		if (!ShouldBlendPhysicsBones() || GetWorld()->WorldType == EWorldType::Editor)
+		{
+			// Flip buffers, update bounds, attachments etc.
+			FinalizeAnimationUpdate();
+		}
+#else
 		if (!ShouldBlendPhysicsBones())
 		{
 			// Flip buffers, update bounds, attachments etc.
 			FinalizeAnimationUpdate();
 		}
+#endif
 	}
 	else 
 	{
