Added IsValid guard to prevent nullptr access violation in FRootMotionSourceGroup::AddStructReferencedObjects while using an outdated replay

#rnx

diff --git a/Engine/Source/Runtime/Engine/Private/GameFramework/RootMotionSource.cpp b/Engine/Source/Runtime/Engine/Private/GameFramework/RootMotionSource.cpp
--- a/Engine/Source/Runtime/Engine/Private/GameFramework/RootMotionSource.cpp
+++ b/Engine/Source/Runtime/Engine/Private/GameFramework/RootMotionSource.cpp
@@ -2238,11 +2238,17 @@ void FRootMotionSourceGroup::AddStructReferencedObjects(class FReferenceCollecto
 {
 	for (const auto& RootMotionSource : RootMotionSources)
 	{
-		RootMotionSource->AddReferencedObjects(Collector);
+		if (RootMotionSource.IsValid())
+		{
+			RootMotionSource->AddReferencedObjects(Collector);
+		}
 	}
 
 	for (const auto& RootMotionSource : PendingAddRootMotionSources)
 	{
-		RootMotionSource->AddReferencedObjects(Collector);
+		if (RootMotionSource.IsValid())
+		{
+			RootMotionSource->AddReferencedObjects(Collector);
+		}
 	}
 }
