Code :
Glider Preset:
Extending the vfx marshaller to set int parameters.
This is the same deal as the boolean thing I did a while back.
#rnx
#texted PIE, compile

diff --git a/Engine/Plugins/FX/Niagara/Source/Niagara/Private/NiagaraComponent.cpp b/Engine/Plugins/FX/Niagara/Source/Niagara/Private/NiagaraComponent.cpp
--- a/Engine/Plugins/FX/Niagara/Source/Niagara/Private/NiagaraComponent.cpp
+++ b/Engine/Plugins/FX/Niagara/Source/Niagara/Private/NiagaraComponent.cpp
@@ -474,6 +474,11 @@ void UNiagaraComponent::SetBoolParameter(FName Parametername, bool Param)
 	SetVariableBool(Parametername, Param);
 }
 
+void UNiagaraComponent::SetIntParameter(FName ParameterName, int Param)
+{
+	SetVariableInt(ParameterName, Param);
+}
+
 void UNiagaraComponent::SetFloatParameter(FName ParameterName, float Param)
 {
 	SetVariableFloat(ParameterName, Param);
diff --git a/Engine/Plugins/FX/Niagara/Source/Niagara/Public/NiagaraComponent.h b/Engine/Plugins/FX/Niagara/Source/Niagara/Public/NiagaraComponent.h
--- a/Engine/Plugins/FX/Niagara/Source/Niagara/Public/NiagaraComponent.h
+++ b/Engine/Plugins/FX/Niagara/Source/Niagara/Public/NiagaraComponent.h
@@ -48,6 +48,7 @@ public:
 
 	/********* UFXSystemComponent *********/
 	void SetBoolParameter(FName ParameterName, bool Param) override;
+	void SetIntParameter(FName ParameterName, int Param) override;
 	void SetFloatParameter(FName ParameterName, float Param) override;
 	void SetVectorParameter(FName ParameterName, FVector Param) override;
 	void SetColorParameter(FName ParameterName, FLinearColor Param) override;
diff --git a/Engine/Source/Runtime/Engine/Classes/Particles/ParticleSystemComponent.h b/Engine/Source/Runtime/Engine/Classes/Particles/ParticleSystemComponent.h
--- a/Engine/Source/Runtime/Engine/Classes/Particles/ParticleSystemComponent.h
+++ b/Engine/Source/Runtime/Engine/Classes/Particles/ParticleSystemComponent.h
@@ -337,12 +337,16 @@ public:
 	UFUNCTION(BlueprintCallable, Category = "Effects|Components|ParticleSystem")
 	virtual void SetBoolParameter(FName ParameterName, bool Param) {}
 
-	/** Change a named float parameter */
+	/** Change a named int parameter */
 	UFUNCTION(BlueprintCallable, Category="Effects|Components|ParticleSystem")
+	virtual void SetIntParameter(FName ParameterName, int Param) {}
+
+	/** Change a named float parameter */
+	UFUNCTION(BlueprintCallable, Category = "Effects|Components|ParticleSystem")
 	virtual void SetFloatParameter(FName ParameterName, float Param) {}
 
 	/** 
-	 *	Set a named vector instance parameter on this ParticleSystemComponent. 
+	 *	Set a named vector instance parameter on this ParticleSystemComponent.
 	 *	Updates the parameter if it already exists, or creates a new entry if not. 
 	 */
 	UFUNCTION(BlueprintCallable, Category="Effects|Components|ParticleSystem")
@@ -1063,8 +1067,21 @@ public:
 	 */
 	void SetEmitterEnable(FName EmitterName, bool bNewEnableState) override;
 
+
+	/**
+	 *	Set a named float instance parameter on this ParticleSystemComponent.
+	 *	This is for function parity with the VFX Marshaller. The bool is converted
+	 *  and then set to a float.
+	 */
 	void SetBoolParameter(FName ParameterName, bool Param) override;
 
+	/**
+	 *	Set a named float instance parameter on this ParticleSystemComponent.
+	 *	This is for function parity with the VFX Marshaller. The int is converted
+	 *  and then set to a float.
+	 */
+	void SetIntParameter(FName ParameterName, int Param) override;
+
 	/** Change a named float parameter */
 	void SetFloatParameter(FName ParameterName, float Param) override;
 
diff --git a/Engine/Source/Runtime/Engine/Private/Particles/ParticleComponents.cpp b/Engine/Source/Runtime/Engine/Private/Particles/ParticleComponents.cpp
--- a/Engine/Source/Runtime/Engine/Private/Particles/ParticleComponents.cpp
+++ b/Engine/Source/Runtime/Engine/Private/Particles/ParticleComponents.cpp
@@ -7084,6 +7084,17 @@ void UParticleSystemComponent::SetBoolParameter(FName Name, bool Value)
 	SetFloatParameter(Name, (Value ? 1.0f : 0.0f));
 
 }
+/**
+ *	Set a named float instance parameter on this ParticleSystemComponent.
+ *	Updates the parameter if it already exists, or creates a new entry if not.
+	This maps a int to a float for parity as cascade doesn't have booleans.
+	This for adding functionality to the parent UFXSystemComponent to set int
+	variables.
+ */
+void UParticleSystemComponent::SetIntParameter(FName Name, int Value)
+{
+	SetFloatParameter(Name, float(Value));
+}
 
 /** 
  *	Set a named float instance parameter on this ParticleSystemComponent. 
