Fix inertia tensors not being calculated correctly and add mass space rotation when calculating inertia/inv inertia for applying torques.

diff --git a/Engine/Source/Runtime/Engine/Private/PhysicsEngine/BodyInstance.cpp b/Engine/Source/Runtime/Engine/Private/PhysicsEngine/BodyInstance.cpp
--- a/Engine/Source/Runtime/Engine/Private/PhysicsEngine/BodyInstance.cpp
+++ b/Engine/Source/Runtime/Engine/Private/PhysicsEngine/BodyInstance.cpp
@@ -36,6 +36,8 @@
 #include "Chaos/ImplicitObjectTransformed.h"
 #include "Chaos/ParticleHandle.h"
 #include "Chaos/TriangleMeshImplicitObject.h"
+#include "Chaos/MassProperties.h"
+#include "Chaos/Utilities.h"
 
 #if WITH_CHAOS
 #include "Chaos/ParticleHandle.h"
@@ -3069,8 +3071,15 @@ void FBodyInstance::UpdateMassProperties()
 
 				// #PHYS2 Refactor out PxMassProperties (Our own impl?)
 				PxQuat MassOrientation;
+#if WITH_CHAOS
+				const PxMat33& IT = TotalMassProperties.inertiaTensor;
+				Chaos::PMatrix<float, 3, 3> InertiaTensor(IT[0][0], IT[0][1], IT[0][2], IT[1][1], IT[1][2], IT[2][2]);
+				const Chaos::TRotation<float, 3> Rotation = Chaos::TransformToLocalSpace(InertiaTensor);
+				const FVector MassSpaceInertiaTensor(InertiaTensor.M[0][0], InertiaTensor.M[1][1], InertiaTensor.M[2][2]);
+				MassOrientation = U2PQuat(Rotation);
+#else
 				const FVector MassSpaceInertiaTensor = P2UVector(PxMassProperties::getMassSpaceInertia(TotalMassProperties.inertiaTensor, MassOrientation));
-
+#endif
 				FPhysicsInterface::SetMass_AssumesLocked(Actor, TotalMassProperties.mass);
 				FPhysicsInterface::SetMassSpaceInertiaTensor_AssumesLocked(Actor, MassSpaceInertiaTensor);
 
diff --git a/Engine/Source/Runtime/Engine/Private/PhysicsEngine/Experimental/PhysInterface_Chaos.cpp b/Engine/Source/Runtime/Engine/Private/PhysicsEngine/Experimental/PhysInterface_Chaos.cpp
--- a/Engine/Source/Runtime/Engine/Private/PhysicsEngine/Experimental/PhysInterface_Chaos.cpp
+++ b/Engine/Source/Runtime/Engine/Private/PhysicsEngine/Experimental/PhysInterface_Chaos.cpp
@@ -1776,33 +1776,17 @@ bool CalculateMassPropertiesOfImplicitType(
 				return false;
 			}
 		else*/ 
-		if (ImplicitObject->GetType(true) & Chaos::ImplicitObjectType::Transformed)
-		{
-			// TODO:  This all very wrong, but is wrong in the same way as scaled. Rotation/Translation are ignored though. The three methods on Transformed no implemented.
-			// Only adding this to hack around CastHelper, as TransformedImplicit is very not supported in that path.
-			const Chaos::TImplicitObjectTransformed<FReal, 3>& Object = ImplicitObject->template GetObjectChecked<Chaos::TImplicitObjectTransformed<FReal, 3>>();
-
-
-			OutMassProperties.Volume = Object.GetVolume();
-			OutMassProperties.Mass = OutMassProperties.Volume * InDensityKGPerCM;
-			OutMassProperties.InertiaTensor = Object.GetInertiaTensor(OutMassProperties.Mass);
-			OutMassProperties.CenterOfMass = Object.GetCenterOfMass();
-			OutMassProperties.RotationOfMass = Chaos::TRotation<float, 3>::FromIdentity();
-			return true;
-		}
-		else
-		{
-			//question: is LocalTM enough to merge these two branches?
-			Chaos::Utilities::CastHelper(*ImplicitObject, FTransform::Identity, [&OutMassProperties, InDensityKGPerCM](const auto& Object, const auto& LocalTM)
+		
+		//todo: Still need to handle scaled
+		Chaos::Utilities::CastHelper(*ImplicitObject, FTransform::Identity, [&OutMassProperties, InDensityKGPerCM](const auto& Object, const auto& LocalTM)
 			{
 				OutMassProperties.Volume = Object.GetVolume();
 				OutMassProperties.Mass = OutMassProperties.Volume * InDensityKGPerCM;
 				OutMassProperties.InertiaTensor = Object.GetInertiaTensor(OutMassProperties.Mass);
-				OutMassProperties.CenterOfMass = Object.GetCenterOfMass();
-				OutMassProperties.RotationOfMass = Chaos::TRotation<float, 3>::FromIdentity();
+				OutMassProperties.CenterOfMass = LocalTM.TransformPosition(Object.GetCenterOfMass());
+				OutMassProperties.RotationOfMass = LocalTM.GetRotation();
 			});
-			return true;
-		}
+		return true;
 	}
 	return false;
 }
@@ -1838,7 +1822,7 @@ void FPhysInterface_Chaos::CalculateMassPropertiesFromShapeCollection(physx::PxM
 	Chaos::PMatrix<float, 3, 3> Tensor;
 	if (MassPropertiesList.Num())
 	{
-		Tensor = Chaos::Combine<float, 3>(MassPropertiesList).InertiaTensor;
+		Tensor = Chaos::CombineWorldSpace<float, 3>(MassPropertiesList,  InDensityKGPerCM).InertiaTensor;
 	}
 	else 
 	{
diff --git a/Engine/Source/Runtime/Engine/Private/PhysicsEngine/Experimental/PhysScene_Chaos.cpp b/Engine/Source/Runtime/Engine/Private/PhysicsEngine/Experimental/PhysScene_Chaos.cpp
--- a/Engine/Source/Runtime/Engine/Private/PhysicsEngine/Experimental/PhysScene_Chaos.cpp
+++ b/Engine/Source/Runtime/Engine/Private/PhysicsEngine/Experimental/PhysScene_Chaos.cpp
@@ -1527,7 +1527,7 @@ void FPhysScene_ChaosInterface::AddTorque_AssumesLocked(FBodyInstance* BodyInsta
 				const Chaos::TVector<float, 3> CurrentTorque = Rigid->Torque();
 				if (bAccelChange)
 				{
-					Rigid->SetTorque(CurrentTorque + (Utilities::ComputeWorldSpaceInertia(Rigid->R(), Rigid->I()) * Torque));
+					Rigid->SetTorque(CurrentTorque + (Utilities::ComputeWorldSpaceInertia(Rigid->R() * Rigid->RotationOfMass(), Rigid->I()) * Torque));
 				}
 				else
 				{
diff --git a/Engine/Source/Runtime/Experimental/Chaos/Private/Chaos/MassProperties.cpp b/Engine/Source/Runtime/Experimental/Chaos/Private/Chaos/MassProperties.cpp
--- a/Engine/Source/Runtime/Experimental/Chaos/Private/Chaos/MassProperties.cpp
+++ b/Engine/Source/Runtime/Experimental/Chaos/Private/Chaos/MassProperties.cpp
@@ -169,6 +169,14 @@ namespace Chaos
 
 	template<class T, int d>
 	TMassProperties<T, d> Combine(const TArray<TMassProperties<T, d>>& MPArray)
+	{
+		TMassProperties<T, d> NewMP = CombineWorldSpace(MPArray, 1.f);
+		NewMP.RotationOfMass = TransformToLocalSpace<T, d>(NewMP.InertiaTensor);
+		return NewMP;
+	}
+
+	template<class T, int d>
+	TMassProperties<T, d> CombineWorldSpace(const TArray<TMassProperties<T, d>>& MPArray, float InDensityKGPerCM)
 	{
 		check(MPArray.Num() > 0);
 		if (MPArray.Num() == 1)
@@ -186,7 +194,7 @@ namespace Chaos
 		NewMP.CenterOfMass /= NewMP.Volume;
 		for (const TMassProperties<T, d>& Child : MPArray)
 		{
-			const T M = Child.Volume;
+			const T M = Child.Volume * InDensityKGPerCM;
 			const TVector<T, d> ParentToChild = Child.CenterOfMass - NewMP.CenterOfMass;
 			const T P0 = ParentToChild[0];
 			const T P1 = ParentToChild[1];
@@ -194,9 +202,8 @@ namespace Chaos
 			const T MP0P0 = M * P0 * P0;
 			const T MP1P1 = M * P1 * P1;
 			const T MP2P2 = M * P2 * P2;
-			NewMP.InertiaTensor += PMatrix<T, d, d>(MP1P1+MP2P2, -M*P1*P0, -M*P2*P0, MP2P2+MP0P0, -M*P2*P1, MP1P1+MP0P0);
+			NewMP.InertiaTensor += PMatrix<T, d, d>(MP1P1 + MP2P2, -M * P1 * P0, -M * P2 * P0, MP2P2 + MP0P0, -M * P2 * P1, MP1P1 + MP0P0);
 		}
-		NewMP.RotationOfMass = TransformToLocalSpace<T, d>(NewMP.InertiaTensor);
 		return NewMP;
 	}
 
@@ -219,4 +226,8 @@ namespace Chaos
 		const TVector<float, 3>& CenterOfMass, PMatrix<float, 3, 3>& OutInertiaTensor, TRotation<float, 3>& OutRotationOfMass);
 
 	template CHAOS_API TMassProperties<float, 3> Combine(const TArray<TMassProperties<float, 3>>& MPArray);
+
+	template CHAOS_API TMassProperties<float, 3> CombineWorldSpace(const TArray<TMassProperties<float, 3>>& MPArray, float InDensityKGPerCM);
+
+	template CHAOS_API TRotation<float, 3> TransformToLocalSpace(PMatrix<float, 3, 3>& Inertia);
 }
diff --git a/Engine/Source/Runtime/Experimental/Chaos/Public/Chaos/MassProperties.h b/Engine/Source/Runtime/Experimental/Chaos/Public/Chaos/MassProperties.h
--- a/Engine/Source/Runtime/Experimental/Chaos/Public/Chaos/MassProperties.h
+++ b/Engine/Source/Runtime/Experimental/Chaos/Public/Chaos/MassProperties.h
@@ -52,5 +52,8 @@ namespace Chaos
 	template<class T, int d>
 	TMassProperties<T, d> Combine(const TArray<TMassProperties<T, d>>& MPArray);
 
+	template<class T, int d>
+	TMassProperties<T, d> CombineWorldSpace(const TArray<TMassProperties<T, d>>& MPArray, float InDensityKGPerCM);
+
 }
 #endif
diff --git a/Engine/Source/Runtime/Experimental/Chaos/Public/Chaos/PerParticleEulerStepVelocity.h b/Engine/Source/Runtime/Experimental/Chaos/Public/Chaos/PerParticleEulerStepVelocity.h
--- a/Engine/Source/Runtime/Experimental/Chaos/Public/Chaos/PerParticleEulerStepVelocity.h
+++ b/Engine/Source/Runtime/Experimental/Chaos/Public/Chaos/PerParticleEulerStepVelocity.h
@@ -39,14 +39,14 @@ class TPerParticleEulerStepVelocity : public TPerParticleRule<T, d>
 		//       Just using W += InvI * (Torque - W x (I * W)) * dt is not correct, since Torque
 		//		 and W are in an inertial frame.
 		//
-		const FMatrix33 WorldInvI = Utilities::ComputeWorldSpaceInertia(InParticles.R(Index), InParticles.InvI(Index));
+		const FMatrix33 WorldInvI = Utilities::ComputeWorldSpaceInertia(InParticles.R(Index) * InParticles.RotationOfMass(Index), InParticles.InvI(Index));
 		InParticles.W(Index) += WorldInvI * InParticles.Torque(Index) * Dt;
 	}
 	
 	inline void Apply(TTransientPBDRigidParticleHandle<T, d>& Particle, const T Dt) const override //-V762
 	{
 		Particle.V() += Particle.F() * Particle.InvM() * Dt;
-		const FMatrix33 WorldInvI = Utilities::ComputeWorldSpaceInertia(Particle.R(), Particle.InvI());
+		const FMatrix33 WorldInvI = Utilities::ComputeWorldSpaceInertia(Particle.R() * Particle.RotationOfMass(), Particle.InvI());
 		Particle.W() += WorldInvI * Particle.Torque() * Dt;
 	}
 };
