Chaos: Fix Cook differences issue

diff --git a/Engine/Source/Runtime/Engine/Private/PhysicsEngine/BodySetup.cpp b/Engine/Source/Runtime/Engine/Private/PhysicsEngine/BodySetup.cpp
--- a/Engine/Source/Runtime/Engine/Private/PhysicsEngine/BodySetup.cpp
+++ b/Engine/Source/Runtime/Engine/Private/PhysicsEngine/BodySetup.cpp
@@ -1059,10 +1059,6 @@ void UBodySetup::Serialize(FArchive& Ar)
 			// processed with Chaos previously
 			Convex.ComputeChaosConvexIndices();
 		}
-		if (bForceIndexRebuild)
-		{
-			BodySetupGuid = FGuid::NewGuid(); // change the guid for dependents (navigation)
-		}
 	}
 #endif
 }
diff --git a/Engine/Source/Runtime/Experimental/Chaos/Public/Chaos/AABB.h b/Engine/Source/Runtime/Experimental/Chaos/Public/Chaos/AABB.h
--- a/Engine/Source/Runtime/Experimental/Chaos/Public/Chaos/AABB.h
+++ b/Engine/Source/Runtime/Experimental/Chaos/Public/Chaos/AABB.h
@@ -31,8 +31,8 @@ namespace Chaos
 		static constexpr int D = d;
 
 		FORCEINLINE TAABB()
-			: MMin()
-			, MMax()
+			: MMin(TVector<T, d>(TNumericLimits<T>::Max()))
+			, MMax(TVector<T, d>(-TNumericLimits<T>::Max()))
 		{
 		}
 
diff --git a/Engine/Source/Runtime/Experimental/Chaos/Public/Chaos/AABBTree.h b/Engine/Source/Runtime/Experimental/Chaos/Public/Chaos/AABBTree.h
--- a/Engine/Source/Runtime/Experimental/Chaos/Public/Chaos/AABBTree.h
+++ b/Engine/Source/Runtime/Experimental/Chaos/Public/Chaos/AABBTree.h
@@ -229,8 +229,8 @@ struct TAABBTreeNode
 		ChildrenBounds[1] = TAABB<T, 3>();
 	}
 	TAABB<T, 3> ChildrenBounds[2];
-	int32 ChildrenNodes[2];
-	bool bLeaf;
+	int32 ChildrenNodes[2] = { 0, 0 };
+	bool bLeaf = false;
 
 	void Serialize(FChaosArchive& Ar)
 	{
diff --git a/Engine/Source/Runtime/NavigationSystem/Private/NavCollision.cpp b/Engine/Source/Runtime/NavigationSystem/Private/NavCollision.cpp
--- a/Engine/Source/Runtime/NavigationSystem/Private/NavCollision.cpp
+++ b/Engine/Source/Runtime/NavigationSystem/Private/NavCollision.cpp
@@ -99,7 +99,7 @@ public:
 
 	virtual FString GetPluginSpecificCacheKeySuffix() const override
 	{
-		const uint16 Version = 13;
+		const uint16 Version = 14;
 
 		return FString::Printf( TEXT("%s_%s_%s_%hu")
 			, *Format.ToString()
