Add cvar "net.CheckNoLoadPackages" to prevent the package map from sync loading packages for which FNetGuidCacheObject::bNoLoad is set. On behalf of Brian.Bekich.

diff --git a/Engine/Source/Runtime/Engine/Private/PackageMapClient.cpp b/Engine/Source/Runtime/Engine/Private/PackageMapClient.cpp
--- a/Engine/Source/Runtime/Engine/Private/PackageMapClient.cpp
+++ b/Engine/Source/Runtime/Engine/Private/PackageMapClient.cpp
@@ -105,6 +105,13 @@ static FAutoConsoleVariableRef CVarQuantizeActorVelocityOnSpawn(
 	TEXT("When enabled, we will quantize Velocity for newly spawned actors to a single decimal of precision.")
 );
 
+static bool GbNetCheckNoLoadPackages = true;
+static FAutoConsoleVariableRef CVarNetCheckNoLoadPackages(
+	TEXT("net.CheckNoLoadPackages"),
+	GbNetCheckNoLoadPackages,
+	TEXT("If enabled, check the no load flag in GetObjectFromNetGUID before forcing a sync load on packages that are not marked IsFullyLoaded")
+);
+
 void BroadcastNetFailure(UNetDriver* Driver, ENetworkFailure::Type FailureType, const FString& ErrorStr)
 {
 	UWorld* World = Driver->GetWorld();
@@ -3164,11 +3171,16 @@ UObject* FNetGUIDCache::GetObjectFromNetGUID( const FNetworkGUID& NetGUID, const
 				// We don't want to hook up this package into the cache yet or return it, because it's only partially loaded.
 				return nullptr;
 			}
-			else
+			else if (!GbNetCheckNoLoadPackages || !CacheObjectPtr->bNoLoad)
 			{
 				// If package isn't fully loaded, load it now
 				UE_LOG(LogNetPackageMap, Log, TEXT("GetObjectFromNetGUID: Blocking load of %s, NetGUID: %s"), *CacheObjectPtr->PathName.ToString(), *NetGUID.ToString());
-				Object = LoadPackage(NULL, *CacheObjectPtr->PathName.ToString(), LOAD_None);				
+				Object = LoadPackage(NULL, *CacheObjectPtr->PathName.ToString(), LOAD_None);
+			}
+			else
+			{
+				// Not fully loaded but we should not be loading it directly.
+				return nullptr;
 			}
 		}
 	}
