Fix static analysis warnings in CharacterMovementComponent.cpp

#rnx

diff --git a/Engine/Source/Runtime/Engine/Private/Components/CharacterMovementComponent.cpp b/Engine/Source/Runtime/Engine/Private/Components/CharacterMovementComponent.cpp
--- a/Engine/Source/Runtime/Engine/Private/Components/CharacterMovementComponent.cpp
+++ b/Engine/Source/Runtime/Engine/Private/Components/CharacterMovementComponent.cpp
@@ -1483,7 +1483,7 @@ void UCharacterMovementComponent::SimulatedTick(float DeltaSeconds)
 		UE_LOG(LogRootMotion, Verbose, TEXT("UCharacterMovementComponent::SimulatedTick"));
 
 		// Tick animations before physics.
-		if( CharacterOwner->GetMesh() )
+		if( CharacterOwner && CharacterOwner->GetMesh() )
 		{
 			TickCharacterPose(DeltaSeconds);
 
@@ -2446,7 +2446,7 @@ void UCharacterMovementComponent::PerformMovement(float DeltaSeconds)
 		else if (CurrentRootMotion.HasActiveRootMotionSources())
 		{
 			FQuat RootMotionRotationQuat;
-			if (CurrentRootMotion.GetOverrideRootMotionRotation(DeltaSeconds, *CharacterOwner, *this, RootMotionRotationQuat))
+			if (CharacterOwner && UpdatedComponent && CurrentRootMotion.GetOverrideRootMotionRotation(DeltaSeconds, *CharacterOwner, *this, RootMotionRotationQuat))
 			{
 				const FQuat OldActorRotationQuat = UpdatedComponent->GetComponentQuat();
 				const FQuat NewActorRotationQuat = RootMotionRotationQuat * OldActorRotationQuat;
@@ -9305,7 +9305,10 @@ void UCharacterMovementComponent::ClientAdjustPosition_Implementation
 	OnClientCorrectionReceived(*ClientData, TimeStamp, WorldShiftedNewLocation, NewVelocity, NewBase, NewBaseBoneName, bHasBase, bBaseRelativePosition, ServerMovementMode);
 
 	// Trust the server's positioning.
-	UpdatedComponent->SetWorldLocation(WorldShiftedNewLocation, false, nullptr, ETeleportType::TeleportPhysics);
+	if (UpdatedComponent)
+	{
+		UpdatedComponent->SetWorldLocation(WorldShiftedNewLocation, false, nullptr, ETeleportType::TeleportPhysics);
+	}
 	Velocity = NewVelocity;
 
 	// Trust the server's movement mode
