Fix paint bug when erasing all instances

diff --git a/Engine/Source/Runtime/Engine/Private/HierarchicalInstancedStaticMesh.cpp b/Engine/Source/Runtime/Engine/Private/HierarchicalInstancedStaticMesh.cpp
--- a/Engine/Source/Runtime/Engine/Private/HierarchicalInstancedStaticMesh.cpp
+++ b/Engine/Source/Runtime/Engine/Private/HierarchicalInstancedStaticMesh.cpp
@@ -2624,6 +2624,8 @@ void UHierarchicalInstancedStaticMeshComponent::BuildTree()
 		UnbuiltInstanceBoundsList.Empty();
 		BuiltInstanceBounds.Init();
 		CacheMeshExtendedBounds = FBoxSphereBounds(ForceInitToZero);
+		PerInstanceRenderData.Reset();
+		MarkRenderStateDirty();
 	}
 
 	// If an Async Build is running it will bail out because data is no longer bIsOutOfDate
@@ -2920,6 +2922,8 @@ void UHierarchicalInstancedStaticMeshComponent::BuildTreeAsync()
 		UnbuiltInstanceBoundsList.Empty();
 		BuiltInstanceBounds.Init();		
 		bIsOutOfDate = false;
+		PerInstanceRenderData.Reset();
+		MarkRenderStateDirty();
 	}
 }
 
