Don't cache soundwave data to/from the DDC when cooking for dedicated server targets

diff --git a/Engine/Source/Runtime/Engine/Private/SoundWave.cpp b/Engine/Source/Runtime/Engine/Private/SoundWave.cpp
--- a/Engine/Source/Runtime/Engine/Private/SoundWave.cpp
+++ b/Engine/Source/Runtime/Engine/Private/SoundWave.cpp
@@ -835,9 +835,12 @@ void USoundWave::PostLoad()
 	{
 		const TArray<ITargetPlatform*>& Platforms = TPM->GetActiveTargetPlatforms();
 
-		for (int32 Index = 0; Index < Platforms.Num(); Index++)
+		for (const ITargetPlatform* Platform : Platforms)
 		{
-			BeginGetCompressedData(Platforms[Index]->GetWaveFormat(this), FPlatformCompressionUtilities::GetCookOverrides(*Platforms[Index]->IniPlatformName()));
+			if (!Platform->IsServerOnly())
+			{
+				BeginGetCompressedData(Platform->GetWaveFormat(this), FPlatformCompressionUtilities::GetCookOverrides(*Platform->IniPlatformName()));
+			}
 		}
 	}
 
