Do not perform project setting based validation on engine visualization materials during DDC warming.

diff --git a/Engine/Source/Runtime/Engine/Private/Materials/HLSLMaterialTranslator.cpp b/Engine/Source/Runtime/Engine/Private/Materials/HLSLMaterialTranslator.cpp
--- a/Engine/Source/Runtime/Engine/Private/Materials/HLSLMaterialTranslator.cpp
+++ b/Engine/Source/Runtime/Engine/Private/Materials/HLSLMaterialTranslator.cpp
@@ -4849,8 +4849,9 @@ void FHLSLMaterialTranslator::UseSceneTextureId(ESceneTextureId SceneTextureId,
 	if (SceneTextureId == PPI_WorldTangent)
 	{
 		static IConsoleVariable* AnisotropicBRDF = IConsoleManager::Get().FindConsoleVariable(TEXT("r.AnisotropicBRDF"));
+		static bool bIsPopulatingDDC = FString(FCommandLine::Get()).Contains(TEXT("Run=DerivedDataCache"));
 
-		if (!AnisotropicBRDF || !AnisotropicBRDF->GetBool())
+		if (!bIsPopulatingDDC && (!AnisotropicBRDF || !AnisotropicBRDF->GetBool()))
 		{
 			Errorf(TEXT("World Tangent scene texture is only available when using anisotropic BRDF."));
 		}
@@ -4859,8 +4860,9 @@ void FHLSLMaterialTranslator::UseSceneTextureId(ESceneTextureId SceneTextureId,
 	if (SceneTextureId == PPI_Anisotropy)
 	{
 		static IConsoleVariable* AnisotropicBRDF = IConsoleManager::Get().FindConsoleVariable(TEXT("r.AnisotropicBRDF"));
+		static bool bIsPopulatingDDC = FString(FCommandLine::Get()).Contains(TEXT("Run=DerivedDataCache"));
 
-		if (!AnisotropicBRDF || !AnisotropicBRDF->GetBool())
+		if (!bIsPopulatingDDC && (!AnisotropicBRDF || !AnisotropicBRDF->GetBool()))
 		{
 			Errorf(TEXT("Anisotropy scene texture is only available when using anisotropic BRDF."));
 		}
