De-reference correct pointer.

diff --git a/Engine/Source/Runtime/Engine/Private/PhysicsEngine/Experimental/PhysScene_Chaos.cpp b/Engine/Source/Runtime/Engine/Private/PhysicsEngine/Experimental/PhysScene_Chaos.cpp
--- a/Engine/Source/Runtime/Engine/Private/PhysicsEngine/Experimental/PhysScene_Chaos.cpp
+++ b/Engine/Source/Runtime/Engine/Private/PhysicsEngine/Experimental/PhysScene_Chaos.cpp
@@ -1927,7 +1927,7 @@ int32 FPhysScene_ChaosInterface::DirtyElementCount(Chaos::ISpatialAccelerationCo
 		}
 		else if (const auto AABBTreeBV = SubStructure->template As<TAABBTree<TAccelerationStructureHandle<FReal, 3>, TBoundingVolume<TAccelerationStructureHandle<FReal, 3>, FReal, 3>, FReal>>())
 		{
-			DirtyElements += AABBTree->NumDirtyElements();
+			DirtyElements += AABBTreeBV->NumDirtyElements();
 		}
 	}
 	return DirtyElements;
