Renamed mobile "Support Distance Field Shadows" setting to "Support Pre-baked Distance Field Shadow Maps" and improve tooltip to reduce confusion

diff --git a/Engine/Source/Runtime/Engine/Classes/Engine/RendererSettings.h b/Engine/Source/Runtime/Engine/Classes/Engine/RendererSettings.h
--- a/Engine/Source/Runtime/Engine/Classes/Engine/RendererSettings.h
+++ b/Engine/Source/Runtime/Engine/Classes/Engine/RendererSettings.h
@@ -750,8 +750,8 @@ class ENGINE_API URendererSettings : public UDeveloperSettings
 
 	UPROPERTY(config, EditAnywhere, Category = MobileShaderPermutationReduction, meta = (
 		ConsoleVariable = "r.Mobile.AllowDistanceFieldShadows",
-		DisplayName = "Support Distance Field Shadows",
-		ToolTip = "Generate shaders for primitives to receive distance field shadows from stationary directional lights. Changing this setting requires restarting the editor.",
+		DisplayName = "Support Pre-baked Distance Field Shadow Maps",
+		ToolTip = "Generate shaders for static primitives render Lightmass-baked distance field shadow maps from stationary directional lights. Changing this setting requires restarting the editor.",
 		ConfigRestartRequired = true))
 		uint32 bMobileAllowDistanceFieldShadows : 1;
 
