Fix static analysis warning in CharacterMovementComponent.cpp

#rnx

diff --git a/Engine/Source/Runtime/Engine/Private/Components/CharacterMovementComponent.cpp b/Engine/Source/Runtime/Engine/Private/Components/CharacterMovementComponent.cpp
--- a/Engine/Source/Runtime/Engine/Private/Components/CharacterMovementComponent.cpp
+++ b/Engine/Source/Runtime/Engine/Private/Components/CharacterMovementComponent.cpp
@@ -1459,9 +1459,12 @@ void UCharacterMovementComponent::SimulatedTick(float DeltaSeconds)
 		UpdateComponentVelocity();
 		bJustTeleported = false;
 
-		CharacterOwner->RootMotionRepMoves.Empty();
-		CurrentRootMotion.Clear();
-		CharacterOwner->SavedRootMotion.Clear();
+		if (CharacterOwner)
+		{
+			CharacterOwner->RootMotionRepMoves.Empty();
+			CurrentRootMotion.Clear();
+			CharacterOwner->SavedRootMotion.Clear();
+		}
 
 		// Note: we do not call the Super implementation, that runs prediction.
 		// We do still need to call these though
