Reorder geometry collection commands and fix bad captures causing solver actor crashes.

diff --git a/Engine/Source/Runtime/Engine/Private/PhysicsEngine/Experimental/PhysScene_Chaos.cpp b/Engine/Source/Runtime/Engine/Private/PhysicsEngine/Experimental/PhysScene_Chaos.cpp
--- a/Engine/Source/Runtime/Engine/Private/PhysicsEngine/Experimental/PhysScene_Chaos.cpp
+++ b/Engine/Source/Runtime/Engine/Private/PhysicsEngine/Experimental/PhysScene_Chaos.cpp
@@ -745,10 +745,10 @@ void RemovePhysicsProxy(ObjectType* InObject, Chaos::FPhysicsSolver* InSolver, F
 	const bool bDedicatedThread = PhysDispatcher->GetMode() == Chaos::EThreadingMode::DedicatedThread;
 
 	// Remove the object from the solver
-	PhysDispatcher->EnqueueCommandImmediate([InObject, InSolver, bDedicatedThread](Chaos::FPersistentPhysicsTask* PhysThread)
+	PhysDispatcher->EnqueueCommandImmediate(InSolver, [InObject, bDedicatedThread](Chaos::FPBDRigidsSolver* InnerSolver)
 	{
 #if CHAOS_PARTICLEHANDLE_TODO
-		InSolver->UnregisterObject(InObject);
+		InnerSolver->UnregisterObject(InObject);
 #endif
 		InObject->OnRemoveFromScene();
 
diff --git a/Engine/Source/Runtime/Experimental/GeometryCollectionEngine/Private/GeometryCollection/GeometryCollectionComponent.cpp b/Engine/Source/Runtime/Experimental/GeometryCollectionEngine/Private/GeometryCollection/GeometryCollectionComponent.cpp
--- a/Engine/Source/Runtime/Experimental/GeometryCollectionEngine/Private/GeometryCollection/GeometryCollectionComponent.cpp
+++ b/Engine/Source/Runtime/Experimental/GeometryCollectionEngine/Private/GeometryCollection/GeometryCollectionComponent.cpp
@@ -232,7 +232,7 @@ void UGeometryCollectionComponent::BeginPlay()
 			}
 			if (Solver != nullptr)
 			{
-				ChaosModule->GetDispatcher()->EnqueueCommandImmediate(Solver, [&InPhysicsProxy = PhysicsProxy](Chaos::FPhysicsSolver* InSolver)
+				ChaosModule->GetDispatcher()->EnqueueCommandImmediate(Solver, [InPhysicsProxy = PhysicsProxy](Chaos::FPhysicsSolver* InSolver)
 				{
 					InPhysicsProxy->ActivateBodies();
 				});
