Make DestroyedReplicatedStaticActors transient. It should never have been saved, and can cause crashes in PIE if entries do happen to get saved in it.

diff --git a/Engine/Source/Runtime/Engine/Classes/Engine/Level.h b/Engine/Source/Runtime/Engine/Classes/Engine/Level.h
--- a/Engine/Source/Runtime/Engine/Classes/Engine/Level.h
+++ b/Engine/Source/Runtime/Engine/Classes/Engine/Level.h
@@ -695,7 +695,7 @@ private:
 	TArray<FPendingAutoReceiveInputActor> PendingAutoReceiveInputActors;
 
 	/** List of replicated static actors that have been destroyed. Used by net drivers to replicate destruction to clients. */
-	UPROPERTY()
+	UPROPERTY(Transient)
 	TArray<FReplicatedStaticActorDestructionInfo> DestroyedReplicatedStaticActors;
 
 #if WITH_EDITORONLY_DATA
