Following bug fix submitted in 4.25 did this clean up in main to minimize risk on the 4.25 release:
Merge BuildTree and BuildTreeAsync. They both go through the same ApplyBuiltTree code now.

#rnx

diff --git a/Engine/Source/Runtime/Engine/Classes/Components/HierarchicalInstancedStaticMeshComponent.h b/Engine/Source/Runtime/Engine/Classes/Components/HierarchicalInstancedStaticMeshComponent.h
--- a/Engine/Source/Runtime/Engine/Classes/Components/HierarchicalInstancedStaticMeshComponent.h
+++ b/Engine/Source/Runtime/Engine/Classes/Components/HierarchicalInstancedStaticMeshComponent.h
@@ -257,8 +257,13 @@ public:
 protected:
 	void BuildTree();
 	void BuildTreeAsync();
+	void ApplyBuildTree(FClusterBuilder& Builder);
+	void ApplyEmpty();
 	void SetPerInstanceLightMapAndEditorData(FStaticMeshInstanceData& PerInstanceData, const TArray<TRefCountPtr<HHitProxy>>& HitProxies);
 
+	void GetInstanceTransforms(TArray<FMatrix>& InstanceTransforms) const;
+	void InitializeInstancingRandomSeed();
+
 	/** Removes specified instances */ 
 	void RemoveInstancesInternal(const int32* InstanceIndices, int32 Num);
 	
diff --git a/Engine/Source/Runtime/Engine/Private/HierarchicalInstancedStaticMesh.cpp b/Engine/Source/Runtime/Engine/Private/HierarchicalInstancedStaticMesh.cpp
--- a/Engine/Source/Runtime/Engine/Private/HierarchicalInstancedStaticMesh.cpp
+++ b/Engine/Source/Runtime/Engine/Private/HierarchicalInstancedStaticMesh.cpp
@@ -2533,103 +2533,28 @@ void UHierarchicalInstancedStaticMeshComponent::BuildTree()
 
 	QUICK_SCOPE_CYCLE_COUNTER(STAT_UHierarchicalInstancedStaticMeshComponent_BuildTree);
 
-	// upload instance edits to GPU, before validing if the mesh is valid, as it's possible that PerInstanceSMData.Num() == 0, so we have to hide everything before doing the build
+	// Make sure while Tree is building that our render state is updated (Command Buffer is processed)
 	if (InstanceUpdateCmdBuffer.NumInlineCommands() > 0 && PerInstanceRenderData.IsValid() && PerInstanceRenderData->InstanceBuffer.RequireCPUAccess)
 	{
-		// if instance data was modified, update GPU copy
-		// if InstanceBuffer was initialized with RequireCPUAccess (always true in editor))
 		PerInstanceRenderData->UpdateFromCommandBuffer(InstanceUpdateCmdBuffer);
 		MarkRenderStateDirty();
 	}
 
-	// all pending edits will be updated
-	InstanceUpdateCmdBuffer.Reset();
-
-	// Verify that the mesh is valid before using it.
-	const bool bMeshIsValid = 
-		// make sure we have instances
-		PerInstanceSMData.Num() > 0 &&
-		// make sure we have an actual staticmesh
-		GetStaticMesh() &&
-		GetStaticMesh()->HasValidRenderData() &&
-		// You really can't use hardware instancing on the consoles with multiple elements because they share the same index buffer. 
-		// @todo: Level error or something to let LDs know this
-		1;//GetStaticMesh()->LODModels(0).Elements.Num() == 1;
-
-	if (bMeshIsValid)
+	if (PerInstanceSMData.Num() > 0 && GetStaticMesh() && GetStaticMesh()->HasValidRenderData())
 	{
-		// If we don't have a random seed for this instanced static mesh component yet, then go ahead and
-		// generate one now.  This will be saved with the static mesh component and used for future generation
-		// of random numbers for this component's instances. (Used by the PerInstanceRandom material expression)
-		while (InstancingRandomSeed == 0)
-		{
-			InstancingRandomSeed = FMath::Rand();
-		}
-		
+		InitializeInstancingRandomSeed();
 		TArray<FMatrix> InstanceTransforms;
-		InstanceTransforms.SetNumUninitialized(PerInstanceSMData.Num());
-		for (int32 Index = 0; Index < PerInstanceSMData.Num(); Index++)
-		{
-			InstanceTransforms[Index] = PerInstanceSMData[Index].Transform;
-		}
+		GetInstanceTransforms(InstanceTransforms);
 
 		FClusterBuilder Builder(InstanceTransforms, PerInstanceSMCustomData, NumCustomDataFloats, GetStaticMesh()->GetBounds().GetBox(), DesiredInstancesPerLeaf(), CurrentDensityScaling, InstancingRandomSeed, PerInstanceSMData.Num() > 0);
 		Builder.BuildTreeAndBuffer();
 
-		NumBuiltInstances = Builder.Result->InstanceReorderTable.Num();
-		NumBuiltRenderInstances = Builder.Result->SortedInstances.Num();
-		OcclusionLayerNumNodes = Builder.Result->OutOcclusionLayerNum;
-		UnbuiltInstanceBounds.Init();
-		UnbuiltInstanceBoundsList.Empty();
-		BuiltInstanceBounds = (Builder.Result->Nodes.Num() > 0 ? FBox(Builder.Result->Nodes[0].BoundMin, Builder.Result->Nodes[0].BoundMax) : FBox(ForceInit));
-
-		ClusterTreePtr = MakeShareable(new TArray<FClusterNode>(MoveTemp(Builder.Result->Nodes)));
-		InstanceReorderTable = MoveTemp(Builder.Result->InstanceReorderTable);
-		SortedInstances = MoveTemp(Builder.Result->SortedInstances);
-		TUniquePtr<FStaticMeshInstanceData> BuiltInstanceData = MoveTemp(Builder.BuiltInstanceData);
-		CacheMeshExtendedBounds = GetStaticMesh()->GetBounds();
-			
-		check(BuiltInstanceData.IsValid());
-		check(BuiltInstanceData->GetNumInstances() == NumBuiltRenderInstances);
-		InstanceCountToRender = NumBuiltInstances;
-		
-		// create per-instance hit-proxies if needed
-		TArray<TRefCountPtr<HHitProxy>> HitProxies;
-		CreateHitProxyData(HitProxies);
-		SetPerInstanceLightMapAndEditorData(*BuiltInstanceData, HitProxies);
-
-		if (PerInstanceRenderData.IsValid())
-		{
-			PerInstanceRenderData->UpdateFromPreallocatedData(*BuiltInstanceData);
-		}
-		else
-		{
-			InitPerInstanceRenderData(false, BuiltInstanceData.Get());
-		}
-		PerInstanceRenderData->HitProxies = MoveTemp(HitProxies);
-		
-		MarkRenderStateDirty();
-		FlushAccumulatedNavigationUpdates();
-		PostBuildStats();
+		ApplyBuildTree(Builder);
 	}
 	else
 	{
-		ClusterTreePtr = MakeShareable(new TArray<FClusterNode>);
-		NumBuiltInstances = 0;
-		NumBuiltRenderInstances = 0;
-		InstanceCountToRender = 0;
-		InstanceReorderTable.Empty();
-		SortedInstances.Empty();
-
-		UnbuiltInstanceBoundsList.Empty();
-		BuiltInstanceBounds.Init();
-		CacheMeshExtendedBounds = FBoxSphereBounds(ForceInitToZero);
-		PerInstanceRenderData.Reset();
-		MarkRenderStateDirty();
+		ApplyEmpty();
 	}
-
-	// If an Async Build is running it will bail out because data is no longer bIsOutOfDate
-	bIsOutOfDate = false;
 }
 
 void UHierarchicalInstancedStaticMeshComponent::BuildTreeAnyThread(
@@ -2681,15 +2606,12 @@ void UHierarchicalInstancedStaticMeshComponent::AcceptPrebuiltTree(TArray<FClust
 	InstanceCountToRender = InNumBuiltRenderInstances;
 
 	// Verify that the mesh is valid before using it.
-	const bool bMeshIsValid = 
+	const bool bMeshIsValid =
 		// make sure we have instances
 		NumBuiltRenderInstances > 0 &&
 		// make sure we have an actual staticmesh
 		GetStaticMesh() &&
-		GetStaticMesh()->HasValidRenderData() &&
-		// You really can't use hardware instancing on the consoles with multiple elements because they share the same index buffer. 
-		// @todo: Level error or something to let LDs know this
-		1;//GetStaticMesh()->LODModels(0).Elements.Num() == 1;
+		GetStaticMesh()->HasValidRenderData();
 
 	if(bMeshIsValid)
 	{
@@ -2729,27 +2651,49 @@ void UHierarchicalInstancedStaticMeshComponent::ApplyBuildTreeAsync(ENamedThread
 	}
 
 	// Completed the build
+	ApplyBuildTree(Builder.Get());
+}
+
+void UHierarchicalInstancedStaticMeshComponent::ApplyEmpty()
+{
+	bIsOutOfDate = false;
+	ClusterTreePtr = MakeShareable(new TArray<FClusterNode>);
+	NumBuiltInstances = 0;
+	NumBuiltRenderInstances = 0;
+	InstanceCountToRender = 0;
+	InstanceReorderTable.Empty();
+	SortedInstances.Empty();
+	CacheMeshExtendedBounds = FBoxSphereBounds(ForceInitToZero);
+
+	UnbuiltInstanceBoundsList.Empty();
+	BuiltInstanceBounds.Init();
+	PerInstanceRenderData.Reset();
+	MarkRenderStateDirty();
+}
+
+void UHierarchicalInstancedStaticMeshComponent::ApplyBuildTree(FClusterBuilder& Builder)
+{
 	bIsOutOfDate = false;
+	InstanceUpdateCmdBuffer.Reset();
 
-	check(Builder->Result->InstanceReorderTable.Num() == PerInstanceSMData.Num());
+	check(Builder.Result->InstanceReorderTable.Num() == PerInstanceSMData.Num());
 
-	NumBuiltInstances = Builder->Result->InstanceReorderTable.Num();
-	NumBuiltRenderInstances = Builder->Result->SortedInstances.Num();		
+	NumBuiltInstances = Builder.Result->InstanceReorderTable.Num();
+	NumBuiltRenderInstances = Builder.Result->SortedInstances.Num();
 
-	ClusterTreePtr = MakeShareable(new TArray<FClusterNode>(MoveTemp(Builder->Result->Nodes)));
-	const TArray<FClusterNode>& ClusterTree = *ClusterTreePtr;
-	InstanceReorderTable = MoveTemp(Builder->Result->InstanceReorderTable);
-	SortedInstances = MoveTemp(Builder->Result->SortedInstances);
+	ClusterTreePtr = MakeShareable(new TArray<FClusterNode>(MoveTemp(Builder.Result->Nodes)));
+
+	InstanceReorderTable = MoveTemp(Builder.Result->InstanceReorderTable);
+	SortedInstances = MoveTemp(Builder.Result->SortedInstances);
 	CacheMeshExtendedBounds = GetStaticMesh()->GetBounds();
-	TUniquePtr<FStaticMeshInstanceData> BuiltInstanceData = MoveTemp(Builder->BuiltInstanceData);
-	OcclusionLayerNumNodes = Builder->Result->OutOcclusionLayerNum;
-	BuiltInstanceBounds = (ClusterTree.Num() > 0 ? FBox(ClusterTree[0].BoundMin, ClusterTree[0].BoundMax) : FBox(ForceInit));
+	TUniquePtr<FStaticMeshInstanceData> BuiltInstanceData = MoveTemp(Builder.BuiltInstanceData);
 
-	UE_LOG(LogStaticMesh, Verbose, TEXT("Built a foliage hierarchy with %d of %d elements in %.1fs."), NumBuiltInstances, PerInstanceSMData.Num(), (float)(FPlatformTime::Seconds() - StartTime));
+	OcclusionLayerNumNodes = Builder.Result->OutOcclusionLayerNum;
+	const TArray<FClusterNode>& ClusterTree = *ClusterTreePtr;
+	BuiltInstanceBounds = (ClusterTree.Num() > 0 ? FBox(ClusterTree[0].BoundMin, ClusterTree[0].BoundMax) : FBox(ForceInit));
 
 	UnbuiltInstanceBounds.Init();
 	UnbuiltInstanceBoundsList.Empty();
-	FlushAccumulatedNavigationUpdates();
 
 	check(BuiltInstanceData.IsValid());
 	check(BuiltInstanceData->GetNumInstances() == NumBuiltRenderInstances);
@@ -2772,10 +2716,11 @@ void UHierarchicalInstancedStaticMeshComponent::ApplyBuildTreeAsync(ENamedThread
 	{
 		InitPerInstanceRenderData(false, BuiltInstanceData.Get());
 	}
-	PerInstanceRenderData->HitProxies = MoveTemp(HitProxies);		
+	PerInstanceRenderData->HitProxies = MoveTemp(HitProxies);
 
-	MarkRenderStateDirty();
+	FlushAccumulatedNavigationUpdates();
 	PostBuildStats();
+	MarkRenderStateDirty();
 }
 
 void UHierarchicalInstancedStaticMeshComponent::OnComponentCreated()
@@ -2838,6 +2783,31 @@ bool UHierarchicalInstancedStaticMeshComponent::BuildTreeIfOutdated(bool Async,
 	return false;
 }
 
+void UHierarchicalInstancedStaticMeshComponent::GetInstanceTransforms(TArray<FMatrix>& InstanceTransforms) const
+{
+	double StartTime = FPlatformTime::Seconds();
+	int32 Num = PerInstanceSMData.Num();
+
+	InstanceTransforms.SetNumUninitialized(Num);
+	for (int32 Index = 0; Index < Num; Index++)
+	{
+		InstanceTransforms[Index] = PerInstanceSMData[Index].Transform;
+	}
+
+	UE_LOG(LogStaticMesh, Verbose, TEXT("Copied %d transforms in %.3fs."), Num, float(FPlatformTime::Seconds() - StartTime));
+}
+
+void UHierarchicalInstancedStaticMeshComponent::InitializeInstancingRandomSeed()
+{
+	// If we don't have a random seed for this instanced static mesh component yet, then go ahead and
+		// generate one now.  This will be saved with the static mesh component and used for future generation
+		// of random numbers for this component's instances. (Used by the PerInstanceRandom material expression)
+	while (InstancingRandomSeed == 0)
+	{
+		InstancingRandomSeed = FMath::Rand();
+	}
+}
+
 void UHierarchicalInstancedStaticMeshComponent::BuildTreeAsync()
 {
 	check(!HasAnyFlags(RF_ClassDefaultObject | RF_ArchetypeObject));
@@ -2849,47 +2819,22 @@ void UHierarchicalInstancedStaticMeshComponent::BuildTreeAsync()
 	check(!bIsAsyncBuilding);
 	check(BuildTreeAsyncTasks.Num() == 0);
 
-	// upload instance edits to GPU, before validing if the mesh is valid, as it's possible that PerInstanceSMData.Num() == 0, so we have to hide everything before doing the build
+	// Make sure while Tree is building that our render state is updated (Command Buffer is processed)
 	if (InstanceUpdateCmdBuffer.NumInlineCommands() > 0 && PerInstanceRenderData.IsValid() && PerInstanceRenderData->InstanceBuffer.RequireCPUAccess)
 	{
-		// if instance data was modified, update GPU copy
-		// if InstanceBuffer was initialized with RequireCPUAccess (always true in editor))
 		PerInstanceRenderData->UpdateFromCommandBuffer(InstanceUpdateCmdBuffer);
 		MarkRenderStateDirty();
 	}
 
 	// Verify that the mesh is valid before using it.
-	const bool bMeshIsValid =
-		// make sure we have instances
-		PerInstanceSMData.Num() > 0 &&
-		// make sure we have an actual staticmesh
-		GetStaticMesh() &&
-		GetStaticMesh()->HasValidRenderData() &&
-		// You really can't use hardware instancing on the consoles with multiple elements because they share the same index buffer. 
-		// @todo: Level error or something to let LDs know this
-		1;//GetStaticMesh()->LODModels(0).Elements.Num() == 1;
-
-	if (bMeshIsValid)
+	if (PerInstanceSMData.Num() > 0 && GetStaticMesh() && GetStaticMesh()->HasValidRenderData())
 	{
 		double StartTime = FPlatformTime::Seconds();
-		// If we don't have a random seed for this instanced static mesh component yet, then go ahead and
-		// generate one now.  This will be saved with the static mesh component and used for future generation
-		// of random numbers for this component's instances. (Used by the PerInstanceRandom material expression)
-		while (InstancingRandomSeed == 0)
-		{
-			InstancingRandomSeed = FMath::Rand();
-		}
-
-		int32 Num = PerInstanceSMData.Num();
+		
+		InitializeInstancingRandomSeed();
 		TArray<FMatrix> InstanceTransforms;
-		InstanceTransforms.SetNumUninitialized(Num);
-		for (int32 Index = 0; Index < Num; Index++)
-		{
-			InstanceTransforms[Index] = PerInstanceSMData[Index].Transform;
-		}
-
-		UE_LOG(LogStaticMesh, Verbose, TEXT("Copied %d transforms in %.3fs."), Num, float(FPlatformTime::Seconds() - StartTime));
-
+		GetInstanceTransforms(InstanceTransforms);
+		
 		TSharedRef<FClusterBuilder, ESPMode::ThreadSafe> Builder(new FClusterBuilder(InstanceTransforms, PerInstanceSMCustomData, NumCustomDataFloats, GetStaticMesh()->GetBounds().GetBox(), DesiredInstancesPerLeaf(), CurrentDensityScaling, InstancingRandomSeed, PerInstanceSMData.Num() > 0));
 
 		bIsAsyncBuilding = true;
@@ -2911,19 +2856,7 @@ void UHierarchicalInstancedStaticMeshComponent::BuildTreeAsync()
 	}
 	else
 	{
-		ClusterTreePtr = MakeShareable(new TArray<FClusterNode>);
-		NumBuiltInstances = 0;
-		NumBuiltRenderInstances = 0;
-		InstanceCountToRender = 0;
-		InstanceReorderTable.Empty();
-		SortedInstances.Empty();
-		CacheMeshExtendedBounds = FBoxSphereBounds(ForceInitToZero);
-
-		UnbuiltInstanceBoundsList.Empty();
-		BuiltInstanceBounds.Init();		
-		bIsOutOfDate = false;
-		PerInstanceRenderData.Reset();
-		MarkRenderStateDirty();
+		ApplyEmpty();
 	}
 }
 
