Fixes for Niagara vector field resources being accessed on the gamethread without it being necessarily initialized.
-Addresses missing curl forces in the mote and the bug effects (possibly others) in cooked (by the book) builds

diff --git a/Engine/Plugins/FX/Niagara/Source/Niagara/Private/NiagaraDataInterfaceVectorField.cpp b/Engine/Plugins/FX/Niagara/Source/Niagara/Private/NiagaraDataInterfaceVectorField.cpp
--- a/Engine/Plugins/FX/Niagara/Source/Niagara/Private/NiagaraDataInterfaceVectorField.cpp
+++ b/Engine/Plugins/FX/Niagara/Source/Niagara/Private/NiagaraDataInterfaceVectorField.cpp
@@ -665,12 +665,11 @@ void UNiagaraDataInterfaceVectorField::PushToRenderThread()
 	bool RT_bTileY = bTileY;
 	bool RT_bTileZ = bTileZ;
 
-	// @todo-threadsafety. This would be a race but I'm taking a ref here. Not ideal in the long term.
-	FTextureRHIRef VolumeTextureRHI = Field ? ((UVectorFieldStatic*)Field)->GetVolumeTextureRef() : (FRHITexture*)GBlackVolumeTexture->TextureRHI;
+	FVectorFieldTextureAccessor TextureAccessor(Field);
 
 	// Push Updates to Proxy.
-	ENQUEUE_RENDER_COMMAND(FUpdateDIColorCurve)(
-		[RT_Proxy, RT_Dimensions, RT_MinBounds, RT_MaxBounds, RT_bTileX, RT_bTileY, RT_bTileZ, VolumeTextureRHI](FRHICommandListImmediate& RHICmdList)
+	ENQUEUE_RENDER_COMMAND(FUpdateDIVectorField)(
+		[RT_Proxy, RT_Dimensions, RT_MinBounds, RT_MaxBounds, RT_bTileX, RT_bTileY, RT_bTileZ, TextureAccessor](FRHICommandListImmediate& RHICmdList)
 	{
 		RT_Proxy->bTileX = RT_bTileX;
 		RT_Proxy->bTileY = RT_bTileY;
@@ -678,7 +677,7 @@ void UNiagaraDataInterfaceVectorField::PushToRenderThread()
 		RT_Proxy->Dimensions = RT_Dimensions;
 		RT_Proxy->MinBounds = RT_MinBounds;
 		RT_Proxy->MaxBounds = RT_MaxBounds;
-		RT_Proxy->TextureRHI = VolumeTextureRHI;
+		RT_Proxy->TextureRHI = TextureAccessor.GetTexture();
 	});
 }
 
diff --git a/Engine/Source/Runtime/Engine/Classes/VectorField/VectorFieldStatic.h b/Engine/Source/Runtime/Engine/Classes/VectorField/VectorFieldStatic.h
--- a/Engine/Source/Runtime/Engine/Classes/VectorField/VectorFieldStatic.h
+++ b/Engine/Source/Runtime/Engine/Classes/VectorField/VectorFieldStatic.h
@@ -112,3 +112,15 @@ private:
 	
 };
 
+// work around for the private nature of FVectorFieldResource
+struct ENGINE_API FVectorFieldTextureAccessor
+{
+	FVectorFieldTextureAccessor(UVectorField* InVectorField);
+	FVectorFieldTextureAccessor(const FVectorFieldTextureAccessor& rhs);
+	~FVectorFieldTextureAccessor();
+
+	FRHITexture* GetTexture() const;
+
+private:
+	TUniquePtr<struct FVectorFieldTextureAccessorImpl> Impl;
+};
\ No newline at end of file
diff --git a/Engine/Source/Runtime/Engine/Private/VectorField.cpp b/Engine/Source/Runtime/Engine/Private/VectorField.cpp
--- a/Engine/Source/Runtime/Engine/Private/VectorField.cpp
+++ b/Engine/Source/Runtime/Engine/Private/VectorField.cpp
@@ -306,6 +306,57 @@ void UVectorFieldStatic::UpdateResource()
 	StaticResource->UpdateResource(this); // Will discard the contents of SourceData
 }
 
+// Simple implementation of an accessor struct for grabbing the, now, inited resource on the render thread
+// Impl structure is required because of the private nature of the FVectorFieldResource, which could be
+// resolved...
+struct FVectorFieldTextureAccessorImpl
+{
+	FVectorFieldTextureAccessorImpl(FVectorFieldResource* InResource)
+		: Resource(InResource)
+	{
+	}
+
+	FVectorFieldTextureAccessorImpl(const FVectorFieldTextureAccessorImpl& rhs)
+		: Resource(rhs.Resource)
+	{
+	}
+
+	TRefCountPtr<FVectorFieldResource> Resource;
+};
+
+FVectorFieldTextureAccessor::FVectorFieldTextureAccessor(UVectorField* InVectorField)
+	: Impl(nullptr)
+{
+	if (UVectorFieldStatic* VectorFieldStatic = Cast<UVectorFieldStatic>(InVectorField))
+	{
+		Impl = MakeUnique<FVectorFieldTextureAccessorImpl>(VectorFieldStatic->Resource);
+	}
+}
+
+FVectorFieldTextureAccessor::FVectorFieldTextureAccessor(const FVectorFieldTextureAccessor& rhs)
+	: Impl(nullptr)
+{
+	if (rhs.Impl)
+	{
+		Impl = MakeUnique<FVectorFieldTextureAccessorImpl>(rhs.Impl->Resource);
+	}
+}
+
+FVectorFieldTextureAccessor::~FVectorFieldTextureAccessor()
+{
+
+}
+
+FRHITexture* FVectorFieldTextureAccessor::GetTexture() const
+{
+	if (Impl && Impl->Resource && Impl->Resource->VolumeTextureRHI)
+	{
+		return Impl->Resource->VolumeTextureRHI;
+	}
+
+	return GBlackVolumeTexture->TextureRHI;
+}
+
 #if WITH_EDITOR
 ENGINE_API void UVectorFieldStatic::SetCPUAccessEnabled()
 {
