Cameras: Prevent users from setting single-instance shakes on shake source actors.

diff --git a/Engine/Source/Runtime/Engine/Classes/Camera/CameraShakeSourceComponent.h b/Engine/Source/Runtime/Engine/Classes/Camera/CameraShakeSourceComponent.h
--- a/Engine/Source/Runtime/Engine/Classes/Camera/CameraShakeSourceComponent.h
+++ b/Engine/Source/Runtime/Engine/Classes/Camera/CameraShakeSourceComponent.h
@@ -30,9 +30,18 @@ public:
 
 	virtual void OnRegister() override;
 
+#if WITH_EDITOR
+    virtual void PreEditChange(FProperty* PropertyAboutToChange) override;
+	virtual void PostEditChangeProperty( struct FPropertyChangedEvent& PropertyChangedEvent) override;
+#endif
+
 private:
 	void UpdateEditorSpriteTexture();
 
+#if WITH_EDITOR
+	TSubclassOf<UCameraShake> PreviousCameraShake;
+#endif
+
 public:
 	/** The attenuation profile for how camera shakes' intensity falls off with distance */
 	UPROPERTY(EditAnywhere, BlueprintReadWrite, Category = Attenuation)
diff --git a/Engine/Source/Runtime/Engine/Private/Camera/CameraShakeSourceComponent.cpp b/Engine/Source/Runtime/Engine/Private/Camera/CameraShakeSourceComponent.cpp
--- a/Engine/Source/Runtime/Engine/Private/Camera/CameraShakeSourceComponent.cpp
+++ b/Engine/Source/Runtime/Engine/Private/Camera/CameraShakeSourceComponent.cpp
@@ -1,8 +1,8 @@
 // Copyright Epic Games, Inc. All Rights Reserved.
 
 #include "Camera/CameraShakeSourceComponent.h"
-#include "Camera/CameraShake.h"
 #include "Camera/CameraModifier_CameraShake.h"
+#include "Camera/CameraShake.h"
 #include "CinematicCameraModule.h"
 #include "Components/BillboardComponent.h"
 #include "Engine/Engine.h"
@@ -12,6 +12,11 @@
 #include "Modules/ModuleManager.h"
 #include "UObject/ConstructorHelpers.h"
 
+#if WITH_EDITOR
+#include "Framework/Notifications/NotificationManager.h"
+#include "Widgets/Notifications/SNotificationList.h"
+#endif
+
 
 UCameraShakeSourceComponent::UCameraShakeSourceComponent(const FObjectInitializer& ObjectInitializer)
 	: Super(ObjectInitializer)
@@ -124,3 +129,43 @@ float UCameraShakeSourceComponent::GetAttenuationFactor(const FVector& Location)
 	}
 	return 1.f;
 }
+
+#if WITH_EDITOR
+
+void UCameraShakeSourceComponent::PreEditChange(FProperty* PropertyAboutToChange)
+{
+	if (PropertyAboutToChange != nullptr && PropertyAboutToChange->GetFName() == GET_MEMBER_NAME_CHECKED(UCameraShakeSourceComponent, CameraShake))
+    {
+        PreviousCameraShake = CameraShake;
+    }
+
+    Super::PreEditChange(PropertyAboutToChange);   
+}
+
+void UCameraShakeSourceComponent::PostEditChangeProperty( struct FPropertyChangedEvent& PropertyChangedEvent)
+{
+    Super::PostEditChangeProperty(PropertyChangedEvent);
+
+    if (PropertyChangedEvent.Property != nullptr && PropertyChangedEvent.Property->GetFName() == GET_MEMBER_NAME_CHECKED(UCameraShakeSourceComponent, CameraShake))
+    {
+        if (CameraShake != nullptr)
+        {
+            // Single instance shakes can't be localized since they're automatically recycled and merged.
+            // If the user is trying to set it, let's notify them and revert to the previous value.
+            const UCameraShake* const CameraShakeCDO = GetDefault<UCameraShake>(CameraShake);
+            if (CameraShakeCDO != nullptr && CameraShakeCDO->bSingleInstance)
+            {
+                FText NotificationText = FText::Format(
+                        NSLOCTEXT("CameraShakeSourceComponent", "SingleInstanceCameraShakeNotAllowed", "{0} is a single instance shake, and therefore can't be localized to a source actor."),
+                        FText::FromName(CameraShake->GetFName()));
+                FNotificationInfo NotificationInfo(NotificationText);
+                NotificationInfo.ExpireDuration = 5.f;
+                FSlateNotificationManager::Get().AddNotification(NotificationInfo);
+
+                CameraShake = PreviousCameraShake;
+            }
+        }
+    }
+}
+
+#endif
