Replicate backout from 4.25 until better fix is found.

#rnx

diff --git a/Engine/Source/Runtime/Engine/Private/HierarchicalInstancedStaticMesh.cpp b/Engine/Source/Runtime/Engine/Private/HierarchicalInstancedStaticMesh.cpp
--- a/Engine/Source/Runtime/Engine/Private/HierarchicalInstancedStaticMesh.cpp
+++ b/Engine/Source/Runtime/Engine/Private/HierarchicalInstancedStaticMesh.cpp
@@ -2667,8 +2667,6 @@ void UHierarchicalInstancedStaticMeshComponent::ApplyEmpty()
 
 	UnbuiltInstanceBoundsList.Empty();
 	BuiltInstanceBounds.Init();
-	PerInstanceRenderData.Reset();
-	MarkRenderStateDirty();
 }
 
 void UHierarchicalInstancedStaticMeshComponent::ApplyBuildTree(FClusterBuilder& Builder)
