Backout CL 11464239 (causes IsInRenderingThread failed check)

#rnx

diff --git a/Engine/Source/Runtime/Engine/Private/HierarchicalInstancedStaticMesh.cpp b/Engine/Source/Runtime/Engine/Private/HierarchicalInstancedStaticMesh.cpp
--- a/Engine/Source/Runtime/Engine/Private/HierarchicalInstancedStaticMesh.cpp
+++ b/Engine/Source/Runtime/Engine/Private/HierarchicalInstancedStaticMesh.cpp
@@ -2624,8 +2624,6 @@ void UHierarchicalInstancedStaticMeshComponent::BuildTree()
 		UnbuiltInstanceBoundsList.Empty();
 		BuiltInstanceBounds.Init();
 		CacheMeshExtendedBounds = FBoxSphereBounds(ForceInitToZero);
-		PerInstanceRenderData.Reset();
-		MarkRenderStateDirty();
 	}
 
 	// If an Async Build is running it will bail out because data is no longer bIsOutOfDate
@@ -2922,8 +2920,6 @@ void UHierarchicalInstancedStaticMeshComponent::BuildTreeAsync()
 		UnbuiltInstanceBoundsList.Empty();
 		BuiltInstanceBounds.Init();		
 		bIsOutOfDate = false;
-		PerInstanceRenderData.Reset();
-		MarkRenderStateDirty();
 	}
 }
 
