Added templatized InitializeDependency.
Also modified InitializeDependency to return the subsystem* should work with backwards compate as bool ops still work.
Would force a code fix if bool b = InitializeDependency(); however i expect the use cases are so rare it's not worth doing more.

diff --git a/Engine/Source/Runtime/Engine/Private/Subsystems/SubsystemCollection.cpp b/Engine/Source/Runtime/Engine/Private/Subsystems/SubsystemCollection.cpp
--- a/Engine/Source/Runtime/Engine/Private/Subsystems/SubsystemCollection.cpp
+++ b/Engine/Source/Runtime/Engine/Private/Subsystems/SubsystemCollection.cpp
@@ -163,7 +163,7 @@ void FSubsystemCollectionBase::Deinitialize()
 	Outer = nullptr;
 }
 
-bool FSubsystemCollectionBase::InitializeDependency(TSubclassOf<USubsystem> SubsystemClass)
+USubsystem* FSubsystemCollectionBase::InitializeDependency(TSubclassOf<USubsystem> SubsystemClass)
 {
 	if (ensureMsgf(SubsystemClass, TEXT("Attempting to add invalid subsystem as dependancy."))
 		&& ensureMsgf(bPopulating, TEXT("InitializeDependancy() should only be called from System USubsystem::Initialization() implementations."))
@@ -171,7 +171,7 @@ bool FSubsystemCollectionBase::InitializeDependency(TSubclassOf<USubsystem> Subs
 	{
 		return AddAndInitializeSubsystem(SubsystemClass);
 	}
-	return false;
+	return nullptr;
 }
 
 void FSubsystemCollectionBase::AddReferencedObjects(FReferenceCollector& Collector)
@@ -184,7 +184,7 @@ FString FSubsystemCollectionBase::GetReferencerName() const
 	return TEXT("FSubsystemCollectionBase");
 }
 
-bool FSubsystemCollectionBase::AddAndInitializeSubsystem(UClass* SubsystemClass)
+USubsystem* FSubsystemCollectionBase::AddAndInitializeSubsystem(UClass* SubsystemClass)
 {
 	if (!SubsystemMap.Contains(SubsystemClass))
 	{
@@ -203,12 +203,13 @@ bool FSubsystemCollectionBase::AddAndInitializeSubsystem(UClass* SubsystemClass)
 				Subsystem->InternalOwningSubsystem = this;
 				Subsystem->Initialize(*this);
 				
-				return true;
+				return Subsystem;
 			}
 		}
-		return false;
+		return nullptr;
 	}
-	return true;
+
+	return SubsystemMap.FindRef(SubsystemClass);
 }
 
 void FSubsystemCollectionBase::RemoveAndDeinitializeSubsystem(USubsystem* Subsystem)
diff --git a/Engine/Source/Runtime/Engine/Public/Subsystems/SubsystemCollection.h b/Engine/Source/Runtime/Engine/Public/Subsystems/SubsystemCollection.h
--- a/Engine/Source/Runtime/Engine/Public/Subsystems/SubsystemCollection.h
+++ b/Engine/Source/Runtime/Engine/Public/Subsystems/SubsystemCollection.h
@@ -22,7 +22,17 @@ public:
 	 * Only call from Initialize() of Systems to ensure initialization order
 	 * Note: Dependencies only work within a collection
 	 */
-	bool InitializeDependency(TSubclassOf<USubsystem> SubsystemClass);
+	USubsystem* InitializeDependency(TSubclassOf<USubsystem> SubsystemClass);
+
+	/**
+	 * Only call from Initialize() of Systems to ensure initialization order
+	 * Note: Dependencies only work within a collection
+	 */
+	template <typename TSubsystemClass>
+	TSubsystemClass* InitializeDependency()
+	{
+		return Cast<TSubsystemClass>(InitializeDependency(TSubsystemClass::StaticClass()));
+	}
 
 	/* FGCObject Interface */
 	virtual void AddReferencedObjects(FReferenceCollector& Collector) override;
@@ -45,7 +55,7 @@ protected:
 	const TSubclassOf<USubsystem>& GetBaseType() const { return BaseType; }
 
 private:
-	bool AddAndInitializeSubsystem(UClass* SubsystemClass);
+	USubsystem* AddAndInitializeSubsystem(UClass* SubsystemClass);
 
 	void RemoveAndDeinitializeSubsystem(USubsystem* Subsystem);
 
