Add identifiable build information to the latest material stats files when mirroring them.

diff --git a/Engine/Source/Runtime/Engine/Private/ShaderCompiler/ShaderCompiler.cpp b/Engine/Source/Runtime/Engine/Private/ShaderCompiler/ShaderCompiler.cpp
--- a/Engine/Source/Runtime/Engine/Private/ShaderCompiler/ShaderCompiler.cpp
+++ b/Engine/Source/Runtime/Engine/Private/ShaderCompiler/ShaderCompiler.cpp
@@ -1620,7 +1620,7 @@ void FShaderCompilerStats::WriteStats()
 		{
 			if(PlatformStats.IsValidIndex(Platform))
 			{
-				const ShaderCompilerStats& Stats = PlatformStats[Platform];
+				const ShaderCompilerStats& Stats = PlatformStats[Platform];	
 				for (const auto& Pair : Stats)
 				{
 					const FString& Path = Pair.Key;
@@ -1653,7 +1653,16 @@ void FShaderCompilerStats::WriteStats()
 				{
 					FParse::Value(FCommandLine::Get(), TEXT("targetplatform="), TargetType);
 				}
-				FString CopyLocation = FPaths::Combine(*MirrorLocation, FApp::GetProjectName(), *FApp::GetBranchName(), FString::Printf(TEXT("Stats-Latest(%s).csv"), *TargetType));
+				FString CopyLocation = FPaths::Combine(*MirrorLocation, FApp::GetProjectName(), *FApp::GetBranchName(), FString::Printf(TEXT("Stats-Latest-%d(%s).csv"), FEngineVersion::Current().GetChangelist() , *TargetType));
+				TArray <FString> ExistingFiles;
+				IFileManager::Get().FindFiles(ExistingFiles, *FPaths::Combine(*MirrorLocation, FApp::GetProjectName(), *FApp::GetBranchName()));
+				for (FString CurFile : ExistingFiles)
+				{
+					if (CurFile.Contains(FString::Printf(TEXT("(%s)"), *TargetType)))
+					{
+						IFileManager::Get().Delete(*FPaths::Combine(*MirrorLocation, FApp::GetProjectName(), *FApp::GetBranchName(), *CurFile), false, true);
+					}
+				}
 				IFileManager::Get().Copy(*CopyLocation, *FileName, true, true);
 			}
 		}
