Remove warning message when display name for material attribute changes, triggered when loading UMaterialExpressionGetMaterialAttributes

diff --git a/Engine/Source/Runtime/Engine/Private/Materials/MaterialExpressions.cpp b/Engine/Source/Runtime/Engine/Private/Materials/MaterialExpressions.cpp
--- a/Engine/Source/Runtime/Engine/Private/Materials/MaterialExpressions.cpp
+++ b/Engine/Source/Runtime/Engine/Private/Materials/MaterialExpressions.cpp
@@ -5874,24 +5874,8 @@ void UMaterialExpressionGetMaterialAttributes::PostLoad()
 
 	for (int i = 1; i < Outputs.Num(); ++i)
 	{
-		const FString DisplayName = FMaterialAttributeDefinitionMap::GetDisplayNameForMaterial(AttributeGetTypes[i-1],Material).ToString();
-		if (Outputs[i].OutputName.ToString() != DisplayName)
-		{
-			FString MaterialName;
-			if (Material)
-			{
-				Material->GetName(MaterialName);
-			}
-			else if (Function)
-			{
-				Function->GetName(MaterialName);
-			}
-
-			const FString AttributeName = FMaterialAttributeDefinitionMap::GetAttributeName(AttributeGetTypes[i - 1]);
-			UE_LOG(LogMaterial, Warning, TEXT("Material \"%s\", serialized attribute '%s' with OutputName '%s', expected '%s'"),
-				*MaterialName, *AttributeName, *(Outputs[i].OutputName.ToString()), *DisplayName);
-			Outputs[i].OutputName = *DisplayName;
-		}
+		const FString DisplayName = FMaterialAttributeDefinitionMap::GetDisplayNameForMaterial(AttributeGetTypes[i - 1], Material).ToString();
+		Outputs[i].OutputName = *DisplayName;
 	}
 }
 #endif // WITH_EDITOR
