Fix cook issue with pre smart name skeletons

diff --git a/Engine/Source/Runtime/Engine/Classes/Animation/SmartName.h b/Engine/Source/Runtime/Engine/Classes/Animation/SmartName.h
--- a/Engine/Source/Runtime/Engine/Classes/Animation/SmartName.h
+++ b/Engine/Source/Runtime/Engine/Classes/Animation/SmartName.h
@@ -184,7 +184,10 @@ struct ENGINE_API FSmartNameContainer
 	const FSmartNameMapping* GetContainer(FName ContainerName) const;
 
 	// Serialize this to the provided archive; required for TMap serialization
-	void Serialize(FArchive& Ar);
+	void Serialize(FArchive& Ar, bool bIsTemplate);
+
+	// Called after load (serialize itself may not be called if the USkeleton we are on is old enough)
+	void PostLoad();
 
 	friend FArchive& operator<<(FArchive& Ar, FSmartNameContainer& Elem);
 
diff --git a/Engine/Source/Runtime/Engine/Private/Animation/Skeleton.cpp b/Engine/Source/Runtime/Engine/Private/Animation/Skeleton.cpp
--- a/Engine/Source/Runtime/Engine/Private/Animation/Skeleton.cpp
+++ b/Engine/Source/Runtime/Engine/Private/Animation/Skeleton.cpp
@@ -162,6 +162,8 @@ void USkeleton::PostLoad()
 	// Cache smart name uids for animation curve names
 	IncreaseAnimCurveUidVersion();
 
+	SmartNames.PostLoad();
+
 	// refresh linked bone indices
 	FSmartNameMapping* CurveMappingTable = SmartNames.GetContainerInternal(USkeleton::AnimCurveMappingName);
 	if (CurveMappingTable)
@@ -248,7 +250,7 @@ void USkeleton::Serialize( FArchive& Ar )
 	// If we should be using smartnames, serialize the mappings
 	if(Ar.UE4Ver() >= VER_UE4_SKELETON_ADD_SMARTNAMES)
 	{
-		SmartNames.Serialize(Ar);
+		SmartNames.Serialize(Ar, IsTemplate());
 
 		AnimCurveMapping = SmartNames.GetContainerInternal(USkeleton::AnimCurveMappingName);
 	}
diff --git a/Engine/Source/Runtime/Engine/Private/Animation/SmartName.cpp b/Engine/Source/Runtime/Engine/Private/Animation/SmartName.cpp
--- a/Engine/Source/Runtime/Engine/Private/Animation/SmartName.cpp
+++ b/Engine/Source/Runtime/Engine/Private/Animation/SmartName.cpp
@@ -291,10 +291,10 @@ const FSmartNameMapping* FSmartNameContainer::GetContainer(FName ContainerName)
 	return NameMappings.Find(ContainerName);
 }
 
-void FSmartNameContainer::Serialize(FArchive& Ar)
+void FSmartNameContainer::Serialize(FArchive& Ar, bool bIsTemplate)
 {
 #if WITH_EDITORONLY_DATA
-	if (Ar.IsCooking())
+	if (Ar.IsCooking() && !bIsTemplate)
 	{
 		Ar << LoadedNameMappings;
 	}
@@ -303,13 +303,13 @@ void FSmartNameContainer::Serialize(FArchive& Ar)
 	{
 		Ar << NameMappings;
 	}
+}
 
+void FSmartNameContainer::PostLoad()
+{
 #if WITH_EDITORONLY_DATA
-	if (Ar.IsLoading())
-	{
-		//To preserve 
-		LoadedNameMappings = NameMappings;
-	}
+	// Preserve Load state for deterministic cooking
+	LoadedNameMappings = NameMappings;
 #endif
 }
 
