#UE4 Fix NaN in SubsurfaceProfile

diff --git a/Engine/Source/Runtime/Engine/Classes/Engine/SubsurfaceProfile.h b/Engine/Source/Runtime/Engine/Classes/Engine/SubsurfaceProfile.h
--- a/Engine/Source/Runtime/Engine/Classes/Engine/SubsurfaceProfile.h
+++ b/Engine/Source/Runtime/Engine/Classes/Engine/SubsurfaceProfile.h
@@ -119,24 +119,7 @@ struct FSubsurfaceProfileStruct
 
 	void Invalidate()
 	{
-		// nicer for debugging / VisualizeSSS
-		ScatterRadius = 0.0f;
-		SubsurfaceColor = FLinearColor(0, 0, 0);
-		FalloffColor = FLinearColor(0, 0, 0);
-		BoundaryColorBleed = FLinearColor::White;
-		ExtinctionScale = 1.0f;
-		ScatteringDistribution = 0.93f;
-		NormalScale = 0.08f;
-		IOR = 1.55f;
-		Roughness0 = 0.75f;
-		Roughness1 = 1.30f;
-		LobeMix = 0.85f;
-		bEnableBurley = false;
-		SurfaceAlbedo = FLinearColor(0, 0, 0);
-		MeanFreePathColor = FLinearColor(0, 0, 0);
-		MeanFreePathDistance = 0.0f;
-		WorldUnitScale = 0.1f;
-		TransmissionTintColor = FLinearColor(0, 0, 0);
+		*this = FSubsurfaceProfileStruct();
 	}
 };
 
