Removed namedevent overhead from FStaticComponentRenderAssetManager_IncrementalUpdate

diff --git a/Engine/Source/Runtime/Engine/Private/Streaming/LevelTextureManager.cpp b/Engine/Source/Runtime/Engine/Private/Streaming/LevelTextureManager.cpp
--- a/Engine/Source/Runtime/Engine/Private/Streaming/LevelTextureManager.cpp
+++ b/Engine/Source/Runtime/Engine/Private/Streaming/LevelTextureManager.cpp
@@ -9,8 +9,6 @@
 #include "Engine/Texture2D.h"
 #include "Engine/World.h"
 
-DECLARE_CYCLE_STAT(TEXT("Render Asset Incremental Update"), FStaticComponentRenderAssetManager_IncrementalUpdate, STATGROUP_StreamingDetailsVerbose);
-
 FLevelRenderAssetManager::FLevelRenderAssetManager(ULevel* InLevel, RenderAssetInstanceTask::FDoWorkTask& AsyncTask)
 	: Level(InLevel)
 	, bIsInitialized(false)
@@ -272,8 +270,6 @@ void FLevelRenderAssetManager::IncrementalUpdate(
 	float Percentage, 
 	bool bUseDynamicStreaming) 
 {
-	SCOPE_CYCLE_COUNTER(FStaticComponentRenderAssetManager_IncrementalUpdate);
-
 	check(Level);
 
 	if (NeedsIncrementalBuild(NumStepsLeftForIncrementalBuild))
