Reducing namedevent overhead from STAT_UParticleSystemComponent_LOD_Inactive

diff --git a/Engine/Source/Runtime/Engine/Private/Particles/ParticleComponents.cpp b/Engine/Source/Runtime/Engine/Private/Particles/ParticleComponents.cpp
--- a/Engine/Source/Runtime/Engine/Private/Particles/ParticleComponents.cpp
+++ b/Engine/Source/Runtime/Engine/Private/Particles/ParticleComponents.cpp
@@ -106,7 +106,6 @@ DECLARE_CYCLE_STAT(TEXT("ParticleComponent UpdateDynamicData GT"), STAT_Particle
 DECLARE_CYCLE_STAT(TEXT("ParticleComponent OrientZAxisTowardCamera GT"), STAT_UParticleSystemComponent_OrientZAxisTowardCamera, STATGROUP_Particles);
 DECLARE_CYCLE_STAT(TEXT("ParticleComponent QueueFinalize GT"), STAT_UParticleSystemComponent_QueueFinalize, STATGROUP_Particles);
 DECLARE_CYCLE_STAT(TEXT("ParticleComponent CheckForReset GT"), STAT_UParticleSystemComponent_CheckForReset, STATGROUP_Particles);
-DECLARE_CYCLE_STAT(TEXT("ParticleComponent LOD_Inactive GT"), STAT_UParticleSystemComponent_LOD_Inactive, STATGROUP_Particles);
 DECLARE_CYCLE_STAT(TEXT("ParticleComponent LOD GT"), STAT_UParticleSystemComponent_LOD, STATGROUP_Particles);
 DECLARE_CYCLE_STAT(TEXT("ParticleComponent QueueTasksGT"), STAT_UParticleSystemComponent_QueueTasks, STATGROUP_Particles);
 DECLARE_CYCLE_STAT(TEXT("ParticleComponent QueueAsyncGT"), STAT_UParticleSystemComponent_QueueAsync, STATGROUP_Particles);
@@ -3563,7 +3562,6 @@ bool UParticleSystemComponent::CanSkipTickDueToVisibility()
 {
 	if (Template && Template->IsLooping() && CanConsiderInvisible() && !bWasDeactivated)
 	{
-		SCOPE_CYCLE_COUNTER(STAT_UParticleSystemComponent_LOD_Inactive);
 		bForcedInActive = true;
 		SpawnEvents.Empty();
 		DeathEvents.Empty();
