Reducing namedevent overhead from STAT_ComponentUpdateBounds

diff --git a/Engine/Source/Runtime/Engine/Private/Components/SceneComponent.cpp b/Engine/Source/Runtime/Engine/Private/Components/SceneComponent.cpp
--- a/Engine/Source/Runtime/Engine/Private/Components/SceneComponent.cpp
+++ b/Engine/Source/Runtime/Engine/Private/Components/SceneComponent.cpp
@@ -49,7 +49,7 @@ DEFINE_LOG_CATEGORY_STATIC(LogSceneComponent, Log, All);
 
 DECLARE_CYCLE_STAT(TEXT("UpdateComponentToWorld"), STAT_UpdateComponentToWorld, STATGROUP_Component);
 DECLARE_CYCLE_STAT(TEXT("UpdateChildTransforms"), STAT_UpdateChildTransforms, STATGROUP_Component);
-DECLARE_CYCLE_STAT(TEXT("Component UpdateBounds"), STAT_ComponentUpdateBounds, STATGROUP_Component);
+DECLARE_CYCLE_STAT(TEXT("Component CalcBounds"), STAT_ComponentCalcBounds, STATGROUP_Component);
 DECLARE_CYCLE_STAT(TEXT("Component UpdateNavData"), STAT_ComponentUpdateNavData, STATGROUP_Component);
 DECLARE_CYCLE_STAT(TEXT("Component PostUpdateNavData"), STAT_ComponentPostUpdateNavData, STATGROUP_Component);
 
@@ -1185,8 +1185,6 @@ void USceneComponent::CalcBoundingCylinder(float& CylinderRadius, float& Cylinde
 
 void USceneComponent::UpdateBounds()
 {
-	SCOPE_CYCLE_COUNTER(STAT_ComponentUpdateBounds);
-
 	// if use parent bound if attach parent exists, and the flag is set
 	// since parents tick first before child, this should work correctly
 	if ( bUseAttachParentBound && GetAttachParent() != nullptr )
@@ -1195,6 +1193,7 @@ void USceneComponent::UpdateBounds()
 	}
 	else
 	{
+		SCOPE_CYCLE_COUNTER(STAT_ComponentCalcBounds);
 		// Calculate new bounds
 		Bounds = CalcBounds(GetComponentTransform());
 	}
