#UE4 Allow non-threading builds to FlushLevelStreaming. Level streaming does not depend on threads.

diff --git a/Engine/Source/Runtime/Engine/Private/World.cpp b/Engine/Source/Runtime/Engine/Private/World.cpp
--- a/Engine/Source/Runtime/Engine/Private/World.cpp
+++ b/Engine/Source/Runtime/Engine/Private/World.cpp
@@ -3477,11 +3477,6 @@ void UWorld::UpdateStreamingLevelPriority(ULevelStreaming* StreamingLevel)
 
 void UWorld::FlushLevelStreaming(EFlushLevelStreamingType FlushType)
 {
-	if (!FPlatformProcess::SupportsMultithreading())
-	{
-		return;
-	}
-
 	AWorldSettings* WorldSettings = GetWorldSettings();
 
 	TGuardValue<EFlushLevelStreamingType> FlushingLevelStreamingGuard(FlushLevelStreamingType, FlushType);
@@ -3568,10 +3563,6 @@ void UWorld::ConditionallyBuildStreamingData()
 
 bool UWorld::IsVisibilityRequestPending() const
 {
-	if (!FPlatformProcess::SupportsMultithreading())
-	{
-		return false;
-	}
 	return (CurrentLevelPendingVisibility != nullptr || CurrentLevelPendingInvisibility != nullptr);
 }
 
