Fix the phys interface inconsistencies between chaos and physx introduced in 11499590

diff --git a/Engine/Source/Runtime/Engine/Private/PhysicsEngine/PhysScene_PhysX.cpp b/Engine/Source/Runtime/Engine/Private/PhysicsEngine/PhysScene_PhysX.cpp
--- a/Engine/Source/Runtime/Engine/Private/PhysicsEngine/PhysScene_PhysX.cpp
+++ b/Engine/Source/Runtime/Engine/Private/PhysicsEngine/PhysScene_PhysX.cpp
@@ -1440,7 +1440,7 @@ void FPhysScene_PhysX::DispatchPhysNotifications_AssumesLocked()
 
 }
 
-void FPhysScene_PhysX::SetUpForFrame(const FVector* NewGrav, float InDeltaSeconds, float InMaxPhysicsDeltaTime)
+void FPhysScene_PhysX::SetUpForFrame(const FVector* NewGrav, float InDeltaSeconds, float InMaxPhysicsDeltaTime, float InMaxSubstepDeltaTime, int32 InMaxSubsteps)
 {
 	DeltaSeconds = InDeltaSeconds;
 	MaxPhysicsDeltaTime = InMaxPhysicsDeltaTime;
diff --git a/Engine/Source/Runtime/Engine/Public/Physics/PhysScene_PhysX.h b/Engine/Source/Runtime/Engine/Public/Physics/PhysScene_PhysX.h
--- a/Engine/Source/Runtime/Engine/Public/Physics/PhysScene_PhysX.h
+++ b/Engine/Source/Runtime/Engine/Public/Physics/PhysScene_PhysX.h
@@ -167,7 +167,7 @@ public:
 	void ApplyWorldOffset(FVector InOffset);
 
 	/** Set the gravity and timing of all physics scenes */
-	ENGINE_API void SetUpForFrame(const FVector* NewGrav, float InDeltaSeconds = 0.0f, float InMaxPhysicsDeltaTime = 0.0f);
+	ENGINE_API void SetUpForFrame(const FVector* NewGrav, float InDeltaSeconds = 0.0f, float InMaxPhysicsDeltaTime = 0.0f, float InMaxSubstepDeltaTime = 0.f, int32 InMaxSubsteps = 1);
 
 	/** Starts a frame */
 	ENGINE_API void StartFrame();
