- Fix when PerSMInstanceData reaches 0 and BuildTreeIfOutdated is called that we clear the render instance data (actual jira fix)
- Fix HISM Undo/Redo which was crashing after the fix because of de-sync between CPU and GPU
- Make sure BuildTreeIfOutdated leaves HISM in a non-outdated state

diff --git a/Engine/Source/Runtime/Engine/Classes/Components/HierarchicalInstancedStaticMeshComponent.h b/Engine/Source/Runtime/Engine/Classes/Components/HierarchicalInstancedStaticMeshComponent.h
--- a/Engine/Source/Runtime/Engine/Classes/Components/HierarchicalInstancedStaticMeshComponent.h
+++ b/Engine/Source/Runtime/Engine/Classes/Components/HierarchicalInstancedStaticMeshComponent.h
@@ -204,6 +204,7 @@ public:
 	virtual FPrimitiveSceneProxy* CreateSceneProxy() override;
 	virtual FBoxSphereBounds CalcBounds(const FTransform& BoundTransform) const override;
 #if WITH_EDITOR
+	virtual void PostEditUndo() override;
 	virtual void PostEditChangeChainProperty(FPropertyChangedChainEvent& PropertyChangedEvent) override;
 #endif
 
diff --git a/Engine/Source/Runtime/Engine/Private/HierarchicalInstancedStaticMesh.cpp b/Engine/Source/Runtime/Engine/Private/HierarchicalInstancedStaticMesh.cpp
--- a/Engine/Source/Runtime/Engine/Private/HierarchicalInstancedStaticMesh.cpp
+++ b/Engine/Source/Runtime/Engine/Private/HierarchicalInstancedStaticMesh.cpp
@@ -1927,6 +1927,16 @@ UHierarchicalInstancedStaticMeshComponent::~UHierarchicalInstancedStaticMeshComp
 }
 
 #if WITH_EDITOR
+void UHierarchicalInstancedStaticMeshComponent::PostEditUndo()
+{
+	Super::PostEditUndo();
+
+	const bool bAsync = false;
+	const bool bForceUpdate = true;
+	// Force because the Outdated Condition will fail (compared values will match)
+	BuildTreeIfOutdated(bAsync, bForceUpdate);
+}
+
 void UHierarchicalInstancedStaticMeshComponent::PostEditChangeChainProperty(FPropertyChangedChainEvent& PropertyChangedEvent)
 {
 	if ((PropertyChangedEvent.Property != NULL && PropertyChangedEvent.Property->GetFName() == "InstancingRandomSeed"))
@@ -2663,10 +2673,17 @@ void UHierarchicalInstancedStaticMeshComponent::ApplyEmpty()
 	InstanceCountToRender = 0;
 	InstanceReorderTable.Empty();
 	SortedInstances.Empty();
-	CacheMeshExtendedBounds = FBoxSphereBounds(ForceInitToZero);
-
 	UnbuiltInstanceBoundsList.Empty();
 	BuiltInstanceBounds.Init();
+	CacheMeshExtendedBounds = GetStaticMesh() && GetStaticMesh()->HasValidRenderData() ? GetStaticMesh()->GetBounds() : FBoxSphereBounds(ForceInitToZero);
+	InstanceUpdateCmdBuffer.Reset();
+	if (PerInstanceRenderData.IsValid())
+	{
+		TUniquePtr<FStaticMeshInstanceData> BuiltInstanceData = MakeUnique<FStaticMeshInstanceData>(GVertexElementTypeSupport.IsSupported(VET_Half2));
+		PerInstanceRenderData->UpdateFromPreallocatedData(*BuiltInstanceData);
+		PerInstanceRenderData->HitProxies.Empty();
+		MarkRenderStateDirty();
+	}
 }
 
 void UHierarchicalInstancedStaticMeshComponent::ApplyBuildTree(FClusterBuilder& Builder)
