Fix ValueSet of PerSMInstanceData not being handled (resetting element would not update the ISM)

diff --git a/Engine/Source/Runtime/Engine/Private/InstancedStaticMesh.cpp b/Engine/Source/Runtime/Engine/Private/InstancedStaticMesh.cpp
--- a/Engine/Source/Runtime/Engine/Private/InstancedStaticMesh.cpp
+++ b/Engine/Source/Runtime/Engine/Private/InstancedStaticMesh.cpp
@@ -2999,6 +2999,10 @@ void UInstancedStaticMeshComponent::PostEditChangeChainProperty(FPropertyChanged
 			{
 				ClearInstances();
 			}
+			else if (PropertyChangedEvent.ChangeType == EPropertyChangeType::ValueSet)
+			{
+				InstanceUpdateCmdBuffer.Edit();
+			}
 			
 			MarkRenderStateDirty();
 		}
