Prevent two full FlushLevelStreaming calls from running into eachother. Fairly easy to cause that with map load in the editor

diff --git a/Engine/Source/Runtime/Engine/Private/World.cpp b/Engine/Source/Runtime/Engine/Private/World.cpp
--- a/Engine/Source/Runtime/Engine/Private/World.cpp
+++ b/Engine/Source/Runtime/Engine/Private/World.cpp
@@ -3477,8 +3477,14 @@ void UWorld::UpdateStreamingLevelPriority(ULevelStreaming* StreamingLevel)
 
 void UWorld::FlushLevelStreaming(EFlushLevelStreamingType FlushType)
 {
-	AWorldSettings* WorldSettings = GetWorldSettings();
+	// We're already full flushing, no reason to flush again
+	if (FlushType == EFlushLevelStreamingType::Full && FlushLevelStreamingType == EFlushLevelStreamingType::Full)
+	{
+		return;
+	}
 
+	AWorldSettings* WorldSettings = GetWorldSettings();
+	
 	TGuardValue<EFlushLevelStreamingType> FlushingLevelStreamingGuard(FlushLevelStreamingType, FlushType);
 
 	// Update internals with current loaded/ visibility flags.
