Fix unexpected state in ULevelStreaming::UpdateStreamingState

The FLevelStreamingWrapper::operator< relies on the streaming level priority for its comparison. The issue is that TMap using this class as a key relied on that operator too for the initial insertion, and since the priority of items already in the map could change afterward, it would break the ordering and prevent proper indexing (for example, Find, Contains, Removed could fail to find items that were actually in the map).

We can actually simply use a map of ULevelStreaming* for the LevelsToProcess map. The insertion in the priority sorted StreamingLevels array is performed using Algo::LowerBound() anyway.

diff --git a/Engine/Source/Runtime/Engine/Classes/Engine/World.h b/Engine/Source/Runtime/Engine/Classes/Engine/World.h
--- a/Engine/Source/Runtime/Engine/Classes/Engine/World.h
+++ b/Engine/Source/Runtime/Engine/Classes/Engine/World.h
@@ -756,31 +756,6 @@ private:
 	int32 SavedTickingCollectionIndex;
 };
 
-USTRUCT()
-struct FLevelStreamingWrapper
-{
-	GENERATED_BODY()
-
-	FLevelStreamingWrapper()
-		: StreamingLevel(nullptr)
-	{}
-
-	FLevelStreamingWrapper(ULevelStreaming* InStreamingLevel)
-		: StreamingLevel(InStreamingLevel)
-	{}
-
-	ULevelStreaming*& Get() { return StreamingLevel; }
-	ULevelStreaming* Get() const { return StreamingLevel; }
-
-
-	bool operator<(const FLevelStreamingWrapper& Other) const;
-	bool operator==(const FLevelStreamingWrapper& Other) const { return StreamingLevel == Other.StreamingLevel; }
-
-private:
-	UPROPERTY()
-	ULevelStreaming* StreamingLevel;
-};
-
 USTRUCT()
 struct FStreamingLevelsToConsider
 {
@@ -794,7 +769,7 @@ private:
 
 	/** Priority sorted array of streaming levels actively being considered. */
 	UPROPERTY()
-	TArray<FLevelStreamingWrapper> StreamingLevels;
+	TArray<ULevelStreaming*> StreamingLevels;
 
 	enum class EProcessReason : uint8
 	{
@@ -803,7 +778,7 @@ private:
 	};
 
 	/** Streaming levels that had their priority changed or were added to the container while consideration was underway. */
-	TSortedMap<FLevelStreamingWrapper, EProcessReason> LevelsToProcess;
+	TMap<ULevelStreaming*, EProcessReason> LevelsToProcess;
 
 	/** Whether the streaming levels are under active consideration or not */
 	bool bStreamingLevelsBeingConsidered;
@@ -816,7 +791,7 @@ private:
 
 public:
 
-	const TArray<FLevelStreamingWrapper>& GetStreamingLevels() const { return StreamingLevels; }
+	const TArray<ULevelStreaming*>& GetStreamingLevels() const { return StreamingLevels; }
 
 	void AddReferencedObjects(UObject* InThis, FReferenceCollector& Collector);
 
diff --git a/Engine/Source/Runtime/Engine/Private/World.cpp b/Engine/Source/Runtime/Engine/Private/World.cpp
--- a/Engine/Source/Runtime/Engine/Private/World.cpp
+++ b/Engine/Source/Runtime/Engine/Private/World.cpp
@@ -3091,24 +3091,6 @@ UWorld* UWorld::DuplicateWorldForPIE(const FString& PackageName, UWorld* OwningW
 	return PIELevelWorld;
 }
 
-bool FLevelStreamingWrapper::operator<(const FLevelStreamingWrapper& Other) const
-{
-	if (StreamingLevel && Other.StreamingLevel)
-	{
-		const int32 Priority = StreamingLevel->GetPriority();
-		const int32 OtherPriority = Other.StreamingLevel->GetPriority();
-
-		if (Priority == OtherPriority)
-		{
-			return ((UPTRINT)StreamingLevel < (UPTRINT)Other.StreamingLevel);
-		}
-
-		return (Priority < OtherPriority);
-	}
-
-	return (StreamingLevel != nullptr);
-}
-
 void FStreamingLevelsToConsider::Add_Internal(ULevelStreaming* StreamingLevel, bool bGuaranteedNotInContainer)
 {
 	if (StreamingLevel)
@@ -3123,10 +3105,27 @@ void FStreamingLevelsToConsider::Add_Internal(ULevelStreaming* StreamingLevel, b
 		}
 		else
 		{
-			FLevelStreamingWrapper WrappedLevel(StreamingLevel);
-			if (bGuaranteedNotInContainer || !StreamingLevels.Contains(WrappedLevel))
+			if (bGuaranteedNotInContainer || !StreamingLevels.Contains(StreamingLevel))
 			{
-				StreamingLevels.Insert(WrappedLevel, Algo::LowerBound(StreamingLevels, WrappedLevel));
+				auto PrioritySort = [](ULevelStreaming* StreamingLevel, ULevelStreaming* OtherStreamingLevel)
+				{
+					if (StreamingLevel && OtherStreamingLevel)
+					{
+						const int32 Priority = StreamingLevel->GetPriority();
+						const int32 OtherPriority = OtherStreamingLevel->GetPriority();
+
+						if (Priority == OtherPriority)
+						{
+							return ((UPTRINT)StreamingLevel < (UPTRINT)OtherStreamingLevel);
+						}
+
+						return (Priority < OtherPriority);
+					}
+
+					return (StreamingLevel != nullptr);
+				};
+
+				StreamingLevels.Insert(StreamingLevel, Algo::LowerBound(StreamingLevels, StreamingLevel, PrioritySort));
 			}
 		}
 	}
@@ -3161,7 +3160,7 @@ void FStreamingLevelsToConsider::RemoveAt(const int32 Index)
 {
 	if (bStreamingLevelsBeingConsidered)
 	{
-		if (ULevelStreaming* StreamingLevel = StreamingLevels[Index].Get())
+		if (ULevelStreaming* StreamingLevel = StreamingLevels[Index])
 		{
 			LevelsToProcess.Remove(StreamingLevel);
 		}
@@ -3177,10 +3176,9 @@ void FStreamingLevelsToConsider::Reevaluate(ULevelStreaming* StreamingLevel)
 		{
 			// If the streaming level is already in the map then it doesn't need to be updated as it is either
 			// already Reevaluate or the more significant Add
-			FLevelStreamingWrapper WrappedLevel(StreamingLevel);
-			if (!LevelsToProcess.Contains(WrappedLevel))
+			if (!LevelsToProcess.Contains(StreamingLevel))
 			{
-				LevelsToProcess.Add(WrappedLevel, EProcessReason::Reevaluate);
+				LevelsToProcess.Add(StreamingLevel, EProcessReason::Reevaluate);
 			}
 		}
 		else
@@ -3219,10 +3217,10 @@ void FStreamingLevelsToConsider::EndConsideration()
 	{
 		// For any streaming level that was added or had its priority changed while we were considering the
 		// streaming levels go through and ensure they are correctly in the map and sorted to the correct location
-		TSortedMap<FLevelStreamingWrapper, EProcessReason> LevelsToProcessCopy = MoveTemp(LevelsToProcess);
-		for (const TPair<FLevelStreamingWrapper,EProcessReason>& LevelToProcessPair : LevelsToProcessCopy)
+		TMap<ULevelStreaming*, EProcessReason> LevelsToProcessCopy = MoveTemp(LevelsToProcess);
+		for (const TPair<ULevelStreaming*,EProcessReason>& LevelToProcessPair : LevelsToProcessCopy)
 		{
-			if (ULevelStreaming* StreamingLevel = LevelToProcessPair.Key.Get())
+			if (ULevelStreaming* StreamingLevel = LevelToProcessPair.Key)
 			{
 				// Remove the level if it is already in the list so we can use Add to place in the correct priority location
 				const bool bIsBeingConsidered = Remove(StreamingLevel);
@@ -3239,9 +3237,9 @@ void FStreamingLevelsToConsider::EndConsideration()
 
 void FStreamingLevelsToConsider::AddReferencedObjects(UObject* InThis, FReferenceCollector& Collector)
 {
-	for (TPair<FLevelStreamingWrapper, EProcessReason>& LevelToProcessPair : LevelsToProcess)
+	for (TPair<ULevelStreaming*, EProcessReason>& LevelToProcessPair : LevelsToProcess)
 	{
-		if (ULevelStreaming*& StreamingLevel = LevelToProcessPair.Key.Get())
+		if (ULevelStreaming*& StreamingLevel = LevelToProcessPair.Key)
 		{
 			Collector.AddReferencedObject(StreamingLevel, InThis);
 		}
@@ -3267,7 +3265,7 @@ void UWorld::UpdateLevelStreaming()
 
 	for (int32 Index = StreamingLevelsToConsider.GetStreamingLevels().Num() - 1; Index >= 0; --Index)
 	{
-		if (ULevelStreaming* StreamingLevel = StreamingLevelsToConsider.GetStreamingLevels()[Index].Get())
+		if (ULevelStreaming* StreamingLevel = StreamingLevelsToConsider.GetStreamingLevels()[Index])
 		{
 			bool bUpdateAgain = true;
 			bool bShouldContinueToConsider = true;
