Fix client crash with "BuildingFoundations with the same decorated level name added to streaminglevels.

If we find a duplicated streamed level that is currently unloading, wait for it to unload instead of falling back to the FailedToLoad state.

diff --git a/Engine/Source/Runtime/Engine/Private/LevelStreaming.cpp b/Engine/Source/Runtime/Engine/Private/LevelStreaming.cpp
--- a/Engine/Source/Runtime/Engine/Private/LevelStreaming.cpp
+++ b/Engine/Source/Runtime/Engine/Private/LevelStreaming.cpp
@@ -902,36 +902,34 @@ bool ULevelStreaming::RequestLevel(UWorld* PersistentWorld, bool bAllowLevelLoad
 		return false;
 	}
 
-	auto ValidateUniqueLevel = [this, PersistentWorld]()
+	// Validate that our new streaming level is unique, check for clash with currently loaded streaming levels
+	for (ULevelStreaming* OtherLevel : PersistentWorld->GetStreamingLevels())
 	{
-		for (ULevelStreaming* OtherLevel : PersistentWorld->GetStreamingLevels())
+		if (OtherLevel == nullptr || OtherLevel == this)
 		{
-			if (OtherLevel == nullptr || OtherLevel == this)
-			{
-				continue;
-			}
+			continue;
+		}
+
+		const ECurrentState OtherState = OtherLevel->GetCurrentState();
+		if (OtherState == ECurrentState::FailedToLoad || OtherState == ECurrentState::Removed || (OtherState == ECurrentState::Unloaded && (OtherLevel->TargetState == ETargetState::Unloaded || OtherLevel->TargetState == ETargetState::UnloadedAndRemoved)))
+		{
+			// If the other level isn't loaded or in the process of being loaded we don't need to consider it
+			continue;
+		}
 
-			const ECurrentState OtherState = OtherLevel->GetCurrentState();
-			if (OtherState == ECurrentState::FailedToLoad || OtherState == ECurrentState::Removed || (OtherState == ECurrentState::Unloaded && (OtherLevel->TargetState == ETargetState::Unloaded || OtherLevel->TargetState == ETargetState::UnloadedAndRemoved)))
+		if (OtherLevel->WorldAsset == WorldAsset)
+		{ 
+			if (OtherLevel->GetIsRequestingUnloadAndRemoval())
 			{
-				// If the other level isn't loaded or in the process of being loaded we don't need to consider it
-				continue;
+				return false; // Cannot load new level now, retry until the OtherLevel is done unloading
 			}
-
-			if (OtherLevel->WorldAsset == WorldAsset)
+			else
 			{
 				UE_LOG(LogLevelStreaming, Warning, TEXT("Streaming Level '%s' uses same destination for level ('%s') as '%s'. Level cannot be loaded again and this StreamingLevel will be flagged as failed to load."), *GetPathName(), *WorldAsset.GetLongPackageName(), *OtherLevel->GetPathName());
+				CurrentState = ECurrentState::FailedToLoad;
 				return false;
 			}
 		}
-
-		return true;
-	};
-
-	if (!ValidateUniqueLevel())
-	{
-		CurrentState = ECurrentState::FailedToLoad;
-		return false;
 	}
 
 	TRACE_LOADTIME_REQUEST_GROUP_SCOPE(TEXT("LevelStreaming - %s"), *GetPathName());
