Added CVS format capability to OBJ LIST console command

#rnx

diff --git a/Engine/Source/Runtime/Engine/Private/UnrealEngine.cpp b/Engine/Source/Runtime/Engine/Private/UnrealEngine.cpp
--- a/Engine/Source/Runtime/Engine/Private/UnrealEngine.cpp
+++ b/Engine/Source/Runtime/Engine/Private/UnrealEngine.cpp
@@ -7524,6 +7524,7 @@ bool UEngine::HandleObjCommand( const TCHAR* Cmd, FOutputDevice& Ar )
 
 		const bool bAlphaSort = FParse::Param( Cmd, TEXT("ALPHASORT") );
 		const bool bCountSort = FParse::Param( Cmd, TEXT("COUNTSORT") );
+		const bool bCSV = FParse::Param(Cmd, TEXT("CSV"));
 
 		if( Objects.Num() )
 		{
@@ -7546,34 +7547,60 @@ bool UEngine::HandleObjCommand( const TCHAR* Cmd, FOutputDevice& Ar )
 			};
 			Objects.Sort( FCompareFSubItem( bAlphaSort ) );
 
-			Ar.Logf( 
-				TEXT("%140s %10s %10s %10s %15s %15s %15s %15s %15s"), 
-				TEXT("Object"), 
-				TEXT("NumKB"), 
-				TEXT("MaxKB"), 
-				TEXT("ResExcKB"),
-				TEXT("ResExcDedSysKB"),
-				TEXT("ResExcShrSysKB"),
-				TEXT("ResExcDedVidKB"),
-				TEXT("ResExcShrVidKB"),
-				TEXT("ResExcUnkKB")
-			);
-
-			for (const FSubItem& ObjItem : Objects)
+			if (bCSV)
+			{
+				Ar.Logf(TEXT("Object, NumKB, MaxKB, ResExcKB, ResExcDedSysKB, ResExcShrSysKB, ResExcDedVidKB, ResExcShrVidKB, ResExcUnkKB"));
+			}
+			else
 			{
 				Ar.Logf(
-					TEXT("%140s %10.2f %10.2f %10.2f %15.2f %15.2f %15.2f %15.2f %15.2f"), 
-					*ObjItem.Object->GetFullName(), 
-					ObjItem.Num / 1024.0f, 
-					ObjItem.Max / 1024.0f, 
-					ObjItem.TrueResourceSize.GetTotalMemoryBytes() / 1024.0f, 
-					ObjItem.TrueResourceSize.GetDedicatedSystemMemoryBytes() / 1024.0f, 
-					ObjItem.TrueResourceSize.GetSharedSystemMemoryBytes() / 1024.0f, 
-					ObjItem.TrueResourceSize.GetDedicatedVideoMemoryBytes() / 1024.0f, 
-					ObjItem.TrueResourceSize.GetSharedVideoMemoryBytes() / 1024.0f, 
-					ObjItem.TrueResourceSize.GetUnknownMemoryBytes() / 1024.0f
+					TEXT("%140s %10s %10s %10s %15s %15s %15s %15s %15s"),
+					TEXT("Object"),
+					TEXT("NumKB"),
+					TEXT("MaxKB"),
+					TEXT("ResExcKB"),
+					TEXT("ResExcDedSysKB"),
+					TEXT("ResExcShrSysKB"),
+					TEXT("ResExcDedVidKB"),
+					TEXT("ResExcShrVidKB"),
+					TEXT("ResExcUnkKB")
 				);
 			}
+
+			for (const FSubItem& ObjItem : Objects)
+			{
+				if (bCSV)
+				{
+					Ar.Logf(
+						TEXT("%s, %.2f, %.2f, %.2f, %.2f, %.2f, %.2f, %.2f, %.2f"),
+						*ObjItem.Object->GetFullName(),
+						ObjItem.Num / 1024.0f,
+						ObjItem.Max / 1024.0f,
+						ObjItem.TrueResourceSize.GetTotalMemoryBytes() / 1024.0f,
+						ObjItem.TrueResourceSize.GetDedicatedSystemMemoryBytes() / 1024.0f,
+						ObjItem.TrueResourceSize.GetSharedSystemMemoryBytes() / 1024.0f,
+						ObjItem.TrueResourceSize.GetDedicatedVideoMemoryBytes() / 1024.0f,
+						ObjItem.TrueResourceSize.GetSharedVideoMemoryBytes() / 1024.0f,
+						ObjItem.TrueResourceSize.GetUnknownMemoryBytes() / 1024.0f
+					);
+
+				}
+				else
+				{
+					Ar.Logf(
+						TEXT("%140s %10.2f %10.2f %10.2f %15.2f %15.2f %15.2f %15.2f %15.2f"),
+						*ObjItem.Object->GetFullName(),
+						ObjItem.Num / 1024.0f,
+						ObjItem.Max / 1024.0f,
+						ObjItem.TrueResourceSize.GetTotalMemoryBytes() / 1024.0f,
+						ObjItem.TrueResourceSize.GetDedicatedSystemMemoryBytes() / 1024.0f,
+						ObjItem.TrueResourceSize.GetSharedSystemMemoryBytes() / 1024.0f,
+						ObjItem.TrueResourceSize.GetDedicatedVideoMemoryBytes() / 1024.0f,
+						ObjItem.TrueResourceSize.GetSharedVideoMemoryBytes() / 1024.0f,
+						ObjItem.TrueResourceSize.GetUnknownMemoryBytes() / 1024.0f
+					);
+				}
+			}
 			Ar.Log(TEXT(""));
 		}
 
@@ -7592,8 +7619,7 @@ bool UEngine::HandleObjCommand( const TCHAR* Cmd, FOutputDevice& Ar )
 				}
 			};
 			List.Sort( FCompareFItem( bAlphaSort, bCountSort ) );
-			const bool bCSV = FParse::Param(Cmd, TEXT("CSV"));
-
+			
 			if (bCSV)
 			{
 				Ar.Logf(TEXT(", Class, Count, NumKB, MaxKB, ResExcKB, ResExcDedSysKB, ResExcShrSysKB, ResExcDedVidKB, ResExcShrVidKB, ResExcUnkKB"));
