Fix crash in GetMaterialFromInternalFaceIndexAndHitLocation if BodySetup is garbage collected during level save.

diff --git a/Engine/Source/Runtime/Engine/Private/PhysicsEngine/Experimental/PhysInterface_Chaos.cpp b/Engine/Source/Runtime/Engine/Private/PhysicsEngine/Experimental/PhysInterface_Chaos.cpp
--- a/Engine/Source/Runtime/Engine/Private/PhysicsEngine/Experimental/PhysInterface_Chaos.cpp
+++ b/Engine/Source/Runtime/Engine/Private/PhysicsEngine/Experimental/PhysInterface_Chaos.cpp
@@ -137,7 +137,7 @@ Chaos::FChaosPhysicsMaterial* GetMaterialFromInternalFaceIndexAndHitLocation(con
 						BodySetup = BodyInst->BodySetup.Get();	//this data should be immutable at runtime so ok to check from worker thread.
 						ECollisionShapeType GeomType = GetGeometryType(Shape);
 
-						if (BodySetup->bSupportUVsAndFaceRemap && GetGeometryType(Shape) == ECollisionShapeType::Trimesh)
+						if (BodySetup && BodySetup->bSupportUVsAndFaceRemap && GetGeometryType(Shape) == ECollisionShapeType::Trimesh)
 						{
 							FVector Scale(1.0f, 1.0f, 1.0f);
 							const Chaos::FImplicitObject* Geometry = Shape.Geometry.Get();
