#UE4 Fix shadow variable

diff --git a/Engine/Source/Runtime/Engine/Private/World.cpp b/Engine/Source/Runtime/Engine/Private/World.cpp
--- a/Engine/Source/Runtime/Engine/Private/World.cpp
+++ b/Engine/Source/Runtime/Engine/Private/World.cpp
@@ -3107,22 +3107,22 @@ void FStreamingLevelsToConsider::Add_Internal(ULevelStreaming* StreamingLevel, b
 		{
 			if (bGuaranteedNotInContainer || !StreamingLevels.Contains(StreamingLevel))
 			{
-				auto PrioritySort = [](ULevelStreaming* StreamingLevel, ULevelStreaming* OtherStreamingLevel)
+				auto PrioritySort = [](ULevelStreaming* LambdaStreamingLevel, ULevelStreaming* OtherStreamingLevel)
 				{
-					if (StreamingLevel && OtherStreamingLevel)
+					if (LambdaStreamingLevel && OtherStreamingLevel)
 					{
-						const int32 Priority = StreamingLevel->GetPriority();
+						const int32 Priority = LambdaStreamingLevel->GetPriority();
 						const int32 OtherPriority = OtherStreamingLevel->GetPriority();
 
 						if (Priority == OtherPriority)
 						{
-							return ((UPTRINT)StreamingLevel < (UPTRINT)OtherStreamingLevel);
+							return ((UPTRINT)LambdaStreamingLevel < (UPTRINT)OtherStreamingLevel);
 						}
 
 						return (Priority < OtherPriority);
 					}
 
-					return (StreamingLevel != nullptr);
+					return (LambdaStreamingLevel != nullptr);
 				};
 
 				StreamingLevels.Insert(StreamingLevel, Algo::LowerBound(StreamingLevels, StreamingLevel, PrioritySort));
