Disable audio mixer toggle if there are folks using the audio device out there in the world.

diff --git a/Engine/Source/Runtime/Engine/Private/AudioDeviceManager.cpp b/Engine/Source/Runtime/Engine/Private/AudioDeviceManager.cpp
--- a/Engine/Source/Runtime/Engine/Private/AudioDeviceManager.cpp
+++ b/Engine/Source/Runtime/Engine/Private/AudioDeviceManager.cpp
@@ -38,10 +38,17 @@ static int32 GCvarIsUsingAudioMixer = 0;
 FAutoConsoleVariableRef CVarIsUsingAudioMixer(
 	TEXT("au.IsUsingAudioMixer"),
 	GCvarIsUsingAudioMixer,
-	TEXT("Whether or not we're currently using the audio mixer. Change to dynamically toggle on/off. Note: sounds will stop. Looping sounds won't automatically resume. \n")
+	TEXT("Whether or not we're currently using the audio mixer. Change to dynamically toggle on/off. This will only take effect if an audio device is currently not in use, unless au.AllowUnsafeAudioMixerToggling is set to 1. Note: sounds will stop. Looping sounds won't automatically resume. \n")
 	TEXT("0: Not Using Audio Mixer, 1: Using Audio Mixer"),
 	ECVF_Default);
 
+static int32 GCvarAllowUnsafeAudioMixerToggling = 0;
+FAutoConsoleVariableRef CVarAllowUnsafeAudioMixerToggling(
+	TEXT("au.AllowUnsafeAudioMixerToggling"),
+	GCvarAllowUnsafeAudioMixerToggling,
+	TEXT("If set to 1, will allow au.IsUsingAudioMixer to swap out the audio engine, even if there are systems in the world currently using the audio engine. \n")
+	TEXT("0: Not Using Audio Mixer, 1: Using Audio Mixer"),
+	ECVF_Default);
 
 static int32 CVarIsVisualizeEnabled = 0;
 FAutoConsoleVariableRef CVarAudioVisualizeEnabled(
@@ -156,11 +163,39 @@ FAudioDeviceManager::~FAudioDeviceManager()
 	}
 }
 
-void FAudioDeviceManager::ToggleAudioMixer()
+bool FAudioDeviceManager::ToggleAudioMixer()
 {
 	// Only need to toggle if we have 2 device module names loaded at init
 	if (AudioDeviceModule && AudioDeviceModuleName.Len() > 0 && AudioMixerModuleName.Len() > 0)
 	{
+		// Check to see if there are still any active handles using this audio device, and either warn or exit early.
+		for (auto& DeviceContainer : Devices)
+		{
+			if (DeviceContainer.Value.NumberOfHandlesToThisDevice > 0)
+			{
+
+				UE_LOG(LogAudio, Warning, TEXT("Attempted to toggle the audio mixer while that audio device was in use."));
+
+#if INSTRUMENT_AUDIODEVICE_HANDLES
+				FString ActiveDeviceHandles;
+				for (auto& StackWalkString : DeviceContainer.Value.HandleCreationStackWalks)
+				{
+					ActiveDeviceHandles += StackWalkString.Value;
+					ActiveDeviceHandles += TEXT("\n\n");
+				}
+
+				UE_LOG(LogAudio, Warning, TEXT("List Of Active Handles to this device: \n%s"), *ActiveDeviceHandles);
+#else
+				UE_LOG(LogAudio, Warning, TEXT("For more information compile with INSTRUMENT_AUDIODEVICE_HANDLES."));
+#endif
+
+				if (!GCvarAllowUnsafeAudioMixerToggling)
+				{
+					return false;
+				}
+			}
+		}
+
 		// Suspend the audio thread
 		FAudioThread::SuspendAudioThread();
 
@@ -202,6 +237,32 @@ void FAudioDeviceManager::ToggleAudioMixer()
 			// and try and get the mix-states to be the same.
 			for (auto& DeviceContainer : Devices)
 			{
+				// Check to see if there are still any active handles using this audio device.
+				if(DeviceContainer.Value.NumberOfHandlesToThisDevice > 0)
+				{
+
+					UE_LOG(LogAudio, Warning, TEXT("Attempted to toggle the audio mixer while that audio device was in use."));
+
+#if INSTRUMENT_AUDIODEVICE_HANDLES
+					FString ActiveDeviceHandles;
+					for (auto& StackWalkString : DeviceContainer.Value.HandleCreationStackWalks)
+					{
+						ActiveDeviceHandles += StackWalkString.Value;
+						ActiveDeviceHandles += TEXT("\n\n");
+					}
+
+					UE_LOG(LogAudio, Warning, TEXT("List Of Active Handles to this device: \n%s"), *ActiveDeviceHandles);
+#else
+					UE_LOG(LogAudio, Warning, TEXT("For more information compile with INSTRUMENT_AUDIODEVICE_HANDLES."));
+#endif
+
+					if(!GCvarAllowUnsafeAudioMixerToggling)
+					{
+						continue;
+					}
+				}
+				
+
 				FAudioDevice*& AudioDevice = DeviceContainer.Value.Device;
 
 				check(AudioDevice);
@@ -291,6 +352,8 @@ void FAudioDeviceManager::ToggleAudioMixer()
 			FAudioThread::ResumeAudioThread();
 		}
 	}
+
+	return false;
 }
 
 bool FAudioDeviceManager::IsUsingAudioMixer() const
@@ -633,15 +696,31 @@ void FAudioDeviceManager::UpdateActiveAudioDevices(bool bGameTicking)
 	{
 		if (bUsingAudioMixer && !GCvarIsUsingAudioMixer)
 		{
-			ToggleAudioMixer();
-			bToggledAudioMixer = true;
-			bUsingAudioMixer = false;
+			const bool bSuccessfullyToggledAudioMixer = ToggleAudioMixer();
+
+			if (bSuccessfullyToggledAudioMixer)
+			{
+				bToggledAudioMixer = true;
+				bUsingAudioMixer = false;
+			}
+			else
+			{
+				GCvarIsUsingAudioMixer = true;
+			}
 		}
 		else if (!bUsingAudioMixer && GCvarIsUsingAudioMixer)
 		{
-			ToggleAudioMixer();
-			bToggledAudioMixer = true;
-			bUsingAudioMixer = true;
+			const bool bSuccessfullyToggledAudioMixer = ToggleAudioMixer();
+
+			if (bSuccessfullyToggledAudioMixer)
+			{
+				bToggledAudioMixer = true;
+				bUsingAudioMixer = true;
+			}
+			else
+			{
+				GCvarIsUsingAudioMixer = false;
+			}
 		}
 	}
 
@@ -883,7 +962,11 @@ void FAudioDeviceManager::AddStackWalkForContainer(Audio::FDeviceId InId, uint32
 
 void FAudioDeviceManager::RemoveStackWalkForContainer(Audio::FDeviceId InId, uint32 StackWalkID)
 {
-	check(Devices.Contains(InId));
+	if (!Devices.Contains(InId))
+	{
+		return;
+	}
+
 	check(Devices[InId].HandleCreationStackWalks.Contains(StackWalkID));
 	FAudioDeviceContainer& Container = Devices[InId];
 	Container.HandleCreationStackWalks.Remove(StackWalkID);
diff --git a/Engine/Source/Runtime/Engine/Public/AudioDeviceManager.h b/Engine/Source/Runtime/Engine/Public/AudioDeviceManager.h
--- a/Engine/Source/Runtime/Engine/Public/AudioDeviceManager.h
+++ b/Engine/Source/Runtime/Engine/Public/AudioDeviceManager.h
@@ -402,7 +402,7 @@ private:
 	uint32 GetNewDeviceID();
 
 	/** Toggles between audio mixer and non-audio mixer audio engine. */
-	void ToggleAudioMixer();
+	bool ToggleAudioMixer();
 
 	/** Load audio device module. */
 	bool LoadDefaultAudioDeviceModule();
