minifix: get rid of redundant check for device handles in ToggleAudioMixer.

diff --git a/Engine/Source/Runtime/Engine/Private/AudioDeviceManager.cpp b/Engine/Source/Runtime/Engine/Private/AudioDeviceManager.cpp
--- a/Engine/Source/Runtime/Engine/Private/AudioDeviceManager.cpp
+++ b/Engine/Source/Runtime/Engine/Private/AudioDeviceManager.cpp
@@ -237,32 +237,6 @@ bool FAudioDeviceManager::ToggleAudioMixer()
 			// and try and get the mix-states to be the same.
 			for (auto& DeviceContainer : Devices)
 			{
-				// Check to see if there are still any active handles using this audio device.
-				if(DeviceContainer.Value.NumberOfHandlesToThisDevice > 0)
-				{
-
-					UE_LOG(LogAudio, Warning, TEXT("Attempted to toggle the audio mixer while that audio device was in use."));
-
-#if INSTRUMENT_AUDIODEVICE_HANDLES
-					FString ActiveDeviceHandles;
-					for (auto& StackWalkString : DeviceContainer.Value.HandleCreationStackWalks)
-					{
-						ActiveDeviceHandles += StackWalkString.Value;
-						ActiveDeviceHandles += TEXT("\n\n");
-					}
-
-					UE_LOG(LogAudio, Warning, TEXT("List Of Active Handles to this device: \n%s"), *ActiveDeviceHandles);
-#else
-					UE_LOG(LogAudio, Warning, TEXT("For more information compile with INSTRUMENT_AUDIODEVICE_HANDLES."));
-#endif
-
-					if(!GCvarAllowUnsafeAudioMixerToggling)
-					{
-						continue;
-					}
-				}
-				
-
 				FAudioDevice*& AudioDevice = DeviceContainer.Value.Device;
 
 				check(AudioDevice);
