Fix typo in au.AllowUnsafeAudioMixerToggling cvar.

diff --git a/Engine/Source/Runtime/Engine/Private/AudioDeviceManager.cpp b/Engine/Source/Runtime/Engine/Private/AudioDeviceManager.cpp
--- a/Engine/Source/Runtime/Engine/Private/AudioDeviceManager.cpp
+++ b/Engine/Source/Runtime/Engine/Private/AudioDeviceManager.cpp
@@ -47,7 +47,7 @@ FAutoConsoleVariableRef CVarAllowUnsafeAudioMixerToggling(
 	TEXT("au.AllowUnsafeAudioMixerToggling"),
 	GCvarAllowUnsafeAudioMixerToggling,
 	TEXT("If set to 1, will allow au.IsUsingAudioMixer to swap out the audio engine, even if there are systems in the world currently using the audio engine. \n")
-	TEXT("0: Not Using Audio Mixer, 1: Using Audio Mixer"),
+	TEXT("0: disable usage of au.IsUsingAudioMixer when the audio device is actively in use, 1: enable usage of au.IsUsingAudioMixer."),
 	ECVF_Default);
 
 static int32 CVarIsVisualizeEnabled = 0;
