Expose some missing runtime virtual texture functionality to blueprint

diff --git a/Engine/Source/Runtime/Engine/Classes/Components/PrimitiveComponent.h b/Engine/Source/Runtime/Engine/Classes/Components/PrimitiveComponent.h
--- a/Engine/Source/Runtime/Engine/Classes/Components/PrimitiveComponent.h
+++ b/Engine/Source/Runtime/Engine/Classes/Components/PrimitiveComponent.h
@@ -592,7 +592,7 @@ public:
 	 * Array of runtime virtual textures into which we render the mesh for this actor. 
 	 * The material also needs to be set up to output to a virtual texture. 
 	 */
-	UPROPERTY(EditAnywhere, BlueprintReadOnly, Category = VirtualTexture, meta = (DisplayName = "Render to Virtual Textures"))
+	UPROPERTY(EditAnywhere, BlueprintReadWrite, Category = VirtualTexture, meta = (DisplayName = "Render to Virtual Textures"))
 	TArray<URuntimeVirtualTexture*> RuntimeVirtualTextures;
 
 	/** Bias to the LOD selected for rendering to runtime virtual textures. */
diff --git a/Engine/Source/Runtime/Engine/Classes/Components/RuntimeVirtualTextureComponent.h b/Engine/Source/Runtime/Engine/Classes/Components/RuntimeVirtualTextureComponent.h
--- a/Engine/Source/Runtime/Engine/Classes/Components/RuntimeVirtualTextureComponent.h
+++ b/Engine/Source/Runtime/Engine/Classes/Components/RuntimeVirtualTextureComponent.h
@@ -9,14 +9,14 @@
 class URuntimeVirtualTexture;
 
 /** Component used to place a URuntimeVirtualTexture in the world. */
-UCLASS(ClassGroup = Rendering, collapsecategories, hidecategories = (Activation, Collision, Cooking, Mobility, LOD, Object, Physics, Rendering), editinlinenew)
+UCLASS(Blueprintable, ClassGroup = Rendering, collapsecategories, hidecategories = (Activation, Collision, Cooking, Mobility, LOD, Object, Physics, Rendering), editinlinenew)
 class ENGINE_API URuntimeVirtualTextureComponent : public USceneComponent
 {
 	GENERATED_UCLASS_BODY()
 
-private:
+protected:
 	/** The virtual texture object to use. */
-	UPROPERTY(EditAnywhere, DuplicateTransient, Category = VirtualTexture)
+	UPROPERTY(EditAnywhere, BlueprintReadOnly, DuplicateTransient, Category = VirtualTexture)
 	URuntimeVirtualTexture* VirtualTexture = nullptr;
 
 	/** Use any streaming low mips when rendering in editor. Set true to view and debug the baked streaming low mips. */
@@ -32,6 +32,7 @@ public:
 	URuntimeVirtualTexture* GetVirtualTexture() const { return VirtualTexture; }
 
 	/** Get the runtime virtual texture UV to World transform on this component. */
+	UFUNCTION(BlueprintPure, Category = VirtualTexture)
 	FTransform GetVirtualTextureTransform() const;
 
 	/** Get if we want use any streaming low mips in the runtime virtual texture set on this component. */
diff --git a/Engine/Source/Runtime/Engine/Classes/VT/RuntimeVirtualTexture.h b/Engine/Source/Runtime/Engine/Classes/VT/RuntimeVirtualTexture.h
--- a/Engine/Source/Runtime/Engine/Classes/VT/RuntimeVirtualTexture.h
+++ b/Engine/Source/Runtime/Engine/Classes/VT/RuntimeVirtualTexture.h
@@ -9,7 +9,7 @@
 #include "RuntimeVirtualTexture.generated.h"
 
 /** Runtime virtual texture UObject */
-UCLASS(ClassGroup = Rendering)
+UCLASS(ClassGroup = Rendering, BlueprintType)
 class ENGINE_API URuntimeVirtualTexture : public UObject
 {
 	GENERATED_UCLASS_BODY()
@@ -17,7 +17,7 @@ class ENGINE_API URuntimeVirtualTexture : public UObject
 
 protected:
 	/** Contents of virtual texture. */
-	UPROPERTY(EditAnywhere, Category = Layout, meta = (DisplayName = "Virtual texture content"), AssetRegistrySearchable)
+	UPROPERTY(EditAnywhere, BlueprintReadOnly, Category = Layout, meta = (DisplayName = "Virtual texture content"), AssetRegistrySearchable)
 	ERuntimeVirtualTextureMaterialType MaterialType = ERuntimeVirtualTextureMaterialType::BaseColor_Normal_Specular;
 
 	/** Enable storing the virtual texture in GPU supported compression formats. Using uncompressed is only recommended for debugging and quality comparisons. */
@@ -25,23 +25,23 @@ protected:
 	bool bCompressTextures = true;
 
 	/** Enable usage of the virtual texture. When disabled there is no rendering into the virtual texture, and sampling will return zero values. */
-	UPROPERTY(EditAnywhere, BlueprintReadOnly, AdvancedDisplay, Category = Layout, meta = (DisplayName = "Enable virtual texture"))
+	UPROPERTY(EditAnywhere, AdvancedDisplay, Category = Layout, meta = (DisplayName = "Enable virtual texture"))
 	bool bEnable = true;
 
 	/** Enable clear before rendering a page of the virtual texture. Disabling this can be an optimization if you know that the texture will always be fully covered by rendering.  */
-	UPROPERTY(EditAnywhere, BlueprintReadOnly, AdvancedDisplay, Category = Layout, meta = (DisplayName = "Enable clear before render"))
+	UPROPERTY(EditAnywhere, AdvancedDisplay, Category = Layout, meta = (DisplayName = "Enable clear before render"))
 	bool bClearTextures = true;
 
 	/** Enable page table channel packing. This reduces page table memory and update cost but can reduce the ability to share physical memory with other virtual textures.  */
-	UPROPERTY(EditAnywhere, BlueprintReadOnly, AdvancedDisplay, Category = Layout, meta = (DisplayName = "Enable packed page table"))
+	UPROPERTY(EditAnywhere, AdvancedDisplay, Category = Layout, meta = (DisplayName = "Enable packed page table"))
 	bool bSinglePhysicalSpace = true;
 
 	/** Enable private page table allocation. This can reduce total page table memory allocation but can also reduce the total number of virtual textures supported. */
-	UPROPERTY(EditAnywhere, BlueprintReadOnly, AdvancedDisplay, Category = Layout, meta = (DisplayName = "Enable private page table"))
+	UPROPERTY(EditAnywhere, AdvancedDisplay, Category = Layout, meta = (DisplayName = "Enable private page table"))
 	bool bPrivateSpace = true;
 
 	/** Enable device scalability settings to modify the TileCount of the virtual texture. */
-	UPROPERTY(EditAnywhere, BlueprintReadOnly, AdvancedDisplay, Category = Layout)
+	UPROPERTY(EditAnywhere, AdvancedDisplay, Category = Layout)
 	bool bEnableScalability = false;
 
 	/** Size of virtual texture along the largest axis. (Actual values increase in powers of 2) */
@@ -53,31 +53,31 @@ protected:
 	 * This replaces the deprecated Size property.
 	 * This is applied to the largest axis in world space and the size for any shorter axis is chosen to maintain aspect ratio.  
 	 */
-	UPROPERTY(EditAnywhere, BlueprintReadOnly, Category = Size, meta = (UIMin = "0", UIMax = "12", DisplayName = "Size of the virtual texture in tiles"))
+	UPROPERTY(EditAnywhere, BluePrintGetter = GetTileCount, Category = Size, meta = (UIMin = "0", UIMax = "12", DisplayName = "Size of the virtual texture in tiles"))
 	int32 TileCount = 8; // 256
 
 	/** Page tile size. (Actual values increase in powers of 2) */
-	UPROPERTY(EditAnywhere, BlueprintReadOnly, Category = Size, meta = (UIMin = "0", UIMax = "4", DisplayName = "Size of each virtual texture tile"))
+	UPROPERTY(EditAnywhere, BluePrintGetter = GetTileSize, Category = Size, meta = (UIMin = "0", UIMax = "4", DisplayName = "Size of each virtual texture tile"))
 	int32 TileSize = 2; // 256
 
 	/** Page tile border size divided by 2 (Actual values increase in multiples of 2). Higher values trigger a higher anisotropic sampling level. */
-	UPROPERTY(EditAnywhere, BlueprintReadOnly, Category = Size, meta = (UIMin = "0", UIMax = "4", DisplayName = "Border padding for each virtual texture tile"))
+	UPROPERTY(EditAnywhere, BluePrintGetter = GetTileBorderSize, Category = Size, meta = (UIMin = "0", UIMax = "4", DisplayName = "Border padding for each virtual texture tile"))
 	int32 TileBorderSize = 2; // 4
 
 	/** Number of low mips to serialize and stream for the virtual texture. This can reduce rendering update cost. */
-	UPROPERTY(EditAnywhere, BlueprintReadOnly, Category = LowMips, meta = (UIMin = "0", UIMax = "6", DisplayName = "Number of low mips to stream to the virtual texture"))
+	UPROPERTY(EditAnywhere, Category = LowMips, meta = (UIMin = "0", UIMax = "6", DisplayName = "Number of low mips to stream to the virtual texture"))
 	int32 StreamLowMips = 0;
 
 	/** Texture object containing streamed low mips. */
-	UPROPERTY(VisibleAnywhere, Category = LowMips, meta = (DisplayName = "Streaming low mip texture"))
+	UPROPERTY(VisibleAnywhere, BlueprintReadOnly, Category = LowMips, meta = (DisplayName = "Streaming low mip texture"))
 	class URuntimeVirtualTextureStreamingProxy* StreamingTexture;
 
 	/** Enable Crunch compression of streamed low mips. ZLib compression is used when Crunch is disabled. */
-	UPROPERTY(EditAnywhere, BlueprintReadOnly, AdvancedDisplay, Category = LowMips, meta = (DisplayName = "Enable Crunch compression"))
+	UPROPERTY(EditAnywhere, AdvancedDisplay, Category = LowMips, meta = (DisplayName = "Enable Crunch compression"))
 	bool bEnableCompressCrunch = false;
 
 	/** Number of low mips to cut from the virtual texture. This can reduce peak virtual texture update cost but will also increase the probability of mip shimmering. */
-	UPROPERTY(EditAnywhere, BlueprintReadOnly, AdvancedDisplay, Category = LowMips, meta = (UIMin = "0", UIMax = "6", DisplayName = "Number of low mips to remove from the virtual texture"))
+	UPROPERTY(EditAnywhere, AdvancedDisplay, Category = LowMips, meta = (UIMin = "0", UIMax = "6", DisplayName = "Number of low mips to remove from the virtual texture"))
 	int32 RemoveLowMips = 0;
 
 public:
@@ -88,14 +88,19 @@ public:
 	ERuntimeVirtualTextureMaterialType GetMaterialType() const { return MaterialType; }
 
 	/** Public getter for virtual texture size */
+	UFUNCTION(BlueprintPure, Category = Size)
 	int32 GetSize() const { return GetTileCount() * GetTileSize(); }
 	/** Public getter for virtual texture tile count */
+	UFUNCTION(BlueprintGetter)
 	int32 GetTileCount() const { return GetTileCount(TileCount); }
 	static int32 GetTileCount(int32 InTileCount) { return 1 << FMath::Clamp(InTileCount, 0, 12); }
 	/** Public getter for virtual texture tile size */
+	UFUNCTION(BlueprintGetter)
 	int32 GetTileSize() const { return 1 << FMath::Clamp(TileSize + 6, 6, 10); }
 	/** Public getter for virtual texture tile border size */
+	UFUNCTION(BlueprintGetter)
 	int32 GetTileBorderSize() const { return 2 * FMath::Clamp(TileBorderSize, 0, 4); }
+	
 	/** Public getter for virtual texture removed low mips */
 	int32 GetRemoveLowMips() const { return FMath::Clamp(RemoveLowMips, 0, 5); }
 	/** Public getter for virtual texture streaming low mips */
diff --git a/Engine/Source/Runtime/Foliage/Public/FoliageType.h b/Engine/Source/Runtime/Foliage/Public/FoliageType.h
--- a/Engine/Source/Runtime/Foliage/Public/FoliageType.h
+++ b/Engine/Source/Runtime/Foliage/Public/FoliageType.h
@@ -529,7 +529,7 @@ public:
 	 * Array of runtime virtual textures into which we render the instances. 
 	 * The mesh material also needs to be set up to output to a virtual texture. 
 	 */
-	UPROPERTY(EditAnywhere, BlueprintReadOnly, Category=VirtualTexture, meta = (DisplayName = "Render to Virtual Textures"))
+	UPROPERTY(EditAnywhere, BlueprintReadWrite, Category=VirtualTexture, meta = (DisplayName = "Render to Virtual Textures"))
 	TArray<URuntimeVirtualTexture*> RuntimeVirtualTextures;
 
 	/**
diff --git a/Engine/Source/Runtime/Landscape/Classes/LandscapeProxy.h b/Engine/Source/Runtime/Landscape/Classes/LandscapeProxy.h
--- a/Engine/Source/Runtime/Landscape/Classes/LandscapeProxy.h
+++ b/Engine/Source/Runtime/Landscape/Classes/LandscapeProxy.h
@@ -545,7 +545,7 @@ public:
 	 * Array of runtime virtual textures into which we render this landscape.
 	 * The material also needs to be set up to output to a virtual texture.
 	 */
-	UPROPERTY(EditAnywhere, BlueprintReadOnly, Category = VirtualTexture, meta = (DisplayName = "Render to Virtual Textures"))
+	UPROPERTY(EditAnywhere, BlueprintReadWrite, Category = VirtualTexture, meta = (DisplayName = "Render to Virtual Textures"))
 	TArray<URuntimeVirtualTexture*> RuntimeVirtualTextures;
 
 	/** 
