fix for sporadic warnings when running cook on Mac.

diff --git a/Engine/Source/Runtime/Engine/Private/AudioDerivedData.cpp b/Engine/Source/Runtime/Engine/Private/AudioDerivedData.cpp
--- a/Engine/Source/Runtime/Engine/Private/AudioDerivedData.cpp
+++ b/Engine/Source/Runtime/Engine/Private/AudioDerivedData.cpp
@@ -1020,27 +1020,27 @@ static void CookSimpleWave(USoundWave* SoundWave, FName FormatName, const IAudio
 	FScopeLock ScopeLock(&SoundWave->RawDataCriticalSection);
 #endif
 
-	SoundWave->RawData.ForceBulkDataResident();
-
-	// check if there is any raw sound data
-	if( SoundWave->RawData.GetBulkDataSize() > 0 )
+	if (!SoundWave->RawData.IsBulkDataLoaded())
 	{
+		SoundWave->RawData.ForceBulkDataResident();
+	}
+	
+	UE_CLOG(SoundWave->RawData.IsBulkDataLoaded(), LogAudioDerivedData, Display, TEXT("calling ForceBulkDataResident for LPCM data for USoundWave %s failed."), *SoundWave->GetFullName());
 
-		// Lock raw wave data.
-		const uint8* RawWaveData = ( const uint8* )SoundWave->RawData.LockReadOnly();
-		bWasLocked = true;
-		int32 RawDataSize = SoundWave->RawData.GetBulkDataSize();
+	// Lock raw wave data.
+	const uint8* RawWaveData = (const uint8*)SoundWave->RawData.LockReadOnly();
+	bWasLocked = true;
+	int32 RawDataSize = SoundWave->RawData.GetBulkDataSize();
 
-		// parse the wave data
-		if( !WaveInfo.ReadWaveHeader( RawWaveData, RawDataSize, 0 ) )
-		{
-			UE_LOG(LogAudioDerivedData, Warning, TEXT( "Only mono or stereo 16 bit waves allowed: %s (%d bytes)" ), *SoundWave->GetFullName(), RawDataSize );
-		}
-		else
-		{
-			Input.AddUninitialized(WaveInfo.SampleDataSize);
-			FMemory::Memcpy(Input.GetData(), WaveInfo.SampleDataStart, WaveInfo.SampleDataSize);
-		}
+	// parse the wave data
+	if (!WaveInfo.ReadWaveHeader(RawWaveData, RawDataSize, 0))
+	{
+		UE_LOG(LogAudioDerivedData, Warning, TEXT("Only mono or stereo 16 bit waves allowed: %s (%d bytes)"), *SoundWave->GetFullName(), RawDataSize);
+	}
+	else
+	{
+		Input.AddUninitialized(WaveInfo.SampleDataSize);
+		FMemory::Memcpy(Input.GetData(), WaveInfo.SampleDataStart, WaveInfo.SampleDataSize);
 	}
 
 	if(!Input.Num())
