Crash fix - not safe to modify the StreamingLevels array length while levels are still being considered (null array entries are fine)

diff --git a/Engine/Source/Runtime/Engine/Private/World.cpp b/Engine/Source/Runtime/Engine/Private/World.cpp
--- a/Engine/Source/Runtime/Engine/Private/World.cpp
+++ b/Engine/Source/Runtime/Engine/Private/World.cpp
@@ -3200,7 +3200,15 @@ bool FStreamingLevelsToConsider::Contains(ULevelStreaming* StreamingLevel) const
 
 void FStreamingLevelsToConsider::Reset()
 {
-	StreamingLevels.Reset();
+	if (bStreamingLevelsBeingConsidered)
+	{
+		// not safe to resize while levels are being considered, just null everything
+		FMemory::Memzero(StreamingLevels.GetData(), StreamingLevels.Num() * sizeof(ULevelStreaming*));
+	}
+	else
+	{
+		StreamingLevels.Reset();
+	}
 	LevelsToProcess.Reset();
 }
 
