Moving a couple of GameViewportClient properties to be protected

diff --git a/Engine/Source/Runtime/Engine/Classes/Engine/GameViewportClient.h b/Engine/Source/Runtime/Engine/Classes/Engine/GameViewportClient.h
--- a/Engine/Source/Runtime/Engine/Classes/Engine/GameViewportClient.h
+++ b/Engine/Source/Runtime/Engine/Classes/Engine/GameViewportClient.h
@@ -887,6 +887,13 @@ private:
 	*/
 	void AddDebugDisplayProperty(class UObject* Obj, TSubclassOf<class UObject> WithinClass, const FName& PropertyName, bool bSpecialProperty = false);
 
+protected:
+	/** Handle to the audio device created for this viewport. Each viewport (for multiple PIE) will have its own audio device. */
+	uint32 AudioDeviceHandle = INDEX_NONE;
+
+	/** Whether or not this audio device is in audio-focus */
+	bool bHasAudioFocus = false;
+
 private:
 	/** Slate window associated with this viewport client.  The same window may host more than one viewport client. */
 	TWeakPtr<SWindow> Window;
@@ -1002,9 +1009,6 @@ private:
 	/** Mouse cursor locking behavior when the viewport is clicked */
 	EMouseLockMode MouseLockMode;
 
-	/** Whether or not this audio device is in audio-focus */
-	bool bHasAudioFocus;
-
 	/** Is the mouse currently over the viewport client */
 	bool bIsMouseOverClient;
 
diff --git a/Engine/Source/Runtime/Engine/Private/GameViewportClient.cpp b/Engine/Source/Runtime/Engine/Private/GameViewportClient.cpp
--- a/Engine/Source/Runtime/Engine/Private/GameViewportClient.cpp
+++ b/Engine/Source/Runtime/Engine/Private/GameViewportClient.cpp
@@ -171,7 +171,6 @@ UGameViewportClient::UGameViewportClient(const FObjectInitializer& ObjectInitial
 	, MouseCaptureMode(EMouseCaptureMode::CapturePermanently)
 	, bHideCursorDuringCapture(false)
 	, MouseLockMode(EMouseLockMode::LockOnCapture)
-	, bHasAudioFocus(false)
 	, bIsMouseOverClient(false)
 #if WITH_EDITOR
 	, bUseMouseForTouchInEditor(false)
@@ -255,7 +254,6 @@ UGameViewportClient::UGameViewportClient(FVTableHelper& Helper)
 	, MouseCaptureMode(EMouseCaptureMode::CapturePermanently)
 	, bHideCursorDuringCapture(false)
 	, MouseLockMode(EMouseLockMode::LockOnCapture)
-	, bHasAudioFocus(false)
 {
 
 }
