Remove warning message when display name for material attribute changes, triggered when loading UMaterialExpressionGetMaterialAttributes

diff --git a/Engine/Source/Runtime/Engine/Private/Materials/MaterialExpressions.cpp b/Engine/Source/Runtime/Engine/Private/Materials/MaterialExpressions.cpp
--- a/Engine/Source/Runtime/Engine/Private/Materials/MaterialExpressions.cpp
+++ b/Engine/Source/Runtime/Engine/Private/Materials/MaterialExpressions.cpp
@@ -5878,24 +5878,11 @@ void UMaterialExpressionGetMaterialAttributes::PostLoad()
 	// Verify serialized attributes
 	check(Outputs.Num() == AttributeGetTypes.Num() + 1);
 
+	// Make sure all outputs have up to date display names
 	for (int i = 1; i < Outputs.Num(); ++i)
 	{
-		const FString DisplayName = FMaterialAttributeDefinitionMap::GetDisplayNameForMaterial(AttributeGetTypes[i-1],Material).ToString();
-		if (Outputs[i].OutputName.ToString() != DisplayName)
-		{
-			FString MaterialName;
-			if (Material)
-			{
-				Material->GetName(MaterialName);
-			}
-			else if (Function)
-			{
-				Function->GetName(MaterialName);
-			}
-
-			UE_LOG(LogMaterial, Warning, TEXT("Serialized attribute that no longer exists (%s) for material \"%s\"."), *(Outputs[i].OutputName.ToString()), *MaterialName);
-			Outputs[i].OutputName = *DisplayName;
-		}
+		const FString DisplayName = FMaterialAttributeDefinitionMap::GetDisplayNameForMaterial(AttributeGetTypes[i-1], Material).ToString();
+		Outputs[i].OutputName = *DisplayName;
 	}
 }
 #endif // WITH_EDITOR
