Remove warning message when display name for material attribute changes, triggered when loading UMaterialExpressionGetMaterialAttributes

diff --git a/Engine/Source/Runtime/Engine/Private/Materials/MaterialExpressions.cpp b/Engine/Source/Runtime/Engine/Private/Materials/MaterialExpressions.cpp
--- a/Engine/Source/Runtime/Engine/Private/Materials/MaterialExpressions.cpp
+++ b/Engine/Source/Runtime/Engine/Private/Materials/MaterialExpressions.cpp
@@ -5872,6 +5872,7 @@ void UMaterialExpressionGetMaterialAttributes::PostLoad()
 	// Verify serialized attributes
 	check(Outputs.Num() == AttributeGetTypes.Num() + 1);
 
+	// Make sure all outputs have up to date display names
 	for (int i = 1; i < Outputs.Num(); ++i)
 	{
 		const FString DisplayName = FMaterialAttributeDefinitionMap::GetDisplayNameForMaterial(AttributeGetTypes[i - 1], Material).ToString();
