Added missing Source.GatherDebugData to FAnimNode_Inertialization::GatherDebugData

diff --git a/Engine/Source/Editor/AnimGraph/Private/AnimGraphNode_Inertialization.cpp b/Engine/Source/Editor/AnimGraph/Private/AnimGraphNode_Inertialization.cpp
--- a/Engine/Source/Editor/AnimGraph/Private/AnimGraphNode_Inertialization.cpp
+++ b/Engine/Source/Editor/AnimGraph/Private/AnimGraphNode_Inertialization.cpp
@@ -2,8 +2,6 @@
 
 #include "AnimGraphNode_Inertialization.h"
 
-#include "EditorCategoryUtils.h"
-
 #define LOCTEXT_NAMESPACE "AnimGraphNode_Inertialization"
 
 
diff --git a/Engine/Source/Runtime/Engine/Private/Animation/AnimNode_Inertialization.cpp b/Engine/Source/Runtime/Engine/Private/Animation/AnimNode_Inertialization.cpp
--- a/Engine/Source/Runtime/Engine/Private/Animation/AnimNode_Inertialization.cpp
+++ b/Engine/Source/Runtime/Engine/Private/Animation/AnimNode_Inertialization.cpp
@@ -268,6 +268,8 @@ void FAnimNode_Inertialization::GatherDebugData(FNodeDebugData& DebugData)
 			InertializationDeficit);
 	}
 	DebugData.AddDebugItem(DebugLine);
+
+	Source.GatherDebugData(DebugData);
 }
 
 
