- Remove various UMaterialExpression::NeedsLoadForClient() impl, and instead always return true for UMaterialExpression::IsEditorOnly(), as we no longer want any expressions in non-editor build
- MaterialCachedData checks for nullptr expressions, and avoids updating if any are encountered...this fixes crash with editor loading cooked/non-editor data

diff --git a/Engine/Source/Runtime/Engine/Private/Materials/Material.cpp b/Engine/Source/Runtime/Engine/Private/Materials/Material.cpp
--- a/Engine/Source/Runtime/Engine/Private/Materials/Material.cpp
+++ b/Engine/Source/Runtime/Engine/Private/Materials/Material.cpp
@@ -2064,6 +2064,10 @@ void UMaterial::UpdateCachedExpressionData()
 
 	if (!bAnyQualityLevelsSet)
 	{
+		if (CachedExpressionData.QualityLevelsUsed.Num() == 0)
+		{
+			CachedExpressionData.QualityLevelsUsed.AddDefaulted(EMaterialQualityLevel::Num);
+		}
 		CachedExpressionData.QualityLevelsUsed[EMaterialQualityLevel::High] = true;
 	}
 }
