Allow non-threading builds to FlushLevelStreaming. Level streaming does not depend on threads.

Prevent FlushLevelStreaming reentrancy from causing issues.

diff --git a/Engine/Source/Runtime/Engine/Private/World.cpp b/Engine/Source/Runtime/Engine/Private/World.cpp
--- a/Engine/Source/Runtime/Engine/Private/World.cpp
+++ b/Engine/Source/Runtime/Engine/Private/World.cpp
@@ -3108,7 +3108,7 @@ void FStreamingLevelsToConsider::Add_Internal(ULevelStreaming* StreamingLevel, b
 			if (bGuaranteedNotInContainer || !StreamingLevels.Contains(StreamingLevel))
 			{
 				auto PrioritySort = [](ULevelStreaming* LambdaStreamingLevel, ULevelStreaming* OtherStreamingLevel)
-				{
+			{
 					if (LambdaStreamingLevel && OtherStreamingLevel)
 					{
 						const int32 Priority = LambdaStreamingLevel->GetPriority();
@@ -3207,7 +3207,7 @@ void FStreamingLevelsToConsider::Reset()
 	}
 	else
 	{
-		StreamingLevels.Reset();
+	StreamingLevels.Reset();
 	}
 	LevelsToProcess.Reset();
 }
@@ -3483,14 +3483,24 @@ void UWorld::UpdateStreamingLevelPriority(ULevelStreaming* StreamingLevel)
 
 void UWorld::FlushLevelStreaming(EFlushLevelStreamingType FlushType)
 {
-	// We're already full flushing, no reason to flush again
-	if (FlushType == EFlushLevelStreamingType::Full && FlushLevelStreamingType == EFlushLevelStreamingType::Full)
+	if (FlushType == FlushLevelStreamingType
+		|| FlushLevelStreamingType == EFlushLevelStreamingType::Full
+		|| FlushType == EFlushLevelStreamingType::None)
 	{
+		// We're already flushing at the correct level, or none was passed in
+		// No need to flush
+		return;
+	}
+	else if (FlushType == EFlushLevelStreamingType::Full && FlushLevelStreamingType == EFlushLevelStreamingType::Visibility)
+	{
+		// If FlushLevelStreaming is called for a full update while we are already doing a visibility update, 
+		// upgrade the stored type and let it go back to the loop that was already running
+		FlushLevelStreamingType = EFlushLevelStreamingType::Full;
 		return;
 	}
 
 	AWorldSettings* WorldSettings = GetWorldSettings();
-	
+
 	TGuardValue<EFlushLevelStreamingType> FlushingLevelStreamingGuard(FlushLevelStreamingType, FlushType);
 
 	// Update internals with current loaded/ visibility flags.
@@ -3503,14 +3513,13 @@ void UWorld::FlushLevelStreaming(EFlushLevelStreamingType FlushType)
 	UpdateLevelStreaming();
 
 	// Making levels visible is spread across several frames so we simply loop till it is done.
-	bool bLevelsPendingVisibility = true;
+	bool bLevelsPendingVisibility = IsVisibilityRequestPending();
 	while( bLevelsPendingVisibility )
 	{
-		bLevelsPendingVisibility = IsVisibilityRequestPending();
-
 		// Tick level streaming to make levels visible.
-		if( bLevelsPendingVisibility )
-		{
+
+		const EFlushLevelStreamingType LastStreamingType = FlushLevelStreamingType;
+
 			// Only flush async loading if we're performing a full flush.
 			if (FlushLevelStreamingType == EFlushLevelStreamingType::Full)
 			{
@@ -3520,7 +3529,10 @@ void UWorld::FlushLevelStreaming(EFlushLevelStreamingType FlushType)
 	
 			// Update level streaming.
 			UpdateLevelStreaming();
-		}
+
+		// If FlushLevelStreaming reentered as a result of FlushAsyncLoading or UpdateLevelStreaming and upgraded
+		// the flush type, we'll need to do at least one additional loop to be certain all processing is complete
+		bLevelsPendingVisibility = ((LastStreamingType != FlushLevelStreamingType) || IsVisibilityRequestPending());
 	}
 	
 	check(!IsVisibilityRequestPending());
