fixed position of debug messages

diff --git a/Engine/Source/Runtime/Engine/Private/DrawDebugHelpers.cpp b/Engine/Source/Runtime/Engine/Private/DrawDebugHelpers.cpp
--- a/Engine/Source/Runtime/Engine/Private/DrawDebugHelpers.cpp
+++ b/Engine/Source/Runtime/Engine/Private/DrawDebugHelpers.cpp
@@ -956,7 +956,7 @@ void DrawDebugFloatHistory(UWorld const & WorldRef, FDebugFloatHistory const & F
 void DrawDebugFloatHistory(UWorld const & WorldRef, FDebugFloatHistory const & FloatHistory, FVector const & DrawLocation, FVector2D const & DrawSize, FColor const & DrawColor, bool const & bPersistent, float const & LifeTime, uint8 const & DepthPriority)
 {
 	APlayerController * PlayerController = WorldRef.GetGameInstance() != nullptr ? WorldRef.GetGameInstance()->GetFirstLocalPlayerController() : nullptr;
-	FRotator const DrawRotation = (PlayerController && PlayerController->PlayerCameraManager) ? PlayerController->PlayerCameraManager->GetCameraCachePOV().Rotation : FRotator(0, 0, 0);
+	FRotator const DrawRotation = (PlayerController && PlayerController->PlayerCameraManager) ? PlayerController->PlayerCameraManager->GetCameraRotation() : FRotator(0, 0, 0);
 
 	FTransform const DrawTransform(DrawRotation, DrawLocation);
 	DrawDebugFloatHistory(WorldRef, FloatHistory, DrawTransform, DrawSize, DrawColor, bPersistent, LifeTime, DepthPriority);
