Fix LevelStreaming when changing desired level between request and completion with an already loaded level. Pending unload is never processed and state machine stays indefinively in the LoadedNotVisible state.

diff --git a/Engine/Source/Runtime/Engine/Private/LevelStreaming.cpp b/Engine/Source/Runtime/Engine/Private/LevelStreaming.cpp
--- a/Engine/Source/Runtime/Engine/Private/LevelStreaming.cpp
+++ b/Engine/Source/Runtime/Engine/Private/LevelStreaming.cpp
@@ -657,6 +657,12 @@ void ULevelStreaming::UpdateStreamingState(bool& bOutUpdateAgain, bool& bOutRede
 			break;
 
 		case ETargetState::LoadedNotVisible:
+			if (LoadedLevel && !IsDesiredLevelLoaded())
+			{
+				// Process PendingUnloadLevel to unblock level streaming state machine (no new request will start while there's a pending level to unload) 
+				// This rare case can happen if desired level (typically LODPackage) changed between last RequestLevel call and AsyncLevelLoadComplete completion callback.
+				DiscardPendingUnloadLevel(World);
+			}
 			UpdateStreamingState_RequestLevel();
 			break;
 
@@ -1836,4 +1842,4 @@ void ULevelStreamingAlwaysLoaded::GetPrestreamPackages(TArray<UObject*>& OutPres
 	OutPrestream.Add(GetLoadedLevel()); // Nulls will be ignored later
 }
 
-#undef LOCTEXT_NAMESPACE
+#undef LOCTEXT_NAMESPACE
\ No newline at end of file
