Fixed Python exposed name of get_components_by_class

#rnx

diff --git a/Engine/Source/Runtime/Engine/Classes/GameFramework/Actor.h b/Engine/Source/Runtime/Engine/Classes/GameFramework/Actor.h
--- a/Engine/Source/Runtime/Engine/Classes/GameFramework/Actor.h
+++ b/Engine/Source/Runtime/Engine/Classes/GameFramework/Actor.h
@@ -2824,7 +2824,7 @@ public:
 	 * Currently returns an array of UActorComponent which must be cast to the correct type.
 	 * This intended to only be used by blueprints. Use GetComponents() in C++.
 	 */
-	UFUNCTION(BlueprintCallable, Category = "Actor", meta = (ComponentClass = "ActorComponent", DisplayName = "GetComponentsByClass", DeterminesOutputType = "ComponentClass"))
+	UFUNCTION(BlueprintCallable, Category = "Actor", meta = (ComponentClass = "ActorComponent", DisplayName = "GetComponentsByClass", ScriptName = "GetComponentsByClass", DeterminesOutputType = "ComponentClass"))
 	TArray<UActorComponent*> K2_GetComponentsByClass(TSubclassOf<UActorComponent> ComponentClass) const;
 
 	/** Gets all the components that inherit from the given class with a given tag. */
