Add blueprint API for setting layer parameters on MID

diff --git a/Engine/Source/Runtime/Engine/Classes/Materials/Material.h b/Engine/Source/Runtime/Engine/Classes/Materials/Material.h
--- a/Engine/Source/Runtime/Engine/Classes/Materials/Material.h
+++ b/Engine/Source/Runtime/Engine/Classes/Materials/Material.h
@@ -1004,6 +1004,7 @@ public:
 #endif // WITH_EDITORONLY_DATA
 	ENGINE_API virtual bool GetTerrainLayerWeightParameterValue(const FHashedMaterialParameterInfo& ParameterInfo, int32& OutWeightmapIndex, FGuid& OutExpressionGuid) const override;
 	ENGINE_API virtual bool UpdateLightmassTextureTracking() override;
+	ENGINE_API virtual int32 GetLayerParameterIndex(EMaterialParameterAssociation Association, UMaterialFunctionInterface* LayerFunction) const override;
 #if WITH_EDITOR
 	ENGINE_API virtual bool GetGroupName(const FHashedMaterialParameterInfo& ParameterInfo, FName& OutGroup) const override;
 	ENGINE_API virtual bool GetParameterDesc(const FHashedMaterialParameterInfo& ParameterInfo, FString& OutDesc, const TArray<struct FStaticMaterialLayersParameter>* MaterialLayersParameters = nullptr) const override;
diff --git a/Engine/Source/Runtime/Engine/Classes/Materials/MaterialInstance.h b/Engine/Source/Runtime/Engine/Classes/Materials/MaterialInstance.h
--- a/Engine/Source/Runtime/Engine/Classes/Materials/MaterialInstance.h
+++ b/Engine/Source/Runtime/Engine/Classes/Materials/MaterialInstance.h
@@ -516,6 +516,7 @@ public:
 	virtual ENGINE_API float GetEmissiveBoost() const override;
 	virtual ENGINE_API float GetDiffuseBoost() const override;
 	virtual ENGINE_API float GetExportResolutionScale() const override;
+	virtual ENGINE_API int32 GetLayerParameterIndex(EMaterialParameterAssociation Association, UMaterialFunctionInterface* LayerFunction) const override;
 #if WITH_EDITOR
 	virtual ENGINE_API bool GetParameterDesc(const FHashedMaterialParameterInfo& ParameterInfo, FString& OutDesc, const TArray<struct FStaticMaterialLayersParameter>* MaterialLayersParameters = nullptr) const override;
 	virtual ENGINE_API bool GetParameterSortPriority(const FHashedMaterialParameterInfo& ParameterInfo, int32& OutSortPriority, const TArray<struct FStaticMaterialLayersParameter>* MaterialLayersParameters = nullptr) const override;
diff --git a/Engine/Source/Runtime/Engine/Classes/Materials/MaterialInstanceDynamic.h b/Engine/Source/Runtime/Engine/Classes/Materials/MaterialInstanceDynamic.h
--- a/Engine/Source/Runtime/Engine/Classes/Materials/MaterialInstanceDynamic.h
+++ b/Engine/Source/Runtime/Engine/Classes/Materials/MaterialInstanceDynamic.h
@@ -17,6 +17,10 @@ class ENGINE_API UMaterialInstanceDynamic : public UMaterialInstance
 	UFUNCTION(BlueprintCallable, meta=(Keywords = "SetFloatParameterValue"), Category="Rendering|Material")
 	void SetScalarParameterValue(FName ParameterName, float Value);
 
+	/** Set a MID scalar (float) parameter value using MPI (to allow access to layer parameters) */
+	UFUNCTION(BlueprintCallable, meta = (Keywords = "SetFloatParameterValue"), Category = "Rendering|Material")
+	void SetScalarParameterValueByInfo(const FMaterialParameterInfo& ParameterInfo, float Value);
+
 	// NOTE: These Index-related functions should be used VERY carefully, and only in cases where optimization is
 	// critical.  Generally that's only if you're using an unusually high number of parameters in a material AND
 	// setting a huge number of parameters in the same frame.
@@ -35,21 +39,41 @@ class ENGINE_API UMaterialInstanceDynamic : public UMaterialInstance
 	UFUNCTION(BlueprintCallable, meta=(DisplayName = "GetScalarParameterValue", ScriptName = "GetScalarParameterValue", Keywords = "GetFloatParameterValue"), Category="Rendering|Material")
 	float K2_GetScalarParameterValue(FName ParameterName);
 
+	/** Get the current scalar (float) parameter value from an MID, using MPI (to allow access to layer parameters) */
+	UFUNCTION(BlueprintCallable, meta = (DisplayName = "GetScalarParameterValueByInfo", ScriptName = "GetScalarParameterValueByInfo", Keywords = "GetFloatParameterValue"), Category = "Rendering|Material")
+	float K2_GetScalarParameterValueByInfo(const FMaterialParameterInfo& ParameterInfo);
+
 	/** Set an MID texture parameter value */
 	UFUNCTION(BlueprintCallable, Category="Rendering|Material")
 	void SetTextureParameterValue(FName ParameterName, class UTexture* Value);
 
+	/** Set an MID texture parameter value using MPI (to allow access to layer parameters) */
+	UFUNCTION(BlueprintCallable, Category = "Rendering|Material")
+	void SetTextureParameterValueByInfo(const FMaterialParameterInfo& ParameterInfo, class UTexture* Value);
+
 	/** Get the current MID texture parameter value */
 	UFUNCTION(BlueprintCallable, meta=(DisplayName = "GetTextureParameterValue", ScriptName = "GetTextureParameterValue"), Category="Rendering|Material")
 	class UTexture* K2_GetTextureParameterValue(FName ParameterName);
 
+	/** Get the current MID texture parameter value, using MPI (to allow access to layer parameters) */
+	UFUNCTION(BlueprintCallable, meta = (DisplayName = "GetTextureParameterValueByInfo", ScriptName = "GetTextureParameterValueByInfo"), Category = "Rendering|Material")
+	class UTexture* K2_GetTextureParameterValueByInfo(const FMaterialParameterInfo& ParameterInfo);
+
 	/** Set an MID vector parameter value */
 	UFUNCTION(BlueprintCallable, meta=(Keywords = "SetColorParameterValue"), Category="Rendering|Material")
 	void SetVectorParameterValue(FName ParameterName, FLinearColor Value);
 
+	/** Set an MID vector parameter value, using MPI (to allow access to layer parameters) */
+	UFUNCTION(BlueprintCallable, meta = (Keywords = "SetColorParameterValue"), Category = "Rendering|Material")
+	void SetVectorParameterValueByInfo(const FMaterialParameterInfo& ParameterInfo, FLinearColor Value);
+
 	/** Get the current MID vector parameter value */
 	UFUNCTION(BlueprintCallable, meta=(DisplayName = "GetVectorParameterValue", ScriptName = "GetVectorParameterValue", Keywords = "GetColorParameterValue"), Category="Rendering|Material")
 	FLinearColor K2_GetVectorParameterValue(FName ParameterName);
+
+	/** Get the current MID vector parameter value, using MPI (to allow access to layer parameters) */
+	UFUNCTION(BlueprintCallable, meta = (DisplayName = "GetVectorParameterValue", ScriptName = "GetVectorParameterValue", Keywords = "GetColorParameterValue"), Category = "Rendering|Material")
+	FLinearColor K2_GetVectorParameterValueByInfo(const FMaterialParameterInfo& ParameterInfo);
 	
 	/**
 	 * Interpolates the scalar and vector parameters of this material instance based on two other material instances, and an alpha blending factor
diff --git a/Engine/Source/Runtime/Engine/Classes/Materials/MaterialInterface.h b/Engine/Source/Runtime/Engine/Classes/Materials/MaterialInterface.h
--- a/Engine/Source/Runtime/Engine/Classes/Materials/MaterialInterface.h
+++ b/Engine/Source/Runtime/Engine/Classes/Materials/MaterialInterface.h
@@ -574,6 +574,9 @@ public:
 		PURE_VIRTUAL(UMaterialInterface::GetStaticComponentMaskParameterDefaultValue,return false;);
 #endif // WITH_EDITOR
 
+	virtual int32 GetLayerParameterIndex(EMaterialParameterAssociation Association, UMaterialFunctionInterface * LayerFunction) const
+		PURE_VIRTUAL(UMaterialInterface::GetLayerParameterIndex, return INDEX_NONE;);
+
 	/** Get textures referenced by expressions, including nested functions. */
 	virtual TArrayView<UObject* const> GetReferencedTextures() const
 		PURE_VIRTUAL(UMaterialInterface::GetReferencedTextures,return TArrayView<UObject* const>(););
@@ -581,6 +584,9 @@ public:
 	virtual void SaveShaderStableKeysInner(const class ITargetPlatform* TP, const struct FStableShaderKeyAndValue& SaveKeyVal)
 		PURE_VIRTUAL(UMaterialInterface::SaveShaderStableKeysInner, );
 
+	UFUNCTION(BlueprintCallable, Category = "Rendering|Material")
+	ENGINE_API FMaterialParameterInfo GetParameterInfo(EMaterialParameterAssociation Association, FName ParameterName, UMaterialFunctionInterface* LayerFunction) const;
+
 	/** @return The material's relevance. */
 	ENGINE_API FMaterialRelevance GetRelevance(ERHIFeatureLevel::Type InFeatureLevel) const;
 	/** @return The material's relevance, from concurrent render thread updates. */
diff --git a/Engine/Source/Runtime/Engine/Private/Materials/Material.cpp b/Engine/Source/Runtime/Engine/Private/Materials/Material.cpp
--- a/Engine/Source/Runtime/Engine/Private/Materials/Material.cpp
+++ b/Engine/Source/Runtime/Engine/Private/Materials/Material.cpp
@@ -4760,6 +4760,18 @@ bool UMaterial::UpdateLightmassTextureTracking()
 	return bTexturesHaveChanged;
 }
 
+int32 UMaterial::GetLayerParameterIndex(EMaterialParameterAssociation Association, UMaterialFunctionInterface* LayerFunction) const
+{
+	int32 Index = INDEX_NONE;
+	switch (Association)
+	{
+	case BlendParameter: Index = CachedExpressionData.DefaultLayerBlends.Find(LayerFunction); break;
+	case LayerParameter: Index = CachedExpressionData.DefaultLayers.Find(LayerFunction); break;
+	default: checkNoEntry(); break;
+	}
+	return Index;
+}
+
 #if WITH_EDITOR
 FExpressionInput* UMaterial::GetExpressionInputForProperty(EMaterialProperty InProperty)
 {
diff --git a/Engine/Source/Runtime/Engine/Private/Materials/MaterialCachedData.cpp b/Engine/Source/Runtime/Engine/Private/Materials/MaterialCachedData.cpp
--- a/Engine/Source/Runtime/Engine/Private/Materials/MaterialCachedData.cpp
+++ b/Engine/Source/Runtime/Engine/Private/Materials/MaterialCachedData.cpp
@@ -344,6 +344,9 @@ bool FMaterialCachedExpressionData::UpdateForExpressions(const TArray<UMaterialE
 				bResult = false;
 			}
 
+			DefaultLayers = LayersExpression->DefaultLayers.Layers;
+			DefaultLayerBlends = LayersExpression->DefaultLayers.Blends;
+
 			LayersExpression->RebuildLayerGraph(false);
 		}
 	}
diff --git a/Engine/Source/Runtime/Engine/Private/Materials/MaterialInstance.cpp b/Engine/Source/Runtime/Engine/Private/Materials/MaterialInstance.cpp
--- a/Engine/Source/Runtime/Engine/Private/Materials/MaterialInstance.cpp
+++ b/Engine/Source/Runtime/Engine/Private/Materials/MaterialInstance.cpp
@@ -1204,6 +1204,26 @@ bool UMaterialInstance::GetRefractionSettings(float& OutBiasValue) const
 	}
 }
 
+int32 UMaterialInstance::GetLayerParameterIndex(EMaterialParameterAssociation Association, UMaterialFunctionInterface* LayerFunction) const
+{
+	check(Association != GlobalParameter);
+
+	int32 Index = INDEX_NONE;
+	for (const FStaticMaterialLayersParameter& LayersParam : GetStaticParameters().MaterialLayersParameters)
+	{
+		if (LayersParam.bOverride)
+		{
+			if (Association == BlendParameter) Index = LayersParam.Value.Blends.Find(LayerFunction);
+			else if (Association == LayerParameter) Index = LayersParam.Value.Layers.Find(LayerFunction);
+		}
+	}
+	if (Index == INDEX_NONE && Parent)
+	{
+		Index = Parent->GetLayerParameterIndex(Association, LayerFunction);
+	}
+	return Index;
+}
+
 void UMaterialInstance::GetTextureExpressionValues(const FMaterialResource* MaterialResource, TArray<UTexture*>& OutTextures, TArray< TArray<int32> >* OutIndices) const
 {
 	check(MaterialResource);
diff --git a/Engine/Source/Runtime/Engine/Private/Materials/MaterialInstanceDynamic.cpp b/Engine/Source/Runtime/Engine/Private/Materials/MaterialInstanceDynamic.cpp
--- a/Engine/Source/Runtime/Engine/Private/Materials/MaterialInstanceDynamic.cpp
+++ b/Engine/Source/Runtime/Engine/Private/Materials/MaterialInstanceDynamic.cpp
@@ -39,24 +39,41 @@ UMaterialInstanceDynamic* UMaterialInstanceDynamic::Create(UMaterialInterface* P
 
 void UMaterialInstanceDynamic::SetVectorParameterValue(FName ParameterName, FLinearColor Value)
 {
-	FMaterialParameterInfo ParameterInfo(ParameterName); // @TODO: This will only work for non-layered parameters
+	FMaterialParameterInfo ParameterInfo(ParameterName);
 	SetVectorParameterValueInternal(ParameterInfo,Value);
 }
 
+void UMaterialInstanceDynamic::SetVectorParameterValueByInfo(const FMaterialParameterInfo& ParameterInfo, FLinearColor Value)
+{
+	SetVectorParameterValueInternal(ParameterInfo, Value);
+}
+
 FLinearColor UMaterialInstanceDynamic::K2_GetVectorParameterValue(FName ParameterName)
 {
 	FLinearColor Result(0,0,0);
-	FMaterialParameterInfo ParameterInfo(ParameterName); // @TODO: This will only work for non-layered parameters
+	FMaterialParameterInfo ParameterInfo(ParameterName);
+	Super::GetVectorParameterValue(ParameterInfo, Result);
+	return Result;
+}
+
+FLinearColor UMaterialInstanceDynamic::K2_GetVectorParameterValueByInfo(const FMaterialParameterInfo& ParameterInfo)
+{
+	FLinearColor Result(0, 0, 0);
 	Super::GetVectorParameterValue(ParameterInfo, Result);
 	return Result;
 }
 
 void UMaterialInstanceDynamic::SetScalarParameterValue(FName ParameterName, float Value)
 {
-	FMaterialParameterInfo ParameterInfo(ParameterName); // @TODO: This will only work for non-layered parameters
+	FMaterialParameterInfo ParameterInfo(ParameterName);
 	SetScalarParameterValueInternal(ParameterInfo,Value);
 }
 
+void UMaterialInstanceDynamic::SetScalarParameterValueByInfo(const FMaterialParameterInfo& ParameterInfo, float Value)
+{
+	SetScalarParameterValueInternal(ParameterInfo, Value);
+}
+
 bool UMaterialInstanceDynamic::InitializeScalarParameterAndGetIndex(const FName& ParameterName, float Value, int32& OutParameterIndex)
 {
 	OutParameterIndex = INDEX_NONE;
@@ -78,7 +95,7 @@ bool UMaterialInstanceDynamic::InitializeVectorParameterAndGetIndex(const FName&
 {
 	OutParameterIndex = INDEX_NONE;
 
-	FMaterialParameterInfo ParameterInfo(ParameterName); // @TODO: This will only work for non-layered parameters
+	FMaterialParameterInfo ParameterInfo(ParameterName);
 	SetVectorParameterValueInternal(ParameterInfo, Value);
 
 	OutParameterIndex = GameThread_FindParameterIndexByName(VectorParameterValues, ParameterInfo);
@@ -94,7 +111,14 @@ bool UMaterialInstanceDynamic::SetVectorParameterByIndex(int32 ParameterIndex, c
 float UMaterialInstanceDynamic::K2_GetScalarParameterValue(FName ParameterName)
 {
 	float Result = 0.f;
-	FMaterialParameterInfo ParameterInfo(ParameterName); // @TODO: This will only work for non-layered parameters
+	FMaterialParameterInfo ParameterInfo(ParameterName);
+	Super::GetScalarParameterValue(ParameterInfo, Result);
+	return Result;
+}
+
+float UMaterialInstanceDynamic::K2_GetScalarParameterValueByInfo(const FMaterialParameterInfo& ParameterInfo)
+{
+	float Result = 0.f;
 	Super::GetScalarParameterValue(ParameterInfo, Result);
 	return Result;
 }
@@ -104,7 +128,7 @@ void UMaterialInstanceDynamic::SetTextureParameterValue(FName ParameterName, UTe
 	// Save the texture renaming as it will be useful to remap the texture streaming data.
 	UTexture* RenamedTexture = NULL;
 
-	FMaterialParameterInfo ParameterInfo(ParameterName); // @TODO: This will only work for non-layered parameters
+	FMaterialParameterInfo ParameterInfo(ParameterName);
 	Super::GetTextureParameterValue(ParameterInfo, RenamedTexture);
 
 	if (Value && RenamedTexture && Value->GetFName() != RenamedTexture->GetFName())
@@ -115,10 +139,31 @@ void UMaterialInstanceDynamic::SetTextureParameterValue(FName ParameterName, UTe
 	SetTextureParameterValueInternal(ParameterInfo,Value);
 }
 
+void UMaterialInstanceDynamic::SetTextureParameterValueByInfo(const FMaterialParameterInfo& ParameterInfo, UTexture* Value)
+{
+	// Save the texture renaming as it will be useful to remap the texture streaming data.
+	UTexture* RenamedTexture = NULL;
+	Super::GetTextureParameterValue(ParameterInfo, RenamedTexture);
+
+	if (Value && RenamedTexture && Value->GetFName() != RenamedTexture->GetFName())
+	{
+		RenamedTextures.FindOrAdd(Value->GetFName()).AddUnique(RenamedTexture->GetFName());
+	}
+
+	SetTextureParameterValueInternal(ParameterInfo, Value);
+}
+
 UTexture* UMaterialInstanceDynamic::K2_GetTextureParameterValue(FName ParameterName)
 {
 	UTexture* Result = NULL;
-	FMaterialParameterInfo ParameterInfo(ParameterName); // @TODO: This will only work for non-layered parameters
+	FMaterialParameterInfo ParameterInfo(ParameterName);
+	Super::GetTextureParameterValue(ParameterInfo, Result);
+	return Result;
+}
+
+UTexture* UMaterialInstanceDynamic::K2_GetTextureParameterValueByInfo(const FMaterialParameterInfo& ParameterInfo)
+{
+	UTexture* Result = NULL;
 	Super::GetTextureParameterValue(ParameterInfo, Result);
 	return Result;
 }
diff --git a/Engine/Source/Runtime/Engine/Private/Materials/MaterialInterface.cpp b/Engine/Source/Runtime/Engine/Private/Materials/MaterialInterface.cpp
--- a/Engine/Source/Runtime/Engine/Private/Materials/MaterialInterface.cpp
+++ b/Engine/Source/Runtime/Engine/Private/Materials/MaterialInterface.cpp
@@ -213,6 +213,22 @@ FMaterialRelevance UMaterialInterface::GetRelevance_Internal(const UMaterial* Ma
 	}
 }
 
+FMaterialParameterInfo UMaterialInterface::GetParameterInfo(EMaterialParameterAssociation Association, FName ParameterName, UMaterialFunctionInterface* LayerFunction) const
+{
+	int32 Index = 0;
+	if (Association != GlobalParameter)
+	{
+		check(LayerFunction);
+		Index = GetLayerParameterIndex(Association, LayerFunction);
+		if (Index == INDEX_NONE)
+		{
+			return FMaterialParameterInfo();
+		}
+	}
+
+	return FMaterialParameterInfo(ParameterName, Association, Index);
+}
+
 FMaterialRelevance UMaterialInterface::GetRelevance(ERHIFeatureLevel::Type InFeatureLevel) const
 {
 	// Find the interface's concrete material.
diff --git a/Engine/Source/Runtime/Engine/Public/MaterialCachedData.h b/Engine/Source/Runtime/Engine/Public/MaterialCachedData.h
--- a/Engine/Source/Runtime/Engine/Public/MaterialCachedData.h
+++ b/Engine/Source/Runtime/Engine/Public/MaterialCachedData.h
@@ -15,6 +15,7 @@ class UFont;
 class UMaterialExpression;
 class URuntimeVirtualTexture;
 class ULandscapeGrassType;
+class UMaterialFunctionInterface;
 
 /** Stores information about a function that this material references, used to know when the material needs to be recompiled. */
 USTRUCT()
@@ -28,7 +29,7 @@ struct FMaterialFunctionInfo
 
 	/** The function which this material has a dependency on. */
 	UPROPERTY()
-	class UMaterialFunctionInterface* Function = nullptr;
+	UMaterialFunctionInterface* Function = nullptr;
 };
 
 /** Stores information about a parameter collection that this material references, used to know when the material needs to be recompiled. */
@@ -185,6 +186,12 @@ struct FMaterialCachedExpressionData
 	UPROPERTY()
 	TArray<FMaterialParameterCollectionInfo> ParameterCollectionInfos;
 
+	UPROPERTY()
+	TArray<UMaterialFunctionInterface*> DefaultLayers;
+
+	UPROPERTY()
+	TArray<UMaterialFunctionInterface*> DefaultLayerBlends;
+
 	UPROPERTY()
 	TArray<ULandscapeGrassType*> GrassTypes;
 
