Add optional WorldPosition input to Runtime Virtual Texture Sample node
This allows us to manipulate sample coordinates.
One use case is to sample using the world position from before material WPO is applied.
Cherry pick integration from Dev-Rendering

diff --git a/Engine/Source/Runtime/Engine/Classes/Materials/MaterialExpressionRuntimeVirtualTextureSample.h b/Engine/Source/Runtime/Engine/Classes/Materials/MaterialExpressionRuntimeVirtualTextureSample.h
--- a/Engine/Source/Runtime/Engine/Classes/Materials/MaterialExpressionRuntimeVirtualTextureSample.h
+++ b/Engine/Source/Runtime/Engine/Classes/Materials/MaterialExpressionRuntimeVirtualTextureSample.h
@@ -34,10 +34,14 @@ class ENGINE_API UMaterialExpressionRuntimeVirtualTextureSample : public UMateri
 {
 	GENERATED_UCLASS_BODY()
 
-	/** UV coordinates input if we want to override standard world position based coordinates. */
-	UPROPERTY()
+	/** Optional UV coordinates input if we want to override standard world position based coordinates. */
+	UPROPERTY(meta = (RequiredInput = "false"))
 	FExpressionInput Coordinates;
 
+	/** Optional world position input to override the default world position. */
+	UPROPERTY(meta = (RequiredInput = "false"))
+	FExpressionInput WorldPosition;
+
 	/** Meaning depends on MipValueMode. A single unit is one mip level.  */
 	UPROPERTY(meta = (RequiredInput = "false"))
 	FExpressionInput MipValue;
diff --git a/Engine/Source/Runtime/Engine/Private/Materials/MaterialExpressions.cpp b/Engine/Source/Runtime/Engine/Private/Materials/MaterialExpressions.cpp
--- a/Engine/Source/Runtime/Engine/Private/Materials/MaterialExpressions.cpp
+++ b/Engine/Source/Runtime/Engine/Private/Materials/MaterialExpressions.cpp
@@ -2312,9 +2312,28 @@ int32 UMaterialExpressionRuntimeVirtualTextureSample::Compile(class FMaterialCom
 	// Compile the coordinates
 	// We use the virtual texture world space transform by default
 	int32 CoordinateIndex = INDEX_NONE;
-	if (Coordinates.GetTracedInput().Expression == nullptr)
+
+	if (Coordinates.GetTracedInput().Expression != nullptr && WorldPosition.GetTracedInput().Expression != nullptr)
+	{
+		Compiler->Errorf(TEXT("Only one of 'Coordinates' and 'WorldPosition' can be used"));
+	}
+
+	if (Coordinates.GetTracedInput().Expression != nullptr)
+	{
+		CoordinateIndex = Coordinates.Compile(Compiler);
+	}
+	else
 	{
-		int32 WorldPositionIndex = Compiler->WorldPosition(WPT_Default);
+		int32 WorldPositionIndex = INDEX_NONE;
+		if (WorldPosition.GetTracedInput().Expression != nullptr)
+		{
+			WorldPositionIndex = WorldPosition.Compile(Compiler);
+		}
+		else
+		{
+			WorldPositionIndex = Compiler->WorldPosition(WPT_Default);
+		}
+		
 		int32 P0, P1, P2;
 		if (bIsParameter)
 		{
@@ -2328,12 +2347,9 @@ int32 UMaterialExpressionRuntimeVirtualTextureSample::Compile(class FMaterialCom
 			P1 = Compiler->VirtualTextureUniform(TextureReferenceIndex[0], 1);
 			P2 = Compiler->VirtualTextureUniform(TextureReferenceIndex[0], 2);
 		}
+
 		CoordinateIndex = Compiler->VirtualTextureWorldToUV(WorldPositionIndex, P0, P1, P2);
 	}
-	else
-	{
-		CoordinateIndex = Coordinates.Compile(Compiler);
-	}
 	
 	// Compile the mip level for the current mip value mode
 	ETextureMipValueMode TextureMipLevelMode = TMVM_None;
