Recoup about 10 minutes of cook time in large project cooks due to repeated calls to UMaterial::UpdateCachedExpressionData.

diff --git a/Engine/Source/Runtime/Engine/Private/Materials/Material.cpp b/Engine/Source/Runtime/Engine/Private/Materials/Material.cpp
--- a/Engine/Source/Runtime/Engine/Private/Materials/Material.cpp
+++ b/Engine/Source/Runtime/Engine/Private/Materials/Material.cpp
@@ -963,7 +963,6 @@ void UMaterial::PreSave(const class ITargetPlatform* TargetPlatform)
 	Super::PreSave(TargetPlatform);
 #if WITH_EDITOR
 	GMaterialsWithDirtyUsageFlags.RemoveAnnotation(this);
-	UpdateCachedExpressionData();
 #endif
 }
 
@@ -3079,7 +3078,7 @@ void UMaterial::CacheResourceShadersForCooking(EShaderPlatform ShaderPlatform, T
 void UMaterial::CacheShadersForResources(EShaderPlatform ShaderPlatform, const TArray<FMaterialResource*>& ResourcesToCache, const ITargetPlatform* TargetPlatform)
 {
 #if WITH_EDITOR
-	UpdateCachedExpressionData();
+	check(!HasAnyFlags(RF_NeedPostLoad));
 #endif
 	for (int32 ResourceIndex = 0; ResourceIndex < ResourcesToCache.Num(); ResourceIndex++)
 	{
diff --git a/Engine/Source/Runtime/Engine/Private/Materials/MaterialInstance.cpp b/Engine/Source/Runtime/Engine/Private/Materials/MaterialInstance.cpp
--- a/Engine/Source/Runtime/Engine/Private/Materials/MaterialInstance.cpp
+++ b/Engine/Source/Runtime/Engine/Private/Materials/MaterialInstance.cpp
@@ -3055,8 +3055,8 @@ void UMaterialInstance::CacheShadersForResources(EShaderPlatform ShaderPlatform,
 {
 	UMaterial* BaseMaterial = GetMaterial();
 #if WITH_EDITOR
-	BaseMaterial->UpdateCachedExpressionData();
-	UpdateCachedLayerParameters();
+	check(!HasAnyFlags(RF_NeedPostLoad));
+	check(!BaseMaterial->HasAnyFlags(RF_NeedPostLoad));
 #endif
 
 	for (int32 ResourceIndex = 0; ResourceIndex < ResourcesToCache.Num(); ResourceIndex++)
