Const-corrected a few transform accessors in FAnimInstanceProxy

From UDN: https://udn.unrealengine.com/questions/555546/faniminstanceproxy-contains-non-const-accessors-th.html

diff --git a/Engine/Source/Runtime/Engine/Public/Animation/AnimInstanceProxy.h b/Engine/Source/Runtime/Engine/Public/Animation/AnimInstanceProxy.h
--- a/Engine/Source/Runtime/Engine/Public/Animation/AnimInstanceProxy.h
+++ b/Engine/Source/Runtime/Engine/Public/Animation/AnimInstanceProxy.h
@@ -450,13 +450,13 @@ public:
 	void ResetDynamics();
 
 	/** Get the relative transform of the component we are running on */
-	const FTransform& GetComponentRelativeTransform() { return ComponentRelativeTransform; }
+	const FTransform& GetComponentRelativeTransform() const { return ComponentRelativeTransform; }
 
 	/** Get the component to world transform of the component we are running on */
-	const FTransform& GetComponentTransform() { return ComponentTransform; }
+	const FTransform& GetComponentTransform() const { return ComponentTransform; }
 
 	/** Get the transform of the actor we are running on */
-	const FTransform& GetActorTransform() { return ActorTransform; }
+	const FTransform& GetActorTransform() const { return ActorTransform; }
 
 #if ANIM_TRACE_ENABLED
 	// Trace montage debug data for the specified slot
