Fixing duplicated script name in metadata

diff --git a/Engine/Source/Runtime/Engine/Classes/Materials/MaterialInstanceDynamic.h b/Engine/Source/Runtime/Engine/Classes/Materials/MaterialInstanceDynamic.h
--- a/Engine/Source/Runtime/Engine/Classes/Materials/MaterialInstanceDynamic.h
+++ b/Engine/Source/Runtime/Engine/Classes/Materials/MaterialInstanceDynamic.h
@@ -72,7 +72,7 @@ class ENGINE_API UMaterialInstanceDynamic : public UMaterialInstance
 	FLinearColor K2_GetVectorParameterValue(FName ParameterName);
 
 	/** Get the current MID vector parameter value, using MPI (to allow access to layer parameters) */
-	UFUNCTION(BlueprintCallable, meta = (DisplayName = "GetVectorParameterValue", ScriptName = "GetVectorParameterValue", Keywords = "GetColorParameterValue"), Category = "Rendering|Material")
+	UFUNCTION(BlueprintCallable, meta = (DisplayName = "GetVectorParameterValueByInfo", ScriptName = "GetVectorParameterValueByInfo", Keywords = "GetColorParameterValue"), Category = "Rendering|Material")
 	FLinearColor K2_GetVectorParameterValueByInfo(const FMaterialParameterInfo& ParameterInfo);
 	
 	/**
