Adding pin names for the new anisotropic pins

diff --git a/Engine/Source/Runtime/Engine/Private/Materials/MaterialShared.cpp b/Engine/Source/Runtime/Engine/Private/Materials/MaterialShared.cpp
--- a/Engine/Source/Runtime/Engine/Private/Materials/MaterialShared.cpp
+++ b/Engine/Source/Runtime/Engine/Private/Materials/MaterialShared.cpp
@@ -3675,9 +3675,13 @@ FText FMaterialAttributeDefinitionMap::GetAttributeOverrideForMaterial(const FGu
 		return LOCTEXT("Specular", "Specular");
 	case MP_Roughness:
 		return LOCTEXT("Roughness", "Roughness");
+	case MP_Anisotropy:
+		return LOCTEXT("Anisotropy", "Anisotropy");
 	case MP_Normal:
 		CustomPinNames.Add({MSM_Hair, "Tangent"});
 		return FText::FromString(GetPinNameFromShadingModelField(Material->GetShadingModels(), CustomPinNames, "Normal"));
+	case MP_Tangent:
+		return LOCTEXT("Tangent", "Tangent");
 	case MP_WorldPositionOffset:
 		return Material->IsUIMaterial() ? LOCTEXT("ScreenPosition", "Screen Position") : LOCTEXT("WorldPositionOffset", "World Position Offset");
 	case MP_WorldDisplacement:
