Update virtual texture on MIC change
Cherry pick integration from Dev-Rendering

diff --git a/Engine/Source/Runtime/Engine/Private/Materials/MaterialInstance.cpp b/Engine/Source/Runtime/Engine/Private/Materials/MaterialInstance.cpp
--- a/Engine/Source/Runtime/Engine/Private/Materials/MaterialInstance.cpp
+++ b/Engine/Source/Runtime/Engine/Private/Materials/MaterialInstance.cpp
@@ -3,6 +3,7 @@
 #include "Materials/MaterialInstance.h"
 #include "Stats/StatsMisc.h"
 #include "EngineGlobals.h"
+#include "EngineModule.h"
 #include "BatchedElements.h"
 #include "Engine/Font.h"
 #include "UObject/UObjectHash.h"
@@ -4203,6 +4204,17 @@ void UMaterialInstance::PostEditChangeProperty(FPropertyChangedEvent& PropertyCh
 	}
 
 	UpdateCachedLayerParameters();
+
+	if (GIsEditor)
+	{
+		// Update virtual textures if necessary
+		//todo[vt]: With many streaming virtual textures this will be slow. Add a flush for runtime virtual texture only? Or maybe work out which primitives or runtime virtual textures are affected flush only affected VT pages?
+		const UMaterial* BaseMaterial = GetMaterial();
+		if (BaseMaterial != nullptr && (BaseMaterial->MaterialDomain == EMaterialDomain::MD_RuntimeVirtualTexture || BaseMaterial->GetCachedExpressionData().bHasRuntimeVirtualTextureOutput))
+		{
+			ENQUEUE_RENDER_COMMAND(FlushVTCommand)([ResourcePtr = Resource](FRHICommandListImmediate& RHICmdList) {	GetRendererModule().FlushVirtualTextureCache();	});
+		}
+	}
 }
 
 #endif // WITH_EDITOR
