When looking up a blueprint property's name from a guid we now prefer names in the super class, when looking up a guid from a name we still prefer guids in the derived class. No longer copying all guids into derived types on compile. This fixes loss of data in derived types when renaming a property

diff --git a/Engine/Source/Editor/KismetCompiler/Private/KismetCompiler.cpp b/Engine/Source/Editor/KismetCompiler/Private/KismetCompiler.cpp
--- a/Engine/Source/Editor/KismetCompiler/Private/KismetCompiler.cpp
+++ b/Engine/Source/Editor/KismetCompiler/Private/KismetCompiler.cpp
@@ -680,16 +680,6 @@ void FKismetCompilerContext::CreateClassVariablesFromBlueprint()
 	if (bRebuildPropertyMap)
 	{
 		NewClass->PropertyGuids.Reset();
-		// Add any chained parent blueprint map values
-		UBlueprint* ParentBP = Cast<UBlueprint>(Blueprint->ParentClass->ClassGeneratedBy);
-		while (ParentBP)
-		{
-			if (UBlueprintGeneratedClass* ParentBPGC = Cast<UBlueprintGeneratedClass>(ParentBP->GeneratedClass))
-			{
-				NewClass->PropertyGuids.Append(ParentBPGC->PropertyGuids);
-			}
-			ParentBP = Cast<UBlueprint>(ParentBP->ParentClass->ClassGeneratedBy);
-		}
 	}
 
 	for (int32 i = 0; i < Blueprint->NewVariables.Num(); ++i)
diff --git a/Engine/Source/Runtime/Engine/Private/BlueprintGeneratedClass.cpp b/Engine/Source/Runtime/Engine/Private/BlueprintGeneratedClass.cpp
--- a/Engine/Source/Runtime/Engine/Private/BlueprintGeneratedClass.cpp
+++ b/Engine/Source/Runtime/Engine/Private/BlueprintGeneratedClass.cpp
@@ -1906,6 +1906,15 @@ FName UBlueprintGeneratedClass::FindPropertyNameFromGuid(const FGuid& PropertyGu
 {
 	FName RedirectedName = NAME_None;
 #if WITH_EDITORONLY_DATA
+	if (UBlueprintGeneratedClass* Super = Cast<UBlueprintGeneratedClass>(GetSuperStruct()))
+	{
+		FName NameFromSuper = Super->FindPropertyNameFromGuid(PropertyGuid);
+		if (NameFromSuper != FName())
+		{
+			return NameFromSuper;
+		}
+	}
+
 	if (const FName* Result = PropertyGuids.FindKey(PropertyGuid))
 	{
 		RedirectedName = *Result;
@@ -1922,6 +1931,10 @@ FGuid UBlueprintGeneratedClass::FindPropertyGuidFromName(const FName InName) con
 	{
 		PropertyGuid = *Result;
 	}
+	else if (UBlueprintGeneratedClass* Super = Cast<UBlueprintGeneratedClass>(GetSuperStruct()))
+	{
+		PropertyGuid = Super->FindPropertyGuidFromName(InName);
+	}
 #endif // WITH_EDITORONLY_DATA
 	return PropertyGuid;
 }
\ No newline at end of file
