Adds ENGINE_API on APostProcessVolume

diff --git a/Engine/Source/Runtime/Engine/Classes/Engine/PostProcessVolume.h b/Engine/Source/Runtime/Engine/Classes/Engine/PostProcessVolume.h
--- a/Engine/Source/Runtime/Engine/Classes/Engine/PostProcessVolume.h
+++ b/Engine/Source/Runtime/Engine/Classes/Engine/PostProcessVolume.h
@@ -18,8 +18,8 @@
 #include "PostProcessVolume.generated.h"
 
 	// for FPostprocessSettings
-UCLASS(autoexpandcategories=PostProcessVolume, hidecategories=(Advanced, Collision, Volume, Brush, Attachment), MinimalAPI)
-class APostProcessVolume : public AVolume, public IInterface_PostProcessVolume
+UCLASS(autoexpandcategories=PostProcessVolume, hidecategories=(Advanced, Collision, Volume, Brush, Attachment))
+class ENGINE_API APostProcessVolume : public AVolume, public IInterface_PostProcessVolume
 {
 	GENERATED_UCLASS_BODY()
 
@@ -51,8 +51,8 @@ class APostProcessVolume : public AVolume, public IInterface_PostProcessVolume
 	uint32 bUnbound:1;
 
 	//~ Begin IInterface_PostProcessVolume Interface
-	ENGINE_API virtual bool EncompassesPoint(FVector Point, float SphereRadius/*=0.f*/, float* OutDistanceToPoint) override;
-	ENGINE_API virtual FPostProcessVolumeProperties GetProperties() const override
+	virtual bool EncompassesPoint(FVector Point, float SphereRadius/*=0.f*/, float* OutDistanceToPoint) override;
+	virtual FPostProcessVolumeProperties GetProperties() const override
 	{
 		FPostProcessVolumeProperties Ret;
 		Ret.bIsEnabled = bEnabled != 0;
@@ -84,7 +84,7 @@ public:
 
 	/** Adds an Blendable (implements IBlendableInterface) to the array of Blendables (if it doesn't exist) and update the weight */
 	UFUNCTION(BlueprintCallable, Category="Rendering")
-	ENGINE_API void AddOrUpdateBlendable(TScriptInterface<IBlendableInterface> InBlendableObject, float InWeight = 1.0f) { Settings.AddBlendable(InBlendableObject, InWeight); }
+	void AddOrUpdateBlendable(TScriptInterface<IBlendableInterface> InBlendableObject, float InWeight = 1.0f) { Settings.AddBlendable(InBlendableObject, InWeight); }
 };
 
 
