Fix for a crash when gathering asset registry tags from a cooked blueprint. The source blueprint won't exist at this point, so don't try and blindly access it without checking it exists first!

diff --git a/Engine/Source/Runtime/Engine/Private/Blueprint.cpp b/Engine/Source/Runtime/Engine/Private/Blueprint.cpp
--- a/Engine/Source/Runtime/Engine/Private/Blueprint.cpp
+++ b/Engine/Source/Runtime/Engine/Private/Blueprint.cpp
@@ -1065,7 +1065,7 @@ void UBlueprint::GetAssetRegistryTags(TArray<FAssetRegistryTag>& OutTags) const
 		for (UBlueprintGeneratedClass* TestBPClass = BlueprintClass; TestBPClass != nullptr; TestBPClass = Cast<UBlueprintGeneratedClass>(TestBPClass->GetSuperClass()))
 		{
 			const UBlueprint* AssociatedBP = Cast<const UBlueprint>(TestBPClass->ClassGeneratedBy);
-			if (AssociatedBP->SimpleConstructionScript != nullptr)
+			if (AssociatedBP && AssociatedBP->SimpleConstructionScript != nullptr)
 			{
 				NumAddedComponents += AssociatedBP->SimpleConstructionScript->GetAllNodesConst().Num();
 			}
