Merging CL 11586770 //UE4/Dev-Core to Release-4.25 (//UE4/Release-4.25)

diff --git a/Engine/Source/Runtime/Engine/Private/AssetManager.cpp b/Engine/Source/Runtime/Engine/Private/AssetManager.cpp
--- a/Engine/Source/Runtime/Engine/Private/AssetManager.cpp
+++ b/Engine/Source/Runtime/Engine/Private/AssetManager.cpp
@@ -1230,7 +1230,7 @@ TSharedPtr<FStreamableHandle> UAssetManager::ChangeBundleStateForPrimaryAssets(c
 				if (Entry.IsValid())
 				{
 					PathsToLoad.Append(Entry.BundleAssets);
-					bNothingToLoadForAsset &= Entry.BundleAssets.Num() == 0;
+					bNothingToLoadForAsset = false;
 				}
 				else
 				{
