Merging CL 11586812 //UE4/Dev-Core to Release-4.25 (//UE4/Release-4.25)

diff --git a/Engine/Source/Runtime/Engine/Classes/Engine/AssetManager.h b/Engine/Source/Runtime/Engine/Classes/Engine/AssetManager.h
--- a/Engine/Source/Runtime/Engine/Classes/Engine/AssetManager.h
+++ b/Engine/Source/Runtime/Engine/Classes/Engine/AssetManager.h
@@ -524,7 +524,7 @@ protected:
 	virtual void RebuildObjectReferenceList();
 
 	/** Called when an internal load handle finishes, handles setting to pending state */
-	virtual void OnAssetStateChangeCompleted(TArray<FPrimaryAssetId> PrimaryAssetId, TSharedPtr<FStreamableHandle> BoundHandle, FStreamableDelegate WrappedDelegate);
+	virtual void OnAssetStateChangeCompleted(FPrimaryAssetId PrimaryAssetId, TSharedPtr<FStreamableHandle> BoundHandle, FStreamableDelegate WrappedDelegate);
 
 	/** Helper function to write out asset reports */
 	virtual bool WriteCustomReport(FString FileName, TArray<FString>& FileLines) const;
diff --git a/Engine/Source/Runtime/Engine/Private/AssetManager.cpp b/Engine/Source/Runtime/Engine/Private/AssetManager.cpp
--- a/Engine/Source/Runtime/Engine/Private/AssetManager.cpp
+++ b/Engine/Source/Runtime/Engine/Private/AssetManager.cpp
@@ -76,17 +76,11 @@ struct FPrimaryAssetData
 
 	/** Pending state of this asset, will be copied to CurrentState when load finishes */
 	FPrimaryAssetLoadState PendingState;
-	
-	/** Asset is considered loaded at all if there is an active handle for it */
-	bool IsLoaded() const { return CurrentState.IsValid(); }
 
-	void MovePendingToCurrent()
-	{
-		Swap(PendingState.BundleNames, CurrentState.BundleNames);
-		Swap(PendingState.Handle, CurrentState.Handle);
+	FPrimaryAssetData() {}
 
-		PendingState.Reset(/* don't cancel handle */ false);	
-	}
+	/** Asset is considered loaded at all if there is an active handle for it */
+	bool IsLoaded() const { return CurrentState.IsValid(); }
 };
 
 /** Structure representing all items of a specific asset type */
@@ -1133,25 +1127,9 @@ void UAssetManager::GetPrimaryAssetTypeInfoList(TArray<FPrimaryAssetTypeInfo>& A
 
 TSharedPtr<FStreamableHandle> UAssetManager::ChangeBundleStateForPrimaryAssets(const TArray<FPrimaryAssetId>& AssetsToChange, const TArray<FName>& AddBundles, const TArray<FName>& RemoveBundles, bool bRemoveAllBundles, FStreamableDelegate DelegateToCall, TAsyncLoadPriority Priority)
 {
-	// ExistingHandles are both primary and secondary assets who are already being loaded.
-	// We need to return a combined handle that waits for newly requested assets as well as
-	// these previously requested assets.
-	//
-	// The code relies on TSet iteraton order being the insertion order.
-	// This matters since FStreamableHandle::GetLoadedAsset() returns the first requested asset
-	// or the first requested asset of the first child handle. 
-
-	TSet<TSharedPtr<FStreamableHandle>> ExistingHandles;
-	TSet<FSoftObjectPath> PathsToLoad;
-	PathsToLoad.Reserve(AssetsToChange.Num());
-
-	TArray<FName> AddBundlesSorted(AddBundles);
-	AddBundlesSorted.Sort(FNameLexicalLess());
-	TArray<FName> TempBundles;
-
-	TArray<FPrimaryAssetData*> NameDatasToLoad;
-	NameDatasToLoad.Reserve(AssetsToChange.Num());
-	TBitArray<> IdsToLoad(false, AssetsToChange.Num());
+	TArray<TSharedPtr<FStreamableHandle> > NewHandles, ExistingHandles;
+	TArray<FPrimaryAssetId> NewAssets;
+	TSharedPtr<FStreamableHandle> ReturnHandle;
 
 	for (const FPrimaryAssetId& PrimaryAssetId : AssetsToChange)
 	{
@@ -1162,44 +1140,31 @@ TSharedPtr<FStreamableHandle> UAssetManager::ChangeBundleStateForPrimaryAssets(c
 			// Iterate list of changes, compute new bundle set
 			bool bLoadIfNeeded = false;
 			
-			// NewBundles will be the set of bundle tags we want to have loaded,
-			// i.e. AddBundles + CurrentBundles - RemoveBundles
-			const TArray<FName>& CurrentBundles = NameData->PendingState.IsValid() ? NameData->PendingState.BundleNames : NameData->CurrentState.BundleNames;
-			const TArray<FName>* NewBundles;
+			// Use pending state if valid
+			TArray<FName> CurrentBundleState = NameData->PendingState.IsValid() ? NameData->PendingState.BundleNames : NameData->CurrentState.BundleNames;
+			TArray<FName> NewBundleState;
 
-			if (bRemoveAllBundles || CurrentBundles.Num() == 0)
-			{
-				// Use AddBundlesSorted to avoid creating temporary array
-				NewBundles = &AddBundlesSorted;
-			}
-			else
+			if (!bRemoveAllBundles)
 			{
-				// Reuse TempBundles to avoid allocations and start with the bundles we want to be loaded
-				NewBundles = &TempBundles;
-				TempBundles = AddBundlesSorted;
+				NewBundleState = CurrentBundleState;
 
-				// Keep CurrentBundles unless they already exist or we explictly should remove them
-				// It is assumed that RemoveBundles do not overlap with AddBundles
-				for (FName CurrentBundle : CurrentBundles)
+				for (const FName& RemoveBundle : RemoveBundles)
 				{
-					if (RemoveBundles.Find(CurrentBundle) == INDEX_NONE &&
-						TempBundles.Find(CurrentBundle) == INDEX_NONE)
-					{
-						TempBundles.Add(CurrentBundle);
-					}
+					NewBundleState.Remove(RemoveBundle);
 				}
+			}
 
-				// Sort bundles if we added anything on top of AddBundlesSorted
-				if (TempBundles.Num() > AddBundlesSorted.Num())
-				{
-					TempBundles.Sort(FNameLexicalLess());
-				}
+			for (const FName& AddBundle : AddBundles)
+			{
+				NewBundleState.AddUnique(AddBundle);
 			}
 
+			NewBundleState.Sort(FNameLexicalLess());
+
 			// If the pending state is valid, check if it is different
 			if (NameData->PendingState.IsValid())
 			{
-				if (NameData->PendingState.BundleNames == *NewBundles)
+				if (NameData->PendingState.BundleNames == NewBundleState)
 				{
 					// This will wait on any existing handles to finish
 					ExistingHandles.Add(NameData->PendingState.Handle);
@@ -1208,29 +1173,30 @@ TSharedPtr<FStreamableHandle> UAssetManager::ChangeBundleStateForPrimaryAssets(c
 				// Clear pending state
 				NameData->PendingState.Reset(true);
 			}
-			else if (NameData->CurrentState.IsValid() && NameData->CurrentState.BundleNames == *NewBundles)
+			else if (NameData->CurrentState.IsValid() && NameData->CurrentState.BundleNames == NewBundleState)
 			{
 				// If no pending, compare with current
 				continue;
 			}
 
+			TSet<FSoftObjectPath> PathsToLoad;
+
 			// Gather asset refs
 			const FSoftObjectPath& AssetPath = NameData->AssetPtr.ToSoftObjectPath();
-			bool bNothingToLoadForAsset = AssetPath.IsNull();
+
 			if (!AssetPath.IsNull())
 			{
 				// Dynamic types can have no base asset path
 				PathsToLoad.Add(AssetPath);
 			}
 			
-			for (const FName& BundleName : *NewBundles)
+			for (const FName& BundleName : NewBundleState)
 			{
 				FAssetBundleEntry Entry = GetAssetBundleEntry(PrimaryAssetId, BundleName);
 
 				if (Entry.IsValid())
 				{
 					PathsToLoad.Append(Entry.BundleAssets);
-					bNothingToLoadForAsset = false;
 				}
 				else
 				{
@@ -1238,118 +1204,100 @@ TSharedPtr<FStreamableHandle> UAssetManager::ChangeBundleStateForPrimaryAssets(c
 				}
 			}
 
-			if (bNothingToLoadForAsset)
+			TSharedPtr<FStreamableHandle> NewHandle;
+
+			FString DebugName = PrimaryAssetId.ToString();
+
+			if (NewBundleState.Num() > 0)
+			{
+				DebugName += TEXT(" (");
+
+				for (int32 i = 0; i < NewBundleState.Num(); i++)
+				{
+					if (i != 0)
+					{
+						DebugName += TEXT(", ");
+					}
+					DebugName += NewBundleState[i].ToString();
+				}
+
+				DebugName += TEXT(")");
+			}
+
+			if (PathsToLoad.Num() == 0)
 			{
 				// New state has no assets to load. Set the CurrentState's bundles and clear the handle
-				NameData->CurrentState.BundleNames = *NewBundles;
+				NameData->CurrentState.BundleNames = NewBundleState;
 				NameData->CurrentState.Handle.Reset();
 				continue;
 			}
 
-			NameDatasToLoad.Add(NameData);
-			NameData->PendingState.BundleNames = *NewBundles;
-			IdsToLoad[&PrimaryAssetId - &AssetsToChange[0]] = true;
-		}
-	}
+			NewHandle = LoadAssetList(PathsToLoad.Array(), FStreamableDelegate(), Priority, DebugName);
 
-	TSharedPtr<FStreamableHandle> ReturnHandle;
-	TStringBuilder<256> DebugName;
-	if (PathsToLoad.Num() > 0)
-	{
-		DebugName << *PathsToLoad.CreateConstIterator();
+			if (!NewHandle.IsValid())
+			{
+				// LoadAssetList already throws an error, no need to do it here as well
+				continue;
+			}
 
-		if (NameDatasToLoad.Num() > 1)
-		{
-			DebugName << " and " << NameDatasToLoad.Num() - 1 << " more";
-		}
-			
-		DebugName << TEXT(" (");
-		int32 LenWithoutElements = DebugName.Len();
-		for (FName AddBundle : AddBundlesSorted)
-		{
-			if (!bRemoveAllBundles)
+			if (NewHandle->HasLoadCompleted())
 			{
-				DebugName << '+';
-			}					
-				
-			DebugName << AddBundle << ", ";
-		}
+				// Copy right into active
+				NameData->CurrentState.BundleNames = NewBundleState;
+				NameData->CurrentState.Handle = NewHandle;
+			}
+			else
+			{
+				// Copy into pending and set delegate
+				NameData->PendingState.BundleNames = NewBundleState;
+				NameData->PendingState.Handle = NewHandle;
 
-		for (FName RemoveBundle : RemoveBundles)
-		{
-			DebugName << '-' << RemoveBundle << ", ";
-		}
+				NewHandle->BindCompleteDelegate(FStreamableDelegate::CreateUObject(this, &UAssetManager::OnAssetStateChangeCompleted, PrimaryAssetId, NewHandle, FStreamableDelegate()));
+			}
 
-		// Trim last comma and space
-		if (DebugName.Len() > LenWithoutElements)
-		{
-			DebugName.RemoveSuffix(2);
+			NewHandles.Add(NewHandle);
+			NewAssets.Add(PrimaryAssetId);
 		}
-		DebugName << TEXT(")");
+	}
 
+	if (NewHandles.Num() > 1 || ExistingHandles.Num() > 0)
+	{
+		// If multiple handles or we have an old handle, need to make wrapper handle
+		NewHandles.Append(ExistingHandles);
 
-		ReturnHandle = LoadAssetList(PathsToLoad.Array(), FStreamableDelegate(), Priority, FString(DebugName));
+		ReturnHandle = StreamableManager.CreateCombinedHandle(NewHandles, FString::Printf(TEXT("%s CreateCombinedHandle"), *GetName()));
 
-		// Update FPrimaryAssetData pending handle
-		for (FPrimaryAssetData* NameData : NameDatasToLoad)
+		// Call delegate or bind to meta handle
+		if (ReturnHandle->HasLoadCompleted())
 		{
-			NameData->PendingState.Handle = ReturnHandle;
-		}
-	}		
-		
-	// Create combined wrapper handle from existing handles
-	if (ExistingHandles.Num() > 0)
-	{
-		// If there is only one handle we don't need a combined handle
-		if (!ReturnHandle.IsValid() && ExistingHandles.Num() == 1)
-		{
-			ReturnHandle = *ExistingHandles.CreateConstIterator();
+			FStreamableHandle::ExecuteDelegate(DelegateToCall);
 		}
 		else
 		{
-			TArray<TSharedPtr<FStreamableHandle>> AllHandles;
-			AllHandles.Reserve(ReturnHandle.IsValid() + ExistingHandles.Num());
-			if (ReturnHandle.IsValid())
-			{
-				AllHandles.Add(ReturnHandle);
-			}
-			for (const TSharedPtr<FStreamableHandle>& ExistingHandle : ExistingHandles)
-			{
-				AllHandles.Add(ExistingHandle);
-			}
-
-			DebugName << " combined";
-
-			ReturnHandle = StreamableManager.CreateCombinedHandle(AllHandles, *DebugName);
+			// Call external callback when completed
+			ReturnHandle->BindCompleteDelegate(DelegateToCall);
 		}
 	}
-
-	// Call delegate and update NameDatasToLoad when done
-	if (!ReturnHandle.IsValid() || ReturnHandle->HasLoadCompleted())
+	else if (NewHandles.Num() == 1)
 	{
-		for (FPrimaryAssetData* NameData : NameDatasToLoad)
+		ReturnHandle = NewHandles[0];
+		ensure(NewAssets.Num() == 1);
+
+		// If only one handle, return it and add callback
+		if (ReturnHandle->HasLoadCompleted())
 		{
-			NameData->MovePendingToCurrent();
+			FStreamableHandle::ExecuteDelegate(DelegateToCall);
+		}
+		else
+		{
+			// Call internal callback and external callback when it finishes
+			ReturnHandle->BindCompleteDelegate(FStreamableDelegate::CreateUObject(this, &UAssetManager::OnAssetStateChangeCompleted, NewAssets[0], ReturnHandle, DelegateToCall));
 		}
-
-		DelegateToCall.ExecuteIfBound();
 	}
 	else
 	{
-		// NameDatasToLoad need updating on completion but they can be deleted. 
-		// Create id array instead and redo Id -> NameData lookup on completion. 
-		TArray<FPrimaryAssetId> NameDataIds;
-		NameDataIds.Reserve(NameDatasToLoad.Num());
-		for (const FPrimaryAssetId& PrimaryAssetId : AssetsToChange)
-		{
-			if (IdsToLoad[&PrimaryAssetId - &AssetsToChange[0]])
-			{
-				NameDataIds.Add(PrimaryAssetId);
-			}
-		}
-		check(NameDatasToLoad.Num() == NameDataIds.Num())
-
-		ReturnHandle->BindCompleteDelegate(FStreamableDelegate::CreateUObject(this, &UAssetManager::OnAssetStateChangeCompleted, MoveTemp(NameDataIds), ReturnHandle, DelegateToCall));
+		// Call completion callback, nothing to do
+		FStreamableHandle::ExecuteDelegate(DelegateToCall);
 	}
 
 	return ReturnHandle;
@@ -1440,28 +1388,31 @@ TSharedPtr<FStreamableHandle> UAssetManager::PreloadPrimaryAssets(const TArray<F
 	return ReturnHandle;
 }
 
-void UAssetManager::OnAssetStateChangeCompleted(TArray<FPrimaryAssetId> PrimaryAssetIds,  TSharedPtr<FStreamableHandle> BoundHandle, FStreamableDelegate WrappedDelegate)
+void UAssetManager::OnAssetStateChangeCompleted(FPrimaryAssetId PrimaryAssetId, TSharedPtr<FStreamableHandle> BoundHandle, FStreamableDelegate WrappedDelegate)
 {
-	for (const FPrimaryAssetId& PrimaryAssetId : PrimaryAssetIds)
+	FPrimaryAssetData* NameData = GetNameData(PrimaryAssetId);
+
+	if (NameData)
 	{
-		if (FPrimaryAssetData* NameData = GetNameData(PrimaryAssetId))
+		if (NameData->PendingState.Handle == BoundHandle)
 		{
-			if (NameData->PendingState.Handle == BoundHandle)
-			{
-				NameData->MovePendingToCurrent();
-			}
-			else
-			{
-				UE_LOG(LogAssetManager, Verbose, TEXT("OnAssetStateChangeCompleted: Received after pending data changed, ignoring (%s)"), *PrimaryAssetId.ToString());
-			}
+			NameData->CurrentState.Handle = NameData->PendingState.Handle;
+			NameData->CurrentState.BundleNames = NameData->PendingState.BundleNames;
+
+			WrappedDelegate.ExecuteIfBound();
+
+			// Clear old state, but don't cancel handle as we just copied it into current
+			NameData->PendingState.Reset(false);
 		}
 		else
 		{
-			UE_LOG(LogAssetManager, Error, TEXT("OnAssetStateChangeCompleted: Received for invalid asset! (%s)"), *PrimaryAssetId.ToString());
+			UE_LOG(LogAssetManager, Verbose, TEXT("OnAssetStateChangeCompleted: Received after pending data changed, ignoring (%s)"), *PrimaryAssetId.ToString());
 		}
 	}
-
-	WrappedDelegate.ExecuteIfBound();
+	else
+	{
+		UE_LOG(LogAssetManager, Error, TEXT("OnAssetStateChangeCompleted: Received for invalid asset! (%s)"), *PrimaryAssetId.ToString());
+	}
 }
 
 TSharedPtr<FStreamableHandle> UAssetManager::LoadPrimaryAssets(const TArray<FPrimaryAssetId>& AssetsToLoad, const TArray<FName>& LoadBundles, FStreamableDelegate DelegateToCall, TAsyncLoadPriority Priority)
diff --git a/Engine/Source/Runtime/Engine/Private/StreamableManager.cpp b/Engine/Source/Runtime/Engine/Private/StreamableManager.cpp
--- a/Engine/Source/Runtime/Engine/Private/StreamableManager.cpp
+++ b/Engine/Source/Runtime/Engine/Private/StreamableManager.cpp
@@ -412,23 +412,13 @@ void FStreamableHandle::GetRequestedAssets(TArray<FSoftObjectPath>& AssetList) c
 
 UObject* FStreamableHandle::GetLoadedAsset() const
 {
-	if (HasLoadCompleted())
-	{
-		if (RequestedAssets.Num() > 0)
-		{
-			const FSoftObjectPath& Ref = RequestedAssets[0];
-			// Try manager, should be faster and will handle redirects better
-			return IsActive() ? OwningManager->GetStreamed(Ref) : Ref.ResolveObject();
-		}
+	TArray<UObject *> LoadedAssets;
 
-		// Check child handles
-		for (const TSharedPtr<FStreamableHandle>& ChildHandle : ChildHandles)
-		{
-			if (UObject* LoadedChildAsset = ChildHandle->GetLoadedAsset())
-			{
-				return LoadedChildAsset;
-			}
-		}	
+	GetLoadedAssets(LoadedAssets);
+
+	if (LoadedAssets.Num() > 0)
+	{
+		return LoadedAssets[0];
 	}
 
 	return nullptr;
