Remove blueprint recompiling on rename as it is redundant with new BP Compiler updates and causes a re-entrance into FileHelpers::SaveWorld

diff --git a/Engine/Source/Runtime/Engine/Private/Blueprint.cpp b/Engine/Source/Runtime/Engine/Private/Blueprint.cpp
--- a/Engine/Source/Runtime/Engine/Private/Blueprint.cpp
+++ b/Engine/Source/Runtime/Engine/Private/Blueprint.cpp
@@ -596,25 +596,7 @@ bool UBlueprint::Rename( const TCHAR* InName, UObject* NewOuter, ERenameFlags Fl
 		return false;
 	}
 
-	bool bSuccess = Super::Rename( InName, NewOuter, Flags );
-
-	// Finally, do a compile, but only if the new name differs from before
-	if(bSuccess && !(Flags & REN_Test) && !(Flags & REN_DoNotDirty) && InName && InName != OldName)
-	{
-		// Gather all blueprints that currently depend on this one.
-		TArray<UBlueprint*> Dependents;
-		FBlueprintEditorUtils::FindDependentBlueprints(this, Dependents);
-
-		FKismetEditorUtilities::CompileBlueprint(this);
-
-		// Recompile dependent blueprints after compiling this one. Otherwise, we can end up with a GLEO during the internal package save, which will include referencers as well.
-		for (UBlueprint* DependentBlueprint : Dependents)
-		{
-			FKismetEditorUtilities::CompileBlueprint(DependentBlueprint);
-		}
-	}
-
-	return bSuccess;
+	return Super::Rename( InName, NewOuter, Flags );
 }
 
 void UBlueprint::PostDuplicate(bool bDuplicateForPIE)
