Fixed crash calling SetupRayTracingCullClusters when raytracing was disabled

diff --git a/Engine/Source/Runtime/Engine/Private/InstancedStaticMesh.cpp b/Engine/Source/Runtime/Engine/Private/InstancedStaticMesh.cpp
--- a/Engine/Source/Runtime/Engine/Private/InstancedStaticMesh.cpp
+++ b/Engine/Source/Runtime/Engine/Private/InstancedStaticMesh.cpp
@@ -1244,6 +1244,11 @@ void FInstancedStaticMeshSceneProxy::GetDynamicRayTracingInstances(struct FRayTr
 
 void FInstancedStaticMeshSceneProxy::SetupRayTracingCullClusters()
 {
+	if (!IsRayTracingEnabled())
+	{
+		return;
+	}
+
 	check(IsInGameThread());
 
 	//#dxr_todo: select the appropriate LOD depending on Context.View
