Fix race condition in RequiresAdjacencyInformation() and FLandscapeComponentSceneProxy::GetStaticMeshElement(). Revealed by parallelized AddStaticMesh

diff --git a/Engine/Source/Runtime/Engine/Private/TessellationRendering.cpp b/Engine/Source/Runtime/Engine/Private/TessellationRendering.cpp
--- a/Engine/Source/Runtime/Engine/Private/TessellationRendering.cpp
+++ b/Engine/Source/Runtime/Engine/Private/TessellationRendering.cpp
@@ -29,7 +29,7 @@ bool MaterialSettingsRequireAdjacencyInformation_GameThread(UMaterialInterface*
   * Rendering thread version that looks at the current shader that will be used. **Will change answer during a shader compile** */
 bool MaterialRenderingRequiresAdjacencyInformation_RenderingThread(UMaterialInterface* Material, const FVertexFactoryType* VertexFactoryType, const FStaticFeatureLevel InFeatureLevel)
 {
-	check(IsInRenderingThread());
+	check(IsInRenderingThread() || IsInParallelRenderingThread());
 
 	//if we pass null here we have to guarantee that the VF supports tessellation (e.g by using type LocalVF)
 	bool VertexFactorySupportsTessellation = !VertexFactoryType || (VertexFactoryType && VertexFactoryType->SupportsTessellationShaders());
@@ -65,7 +65,7 @@ bool RequiresAdjacencyInformation(UMaterialInterface* Material, const FVertexFac
 		return false;
 	}
 
-	if (IsInRenderingThread())
+	if (IsInRenderingThread() || IsInParallelRenderingThread())
 	{
 		return MaterialRenderingRequiresAdjacencyInformation_RenderingThread(Material, VertexFactoryType, InFeatureLevel);
 	}
diff --git a/Engine/Source/Runtime/Landscape/Private/LandscapeRender.cpp b/Engine/Source/Runtime/Landscape/Private/LandscapeRender.cpp
--- a/Engine/Source/Runtime/Landscape/Private/LandscapeRender.cpp
+++ b/Engine/Source/Runtime/Landscape/Private/LandscapeRender.cpp
@@ -2300,7 +2300,7 @@ bool FLandscapeComponentSceneProxy::GetStaticMeshElement(int32 LODIndex, bool bF
 		}
 
 		UMaterialInstance* MaterialInstance = Cast<UMaterialInstance>(AvailableMaterials[MaterialIndex]);
-		bool HasTessellationEnabled = (GetScene().GetFeatureLevel() >= ERHIFeatureLevel::SM5) ? MaterialInstance != nullptr && MaterialInstance->GetMaterial_Concurrent()->D3D11TessellationMode != EMaterialTessellationMode::MTM_NoTessellation && MaterialIndexToDisabledTessellationMaterial[MaterialIndex] != INDEX_NONE : false;
+		bool HasTessellationEnabled = (GetScene().GetFeatureLevel() >= ERHIFeatureLevel::SM5) ? MaterialInstance != nullptr && RequiresAdjacencyInformation(MaterialInstance, VertexFactory->GetType(), GetScene().GetFeatureLevel()) && MaterialIndexToDisabledTessellationMaterial[MaterialIndex] != INDEX_NONE : false;
 
 		if (HasTessellationEnabled)
 		{
