Change type of FStreamingLevelsToConsider::LevelsToProcess from a TMap<> to a TSortedMap<>, saves 64 byte per world
Should also be faster since this map usually contains very few elements.

diff --git a/Engine/Source/Runtime/Engine/Classes/Engine/World.h b/Engine/Source/Runtime/Engine/Classes/Engine/World.h
--- a/Engine/Source/Runtime/Engine/Classes/Engine/World.h
+++ b/Engine/Source/Runtime/Engine/Classes/Engine/World.h
@@ -778,7 +778,7 @@ private:
 	};
 
 	/** Streaming levels that had their priority changed or were added to the container while consideration was underway. */
-	TMap<ULevelStreaming*, EProcessReason> LevelsToProcess;
+	TSortedMap<ULevelStreaming*, EProcessReason> LevelsToProcess;
 
 	/** Whether the streaming levels are under active consideration or not */
 	bool bStreamingLevelsBeingConsidered;
diff --git a/Engine/Source/Runtime/Engine/Private/World.cpp b/Engine/Source/Runtime/Engine/Private/World.cpp
--- a/Engine/Source/Runtime/Engine/Private/World.cpp
+++ b/Engine/Source/Runtime/Engine/Private/World.cpp
@@ -3225,7 +3225,7 @@ void FStreamingLevelsToConsider::EndConsideration()
 	{
 		// For any streaming level that was added or had its priority changed while we were considering the
 		// streaming levels go through and ensure they are correctly in the map and sorted to the correct location
-		TMap<ULevelStreaming*, EProcessReason> LevelsToProcessCopy = MoveTemp(LevelsToProcess);
+		TSortedMap<ULevelStreaming*, EProcessReason> LevelsToProcessCopy = MoveTemp(LevelsToProcess);
 		for (const TPair<ULevelStreaming*,EProcessReason>& LevelToProcessPair : LevelsToProcessCopy)
 		{
 			if (ULevelStreaming* StreamingLevel = LevelToProcessPair.Key)
