Fixed crash isolating materials in the static mesh editor with ray tracing enabled

#[fyi] yujiang.wang

diff --git a/Engine/Source/Runtime/Engine/Private/StaticMeshRender.cpp b/Engine/Source/Runtime/Engine/Private/StaticMeshRender.cpp
--- a/Engine/Source/Runtime/Engine/Private/StaticMeshRender.cpp
+++ b/Engine/Source/Runtime/Engine/Private/StaticMeshRender.cpp
@@ -1588,7 +1588,12 @@ void FStaticMeshSceneProxy::GetDynamicRayTracingInstances(FRayTracingMaterialGat
 				FMeshBatch Mesh;
 	
 				bool bResult = GetMeshElement(LODIndex, BatchIndex, SectionIndex, PrimitiveDPG, false, false, Mesh);
-				check(bResult);
+				if (!bResult)
+				{
+					// Hidden material
+					Mesh.MaterialRenderProxy = UMaterial::GetDefaultMaterial(MD_Surface)->GetRenderProxy();
+					Mesh.VertexFactory = &RenderData->LODVertexFactories[LODIndex].VertexFactory;
+				}
 				Mesh.SegmentIndex = SectionIndex;
 				RayTracingInstance.Materials.Add(Mesh);
 			}
