Expose pre-skinned vertex buffer to local vertex factory so that pre skin vertex position is available to material nodes that use it when derived vertex factories generate position buffers that are post skinned. Particularly relevant for GPUSkinCache.

diff --git a/Engine/Shaders/Private/LocalVertexFactory.ush b/Engine/Shaders/Private/LocalVertexFactory.ush
--- a/Engine/Shaders/Private/LocalVertexFactory.ush
+++ b/Engine/Shaders/Private/LocalVertexFactory.ush
@@ -307,6 +307,8 @@ struct FVertexFactoryIntermediates
 	float4 PerInstanceParams;
 #endif	// USE_INSTANCING
 	uint PrimitiveId;
+
+	float3 PreSkinPosition;
 };
 
 #if USE_INSTANCING
@@ -505,7 +507,7 @@ FMaterialVertexParameters GetMaterialVertexParameters(FVertexFactoryInput Input,
 	Result.PrevFrameLocalToWorld = GetPrimitiveData(Intermediates.PrimitiveId).PreviousLocalToWorld;
 #endif	// USE_INSTANCING
 
-	Result.PreSkinnedPosition = Input.Position.xyz;
+	Result.PreSkinnedPosition = Intermediates.PreSkinPosition.xyz;
 	Result.PreSkinnedNormal = TangentToLocal[2]; //TangentBias(Input.TangentZ.xyz);
 
 #if MANUAL_VERTEX_FETCH && NUM_MATERIAL_TEXCOORDS_VERTEX
@@ -802,6 +804,16 @@ FVertexFactoryIntermediates GetVertexFactoryIntermediates(FVertexFactoryInput In
 	Intermediates.PerInstanceParams.z = 1;
 	Intermediates.PerInstanceParams.w = 0;
 #endif	// USE_INSTANCING
+
+#if GPUSKIN_PASS_THROUGH
+	uint PreSkinVertexOffset = LocalVF.PreSkinBaseVertexIndex + Input.VertexId * 3;
+	Intermediates.PreSkinPosition.x = LocalVF.VertexFetch_PreSkinPositionBuffer[PreSkinVertexOffset + 0];
+	Intermediates.PreSkinPosition.y = LocalVF.VertexFetch_PreSkinPositionBuffer[PreSkinVertexOffset + 1];
+	Intermediates.PreSkinPosition.z = LocalVF.VertexFetch_PreSkinPositionBuffer[PreSkinVertexOffset + 2];
+#else
+	Intermediates.PreSkinPosition = Input.Position.xyz;
+#endif
+
 	return Intermediates;
 }
 
diff --git a/Engine/Source/Runtime/Engine/Private/GPUSkinCache.cpp b/Engine/Source/Runtime/Engine/Private/GPUSkinCache.cpp
--- a/Engine/Source/Runtime/Engine/Private/GPUSkinCache.cpp
+++ b/Engine/Source/Runtime/Engine/Private/GPUSkinCache.cpp
@@ -260,6 +260,14 @@ public:
 			return PositionBuffer;
 		}
 
+		inline FRHIShaderResourceView* GetPreSkinPositionSRV()
+		{
+			check(SourceVertexFactory);
+			check(SourceVertexFactory->GetPositionsSRV());
+
+			return SourceVertexFactory->GetPositionsSRV().GetReference();
+		}
+
 		inline FRWBuffer* GetTangentRWBuffer()
 		{
 			return TangentBuffer;
@@ -267,7 +275,7 @@ public:
 
 		void UpdateVertexFactoryDeclaration()
 		{
-			TargetVertexFactory->UpdateVertexDeclaration(SourceVertexFactory, GetPositionRWBuffer(), GetTangentRWBuffer());
+			TargetVertexFactory->UpdateVertexDeclaration(SourceVertexFactory, GetPositionRWBuffer(), GetPreSkinPositionSRV(), GetTangentRWBuffer());
 		}
 	};
 
@@ -1475,7 +1483,6 @@ void FGPUSkinCache::DispatchUpdateSkinning(FRHICommandListImmediate& RHICmdList,
 		INC_DWORD_STAT_BY(STAT_GPUSkinCache_TotalNumVertices, VertexCountAlign64 * 64);
 		RHICmdList.DispatchComputeShader(VertexCountAlign64, 1, 1);
 		Shader->UnsetParameters(RHICmdList);
-
 	}
 
 	check(DispatchData.PreviousPositionBuffer != DispatchData.PositionBuffer);
diff --git a/Engine/Source/Runtime/Engine/Private/GPUSkinVertexFactory.cpp b/Engine/Source/Runtime/Engine/Private/GPUSkinVertexFactory.cpp
--- a/Engine/Source/Runtime/Engine/Private/GPUSkinVertexFactory.cpp
+++ b/Engine/Source/Runtime/Engine/Private/GPUSkinVertexFactory.cpp
@@ -407,6 +407,7 @@ void TGPUSkinVertexFactory<BoneInfluenceType>::CopyDataTypeForPassthroughFactory
 	DestDataType.TextureCoordinates = Data.TextureCoordinates;
 	DestDataType.ColorComponent = Data.ColorComponent;
 	DestDataType.PositionComponentSRV = Data.PositionComponentSRV;
+	DestDataType.PreSkinPositionComponentSRV = Data.PositionComponentSRV;
 	DestDataType.TangentsSRV = Data.TangentsSRV;
 	DestDataType.ColorComponentsSRV = Data.ColorComponentsSRV;
 	DestDataType.ColorIndexMask = Data.ColorIndexMask;
@@ -691,7 +692,7 @@ bool FGPUSkinPassthroughVertexFactory::ShouldCompilePermutation(const FVertexFac
 		(Parameters.MaterialParameters.bIsUsedWithSkeletalMesh || Parameters.MaterialParameters.bIsSpecialEngineMaterial);
 }
 
-void FGPUSkinPassthroughVertexFactory::InternalUpdateVertexDeclaration(FGPUBaseSkinVertexFactory* SourceVertexFactory, struct FRWBuffer* PositionRWBuffer, struct FRWBuffer* TangentRWBuffer)
+void FGPUSkinPassthroughVertexFactory::InternalUpdateVertexDeclaration(FGPUBaseSkinVertexFactory* SourceVertexFactory, struct FRWBuffer* PositionRWBuffer, class FRHIShaderResourceView* PreSkinPositionSRV, struct FRWBuffer* TangentRWBuffer)
 {
 	// Point this vertex buffer to the RWBuffer
 	PositionVBAlias.VertexBufferRHI = PositionRWBuffer->Buffer;
@@ -708,6 +709,7 @@ void FGPUSkinPassthroughVertexFactory::InternalUpdateVertexDeclaration(FGPUBaseS
 	{
 		Data.TangentsSRV = TangentRWBuffer ? TangentRWBuffer->SRV : SourceVertexFactory->GetTangentsSRV();
 		Data.PositionComponentSRV = PositionRWBuffer->SRV;
+		Data.PreSkinPositionComponentSRV = PreSkinPositionSRV;
 	}
 
 	Data.TangentBasisComponents[0] = SourceVertexFactory->GetTangentStreamComponent(0);
diff --git a/Engine/Source/Runtime/Engine/Private/LocalVertexFactory.cpp b/Engine/Source/Runtime/Engine/Private/LocalVertexFactory.cpp
--- a/Engine/Source/Runtime/Engine/Private/LocalVertexFactory.cpp
+++ b/Engine/Source/Runtime/Engine/Private/LocalVertexFactory.cpp
@@ -46,7 +46,9 @@ TUniformBufferRef<FLocalVertexFactoryUniformShaderParameters> CreateLocalVFUnifo
 	const FLocalVertexFactory* LocalVertexFactory, 
 	uint32 LODLightmapDataIndex, 
 	FColorVertexBuffer* OverrideColorVertexBuffer, 
-	int32 BaseVertexIndex)
+	int32 BaseVertexIndex,
+	int32 PreSkinBaseVertexIndex
+	)
 {
 	FLocalVertexFactoryUniformShaderParameters UniformParameters;
 
@@ -56,6 +58,7 @@ TUniformBufferRef<FLocalVertexFactoryUniformShaderParameters> CreateLocalVFUnifo
 	if (RHISupportsManualVertexFetch(GMaxRHIShaderPlatform))
 	{
 		UniformParameters.VertexFetch_PositionBuffer = LocalVertexFactory->GetPositionsSRV();
+		UniformParameters.VertexFetch_PreSkinPositionBuffer = LocalVertexFactory->GetPreSkinPositionSRV();
 
 		UniformParameters.VertexFetch_PackedTangentsBuffer = LocalVertexFactory->GetTangentsSRV();
 		UniformParameters.VertexFetch_TexCoordBuffer = LocalVertexFactory->GetTextureCoordinatesSRV();
@@ -73,6 +76,7 @@ TUniformBufferRef<FLocalVertexFactoryUniformShaderParameters> CreateLocalVFUnifo
 	}
 	else
 	{
+		UniformParameters.VertexFetch_PreSkinPositionBuffer = GNullColorVertexBuffer.VertexBufferSRV;
 		UniformParameters.VertexFetch_PackedTangentsBuffer = GNullColorVertexBuffer.VertexBufferSRV;
 		UniformParameters.VertexFetch_TexCoordBuffer = GNullColorVertexBuffer.VertexBufferSRV;
 	}
@@ -85,7 +89,10 @@ TUniformBufferRef<FLocalVertexFactoryUniformShaderParameters> CreateLocalVFUnifo
 	const int32 NumTexCoords = LocalVertexFactory->GetNumTexcoords();
 	const int32 LightMapCoordinateIndex = LocalVertexFactory->GetLightMapCoordinateIndex();
 	const int32 EffectiveBaseVertexIndex = RHISupportsAbsoluteVertexID(GMaxRHIShaderPlatform) ? 0 : BaseVertexIndex;
+	const int32 EffectivePreSkinBaseVertexIndex = RHISupportsAbsoluteVertexID(GMaxRHIShaderPlatform) ? 0 : PreSkinBaseVertexIndex;
+
 	UniformParameters.VertexFetch_Parameters = {ColorIndexMask, NumTexCoords, LightMapCoordinateIndex, EffectiveBaseVertexIndex};
+	UniformParameters.PreSkinBaseVertexIndex = EffectivePreSkinBaseVertexIndex;
 
 	return TUniformBufferRef<FLocalVertexFactoryUniformShaderParameters>::CreateUniformBufferImmediate(UniformParameters, UniformBuffer_MultiFrame);
 }
@@ -378,10 +385,11 @@ void FLocalVertexFactory::InitRHI()
 	check(IsValidRef(GetDeclaration()));
 
 	const int32 DefaultBaseVertexIndex = 0;
+	const int32 DefaultPreSkinBaseVertexIndex = 0;
 	if (RHISupportsManualVertexFetch(GMaxRHIShaderPlatform) || bCanUseGPUScene)
 	{
 		SCOPED_LOADTIMER(FLocalVertexFactory_InitRHI_CreateLocalVFUniformBuffer);
-		UniformBuffer = CreateLocalVFUniformBuffer(this, Data.LODLightmapDataIndex, nullptr, DefaultBaseVertexIndex);
+		UniformBuffer = CreateLocalVFUniformBuffer(this, Data.LODLightmapDataIndex, nullptr, DefaultBaseVertexIndex, DefaultPreSkinBaseVertexIndex);
 	}
 
 	check(IsValidRef(GetDeclaration()));
diff --git a/Engine/Source/Runtime/Engine/Private/ModelRender.cpp b/Engine/Source/Runtime/Engine/Private/ModelRender.cpp
--- a/Engine/Source/Runtime/Engine/Private/ModelRender.cpp
+++ b/Engine/Source/Runtime/Engine/Private/ModelRender.cpp
@@ -707,7 +707,7 @@ private:
 				ENQUEUE_RENDER_COMMAND(FElementInfoCreateLocalVFUniformBuffer)(
 					[UniformBufferPtr, VertexFactory = InVertexFactory, ElementIndex = InElementIndex](FRHICommandListImmediate& RHICmdList)
 				{
-					*UniformBufferPtr = CreateLocalVFUniformBuffer(VertexFactory, ElementIndex, nullptr, 0);
+					*UniformBufferPtr = CreateLocalVFUniformBuffer(VertexFactory, ElementIndex, nullptr, 0, 0);
 				});
 			}
 		}
diff --git a/Engine/Source/Runtime/Engine/Private/SkeletalRenderGPUSkin.cpp b/Engine/Source/Runtime/Engine/Private/SkeletalRenderGPUSkin.cpp
--- a/Engine/Source/Runtime/Engine/Private/SkeletalRenderGPUSkin.cpp
+++ b/Engine/Source/Runtime/Engine/Private/SkeletalRenderGPUSkin.cpp
@@ -1381,9 +1381,10 @@ void FGPUSkinPassthroughVertexFactory::SetData(const FDataType& InData)
 {
 	FLocalVertexFactory::SetData(InData);
 	const int32 DefaultBaseVertexIndex = 0;
+	const int32 DefaultPreSkinBaseVertexIndex = 0;
 	if (RHISupportsManualVertexFetch(GMaxRHIShaderPlatform))
 	{
-		UniformBuffer = CreateLocalVFUniformBuffer(this, Data.LODLightmapDataIndex, nullptr, DefaultBaseVertexIndex);
+		UniformBuffer = CreateLocalVFUniformBuffer(this, Data.LODLightmapDataIndex, nullptr, DefaultBaseVertexIndex, DefaultPreSkinBaseVertexIndex);
 	}
 }
 
diff --git a/Engine/Source/Runtime/Engine/Private/StaticMeshRender.cpp b/Engine/Source/Runtime/Engine/Private/StaticMeshRender.cpp
--- a/Engine/Source/Runtime/Engine/Private/StaticMeshRender.cpp
+++ b/Engine/Source/Runtime/Engine/Private/StaticMeshRender.cpp
@@ -1887,7 +1887,7 @@ FStaticMeshSceneProxy::FLODInfo::FLODInfo(const UStaticMeshComponent* InComponen
 					{
 						checkf(LocalVF->GetTangentsSRV(), TEXT("LOD: %i of PathName: %s has a null tangents srv."), LODIndex, *ComponentPathName);
 						checkf(LocalVF->GetTextureCoordinatesSRV(), TEXT("LOD: %i of PathName: %s has a null texcoord srv."), LODIndex, *ComponentPathName);
-						*UniformBufferPtr = CreateLocalVFUniformBuffer(LocalVF, LODIndex, VertexBuffer, 0);
+						*UniformBufferPtr = CreateLocalVFUniformBuffer(LocalVF, LODIndex, VertexBuffer, 0, 0);
 					});
 				}
 			}
diff --git a/Engine/Source/Runtime/Engine/Public/GPUSkinVertexFactory.h b/Engine/Source/Runtime/Engine/Public/GPUSkinVertexFactory.h
--- a/Engine/Source/Runtime/Engine/Public/GPUSkinVertexFactory.h
+++ b/Engine/Source/Runtime/Engine/Public/GPUSkinVertexFactory.h
@@ -483,11 +483,11 @@ public:
 	static void ModifyCompilationEnvironment(const FVertexFactoryShaderPermutationParameters& Parameters, FShaderCompilerEnvironment& OutEnvironment);
 	static bool ShouldCompilePermutation(const FVertexFactoryShaderPermutationParameters& Parameters);
 
-	inline void UpdateVertexDeclaration(FGPUBaseSkinVertexFactory* SourceVertexFactory, struct FRWBuffer* PositionRWBuffer, struct FRWBuffer* TangentRWBuffer)
+	inline void UpdateVertexDeclaration(FGPUBaseSkinVertexFactory* SourceVertexFactory, struct FRWBuffer* PositionRWBuffer, class FRHIShaderResourceView* PreSkinPositionSRV, struct FRWBuffer* TangentRWBuffer)
 	{
 		if (PositionStreamIndex == -1)
 		{
-			InternalUpdateVertexDeclaration(SourceVertexFactory, PositionRWBuffer, TangentRWBuffer);
+			InternalUpdateVertexDeclaration(SourceVertexFactory, PositionRWBuffer, PreSkinPositionSRV, TangentRWBuffer);
 		}
 	}
 
@@ -530,7 +530,7 @@ protected:
 	int32 PositionStreamIndex;
 	int32 TangentStreamIndex;
 
-	void InternalUpdateVertexDeclaration(FGPUBaseSkinVertexFactory* SourceVertexFactory, struct FRWBuffer* PositionRWBuffer, struct FRWBuffer* TangentRWBuffer);
+	void InternalUpdateVertexDeclaration(FGPUBaseSkinVertexFactory* SourceVertexFactory, struct FRWBuffer* PositionRWBuffer, class FRHIShaderResourceView* PreSkinPositionSRV, struct FRWBuffer* TangentRWBuffer);
 };
 
 /** Vertex factory with vertex stream components for GPU-skinned and morph target streams */
diff --git a/Engine/Source/Runtime/Engine/Public/LocalVertexFactory.h b/Engine/Source/Runtime/Engine/Public/LocalVertexFactory.h
--- a/Engine/Source/Runtime/Engine/Public/LocalVertexFactory.h
+++ b/Engine/Source/Runtime/Engine/Public/LocalVertexFactory.h
@@ -17,9 +17,11 @@ struct FMeshBatchElement;
 
 BEGIN_GLOBAL_SHADER_PARAMETER_STRUCT(FLocalVertexFactoryUniformShaderParameters,ENGINE_API)
 	SHADER_PARAMETER(FIntVector4,VertexFetch_Parameters)
+	SHADER_PARAMETER(int32, PreSkinBaseVertexIndex)
 	SHADER_PARAMETER(uint32,LODLightmapDataIndex)
 	SHADER_PARAMETER_SRV(Buffer<float2>, VertexFetch_TexCoordBuffer)
 	SHADER_PARAMETER_SRV(Buffer<float>, VertexFetch_PositionBuffer)
+	SHADER_PARAMETER_SRV(Buffer<float>, VertexFetch_PreSkinPositionBuffer)
 	SHADER_PARAMETER_SRV(Buffer<float4>, VertexFetch_PackedTangentsBuffer)
 	SHADER_PARAMETER_SRV(Buffer<float4>, VertexFetch_ColorComponentsBuffer)
 END_GLOBAL_SHADER_PARAMETER_STRUCT()
@@ -28,7 +30,9 @@ extern TUniformBufferRef<FLocalVertexFactoryUniformShaderParameters> CreateLocal
 	const class FLocalVertexFactory* VertexFactory, 
 	uint32 LODLightmapDataIndex, 
 	class FColorVertexBuffer* OverrideColorVertexBuffer, 
-	int32 BaseVertexIndex);
+	int32 BaseVertexIndex,
+	int32 PreSkinBaseVertexIndex
+	);
 
 /**
  * A vertex factory which simply transforms explicit vertex attributes from local to world space.
@@ -48,6 +52,7 @@ public:
 
 	struct FDataType : public FStaticMeshDataType
 	{
+		FRHIShaderResourceView* PreSkinPositionComponentSRV = nullptr;
 	};
 
 	/**
@@ -100,6 +105,11 @@ public:
 		return Data.PositionComponentSRV;
 	}
 
+	inline FRHIShaderResourceView* GetPreSkinPositionSRV() const
+	{
+		return Data.PreSkinPositionComponentSRV ? Data.PreSkinPositionComponentSRV : GNullColorVertexBuffer.VertexBufferSRV;
+	}
+
 	inline FRHIShaderResourceView* GetTangentsSRV() const
 	{
 		return Data.TangentsSRV;
