Cancel existing compression attempts if we are requesting a new compression

diff --git a/Engine/Source/Runtime/Engine/Private/Animation/AnimSequence.cpp b/Engine/Source/Runtime/Engine/Private/Animation/AnimSequence.cpp
--- a/Engine/Source/Runtime/Engine/Private/Animation/AnimSequence.cpp
+++ b/Engine/Source/Runtime/Engine/Private/Animation/AnimSequence.cpp
@@ -2177,7 +2177,7 @@ void UAnimSequence::RequestAnimCompression(FRequestAnimCompressionParams Params)
 		return;
 	}
 
-	WaitOnExistingCompression(false);
+	WaitOnExistingCompression(true);
 
 	if (BoneCompressionSettings == nullptr || !BoneCompressionSettings->AreSettingsValid())
 	{
