Handle shutdown while animations are still async compressing

diff --git a/Engine/Source/Runtime/Engine/Private/Animation/AnimCompress_RemoveLinearKeys.cpp b/Engine/Source/Runtime/Engine/Private/Animation/AnimCompress_RemoveLinearKeys.cpp
--- a/Engine/Source/Runtime/Engine/Private/Animation/AnimCompress_RemoveLinearKeys.cpp
+++ b/Engine/Source/Runtime/Engine/Private/Animation/AnimCompress_RemoveLinearKeys.cpp
@@ -620,6 +620,10 @@ void UAnimCompress_RemoveLinearKeys::ProcessAnimationTracks(
 	// setup the raw bone transformation and find all end effectors
 	for ( int32 BoneIndex = 0; BoneIndex < NumBones; ++BoneIndex )
 	{
+		if (CompressibleAnimData.IsCancelled())
+		{
+			return;
+		}
 		// get the raw world-atoms for this bone
 		UpdateWorldBoneTransformTable(
 			CompressibleAnimData,
@@ -641,6 +645,10 @@ void UAnimCompress_RemoveLinearKeys::ProcessAnimationTracks(
 	// for each bone...
 	for ( int32 BoneIndex = 0; BoneIndex < NumBones; ++BoneIndex )
 	{
+		if (CompressibleAnimData.IsCancelled())
+		{
+			return;
+		}
 		const FBoneData& Bone = CompressibleAnimData.BoneData[BoneIndex];
 		const int32 ParentBoneIndex = Bone.GetParent();
 
diff --git a/Engine/Source/Runtime/Engine/Private/Animation/AnimCompressionDerivedData.cpp b/Engine/Source/Runtime/Engine/Private/Animation/AnimCompressionDerivedData.cpp
--- a/Engine/Source/Runtime/Engine/Private/Animation/AnimCompressionDerivedData.cpp
+++ b/Engine/Source/Runtime/Engine/Private/Animation/AnimCompressionDerivedData.cpp
@@ -75,6 +75,8 @@ bool FDerivedDataAnimationCompression::Build( TArray<uint8>& OutDataArray )
 											OutputStr);
 	}
 
+	bCompressionSuccessful = bCompressionSuccessful && !DataToCompress.IsCancelled();
+
 	if (bCompressionSuccessful)
 	{
 		const double CompressionEndTime = FPlatformTime::Seconds();
@@ -94,7 +96,7 @@ bool FDerivedDataAnimationCompression::Build( TArray<uint8>& OutDataArray )
 }
 
 const uint64 GigaBytes = 1024 * 1024 * 1024;
-const uint64 MAX_ASYNC_COMPRESSION_MEM_USAGE = 2 * GigaBytes;
+const uint64 MAX_ASYNC_COMPRESSION_MEM_USAGE = 3 * GigaBytes;
 const int32 MAX_ACTIVE_COMPRESSIONS = 2;
 
 FAsyncCompressedAnimationsManagement* GAsyncCompressedAnimationsTracker = nullptr;
@@ -106,33 +108,120 @@ FAsyncCompressedAnimationsManagement& FAsyncCompressedAnimationsManagement::Get(
 	return SingletonInstance;
 }
 
+FAsyncCompressedAnimationsManagement::FAsyncCompressedAnimationsManagement() : ActiveMemoryUsage(0), LastTickTime(0.0)
+{
+	FCoreDelegates::OnPreExit.AddRaw(this, &FAsyncCompressedAnimationsManagement::Shutdown);
+}
+
+void FAsyncCompressedAnimationsManagement::Shutdown()
+{
+	TArray<uint8> OutData; // For canceling
+
+	for (FActiveAsyncCompressionTask& ActiveTask : ActiveAsyncCompressionTasks)
+	{
+		ActiveTask.DataToCompress->IsCancelledSignal.Cancel();
+		if (ActiveTask.Sequence)
+		{
+			ActiveTask.Sequence->ApplyCompressedData(FString(), false, OutData); // Clear active compression on Sequence
+			ActiveTask.Sequence = nullptr;
+		}
+	}
+
+	for (const FQueuedAsyncCompressionWork& QueuedTask : QueuedAsyncCompressionWork)
+	{
+		if (QueuedTask.Anim)
+		{
+			QueuedTask.Anim->ApplyCompressedData(FString(), false, OutData); // Clear active compression on Sequence
+		}
+	}
+
+	QueuedAsyncCompressionWork.Reset();
+
+	const double PreExitStartTimer = FPlatformTime::Seconds();
+	while(ActiveAsyncCompressionTasks.Num() > 0)
+	{
+		Tick(0.f); //Give active tasks that have already started a chance to finish
+		
+		const double PreExitDuration = FPlatformTime::Seconds() - PreExitStartTimer;
+
+		if (PreExitDuration > 10.0)
+		{
+			UE_LOG(LogAnimationCompression, Warning, TEXT("Async Compression Pre Init Timer hit, Active Compressions:%i"), ActiveAsyncCompressionTasks.Num());
+			return;
+		}
+		FPlatformProcess::Sleep(1.f);
+	} 
+}
+
 void FAsyncCompressedAnimationsManagement::OnActiveCompressionFinished(int32 ActiveAnimIndex)
 {
 	FDerivedDataCacheInterface& DerivedDataCache = GetDerivedDataCacheRef();
 
 	FActiveAsyncCompressionTask& Task = ActiveAsyncCompressionTasks[ActiveAnimIndex];
 
-	TArray<uint8> OutData;
-	bool bBuiltLocally = false;
-	if (DerivedDataCache.GetAsynchronousResults(Task.AsyncHandle, OutData, &bBuiltLocally))
-	{
-		Task.Sequence->ApplyCompressedData(Task.CacheKey, Task.bPerformFrameStripping, OutData);
-	}
-	else
+	if(Task.Sequence)
 	{
-		UE_LOG(LogAnimationCompression, Fatal, TEXT("Failed to get async compressed animation data for anim '%s'"), *Task.Sequence->GetName());
-		Task.Sequence->ApplyCompressedData(FString(), false, OutData); // Clear active compression on Sequence
+		TArray<uint8> OutData;
+		bool bBuiltLocally = false;
+		if (DerivedDataCache.GetAsynchronousResults(Task.AsyncHandle, OutData, &bBuiltLocally))
+		{
+			Task.Sequence->ApplyCompressedData(Task.CacheKey, Task.bPerformFrameStripping, OutData);
+		}
+		else
+		{
+			UE_LOG(LogAnimationCompression, Fatal, TEXT("Failed to get async compressed animation data for anim '%s'"), *Task.Sequence->GetName());
+			Task.Sequence->ApplyCompressedData(FString(), false, OutData); // Clear active compression on Sequence
+		}
 	}
+
 	ActiveMemoryUsage -= Task.TaskSize;
 	ActiveAsyncCompressionTasks.RemoveAtSwap(ActiveAnimIndex, 1, false);
 }
 
+float FormatBytes(uint64 Bytes, const TCHAR*& OutPostFix)
+{
+	uint64 Size = GigaBytes;
+
+	if (Bytes > Size)
+	{
+		OutPostFix = TEXT("GB");
+		return (float)Bytes / (float)Size;
+	}
+
+	Size = Size / 1024;
+	if (Bytes > Size)
+	{
+		OutPostFix = TEXT("MB");
+		return (float)Bytes / (float)Size;
+	}
+
+	Size = Size / 1024;
+	if (Bytes > Size)
+	{
+		OutPostFix = TEXT("KB");
+		return (float)Bytes / (float)Size;
+	}
+
+	OutPostFix = TEXT("B");
+	return (float)Bytes;
+}
+
+#define ASYNC_MEM_LOG 0
+
 void FAsyncCompressedAnimationsManagement::Tick(float DeltaTime)
 {
 	const double MaxProcessingTime = 0.1; // try not to hang the editor too much
 	const double EndTime = FPlatformTime::Seconds() + MaxProcessingTime;
 	const double StartTime = FPlatformTime::Seconds();
 
+#if ASYNC_MEM_LOG
+	const TCHAR* BytesPostFix = TEXT("ERROR");
+	const float FormatSize = FormatBytes(ActiveMemoryUsage, BytesPostFix);
+	FPlatformMisc::LowLevelOutputDebugStringf(TEXT("FAsyncCompressedAnimationsManagement::Tick\n\tDelta:%.2f TimeSinceLastRun:%.2f\n\tActive:%i Queued:%i\n\tMemUsage:%.2f%s\n"), DeltaTime, StartTime - LastTickTime, ActiveAsyncCompressionTasks.Num(), QueuedAsyncCompressionWork.Num(), FormatSize, BytesPostFix);
+#endif
+
+	LastTickTime = StartTime;
+
 	FDerivedDataCacheInterface& DerivedDataCache = GetDerivedDataCacheRef();
 
 	for (int32 ActiveAnim = ActiveAsyncCompressionTasks.Num() - 1; ActiveAnim >= 0; --ActiveAnim)
@@ -152,7 +241,17 @@ void FAsyncCompressedAnimationsManagement::Tick(float DeltaTime)
 
 	const bool bHasQueuedTasks = QueuedAsyncCompressionWork.Num() > 0;
 
-	while (ActiveAsyncCompressionTasks.Num() < MAX_ACTIVE_COMPRESSIONS)
+	StartQueuedTasks(MAX_ACTIVE_COMPRESSIONS);
+
+	if (bHasQueuedTasks && QueuedAsyncCompressionWork.Num() == 0)
+	{
+		QueuedAsyncCompressionWork.Empty(); //free memory
+	}
+}
+
+void FAsyncCompressedAnimationsManagement::StartQueuedTasks(int32 MaxActiveTasks)
+{
+	while (ActiveAsyncCompressionTasks.Num() < MaxActiveTasks)
 	{
 		if (QueuedAsyncCompressionWork.Num() == 0)
 		{
@@ -161,11 +260,6 @@ void FAsyncCompressedAnimationsManagement::Tick(float DeltaTime)
 		FQueuedAsyncCompressionWork NewTask = QueuedAsyncCompressionWork.Pop(false);
 		StartAsyncWork(NewTask.Compressor, NewTask.Anim, NewTask.Compressor.GetMemoryUsage(), NewTask.bPerformFrameStripping);
 	}
-
-	if (bHasQueuedTasks && QueuedAsyncCompressionWork.Num() == 0)
-	{
-		QueuedAsyncCompressionWork.Empty(); //free memory
-	}
 }
 
 void FAsyncCompressedAnimationsManagement::AddReferencedObjects(FReferenceCollector& Collector)
@@ -192,7 +286,14 @@ bool FAsyncCompressedAnimationsManagement::RequestAsyncCompression(FDerivedDataA
 
 	bool bWasAsync = true;
 
-	if (ActiveMemoryUsage + NewTaskSize >= MAX_ASYNC_COMPRESSION_MEM_USAGE)
+	const bool bOutOfMemForQueue = (ActiveMemoryUsage + NewTaskSize) >= MAX_ASYNC_COMPRESSION_MEM_USAGE;
+	const bool bGameThreadStarved = FPlatformTime::Seconds() - LastTickTime > 2.f;
+
+	//Boost Max active compressions in the case of game thread not running for extended periods (assumes heavy loading situation)
+	// so that memory usage of async compressions does not grow too large
+	const int32 MaxActiveCompressions = bGameThreadStarved ? (MAX_ACTIVE_COMPRESSIONS * 6) : MAX_ACTIVE_COMPRESSIONS;
+
+	if (bOutOfMemForQueue || bGameThreadStarved)
 	{
 		Tick(0.f); // Try to free up some memory
 	}
@@ -200,12 +301,13 @@ bool FAsyncCompressedAnimationsManagement::RequestAsyncCompression(FDerivedDataA
 	const bool bCanRunASync = (ActiveMemoryUsage + NewTaskSize < MAX_ASYNC_COMPRESSION_MEM_USAGE);
 	const bool bForceSync = false; //debugging override
 
+
 	if (bCanRunASync && !bForceSync)
 	{
 		//Queue Async
 		ActiveMemoryUsage += NewTaskSize;
 
-		if(ActiveAsyncCompressionTasks.Num() < MAX_ACTIVE_COMPRESSIONS)
+		if(ActiveAsyncCompressionTasks.Num() < MaxActiveCompressions)
 		{
 			StartAsyncWork(Compressor, Anim, NewTaskSize, bPerformFrameStripping);
 		}
@@ -220,6 +322,9 @@ bool FAsyncCompressedAnimationsManagement::RequestAsyncCompression(FDerivedDataA
 		GetDerivedDataCacheRef().GetSynchronous(&Compressor, OutData);
 		bWasAsync = false;
 	}
+
+	StartQueuedTasks(MaxActiveCompressions);
+
 	return bWasAsync;
 }
 
diff --git a/Engine/Source/Runtime/Engine/Private/Animation/AnimCompressionDerivedData.h b/Engine/Source/Runtime/Engine/Private/Animation/AnimCompressionDerivedData.h
--- a/Engine/Source/Runtime/Engine/Private/Animation/AnimCompressionDerivedData.h
+++ b/Engine/Source/Runtime/Engine/Private/Animation/AnimCompressionDerivedData.h
@@ -45,7 +45,7 @@ public:
 
 	uint64 GetMemoryUsage() const
 	{
-		return DataToCompressPtr.IsValid() ? DataToCompressPtr->GetApproxRawSize() : 0;
+		return DataToCompressPtr.IsValid() ? DataToCompressPtr->GetApproxMemoryUsage() : 0;
 	}
 
 	virtual const TCHAR* GetPluginName() const override
diff --git a/Engine/Source/Runtime/Engine/Public/Animation/AnimCompressionDerivedDataPublic.h b/Engine/Source/Runtime/Engine/Public/Animation/AnimCompressionDerivedDataPublic.h
--- a/Engine/Source/Runtime/Engine/Public/Animation/AnimCompressionDerivedDataPublic.h
+++ b/Engine/Source/Runtime/Engine/Public/Animation/AnimCompressionDerivedDataPublic.h
@@ -69,17 +69,24 @@ private:
 
 	void StartAsyncWork(FDerivedDataAnimationCompression& Compressor, UAnimSequence* Anim, const uint64 NewTaskSize, const bool bPerformFrameStripping);
 
+	void StartQueuedTasks(int32 MaxActiveTasks);
+
 	bool WaitOnActiveCompression(UAnimSequence* Anim);
 
 	void OnActiveCompressionFinished(int32 ActiveAnimIndex);
 
-	FAsyncCompressedAnimationsManagement() : ActiveMemoryUsage(0) {}
+	FAsyncCompressedAnimationsManagement();
+
+	void Shutdown();
 
 	//Array of active and queued compression jobs
 	TArray<FActiveAsyncCompressionTask> ActiveAsyncCompressionTasks;
 	TArray<FQueuedAsyncCompressionWork> QueuedAsyncCompressionWork;
 	uint64 ActiveMemoryUsage;
 
+	//Cache the last tick time to keep queue pumping during large loading segments
+	double LastTickTime;
+
 	/* Begin FTickableEditorObject */
 	virtual void Tick(float DeltaTime) override;
 	virtual ETickableTickType GetTickableTickType() const override { return ETickableTickType::Always; }
diff --git a/Engine/Source/Runtime/Engine/Public/Animation/AnimCompressionTypes.h b/Engine/Source/Runtime/Engine/Public/Animation/AnimCompressionTypes.h
--- a/Engine/Source/Runtime/Engine/Public/Animation/AnimCompressionTypes.h
+++ b/Engine/Source/Runtime/Engine/Public/Animation/AnimCompressionTypes.h
@@ -164,6 +164,30 @@ struct FCompressedOffsetData : public FCompressedOffsetDataBase<TArray<int32>>
 
 FArchive& operator<<(FArchive& Ar, FCompressedOffsetData& D);
 
+// Mechanism for cancelling in flight compression
+struct FCancelCompressionSignal
+{
+private:
+	FThreadSafeBool Signal;
+
+public:
+	FCancelCompressionSignal& operator=(const FCancelCompressionSignal& Other)
+	{
+		Signal = (bool)Other.Signal;
+		return *this;
+	}
+
+	void Cancel()
+	{
+		Signal = true;
+	}
+
+	bool IsCancelled() const
+	{
+		return Signal;
+	}
+};
+
 struct ENGINE_API FCompressibleAnimData
 {
 public:
@@ -204,20 +228,28 @@ public:
 	FString FullName;
 	FName   AnimFName;
 
-	int32 GetApproxRawBoneSize() const
+	FCancelCompressionSignal IsCancelledSignal;
+
+	static int32 GetApproxRawDataArraySize(const TArray<FRawAnimSequenceTrack>& AnimData)
 	{
-		int32 Total = sizeof(FRawAnimSequenceTrack) * RawAnimationData.Num();
-		for (int32 i = 0; i < RawAnimationData.Num(); ++i)
+		int32 Total = sizeof(FRawAnimSequenceTrack) * AnimData.Num();
+		for (int32 i = 0; i < AnimData.Num(); ++i)
 		{
-			const FRawAnimSequenceTrack& RawTrack = RawAnimationData[i];
+			const FRawAnimSequenceTrack& RawTrack = AnimData[i];
 			Total +=
 				sizeof(FVector) * RawTrack.PosKeys.Num() +
 				sizeof(FQuat) * RawTrack.RotKeys.Num() +
 				sizeof(FVector) * RawTrack.ScaleKeys.Num();
 		}
+
 		return Total;
 	}
 
+	int32 GetApproxRawBoneSize() const
+	{
+		return GetApproxRawDataArraySize(RawAnimationData);
+	}
+
 	int32 GetApproxRawCurveSize() const
 	{
 		int32 Total = 0;
@@ -234,6 +266,15 @@ public:
 		return GetApproxRawBoneSize() + GetApproxRawCurveSize();
 	}
 
+	uint64 GetApproxMemoryUsage() const
+	{
+		const uint64 MemUsage	= GetApproxRawSize()
+								+ TrackToSkeletonMapTable.GetAllocatedSize()
+								+ GetApproxRawDataArraySize(AdditiveBaseAnimationData)
+								+ BoneData.GetAllocatedSize();
+		return MemUsage;
+	}
+
 	void Update(struct FCompressedAnimSequence& CompressedData) const;
 
 	void AddReferencedObjects(FReferenceCollector& Collector)
@@ -242,6 +283,11 @@ public:
 		Collector.AddReferencedObject(CurveCompressionSettings);
 	}
 
+	bool IsCancelled() const
+	{
+		return IsCancelledSignal.IsCancelled();
+	}
+
 private:
 
 };
