UE4.25 - Fix compile issue
#rnx

diff --git a/Engine/Source/Runtime/Engine/Public/LocalVertexFactory.h b/Engine/Source/Runtime/Engine/Public/LocalVertexFactory.h
--- a/Engine/Source/Runtime/Engine/Public/LocalVertexFactory.h
+++ b/Engine/Source/Runtime/Engine/Public/LocalVertexFactory.h
@@ -107,7 +107,7 @@ public:
 
 	inline FRHIShaderResourceView* GetPreSkinPositionSRV() const
 	{
-		return Data.PreSkinPositionComponentSRV ? Data.PreSkinPositionComponentSRV : GNullColorVertexBuffer.VertexBufferSRV;
+		return Data.PreSkinPositionComponentSRV ? Data.PreSkinPositionComponentSRV : GNullColorVertexBuffer.VertexBufferSRV.GetReference();
 	}
 
 	inline FRHIShaderResourceView* GetTangentsSRV() const
