Remove USkeleton reference from anim compression to address thread safety

diff --git a/Engine/Source/Runtime/Engine/Private/Animation/AnimBoneCompressionSettings.cpp b/Engine/Source/Runtime/Engine/Private/Animation/AnimBoneCompressionSettings.cpp
--- a/Engine/Source/Runtime/Engine/Private/Animation/AnimBoneCompressionSettings.cpp
+++ b/Engine/Source/Runtime/Engine/Private/Animation/AnimBoneCompressionSettings.cpp
@@ -132,7 +132,7 @@ public:
 
 bool UAnimBoneCompressionSettings::Compress(const FCompressibleAnimData& AnimSeq, FCompressibleAnimDataResult& OutCompressedData) const
 {
-	if (AnimSeq.Skeleton == nullptr || !AreSettingsValid())
+	if (!AreSettingsValid())
 	{
 		return false;
 	}
diff --git a/Engine/Source/Runtime/Engine/Private/Animation/AnimCompress_RemoveLinearKeys.cpp b/Engine/Source/Runtime/Engine/Private/Animation/AnimCompress_RemoveLinearKeys.cpp
--- a/Engine/Source/Runtime/Engine/Private/Animation/AnimCompress_RemoveLinearKeys.cpp
+++ b/Engine/Source/Runtime/Engine/Private/Animation/AnimCompress_RemoveLinearKeys.cpp
@@ -596,7 +596,7 @@ void UAnimCompress_RemoveLinearKeys::ProcessAnimationTracks(
 	const float FrameRate = (float)(LastFrame) / SequenceLength;
 	const float TimePerFrame = SequenceLength / (float)(LastFrame);
 
-	const TArray<FTransform>& RefPose = CompressibleAnimData.Skeleton->GetRefLocalPoses();
+	const TArray<FTransform>& RefPose = CompressibleAnimData.RefLocalPoses;
 	const bool bHasScale =  (ScaleTracks.Num() > 0);
 
 	// make sure the parent key scale is properly bound to 1.0 or more
diff --git a/Engine/Source/Runtime/Engine/Private/Animation/AnimCompressionDerivedData.cpp b/Engine/Source/Runtime/Engine/Private/Animation/AnimCompressionDerivedData.cpp
--- a/Engine/Source/Runtime/Engine/Private/Animation/AnimCompressionDerivedData.cpp
+++ b/Engine/Source/Runtime/Engine/Private/Animation/AnimCompressionDerivedData.cpp
@@ -89,7 +89,7 @@ bool FDerivedDataAnimationCompression::Build( TArray<uint8>& OutDataArray )
 		OutData.BoneCompressionCodec = CompressionResult.Codec;
 
 		FMemoryWriter Ar(OutDataArray, true);
-		OutData.SerializeCompressedData(Ar, true, nullptr, DataToCompress.Skeleton, DataToCompress.BoneCompressionSettings, DataToCompress.CurveCompressionSettings); //Save out compressed
+		OutData.SerializeCompressedData(Ar, true, nullptr, nullptr, DataToCompress.BoneCompressionSettings, DataToCompress.CurveCompressionSettings); //Save out compressed
 	}
 
 	return bCompressionSuccessful;
diff --git a/Engine/Source/Runtime/Engine/Private/Animation/AnimCompressionDerivedData.h b/Engine/Source/Runtime/Engine/Private/Animation/AnimCompressionDerivedData.h
--- a/Engine/Source/Runtime/Engine/Private/Animation/AnimCompressionDerivedData.h
+++ b/Engine/Source/Runtime/Engine/Private/Animation/AnimCompressionDerivedData.h
@@ -38,7 +38,6 @@ public:
 	void SetCompressibleData(FCompressibleAnimRef InCompressibleAnimData)
 	{
 		DataToCompressPtr = InCompressibleAnimData;
-		check(DataToCompressPtr->Skeleton != nullptr);
 	}
 
 	FCompressibleAnimPtr GetCompressibleData() const { return DataToCompressPtr; }
diff --git a/Engine/Source/Runtime/Engine/Private/Animation/AnimCompressionTypes.cpp b/Engine/Source/Runtime/Engine/Private/Animation/AnimCompressionTypes.cpp
--- a/Engine/Source/Runtime/Engine/Private/Animation/AnimCompressionTypes.cpp
+++ b/Engine/Source/Runtime/Engine/Private/Animation/AnimCompressionTypes.cpp
@@ -128,7 +128,6 @@ void StripFramesOdd(TArray<ArrayValue>& Keys, const int32 NumFrames)
 FCompressibleAnimData::FCompressibleAnimData(class UAnimSequence* InSeq, const bool bPerformStripping)
 	: CurveCompressionSettings(InSeq->CurveCompressionSettings)
 	, BoneCompressionSettings(InSeq->BoneCompressionSettings)
-	, Skeleton(InSeq->GetSkeleton())
 	, TrackToSkeletonMapTable(InSeq->GetRawTrackToSkeletonMapTable())
 	, Interpolation(InSeq->Interpolation)
 	, SequenceLength(InSeq->SequenceLength)
@@ -144,8 +143,12 @@ FCompressibleAnimData::FCompressibleAnimData(class UAnimSequence* InSeq, const b
 	, AnimFName(InSeq->GetFName())
 {
 #if WITH_EDITOR
+	USkeleton* Skeleton = InSeq->GetSkeleton();
 	FAnimationUtils::BuildSkeletonMetaData(Skeleton, BoneData);
 
+	RefLocalPoses = Skeleton->GetRefLocalPoses();
+	RefSkeleton = Skeleton->GetReferenceSkeleton();
+
 	const bool bHasVirtualBones = InSeq->GetSkeleton()->GetVirtualBones().Num() > 0;
 
 	if (InSeq->CanBakeAdditive())
@@ -205,7 +208,6 @@ FCompressibleAnimData::FCompressibleAnimData(class UAnimSequence* InSeq, const b
 FCompressibleAnimData::FCompressibleAnimData(UAnimBoneCompressionSettings* InBoneCompressionSettings, UAnimCurveCompressionSettings* InCurveCompressionSettings, USkeleton* InSkeleton, EAnimInterpolationType InInterpolation, float InSequenceLength, int32 InNumFrames)
 	: CurveCompressionSettings(InCurveCompressionSettings)
 	, BoneCompressionSettings(InBoneCompressionSettings)
-	, Skeleton(InSkeleton)
 	, Interpolation(InInterpolation)
 	, SequenceLength(InSequenceLength)
 	, NumFrames(InNumFrames)
@@ -213,14 +215,15 @@ FCompressibleAnimData::FCompressibleAnimData(UAnimBoneCompressionSettings* InBon
 	, ErrorThresholdScale(1.f)
 {
 #if WITH_EDITOR
-	FAnimationUtils::BuildSkeletonMetaData(Skeleton, BoneData);
+	RefLocalPoses = InSkeleton->GetRefLocalPoses();
+	RefSkeleton = InSkeleton->GetReferenceSkeleton();
+	FAnimationUtils::BuildSkeletonMetaData(InSkeleton, BoneData);
 #endif
 }
 
 FCompressibleAnimData::FCompressibleAnimData()
 : CurveCompressionSettings(nullptr)
 , BoneCompressionSettings(nullptr)
-, Skeleton(nullptr)
 , Interpolation((EAnimInterpolationType)0)
 , SequenceLength(0.f)
 , NumFrames(0)
diff --git a/Engine/Source/Runtime/Engine/Private/Animation/AnimationUtils.cpp b/Engine/Source/Runtime/Engine/Private/Animation/AnimationUtils.cpp
--- a/Engine/Source/Runtime/Engine/Private/Animation/AnimationUtils.cpp
+++ b/Engine/Source/Runtime/Engine/Private/Animation/AnimationUtils.cpp
@@ -249,10 +249,7 @@ void FAnimationUtils::ComputeCompressionError(const FCompressibleAnimData& Compr
 		float ErrorCount = 0.0f;
 		float ErrorTotal = 0.0f;
 
-		USkeleton* Skeleton = CompressibleAnimData.Skeleton;
-		check ( Skeleton );
-
-		const TArray<FTransform>& RefPose = Skeleton->GetRefLocalPoses();
+		const TArray<FTransform>& RefPose = CompressibleAnimData.RefLocalPoses;
 
 		TArray<FTransform> RawTransforms;
 		TArray<FTransform> NewTransforms;
@@ -272,7 +269,7 @@ void FAnimationUtils::ComputeCompressionError(const FCompressibleAnimData& Compr
 		{
 			FCachedBoneIndexData& BoneIndexData = CachedBoneIndexData[BoneIndex];
 			BoneIndexData.TrackIndex = GetAnimTrackIndexForSkeletonBone(BoneIndex, CompressibleAnimData.TrackToSkeletonMapTable);
-			BoneIndexData.ParentIndex = Skeleton->GetReferenceSkeleton().GetParentIndex(BoneIndex);
+			BoneIndexData.ParentIndex = CompressibleAnimData.RefSkeleton.GetParentIndex(BoneIndex);
 		}
 
 		// Check the precondition that parents occur before children in the RequiredBones array.
@@ -604,10 +601,7 @@ void FAnimationUtils::TallyErrorsFromPerturbation(
 	const int32 NumBones = CompressibleAnimData.BoneData.Num();
 
 
-	USkeleton* Skeleton = CompressibleAnimData.Skeleton;
-	check ( Skeleton );
-
-	const TArray<FTransform>& RefPose = Skeleton->GetRefLocalPoses();
+	const TArray<FTransform>& RefPose = CompressibleAnimData.RefLocalPoses;
 
 	TArray<FTransform> RawAtoms;
 	TArray<FTransform> NewAtomsT;
@@ -690,7 +684,7 @@ void FAnimationUtils::TallyErrorsFromPerturbation(
 				// For all bones below the root, final component-space transform is relative transform * component-space transform of parent.
 				if ( BoneIndex > 0 )
 				{
-					const int32 ParentIndex = Skeleton->GetReferenceSkeleton().GetParentIndex(BoneIndex);
+					const int32 ParentIndex = CompressibleAnimData.RefSkeleton.GetParentIndex(BoneIndex);
 
 					// Check the precondition that parents occur before children in the RequiredBones array.
 					check( ParentIndex != INDEX_NONE );
@@ -945,11 +939,6 @@ bool FAnimationUtils::CompressAnimBones(FCompressibleAnimData& AnimSeq, FCompres
 	// Clear any previous data we might have even if we end up failing to compress
 	Target = FCompressibleAnimDataResult();
 
-	if (AnimSeq.Skeleton == nullptr)
-	{
-		return false;
-	}
-
 	if (AnimSeq.BoneCompressionSettings == nullptr || !AnimSeq.BoneCompressionSettings->AreSettingsValid())
 	{
 		return false;
diff --git a/Engine/Source/Runtime/Engine/Public/Animation/AnimCompressionTypes.h b/Engine/Source/Runtime/Engine/Public/Animation/AnimCompressionTypes.h
--- a/Engine/Source/Runtime/Engine/Public/Animation/AnimCompressionTypes.h
+++ b/Engine/Source/Runtime/Engine/Public/Animation/AnimCompressionTypes.h
@@ -201,7 +201,9 @@ public:
 
 	UAnimBoneCompressionSettings* BoneCompressionSettings;
 
-	USkeleton* Skeleton;
+	// Data from USkeleton
+	TArray<FTransform> RefLocalPoses;
+	FReferenceSkeleton RefSkeleton;
 
 	TArray<FTrackToSkeletonMap> TrackToSkeletonMapTable;
 
