State machine recorded weights are now correctly reported

Weights are now recorded at the end of the machine update, rather than part-way through active transition updates, where some weights will not have been recalculated yet.

diff --git a/Engine/Source/Runtime/Engine/Private/Animation/AnimNode_StateMachine.cpp b/Engine/Source/Runtime/Engine/Private/Animation/AnimNode_StateMachine.cpp
--- a/Engine/Source/Runtime/Engine/Private/Animation/AnimNode_StateMachine.cpp
+++ b/Engine/Source/Runtime/Engine/Private/Animation/AnimNode_StateMachine.cpp
@@ -550,16 +550,23 @@ void FAnimNode_StateMachine::Update_AnyThread(const FAnimationUpdateContext& Con
 	if (ActiveTransitionArray.Num() == 0 && !IsAConduitState(CurrentState) && !StatesUpdated.Contains(CurrentState))
 	{
 		StatePoseLinks[CurrentState].Update(Context);
-
-		const float StateWeight = GetStateWeight(CurrentState);
-		const float CurrentStateElapsedTime = GetCurrentStateElapsedTime();
-		Context.AnimInstanceProxy->RecordStateWeight(StateMachineIndexInClass, CurrentState, StateWeight, CurrentStateElapsedTime);
-
-		TRACE_ANIM_STATE_MACHINE_STATE(Context, StateMachineIndexInClass, CurrentState, StateWeight, CurrentStateElapsedTime);
 	}
 
 	ElapsedTime += Context.GetDeltaTime();
 
+	// Record state weights after transitions/updates are completed
+	for (int32 StateIndex = 0; StateIndex < StatePoseLinks.Num(); ++StateIndex)
+	{
+		const float StateWeight = GetStateWeight(StateIndex);
+		if(StateWeight > 0.0f)
+		{
+			const float CurrentStateElapsedTime = StateIndex == CurrentState ? GetCurrentStateElapsedTime() : 0.0f;
+			Context.AnimInstanceProxy->RecordStateWeight(StateMachineIndexInClass, StateIndex, StateWeight, CurrentStateElapsedTime);
+
+			TRACE_ANIM_STATE_MACHINE_STATE(Context, StateMachineIndexInClass, StateIndex, StateWeight, CurrentStateElapsedTime);
+		}
+	}
+
 #if ANIM_TRACE_ENABLED
 	TRACE_ANIM_NODE_VALUE(Context, TEXT("Name"), GetMachineDescription()->MachineName);
 	TRACE_ANIM_NODE_VALUE(Context, TEXT("Current State"), GetStateInfo().StateName);
@@ -1072,12 +1079,6 @@ void FAnimNode_StateMachine::UpdateState(int32 StateIndex, const FAnimationUpdat
 	{
 		StatesUpdated.Add(StateIndex);
 		StatePoseLinks[StateIndex].Update(Context);
-
-		const float StateWeight = GetStateWeight(StateIndex);
-		const float CurrentStateElapsedTime = StateIndex == CurrentState ? GetCurrentStateElapsedTime() : 0.0f;
-		Context.AnimInstanceProxy->RecordStateWeight(StateMachineIndexInClass, StateIndex, StateWeight, CurrentStateElapsedTime);
-
-		TRACE_ANIM_STATE_MACHINE_STATE(Context, StateMachineIndexInClass, StateIndex, StateWeight, CurrentStateElapsedTime);
 	}
 }
 
