- Add an Auto Destroy Subsystem that checks actors who have bAutoDestroyWhenFinished set to true
- Encapsulate bAutoDestroyWhenFinished into getter/setter functions that appropriately register to the subsystem and replace the references to existing calls

diff --git a/Engine/Source/Runtime/Engine/Classes/Engine/AutoDestroySubsystem.h b/Engine/Source/Runtime/Engine/Classes/Engine/AutoDestroySubsystem.h
new file mode 100644
--- /dev/null
+++ b/Engine/Source/Runtime/Engine/Classes/Engine/AutoDestroySubsystem.h
@@ -0,0 +1,81 @@
+// Copyright Epic Games, Inc. All Rights Reserved.
+
+#pragma once
+
+#include "Subsystems/WorldSubsystem.h"
+#include "Tickable.h"
+#include "Stats/Stats.h"
+#include "Engine/EngineTypes.h"
+
+#include "AutoDestroySubsystem.generated.h"
+
+class AActor;
+struct FLatentActionManager;
+
+/**
+* The Auto destroy subsystem manages actors who have bAutoDestroyWhenFinished
+* set as true. This ensures that even actors who do not have Tick enabled 
+* get properly destroyed, as well as decouple this behavior from AActor::Tick
+*/
+UCLASS()
+class UAutoDestroySubsystem : public UWorldSubsystem, public FTickableGameObject
+{
+	GENERATED_BODY()
+
+public:
+
+	/**
+	* Registers the given actor if it has the bAutoDestroyWhenFinished flag set
+	*
+	* @return True if actor registered
+	*/
+	bool RegisterActor(AActor* ActorToRegister);
+
+	/**
+	* Remove this actor from the array of actors to poll with this subsystem
+	* 
+	* @return	True if this actor was removed
+	*/
+	bool UnregisterActor(AActor* ActorToRemove);
+
+protected:
+
+	//~FTickableGameObject interface
+
+	ETickableTickType GetTickableTickType() const override;
+
+	/** 
+	* Update any Actors that have bAutoDestroyWhenFinished set to true. 
+	* If the actor should be destroyed, then destroy it and it's components
+	*/
+	void Tick(float DeltaTime) override;
+
+	TStatId GetStatId() const override { RETURN_QUICK_DECLARE_CYCLE_STAT(UAutoDestroySubsystem, STATGROUP_Tickables); }
+
+	//~End of FTickableGameObject interface
+
+
+	//~USubsystem interface
+	void Deinitialize() override;
+	//~End of USubsystem interface
+
+private:
+	
+	/** Callback for a registered actor's End Play so we can remove it from our known actors */
+	UFUNCTION()
+	void OnActorEndPlay(AActor* Actor, EEndPlayReason::Type EndPlayReason);
+
+	/** 
+	* Returns true if there are no latent actions on the given actor or the actor cannot have latent actions
+	*/
+	static bool CheckLatentActionsOnActor(FLatentActionManager& LatentActionManager, AActor* ActorToCheck, float WorldDeltaTime);
+
+	/** 
+	* Check each component on the given actor and if they are all ready for auto destroy, return true
+	*/
+	static bool ActorComponentsAreReadyForDestroy(AActor* const ActorToCheck);
+
+	/** Actors to check if they should auto destroy or not */
+	UPROPERTY()
+	TArray<AActor*> ActorsToPoll;
+};
\ No newline at end of file
diff --git a/Engine/Source/Runtime/Engine/Classes/GameFramework/Actor.h b/Engine/Source/Runtime/Engine/Classes/GameFramework/Actor.h
--- a/Engine/Source/Runtime/Engine/Classes/GameFramework/Actor.h
+++ b/Engine/Source/Runtime/Engine/Classes/GameFramework/Actor.h
@@ -256,10 +256,12 @@ public:
 	UPROPERTY(EditDefaultsOnly, Category=Tick)
 	uint8 bAllowTickBeforeBeginPlay:1;
 
+private:
 	/** If true then destroy self when "finished", meaning all relevant components report that they are done and no timelines or timers are in flight. */
-	UPROPERTY(BlueprintReadWrite, Category=Actor)
+	UPROPERTY(BlueprintSetter=SetAutoDestroyWhenFinished, Category=Actor)
 	uint8 bAutoDestroyWhenFinished:1;
 
+public:
 	/** If true, all input on the stack below this actor will not be considered */
 	UPROPERTY(EditDefaultsOnly, Category=Input)
 	uint8 bBlockInput:1;
@@ -319,6 +321,11 @@ public:
 	 */
 	virtual bool HasLocalNetOwner() const;
 
+	bool GetAutoDestroyWhenFinished() const { return bAutoDestroyWhenFinished; }
+
+	UFUNCTION(BlueprintSetter)
+	void SetAutoDestroyWhenFinished(bool bVal);
+
 protected:
 	/**
 	 * If true, this actor will replicate to remote machines
diff --git a/Engine/Source/Runtime/Engine/Private/Actor.cpp b/Engine/Source/Runtime/Engine/Private/Actor.cpp
--- a/Engine/Source/Runtime/Engine/Private/Actor.cpp
+++ b/Engine/Source/Runtime/Engine/Private/Actor.cpp
@@ -42,6 +42,7 @@
 #include "DeviceProfiles/DeviceProfile.h"
 #include "ObjectTrace.h"
 #include "Net/Core/PushModel/PushModel.h"
+#include "Engine/AutoDestroySubsystem.h"
 
 DEFINE_LOG_CATEGORY(LogActor);
 
@@ -1009,8 +1010,6 @@ void AActor::TickActor( float DeltaSeconds, ELevelTick TickType, FActorTickFunct
 
 void AActor::Tick( float DeltaSeconds )
 {
-	bool bShouldAutoDestroy = bAutoDestroyWhenFinished;
-
 	if (GetClass()->HasAnyClassFlags(CLASS_CompiledFromBlueprint) || !GetClass()->HasAnyClassFlags(CLASS_Native))
 	{
 		// Blueprint code outside of the construction script should not run in the editor
@@ -1032,28 +1031,6 @@ void AActor::Tick( float DeltaSeconds )
 		{
 			FLatentActionManager& LatentActionManager = MyWorld->GetLatentActionManager();
 			LatentActionManager.ProcessLatentActions(this, MyWorld->GetDeltaSeconds());
-			if (bShouldAutoDestroy)
-			{
-				bShouldAutoDestroy = (LatentActionManager.GetNumActionsForObject(this) == 0);
-			}
-		}
-	}
-
-	if (bShouldAutoDestroy)
-	{
-		for (UActorComponent* const Comp : GetComponents())
-		{
-			if (Comp && !Comp->IsReadyForOwnerToAutoDestroy())
-			{
-				bShouldAutoDestroy = false;
-				break;
-			}
-		}
-
-		// die!
-		if (bShouldAutoDestroy)
-		{
-			Destroy();
 		}
 	}
 }
@@ -1662,6 +1639,23 @@ bool AActor::HasNetOwner() const
 	return TopOwner->HasNetOwner();
 }
 
+void AActor::SetAutoDestroyWhenFinished(bool bVal)
+{
+	bAutoDestroyWhenFinished = bVal;
+	if (UWorld* MyWorld = GetWorld())
+	{
+		UAutoDestroySubsystem* AutoDestroySub = MyWorld->GetSubsystem<UAutoDestroySubsystem>();
+		if (bAutoDestroyWhenFinished && (HasActorBegunPlay() || IsActorBeginningPlay()))
+		{
+			AutoDestroySub->RegisterActor(this);
+		}
+		else
+		{
+			AutoDestroySub->UnregisterActor(this);
+		}
+	}
+}
+
 void AActor::K2_AttachRootComponentTo(USceneComponent* InParent, FName InSocketName, EAttachLocation::Type AttachLocationType /*= EAttachLocation::KeepRelativeOffset */, bool bWeldSimulatedBodies /*=true*/)
 {
 PRAGMA_DISABLE_DEPRECATION_WARNINGS
@@ -3490,6 +3484,17 @@ void AActor::BeginPlay()
 		}
 	}
 
+	if (GetAutoDestroyWhenFinished())
+	{
+		if (UWorld* MyWorld = GetWorld())
+		{
+			if (UAutoDestroySubsystem* AutoDestroySys = MyWorld->GetSubsystem<UAutoDestroySubsystem>())
+			{
+				AutoDestroySys->RegisterActor(this);
+			}			
+		}
+	}
+
 	ReceiveBeginPlay();
 
 	ActorHasBegunPlay = EActorBeginPlayState::HasBegunPlay;
diff --git a/Engine/Source/Runtime/Engine/Private/Engine/AutoDestroySubsystem.cpp b/Engine/Source/Runtime/Engine/Private/Engine/AutoDestroySubsystem.cpp
new file mode 100644
--- /dev/null
+++ b/Engine/Source/Runtime/Engine/Private/Engine/AutoDestroySubsystem.cpp
@@ -0,0 +1,122 @@
+// Copyright Epic Games, Inc. All Rights Reserved.
+
+#include "Engine/AutoDestroySubsystem.h"
+#include "Engine/World.h"
+#include "Logging/LogMacros.h"
+#include "GameFramework/Actor.h"
+
+DECLARE_LOG_CATEGORY_EXTERN(LogAutoDestroySubsystem, Log, All);
+DEFINE_LOG_CATEGORY(LogAutoDestroySubsystem);
+
+void UAutoDestroySubsystem::Deinitialize()
+{
+	// Ensure that the resources from the world are cleaned up and callbacks unregistered
+	for (AActor* Actor : ActorsToPoll)
+	{
+		if (Actor)
+		{
+			Actor->OnEndPlay.RemoveAll(this);
+		}
+	}
+	ActorsToPoll.Empty();
+}
+
+bool UAutoDestroySubsystem::RegisterActor(AActor* ActorToRegister)
+{
+	if (ActorToRegister && ActorToRegister->GetAutoDestroyWhenFinished())
+	{
+		ActorsToPoll.AddUnique(ActorToRegister);
+		ActorToRegister->OnEndPlay.AddDynamic(this, &UAutoDestroySubsystem::OnActorEndPlay);
+		return true;
+	}
+
+	return false;
+}
+
+void UAutoDestroySubsystem::OnActorEndPlay(AActor* Actor, EEndPlayReason::Type EndPlayReason)
+{
+	UnregisterActor(Actor);
+}
+
+bool UAutoDestroySubsystem::UnregisterActor(AActor* ActorToRemove)
+{
+	if (ActorToRemove && ActorsToPoll.Remove(ActorToRemove))
+	{
+		ActorToRemove->OnEndPlay.RemoveAll(this);
+		return true;
+	}
+	return false;
+}
+
+bool UAutoDestroySubsystem::CheckLatentActionsOnActor(FLatentActionManager& LatentActionManager, AActor* ActorToCheck, float WorldDeltaTime)
+{
+	UClass* const ActorClass = ActorToCheck ? ActorToCheck->GetClass() : nullptr;
+
+	// Check if the latent action manager knows that this actor should be destroyed 	
+	if (ActorClass && (ActorClass->HasAnyClassFlags(CLASS_CompiledFromBlueprint) || !ActorClass->HasAnyClassFlags(CLASS_Native)))
+	{
+		return LatentActionManager.GetNumActionsForObject(ActorToCheck) == 0;
+	}
+
+	return true;
+}
+
+bool UAutoDestroySubsystem::ActorComponentsAreReadyForDestroy(AActor* const ActorToCheck)
+{
+	if (ActorToCheck)
+	{
+		// Check if this actors components are ready to be destroyed before marking it for auto destruction
+		for (UActorComponent* const Comp : ActorToCheck->GetComponents())
+		{
+			if (Comp && !Comp->IsReadyForOwnerToAutoDestroy())
+			{
+				return false;
+			}
+		}
+	}	
+
+	return true;
+}
+
+ETickableTickType UAutoDestroySubsystem::GetTickableTickType() const
+{
+	// The CDO of this should never tick
+	return IsTemplate() ? ETickableTickType::Never : ETickableTickType::Always;
+}
+
+void UAutoDestroySubsystem::Tick(float DeltaTime)
+{
+	UWorld* MyWorld = GetWorld();
+	
+	if (MyWorld && ActorsToPoll.Num())
+	{
+		// Use a copied array because if the user decides to override the OnDestroy of an actor and 
+		// set it's AutoDestroy flag, we could end up modifiying the array
+		TArray<AActor*, TInlineAllocator<8>> ActorsToDestroy;
+
+		FLatentActionManager& LatentActionManager = MyWorld->GetLatentActionManager();
+
+		for (int32 i = ActorsToPoll.Num() - 1; i >= 0; --i)
+		{
+			AActor* const CurActor = ActorsToPoll[i];
+
+			if (CurActor && CurActor->GetAutoDestroyWhenFinished())
+			{
+				if (CheckLatentActionsOnActor(LatentActionManager, CurActor, DeltaTime) && ActorComponentsAreReadyForDestroy(CurActor))
+				{
+					ActorsToDestroy.Emplace(CurActor);
+					ActorsToPoll.RemoveSwap(CurActor);
+				}
+			}
+		}
+
+		// Actually destroy all necessary actors
+		for (AActor* CurActor : ActorsToDestroy)
+		{
+			if(CurActor)
+			{
+				CurActor->Destroy();
+			}
+		}
+	}
+}
\ No newline at end of file
