Update the split screen game viewport client to have higher precision float values when creating the splitscreen info in order to remove unwanted black bars at higher resolutions

diff --git a/Engine/Source/Runtime/Engine/Private/GameViewportClient.cpp b/Engine/Source/Runtime/Engine/Private/GameViewportClient.cpp
--- a/Engine/Source/Runtime/Engine/Private/GameViewportClient.cpp
+++ b/Engine/Source/Runtime/Engine/Private/GameViewportClient.cpp
@@ -182,6 +182,9 @@ UGameViewportClient::UGameViewportClient(const FObjectInitializer& ObjectInitial
 
 	SplitscreenInfo.Init(FSplitscreenData(), ESplitScreenType::SplitTypeCount);
 
+	static float OneOverThree = 1.0f / 3.0f;
+	static float TwoOverThree = 2.0f / 3.0f;
+
 	SplitscreenInfo[ESplitScreenType::None].PlayerData.Add(FPerPlayerSplitscreenData(1.0f, 1.0f, 0.0f, 0.0f));
 
 	SplitscreenInfo[ESplitScreenType::TwoPlayer_Horizontal].PlayerData.Add(FPerPlayerSplitscreenData(1.0f, 0.5f, 0.0f, 0.0f));
@@ -198,13 +201,13 @@ UGameViewportClient::UGameViewportClient(const FObjectInitializer& ObjectInitial
 	SplitscreenInfo[ESplitScreenType::ThreePlayer_FavorBottom].PlayerData.Add(FPerPlayerSplitscreenData(0.5f, 0.5f, 0.5f, 0.0f));
 	SplitscreenInfo[ESplitScreenType::ThreePlayer_FavorBottom].PlayerData.Add(FPerPlayerSplitscreenData(1.0f, 0.5f, 0.0f, 0.5f));
 
-	SplitscreenInfo[ESplitScreenType::ThreePlayer_Vertical].PlayerData.Add(FPerPlayerSplitscreenData(0.333f, 1.0f, 0.0f, 0.0f));
-	SplitscreenInfo[ESplitScreenType::ThreePlayer_Vertical].PlayerData.Add(FPerPlayerSplitscreenData(0.333f, 1.0f, 0.333f, 0.0f));
-	SplitscreenInfo[ESplitScreenType::ThreePlayer_Vertical].PlayerData.Add(FPerPlayerSplitscreenData(0.333f, 1.0f, 0.666f, 0.0f));
+	SplitscreenInfo[ESplitScreenType::ThreePlayer_Vertical].PlayerData.Add(FPerPlayerSplitscreenData(OneOverThree, 1.0f, 0.0f, 0.0f));
+	SplitscreenInfo[ESplitScreenType::ThreePlayer_Vertical].PlayerData.Add(FPerPlayerSplitscreenData(OneOverThree, 1.0f, OneOverThree, 0.0f));
+	SplitscreenInfo[ESplitScreenType::ThreePlayer_Vertical].PlayerData.Add(FPerPlayerSplitscreenData(OneOverThree, 1.0f, TwoOverThree, 0.0f));
 
-	SplitscreenInfo[ESplitScreenType::ThreePlayer_Horizontal].PlayerData.Add(FPerPlayerSplitscreenData(1.0f, 0.333f, 0.0f, 0.0f));
-	SplitscreenInfo[ESplitScreenType::ThreePlayer_Horizontal].PlayerData.Add(FPerPlayerSplitscreenData(1.0f, 0.333f, 0.0f, 0.333f));
-	SplitscreenInfo[ESplitScreenType::ThreePlayer_Horizontal].PlayerData.Add(FPerPlayerSplitscreenData(1.0f, 0.333f, 0.0f, 0.666f));
+	SplitscreenInfo[ESplitScreenType::ThreePlayer_Horizontal].PlayerData.Add(FPerPlayerSplitscreenData(1.0f, OneOverThree, 0.0f, 0.0f));
+	SplitscreenInfo[ESplitScreenType::ThreePlayer_Horizontal].PlayerData.Add(FPerPlayerSplitscreenData(1.0f, OneOverThree, 0.0f, OneOverThree));
+	SplitscreenInfo[ESplitScreenType::ThreePlayer_Horizontal].PlayerData.Add(FPerPlayerSplitscreenData(1.0f, OneOverThree, 0.0f, TwoOverThree));
 
 	SplitscreenInfo[ESplitScreenType::FourPlayer_Grid].PlayerData.Add(FPerPlayerSplitscreenData(0.5f, 0.5f, 0.0f, 0.0f));
 	SplitscreenInfo[ESplitScreenType::FourPlayer_Grid].PlayerData.Add(FPerPlayerSplitscreenData(0.5f, 0.5f, 0.5f, 0.0f));
