Make Actor, ActorComponent, and SceneComponent replicated properties deprecated in 4.24 private

diff --git a/Engine/Source/Runtime/Engine/Classes/Components/ActorComponent.h b/Engine/Source/Runtime/Engine/Classes/Components/ActorComponent.h
--- a/Engine/Source/Runtime/Engine/Classes/Components/ActorComponent.h
+++ b/Engine/Source/Runtime/Engine/Classes/Components/ActorComponent.h
@@ -141,16 +141,15 @@ protected:
 	/** If the physics state is currently created for this component */
 	uint8 bPhysicsStateCreated:1;
 
-	/** Is this component currently replicating? Should the network code consider it for replication? Owning Actor must be replicating first! */
-	UE_DEPRECATED(4.24, "This member will be made private. Please use GetIsReplicated, SetIsReplicated, or SetIsReplicatedByDefault for constructors.")
-	UPROPERTY(EditDefaultsOnly, BlueprintReadOnly, Replicated, Category=ComponentReplication,meta=(DisplayName = "Component Replicates"))
-	uint8 bReplicates:1;
-
 	/** Is this component safe to ID over the network by name?  */
 	UPROPERTY()
 	uint8 bNetAddressable:1;
 
 private:
+	/** Is this component currently replicating? Should the network code consider it for replication? Owning Actor must be replicating first! */
+	UPROPERTY(EditDefaultsOnly, BlueprintReadOnly, Replicated, Category=ComponentReplication,meta=(DisplayName = "Component Replicates", AllowPrivateAccess = "true"))
+	uint8 bReplicates:1;
+
 	/** Is this component in need of its whole state being sent to the renderer? */
 	uint8 bRenderStateDirty:1;
 
@@ -198,11 +197,13 @@ public:
 	UPROPERTY(EditAnywhere, BlueprintReadOnly, Category=Activation)
 	uint8 bAutoActivate:1;
 
+private:
 	/** Whether the component is currently active. */
-	UE_DEPRECATED(4.24, "This member will be made private. Please use IsActive or SetActive.")
 	UPROPERTY(transient, ReplicatedUsing=OnRep_IsActive)
 	uint8 bIsActive:1;
 
+public:
+
 	/** True if this component can be modified when it was inherited from a parent actor class */
 	UPROPERTY(EditDefaultsOnly, Category="Variable")
 	uint8 bEditableWhenInherited:1;
@@ -387,14 +388,12 @@ public:
 	UFUNCTION(BlueprintCallable, Category="Components|Activation", meta=(UnsafeDuringActorConstruction="true"))
 	virtual void ToggleActive();
 
-PRAGMA_DISABLE_DEPRECATION_WARNINGS
 	/**
 	 * Returns whether the component is active or not
 	 * @return - The active state of the component.
 	 */
 	UFUNCTION(BlueprintCallable, Category="Components|Activation", meta=(UnsafeDuringActorConstruction="true"))
 	bool IsActive() const { return bIsActive; }
-PRAGMA_ENABLE_DEPRECATION_WARNINGS
 
 	/**
 	 * Sets whether the component should be auto activate or not. Only safe during construction scripts.
@@ -423,13 +422,11 @@ PRAGMA_ENABLE_DEPRECATION_WARNINGS
 	UFUNCTION(BlueprintCallable, Category="Components")
 	void SetIsReplicated(bool ShouldReplicate);
 
-PRAGMA_DISABLE_DEPRECATION_WARNINGS
 	/** Returns whether replication is enabled or not. */
 	FORCEINLINE bool GetIsReplicated() const
 	{
 		return bReplicates;
 	}
-PRAGMA_ENABLE_DEPRECATION_WARNINGS
 
 	/** Allows a component to replicate other subobject on the actor  */
 	virtual bool ReplicateSubobjects(class UActorChannel *Channel, class FOutBunch *Bunch, FReplicationFlags *RepFlags);
@@ -960,7 +957,6 @@ private:
 	friend class FComponentRecreateRenderStateContext;
 	friend struct FActorComponentTickFunction;
 
-PRAGMA_DISABLE_DEPRECATION_WARNINGS
 	//~ Begin Methods for Replicated Members.
 protected:
 
@@ -994,7 +990,6 @@ public:
 	void SetActiveFlag(const bool bNewIsActive);
 
 	//~ End Methods for Replicated Members.
-PRAGMA_ENABLE_DEPRECATION_WARNINGS
 
 protected:
 
diff --git a/Engine/Source/Runtime/Engine/Classes/Components/SceneComponent.h b/Engine/Source/Runtime/Engine/Classes/Components/SceneComponent.h
--- a/Engine/Source/Runtime/Engine/Classes/Components/SceneComponent.h
+++ b/Engine/Source/Runtime/Engine/Classes/Components/SceneComponent.h
@@ -148,24 +148,23 @@ public:
 	/** Current bounds of the component */
 	FBoxSphereBounds Bounds;
 
+private:
 	/** Location of the component relative to its parent */
-	UE_DEPRECATED(4.24, "This member will be made private. Please use GetRelativeLocation or SetRelativeLocation.")
-	UPROPERTY(EditAnywhere, BlueprintReadOnly, ReplicatedUsing=OnRep_Transform, Category = Transform)
+	UPROPERTY(EditAnywhere, BlueprintReadOnly, ReplicatedUsing=OnRep_Transform, Category = Transform, meta=(AllowPrivateAccess="true"))
 	FVector RelativeLocation;
 
 	/** Rotation of the component relative to its parent */
-	UE_DEPRECATED(4.24, "This member will be made private. Please use GetRelativeRotation or SetRelativeRotation.")
-	UPROPERTY(EditAnywhere, BlueprintReadOnly, ReplicatedUsing=OnRep_Transform, Category=Transform)
+	UPROPERTY(EditAnywhere, BlueprintReadOnly, ReplicatedUsing=OnRep_Transform, Category=Transform, meta=(AllowPrivateAccess="true"))
 	FRotator RelativeRotation;
 
 	/**
 	*	Non-uniform scaling of the component relative to its parent.
 	*	Note that scaling is always applied in local space (no shearing etc)
 	*/
-	UE_DEPRECATED(4.24, "This member will be made private. Please use GetRelativeScale3D or SetRelativeScale3D.")
-	UPROPERTY(BlueprintReadOnly, ReplicatedUsing=OnRep_Transform, interp, Category=Transform)
+	UPROPERTY(BlueprintReadOnly, ReplicatedUsing=OnRep_Transform, interp, Category=Transform, meta=(AllowPrivateAccess="true"))
 	FVector RelativeScale3D;
 
+public:
 	/**
 	* Velocity of the component.
 	* @see GetComponentVelocity()
@@ -182,32 +181,22 @@ private:
 	* This should only be set to true as a result of UpdateOverlaps. To dirty this flag see ClearSkipUpdateOverlaps() which is expected when state affecting UpdateOverlaps changes (attachment, Collision settings, etc...) */
 	uint8 bSkipUpdateOverlaps : 1;
 
-public:
 	/** If RelativeLocation should be considered relative to the world, rather than the parent */
-	UE_DEPRECATED(4.24, "This member will be made private. Please use IsUsingAbsoluteLocation or SetUsingAbsoluteLocation.")
-	UPROPERTY(EditAnywhere, AdvancedDisplay, BlueprintReadWrite, ReplicatedUsing=OnRep_Transform, Category=Transform)
+	UPROPERTY(EditAnywhere, AdvancedDisplay, BlueprintReadWrite, ReplicatedUsing=OnRep_Transform, Category=Transform, meta=(AllowPrivateAccess="true"))
 	uint8 bAbsoluteLocation:1;
 
 	/** If RelativeRotation should be considered relative to the world, rather than the parent */
-	UE_DEPRECATED(4.24, "This member will be made private. Please use IsUsingAbsoluteRotation or SetUsingAbsoluteRotation.")
-	UPROPERTY(EditAnywhere, AdvancedDisplay, BlueprintReadWrite, ReplicatedUsing=OnRep_Transform, Category=Transform)
+	UPROPERTY(EditAnywhere, AdvancedDisplay, BlueprintReadWrite, ReplicatedUsing=OnRep_Transform, Category=Transform, meta=(AllowPrivateAccess="true"))
 	uint8 bAbsoluteRotation:1;
 
 	/** If RelativeScale3D should be considered relative to the world, rather than the parent */
-	UE_DEPRECATED(4.24, "This member will be made private. Please use IsUsingAbsoluteScale or SetUsingAbsoluteScale.")
-	UPROPERTY(EditAnywhere, AdvancedDisplay, BlueprintReadWrite, ReplicatedUsing=OnRep_Transform, Category=Transform)
+	UPROPERTY(EditAnywhere, AdvancedDisplay, BlueprintReadWrite, ReplicatedUsing=OnRep_Transform, Category=Transform, meta=(AllowPrivateAccess="true"))
 	uint8 bAbsoluteScale:1;
 
 	/** Whether to completely draw the primitive; if false, the primitive is not drawn, does not cast a shadow. */
-	UE_DEPRECATED(4.24, "This member will be made private. Please use IsVisible or SetVisibility.")
-	UPROPERTY(EditAnywhere, BlueprintReadOnly, ReplicatedUsing=OnRep_Visibility,  Category = Rendering)
+	UPROPERTY(EditAnywhere, BlueprintReadOnly, ReplicatedUsing=OnRep_Visibility,  Category = Rendering, meta=(AllowPrivateAccess="true"))
 	uint8 bVisible:1;
 
-	/** Whether to hide the primitive in game, if the primitive is Visible. */
-	UPROPERTY(Interp, EditAnywhere, BlueprintReadOnly, Category=Rendering, meta=(SequencerTrackClass = "MovieSceneVisibilityTrack"))
-	uint8 bHiddenInGame:1;
-
-private:
 	/** Whether or not we should be attached. */
 	UPROPERTY(Transient, Replicated)
 	uint8 bShouldBeAttached : 1;
@@ -226,6 +215,10 @@ private:
 	uint8 bShouldUpdatePhysicsVolume:1;
 
 public:
+	/** Whether to hide the primitive in game, if the primitive is Visible. */
+	UPROPERTY(Interp, EditAnywhere, BlueprintReadOnly, Category=Rendering, meta=(SequencerTrackClass = "MovieSceneVisibilityTrack"))
+	uint8 bHiddenInGame:1;
+
 	/** If true, a change in the bounds of the component will call trigger a streaming data rebuild */
 	UPROPERTY()
 	uint8 bBoundsChangeTriggersStreamingDataRebuild:1;
@@ -1258,7 +1251,6 @@ private:
 	friend class FScopedPreventAttachedComponentMove;
 	friend struct FDirectAttachChildrenAccessor;
 
-PRAGMA_DISABLE_DEPRECATION_WARNINGS
 	//~ Begin Methods for Replicated Members.
 private:
 
@@ -1516,7 +1508,6 @@ public:
 	void SetVisibleFlag(const bool bInVisible);
 	
 	//~ End Methods for Replicated Members.
-PRAGMA_ENABLE_DEPRECATION_WARNINGS
 };
 
 /** 
diff --git a/Engine/Source/Runtime/Engine/Classes/GameFramework/Actor.h b/Engine/Source/Runtime/Engine/Classes/GameFramework/Actor.h
--- a/Engine/Source/Runtime/Engine/Classes/GameFramework/Actor.h
+++ b/Engine/Source/Runtime/Engine/Classes/GameFramework/Actor.h
@@ -159,14 +159,6 @@ public:
 	UPROPERTY(EditDefaultsOnly, Category=Tick)
 	struct FActorTickFunction PrimaryActorTick;
 
-	/**
-	 * Allows us to only see this Actor in the Editor, and not in the actual game.
-	 * @see SetActorHiddenInGame()
-	 */
-	UE_DEPRECATED(4.24, "This member will be made private. Please use IsHidden or SetHidden.")
-	UPROPERTY(Interp, EditAnywhere, Category=Rendering, BlueprintReadOnly, Replicated, meta=(DisplayName="Actor Hidden In Game", SequencerTrackClass="MovieSceneVisibilityTrack"))
-	uint8 bHidden:1;
-
 	/** If true, when the actor is spawned it will be sent to the client but receive no further replication updates from the server afterwards. */
 	UPROPERTY()
 	uint8 bNetTemporary:1;
@@ -183,21 +175,26 @@ public:
 	UPROPERTY(Category=Replication, EditDefaultsOnly, BlueprintReadWrite)
 	uint8 bAlwaysRelevant:1;    
 
+	/** Called on client when updated bReplicateMovement value is received for this actor. */
+	UFUNCTION()
+	virtual void OnRep_ReplicateMovement();
+
+private:
 	/**
 	 * If true, replicate movement/location related properties.
 	 * Actor must also be set to replicate.
 	 * @see SetReplicates()
 	 * @see https://docs.unrealengine.com/latest/INT/Gameplay/Networking/Replication/
 	 */
-	UE_DEPRECATED(4.24, "This member will be made private. Please use IsReplicatingMovement or SetReplicatingMovement.")
 	UPROPERTY(ReplicatedUsing=OnRep_ReplicateMovement, Category=Replication, EditDefaultsOnly)
 	uint8 bReplicateMovement:1;    
 
-	/** Called on client when updated bReplicateMovement value is received for this actor. */
-	UFUNCTION()
-	virtual void OnRep_ReplicateMovement();
-
-private:
+	/**
+	 * Allows us to only see this Actor in the Editor, and not in the actual game.
+	 * @see SetActorHiddenInGame()
+	 */
+	UPROPERTY(Interp, EditAnywhere, Category=Rendering, BlueprintReadOnly, Replicated, meta=(AllowPrivateAccess="true", DisplayName="Actor Hidden In Game", SequencerTrackClass="MovieSceneVisibilityTrack"))
+	uint8 bHidden:1;
 
 	UPROPERTY(Replicated)
 	uint8 bTearOff:1; 
@@ -261,20 +258,19 @@ private:
 	UPROPERTY(BlueprintSetter=SetAutoDestroyWhenFinished, Category=Actor)
 	uint8 bAutoDestroyWhenFinished:1;
 
-public:
-	/** If true, all input on the stack below this actor will not be considered */
-	UPROPERTY(EditDefaultsOnly, Category=Input)
-	uint8 bBlockInput:1;
-
 	/**
 	 * Whether this actor can take damage. Must be true for damage events (e.g. ReceiveDamage()) to be called.
 	 * @see https://www.unrealengine.com/blog/damage-in-ue4
 	 * @see TakeDamage(), ReceiveDamage()
 	 */
-	UE_DEPRECATED(4.24, "This member will be made private. Please use CanBeDamaged or SetCanBeDamaged.")
-	UPROPERTY(EditAnywhere, BlueprintReadWrite, SaveGame, Replicated, Category=Actor)
+	UPROPERTY(EditAnywhere, BlueprintReadWrite, SaveGame, Replicated, Category=Actor, meta=(AllowPrivateAccess="true"))
 	uint8 bCanBeDamaged:1;
 
+public:
+	/** If true, all input on the stack below this actor will not be considered */
+	UPROPERTY(EditDefaultsOnly, Category=Input)
+	uint8 bBlockInput:1;
+
 	/** This actor collides with the world when placing in the editor, even if RootComponent collision is disabled. Does not affect spawning, @see SpawnCollisionHandlingMethod */
 	UPROPERTY()
 	uint8 bCollideWhenPlacing:1;
@@ -482,17 +478,18 @@ public:
 
 	/** Returns how much control the local machine has over this actor. */
 	UFUNCTION(BlueprintCallable, Category=Replication)
-	ENetRole GetLocalRole() const;
+	ENetRole GetLocalRole() const { return Role; }
 
 	/** Returns how much control the remote machine has over this actor. */
 	UFUNCTION(BlueprintCallable, Category=Replication)
 	ENetRole GetRemoteRole() const;
 
+private:
 	/** Used for replication of our RootComponent's position and velocity */
-	UE_DEPRECATED(4.24, "This member will be made private. Please use GetReplicatedMovement or SetReplicatedMovement.")
 	UPROPERTY(EditDefaultsOnly, ReplicatedUsing=OnRep_ReplicatedMovement, Category=Replication, AdvancedDisplay)
 	struct FRepMovement ReplicatedMovement;
 
+public:
 	/** How long this Actor lives before dying, 0=forever. Note this is the INITIAL value and should not be modified once play has begun. */
 	UPROPERTY(EditAnywhere, BlueprintReadOnly, Category=Actor)
 	float InitialLifeSpan;
@@ -535,11 +532,12 @@ public:
 	UFUNCTION()
 	virtual void OnRep_AttachmentReplication();
 
+private:
 	/** Describes how much control the local machine has over the actor. */
-	UE_DEPRECATED(4.24, "This member will be made private. Please use GetLocalRole or SetRole.")
 	UPROPERTY(Replicated)
 	TEnumAsByte<enum ENetRole> Role;
 
+public:
 	/** Dormancy setting for actor to take itself off of the replication list without being destroyed on clients. */
 	UPROPERTY(BlueprintReadOnly, EditDefaultsOnly, Category=Replication)
 	TEnumAsByte<enum ENetDormancy> NetDormancy;
@@ -635,12 +633,13 @@ public:
 
 	/** Called by the networking system to call PreReplication on this actor and its components using the given NetDriver to find or create RepChangedPropertyTrackers. */
 	void CallPreReplication(UNetDriver* NetDriver);	
-	
+
+private:
 	/** Pawn responsible for damage and other gameplay events caused by this actor. */
-	UE_DEPRECATED(4.24, "This member will be made private. Please use GetInstigator or SetInstigator.")
-	UPROPERTY(BlueprintReadWrite, ReplicatedUsing=OnRep_Instigator, meta=(ExposeOnSpawn=true), Category=Actor)
+	UPROPERTY(BlueprintReadWrite, ReplicatedUsing=OnRep_Instigator, meta=(ExposeOnSpawn=true, AllowPrivateAccess=true), Category=Actor)
 	class APawn* Instigator;
 
+public:
 	/** Called on clients when Instigator is replicated. */
 	UFUNCTION()
 	virtual void OnRep_Instigator();
@@ -3258,8 +3257,6 @@ private:
 		return (RootComponent != nullptr) ? RootComponent->GetRightVector() : FVector::RightVector;
 	}
 	
-PRAGMA_DISABLE_DEPRECATION_WARNINGS
-
 	//~ Begin Methods for Replicated Members.
 public:
 
@@ -3390,7 +3387,6 @@ public:
 	void SetInstigator(APawn* InInstigator);
 
 	//~ End Methods for Replicated Members.
-PRAGMA_ENABLE_DEPRECATION_WARNINGS
 };
 
 /** Internal struct used by level code to mark actors as destroyed */
@@ -3562,14 +3558,6 @@ FORCEINLINE_DEBUGGABLE const AActor* AActor::GetNetOwner() const
 	return Owner;
 }
 
-PRAGMA_DISABLE_DEPRECATION_WARNINGS
-FORCEINLINE_DEBUGGABLE ENetRole AActor::GetLocalRole() const
-{
-	return Role;
-}
-PRAGMA_ENABLE_DEPRECATION_WARNINGS
-
-
 FORCEINLINE_DEBUGGABLE ENetRole AActor::GetRemoteRole() const
 {
 	return RemoteRole;
diff --git a/Engine/Source/Runtime/Engine/Private/Actor.cpp b/Engine/Source/Runtime/Engine/Private/Actor.cpp
--- a/Engine/Source/Runtime/Engine/Private/Actor.cpp
+++ b/Engine/Source/Runtime/Engine/Private/Actor.cpp
@@ -1055,9 +1055,7 @@ void AActor::PreReplication(IRepChangedPropertyTracker & ChangedPropertyTracker)
 
 	GatherCurrentMovement();
 
-PRAGMA_DISABLE_DEPRECATION_WARNINGS
 	DOREPLIFETIME_ACTIVE_OVERRIDE(AActor, ReplicatedMovement, IsReplicatingMovement());
-PRAGMA_ENABLE_DEPRECATION_WARNINGS
 
 	// Don't need to replicate AttachmentReplication if the root component replicates, because it already handles it.
 	DOREPLIFETIME_ACTIVE_OVERRIDE(AActor, AttachmentReplication, RootComponent && !RootComponent->GetIsReplicated());
@@ -2563,7 +2561,6 @@ void AActor::EndViewTarget( APlayerController* PC )
 	K2_OnEndViewTarget(PC);
 }
 
-PRAGMA_DISABLE_DEPRECATION_WARNINGS
 APawn* AActor::GetInstigator() const
 {
 	return Instigator;
@@ -2573,7 +2570,6 @@ AController* AActor::GetInstigatorController() const
 {
 	return Instigator ? Instigator->Controller : nullptr;
 }
-PRAGMA_ENABLE_DEPRECATION_WARNINGS
 
 void AActor::CalcCamera(float DeltaTime, FMinimalViewInfo& OutResult)
 {
@@ -3357,9 +3353,7 @@ void AActor::ExchangeNetRoles(bool bRemoteOwned)
 	{
 		if (bRemoteOwned)
 		{
-PRAGMA_DISABLE_DEPRECATION_WARNINGS
 			Exchange( Role, RemoteRole );
-PRAGMA_ENABLE_DEPRECATION_WARNINGS
 		}
 		bExchangedRoles = true;
 	}
@@ -3367,9 +3361,7 @@ PRAGMA_ENABLE_DEPRECATION_WARNINGS
 
 void AActor::SwapRoles()
 {
-PRAGMA_DISABLE_DEPRECATION_WARNINGS
 	Swap(Role, RemoteRole);
-PRAGMA_ENABLE_DEPRECATION_WARNINGS
 
 	MARK_PROPERTY_DIRTY_FROM_NAME(AActor, RemoteRole, this);
 	MARK_PROPERTY_DIRTY_FROM_NAME(AActor, Role, this);
@@ -5012,42 +5004,32 @@ void AActor::SetLODParent(UPrimitiveComponent* InLODParent, float InParentDrawDi
 
 void AActor::SetHidden(bool bInHidden)
 {
-PRAGMA_DISABLE_DEPRECATION_WARNINGS
 	bHidden = bInHidden;
 	MARK_PROPERTY_DIRTY_FROM_NAME(AActor, bHidden, this);
-PRAGMA_ENABLE_DEPRECATION_WARNINGS
 }
 
 void AActor::SetReplicatingMovement(bool bInReplicateMovement)
 {
-PRAGMA_DISABLE_DEPRECATION_WARNINGS
 	bReplicateMovement = bInReplicateMovement;
 	MARK_PROPERTY_DIRTY_FROM_NAME(AActor, bReplicateMovement, this);
-PRAGMA_ENABLE_DEPRECATION_WARNINGS
 }
 
 void AActor::SetCanBeDamaged(bool bInCanBeDamaged)
 {
-PRAGMA_DISABLE_DEPRECATION_WARNINGS
 	bCanBeDamaged = bInCanBeDamaged;
 	MARK_PROPERTY_DIRTY_FROM_NAME(AActor, bCanBeDamaged, this);
-PRAGMA_ENABLE_DEPRECATION_WARNINGS
 }
 
 void AActor::SetRole(ENetRole InRole)
 {
-PRAGMA_DISABLE_DEPRECATION_WARNINGS
 	Role = InRole;
 	MARK_PROPERTY_DIRTY_FROM_NAME(AActor, Role, this);
-PRAGMA_ENABLE_DEPRECATION_WARNINGS
 }
 
 FRepMovement& AActor::GetReplicatedMovement_Mutable()
 {
-PRAGMA_DISABLE_DEPRECATION_WARNINGS
 	MARK_PROPERTY_DIRTY_FROM_NAME(AActor, ReplicatedMovement, this);
 	return ReplicatedMovement;
-PRAGMA_ENABLE_DEPRECATION_WARNINGS
 }
 
 void AActor::SetReplicatedMovement(const FRepMovement& InReplicatedMovement)
@@ -5057,10 +5039,8 @@ void AActor::SetReplicatedMovement(const FRepMovement& InReplicatedMovement)
 
 void AActor::SetInstigator(APawn* InInstigator)
 {
-PRAGMA_DISABLE_DEPRECATION_WARNINGS
 	Instigator = InInstigator;
 	MARK_PROPERTY_DIRTY_FROM_NAME(AActor, Instigator, this);
-PRAGMA_ENABLE_DEPRECATION_WARNINGS
 }
 
 #undef LOCTEXT_NAMESPACE
diff --git a/Engine/Source/Runtime/Engine/Private/ActorReplication.cpp b/Engine/Source/Runtime/Engine/Private/ActorReplication.cpp
--- a/Engine/Source/Runtime/Engine/Private/ActorReplication.cpp
+++ b/Engine/Source/Runtime/Engine/Private/ActorReplication.cpp
@@ -145,10 +145,8 @@ void AActor::PostNetReceive()
 		bNetCheckedInitialPhysicsState = true;
 	}
 
-PRAGMA_DISABLE_DEPRECATION_WARNINGS
 	ExchangeB(bHidden, ActorReplication::SavedbHidden);
 	Exchange(Owner, ActorReplication::SavedOwner);
-PRAGMA_ENABLE_DEPRECATION_WARNINGS
 
 	if (IsHidden() != ActorReplication::SavedbHidden)
 	{
@@ -336,7 +334,6 @@ bool AActor::IsReplayRelevantFor(const AActor* RealViewer, const AActor* ViewTar
 
 void AActor::GatherCurrentMovement()
 {
-PRAGMA_DISABLE_DEPRECATION_WARNINGS
 	if (IsReplicatingMovement() || (RootComponent && RootComponent->GetAttachParent()))
 	{
 		bool bWasAttachmentModified = false;
@@ -409,7 +406,6 @@ PRAGMA_DISABLE_DEPRECATION_WARNINGS
 			MARK_PROPERTY_DIRTY_FROM_NAME(AActor, AttachmentReplication, this);
 		}
 	}
-PRAGMA_ENABLE_DEPRECATION_WARNINGS
 }
 
 void AActor::GetLifetimeReplicatedProps( TArray< FLifetimeProperty > & OutLifetimeProps ) const
diff --git a/Engine/Source/Runtime/Engine/Private/Components/ActorComponent.cpp b/Engine/Source/Runtime/Engine/Private/Components/ActorComponent.cpp
--- a/Engine/Source/Runtime/Engine/Private/Components/ActorComponent.cpp
+++ b/Engine/Source/Runtime/Engine/Private/Components/ActorComponent.cpp
@@ -1812,10 +1812,8 @@ void UActorComponent::SetIsReplicated(bool bShouldReplicate)
 
 		if (GetComponentClassCanReplicate())
 		{
-PRAGMA_DISABLE_DEPRECATION_WARNINGS
 			bReplicates = bShouldReplicate;
 			MARK_PROPERTY_DIRTY_FROM_NAME(UActorComponent, bReplicates, this);
-PRAGMA_ENABLE_DEPRECATION_WARNINGS
 
 			if (AActor* MyOwner = GetOwner())
 			{
@@ -2033,10 +2031,8 @@ void UActorComponent::SetIsReplicatedByDefault(const bool bNewReplicates)
 	// Don't bother checking parent here.
 	if (LIKELY(NeedsInitialization()))
 	{
-PRAGMA_DISABLE_DEPRECATION_WARNINGS
 		bReplicates = bNewReplicates;
 		MARK_PROPERTY_DIRTY_FROM_NAME(UActorComponent, bReplicates, this);
-PRAGMA_ENABLE_DEPRECATION_WARNINGS
 	}
 	else
 	{
@@ -2047,10 +2043,8 @@ PRAGMA_ENABLE_DEPRECATION_WARNINGS
 
 void UActorComponent::SetActiveFlag(const bool bNewIsActive)
 {
-PRAGMA_DISABLE_DEPRECATION_WARNINGS
 	bIsActive = bNewIsActive;
 	MARK_PROPERTY_DIRTY_FROM_NAME(UActorComponent, bIsActive, this);
-PRAGMA_ENABLE_DEPRECATION_WARNINGS
 }
 
 bool UActorComponent::OwnerNeedsInitialization() const
diff --git a/Engine/Source/Runtime/Engine/Private/Components/SceneComponent.cpp b/Engine/Source/Runtime/Engine/Private/Components/SceneComponent.cpp
--- a/Engine/Source/Runtime/Engine/Private/Components/SceneComponent.cpp
+++ b/Engine/Source/Runtime/Engine/Private/Components/SceneComponent.cpp
@@ -3367,22 +3367,18 @@ FScopedPreventAttachedComponentMove::FScopedPreventAttachedComponentMove(USceneC
 		bSavedNonAbsoluteComponent = !(bSavedAbsoluteLocation && bSavedAbsoluteRotation && bSavedAbsoluteScale);
 
 		// These are only going to be changed temporarily and reset, so we'll allow access here.
-PRAGMA_DISABLE_DEPRECATION_WARNINGS
 		// Use absolute (stay in world space no matter what parent does)
 		Owner->bAbsoluteLocation = true;
 		Owner->bAbsoluteRotation = true;
 		Owner->bAbsoluteScale = true;
-PRAGMA_ENABLE_DEPRECATION_WARNINGS
 
 		if (bSavedNonAbsoluteComponent && Owner->GetAttachParent())
 		{
-PRAGMA_DISABLE_DEPRECATION_WARNINGS
 			// Make RelativeLocation etc relative to the world.
 			Component->ConditionalUpdateComponentToWorld();
 			Owner->RelativeLocation = Owner->GetComponentLocation();
 			Owner->RelativeRotation = Owner->GetComponentRotation();
 			Owner->RelativeScale3D = Owner->GetComponentScale();
-PRAGMA_ENABLE_DEPRECATION_WARNINGS
 		}
 	}
 	else
@@ -3397,11 +3393,9 @@ FScopedPreventAttachedComponentMove::~FScopedPreventAttachedComponentMove()
 {
 	if (Owner)
 	{
-PRAGMA_DISABLE_DEPRECATION_WARNINGS
 		Owner->bAbsoluteLocation = bSavedAbsoluteLocation;
 		Owner->bAbsoluteRotation = bSavedAbsoluteRotation;
 		Owner->bAbsoluteScale = bSavedAbsoluteScale;
-PRAGMA_ENABLE_DEPRECATION_WARNINGS
 
 		if (bSavedNonAbsoluteComponent && Owner->GetAttachParent())
 		{
@@ -3509,13 +3503,11 @@ void FScopedMovementUpdate::RevertMove()
 		
 		if (IsTransformDirty())
 		{
-PRAGMA_DISABLE_DEPRECATION_WARNINGS
 			// Teleport to start
 			Component->ComponentToWorld = InitialTransform;
 			Component->RelativeLocation = InitialRelativeLocation;
 			Component->RelativeRotation = InitialRelativeRotation;
 			Component->RelativeScale3D = InitialRelativeScale;
-PRAGMA_ENABLE_DEPRECATION_WARNINGS
 
 			if (!IsDeferringUpdates())
 			{
@@ -3833,18 +3825,14 @@ void USceneComponent::K2_AddWorldTransform(const FTransform& DeltaTransform, boo
 
 void USceneComponent::SetVisibleFlag(const bool bNewVisible)
 {
-PRAGMA_DISABLE_DEPRECATION_WARNINGS
 	bVisible = bNewVisible;
 	MARK_PROPERTY_DIRTY_FROM_NAME(USceneComponent, bVisible, this);
-PRAGMA_ENABLE_DEPRECATION_WARNINGS
 }
 
 FVector& USceneComponent::GetRelativeLocation_DirectMutable()
 {
-PRAGMA_DISABLE_DEPRECATION_WARNINGS
 	MARK_PROPERTY_DIRTY_FROM_NAME(USceneComponent, RelativeLocation, this);
 	return RelativeLocation;
-PRAGMA_ENABLE_DEPRECATION_WARNINGS
 }
 
 void USceneComponent::SetRelativeLocation_Direct(const FVector NewRelativeLocation)
@@ -3854,10 +3842,8 @@ void USceneComponent::SetRelativeLocation_Direct(const FVector NewRelativeLocati
 
 FRotator& USceneComponent::GetRelativeRotation_DirectMutable()
 {
-PRAGMA_DISABLE_DEPRECATION_WARNINGS
 	MARK_PROPERTY_DIRTY_FROM_NAME(USceneComponent, RelativeRotation, this);
 	return RelativeRotation;
-PRAGMA_ENABLE_DEPRECATION_WARNINGS
 }
 
 void USceneComponent::SetRelativeRotation_Direct(const FRotator NewRelativeRotation)
@@ -3867,10 +3853,8 @@ void USceneComponent::SetRelativeRotation_Direct(const FRotator NewRelativeRotat
 
 FVector& USceneComponent::GetRelativeScale3D_DirectMutable()
 {
-PRAGMA_DISABLE_DEPRECATION_WARNINGS
 	MARK_PROPERTY_DIRTY_FROM_NAME(USceneComponent, RelativeScale3D, this);
 	return RelativeScale3D;
-PRAGMA_ENABLE_DEPRECATION_WARNINGS
 }
 
 void USceneComponent::SetRelativeScale3D_Direct(const FVector NewRelativeScale3D)
@@ -3880,42 +3864,32 @@ void USceneComponent::SetRelativeScale3D_Direct(const FVector NewRelativeScale3D
 
 void USceneComponent::SetUsingAbsoluteLocation(bool bInAbsoluteLocation)
 {
-PRAGMA_DISABLE_DEPRECATION_WARNINGS
 	bAbsoluteLocation = bInAbsoluteLocation;
 	MARK_PROPERTY_DIRTY_FROM_NAME(USceneComponent, bAbsoluteLocation, this);
-PRAGMA_ENABLE_DEPRECATION_WARNINGS
 }
 
 void USceneComponent::SetUsingAbsoluteRotation(bool bInAbsoluteRotation)
 {
-PRAGMA_DISABLE_DEPRECATION_WARNINGS
 	bAbsoluteRotation = bInAbsoluteRotation;
 	MARK_PROPERTY_DIRTY_FROM_NAME(USceneComponent, bAbsoluteRotation, this);
-PRAGMA_ENABLE_DEPRECATION_WARNINGS
 }
 
 void USceneComponent::SetUsingAbsoluteScale(bool bInAbsoluteScale)
 {
-PRAGMA_DISABLE_DEPRECATION_WARNINGS
 	bAbsoluteScale = bInAbsoluteScale;
 	MARK_PROPERTY_DIRTY_FROM_NAME(USceneComponent, bAbsoluteScale, this);
-PRAGMA_ENABLE_DEPRECATION_WARNINGS
 }
 
 void USceneComponent::SetAttachParent(USceneComponent* NewAttachParent)
 {
-PRAGMA_DISABLE_DEPRECATION_WARNINGS
 	AttachParent = NewAttachParent;
 	MARK_PROPERTY_DIRTY_FROM_NAME(USceneComponent, AttachParent, this);
-PRAGMA_ENABLE_DEPRECATION_WARNINGS
 }
 
 void USceneComponent::SetAttachSocketName(FName NewSocketName)
 {
-PRAGMA_DISABLE_DEPRECATION_WARNINGS
 	AttachSocketName = NewSocketName;
 	MARK_PROPERTY_DIRTY_FROM_NAME(USceneComponent, AttachSocketName, this);
-PRAGMA_ENABLE_DEPRECATION_WARNINGS
 }
 
 void USceneComponent::ModifiedAttachChildren()
