MoviePipeline: Render (Local) now respects the level specified in the Movie Render Queue for a job, even if it doesn't match the current editor level.

diff --git a/Engine/Source/Editor/UnrealEd/Private/PlayLevel.cpp b/Engine/Source/Editor/UnrealEd/Private/PlayLevel.cpp
--- a/Engine/Source/Editor/UnrealEd/Private/PlayLevel.cpp
+++ b/Engine/Source/Editor/UnrealEd/Private/PlayLevel.cpp
@@ -1662,6 +1662,7 @@ void UEditorEngine::CreateNewPlayInEditorInstance(FRequestPlaySessionParams &InR
 		GameInstancePIEParameters.EditorPlaySettings = PlayInEditorSessionInfo->OriginalRequestParams.EditorPlaySettings;
 		GameInstancePIEParameters.WorldFeatureLevel = PreviewPlatform.GetEffectivePreviewFeatureLevel();
 		GameInstancePIEParameters.NetMode = InNetMode;
+		GameInstancePIEParameters.OverrideMapURL = InRequestParams.GlobalMapOverride;
 
 		PIELoginInfo.GameInstancePIEParameters = GameInstancePIEParameters;
 
diff --git a/Engine/Source/Runtime/Engine/Private/GameInstance.cpp b/Engine/Source/Runtime/Engine/Private/GameInstance.cpp
--- a/Engine/Source/Runtime/Engine/Private/GameInstance.cpp
+++ b/Engine/Source/Runtime/Engine/Private/GameInstance.cpp
@@ -273,8 +273,21 @@ FGameInstancePIEResult UGameInstance::InitializeForPlayInEditor(int32 PIEInstanc
 	}
 	else
 	{
-		// Standard PIE path: just duplicate the EditorWorld
-		NewWorld = EditorEngine->CreatePIEWorldByDuplication(*WorldContext, EditorEngine->EditorWorld, PIEMapName);
+		if (Params.OverrideMapURL.Len() > 0)
+		{
+			// Attempt to load the target world asset
+			FSoftObjectPath TargetWorld = FSoftObjectPath(Params.OverrideMapURL);
+			UWorld* WorldToDuplicate = Cast<UWorld>(TargetWorld.TryLoad());
+			if (WorldToDuplicate)
+			{
+				NewWorld = EditorEngine->CreatePIEWorldByDuplication(*WorldContext, WorldToDuplicate, PIEMapName);
+			}
+		}
+		else
+		{
+			// Standard PIE path: just duplicate the EditorWorld
+			NewWorld = EditorEngine->CreatePIEWorldByDuplication(*WorldContext, EditorEngine->EditorWorld, PIEMapName);
+		}
 
 		// Duplication can result in unreferenced objects, so indicate that we should do a GC pass after initializing the world context
 		bNeedsGarbageCollection = true;
