Fixed string/TCHAR return for Chaos version

diff --git a/Engine/Source/Runtime/Engine/Private/PhysicsEngine/Experimental/ChaosDerivedData.cpp b/Engine/Source/Runtime/Engine/Private/PhysicsEngine/Experimental/ChaosDerivedData.cpp
--- a/Engine/Source/Runtime/Engine/Private/PhysicsEngine/Experimental/ChaosDerivedData.cpp
+++ b/Engine/Source/Runtime/Engine/Private/PhysicsEngine/Experimental/ChaosDerivedData.cpp
@@ -29,7 +29,7 @@ const TCHAR* FChaosDerivedDataCooker::GetVersionString() const
 	// for their own DDC or serialized data - change Chaos::ChaosVersionString in Chaos/Core.h to bump our
 	// Chaos data version. Callers can also rely on that version in their builders and avoid bad serialization
 	// when basic Chaos data changes
-	return Chaos::ChaosVersionString;
+	return *Chaos::ChaosVersionString;
 }
 
 FString FChaosDerivedDataCooker::GetPluginSpecificCacheKeySuffix() const
