Propagate changes made to layers in base material to derived MIs

diff --git a/Engine/Source/Editor/MaterialEditor/Private/SMaterialLayersFunctionsTree.cpp b/Engine/Source/Editor/MaterialEditor/Private/SMaterialLayersFunctionsTree.cpp
--- a/Engine/Source/Editor/MaterialEditor/Private/SMaterialLayersFunctionsTree.cpp
+++ b/Engine/Source/Editor/MaterialEditor/Private/SMaterialLayersFunctionsTree.cpp
@@ -139,53 +139,6 @@ void SMaterialLayersFunctionsInstanceTreeItem::OnNameChanged(const FText& InText
 	InTree->FunctionInstanceHandle->NotifyPostChange();
 }
 
-
-int32 GetNewParamIndex(EMaterialParameterAssociation InAssociation, int32 InIndex, int32 OriginalIndex, int32 NewIndex, int32 OriginalBlendIndex, int32 NewBlendIndex)
-{
-	int32 ParamIndex = InIndex;
-	int32 OriginalIndexToUse;
-	int32 NewIndexToUse;
-	if (InAssociation == EMaterialParameterAssociation::LayerParameter)
-	{
-		OriginalIndexToUse = OriginalIndex;
-		NewIndexToUse = NewIndex;
-	}
-	else if (InAssociation == EMaterialParameterAssociation::BlendParameter)
-	{
-		OriginalIndexToUse = OriginalBlendIndex;
-		NewIndexToUse = NewBlendIndex;
-	}
-	else
-	{
-		return ParamIndex;
-	}
-
-	if (NewIndexToUse < OriginalIndexToUse)
-	{
-		if (InIndex == OriginalIndexToUse)
-		{
-			ParamIndex = NewIndexToUse;
-		}
-		else if (InIndex >= NewIndexToUse && InIndex < OriginalIndexToUse)
-		{
-			ParamIndex++;
-		}
-	}
-	else if (NewIndexToUse > OriginalIndexToUse)
-	{
-		if (InIndex == OriginalIndexToUse)
-		{
-			ParamIndex = NewIndexToUse;
-		}
-		else if (InIndex <= NewIndexToUse && InIndex > OriginalIndexToUse)
-		{
-			ParamIndex--;
-		}
-	}
-
-	return ParamIndex;
-}
-
 FReply SMaterialLayersFunctionsInstanceTreeItem::OnLayerDrop(const FDragDropEvent& DragDropEvent)
 {
 	if (!bIsHoveredDragTarget)
@@ -213,77 +166,30 @@ FReply SMaterialLayersFunctionsInstanceTreeItem::OnLayerDrop(const FDragDropEven
 		if (SwappingPropertyData != SwappablePropertyData)
 		{
 			int32 OriginalIndex = SwappingPropertyData->ParameterInfo.Index;
-			int32 NewIndex = SwappablePropertyData->ParameterInfo.Index;
-			int32 OriginalBlendIndex = SwappingPropertyData->ParameterInfo.Index - 1;
-			int32 NewBlendIndex = SwappablePropertyData->ParameterInfo.Index;
-			if (SwappablePropertyData->ParameterInfo.Association == EMaterialParameterAssociation::BlendParameter)
+			if (SwappingPropertyData->ParameterInfo.Association == EMaterialParameterAssociation::BlendParameter)
 			{
-				NewIndex++;
+				OriginalIndex++;
 			}
 
-			if (SwappablePropertyData->ParameterInfo.Association != EMaterialParameterAssociation::BlendParameter)
-			{
-				NewBlendIndex--;
-			}
-
-
-			for (int32 ParamIt = 0; ParamIt < Tree->MaterialEditorInstance->SourceInstance->ScalarParameterValues.Num(); ParamIt++)
-			{
-				FScalarParameterValue* Param = &Tree->MaterialEditorInstance->SourceInstance->ScalarParameterValues[ParamIt];
-				Param->ParameterInfo.Index = GetNewParamIndex(Param->ParameterInfo.Association, Param->ParameterInfo.Index, OriginalIndex, NewIndex, OriginalBlendIndex, NewBlendIndex);
-			}
-
-			for (int32 ParamIt = 0; ParamIt < Tree->MaterialEditorInstance->SourceInstance->VectorParameterValues.Num(); ParamIt++)
+			int32 NewIndex = SwappablePropertyData->ParameterInfo.Index;
+			if (SwappablePropertyData->ParameterInfo.Association == EMaterialParameterAssociation::BlendParameter)
 			{
-				FVectorParameterValue* Param = &Tree->MaterialEditorInstance->SourceInstance->VectorParameterValues[ParamIt];
-				Param->ParameterInfo.Index = GetNewParamIndex(Param->ParameterInfo.Association, Param->ParameterInfo.Index, OriginalIndex, NewIndex, OriginalBlendIndex, NewBlendIndex);
+				NewIndex++;
 			}
 
-			for (int32 ParamIt = 0; ParamIt < Tree->MaterialEditorInstance->SourceInstance->TextureParameterValues.Num(); ParamIt++)
+			if (OriginalIndex != NewIndex)
 			{
-				FTextureParameterValue* Param = &Tree->MaterialEditorInstance->SourceInstance->TextureParameterValues[ParamIt];
-				Param->ParameterInfo.Index = GetNewParamIndex(Param->ParameterInfo.Association, Param->ParameterInfo.Index, OriginalIndex, NewIndex, OriginalBlendIndex, NewBlendIndex);
-			}
+				Tree->MaterialEditorInstance->SourceInstance->SwapLayerParameterIndices(OriginalIndex, NewIndex);
 
-			for (int32 ParamIt = 0; ParamIt < Tree->MaterialEditorInstance->SourceInstance->RuntimeVirtualTextureParameterValues.Num(); ParamIt++)
-			{
-				FRuntimeVirtualTextureParameterValue* Param = &Tree->MaterialEditorInstance->SourceInstance->RuntimeVirtualTextureParameterValues[ParamIt];
-				Param->ParameterInfo.Index = GetNewParamIndex(Param->ParameterInfo.Association, Param->ParameterInfo.Index, OriginalIndex, NewIndex, OriginalBlendIndex, NewBlendIndex);
-			}
+				// Need to save the moving and target expansion states before swapping
+				const bool bOriginalSwappableExpansion = IsItemExpanded();
+				const bool bOriginalSwappingExpansion = LayerPtr->IsItemExpanded();
 
-			for (int32 ParamIt = 0; ParamIt < Tree->MaterialEditorInstance->SourceInstance->FontParameterValues.Num(); ParamIt++)
-			{
-				FFontParameterValue* Param = &Tree->MaterialEditorInstance->SourceInstance->FontParameterValues[ParamIt];
-				Param->ParameterInfo.Index = GetNewParamIndex(Param->ParameterInfo.Association, Param->ParameterInfo.Index, OriginalIndex, NewIndex, OriginalBlendIndex, NewBlendIndex);
-			}
+				TArray<void*> StructPtrs;
+				Tree->FunctionInstanceHandle->AccessRawData(StructPtrs);
+				FMaterialLayersFunctions* MaterialLayersFunctions = static_cast<FMaterialLayersFunctions*>(StructPtrs[0]);
+				MaterialLayersFunctions->MoveBlendedLayer(OriginalIndex, NewIndex);
 
-			if (NewIndex > OriginalIndex)
-			{
-				NewIndex += 1;
-			}
-			if (NewBlendIndex > OriginalBlendIndex)
-			{
-				NewBlendIndex += 1;
-			}
-			TSharedPtr<IPropertyHandle> ParentHandle = Tree->FunctionInstanceHandle;
-			TSharedPtr<IPropertyHandleArray> LayerHandle = ParentHandle->GetChildHandle("Layers")->AsArray();
-			TSharedPtr<IPropertyHandleArray> BlendHandle = ParentHandle->GetChildHandle("Blends")->AsArray();
-			TSharedPtr<IPropertyHandleArray> LayerNameHandle = ParentHandle->GetChildHandle("LayerNames")->AsArray();
-			TSharedPtr<IPropertyHandleArray> LayerFilterHandle = ParentHandle->GetChildHandle("RestrictToLayerRelatives")->AsArray();
-			TSharedPtr<IPropertyHandleArray> BlendFilterHandle = ParentHandle->GetChildHandle("RestrictToBlendRelatives")->AsArray();
-			TSharedPtr<IPropertyHandleArray> LayerStateHandle = ParentHandle->GetChildHandle("LayerStates")->AsArray();
-
-			if (LayerHandle.IsValid() && BlendHandle.IsValid())
-			{
-				// Need to save the moving and target expansion states before swapping
-				bool bOriginalSwappableExpansion = IsItemExpanded();
-				bool bOriginalSwappingExpansion = LayerPtr->IsItemExpanded();
-				LayerHandle->MoveElementTo(OriginalIndex, NewIndex);
-				LayerNameHandle->MoveElementTo(OriginalIndex, NewIndex);
-				LayerFilterHandle->MoveElementTo(OriginalIndex, NewIndex);
-				LayerStateHandle->MoveElementTo(OriginalIndex, NewIndex);
-				BlendHandle->MoveElementTo(OriginalBlendIndex, NewBlendIndex);
-				BlendFilterHandle->MoveElementTo(OriginalBlendIndex, NewBlendIndex);
 				Tree->OnExpansionChanged(SwappablePropertyData, bOriginalSwappingExpansion);
 				Tree->OnExpansionChanged(SwappingPropertyData, bOriginalSwappableExpansion);
 				Tree->FunctionInstanceHandle->NotifyPostChange();
diff --git a/Engine/Source/Runtime/Engine/Classes/Materials/MaterialExpressionMaterialAttributeLayers.h b/Engine/Source/Runtime/Engine/Classes/Materials/MaterialExpressionMaterialAttributeLayers.h
--- a/Engine/Source/Runtime/Engine/Classes/Materials/MaterialExpressionMaterialAttributeLayers.h
+++ b/Engine/Source/Runtime/Engine/Classes/Materials/MaterialExpressionMaterialAttributeLayers.h
@@ -44,7 +44,6 @@ class UMaterialExpressionMaterialAttributeLayers : public UMaterialExpression
 	}
 
 #if WITH_EDITOR
-
 	const TArray<FText>& GetLayerNames() const
 	{
 		return ParamLayers ? ParamLayers->LayerNames : DefaultLayers.LayerNames;
@@ -59,7 +58,17 @@ class UMaterialExpressionMaterialAttributeLayers : public UMaterialExpression
 	{
 		return ParamLayers ? ParamLayers->RestrictToBlendRelatives : DefaultLayers.RestrictToBlendRelatives;
 	}
-#endif
+
+	const TArray<FGuid>& GetLayerGuids() const
+	{
+		return ParamLayers ? ParamLayers->LayerGuids : DefaultLayers.LayerGuids;
+	}
+
+	const TArray<FGuid>& GetParentLayerGuids() const
+	{
+		return ParamLayers ? ParamLayers->ParentLayerGuids : DefaultLayers.ParentLayerGuids;
+	}
+#endif // WITH_EDITOR
 
 	const TArray<bool>& GetLayerStates() const
 	{
diff --git a/Engine/Source/Runtime/Engine/Classes/Materials/MaterialInstance.h b/Engine/Source/Runtime/Engine/Classes/Materials/MaterialInstance.h
--- a/Engine/Source/Runtime/Engine/Classes/Materials/MaterialInstance.h
+++ b/Engine/Source/Runtime/Engine/Classes/Materials/MaterialInstance.h
@@ -583,6 +583,8 @@ public:
 	/** Ensure's static permutations for current parameters and overrides are upto date. */
 	ENGINE_API void UpdateStaticPermutation();
 
+	ENGINE_API void SwapLayerParameterIndices(int32 OriginalIndex, int32 NewIndex);
+
 #endif // WITH_EDITOR
 
 	//~ End UObject Interface.
diff --git a/Engine/Source/Runtime/Engine/Classes/Materials/MaterialLayersFunctions.h b/Engine/Source/Runtime/Engine/Classes/Materials/MaterialLayersFunctions.h
--- a/Engine/Source/Runtime/Engine/Classes/Materials/MaterialLayersFunctions.h
+++ b/Engine/Source/Runtime/Engine/Classes/Materials/MaterialLayersFunctions.h
@@ -165,18 +165,36 @@ struct ENGINE_API FMaterialLayersFunctions
 		//TODO: Investigate whether this is really required given it is only used by FMaterialShaderMapId AND that one also uses UpdateHash
 		void AppendKeyString(FString& KeyString) const;
 	};
+
+	static const FGuid UninitializedParentGuid;
+	static const FGuid NoParentGuid;
 		
 	FMaterialLayersFunctions()
 	{
 		// Default to a non-blended "background" layer
 		Layers.AddDefaulted();
+		LayerStates.Add(true);
 #if WITH_EDITOR
 		FText LayerName = FText(LOCTEXT("Background", "Background"));
 		LayerNames.Add(LayerName);
-		RestrictToLayerRelatives.Push(false);
+		RestrictToLayerRelatives.Add(false);
+		LayerGuids.Add(FGuid::NewGuid());
+		ParentLayerGuids.Add(NoParentGuid);
 #endif
-		LayerStates.Push(true);
+	}
 
+	void Empty()
+	{
+		Layers.Empty();
+		Blends.Empty();
+		LayerStates.Empty();
+#if WITH_EDITOR
+		LayerNames.Empty();
+		RestrictToLayerRelatives.Empty();
+		RestrictToBlendRelatives.Empty();
+		LayerGuids.Empty();
+		ParentLayerGuids.Empty();
+#endif
 	}
 
 	UPROPERTY(EditAnywhere, Category=MaterialLayers)
@@ -185,6 +203,9 @@ struct ENGINE_API FMaterialLayersFunctions
 	UPROPERTY(EditAnywhere, Category=MaterialLayers)
 	TArray<class UMaterialFunctionInterface*> Blends;
 
+	UPROPERTY(EditAnywhere, Category = MaterialLayers)
+	TArray<bool> LayerStates;
+
 #if WITH_EDITORONLY_DATA
 	UPROPERTY(EditAnywhere, Category = MaterialLayers)
 	TArray<FText> LayerNames;
@@ -194,43 +215,39 @@ struct ENGINE_API FMaterialLayersFunctions
 
 	UPROPERTY(EditAnywhere, Category = MaterialLayers)
 	TArray<bool> RestrictToBlendRelatives;
-#endif
 
+	/** Guid that identifies each layer in this stack */
 	UPROPERTY(EditAnywhere, Category = MaterialLayers)
-	TArray<bool> LayerStates;
+	TArray<FGuid> LayerGuids;
+
+	/**
+	 * Refers to the layer in the parent's LayerGuids list used to initialize this layer
+	 * - Special value of 'NoParentGuid' means this layer was created in this material, not based on any layer in the parent
+	 * - Special value of 'UninitializedParentGuid' means this data was serialized before these guids existed...layers with this value will attempt to match a parent layer with the same resources assigned
+	 */
+	UPROPERTY(EditAnywhere, Category = MaterialLayers)
+	TArray<FGuid> ParentLayerGuids;
+
+	/**
+	 * List of Guids that exist in the parent material that have been explicitly deleted
+	 * This is needed to distinguish these layers from newly added layers in the parent material
+	 */
+	UPROPERTY(EditAnywhere, Category = MaterialLayers)
+	TArray<FGuid> DeletedParentLayerGuids;
+#endif // WITH_EDITORONLY_DATA
 
 	UPROPERTY()
 	FString KeyString_DEPRECATED;
 
-	void AppendBlendedLayer()
-	{
-		Layers.AddDefaulted();
-		Blends.AddDefaulted();
-#if WITH_EDITOR
-		FText LayerName = FText::Format(LOCTEXT("LayerPrefix", "Layer {0}"), Layers.Num()-1);
-		LayerNames.Add(LayerName);
-		RestrictToLayerRelatives.Push(false);
-		RestrictToBlendRelatives.Push(false);
-#endif
-		LayerStates.Push(true);
-	}
+	void AppendBlendedLayer();
 
-	void RemoveBlendedLayerAt(int32 Index)
-	{
-		if (Layers.IsValidIndex(Index))
-		{
-			check(Layers.IsValidIndex(Index) && Blends.IsValidIndex(Index-1) && LayerStates.IsValidIndex(Index));
-			Layers.RemoveAt(Index);
-			Blends.RemoveAt(Index-1);
-			LayerStates.RemoveAt(Index);
-#if WITH_EDITOR
-			check(LayerNames.IsValidIndex(Index) && RestrictToLayerRelatives.IsValidIndex(Index) && RestrictToBlendRelatives.IsValidIndex(Index-1));
-			LayerNames.RemoveAt(Index);
-			RestrictToLayerRelatives.RemoveAt(Index);
-			RestrictToBlendRelatives.RemoveAt(Index-1);
-#endif //WITH_EDITOR
-		}
-	}
+	void AddLayerCopy(const FMaterialLayersFunctions& Source, int32 SourceLayerIndex, const FGuid& ParentGuid);
+
+	void InsertLayerCopy(const FMaterialLayersFunctions& Source, int32 SourceLayerIndex, const FGuid& ParentGuid, int32 LayerIndex);
+
+	void RemoveBlendedLayerAt(int32 Index);
+
+	void MoveBlendedLayer(int32 SrcLayerIndex, int32 DstLayerIndex);
 
 	void ToggleBlendedLayerVisibility(int32 Index)
 	{
@@ -244,14 +261,14 @@ struct ENGINE_API FMaterialLayersFunctions
 		LayerStates[Index] = InNewVisibility;
 	}
 
-	bool GetLayerVisibility(int32 Index)
+	bool GetLayerVisibility(int32 Index) const
 	{
 		check(LayerStates.IsValidIndex(Index));
 		return LayerStates[Index];
 	}
 
 #if WITH_EDITORONLY_DATA
-	FText GetLayerName(int32 Counter)
+	FText GetLayerName(int32 Counter) const
 	{
 		FText LayerName = FText::Format(LOCTEXT("LayerPrefix", "Layer {0}"), Counter);
 		if (LayerNames.IsValidIndex(Counter))
@@ -261,7 +278,11 @@ struct ENGINE_API FMaterialLayersFunctions
 		return LayerName;
 	}
 
-#endif
+	void CopyGuidsToParent();
+
+	bool ResolveParent(const FMaterialLayersFunctions& Parent, TArray<int32>& OutRemapLayerIndices);
+
+#endif // WITH_EDITORONLY_DATA
 
 	const ID GetID() const;
 
@@ -270,6 +291,8 @@ struct ENGINE_API FMaterialLayersFunctions
 
 	void SerializeForDDC(FArchive& Ar);
 
+	void PostSerialize(const FArchive& Ar);
+
 	FORCEINLINE bool operator==(const FMaterialLayersFunctions& Other) const
 	{
 		return Layers == Other.Layers && Blends == Other.Blends && LayerStates == Other.LayerStates;
@@ -281,4 +304,10 @@ struct ENGINE_API FMaterialLayersFunctions
 	}
 };
 
+template<>
+struct TStructOpsTypeTraits<FMaterialLayersFunctions> : TStructOpsTypeTraitsBase2<FMaterialLayersFunctions>
+{
+	enum { WithPostSerialize = true };
+};
+
 #undef LOCTEXT_NAMESPACE
\ No newline at end of file
diff --git a/Engine/Source/Runtime/Engine/Private/Materials/MaterialExpressions.cpp b/Engine/Source/Runtime/Engine/Private/Materials/MaterialExpressions.cpp
--- a/Engine/Source/Runtime/Engine/Private/Materials/MaterialExpressions.cpp
+++ b/Engine/Source/Runtime/Engine/Private/Materials/MaterialExpressions.cpp
@@ -6733,14 +6733,7 @@ bool UMaterialExpressionMaterialAttributeLayers::IsNamedParameter(const FHashedM
 {
 	if (ParameterInfo.Name == ParameterName)
 	{
-		OutLayers.Layers = GetLayers();
-		OutLayers.Blends = GetBlends();
-		OutLayers.LayerStates = GetLayerStates();
-#if WITH_EDITOR
-		OutLayers.RestrictToLayerRelatives = GetShouldFilterLayers();
-		OutLayers.RestrictToBlendRelatives = GetShouldFilterBlends();
-		OutLayers.LayerNames = GetLayerNames();
-#endif
+		OutLayers = (ParamLayers != nullptr) ? *ParamLayers : DefaultLayers;
 		OutExpressionGuid = ExpressionGUID;
 		return true;
 	}
@@ -12465,6 +12458,9 @@ void FMaterialLayersFunctions::ID::AppendKeyString(FString& KeyString) const
 // FMaterialLayersFunctions
 ///////////////////////////////////////////////////////////////////////////////
 
+const FGuid FMaterialLayersFunctions::UninitializedParentGuid(0u, 0u, 0u, 0u);
+const FGuid FMaterialLayersFunctions::NoParentGuid(1u, 0u, 0u, 0u);
+
 const FMaterialLayersFunctions::ID FMaterialLayersFunctions::GetID() const
 {
 	FMaterialLayersFunctions::ID Result;
@@ -12541,6 +12537,255 @@ void FMaterialLayersFunctions::SerializeForDDC(FArchive& Ar)
 	Ar << KeyString_DEPRECATED;
 }
 
+void FMaterialLayersFunctions::PostSerialize(const FArchive& Ar)
+{
+#if WITH_EDITORONLY_DATA
+	if (Ar.IsLoading())
+	{
+		check(LayerGuids.Num() == ParentLayerGuids.Num());
+		if (LayerGuids.Num() != Layers.Num())
+		{
+			const int32 NumLayers = Layers.Num();
+			LayerGuids.Empty(NumLayers);
+			ParentLayerGuids.Empty(NumLayers);
+			for (int32 i = 0; i < NumLayers; ++i)
+			{
+				LayerGuids.Add(FGuid::NewGuid());
+				ParentLayerGuids.Add(UninitializedParentGuid);
+			}
+		}
+	}
+#endif // WITH_EDITORONLY_DATA
+}
+
+void FMaterialLayersFunctions::AppendBlendedLayer()
+{
+	Layers.AddDefaulted();
+	Blends.AddDefaulted();
+	LayerStates.Add(true);
+#if WITH_EDITOR
+	FText LayerName = FText::Format(LOCTEXT("LayerPrefix", "Layer {0}"), Layers.Num() - 1);
+	LayerNames.Add(LayerName);
+	RestrictToLayerRelatives.Add(false);
+	RestrictToBlendRelatives.Add(false);
+	LayerGuids.Add(FGuid::NewGuid());
+	ParentLayerGuids.Add(NoParentGuid);
+#endif
+}
+
+void FMaterialLayersFunctions::AddLayerCopy(const FMaterialLayersFunctions& Source, int32 SourceLayerIndex, const FGuid& ParentGuid)
+{
+	check(ParentGuid != UninitializedParentGuid);
+	const int32 LayerIndex = Layers.Num();
+
+	Layers.Add(Source.Layers[SourceLayerIndex]);
+	if (LayerIndex > 0)
+	{
+		Blends.Add(Source.Blends[SourceLayerIndex - 1]);
+	}
+	LayerStates.Add(Source.LayerStates[SourceLayerIndex]);
+#if WITH_EDITOR
+	LayerNames.Add(Source.LayerNames[SourceLayerIndex]);
+	RestrictToLayerRelatives.Add(Source.RestrictToLayerRelatives[SourceLayerIndex]);
+	if (LayerIndex > 0)
+	{
+		RestrictToBlendRelatives.Add(Source.RestrictToBlendRelatives[SourceLayerIndex - 1]);
+	}
+	LayerGuids.Add(Source.LayerGuids[SourceLayerIndex]);
+	ParentLayerGuids.Add(ParentGuid);
+#endif
+}
+
+void FMaterialLayersFunctions::InsertLayerCopy(const FMaterialLayersFunctions& Source, int32 SourceLayerIndex, const FGuid& ParentGuid, int32 LayerIndex)
+{
+	check(ParentGuid != UninitializedParentGuid);
+	check(LayerIndex > 0);
+	Layers.Insert(Source.Layers[SourceLayerIndex], LayerIndex);
+	Blends.Insert(Source.Blends[SourceLayerIndex - 1], LayerIndex - 1);
+	LayerStates.Insert(Source.LayerStates[SourceLayerIndex], LayerIndex);
+#if WITH_EDITOR
+	LayerNames.Insert(Source.LayerNames[SourceLayerIndex], LayerIndex);
+	RestrictToLayerRelatives.Insert(Source.RestrictToLayerRelatives[SourceLayerIndex], LayerIndex);
+	RestrictToBlendRelatives.Insert(Source.RestrictToBlendRelatives[SourceLayerIndex - 1], LayerIndex - 1);
+	LayerGuids.Insert(Source.LayerGuids[SourceLayerIndex], LayerIndex);
+	ParentLayerGuids.Insert(ParentGuid, LayerIndex);
+#endif
+}
+
+void FMaterialLayersFunctions::RemoveBlendedLayerAt(int32 Index)
+{
+	if (Layers.IsValidIndex(Index))
+	{
+		check(Layers.IsValidIndex(Index) && Blends.IsValidIndex(Index - 1) && LayerStates.IsValidIndex(Index));
+		Layers.RemoveAt(Index);
+		Blends.RemoveAt(Index - 1);
+		LayerStates.RemoveAt(Index);
+#if WITH_EDITOR
+		check(LayerNames.IsValidIndex(Index) && RestrictToLayerRelatives.IsValidIndex(Index) && RestrictToBlendRelatives.IsValidIndex(Index - 1));
+
+		const FGuid& ParentGuid = ParentLayerGuids[Index];
+		if (ParentGuid != NoParentGuid && ParentGuid != UninitializedParentGuid)
+		{
+			// Save the parent guid as explicitly deleted, so it's not added back
+			check(!DeletedParentLayerGuids.Contains(ParentGuid));
+			DeletedParentLayerGuids.Add(ParentGuid);
+		}
+
+		LayerNames.RemoveAt(Index);
+		RestrictToLayerRelatives.RemoveAt(Index);
+		RestrictToBlendRelatives.RemoveAt(Index - 1);
+		LayerGuids.RemoveAt(Index);
+		ParentLayerGuids.RemoveAt(Index);
+#endif //WITH_EDITOR
+	}
+}
+
+void FMaterialLayersFunctions::MoveBlendedLayer(int32 SrcLayerIndex, int32 DstLayerIndex)
+{
+	check(SrcLayerIndex > 0);
+	check(DstLayerIndex > 0);
+	if (SrcLayerIndex != DstLayerIndex)
+	{
+		Layers.Swap(SrcLayerIndex, DstLayerIndex);
+		Blends.Swap(SrcLayerIndex - 1, DstLayerIndex - 1);
+		LayerStates.Swap(SrcLayerIndex, DstLayerIndex);
+#if WITH_EDITOR
+		LayerNames.Swap(SrcLayerIndex, DstLayerIndex);
+		RestrictToLayerRelatives.Swap(SrcLayerIndex, DstLayerIndex);
+		RestrictToBlendRelatives.Swap(SrcLayerIndex - 1, DstLayerIndex - 1);
+		LayerGuids.Swap(SrcLayerIndex, DstLayerIndex);
+		// Disconnect layers from parent when they're moved at the instance level
+		if (ParentLayerGuids[SrcLayerIndex] != NoParentGuid)
+		{
+			check(ParentLayerGuids[SrcLayerIndex] != UninitializedParentGuid);
+			DeletedParentLayerGuids.Add(ParentLayerGuids[SrcLayerIndex]);
+			ParentLayerGuids[SrcLayerIndex] = NoParentGuid;
+		}
+		if (ParentLayerGuids[DstLayerIndex] != NoParentGuid)
+		{
+			check(ParentLayerGuids[DstLayerIndex] != UninitializedParentGuid);
+			DeletedParentLayerGuids.Add(ParentLayerGuids[DstLayerIndex]);
+			ParentLayerGuids[DstLayerIndex] = NoParentGuid;
+		}
+#endif //WITH_EDITOR
+	}
+}
+
+#if WITH_EDITORONLY_DATA
+void FMaterialLayersFunctions::CopyGuidsToParent()
+{
+	ParentLayerGuids = LayerGuids;
+}
+
+bool FMaterialLayersFunctions::ResolveParent(const FMaterialLayersFunctions& Parent, TArray<int32>& OutRemapLayerIndices)
+{
+	check(LayerGuids.Num() == Layers.Num());
+	check(ParentLayerGuids.Num() == Layers.Num());
+
+	FMaterialLayersFunctions ResolvedLayers;
+	TArray<int32> ResolvedLayerIndices;
+
+	// Start with base layer
+	ResolvedLayers.Empty();
+	ResolvedLayers.AddLayerCopy(*this, 0, NoParentGuid);
+
+	for (int32 ParentLayerIndex = 1; ParentLayerIndex < Parent.Layers.Num(); ++ParentLayerIndex)
+	{
+		const FGuid& ParentGuid = Parent.LayerGuids[ParentLayerIndex];
+		if (DeletedParentLayerGuids.Contains(ParentGuid))
+		{
+			// Layer was deleted
+			ResolvedLayers.DeletedParentLayerGuids.Add(ParentGuid);
+			continue;
+		}
+
+		int32 LayerIndex = ParentLayerGuids.Find(ParentGuid);
+		if (LayerIndex == INDEX_NONE)
+		{
+			// Check to see if we have any layers with parents that haven't been initialized yet, that match this parent layer
+			for (int32 CheckLayerIndex = 1; CheckLayerIndex < Layers.Num(); ++CheckLayerIndex)
+			{
+				if (Layers[CheckLayerIndex] == Parent.Layers[ParentLayerIndex] &&
+					Blends[CheckLayerIndex - 1] && Parent.Blends[CheckLayerIndex - 1] &&
+					ParentLayerGuids[CheckLayerIndex] == UninitializedParentGuid)
+				{
+					ParentLayerGuids[CheckLayerIndex] = ParentGuid;
+					LayerIndex = CheckLayerIndex;
+					break;
+				}
+			}
+		}
+
+		if (LayerIndex != INDEX_NONE)
+		{
+			// We have a copy of the layer from the parent, use that
+			check(LayerIndex != 0); // should not find layer 0
+			check(!ResolvedLayerIndices.Contains(LayerIndex));
+			ResolvedLayers.AddLayerCopy(*this, LayerIndex, ParentGuid);
+			ResolvedLayers.LayerNames[ParentLayerIndex] = Parent.LayerNames[ParentLayerIndex]; // Copy (potentially) updated name from parent
+			ResolvedLayerIndices.Add(LayerIndex);
+		}
+		else
+		{
+			// New layer added from parent
+			ResolvedLayers.AddLayerCopy(Parent, ParentLayerIndex, ParentGuid);
+		}
+	}
+
+	// Insert any layers that aren't linked to parent
+	for (int32 LayerIndex = 1; LayerIndex < Layers.Num(); ++LayerIndex)
+	{
+		const FGuid& ParentGuid = ParentLayerGuids[LayerIndex];
+		if (ParentGuid != NoParentGuid)
+		{
+			continue;
+		}
+
+		if (ResolvedLayerIndices.Contains(LayerIndex))
+		{
+			// already added this layer
+			continue;
+		}
+
+		// Layer doesn't exist in parent, merge it in
+		int32 PrevLayerIndex = FMath::Min(LayerIndex - 1, ResolvedLayers.Layers.Num() - 1);
+		while (PrevLayerIndex > 0)
+		{
+			if (ResolvedLayers.LayerGuids[PrevLayerIndex] == LayerGuids[PrevLayerIndex])
+			{
+				break;
+			}
+			--PrevLayerIndex;
+		}
+
+		ResolvedLayers.InsertLayerCopy(*this, LayerIndex, NoParentGuid, PrevLayerIndex + 1);
+	}
+
+	OutRemapLayerIndices.SetNumUninitialized(Layers.Num());
+	OutRemapLayerIndices[0] = 0;
+
+	bool bUpdatedLayers = Layers.Num() != ResolvedLayers.Layers.Num() ||
+		DeletedParentLayerGuids.Num() != ResolvedLayers.DeletedParentLayerGuids.Num();
+
+	for (int32 PrevLayerIndex = 1; PrevLayerIndex < Layers.Num(); ++PrevLayerIndex)
+	{
+		const FGuid& LayerGuid = LayerGuids[PrevLayerIndex];
+		const int32 LayerIndex = ResolvedLayers.LayerGuids.Find(LayerGuid);
+		OutRemapLayerIndices[PrevLayerIndex] = LayerIndex;
+		if (PrevLayerIndex != LayerIndex)
+		{
+			bUpdatedLayers = true;
+		}
+	}
+
+	if (bUpdatedLayers)
+	{
+		*this = MoveTemp(ResolvedLayers);
+	}
+	return bUpdatedLayers;
+}
+#endif // WITH_EDITORONLY_DATA
+
 ///////////////////////////////////////////////////////////////////////////////
 // UMaterialExpressionMaterialFunctionCall
 ///////////////////////////////////////////////////////////////////////////////
diff --git a/Engine/Source/Runtime/Engine/Private/Materials/MaterialInstance.cpp b/Engine/Source/Runtime/Engine/Private/Materials/MaterialInstance.cpp
--- a/Engine/Source/Runtime/Engine/Private/Materials/MaterialInstance.cpp
+++ b/Engine/Source/Runtime/Engine/Private/Materials/MaterialInstance.cpp
@@ -365,6 +365,84 @@ void GameThread_UpdateMIParameter(const UMaterialInstance* Instance, const Param
 		});
 }
 
+template<typename ParameterType>
+static void RemapLayerParameterIndicesArray(TArray<ParameterType>& Parameters, const TArray<int32>& RemapLayerIndices)
+{
+	int32 ParameterIndex = 0;
+	while (ParameterIndex < Parameters.Num())
+	{
+		ParameterType& Parameter = Parameters[ParameterIndex];
+		bool bRemovedParameter = false;
+		if (Parameter.ParameterInfo.Association == LayerParameter)
+		{
+			const int32 NewIndex = RemapLayerIndices[Parameter.ParameterInfo.Index];
+			if (NewIndex != INDEX_NONE)
+			{
+				Parameter.ParameterInfo.Index = NewIndex;
+			}
+			else
+			{
+				bRemovedParameter = true;
+			}
+		}
+		else if (Parameter.ParameterInfo.Association == BlendParameter)
+		{
+			const int32 NewIndex = RemapLayerIndices[Parameter.ParameterInfo.Index + 1];
+			if (NewIndex != INDEX_NONE)
+			{
+				Parameter.ParameterInfo.Index = NewIndex - 1;
+			}
+			else
+			{
+				bRemovedParameter = true;
+			}
+		}
+		if (bRemovedParameter)
+		{
+			Parameters.RemoveAt(ParameterIndex);
+		}
+		else
+		{
+			++ParameterIndex;
+		}
+	}
+}
+
+template<typename ParameterType>
+static void SwapLayerParameterIndicesArray(TArray<ParameterType>& Parameters, int32 OriginalIndex, int32 NewIndex)
+{
+	check(OriginalIndex > 0);
+	check(NewIndex > 0);
+
+	for (ParameterType& Parameter : Parameters)
+	{
+		if (Parameter.ParameterInfo.Association == LayerParameter)
+		{
+			if(Parameter.ParameterInfo.Index == OriginalIndex) Parameter.ParameterInfo.Index = NewIndex;
+			else if (Parameter.ParameterInfo.Index == NewIndex) Parameter.ParameterInfo.Index = OriginalIndex;
+		}
+		else if (Parameter.ParameterInfo.Association == BlendParameter)
+		{
+			if (Parameter.ParameterInfo.Index == OriginalIndex - 1) Parameter.ParameterInfo.Index = NewIndex - 1;
+			else if (Parameter.ParameterInfo.Index == NewIndex - 1) Parameter.ParameterInfo.Index = OriginalIndex - 1;
+		}
+	}
+}
+
+void UMaterialInstance::SwapLayerParameterIndices(int32 OriginalIndex, int32 NewIndex)
+{
+	if (OriginalIndex != NewIndex)
+	{
+		SwapLayerParameterIndicesArray(ScalarParameterValues, OriginalIndex, NewIndex);
+		SwapLayerParameterIndicesArray(VectorParameterValues, OriginalIndex, NewIndex);
+		SwapLayerParameterIndicesArray(TextureParameterValues, OriginalIndex, NewIndex);
+		SwapLayerParameterIndicesArray(RuntimeVirtualTextureParameterValues, OriginalIndex, NewIndex);
+		SwapLayerParameterIndicesArray(FontParameterValues, OriginalIndex, NewIndex);
+		SwapLayerParameterIndicesArray(StaticParameters.StaticSwitchParameters, OriginalIndex, NewIndex);
+		SwapLayerParameterIndicesArray(StaticParameters.MaterialLayersParameters, OriginalIndex, NewIndex);
+	}
+}
+
 bool UMaterialInstance::UpdateParameters()
 {
 	bool bDirty = false;
@@ -420,6 +498,30 @@ bool UMaterialInstance::UpdateParameters()
 				bDirty |= CustomParameterSetUpdater.Execute(StaticParameters, ParentMaterial);
 			}
 		}
+
+		if (Parent)
+		{
+			for (FStaticMaterialLayersParameter& LayersParam : StaticParameters.MaterialLayersParameters)
+			{
+				FMaterialLayersFunctions ParentLayers;
+				FGuid ParentGuid;
+				if (Parent->GetMaterialLayersParameterValue(LayersParam.ParameterInfo, ParentLayers, ParentGuid))
+				{
+					TArray<int32> RemapLayerIndices;
+					if (LayersParam.Value.ResolveParent(ParentLayers, RemapLayerIndices))
+					{
+						RemapLayerParameterIndicesArray(ScalarParameterValues, RemapLayerIndices);
+						RemapLayerParameterIndicesArray(VectorParameterValues, RemapLayerIndices);
+						RemapLayerParameterIndicesArray(TextureParameterValues, RemapLayerIndices);
+						RemapLayerParameterIndicesArray(RuntimeVirtualTextureParameterValues, RemapLayerIndices);
+						RemapLayerParameterIndicesArray(FontParameterValues, RemapLayerIndices);
+						RemapLayerParameterIndicesArray(StaticParameters.StaticSwitchParameters, RemapLayerIndices);
+						RemapLayerParameterIndicesArray(StaticParameters.MaterialLayersParameters, RemapLayerIndices);
+						bDirty = true;
+					}
+				}
+			}
+		}
 	}
 #endif // WITH_EDITOR
 
@@ -2036,6 +2138,8 @@ void UMaterialInstance::GetStaticParameterValues(FStaticParameterSet& OutStaticP
 			ParentParameter.ParameterInfo = ParameterInfo;
 
 			Parent->GetMaterialLayersParameterValue(ParameterInfo, ParentParameter.Value, ExpressionId);
+			ParentParameter.Value.CopyGuidsToParent(); // Set parent guids for layers from parent material
+
 			ParentParameter.ExpressionGUID = ExpressionId;
 			// If the SourceInstance is overriding this parameter, use its settings
 			for (const FStaticMaterialLayersParameter& LayersParam : StaticParameters.MaterialLayersParameters)
