Fix layer parameters being inherited from one MIC to another MIC
- Also remove various editor-only debug code to verify new cached parameter code matches old code, as the old code was also wrong in this case

diff --git a/Engine/Source/Runtime/Engine/Classes/Materials/MaterialInstance.h b/Engine/Source/Runtime/Engine/Classes/Materials/MaterialInstance.h
--- a/Engine/Source/Runtime/Engine/Classes/Materials/MaterialInstance.h
+++ b/Engine/Source/Runtime/Engine/Classes/Materials/MaterialInstance.h
@@ -706,14 +706,6 @@ public:
 #if WITH_EDITOR
 	void BeginAllowCachingStaticParameterValues();
 	void EndAllowCachingStaticParameterValues();
-
-private:
-	bool GetScalarLayerParameterValue_Legacy(const FHashedMaterialParameterInfo& ParameterInfo, bool bOveriddenOnly, float& OutValue, bool& OutResult) const;
-	bool GetVectorLayerParameterValue_Legacy(const FHashedMaterialParameterInfo& ParameterInfo, bool bOveriddenOnly, FLinearColor& OutValue, bool& OutResult) const;
-	bool GetTextureLayerParameterValue_Legacy(const FHashedMaterialParameterInfo& ParameterInfo, bool bOveriddenOnly, UTexture*& OutValue, bool& OutResult) const;
-	bool GetFontLayerParameterValue_Legacy(const FHashedMaterialParameterInfo& ParameterInfo, bool bOveriddenOnly, UFont*& OutValue, int32& OutFontPage, bool& OutResult) const;
-	bool GetRuntimeVirtualTextureLayerParameterValue_Legacy(const FHashedMaterialParameterInfo& ParameterInfo, bool bOveriddenOnly, URuntimeVirtualTexture*& OutValue, bool& OutResult) const;
-
 #endif // WITH_EDITOR
 
 protected:
diff --git a/Engine/Source/Runtime/Engine/Private/Materials/Material.cpp b/Engine/Source/Runtime/Engine/Private/Materials/Material.cpp
--- a/Engine/Source/Runtime/Engine/Private/Materials/Material.cpp
+++ b/Engine/Source/Runtime/Engine/Private/Materials/Material.cpp
@@ -2041,7 +2041,8 @@ void UMaterial::UpdateCachedExpressionData()
 	{
 		FMaterialCachedExpressionData UpdatedCachedExpressionData;
 		UpdatedCachedExpressionData.Reset();
-		if (UpdatedCachedExpressionData.UpdateForExpressions(Expressions, EMaterialParameterAssociation::GlobalParameter, -1))
+		FMaterialCachedExpressionContext Context(nullptr); // UMaterial have no parent
+		if (UpdatedCachedExpressionData.UpdateForExpressions(Context, Expressions, EMaterialParameterAssociation::GlobalParameter, -1))
 		{
 			// Only update our cached data if the update succeeded
 			// It's possible we could have some nullptr UMaterialExpressions, if we're loading cooked data here
diff --git a/Engine/Source/Runtime/Engine/Private/Materials/MaterialCachedData.cpp b/Engine/Source/Runtime/Engine/Private/Materials/MaterialCachedData.cpp
--- a/Engine/Source/Runtime/Engine/Private/Materials/MaterialCachedData.cpp
+++ b/Engine/Source/Runtime/Engine/Private/Materials/MaterialCachedData.cpp
@@ -96,7 +96,7 @@ static int32 TryAddParameter(FMaterialCachedParameters& CachedParameters, EMater
 	return INDEX_NONE;
 }
 
-bool FMaterialCachedExpressionData::UpdateForFunction(UMaterialFunctionInterface* Function, EMaterialParameterAssociation Association, int32 ParameterIndex)
+bool FMaterialCachedExpressionData::UpdateForFunction(const FMaterialCachedExpressionContext& Context, UMaterialFunctionInterface* Function, EMaterialParameterAssociation Association, int32 ParameterIndex)
 {
 	if (!Function)
 	{
@@ -107,7 +107,7 @@ bool FMaterialCachedExpressionData::UpdateForFunction(UMaterialFunctionInterface
 	const TArray<UMaterialExpression*>* FunctionExpressions = Function->GetFunctionExpressions();
 	if (FunctionExpressions)
 	{
-		if (!UpdateForExpressions(*FunctionExpressions, Association, ParameterIndex))
+		if (!UpdateForExpressions(Context, *FunctionExpressions, Association, ParameterIndex))
 		{
 			bResult = false;
 		}
@@ -176,12 +176,12 @@ bool FMaterialCachedExpressionData::UpdateForFunction(UMaterialFunctionInterface
 	return bResult;
 }
 
-bool FMaterialCachedExpressionData::UpdateForLayerFunctions(const FMaterialLayersFunctions& LayerFunctions)
+bool FMaterialCachedExpressionData::UpdateForLayerFunctions(const FMaterialCachedExpressionContext& Context, const FMaterialLayersFunctions& LayerFunctions)
 {
 	bool bResult = true;
 	for (int32 LayerIndex = 0; LayerIndex < LayerFunctions.Layers.Num(); ++LayerIndex)
 	{
-		if (!UpdateForFunction(LayerFunctions.Layers[LayerIndex], LayerParameter, LayerIndex))
+		if (!UpdateForFunction(Context, LayerFunctions.Layers[LayerIndex], LayerParameter, LayerIndex))
 		{
 			bResult = false;
 		}
@@ -189,7 +189,7 @@ bool FMaterialCachedExpressionData::UpdateForLayerFunctions(const FMaterialLayer
 
 	for (int32 BlendIndex = 0; BlendIndex < LayerFunctions.Blends.Num(); ++BlendIndex)
 	{
-		if (!UpdateForFunction(LayerFunctions.Blends[BlendIndex], BlendParameter, BlendIndex))
+		if (!UpdateForFunction(Context, LayerFunctions.Blends[BlendIndex], BlendParameter, BlendIndex))
 		{
 			bResult = false;
 		}
@@ -197,7 +197,7 @@ bool FMaterialCachedExpressionData::UpdateForLayerFunctions(const FMaterialLayer
 	return bResult;
 }
 
-bool FMaterialCachedExpressionData::UpdateForExpressions(const TArray<UMaterialExpression*>& Expressions, EMaterialParameterAssociation Association, int32 ParameterIndex)
+bool FMaterialCachedExpressionData::UpdateForExpressions(const FMaterialCachedExpressionContext& Context, const TArray<UMaterialExpression*>& Expressions, EMaterialParameterAssociation Association, int32 ParameterIndex)
 {
 	bool bResult = true;
 	for (UMaterialExpression* Expression : Expressions)
@@ -221,7 +221,12 @@ bool FMaterialCachedExpressionData::UpdateForExpressions(const TArray<UMaterialE
 			const int32 Index = TryAddParameter(Parameters, EMaterialParameterType::Scalar, ParameterInfo, ExpressionScalarParameter->ExpressionGUID);
 			if (Index != INDEX_NONE)
 			{
-				Parameters.ScalarValues.Insert(ExpressionScalarParameter->DefaultValue, Index);
+				float Value = ExpressionScalarParameter->DefaultValue;
+				if (Context.Parent)
+				{
+					Context.Parent->GetScalarParameterDefaultValue(ParameterInfo, Value, false, true);
+				}
+				Parameters.ScalarValues.Insert(Value, Index);
 				Parameters.ScalarMinMaxValues.Insert(FVector2D(ExpressionScalarParameter->SliderMin, ExpressionScalarParameter->SliderMax), Index);
 				if (ExpressionScalarParameter->IsUsedAsAtlasPosition())
 				{
@@ -242,7 +247,12 @@ bool FMaterialCachedExpressionData::UpdateForExpressions(const TArray<UMaterialE
 			const int32 Index = TryAddParameter(Parameters, EMaterialParameterType::Vector, ParameterInfo, ExpressionVectorParameter->ExpressionGUID);
 			if (Index != INDEX_NONE)
 			{
-				Parameters.VectorValues.Insert(ExpressionVectorParameter->DefaultValue, Index);
+				FLinearColor Value = ExpressionVectorParameter->DefaultValue;
+				if (Context.Parent)
+				{
+					Context.Parent->GetVectorParameterDefaultValue(ParameterInfo, Value, false, true);
+				}
+				Parameters.VectorValues.Insert(Value, Index);
 				Parameters.VectorChannelNameValues.Insert(ExpressionVectorParameter->ChannelNames, Index);
 				Parameters.VectorUsedAsChannelMaskValues.Insert(ExpressionVectorParameter->IsUsedAsChannelMask(), Index);
 			}
@@ -253,7 +263,12 @@ bool FMaterialCachedExpressionData::UpdateForExpressions(const TArray<UMaterialE
 			const int32 Index = TryAddParameter(Parameters, EMaterialParameterType::Texture, ParameterInfo, ExpressionTextureParameter->ExpressionGUID);
 			if (Index != INDEX_NONE)
 			{
-				Parameters.TextureValues.Insert(ExpressionTextureParameter->Texture, Index);
+				UTexture* Value = ExpressionTextureParameter->Texture;
+				if (Context.Parent)
+				{
+					Context.Parent->GetTextureParameterDefaultValue(ParameterInfo, Value, true);
+				}
+				Parameters.TextureValues.Insert(Value, Index);
 				Parameters.TextureChannelNameValues.Insert(ExpressionTextureParameter->ChannelNames, Index);
 			}
 		}
@@ -263,8 +278,14 @@ bool FMaterialCachedExpressionData::UpdateForExpressions(const TArray<UMaterialE
 			const int32 Index = TryAddParameter(Parameters, EMaterialParameterType::Font, ParameterInfo, ExpressionFontParameter->ExpressionGUID);
 			if (Index != INDEX_NONE)
 			{
-				Parameters.FontValues.Insert(ExpressionFontParameter->Font, Index);
-				Parameters.FontPageValues.Insert(ExpressionFontParameter->FontTexturePage, Index);
+				UFont* Font = ExpressionFontParameter->Font;
+				int32 FontTexturePage = ExpressionFontParameter->FontTexturePage;
+				if (Context.Parent)
+				{
+					Context.Parent->GetFontParameterDefaultValue(ParameterInfo, Font, FontTexturePage, true);
+				}
+				Parameters.FontValues.Insert(Font, Index);
+				Parameters.FontPageValues.Insert(FontTexturePage, Index);
 			}
 		}
 		else if (UMaterialExpressionRuntimeVirtualTextureSampleParameter* ExpressionRTVParameter = Cast<UMaterialExpressionRuntimeVirtualTextureSampleParameter>(Expression))
@@ -273,7 +294,12 @@ bool FMaterialCachedExpressionData::UpdateForExpressions(const TArray<UMaterialE
 			const int32 Index = TryAddParameter(Parameters, EMaterialParameterType::RuntimeVirtualTexture, ParameterInfo, ExpressionRTVParameter->ExpressionGUID);
 			if (Index != INDEX_NONE)
 			{
-				Parameters.RuntimeVirtualTextureValues.Insert(ExpressionRTVParameter->VirtualTexture, Index);
+				URuntimeVirtualTexture* Value = ExpressionRTVParameter->VirtualTexture;
+				if (Context.Parent)
+				{
+					Context.Parent->GetRuntimeVirtualTextureParameterDefaultValue(ParameterInfo, Value, true);
+				}
+				Parameters.RuntimeVirtualTextureValues.Insert(Value, Index);
 			}
 		}
 		else if (UMaterialExpressionCollectionParameter* ExpressionCollectionParameter = Cast<UMaterialExpressionCollectionParameter>(Expression))
@@ -327,7 +353,7 @@ bool FMaterialCachedExpressionData::UpdateForExpressions(const TArray<UMaterialE
 		}
 		else if (UMaterialExpressionMaterialFunctionCall* FunctionCall = Cast<UMaterialExpressionMaterialFunctionCall>(Expression))
 		{
-			if (!UpdateForFunction(FunctionCall->MaterialFunction, GlobalParameter, -1))
+			if (!UpdateForFunction(Context, FunctionCall->MaterialFunction, GlobalParameter, -1))
 			{
 				bResult = false;
 			}
@@ -339,7 +365,7 @@ bool FMaterialCachedExpressionData::UpdateForExpressions(const TArray<UMaterialE
 		else if (UMaterialExpressionMaterialAttributeLayers* LayersExpression = Cast<UMaterialExpressionMaterialAttributeLayers>(Expression))
 		{
 			checkf(Association == GlobalParameter, TEXT("UMaterialExpressionMaterialAttributeLayers can't be nested"));
-			if (!UpdateForLayerFunctions(LayersExpression->DefaultLayers))
+			if (!UpdateForLayerFunctions(Context, LayersExpression->DefaultLayers))
 			{
 				bResult = false;
 			}
@@ -422,6 +448,21 @@ bool FMaterialCachedParameters::IsParameterValid(EMaterialParameterType Type, in
 	return !bOveriddenOnly || Entry.Overrides[Index];
 }
 
+bool FMaterialCachedParameters::IsDefaultParameterValid(EMaterialParameterType Type, int32 Index, bool bOveriddenOnly, bool bCheckOwnedGlobalOverrides) const
+{
+	const FMaterialCachedParameterEntry& Entry = Entries[(int32)Type];
+	const bool bOveridden = Entry.Overrides[Index];
+	if (!bCheckOwnedGlobalOverrides && bOveridden)
+	{
+		return false;
+	}
+	if (bOveriddenOnly && !bOveridden)
+	{
+		return false;
+	}
+	return true;
+}
+
 void FMaterialCachedParameters::GetAllParameterInfoOfType(EMaterialParameterType Type, bool bEmptyOutput, TArray<FMaterialParameterInfo>& OutParameterInfo, TArray<FGuid>& OutParameterIds) const
 {
 	const FMaterialCachedParameterEntry& Entry = Entries[(int32)Type];
diff --git a/Engine/Source/Runtime/Engine/Private/Materials/MaterialInstance.cpp b/Engine/Source/Runtime/Engine/Private/Materials/MaterialInstance.cpp
--- a/Engine/Source/Runtime/Engine/Private/Materials/MaterialInstance.cpp
+++ b/Engine/Source/Runtime/Engine/Private/Materials/MaterialInstance.cpp
@@ -692,180 +692,6 @@ bool UMaterialInstance::GetScalarParameterSliderMinMax(const FHashedMaterialPara
 }
 #endif // WITH_EDITOR
 
-#if WITH_EDITOR
-
-#define checkLayerParameter(x) check(x)
-
-bool UMaterialInstance::GetScalarLayerParameterValue_Legacy(const FHashedMaterialParameterInfo& ParameterInfo, bool bOveriddenOnly, float& OutValue, bool& OutResult) const
-{
-	UMaterialExpressionScalarParameter* Parameter = nullptr;
-	for (const FStaticMaterialLayersParameter& LayersParam : StaticParameters.MaterialLayersParameters)
-	{
-		if (LayersParam.bOverride)
-		{
-			UMaterialFunctionInterface* Function = LayersParam.GetParameterAssociatedFunction(ParameterInfo);
-			UMaterialFunctionInterface* ParameterOwner = nullptr;
-
-			if (Function && Function->OverrideNamedScalarParameter(ParameterInfo, OutValue))
-			{
-				OutResult = true;
-				return true;
-			}
-
-			if (Function && Function->GetNamedParameterOfType(ParameterInfo, Parameter, &ParameterOwner))
-			{
-				if (ParameterOwner->OverrideNamedScalarParameter(ParameterInfo, OutValue))
-				{
-					OutResult = true;
-					return true;
-				}
-				
-				Parameter->IsNamedParameter(ParameterInfo, OutValue);
-				OutResult = !bOveriddenOnly;
-				return true;
-			}
-		}
-	}
-
-	return false;
-}
-
-bool UMaterialInstance::GetVectorLayerParameterValue_Legacy(const FHashedMaterialParameterInfo& ParameterInfo, bool bOveriddenOnly, FLinearColor& OutValue, bool& OutResult) const
-{
-	UMaterialExpressionVectorParameter* Parameter = nullptr;
-	for (const FStaticMaterialLayersParameter& LayersParam : StaticParameters.MaterialLayersParameters)
-	{
-		if (LayersParam.bOverride)
-		{
-			UMaterialFunctionInterface* Function = LayersParam.GetParameterAssociatedFunction(ParameterInfo);
-			UMaterialFunctionInterface* ParameterOwner = nullptr;
-
-			if (Function && Function->OverrideNamedVectorParameter(ParameterInfo, OutValue))
-			{
-				OutResult = true;
-				return true;
-			}
-
-			if (Function && Function->GetNamedParameterOfType(ParameterInfo, Parameter, &ParameterOwner))
-			{
-				if (ParameterOwner->OverrideNamedVectorParameter(ParameterInfo, OutValue))
-				{
-					OutResult = true;
-					return true;
-				}
-
-				Parameter->IsNamedParameter(ParameterInfo, OutValue);
-				OutResult = !bOveriddenOnly;
-				return true;
-			}
-		}
-	}
-
-	return false;
-}
-
-bool UMaterialInstance::GetTextureLayerParameterValue_Legacy(const FHashedMaterialParameterInfo& ParameterInfo, bool bOveriddenOnly, UTexture*& OutValue, bool& OutResult) const
-{
-	UMaterialExpressionTextureSampleParameter* Parameter = nullptr;
-	for (const FStaticMaterialLayersParameter& LayersParam : StaticParameters.MaterialLayersParameters)
-	{
-		if (LayersParam.bOverride)
-		{
-			UMaterialFunctionInterface* Function = LayersParam.GetParameterAssociatedFunction(ParameterInfo);
-			UMaterialFunctionInterface* ParameterOwner = nullptr;
-
-			if (Function && Function->OverrideNamedTextureParameter(ParameterInfo, OutValue))
-			{
-				OutResult = true;
-				return true;
-			}
-
-			if (Function && Function->GetNamedParameterOfType(ParameterInfo, Parameter, &ParameterOwner))
-			{
-				// TODO - Different logic than other parameter types for 'OutResult', intentional??
-				if (!ParameterOwner->OverrideNamedTextureParameter(ParameterInfo, OutValue))
-				{
-					Parameter->IsNamedParameter(ParameterInfo, OutValue);
-				}
-				OutResult = true;
-				return true;
-			}
-		}
-	}
-
-	return false;
-}
-
-bool UMaterialInstance::GetFontLayerParameterValue_Legacy(const FHashedMaterialParameterInfo& ParameterInfo, bool bOveriddenOnly, UFont*& OutFontValue, int32& OutFontPage, bool& OutResult) const
-{
-	UMaterialExpressionFontSampleParameter* Parameter = nullptr;
-	for (const FStaticMaterialLayersParameter& LayersParam : StaticParameters.MaterialLayersParameters)
-	{
-		if (LayersParam.bOverride)
-		{
-			UMaterialFunctionInterface* Function = LayersParam.GetParameterAssociatedFunction(ParameterInfo);
-			UMaterialFunctionInterface* ParameterOwner = nullptr;
-
-			if (Function && Function->OverrideNamedFontParameter(ParameterInfo, OutFontValue, OutFontPage))
-			{
-				OutResult = true;
-				return true;
-			}
-
-			if (Function && Function->GetNamedParameterOfType(ParameterInfo, Parameter, &ParameterOwner))
-			{
-				if (ParameterOwner->OverrideNamedFontParameter(ParameterInfo, OutFontValue, OutFontPage))
-				{
-					OutResult = true;
-					return true;
-				}
-				Parameter->IsNamedParameter(ParameterInfo, OutFontValue, OutFontPage);
-				OutResult = !bOveriddenOnly;
-				return true;
-			}
-		}
-	}
-
-	return false;
-}
-
-bool UMaterialInstance::GetRuntimeVirtualTextureLayerParameterValue_Legacy(const FHashedMaterialParameterInfo& ParameterInfo, bool bOveriddenOnly, URuntimeVirtualTexture*& OutValue, bool& OutResult) const
-{
-	UMaterialExpressionRuntimeVirtualTextureSampleParameter* Parameter = nullptr;
-	for (const FStaticMaterialLayersParameter& LayersParam : StaticParameters.MaterialLayersParameters)
-	{
-		if (LayersParam.bOverride)
-		{
-			UMaterialFunctionInterface* Function = LayersParam.GetParameterAssociatedFunction(ParameterInfo);
-			UMaterialFunctionInterface* ParameterOwner = nullptr;
-
-			if (Function && Function->OverrideNamedRuntimeVirtualTextureParameter(ParameterInfo, OutValue))
-			{
-				OutResult = true;
-				return true;
-			}
-
-			if (Function && Function->GetNamedParameterOfType(ParameterInfo, Parameter, &ParameterOwner))
-			{
-				if (!ParameterOwner->OverrideNamedRuntimeVirtualTextureParameter(ParameterInfo, OutValue))
-				{
-					Parameter->IsNamedParameter(ParameterInfo, OutValue);
-				}
-				OutResult = true;
-				return true;
-			}
-		}
-	}
-
-	return false;
-}
-
-#endif // WITH_EDITOR
-
-#if !defined(checkLayerParameter)
-#define checkLayerParameter(x)
-#endif
-
 bool UMaterialInstance::GetScalarParameterValue(const FHashedMaterialParameterInfo& ParameterInfo, float& OutValue, bool bOveriddenOnly) const
 {
 	bool bFoundAValue = false;
@@ -886,25 +712,13 @@ bool UMaterialInstance::GetScalarParameterValue(const FHashedMaterialParameterIn
 	// Instance-included default
 	if (ParameterInfo.Association != EMaterialParameterAssociation::GlobalParameter)
 	{
-#if WITH_EDITOR
-		bool bLegacyFoundValue = false;
-		bool bLegacyResult = false;
-		float LegacyValue = 0.0f;
-		bLegacyFoundValue = GetScalarLayerParameterValue_Legacy(ParameterInfo, bOveriddenOnly, LegacyValue, bLegacyResult);
-#endif // WITH_EDITOR
-
 		const int32 ParameterIndex = CachedLayerParameters.FindParameterIndex(EMaterialParameterType::Scalar, ParameterInfo);
 		if (ParameterIndex != INDEX_NONE)
 		{
-			checkLayerParameter(bLegacyFoundValue);
 			const bool bResult = CachedLayerParameters.IsParameterValid(EMaterialParameterType::Scalar, ParameterIndex, bOveriddenOnly);
 			OutValue = CachedLayerParameters.ScalarValues[ParameterIndex];
-			checkLayerParameter(OutValue == LegacyValue);
-			checkLayerParameter(bLegacyResult == bResult);
 			return bResult;
 		}
-
-		checkLayerParameter(!bLegacyFoundValue);
 	}
 	
 	// Next material in hierarchy
@@ -992,25 +806,13 @@ bool UMaterialInstance::GetVectorParameterValue(const FHashedMaterialParameterIn
 	// Instance-included default
 	if (ParameterInfo.Association != EMaterialParameterAssociation::GlobalParameter)
 	{
-#if WITH_EDITOR
-		bool bLegacyFoundValue = false;
-		bool bLegacyResult = false;
-		FLinearColor LegacyValue(ForceInitToZero);
-		bLegacyFoundValue = GetVectorLayerParameterValue_Legacy(ParameterInfo, bOveriddenOnly, LegacyValue, bLegacyResult);
-#endif // WITH_EDITOR
-
 		const int32 ParameterIndex = CachedLayerParameters.FindParameterIndex(EMaterialParameterType::Vector, ParameterInfo);
 		if (ParameterIndex != INDEX_NONE)
 		{
-			checkLayerParameter(bLegacyFoundValue);
 			const bool bResult = CachedLayerParameters.IsParameterValid(EMaterialParameterType::Vector, ParameterIndex, bOveriddenOnly);
 			OutValue = CachedLayerParameters.VectorValues[ParameterIndex];
-			checkLayerParameter(OutValue == LegacyValue);
-			checkLayerParameter(bLegacyResult == bResult);
 			return bResult;
 		}
-
-		checkLayerParameter(!bLegacyFoundValue);
 	}
 	
 	// Next material in hierarchy
@@ -1105,25 +907,13 @@ bool UMaterialInstance::GetTextureParameterValue(const FHashedMaterialParameterI
 	// Instance-included default
 	if (ParameterInfo.Association != EMaterialParameterAssociation::GlobalParameter)
 	{
-#if WITH_EDITOR
-		bool bLegacyFoundValue = false;
-		bool bLegacyResult = false;
-		UTexture* LegacyValue = nullptr;
-		bLegacyFoundValue = GetTextureLayerParameterValue_Legacy(ParameterInfo, bOveriddenOnly, LegacyValue, bLegacyResult);
-#endif // WITH_EDITOR
-
 		const int32 ParameterIndex = CachedLayerParameters.FindParameterIndex(EMaterialParameterType::Texture, ParameterInfo);
 		if (ParameterIndex != INDEX_NONE)
 		{
-			checkLayerParameter(bLegacyFoundValue);
 			const bool bResult = CachedLayerParameters.IsParameterValid(EMaterialParameterType::Texture, ParameterIndex, bOveriddenOnly);
 			OutValue = CachedLayerParameters.TextureValues[ParameterIndex];
-			checkLayerParameter(OutValue == LegacyValue);
-			checkLayerParameter(bLegacyResult == bResult);
 			return bResult;
 		}
-
-		checkLayerParameter(!bLegacyFoundValue);
 	}
 	
 	// Next material in hierarchy
@@ -1156,25 +946,13 @@ bool UMaterialInstance::GetRuntimeVirtualTextureParameterValue(const FHashedMate
 	// Instance-included default
 	if (ParameterInfo.Association != EMaterialParameterAssociation::GlobalParameter)
 	{
-#if WITH_EDITOR
-		bool bLegacyFoundValue = false;
-		bool bLegacyResult = false;
-		URuntimeVirtualTexture* LegacyValue = nullptr;
-		bLegacyFoundValue = GetRuntimeVirtualTextureLayerParameterValue_Legacy(ParameterInfo, bOveriddenOnly, LegacyValue, bLegacyResult);
-#endif // WITH_EDITOR
-
 		const int32 ParameterIndex = CachedLayerParameters.FindParameterIndex(EMaterialParameterType::RuntimeVirtualTexture, ParameterInfo);
 		if (ParameterIndex != INDEX_NONE)
 		{
-			checkLayerParameter(bLegacyFoundValue);
 			const bool bResult = CachedLayerParameters.IsParameterValid(EMaterialParameterType::RuntimeVirtualTexture, ParameterIndex, bOveriddenOnly);
 			OutValue = CachedLayerParameters.RuntimeVirtualTextureValues[ParameterIndex];
-			checkLayerParameter(OutValue == LegacyValue);
-			checkLayerParameter(bLegacyResult == bResult);
 			return bResult;
 		}
-
-		checkLayerParameter(!bLegacyFoundValue);
 	}
 
 	// Next material in hierarchy
@@ -1241,28 +1019,14 @@ bool UMaterialInstance::GetFontParameterValue(const FHashedMaterialParameterInfo
 	// Instance-included default
 	if (ParameterInfo.Association != EMaterialParameterAssociation::GlobalParameter)
 	{
-#if WITH_EDITOR
-		bool bLegacyFoundValue = false;
-		bool bLegacyResult = false;
-		UFont* LegacyValue = nullptr;
-		int32 LegacyFontPageValue = INDEX_NONE;
-		bLegacyFoundValue = GetFontLayerParameterValue_Legacy(ParameterInfo, bOveriddenOnly, LegacyValue, LegacyFontPageValue, bLegacyResult);
-#endif // WITH_EDITOR
-
 		const int32 ParameterIndex = CachedLayerParameters.FindParameterIndex(EMaterialParameterType::Font, ParameterInfo);
 		if (ParameterIndex != INDEX_NONE)
 		{
-			checkLayerParameter(bLegacyFoundValue);
 			const bool bResult = CachedLayerParameters.IsParameterValid(EMaterialParameterType::Font, ParameterIndex, bOveriddenOnly);
 			OutFontValue = CachedLayerParameters.FontValues[ParameterIndex];
 			OutFontPage = CachedLayerParameters.FontPageValues[ParameterIndex];
-			checkLayerParameter(OutFontValue == LegacyValue);
-			checkLayerParameter(OutFontPage == LegacyFontPageValue);
-			checkLayerParameter(bLegacyResult == bResult);
 			return bResult;
 		}
-
-		checkLayerParameter(!bLegacyFoundValue);
 	}
 	
 	// Next material in hierarchy
@@ -2383,31 +2147,7 @@ bool UMaterialInstance::GetScalarParameterDefaultValue(const FHashedMaterialPara
 	}
 #endif
 
-	// In the case of duplicate parameters with different values, this will return the
-	// first matching expression found, not necessarily the one that's used for rendering
-	if (ParameterInfo.Association != EMaterialParameterAssociation::GlobalParameter)
-	{
-#if WITH_EDITOR
-		bool bLegacyFoundValue = false;
-		bool bLegacyResult = false;
-		float LegacyValue = 0.0f;
-		// Legacy code doesn't check bOveriddenOnly here
-		bLegacyFoundValue = GetScalarLayerParameterValue_Legacy(ParameterInfo, false, LegacyValue, bLegacyResult);
-#endif // WITH_EDITOR
-
-		const int32 ParameterIndex = CachedLayerParameters.FindParameterIndex(EMaterialParameterType::Scalar, ParameterInfo);
-		if (ParameterIndex != INDEX_NONE)
-		{
-			checkLayerParameter(bLegacyFoundValue);
-			OutValue = CachedLayerParameters.ScalarValues[ParameterIndex];
-			checkLayerParameter(OutValue == LegacyValue);
-			checkLayerParameter(bLegacyResult);
-			return true;
-		}
-
-		checkLayerParameter(!bLegacyFoundValue);
-	}
-	else if (bCheckOwnedGlobalOverrides)
+	if (bCheckOwnedGlobalOverrides)
 	{
 		// Parameters overridden by this instance
 		for (const FScalarParameterValue& ScalarParam : ScalarParameterValues)
@@ -2419,7 +2159,19 @@ bool UMaterialInstance::GetScalarParameterDefaultValue(const FHashedMaterialPara
 			}
 		}
 	}
-	
+
+	// In the case of duplicate parameters with different values, this will return the
+	// first matching expression found, not necessarily the one that's used for rendering
+	if (ParameterInfo.Association != EMaterialParameterAssociation::GlobalParameter)
+	{
+		const int32 ParameterIndex = CachedLayerParameters.FindParameterIndex(EMaterialParameterType::Scalar, ParameterInfo);
+		if (ParameterIndex != INDEX_NONE)
+		{
+			OutValue = CachedLayerParameters.ScalarValues[ParameterIndex];
+			return CachedLayerParameters.IsDefaultParameterValid(EMaterialParameterType::Scalar, ParameterIndex, bOveriddenOnly, bCheckOwnedGlobalOverrides);
+		}
+	}
+
 	if (Parent)
 	{
 #if WITH_EDITOR
@@ -2442,31 +2194,7 @@ bool UMaterialInstance::GetVectorParameterDefaultValue(const FHashedMaterialPara
 
 	// In the case of duplicate parameters with different values, this will return the
 	// first matching expression found, not necessarily the one that's used for rendering
-	UMaterialExpressionVectorParameter* Parameter = nullptr;
-
-	if (ParameterInfo.Association != EMaterialParameterAssociation::GlobalParameter)
-	{
-#if WITH_EDITOR
-		bool bLegacyFoundValue = false;
-		bool bLegacyResult = false;
-		FLinearColor LegacyValue(ForceInitToZero);
-		// Legacy code doesn't check bOveriddenOnly here
-		bLegacyFoundValue = GetVectorLayerParameterValue_Legacy(ParameterInfo, false, LegacyValue, bLegacyResult);
-#endif // WITH_EDITOR
-
-		const int32 ParameterIndex = CachedLayerParameters.FindParameterIndex(EMaterialParameterType::Vector, ParameterInfo);
-		if (ParameterIndex != INDEX_NONE)
-		{
-			checkLayerParameter(bLegacyFoundValue);
-			OutValue = CachedLayerParameters.VectorValues[ParameterIndex];
-			checkLayerParameter(OutValue == LegacyValue);
-			checkLayerParameter(bLegacyResult);
-			return true;
-		}
-
-		checkLayerParameter(!bLegacyFoundValue);
-	}
-	else if (bCheckOwnedGlobalOverrides)
+	if (bCheckOwnedGlobalOverrides)
 	{
 		// Parameters overridden by this instance
 		for (const FVectorParameterValue& VectorParam : VectorParameterValues)
@@ -2478,7 +2206,17 @@ bool UMaterialInstance::GetVectorParameterDefaultValue(const FHashedMaterialPara
 			}
 		}
 	}
-	
+
+	if (ParameterInfo.Association != EMaterialParameterAssociation::GlobalParameter)
+	{
+		const int32 ParameterIndex = CachedLayerParameters.FindParameterIndex(EMaterialParameterType::Vector, ParameterInfo);
+		if (ParameterIndex != INDEX_NONE)
+		{
+			OutValue = CachedLayerParameters.VectorValues[ParameterIndex];
+			return CachedLayerParameters.IsDefaultParameterValid(EMaterialParameterType::Vector, ParameterIndex, bOveriddenOnly, bCheckOwnedGlobalOverrides);
+		}
+	}
+
 	if (Parent)
 	{
 #if WITH_EDITOR
@@ -2501,31 +2239,7 @@ bool UMaterialInstance::GetTextureParameterDefaultValue(const FHashedMaterialPar
 
 	// In the case of duplicate parameters with different values, this will return the
 	// first matching expression found, not necessarily the one that's used for rendering
-	UMaterialExpressionTextureSampleParameter* Parameter = nullptr;
-
-	if (ParameterInfo.Association != EMaterialParameterAssociation::GlobalParameter)
-	{
-#if WITH_EDITOR
-		bool bLegacyFoundValue = false;
-		bool bLegacyResult = false;
-		UTexture* LegacyValue = nullptr;
-		// Legacy code doesn't check bOveriddenOnly here
-		bLegacyFoundValue = GetTextureLayerParameterValue_Legacy(ParameterInfo, false, LegacyValue, bLegacyResult);
-#endif // WITH_EDITOR
-
-		const int32 ParameterIndex = CachedLayerParameters.FindParameterIndex(EMaterialParameterType::Texture, ParameterInfo);
-		if (ParameterIndex != INDEX_NONE)
-		{
-			checkLayerParameter(bLegacyFoundValue);
-			OutValue = CachedLayerParameters.TextureValues[ParameterIndex];
-			checkLayerParameter(OutValue == LegacyValue);
-			checkLayerParameter(bLegacyResult);
-			return true;
-		}
-
-		checkLayerParameter(!bLegacyFoundValue);
-	}
-	else if (bCheckOwnedGlobalOverrides)
+	if (bCheckOwnedGlobalOverrides)
 	{
 		// Parameters overridden by this instance
 		for (const FTextureParameterValue& TextureParam : TextureParameterValues)
@@ -2537,7 +2251,17 @@ bool UMaterialInstance::GetTextureParameterDefaultValue(const FHashedMaterialPar
 			}
 		}
 	}
-	
+
+	if (ParameterInfo.Association != EMaterialParameterAssociation::GlobalParameter)
+	{
+		const int32 ParameterIndex = CachedLayerParameters.FindParameterIndex(EMaterialParameterType::Texture, ParameterInfo);
+		if (ParameterIndex != INDEX_NONE)
+		{
+			OutValue = CachedLayerParameters.TextureValues[ParameterIndex];
+			return CachedLayerParameters.IsDefaultParameterValid(EMaterialParameterType::Texture, ParameterIndex, false, bCheckOwnedGlobalOverrides);
+		}
+	}
+
 	if (Parent)
 	{
 #if WITH_EDITOR
@@ -2560,31 +2284,7 @@ bool UMaterialInstance::GetRuntimeVirtualTextureParameterDefaultValue(const FHas
 
 	// In the case of duplicate parameters with different values, this will return the
 	// first matching expression found, not necessarily the one that's used for rendering
-	UMaterialExpressionRuntimeVirtualTextureSampleParameter* Parameter = nullptr;
-
-	if (ParameterInfo.Association != EMaterialParameterAssociation::GlobalParameter)
-	{
-#if WITH_EDITOR
-		bool bLegacyFoundValue = false;
-		bool bLegacyResult = false;
-		URuntimeVirtualTexture* LegacyValue = nullptr;
-		// Legacy code doesn't check bOveriddenOnly here
-		bLegacyFoundValue = GetRuntimeVirtualTextureLayerParameterValue_Legacy(ParameterInfo, false, LegacyValue, bLegacyResult);
-#endif // WITH_EDITOR
-
-		const int32 ParameterIndex = CachedLayerParameters.FindParameterIndex(EMaterialParameterType::RuntimeVirtualTexture, ParameterInfo);
-		if (ParameterIndex != INDEX_NONE)
-		{
-			checkLayerParameter(bLegacyFoundValue);
-			OutValue = CachedLayerParameters.RuntimeVirtualTextureValues[ParameterIndex];
-			checkLayerParameter(OutValue == LegacyValue);
-			checkLayerParameter(bLegacyResult);
-			return true;
-		}
-
-		checkLayerParameter(!bLegacyFoundValue);
-	}
-	else if (bCheckOwnedGlobalOverrides)
+	if (bCheckOwnedGlobalOverrides)
 	{
 		// Parameters overridden by this instance
 		for (const FRuntimeVirtualTextureParameterValue& RuntimeVirtualTextureParam : RuntimeVirtualTextureParameterValues)
@@ -2597,6 +2297,16 @@ bool UMaterialInstance::GetRuntimeVirtualTextureParameterDefaultValue(const FHas
 		}
 	}
 
+	if (ParameterInfo.Association != EMaterialParameterAssociation::GlobalParameter)
+	{
+		const int32 ParameterIndex = CachedLayerParameters.FindParameterIndex(EMaterialParameterType::RuntimeVirtualTexture, ParameterInfo);
+		if (ParameterIndex != INDEX_NONE)
+		{
+			OutValue = CachedLayerParameters.RuntimeVirtualTextureValues[ParameterIndex];
+			return CachedLayerParameters.IsDefaultParameterValid(EMaterialParameterType::RuntimeVirtualTexture, ParameterIndex, false, bCheckOwnedGlobalOverrides);
+		}
+	}
+
 	if (Parent)
 	{
 #if WITH_EDITOR
@@ -2619,34 +2329,7 @@ bool UMaterialInstance::GetFontParameterDefaultValue(const FHashedMaterialParame
 
 	// In the case of duplicate parameters with different values, this will return the
 	// first matching expression found, not necessarily the one that's used for rendering
-	UMaterialExpressionFontSampleParameter* Parameter = nullptr;
-
-	if (ParameterInfo.Association != EMaterialParameterAssociation::GlobalParameter)
-	{
-#if WITH_EDITOR
-		bool bLegacyFoundValue = false;
-		bool bLegacyResult = false;
-		UFont* LegacyValue = nullptr;
-		int32 LegacyFontPageValue = INDEX_NONE;
-		// Legacy code doesn't check bOveriddenOnly here
-		bLegacyFoundValue = GetFontLayerParameterValue_Legacy(ParameterInfo, false, LegacyValue, LegacyFontPageValue, bLegacyResult);
-#endif // WITH_EDITOR
-
-		const int32 ParameterIndex = CachedLayerParameters.FindParameterIndex(EMaterialParameterType::Font, ParameterInfo);
-		if (ParameterIndex != INDEX_NONE)
-		{
-			checkLayerParameter(bLegacyFoundValue);
-			OutFontValue = CachedLayerParameters.FontValues[ParameterIndex];
-			OutFontPage = CachedLayerParameters.FontPageValues[ParameterIndex];
-			checkLayerParameter(OutFontValue == LegacyValue);
-			checkLayerParameter(OutFontPage == LegacyFontPageValue);
-			checkLayerParameter(bLegacyResult);
-			return true;
-		}
-
-		checkLayerParameter(!bLegacyFoundValue);
-	}
-	else if (bCheckOwnedGlobalOverrides)
+	if (bCheckOwnedGlobalOverrides)
 	{
 		// Parameters overridden by this instance
 		for (const FFontParameterValue& FontParam : FontParameterValues)
@@ -2659,7 +2342,18 @@ bool UMaterialInstance::GetFontParameterDefaultValue(const FHashedMaterialParame
 			}
 		}
 	}
-	
+
+	if (ParameterInfo.Association != EMaterialParameterAssociation::GlobalParameter)
+	{
+		const int32 ParameterIndex = CachedLayerParameters.FindParameterIndex(EMaterialParameterType::Font, ParameterInfo);
+		if (ParameterIndex != INDEX_NONE)
+		{
+			OutFontValue = CachedLayerParameters.FontValues[ParameterIndex];
+			OutFontPage = CachedLayerParameters.FontPageValues[ParameterIndex];
+			return CachedLayerParameters.IsDefaultParameterValid(EMaterialParameterType::RuntimeVirtualTexture, ParameterIndex, false, bCheckOwnedGlobalOverrides);
+		}
+	}
+
 	if (Parent)
 	{
 #if WITH_EDITOR
@@ -4245,7 +3939,8 @@ void UMaterialInstance::UpdateCachedLayerParameters()
 	bool bCachedDataValid = true;
 	for (FStaticMaterialLayersParameter& LayerParameters : StaticParameters.MaterialLayersParameters)
 	{
-		if (!CachedExpressionData.UpdateForLayerFunctions(LayerParameters.Value))
+		FMaterialCachedExpressionContext Context(Parent);
+		if (!CachedExpressionData.UpdateForLayerFunctions(Context, LayerParameters.Value))
 		{
 			bCachedDataValid = false;
 		}
diff --git a/Engine/Source/Runtime/Engine/Public/MaterialCachedData.h b/Engine/Source/Runtime/Engine/Public/MaterialCachedData.h
--- a/Engine/Source/Runtime/Engine/Public/MaterialCachedData.h
+++ b/Engine/Source/Runtime/Engine/Public/MaterialCachedData.h
@@ -111,6 +111,7 @@ struct FMaterialCachedParameters
 	int32 FindParameterIndex(EMaterialParameterType Type, const FHashedMaterialParameterInfo& HashedParameterInfo, bool bOveriddenOnly) const;
 	int32 FindParameterIndex(EMaterialParameterType Type, const FHashedMaterialParameterInfo& HashedParameterInfo) const;
 	bool IsParameterValid(EMaterialParameterType Type, int32 Index, bool bOveriddenOnly) const;
+	bool IsDefaultParameterValid(EMaterialParameterType Type, int32 Index, bool bOveriddenOnly, bool bCheckOwnedGlobalOverrides) const;
 	void GetAllParameterInfoOfType(EMaterialParameterType Type, bool bEmptyOutput, TArray<FMaterialParameterInfo>& OutParameterInfo, TArray<FGuid>& OutParameterIds) const;
 	void GetAllGlobalParameterInfoOfType(EMaterialParameterType Type, bool bEmptyOutput, TArray<FMaterialParameterInfo>& OutParameterInfo, TArray<FGuid>& OutParameterIds) const;
 	void Reset();
@@ -157,6 +158,12 @@ struct FMaterialCachedParameters
 #endif // WITH_EDITORONLY_DATA
 };
 
+struct FMaterialCachedExpressionContext
+{
+	explicit FMaterialCachedExpressionContext(UMaterialInterface* InParent = nullptr) : Parent(InParent) {}
+	UMaterialInterface* Parent;
+};
+
 USTRUCT()
 struct FMaterialCachedExpressionData
 {
@@ -164,9 +171,9 @@ struct FMaterialCachedExpressionData
 
 #if WITH_EDITOR
 	/** Returns 'false' if update is incomplete, due to missing expression data (stripped from non-editor build) */
-	bool UpdateForExpressions(const TArray<UMaterialExpression*>& Expressions, EMaterialParameterAssociation Association, int32 ParameterIndex);
-	bool UpdateForFunction(UMaterialFunctionInterface* Function, EMaterialParameterAssociation Association, int32 ParameterIndex);
-	bool UpdateForLayerFunctions(const FMaterialLayersFunctions& LayerFunctions);
+	bool UpdateForExpressions(const FMaterialCachedExpressionContext& Context, const TArray<UMaterialExpression*>& Expressions, EMaterialParameterAssociation Association, int32 ParameterIndex);
+	bool UpdateForFunction(const FMaterialCachedExpressionContext& Context, UMaterialFunctionInterface* Function, EMaterialParameterAssociation Association, int32 ParameterIndex);
+	bool UpdateForLayerFunctions(const FMaterialCachedExpressionContext& Context, const FMaterialLayersFunctions& LayerFunctions);
 #endif // WITH_EDITOR
 
 	void Reset();
