Revised fix for variables on spawned Blueprint instances not being initialized to modified CDO values for native struct types with one or more non-reflected fields exposed for edit using details customization.

diff --git a/Engine/Source/Runtime/Engine/Private/BlueprintGeneratedClass.cpp b/Engine/Source/Runtime/Engine/Private/BlueprintGeneratedClass.cpp
--- a/Engine/Source/Runtime/Engine/Private/BlueprintGeneratedClass.cpp
+++ b/Engine/Source/Runtime/Engine/Private/BlueprintGeneratedClass.cpp
@@ -437,8 +437,22 @@ bool UBlueprintGeneratedClass::BuildCustomPropertyListForPostConstruction(FCusto
 					*CurrentNodePtr = new FCustomPropertyListNode(Property, Idx);
 					CustomPropertyListForPostConstruction.Add(*CurrentNodePtr);
 
-					// Recursively gather up all struct fields that differ and assign to the current node's sub property list.
-					if (BuildCustomPropertyListForPostConstruction((*CurrentNodePtr)->SubPropertyList, StructProperty->Struct, PropertyValue, DefaultPropertyValue))
+					UScriptStruct::ICppStructOps* CppStructOps = nullptr;
+					if (StructProperty->Struct)
+					{
+						CppStructOps = StructProperty->Struct->GetCppStructOps();
+					}
+
+					// Check if we should initialize using the full value (e.g. a USTRUCT with one or more non-reflected fields).
+					bool bIsIdentical = false;
+					const uint32 PortFlags = 0;
+					if(!CppStructOps || !CppStructOps->HasIdentical() || !CppStructOps->Identical(PropertyValue, DefaultPropertyValue, PortFlags, bIsIdentical))
+					{
+						// Recursively gather up all struct fields that differ and assign to the current node's sub property list.
+						bIsIdentical = !BuildCustomPropertyListForPostConstruction((*CurrentNodePtr)->SubPropertyList, StructProperty->Struct, PropertyValue, DefaultPropertyValue);
+					}
+
+					if (!bIsIdentical)
 					{
 						// Advance to the next node in the list.
 						CurrentNodePtr = &(*CurrentNodePtr)->PropertyListNext;
@@ -643,11 +657,15 @@ void UBlueprintGeneratedClass::InitPropertiesFromCustomList(const FCustomPropert
 
 		if (const FStructProperty* StructProperty = CastField<FStructProperty>(CustomPropertyListNode->Property))
 		{
-			// This should never be NULL; we should not be recording the StructProperty without at least one sub property, but we'll verify just to be sure.
-			if (ensure(CustomPropertyListNode->SubPropertyList != nullptr))
+			if (CustomPropertyListNode->SubPropertyList != nullptr)
 			{
 				InitPropertiesFromCustomList(CustomPropertyListNode->SubPropertyList, StructProperty->Struct, PropertyValue, DefaultPropertyValue);
 			}
+			else
+			{
+				// A NULL sub-property list indicates that we should copy the entire default value (e.g. a struct with one or more non-reflected fields).
+				StructProperty->CopySingleValue(PropertyValue, DefaultPropertyValue);
+			}
 		}
 		else if (const FArrayProperty* ArrayProperty = CastField<FArrayProperty>(CustomPropertyListNode->Property))
 		{
