Added a IsPendingKillPending() check in ALODActor::UpdateSubActorLODParents()

diff --git a/Engine/Source/Runtime/Engine/Private/LODActor.cpp b/Engine/Source/Runtime/Engine/Private/LODActor.cpp
--- a/Engine/Source/Runtime/Engine/Private/LODActor.cpp
+++ b/Engine/Source/Runtime/Engine/Private/LODActor.cpp
@@ -1010,7 +1010,7 @@ void ALODActor::UpdateSubActorLODParents()
 {
 	for (AActor* Actor : SubActors)
 	{	
-		if(Actor)
+		if (Actor && !Actor->IsPendingKillPending())
 		{
 			UStaticMeshComponent* LODComponent = GetLODComponentForActor(Actor);
 			Actor->SetLODParent(LODComponent, LODComponent->MinDrawDistance);
