Add missing forward declaration

diff --git a/Engine/Source/Runtime/Engine/Public/MaterialCachedData.h b/Engine/Source/Runtime/Engine/Public/MaterialCachedData.h
--- a/Engine/Source/Runtime/Engine/Public/MaterialCachedData.h
+++ b/Engine/Source/Runtime/Engine/Public/MaterialCachedData.h
@@ -16,6 +16,7 @@ class UMaterialExpression;
 class URuntimeVirtualTexture;
 class ULandscapeGrassType;
 class UMaterialFunctionInterface;
+class UMaterialInterface;
 
 /** Stores information about a function that this material references, used to know when the material needs to be recompiled. */
 USTRUCT()
