Wrap some stuff in iif WITH_EDITOR

diff --git a/Engine/Source/Runtime/Engine/Private/Materials/MaterialInstance.cpp b/Engine/Source/Runtime/Engine/Private/Materials/MaterialInstance.cpp
--- a/Engine/Source/Runtime/Engine/Private/Materials/MaterialInstance.cpp
+++ b/Engine/Source/Runtime/Engine/Private/Materials/MaterialInstance.cpp
@@ -365,6 +365,7 @@ void GameThread_UpdateMIParameter(const UMaterialInstance* Instance, const Param
 		});
 }
 
+#if WITH_EDITOR
 template<typename ParameterType>
 static void RemapLayerParameterIndicesArray(TArray<ParameterType>& Parameters, const TArray<int32>& RemapLayerIndices)
 {
@@ -442,6 +443,7 @@ void UMaterialInstance::SwapLayerParameterIndices(int32 OriginalIndex, int32 New
 		SwapLayerParameterIndicesArray(StaticParameters.MaterialLayersParameters, OriginalIndex, NewIndex);
 	}
 }
+#endif // WITH_EDITOR
 
 bool UMaterialInstance::UpdateParameters()
 {
