CIS Static Analysis fixes, protecting access to potentially null solvers during object removal

diff --git a/Engine/Source/Runtime/Engine/Private/PhysicsEngine/Experimental/PhysScene_Chaos.cpp b/Engine/Source/Runtime/Engine/Private/PhysicsEngine/Experimental/PhysScene_Chaos.cpp
--- a/Engine/Source/Runtime/Engine/Private/PhysicsEngine/Experimental/PhysScene_Chaos.cpp
+++ b/Engine/Source/Runtime/Engine/Private/PhysicsEngine/Experimental/PhysScene_Chaos.cpp
@@ -835,33 +835,40 @@ void FPhysScene_Chaos::RemoveObject(FGeometryCollectionPhysicsProxy* InObject)
 void FPhysScene_Chaos::RemoveObject(FFieldSystemPhysicsProxy* InObject)
 {
 	Chaos::FPhysicsSolver* CurrSceneSolver = InObject->GetSolver();
-	if(CurrSceneSolver && !CurrSceneSolver->UnregisterObject(InObject))
+	if(CurrSceneSolver)
 	{
-		UE_LOG(LogChaos, Warning, TEXT("Attempted to remove an object that wasn't found in its solver's gamethread storage - it's likely the solver has been mistakenly changed."));
-	}
-	RemoveFromComponentMaps(InObject);
+		if(!CurrSceneSolver->UnregisterObject(InObject))
+		{
+			UE_LOG(LogChaos, Warning, TEXT("Attempted to remove an object that wasn't found in its solver's gamethread storage - it's likely the solver has been mistakenly changed."));
+		}
+		RemoveFromComponentMaps(InObject);
 
-	if (Chaos::IDispatcher* Dispatcher = GetDispatcher())
-	{
-		TArray<Chaos::FPhysicsSolver*> SolverList = ChaosModule->GetSolversMutable(CurrSceneSolver->GetOwner());
-		
-		for(Chaos::FPhysicsSolver* Solver : SolverList)
+		if(Chaos::IDispatcher* Dispatcher = GetDispatcher())
 		{
-			if(true || Solver->HasActiveParticles())
-			{
-				Solver->RegisterObject(InObject);
+			TArray<Chaos::FPhysicsSolver*> SolverList = ChaosModule->GetSolversMutable(CurrSceneSolver->GetOwner());
 
-				if(/*bDedicatedThread && */Dispatcher)
+			for(Chaos::FPhysicsSolver* Solver : SolverList)
+			{
+				if(true || Solver->HasActiveParticles())
 				{
-					// Pass the proxy off to the physics thread
-					Dispatcher->EnqueueCommandImmediate([InObject, InSolver = Solver](Chaos::FPersistentPhysicsTask* PhysThread)
-						{
-							InSolver->UnregisterObject(InObject);
-						});
+					Solver->RegisterObject(InObject);
+
+					if(/*bDedicatedThread && */Dispatcher)
+					{
+						// Pass the proxy off to the physics thread
+						Dispatcher->EnqueueCommandImmediate([InObject, InSolver = Solver](Chaos::FPersistentPhysicsTask* PhysThread)
+							{
+								InSolver->UnregisterObject(InObject);
+							});
+					}
 				}
 			}
 		}
 	}
+	else
+	{
+		UE_LOG(LogChaos, Warning, TEXT("Attempted to remove an object but no solver had been set."));
+	}
 }
 
 #if XGE_FIXED
