Make sure the FSkeletalMeshLODModel::CopyStructure is thread safe
#rnx

diff --git a/Engine/Plugins/Experimental/SkeletalReduction/Source/Private/SkeletalMeshReductionPlugin.cpp b/Engine/Plugins/Experimental/SkeletalReduction/Source/Private/SkeletalMeshReductionPlugin.cpp
--- a/Engine/Plugins/Experimental/SkeletalReduction/Source/Private/SkeletalMeshReductionPlugin.cpp
+++ b/Engine/Plugins/Experimental/SkeletalReduction/Source/Private/SkeletalMeshReductionPlugin.cpp
@@ -1327,7 +1327,7 @@ namespace
 
 	void Empty(FSkeletalMeshLODModel& LODModel)
 	{
-		LODModel = FSkeletalMeshLODModel();
+		LODModel.Empty();
 	}
 
 }
@@ -1591,9 +1591,7 @@ void FQuadricSkeletalMeshReduction::ConvertToFSkeletalMeshLODModel( const int32
 	Options.bComputeTangents = false;
 	Options.bUseMikkTSpace = true; //Avoid builtin build by specifying true for mikkt space
 	Options.bComputeWeightedNormals = false;
-	Options.OverlappingThresholds.ThresholdPosition = 0.0f;
-	Options.OverlappingThresholds.ThresholdTangentNormal = 0.0f;
-	Options.OverlappingThresholds.ThresholdUV = 0.0f;
+	//Leave the default threshold
 	Options.bRemoveDegenerateTriangles = false;
 	IMeshUtilities& MeshUtilities = FModuleManager::Get().LoadModuleChecked<IMeshUtilities>("MeshUtilities");
 	
@@ -2007,8 +2005,8 @@ void FQuadricSkeletalMeshReduction::ReduceSkeletalMesh(USkeletalMesh& SkeletalMe
 		}
 	}
 
-	// Reduce LOD model with SrcMesh
-	if (ReduceSkeletalLODModel(*SrcModel, *NewModel, SkeletalMesh.GetImportedBounds(), SkeletalMesh.RefSkeleton, Settings, ImportantBones, RelativeToRefPoseMatrices, LODIndex, bReducingSourceModel))
+	// Reduce LOD model with SrcMesh if src mesh has more then 1 triangle
+	if (SrcModel->NumVertices > 3 && ReduceSkeletalLODModel(*SrcModel, *NewModel, SkeletalMesh.GetImportedBounds(), SkeletalMesh.RefSkeleton, Settings, ImportantBones, RelativeToRefPoseMatrices, LODIndex, bReducingSourceModel))
 	{
 		FSkeletalMeshLODInfo* ReducedLODInfoPtr = SkeletalMesh.GetLODInfo(LODIndex);
 		check(ReducedLODInfoPtr);
@@ -2072,7 +2070,6 @@ void FQuadricSkeletalMeshReduction::ReduceSkeletalMesh(USkeletalMesh& SkeletalMe
 	}
 	else
 	{
-		//TODO: This is not thread safe.
 		FSkeletalMeshLODModel::CopyStructure(NewModel, SrcModel);
 
 		// Do any joint-welding / bone removal.
diff --git a/Engine/Source/Developer/SkeletalMeshUtilitiesCommon/Private/LODUtilities.cpp b/Engine/Source/Developer/SkeletalMeshUtilitiesCommon/Private/LODUtilities.cpp
--- a/Engine/Source/Developer/SkeletalMeshUtilitiesCommon/Private/LODUtilities.cpp
+++ b/Engine/Source/Developer/SkeletalMeshUtilitiesCommon/Private/LODUtilities.cpp
@@ -1169,7 +1169,7 @@ void FLODUtilities::RestoreSkeletalMeshLODImportedData(USkeletalMesh* SkeletalMe
 		LODInfo->LODMaterialMap.Empty();
 		LODInfo->bHasBeenSimplified = false;
 		//Copy the SkeletalMeshLODModel
-		SkeletalMesh->GetImportedModel()->LODModels[LodIndex] = ImportedBaseLODModel;
+		FSkeletalMeshLODModel::CopyStructure(&(SkeletalMesh->GetImportedModel()->LODModels[LodIndex]), &ImportedBaseLODModel);
 		//Copy the morph target deltas
 		bool bInitMorphTargetData = false;
 		for (UMorphTarget *MorphTarget : SkeletalMesh->MorphTargets)
diff --git a/Engine/Source/Editor/UnrealEd/Private/Fbx/FbxSkeletalMeshImport.cpp b/Engine/Source/Editor/UnrealEd/Private/Fbx/FbxSkeletalMeshImport.cpp
--- a/Engine/Source/Editor/UnrealEd/Private/Fbx/FbxSkeletalMeshImport.cpp
+++ b/Engine/Source/Editor/UnrealEd/Private/Fbx/FbxSkeletalMeshImport.cpp
@@ -3842,7 +3842,7 @@ void UnFbx::FFbxImporter::InsertNewLODToBaseSkeletalMesh(USkeletalMesh* InSkelet
 	BaseSkeletalMesh->PreEditChange(NULL);
 
 	// Assign new FSkeletalMeshLODModel to desired slot in selected skeletal mesh.
-	DestImportedResource->LODModels[DesiredLOD] = NewLODModel;
+	FSkeletalMeshLODModel::CopyStructure(&(DestImportedResource->LODModels[DesiredLOD]), &NewLODModel);
 	//Copy the import data into the base skeletalmesh for the imported LOD
 	USkeletalMesh::CopyImportedData(0, InSkeletalMesh, DesiredLOD, BaseSkeletalMesh);
 	
diff --git a/Engine/Source/Runtime/Engine/Private/SkeletalMeshLODModel.cpp b/Engine/Source/Runtime/Engine/Private/SkeletalMeshLODModel.cpp
--- a/Engine/Source/Runtime/Engine/Private/SkeletalMeshLODModel.cpp
+++ b/Engine/Source/Runtime/Engine/Private/SkeletalMeshLODModel.cpp
@@ -20,6 +20,7 @@
 #include "UObject/FortniteMainBranchObjectVersion.h"
 #include "GPUSkinVertexFactory.h"
 #include "UObject/AnimObjectVersion.h"
+#include "Misc/ScopeLock.h"
 
 /*-----------------------------------------------------------------------------
 FSoftSkinVertex
@@ -1158,12 +1159,25 @@ void FSkeletalMeshLODModel::UpdateChunkedSectionInfo(const FString& SkeletalMesh
 	}
 }
 
-bool FSkeletalMeshLODModel::CopyStructure(FSkeletalMeshLODModel* Destination, FSkeletalMeshLODModel* Source)
+void FSkeletalMeshLODModel::CopyStructure(FSkeletalMeshLODModel* Destination, FSkeletalMeshLODModel* Source)
 {
-	if (Source->RawPointIndices.IsLocked() || Source->LegacyRawPointIndices.IsLocked() || Source->RawSkeletalMeshBulkData_DEPRECATED.GetBulkData().IsLocked())
-	{
-		return false;
-	}
+	//The private Lock should always be valid
+	check(Source);
+	check(Destination);
+	check(Source->BulkDataReadMutex);
+	check(Destination->BulkDataReadMutex);
+	//Lock both mutex before touching the bulk data
+	FScopeLock LockSource(Source->BulkDataReadMutex);
+	FScopeLock LockDestination(Destination->BulkDataReadMutex);
+
+
+	FCriticalSection* DestinationBulkDataReadMutex = Destination->BulkDataReadMutex;
+
+	//Empty the Destination BulkData to avoid leaks
+	Destination->RawPointIndices.RemoveBulkData();
+	Destination->LegacyRawPointIndices.RemoveBulkData();
+	Destination->RawSkeletalMeshBulkData_DEPRECATED.EmptyBulkData();
+
 	// Bulk data arrays need to be locked before a copy can be made.
 	Source->RawPointIndices.Lock(LOCK_READ_ONLY);
 	Source->LegacyRawPointIndices.Lock(LOCK_READ_ONLY);
@@ -1173,7 +1187,8 @@ bool FSkeletalMeshLODModel::CopyStructure(FSkeletalMeshLODModel* Destination, FS
 	Source->RawPointIndices.Unlock();
 	Source->LegacyRawPointIndices.Unlock();
 
-	return true;
+	//Make sure the mutex of the copy is set back to the original destination mutex, we can recycle the pointer.
+	Destination->BulkDataReadMutex = DestinationBulkDataReadMutex;
 }
 
 #endif // WITH_EDITOR
diff --git a/Engine/Source/Runtime/Engine/Public/Rendering/SkeletalMeshLODModel.h b/Engine/Source/Runtime/Engine/Public/Rendering/SkeletalMeshLODModel.h
--- a/Engine/Source/Runtime/Engine/Public/Rendering/SkeletalMeshLODModel.h
+++ b/Engine/Source/Runtime/Engine/Public/Rendering/SkeletalMeshLODModel.h
@@ -14,6 +14,8 @@
 #include "SkeletalMeshTypes.h"
 #include "Rendering/SkeletalMeshLODImporterData.h"
 #include "Animation/SkinWeightProfile.h"
+#include "CoreTypes.h"
+#include "HAL/CriticalSection.h"
 
 //
 //	FSoftSkinVertex
@@ -361,8 +363,40 @@ public:
 		, bIsRawSkeletalMeshBulkDataEmpty(true)
 		, BuildStringID(TEXT(""))
 	{
+		//Allocate the private mutex
+		BulkDataReadMutex = new FCriticalSection();
 	}
 
+	~FSkeletalMeshLODModel()
+	{
+		//Release the allocate resources
+		if(BulkDataReadMutex != nullptr)
+		{
+			delete BulkDataReadMutex;
+			BulkDataReadMutex = nullptr;
+		}
+	}
+
+	/*Empty the skeletal mesh LOD model. Empty copy a default constructed FSkeletalMeshLODModel but will not copy the BulkDataReadMutex which will be the same after*/
+	void Empty()
+	{
+		//We need to restore the mutex after the copy
+		FCriticalSection* BackupBulkDataReadMutex = BulkDataReadMutex;
+		*this = FSkeletalMeshLODModel();
+		BulkDataReadMutex = BackupBulkDataReadMutex;
+	}
+
+private:
+	//Mutex use by the CopyStructure function. It's a pointer because FCriticalSection privatize the operator= function, which will prevent this class operator= to use the default.
+	//We want to avoid having a custom equal operator that will get deprecated if dev forget to add the new member in this class
+	//The CopyStructure function will copy everything but make sure the destination mutex is set to a new mutex pointer.
+	FCriticalSection* BulkDataReadMutex;
+
+	//Use the static FSkeletalMeshLODModel::CopyStructure function to copy from one instance to another
+	//The reason is we want the copy to be multithread safe and use the BulkDataReadMutex.
+	FSkeletalMeshLODModel& operator=(const FSkeletalMeshLODModel& Other) = default;
+
+public:
 	/**
 	* Special serialize function passing the owning UObject along as required by FUnytpedBulkData
 	* serialization.
@@ -428,17 +462,19 @@ public:
 	* Build a derive data key with the user section data (UserSectionsData) and the original bulk data
 	*/
 	ENGINE_API FString GetLODModelDeriveDataKey() const;
-	
-	/**
-	* Copy one structure to the other, make sure all bulk data is unlock and the data can be read before copying.
-	*/
-	static ENGINE_API bool CopyStructure(FSkeletalMeshLODModel* Destination, FSkeletalMeshLODModel* Source);
 
 	/**
 	* This function will update the chunked information for each section. Only old data before the 
 	* skeletal mesh build refactor should need to call this function.
 	*/
 	ENGINE_API void UpdateChunkedSectionInfo(const FString& SkeletalMeshName);
+
+	/**
+	* Copy one structure to the other, make sure all bulk data is unlock and the data can be read before copying.
+	*
+	* It also use a private mutex to make sure it's thread safe to copy the same source multiple time in multiple thread.
+	*/
+	static ENGINE_API void CopyStructure(FSkeletalMeshLODModel* Destination, FSkeletalMeshLODModel* Source);
 };
 
 #endif // WITH_EDITOR
\ No newline at end of file
