Fix issue where new/changed screenshots could not be approved in the editor on some platforms (Mac, Xbox, others?)

The issue was that the code that captures screenshots on a device and sends them was the editor was setting a 'path' variable that represented what a screenshot *would* be saved at on the local system. (The screenshot was not actually saved, it's returned as raw data in the message).

This meant that the editor would try to save for screenshots to a local path that did not exist (e.g. g:\settings\ on xbox) and fail.

The fixes here are -

1) Add some warnings to the code that tries to save screenshots if that fails

2) Change the names of the screenshot functions in AutomationCommon to clarify what they should be used for.

3) Change the 'path' variable in the message to screenshot name.

4) Calculate the path to use for screenshots on the editor.

#tests ran EngineTest on WIn64, PS4, and XboxOne (which now passes).

diff --git a/Engine/Plugins/Tests/EditorTests/Source/EditorTests/Private/UnrealEd/EditorAssetAutomationTests.cpp b/Engine/Plugins/Tests/EditorTests/Source/EditorTests/Private/UnrealEd/EditorAssetAutomationTests.cpp
--- a/Engine/Plugins/Tests/EditorTests/Source/EditorTests/Private/UnrealEd/EditorAssetAutomationTests.cpp
+++ b/Engine/Plugins/Tests/EditorTests/Source/EditorTests/Private/UnrealEd/EditorAssetAutomationTests.cpp
@@ -962,7 +962,7 @@ namespace ImportExportAssetHelper
 			{
 				FString ScreenshotName;
 				const FString TestName = FString::Printf(TEXT("AssetImportExport/Screenshots/%s"), *ImportedAsset->GetName());
-				AutomationCommon::GetScreenshotPath(TestName, ScreenshotName);
+				ScreenshotName = AutomationCommon::GetScreenshotName(TestName);
 
 				TSharedRef<SWidget> WindowRef = ActiveWindow.ToSharedRef();
 
@@ -973,7 +973,7 @@ namespace ImportExportAssetHelper
 					FAutomationScreenshotData Data;
 					Data.Width = OutImageSize.X;
 					Data.Height = OutImageSize.Y;
-					Data.Path = ScreenshotName;
+					Data.ScreenshotName = ScreenshotName;
 					FAutomationTestFramework::Get().OnScreenshotCaptured().ExecuteIfBound(OutImageData, Data);
 				}
 
diff --git a/Engine/Plugins/Tests/EditorTests/Source/EditorTests/Private/UnrealEd/EditorAutomationTests.cpp b/Engine/Plugins/Tests/EditorTests/Source/EditorTests/Private/UnrealEd/EditorAutomationTests.cpp
--- a/Engine/Plugins/Tests/EditorTests/Source/EditorTests/Private/UnrealEd/EditorAutomationTests.cpp
+++ b/Engine/Plugins/Tests/EditorTests/Source/EditorTests/Private/UnrealEd/EditorAutomationTests.cpp
@@ -281,7 +281,7 @@ bool FLoadAllMapsInEditorTest::RunTest(const FString& Parameters)
 
 		//Create a screen shot filename and path
 		const FString LoadAllMapsTestName = FString::Printf(TEXT("LoadAllMaps_Editor/%s"), *FPaths::GetBaseFilename(MapName));
-		AutomationCommon::GetScreenshotPath(LoadAllMapsTestName, WindowParameters.ScreenshotName);
+		WindowParameters.ScreenshotName = AutomationCommon::GetScreenshotName(LoadAllMapsTestName);
 
 		//Get the current number of seconds.  This will be used to track how long it took to load the map.
 		MapLoadStartTime = FPlatformTime::Seconds();
diff --git a/Engine/Plugins/Tests/EditorTests/Source/EditorTests/Private/UnrealEd/EditorBuildPromotionTests.cpp b/Engine/Plugins/Tests/EditorTests/Source/EditorTests/Private/UnrealEd/EditorBuildPromotionTests.cpp
--- a/Engine/Plugins/Tests/EditorTests/Source/EditorTests/Private/UnrealEd/EditorBuildPromotionTests.cpp
+++ b/Engine/Plugins/Tests/EditorTests/Source/EditorTests/Private/UnrealEd/EditorBuildPromotionTests.cpp
@@ -2388,7 +2388,7 @@ bool FBuildPromotionPIETest::RunTest(const FString& Parameters)
 	}
 
 	WindowScreenshotParameters ScreenshotParams;
-	AutomationCommon::GetScreenshotPath(TEXT("EditorBuildPromotion/RunMap"), ScreenshotParams.ScreenshotName);
+	ScreenshotParams.ScreenshotName = AutomationCommon::GetScreenshotName(TEXT("EditorBuildPromotion/RunMap"));
 	ScreenshotParams.CurrentWindow = AllWindows[0];
 	//Wait for the play world to come up
 	ADD_LATENT_AUTOMATION_COMMAND(FWaitLatentCommand(1.f));
diff --git a/Engine/Plugins/Tests/RuntimeTests/Source/RuntimeTests/Private/Shaders/ShaderComparisonTests.cpp b/Engine/Plugins/Tests/RuntimeTests/Source/RuntimeTests/Private/Shaders/ShaderComparisonTests.cpp
--- a/Engine/Plugins/Tests/RuntimeTests/Source/RuntimeTests/Private/Shaders/ShaderComparisonTests.cpp
+++ b/Engine/Plugins/Tests/RuntimeTests/Source/RuntimeTests/Private/Shaders/ShaderComparisonTests.cpp
@@ -144,14 +144,14 @@ bool FCompareBasepassShaders::RunTest(const FString& Parameters)
     for (int i = 0; i < NUM_CAMERAS; ++i)
     {
         FString CompareBasepassShadersTestName = FString::Printf(TEXT("CompareBasepassShaders_Game/%s/%d"), *FPaths::GetBaseFilename(MapName), i);
-        AutomationCommon::GetScreenshotPath(CompareBasepassShadersTestName, ScreenshotFileName[i]);
+		ScreenshotFileName[i] = AutomationCommon::GetScreenshotName(CompareBasepassShadersTestName);
         CompareBasepassShadersTestName = FString::Printf(TEXT("Incoming/CompareBasepassShaders_Game/%s/%d"), *FPaths::GetBaseFilename(MapName), i);
-        AutomationCommon::GetScreenshotPath(CompareBasepassShadersTestName, RealScreenshotFileName[i]);
+		RealScreenshotFileName[i] = AutomationCommon::GetScreenshotName(CompareBasepassShadersTestName);
         RealScreenshotFileName[i] = "../../../" + RealScreenshotFileName[i];
         CompareBasepassShadersTestName = FString::Printf(TEXT("CompareBasepassShaders_Game/%s_fp16/%d"), *FPaths::GetBaseFilename(MapName), i);
-        AutomationCommon::GetScreenshotPath(CompareBasepassShadersTestName, FP16ScreenshotFileName[i]);
+		FP16ScreenshotFileName[i] = AutomationCommon::GetScreenshotName(CompareBasepassShadersTestName);
         CompareBasepassShadersTestName = FString::Printf(TEXT("Incoming/CompareBasepassShaders_Game/%s_fp16/%d"), *FPaths::GetBaseFilename(MapName), i);
-        AutomationCommon::GetScreenshotPath(CompareBasepassShadersTestName, RealFP16ScreenshotFileName[i]);
+		RealFP16ScreenshotFileName[i] = AutomationCommon::GetScreenshotName(CompareBasepassShadersTestName);
         RealFP16ScreenshotFileName[i] = "../../../" + RealFP16ScreenshotFileName[i];
     }
 
diff --git a/Engine/Source/Developer/AutomationController/Private/AutomationControllerManger.cpp b/Engine/Source/Developer/AutomationController/Private/AutomationControllerManger.cpp
--- a/Engine/Source/Developer/AutomationController/Private/AutomationControllerManger.cpp
+++ b/Engine/Source/Developer/AutomationController/Private/AutomationControllerManger.cpp
@@ -1036,8 +1036,20 @@ void FAutomationControllerManager::HandleReceivedScreenShot(const FAutomationWor
 
 	bool bTree = true;
 	FString FileName = ScreenshotIncomingFolder / Message.ScreenShotName;
-	IFileManager::Get().MakeDirectory(*FPaths::GetPath(FileName), bTree);
-	FFileHelper::SaveArrayToFile(Message.ScreenImage, *FileName);
+
+	FString DirectoryPath = FPaths::GetPath(FileName);
+
+	if (!IFileManager::Get().MakeDirectory(*DirectoryPath, bTree))
+	{
+		UE_LOG(LogAutomationController, Error, TEXT("Failed to create directory %s for incoming screenshot"), *DirectoryPath);
+		return;
+	}
+	
+	if (!FFileHelper::SaveArrayToFile(Message.ScreenImage, *FileName))
+	{
+		UE_LOG(LogAutomationController, Error, TEXT("Failed to save screenshot to %s"), *FileName);
+		return;
+	}
 
 	// TODO Automation There is identical code in, Engine\Source\Runtime\AutomationWorker\Private\AutomationWorkerModule.cpp,
 	// need to move this code into common area.
diff --git a/Engine/Source/Developer/FunctionalTesting/Private/AutomationBlueprintFunctionLibrary.cpp b/Engine/Source/Developer/FunctionalTesting/Private/AutomationBlueprintFunctionLibrary.cpp
--- a/Engine/Source/Developer/FunctionalTesting/Private/AutomationBlueprintFunctionLibrary.cpp
+++ b/Engine/Source/Developer/FunctionalTesting/Private/AutomationBlueprintFunctionLibrary.cpp
@@ -443,7 +443,7 @@ public:
 
 			bool bAttemptToCompareShot = FAutomationTestFramework::Get().OnScreenshotCaptured().ExecuteIfBound(InImageData, Data);
 
-			UE_LOG(AutomationFunctionLibrary, Log, TEXT("Screenshot captured as %s"), *Data.Path);
+			UE_LOG(AutomationFunctionLibrary, Log, TEXT("Screenshot captured as %s"), *Data.ScreenshotName);
 
 			if ( GIsAutomationTesting )
 			{
@@ -547,7 +547,7 @@ public:
 
 		bool bAttemptToCompareShot = FAutomationTestFramework::Get().OnScreenshotCaptured().ExecuteIfBound(InImageData, Data);
 
-		UE_LOG(AutomationFunctionLibrary, Log, TEXT("Screenshot captured as %s"), *Data.Path);
+		UE_LOG(AutomationFunctionLibrary, Log, TEXT("Screenshot captured as %s"), *Data.ScreenshotName);
 
 		FAutomationTestFramework::Get().OnScreenshotCompared.AddRaw(this, &FAutomationHighResScreenshotGrabber::OnComparisonComplete);
 	}
diff --git a/Engine/Source/Developer/FunctionalTesting/Private/ScreenshotFunctionalTestBase.cpp b/Engine/Source/Developer/FunctionalTesting/Private/ScreenshotFunctionalTestBase.cpp
--- a/Engine/Source/Developer/FunctionalTesting/Private/ScreenshotFunctionalTestBase.cpp
+++ b/Engine/Source/Developer/FunctionalTesting/Private/ScreenshotFunctionalTestBase.cpp
@@ -150,7 +150,7 @@ void AScreenshotFunctionalTestBase::OnScreenShotCaptured(int32 InSizeX, int32 In
 
 	FAutomationTestFramework::Get().OnScreenshotCaptured().ExecuteIfBound(InImageData, Data);
 
-	UE_LOG(LogScreenshotFunctionalTest, Log, TEXT("Screenshot captured as %s"), *Data.Path);
+	UE_LOG(LogScreenshotFunctionalTest, Log, TEXT("Screenshot captured as %s"), *Data.ScreenshotName);
 #endif
 }
 
diff --git a/Engine/Source/Runtime/AutomationWorker/Private/AutomationWorkerModule.cpp b/Engine/Source/Runtime/AutomationWorker/Private/AutomationWorkerModule.cpp
--- a/Engine/Source/Runtime/AutomationWorker/Private/AutomationWorkerModule.cpp
+++ b/Engine/Source/Runtime/AutomationWorker/Private/AutomationWorkerModule.cpp
@@ -28,6 +28,8 @@
 
 #define LOCTEXT_NAMESPACE "AutomationTest"
 
+DEFINE_LOG_CATEGORY_STATIC(LogAutomationWorker, Log, All);
+
 IMPLEMENT_MODULE(FAutomationWorkerModule, AutomationWorker);
 
 
@@ -425,23 +427,40 @@ void FAutomationWorkerModule::HandleScreenShotCapturedWithName(const TArray<FCol
 	{
 		FAutomationWorkerScreenImage* Message = new FAutomationWorkerScreenImage();
 
-		Message->ScreenShotName = FPaths::ProjectDir() / Data.Path;
-		FPaths::MakePathRelativeTo(Message->ScreenShotName, *FPaths::AutomationDir());
+		Message->ScreenShotName = Data.ScreenshotName;
 		Message->ScreenImage = CompressedBitmap;
 		Message->Metadata = Metadata;
+
+		UE_LOG(LogAutomationWorker, Log, TEXT("Sending screenshot %s"), *Message->ScreenShotName);
+
 		MessageEndpoint->Send(Message, TestRequesterAddress);
 	}
 	else
 	{
 		//Save locally
 		const bool bTree = true;
-		IFileManager::Get().MakeDirectory(*FPaths::GetPath(Data.Path), bTree);
-		FFileHelper::SaveArrayToFile(CompressedBitmap, *Data.Path);
+
+		FString LocalFile = AutomationCommon::GetLocalPathForScreenshot(Data.ScreenshotName);
+		FString DirectoryPath = FPaths::GetPath(LocalFile);
+
+		UE_LOG(LogAutomationWorker, Log, TEXT("Saving screenshot %s as %s"),*Data.ScreenshotName, *LocalFile);
+
+		if (!IFileManager::Get().MakeDirectory(*DirectoryPath, bTree))
+		{
+			UE_LOG(LogAutomationWorker, Error, TEXT("Failed to create directory %s for incoming screenshot"), *DirectoryPath);
+			return;
+		}
+
+		if (!FFileHelper::SaveArrayToFile(CompressedBitmap, *LocalFile))
+		{
+			UE_LOG(LogAutomationWorker, Error, TEXT("Failed to save screenshot to %s"), *LocalFile);
+			return;
+		}
 
 		FString Json;
 		if ( FJsonObjectConverter::UStructToJsonObjectString(Metadata, Json) )
 		{
-			FString MetadataPath = FPaths::ChangeExtension(Data.Path, TEXT("json"));
+			FString MetadataPath = FPaths::ChangeExtension(LocalFile, TEXT("json"));
 			FFileHelper::SaveStringToFile(Json, *MetadataPath, FFileHelper::EEncodingOptions::ForceUTF8WithoutBOM);
 		}
 	}
diff --git a/Engine/Source/Runtime/Core/Public/Misc/AutomationTest.h b/Engine/Source/Runtime/Core/Public/Misc/AutomationTest.h
--- a/Engine/Source/Runtime/Core/Public/Misc/AutomationTest.h
+++ b/Engine/Source/Runtime/Core/Public/Misc/AutomationTest.h
@@ -594,7 +594,8 @@ struct FAutomationScreenshotData
 	bool bIgnoreAntiAliasing;
 	bool bIgnoreColors;
 
-	FString Path;
+	// Name of the screenshot generated from AutomationCommon::GetScreenShotName()
+	FString ScreenshotName;
 
 	FAutomationScreenshotData()
 		: Id()
diff --git a/Engine/Source/Runtime/Engine/Private/Tests/AutomationCommon.cpp b/Engine/Source/Runtime/Engine/Private/Tests/AutomationCommon.cpp
--- a/Engine/Source/Runtime/Engine/Private/Tests/AutomationCommon.cpp
+++ b/Engine/Source/Runtime/Engine/Private/Tests/AutomationCommon.cpp
@@ -82,16 +82,19 @@ namespace AutomationCommon
 #if (WITH_DEV_AUTOMATION_TESTS || WITH_PERF_AUTOMATION_TESTS)
 
 	/** Gets a path used for automation testing (PNG sent to the AutomationTest folder) */
-	void GetScreenshotPath(const FString& TestName, FString& OutScreenshotName)
+	FString GetScreenshotName(const FString& TestName)
 	{
-		FString PathName = FPaths::AutomationDir() + TestName / FPlatformProperties::IniPlatformName();
+		FString PathName = TestName / FPlatformProperties::IniPlatformName();
 		PathName = PathName + TEXT("/") + GetRenderDetailsString();
 
-		FPaths::MakePathRelativeTo(PathName, *FPaths::ProjectDir());
-
-		OutScreenshotName = FString::Printf(TEXT("%s/%s.png"), *PathName, *FPlatformMisc::GetDeviceId());
+		return FString::Printf(TEXT("%s/%s.png"), *PathName, *FPlatformMisc::GetDeviceId());
 	}
 
+	FString GetLocalPathForScreenshot(const FString& InScreenshotName)
+	{
+		return FPaths::AutomationDir() + InScreenshotName;
+	}
+	
 	FAutomationScreenshotData BuildScreenshotData(const FString& MapOrContext, const FString& TestName, int32 Width, int32 Height)
 	{
 		FAutomationScreenshotData Data;
@@ -131,14 +134,14 @@ namespace AutomationCommon
 		// Device's native resolution (we want to use a hardware dump of the frontbuffer at the native resolution so we compare what we actually output rather than what we think we rendered)
 
 		const FString MapAndTest = MapOrContext + TEXT("/") + Data.Name;
-		AutomationCommon::GetScreenshotPath(MapAndTest, Data.Path);
+		Data.ScreenshotName = GetScreenshotName(MapAndTest);
 
 		return Data;
 	}
 #endif
 
 	/** These save a PNG and get sent over the network */
-	static void SaveWindowAsScreenshot(TSharedRef<SWindow> Window, const FString& FileName)
+	static void SaveWindowAsScreenshot(TSharedRef<SWindow> Window, const FString& ScreenshotName)
 	{
 		TSharedRef<SWidget> WindowRef = Window;
 
@@ -149,7 +152,7 @@ namespace AutomationCommon
 			FAutomationScreenshotData Data;
 			Data.Width = OutImageSize.X;
 			Data.Height = OutImageSize.Y;
-			Data.Path = FileName;
+			Data.ScreenshotName = ScreenshotName;
 			FAutomationTestFramework::Get().OnScreenshotCaptured().ExecuteIfBound(OutImageData, Data);
 		}
 	}
diff --git a/Engine/Source/Runtime/Engine/Private/Tests/EngineAutomationTests.cpp b/Engine/Source/Runtime/Engine/Private/Tests/EngineAutomationTests.cpp
--- a/Engine/Source/Runtime/Engine/Private/Tests/EngineAutomationTests.cpp
+++ b/Engine/Source/Runtime/Engine/Private/Tests/EngineAutomationTests.cpp
@@ -229,7 +229,7 @@ bool FLoadAllMapsInGameTest::RunTest(const FString& Parameters)
 		//Generate the screen shot name and path
 		FString ScreenshotFileName;
 		const FString LoadAllMapsTestName = FString::Printf(TEXT("LoadAllMaps_Game/%s"), *FPaths::GetBaseFilename(MapName));
-		AutomationCommon::GetScreenshotPath(LoadAllMapsTestName, ScreenshotFileName);
+		ScreenshotFileName = AutomationCommon::GetLocalPathForScreenshot(AutomationCommon::GetScreenshotName(LoadAllMapsTestName));
 
 		//Give the map some time to load
 		ADD_LATENT_AUTOMATION_COMMAND(FEngineWaitLatentCommand(1.5f));
diff --git a/Engine/Source/Runtime/Engine/Public/Tests/AutomationCommon.h b/Engine/Source/Runtime/Engine/Public/Tests/AutomationCommon.h
--- a/Engine/Source/Runtime/Engine/Public/Tests/AutomationCommon.h
+++ b/Engine/Source/Runtime/Engine/Public/Tests/AutomationCommon.h
@@ -28,8 +28,18 @@ namespace AutomationCommon
 
 #if (WITH_DEV_AUTOMATION_TESTS || WITH_PERF_AUTOMATION_TESTS)
 
-	/** Gets a path used for automation testing (PNG sent to the AutomationTest folder) */
-	ENGINE_API void GetScreenshotPath(const FString& TestName, FString& OutScreenshotName);
+	/** Gets a name to be used for this screenshot.  This will return something like 
+		TestName/PlatformName/DeviceName.png. It's important to understand that a screenshot
+		generated on a device will likely have a different absolute path than the editor so this
+		name should be used with	*/
+	ENGINE_API FString GetScreenshotName(const FString& TestName);
+
+	/** 
+	This function takes the result of GetScreenshotName and will return the complete path to where a
+	screenshot can/should be found on the local device. This cannot reliably be used when communicating between
+	the editor and a test worker!
+	*/
+	ENGINE_API FString GetLocalPathForScreenshot(const FString& InScreenshotName);
 
 	ENGINE_API FAutomationScreenshotData BuildScreenshotData(const FString& MapOrContext, const FString& TestName, int32 Width, int32 Height);
 
