Add CSV Profiler events that can track when a Connection's Actor (or its Pawn) are stuck on network loads for a long time.

diff --git a/Engine/Source/Runtime/Engine/Classes/Engine/PackageMapClient.h b/Engine/Source/Runtime/Engine/Classes/Engine/PackageMapClient.h
--- a/Engine/Source/Runtime/Engine/Classes/Engine/PackageMapClient.h
+++ b/Engine/Source/Runtime/Engine/Classes/Engine/PackageMapClient.h
@@ -317,6 +317,10 @@ private:
 
 		FNetworkGUID NetGUID;
 		double RequestStartTime;
+
+#if CSV_PROFILER
+		bool bWasRequestedByOwnerOrPawn = false;
+#endif
 	};
 
 	/** Set of packages that are currently pending Async loads, referenced by package name. */
@@ -336,6 +340,53 @@ private:
 	 * pending guids. 
 	 */
 	TMap<FNetworkGUID, TWeakPtr<FQueuedBunchObjectReference>> QueuedBunchObjectReferences;
+
+#if CSV_PROFILER
+public:
+
+	bool IsTrackingOwnerOrPawn() const;
+
+	struct FIsOwnerOrPawnHelper
+	{
+	private:
+		friend class UActorChannel;
+
+		FIsOwnerOrPawnHelper(
+			class FNetGUIDCache* const InGuidCache,
+			const class AActor* InConnectionActor,
+			const class AActor* ChannelActor);
+
+	public:
+
+		~FIsOwnerOrPawnHelper();
+
+		bool IsOwnerOrPawn() const;
+
+	private:
+
+		class FNetGUIDCache* const GuidCache;
+		const class AActor* ConnectionActor;
+		const class AActor* ChannelActor;
+		mutable int8 CachedResult = INDEX_NONE;
+
+		// No copying, moving, or constructing off stack.
+		// References or pointers to this class should never be kept alive.
+		FIsOwnerOrPawnHelper(const FIsOwnerOrPawnHelper&) = delete;
+		FIsOwnerOrPawnHelper(FIsOwnerOrPawnHelper&&) = delete;
+
+		FIsOwnerOrPawnHelper& operator=(const FIsOwnerOrPawnHelper&) = delete;
+		FIsOwnerOrPawnHelper& operator=(FIsOwnerOrPawnHelper&&) = delete;
+
+		void* operator new (size_t) = delete;
+		void* operator new[](size_t) = delete;
+		void operator delete (void *) = delete;
+		void operator delete[](void*) = delete;
+	};
+
+private:
+
+	FIsOwnerOrPawnHelper* TrackingOwnerOrPawnHelper = nullptr;
+#endif
 };
 
 class ENGINE_API FPackageMapAckState
@@ -527,4 +578,4 @@ public:
 private:
 
 	FNetQueuedActorDelinquencyAnalytics DelinquentQueuedActors;
-};
+};
\ No newline at end of file
diff --git a/Engine/Source/Runtime/Engine/Private/DataChannel.cpp b/Engine/Source/Runtime/Engine/Private/DataChannel.cpp
--- a/Engine/Source/Runtime/Engine/Private/DataChannel.cpp
+++ b/Engine/Source/Runtime/Engine/Private/DataChannel.cpp
@@ -2389,6 +2389,10 @@ bool UActorChannel::ProcessQueuedBunches()
 
 			if (UPackageMapClient * PackageMapClient = Cast< UPackageMapClient >(Connection->PackageMap))
 			{
+#if CSV_PROFILER
+				FNetGUIDCache::FIsOwnerOrPawnHelper Helper(Connection->Driver->GuidCache.Get(), Connection->OwningActor, Actor);
+#endif
+
 				PackageMapClient->SetHasQueuedBunches(ActorNetGUID, false);
 			}
 
@@ -2464,9 +2468,12 @@ void UActorChannel::ReceivedBunch( FInBunch & Bunch )
 			QueuedObjectsToTrack.SetNum(NumMustBeMappedGUIDs);
 			//UE_LOG( LogNetTraffic, Warning, TEXT( "Read must be mapped GUID's. NumMustBeMappedGUIDs: %i" ), NumMustBeMappedGUIDs );
 
-			UPackageMapClient* PackageMapClient = CastChecked<UPackageMapClient>(Connection->PackageMap);
 			FNetGUIDCache* GuidCache = Connection->Driver->GuidCache.Get();
 
+#if CSV_PROFILER
+			FNetGUIDCache::FIsOwnerOrPawnHelper Helper(GuidCache, Connection->OwningActor, Actor);
+#endif
+
 			for (int32 i = 0; i < NumMustBeMappedGUIDs; i++)
 			{
 				FNetworkGUID NetGUID;
@@ -2474,7 +2481,7 @@ void UActorChannel::ReceivedBunch( FInBunch & Bunch )
 
 				// If we have async package map loading disabled, we have to ignore NumMustBeMappedGUIDs
 				//	(this is due to the fact that async loading could have been enabled on the server side)
-				if (!Connection->Driver->GuidCache->ShouldAsyncLoad())
+				if (!GuidCache->ShouldAsyncLoad())
 				{
 					continue;
 				}
diff --git a/Engine/Source/Runtime/Engine/Private/PackageMapClient.cpp b/Engine/Source/Runtime/Engine/Private/PackageMapClient.cpp
--- a/Engine/Source/Runtime/Engine/Private/PackageMapClient.cpp
+++ b/Engine/Source/Runtime/Engine/Private/PackageMapClient.cpp
@@ -27,6 +27,7 @@
 #include "Serialization/MemoryWriter.h"
 #include "Components/ChildActorComponent.h"
 #include "Net/NetworkGranularMemoryLogging.h"
+#include "GameFramework/Controller.h"
 
 #if WITH_EDITOR
 #include "UObject/ObjectRedirector.h"
@@ -39,6 +40,8 @@
 static const int GUID_PACKET_NOT_ACKED	= -2;		
 static const int GUID_PACKET_ACKED		= -1;		
 
+CSV_DEFINE_CATEGORY(PackageMap, true);
+
 /**
  * Don't allow infinite recursion of InternalLoadObject - an attacker could
  * send malicious packets that cause a stack overflow on the server.
@@ -51,22 +54,38 @@ static TAutoConsoleVariable<int32> CVarReservedNetGuidSize(TEXT("net.ReservedNet
 extern FAutoConsoleVariableRef CVarEnableMultiplayerWorldOriginRebasing;
 extern TAutoConsoleVariable<int32> CVarFilterGuidRemapping;
 
-static float GGuidCacheTrackAsyncLoadingGUIDTreshold = 0.f;
+static float GGuidCacheTrackAsyncLoadingGUIDThreshold = 0.f;
 static FAutoConsoleVariableRef CVarTrackAsyncLoadingGUIDTreshold(
 	TEXT("net.TrackAsyncLoadingGUIDThreshold"),
-	GGuidCacheTrackAsyncLoadingGUIDTreshold,
+	GGuidCacheTrackAsyncLoadingGUIDThreshold,
 	TEXT("When > 0, any objects that take longer than the threshold to async load will be tracked."
 		" Threshold in seconds, @see FNetGUIDCache::ConsumeDelinquencyAnalytics. Used for Debugging and Analytics")
 );
 
-static float GPackageMapTrackQueuedActorTreshold = 0.f;
+static float GGuidCacheTrackAsyncLoadingGUIDThresholdOwner = 0.f;
+static FAutoConsoleVariableRef CVarTrackAsyncLoadingGUIDThresholdOwner(
+	TEXT("net.TrackAsyncLoadingGUIDThresholdOwner"),
+	GGuidCacheTrackAsyncLoadingGUIDThresholdOwner,
+	TEXT("When > 0, if the Net Connection's owning Controller or Pawn is waiting on Async Loads for longer than this"
+		" threshold, we will fire a CSV Event to track it. Used for Debugging and Profiling")
+);
+
+static float GPackageMapTrackQueuedActorThreshold = 0.f;
 static FAutoConsoleVariableRef CVarTrackQueuedActorTreshold(
 	TEXT("net.TrackQueuedActorThreshold"),
-	GPackageMapTrackQueuedActorTreshold,
+	GPackageMapTrackQueuedActorThreshold,
 	TEXT("When > 0, any actors that spend longer than the threshold with queued bunches will be tracked."
 		" Threshold in seconds, @see UPackageMap::ConsumeDelinquencyAnalytics. Used for Debugging and Analytics")
 );
 
+static float GPackageMapTrackQueuedActorThresholdOwner = 0.f;
+static FAutoConsoleVariableRef CVarTrackQueuedActorOwnerThreshold(
+	TEXT("net.TrackQueuedActorThresholdOwner"),
+	GPackageMapTrackQueuedActorThresholdOwner,
+	TEXT("When > 0, if the Net Connection's owning Controller or Pawn has Queued Bunches for longer than this"
+		" threshold, we will fire a CSV Event to track it. Used for Debugging and Profiling")
+);
+
 static int32 GDelinquencyNumberOfTopOffendersToTrack = 10;
 static FAutoConsoleVariableRef CVarDelinquencyNumberOfTopOffendersToTrack(
 	TEXT("net.DelinquencyNumberOfTopOffendersToTrack"),
@@ -2217,7 +2236,7 @@ void UPackageMapClient::SetHasQueuedBunches(const FNetworkGUID& NetGUID, bool bH
 {
 	if (bHasQueuedBunches)
 	{
-		if (GPackageMapTrackQueuedActorTreshold > 0.f)
+		if (GPackageMapTrackQueuedActorThreshold > 0.f)
 		{
 			if (UNetDriver const * const NetDriver = ((Connection) ? Connection->GetDriver() : nullptr))
 			{
@@ -2231,24 +2250,47 @@ void UPackageMapClient::SetHasQueuedBunches(const FNetworkGUID& NetGUID, bool bH
 	{
 		double StartTime = 0.f;
 
+		const bool bNormalQueueEnabled = GPackageMapTrackQueuedActorThreshold > 0.f;
+		
+#if CSV_PROFILER		
+		const bool bOwnerQueueEnabled = GPackageMapTrackQueuedActorThresholdOwner > 0.f;
+#else
+		constexpr bool bOwnerQueueEnabled = false;
+#endif		
+
 		// We try to remove the value regardless of whether or not the CVar is on.
 		// That way if it's toggled on and off, we don't end up wasting resources.
 		// If it is disabled with entries in DelinquentQueuedActors, it will be up
 		// to clients to clear out the map by calling ConsumeDelinquencyAnalytics.
 		if (CurrentQueuedBunchNetGUIDs.RemoveAndCopyValue(NetGUID, StartTime) &&
-			GPackageMapTrackQueuedActorTreshold &&
+			(bNormalQueueEnabled || bOwnerQueueEnabled) &&
 			GuidCache)
 		{
 			if (UNetDriver const * const NetDriver = ((Connection) ? Connection->GetDriver() : nullptr))
 			{
 				const double QueuedTime = NetDriver->GetElapsedTime() - StartTime;
-				if (QueuedTime > GPackageMapTrackQueuedActorTreshold)
+				const bool bAboveNormalQueuedTime = bNormalQueueEnabled && QueuedTime > GPackageMapTrackQueuedActorThreshold;
+				const bool bAboveOwnerQueuedTime = bOwnerQueueEnabled && QueuedTime > GPackageMapTrackQueuedActorThresholdOwner;
+
+				if (bAboveNormalQueuedTime || bAboveOwnerQueuedTime)
 				{
 					if (FNetGuidCacheObject const * const CacheObject = GuidCache->ObjectLookup.Find(NetGUID))
 					{
 						if (UObject const * const Object = CacheObject->Object.Get())
 						{
-							DelinquentQueuedActors.DelinquentQueuedActors.Emplace(Object->GetClass()->GetFName(), QueuedTime);
+							const FName ObjectClass = Object->GetClass()->GetFName();
+
+							if (bAboveNormalQueuedTime)
+							{
+								DelinquentQueuedActors.DelinquentQueuedActors.Emplace(ObjectClass, QueuedTime);
+							}
+
+#if CSV_PROFILER
+							if (bAboveOwnerQueuedTime && GuidCache->IsTrackingOwnerOrPawn())
+							{
+								CSV_EVENT(PackageMap, TEXT("Owner Net Stall Queued Actor (QueueTime=%.2f)"), *ObjectClass.ToString(), QueuedTime);
+							}
+#endif
 						}
 					}
 				}
@@ -2878,6 +2920,10 @@ PRAGMA_ENABLE_DEPRECATION_WARNINGS
 	{
 		UE_LOG(LogNetPackageMap, Log, TEXT("ValidateAsyncLoadingPackage: Already async loading package. Path: %s, NetGUID: %s"), *CacheObject.PathName.ToString(), *NetGUID.ToString());
 	}
+	
+#if CSV_PROFILER
+	PendingLoadRequest.bWasRequestedByOwnerOrPawn |= IsTrackingOwnerOrPawn();
+#endif
 }
 
 void FNetGUIDCache::StartAsyncLoadingPackage(FNetGuidCacheObject& CacheObject, const FNetworkGUID NetGUID, const bool bWasAlreadyAsyncLoading)
@@ -2887,7 +2933,13 @@ PRAGMA_DISABLE_DEPRECATION_WARNINGS
 	PendingAsyncPackages.Add(CacheObject.PathName, NetGUID);
 PRAGMA_ENABLE_DEPRECATION_WARNINGS
 
-	PendingAsyncLoadRequests.Emplace(CacheObject.PathName, FPendingAsyncLoadRequest(NetGUID, Driver->GetElapsedTime()));
+	FPendingAsyncLoadRequest LoadRequest(NetGUID, Driver->Driver->GetElapsedTime());
+	
+#if CSV_PROFILER
+	LoadRequest.bWasRequestedByOwnerOrPawn = IsTrackingOwnerOrPawn();
+#endif
+
+	PendingAsyncLoadRequests.Emplace(CacheObject.PathName, MoveTemp(LoadRequest));
 
 	DelinquentAsyncLoads.MaxConcurrentAsyncLoads = FMath::Max<uint32>(DelinquentAsyncLoads.MaxConcurrentAsyncLoads, PendingAsyncLoadRequests.Num());
 
@@ -2939,12 +2991,22 @@ void FNetGUIDCache::AsyncPackageCallback(const FName& PackageName, UPackage * Pa
 		// This won't be the exact amount of time that we spent loading the package, but should
 		// give us a close enough estimate (within a frame time).
 		const double LoadTime = (Driver->GetElapsedTime() - PendingLoadRequest->RequestStartTime);
-		if (GGuidCacheTrackAsyncLoadingGUIDTreshold > 0.f &&
-			LoadTime >= GGuidCacheTrackAsyncLoadingGUIDTreshold)
+		if (GGuidCacheTrackAsyncLoadingGUIDThreshold > 0.f &&
+			LoadTime >= GGuidCacheTrackAsyncLoadingGUIDThreshold)
 		{
 			DelinquentAsyncLoads.DelinquentAsyncLoads.Emplace(PackageName, LoadTime);
 		}
 
+#if CSV_PROFILER
+		if (PendingLoadRequest->bWasRequestedByOwnerOrPawn &&
+			GGuidCacheTrackAsyncLoadingGUIDThresholdOwner > 0.f &&
+			LoadTime >= GGuidCacheTrackAsyncLoadingGUIDThresholdOwner &&
+			Driver->ServerConnection)
+		{
+			CSV_EVENT(PackageMap, TEXT("Owner Net Stall Async Load (Package=%s|LoadTime=%.2f)"), *PackageName.ToString(), LoadTime);
+		}
+#endif
+
 		PendingAsyncLoadRequests.Remove(PackageName);
 
 PRAGMA_DISABLE_DEPRECATION_WARNINGS
@@ -3716,3 +3778,54 @@ FAutoConsoleCommand	ListNetGUIDExportsCommand(
 	);
 
 // ----------------------------------------------------------------
+
+#if CSV_PROFILER
+bool FNetGUIDCache::IsTrackingOwnerOrPawn() const
+{
+	return TrackingOwnerOrPawnHelper && TrackingOwnerOrPawnHelper->IsOwnerOrPawn();
+}
+
+FNetGUIDCache::FIsOwnerOrPawnHelper::FIsOwnerOrPawnHelper(
+	FNetGUIDCache* InGuidCache,
+	const AActor* InConnectionActor,
+	const AActor* InChannelActor)
+
+	: GuidCache(InGuidCache)
+	, ConnectionActor(InConnectionActor)
+	, ChannelActor(InChannelActor)
+{
+	if (GuidCache)
+	{
+		GuidCache->TrackingOwnerOrPawnHelper = this;
+	}
+}
+
+FNetGUIDCache::FIsOwnerOrPawnHelper::~FIsOwnerOrPawnHelper()
+{
+	if (GuidCache)
+	{
+		GuidCache->TrackingOwnerOrPawnHelper = nullptr;
+	}
+}
+
+bool FNetGUIDCache::FIsOwnerOrPawnHelper::IsOwnerOrPawn() const
+{
+	if (CachedResult == INDEX_NONE)
+	{
+		bool bIsOwnerOrPawn = false;
+		if (ChannelActor && ConnectionActor)
+		{
+			bIsOwnerOrPawn = (ChannelActor == ConnectionActor);
+			if (!bIsOwnerOrPawn)
+			{
+				const AController* Controller = Cast<AController>(ConnectionActor);
+				bIsOwnerOrPawn = Controller ? (Controller->GetPawn() == ChannelActor) : false;
+			}
+		}
+
+		CachedResult = bIsOwnerOrPawn ? 1 : 0;
+	}
+
+	return !!CachedResult;
+}
+#endif
\ No newline at end of file
