Fix typo from bad merge.

#rnx

diff --git a/Engine/Source/Runtime/Engine/Private/PackageMapClient.cpp b/Engine/Source/Runtime/Engine/Private/PackageMapClient.cpp
--- a/Engine/Source/Runtime/Engine/Private/PackageMapClient.cpp
+++ b/Engine/Source/Runtime/Engine/Private/PackageMapClient.cpp
@@ -2933,7 +2933,7 @@ PRAGMA_DISABLE_DEPRECATION_WARNINGS
 	PendingAsyncPackages.Add(CacheObject.PathName, NetGUID);
 PRAGMA_ENABLE_DEPRECATION_WARNINGS
 
-	FPendingAsyncLoadRequest LoadRequest(NetGUID, Driver->Driver->GetElapsedTime());
+	FPendingAsyncLoadRequest LoadRequest(NetGUID, Driver->GetElapsedTime());
 	
 #if CSV_PROFILER
 	LoadRequest.bWasRequestedByOwnerOrPawn = IsTrackingOwnerOrPawn();
