Don't default construct FMaterialLayersFunctions with a new guid value, breaks check for unitialized members which assumes 2 default-constructed structs will be equal

diff --git a/Engine/Source/Runtime/Engine/Classes/Materials/MaterialLayersFunctions.h b/Engine/Source/Runtime/Engine/Classes/Materials/MaterialLayersFunctions.h
--- a/Engine/Source/Runtime/Engine/Classes/Materials/MaterialLayersFunctions.h
+++ b/Engine/Source/Runtime/Engine/Classes/Materials/MaterialLayersFunctions.h
@@ -168,6 +168,7 @@ struct ENGINE_API FMaterialLayersFunctions
 
 	static const FGuid UninitializedParentGuid;
 	static const FGuid NoParentGuid;
+	static const FGuid BackgroundGuid;
 		
 	FMaterialLayersFunctions()
 	{
@@ -178,7 +179,10 @@ struct ENGINE_API FMaterialLayersFunctions
 		FText LayerName = FText(LOCTEXT("Background", "Background"));
 		LayerNames.Add(LayerName);
 		RestrictToLayerRelatives.Add(false);
-		LayerGuids.Add(FGuid::NewGuid());
+		// Use a consistent Guid for the background layer
+		// This layer never needs to resolve, so doesn't need to be unique
+		// Default constructor assigning different guids will break FStructUtils::AttemptToFindUninitializedScriptStructMembers
+		LayerGuids.Add(BackgroundGuid);
 		ParentLayerGuids.Add(NoParentGuid);
 #endif
 	}
diff --git a/Engine/Source/Runtime/Engine/Private/Materials/MaterialExpressions.cpp b/Engine/Source/Runtime/Engine/Private/Materials/MaterialExpressions.cpp
--- a/Engine/Source/Runtime/Engine/Private/Materials/MaterialExpressions.cpp
+++ b/Engine/Source/Runtime/Engine/Private/Materials/MaterialExpressions.cpp
@@ -12460,6 +12460,7 @@ void FMaterialLayersFunctions::ID::AppendKeyString(FString& KeyString) const
 
 const FGuid FMaterialLayersFunctions::UninitializedParentGuid(0u, 0u, 0u, 0u);
 const FGuid FMaterialLayersFunctions::NoParentGuid(1u, 0u, 0u, 0u);
+const FGuid FMaterialLayersFunctions::BackgroundGuid(2u, 0u, 0u, 0u);
 
 const FMaterialLayersFunctions::ID FMaterialLayersFunctions::GetID() const
 {
@@ -12548,10 +12549,17 @@ void FMaterialLayersFunctions::PostSerialize(const FArchive& Ar)
 			const int32 NumLayers = Layers.Num();
 			LayerGuids.Empty(NumLayers);
 			ParentLayerGuids.Empty(NumLayers);
-			for (int32 i = 0; i < NumLayers; ++i)
+
+			if (NumLayers > 0)
 			{
-				LayerGuids.Add(FGuid::NewGuid());
-				ParentLayerGuids.Add(UninitializedParentGuid);
+				LayerGuids.Add(BackgroundGuid);
+				ParentLayerGuids.Add(NoParentGuid);
+
+				for (int32 i = 1; i < NumLayers; ++i)
+				{
+					LayerGuids.Add(FGuid::NewGuid());
+					ParentLayerGuids.Add(UninitializedParentGuid);
+				}
 			}
 		}
 	}
