Added GPU sorting support for  GPU particles on Android.

diff --git a/Engine/Source/Runtime/Engine/Private/Particles/ParticleGpuSimulation.cpp b/Engine/Source/Runtime/Engine/Private/Particles/ParticleGpuSimulation.cpp
--- a/Engine/Source/Runtime/Engine/Private/Particles/ParticleGpuSimulation.cpp
+++ b/Engine/Source/Runtime/Engine/Private/Particles/ParticleGpuSimulation.cpp
@@ -2917,8 +2917,7 @@ public:
 
 				const bool bTranslucent = RendersWithTranslucentMaterial();
 				const bool bAllowSorting = FXConsoleVariables::bAllowGPUSorting
-					&& FeatureLevel == ERHIFeatureLevel::SM5
-					// && RHISupportsComputeShaders(ViewFamily.GetShaderPlatform())
+					&& RHISupportsComputeShaders(ViewFamily.GetShaderPlatform())
 					&& bTranslucent;
 
 				// Iterate over views and assign parameters for each.
diff --git a/Engine/Source/Runtime/Engine/Private/Particles/ParticleSortingGPU.cpp b/Engine/Source/Runtime/Engine/Private/Particles/ParticleSortingGPU.cpp
--- a/Engine/Source/Runtime/Engine/Private/Particles/ParticleSortingGPU.cpp
+++ b/Engine/Source/Runtime/Engine/Private/Particles/ParticleSortingGPU.cpp
@@ -49,8 +49,7 @@ public:
 
 	static bool ShouldCompilePermutation(const FGlobalShaderPermutationParameters& Parameters)
 	{
-		return IsFeatureLevelSupported(Parameters.Platform, ERHIFeatureLevel::SM5);
-		// return RHISupportsComputeShaders(Parameters.Platform);
+		return RHISupportsComputeShaders(Parameters.Platform);
 	}
 
 	static void ModifyCompilationEnvironment(const FGlobalShaderPermutationParameters& Parameters, FShaderCompilerEnvironment& OutEnvironment)
