Change AsyncTraceTask priority from background thread/high task priority to normal thread/normal task priority.

diff --git a/Engine/Source/Runtime/Engine/Private/Collision/WorldCollisionAsync.cpp b/Engine/Source/Runtime/Engine/Private/Collision/WorldCollisionAsync.cpp
--- a/Engine/Source/Runtime/Engine/Private/Collision/WorldCollisionAsync.cpp
+++ b/Engine/Source/Runtime/Engine/Private/Collision/WorldCollisionAsync.cpp
@@ -215,9 +215,8 @@ namespace
 	FAutoConsoleTaskPriority CPrio_FAsyncTraceTask(
 		TEXT("TaskGraph.TaskPriorities.AsyncTraceTask"),
 		TEXT("Task and thread priority for async traces."),
-		ENamedThreads::BackgroundThreadPriority, // if we have background priority task threads, then use them...
-		ENamedThreads::HighTaskPriority, // .. at normal task priority
-		ENamedThreads::NormalTaskPriority // if we don't have background threads, then use normal priority threads at normal task priority instead
+		ENamedThreads::NormalThreadPriority, // Use Normal thread and normal task priority
+		ENamedThreads::NormalTaskPriority 
 		);
 
 
