GameViewportClient:
-Added getter for currently using software cursor flag
-Added getter for grabbing the current slate widget cursor, and setter for adding slate widget as cursor

diff --git a/Engine/Source/Runtime/Engine/Classes/Engine/GameViewportClient.h b/Engine/Source/Runtime/Engine/Classes/Engine/GameViewportClient.h
--- a/Engine/Source/Runtime/Engine/Classes/Engine/GameViewportClient.h
+++ b/Engine/Source/Runtime/Engine/Classes/Engine/GameViewportClient.h
@@ -787,6 +787,11 @@ public:
 		bUseSoftwareCursorWidgets = bInUseSoftwareCursorWidgets;
 	}
 
+	/**
+	* Get whether or not the viewport is currently using software cursor
+	*/
+	bool GetIsUsingSoftwareCursorWidgets() { return bUseSoftwareCursorWidgets; }
+
 #if WITH_EDITOR
 	/** Accessor for delegate called when a game viewport received input key */
 	FOnGameViewportInputKey& OnGameViewportInputKey()
@@ -803,9 +808,15 @@ public:
 
 	void SetVirtualCursorWidget(EMouseCursor::Type Cursor, class UUserWidget* Widget);
 
+	/** Add a cursor to the set based on the enum and a slate widget */
+	void AddSoftwareCursorFromSlateWidget(EMouseCursor::Type InCursorType, TSharedPtr<SWidget> CursorWidgetPtr);
+
 	/** Adds a cursor to the set based on the enum and the class reference to it. */
 	void AddSoftwareCursor(EMouseCursor::Type Cursor, const FSoftClassPath& CursorClass);
 
+	/** Get the slate widget of the current software cursor */
+	TSharedPtr<SWidget> GetSoftwareCursorWidget(EMouseCursor::Type Cursor) const;
+
 	/** Does the viewport client have a software cursor set up for the given enum? */
 	bool HasSoftwareCursor(EMouseCursor::Type Cursor) const;
 
diff --git a/Engine/Source/Runtime/Engine/Private/GameViewportClient.cpp b/Engine/Source/Runtime/Engine/Private/GameViewportClient.cpp
--- a/Engine/Source/Runtime/Engine/Private/GameViewportClient.cpp
+++ b/Engine/Source/Runtime/Engine/Private/GameViewportClient.cpp
@@ -916,6 +916,27 @@ void UGameViewportClient::AddSoftwareCursor(EMouseCursor::Type Cursor, const FSo
 	}
 }
 
+void UGameViewportClient::AddSoftwareCursorFromSlateWidget(EMouseCursor::Type InCursorType, TSharedPtr<SWidget> CursorWidgetPtr)
+{
+	if (CursorWidgetPtr.IsValid())
+	{
+		CursorWidgets.Emplace(InCursorType, CursorWidgetPtr);
+	}
+}
+
+TSharedPtr<SWidget> UGameViewportClient::GetSoftwareCursorWidget(EMouseCursor::Type Cursor) const
+{
+	if (CursorWidgets.Contains(Cursor))
+	{
+		const TSharedPtr<SWidget> CursorWidgetPtr = CursorWidgets.FindRef(Cursor);
+		if (CursorWidgetPtr.IsValid())
+		{	
+			return CursorWidgetPtr;
+		}
+	}
+	return nullptr;
+}
+
 bool UGameViewportClient::HasSoftwareCursor(EMouseCursor::Type Cursor) const
 {
 	return CursorWidgets.Contains(Cursor);
