Fix null reference from bad merge

diff --git a/Engine/Source/Runtime/Engine/Private/SoundSubmix.cpp b/Engine/Source/Runtime/Engine/Private/SoundSubmix.cpp
--- a/Engine/Source/Runtime/Engine/Private/SoundSubmix.cpp
+++ b/Engine/Source/Runtime/Engine/Private/SoundSubmix.cpp
@@ -520,7 +520,7 @@ void USoundSubmixWithParentBase::SetParentSubmix(USoundSubmixBase* InParentSubmi
 {
 	if (ParentSubmix != InParentSubmix)
 	{
-		if (ParentSubmix != nullptr)
+		if (ParentSubmix)
 		{
 			ParentSubmix->Modify();
 			ParentSubmix->ChildSubmixes.Remove(this);
@@ -528,8 +528,11 @@ void USoundSubmixWithParentBase::SetParentSubmix(USoundSubmixBase* InParentSubmi
 
 		Modify();
 		ParentSubmix = InParentSubmix;
-		ParentSubmix->ChildSubmixes.AddUnique(this);
+		if (ParentSubmix)
+		{
+			ParentSubmix->ChildSubmixes.AddUnique(this);
 		}
+	}
 }
 
 #if WITH_EDITOR
