Fix failed check relating to material layer parameters

diff --git a/Engine/Source/Runtime/Engine/Private/Materials/Material.cpp b/Engine/Source/Runtime/Engine/Private/Materials/Material.cpp
--- a/Engine/Source/Runtime/Engine/Private/Materials/Material.cpp
+++ b/Engine/Source/Runtime/Engine/Private/Materials/Material.cpp
@@ -2344,8 +2344,8 @@ bool UMaterial::GetTextureParameterValue(const FHashedMaterialParameterInfo& Par
 	UTexture* OldValue = nullptr;
 	const bool bOldResult = GetTextureParameterValue_Legacy(ParameterInfo, OldValue, bOveriddenOnly);
 
-	check(OutValue == OldValue);
 	check(bOldResult == bResult);
+	check(!bOldResult || OutValue == OldValue); // if result is false, texture value is undefined
 #endif
 	return bResult;
 }
