Add option on a URuntimeVirtualTextureComponent to enable/disable the use of streaming mips.

diff --git a/Engine/Source/Editor/VirtualTexturingEditor/Private/RuntimeVirtualTextureDetailsCustomization.cpp b/Engine/Source/Editor/VirtualTexturingEditor/Private/RuntimeVirtualTextureDetailsCustomization.cpp
--- a/Engine/Source/Editor/VirtualTexturingEditor/Private/RuntimeVirtualTextureDetailsCustomization.cpp
+++ b/Engine/Source/Editor/VirtualTexturingEditor/Private/RuntimeVirtualTextureDetailsCustomization.cpp
@@ -244,7 +244,7 @@ void FRuntimeVirtualTextureComponentDetailsCustomization::CustomizeDetails(IDeta
 	];
 
 	// Use existing property to add build button
-	TSharedPtr<IPropertyHandle> LowMipsValue = DetailBuilder.GetProperty("bUseStreamingLowMipsInEditor");
+	TSharedPtr<IPropertyHandle> LowMipsValue = DetailBuilder.GetProperty("bUseStreamingLowMips");
 	DetailBuilder.HideProperty(LowMipsValue);
 
 	IDetailCategoryBuilder& VirtualTextureCategory = DetailBuilder.EditCategory("VirtualTexture", FText::GetEmpty());
diff --git a/Engine/Source/Runtime/Engine/Classes/Components/RuntimeVirtualTextureComponent.h b/Engine/Source/Runtime/Engine/Classes/Components/RuntimeVirtualTextureComponent.h
--- a/Engine/Source/Runtime/Engine/Classes/Components/RuntimeVirtualTextureComponent.h
+++ b/Engine/Source/Runtime/Engine/Classes/Components/RuntimeVirtualTextureComponent.h
@@ -19,8 +19,12 @@ protected:
 	UPROPERTY(EditAnywhere, BlueprintReadOnly, DuplicateTransient, Category = VirtualTexture)
 	URuntimeVirtualTexture* VirtualTexture = nullptr;
 
-	/** Use any streaming low mips when rendering in editor. Set true to view and debug the baked streaming low mips. */
+	/** Use any streaming low mips stored in the virtual texture object. */
 	UPROPERTY(EditAnywhere, AdvancedDisplay, Category = VirtualTexture)
+	bool bUseStreamingLowMips = true;
+
+	/** Use any streaming low mips when rendering in editor. Set true to view and debug the baked streaming low mips. */
+	UPROPERTY(EditAnywhere, AdvancedDisplay, Category = VirtualTexture, meta = (EditCondition = "bUseStreamingLowMips"))
 	bool bUseStreamingLowMipsInEditor = false;
 
 	/** Actor to copy the bounds from to set up the transform. */
diff --git a/Engine/Source/Runtime/Engine/Private/Components/RuntimeVirtualTextureComponent.cpp b/Engine/Source/Runtime/Engine/Private/Components/RuntimeVirtualTextureComponent.cpp
--- a/Engine/Source/Runtime/Engine/Private/Components/RuntimeVirtualTextureComponent.cpp
+++ b/Engine/Source/Runtime/Engine/Private/Components/RuntimeVirtualTextureComponent.cpp
@@ -58,9 +58,9 @@ FTransform URuntimeVirtualTextureComponent::GetVirtualTextureTransform() const
 bool URuntimeVirtualTextureComponent::IsStreamingLowMips() const
 {
 #if WITH_EDITOR
-	return bUseStreamingLowMipsInEditor;
+	return bUseStreamingLowMips && bUseStreamingLowMipsInEditor;
 #else
-	return true;
+	return bUseStreamingLowMips;
 #endif
 }
 
