Don't replace objects marked pending kill with null when running copy properties for unrelated objects.

diff --git a/Engine/Source/Runtime/Engine/Private/UnrealEngine.cpp b/Engine/Source/Runtime/Engine/Private/UnrealEngine.cpp
--- a/Engine/Source/Runtime/Engine/Private/UnrealEngine.cpp
+++ b/Engine/Source/Runtime/Engine/Private/UnrealEngine.cpp
@@ -14380,14 +14380,14 @@ void UEngine::CopyPropertiesForUnrelatedObjects(UObject* OldObject, UObject* New
 
 	FFindInstancedReferenceSubobjectHelper::Duplicate(OldObject, NewObject, ReferenceReplacementMap, ComponentsOnNewObject);
 
-	// Replace anything with an outer of the old object with NULL, unless it already has a replacement
+	// Replace anything with an outer of the old object with NULL, unless it already has a replacement or is marked pending kill
 	ForEachObjectWithOuter(OldObject, [&ReferenceReplacementMap](UObject* ObjectInOuter)
 	{
 		if (!ReferenceReplacementMap.Contains(ObjectInOuter))
 		{
 			ReferenceReplacementMap.Add(ObjectInOuter, nullptr);
 		}
-	});
+	}, true, RF_NoFlags, EInternalObjectFlags::PendingKill);
 
 	if(Params.bClearReferences)
 	{
