FSceneViewport will no longer set all users to focus the game viewport when acquiring focus automatically
- Was already only responding to receiving focus from the keyboard user, but then still set the focus for all users

diff --git a/Engine/Source/Runtime/Engine/Private/Slate/SceneViewport.cpp b/Engine/Source/Runtime/Engine/Private/Slate/SceneViewport.cpp
--- a/Engine/Source/Runtime/Engine/Private/Slate/SceneViewport.cpp
+++ b/Engine/Source/Runtime/Engine/Private/Slate/SceneViewport.cpp
@@ -513,7 +513,7 @@ FReply FSceneViewport::AcquireFocusAndCapture(FIntPoint MousePosition, EFocusCau
 	TSharedRef<SViewport> ViewportWidgetRef = ViewportWidget.Pin().ToSharedRef();
 
 	// Mouse down should focus viewport for user input
-	ReplyState.SetUserFocus(ViewportWidgetRef, FocusCause, true);
+	ReplyState.SetUserFocus(ViewportWidgetRef, FocusCause);
 
 	UWorld* World = ViewportClient->GetWorld();
 	if (World && World->IsGameWorld() && World->GetGameInstance() && (World->GetGameInstance()->GetFirstLocalPlayerController() || World->IsPlayInEditor()))
