#UE4 Fix for crash when asking to create a MID with the same name and outer as another, which overwrites it and orphans the material's render proxy without properly destroying it, which can result in a crash when it accesses it's owner.

diff --git a/Engine/Source/Runtime/Engine/Private/Materials/MaterialInstanceDynamic.cpp b/Engine/Source/Runtime/Engine/Private/Materials/MaterialInstanceDynamic.cpp
--- a/Engine/Source/Runtime/Engine/Private/Materials/MaterialInstanceDynamic.cpp
+++ b/Engine/Source/Runtime/Engine/Private/Materials/MaterialInstanceDynamic.cpp
@@ -32,6 +32,11 @@ UMaterialInstanceDynamic* UMaterialInstanceDynamic::Create(UMaterialInterface* P
 {
 	LLM_SCOPE(ELLMTag::MaterialInstance);
 	UObject* Outer = InOuter ? InOuter : GetTransientPackage();
+	if (Name != NAME_None && FindObjectFast<UMaterialInstanceDynamic>(Outer, *Name.ToString(), true) != nullptr)
+	{
+		// If a MID is made with the same name and outer as another, it will overwrite it. To avoid this we will change the name when there is a collision.
+		Name = MakeUniqueObjectName(Outer, UMaterialInstanceDynamic::StaticClass(), Name);
+	}
 	UMaterialInstanceDynamic* MID = NewObject<UMaterialInstanceDynamic>(Outer, Name);
 	MID->SetParentInternal(ParentMaterial, false);
 	return MID;
