* Various fixes to WorldBrowser and scripted streaming level creation to avoid crashing when setting an invalid level name
* Repair world streaming levels when they point to an invalid level package

diff --git a/Engine/Source/Editor/UnrealEd/Private/EditorLevelUtils.cpp b/Engine/Source/Editor/UnrealEd/Private/EditorLevelUtils.cpp
--- a/Engine/Source/Editor/UnrealEd/Private/EditorLevelUtils.cpp
+++ b/Engine/Source/Editor/UnrealEd/Private/EditorLevelUtils.cpp
@@ -666,17 +666,20 @@ ULevelStreaming* UEditorLevelUtils::CreateNewStreamingLevelForWorld(UWorld& InWo
 	if (bNewWorldSaved)
 	{
 		NewStreamingLevel = AddLevelToWorld(WorldToAddLevelTo, *NewPackageName, LevelStreamingClass);
-		NewLevel = NewStreamingLevel->GetLoadedLevel();
-		// If we are moving the selected actors to the new level move them now
-		if (bMoveSelectedActorsIntoNewLevel)
+		if (NewStreamingLevel != nullptr)
 		{
-			MoveSelectedActorsToLevel(NewStreamingLevel);
-		}
+			NewLevel = NewStreamingLevel->GetLoadedLevel();
+			// If we are moving the selected actors to the new level move them now
+			if (bMoveSelectedActorsIntoNewLevel)
+			{
+				MoveSelectedActorsToLevel(NewStreamingLevel);
+			}
 
-		// Finally make the new level the current one
-		if (WorldToAddLevelTo->SetCurrentLevel(NewLevel))
-		{
-			FEditorDelegates::NewCurrentLevel.Broadcast();
+			// Finally make the new level the current one
+			if (WorldToAddLevelTo->SetCurrentLevel(NewLevel))
+			{
+				FEditorDelegates::NewCurrentLevel.Broadcast();
+			}
 		}
 	}
 
diff --git a/Engine/Source/Editor/WorldBrowser/Private/Tiles/WorldTileCollectionModel.cpp b/Engine/Source/Editor/WorldBrowser/Private/Tiles/WorldTileCollectionModel.cpp
--- a/Engine/Source/Editor/WorldBrowser/Private/Tiles/WorldTileCollectionModel.cpp
+++ b/Engine/Source/Editor/WorldBrowser/Private/Tiles/WorldTileCollectionModel.cpp
@@ -1904,8 +1904,12 @@ void FWorldTileCollectionModel::OnPostSaveWorld(uint32 SaveFlags, UWorld* World,
 void FWorldTileCollectionModel::OnNewCurrentLevel()
 {
 	TSharedPtr<FLevelModel> CurrentLevelModel = FindLevelModel(CurrentWorld->GetCurrentLevel());
-	// Make sure level will be in focus
-	Focus(CurrentLevelModel->GetLevelBounds(), FWorldTileCollectionModel::OriginAtCenter);
+	// it's possible the level model has not been registered yet because the RefreshLevelBrowser event hasn't been broadcasted yet : 
+	if (CurrentLevelModel.IsValid())
+	{
+		// Make sure level will be in focus
+		Focus(CurrentLevelModel->GetLevelBounds(), FWorldTileCollectionModel::OriginAtCenter);
+	}
 }
 
 bool FWorldTileCollectionModel::HasMeshProxySupport() const
diff --git a/Engine/Source/Runtime/Engine/Classes/Engine/LevelStreaming.h b/Engine/Source/Runtime/Engine/Classes/Engine/LevelStreaming.h
--- a/Engine/Source/Runtime/Engine/Classes/Engine/LevelStreaming.h
+++ b/Engine/Source/Runtime/Engine/Classes/Engine/LevelStreaming.h
@@ -415,6 +415,9 @@ public:
 	UFUNCTION(BlueprintPure, meta = (BlueprintInternalUseOnly = "true"))
 	ALevelScriptActor* GetLevelScriptActor();
 
+	/** Returns false if the level package associated to that streaming level is invalid. */
+	bool IsValidStreamingLevel() const;
+
 #if WITH_EDITOR
 	/** Get the folder path for this level for use in the world browser. Only available in editor builds */
 	const FName& GetFolderPath() const;
diff --git a/Engine/Source/Runtime/Engine/Classes/Engine/World.h b/Engine/Source/Runtime/Engine/Classes/Engine/World.h
--- a/Engine/Source/Runtime/Engine/Classes/Engine/World.h
+++ b/Engine/Source/Runtime/Engine/Classes/Engine/World.h
@@ -1322,6 +1322,9 @@ private:
 
 	/** Utility function that is used to ensure that a World has the correct WorldSettings */
 	void RepairWorldSettings();
+
+	/** Utility function to cleanup streaming levels that point to invalid level packages */
+	void RepairStreamingLevels();
 	
 #if INCLUDE_CHAOS
 	/** Utility function that is used to ensure that a World has the correct ChaosActor */
diff --git a/Engine/Source/Runtime/Engine/Private/LevelStreaming.cpp b/Engine/Source/Runtime/Engine/Private/LevelStreaming.cpp
--- a/Engine/Source/Runtime/Engine/Private/LevelStreaming.cpp
+++ b/Engine/Source/Runtime/Engine/Private/LevelStreaming.cpp
@@ -311,25 +311,22 @@ void ULevelStreaming::PostLoad()
 	}
 #endif
 
-	if ( !PIESession && !WorldAsset.IsNull() )
+	if (!IsValidStreamingLevel())
 	{
 		const FString WorldPackageName = GetWorldAssetPackageName();
-		if (FPackageName::DoesPackageExist(WorldPackageName) == false)
-		{
-			UE_LOG(LogLevelStreaming, Display, TEXT("Failed to find streaming level package file: %s. This streaming level may not load or save properly."), *WorldPackageName);
+		UE_LOG(LogLevelStreaming, Display, TEXT("Failed to find streaming level package file: %s. This streaming level may not load or save properly."), *WorldPackageName);
 #if WITH_EDITOR
-			if (GIsEditor)
-			{
-				// Launch notification to inform user of default change
-				FFormatNamedArguments Args;
-				Args.Add(TEXT("PackageName"), FText::FromString(WorldPackageName));
-				FNotificationInfo Info(FText::Format(LOCTEXT("LevelStreamingFailToStreamLevel", "Failed to find streamed level {PackageName}, please fix the reference to it in the Level Browser"), Args));
-				Info.ExpireDuration = 7.0f;
+		if (GIsEditor)
+		{
+			// Launch notification to inform user of default change
+			FFormatNamedArguments Args;
+			Args.Add(TEXT("PackageName"), FText::FromString(WorldPackageName));
+			FNotificationInfo Info(FText::Format(LOCTEXT("LevelStreamingFailToStreamLevel", "Failed to find streamed level {PackageName}, please fix the reference to it in the Level Browser"), Args));
+			Info.ExpireDuration = 7.0f;
 
-				FSlateNotificationManager::Get().AddNotification(Info);
-			}
-#endif // WITH_EDITOR
+			FSlateNotificationManager::Get().AddNotification(Info);
 		}
+#endif // WITH_EDITOR
 	}
 
 #if WITH_EDITOR
@@ -1682,6 +1679,17 @@ ALevelScriptActor* ULevelStreaming::GetLevelScriptActor()
 	return nullptr;
 }
 
+bool ULevelStreaming::IsValidStreamingLevel() const
+{
+	const bool PIESession = GetWorld()->WorldType == EWorldType::PIE || GetOutermost()->HasAnyPackageFlags(PKG_PlayInEditor);
+	if (!PIESession && !WorldAsset.IsNull())
+	{
+		const FString WorldPackageName = GetWorldAssetPackageName();
+		return FPackageName::DoesPackageExist(WorldPackageName);
+	}
+	return true;
+}
+
 #if WITH_EDITOR
 void ULevelStreaming::PostEditUndo()
 {
diff --git a/Engine/Source/Runtime/Engine/Private/World.cpp b/Engine/Source/Runtime/Engine/Private/World.cpp
--- a/Engine/Source/Runtime/Engine/Private/World.cpp
+++ b/Engine/Source/Runtime/Engine/Private/World.cpp
@@ -961,6 +961,7 @@ void UWorld::PostLoad()
 #endif
 #if WITH_EDITOR
 	RepairWorldSettings();
+	RepairStreamingLevels();
 #endif
 #if INCLUDE_CHAOS
 	//RepairChaosActors();
@@ -1275,6 +1276,22 @@ void UWorld::RepairChaosActors()
 }
 #endif
 
+void UWorld::RepairStreamingLevels()
+{
+	for (int32 Index = 0; Index < StreamingLevels.Num(); )
+	{
+		if (!StreamingLevels[Index]->IsValidStreamingLevel())
+		{
+			StreamingLevels.RemoveAtSwap(Index);
+		}
+		else
+		{
+			++Index;
+		}
+	}
+}
+
+
 void UWorld::RepairWorldSettings()
 {
 	AWorldSettings* ExistingWorldSettings = PersistentLevel->GetWorldSettings(false);
@@ -1406,11 +1423,13 @@ void UWorld::InitWorld(const InitializationValues IVS)
 
 #if WITH_EDITOR
 	RepairWorldSettings();
+	RepairStreamingLevels();
 #endif
 #if INCLUDE_CHAOS
 	//RepairChaosActors();
 #endif
 
+
 	// initialize DefaultPhysicsVolume for the world
 	// Spawned on demand by this function.
 	DefaultPhysicsVolume = GetDefaultPhysicsVolume();
