Allow SpawnActor to take a name as part of its parameters without it having to be fatal if that name is in use. Possible alternative modes are to return null and optionally report an error or to treat the name as a request and generate a unique one if needed.
Name uniqueness logic is now correctly applied when the supplied template is null.

diff --git a/Engine/Source/Runtime/Engine/Classes/Engine/World.h b/Engine/Source/Runtime/Engine/Classes/Engine/World.h
--- a/Engine/Source/Runtime/Engine/Classes/Engine/World.h
+++ b/Engine/Source/Runtime/Engine/Classes/Engine/World.h
@@ -531,13 +531,32 @@ public:
 	uint8	bAllowDuringConstructionScript:1;
 
 #if WITH_EDITOR
-	/** Determines whether the begin play cycle will run on the spawned actor when in the editor. */
+	/* Determines whether the begin play cycle will run on the spawned actor when in the editor. */
 	uint8	bTemporaryEditorActor:1;
 
-	/* Determines wether or not the actor should be hidden from the Scene Outliner */
+	/* Determines whether or not the actor should be hidden from the Scene Outliner */
 	uint8	bHideFromSceneOutliner:1;
 #endif
-	
+
+	/* Modes that SpawnActor can use the supplied name when it is not None. */
+	enum class ESpawnActorNameMode : uint8
+	{
+		/* Fatal if unavailable, application will assert */
+		Required_Fatal,
+
+		/* Report an error return null if unavailable */
+		Required_ErrorAndReturnNull,
+
+		/* Return null if unavailable */
+		Required_ReturnNull,
+
+		/* If the supplied Name is already in use the generate an unused one using the supplied version as a base */
+		Requested
+	};
+
+	/* In which way should SpawnActor should treat the supplied Name if not none. */
+	ESpawnActorNameMode NameMode;
+
 	/* Flags used to describe the spawned actor/object instance. */
 	EObjectFlags ObjectFlags;		
 };
diff --git a/Engine/Source/Runtime/Engine/Private/Components/ChildActorComponent.cpp b/Engine/Source/Runtime/Engine/Private/Components/ChildActorComponent.cpp
--- a/Engine/Source/Runtime/Engine/Private/Components/ChildActorComponent.cpp
+++ b/Engine/Source/Runtime/Engine/Private/Components/ChildActorComponent.cpp
@@ -563,6 +563,7 @@ void UChildActorComponent::CreateChildActor()
 				Params.bAllowDuringConstructionScript = true;
 				Params.OverrideLevel = (MyOwner ? MyOwner->GetLevel() : nullptr);
 				Params.Name = ChildActorName;
+				Params.NameMode = FActorSpawnParameters::ESpawnActorNameMode::Requested;
 				if (ChildActorTemplate && ChildActorTemplate->GetClass() == ChildActorClass)
 				{
 					Params.Template = ChildActorTemplate;
diff --git a/Engine/Source/Runtime/Engine/Private/LevelActor.cpp b/Engine/Source/Runtime/Engine/Private/LevelActor.cpp
--- a/Engine/Source/Runtime/Engine/Private/LevelActor.cpp
+++ b/Engine/Source/Runtime/Engine/Private/LevelActor.cpp
@@ -384,19 +384,37 @@ AActor* UWorld::SpawnActor( UClass* Class, FTransform const* UserTransformPtr, c
 		// Use class's default actor as a template.
 		Template = Class->GetDefaultObject<AActor>();
 	}
-	else if (!Template->HasAnyFlags(RF_ClassDefaultObject))
+	check(Template);
+
+	if (NewActorName.IsNone())
 	{
-		if (NewActorName.IsNone())
+		// If we are using a template object and haven't specified a name, create a name relative to the template, otherwise let the default object naming behavior in Stat
+		if (!Template->HasAnyFlags(RF_ClassDefaultObject))
 		{
 			NewActorName = MakeUniqueObjectName(LevelToSpawnIn, Template->GetClass(), *Template->GetFName().GetPlainNameString());
 		}
-		else if (StaticFindObjectFast(nullptr, LevelToSpawnIn, NewActorName))
+	}
+	else if (StaticFindObjectFast(nullptr, LevelToSpawnIn, NewActorName))
+	{
+		// If the supplied name is already in use, then either fail in the requested manner or determine a new name to use if the caller indicates that's ok
+
+		if (SpawnParameters.NameMode == FActorSpawnParameters::ESpawnActorNameMode::Requested)
 		{
-			UE_LOG(LogSpawn, Fatal, TEXT("An actor of name '%s' already exists in level '%s'."), *NewActorName.ToString(), *LevelToSpawnIn->GetFullName());
+			NewActorName = MakeUniqueObjectName(LevelToSpawnIn, Template->GetClass(), *NewActorName.GetPlainNameString());
+		}
+		else
+		{
+			if (SpawnParameters.NameMode == FActorSpawnParameters::ESpawnActorNameMode::Required_Fatal)
+			{
+				UE_LOG(LogSpawn, Fatal, TEXT("An actor of name '%s' already exists in level '%s'."), *NewActorName.ToString(), *LevelToSpawnIn->GetFullName());
+			}
+			else if (SpawnParameters.NameMode == FActorSpawnParameters::ESpawnActorNameMode::Required_ErrorAndReturnNull)
+			{
+				UE_LOG(LogSpawn, Error, TEXT("An actor of name '%s' already exists in level '%s'."), *NewActorName.ToString(), *LevelToSpawnIn->GetFullName());
+			}
 			return nullptr;
 		}
 	}
-	check(Template);
 
 	// See if we can spawn on ded.server/client only etc (check NeedsLoadForClient & NeedsLoadForServer)
 	if(!CanCreateInCurrentContext(Template))
diff --git a/Engine/Source/Runtime/Engine/Private/World.cpp b/Engine/Source/Runtime/Engine/Private/World.cpp
--- a/Engine/Source/Runtime/Engine/Private/World.cpp
+++ b/Engine/Source/Runtime/Engine/Private/World.cpp
@@ -172,6 +172,7 @@ FActorSpawnParameters::FActorSpawnParameters()
 , bTemporaryEditorActor(false)
 , bHideFromSceneOutliner(false)
 #endif
+, NameMode(ESpawnActorNameMode::Required_Fatal)
 , ObjectFlags(RF_Transactional)
 {
 }
