Blueprints: BlueprintDisplayName and BlueprintDescription now get cleared on the new copy when duplicating a Blueprint asset

diff --git a/Engine/Source/Runtime/Engine/Classes/Engine/Blueprint.h b/Engine/Source/Runtime/Engine/Classes/Engine/Blueprint.h
--- a/Engine/Source/Runtime/Engine/Classes/Engine/Blueprint.h
+++ b/Engine/Source/Runtime/Engine/Classes/Engine/Blueprint.h
@@ -525,11 +525,11 @@ public:
 	TEnumAsByte<enum EBlueprintStatus> Status;
 
 	/** Overrides the BP's display name in the editor UI */
-	UPROPERTY(EditAnywhere, Category=BlueprintOptions)
+	UPROPERTY(EditAnywhere, Category=BlueprintOptions, DuplicateTransient)
 	FString BlueprintDisplayName;
 
-	/** Shows up in the content browser when the blueprint is hovered, exported as Hidden in GetAssetRegistryTags */
-	UPROPERTY(EditAnywhere, Category=BlueprintOptions, meta=(MultiLine=true))
+	/** Shows up in the content browser tooltip when the blueprint is hovered */
+	UPROPERTY(EditAnywhere, Category=BlueprintOptions, meta=(MultiLine=true), DuplicateTransient)
 	FString BlueprintDescription;
 
 	/** The category of the Blueprint, used to organize this Blueprint class when displayed in palette windows */
