Blueprints: Deleting unused variable UBlueprint::SearchGuid

diff --git a/Engine/Source/Runtime/Engine/Classes/Engine/Blueprint.h b/Engine/Source/Runtime/Engine/Classes/Engine/Blueprint.h
--- a/Engine/Source/Runtime/Engine/Classes/Engine/Blueprint.h
+++ b/Engine/Source/Runtime/Engine/Classes/Engine/Blueprint.h
@@ -539,10 +539,6 @@ public:
 	/** Additional HideCategories. These are added to HideCategories from parent. */
 	UPROPERTY(EditAnywhere, Category=BlueprintOptions)
 	TArray<FString> HideCategories;
-	 
-	/** Guid key for finding searchable data for Blueprint in the DDC */
-	UPROPERTY()
-	FGuid SearchGuid;
 
 #endif //WITH_EDITORONLY_DATA
 
