Fixes from UDN: Copy-paste error and FTexture2DArrayResource constructor initialization.
https://udn.unrealengine.com/questions/544210/crash-starting-up-in.html

diff --git a/Engine/Source/Runtime/Engine/Private/Materials/MaterialUniformExpressions.cpp b/Engine/Source/Runtime/Engine/Private/Materials/MaterialUniformExpressions.cpp
--- a/Engine/Source/Runtime/Engine/Private/Materials/MaterialUniformExpressions.cpp
+++ b/Engine/Source/Runtime/Engine/Private/Materials/MaterialUniformExpressions.cpp
@@ -1286,7 +1286,7 @@ void FUniformExpressionSet::FillUniformBuffer(const FMaterialRenderContext& Mate
 				check(GBlackArrayTexture->TextureRHI);
 				*ResourceTableTexturePtr = GBlackArrayTexture->TextureRHI;
 				check(GBlackArrayTexture->SamplerStateRHI);
-				*ResourceTableTexturePtr = GBlackArrayTexture->SamplerStateRHI;
+				*ResourceTableSamplerPtr = GBlackArrayTexture->SamplerStateRHI;
 			}
 		}
 
diff --git a/Engine/Source/Runtime/Engine/Public/TextureResource.h b/Engine/Source/Runtime/Engine/Public/TextureResource.h
--- a/Engine/Source/Runtime/Engine/Public/TextureResource.h
+++ b/Engine/Source/Runtime/Engine/Public/TextureResource.h
@@ -408,6 +408,15 @@ public:
 
 	FTexture2DArrayResource() :
 		SizeX(0),
+		SizeY(0),
+		NumMips(0),
+		NumSlices(0),
+		LODGroup(),
+		Format(PF_Unknown),
+		Filter(SF_Point),
+		SamplerXAddress(AM_Wrap),
+		SamplerYAddress(AM_Wrap),
+		SamplerZAddress(AM_Wrap),
 		bDirty(false),
 		bPreventingReallocation(false)
 	{}
