Redo AssetManager / StreamableHandle optimization except for a small part that broke EngineTest LoadAsset blueprint test

diff --git a/Engine/Source/Runtime/Engine/Classes/Engine/AssetManager.h b/Engine/Source/Runtime/Engine/Classes/Engine/AssetManager.h
--- a/Engine/Source/Runtime/Engine/Classes/Engine/AssetManager.h
+++ b/Engine/Source/Runtime/Engine/Classes/Engine/AssetManager.h
@@ -524,7 +524,7 @@ protected:
 	virtual void RebuildObjectReferenceList();
 
 	/** Called when an internal load handle finishes, handles setting to pending state */
-	virtual void OnAssetStateChangeCompleted(FPrimaryAssetId PrimaryAssetId, TSharedPtr<FStreamableHandle> BoundHandle, FStreamableDelegate WrappedDelegate);
+	virtual void OnAssetStateChangeCompleted(TArray<FPrimaryAssetId> PrimaryAssetId, TSharedPtr<FStreamableHandle> BoundHandle, FStreamableDelegate WrappedDelegate);
 
 	/** Helper function to write out asset reports */
 	virtual bool WriteCustomReport(FString FileName, TArray<FString>& FileLines) const;
diff --git a/Engine/Source/Runtime/Engine/Private/AssetManager.cpp b/Engine/Source/Runtime/Engine/Private/AssetManager.cpp
--- a/Engine/Source/Runtime/Engine/Private/AssetManager.cpp
+++ b/Engine/Source/Runtime/Engine/Private/AssetManager.cpp
@@ -76,11 +76,17 @@ struct FPrimaryAssetData
 
 	/** Pending state of this asset, will be copied to CurrentState when load finishes */
 	FPrimaryAssetLoadState PendingState;
-
-	FPrimaryAssetData() {}
-
+	
 	/** Asset is considered loaded at all if there is an active handle for it */
 	bool IsLoaded() const { return CurrentState.IsValid(); }
+
+	void MovePendingToCurrent()
+	{
+		Swap(PendingState.BundleNames, CurrentState.BundleNames);
+		Swap(PendingState.Handle, CurrentState.Handle);
+
+		PendingState.Reset(/* don't cancel handle */ false);	
+	}
 };
 
 /** Structure representing all items of a specific asset type */
@@ -1127,9 +1133,25 @@ void UAssetManager::GetPrimaryAssetTypeInfoList(TArray<FPrimaryAssetTypeInfo>& A
 
 TSharedPtr<FStreamableHandle> UAssetManager::ChangeBundleStateForPrimaryAssets(const TArray<FPrimaryAssetId>& AssetsToChange, const TArray<FName>& AddBundles, const TArray<FName>& RemoveBundles, bool bRemoveAllBundles, FStreamableDelegate DelegateToCall, TAsyncLoadPriority Priority)
 {
-	TArray<TSharedPtr<FStreamableHandle> > NewHandles, ExistingHandles;
-	TArray<FPrimaryAssetId> NewAssets;
-	TSharedPtr<FStreamableHandle> ReturnHandle;
+	// ExistingHandles are both primary and secondary assets who are already being loaded.
+	// We need to return a combined handle that waits for newly requested assets as well as
+	// these previously requested assets.
+	//
+	// The code relies on TSet iteraton order being the insertion order.
+	// This matters since FStreamableHandle::GetLoadedAsset() returns the first requested asset
+	// or the first requested asset of the first child handle. 
+
+	TSet<TSharedPtr<FStreamableHandle>> ExistingHandles;
+	TSet<FSoftObjectPath> PathsToLoad;
+	PathsToLoad.Reserve(AssetsToChange.Num());
+
+	TArray<FName> AddBundlesSorted(AddBundles);
+	AddBundlesSorted.Sort(FNameLexicalLess());
+	TArray<FName> TempBundles;
+
+	TArray<FPrimaryAssetData*> NameDatasToLoad;
+	NameDatasToLoad.Reserve(AssetsToChange.Num());
+	TBitArray<> IdsToLoad(false, AssetsToChange.Num());
 
 	for (const FPrimaryAssetId& PrimaryAssetId : AssetsToChange)
 	{
@@ -1140,31 +1162,44 @@ TSharedPtr<FStreamableHandle> UAssetManager::ChangeBundleStateForPrimaryAssets(c
 			// Iterate list of changes, compute new bundle set
 			bool bLoadIfNeeded = false;
 			
-			// Use pending state if valid
-			TArray<FName> CurrentBundleState = NameData->PendingState.IsValid() ? NameData->PendingState.BundleNames : NameData->CurrentState.BundleNames;
-			TArray<FName> NewBundleState;
+			// NewBundles will be the set of bundle tags we want to have loaded,
+			// i.e. AddBundles + CurrentBundles - RemoveBundles
+			const TArray<FName>& CurrentBundles = NameData->PendingState.IsValid() ? NameData->PendingState.BundleNames : NameData->CurrentState.BundleNames;
+			const TArray<FName>* NewBundles;
 
-			if (!bRemoveAllBundles)
+			if (bRemoveAllBundles || CurrentBundles.Num() == 0)
+			{
+				// Use AddBundlesSorted to avoid creating temporary array
+				NewBundles = &AddBundlesSorted;
+			}
+			else
 			{
-				NewBundleState = CurrentBundleState;
+				// Reuse TempBundles to avoid allocations and start with the bundles we want to be loaded
+				NewBundles = &TempBundles;
+				TempBundles = AddBundlesSorted;
 
-				for (const FName& RemoveBundle : RemoveBundles)
+				// Keep CurrentBundles unless they already exist or we explictly should remove them
+				// It is assumed that RemoveBundles do not overlap with AddBundles
+				for (FName CurrentBundle : CurrentBundles)
 				{
-					NewBundleState.Remove(RemoveBundle);
+					if (RemoveBundles.Find(CurrentBundle) == INDEX_NONE &&
+						TempBundles.Find(CurrentBundle) == INDEX_NONE)
+					{
+						TempBundles.Add(CurrentBundle);
+					}
 				}
-			}
 
-			for (const FName& AddBundle : AddBundles)
-			{
-				NewBundleState.AddUnique(AddBundle);
+				// Sort bundles if we added anything on top of AddBundlesSorted
+				if (TempBundles.Num() > AddBundlesSorted.Num())
+				{
+					TempBundles.Sort(FNameLexicalLess());
+				}
 			}
 
-			NewBundleState.Sort(FNameLexicalLess());
-
 			// If the pending state is valid, check if it is different
 			if (NameData->PendingState.IsValid())
 			{
-				if (NameData->PendingState.BundleNames == NewBundleState)
+				if (NameData->PendingState.BundleNames == *NewBundles)
 				{
 					// This will wait on any existing handles to finish
 					ExistingHandles.Add(NameData->PendingState.Handle);
@@ -1173,30 +1208,29 @@ TSharedPtr<FStreamableHandle> UAssetManager::ChangeBundleStateForPrimaryAssets(c
 				// Clear pending state
 				NameData->PendingState.Reset(true);
 			}
-			else if (NameData->CurrentState.IsValid() && NameData->CurrentState.BundleNames == NewBundleState)
+			else if (NameData->CurrentState.IsValid() && NameData->CurrentState.BundleNames == *NewBundles)
 			{
 				// If no pending, compare with current
 				continue;
 			}
 
-			TSet<FSoftObjectPath> PathsToLoad;
-
 			// Gather asset refs
 			const FSoftObjectPath& AssetPath = NameData->AssetPtr.ToSoftObjectPath();
-
+			bool bNothingToLoadForAsset = AssetPath.IsNull();
 			if (!AssetPath.IsNull())
 			{
 				// Dynamic types can have no base asset path
 				PathsToLoad.Add(AssetPath);
 			}
 			
-			for (const FName& BundleName : NewBundleState)
+			for (const FName& BundleName : *NewBundles)
 			{
 				FAssetBundleEntry Entry = GetAssetBundleEntry(PrimaryAssetId, BundleName);
 
 				if (Entry.IsValid())
 				{
 					PathsToLoad.Append(Entry.BundleAssets);
+					bNothingToLoadForAsset &= Entry.BundleAssets.Num() == 0;
 				}
 				else
 				{
@@ -1204,100 +1238,110 @@ TSharedPtr<FStreamableHandle> UAssetManager::ChangeBundleStateForPrimaryAssets(c
 				}
 			}
 
-			TSharedPtr<FStreamableHandle> NewHandle;
-
-			FString DebugName = PrimaryAssetId.ToString();
-
-			if (NewBundleState.Num() > 0)
-			{
-				DebugName += TEXT(" (");
-
-				for (int32 i = 0; i < NewBundleState.Num(); i++)
-				{
-					if (i != 0)
-					{
-						DebugName += TEXT(", ");
-					}
-					DebugName += NewBundleState[i].ToString();
-				}
-
-				DebugName += TEXT(")");
-			}
-
-			if (PathsToLoad.Num() == 0)
+			if (bNothingToLoadForAsset)
 			{
 				// New state has no assets to load. Set the CurrentState's bundles and clear the handle
-				NameData->CurrentState.BundleNames = NewBundleState;
+				NameData->CurrentState.BundleNames = *NewBundles;
 				NameData->CurrentState.Handle.Reset();
 				continue;
 			}
 
-			NewHandle = LoadAssetList(PathsToLoad.Array(), FStreamableDelegate(), Priority, DebugName);
+			NameDatasToLoad.Add(NameData);
+			NameData->PendingState.BundleNames = *NewBundles;
+			IdsToLoad[&PrimaryAssetId - &AssetsToChange[0]] = true;
+		}
+	}
 
-			if (!NewHandle.IsValid())
-			{
-				// LoadAssetList already throws an error, no need to do it here as well
-				continue;
-			}
+	TSharedPtr<FStreamableHandle> ReturnHandle;
+	TStringBuilder<256> DebugName;
+	if (PathsToLoad.Num() > 0)
+	{
+		DebugName << *PathsToLoad.CreateConstIterator();
 
-			if (NewHandle->HasLoadCompleted())
-			{
-				// Copy right into active
-				NameData->CurrentState.BundleNames = NewBundleState;
-				NameData->CurrentState.Handle = NewHandle;
-			}
-			else
+		if (NameDatasToLoad.Num() > 1)
+		{
+			DebugName << " and " << NameDatasToLoad.Num() - 1 << " more";
+		}
+			
+		DebugName << TEXT(" (");
+		int32 LenWithoutElements = DebugName.Len();
+		for (FName AddBundle : AddBundlesSorted)
+		{
+			if (!bRemoveAllBundles)
 			{
-				// Copy into pending and set delegate
-				NameData->PendingState.BundleNames = NewBundleState;
-				NameData->PendingState.Handle = NewHandle;
+				DebugName << '+';
+			}					
+				
+			DebugName << AddBundle << ", ";
+		}
 
-				NewHandle->BindCompleteDelegate(FStreamableDelegate::CreateUObject(this, &UAssetManager::OnAssetStateChangeCompleted, PrimaryAssetId, NewHandle, FStreamableDelegate()));
-			}
+		for (FName RemoveBundle : RemoveBundles)
+		{
+			DebugName << '-' << RemoveBundle << ", ";
+		}
 
-			NewHandles.Add(NewHandle);
-			NewAssets.Add(PrimaryAssetId);
+		// Trim last comma and space
+		if (DebugName.Len() > LenWithoutElements)
+		{
+			DebugName.RemoveSuffix(2);
 		}
-	}
+		DebugName << TEXT(")");
 
-	if (NewHandles.Num() > 1 || ExistingHandles.Num() > 0)
-	{
-		// If multiple handles or we have an old handle, need to make wrapper handle
-		NewHandles.Append(ExistingHandles);
 
-		ReturnHandle = StreamableManager.CreateCombinedHandle(NewHandles, FString::Printf(TEXT("%s CreateCombinedHandle"), *GetName()));
+		ReturnHandle = LoadAssetList(PathsToLoad.Array(), FStreamableDelegate(), Priority, FString(DebugName));
 
-		// Call delegate or bind to meta handle
-		if (ReturnHandle->HasLoadCompleted())
+		// Update FPrimaryAssetData pending handle
+		for (FPrimaryAssetData* NameData : NameDatasToLoad)
 		{
-			FStreamableHandle::ExecuteDelegate(DelegateToCall);
-		}
-		else
-		{
-			// Call external callback when completed
-			ReturnHandle->BindCompleteDelegate(DelegateToCall);
+			NameData->PendingState.Handle = ReturnHandle;
 		}
 	}
-	else if (NewHandles.Num() == 1)
+		
+	// Create combined wrapper handle from existing handles
+	if (ExistingHandles.Num() > 0)
 	{
-		ReturnHandle = NewHandles[0];
-		ensure(NewAssets.Num() == 1);
-
-		// If only one handle, return it and add callback
-		if (ReturnHandle->HasLoadCompleted())
+		TArray<TSharedPtr<FStreamableHandle>> AllHandles;
+		AllHandles.Reserve(ReturnHandle.IsValid() + ExistingHandles.Num());
+		if (ReturnHandle.IsValid())
 		{
-			FStreamableHandle::ExecuteDelegate(DelegateToCall);
+			AllHandles.Add(ReturnHandle);
 		}
-		else
+		for (const TSharedPtr<FStreamableHandle>& ExistingHandle : ExistingHandles)
 		{
-			// Call internal callback and external callback when it finishes
-			ReturnHandle->BindCompleteDelegate(FStreamableDelegate::CreateUObject(this, &UAssetManager::OnAssetStateChangeCompleted, NewAssets[0], ReturnHandle, DelegateToCall));
+			AllHandles.Add(ExistingHandle);
 		}
+
+		DebugName << " combined";
+
+		ReturnHandle = StreamableManager.CreateCombinedHandle(AllHandles, *DebugName);
+	}
+
+	// Call delegate and update NameDatasToLoad when done
+	if (!ReturnHandle.IsValid() || ReturnHandle->HasLoadCompleted())
+	{
+		for (FPrimaryAssetData* NameData : NameDatasToLoad)
+		{
+			NameData->MovePendingToCurrent();
+		}
+
+		DelegateToCall.ExecuteIfBound();
 	}
 	else
 	{
-		// Call completion callback, nothing to do
-		FStreamableHandle::ExecuteDelegate(DelegateToCall);
+		// NameDatasToLoad need updating on completion but they can be deleted. 
+		// Create id array instead and redo Id -> NameData lookup on completion. 
+		TArray<FPrimaryAssetId> NameDataIds;
+		NameDataIds.Reserve(NameDatasToLoad.Num());
+		for (const FPrimaryAssetId& PrimaryAssetId : AssetsToChange)
+		{
+			if (IdsToLoad[&PrimaryAssetId - &AssetsToChange[0]])
+			{
+				NameDataIds.Add(PrimaryAssetId);
+			}
+		}
+		check(NameDatasToLoad.Num() == NameDataIds.Num())
+
+		ReturnHandle->BindCompleteDelegate(FStreamableDelegate::CreateUObject(this, &UAssetManager::OnAssetStateChangeCompleted, MoveTemp(NameDataIds), ReturnHandle, DelegateToCall));
 	}
 
 	return ReturnHandle;
@@ -1388,31 +1432,28 @@ TSharedPtr<FStreamableHandle> UAssetManager::PreloadPrimaryAssets(const TArray<F
 	return ReturnHandle;
 }
 
-void UAssetManager::OnAssetStateChangeCompleted(FPrimaryAssetId PrimaryAssetId, TSharedPtr<FStreamableHandle> BoundHandle, FStreamableDelegate WrappedDelegate)
+void UAssetManager::OnAssetStateChangeCompleted(TArray<FPrimaryAssetId> PrimaryAssetIds,  TSharedPtr<FStreamableHandle> BoundHandle, FStreamableDelegate WrappedDelegate)
 {
-	FPrimaryAssetData* NameData = GetNameData(PrimaryAssetId);
-
-	if (NameData)
+	for (const FPrimaryAssetId& PrimaryAssetId : PrimaryAssetIds)
 	{
-		if (NameData->PendingState.Handle == BoundHandle)
+		if (FPrimaryAssetData* NameData = GetNameData(PrimaryAssetId))
 		{
-			NameData->CurrentState.Handle = NameData->PendingState.Handle;
-			NameData->CurrentState.BundleNames = NameData->PendingState.BundleNames;
-
-			WrappedDelegate.ExecuteIfBound();
-
-			// Clear old state, but don't cancel handle as we just copied it into current
-			NameData->PendingState.Reset(false);
+			if (NameData->PendingState.Handle == BoundHandle)
+			{
+				NameData->MovePendingToCurrent();
+			}
+			else
+			{
+				UE_LOG(LogAssetManager, Verbose, TEXT("OnAssetStateChangeCompleted: Received after pending data changed, ignoring (%s)"), *PrimaryAssetId.ToString());
+			}
 		}
 		else
 		{
-			UE_LOG(LogAssetManager, Verbose, TEXT("OnAssetStateChangeCompleted: Received after pending data changed, ignoring (%s)"), *PrimaryAssetId.ToString());
+			UE_LOG(LogAssetManager, Error, TEXT("OnAssetStateChangeCompleted: Received for invalid asset! (%s)"), *PrimaryAssetId.ToString());
 		}
 	}
-	else
-	{
-		UE_LOG(LogAssetManager, Error, TEXT("OnAssetStateChangeCompleted: Received for invalid asset! (%s)"), *PrimaryAssetId.ToString());
-	}
+
+	WrappedDelegate.ExecuteIfBound();
 }
 
 TSharedPtr<FStreamableHandle> UAssetManager::LoadPrimaryAssets(const TArray<FPrimaryAssetId>& AssetsToLoad, const TArray<FName>& LoadBundles, FStreamableDelegate DelegateToCall, TAsyncLoadPriority Priority)
diff --git a/Engine/Source/Runtime/Engine/Private/StreamableManager.cpp b/Engine/Source/Runtime/Engine/Private/StreamableManager.cpp
--- a/Engine/Source/Runtime/Engine/Private/StreamableManager.cpp
+++ b/Engine/Source/Runtime/Engine/Private/StreamableManager.cpp
@@ -412,13 +412,23 @@ void FStreamableHandle::GetRequestedAssets(TArray<FSoftObjectPath>& AssetList) c
 
 UObject* FStreamableHandle::GetLoadedAsset() const
 {
-	TArray<UObject *> LoadedAssets;
-
-	GetLoadedAssets(LoadedAssets);
-
-	if (LoadedAssets.Num() > 0)
+	if (HasLoadCompleted())
 	{
-		return LoadedAssets[0];
+		if (RequestedAssets.Num() > 0)
+		{
+			const FSoftObjectPath& Ref = RequestedAssets[0];
+			// Try manager, should be faster and will handle redirects better
+			return IsActive() ? OwningManager->GetStreamed(Ref) : Ref.ResolveObject();
+		}
+
+		// Check child handles
+		for (const TSharedPtr<FStreamableHandle>& ChildHandle : ChildHandles)
+		{
+			if (UObject* LoadedChildAsset = ChildHandle->GetLoadedAsset())
+			{
+				return LoadedChildAsset;
+			}
+		}	
 	}
 
 	return nullptr;
