Concurrency feature: retrigger threshold

diff --git a/Engine/Source/Runtime/Engine/Classes/Sound/SoundConcurrency.h b/Engine/Source/Runtime/Engine/Classes/Sound/SoundConcurrency.h
--- a/Engine/Source/Runtime/Engine/Classes/Sound/SoundConcurrency.h
+++ b/Engine/Source/Runtime/Engine/Classes/Sound/SoundConcurrency.h
@@ -73,6 +73,10 @@ struct ENGINE_API FSoundConcurrencySettings
 	UPROPERTY(EditAnywhere, BlueprintReadWrite, Category = Concurrency)
 	TEnumAsByte<EMaxConcurrentResolutionRule::Type> ResolutionRule;
 
+	/** Amount of time to wait (in seconds) between different sounds which play with this concurrency. Sounds rejected from this will ignore virtualization settings. */
+	UPROPERTY(EditAnywhere, BlueprintReadWrite, Category = Concurrency, meta = (UIMin = "0.0", ClampMin = "0.0"))
+	float RetriggerTime;
+
 private:
 	/**
 	 * Ducking factor to apply per older voice instance (generation), which compounds as new voices play
@@ -113,6 +117,7 @@ public:
 		: MaxCount(16)
 		, bLimitToOwner(0)
 		, ResolutionRule(EMaxConcurrentResolutionRule::StopFarthestThenOldest)
+		, RetriggerTime(0.0f)
 		, VolumeScale(1.0f)
 		, VolumeScaleAttackTime(0.01f)
 		, bVolumeScaleCanRelease(0)
@@ -260,6 +265,9 @@ struct FSoundInstanceEntry
 /** Type for mapping an object id to a concurrency entry. */
 typedef TMap<FConcurrencyObjectID, FConcurrencyGroupID> FConcurrencyMap;
 
+/** Type for mapping concurrency group id to when the group last played. */
+typedef TMap<FConcurrencyGroupID, float> FLastTimePlayedMap;
+
 struct FOwnerConcurrencyMapEntry
 {
 	FConcurrencyMap ConcurrencyObjectToConcurrencyGroup;
@@ -299,6 +307,9 @@ public:
 	void UpdateSoundsToCull();
 
 private: // Methods
+	/** Returns whether or not the sound is rate-limited using retrigger threshold */
+	bool IsRateLimited(const FConcurrencyHandle& InHandle);
+
 	/** Evaluates whether or not the sound can play given the concurrency group's rules. Appends permissible
 	sounds to evict in order for sound to play (if required) and returns the desired concurrency group. */
 	FConcurrencyGroup* CanPlaySound(const FActiveSound& NewActiveSound, const FConcurrencyGroupID GroupID, TArray<FActiveSound*>& OutSoundsToEvict, bool bIsRetriggering);
@@ -330,6 +341,9 @@ private: // Data
 	/** Owning audio device ptr for the concurrency manager. */
 	FAudioDevice* AudioDevice;
 
+	/** A map of when a sound last played on the concurrency group. */
+	FLastTimePlayedMap LastTimePlayedMap;
+
 	/** Global concurrency map that maps individual sounds instances to shared USoundConcurrency UObjects. */
 	FConcurrencyMap ConcurrencyMap;
 
diff --git a/Engine/Source/Runtime/Engine/Private/SoundConcurrency.cpp b/Engine/Source/Runtime/Engine/Private/SoundConcurrency.cpp
--- a/Engine/Source/Runtime/Engine/Private/SoundConcurrency.cpp
+++ b/Engine/Source/Runtime/Engine/Private/SoundConcurrency.cpp
@@ -303,6 +303,12 @@ void FSoundConcurrencyManager::CreateNewGroupsFromHandles(
 {
 	for (const FConcurrencyHandle& ConcurrencyHandle : ConcurrencyHandles)
 	{
+		// Add an entry to the map if it hasn't already been added
+		if (!LastTimePlayedMap.Contains(ConcurrencyHandle.ObjectID))
+		{
+			LastTimePlayedMap.Add(ConcurrencyHandle.ObjectID, FPlatformTime::ToSeconds64(FPlatformTime::Cycles64()));
+		}
+	
 		switch (ConcurrencyHandle.GetMode(NewActiveSound))
 		{
 			case EConcurrencyMode::Group:
@@ -415,6 +421,28 @@ FConcurrencyGroup& FSoundConcurrencyManager::CreateNewConcurrencyGroup(const FCo
 	return *ConcurrencyGroups.FindRef(GroupID);
 }
 
+bool FSoundConcurrencyManager::IsRateLimited(const FConcurrencyHandle& InHandle)
+{
+	const float* LastTimePlayed = LastTimePlayedMap.Find(InHandle.ObjectID);
+	const float CurrentTime = FPlatformTime::ToSeconds64(FPlatformTime::Cycles64());
+
+	if (LastTimePlayed)
+	{
+		const float DeltaTime = CurrentTime - *LastTimePlayed;
+
+		if (*LastTimePlayed > 0.0f && InHandle.Settings.RetriggerTime > 0.0f && DeltaTime < InHandle.Settings.RetriggerTime)
+		{
+			// Retrieve the current game time
+			UE_LOG(LogAudio, VeryVerbose, TEXT("Rejected Sound for Group ID (%d) with DeltaTime: %.2f"), InHandle.ObjectID, DeltaTime);
+			return true;
+		}
+
+		LastTimePlayedMap.Add(InHandle.ObjectID, CurrentTime);
+	}
+
+	return false;
+}
+
 FConcurrencyGroup* FSoundConcurrencyManager::CanPlaySound(const FActiveSound& NewActiveSound, const FConcurrencyGroupID GroupID, TArray<FActiveSound*>& OutSoundsToEvict, bool bIsRetriggering)
 {
 	check(GroupID != 0);
@@ -573,6 +601,12 @@ FActiveSound* FSoundConcurrencyManager::EvaluateConcurrency(const FActiveSound&
 
 	for (const FConcurrencyHandle& ConcurrencyHandle : ConcurrencyHandles)
 	{
+		// If this concurrency handle has been rate limited, early exit
+		if (IsRateLimited(ConcurrencyHandle))
+		{
+			return nullptr;
+		}
+
 		switch (ConcurrencyHandle.GetMode(NewActiveSound))
 		{
 			case EConcurrencyMode::Group:
@@ -845,4 +879,4 @@ void FSoundConcurrencyManager::UpdateSoundsToCull()
 	{
 		ConcurrenyGroupEntry.Value->CullSoundsDueToMaxConcurrency();
 	}
-}
\ No newline at end of file
+}
