Add stream cache LLM tags. Also add an LLM tag specifically for the TArray we store compressed audio data in so that we can accurately assess memory overhead of any other state we hold onto (soundwave fnames, etc.)

diff --git a/Engine/Source/Runtime/Core/Private/HAL/LowLevelMemTracker.cpp b/Engine/Source/Runtime/Core/Private/HAL/LowLevelMemTracker.cpp
--- a/Engine/Source/Runtime/Core/Private/HAL/LowLevelMemTracker.cpp
+++ b/Engine/Source/Runtime/Core/Private/HAL/LowLevelMemTracker.cpp
@@ -63,6 +63,8 @@ DECLARE_LLM_MEMORY_STAT(TEXT("AudioDecompress"), STAT_AudioDecompressLLM, STATGR
 DECLARE_LLM_MEMORY_STAT(TEXT("AudioRealtimePrecache"), STAT_AudioRealtimePrecacheLLM, STATGROUP_LLMFULL);
 DECLARE_LLM_MEMORY_STAT(TEXT("AudioFullDecompress"), STAT_AudioFullDecompressLLM, STATGROUP_LLMFULL);
 DECLARE_LLM_MEMORY_STAT(TEXT("AudioVoiceChat"), STAT_AudioVoiceChatLLM, STATGROUP_LLMFULL);
+DECLARE_LLM_MEMORY_STAT(TEXT("AudioStreamCache"), STAT_AudioStreamCacheLLM, STATGROUP_LLMFULL);
+DECLARE_LLM_MEMORY_STAT(TEXT("AudioStreamCacheCompressedData"), STAT_AudioStreamCacheCompressedDataLLM, STATGROUP_LLMFULL);
 DECLARE_LLM_MEMORY_STAT(TEXT("FName"), STAT_FNameLLM, STATGROUP_LLMFULL);
 DECLARE_LLM_MEMORY_STAT(TEXT("Networking"), STAT_NetworkingLLM, STATGROUP_LLMFULL);
 DECLARE_LLM_MEMORY_STAT(TEXT("Meshes"), STAT_MeshesLLM, STATGROUP_LLMFULL);
diff --git a/Engine/Source/Runtime/Core/Public/HAL/LowLevelMemTracker.h b/Engine/Source/Runtime/Core/Public/HAL/LowLevelMemTracker.h
--- a/Engine/Source/Runtime/Core/Public/HAL/LowLevelMemTracker.h
+++ b/Engine/Source/Runtime/Core/Public/HAL/LowLevelMemTracker.h
@@ -164,6 +164,8 @@ enum class ELLMTagSet : uint8
 	macro(AudioRealtimePrecache,				"AudioRealtimePrecache",		GET_STATFNAME(STAT_AudioRealtimePrecacheLLM),				GET_STATFNAME(STAT_AudioSummaryLLM),			ELLMTag::Audio)\
 	macro(AudioFullDecompress,					"AudioFullDecompress",			GET_STATFNAME(STAT_AudioFullDecompressLLM),					GET_STATFNAME(STAT_AudioSummaryLLM),			ELLMTag::Audio)\
 	macro(AudioVoiceChat,						"AudioVoiceChat",				GET_STATFNAME(STAT_AudioVoiceChatLLM),						GET_STATFNAME(STAT_AudioSummaryLLM),			ELLMTag::Audio)\
+	macro(AudioStreamCache,						"AudioStreamCache",				GET_STATFNAME(STAT_AudioStreamCacheLLM),					GET_STATFNAME(STAT_AudioSummaryLLM),			ELLMTag::Audio)\
+	macro(AudioStreamCacheCompressedData,		"AudioStreamCacheCompressedData",GET_STATFNAME(STAT_AudioStreamCacheCompressedDataLLM),		GET_STATFNAME(STAT_AudioSummaryLLM),			ELLMTag::Audio)\
 	macro(FName,								"FName",						GET_STATFNAME(STAT_FNameLLM),								GET_STATFNAME(STAT_EngineSummaryLLM),			-1)\
 	macro(Networking,							"Networking",					GET_STATFNAME(STAT_NetworkingLLM),							GET_STATFNAME(STAT_EngineSummaryLLM),			-1)\
 	macro(Meshes,								"Meshes",						GET_STATFNAME(STAT_MeshesLLM),								GET_STATFNAME(STAT_MeshesSummaryLLM),			-1)\
diff --git a/Engine/Source/Runtime/Engine/Private/AudioStreamingCache.cpp b/Engine/Source/Runtime/Engine/Private/AudioStreamingCache.cpp
--- a/Engine/Source/Runtime/Engine/Private/AudioStreamingCache.cpp
+++ b/Engine/Source/Runtime/Engine/Private/AudioStreamingCache.cpp
@@ -163,6 +163,7 @@ static FAutoConsoleCommand GDisableProfilingAudioCacheCommand(
 
 FCachedAudioStreamingManager::FCachedAudioStreamingManager(const FCachedAudioStreamingManagerParams& InitParams)
 {
+	LLM_SCOPE(ELLMTag::AudioStreamCache);
 	check(FPlatformCompressionUtilities::IsCurrentPlatformUsingStreamCaching());
 	checkf(InitParams.Caches.Num() > 0, TEXT("FCachedAudioStreamingManager should be initialized with dimensions for at least one cache."));
 
@@ -187,6 +188,8 @@ void FCachedAudioStreamingManager::UpdateResourceStreaming(float DeltaTime, bool
 
 int32 FCachedAudioStreamingManager::BlockTillAllRequestsFinished(float TimeLimit, bool bLogResults)
 {
+	LLM_SCOPE(ELLMTag::AudioStreamCache);
+
 	// TODO: Honor TimeLimit and bLogResults. Since we cancel any in flight read requests, this should not spin out.
 	for (FAudioChunkCache& Cache : CacheArray)
 	{
@@ -292,6 +295,7 @@ bool FCachedAudioStreamingManager::IsManagedStreamingSoundSource(const FSoundSou
 
 FAudioChunkHandle FCachedAudioStreamingManager::GetLoadedChunk(const USoundWave* SoundWave, uint32 ChunkIndex, bool bBlockForLoad, bool bForImmediatePlayback) const
 {
+	LLM_SCOPE(ELLMTag::AudioStreamCache);
 	bBlockForLoad |= (ForceBlockForLoadCVar != 0);
 
 	// If this sound wave is managed by a cache, use that to get the chunk:
@@ -378,6 +382,7 @@ FAudioChunkHandle FCachedAudioStreamingManager::GetLoadedChunk(const USoundWave*
 
 FAudioChunkCache* FCachedAudioStreamingManager::GetCacheForWave(const USoundWave* InSoundWave) const
 {
+	LLM_SCOPE(ELLMTag::AudioStreamCache);
 	check(InSoundWave);
 
 	// We only cache chunks beyond the zeroth chunk of audio (which is inlined directly on the asset)
@@ -394,6 +399,7 @@ FAudioChunkCache* FCachedAudioStreamingManager::GetCacheForWave(const USoundWave
 
 FAudioChunkCache* FCachedAudioStreamingManager::GetCacheForChunkSize(uint32 InChunkSize) const
 {
+	LLM_SCOPE(ELLMTag::AudioStreamCache);
 	// Iterate over our caches until we find the lowest MaxChunkSize cache this sound's chunks will fit into. 
 	for (int32 CacheIndex = 0; CacheIndex < CacheArray.Num(); CacheIndex++)
 	{
@@ -413,6 +419,7 @@ FAudioChunkCache* FCachedAudioStreamingManager::GetCacheForChunkSize(uint32 InCh
 
 int32 FCachedAudioStreamingManager::GetNextChunkIndex(const USoundWave* InSoundWave, uint32 CurrentChunkIndex) const
 {
+	LLM_SCOPE(ELLMTag::AudioStreamCache);
 	check(InSoundWave);
 	// TODO: Figure out a way to tell whether this wave is looping or not.
 	// if(bNotLooping) return ((int32) CurrentChunkIndex) < (InSoundWave->RunningPlatformData->Chunks.Num() - 1);
@@ -438,6 +445,7 @@ int32 FCachedAudioStreamingManager::GetNextChunkIndex(const USoundWave* InSoundW
 
 void FCachedAudioStreamingManager::AddReferenceToChunk(const FAudioChunkHandle& InHandle)
 {
+	LLM_SCOPE(ELLMTag::AudioStreamCache);
 	FAudioChunkCache* Cache = GetCacheForChunkSize(InHandle.CachedDataNumBytes);
 	check(Cache);
 
@@ -456,6 +464,7 @@ void FCachedAudioStreamingManager::AddReferenceToChunk(const FAudioChunkHandle&
 
 void FCachedAudioStreamingManager::RemoveReferenceToChunk(const FAudioChunkHandle& InHandle)
 {
+	LLM_SCOPE(ELLMTag::AudioStreamCache);
 	FAudioChunkCache* Cache = GetCacheForChunkSize(InHandle.CachedDataNumBytes);
 	check(Cache);
 
@@ -474,6 +483,7 @@ void FCachedAudioStreamingManager::RemoveReferenceToChunk(const FAudioChunkHandl
 
 bool FCachedAudioStreamingManager::RequestChunk(USoundWave* SoundWave, uint32 ChunkIndex, TFunction<void(EAudioChunkLoadResult)> OnLoadCompleted, ENamedThreads::Type ThreadToCallOnLoadCompletedOn, bool bForImmediatePlayback)
 {
+	LLM_SCOPE(ELLMTag::AudioStreamCache);
 	FAudioChunkCache* Cache = GetCacheForWave(SoundWave);
 	if (Cache)
 	{
@@ -506,6 +516,7 @@ FAudioChunkCache::FAudioChunkCache(uint32 InMaxChunkSize, uint32 NumChunks, uint
 
 FAudioChunkCache::~FAudioChunkCache()
 {
+	LLM_SCOPE(ELLMTag::AudioStreamCache);
 	// While this is handled by the default destructor, we do this to ensure that we don't leak async read operations.
 	CachePool.Reset();
 	check(NumberOfLoadsInFlight.GetValue() == 0);
@@ -699,7 +710,11 @@ uint64 FAudioChunkCache::TrimMemory(uint64 BytesToFree)
 			BytesFreed += CurrentElement->ChunkData.Num();
 			MemoryCounterBytes -= CurrentElement->ChunkData.Num();
 			// Empty the chunk data and invalidate the key.
-			CurrentElement->ChunkData.Empty();
+			{
+				LLM_SCOPE(ELLMTag::AudioStreamCacheCompressedData);
+				CurrentElement->ChunkData.Empty();
+			}
+
 			CurrentElement->ChunkDataSize = 0;
 			CurrentElement->Key = FChunkKey();
 
@@ -1087,7 +1102,6 @@ FAudioChunkCache::FCacheElement* FAudioChunkCache::EvictLeastRecentChunk()
 
 void FAudioChunkCache::KickOffAsyncLoad(FCacheElement* CacheElement, const FChunkKey& InKey, TFunction<void(EAudioChunkLoadResult)> OnLoadCompleted, ENamedThreads::Type CallbackThread, bool bNeededForPlayback)
 {
-	LLM_SCOPE(ELLMTag::Audio);
 	check(CacheElement);
 	
 	const FStreamedAudioChunk& Chunk = InKey.SoundWave->RunningPlatformData->Chunks[InKey.ChunkIndex];
@@ -1096,8 +1110,14 @@ void FAudioChunkCache::KickOffAsyncLoad(FCacheElement* CacheElement, const FChun
 	EAsyncIOPriorityAndFlags AsyncIOPriority = GetAsyncPriorityForChunk(InKey, bNeededForPlayback);
 
 	MemoryCounterBytes -= CacheElement->ChunkData.Num();
-	// Reallocate our chunk data This allows us to shrink if possible.
-	CacheElement->ChunkData.SetNumUninitialized(Chunk.AudioDataSize, true);
+
+	{
+		LLM_SCOPE(ELLMTag::AudioStreamCacheCompressedData);
+
+		// Reallocate our chunk data This allows us to shrink if possible.
+		CacheElement->ChunkData.SetNumUninitialized(Chunk.AudioDataSize, true);
+	}
+
 	MemoryCounterBytes += CacheElement->ChunkData.Num();
 
 #if DEBUG_STREAM_CACHE
@@ -1162,7 +1182,7 @@ void FAudioChunkCache::KickOffAsyncLoad(FCacheElement* CacheElement, const FChun
 		check(ChunkDataSize <= ChunkBulkDataSize);
 		check(((int32)ChunkDataSize) <= CacheElement->ChunkData.Num());
 
-		// If we ever want to eliminate zero-padding in chunks, that could be done here:
+		// If we ever want to eliminate zero-padding in chunks, that could be verified here:
 		//ensureAlwaysMsgf(AudioChunkSize == ChunkBulkDataSize, TEXT("For memory load on demand, we do not zero-pad to page sizes."));
 		
 		NumberOfLoadsInFlight.Increment();
