Add an option to output shader cook stats to UE_LOG(as well as file).

diff --git a/Engine/Source/Runtime/Engine/Private/ShaderCompiler/ShaderCompiler.cpp b/Engine/Source/Runtime/Engine/Private/ShaderCompiler/ShaderCompiler.cpp
--- a/Engine/Source/Runtime/Engine/Private/ShaderCompiler/ShaderCompiler.cpp
+++ b/Engine/Source/Runtime/Engine/Private/ShaderCompiler/ShaderCompiler.cpp
@@ -160,6 +160,15 @@ static FAutoConsoleVariableRef CVarDumpSCWJobInfoOnCrash(
 	TEXT("When set to 1, it will dump a job list to help track down crashes that happened on ShaderCompileWorker.")
 );
 
+
+static int32 GLogShaderCompilerStats = 0;
+static FAutoConsoleVariableRef CVarLogShaderCompilerStats(
+	TEXT("r.LogShaderCompilerStats"),
+	GLogShaderCompilerStats,
+	TEXT("When set to 1, Log detailed shader compiler stats.")
+);
+
+
 static int32 GShowShaderWarnings = 0;
 static FAutoConsoleVariableRef CVarShowShaderWarnings(
 	TEXT("r.ShowShaderCompilerWarnings"),
@@ -1635,6 +1644,10 @@ void FShaderCompilerStats::WriteStats()
 					StatWriter.AddColumn(TEXT("%u"), SingleStats.CompiledDouble);
 					StatWriter.AddColumn(TEXT("%u"), SingleStats.CookedDouble);
 					StatWriter.CycleRow();
+					if(GLogShaderCompilerStats)
+					{
+						UE_LOG(LogShaderCompilers, Log, TEXT("SHADERSTATS %s, %u, %u, %u, %u, %u, %u\n"), *Path, Platform, SingleStats.Compiled, SingleStats.Cooked, SingleStats.PermutationCompilations.Num(), SingleStats.CompiledDouble, SingleStats.CookedDouble);
+					}
 				}
 			}
 		}
