Fix physics interface inconsistency for physx immediate mode

diff --git a/Engine/Source/Runtime/Engine/Private/PhysicsEngine/Experimental/PhysScene_ImmediatePhysX.cpp b/Engine/Source/Runtime/Engine/Private/PhysicsEngine/Experimental/PhysScene_ImmediatePhysX.cpp
--- a/Engine/Source/Runtime/Engine/Private/PhysicsEngine/Experimental/PhysScene_ImmediatePhysX.cpp
+++ b/Engine/Source/Runtime/Engine/Private/PhysicsEngine/Experimental/PhysScene_ImmediatePhysX.cpp
@@ -706,7 +706,7 @@ void FPhysScene_ImmediatePhysX::DispatchPhysNotifications_AssumesLocked()
     FPhysicsDelegates::OnPhysDispatchNotifications.Broadcast(this);
 }
 
-void FPhysScene_ImmediatePhysX::SetUpForFrame(const FVector* NewGrav, float InDeltaSeconds, float InMaxPhysicsDeltaTime)
+void FPhysScene_ImmediatePhysX::SetUpForFrame(const FVector* NewGrav, float InDeltaSeconds, float InMaxPhysicsDeltaTime, float InMaxSubstepDeltaTime, int32 InMaxSubsteps)
 {
 	DeltaSeconds = FMath::Min(InDeltaSeconds, 0.033f);
 	//Create dynamic bodies and integrate their unconstrained velocities
diff --git a/Engine/Source/Runtime/Engine/Public/Physics/Experimental/PhysScene_ImmediatePhysX.h b/Engine/Source/Runtime/Engine/Public/Physics/Experimental/PhysScene_ImmediatePhysX.h
--- a/Engine/Source/Runtime/Engine/Public/Physics/Experimental/PhysScene_ImmediatePhysX.h
+++ b/Engine/Source/Runtime/Engine/Public/Physics/Experimental/PhysScene_ImmediatePhysX.h
@@ -490,7 +490,7 @@ public:
 	void ApplyWorldOffset(FVector InOffset);
 
 	/** Set the gravity and timing of all physics scenes */
-	ENGINE_API void SetUpForFrame(const FVector* NewGrav, float InDeltaSeconds = 0.0f, float InMaxPhysicsDeltaTime = 0.0f);
+	ENGINE_API void SetUpForFrame(const FVector* NewGrav, float InDeltaSeconds = 0.0f, float InMaxPhysicsDeltaTime = 0.0f, float InMaxSubstepDeltaTime = 0.f, int32 InMaxSubsteps = 1);
 
 	/** Starts a frame */
 	ENGINE_API void StartFrame();
