Fix uninitialized static arrays in cooked PC build
#rnx

diff --git a/Engine/Source/Runtime/Engine/Private/Audio/AudioAddressPattern.cpp b/Engine/Source/Runtime/Engine/Private/Audio/AudioAddressPattern.cpp
--- a/Engine/Source/Runtime/Engine/Private/Audio/AudioAddressPattern.cpp
+++ b/Engine/Source/Runtime/Engine/Private/Audio/AudioAddressPattern.cpp
@@ -2,11 +2,6 @@
 // Pattern-Matching based on OSC 1.0 Protocol
 #include "Audio/AudioAddressPattern.h"
 
-namespace
-{
-	const TArray<TCHAR> AddressInvalidChars = { ' ', '#', };
-	const TArray<TCHAR> AddressPatternChars = { ',', '*', '?', '[', ']', '{', '}' };
-} // namespace <>
 
 bool FAudioAddressPattern::BracePatternMatches(const FString& Pattern, int32 PatternStartIndex, int32 PatternEndIndex, const FString& Part, int32& PartIter)
 {
@@ -133,6 +128,18 @@ int32 FAudioAddressPattern::FindPatternTerminatorIndex(const FString& Pattern, i
 	return EndIndex;
 }
 
+const TArray<TCHAR>& FAudioAddressPattern::GetInvalidChars()
+{
+	static const TArray<TCHAR> InvalidChars = { ' ', '#', };
+	return InvalidChars;
+}
+
+const TArray<TCHAR>& FAudioAddressPattern::GetPatternChars()
+{
+	static const TArray<TCHAR> PatternChars = { ',', '*', '?', '[', ']', '{', '}' };
+	return PatternChars;
+}
+
 bool FAudioAddressPattern::IsValidPatternPart(const FString& Part)
 {
 	bool bInBrackets = false;
@@ -145,7 +152,7 @@ bool FAudioAddressPattern::IsValidPatternPart(const FString& Part)
 			return false;
 		}
 
-		if (AddressInvalidChars.Contains(Char))
+		if (GetInvalidChars().Contains(Char))
 		{
 			return false;
 		}
@@ -258,12 +265,12 @@ bool FAudioAddressPattern::IsValidPath(const FString& Path, bool bInvalidateSepa
 			continue;
 		}
 
-		if (AddressInvalidChars.Contains(Char))
+		if (GetInvalidChars().Contains(Char))
 		{
 			return false;
 		}
 
-		if (AddressPatternChars.Contains(Char))
+		if (GetPatternChars().Contains(Char))
 		{
 			return false;
 		}
diff --git a/Engine/Source/Runtime/Engine/Public/Audio/AudioAddressPattern.h b/Engine/Source/Runtime/Engine/Public/Audio/AudioAddressPattern.h
--- a/Engine/Source/Runtime/Engine/Public/Audio/AudioAddressPattern.h
+++ b/Engine/Source/Runtime/Engine/Public/Audio/AudioAddressPattern.h
@@ -8,6 +8,9 @@ class ENGINE_API FAudioAddressPattern
 {
 	static int32 FindPatternTerminatorIndex(const FString& Pattern, int32 PatternIter, TCHAR Terminator);
 
+	static const TArray<TCHAR>& GetInvalidChars();
+	static const TArray<TCHAR>& GetPatternChars();
+	
 public:
 	static bool BracePatternMatches(const FString& Pattern, int32 PatternStartIndex, int32 PatternEndIndex, const FString& Part, int32& PartIter);
 	static bool BracketPatternMatches(const FString& Pattern, int32 PatternStartIndex, int32 PatternEndIndex, TCHAR MatchChar);
