Reduce the noise of duplicate notifications when changing the mobility of blueprint component by eliminating the notification for each instance of it in the world.
#rnx

diff --git a/Engine/Source/Runtime/Engine/Private/Components/SceneComponent.cpp b/Engine/Source/Runtime/Engine/Private/Components/SceneComponent.cpp
--- a/Engine/Source/Runtime/Engine/Private/Components/SceneComponent.cpp
+++ b/Engine/Source/Runtime/Engine/Private/Components/SceneComponent.cpp
@@ -266,6 +266,7 @@ static int32 SetAncestorMobility(USceneComponent const* SceneComponentObject, EC
  * This method walks the hierarchy and alters parent/child component's Mobility as a result of
  * this property change.
  */
+bool GNotifyAboutMobilityUpdate = true;
 static void UpdateAttachedMobility(USceneComponent* ComponentThatChanged)
 {
 	// Attached parent components can't be more mobile than their children. This means that 
@@ -319,19 +320,22 @@ static void UpdateAttachedMobility(USceneComponent* ComponentThatChanged)
 			{
 				ComponentThatChanged->DetachFromComponent(FDetachmentTransformRules::KeepWorldTransform);
 
-				// Fire off a notification
-				FText NotificationText = FText::Format(LOCTEXT("ComponentDetachedFromParentDueToMobility", "Caused {0} to be detached from its parent {1} because it does not allow to be static"), FText::FromName(ComponentThatChanged->GetFName()), FText::FromName(ParentComponent->GetFName()));
-				FNotificationInfo Info(NotificationText);
-				Info.bFireAndForget = true;
-				Info.bUseThrobber = true;
-				Info.ExpireDuration = 2.0f;
-				FSlateNotificationManager::Get().AddNotification(Info);
+				if (GNotifyAboutMobilityUpdate)
+				{
+					// Fire off a notification
+					FText NotificationText = FText::Format(LOCTEXT("ComponentDetachedFromParentDueToMobility", "Caused {0} to be detached from its parent {1} because it does not allow to be static"), FText::FromName(ComponentThatChanged->GetFName()), FText::FromName(ParentComponent->GetFName()));
+					FNotificationInfo Info(NotificationText);
+					Info.bFireAndForget = true;
+					Info.bUseThrobber = true;
+					Info.ExpireDuration = 2.0f;
+					FSlateNotificationManager::Get().AddNotification(Info);
+				}
 			}
 		}		
 	}
 
 	// if we altered any components (other than the ones selected), then notify the user
-	if(NumMobilityChanges > 0)
+	if (GNotifyAboutMobilityUpdate && NumMobilityChanges > 0)
 	{
 		FText NotificationText = LOCTEXT("MobilityAlteredSingularNotification", "Caused 1 component to also change Mobility");
 		if(NumMobilityChanges > 1)
@@ -502,11 +506,11 @@ void USceneComponent::PostEditChangeProperty(FPropertyChangedEvent& PropertyChan
 	{
 		UpdateAttachedMobility(this);
 	}
-	if (bIsEditorOnly && PropertyName == GET_MEMBER_NAME_CHECKED(UActorComponent, bIsEditorOnly))
+	else if (bIsEditorOnly && PropertyName == GET_MEMBER_NAME_CHECKED(UActorComponent, bIsEditorOnly))
 	{
 		UpdateAttachedIsEditorOnly(this);
 	}
-	if (PropertyName == USceneComponent::GetVisiblePropertyName())
+	else if (PropertyName == USceneComponent::GetVisiblePropertyName())
 	{
 		OnVisibilityChanged();
 	}
@@ -515,6 +519,9 @@ void USceneComponent::PostEditChangeProperty(FPropertyChangedEvent& PropertyChan
 		OnHiddenInGameChanged();
 	}
 
+	// If this is a template object when the property change is propagated to instances we don't want duplicate notification toasts
+	TGuardValue<bool> MobilityNotificationGuard(GNotifyAboutMobilityUpdate, (IsTemplate() ? false : GNotifyAboutMobilityUpdate));
+
 	Super::PostEditChangeProperty(PropertyChangedEvent);
 
 	const bool bLocationChanged = (PropertyName == LocationName || MemberPropertyName == LocationName);
@@ -540,6 +547,9 @@ void USceneComponent::PostEditChangeChainProperty(FPropertyChangedChainEvent& Pr
 		UpdateAttachedIsEditorOnly(this);
 	}
 
+	// If this is a template object when the property change is propagated to instances we don't want duplicate notification toasts
+	TGuardValue<bool> MobilityNotificationGuard(GNotifyAboutMobilityUpdate, (IsTemplate() ? false : GNotifyAboutMobilityUpdate));
+
 	Super::PostEditChangeChainProperty(PropertyChangedEvent);
 }
 
