Fixed invalid IsInGameThread() assert during ray tracing ISM preprocessing.

diff --git a/Engine/Source/Runtime/Engine/Private/InstancedStaticMesh.cpp b/Engine/Source/Runtime/Engine/Private/InstancedStaticMesh.cpp
--- a/Engine/Source/Runtime/Engine/Private/InstancedStaticMesh.cpp
+++ b/Engine/Source/Runtime/Engine/Private/InstancedStaticMesh.cpp
@@ -1250,8 +1250,6 @@ void FInstancedStaticMeshSceneProxy::SetupRayTracingCullClusters()
 		return;
 	}
 
-	check(IsInGameThread());
-
 	//#dxr_todo: select the appropriate LOD depending on Context.View
 	int32 LOD = 0;
 	if (RenderData->LODResources.Num() > LOD && RenderData->LODResources[LOD].RayTracingGeometry.IsInitialized())
