Don't initialize XR devices when no project is loaded

diff --git a/Engine/Source/Runtime/Engine/Private/UnrealEngine.cpp b/Engine/Source/Runtime/Engine/Private/UnrealEngine.cpp
--- a/Engine/Source/Runtime/Engine/Private/UnrealEngine.cpp
+++ b/Engine/Source/Runtime/Engine/Private/UnrealEngine.cpp
@@ -1489,11 +1489,14 @@ void UEngine::Init(IEngineLoop* InEngineLoop)
 
 	FCoreUObjectDelegates::GetPreGarbageCollectDelegate().AddStatic(UEngine::PreGarbageCollect);
 
-	// Initialize the HMDs and motion controllers, if any
-	InitializeHMDDevice();
+	if (!FApp::IsProjectNameEmpty())
+	{
+		// Initialize the HMDs and motion controllers, if any
+		InitializeHMDDevice();
 
-	// Initialize attached eye tracking devices, if any
-	InitializeEyeTrackingDevice();
+		// Initialize attached eye tracking devices, if any
+		InitializeEyeTrackingDevice();
+	}
 
 	// Disable the screensaver when running the game.
 	if( GIsClient && !GIsEditor )
