Fix case in PrimitiveComponent's MoveComponent code where we would possibly choose the wrong blocking hit.

diff --git a/Engine/Source/Runtime/Engine/Private/Components/PrimitiveComponent.cpp b/Engine/Source/Runtime/Engine/Private/Components/PrimitiveComponent.cpp
--- a/Engine/Source/Runtime/Engine/Private/Components/PrimitiveComponent.cpp
+++ b/Engine/Source/Runtime/Engine/Private/Components/PrimitiveComponent.cpp
@@ -2148,7 +2148,7 @@ bool UPrimitiveComponent::MoveComponentImpl( const FVector& Delta, const FQuat&
 					{
 						if (!ShouldIgnoreHitResult(MyWorld, TestHit, Delta, Actor, MoveFlags))
 						{
-							if (TestHit.Time == 0.f)
+							if (TestHit.bStartPenetrating)
 							{
 								// We may have multiple initial hits, and want to choose the one with the normal most opposed to our movement.
 								const float NormalDotDelta = (TestHit.ImpactNormal | Delta);
