Pushed various material and expression properties to be editor-only. Saves about 500kb in a basic test scenario with a few thousand objects.

diff --git a/Engine/Source/Runtime/Engine/Classes/Materials/MaterialExpressionAntialiasedTextureMask.h b/Engine/Source/Runtime/Engine/Classes/Materials/MaterialExpressionAntialiasedTextureMask.h
--- a/Engine/Source/Runtime/Engine/Classes/Materials/MaterialExpressionAntialiasedTextureMask.h
+++ b/Engine/Source/Runtime/Engine/Classes/Materials/MaterialExpressionAntialiasedTextureMask.h
@@ -31,18 +31,17 @@ class UMaterialExpressionAntialiasedTextureMask : public UMaterialExpressionText
 	UPROPERTY(EditAnywhere, Category=MaterialExpressionAntialiasedTextureMask)
 	TEnumAsByte<enum ETextureColorChannel> Channel;
 
-
+#if WITH_EDITOR
 	//~ Begin UMaterialExpressionTextureSampleParameter Interface
 	virtual bool TextureIsValid(UTexture* InTexture, FString& OutMessage) override;
 	virtual void SetDefaultTexture() override;
 	//~ End UMaterialExpressionTextureSampleParameter Interface
 
 	//~ Begin UMaterialExpression Interface
-#if WITH_EDITOR
 	virtual int32 Compile(class FMaterialCompiler* Compiler, int32 OutputIndex) override;
 	virtual void GetCaption(TArray<FString>& OutCaptions) const override;
-#endif
 	//~ End UMaterialExpression Interface
+#endif // WITH_EDITOR
 };
 
 
diff --git a/Engine/Source/Runtime/Engine/Classes/Materials/MaterialExpressionChannelMaskParameter.h b/Engine/Source/Runtime/Engine/Classes/Materials/MaterialExpressionChannelMaskParameter.h
--- a/Engine/Source/Runtime/Engine/Classes/Materials/MaterialExpressionChannelMaskParameter.h
+++ b/Engine/Source/Runtime/Engine/Classes/Materials/MaterialExpressionChannelMaskParameter.h
@@ -27,8 +27,10 @@ class UMaterialExpressionChannelMaskParameter : public UMaterialExpressionVector
 	UPROPERTY(EditAnywhere, Category=MaterialExpressionChannelMaskParameter)
 	TEnumAsByte<EChannelMaskParameterColor::Type> MaskChannel;
 
+#if WITH_EDITORONLY_DATA
 	UPROPERTY()
 	FExpressionInput Input;
+#endif
 
 #if WITH_EDITOR
 	virtual bool SetParameterValue(FName InParameterName, FLinearColor InValue) override;
diff --git a/Engine/Source/Runtime/Engine/Classes/Materials/MaterialExpressionFeatureLevelSwitch.h b/Engine/Source/Runtime/Engine/Classes/Materials/MaterialExpressionFeatureLevelSwitch.h
--- a/Engine/Source/Runtime/Engine/Classes/Materials/MaterialExpressionFeatureLevelSwitch.h
+++ b/Engine/Source/Runtime/Engine/Classes/Materials/MaterialExpressionFeatureLevelSwitch.h
@@ -18,8 +18,10 @@ class UMaterialExpressionFeatureLevelSwitch : public UMaterialExpression
 	UPROPERTY()
 	FExpressionInput Default;
 
+#if WITH_EDITORONLY_DATA
 	UPROPERTY()
 	FExpressionInput Inputs[ERHIFeatureLevel::Num];
+#endif
 
 	//~ Begin UObject Interface.
 	virtual void Serialize(FStructuredArchive::FRecord Record) override;
diff --git a/Engine/Source/Runtime/Engine/Classes/Materials/MaterialExpressionStaticComponentMaskParameter.h b/Engine/Source/Runtime/Engine/Classes/Materials/MaterialExpressionStaticComponentMaskParameter.h
--- a/Engine/Source/Runtime/Engine/Classes/Materials/MaterialExpressionStaticComponentMaskParameter.h
+++ b/Engine/Source/Runtime/Engine/Classes/Materials/MaterialExpressionStaticComponentMaskParameter.h
@@ -14,8 +14,10 @@ class UMaterialExpressionStaticComponentMaskParameter : public UMaterialExpressi
 {
 	GENERATED_UCLASS_BODY()
 
+#if WITH_EDITORONLY_DATA
 	UPROPERTY()
 	FExpressionInput Input;
+#endif
 
 	UPROPERTY(EditAnywhere, Category=MaterialExpressionStaticComponentMaskParameter)
 	uint32 DefaultR:1;
diff --git a/Engine/Source/Runtime/Engine/Classes/Materials/MaterialExpressionStaticSwitchParameter.h b/Engine/Source/Runtime/Engine/Classes/Materials/MaterialExpressionStaticSwitchParameter.h
--- a/Engine/Source/Runtime/Engine/Classes/Materials/MaterialExpressionStaticSwitchParameter.h
+++ b/Engine/Source/Runtime/Engine/Classes/Materials/MaterialExpressionStaticSwitchParameter.h
@@ -14,12 +14,13 @@ class UMaterialExpressionStaticSwitchParameter : public UMaterialExpressionStati
 {
 	GENERATED_UCLASS_BODY()
 
+#if WITH_EDITORONLY_DATA
 	UPROPERTY()
 	FExpressionInput A;
 
 	UPROPERTY()
 	FExpressionInput B;
-
+#endif
 
 	//~ Begin UMaterialExpression Interface
 #if WITH_EDITOR
diff --git a/Engine/Source/Runtime/Engine/Classes/Materials/MaterialExpressionTextureBase.h b/Engine/Source/Runtime/Engine/Classes/Materials/MaterialExpressionTextureBase.h
--- a/Engine/Source/Runtime/Engine/Classes/Materials/MaterialExpressionTextureBase.h
+++ b/Engine/Source/Runtime/Engine/Classes/Materials/MaterialExpressionTextureBase.h
@@ -23,12 +23,14 @@ class ENGINE_API UMaterialExpressionTextureBase : public UMaterialExpression
 	UPROPERTY(EditAnywhere, BlueprintReadWrite, Category=MaterialExpressionTextureBase)
 	class UTexture* Texture;
 
+#if WITH_EDITORONLY_DATA
 	UPROPERTY(EditAnywhere, BlueprintReadWrite, Category=MaterialExpressionTextureBase)
 	TEnumAsByte<enum EMaterialSamplerType> SamplerType;
 	
 	/** Is default selected texture when using mesh paint mode texture painting */
 	UPROPERTY(EditAnywhere, Category=MaterialExpressionTextureBase)
 	uint8 IsDefaultMeshpaintTexture:1;
+#endif // WITH_EDITORONLY_DATA
 	
 	//~ Begin UObject Interface
 #if WITH_EDITOR
@@ -52,6 +54,7 @@ class ENGINE_API UMaterialExpressionTextureBase : public UMaterialExpression
 	virtual UObject* GetReferencedTexture() const override;
 	virtual bool CanReferenceTexture() const override { return true; }
 
+#if WITH_EDITOR
 	/**
 	 * Automatically determines and set the sampler type for the current texture.
 	 */
@@ -63,4 +66,5 @@ class ENGINE_API UMaterialExpressionTextureBase : public UMaterialExpression
 	 * @returns the default sampler type for the specified texture.
 	 */
 	static EMaterialSamplerType GetSamplerTypeForTexture( const UTexture* Texture, bool ForceNoVT = false );
+#endif // WITH_EDITOR
 };
diff --git a/Engine/Source/Runtime/Engine/Classes/Materials/MaterialExpressionTextureObjectParameter.h b/Engine/Source/Runtime/Engine/Classes/Materials/MaterialExpressionTextureObjectParameter.h
--- a/Engine/Source/Runtime/Engine/Classes/Materials/MaterialExpressionTextureObjectParameter.h
+++ b/Engine/Source/Runtime/Engine/Classes/Materials/MaterialExpressionTextureObjectParameter.h
@@ -17,10 +17,11 @@ class UMaterialExpressionTextureObjectParameter : public UMaterialExpressionText
 {
 	GENERATED_UCLASS_BODY()
 	
+#if WITH_EDITOR
 	//~ Begin UMaterialExpressionTextureSampleParameter Interface
 	virtual bool TextureIsValid(UTexture* InTexture, FString& OutMessage) override;
 	//~ End UMaterialExpressionTextureSampleParameter Interface
-
+#endif
 
 	//~ Begin UMaterialExpression Interface
 #if WITH_EDITOR
diff --git a/Engine/Source/Runtime/Engine/Classes/Materials/MaterialExpressionTextureSample.h b/Engine/Source/Runtime/Engine/Classes/Materials/MaterialExpressionTextureSample.h
--- a/Engine/Source/Runtime/Engine/Classes/Materials/MaterialExpressionTextureSample.h
+++ b/Engine/Source/Runtime/Engine/Classes/Materials/MaterialExpressionTextureSample.h
@@ -38,6 +38,7 @@ class ENGINE_API UMaterialExpressionTextureSample : public UMaterialExpressionTe
 	UPROPERTY(meta = (RequiredInput = "false", ToolTip = "Defaults to 'ConstCoordinate' if not specified"))
 	FExpressionInput Coordinates;
 
+#if WITH_EDITORONLY_DATA
 	/** 
 	 * Texture object input which overrides Texture if specified. 
 	 * This only shows up in material functions and is used to implement texture parameters without actually putting the texture parameter in the function.
@@ -90,6 +91,7 @@ public:
 	/** only used if MipValue is not hooked up */
 	UPROPERTY(EditAnywhere, Category = MaterialExpressionTextureSample)
 	int32 ConstMipValue;
+#endif // WITH_EDITORONLY_DATA
 
 	//~ Begin UObject Interface
 #if WITH_EDITOR
diff --git a/Engine/Source/Runtime/Engine/Classes/Materials/MaterialExpressionTextureSampleParameter.h b/Engine/Source/Runtime/Engine/Classes/Materials/MaterialExpressionTextureSampleParameter.h
--- a/Engine/Source/Runtime/Engine/Classes/Materials/MaterialExpressionTextureSampleParameter.h
+++ b/Engine/Source/Runtime/Engine/Classes/Materials/MaterialExpressionTextureSampleParameter.h
@@ -61,7 +61,6 @@ class ENGINE_API UMaterialExpressionTextureSampleParameter : public UMaterialExp
 	bool SetParameterValue(FName InParameterName, UTexture* InValue);
 	virtual void PostEditChangeProperty(FPropertyChangedEvent& PropertyChangedEvent) override;
 	void ApplyChannelNames();
-#endif
 
 	/**
 	 * Return true if the texture is a movie texture
@@ -76,6 +75,7 @@ class ENGINE_API UMaterialExpressionTextureSampleParameter : public UMaterialExp
 	 *	Sets the default texture if none is set
 	 */
 	virtual void SetDefaultTexture();
+#endif // WITH_EDITOR
 
 	virtual FGuid& GetParameterExpressionId() override
 	{
diff --git a/Engine/Source/Runtime/Engine/Classes/Materials/MaterialExpressionTextureSampleParameter2DArray.h b/Engine/Source/Runtime/Engine/Classes/Materials/MaterialExpressionTextureSampleParameter2DArray.h
--- a/Engine/Source/Runtime/Engine/Classes/Materials/MaterialExpressionTextureSampleParameter2DArray.h
+++ b/Engine/Source/Runtime/Engine/Classes/Materials/MaterialExpressionTextureSampleParameter2DArray.h
@@ -19,7 +19,9 @@ class UMaterialExpressionTextureSampleParameter2DArray : public UMaterialExpress
 #endif
 
 	// Begin UMaterialExpressionTextureSampleParameter Interface
+#if WITH_EDITOR
 	virtual bool TextureIsValid( UTexture* InTexture, FString& OutMessage) override;
+#endif
 	virtual const TCHAR* GetRequirements();
 	// End UMaterialExpressionTextureSampleParameter Interface
 };
\ No newline at end of file
diff --git a/Engine/Source/Runtime/Engine/Classes/Materials/MaterialExpressionTextureSampleParameterCube.h b/Engine/Source/Runtime/Engine/Classes/Materials/MaterialExpressionTextureSampleParameterCube.h
--- a/Engine/Source/Runtime/Engine/Classes/Materials/MaterialExpressionTextureSampleParameterCube.h
+++ b/Engine/Source/Runtime/Engine/Classes/Materials/MaterialExpressionTextureSampleParameterCube.h
@@ -15,21 +15,15 @@ class UMaterialExpressionTextureSampleParameterCube : public UMaterialExpression
 {
 	GENERATED_UCLASS_BODY()
 
-
-
-
-	//~ Begin UMaterialExpression Interface
 #if WITH_EDITOR
+	//~ Begin UMaterialExpression Interface
 	virtual void GetCaption(TArray<FString>& OutCaptions) const override;
 	virtual int32 Compile(class FMaterialCompiler* Compiler, int32 OutputIndex) override;
-#endif
 	//~ End UMaterialExpression Interface
 
 	//~ Begin UMaterialExpressionTextureSampleParameter Interface
 	virtual bool TextureIsValid(UTexture* InTexture, FString& OutMessage) override;
 	virtual void SetDefaultTexture() override;
 	//~ End UMaterialExpressionTextureSampleParameter Interface
+#endif // WITH_EDITOR
 };
-
-
-
diff --git a/Engine/Source/Runtime/Engine/Classes/Materials/MaterialExpressionTextureSampleParameterSubUV.h b/Engine/Source/Runtime/Engine/Classes/Materials/MaterialExpressionTextureSampleParameterSubUV.h
--- a/Engine/Source/Runtime/Engine/Classes/Materials/MaterialExpressionTextureSampleParameterSubUV.h
+++ b/Engine/Source/Runtime/Engine/Classes/Materials/MaterialExpressionTextureSampleParameterSubUV.h
@@ -18,18 +18,14 @@ class UMaterialExpressionTextureSampleParameterSubUV : public UMaterialExpressio
 	UPROPERTY(EditAnywhere, Category=MaterialExpressionTextureSampleParameterSubUV)
 	uint32 bBlend:1;
 
-
+#if WITH_EDITOR
 	//~ Begin UMaterialExpressionTextureSampleParameter Interface
 	virtual bool TextureIsValid(UTexture* InTexture, FString& OutMessage) override;
 	//~ End UMaterialExpressionTextureSampleParameter Interface
 
 	//~ Begin UMaterialExpression Interface
-#if WITH_EDITOR
 	virtual int32 Compile(class FMaterialCompiler* Compiler, int32 OutputIndex) override;
 	virtual void GetCaption(TArray<FString>& OutCaptions) const override;
-#endif
 	//~ End UMaterialExpression Interface
+#endif // WITH_EDITOR
 };
-
-
-
diff --git a/Engine/Source/Runtime/Engine/Classes/Materials/MaterialExpressionTextureSampleParameterVolume.h b/Engine/Source/Runtime/Engine/Classes/Materials/MaterialExpressionTextureSampleParameterVolume.h
--- a/Engine/Source/Runtime/Engine/Classes/Materials/MaterialExpressionTextureSampleParameterVolume.h
+++ b/Engine/Source/Runtime/Engine/Classes/Materials/MaterialExpressionTextureSampleParameterVolume.h
@@ -1,6 +1,5 @@
 // Copyright Epic Games, Inc. All Rights Reserved.
 
-
 #pragma once
 
 #include "CoreMinimal.h"
@@ -14,18 +13,16 @@ UCLASS(collapsecategories, hidecategories=Object, MinimalAPI)
 class UMaterialExpressionTextureSampleParameterVolume : public UMaterialExpressionTextureSampleParameter
 {
 	GENERATED_UCLASS_BODY()
-	//~ Begin UMaterialExpression Interface
+
 #if WITH_EDITOR
+	//~ Begin UMaterialExpression Interface
 	virtual void GetCaption(TArray<FString>& OutCaptions) const override;
 	virtual int32 Compile(class FMaterialCompiler* Compiler, int32 OutputIndex) override;
-#endif
 	//~ End UMaterialExpression Interface
 
 	//~ Begin UMaterialExpressionTextureSampleParameter Interface
 	virtual bool TextureIsValid(UTexture* InTexture, FString& OutMessage) override;
 	virtual void SetDefaultTexture() override;
 	//~ End UMaterialExpressionTextureSampleParameter Interface
+#endif // WITH_EDITOR
 };
-
-
-
diff --git a/Engine/Source/Runtime/Engine/Classes/Materials/MaterialFunction.h b/Engine/Source/Runtime/Engine/Classes/Materials/MaterialFunction.h
--- a/Engine/Source/Runtime/Engine/Classes/Materials/MaterialFunction.h
+++ b/Engine/Source/Runtime/Engine/Classes/Materials/MaterialFunction.h
@@ -74,12 +74,12 @@ class UMaterialFunction : public UMaterialFunctionInterface
 
 	UPROPERTY()
 	TArray<class UMaterialExpressionMaterialFunctionCall*> DependentFunctionExpressionCandidates;
-#endif // WITH_EDITORONLY_DATA
 
 private:
 	/** Transient flag used to track re-entrance in recursive functions like IsDependent. */
 	UPROPERTY(transient)
 	uint8 bReentrantFlag:1;
+#endif // WITH_EDITORONLY_DATA
 
 public:
 	//~ Begin UObject Interface.
@@ -151,8 +151,10 @@ public:
 #endif
 	virtual const FString* GetDescription() const override { return &Description; }
 
+#if WITH_EDITOR
 	virtual bool GetReentrantFlag() const override { return bReentrantFlag; }
 	virtual void SetReentrantFlag(const bool bIsReentrant) override { bReentrantFlag = bIsReentrant; }
+#endif
 	//~ End UMaterialFunctionInterface interface
 
 
diff --git a/Engine/Source/Runtime/Engine/Classes/Materials/MaterialFunctionInstance.h b/Engine/Source/Runtime/Engine/Classes/Materials/MaterialFunctionInstance.h
--- a/Engine/Source/Runtime/Engine/Classes/Materials/MaterialFunctionInstance.h
+++ b/Engine/Source/Runtime/Engine/Classes/Materials/MaterialFunctionInstance.h
@@ -151,6 +151,7 @@ class UMaterialFunctionInstance : public UMaterialFunctionInterface
 		return BaseFunction ? BaseFunction->GetDescription() : nullptr;
 	}
 
+#if WITH_EDITOR
 	virtual bool GetReentrantFlag() const override
 	{
 		const UMaterialFunctionInterface* BaseFunction = GetBaseFunction();
@@ -164,6 +165,7 @@ class UMaterialFunctionInstance : public UMaterialFunctionInterface
 			BaseFunction->SetReentrantFlag(bIsReentrant);
 		}
 	}
+#endif // WITH_EDITOR
 
 public:
 	virtual bool OverrideNamedScalarParameter(const FHashedMaterialParameterInfo& ParameterInfo, float& OutValue) override;
diff --git a/Engine/Source/Runtime/Engine/Classes/Materials/MaterialInstance.h b/Engine/Source/Runtime/Engine/Classes/Materials/MaterialInstance.h
--- a/Engine/Source/Runtime/Engine/Classes/Materials/MaterialInstance.h
+++ b/Engine/Source/Runtime/Engine/Classes/Materials/MaterialInstance.h
@@ -365,12 +365,18 @@ class UMaterialInstance : public UMaterialInterface
 
 	FORCEINLINE bool GetReentrantFlag() const
 	{
+#if WITH_EDITOR
 		return ReentrantFlag[IsInGameThread() ? 0 : 1];
+#else
+		return false;
+#endif
 	}
 
 	FORCEINLINE void SetReentrantFlag(const bool bValue)
 	{
+#if WITH_EDITOR
 		ReentrantFlag[IsInGameThread() ? 0 : 1] = bValue;
+#endif
 	}
 
 	/** Scalar parameters. */
@@ -408,8 +414,10 @@ class UMaterialInstance : public UMaterialInterface
 
 	FMaterialShadingModelField ShadingModels;
 
-	/** Flag to detect cycles in the material instance graph. */
+#if WITH_EDITOR
+	/** Flag to detect cycles in the material instance graph, this is only used at content creation time where the hierarchy can be changed. */
 	bool ReentrantFlag[2];
+#endif
 
 	/** 
 	 * FMaterialRenderProxy derivative that represent this material instance to the renderer, when the renderer needs to fetch parameter values. 
diff --git a/Engine/Source/Runtime/Engine/Private/Materials/MaterialExpressions.cpp b/Engine/Source/Runtime/Engine/Private/Materials/MaterialExpressions.cpp
--- a/Engine/Source/Runtime/Engine/Private/Materials/MaterialExpressions.cpp
+++ b/Engine/Source/Runtime/Engine/Private/Materials/MaterialExpressions.cpp
@@ -1318,10 +1318,6 @@ bool UMaterialExpression::HasConnectedOutputs() const
 	return bIsConnected;
 }
 
-
-#endif // WITH_EDITOR
-
-
 bool UMaterialExpression::ContainsInputLoop(const bool bStopOnFunctionCall /*= true*/)
 {
 	TArray<FMaterialExpressionKey> ExpressionStack;
@@ -1331,7 +1327,6 @@ bool UMaterialExpression::ContainsInputLoop(const bool bStopOnFunctionCall /*= t
 
 bool UMaterialExpression::ContainsInputLoopInternal(TArray<FMaterialExpressionKey>& ExpressionStack, TSet<FMaterialExpressionKey>& VisitedExpressions, const bool bStopOnFunctionCall)
 {
-#if WITH_EDITORONLY_DATA
 	const TArray<FExpressionInput*> Inputs = GetInputs();
 	for (int32 Index = 0; Index < Inputs.Num(); ++Index)
 	{
@@ -1365,13 +1360,16 @@ bool UMaterialExpression::ContainsInputLoopInternal(TArray<FMaterialExpressionKe
 			}
 		}
 	}
-#endif // WITH_EDITORONLY_DATA
+
 	return false;
 }
+#endif // WITH_EDITOR
 
 UMaterialExpressionTextureBase::UMaterialExpressionTextureBase(const FObjectInitializer& ObjectInitializer)
 	: Super(ObjectInitializer)
+#if WITH_EDITORONLY_DATA
 	, IsDefaultMeshpaintTexture(false)
+#endif
 {}
 
 #if WITH_EDITOR
@@ -1440,8 +1438,6 @@ FText UMaterialExpressionTextureBase::GetPreviewOverlayText() const
 	}
 }
 
-#endif // WITH_EDITOR
-
 void UMaterialExpressionTextureBase::AutoSetSampleType()
 {
 	if ( Texture )
@@ -1450,11 +1446,6 @@ void UMaterialExpressionTextureBase::AutoSetSampleType()
 	}
 }
 
-UObject* UMaterialExpressionTextureBase::GetReferencedTexture() const
-{
-	return Texture;
-}
-
 EMaterialSamplerType UMaterialExpressionTextureBase::GetSamplerTypeForTexture(const UTexture* Texture, bool ForceNoVT)
 {
 	if (Texture)
@@ -1490,11 +1481,17 @@ EMaterialSamplerType UMaterialExpressionTextureBase::GetSamplerTypeForTexture(co
 	}
 	return SAMPLERTYPE_Color;
 }
+#endif // WITH_EDITOR
+
+UObject* UMaterialExpressionTextureBase::GetReferencedTexture() const
+{
+	return Texture;
+}
 
 UMaterialExpressionTextureSample::UMaterialExpressionTextureSample(const FObjectInitializer& ObjectInitializer)
 	: Super(ObjectInitializer)
-	, bShowTextureInputPin(true)
 {
+#if WITH_EDITORONLY_DATA
 	// Structure to hold one-time initialization
 	struct FConstructorStatics
 	{
@@ -1506,7 +1503,6 @@ UMaterialExpressionTextureSample::UMaterialExpressionTextureSample(const FObject
 	};
 	static FConstructorStatics ConstructorStatics;
 
-#if WITH_EDITORONLY_DATA
 	MenuCategories.Add(ConstructorStatics.NAME_Texture);
 
 	Outputs.Reset();
@@ -1518,14 +1514,14 @@ UMaterialExpressionTextureSample::UMaterialExpressionTextureSample(const FObject
 	Outputs.Add(FExpressionOutput(TEXT("RGBA"), 1, 1, 1, 1, 1));
 
 	bShowOutputNameOnPin = true;
+	bShowTextureInputPin = true;
 	bCollapsed = false;
-#endif // WITH_EDITORONLY_DATA
 
 	MipValueMode = TMVM_None;
-
 	ConstCoordinate = 0;
 	ConstMipValue = INDEX_NONE;
 	AutomaticViewMipBias = true;
+#endif // WITH_EDITORONLY_DATA
 }
 
 #if WITH_EDITOR
@@ -1675,9 +1671,6 @@ FName UMaterialExpressionTextureSample::GetInputName(int32 InputIndex) const
 }
 #undef IF_INPUT_RETURN
 
-#endif // WITH_EDITOR
-
-
 /**
  * Verify that the texture and sampler type. Generates a compiler waring if 
  * they do not.
@@ -1735,7 +1728,6 @@ static bool VerifySamplerType(
 	return true;
 }
 
-#if WITH_EDITOR
 int32 UMaterialExpressionTextureSample::Compile(class FMaterialCompiler* Compiler, int32 OutputIndex)
 {
 	UMaterialExpression* InputExpression = TextureObject.GetTracedInput().Expression;
@@ -2656,6 +2648,7 @@ UMaterialExpressionAdd::UMaterialExpressionAdd(const FObjectInitializer& ObjectI
 UMaterialExpressionTextureSampleParameter::UMaterialExpressionTextureSampleParameter(const FObjectInitializer& ObjectInitializer)
 	: Super(ObjectInitializer)
 {
+#if WITH_EDITORONLY_DATA
 	// Structure to hold one-time initialization
 	struct FConstructorStatics
 	{
@@ -2669,7 +2662,6 @@ UMaterialExpressionTextureSampleParameter::UMaterialExpressionTextureSampleParam
 	bIsParameterExpression = true;
 	bShowTextureInputPin = false;
 
-#if WITH_EDITORONLY_DATA
 	MenuCategories.Empty();
 	MenuCategories.Add( ConstructorStatics.NAME_Obsolete);
 	SortPriority = 0;
@@ -2785,7 +2777,6 @@ void UMaterialExpressionTextureSampleParameter::ApplyChannelNames()
 		Outputs[4].OutputName = !ChannelNames.A.IsEmpty() ? FName(*ChannelNames.A.ToString()) : Alpha;
 	}
 }
-#endif
 
 bool UMaterialExpressionTextureSampleParameter::TextureIsValid(UTexture* /*InTexture*/, FString& OutMessage)
 {
@@ -2797,6 +2788,7 @@ void UMaterialExpressionTextureSampleParameter::SetDefaultTexture()
 {
 	// Does nothing in the base case...
 }
+#endif // WITH_EDITOR
 
 void UMaterialExpressionTextureSampleParameter::GetAllParameterInfo(TArray<FMaterialParameterInfo> &OutParameterInfo, TArray<FGuid> &OutParameterIds, const FMaterialParameterInfo& InBaseParameterInfo) const
 {
@@ -2855,6 +2847,7 @@ UMaterialExpressionTextureObjectParameter::UMaterialExpressionTextureObjectParam
 #endif // WITH_EDITORONLY_DATA
 }
 
+#if WITH_EDITOR
 bool UMaterialExpressionTextureObjectParameter::TextureIsValid(UTexture* InTexture, FString& OutMessage)
 {
 	if (!InTexture)
@@ -2866,7 +2859,6 @@ bool UMaterialExpressionTextureObjectParameter::TextureIsValid(UTexture* InTextu
 	return true;
 }
 
-#if WITH_EDITOR
 void UMaterialExpressionTextureObjectParameter::GetCaption(TArray<FString>& OutCaptions) const
 {
 	OutCaptions.Add(TEXT("Param Tex Object")); 
@@ -3150,7 +3142,6 @@ void UMaterialExpressionTextureSampleParameter2D::GetCaption(TArray<FString>& Ou
 	OutCaptions.Add(TEXT("Param2D")); 
 	OutCaptions.Add(FString::Printf(TEXT("'%s'"), *ParameterName.ToString()));
 }
-#endif // WITH_EDITOR
 
 bool UMaterialExpressionTextureSampleParameter2D::TextureIsValid(UTexture* InTexture, FString& OutMessage)
 {
@@ -3184,8 +3175,6 @@ void UMaterialExpressionTextureSampleParameter2D::SetDefaultTexture()
 	Texture = LoadObject<UTexture2D>(nullptr, TEXT("/Engine/EngineResources/DefaultTexture.DefaultTexture"), nullptr, LOAD_None, nullptr);
 }
 
-#if WITH_EDITOR
-
 bool UMaterialExpressionTextureSampleParameter::MatchesSearchQuery( const TCHAR* SearchQuery )
 {
 	if( ParameterName.ToString().Contains(SearchQuery) )
@@ -3254,7 +3243,6 @@ void UMaterialExpressionTextureSampleParameterCube::GetCaption(TArray<FString>&
 	OutCaptions.Add(TEXT("ParamCube")); 
 	OutCaptions.Add(FString::Printf(TEXT("'%s'"), *ParameterName.ToString()));
 }
-#endif // WITH_EDITOR
 
 bool UMaterialExpressionTextureSampleParameterCube::TextureIsValid(UTexture* InTexture, FString& OutMessage)
 {
@@ -3276,6 +3264,7 @@ void UMaterialExpressionTextureSampleParameterCube::SetDefaultTexture()
 {
 	Texture = LoadObject<UTextureCube>(nullptr, TEXT("/Engine/EngineResources/DefaultTextureCube.DefaultTextureCube"), nullptr, LOAD_None, nullptr);
 }
+#endif // WITH_EDITOR
 
 //
 //  UMaterialExpressionTextureSampleParameter2DArray
@@ -3323,7 +3312,6 @@ void UMaterialExpressionTextureSampleParameter2DArray::GetCaption(TArray<FString
 	OutCaptions.Add(TEXT("Param2DArray"));
 	OutCaptions.Add(FString::Printf(TEXT("'%s'"), *ParameterName.ToString()));
 }
-#endif
 
 bool UMaterialExpressionTextureSampleParameter2DArray::TextureIsValid(UTexture* InTexture, FString& OutMessage)
 {
@@ -3340,6 +3328,7 @@ bool UMaterialExpressionTextureSampleParameter2DArray::TextureIsValid(UTexture*
 
 	return true;
 }
+#endif // WITH_EDITOR
 
 const TCHAR* UMaterialExpressionTextureSampleParameter2DArray::GetRequirements()
 {
@@ -3392,7 +3381,6 @@ void UMaterialExpressionTextureSampleParameterVolume::GetCaption(TArray<FString>
 	OutCaptions.Add(TEXT("ParamVolume")); 
 	OutCaptions.Add(FString::Printf(TEXT("'%s'"), *ParameterName.ToString()));
 }
-#endif // WITH_EDITOR
 
 bool UMaterialExpressionTextureSampleParameterVolume::TextureIsValid(UTexture* InTexture, FString& OutMessage)
 {
@@ -3414,6 +3402,7 @@ void UMaterialExpressionTextureSampleParameterVolume::SetDefaultTexture()
 {
 	Texture = LoadObject<UVolumeTexture>(nullptr, TEXT("/Engine/EngineResources/DefaultVolumeTexture.DefaultVolumeTexture"), nullptr, LOAD_None, nullptr);
 }
+#endif // WITH_EDITOR
 
 /** 
  * Performs a SubUV operation, which is doing a texture lookup into a sub rectangle of a texture, and optionally blending with another rectangle.  
@@ -3470,14 +3459,12 @@ void UMaterialExpressionTextureSampleParameterSubUV::GetCaption(TArray<FString>&
 {
 	OutCaptions.Add(TEXT("Parameter SubUV"));
 }
-#endif // WITH_EDITOR
 
 bool UMaterialExpressionTextureSampleParameterSubUV::TextureIsValid(UTexture* InTexture, FString& OutMessage)
 {
 	return UMaterialExpressionTextureSampleParameter2D::TextureIsValid(InTexture, OutMessage);
 }
 
-#if WITH_EDITOR
 int32 UMaterialExpressionAdd::Compile(class FMaterialCompiler* Compiler, int32 OutputIndex)
 {
 	// if the input is hooked up, use it, otherwise use the internal constant
@@ -3488,7 +3475,6 @@ int32 UMaterialExpressionAdd::Compile(class FMaterialCompiler* Compiler, int32 O
 	return Compiler->Add(Arg1, Arg2);
 }
 
-
 void UMaterialExpressionAdd::GetCaption(TArray<FString>& OutCaptions) const
 {
 	FString ret = TEXT("Add");
@@ -8263,6 +8249,10 @@ void UMaterialExpressionFeatureLevelSwitch::Serialize(FStructuredArchive::FRecor
 	FArchive& UnderlyingArchive = Record.GetUnderlyingArchive();
 	UnderlyingArchive.UsingCustomVersion(FRenderingObjectVersion::GUID);
 
+#if !WITH_EDITORONLY_DATA
+	FExpressionInput Inputs[ERHIFeatureLevel::Num];
+#endif
+
 	if (UnderlyingArchive.IsLoading() && UnderlyingArchive.UE4Ver() < VER_UE4_RENAME_SM3_TO_ES3_1)
 	{
 		// Copy the ES2 input to SM3 (since SM3 will now become ES3_1 and we don't want broken content)
@@ -11674,7 +11664,9 @@ bool UMaterialFunction::IsDependent(UMaterialFunctionInterface* OtherFunction)
 		return true;
 	}
 
+#if WITH_EDITOR
 	SetReentrantFlag(true);
+#endif
 
 	bool bIsDependent = false;
 	for (int32 ExpressionIndex = 0; ExpressionIndex < FunctionExpressions.Num(); ExpressionIndex++)
@@ -11690,7 +11682,9 @@ bool UMaterialFunction::IsDependent(UMaterialFunctionInterface* OtherFunction)
 		}
 	}
 
+#if WITH_EDITOR
 	SetReentrantFlag(false);
+#endif
 
 	return bIsDependent;
 }
@@ -15583,7 +15577,6 @@ void UMaterialExpressionAntialiasedTextureMask::GetCaption(TArray<FString>& OutC
 	OutCaptions.Add(TEXT("AAMasked Param2D")); 
 	OutCaptions.Add(FString::Printf(TEXT("'%s'"), *ParameterName.ToString()));
 }
-#endif // WITH_EDITOR
 
 bool UMaterialExpressionAntialiasedTextureMask::TextureIsValid(UTexture* InTexture, FString& OutMessage)
 {
@@ -15606,6 +15599,7 @@ void UMaterialExpressionAntialiasedTextureMask::SetDefaultTexture()
 {
 	Texture = LoadObject<UTexture2D>(nullptr, TEXT("/Engine/EngineResources/DefaultTexture.DefaultTexture"), nullptr, LOAD_None, nullptr);
 }
+#endif // WITH_EDITOR
 
 //
 //	UMaterialExpressionDecalDerivative
diff --git a/Engine/Source/Runtime/Engine/Private/Materials/MaterialInstance.cpp b/Engine/Source/Runtime/Engine/Private/Materials/MaterialInstance.cpp
--- a/Engine/Source/Runtime/Engine/Private/Materials/MaterialInstance.cpp
+++ b/Engine/Source/Runtime/Engine/Private/Materials/MaterialInstance.cpp
@@ -534,8 +534,10 @@ UMaterialInstance::UMaterialInstance(const FObjectInitializer& ObjectInitializer
 	: Super(ObjectInitializer)
 {
 	bHasStaticPermutationResource = false;
+#if WITH_EDITOR
 	ReentrantFlag[0] = false;
 	ReentrantFlag[1] = false;
+#endif
 	ShadingModels = MSM_Unlit;
 
 	PhysMaterial = nullptr;
diff --git a/Engine/Source/Runtime/Engine/Private/Materials/MaterialInstanceSupport.h b/Engine/Source/Runtime/Engine/Private/Materials/MaterialInstanceSupport.h
--- a/Engine/Source/Runtime/Engine/Private/Materials/MaterialInstanceSupport.h
+++ b/Engine/Source/Runtime/Engine/Private/Materials/MaterialInstanceSupport.h
@@ -31,7 +31,9 @@ void RecacheMaterialInstanceUniformExpressions(const UMaterialInterface* ParentM
 class FMICReentranceGuard
 {
 public:
-
+#if !WITH_EDITOR
+	FMICReentranceGuard(const UMaterialInstance* InMaterial) {}
+#else
 	FMICReentranceGuard(const UMaterialInstance* InMaterial)
 	{
 		Material = const_cast<UMaterialInstance*>(InMaterial);
@@ -53,6 +55,7 @@ public:
 
 private:
 	UMaterialInstance* Material;
+#endif // WITH_EDITOR
 };
 
 /**
