Single UniformExpressionCache per FMaterialRenderProxy in non-editor builds, and one per feature level in the editor

diff --git a/Engine/Source/Runtime/Engine/Public/MaterialShared.h b/Engine/Source/Runtime/Engine/Public/MaterialShared.h
--- a/Engine/Source/Runtime/Engine/Public/MaterialShared.h
+++ b/Engine/Source/Runtime/Engine/Public/MaterialShared.h
@@ -2056,6 +2056,24 @@ struct FUniformExpressionCache
 	void ResetAllocatedVTs();
 };
 
+struct FUniformExpressionCacheContainer
+{
+	inline FUniformExpressionCache& operator[](int32 Index)
+	{
+#if WITH_EDITOR
+		return Elements[Index];
+#else
+		return Elements;
+#endif
+	}
+private:
+#if WITH_EDITOR
+	FUniformExpressionCache Elements[ERHIFeatureLevel::Num];
+#else
+	FUniformExpressionCache Elements;
+#endif
+};
+
 class USubsurfaceProfile;
 
 /**
@@ -2066,7 +2084,7 @@ class FMaterialRenderProxy : public FRenderResource
 public:
 
 	/** Cached uniform expressions. */
-	mutable FUniformExpressionCache UniformExpressionCache[ERHIFeatureLevel::Num];
+	mutable FUniformExpressionCacheContainer UniformExpressionCache;
 
 	/** Cached external texture immutable samplers */
 	mutable FImmutableSamplerState ImmutableSamplerState;
