Batch material instance parameter initialization to reduce amount of GraphTasks during loading

diff --git a/Engine/Source/Runtime/Engine/Private/Materials/MaterialInstance.cpp b/Engine/Source/Runtime/Engine/Private/Materials/MaterialInstance.cpp
--- a/Engine/Source/Runtime/Engine/Private/Materials/MaterialInstance.cpp
+++ b/Engine/Source/Runtime/Engine/Private/Materials/MaterialInstance.cpp
@@ -349,6 +349,15 @@ void FMaterialInstanceResource::GameThread_SetParent(UMaterialInterface* ParentM
 	}
 }
 
+void FMaterialInstanceResource::InitMIParameters(FMaterialInstanceParameterSet& ParameterSet)
+{
+	InvalidateUniformExpressionCache(false);
+	Swap(ScalarParameterArray, ParameterSet.ScalarParameters);
+	Swap(VectorParameterArray, ParameterSet.VectorParameters);
+	Swap(TextureParameterArray, ParameterSet.TextureParameters);
+	Swap(RuntimeVirtualTextureParameterArray, ParameterSet.RuntimeVirtualTextureParameters);
+}
+
 /**
 * Updates a parameter on the material instance from the game thread.
 */
@@ -561,16 +570,63 @@ void UMaterialInstance::PostInitProperties()
 /**
  * Initializes MI parameters from the game thread.
  */
-template <typename ParameterType>
-void GameThread_InitMIParameters(const UMaterialInstance* Instance, const TArray<ParameterType>& Parameters)
+void GameThread_InitMIParameters(const UMaterialInstance& Instance)
 {
-	if (!Instance->HasAnyFlags(RF_ClassDefaultObject))
+	if (Instance.HasAnyFlags(RF_ClassDefaultObject))
 	{
-		for (int32 ParameterIndex = 0; ParameterIndex < Parameters.Num(); ParameterIndex++)
-		{
-			GameThread_UpdateMIParameter(Instance, Parameters[ParameterIndex]);
-		}
+		return;
+	}
+
+	FMaterialInstanceResource* Resource = Instance.Resource;
+	FMaterialInstanceParameterSet ParameterSet;
+
+	// Scalar parameters
+	ParameterSet.ScalarParameters.Reserve(Instance.ScalarParameterValues.Num());
+	for (const FScalarParameterValue& Parameter : Instance.ScalarParameterValues)
+	{
+		auto& ParamRef = ParameterSet.ScalarParameters.AddDefaulted_GetRef();
+		ParamRef.Info = Parameter.ParameterInfo;
+		ParamRef.Value = FScalarParameterValue::GetValue(Parameter);
+	}
+
+	// Vector parameters
+	ParameterSet.VectorParameters.Reserve(Instance.VectorParameterValues.Num());
+	for (const FVectorParameterValue& Parameter : Instance.VectorParameterValues)
+	{
+		auto& ParamRef = ParameterSet.VectorParameters.AddDefaulted_GetRef();
+		ParamRef.Info = Parameter.ParameterInfo;
+		ParamRef.Value = FVectorParameterValue::GetValue(Parameter);
+	}
+
+	// Texture + Fonts parameters
+	ParameterSet.TextureParameters.Reserve(Instance.TextureParameterValues.Num() + Instance.FontParameterValues.Num());
+	for (const FTextureParameterValue& Parameter : Instance.TextureParameterValues)
+	{
+		auto& ParamRef = ParameterSet.TextureParameters.AddDefaulted_GetRef();
+		ParamRef.Info = Parameter.ParameterInfo;
+		ParamRef.Value = FTextureParameterValue::GetValue(Parameter);
 	}
+	for (const FFontParameterValue& Parameter : Instance.FontParameterValues)
+	{
+		auto& ParamRef = ParameterSet.TextureParameters.AddDefaulted_GetRef();
+		ParamRef.Info = Parameter.ParameterInfo;
+		ParamRef.Value = FFontParameterValue::GetValue(Parameter);
+	}
+
+	// RuntimeVirtualTexture parameters
+	ParameterSet.RuntimeVirtualTextureParameters.Reserve(Instance.RuntimeVirtualTextureParameterValues.Num());
+	for (const FRuntimeVirtualTextureParameterValue& Parameter : Instance.RuntimeVirtualTextureParameterValues)
+	{
+		auto& ParamRef = ParameterSet.RuntimeVirtualTextureParameters.AddDefaulted_GetRef();
+		ParamRef.Info = Parameter.ParameterInfo;
+		ParamRef.Value = FRuntimeVirtualTextureParameterValue::GetValue(Parameter);
+	}
+	
+	ENQUEUE_RENDER_COMMAND(InitMIParameters)(
+		[Resource, Parameters = MoveTemp(ParameterSet)](FRHICommandListImmediate& RHICmdList) mutable
+		{
+			Resource->InitMIParameters(Parameters);
+		});
 }
 
 void UMaterialInstance::InitResources()
@@ -608,12 +664,7 @@ void UMaterialInstance::InitResources()
 		Resource->GameThread_SetParent(SafeParent);
 	}
 
-	GameThread_InitMIParameters(this, ScalarParameterValues);
-	GameThread_InitMIParameters(this, VectorParameterValues);
-	GameThread_InitMIParameters(this, TextureParameterValues);
-	GameThread_InitMIParameters(this, FontParameterValues);
-	GameThread_InitMIParameters(this, RuntimeVirtualTextureParameterValues);
-
+	GameThread_InitMIParameters(*this);
 	PropagateDataToMaterialProxy();
 
 	CacheMaterialInstanceUniformExpressions(this);
diff --git a/Engine/Source/Runtime/Engine/Private/Materials/MaterialInstanceSupport.h b/Engine/Source/Runtime/Engine/Private/Materials/MaterialInstanceSupport.h
--- a/Engine/Source/Runtime/Engine/Private/Materials/MaterialInstanceSupport.h
+++ b/Engine/Source/Runtime/Engine/Private/Materials/MaterialInstanceSupport.h
@@ -106,15 +106,17 @@ public:
 
 	void GameThread_SetParent(UMaterialInterface* ParentMaterialInterface);
 
+	void InitMIParameters(struct FMaterialInstanceParameterSet& ParameterSet);
+
 	/**
 	 * Clears all parameters set on this material instance.
 	 */
 	void RenderThread_ClearParameters()
 	{
-		VectorParameterArray.Reset();
-		ScalarParameterArray.Reset();
-		TextureParameterArray.Reset();
-		RuntimeVirtualTextureParameterArray.Reset();
+		VectorParameterArray.Empty();
+		ScalarParameterArray.Empty();
+		TextureParameterArray.Empty();
+		RuntimeVirtualTextureParameterArray.Empty();
 		InvalidateUniformExpressionCache(false);
 	}
 
@@ -198,6 +200,14 @@ template <> FORCEINLINE const TArray<FMaterialInstanceResource::TNamedParameter<
 template <> FORCEINLINE const TArray<FMaterialInstanceResource::TNamedParameter<const UTexture*> >& FMaterialInstanceResource::GetValueArray() const { return TextureParameterArray; }
 template <> FORCEINLINE const TArray<FMaterialInstanceResource::TNamedParameter<const URuntimeVirtualTexture*> >& FMaterialInstanceResource::GetValueArray() const { return RuntimeVirtualTextureParameterArray; }
 
+struct FMaterialInstanceParameterSet
+{
+	TArray<FMaterialInstanceResource::TNamedParameter<float> > ScalarParameters;
+	TArray<FMaterialInstanceResource::TNamedParameter<FLinearColor> > VectorParameters;
+	TArray<FMaterialInstanceResource::TNamedParameter<const UTexture*> > TextureParameters;
+	TArray<FMaterialInstanceResource::TNamedParameter<const URuntimeVirtualTexture*> > RuntimeVirtualTextureParameters;
+};
+	
 /** Finds a parameter by name from the game thread. */
 template <typename ParameterType>
 ParameterType* GameThread_FindParameterByName(TArray<ParameterType>& Parameters, const FHashedMaterialParameterInfo& ParameterInfo)
