Allocate FMaterialResource only for current feature level in a cooked builds

diff --git a/Engine/Source/Runtime/Engine/Private/Materials/Material.cpp b/Engine/Source/Runtime/Engine/Private/Materials/Material.cpp
--- a/Engine/Source/Runtime/Engine/Private/Materials/Material.cpp
+++ b/Engine/Source/Runtime/Engine/Private/Materials/Material.cpp
@@ -3004,9 +3004,12 @@ void UMaterial::CacheResourceShadersForRendering(bool bRegenerateId)
 			// In cooked build, there is no shader compilation but this is still needed
 			// to register the loaded shadermap
 			ResourcesToCache.Reset();
-			check(MaterialResources[LocalActiveQL][FeatureLevel]);
-			ResourcesToCache.Add(MaterialResources[LocalActiveQL][FeatureLevel]);
-			CacheShadersForResources(ShaderPlatform, ResourcesToCache);
+			FMaterialResource* CurrentResource = MaterialResources[LocalActiveQL][FeatureLevel];
+			if (CurrentResource)
+			{
+				ResourcesToCache.Add(MaterialResources[LocalActiveQL][FeatureLevel]);
+				CacheShadersForResources(ShaderPlatform, ResourcesToCache);
+			}
 		}
 
 		FString AdditionalFormatToCache = GCompileMaterialsForShaderFormatCVar->GetString();
@@ -3121,7 +3124,10 @@ void UMaterial::FlushResourceShaderMaps()
 			for(int32 QualityLevelIndex = 0; QualityLevelIndex < EMaterialQualityLevel::Num; QualityLevelIndex++)
 			{
 				FMaterialResource* CurrentResource = MaterialResources[QualityLevelIndex][InFeatureLevel];
-				CurrentResource->ReleaseShaderMap();
+				if (CurrentResource)
+				{
+					CurrentResource->ReleaseShaderMap();
+				}
 			}
 		});
 	}
@@ -3375,7 +3381,16 @@ void UMaterial::UpdateResourceAllocations(FMaterialResourceDeferredDeletionArray
 			}
 		}
 #else
-		for (int32 FeatureLevelIndex = 0; FeatureLevelIndex < ERHIFeatureLevel::Num; FeatureLevelIndex++)
+		// Initialize only current feature level in a cooked game and all in the editor
+		int32 FeatureLevelMin = 0;
+		int32 FeatureLevelMax = ERHIFeatureLevel::Num;
+		if (FPlatformProperties::RequiresCookedData())
+		{
+			FeatureLevelMin = GMaxRHIFeatureLevel;
+			FeatureLevelMax = GMaxRHIFeatureLevel + 1;
+		}
+
+		for (int32 FeatureLevelIndex = FeatureLevelMin; FeatureLevelIndex < FeatureLevelMax; FeatureLevelIndex++)
 		{
 			EShaderPlatform ShaderPlatform = GShaderPlatformForFeatureLevel[FeatureLevelIndex];
 			TArray<bool, TInlineAllocator<EMaterialQualityLevel::Num> > QualityLevelsUsed;
@@ -4595,12 +4610,14 @@ void UMaterial::BackupMaterialShadersToMemory(TMap<FMaterialShaderMap*, TUniqueP
 				for (int32 FeatureLevelIndex = 0; FeatureLevelIndex < ERHIFeatureLevel::Num; FeatureLevelIndex++)
 				{
 					FMaterialResource* CurrentResource = BaseMaterial->MaterialResources[QualityLevelIndex][FeatureLevelIndex];
-					FMaterialShaderMap* ShaderMap = CurrentResource->GetGameThreadShaderMap();
-
-					if (ShaderMap && !ShaderMapToSerializedShaderData.Contains(ShaderMap))
+					if (CurrentResource)
 					{
-						TArray<uint8>* ShaderData = ShaderMap->BackupShadersToMemory();
-						ShaderMapToSerializedShaderData.Emplace(ShaderMap, ShaderData);
+						FMaterialShaderMap* ShaderMap = CurrentResource->GetGameThreadShaderMap();
+						if (ShaderMap && !ShaderMapToSerializedShaderData.Contains(ShaderMap))
+						{
+							TArray<uint8>* ShaderData = ShaderMap->BackupShadersToMemory();
+							ShaderMapToSerializedShaderData.Emplace(ShaderMap, ShaderData);
+						}
 					}
 				}
 			}
@@ -4652,15 +4669,17 @@ void UMaterial::RestoreMaterialShadersFromMemory(const TMap<FMaterialShaderMap*,
 				for (int32 FeatureLevelIndex = 0; FeatureLevelIndex < ERHIFeatureLevel::Num; FeatureLevelIndex++)
 				{
 					FMaterialResource* CurrentResource = BaseMaterial->MaterialResources[QualityLevelIndex][FeatureLevelIndex];
-					FMaterialShaderMap* ShaderMap = CurrentResource->GetGameThreadShaderMap();
-
-					if (ShaderMap)
+					if (CurrentResource)
 					{
-						const TUniquePtr<TArray<uint8> >* ShaderData = ShaderMapToSerializedShaderData.Find(ShaderMap);
-
-						if (ShaderData)
+						FMaterialShaderMap* ShaderMap = CurrentResource->GetGameThreadShaderMap();
+						if (ShaderMap)
 						{
-							ShaderMap->RestoreShadersFromMemory(**ShaderData);
+							const TUniquePtr<TArray<uint8> >* ShaderData = ShaderMapToSerializedShaderData.Find(ShaderMap);
+
+							if (ShaderData)
+							{
+								ShaderMap->RestoreShadersFromMemory(**ShaderData);
+							}
 						}
 					}
 				}
diff --git a/Engine/Source/Runtime/Engine/Private/Materials/MaterialInstance.cpp b/Engine/Source/Runtime/Engine/Private/Materials/MaterialInstance.cpp
--- a/Engine/Source/Runtime/Engine/Private/Materials/MaterialInstance.cpp
+++ b/Engine/Source/Runtime/Engine/Private/Materials/MaterialInstance.cpp
@@ -2723,8 +2723,11 @@ void UMaterialInstance::GetAllShaderMaps(TArray<FMaterialShaderMap*>& OutShaderM
 		for (int32 FeatureLevelIndex = 0; FeatureLevelIndex < ERHIFeatureLevel::Num; FeatureLevelIndex++)
 		{
 			FMaterialResource* CurrentResource = StaticPermutationMaterialResources[QualityLevelIndex][FeatureLevelIndex];
-			FMaterialShaderMap* ShaderMap = CurrentResource->GetGameThreadShaderMap();
-			OutShaderMaps.Add(ShaderMap);
+			if (CurrentResource)
+			{
+				FMaterialShaderMap* ShaderMap = CurrentResource->GetGameThreadShaderMap();
+				OutShaderMaps.Add(ShaderMap);
+			}
 		}
 	}
 }
@@ -2775,7 +2778,16 @@ void UMaterialInstance::UpdatePermutationAllocations(FMaterialResourceDeferredDe
 			}
 		}
 #else
-		for (int32 FeatureLevelIndex = 0; FeatureLevelIndex < ERHIFeatureLevel::Num; FeatureLevelIndex++)
+		// Initialize only current feature level in a cooked game and all in the editor
+		int32 FeatureLevelMin = 0;
+		int32 FeatureLevelMax = ERHIFeatureLevel::Num;
+		if (FPlatformProperties::RequiresCookedData())
+		{
+			FeatureLevelMin = GMaxRHIFeatureLevel;
+			FeatureLevelMax = GMaxRHIFeatureLevel + 1;
+		}
+		
+		for (int32 FeatureLevelIndex = FeatureLevelMin; FeatureLevelIndex < FeatureLevelMax; FeatureLevelIndex++)
 		{
 			EShaderPlatform ShaderPlatform = GShaderPlatformForFeatureLevel[FeatureLevelIndex];
 			TArray<bool, TInlineAllocator<EMaterialQualityLevel::Num> > QualityLevelsUsed;
@@ -2843,9 +2855,12 @@ void UMaterialInstance::CacheResourceShadersForRendering()
 			// In cooked build, there is no shader compilation but this is still needed to
 			// register the loaded shadermap
 			ResourcesToCache.Reset();
-			check(StaticPermutationMaterialResources[LocalActiveQL][FeatureLevel]);
-			ResourcesToCache.Add(StaticPermutationMaterialResources[LocalActiveQL][FeatureLevel]);
-			CacheShadersForResources(ShaderPlatform, ResourcesToCache);
+			FMaterialResource* CurrentResource = StaticPermutationMaterialResources[LocalActiveQL][FeatureLevel];
+			if (CurrentResource)
+			{
+				ResourcesToCache.Add(CurrentResource);
+				CacheShadersForResources(ShaderPlatform, ResourcesToCache);
+			}
 		}
 	}
 
diff --git a/Engine/Source/Runtime/Engine/Private/Materials/MaterialShared.cpp b/Engine/Source/Runtime/Engine/Private/Materials/MaterialShared.cpp
--- a/Engine/Source/Runtime/Engine/Private/Materials/MaterialShared.cpp
+++ b/Engine/Source/Runtime/Engine/Private/Materials/MaterialShared.cpp
@@ -3116,7 +3116,10 @@ FMaterialUpdateContext::~FMaterialUpdateContext()
 				for (int32 FeatureLevelIndex = 0; FeatureLevelIndex < ERHIFeatureLevel::Num; FeatureLevelIndex++)
 				{
 					FMaterialResource* CurrentResource = Material->MaterialResources[QualityLevelIndex][FeatureLevelIndex];
-					MaterialResourcesToUpdate.Add(CurrentResource);
+					if (CurrentResource)
+					{
+						MaterialResourcesToUpdate.Add(CurrentResource);
+					}
 				}
 			}
 		}
@@ -3172,7 +3175,10 @@ FMaterialUpdateContext::~FMaterialUpdateContext()
 					for (int32 FeatureLevelIndex = 0; FeatureLevelIndex < ERHIFeatureLevel::Num; FeatureLevelIndex++)
 					{
 						FMaterialResource* CurrentResource = MI->StaticPermutationMaterialResources[QualityLevelIndex][FeatureLevelIndex];
-						MaterialResourcesToUpdate.Add(CurrentResource);
+						if (CurrentResource)
+						{
+							MaterialResourcesToUpdate.Add(CurrentResource);
+						}
 					}
 				}
 			}
