Added MID creation flags for the BP create function. The only flags right now i None and Transient. It's implemented as a bitflag mask so that multiple flags can be set and manipulated throught BP. This iniitial implementation is to give users the option to create MIDs as transient objects, never to be exported.

diff --git a/Engine/Source/Runtime/Engine/Classes/Kismet/KismetMaterialLibrary.h b/Engine/Source/Runtime/Engine/Classes/Kismet/KismetMaterialLibrary.h
--- a/Engine/Source/Runtime/Engine/Classes/Kismet/KismetMaterialLibrary.h
+++ b/Engine/Source/Runtime/Engine/Classes/Kismet/KismetMaterialLibrary.h
@@ -3,16 +3,27 @@
 #pragma once
 
 #include "CoreMinimal.h"
+#include "Misc/EnumClassFlags.h"
 #include "UObject/ObjectMacros.h"
 #include "Kismet/BlueprintFunctionLibrary.h"
 #include "KismetMaterialLibrary.generated.h"
 
 class UMaterialParameterCollection;
 
+UENUM(BlueprintType, meta = (Bitflags, UseEnumValuesAsMaskValuesInEditor = "true"))
+enum class EMIDCreationFlags : uint8
+{
+	None					UMETA(DisplayName = "None"),
+	Transient = 1 << 0		UMETA(DisplayName = "Transient")
+};
+ENUM_CLASS_FLAGS(EMIDCreationFlags)
+
 UCLASS(MinimalAPI, meta=(ScriptName="MaterialLibrary"))
 class UKismetMaterialLibrary : public UBlueprintFunctionLibrary
 {
-	GENERATED_UCLASS_BODY()
+	GENERATED_BODY()
+
+public:
 
 	/** Sets a scalar parameter value on the material collection instance. Logs if ParameterName is invalid. */
 	UFUNCTION(BlueprintCallable, Category="Rendering|Material", meta=(Keywords="SetFloatParameterValue", WorldContext="WorldContextObject", MaterialParameterCollectionFunction = "true"))
@@ -32,5 +43,5 @@ class UKismetMaterialLibrary : public UBlueprintFunctionLibrary
 
 	/** Creates a Dynamic Material Instance which you can modify during gameplay. */
 	UFUNCTION(BlueprintCallable, Category="Rendering|Material", meta=(WorldContext="WorldContextObject", MaterialParameterCollectionFunction = "true"))
-	static ENGINE_API class UMaterialInstanceDynamic* CreateDynamicMaterialInstance(UObject* WorldContextObject, class UMaterialInterface* Parent, FName OptionalName = NAME_None);
+	static ENGINE_API class UMaterialInstanceDynamic* CreateDynamicMaterialInstance(UObject* WorldContextObject, class UMaterialInterface* Parent, FName OptionalName = NAME_None, EMIDCreationFlags CreationFlags = EMIDCreationFlags::None);
 };
diff --git a/Engine/Source/Runtime/Engine/Private/KismetMaterialLibrary.cpp b/Engine/Source/Runtime/Engine/Private/KismetMaterialLibrary.cpp
--- a/Engine/Source/Runtime/Engine/Private/KismetMaterialLibrary.cpp
+++ b/Engine/Source/Runtime/Engine/Private/KismetMaterialLibrary.cpp
@@ -15,12 +15,6 @@
 
 #define LOCTEXT_NAMESPACE "KismetMaterialLibrary"
 
-
-UKismetMaterialLibrary::UKismetMaterialLibrary(const FObjectInitializer& ObjectInitializer)
-	: Super(ObjectInitializer)
-{
-}
-
 void UKismetMaterialLibrary::SetScalarParameterValue(UObject* WorldContextObject, UMaterialParameterCollection* Collection, FName ParameterName, float ParameterValue)
 {
 	if (Collection)
@@ -125,7 +119,7 @@ FLinearColor UKismetMaterialLibrary::GetVectorParameterValue(UObject* WorldConte
 	return ParameterValue;
 }
 
-class UMaterialInstanceDynamic* UKismetMaterialLibrary::CreateDynamicMaterialInstance(UObject* WorldContextObject, class UMaterialInterface* Parent, FName OptionalName)
+class UMaterialInstanceDynamic* UKismetMaterialLibrary::CreateDynamicMaterialInstance(UObject* WorldContextObject, class UMaterialInterface* Parent, FName OptionalName, EMIDCreationFlags CreationFlags)
 {
 	UMaterialInstanceDynamic* NewMID = nullptr;
 
@@ -135,7 +129,7 @@ class UMaterialInstanceDynamic* UKismetMaterialLibrary::CreateDynamicMaterialIns
 		// MIDs need to be created within a persistent object if in the construction script (or editor utility) or else they will not be saved.
 		// If this MID is created at runtime then put it in the transient package
 		UWorld* World = GEngine->GetWorldFromContextObject(WorldContextObject, EGetWorldErrorMode::ReturnNull);
-		UObject* MIDOuter = (World && (World->bIsRunningConstructionScript  || !World->IsGameWorld()) ? WorldContextObject : nullptr);
+		UObject* MIDOuter = (!EnumHasAnyFlags(CreationFlags, EMIDCreationFlags::Transient) &&  World && (World->bIsRunningConstructionScript  || !World->IsGameWorld()) ? WorldContextObject : nullptr);
 		NewMID = UMaterialInstanceDynamic::Create(Parent, MIDOuter, OptionalName);
 		if (MIDOuter == nullptr)
 		{
