Fix compile error from merge

diff --git a/Engine/Source/Runtime/Engine/Private/AudioStreamingCache.cpp b/Engine/Source/Runtime/Engine/Private/AudioStreamingCache.cpp
--- a/Engine/Source/Runtime/Engine/Private/AudioStreamingCache.cpp
+++ b/Engine/Source/Runtime/Engine/Private/AudioStreamingCache.cpp
@@ -1188,7 +1188,7 @@ void FAudioChunkCache::KickOffAsyncLoad(FCacheElement* CacheElement, const FChun
 		
 		NumberOfLoadsInFlight.Increment();
 
-		FAsyncFileCallBack AsyncFileCallBack = [this, OnLoadCompleted, CacheElement, InKey, ChunkDataSize, CallbackThread](bool bWasCancelled, IAsyncReadRequest* Request)
+		FBulkDataIORequestCallBack AsyncFileCallBack = [this, OnLoadCompleted, CacheElement, InKey, ChunkDataSize](bool bWasCancelled, IBulkDataIORequest*)
 		{
 			// Populate key and DataSize. The async read request was set up to write directly into CacheElement->ChunkData.
 			CacheElement->Key = InKey;
