Implement radial forces.

diff --git a/Engine/Source/Runtime/Engine/Private/PhysicsEngine/Experimental/PhysScene_Chaos.cpp b/Engine/Source/Runtime/Engine/Private/PhysicsEngine/Experimental/PhysScene_Chaos.cpp
--- a/Engine/Source/Runtime/Engine/Private/PhysicsEngine/Experimental/PhysScene_Chaos.cpp
+++ b/Engine/Source/Runtime/Engine/Private/PhysicsEngine/Experimental/PhysScene_Chaos.cpp
@@ -1442,33 +1442,60 @@ void FPhysScene_ChaosInterface::AddForceAtPosition_AssumesLocked(FBodyInstance*
 
 void FPhysScene_ChaosInterface::AddRadialForceToBody_AssumesLocked(FBodyInstance* BodyInstance, const FVector& Origin, const float Radius, const float Strength, const uint8 Falloff, bool bAccelChange, bool bAllowSubstepping)
 {
-	// #todo : Implement
-#if 0
-	Chaos::TVector<float, 3> Direction = (static_cast<FVector>(Scene.GetSolver()->GetRigidParticles().X(Index)) - Origin);
-	Chaos::TVector<float, 3> Force(0);
-	const float Distance = Direction.Size();
-
-	if(Distance > Radius)
+	FPhysicsActorHandle& Handle = BodyInstance->GetPhysicsActorHandle();
+	if (ensure(FPhysicsInterface::IsValid(Handle)))
 	{
-		return;
-	}
+		Chaos::TPBDRigidParticle<float, 3>* Rigid = Handle->CastToRigidParticle();
+		
+		if (ensure(Rigid))
+		{
+			Chaos::EObjectStateType ObjectState = Rigid->ObjectState();
+			if (CHAOS_ENSURE(ObjectState == Chaos::EObjectStateType::Dynamic || ObjectState == Chaos::EObjectStateType::Sleeping))
+			{
+				const Chaos::FVec3& CurrentForce = Rigid->F();
+				const Chaos::FVec3& CurrentTorque = Rigid->Torque();
+				const Chaos::FVec3 WorldCOM = Chaos::FParticleUtilitiesGT::GetCoMWorldPosition(Rigid);
 
-	Direction = Direction.GetSafeNormal();
-	
-	check(Falloff == RIF_Constant || Falloff == RIF_Linear);
+				Chaos::FVec3 Direction = WorldCOM - Origin;
+				const float Distance = Direction.Size();
+				if (Distance > Radius)
+				{
+					return;
+				}
 
-	if(Falloff == RIF_Constant)
-	{
-		Force = Strength * Direction;
-	}
+				Rigid->SetObjectState(Chaos::EObjectStateType::Dynamic);
 
-	if(Falloff == RIF_Linear)
-	{
-		Force = (Radius - Distance) / Radius * Strength * Direction;
+				if (Distance < 1e-4)
+				{
+					Direction = Chaos::FVec3(1, 0, 0);
+				}
+				else
+				{
+					Direction = Direction.GetUnsafeNormal();
+				}
+				Chaos::FVec3 Force(0, 0, 0);
+				CHAOS_ENSURE(Falloff < RIF_MAX);
+				if (Falloff == ERadialImpulseFalloff::RIF_Constant)
+				{
+					Force = Strength * Direction;
+				}
+				if (Falloff == ERadialImpulseFalloff::RIF_Linear)
+				{
+					Force = (Radius - Distance) / Radius * Strength * Direction;
+				}
+				if (bAccelChange)
+				{
+					const float Mass = Rigid->M();
+					const Chaos::TVector<float, 3> TotalAcceleration = CurrentForce + (Force * Mass);
+					Rigid->SetF(TotalAcceleration);
+				}
+				else
+				{
+					Rigid->SetF(CurrentForce + Force);
+				}
+			}
+		}
 	}
-
-	AddForce(bAccelChange ? (Force * Scene.GetSolver()->GetRigidParticles().M(Index)) : Force, Id);
-#endif
 }
 
 void FPhysScene_ChaosInterface::ClearForces_AssumesLocked(FBodyInstance* BodyInstance, bool bAllowSubstepping)
