#UE4 Fix editor error in CurveLinearColorAtlas where loading an atlas that was just changed may result in texture data not being available during load.

diff --git a/Engine/Source/Runtime/Engine/Private/Curves/CurveLinearColorAtlas.cpp b/Engine/Source/Runtime/Engine/Private/Curves/CurveLinearColorAtlas.cpp
--- a/Engine/Source/Runtime/Engine/Private/Curves/CurveLinearColorAtlas.cpp
+++ b/Engine/Source/Runtime/Engine/Private/Curves/CurveLinearColorAtlas.cpp
@@ -84,6 +84,11 @@ void UCurveLinearColorAtlas::PostEditChangeProperty(struct FPropertyChangedEvent
 void UCurveLinearColorAtlas::PostLoad()
 {
 #if WITH_EDITOR
+	if (FApp::CanEverRender())
+	{
+		FinishCachePlatformData();
+	}
+
 	for (int32 i = 0; i < GradientCurves.Num(); ++i)
 	{
 		if (GradientCurves[i] != nullptr)
