Fix reported size of FUniformExpressionSet. GetAllocatedSize was returning the size of ref counted pointer arrays but missing the contents.

diff --git a/Engine/Source/Runtime/Engine/Public/MaterialShared.h b/Engine/Source/Runtime/Engine/Public/MaterialShared.h
--- a/Engine/Source/Runtime/Engine/Public/MaterialShared.h
+++ b/Engine/Source/Runtime/Engine/Public/MaterialShared.h
@@ -543,16 +543,26 @@ public:
 
 	uint32 GetAllocatedSize() const
 	{
-		/*return UniformVectorExpressions.GetAllocatedSize()
-			+ UniformScalarExpressions.GetAllocatedSize()
+		/*	// sizeof(TArray<TRefCountPtr>)+sizeof(Contents)*ArrayLength
+		return UniformScalarExpressions.GetAllocatedSize()
+			+ sizeof(FMaterialUniformExpressionScalarParameter) * UniformScalarExpressions.Num()
+			+ UniformVectorExpressions.GetAllocatedSize()
+			+ sizeof(FMaterialUniformExpressionVectorParameter) * UniformVectorExpressions.Num()
 			+ Uniform2DTextureExpressions.GetAllocatedSize()
+			+ sizeof(FMaterialUniformExpressionTexture) * Uniform2DTextureExpressions.Num()
 			+ UniformCubeTextureExpressions.GetAllocatedSize()
+			+ sizeof(FMaterialUniformExpressionTexture) * UniformCubeTextureExpressions.Num()
 			+ Uniform2DArrayTextureExpressions.GetAllocatedSize()
+			+ sizeof(FMaterialUniformExpressionTexture) * Uniform2DArrayTextureExpressions.Num()
 			+ UniformVolumeTextureExpressions.GetAllocatedSize()
+			+ sizeof(FMaterialUniformExpressionTexture) * UniformVolumeTextureExpressions.Num()
 			+ UniformVirtualTextureExpressions.GetAllocatedSize()
+			+ sizeof(FMaterialUniformExpressionTexture) * UniformVirtualTextureExpressions.Num()
 			+ UniformExternalTextureExpressions.GetAllocatedSize()
+			+ sizeof(FMaterialUniformExpressionExternalTexture) * UniformExternalTextureExpressions.Num()
 			+ VTStacks.GetAllocatedSize()
-			+ ParameterCollections.GetAllocatedSize();*/
+			+ ParameterCollections.GetAllocatedSize()
+			+ (UniformBufferStruct ? (sizeof(FShaderParametersMetadata) + UniformBufferStruct->GetMembers().GetAllocatedSize()) : 0);*/
 		return 0u;
 	}
 
