Fix crash in UnregisterObject in single threaded configuration.
Only lock FPBDRigidActiveParticlesBuffer when threaded.

diff --git a/Engine/Source/Runtime/Engine/Private/PhysicsEngine/Experimental/PhysScene_Chaos.cpp b/Engine/Source/Runtime/Engine/Private/PhysicsEngine/Experimental/PhysScene_Chaos.cpp
--- a/Engine/Source/Runtime/Engine/Private/PhysicsEngine/Experimental/PhysScene_Chaos.cpp
+++ b/Engine/Source/Runtime/Engine/Private/PhysicsEngine/Experimental/PhysScene_Chaos.cpp
@@ -2127,69 +2127,72 @@ void FPhysScene_ChaosInterface::SyncBodies(Chaos::FPhysicsSolver* Solver)
 	DECLARE_SCOPE_CYCLE_COUNTER(TEXT("SyncBodies"), STAT_SyncBodies, STATGROUP_Physics);
 	TArray<FPhysScenePendingComponentTransform_Chaos> PendingTransforms;
 
-	Chaos::FPBDRigidActiveParticlesBufferAccessor Accessor(Solver->GetActiveParticlesBuffer());
+	{
+		Chaos::FPBDRigidActiveParticlesBufferAccessor Accessor(Solver->GetActiveParticlesBuffer());
 
-	TSet<FGeometryCollectionPhysicsProxy*> GCProxies;
+		TSet<FGeometryCollectionPhysicsProxy*> GCProxies;
 
-	const Chaos::FPBDRigidActiveParticlesBufferOut* ActiveParticleBuffer = Accessor.GetSolverOutData();
-	for (Chaos::TGeometryParticle<float, 3>* ActiveParticle : ActiveParticleBuffer->ActiveGameThreadParticles)
-	{
-		if (IPhysicsProxyBase * ProxyBase = ActiveParticle->Proxy)
+		const Chaos::FPBDRigidActiveParticlesBufferOut* ActiveParticleBuffer = Accessor.GetSolverOutData();
+		for (Chaos::TGeometryParticle<float, 3>* ActiveParticle : ActiveParticleBuffer->ActiveGameThreadParticles)
 		{
-			if (ProxyBase->GetType() == EPhysicsProxyType::SingleRigidParticleType)
+			if (IPhysicsProxyBase * ProxyBase = ActiveParticle->Proxy)
 			{
-				FSingleParticlePhysicsProxy< Chaos::TPBDRigidParticle<float, 3> > * Proxy = static_cast<FSingleParticlePhysicsProxy< Chaos::TPBDRigidParticle<float, 3> >*>(ProxyBase);
-				Proxy->PullFromPhysicsState();
-
-				if (FBodyInstance* BodyInstance = FPhysicsUserData::Get<FBodyInstance>(ActiveParticle->UserData()))
+				if (ProxyBase->GetType() == EPhysicsProxyType::SingleRigidParticleType)
 				{
-					if (BodyInstance->OwnerComponent.IsValid())
+					FSingleParticlePhysicsProxy< Chaos::TPBDRigidParticle<float, 3> > * Proxy = static_cast<FSingleParticlePhysicsProxy< Chaos::TPBDRigidParticle<float, 3> >*>(ProxyBase);
+					Proxy->PullFromPhysicsState();
+
+					if (FBodyInstance* BodyInstance = FPhysicsUserData::Get<FBodyInstance>(ActiveParticle->UserData()))
 					{
-						UPrimitiveComponent* OwnerComponent = BodyInstance->OwnerComponent.Get();
-						if (OwnerComponent != nullptr)
+						if (BodyInstance->OwnerComponent.IsValid())
 						{
-							bool bPendingMove = false;
-							if (BodyInstance->InstanceBodyIndex == INDEX_NONE)
+							UPrimitiveComponent* OwnerComponent = BodyInstance->OwnerComponent.Get();
+							if (OwnerComponent != nullptr)
 							{
-								Chaos::TRigidTransform<float, 3> NewTransform(ActiveParticle->X(), ActiveParticle->R());
-
-								if (!NewTransform.EqualsNoScale(OwnerComponent->GetComponentTransform()))
+								bool bPendingMove = false;
+								if (BodyInstance->InstanceBodyIndex == INDEX_NONE)
 								{
-									bPendingMove = true;
-									const FVector MoveBy = NewTransform.GetLocation() - OwnerComponent->GetComponentTransform().GetLocation();
-									const FQuat NewRotation = NewTransform.GetRotation();
-									PendingTransforms.Add(FPhysScenePendingComponentTransform_Chaos(OwnerComponent, MoveBy, NewRotation, Proxy->HasAwakeEvent()));
+									Chaos::TRigidTransform<float, 3> NewTransform(ActiveParticle->X(), ActiveParticle->R());
+
+									if (!NewTransform.EqualsNoScale(OwnerComponent->GetComponentTransform()))
+									{
+										bPendingMove = true;
+										const FVector MoveBy = NewTransform.GetLocation() - OwnerComponent->GetComponentTransform().GetLocation();
+										const FQuat NewRotation = NewTransform.GetRotation();
+										PendingTransforms.Add(FPhysScenePendingComponentTransform_Chaos(OwnerComponent, MoveBy, NewRotation, Proxy->HasAwakeEvent()));
+									}
 								}
-							}
 
-							if (Proxy->HasAwakeEvent() && !bPendingMove)
-							{
-								PendingTransforms.Add(FPhysScenePendingComponentTransform_Chaos(OwnerComponent));
+								if (Proxy->HasAwakeEvent() && !bPendingMove)
+								{
+									PendingTransforms.Add(FPhysScenePendingComponentTransform_Chaos(OwnerComponent));
+								}
+								Proxy->ClearEvents();
 							}
-							Proxy->ClearEvents();
 						}
 					}
 				}
+				else if(ProxyBase->GetType() == EPhysicsProxyType::GeometryCollectionType)
+				{
+					FGeometryCollectionPhysicsProxy* Proxy = static_cast<FGeometryCollectionPhysicsProxy*>(ProxyBase);
+					GCProxies.Add(Proxy);
+				}
 			}
-			else if(ProxyBase->GetType() == EPhysicsProxyType::GeometryCollectionType)
+		}
+		for (IPhysicsProxyBase* ProxyBase : ActiveParticleBuffer->PhysicsParticleProxies) 
+		{
+			if(ProxyBase->GetType() == EPhysicsProxyType::GeometryCollectionType)
 			{
 				FGeometryCollectionPhysicsProxy* Proxy = static_cast<FGeometryCollectionPhysicsProxy*>(ProxyBase);
 				GCProxies.Add(Proxy);
 			}
+			else
+			{
+				ensure(false); // Unhandled physics only particle proxy!
+			}
 		}
 	}
-	for (IPhysicsProxyBase* ProxyBase : ActiveParticleBuffer->PhysicsParticleProxies) 
-	{
-		if(ProxyBase->GetType() == EPhysicsProxyType::GeometryCollectionType)
-		{
-			FGeometryCollectionPhysicsProxy* Proxy = static_cast<FGeometryCollectionPhysicsProxy*>(ProxyBase);
-			GCProxies.Add(Proxy);
-		}
-		else
-		{
-			ensure(false); // Unhandled physics only particle proxy!
-		}
-	}
+	
 	for (auto* GCProxy : GCProxies)
 	{
 		GCProxy->PullFromPhysicsState();
diff --git a/Engine/Source/Runtime/Experimental/ChaosSolvers/Private/PBDRigidActiveParticlesBuffer.cpp b/Engine/Source/Runtime/Experimental/ChaosSolvers/Private/PBDRigidActiveParticlesBuffer.cpp
--- a/Engine/Source/Runtime/Experimental/ChaosSolvers/Private/PBDRigidActiveParticlesBuffer.cpp
+++ b/Engine/Source/Runtime/Experimental/ChaosSolvers/Private/PBDRigidActiveParticlesBuffer.cpp
@@ -6,25 +6,25 @@
 namespace Chaos
 {
 
-	FPBDRigidActiveParticlesBuffer::FPBDRigidActiveParticlesBuffer(const Chaos::EMultiBufferMode& InBufferMode) : BufferMode(InBufferMode)
+	FPBDRigidActiveParticlesBuffer::FPBDRigidActiveParticlesBuffer(const Chaos::EMultiBufferMode& InBufferMode, bool bInSingleThreaded) : BufferMode(InBufferMode), bUseLock(!bInSingleThreaded)
 	{
 		SolverDataOut = Chaos::FMultiBufferFactory<FPBDRigidActiveParticlesBufferOut>::CreateBuffer(InBufferMode);
 	}
 
 	void FPBDRigidActiveParticlesBuffer::CaptureSolverData(FPBDRigidsSolver* Solver)
 	{
-		ResourceOutLock.WriteLock();
+		WriteLock();
 		BufferPhysicsResults(Solver);
 		FlipDataOut();
-		ResourceOutLock.WriteUnlock();
+		WriteUnlock();
 	}
 
 	void FPBDRigidActiveParticlesBuffer::RemoveActiveParticleFromConsumerBuffer(TGeometryParticle<FReal, 3>* Particle)
 	{
-		ResourceOutLock.WriteLock();
+		WriteLock();
 		auto& ActiveGameThreadParticles = SolverDataOut->GetConsumerBufferMutable()->ActiveGameThreadParticles;
 		ActiveGameThreadParticles.RemoveSingleSwap(Particle);
-		ResourceOutLock.WriteUnlock();
+		WriteUnlock();
 	}
 
 	void FPBDRigidActiveParticlesBuffer::BufferPhysicsResults(FPBDRigidsSolver* Solver)
@@ -52,4 +52,35 @@ namespace Chaos
 		}
 	}
 
+	void FPBDRigidActiveParticlesBuffer::ReadLock()
+	{
+		if (bUseLock)
+		{
+			ResourceOutLock.ReadLock();
+		}
+	}
+
+	void FPBDRigidActiveParticlesBuffer::ReadUnlock()
+	{
+		if (bUseLock)
+		{
+			ResourceOutLock.ReadUnlock();
+		}
+	}
+
+	void FPBDRigidActiveParticlesBuffer::WriteLock()
+	{
+		if (bUseLock)
+		{
+			ResourceOutLock.WriteLock();
+		}
+	}
+
+	void FPBDRigidActiveParticlesBuffer::WriteUnlock()
+	{
+		if (bUseLock)
+		{
+			ResourceOutLock.WriteUnlock();
+		}
+	}
 }
diff --git a/Engine/Source/Runtime/Experimental/ChaosSolvers/Private/PBDRigidsSolver.cpp b/Engine/Source/Runtime/Experimental/ChaosSolvers/Private/PBDRigidsSolver.cpp
--- a/Engine/Source/Runtime/Experimental/ChaosSolvers/Private/PBDRigidsSolver.cpp
+++ b/Engine/Source/Runtime/Experimental/ChaosSolvers/Private/PBDRigidsSolver.cpp
@@ -192,7 +192,7 @@ namespace Chaos
 		, MEvolution(new FPBDRigidsEvolution(Particles, ChaosSolverCollisionDefaultIterationsCVar, ChaosSolverCollisionDefaultPushoutIterationsCVar, BufferingModeIn == Chaos::EMultiBufferMode::Single))
 		, MEventManager(new FEventManager(BufferingModeIn))
 		, MSolverEventFilters(new FSolverEventFilters())
-		, MActiveParticlesBuffer(new FActiveParticlesBuffer(BufferingModeIn))
+		, MActiveParticlesBuffer(new FActiveParticlesBuffer(BufferingModeIn, BufferingModeIn == Chaos::EMultiBufferMode::Single))
 		, MCurrentLock(new FCriticalSection())
 	{
 		UE_LOG(LogPBDRigidsSolver, Verbose, TEXT("PBDRigidsSolver::PBDRigidsSolver()"));
diff --git a/Engine/Source/Runtime/Experimental/ChaosSolvers/Public/PBDRigidActiveParticlesBuffer.h b/Engine/Source/Runtime/Experimental/ChaosSolvers/Public/PBDRigidActiveParticlesBuffer.h
--- a/Engine/Source/Runtime/Experimental/ChaosSolvers/Public/PBDRigidActiveParticlesBuffer.h
+++ b/Engine/Source/Runtime/Experimental/ChaosSolvers/Public/PBDRigidActiveParticlesBuffer.h
@@ -26,11 +26,16 @@ namespace Chaos
 		friend class FPBDRigidActiveParticlesBufferAccessor;
 
 	public:
-		FPBDRigidActiveParticlesBuffer(const Chaos::EMultiBufferMode& InBufferMode);
+		FPBDRigidActiveParticlesBuffer(const Chaos::EMultiBufferMode& InBufferMode, bool bInSingleThreaded);
 
 		void CaptureSolverData(FPBDRigidsSolver* Solver);
 
 		void RemoveActiveParticleFromConsumerBuffer(TGeometryParticle<FReal, 3>* Particle);
+
+		void ReadLock();
+		void ReadUnlock();
+		void WriteLock();
+		void WriteUnlock();
 	
 	private:
 		const FPBDRigidActiveParticlesBufferOut* GetSolverOutData() const
@@ -53,6 +58,7 @@ namespace Chaos
 
 		Chaos::EMultiBufferMode BufferMode;
 		FRWLock ResourceOutLock;
+		bool bUseLock;
 
 		// Physics thread to game thread
 		TUniquePtr<IBufferResource<FPBDRigidActiveParticlesBufferOut>> SolverDataOut;
@@ -64,7 +70,7 @@ namespace Chaos
 		FPBDRigidActiveParticlesBufferAccessor(FPBDRigidActiveParticlesBuffer* InManager) : Manager(InManager)
 		{
 			check(InManager);
-			Manager->ResourceOutLock.ReadLock();
+			Manager->ReadLock();
 		}
 
 		const FPBDRigidActiveParticlesBufferOut* GetSolverOutData() const
@@ -74,7 +80,7 @@ namespace Chaos
 
 		~FPBDRigidActiveParticlesBufferAccessor()
 		{
-			Manager->ResourceOutLock.ReadUnlock();
+			Manager->ReadUnlock();
 		}
 
 	private:
