Make sure Push Model marks RemoteRole dirty when we change Net Owners.
Because of shared comparisons and scoped role downgrades, we can't rely on this behavior from AActor::SetOwner.

diff --git a/Engine/Source/Runtime/Engine/Private/RepLayout.cpp b/Engine/Source/Runtime/Engine/Private/RepLayout.cpp
--- a/Engine/Source/Runtime/Engine/Private/RepLayout.cpp
+++ b/Engine/Source/Runtime/Engine/Private/RepLayout.cpp
@@ -1146,6 +1146,7 @@ struct FComparePropertiesSharedParams
 	const TBitArray<>* const PushModelProperties = nullptr;
 	const bool bValidateProperties = false;
 	const bool bIsNetworkProfilerActive = false;
+	const bool bChangedNetOwner = false;
 #if (WITH_PUSH_VALIDATION_SUPPORT || USE_NETWORK_PROFILER)
 	TBitArray<> PropertiesCompared;
 	TBitArray<> PropertiesChanged;
@@ -1299,7 +1300,10 @@ static void CompareParentProperties(
 		// This really shouldn't cause any problems, but it will cause Validation to consistently fail (if GbPushModelValidateProperties
 		// is enabled).
 		// To get around this, and just to be sure our code paths are consistent, we forcibly mark Role and RemoteRole dirty here.
-		if (EnumHasAnyFlags(SharedParams.Flags, ERepLayoutFlags::IsActor) && SharedParams.bIsInitial)
+		// Also, if we changed Net Owner state, then our RemoteRole will have changed (due to FScopeRoleDowngrade) even though
+		// the value hasn't changed on the Actor through normal gameplay code.
+		// Instead of marking RemoteRole dirty in Scoped Role Downgrade (which would cause tons dirties), we detect that here.
+		if (UNLIKELY(EnumHasAnyFlags(SharedParams.Flags, ERepLayoutFlags::IsActor) && (SharedParams.bIsInitial || SharedParams.bChangedNetOwner)))
 		{
 			SharedParams.PushModelState->MarkPropertyDirty((int32)AActor::ENetFields_Private::RemoteRole);
 			SharedParams.PushModelState->MarkPropertyDirty((int32)AActor::ENetFields_Private::Role);
@@ -1510,7 +1514,8 @@ bool FRepLayout::CompareProperties(
 		/*PushModelState=*/UE4_RepLayout_Private::GetPerNetDriverState(RepChangelistState),
 		/*PushModelProperties=*/ LocalPushModelProperties,	
 		/*bValidateProperties=*/GbPushModelValidateProperties,
-		/*bIsNetworkProfilerActive=*/UE4_RepLayout_Private::IsNetworkProfilerComparisonTrackingEnabled()
+		/*bIsNetworkProfilerActive=*/UE4_RepLayout_Private::IsNetworkProfilerComparisonTrackingEnabled(),
+		/*bChangedNetOwner=*/ RepState && RepState->RepFlags.bNetOwner != RepFlags.bNetOwner
 	};
 
 	FComparePropertiesStackParams StackParams{
