Fix for broken debug display w/ au.Debug.SoundCues.Minimal
FAudioStats::FStatSoundInfo structure changed layout.

diff --git a/Engine/Source/Runtime/Engine/Private/AudioDebug.cpp b/Engine/Source/Runtime/Engine/Private/AudioDebug.cpp
--- a/Engine/Source/Runtime/Engine/Private/AudioDebug.cpp
+++ b/Engine/Source/Runtime/Engine/Private/AudioDebug.cpp
@@ -967,14 +967,14 @@ int32 FAudioDebugger::RenderStatCues(UWorld* World, FViewport* Viewport, FCanvas
 			{
 				if (WaveInstanceInfo.Volume >= MinDisplayVolume)
 				{
-					if (!StatSoundInfo.SoundName.Split(TEXT("."), &SoundPath, &SoundName))
+					if (!StatSoundInfo.SoundPath.Split(TEXT("."), &SoundPath, &SoundName))
 					{
-						if (!StatSoundInfo.SoundName.Split(SUBOBJECT_DELIMITER, &SoundPath, &SoundName))
+						if (!StatSoundInfo.SoundPath.Split(SUBOBJECT_DELIMITER, &SoundPath, &SoundName))
 						{
-							SoundName.Reset();
-							SoundPath = StatSoundInfo.SoundName;
+							SoundPath = StatSoundInfo.SoundPath;
 						}
 					}
+					SoundName = StatSoundInfo.SoundName;
 					SplitNames.Emplace(SoundName, SoundPath);
 
 					if (SoundName.Len() > MaxNameLength)
@@ -1008,14 +1008,14 @@ int32 FAudioDebugger::RenderStatCues(UWorld* World, FViewport* Viewport, FCanvas
 					FString MuteSoloReason;
 					bool bMutedOrSoloed = false;
 
-				if (FSoundSource::FDebugInfo* DebugInfo = WaveInstanceInfo.DebugInfo.Get())
-				{
+					if (FSoundSource::FDebugInfo* DebugInfo = WaveInstanceInfo.DebugInfo.Get())
+					{
 						// Color code same as icons. Red (mute), Yellow (solo), White (normal).
-					FScopeLock Lock(&DebugInfo->CS);
+						FScopeLock Lock(&DebugInfo->CS);
 						Color = DebugInfo->bIsMuted ? FColor::Red : DebugInfo->bIsSoloed ? FColor::Yellow : FColor::White;
 						bMutedOrSoloed = DebugInfo->bIsMuted || DebugInfo->bIsSoloed;
-					MuteSoloReason = !DebugInfo->MuteSoloReason.IsEmpty() ? FString::Printf(TEXT(" - %s"), *DebugInfo->MuteSoloReason) : TEXT("");
-				}
+						MuteSoloReason = !DebugInfo->MuteSoloReason.IsEmpty() ? FString::Printf(TEXT(" - %s"), *DebugInfo->MuteSoloReason) : TEXT("");
+					}
 
 					const int32 SoundNameIndex = ActiveSoundCount++;
 					const FString LeadingNumber = FString::Printf(TEXT("%4i. "), SoundNameIndex);
@@ -1037,12 +1037,12 @@ int32 FAudioDebugger::RenderStatCues(UWorld* World, FViewport* Viewport, FCanvas
 						Canvas->DrawShadowedString(CurrentX, Y, *PathAndMuting, StatsFont, Color);
 					}
 
-				Y += FontHeight;
-				break;
+					Y += FontHeight;
+					break;
+				}
 			}
 		}
 	}
-	}
 
 	Canvas->DrawShadowedString(X, Y, *FString::Printf(TEXT("Total: %i"), ActiveSoundCount), GetStatsFont(), GetBodyColor());
 	Y += FontHeight;
