Cameras: Add methods for ending camera shakes of a given type originating from a given source. Make other recently added method names more consistent.

diff --git a/Engine/Source/Runtime/Engine/Classes/Camera/CameraModifier_CameraShake.h b/Engine/Source/Runtime/Engine/Classes/Camera/CameraModifier_CameraShake.h
--- a/Engine/Source/Runtime/Engine/Classes/Camera/CameraModifier_CameraShake.h
+++ b/Engine/Source/Runtime/Engine/Classes/Camera/CameraModifier_CameraShake.h
@@ -106,8 +106,21 @@ public:
 	 */
 	virtual void RemoveAllCameraShakesOfClass(TSubclassOf<UCameraShake> ShakeClass, bool bImmediately = true);
 
+	/**
+	 * Stops and removes all camera shakes originating from the given source.
+	 * @param SourceComponent   The shake source.
+	 * @param bImmediately      If true, shake stops right away regardless of blend out settings. If false, shake may blend out according to its settings.
+	 */
 	virtual void RemoveAllCameraShakesFromSource(const UCameraShakeSourceComponent* SourceComponent, bool bImmediately = true);
 
+	/**
+	 * Stops and removes all camera shakes of the given class originating from the given source.
+	 * @param ShakeClasss       The camera shake class to remove.
+	 * @param SourceComponent   The shake source.
+	 * @param bImmediately      If true, shake stops right away regardless of blend out settings. If false, shake may blend out according to its settings.
+	 */
+	virtual void RemoveAllCameraShakesOfClassFromSource(TSubclassOf<UCameraShake> ShakeClass, const UCameraShakeSourceComponent* SourceComponent, bool bImmediately = true);
+
 	/** 
 	 * Stops and removes all camera shakes from the camera. 
 	 * @param bImmediately		If true, shake stops right away regardless of blend out settings. If false, shake may blend out according to its settings.
diff --git a/Engine/Source/Runtime/Engine/Classes/Camera/CameraShakeSourceComponent.h b/Engine/Source/Runtime/Engine/Classes/Camera/CameraShakeSourceComponent.h
--- a/Engine/Source/Runtime/Engine/Classes/Camera/CameraShakeSourceComponent.h
+++ b/Engine/Source/Runtime/Engine/Classes/Camera/CameraShakeSourceComponent.h
@@ -79,6 +79,10 @@ public:
 	UFUNCTION(BlueprintCallable, Category = CameraShake)
 	void PlayCameraShake(TSubclassOf<UCameraShake> InCameraShake);
 
+	/** Stops a camera shake originating from this source */
+	UFUNCTION(BlueprintCallable, Category = CameraShake)
+	void StopAllCameraShakesOfType(TSubclassOf<UCameraShake> InCameraShake, bool bImmediately = true);
+
 	/** Stops all currently active camera shakes that are originating from this source from all player controllers */
 	UFUNCTION(BlueprintCallable, Category = CameraShake)
 	void StopAllCameraShakes(bool bImmediately = true);
diff --git a/Engine/Source/Runtime/Engine/Classes/Camera/PlayerCameraManager.h b/Engine/Source/Runtime/Engine/Classes/Camera/PlayerCameraManager.h
--- a/Engine/Source/Runtime/Engine/Classes/Camera/PlayerCameraManager.h
+++ b/Engine/Source/Runtime/Engine/Classes/Camera/PlayerCameraManager.h
@@ -792,14 +792,18 @@ public:
 	UFUNCTION(BlueprintCallable, Category = "Camera Shakes")
 	virtual void StopAllInstancesOfCameraShake(TSubclassOf<class UCameraShake> Shake, bool bImmediately = true);
 
-	/** Stops playing all shakes originating from the given source. */
-	UFUNCTION(BlueprintCallable, Category = "Camera Shakes")
-	virtual void StopAllInstancesOfCameraShakeFromSource(class UCameraShakeSourceComponent* SourceComponent, bool bImmediately = true);
-
 	/** Stops all active camera shakes on this camera. */
 	UFUNCTION(BlueprintCallable, Category = "Camera Shakes")
 	virtual void StopAllCameraShakes(bool bImmediately = true);
 
+	/** Stops playing all shakes of the given class originating from the given source. */
+	UFUNCTION(BlueprintCallable, Category = "Camera Shakes")
+	virtual void StopAllInstancesOfCameraShakeFromSource(TSubclassOf<class UCameraShake> Shake, class UCameraShakeSourceComponent* SourceComponent, bool bImmediately = true);
+
+	/** Stops playing all shakes originating from the given source. */
+	UFUNCTION(BlueprintCallable, Category = "Camera Shakes")
+	virtual void StopAllCameraShakesFromSource(class UCameraShakeSourceComponent* SourceComponent, bool bImmediately = true);
+
 	//
 	//  CameraAnim fades.
 	//
diff --git a/Engine/Source/Runtime/Engine/Private/Camera/CameraModifier_CameraShake.cpp b/Engine/Source/Runtime/Engine/Private/Camera/CameraModifier_CameraShake.cpp
--- a/Engine/Source/Runtime/Engine/Private/Camera/CameraModifier_CameraShake.cpp
+++ b/Engine/Source/Runtime/Engine/Private/Camera/CameraModifier_CameraShake.cpp
@@ -239,6 +239,25 @@ void UCameraModifier_CameraShake::RemoveAllCameraShakesFromSource(const UCameraS
 	}
 }
 
+void UCameraModifier_CameraShake::RemoveAllCameraShakesOfClassFromSource(TSubclassOf<UCameraShake> ShakeClass, const UCameraShakeSourceComponent* SourceComponent, bool bImmediately)
+{
+	for (int32 i = ActiveShakes.Num() - 1; i >= 0; --i)
+	{
+		FActiveCameraShakeInfo ShakeInfo = ActiveShakes[i];  // Copy struct because we might delete it.
+		if (ShakeInfo.ShakeSource.Get() == SourceComponent && 
+				ShakeInfo.ShakeInstance != nullptr && 
+				ShakeInfo.ShakeInstance->GetClass()->IsChildOf(ShakeClass))
+		{
+			ShakeInfo.ShakeInstance->StopShake(bImmediately);
+			if (bImmediately)
+			{
+				ActiveShakes.RemoveAt(i, 1);
+				SaveShakeInExpiredPool(ShakeInfo.ShakeInstance);
+			}
+		}
+	}
+}
+
 void UCameraModifier_CameraShake::RemoveAllCameraShakes(bool bImmediately)
 {
 	// Clean up any active camera shake anims
diff --git a/Engine/Source/Runtime/Engine/Private/Camera/CameraShakeSourceComponent.cpp b/Engine/Source/Runtime/Engine/Private/Camera/CameraShakeSourceComponent.cpp
--- a/Engine/Source/Runtime/Engine/Private/Camera/CameraShakeSourceComponent.cpp
+++ b/Engine/Source/Runtime/Engine/Private/Camera/CameraShakeSourceComponent.cpp
@@ -92,6 +92,18 @@ void UCameraShakeSourceComponent::PlayCameraShake(TSubclassOf<UCameraShake> InCa
 	}
 }
 
+void UCameraShakeSourceComponent::StopAllCameraShakesOfType(TSubclassOf<UCameraShake> InCameraShake, bool bImmediately)
+{
+	for (FConstPlayerControllerIterator Iterator = GetWorld()->GetPlayerControllerIterator(); Iterator; ++Iterator)
+	{
+		APlayerController* PlayerController = Iterator->Get();
+		if (PlayerController != nullptr && PlayerController->PlayerCameraManager != nullptr)
+		{
+			PlayerController->PlayerCameraManager->StopAllInstancesOfCameraShakeFromSource(InCameraShake, this, bImmediately);
+		}
+	}
+}
+
 void UCameraShakeSourceComponent::StopAllCameraShakes(bool bImmediately)
 {
 	for (FConstPlayerControllerIterator Iterator = GetWorld()->GetPlayerControllerIterator(); Iterator; ++Iterator)
@@ -99,7 +111,7 @@ void UCameraShakeSourceComponent::StopAllCameraShakes(bool bImmediately)
 		APlayerController* PlayerController = Iterator->Get();
 		if (PlayerController != nullptr && PlayerController->PlayerCameraManager != nullptr)
 		{
-			PlayerController->PlayerCameraManager->StopAllInstancesOfCameraShakeFromSource(this, bImmediately);
+			PlayerController->PlayerCameraManager->StopAllCameraShakesFromSource(this, bImmediately);
 		}
 	}
 }
diff --git a/Engine/Source/Runtime/Engine/Private/PlayerCameraManager.cpp b/Engine/Source/Runtime/Engine/Private/PlayerCameraManager.cpp
--- a/Engine/Source/Runtime/Engine/Private/PlayerCameraManager.cpp
+++ b/Engine/Source/Runtime/Engine/Private/PlayerCameraManager.cpp
@@ -1305,19 +1305,27 @@ void APlayerCameraManager::StopAllInstancesOfCameraShake(TSubclassOf<class UCame
 	}
 }
 
-void APlayerCameraManager::StopAllInstancesOfCameraShakeFromSource(class UCameraShakeSourceComponent* SourceComponent, bool bImmediately)
+void APlayerCameraManager::StopAllCameraShakes(bool bImmediately)
 {
-	if (SourceComponent && CachedCameraShakeMod)
+	if (CachedCameraShakeMod)
 	{
-		CachedCameraShakeMod->RemoveAllCameraShakesFromSource(SourceComponent, bImmediately);
+		CachedCameraShakeMod->RemoveAllCameraShakes(bImmediately);
 	}
 }
 
-void APlayerCameraManager::StopAllCameraShakes(bool bImmediately)
+void APlayerCameraManager::StopAllInstancesOfCameraShakeFromSource(TSubclassOf<UCameraShake> ShakeClass, class UCameraShakeSourceComponent* SourceComponent, bool bImmediately)
 {
-	if (CachedCameraShakeMod)
+	if (ShakeClass && SourceComponent && CachedCameraShakeMod)
 	{
-		CachedCameraShakeMod->RemoveAllCameraShakes(bImmediately);
+		CachedCameraShakeMod->RemoveAllCameraShakesOfClassFromSource(ShakeClass, SourceComponent, bImmediately);
+	}
+}
+
+void APlayerCameraManager::StopAllCameraShakesFromSource(class UCameraShakeSourceComponent* SourceComponent, bool bImmediately)
+{
+	if (SourceComponent && CachedCameraShakeMod)
+	{
+		CachedCameraShakeMod->RemoveAllCameraShakesFromSource(SourceComponent, bImmediately);
 	}
 }
 
diff --git a/Engine/Source/Runtime/Engine/Private/PlayerController.cpp b/Engine/Source/Runtime/Engine/Private/PlayerController.cpp
--- a/Engine/Source/Runtime/Engine/Private/PlayerController.cpp
+++ b/Engine/Source/Runtime/Engine/Private/PlayerController.cpp
@@ -4351,7 +4351,7 @@ void APlayerController::ClientStopCameraShakesFromSource(class UCameraShakeSourc
 {
 	if (PlayerCameraManager != NULL)
 	{
-		PlayerCameraManager->StopAllInstancesOfCameraShakeFromSource(SourceComponent, bImmediately);
+		PlayerCameraManager->StopAllCameraShakesFromSource(SourceComponent, bImmediately);
 	}
 }
 
