Fix unity build.

diff --git a/Engine/Source/Runtime/Engine/Classes/PhysicsEngine/ShapeElem.h b/Engine/Source/Runtime/Engine/Classes/PhysicsEngine/ShapeElem.h
--- a/Engine/Source/Runtime/Engine/Classes/PhysicsEngine/ShapeElem.h
+++ b/Engine/Source/Runtime/Engine/Classes/PhysicsEngine/ShapeElem.h
@@ -6,6 +6,7 @@
 #include "UObject/ObjectMacros.h"
 #include "EngineDefines.h"
 #include "PhysxUserData.h"
+#include "PhysicsCore/Public/PhysicsInterfaceTypesCore.h"
 #include "ShapeElem.generated.h"
 
 namespace EAggCollisionShape
