Add support for per player view csv stats

#rnx

diff --git a/Engine/Source/Runtime/Core/Private/ProfilingDebugging/CsvProfiler.cpp b/Engine/Source/Runtime/Core/Private/ProfilingDebugging/CsvProfiler.cpp
--- a/Engine/Source/Runtime/Core/Private/ProfilingDebugging/CsvProfiler.cpp
+++ b/Engine/Source/Runtime/Core/Private/ProfilingDebugging/CsvProfiler.cpp
@@ -3274,6 +3274,12 @@ bool FCsvProfiler::EnableCategoryByString(const FString& CategoryName) const
 	return false;
 }
 
+void FCsvProfiler::EnableCategoryByIndex(uint32 CategoryIndex, bool bEnable) const
+{
+	check(CategoryIndex < CSV_MAX_CATEGORY_COUNT);
+	GCsvCategoriesEnabled[CategoryIndex] = bEnable;
+}
+
 bool FCsvProfiler::IsCapturing_Renderthread()
 {
 	check(IsInRenderingThread());
diff --git a/Engine/Source/Runtime/Core/Public/ProfilingDebugging/CsvProfiler.h b/Engine/Source/Runtime/Core/Public/ProfilingDebugging/CsvProfiler.h
--- a/Engine/Source/Runtime/Core/Public/ProfilingDebugging/CsvProfiler.h
+++ b/Engine/Source/Runtime/Core/Public/ProfilingDebugging/CsvProfiler.h
@@ -241,6 +241,8 @@ public:
 
 	CORE_API static void SetMetadata(const TCHAR* Key, const TCHAR* Value);
 
+	static CORE_API int32 RegisterCategory(const FString& Name, bool bEnableByDefault, bool bIsGlobal);
+
 	template <typename FmtType, typename... Types>
 	FORCEINLINE static void RecordEventf(int32 CategoryIndex, const FmtType& Fmt, Types... Args)
 	{
@@ -264,6 +266,7 @@ public:
 	CORE_API int32 GetNumFrameToCaptureOnEvent();
 
 	CORE_API bool EnableCategoryByString(const FString& CategoryName) const;
+	CORE_API void EnableCategoryByIndex(uint32 CategoryIndex, bool bEnable) const;
 
 	/** Per-frame update */
 	CORE_API void BeginFrame();
@@ -302,7 +305,6 @@ public:
 private:
 	CORE_API static void VARARGS RecordEventfInternal(int32 CategoryIndex, const TCHAR* Fmt, ...);
 
-	static CORE_API int32 RegisterCategory(const FString& Name, bool bEnableByDefault, bool bIsGlobal);
 	static int32 GetCategoryIndex(const FString& Name);
 
 	void FinalizeCsvFile();
diff --git a/Engine/Source/Runtime/Engine/Classes/Engine/GameViewportClient.h b/Engine/Source/Runtime/Engine/Classes/Engine/GameViewportClient.h
--- a/Engine/Source/Runtime/Engine/Classes/Engine/GameViewportClient.h
+++ b/Engine/Source/Runtime/Engine/Classes/Engine/GameViewportClient.h
@@ -820,6 +820,8 @@ public:
 	/** Does the viewport client have a software cursor set up for the given enum? */
 	bool HasSoftwareCursor(EMouseCursor::Type Cursor) const;
 
+	void EnableCsvPlayerStats(int32 LocalPlayerCount);
+
 private:
 	/** Resets the platform type shape to nullptr, to restore it to the OS default. */
 	void ResetHardwareCursorStates();
@@ -944,7 +946,8 @@ private:
 	bool SetDisplayConfiguration( const FIntPoint* Dimensions, EWindowMode::Type WindowMode);
 
 	/** Updates CSVProfiler camera stats */
-	void UpdateCsvCameraStats(const FSceneView* View);
+	void UpdateCsvCameraStats(const TMap<ULocalPlayer*, FSceneView*>& PlayerViewMap);
+
 
 #if WITH_EDITOR
 	/** Delegate called when game viewport client received input key */
diff --git a/Engine/Source/Runtime/Engine/Private/GameViewportClient.cpp b/Engine/Source/Runtime/Engine/Private/GameViewportClient.cpp
--- a/Engine/Source/Runtime/Engine/Private/GameViewportClient.cpp
+++ b/Engine/Source/Runtime/Engine/Private/GameViewportClient.cpp
@@ -62,7 +62,6 @@
 #include "IImageWrapperModule.h"
 #include "HAL/PlatformApplicationMisc.h"
 
-CSV_DEFINE_CATEGORY(View, true);
 
 #if WITH_EDITOR
 #include "Settings/LevelEditorPlaySettings.h"
@@ -116,43 +115,92 @@ static TAutoConsoleVariable<float> CVarSecondaryScreenPercentage( // TODO: make
 	TEXT(" 1: override secondary screen percentage."),
 	ECVF_Default);
 
+#if CSV_PROFILER
+struct FCsvLocalPlayer
+{
+	FCsvLocalPlayer()
+	{
+		CategoryIndex = INDEX_NONE;
+		PrevViewOrigin = FVector::ZeroVector;
+		LastFrame = 0;
+		PrevTime = 0.0;
+	}
 
-void UGameViewportClient::UpdateCsvCameraStats(const FSceneView* View)
+	uint32 CategoryIndex;
+	uint32 LastFrame;
+	double PrevTime;
+	FVector PrevViewOrigin;
+};
+static TMap<uint32, FCsvLocalPlayer> GCsvLocalPlayers;
+#endif
+
+void UGameViewportClient::EnableCsvPlayerStats(int32 LocalPlayerCount)
 {
 #if CSV_PROFILER
-	if (!View)
+	if (GCsvLocalPlayers.Num() < LocalPlayerCount)
 	{
-		return;
+		for (int PlayerIndex = GCsvLocalPlayers.Num(); PlayerIndex < LocalPlayerCount; PlayerIndex++)
+		{
+			FCsvLocalPlayer& CsvData = GCsvLocalPlayers.Add(PlayerIndex);
+			uint32 index = GCsvLocalPlayers.Num() - 1;
+			FString CategoryName = (PlayerIndex == 0) ? TEXT("View") : FString::Printf(TEXT("View%d"), index);
+			CsvData.CategoryIndex = FCsvProfiler::Get()->RegisterCategory(CategoryName, (PlayerIndex == 0) ? true : false, false);
+		}
+	}
+
+	int32 PlayerIndex = 0;
+	for (auto& KV : GCsvLocalPlayers)
+	{
+		FCsvLocalPlayer& value = KV.Value;
+		FCsvProfiler::Get()->EnableCategoryByIndex(value.CategoryIndex, PlayerIndex < LocalPlayerCount);
+		PlayerIndex++;
 	}
-	static uint32 PrevFrameNumber = GFrameNumber;
-	static double PrevTime = 0.0;
-	static FVector PrevViewOrigin = FVector(ForceInitToZero);
-
-	// TODO: support multiple views/view families, e.g for splitscreen. For now, we just output stats for the first one.
-	if (GFrameNumber != PrevFrameNumber)
-	{
-		FVector ViewOrigin = View->ViewMatrices.GetViewOrigin();
-		FVector ForwardVec = View->ViewMatrices.GetOverriddenTranslatedViewMatrix().GetColumn(2);
-		FVector UpVec = View->ViewMatrices.GetOverriddenTranslatedViewMatrix().GetColumn(1);
-		FVector Diff = ViewOrigin - PrevViewOrigin;
-		double CurrentTime = FPlatformTime::Seconds();
-		double DeltaT = CurrentTime - PrevTime;
-		FVector Velocity = Diff / float(DeltaT);
-		float CameraSpeed = Velocity.Size();
-		PrevViewOrigin = ViewOrigin;
-		PrevTime = CurrentTime;
-		PrevFrameNumber = GFrameNumber;
-
-		CSV_CUSTOM_STAT(View, PosX, View->ViewMatrices.GetViewOrigin().X, ECsvCustomStatOp::Set);
-		CSV_CUSTOM_STAT(View, PosY, View->ViewMatrices.GetViewOrigin().Y, ECsvCustomStatOp::Set);
-		CSV_CUSTOM_STAT(View, PosZ, View->ViewMatrices.GetViewOrigin().Z, ECsvCustomStatOp::Set);
-		CSV_CUSTOM_STAT(View, ForwardX, ForwardVec.X, ECsvCustomStatOp::Set);
-		CSV_CUSTOM_STAT(View, ForwardY, ForwardVec.Y, ECsvCustomStatOp::Set);
-		CSV_CUSTOM_STAT(View, ForwardZ, ForwardVec.Z, ECsvCustomStatOp::Set);
-		CSV_CUSTOM_STAT(View, UpX, UpVec.X, ECsvCustomStatOp::Set);
-		CSV_CUSTOM_STAT(View, UpY, UpVec.Y, ECsvCustomStatOp::Set);
-		CSV_CUSTOM_STAT(View, UpZ, UpVec.Z, ECsvCustomStatOp::Set);
-		CSV_CUSTOM_STAT(View, Speed, CameraSpeed, ECsvCustomStatOp::Set);
+#endif
+}
+
+void UGameViewportClient::UpdateCsvCameraStats(const TMap<ULocalPlayer*, FSceneView*>& PlayerViewMap)
+{
+#if CSV_PROFILER
+
+	for (TMap<ULocalPlayer*, FSceneView*>::TConstIterator It(PlayerViewMap); It; ++It)
+	{
+		ULocalPlayer* LocalPlayer = It.Key();
+		FSceneView* SceneView = It.Value();
+
+		uint32 PlayerIndex = ConvertLocalPlayerToGamePlayerIndex(LocalPlayer);
+		if (PlayerIndex != INDEX_NONE)
+		{
+			FCsvLocalPlayer& CsvData = GCsvLocalPlayers.FindOrAdd(PlayerIndex);
+			if (CsvData.CategoryIndex == INDEX_NONE)
+			{
+				uint32 index = GCsvLocalPlayers.Num() - 1;
+				FString CategoryName = LocalPlayer->IsPrimaryPlayer() ? TEXT("View") : FString::Printf(TEXT("View%d"), index);
+				CsvData.CategoryIndex = FCsvProfiler::Get()->RegisterCategory(CategoryName, (index == 0) ? true : false, false);
+			}
+
+			FVector ViewOrigin = SceneView->ViewMatrices.GetViewOrigin();
+			FVector ForwardVec = SceneView->ViewMatrices.GetOverriddenTranslatedViewMatrix().GetColumn(2);
+			FVector UpVec = SceneView->ViewMatrices.GetOverriddenTranslatedViewMatrix().GetColumn(1);
+			FVector Diff = ViewOrigin - CsvData.PrevViewOrigin;
+			double CurrentTime = FPlatformTime::Seconds();
+			double DeltaT = CurrentTime - CsvData.PrevTime;
+			FVector Velocity = Diff / float(DeltaT);
+			float CameraSpeed = Velocity.Size();
+			CsvData.PrevViewOrigin = ViewOrigin;
+			CsvData.LastFrame = GFrameNumber;
+			CsvData.PrevTime = CurrentTime;
+
+			FCsvProfiler::RecordCustomStat("PosX", CsvData.CategoryIndex, ViewOrigin.X, ECsvCustomStatOp::Set);
+			FCsvProfiler::RecordCustomStat("PosY", CsvData.CategoryIndex, ViewOrigin.Y, ECsvCustomStatOp::Set);
+			FCsvProfiler::RecordCustomStat("PosZ", CsvData.CategoryIndex, ViewOrigin.Z, ECsvCustomStatOp::Set);
+			FCsvProfiler::RecordCustomStat("ForwardX", CsvData.CategoryIndex, ForwardVec.X, ECsvCustomStatOp::Set);
+			FCsvProfiler::RecordCustomStat("ForwardY", CsvData.CategoryIndex, ForwardVec.Y, ECsvCustomStatOp::Set);
+			FCsvProfiler::RecordCustomStat("ForwardZ", CsvData.CategoryIndex, ForwardVec.Z, ECsvCustomStatOp::Set);
+			FCsvProfiler::RecordCustomStat("UpX", CsvData.CategoryIndex, UpVec.X, ECsvCustomStatOp::Set);
+			FCsvProfiler::RecordCustomStat("UpY", CsvData.CategoryIndex, UpVec.Y, ECsvCustomStatOp::Set);
+			FCsvProfiler::RecordCustomStat("UpZ", CsvData.CategoryIndex, UpVec.Z, ECsvCustomStatOp::Set);
+			FCsvProfiler::RecordCustomStat("Speed", CsvData.CategoryIndex, CameraSpeed, ECsvCustomStatOp::Set);
+		}
 	}
 #endif
 }
@@ -1384,10 +1432,6 @@ void UGameViewportClient::Draw(FViewport* InViewport, FCanvas* SceneCanvas)
 					#if RHI_RAYTRACING
 						View->SetupRayTracedRendering();
 					#endif
-
-					#if CSV_PROFILER
-						UpdateCsvCameraStats(View);
-					#endif
 					}
 
 					// Add view information for resource streaming. Allow up to 5X boost for small FOV.
@@ -1399,6 +1443,10 @@ void UGameViewportClient::Draw(FViewport* InViewport, FCanvas* SceneCanvas)
 		}
 	}
 
+#if CSV_PROFILER
+	UpdateCsvCameraStats(PlayerViewMap);
+#endif
+
 	FinalizeViews(&ViewFamily, PlayerViewMap);
 
 	// Update level streaming.
