Fix a crash when textures use ETextureMipLoadOptions::OnlyFirstMip together with FCompactByteBulkData

diff --git a/Engine/Source/Runtime/Engine/Private/TextureDerivedData.cpp b/Engine/Source/Runtime/Engine/Private/TextureDerivedData.cpp
--- a/Engine/Source/Runtime/Engine/Private/TextureDerivedData.cpp
+++ b/Engine/Source/Runtime/Engine/Private/TextureDerivedData.cpp
@@ -27,6 +27,7 @@
 #include "Engine/VolumeTexture.h"
 #include "VT/VirtualTextureBuildSettings.h"
 #include "VT/VirtualTextureBuiltData.h"
+#include "HAL/FileManager.h"
 
 #if WITH_EDITOR
 
@@ -1079,7 +1080,37 @@ bool FTexturePlatformData::TryLoadMips(int32 FirstMipToLoad, void** OutMipData,
 		{
 			if (OutMipData != nullptr)
 			{
-#if PLATFORM_SUPPORTS_TEXTURE_STREAMING && !TEXTURE2DMIPMAP_USE_COMPACT_BULKDATA
+#if TEXTURE2DMIPMAP_USE_COMPACT_BULKDATA
+				if (Mip.BulkData.IsInlined())
+				{
+					Mip.BulkData.GetCopy(&OutMipData[MipIndex - FirstMipToLoad], true);
+				}
+				else
+				{
+					check(Texture);
+					FString FileName;
+					verify(Texture->GetMipDataFilename(MipIndex, FileName));
+
+					check(FileName.EndsWith(TEXT(".ubulk")) || FileName.EndsWith(TEXT(".uptnl")));
+					check(!Mip.BulkData.IsStoredCompressedOnDisk() && Mip.BulkData.GetBulkDataSize() > 0);
+
+					FArchive* Ar = IFileManager::Get().CreateFileReader(*FileName, FILEREAD_Silent);
+					if (!Ar)
+					{
+						continue;
+					}
+
+					void*& Dest = OutMipData[MipIndex - FirstMipToLoad];
+					if (!Dest)
+					{
+						Dest = FMemory::Malloc(Mip.BulkData.GetBulkDataSize());
+					}
+					Ar->Seek(Mip.BulkData.GetBulkDataOffsetInFile());
+					Ar->Serialize(Dest, Mip.BulkData.GetBulkDataSize());
+					delete Ar;
+				}
+#else
+#if PLATFORM_SUPPORTS_TEXTURE_STREAMING
 				// We want to make sure that any non-streamed mips are coming from the texture asset file, and not from an external bulk file.
 				// But because "r.TextureStreaming" is driven by the project setting as well as the command line option "-NoTextureStreaming", 
 				// is it possible for streaming mips to be loaded in non streaming ways.
@@ -1089,6 +1120,7 @@ bool FTexturePlatformData::TryLoadMips(int32 FirstMipToLoad, void** OutMipData,
 				}
 #endif
 				Mip.BulkData.GetCopy(&OutMipData[MipIndex - FirstMipToLoad], true);
+#endif
 			}
 			NumMipsCached++;
 		}
