Remove unused GetAllocatedSize() method...this memory is now accounted for as part of the larger frozen shader memory image

diff --git a/Engine/Source/Runtime/Engine/Public/MaterialShared.h b/Engine/Source/Runtime/Engine/Public/MaterialShared.h
--- a/Engine/Source/Runtime/Engine/Public/MaterialShared.h
+++ b/Engine/Source/Runtime/Engine/Public/MaterialShared.h
@@ -541,31 +541,6 @@ public:
 		return UniformBufferLayout;
 	}
 
-	uint32 GetAllocatedSize() const
-	{
-		/*	// sizeof(TArray<TRefCountPtr>)+sizeof(Contents)*ArrayLength
-		return UniformScalarExpressions.GetAllocatedSize()
-			+ sizeof(FMaterialUniformExpressionScalarParameter) * UniformScalarExpressions.Num()
-			+ UniformVectorExpressions.GetAllocatedSize()
-			+ sizeof(FMaterialUniformExpressionVectorParameter) * UniformVectorExpressions.Num()
-			+ Uniform2DTextureExpressions.GetAllocatedSize()
-			+ sizeof(FMaterialUniformExpressionTexture) * Uniform2DTextureExpressions.Num()
-			+ UniformCubeTextureExpressions.GetAllocatedSize()
-			+ sizeof(FMaterialUniformExpressionTexture) * UniformCubeTextureExpressions.Num()
-			+ Uniform2DArrayTextureExpressions.GetAllocatedSize()
-			+ sizeof(FMaterialUniformExpressionTexture) * Uniform2DArrayTextureExpressions.Num()
-			+ UniformVolumeTextureExpressions.GetAllocatedSize()
-			+ sizeof(FMaterialUniformExpressionTexture) * UniformVolumeTextureExpressions.Num()
-			+ UniformVirtualTextureExpressions.GetAllocatedSize()
-			+ sizeof(FMaterialUniformExpressionTexture) * UniformVirtualTextureExpressions.Num()
-			+ UniformExternalTextureExpressions.GetAllocatedSize()
-			+ sizeof(FMaterialUniformExpressionExternalTexture) * UniformExternalTextureExpressions.Num()
-			+ VTStacks.GetAllocatedSize()
-			+ ParameterCollections.GetAllocatedSize()
-			+ (UniformBufferStruct ? (sizeof(FShaderParametersMetadata) + UniformBufferStruct->GetMembers().GetAllocatedSize()) : 0);*/
-		return 0u;
-	}
-
 	inline const FMaterialVectorParameterInfo& GetVectorParameter(uint32 Index) const { return UniformVectorParameters[Index]; }
 	inline const FMaterialScalarParameterInfo& GetScalarParameter(uint32 Index) const { return UniformScalarParameters[Index]; }
 	inline const FMaterialTextureParameterInfo& GetTextureParameter(EMaterialTextureParameterType Type, int32 Index) const { return UniformTextureParameters[(uint32)Type][Index]; }
