Fix crash when MI layer expressions reference quality

diff --git a/Engine/Source/Runtime/Engine/Private/Materials/MaterialInstance.cpp b/Engine/Source/Runtime/Engine/Private/Materials/MaterialInstance.cpp
--- a/Engine/Source/Runtime/Engine/Private/Materials/MaterialInstance.cpp
+++ b/Engine/Source/Runtime/Engine/Private/Materials/MaterialInstance.cpp
@@ -4001,6 +4001,7 @@ void UMaterialInstance::UpdateStaticPermutation(const FStaticParameterSet& NewPa
 void UMaterialInstance::UpdateCachedLayerParameters()
 {
 	FMaterialCachedExpressionData CachedExpressionData;
+	CachedExpressionData.Reset();
 	if (Parent)
 	{
 		CachedExpressionData.ReferencedTextures = Parent->GetReferencedTextures();
