Implement GetResourceSizeEx for UVectorFieldStatic so that obj list reports size estiimates of resource memory

diff --git a/Engine/Source/Runtime/Engine/Classes/VectorField/VectorFieldStatic.h b/Engine/Source/Runtime/Engine/Classes/VectorField/VectorFieldStatic.h
--- a/Engine/Source/Runtime/Engine/Classes/VectorField/VectorFieldStatic.h
+++ b/Engine/Source/Runtime/Engine/Classes/VectorField/VectorFieldStatic.h
@@ -59,6 +59,7 @@ public:
 	//~ Begin UObject Interface.
 	virtual void PostLoad() override;
 	virtual void BeginDestroy() override;
+	virtual void GetResourceSizeEx(FResourceSizeEx& CumulativeResourceSize) override;
 #if WITH_EDITOR
 	virtual void PostEditChangeProperty(FPropertyChangedEvent& PropertyChangedEvent) override;
 #endif // WITH_EDITOR
diff --git a/Engine/Source/Runtime/Engine/Private/VectorField.cpp b/Engine/Source/Runtime/Engine/Private/VectorField.cpp
--- a/Engine/Source/Runtime/Engine/Private/VectorField.cpp
+++ b/Engine/Source/Runtime/Engine/Private/VectorField.cpp
@@ -263,6 +263,22 @@ public:
 			});
 	}
 
+	void GetResourceSizeEx(FResourceSizeEx& CumulativeResourceSize) const
+	{
+		if (VolumeData)
+		{
+			CumulativeResourceSize.AddDedicatedSystemMemoryBytes(SizeX * SizeY * SizeZ * sizeof(FFloat16Color));
+		}
+		if (VolumeTextureRHI)
+		{
+			const FPixelFormatInfo& FormatInfo = GPixelFormats[PF_FloatRGBA];
+			const SIZE_T NumBlocksX = (SIZE_T)FMath::DivideAndRoundUp(SizeX, FormatInfo.BlockSizeX);
+			const SIZE_T NumBlocksY = (SIZE_T)FMath::DivideAndRoundUp(SizeY, FormatInfo.BlockSizeY);
+			const SIZE_T NumBlocksZ = (SIZE_T)FMath::DivideAndRoundUp(SizeZ, FormatInfo.BlockSizeZ);
+			CumulativeResourceSize.AddDedicatedVideoMemoryBytes(NumBlocksX * NumBlocksY * NumBlocksZ * (SIZE_T)FormatInfo.BlockBytes);
+		}
+	}
+
 private:
 
 	/** Static volume texture data. */
@@ -486,6 +502,21 @@ void UVectorFieldStatic::BeginDestroy()
 	Super::BeginDestroy();
 }
 
+void UVectorFieldStatic::GetResourceSizeEx(FResourceSizeEx& CumulativeResourceSize)
+{
+	Super::GetResourceSizeEx(CumulativeResourceSize);
+
+	CumulativeResourceSize.AddDedicatedSystemMemoryBytes(CPUData.GetAllocatedSize());
+	if (SourceData.IsBulkDataLoaded())
+	{
+		CumulativeResourceSize.AddDedicatedSystemMemoryBytes(SourceData.GetBulkDataSize());
+	}
+	if (Resource)
+	{
+		((const FVectorFieldStaticResource*)Resource)->GetResourceSizeEx(CumulativeResourceSize);
+	}
+}
+
 #if WITH_EDITOR
 void UVectorFieldStatic::PostEditChangeProperty(FPropertyChangedEvent& PropertyChangedEvent)
 {
