Removed materials that are only referenced by stripped static mesh LODs.

diff --git a/Engine/Source/Runtime/Engine/Private/StaticMesh.cpp b/Engine/Source/Runtime/Engine/Private/StaticMesh.cpp
--- a/Engine/Source/Runtime/Engine/Private/StaticMesh.cpp
+++ b/Engine/Source/Runtime/Engine/Private/StaticMesh.cpp
@@ -4481,6 +4481,45 @@ void UStaticMesh::Serialize(FArchive& Ar)
 	bool bCooked = Ar.IsCooking();
 	Ar << bCooked;
 
+#if WITH_EDITOR
+	// Before we save out cooked data, see if we can trim any references
+	if (bCooked && Ar.IsSaving() && Ar.CookingTarget())
+	{
+		const ITargetPlatform* CookingTarget = Ar.CookingTarget();
+
+		const bool bWantToStripLODs = RenderData && (CVarStripMinLodDataDuringCooking.GetValueOnAnyThread() != 0);
+		const int32 MinLODIndex = MinLOD.GetValueForPlatformIdentifiers(
+			CookingTarget->GetPlatformInfo().PlatformGroupName,
+			CookingTarget->GetPlatformInfo().VanillaPlatformName);
+
+		if (bWantToStripLODs)
+		{
+			// Move materials into temp array
+			TArray<UMaterialInterface*> UnusedMaterials;
+			for (FStaticMaterial& Material : StaticMaterials)
+			{
+				UnusedMaterials.Add(Material.MaterialInterface);
+				Material.MaterialInterface = nullptr;
+			}
+
+			// For each material still referenced, move it back
+			for (int32 LODIndex = MinLODIndex; LODIndex < RenderData->LODResources.Num(); ++LODIndex)
+			{
+				const FStaticMeshLODResources& LOD = RenderData->LODResources[LODIndex];
+
+				for (const FStaticMeshSection& Section : LOD.Sections)
+				{
+					if (UnusedMaterials[Section.MaterialIndex])
+					{
+						StaticMaterials[Section.MaterialIndex].MaterialInterface = UnusedMaterials[Section.MaterialIndex];
+						UnusedMaterials[Section.MaterialIndex] = nullptr;
+					}
+				}
+			}
+		}
+	}
+#endif // WITH_EDITOR
+
 #if WITH_EDITORONLY_DATA
 	if (Ar.IsLoading() && Ar.UE4Ver() < VER_UE4_REMOVE_ZERO_TRIANGLE_SECTIONS)
 	{
