(4.25) Fix ChildActorName changing on load
Fix a Child Actor's Label changing when the name is staying consistent

diff --git a/Engine/Source/Runtime/Engine/Private/Components/ChildActorComponent.cpp b/Engine/Source/Runtime/Engine/Private/Components/ChildActorComponent.cpp
--- a/Engine/Source/Runtime/Engine/Private/Components/ChildActorComponent.cpp
+++ b/Engine/Source/Runtime/Engine/Private/Components/ChildActorComponent.cpp
@@ -289,16 +289,21 @@ FChildActorComponentInstanceData::FChildActorComponentInstanceData(const UChildA
 	, ChildActorName(Component->GetChildActorName())
 	, ComponentInstanceData(nullptr)
 {
-	if (Component->GetChildActor())
+	if (AActor* ChildActor = Component->GetChildActor())
 	{
-		ComponentInstanceData = MakeShared<FComponentInstanceDataCache>(Component->GetChildActor());
+		if (ChildActorName.IsNone())
+		{
+			ChildActorName = ChildActor->GetFName();
+		}
+
+		ComponentInstanceData = MakeShared<FComponentInstanceDataCache>(ChildActor);
 		// If it is empty dump it
 		if (!ComponentInstanceData->HasInstanceData())
 		{
 			ComponentInstanceData.Reset();
 		}
 
-		USceneComponent* ChildRootComponent = Component->GetChildActor()->GetRootComponent();
+		USceneComponent* ChildRootComponent = ChildActor->GetRootComponent();
 		if (ChildRootComponent)
 		{
 			for (USceneComponent* AttachedComponent : ChildRootComponent->GetAttachChildren())
@@ -306,7 +311,7 @@ FChildActorComponentInstanceData::FChildActorComponentInstanceData(const UChildA
 				if (AttachedComponent)
 				{
 					AActor* AttachedActor = AttachedComponent->GetOwner();
-					if (AttachedActor != Component->GetChildActor())
+					if (AttachedActor != ChildActor)
 					{
 						FChildActorAttachedActorInfo Info;
 						Info.Actor = AttachedActor;
@@ -393,6 +398,9 @@ void UChildActorComponent::ApplyComponentInstanceData(FChildActorComponentInstan
 			if (ChildActor->Rename(*ChildActorNameString, nullptr, REN_Test))
 			{
 				ChildActor->Rename(*ChildActorNameString, nullptr, REN_DoNotDirty | REN_ForceNoResetLoaders);
+#if WITH_EDITOR
+				ChildActor->ClearActorLabel();
+#endif
 			}
 		}
 
