Cleanup ActorLabel code
#rnx

diff --git a/Engine/Source/Runtime/Engine/Classes/GameFramework/Actor.h b/Engine/Source/Runtime/Engine/Classes/GameFramework/Actor.h
--- a/Engine/Source/Runtime/Engine/Classes/GameFramework/Actor.h
+++ b/Engine/Source/Runtime/Engine/Classes/GameFramework/Actor.h
@@ -708,7 +708,7 @@ private:
 	 * and call AActor::SetActorLabel() or FActorLabelUtilities::SetActorLabelUnique() to change the label.  Never set the label directly.
 	 */
 	UPROPERTY()
-	FString ActorLabel;
+	mutable FString ActorLabel;
 
 	/** The folder path of this actor in the world (empty=root, / separated)*/
 	UPROPERTY()
diff --git a/Engine/Source/Runtime/Engine/Private/ActorEditor.cpp b/Engine/Source/Runtime/Engine/Private/ActorEditor.cpp
--- a/Engine/Source/Runtime/Engine/Private/ActorEditor.cpp
+++ b/Engine/Source/Runtime/Engine/Private/ActorEditor.cpp
@@ -773,20 +773,15 @@ const FString& AActor::GetActorLabel() const
 
 	if( ActorLabel.IsEmpty() )
 	{
-		// Treating ActorLabel as mutable here (no 'mutable' keyword in current script compiler)
-		AActor* MutableThis = const_cast< AActor* >( this );
-
 		// Get the class
 		UClass* ActorClass = GetClass();
 
-		// NOTE: Calling GetName() is actually fairly slow (does ANSI->Wide conversion, lots of copies, etc.)
 		FString DefaultActorLabel = ActorClass->GetName();
 
 		// Strip off the ugly "_C" suffix for Blueprint class actor instances
-		UBlueprint* GeneratedByClassBlueprint = Cast<UBlueprint>( ActorClass->ClassGeneratedBy );
-		if( GeneratedByClassBlueprint != nullptr && DefaultActorLabel.EndsWith( TEXT( "_C" ) ) )
+		if (Cast<UBlueprint>(ActorClass->ClassGeneratedBy))
 		{
-			DefaultActorLabel.RemoveFromEnd( TEXT( "_C" ) );
+			DefaultActorLabel.RemoveFromEnd(TEXT("_C"), ESearchCase::CaseSensitive);
 		}
 
 		// We want the actor's label to be initially unique, if possible, so we'll use the number of the
@@ -805,7 +800,7 @@ const FString& AActor::GetActorLabel() const
 		// Remember, there could already be an actor with the same label in the level.  But that's OK, because
 		// actor labels aren't supposed to be unique.  We just try to make them unique initially to help
 		// disambiguate when opening up a new level and there are hundreds of actors of the same type.
-		MutableThis->ActorLabel = DefaultActorLabel;
+		ActorLabel = MoveTemp(DefaultActorLabel);
 	}
 
 	return ActorLabel;
@@ -820,8 +815,7 @@ void AActor::SetActorLabel( const FString& NewActorLabelDirty, bool bMarkDirty )
 void AActor::SetActorLabelInternal(const FString& NewActorLabelDirty, bool bMakeGloballyUniqueFName, bool bMarkDirty)
 {
 	// Clean up the incoming string a bit
-	FString NewActorLabel = NewActorLabelDirty;
-	NewActorLabel.TrimStartAndEndInline();
+	FString NewActorLabel = NewActorLabelDirty.TrimStartAndEnd();
 
 	// Validate incoming string before proceeding
 	FText OutErrorMessage;
@@ -839,7 +833,7 @@ void AActor::SetActorLabelInternal(const FString& NewActorLabelDirty, bool bMake
 			{
 				// Store new label
 				Modify(bMarkDirty);
-				ActorLabel = NewActorLabel;
+				ActorLabel = MoveTemp(NewActorLabel);
 			}
 		}
 
@@ -895,7 +889,7 @@ bool AActor::IsActorLabelEditable() const
 
 void AActor::ClearActorLabel()
 {
-	ActorLabel = TEXT("");
+	ActorLabel.Reset();
 }
 
 const FName& AActor::GetFolderPath() const
