Workaround for Xcode 11.1 bug causing it to complain about EIndex enum in AnimCurveTypes conflicting with Chaos

diff --git a/Engine/Source/Runtime/Engine/Private/Animation/AnimCurveTypes.cpp b/Engine/Source/Runtime/Engine/Private/Animation/AnimCurveTypes.cpp
--- a/Engine/Source/Runtime/Engine/Private/Animation/AnimCurveTypes.cpp
+++ b/Engine/Source/Runtime/Engine/Private/Animation/AnimCurveTypes.cpp
@@ -138,18 +138,18 @@ FVector FVectorCurve::Evaluate(float CurrentTime, float BlendWeight) const
 {
 	FVector Value;
 
-	Value.X = FloatCurves[X].Eval(CurrentTime)*BlendWeight;
-	Value.Y = FloatCurves[Y].Eval(CurrentTime)*BlendWeight;
-	Value.Z = FloatCurves[Z].Eval(CurrentTime)*BlendWeight;
+	Value.X = FloatCurves[(int32)EIndex::X].Eval(CurrentTime)*BlendWeight;
+	Value.Y = FloatCurves[(int32)EIndex::Y].Eval(CurrentTime)*BlendWeight;
+	Value.Z = FloatCurves[(int32)EIndex::Z].Eval(CurrentTime)*BlendWeight;
 
 	return Value;
 }
 
 void FVectorCurve::UpdateOrAddKey(const FVector& NewKey, float CurrentTime)
 {
-	FloatCurves[X].UpdateOrAddKey(CurrentTime, NewKey.X);
-	FloatCurves[Y].UpdateOrAddKey(CurrentTime, NewKey.Y);
-	FloatCurves[Z].UpdateOrAddKey(CurrentTime, NewKey.Z);
+	FloatCurves[(int32)EIndex::X].UpdateOrAddKey(CurrentTime, NewKey.X);
+	FloatCurves[(int32)EIndex::Y].UpdateOrAddKey(CurrentTime, NewKey.Y);
+	FloatCurves[(int32)EIndex::Z].UpdateOrAddKey(CurrentTime, NewKey.Z);
 }
 
 void FVectorCurve::GetKeys(TArray<float>& OutTimes, TArray<FVector>& OutValues)
@@ -185,9 +185,9 @@ void FVectorCurve::GetKeys(TArray<float>& OutTimes, TArray<FVector>& OutValues)
 
 void FVectorCurve::Resize(float NewLength, bool bInsert/* whether insert or remove*/, float OldStartTime, float OldEndTime)
 {
-	FloatCurves[X].ReadjustTimeRange(0, NewLength, bInsert, OldStartTime, OldEndTime);
-	FloatCurves[Y].ReadjustTimeRange(0, NewLength, bInsert, OldStartTime, OldEndTime);
-	FloatCurves[Z].ReadjustTimeRange(0, NewLength, bInsert, OldStartTime, OldEndTime);
+	FloatCurves[(int32)EIndex::X].ReadjustTimeRange(0, NewLength, bInsert, OldStartTime, OldEndTime);
+	FloatCurves[(int32)EIndex::Y].ReadjustTimeRange(0, NewLength, bInsert, OldStartTime, OldEndTime);
+	FloatCurves[(int32)EIndex::Z].ReadjustTimeRange(0, NewLength, bInsert, OldStartTime, OldEndTime);
 }
 
 int32 FVectorCurve::GetNumKeys()
diff --git a/Engine/Source/Runtime/Engine/Public/Animation/AnimCurveTypes.h b/Engine/Source/Runtime/Engine/Public/Animation/AnimCurveTypes.h
--- a/Engine/Source/Runtime/Engine/Public/Animation/AnimCurveTypes.h
+++ b/Engine/Source/Runtime/Engine/Public/Animation/AnimCurveTypes.h
@@ -182,7 +182,7 @@ struct FVectorCurve : public FAnimCurveBase
 {
 	GENERATED_USTRUCT_BODY()
 
-	enum EIndex
+	enum class EIndex
 	{
 		X = 0, 
 		Y, 
