Cameras: add utility functions to get a camera shake duration and blend in/out times.

diff --git a/Engine/Source/Runtime/Engine/Classes/Camera/CameraShake.h b/Engine/Source/Runtime/Engine/Classes/Camera/CameraShake.h
--- a/Engine/Source/Runtime/Engine/Classes/Camera/CameraShake.h
+++ b/Engine/Source/Runtime/Engine/Classes/Camera/CameraShake.h
@@ -204,6 +204,48 @@ public:
 	UPROPERTY(EditAnywhere, Category = AnimShake)
 	uint32 bRandomAnimSegment : 1;
 
+public:
+	/**
+	 * Gets the duration of this camera shake in seconds.
+	 *
+	 * The value could be 0 or negative if the shake uses the oscillator, meaning, respectively,
+	 * no oscillation, or indefinite oscillation.
+	 */
+	float GetCameraShakeDuration() const;
+
+	/**
+	 * Gets the duration of this camera shake's blend in and out.
+	 *
+	 * The values could be 0 or negative if there's no blend in and/or out.
+	 */
+	void GetCameraShakeBlendTimes(float& OutBlendIn, float& OutBlendOut) const;
+
+	/**
+	 * Gets the default duration for camera shakes of the given class.
+	 */
+	static bool GetCameraShakeDuration(TSubclassOf<UCameraShake> CameraShakeClass, float& OutDuration)
+	{
+		if (const UCameraShake* CDO = CameraShakeClass->GetDefaultObject<UCameraShake>())
+		{
+			OutDuration = CDO->GetCameraShakeDuration();
+			return true;
+		}
+		return false;
+	}
+
+	/**
+	 * Gets the default blend in/out durations for camera shakes of the given class.
+	 */
+	static bool GetCameraShakeBlendTimes(TSubclassOf<UCameraShake> CameraShakeClass, float& OutBlendIn, float& OutBlendOut)
+	{
+		if (const UCameraShake* CDO = CameraShakeClass->GetDefaultObject<UCameraShake>())
+		{
+			CDO->GetCameraShakeBlendTimes(OutBlendIn, OutBlendOut);
+			return true;
+		}
+		return false;
+	}
+
 protected:
 
 	// INSTANCE DATA
diff --git a/Engine/Source/Runtime/Engine/Private/Camera/CameraShake.cpp b/Engine/Source/Runtime/Engine/Private/Camera/CameraShake.cpp
--- a/Engine/Source/Runtime/Engine/Private/Camera/CameraShake.cpp
+++ b/Engine/Source/Runtime/Engine/Private/Camera/CameraShake.cpp
@@ -2,6 +2,7 @@
 
 #include "Camera/CameraShake.h"
 #include "Camera/PlayerCameraManager.h"
+#include "Camera/CameraAnim.h"
 #include "Camera/CameraAnimInst.h"
 #include "Engine/Engine.h"
 #include "IXRTrackingSystem.h" // for IsHeadTrackingAllowed()
@@ -427,3 +428,36 @@ void UCameraShake::SetCurrentTimeAndApplyShake(float NewTime, FMinimalViewInfo&
 		AnimInst->ApplyToView(POV);
 	}
 }
+
+float UCameraShake::GetCameraShakeDuration() const
+{
+	if (Anim != nullptr)
+	{
+		if (bRandomAnimSegment)
+		{
+			return RandomAnimSegmentDuration;
+		}
+		else
+		{
+			return Anim->AnimLength;
+		}
+	}
+	else
+	{
+		return OscillationDuration;
+	}
+}
+
+void UCameraShake::GetCameraShakeBlendTimes(float& OutBlendIn, float& OutBlendOut) const
+{
+	if (Anim != nullptr)
+	{
+		OutBlendIn = AnimBlendInTime;
+		OutBlendOut = AnimBlendOutTime;
+	}
+	else
+	{
+		OutBlendIn = OscillationBlendInTime;
+		OutBlendOut = OscillationBlendOutTime;
+	}
+}
