Fix for RF_Standalone flags going missing non-deterministically on MapBuildDataRegistry assets.

diff --git a/Engine/Source/Runtime/Engine/Private/Level.cpp b/Engine/Source/Runtime/Engine/Private/Level.cpp
--- a/Engine/Source/Runtime/Engine/Private/Level.cpp
+++ b/Engine/Source/Runtime/Engine/Private/Level.cpp
@@ -1912,8 +1912,11 @@ UMapBuildDataRegistry* ULevel::GetOrCreateMapBuildData()
 			// Release rendering data depending on MapBuildData, before we destroy MapBuildData
 			MapBuildData->InvalidateStaticLighting(GetWorld(), true, nullptr);
 
-			// Allow the legacy registry to be GC'ed
-			MapBuildData->ClearFlags(RF_Standalone);
+			// Allow the legacy registry to be GC'ed, but not if the registry is itself an asset, because assets are always standalone
+			if (!MapBuildData->IsAsset())
+			{
+				MapBuildData->ClearFlags(RF_Standalone);
+			}
 		}
 
 		UPackage* BuiltDataPackage = CreateMapBuildDataPackage();
diff --git a/Engine/Source/Runtime/Engine/Private/World.cpp b/Engine/Source/Runtime/Engine/Private/World.cpp
--- a/Engine/Source/Runtime/Engine/Private/World.cpp
+++ b/Engine/Source/Runtime/Engine/Private/World.cpp
@@ -4283,7 +4283,7 @@ void UWorld::CleanupWorldInternal(bool bSessionEnded, bool bCleanupResources, UW
 
 		if (WorldType != EWorldType::PIE)
 		{
-			if (PersistentLevel && PersistentLevel->MapBuildData)
+			if (PersistentLevel && PersistentLevel->MapBuildData && !PersistentLevel->MapBuildData->IsAsset())
 			{
 				PersistentLevel->MapBuildData->ClearFlags(RF_Standalone);
 
