Skin weight profiles changes:
- Implemented async GPU readback system to support non-cpu accessible skeletal mesh platforms
- Improved data layout to improve runtime cost, and adapted system to work with new skin weight buffers
- Added object version for new skin weight data layout
- Bumped skeletal mesh DDC key
Misc:
- Fixed logging issue when clamping the number of bone influences

diff --git a/Engine/Source/Developer/MeshUtilities/Private/MeshUtilities.cpp b/Engine/Source/Developer/MeshUtilities/Private/MeshUtilities.cpp
--- a/Engine/Source/Developer/MeshUtilities/Private/MeshUtilities.cpp
+++ b/Engine/Source/Developer/MeshUtilities/Private/MeshUtilities.cpp
@@ -114,6 +114,7 @@
 #include "IMeshReductionManagerModule.h"
 #include "MeshMergeModule.h"
 
+#include "Animation/SkinWeightProfile.h"
 
 DEFINE_LOG_CATEGORY(LogMeshUtilities);
 /*------------------------------------------------------------------------------
@@ -4066,8 +4067,8 @@ public:
 					InfluenceCount = FMath::Min<uint32>(InfluenceCount, MAX_TOTAL_INFLUENCES);
 					if (InfluenceCount > EXTRA_BONE_INFLUENCES && !FGPUBaseSkinVertexFactory::UseUnlimitedBoneInfluences(InfluenceCount))
 					{
-						InfluenceCount = EXTRA_BONE_INFLUENCES;
 						UE_LOG(LogSkeletalMesh, Warning, TEXT("Skeletal mesh of %d bone influences requires unlimited bone influence mode on. Influence truncated to %d."), InfluenceCount, EXTRA_BONE_INFLUENCES);
+						InfluenceCount = EXTRA_BONE_INFLUENCES;
 					}
 
 					// Setup the vertex influences.
@@ -5746,6 +5747,8 @@ void FMeshUtilities::GenerateRuntimeSkinWeightData(const FSkeletalMeshLODModel*
 
 		// Determine how many influences each skinweight can contain
 		const int32 NumInfluences = TargetLODModel.GetMaxBoneInfluences();
+		InOutSkinWeightOverrideData.NumWeightsPerVertex = NumInfluences;
+
 
 		TArray<FRawSkinWeight> UniqueWeights;
 		for (int32 VertexIndex = 0; VertexIndex < TargetVertices.Num(); ++VertexIndex)
@@ -5783,31 +5786,22 @@ void FMeshUtilities::GenerateRuntimeSkinWeightData(const FSkeletalMeshLODModel*
 				// If one hasn't been added yet, create a new one
 				if (OverrideIndex == INDEX_NONE)
 				{
-					FRuntimeSkinWeightProfileData::FSkinWeightOverrideInfo& DeltaOverride = InOutSkinWeightOverrideData.OverridesInfo.AddDefaulted_GetRef();
-
-					// Store offset into array and total number of influences to read
-					DeltaOverride.InfluencesOffset = InOutSkinWeightOverrideData.Weights.Num();
-					DeltaOverride.NumInfluences = 0;
+					OverrideIndex = UniqueWeights.Num();
 
 					// Write out non-zero weighted influences only
 					for (int32 InfluenceIndex = 0; InfluenceIndex < NumInfluences; ++InfluenceIndex)
 					{
-						if (SourceSkinWeight.InfluenceWeights[InfluenceIndex] > 0)
-						{
-							const uint32 Index = SourceSkinWeight.InfluenceBones[InfluenceIndex] << 16;
-							const uint32 Weight = SourceSkinWeight.InfluenceWeights[InfluenceIndex];
-							const uint32 Value = Index | Weight;
+						const FBoneIndexType Index = SourceSkinWeight.InfluenceBones[InfluenceIndex];
+						const uint8 Weight = SourceSkinWeight.InfluenceWeights[InfluenceIndex];
 
-							InOutSkinWeightOverrideData.Weights.Add(Value);
-							++DeltaOverride.NumInfluences;
-						}
+						InOutSkinWeightOverrideData.BoneIDs.Add(Index);
+						InOutSkinWeightOverrideData.BoneWeights.Add(Weight);
 					}
 
-					OverrideIndex = InOutSkinWeightOverrideData.OverridesInfo.Num() - 1;
 					UniqueWeights.Add(SourceSkinWeight);
 				}
 
-				InOutSkinWeightOverrideData.VertexIndexOverrideIndex.Add(VertexIndex, OverrideIndex);
+				InOutSkinWeightOverrideData.VertexIndexToInfluenceOffset.Add(VertexIndex, OverrideIndex);
 			}
 		}
 	}
diff --git a/Engine/Source/Developer/MeshUtilities/Public/MeshUtilities.h b/Engine/Source/Developer/MeshUtilities/Public/MeshUtilities.h
--- a/Engine/Source/Developer/MeshUtilities/Public/MeshUtilities.h
+++ b/Engine/Source/Developer/MeshUtilities/Public/MeshUtilities.h
@@ -11,7 +11,6 @@
 #include "MeshBuild.h"
 
 #include "IMeshMergeUtilities.h"
-#include "Animation/SkinWeightProfile.h"
 
 class UMeshComponent;
 class USkeletalMesh;
@@ -363,5 +362,5 @@ public:
 	virtual void FindOverlappingCorners(FOverlappingCorners& OutOverlappingCorners, const TArray<FVector>& InVertices, const TArray<uint32>& InIndices, float ComparisonThreshold) const = 0;
 
 	/** Used to generate runtime skin weight data from Editor-only data */
-	virtual void GenerateRuntimeSkinWeightData(const FSkeletalMeshLODModel* ImportedModel, const TArray<FRawSkinWeight>& InRawSkinWeights, FRuntimeSkinWeightProfileData& InOutSkinWeightOverrideData) const = 0;
+	virtual void GenerateRuntimeSkinWeightData(const FSkeletalMeshLODModel* ImportedModel, const TArray<FRawSkinWeight>& InRawSkinWeights, struct FRuntimeSkinWeightProfileData& InOutSkinWeightOverrideData) const = 0;
 };
diff --git a/Engine/Source/Developer/SkeletalMeshUtilitiesCommon/Private/LODUtilities.cpp b/Engine/Source/Developer/SkeletalMeshUtilitiesCommon/Private/LODUtilities.cpp
--- a/Engine/Source/Developer/SkeletalMeshUtilitiesCommon/Private/LODUtilities.cpp
+++ b/Engine/Source/Developer/SkeletalMeshUtilitiesCommon/Private/LODUtilities.cpp
@@ -1600,7 +1600,7 @@ bool FLODUtilities::UpdateAlternateSkinWeights(FSkeletalMeshLODModel& LODModelDe
 	int32 ProfileIndex = 0;
 	if (!ImportDataDest.AlternateInfluenceProfileNames.Find(ProfileNameDest.ToString(), ProfileIndex))
 	{
-		UE_LOG(LogLODUtilities, Error, TEXT("Failed to import Skin Weight Profile the alternate skinning imported source data is not available."), *SkeletalMeshName);
+		UE_LOG(LogLODUtilities, Warning, TEXT("Failed to import Skin Weight Profile the alternate skinning imported source data is not available."), *SkeletalMeshName);
 		return false;
 	}
 	check(ImportDataDest.AlternateInfluences.IsValidIndex(ProfileIndex));
diff --git a/Engine/Source/Editor/Persona/Private/AnimationEditorViewportClient.cpp b/Engine/Source/Editor/Persona/Private/AnimationEditorViewportClient.cpp
--- a/Engine/Source/Editor/Persona/Private/AnimationEditorViewportClient.cpp
+++ b/Engine/Source/Editor/Persona/Private/AnimationEditorViewportClient.cpp
@@ -962,7 +962,7 @@ FText FAnimationViewportClient::GetDisplayInfo(bool bDisplayAllInfo) const
 			
 			const FName ProfileName = PreviewMeshComponent->GetCurrentSkinWeightProfileName();
 			const FRuntimeSkinWeightProfileData* OverrideData = LODData.SkinWeightProfilesData.GetOverrideData(ProfileName);
-			TextValue = ConcatenateLine(TextValue, FText::Format(LOCTEXT("NumSkinWeightOverrides", "Skin Weight Profile Weights: {0}"),	OverrideData ? FText::AsNumber(OverrideData->OverridesInfo.Num()) : LOCTEXT("NoSkinWeightsOverridesForLOD", "no data for LOD")));
+			TextValue = ConcatenateLine(TextValue, FText::Format(LOCTEXT("NumSkinWeightOverrides", "Skin Weight Profile Weights: {0}"),	OverrideData ? FText::AsNumber(OverrideData->BoneWeights.Num() / OverrideData->NumWeightsPerVertex) : LOCTEXT("NoSkinWeightsOverridesForLOD", "no data for LOD")));
 		}
 	}
 
diff --git a/Engine/Source/Runtime/Core/Public/UObject/ObjectVersion.h b/Engine/Source/Runtime/Core/Public/UObject/ObjectVersion.h
--- a/Engine/Source/Runtime/Core/Public/UObject/ObjectVersion.h
+++ b/Engine/Source/Runtime/Core/Public/UObject/ObjectVersion.h
@@ -633,6 +633,8 @@ enum EUnrealEngineObjectUE4Version
 	VER_UE4_FIX_WIDE_STRING_CRC,
 	// Added package owner to allow private references
 	VER_UE4_ADDED_PACKAGE_OWNER,
+	// Changed the data layout for skin weight profile data
+	VER_UE4_SKINWEIGHT_PROFILE_DATA_LAYOUT_CHANGES,
 
 	// -----<new versions can be added before this line>-------------------------------------------------
 	// - this needs to be the last line (see note below)
diff --git a/Engine/Source/Runtime/D3D12RHI/Private/D3D12Commands.cpp b/Engine/Source/Runtime/D3D12RHI/Private/D3D12Commands.cpp
--- a/Engine/Source/Runtime/D3D12RHI/Private/D3D12Commands.cpp
+++ b/Engine/Source/Runtime/D3D12RHI/Private/D3D12Commands.cpp
@@ -424,8 +424,11 @@ void FD3D12CommandContext::RHICopyToStagingBuffer(FRHIVertexBuffer* SourceBuffer
 		FD3D12Resource* pDestResource = StagingBuffer->StagedRead;
 		D3D12_RESOURCE_DESC const& DestBufferDesc = pDestResource->GetDesc();
 
-		FD3D12DynamicRHI::TransitionResource(CommandListHandle, pSourceResource, D3D12_RESOURCE_STATE_COPY_SOURCE, 0);
-		CommandListHandle.FlushResourceBarriers();	// Must flush so the desired state is actually set.
+		if (pSourceResource->RequiresResourceStateTracking())
+		{
+			FD3D12DynamicRHI::TransitionResource(CommandListHandle, pSourceResource, D3D12_RESOURCE_STATE_COPY_SOURCE, 0);
+			CommandListHandle.FlushResourceBarriers();	// Must flush so the desired state is actually set.
+		}
 
 		numCopies++;
 
diff --git a/Engine/Source/Runtime/D3D12RHI/Public/D3D12Resources.h b/Engine/Source/Runtime/D3D12RHI/Public/D3D12Resources.h
--- a/Engine/Source/Runtime/D3D12RHI/Public/D3D12Resources.h
+++ b/Engine/Source/Runtime/D3D12RHI/Public/D3D12Resources.h
@@ -1091,7 +1091,7 @@ public:
 	FD3D12GPUFence(FName InName, FD3D12Fence* InFence)
 		: FRHIGPUFence(InName)
 		, Fence(InFence)
-		, Value(0)
+		, Value(MAX_uint64)
 	{}
 
 	void WriteInternal(ED3D12CommandQueueType QueueType);
diff --git a/Engine/Source/Runtime/Engine/Classes/Components/SkinnedMeshComponent.h b/Engine/Source/Runtime/Engine/Classes/Components/SkinnedMeshComponent.h
--- a/Engine/Source/Runtime/Engine/Classes/Components/SkinnedMeshComponent.h
+++ b/Engine/Source/Runtime/Engine/Classes/Components/SkinnedMeshComponent.h
@@ -475,6 +475,9 @@ protected:
 
 	/** Whether or not a Skin Weight profile is currently set for this component */
 	uint8 bSkinWeightProfileSet:1;
+
+	/** Whether or not a Skin Weight profile is currently pending load and creation for this component */
+	uint8 bSkinWeightProfilePending:1;
 public:
 
 	/** Whether we should use the min lod specified in MinLodModel for this component instead of the min lod in the mesh */
@@ -998,9 +1001,12 @@ public:
 	UFUNCTION(BlueprintCallable, Category = "Components|SkinnedMesh")
 	bool IsUsingSkinWeightProfile() const { return bSkinWeightProfileSet == 1;  }
 
-protected:
+	/** Check whether or not a Skin Weight Profile is currently pending load / create */
+	bool IsSkinWeightProfilePending() const { return bSkinWeightProfilePending == 1; }
+
 	/** Queues an update of the Skin Weight Buffer used by the current MeshObject */
 	void UpdateSkinWeightOverrideBuffer();
+protected:	
 
 	/** Name of currently set up Skin Weight profile, otherwise is 'none' */
 	FName CurrentSkinWeightProfileName;
diff --git a/Engine/Source/Runtime/Engine/Private/Animation/SkinWeightProfile.cpp b/Engine/Source/Runtime/Engine/Private/Animation/SkinWeightProfile.cpp
--- a/Engine/Source/Runtime/Engine/Private/Animation/SkinWeightProfile.cpp
+++ b/Engine/Source/Runtime/Engine/Private/Animation/SkinWeightProfile.cpp
@@ -8,11 +8,58 @@
 #include "Components/SkinnedMeshComponent.h"
 #include "Rendering/SkeletalMeshRenderData.h"
 #include "Rendering/SkinWeightVertexBuffer.h"
+#include "Animation/SkinWeightProfileManager.h"
 
 #if WITH_EDITOR
 #include "Rendering/SkeletalMeshLODImporterData.h"
 #include "Animation/DebugSkelMeshComponent.h"
 #endif
+#include "ComponentRecreateRenderStateContext.h"
+#include "SkeletalMeshTypes.h"
+#include "UObject/WeakObjectPtrTemplates.h"
+#include "Engine/GameEngine.h"
+#include "UObject/ObjectVersion.h"
+
+
+class ENGINE_API FSkinnedMeshComponentUpdateSkinWeightsContext
+{
+public:
+	FSkinnedMeshComponentUpdateSkinWeightsContext(USkeletalMesh* InSkeletalMesh)
+	{
+		for (TObjectIterator<USkinnedMeshComponent> It; It; ++It)
+		{
+			if (It->SkeletalMesh == InSkeletalMesh)
+			{
+				checkf(!It->IsUnreachable(), TEXT("%s"), *It->GetFullName());
+
+				if (It->IsRenderStateCreated())
+				{
+					check(It->IsRegistered());
+					MeshComponents.Add(*It);
+				}
+			}
+		}
+	}
+
+	~FSkinnedMeshComponentUpdateSkinWeightsContext()
+	{
+		const int32 ComponentCount = MeshComponents.Num();
+		for (int32 ComponentIndex = 0; ComponentIndex < ComponentCount; ++ComponentIndex)
+		{
+			USkinnedMeshComponent* Component = MeshComponents[ComponentIndex];
+
+			if (Component->IsRegistered())
+			{
+				Component->UpdateSkinWeightOverrideBuffer();
+			}
+		}
+	}
+
+
+private:
+	TArray< class USkinnedMeshComponent*> MeshComponents;
+};
+
 
 static void OnDefaultProfileCVarsChanged(IConsoleVariable* Variable)
 {
@@ -90,26 +137,22 @@ static FAutoConsoleVariableRef CVarSkinWeightProfilesAllowedFromLOD(
 
 FArchive& operator<<(FArchive& Ar, FRuntimeSkinWeightProfileData& OverrideData)
 {
-	Ar.UsingCustomVersion(FAnimObjectVersion::GUID);
-	Ar << OverrideData.OverridesInfo;
-	if (Ar.IsLoading() && Ar.CustomVer(FAnimObjectVersion::GUID) < FAnimObjectVersion::IncreaseBoneIndexLimitPerChunk)
+#if WITH_EDITOR
+	if (Ar.UE4Ver() < VER_UE4_SKINWEIGHT_PROFILE_DATA_LAYOUT_CHANGES)
 	{
-		TArray<uint16> OldWeights;
-		Ar << OldWeights;
-		// repack index_8bit|weight_8bit into index_16bit|weight_16bit
-		for (int32 i = 0; i < OldWeights.Num(); ++i)
-		{
-			uint16 OldWeight = OldWeights[i];
-			const uint8 Index = OldWeight >> 8;
-			const uint8 Weight = OldWeight & 0xFF;
-			OverrideData.Weights[i] = (Index << 16) | Weight;
-		}
+		Ar << OverrideData.OverridesInfo_DEPRECATED;
+		Ar << OverrideData.Weights_DEPRECATED;
 	}
 	else
-	{
-		Ar << OverrideData.Weights;
+#endif
+	{	
+		Ar << OverrideData.BoneIDs;
+		Ar << OverrideData.BoneWeights;
+		Ar << OverrideData.NumWeightsPerVertex;
 	}
-	Ar << OverrideData.VertexIndexOverrideIndex;
+	
+	Ar << OverrideData.VertexIndexToInfluenceOffset;
+
 
 	return Ar;
 }
@@ -123,8 +166,10 @@ FArchive& operator<<(FArchive& Ar, FSkinWeightProfilesData& LODData)
 
 FArchive& operator<<(FArchive& Ar, FRuntimeSkinWeightProfileData::FSkinWeightOverrideInfo& OverrideInfo)
 {
+#if WITH_EDITOR
 	Ar << OverrideInfo.InfluencesOffset;
-	Ar << OverrideInfo.NumInfluences;
+	Ar << OverrideInfo.NumInfluences_DEPRECATED;
+#endif
 
 	return Ar;
 }
@@ -186,10 +231,14 @@ void FSkinWeightProfilesData::Init(FSkinWeightVertexBuffer* InBaseBuffer)
 
 FSkinWeightProfilesData::~FSkinWeightProfilesData()
 {
+	ReleaseResources();
+
+	BaseBuffer = nullptr;
+	DefaultOverrideSkinWeightBuffer = nullptr;
+
 	bDefaultOverriden = false;
 	bStaticOverriden = false;
 	DefaultProfileName = NAME_None;
-	ReleaseResources();
 }
 
 #if !WITH_EDITOR
@@ -221,7 +270,12 @@ void FSkinWeightProfilesData::OverrideBaseBufferSkinWeightData(USkeletalMesh* Me
 		// If we found a profile try and find the override skin weights and apply if found
 		if (DefaultProfileIndex != INDEX_NONE)
 		{
-			ApplyOverrideProfile(BaseBuffer, Profiles[DefaultProfileIndex].Name);
+			const FRuntimeSkinWeightProfileData* ProfilePtr = OverrideData.Find(Profiles[DefaultProfileIndex].Name);
+			if (ProfilePtr)
+			{
+				ProfilePtr->ApplyDefaultOverride(BaseBuffer);
+			}
+
 			bDefaultOverriden = true;
 			bStaticOverriden = true;
 			DefaultProfileName = Profiles[DefaultProfileIndex].Name;
@@ -230,7 +284,7 @@ void FSkinWeightProfilesData::OverrideBaseBufferSkinWeightData(USkeletalMesh* Me
 }
 #endif 
 
-void FSkinWeightProfilesData::SetDynamicDefaultSkinWeightProfile(USkeletalMesh* Mesh, int32 LODIndex)
+void FSkinWeightProfilesData::SetDynamicDefaultSkinWeightProfile(USkeletalMesh* Mesh, int32 LODIndex, bool bSerialization /*= false*/)
 {
 	if (bStaticOverriden)
 	{
@@ -268,9 +322,79 @@ void FSkinWeightProfilesData::SetDynamicDefaultSkinWeightProfile(USkeletalMesh*
 			const bool bDifferentDefaultProfile = bNoDefaultProfile && (!bDefaultOverriden || DefaultProfileName != Profiles[DefaultProfileIndex].Name);
 			if (bNoDefaultProfile || bDifferentDefaultProfile)
 			{
-				DefaultOverrideSkinWeightBuffer = GetOverrideBuffer(Profiles[DefaultProfileIndex].Name);
+				if (GetOverrideBuffer(Profiles[DefaultProfileIndex].Name) == nullptr)
+				{
+					if (bSerialization)
+					{
+						// During serialization the CPU copy of the weight should still be available
+						const void* BaseBufferData = BaseBuffer->GetDataVertexBuffer()->GetWeightData();
+						
+						if (ensure(BaseBufferData))
+						{
+							FSkinWeightVertexBuffer* OverrideBuffer = nullptr;
+							
+							const FName& ProfileName = Profiles[DefaultProfileIndex].Name;
+							const FRuntimeSkinWeightProfileData* ProfilePtr = OverrideData.Find(ProfileName);
+							if (ProfilePtr)
+							{
+								FSkinnedMeshComponentUpdateSkinWeightsContext Context(Mesh);
+
+								OverrideBuffer = new FSkinWeightVertexBuffer();
+								OverrideBuffer->CopyMetaData(*BaseBuffer);
+								ProfileNameToBuffer.Add(ProfileName, OverrideBuffer);
+								
+								const uint32 NumInfluences = BaseBuffer->GetMaxBoneInfluences();
+								OverrideBuffer->SetMaxBoneInfluences(NumInfluences);
+								const bool bUse16BitBoneIndex = BaseBuffer->Use16BitBoneIndex();
+								OverrideBuffer->SetUse16BitBoneIndex(bUse16BitBoneIndex);
+								
+								ProfilePtr->ApplyOverrides(OverrideBuffer, BaseBufferData, BaseBuffer->GetNumVertices());
+
+								DefaultOverrideSkinWeightBuffer = OverrideBuffer;
+								bDefaultOverriden = true;
+								DefaultProfileName = Profiles[DefaultProfileIndex].Name;
+
+								OverrideBuffer->BeginInitResources();
+							}
+						}
+					}
+					else
+					{
+						FSkinWeightProfilesData* DataPtr = this;
+						FRequestFinished Callback = [DataPtr](TWeakObjectPtr<USkeletalMesh> WeakMesh, FName ProfileName)
+						{
+							if (WeakMesh.IsValid())
+							{
+								FSkinnedMeshComponentRecreateRenderStateContext RecreateState(WeakMesh.Get());
+								DataPtr->bDefaultOverriden = true;
+								DataPtr->DefaultProfileName = ProfileName;
+								DataPtr->SetupDynamicDefaultSkinweightProfile();								
+							}
+						};
+
+						UWorld* World = nullptr;
+						UGameEngine* GameEngine = Cast<UGameEngine>(GEngine);
+						if (GameEngine)
+						{
+							World = GameEngine->GetGameWorld();
+						}
+
+						if (World)
+						{
+							if (FSkinWeightProfileManager * Manager = FSkinWeightProfileManager::Get(World))
+							{
+								Manager->RequestSkinWeightProfile(Profiles[DefaultProfileIndex].Name, Mesh, Mesh, Callback, LODIndex);
+							}
+						}
+					}
+				}
+				else
+				{
 				bDefaultOverriden = true;
 				DefaultProfileName = Profiles[DefaultProfileIndex].Name;
+
+					SetupDynamicDefaultSkinweightProfile();
+				}
 			}
 		}
 	}
@@ -288,7 +412,7 @@ void FSkinWeightProfilesData::ClearDynamicDefaultSkinWeightProfile(USkeletalMesh
 	{
 		if (DefaultOverrideSkinWeightBuffer != nullptr)
 		{
-			ReleaseBuffer(DefaultProfileName);
+			ReleaseBuffer(DefaultProfileName, true);
 			DefaultOverrideSkinWeightBuffer = nullptr;			
 		}
 
@@ -297,11 +421,22 @@ void FSkinWeightProfilesData::ClearDynamicDefaultSkinWeightProfile(USkeletalMesh
 	}
 }
 
-FSkinWeightVertexBuffer* FSkinWeightProfilesData::GetOverrideBuffer(const FName& ProfileName)
+void FSkinWeightProfilesData::SetupDynamicDefaultSkinweightProfile()
 {
-	SCOPED_NAMED_EVENT(FSkinWeightProfilesData_GetOverrideBuffer, FColor::Red);
+	if (ProfileNameToBuffer.Contains(DefaultProfileName) && bDefaultOverriden && !bStaticOverriden)
+	{
+		DefaultOverrideSkinWeightBuffer = ProfileNameToBuffer.FindChecked(DefaultProfileName);
+	}
+}
 
-	LLM_SCOPE(ELLMTag::SkeletalMesh);
+bool FSkinWeightProfilesData::ContainsProfile(const FName& ProfileName) const
+{
+	return OverrideData.Contains(ProfileName);
+}
+
+FSkinWeightVertexBuffer* FSkinWeightProfilesData::GetOverrideBuffer(const FName& ProfileName) const
+{
+	SCOPED_NAMED_EVENT(FSkinWeightProfilesData_GetOverrideBuffer, FColor::Red);
 
 	// In case we have overridden the default skin weight buffer we do not need to create an override buffer, if it was statically overridden we cannot load any other profile
 	if (bDefaultOverriden && (ProfileName == DefaultProfileName || bStaticOverriden))
@@ -316,31 +451,24 @@ FSkinWeightVertexBuffer* FSkinWeightProfilesData::GetOverrideBuffer(const FName&
 
 	if (BaseBuffer)
 	{
+		check(BaseBuffer->GetNumVertices() > 0);
 		if (ProfileNameToBuffer.Contains(ProfileName))
 		{
-			return ProfileNameToBuffer.FindChecked(ProfileName);
+			FSkinWeightVertexBuffer* Buffer = ProfileNameToBuffer.FindChecked(ProfileName);
+			return Buffer;
 		}
-		else
-		{
-			if (OverrideData.Contains(ProfileName))
-			{
-				FSkinWeightVertexBuffer* OverrideBuffer = new FSkinWeightVertexBuffer();
-				ProfileNameToBuffer.Add(ProfileName, OverrideBuffer);
-				OverrideBuffer->SetNeedsCPUAccess(true);
-
-				ApplyOverrideProfile(OverrideBuffer, ProfileName);
-
-				INC_DWORD_STAT_BY(STAT_SkeletalMeshVertexMemory, OverrideBuffer->GetVertexDataSize());
-				OverrideBuffer->BeginInitResources();
-
-				return OverrideBuffer;
-			}
 		}
-	}
 
 	return nullptr;
 }
 
+
+bool FSkinWeightProfilesData::ContainsOverrideBuffer(const FName& ProfileName) const
+{
+	return ProfileNameToBuffer.Contains(ProfileName);
+}
+
+#if WITH_EDITOR
 const FRuntimeSkinWeightProfileData* FSkinWeightProfilesData::GetOverrideData(const FName& ProfileName) const
 {
 	return OverrideData.Find(ProfileName);
@@ -350,24 +478,11 @@ FRuntimeSkinWeightProfileData& FSkinWeightProfilesData::AddOverrideData(const FN
 {
 	return OverrideData.FindOrAdd(ProfileName);
 }
+#endif // WITH_EDITOR
 
-void FSkinWeightProfilesData::ApplyOverrideProfile(FSkinWeightVertexBuffer* OverrideBuffer, const FName& ProfileName)
-{
-	const uint32 NumInfluences = BaseBuffer->GetMaxBoneInfluences();
-	OverrideBuffer->SetMaxBoneInfluences(NumInfluences);
-	const bool bUse16BitBoneIndex = BaseBuffer->Use16BitBoneIndex();
-	OverrideBuffer->SetUse16BitBoneIndex(bUse16BitBoneIndex);
-
-	const FRuntimeSkinWeightProfileData* ProfilePtr = OverrideData.Find(ProfileName);
-	if (ProfilePtr)
-	{
-		ProfilePtr->ApplyOverrides(OverrideBuffer, BaseBuffer);
-	}
-}
-
-void FSkinWeightProfilesData::ReleaseBuffer(const FName& ProfileName)
+void FSkinWeightProfilesData::ReleaseBuffer(const FName& ProfileName, bool bForceRelease /*= false*/)
 {
-	if (ProfileNameToBuffer.Contains(ProfileName) && (!bDefaultOverriden || ProfileName != DefaultProfileName))
+	if (ProfileNameToBuffer.Contains(ProfileName) && (!bDefaultOverriden || ProfileName != DefaultProfileName || bForceRelease))
 	{
 		FSkinWeightVertexBuffer* Buffer = nullptr;
 		ProfileNameToBuffer.RemoveAndCopyValue(ProfileName, Buffer);
@@ -400,6 +515,8 @@ void FSkinWeightProfilesData::ReleaseResources()
 
 	Buffers.Remove(nullptr);
 
+	ResetGPUReadback();
+
 	if (Buffers.Num())
 	{
 		ENQUEUE_RENDER_COMMAND(ReleaseSkinSkinWeightProfilesDataBufferCommand)(
@@ -450,6 +567,8 @@ void FSkinWeightProfilesData::ReleaseCPUResources()
 	{
 		It->Value = FRuntimeSkinWeightProfileData();
 	}
+
+	ResetGPUReadback();
 }
 
 template <bool bRenderThread>
@@ -474,6 +593,139 @@ void FSkinWeightProfilesData::CreateRHIBuffers_Internal(TArray<TPair<FName, FSki
 	}
 }
 
+bool FSkinWeightProfilesData::IsPendingReadback() const
+{
+	return !ReadbackData.BufferReadback.IsValid();
+}
+
+void FSkinWeightProfilesData::EnqueueGPUReadback()
+{
+	ensure(!ReadbackData.BufferReadback.IsValid());
+
+	if (BaseBuffer && BaseBuffer->GetDataVertexBuffer()->IsWeightDataValid() && BaseBuffer->GetDataVertexBuffer()->GetVertexDataSize())
+	{
+		static const FName ReadbackName = TEXT("ReadbackSkinWeightBuffer");
+		FSkinweightReadbackData* DataPtr = &ReadbackData;
+		ENQUEUE_RENDER_COMMAND(FSkinWeightProfilesData_EnqueueGPUReadback)(
+			[DataPtr, this](FRHICommandListImmediate& RHICmdList)
+		{
+			DataPtr->BufferReadback.Reset(new FRHIGPUBufferReadback(ReadbackName));
+			if (DataPtr->BufferReadback.IsValid() && BaseBuffer->GetDataVertexBuffer()->VertexBufferRHI->GetSize())
+			{
+				DataPtr->BufferReadback->EnqueueCopy(RHICmdList, BaseBuffer->GetDataVertexBuffer()->VertexBufferRHI);
+			}
+		});
+	}
+}
+
+bool FSkinWeightProfilesData::IsGPUReadbackFinished() const
+{
+	return !IsPendingReadback() && ReadbackData.BufferReadback->IsReady();
+}
+
+void FSkinWeightProfilesData::EnqueueDataReadback()
+{
+	ensure(ReadbackData.ReadbackData.Num() == 0 && ReadbackData.BufferReadback->IsReady());
+	
+	if ( BaseBuffer )
+	{
+		ReadbackData.ReadbackData.SetNumZeroed(BaseBuffer->GetVertexDataSize());
+
+		FSkinweightReadbackData* DataPtr = &ReadbackData;
+		ENQUEUE_RENDER_COMMAND(FSkinWeightProfilesData_EnqueueDataReadback)(
+				[DataPtr](FRHICommandListImmediate& RHICmdList)
+		{
+			if (DataPtr && DataPtr->BufferReadback.IsValid())
+			{
+				ensure(DataPtr->BufferReadback->IsReady());
+				const void* BufferPtr = DataPtr->BufferReadback->Lock(DataPtr->ReadbackData.Num());
+				FMemory::Memcpy(DataPtr->ReadbackData.GetData(), BufferPtr, DataPtr->ReadbackData.Num());
+				DataPtr->BufferReadback->Unlock();
+
+				DataPtr->ReadbackFinishedFrameIndex = GFrameNumberRenderThread;
+			}
+		});
+	}
+}
+
+bool FSkinWeightProfilesData::IsDataReadbackPending() const
+{
+	return ReadbackData.ReadbackData.Num() > 0;
+}
+
+bool FSkinWeightProfilesData::IsDataReadbackFinished() const
+{
+	return !IsPendingReadback() && IsGPUReadbackFinished() && ReadbackData.ReadbackFinishedFrameIndex != INDEX_NONE && GFrameNumberRenderThread > ReadbackData.ReadbackFinishedFrameIndex;
+}
+
+void FSkinWeightProfilesData::ResetGPUReadback()
+{
+	ReadbackData.BufferReadback.Reset();
+	ReadbackData.ReadbackData.Empty();
+	ReadbackData.ReadbackFinishedFrameIndex = INDEX_NONE;
+}
+
+void FSkinWeightProfilesData::InitialiseProfileBuffer(const FName& ProfileName)
+{
+	LLM_SCOPE(ELLMTag::SkeletalMesh);
+
+	if (ProfileNameToBuffer.Contains(ProfileName))
+	{
+		// Maybe it has already been previously set-up and not unloaded yet
+		return;
+	}
+
+	if (BaseBuffer)
+	{
+		const void* BaseBufferData = nullptr;
+		if (BaseBuffer->GetNeedsCPUAccess())
+		{
+			BaseBufferData = BaseBuffer->GetDataVertexBuffer()->GetWeightData();
+		}
+		else
+		{
+			ensure(IsDataReadbackFinished());
+			BaseBufferData = ReadbackData.ReadbackData.GetData();
+		}
+		
+		if (ensure(BaseBufferData))
+		{
+			FSkinWeightVertexBuffer* OverrideBuffer = nullptr;
+
+			OverrideBuffer = new FSkinWeightVertexBuffer();
+			OverrideBuffer->CopyMetaData(*BaseBuffer);
+			ProfileNameToBuffer.Add(ProfileName, OverrideBuffer);
+
+			const uint32 NumInfluences = BaseBuffer->GetMaxBoneInfluences();
+			OverrideBuffer->SetMaxBoneInfluences(NumInfluences);
+			const bool bUse16BitBoneIndex = BaseBuffer->Use16BitBoneIndex();
+			OverrideBuffer->SetUse16BitBoneIndex(bUse16BitBoneIndex);
+			
+			const FRuntimeSkinWeightProfileData* ProfilePtr = OverrideData.Find(ProfileName);
+			if (ProfilePtr)
+			{
+				ProfilePtr->ApplyOverrides(OverrideBuffer, BaseBufferData, BaseBuffer->GetNumVertices());
+			}
+
+			OverrideBuffer->BeginInitResources();
+		}	
+	}
+}
+
+void FSkinWeightProfilesData::ApplyOverrideProfile(FSkinWeightVertexBuffer* OverrideBuffer, const FName& ProfileName)
+{
+	const uint32 NumInfluences = BaseBuffer->GetMaxBoneInfluences();
+	OverrideBuffer->SetMaxBoneInfluences(NumInfluences);
+	const bool bUse16BitBoneIndex = BaseBuffer->Use16BitBoneIndex();
+	OverrideBuffer->SetUse16BitBoneIndex(bUse16BitBoneIndex);
+
+	const FRuntimeSkinWeightProfileData* ProfilePtr = OverrideData.Find(ProfileName);
+	if (ProfilePtr)
+	{
+		ProfilePtr->ApplyOverrides(OverrideBuffer, BaseBuffer->GetDataVertexBuffer()->GetWeightData(), BaseBuffer->GetNumVertices());
+	}	
+}
+
 void FSkinWeightProfilesData::CreateRHIBuffers_RenderThread(TArray<TPair<FName, FSkinWeightRHIInfo>>& OutBuffers)
 {
 	CreateRHIBuffers_Internal<true>(OutBuffers);
@@ -484,3 +736,74 @@ void FSkinWeightProfilesData::CreateRHIBuffers_Async(TArray<TPair<FName, FSkinWe
 	CreateRHIBuffers_Internal<false>(OutBuffers);
 }
 
+void FRuntimeSkinWeightProfileData::ApplyOverrides(FSkinWeightVertexBuffer* OverrideBuffer, const void* DataBuffer, const int32 NumVerts) const
+{
+	if (DataBuffer)
+	{
+		if (OverrideBuffer)
+		{
+			const FSkinWeightInfo* BaseSkinWeights = (const FSkinWeightInfo*)DataBuffer;
+			OverrideBuffer->CopySkinWeightInfoData(TArrayView<const FSkinWeightInfo>(BaseSkinWeights, NumVerts));
+
+			uint8* TargetSkinWeightData = (uint8*)OverrideBuffer->GetDataVertexBuffer()->GetWeightData();
+			
+			const uint8 VertexStride = OverrideBuffer->GetConstantInfluencesVertexStride();
+			const uint8 WeightDataOffset = OverrideBuffer->GetBoneIndexByteSize() * OverrideBuffer->GetMaxBoneInfluences();
+
+			// Apply overrides
+			for (auto VertexIndexOverridePair : VertexIndexToInfluenceOffset)
+			{
+				const uint32 VertexIndex = VertexIndexOverridePair.Key;
+				const uint32 InfluenceOffset = VertexIndexOverridePair.Value;
+				
+				const uint8* BoneData = TargetSkinWeightData + (VertexIndex * VertexStride);
+				const uint8* WeightData = BoneData + WeightDataOffset;				
+
+#if !UE_BUILD_SHIPPING
+				uint32 VertexOffset = 0;
+				uint32 VertexInfluenceCount = 0;
+				OverrideBuffer->GetVertexInfluenceOffsetCount(VertexIndex, VertexOffset, VertexInfluenceCount);
+				check(NumWeightsPerVertex <= VertexInfluenceCount);
+				check((void*)(((uint8*)TargetSkinWeightData) + VertexOffset) == (void*)BoneData);
+#endif				
+				FMemory::Memcpy((void*)BoneData, &BoneIDs[InfluenceOffset * NumWeightsPerVertex], sizeof(FBoneIndexType) * NumWeightsPerVertex);
+				FMemory::Memcpy((void*)WeightData, &BoneWeights[InfluenceOffset * NumWeightsPerVertex], sizeof(uint8) * NumWeightsPerVertex);
+			}
+		}
+	}
+}
+
+void FRuntimeSkinWeightProfileData::ApplyDefaultOverride(FSkinWeightVertexBuffer* Buffer) const
+{
+	if (Buffer)
+	{
+		const int32 ExpectedNumVerts = Buffer->GetNumVertices();
+		if (ExpectedNumVerts)
+		{
+			uint8* TargetSkinWeightData = (uint8*)Buffer->GetDataVertexBuffer()->GetWeightData();
+
+			const uint8 VertexStride = Buffer->GetConstantInfluencesVertexStride();
+			const uint8 WeightDataOffset = Buffer->GetBoneIndexByteSize() * Buffer->GetMaxBoneInfluences();
+
+			for (auto VertexIndexOverridePair : VertexIndexToInfluenceOffset)
+			{
+				const uint32 VertexIndex = VertexIndexOverridePair.Key;
+				const uint32 InfluenceOffset = VertexIndexOverridePair.Value;
+
+				const uint8* BoneData = TargetSkinWeightData + (VertexIndex * VertexStride);
+				const uint8* WeightData = BoneData + WeightDataOffset;
+
+#if !UE_BUILD_SHIPPING
+				uint32 VertexOffset = 0;
+				uint32 VertexInfluenceCount = 0;
+				Buffer->GetVertexInfluenceOffsetCount(VertexIndex, VertexOffset, VertexInfluenceCount);
+				check(NumWeightsPerVertex <= VertexInfluenceCount);
+				check((void*)(((uint8*)TargetSkinWeightData) + VertexOffset) == (void*)BoneData);
+#endif
+				
+				FMemory::Memcpy((void*)BoneData, &BoneIDs[InfluenceOffset * NumWeightsPerVertex], sizeof(FBoneIndexType) * NumWeightsPerVertex);
+				FMemory::Memcpy((void*)WeightData, &BoneWeights[InfluenceOffset * NumWeightsPerVertex], sizeof(uint8) * NumWeightsPerVertex);
+			}
+		}
+	}
+}
diff --git a/Engine/Source/Runtime/Engine/Private/Animation/SkinWeightProfileManager.cpp b/Engine/Source/Runtime/Engine/Private/Animation/SkinWeightProfileManager.cpp
new file mode 100644
--- /dev/null
+++ b/Engine/Source/Runtime/Engine/Private/Animation/SkinWeightProfileManager.cpp
@@ -0,0 +1,372 @@
+// Copyright Epic Games, Inc. All Rights Reserved.
+
+#include "Animation/SkinWeightProfileManager.h"
+#include "Stats/Stats.h"
+#include "Animation/SkinWeightProfile.h"
+#include "Rendering/SkeletalMeshLODRenderData.h"
+#include "Rendering/SkinWeightVertexBuffer.h"
+#include "Engine/SkeletalMesh.h"
+#include "Rendering/SkeletalMeshRenderData.h"
+#include "RenderingThread.h"
+#include "RHICommandList.h"
+#include "UObject/WeakObjectPtrTemplates.h"
+#include "Algo/Transform.h"
+#include "Engine/Engine.h"
+#include "Engine/Canvas.h"
+#include "GameFramework/HUD.h"
+#include "Logging/LogMacros.h"
+#include "ProfilingDebugging/CsvProfiler.h"
+
+DEFINE_LOG_CATEGORY_STATIC(LogSkinWeightProfileManager, Warning, Display);
+
+#define PROFILE_SKIN_WEIGHT_MANAGER (!UE_BUILD_SHIPPING)
+
+static int32 GAllowCPU = 0;
+static FAutoConsoleVariableRef CVarAllowCPU (
+	TEXT("SkinWeightProfileManager.AllowCPU"),
+	GAllowCPU,
+	TEXT("Whether or not to allow cpu buffer generation"),
+	ECVF_Cheat
+);
+
+FSkinWeightProfileManager::FSkinWeightProfileManager(UWorld* InWorld) : LastGamethreadProfileIndex(INDEX_NONE), WeakWorld(nullptr)
+{
+	if (InWorld)
+	{
+		TickFunction.TickGroup = ETickingGroup::TG_PrePhysics;
+		TickFunction.EndTickGroup = ETickingGroup::TG_PostUpdateWork;
+		TickFunction.bCanEverTick = true;
+		TickFunction.bStartWithTickEnabled = true;
+		TickFunction.bTickEvenWhenPaused = false;
+		TickFunction.Owner = this;
+		TickFunction.RegisterTickFunction(InWorld->PersistentLevel);
+
+		WeakWorld = InWorld;
+	}	
+}
+
+TMap<UWorld*, FSkinWeightProfileManager*> FSkinWeightProfileManager::WorldManagers;
+
+void FSkinWeightProfileManager::OnWorldInit(UWorld* World, const UWorld::InitializationValues IVS)
+{
+	check(WorldManagers.Find(World) == nullptr);
+
+	FSkinWeightProfileManager* NewWorldMan = new FSkinWeightProfileManager(World);
+	WorldManagers.Add(World) = NewWorldMan;
+}
+
+void FSkinWeightProfileManager::OnWorldCleanup(UWorld* World, bool bSessionEnded, bool bCleanupResources)
+{
+	if (FSkinWeightProfileManager* *WorldMan = WorldManagers.Find(World))
+	{
+		delete (*WorldMan);
+		WorldManagers.Remove(World);
+	}
+}
+
+void FSkinWeightProfileManager::OnPreWorldFinishDestroy(UWorld* World)
+{
+	if (FSkinWeightProfileManager * *WorldMan = WorldManagers.Find(World))
+	{
+		delete (*WorldMan);
+		WorldManagers.Remove(World);
+	}
+}
+
+void FSkinWeightProfileManager::OnWorldBeginTearDown(UWorld* World)
+{
+	if (FSkinWeightProfileManager * *WorldMan = WorldManagers.Find(World))
+	{
+		delete (*WorldMan);
+		WorldManagers.Remove(World);
+	}
+}
+
+FSkinWeightProfileManager* FSkinWeightProfileManager::Get(UWorld* World)
+{
+	if (FSkinWeightProfileManager * *WorldMan = WorldManagers.Find(World))
+	{
+		return *WorldMan;
+	}
+
+	return nullptr;
+}
+
+void FSkinWeightProfileManager::OnStartup()
+{
+	FWorldDelegates::OnPreWorldInitialization.AddStatic(&OnWorldInit);
+	FWorldDelegates::OnPostWorldCleanup.AddStatic(&OnWorldCleanup);
+	FWorldDelegates::OnPreWorldFinishDestroy.AddStatic(&OnPreWorldFinishDestroy);
+	FWorldDelegates::OnWorldBeginTearDown.AddStatic(&OnWorldBeginTearDown);
+}
+
+void FSkinWeightProfileManager::OnShutdown()
+{
+	for (TPair<UWorld*, FSkinWeightProfileManager*>& Pair : WorldManagers)
+	{
+		if (Pair.Value)
+		{
+			delete Pair.Value;
+		}
+	}
+	WorldManagers.Empty();
+}
+
+void FSkinWeightProfileManager::RequestSkinWeightProfile(FName InProfileName, USkeletalMesh* Mesh, UObject* Requester, FRequestFinished& Callback, int32 LODIndex /*= INDEX_NONE*/)
+{
+	// Make sure we have an actual skeletal mesh
+ 	if (Mesh)
+	{
+		// Setup a request structure
+		FSetProfileRequest ProfileRequest;
+		ProfileRequest.ProfileName = InProfileName;
+		ProfileRequest.Callback = Callback;
+		ProfileRequest.WeakSkeletalMesh = Mesh;
+		ProfileRequest.IdentifyingObject = Requester;
+
+		UE_LOG(LogSkinWeightProfileManager, Display, TEXT("RequestSkinWeightProfile [%s | %s]"), *Mesh->GetName(), *Requester->GetName());
+
+		if (LODIndex == INDEX_NONE)
+		{
+			const int32 NumLODS = Mesh->GetLODNum();
+			for (int32 index = 0; index < NumLODS; ++index)
+			{
+				ProfileRequest.LODIndices.Add(index);
+			}
+		}
+		else
+		{
+			ProfileRequest.LODIndices.Add(LODIndex);
+		}
+		
+		// Add the profile request
+		const int32 Index = PendingSetProfileRequests.Add(ProfileRequest);
+
+		// Check whether or not we might have already setup a readback for this skeletal mesh, and add a request to it
+		int32& NumRequestsForMesh = PendingMeshes.FindOrAdd(Mesh, 0);
+		++NumRequestsForMesh;
+
+		TickFunction.SetTickFunctionEnable(true);
+	}
+}
+
+void FSkinWeightProfileManager::CancelSkinWeightProfileRequest(UObject* Requester)
+{
+	QUICK_SCOPE_CYCLE_COUNTER(STAT_SkinWeightProfileManager_CancelSkinWeightProfileRequest);
+
+	for (FSetProfileRequest& Request : PendingSetProfileRequests)
+	{
+		if (Request.IdentifyingObject == Requester || Request.WeakSkeletalMesh == Requester)
+		{
+			CanceledRequest.Add(Request);
+			UE_LOG(LogSkinWeightProfileManager, Display, TEXT("Cancel [%s]"), *Requester->GetName());			
+		}
+	}
+}
+
+void FSkinWeightProfileManager::DoTick(float DeltaTime, ENamedThreads::Type CurrentThread, const FGraphEventRef& MyCompletionGraphEvent)
+{	
+#if PROFILE_SKIN_WEIGHT_MANAGER 
+	CSV_SCOPED_TIMING_STAT_EXCLUSIVE(SkinWeightProfileManagerGT);	
+	CSV_CUSTOM_STAT_GLOBAL(SkinWeightProfileManager_PendingMeshes, PendingMeshes.Num(), ECsvCustomStatOp::Set);
+	CSV_CUSTOM_STAT_GLOBAL(SkinWeightProfileManager_PendingRequests, PendingSetProfileRequests.Num(), ECsvCustomStatOp::Set);
+#endif // PROFILE_SKIN_WEIGHT_MANAGER 
+
+	QUICK_SCOPE_CYCLE_COUNTER(STAT_SkinWeightProfileManager_Tick);
+
+	LastGamethreadProfileIndex = INDEX_NONE;
+
+	if (PendingSetProfileRequests.Num())
+	{
+		if (AsyncTask.IsValid())
+		{
+			FTaskGraphInterface::Get().WaitUntilTaskCompletes(AsyncTask, ENamedThreads::GameThread);
+			AsyncTask.SafeRelease();
+		}
+
+		for (int32 RequestIndex = 0; RequestIndex < PendingSetProfileRequests.Num(); ++RequestIndex)
+		{
+			QUICK_SCOPE_CYCLE_COUNTER(STAT_SkinWeightProfileManager_TickMesh);
+
+			const FSetProfileRequest& Request = PendingSetProfileRequests[RequestIndex];
+			
+			// Remove if canceled or objects have gone stale 
+			bool bRemove = Request.IdentifyingObject.IsStale() || Request.WeakSkeletalMesh.IsStale() ||
+				CanceledRequest.ContainsByPredicate([Request](FSetProfileRequest& B)
+			{
+				return Request.IdentifyingObject == B.IdentifyingObject && Request.WeakSkeletalMesh == B.WeakSkeletalMesh && Request.ProfileName == B.ProfileName;
+			});
+
+			const USkeletalMesh* SkeletalMesh = Request.WeakSkeletalMesh.Get();
+			FSkeletalMeshRenderData* RenderData = SkeletalMesh ? SkeletalMesh->GetResourceForRendering() : nullptr;
+			// Or if skeletal mesh / render data is invalid
+			bRemove = bRemove ? true : (RenderData == nullptr);
+				
+			if (!bRemove)
+			{	
+				bool bAllBuffersReady = true;
+				for (const int32 LODIndex : Request.LODIndices)
+				{
+					if (!RenderData->LODRenderData[LODIndex].SkinWeightProfilesData.ContainsOverrideBuffer(Request.ProfileName))
+					{
+						bAllBuffersReady = false;
+						break;
+					}
+				}
+
+				if (bAllBuffersReady)
+				{
+					if (Request.IdentifyingObject.IsValid())
+					{
+						UE_LOG(LogSkinWeightProfileManager, Display, TEXT("Callback [%s | %s]"), *SkeletalMesh->GetName(), *Request.IdentifyingObject->GetName());
+						Request.Callback(Request.WeakSkeletalMesh, Request.ProfileName);
+						bRemove = true;
+					}
+				}
+			}
+			
+			if (bRemove)
+			{
+				CleanupRequest(Request);
+				CanceledRequest.Remove(Request);
+				PendingSetProfileRequests.RemoveAtSwap(RequestIndex);
+				--RequestIndex;
+			}
+		}
+
+		// This is the last index that we will process asynchronously, while the task is running we do not remove request from the array so access should be safe
+		LastGamethreadProfileIndex = PendingSetProfileRequests.Num() - 1;
+
+		if (MyCompletionGraphEvent)
+		{
+			AsyncTask = TGraphTask<FSkinWeightProfileManagerAsyncTask>::CreateTask(nullptr, CurrentThread).ConstructAndDispatchWhenReady(this);
+			TickFunction.GetCompletionHandle()->DontCompleteUntil(AsyncTask);
+		}
+	}
+	else
+	{
+		TickFunction.SetTickFunctionEnable(false);
+	}
+}
+
+void FSkinWeightProfileManager::Tick(float DeltaTime)
+{
+	DoTick(DeltaTime, ENamedThreads::GameThread, nullptr);
+}
+
+bool FSkinWeightProfileManager::IsTickable() const
+{
+	return WeakWorld.IsExplicitlyNull();
+}
+
+TStatId FSkinWeightProfileManager::GetStatId() const
+{
+	RETURN_QUICK_DECLARE_CYCLE_STAT(FSkinWeightProfileManager, STATGROUP_Tickables);
+}
+
+void FSkinWeightProfileManager::CleanupRequest(const FSetProfileRequest& Request)
+{
+	QUICK_SCOPE_CYCLE_COUNTER(STAT_SkinWeightProfileManager_CleanupRequest);
+	// Remove request for its pending skeletal mesh
+	if (int32* RequestsPtr = PendingMeshes.Find(Request.WeakSkeletalMesh))
+	{
+		int32& NumPendingRequests = *RequestsPtr;
+		--NumPendingRequests;
+
+		// In case all requests have finished clean up the readback data and remove skeletal mesh as pending
+		if (NumPendingRequests == 0)
+		{
+			PendingMeshes.Remove(Request.WeakSkeletalMesh);
+
+			if (USkeletalMesh * SkeletalMesh = Request.WeakSkeletalMesh.Get())
+			{
+				if (FSkeletalMeshRenderData * RenderData = SkeletalMesh->GetResourceForRendering())
+				{
+					for (FSkeletalMeshLODRenderData& LODRenderData : RenderData->LODRenderData)
+					{
+						FSkinWeightProfilesData& SkinweightData = LODRenderData.SkinWeightProfilesData;
+						SkinweightData.ResetGPUReadback();
+					}
+				}
+			}
+		}
+	}
+	
+}
+
+bool FSkinWeightProfileManager::IsTickableWhenPaused() const
+{
+	return false;
+}
+
+bool FSkinWeightProfileManager::IsTickableInEditor() const
+{
+	return true;
+}
+
+void FSkinWeightProfileManagerTickFunction::ExecuteTick(float DeltaTime, ELevelTick TickType, ENamedThreads::Type CurrentThread, const FGraphEventRef& MyCompletionGraphEvent)
+{
+	Owner->DoTick(DeltaTime, CurrentThread, MyCompletionGraphEvent);
+}
+
+FString FSkinWeightProfileManagerTickFunction::DiagnosticMessage()
+{
+	return TEXT("FSkinWeightProfileManagerTickFunction::Tick");
+}
+
+FName FSkinWeightProfileManagerTickFunction::DiagnosticContext(bool bDetailed)
+{
+	return FName("FSkinWeightProfileManager");
+}
+
+void FSkinWeightProfileManagerAsyncTask::DoTask(ENamedThreads::Type CurrentThread, const FGraphEventRef& MyCompletionGraphEvent)
+{
+#if PROFILE_SKIN_WEIGHT_MANAGER 
+	CSV_SCOPED_TIMING_STAT_GLOBAL(SkinWeightProfileManagerAsync);
+#endif // PROFILE_SKIN_WEIGHT_MANAGER 
+	
+	const TArray<FSetProfileRequest>& OwnerRequests = Owner->PendingSetProfileRequests;
+
+	checkf(Owner->LastGamethreadProfileIndex < OwnerRequests.Num(), TEXT("Invalid state"));
+
+	for (int32 RequestIndex = 0; RequestIndex <= Owner->LastGamethreadProfileIndex; ++RequestIndex)
+	{
+		QUICK_SCOPE_CYCLE_COUNTER(STAT_SkinWeightProfileManager_TickMeshAsync);
+
+		const FSetProfileRequest& Request = OwnerRequests[RequestIndex];
+		if (const USkeletalMesh* SkeletalMesh = Request.WeakSkeletalMesh.Get())
+		{
+			if (FSkeletalMeshRenderData* RenderData = SkeletalMesh->GetResourceForRendering())
+			{
+				for (const int32 LODIndex : Request.LODIndices)
+				{
+					FSkeletalMeshLODRenderData& LODRenderData = RenderData->LODRenderData[LODIndex];
+					FSkinWeightProfilesData& SkinweightData = LODRenderData.SkinWeightProfilesData;
+					
+					if (!!GAllowCPU && LODRenderData.SkinWeightVertexBuffer.GetNeedsCPUAccess())
+					{
+						SkinweightData.InitialiseProfileBuffer(Request.ProfileName);
+					}
+					else
+					{
+						if (SkinweightData.IsPendingReadback())
+						{
+							SkinweightData.EnqueueGPUReadback();
+						}
+						else if (SkinweightData.IsGPUReadbackFinished())
+						{
+							if (SkinweightData.IsGPUReadbackFinished() && !SkinweightData.IsDataReadbackPending())
+							{						
+								SkinweightData.EnqueueDataReadback();
+							}
+							else if (SkinweightData.IsDataReadbackFinished())
+							{						
+								SkinweightData.InitialiseProfileBuffer(Request.ProfileName);
+							}
+						}
+					}
+				}
+			}
+		}
+	}
+}
diff --git a/Engine/Source/Runtime/Engine/Private/Components/SkinnedMeshComponent.cpp b/Engine/Source/Runtime/Engine/Private/Components/SkinnedMeshComponent.cpp
--- a/Engine/Source/Runtime/Engine/Private/Components/SkinnedMeshComponent.cpp
+++ b/Engine/Source/Runtime/Engine/Private/Components/SkinnedMeshComponent.cpp
@@ -25,6 +25,7 @@
 #include "SkeletalMeshTypes.h"
 #include "Animation/MorphTarget.h"
 #include "AnimationRuntime.h"
+#include "Animation/SkinWeightProfileManager.h"
 
 DEFINE_LOG_CATEGORY_STATIC(LogSkinnedMeshComp, Log, All);
 
@@ -3178,6 +3179,15 @@ void USkinnedMeshComponent::BeginDestroy()
 	Super::BeginDestroy();
 	ReleaseResources();
 
+	if (bSkinWeightProfilePending)
+	{
+		bSkinWeightProfilePending = false;
+		if (FSkinWeightProfileManager * Manager = FSkinWeightProfileManager::Get(GetWorld()))
+		{
+			Manager->CancelSkinWeightProfileRequest(this);
+		}
+	}
+
 	// Release ref pose override if allocated
 	if (RefPoseOverride)
 	{
@@ -3563,37 +3573,90 @@ void USkinnedMeshComponent::ClearSkinWeightOverride(int32 LODIndex)
 
 bool USkinnedMeshComponent::SetSkinWeightProfile(FName InProfileName)
 {
+	bool bContainsProfile = false;
+
 	if (FSkeletalMeshRenderData* SkelMeshRenderData = GetSkeletalMeshRenderData())
 	{
 		// Ensure the LOD infos array is initialized
 		InitLODInfos();
-        for (int32 LODIndex = 0; LODIndex < LODInfo.Num(); ++LODIndex)
+		for (int32 LODIndex = 0; LODIndex < LODInfo.Num(); ++LODIndex)
         {
 			// Check whether or not setting a profile is allow for this LOD index
 			if (LODIndex > GSkinWeightProfilesAllowedFromLOD)
 			{
 				FSkeletalMeshLODRenderData& RenderData = SkelMeshRenderData->LODRenderData[LODIndex];
-				
+
+				bContainsProfile |= RenderData.SkinWeightProfilesData.ContainsProfile(InProfileName);
+
 				// Retrieve this profile's skin weight buffer
 				FSkinWeightVertexBuffer* Buffer = RenderData.SkinWeightProfilesData.GetOverrideBuffer(InProfileName);
+        
+				FSkelMeshComponentLODInfo& Info = LODInfo[LODIndex];
+				Info.OverrideProfileSkinWeights = Buffer;
+                
 				if (Buffer != nullptr)
 				{
-					FSkelMeshComponentLODInfo& Info = LODInfo[LODIndex];
-					Info.OverrideProfileSkinWeights = Buffer;
-
 					bSkinWeightProfileSet = true;
-					CurrentSkinWeightProfileName = InProfileName;
 				}
 			}
         }
 
-		if (bSkinWeightProfileSet)
+		if (bContainsProfile)
 		{
-			UpdateSkinWeightOverrideBuffer();
+			CurrentSkinWeightProfileName = InProfileName;
+
+			if (bSkinWeightProfileSet)
+			{
+				UpdateSkinWeightOverrideBuffer();
+			}
+			else 
+			{
+				TWeakObjectPtr<USkinnedMeshComponent> WeakComponent = this;
+				FRequestFinished Callback = [WeakComponent](TWeakObjectPtr<USkeletalMesh> WeakMesh, FName ProfileName)
+				{
+					// Ensure that the request objects are still valid
+					if (WeakMesh.IsValid() && WeakComponent.IsValid())
+					{
+						USkinnedMeshComponent* Component = WeakComponent.Get();
+						Component->InitLODInfos();
+
+						Component->bSkinWeightProfilePending = false;
+						Component->bSkinWeightProfileSet = true;
+
+						if (FSkeletalMeshRenderData * RenderData = WeakMesh->GetResourceForRendering())
+						{
+							const int32 NumLODs = RenderData->LODRenderData.Num();
+							for (int32 Index = 0; Index < NumLODs; ++Index)
+							{
+								FSkeletalMeshLODRenderData& LODRenderData = RenderData->LODRenderData[Index];
+								FSkinWeightProfilesData& SkinweightData = LODRenderData.SkinWeightProfilesData;
+
+								// Check whether or not setting a profile is allow for this LOD index
+								if (Index > GSkinWeightProfilesAllowedFromLOD)
+								{
+									// Retrieve this profile's skin weight buffer
+									FSkinWeightVertexBuffer* Buffer = SkinweightData.GetOverrideBuffer(ProfileName);
+									FSkelMeshComponentLODInfo& Info = Component->LODInfo[Index];
+									Info.OverrideProfileSkinWeights = Buffer;
+								}
+							}
+
+							Component->UpdateSkinWeightOverrideBuffer();
+						}
+					}
+				};
+
+				// Put in a skin weight profile request
+				if (FSkinWeightProfileManager* Manager = FSkinWeightProfileManager::Get(GetWorld()))
+				{
+					Manager->RequestSkinWeightProfile(InProfileName, SkeletalMesh, this, Callback);
+					bSkinWeightProfilePending = true;
+				}
+			}
 		}
 	}
 
-	return bSkinWeightProfileSet;
+	return bContainsProfile;
 }
 
 void USkinnedMeshComponent::ClearSkinWeightProfile()
@@ -3602,20 +3665,33 @@ void USkinnedMeshComponent::ClearSkinWeightProfile()
 	{	
 		bool bCleared = false;
 
-		// Clear skin weight buffer set for all of the LODs
-		for (int32 LODIndex = 0; LODIndex < LODInfo.Num(); ++LODIndex)
+		if (bSkinWeightProfileSet)
 		{
-			FSkelMeshComponentLODInfo& Info = LODInfo[LODIndex];
-			bCleared |= (Info.OverrideProfileSkinWeights != nullptr);
-			Info.OverrideProfileSkinWeights = nullptr;
+			InitLODInfos();
+			// Clear skin weight buffer set for all of the LODs
+			for (int32 LODIndex = 0; LODIndex < LODInfo.Num(); ++LODIndex)
+			{
+				FSkelMeshComponentLODInfo& Info = LODInfo[LODIndex];
+				bCleared |= (Info.OverrideProfileSkinWeights != nullptr);
+				Info.OverrideProfileSkinWeights = nullptr;
+			}
+
+			if (bCleared)
+			{
+				UpdateSkinWeightOverrideBuffer();
+			}
 		}
 
-		if (bCleared)
+		if (bSkinWeightProfilePending)
 		{
-			UpdateSkinWeightOverrideBuffer();
+			if (FSkinWeightProfileManager * Manager = FSkinWeightProfileManager::Get(GetWorld()))
+			{
+				Manager->CancelSkinWeightProfileRequest(this);
+			}
 		}
 	}
 
+	bSkinWeightProfilePending = false;
 	bSkinWeightProfileSet = false;
 	CurrentSkinWeightProfileName = NAME_None;
 }
@@ -3647,6 +3723,16 @@ void USkinnedMeshComponent::UnloadSkinWeightProfile(FName InProfileName)
 				UpdateSkinWeightOverrideBuffer();
 			}
 		}
+
+		if (bSkinWeightProfilePending)
+		{
+			if (FSkinWeightProfileManager * Manager = FSkinWeightProfileManager::Get(GetWorld()))
+			{
+				Manager->CancelSkinWeightProfileRequest(this);
+			}
+
+			bSkinWeightProfilePending = false;
+		}
 	}
 
 	if (CurrentSkinWeightProfileName == InProfileName)
diff --git a/Engine/Source/Runtime/Engine/Private/LevelTick.cpp b/Engine/Source/Runtime/Engine/Private/LevelTick.cpp
--- a/Engine/Source/Runtime/Engine/Private/LevelTick.cpp
+++ b/Engine/Source/Runtime/Engine/Private/LevelTick.cpp
@@ -896,7 +896,7 @@ void UWorld::MarkActorComponentForNeededEndOfFrameUpdate(UActorComponent* Compon
 		if (!bForceGameThread)
 		{
 			bool bAllowConcurrentUpdates = FApp::ShouldUseThreadingForPerformance() && 
-				(GIsEditor ? !!CVarAllowAsyncRenderThreadUpdatesEditor.GetValueOnGameThread() : !!CVarAllowAsyncRenderThreadUpdates.GetValueOnGameThread());
+				(GIsEditor ? !!CVarAllowAsyncRenderThreadUpdatesEditor.GetValueOnAnyThread() : !!CVarAllowAsyncRenderThreadUpdates.GetValueOnAnyThread());
 			bForceGameThread = !bAllowConcurrentUpdates 
 								// When there is no rendering thread force all updates on game thread,
 								// to avoid modifying scene structures from multiple task threads
diff --git a/Engine/Source/Runtime/Engine/Private/Rendering/SkinWeightVertexBuffer.cpp b/Engine/Source/Runtime/Engine/Private/Rendering/SkinWeightVertexBuffer.cpp
--- a/Engine/Source/Runtime/Engine/Private/Rendering/SkinWeightVertexBuffer.cpp
+++ b/Engine/Source/Runtime/Engine/Private/Rendering/SkinWeightVertexBuffer.cpp
@@ -184,7 +184,7 @@ FVertexBufferRHIRef FSkinWeightLookupVertexBuffer::CreateRHIBuffer_Internal()
 		// Create the vertex buffer.
 		FResourceArrayInterface* ResourceArray = LookupData ? LookupData->GetResourceArray() : nullptr;
 		const uint32 SizeInBytes = ResourceArray ? ResourceArray->GetResourceDataSize() : 0;
-		const uint32 BuffFlags = BUF_Static | BUF_ShaderResource;
+		const uint32 BuffFlags = BUF_Static | BUF_ShaderResource | BUF_SourceCopy;
 		FRHIResourceCreateInfo CreateInfo(ResourceArray);
 		CreateInfo.bWithoutNativeResource = !LookupData;
 
@@ -394,7 +394,7 @@ FVertexBufferRHIRef FSkinWeightDataVertexBuffer::CreateRHIBuffer_Internal()
 		// Create the vertex buffer.
 		FResourceArrayInterface* ResourceArray = WeightData ? WeightData->GetResourceArray() : nullptr;
 		const uint32 SizeInBytes = ResourceArray ? ResourceArray->GetResourceDataSize() : 0;
-		const uint32 BuffFlags = BUF_Static | BUF_ShaderResource;
+		const uint32 BuffFlags = BUF_Static | BUF_ShaderResource | BUF_SourceCopy;
 		FRHIResourceCreateInfo CreateInfo(ResourceArray);
 		CreateInfo.bWithoutNativeResource = !WeightData;
 
@@ -421,6 +421,11 @@ FVertexBufferRHIRef FSkinWeightDataVertexBuffer::CreateRHIBuffer_Async()
 	return CreateRHIBuffer_Internal<false>();
 }
 
+bool FSkinWeightDataVertexBuffer::IsWeightDataValid() const
+{
+	return WeightData != nullptr;
+}
+
 void FSkinWeightDataVertexBuffer::InitRHI()
 {
 	SCOPED_LOADTIMER(FSkinWeightVertexBuffer_InitRHI);
@@ -590,6 +595,11 @@ void FSkinWeightDataVertexBuffer::ResetVertexBoneWeights(uint32 VertexWeightOffs
 	}
 }
 
+void FSkinWeightDataVertexBuffer::CopyDataFromBuffer(const TArrayView<const FSkinWeightInfo>& SkinWeightData)
+{
+	Init(SkinWeightData.Num() * GetMaxBoneInfluences(), SkinWeightData.Num());
+	FMemory::Memcpy(Data, SkinWeightData.GetData(), GetVertexDataSize());
+}
 
 /*-----------------------------------------------------------------------------
 FSkinWeightVertexBuffer
@@ -903,4 +913,9 @@ FSkinWeightInfo FSkinWeightVertexBuffer::GetVertexSkinWeights(uint32 VertexIndex
 		OutVertex.InfluenceWeights[InfluenceIdx] = GetBoneWeight(VertexIndex, InfluenceIdx);
 	}
 	return OutVertex;
-}
\ No newline at end of file
+}
+
+void FSkinWeightVertexBuffer::CopySkinWeightInfoData(const TArrayView<const FSkinWeightInfo>& SkinWeightData)
+{
+	DataVertexBuffer.CopyDataFromBuffer(SkinWeightData);
+}
diff --git a/Engine/Source/Runtime/Engine/Private/SkeletalMesh.cpp b/Engine/Source/Runtime/Engine/Private/SkeletalMesh.cpp
--- a/Engine/Source/Runtime/Engine/Private/SkeletalMesh.cpp
+++ b/Engine/Source/Runtime/Engine/Private/SkeletalMesh.cpp
@@ -98,6 +98,7 @@
 #if RHI_RAYTRACING
 #include "RayTracingInstance.h"
 #endif
+#include "Animation/SkinWeightProfileManager.h"
 
 #define LOCTEXT_NAMESPACE "SkeltalMesh"
 
@@ -1439,6 +1440,11 @@ void USkeletalMesh::BeginDestroy()
 	// Safely unlink mesh from list of streamable ones.
 	UnlinkStreaming();
 
+	if (FSkinWeightProfileManager * Manager = FSkinWeightProfileManager::Get(GetWorld()))
+	{
+		Manager->CancelSkinWeightProfileRequest(this);
+	}
+
 	// remove the cache of link up
 	if (Skeleton)
 	{
diff --git a/Engine/Source/Runtime/Engine/Private/SkeletalMeshLODRenderData.cpp b/Engine/Source/Runtime/Engine/Private/SkeletalMeshLODRenderData.cpp
--- a/Engine/Source/Runtime/Engine/Private/SkeletalMeshLODRenderData.cpp
+++ b/Engine/Source/Runtime/Engine/Private/SkeletalMeshLODRenderData.cpp
@@ -192,10 +192,6 @@ void FSkeletalMeshLODRenderData::InitResources(bool bNeedsVertexColors, int32 LO
 	BeginInitResource(&StaticVertexBuffers.PositionVertexBuffer);
 	BeginInitResource(&StaticVertexBuffers.StaticMeshVertexBuffer);
 
-	if (GSkinWeightProfilesLoadByDefaultMode == 3)
-	{
-		SkinWeightProfilesData.SetDynamicDefaultSkinWeightProfile(Owner, LODIndex);
-	}
 	SkinWeightVertexBuffer.BeginInitResources();
 
 	if (bNeedsVertexColors)
@@ -891,12 +887,20 @@ void FSkeletalMeshLODRenderData::SerializeStreamedData(FArchive& Ar, USkeletalMe
 	Ar << SkinWeightProfilesData;
 	SkinWeightProfilesData.Init(&SkinWeightVertexBuffer);
 
-	if (Ar.IsLoading() && GSkinWeightProfilesLoadByDefaultMode == 1)
+	if (Ar.IsLoading())
 	{
 #if !WITH_EDITOR
-		// Only allow overriding the base buffer in non-editor builds as it could otherwise be serialized into the asset
-		SkinWeightProfilesData.OverrideBaseBufferSkinWeightData(Owner, LODIdx);
+		if (GSkinWeightProfilesLoadByDefaultMode == 1)
+		{
+			// Only allow overriding the base buffer in non-editor builds as it could otherwise be serialized into the asset
+			SkinWeightProfilesData.OverrideBaseBufferSkinWeightData(Owner, LODIdx);
+		} else
+	
 #endif
+		if (GSkinWeightProfilesLoadByDefaultMode == 3)
+		{
+			SkinWeightProfilesData.SetDynamicDefaultSkinWeightProfile(Owner, LODIdx, true);
+		}
 	}
 }
 
@@ -921,14 +925,6 @@ void FSkeletalMeshLODRenderData::SerializeAvailabilityInfo(FArchive& Ar, USkelet
 	}
 	SkinWeightProfilesData.SerializeMetaData(Ar);
 	SkinWeightProfilesData.Init(&SkinWeightVertexBuffer);
-
-	if (Ar.IsLoading() && GSkinWeightProfilesLoadByDefaultMode == 1)
-	{
-#if !WITH_EDITOR
-		// Only allow overriding the base buffer in non-editor builds as it could otherwise be serialized into the asset
-		SkinWeightProfilesData.OverrideBaseBufferSkinWeightData(Owner, LODIdx);
-#endif
-	}
 }
 
 void FSkeletalMeshLODRenderData::Serialize(FArchive& Ar, UObject* Owner, int32 Idx)
diff --git a/Engine/Source/Runtime/Engine/Private/SkeletalMeshRenderData.cpp b/Engine/Source/Runtime/Engine/Private/SkeletalMeshRenderData.cpp
--- a/Engine/Source/Runtime/Engine/Private/SkeletalMeshRenderData.cpp
+++ b/Engine/Source/Runtime/Engine/Private/SkeletalMeshRenderData.cpp
@@ -59,7 +59,7 @@ static void SerializeLODInfoForDDC(USkeletalMesh* SkeletalMesh, FString& KeySuff
 // differences, etc.) replace the version GUID below with a new one.
 // In case of merge conflicts with DDC versions, you MUST generate a new GUID
 // and set this new GUID as the version.
-#define SKELETALMESH_DERIVEDDATA_VER TEXT("D63A7B68FF6E4919AE164CD429AD51D7")
+#define SKELETALMESH_DERIVEDDATA_VER TEXT("D6E4934179CC45EFA985EE626BF24636")
 
 static const FString& GetSkeletalMeshDerivedDataVersion()
 {
diff --git a/Engine/Source/Runtime/Engine/Private/UnrealEngine.cpp b/Engine/Source/Runtime/Engine/Private/UnrealEngine.cpp
--- a/Engine/Source/Runtime/Engine/Private/UnrealEngine.cpp
+++ b/Engine/Source/Runtime/Engine/Private/UnrealEngine.cpp
@@ -234,6 +234,7 @@ UnrealEngine.cpp: Implements the UEngine class and helpers.
 #include "ObjectTrace.h"
 #include "StudioAnalytics.h"
 #include "TraceFilter.h"
+#include "Animation/SkinWeightProfileManager.h"
 
 DEFINE_LOG_CATEGORY(LogEngine);
 IMPLEMENT_MODULE( FEngineModule, Engine );
@@ -278,6 +279,8 @@ void FEngineModule::StartupModule()
 #if TRACE_FILTERING_ENABLED
 	FTraceFilter::Init();
 #endif
+
+	FSkinWeightProfileManager::OnStartup();
 }
 
 void FEngineModule::ShutdownModule()
@@ -291,6 +294,8 @@ void FEngineModule::ShutdownModule()
 #endif
 
 	FParticleSystemWorldManager::OnShutdown();
+
+	FSkinWeightProfileManager::OnShutdown();
 }
 
 /* Global variables
@@ -2201,7 +2206,7 @@ void UEngine::ReinitializeCustomTimeStep()
 	{
 		if (bIsCurrentCustomTimeStepInitialized)
 		{
-			CustomTimeStep->Shutdown(this);
+		CustomTimeStep->Shutdown(this);
 		}
 		bIsCurrentCustomTimeStepInitialized = CustomTimeStep->Initialize(this);
 	}
@@ -2220,7 +2225,7 @@ bool UEngine::SetCustomTimeStep(UEngineCustomTimeStep* InCustomTimeStep)
 		CustomTimeStep = InCustomTimeStep && !InCustomTimeStep->IsPendingKill() ? InCustomTimeStep : nullptr;
 
 		if (CustomTimeStep)
-		{
+			{
 			bIsCurrentCustomTimeStepInitialized = CustomTimeStep->Initialize(this);
 		}
 		OnCustomTimeStepChanged().Broadcast();
@@ -2235,7 +2240,7 @@ void UEngine::ReinitializeTimecodeProvider()
 	{
 		if (bIsCurrentTimecodeProviderInitialized)
 		{
-			Provider->Shutdown(this);
+	Provider->Shutdown(this);
 		}
 		bIsCurrentTimecodeProviderInitialized = Provider->Initialize(this);
 	}
@@ -2244,7 +2249,7 @@ void UEngine::ReinitializeTimecodeProvider()
 bool UEngine::SetTimecodeProvider(UTimecodeProvider* InTimecodeProvider)
 {
 	if (InTimecodeProvider != TimecodeProvider)
-	{
+			{
 		if (TimecodeProvider && bIsCurrentTimecodeProviderInitialized)
 		{
 			TimecodeProvider->Shutdown(this);
@@ -2254,7 +2259,7 @@ bool UEngine::SetTimecodeProvider(UTimecodeProvider* InTimecodeProvider)
 		TimecodeProvider = InTimecodeProvider && !InTimecodeProvider->IsPendingKill() ? InTimecodeProvider : nullptr;
 
 		if (TimecodeProvider)
-		{
+			{
 			bIsCurrentTimecodeProviderInitialized = TimecodeProvider->Initialize(this);
 		}
 		OnTimecodeProviderChanged().Broadcast();
@@ -2420,12 +2425,12 @@ static void LoadCustomTimeStep(UEngine* Engine)
 static void LoadTimecodeProvider(UEngine* Engine)
 {
 	if (Engine->TimecodeProviderClassName.IsValid())
-	{
+{
 		if (UClass* TimecodeProviderClass = Engine->TimecodeProviderClassName.TryLoadClass<UTimecodeProvider>())
-		{
+	{
 			UTimecodeProvider* NewTimecodeProvider = NewObject<UTimecodeProvider>(Engine, TimecodeProviderClass);
 			if (!Engine->SetTimecodeProvider(NewTimecodeProvider))
-			{
+		{
 				UE_LOG(LogEngine, Warning, TEXT("InitializeTimecodeProvider - Failed to intialize TimecodeProvider '%s'."), *NewTimecodeProvider->GetPathName());
 			}
 		}
@@ -3358,7 +3363,7 @@ void UEngine::SwapControllerId(ULocalPlayer *NewPlayer, const int32 CurrentContr
 
 APlayerController* UEngine::GetFirstLocalPlayerController(const UWorld* InWorld)
 {
-	const FWorldContext& Context = GetWorldContextFromWorldChecked(InWorld);
+	const FWorldContext &Context = GetWorldContextFromWorldChecked(InWorld);
 
 	return (Context.OwningGameInstance ? Context.OwningGameInstance->GetFirstLocalPlayerController(InWorld) : nullptr);
 }
@@ -5008,14 +5013,14 @@ bool UEngine::HandleKismetEventCommand(UWorld* InWorld, const TCHAR* Cmd, FOutpu
 			const bool bIsCDOOwnedArchtype = IsCDOOwnedArchtype(Obj);
 			if (!bIsCDOOwnedArchtype)
 			{
-				UWorld const* const ObjWorld = Obj->GetWorld();
-				if (ObjWorld == InWorld)
-				{
-					const bool bSucceeded = Obj->CallFunctionByNameWithArguments(Cmd, Ar, nullptr, true);
-					NumInstanceCallsSucceeded += bSucceeded ? 1 : 0;
-				}
+			UWorld const* const ObjWorld = Obj->GetWorld();
+			if (ObjWorld == InWorld)
+			{
+				const bool bSucceeded = Obj->CallFunctionByNameWithArguments(Cmd, Ar, nullptr, true);
+				NumInstanceCallsSucceeded += bSucceeded ? 1 : 0;
 			}
 		}
+		}
 
 		Ar.Logf(TEXT("Called '%s' on everything in the world and %d instances succeeded"), Cmd, NumInstanceCallsSucceeded);
 	}
@@ -9306,8 +9311,8 @@ void UEngine::ClearOnScreenDebugMessages()
 	{	
 		// Because some components add their message in concurrent work, we need a CS here.
 		FScopeLock ScopeLock(&GOnScreenMessageCS);
-		ScreenMessages.Empty();
-		PriorityScreenMessages.Empty();
+	ScreenMessages.Empty();
+	PriorityScreenMessages.Empty();
 	}
 #endif // !UE_BUILD_SHIPPING
 }
@@ -10121,55 +10126,55 @@ float UEngine::DrawOnscreenDebugMessages(UWorld* World, FViewport* Viewport, FCa
 		if (PriorityScreenMessages.Num() > 0)
 		{
 			FCanvasTextItem MessageTextItem(FVector2D(0, 0), FText::GetEmpty(), GetSmallFont(), FLinearColor::White);
-			MessageTextItem.EnableShadow(FLinearColor::Black);
+		MessageTextItem.EnableShadow(FLinearColor::Black);
 			for (int32 PrioIndex = PriorityScreenMessages.Num() - 1; PrioIndex >= 0; PrioIndex--)
-			{
+		{
 				FScreenMessageString& Message = PriorityScreenMessages[PrioIndex];
-				if (YPos < MaxYPos)
-				{
-					MessageTextItem.Text = FText::FromString(Message.ScreenMessage);
-					MessageTextItem.SetColor(Message.DisplayColor);
-					MessageTextItem.Scale = Message.TextScale;
-					Canvas->DrawItem(MessageTextItem, FVector2D(MessageX, YPos));
-					YPos += MessageTextItem.DrawnSize.Y * 1.15f;
-				}
-				if (!HasUpdatedScreenDebugMessages.IsSet())
+			if (YPos < MaxYPos)
+			{
+				MessageTextItem.Text = FText::FromString(Message.ScreenMessage);
+				MessageTextItem.SetColor(Message.DisplayColor);
+				MessageTextItem.Scale = Message.TextScale;
+				Canvas->DrawItem(MessageTextItem, FVector2D(MessageX, YPos));
+				YPos += MessageTextItem.DrawnSize.Y * 1.15f;
+			}
+			if (!HasUpdatedScreenDebugMessages.IsSet())
+			{
+				Message.CurrentTimeDisplayed += World->GetDeltaSeconds();
+				if (Message.CurrentTimeDisplayed >= Message.TimeToDisplay)
 				{
-					Message.CurrentTimeDisplayed += World->GetDeltaSeconds();
-					if (Message.CurrentTimeDisplayed >= Message.TimeToDisplay)
-					{
 						PriorityScreenMessages.RemoveAt(PrioIndex);
-					}
 				}
 			}
 		}
+	}
 
 		if (ScreenMessages.Num() > 0)
-		{
+	{
 			FCanvasTextItem MessageTextItem(FVector2D(0, 0), FText::GetEmpty(), GetSmallFont(), FLinearColor::White);
-			MessageTextItem.EnableShadow(FLinearColor::Black);
+		MessageTextItem.EnableShadow(FLinearColor::Black);
 			for (TMap<int32, FScreenMessageString>::TIterator MsgIt(ScreenMessages); MsgIt; ++MsgIt)
+		{
+			FScreenMessageString& Message = MsgIt.Value();
+			if (YPos < MaxYPos)
+			{
+				MessageTextItem.Text = FText::FromString(Message.ScreenMessage);
+				MessageTextItem.SetColor(Message.DisplayColor);
+				MessageTextItem.Scale = Message.TextScale;
+				Canvas->DrawItem(MessageTextItem, FVector2D(MessageX, YPos));
+				YPos += MessageTextItem.DrawnSize.Y * 1.15f;
+			}
+			if (!HasUpdatedScreenDebugMessages.IsSet())
 			{
-				FScreenMessageString& Message = MsgIt.Value();
-				if (YPos < MaxYPos)
+				Message.CurrentTimeDisplayed += World->GetDeltaSeconds();
+				if (Message.CurrentTimeDisplayed >= Message.TimeToDisplay)
 				{
-					MessageTextItem.Text = FText::FromString(Message.ScreenMessage);
-					MessageTextItem.SetColor(Message.DisplayColor);
-					MessageTextItem.Scale = Message.TextScale;
-					Canvas->DrawItem(MessageTextItem, FVector2D(MessageX, YPos));
-					YPos += MessageTextItem.DrawnSize.Y * 1.15f;
-				}
-				if (!HasUpdatedScreenDebugMessages.IsSet())
-				{
-					Message.CurrentTimeDisplayed += World->GetDeltaSeconds();
-					if (Message.CurrentTimeDisplayed >= Message.TimeToDisplay)
-					{
-						MsgIt.RemoveCurrent();
-					}
+					MsgIt.RemoveCurrent();
 				}
 			}
 		}
 	}
+	}
 	// Flag variable that the update has already been done this frame
 	HasUpdatedScreenDebugMessages = true;
 #endif //#if !(UE_BUILD_SHIPPING || UE_BUILD_TEST)
@@ -15991,29 +15996,29 @@ int32 UEngine::RenderStatTimecode(UWorld* World, FViewport* Viewport, FCanvas* C
 	const UTimecodeProvider* Provider = GetTimecodeProvider();
 	if (Provider)
 	{
-		ETimecodeProviderSynchronizationState State = Provider->GetSynchronizationState();
-		FString ProviderName = Provider->GetName();
-		float CharWidth, CharHeight;
-		Font->GetCharSize(TEXT(' '), CharWidth, CharHeight);
-		int32 NewX = X - Font->GetStringSize(*ProviderName) - (int32)CharWidth;
-		switch (State)
-		{
-		case ETimecodeProviderSynchronizationState::Closed:
-			Canvas->DrawShadowedString(NewX, Y, *FString::Printf(TEXT("%s TC: Closed"), *ProviderName), Font, FColor::Red);
-			break;
-		case ETimecodeProviderSynchronizationState::Error:
-			Canvas->DrawShadowedString(NewX, Y, *FString::Printf(TEXT("%s TC: Error"), *ProviderName), Font, FColor::Red);
-			break;
-		case ETimecodeProviderSynchronizationState::Synchronized:
-			Canvas->DrawShadowedString(NewX, Y, *FString::Printf(TEXT("%s TC: Synchronized"), *ProviderName), Font, FColor::Green);
-			break;
-		case ETimecodeProviderSynchronizationState::Synchronizing:
-			Canvas->DrawShadowedString(NewX, Y, *FString::Printf(TEXT("%s TC: Synchronizing"), *ProviderName), Font, FColor::Yellow);
-			break;
-		default:
-			check(false);
-			break;
-		}
+	ETimecodeProviderSynchronizationState State = Provider->GetSynchronizationState();
+	FString ProviderName = Provider->GetName();
+	float CharWidth, CharHeight;
+	Font->GetCharSize(TEXT(' '), CharWidth, CharHeight);
+	int32 NewX = X - Font->GetStringSize(*ProviderName) - (int32)CharWidth;
+	switch (State)
+	{
+	case ETimecodeProviderSynchronizationState::Closed:
+		Canvas->DrawShadowedString(NewX, Y, *FString::Printf(TEXT("%s TC: Closed"), *ProviderName), Font, FColor::Red);
+		break;
+	case ETimecodeProviderSynchronizationState::Error:
+		Canvas->DrawShadowedString(NewX, Y, *FString::Printf(TEXT("%s TC: Error"), *ProviderName), Font, FColor::Red);
+		break;
+	case ETimecodeProviderSynchronizationState::Synchronized:
+		Canvas->DrawShadowedString(NewX, Y, *FString::Printf(TEXT("%s TC: Synchronized"), *ProviderName), Font, FColor::Green);
+		break;
+	case ETimecodeProviderSynchronizationState::Synchronizing:
+		Canvas->DrawShadowedString(NewX, Y, *FString::Printf(TEXT("%s TC: Synchronizing"), *ProviderName), Font, FColor::Yellow);
+		break;
+	default:
+		check(false);
+		break;
+	}
 	}
 	Y += RowHeight;
 	Canvas->DrawShadowedString(X, Y, *FString::Printf(TEXT("TC: %s"), *FApp::GetTimecode().ToString()), Font, FColor::Green);
diff --git a/Engine/Source/Runtime/Engine/Public/Animation/SkinWeightProfile.h b/Engine/Source/Runtime/Engine/Public/Animation/SkinWeightProfile.h
--- a/Engine/Source/Runtime/Engine/Public/Animation/SkinWeightProfile.h
+++ b/Engine/Source/Runtime/Engine/Public/Animation/SkinWeightProfile.h
@@ -9,6 +9,9 @@
 #include "RenderingThread.h"
 #include "HAL/UnrealMemory.h"
 #include "BoneIndices.h"
+#include "RHI/Public/RHIGPUReadback.h"
+#include "Core/Public/Templates/UniquePtr.h"
+
 #include "SkinWeightProfile.generated.h"
 
 class USkeletalMesh;
@@ -81,61 +84,40 @@ struct FRuntimeSkinWeightProfileData
 	{
 		/** Offset into FRuntimeSkinWeightOverrideData.Weights */
 		uint32 InfluencesOffset;
+#if WITH_EDITORONLY_DATA
 		/** Number of influences to be read from FRuntimeSkinWeightOverrideData.Weights */
-		uint8 NumInfluences;
-
+		uint8 NumInfluences_DEPRECATED;
+#endif
 		friend FArchive& operator<<(FArchive& Ar, FSkinWeightOverrideInfo& OverrideInfo);
 	};
 
-	void ApplyOverrides(FSkinWeightVertexBuffer* OverrideBuffer, const FSkinWeightVertexBuffer* BaseBuffer) const
-	{
-		const uint32 ExpectedNumVerts = BaseBuffer->GetNumVertices();
-		if (ExpectedNumVerts)
-		{
-			TArray<FSkinWeightInfo> OverrideArray;
-			BaseBuffer->GetSkinWeights(OverrideArray);
-
-			// Apply overrides
-			{
-				for (auto VertexIndexOverridePair : VertexIndexOverrideIndex)
-				{
-					const uint32 VertexIndex = VertexIndexOverridePair.Key;
-					FSkinWeightInfo& Entry = OverrideArray[VertexIndex];
+	void ApplyOverrides(FSkinWeightVertexBuffer* OverrideBuffer, const void* DataBuffer, const int32 NumVerts) const;	
+	void ApplyDefaultOverride(FSkinWeightVertexBuffer* Buffer) const;
 
-					const uint32 OverrideIndex = VertexIndexOverridePair.Value;
-					const FSkinWeightOverrideInfo& OverrideInfo = OverridesInfo[OverrideIndex];
-
-					FMemory::Memzero(Entry.InfluenceBones);
-					FMemory::Memzero(Entry.InfluenceWeights);
-
-					for (int32 Index = 0; Index < OverrideInfo.NumInfluences; ++Index)
-					{
-						const uint32 WeightData = Weights[OverrideInfo.InfluencesOffset + Index];
+#if WITH_EDITORONLY_DATA
+	/** Per skin weight offset into Weights array and number of weights stored */
+	TArray<FSkinWeightOverrideInfo> OverridesInfo_DEPRECATED;
+	/** Bulk data containing all Weights, stored as bone id in upper and weight in lower (8) bits */
+	TArray<uint16> Weights_DEPRECATED;	
+#endif 
 
-						Entry.InfluenceBones[Index] = (WeightData) >> 16;
-						Entry.InfluenceWeights[Index] = (WeightData & 0xFF);
-					}
-				}
-			}
+	TArray<FBoneIndexType> BoneIDs;
+	TArray<uint8> BoneWeights;
+	/** Map between Vertex Indices and the influence offset into BoneIDs/BoneWeights (DEPRECATED and entries of OverridesInfo) */
+	TMap<uint32, uint32> VertexIndexToInfluenceOffset;
 
-			(*OverrideBuffer) = OverrideArray;
-		}
-		else
-		{
-			OverrideBuffer->CopyMetaData(*BaseBuffer);
-		}
-	}
-
-	/** Per skin weight offset into Weights array and number of weights stored */
-	TArray<FSkinWeightOverrideInfo> OverridesInfo;
-	/** Bulk data containing all Weights, stored as bone id in upper and weight in lower (16) bits */
-	TArray<uint32> Weights;
-	/** Map between Vertex Indices and entries of OverridesInfo */
-	TMap<uint32, uint32> VertexIndexOverrideIndex;
+	uint8 NumWeightsPerVertex;
 	
 	friend FArchive& operator<<(FArchive& Ar, FRuntimeSkinWeightProfileData& OverrideData);
 };
 
+struct FSkinweightReadbackData
+{
+	TUniquePtr<FRHIGPUBufferReadback> BufferReadback;
+	TArray<uint8> ReadbackData;
+	uint32 ReadbackFinishedFrameIndex;
+};
+
 /** Runtime structure for keeping track of skin weight profile(s) and the associated buffer */
 struct ENGINE_API FSkinWeightProfilesData
 {
@@ -148,17 +130,21 @@ struct ENGINE_API FSkinWeightProfilesData
 	// Mark this as non-editor only to prevent mishaps from users
 	void OverrideBaseBufferSkinWeightData(USkeletalMesh* Mesh, int32 LODIndex);
 #endif 
-	void SetDynamicDefaultSkinWeightProfile(USkeletalMesh* Mesh, int32 LODIndex);	
+	void SetDynamicDefaultSkinWeightProfile(USkeletalMesh* Mesh, int32 LODIndex, bool bSerialization = false);	
 	void ClearDynamicDefaultSkinWeightProfile(USkeletalMesh* Mesh, int32 LODIndex);
+	void SetupDynamicDefaultSkinweightProfile();
 	FSkinWeightVertexBuffer* GetDefaultOverrideBuffer() const { return DefaultOverrideSkinWeightBuffer; }
 
-	void ApplyOverrideProfile(FSkinWeightVertexBuffer* OverrideBuffer, const FName& ProfileName);
+	bool ContainsProfile(const FName& ProfileName) const;
+	FSkinWeightVertexBuffer* GetOverrideBuffer(const FName& ProfileName) const;
+	bool ContainsOverrideBuffer(const FName& ProfileName) const;
 	
-	FSkinWeightVertexBuffer* GetOverrideBuffer(const FName& ProfileName);
+#if WITH_EDITOR
 	const FRuntimeSkinWeightProfileData* GetOverrideData(const FName& ProfileName) const;
 	FRuntimeSkinWeightProfileData& AddOverrideData(const FName& ProfileName);
+#endif // WITH_EDITOR
 	
-	void ReleaseBuffer(const FName& ProfileName);
+	void ReleaseBuffer(const FName& ProfileName, bool bForceRelease = false);
 	void ReleaseResources();
 
 	SIZE_T GetResourcesSize() const;
@@ -192,7 +178,18 @@ struct ENGINE_API FSkinWeightProfilesData
 		}
 	}
 
+	bool IsPendingReadback() const;
+	void EnqueueGPUReadback();
+	bool IsGPUReadbackFinished() const;
+	void EnqueueDataReadback();
+	bool IsDataReadbackPending() const;
+	bool IsDataReadbackFinished() const;
+	void ResetGPUReadback();
+	void InitialiseProfileBuffer(const FName& ProfileName);
+
 protected:
+	void ApplyOverrideProfile(FSkinWeightVertexBuffer* OverrideBuffer, const FName& ProfileName);
+
 	template <bool bRenderThread>
 	void CreateRHIBuffers_Internal(TArray<TPair<FName, FSkinWeightRHIInfo>>& OutBuffers);
 
@@ -205,5 +202,8 @@ protected:
 	bool bDefaultOverriden;
 	bool bStaticOverriden;
 	FName DefaultProfileName;
+
+protected:
+	FSkinweightReadbackData ReadbackData;
 };
 
diff --git a/Engine/Source/Runtime/Engine/Public/Animation/SkinWeightProfileManager.h b/Engine/Source/Runtime/Engine/Public/Animation/SkinWeightProfileManager.h
new file mode 100644
--- /dev/null
+++ b/Engine/Source/Runtime/Engine/Public/Animation/SkinWeightProfileManager.h
@@ -0,0 +1,150 @@
+// Copyright Epic Games, Inc. All Rights Reserved.
+
+#pragma once
+
+#include "Stats/Stats.h"
+#include "UObject/WeakObjectPtrTemplates.h"
+#include "Components/SkinnedMeshComponent.h"
+#include "Engine/SkeletalMesh.h"
+#include "RHIGPUReadback.h"
+#include "Tickable.h"
+#include "Templates/Function.h"
+#include "Engine/EngineBaseTypes.h"
+#include "UObject/GCObject.h"
+#include "UObject/ObjectMacros.h"
+
+#include "SkinWeightProfileManager.generated.h"
+
+class FSkinWeightProfileManager;
+
+typedef TFunction<void(TWeakObjectPtr<USkeletalMesh> WeakMesh, FName ProfileName)> FRequestFinished;
+
+/** Describes a single skin weight profile request */
+struct FSetProfileRequest
+{
+	/** Name of the skin weight profile to be loaded */
+	FName ProfileName;
+	/** LOD Indices to load the profile for */
+	TArray<int32> LODIndices;
+	/** Called when the profile request has finished and data is ready (called from GT only) */
+	FRequestFinished Callback;
+
+	/** Weak UObject that is responsible for this request */
+	TWeakObjectPtr<UObject> IdentifyingObject;
+	/** Weak skeletal mesh for which the skin weight profile is to be loaded */
+	TWeakObjectPtr<USkeletalMesh> WeakSkeletalMesh;
+
+	friend bool operator==(const FSetProfileRequest& A, const FSetProfileRequest& B)
+	{
+		return A.ProfileName == B.ProfileName && A.WeakSkeletalMesh == B.WeakSkeletalMesh && A.IdentifyingObject == B.IdentifyingObject;
+	}
+
+	friend uint32 GetTypeHash(FSetProfileRequest A)
+	{
+		return HashCombine(GetTypeHash(A.ProfileName), GetTypeHash(A.WeakSkeletalMesh));
+	}
+};
+
+/** Async task handling the skin weight buffer generation */
+class FSkinWeightProfileManagerAsyncTask
+{
+	FSkinWeightProfileManager* Owner;
+
+public:
+	FSkinWeightProfileManagerAsyncTask(FSkinWeightProfileManager* InOwner)
+		: Owner(InOwner)
+	{
+
+	}
+
+	FORCEINLINE TStatId GetStatId() const
+	{
+		RETURN_QUICK_DECLARE_CYCLE_STAT(FSkinWeightProfileManagerAsyncTask, STATGROUP_TaskGraphTasks);
+	}
+
+	ENamedThreads::Type GetDesiredThread()
+	{
+#if WITH_EDITOR
+		// Force task to run on GT since it can cause a stall waiting for RT
+		return ENamedThreads::GameThread;
+#else
+
+#if PLATFORM_ANDROID
+		return IsRunningRHIInSeparateThread() ? ENamedThreads::RHIThread : ENamedThreads::ActualRenderingThread;
+#else
+		return ENamedThreads::AnyBackgroundHiPriTask;
+#endif
+
+#endif 		
+	}
+
+	static ESubsequentsMode::Type GetSubsequentsMode() { return ESubsequentsMode::TrackSubsequents; }
+	void DoTask(ENamedThreads::Type CurrentThread, const FGraphEventRef& MyCompletionGraphEvent);
+};
+
+USTRUCT()
+struct FSkinWeightProfileManagerTickFunction : public FTickFunction
+{
+	GENERATED_USTRUCT_BODY()
+
+	//~ FTickFunction Interface
+	virtual void ExecuteTick(float DeltaTime, ELevelTick TickType, ENamedThreads::Type CurrentThread, const FGraphEventRef& MyCompletionGraphEvent) override;
+	virtual FString DiagnosticMessage() override;
+	virtual FName DiagnosticContext(bool bDetailed)  override;
+	//~ FTickFunction Interface
+
+	FSkinWeightProfileManager* Owner;
+};
+
+template<>
+struct TStructOpsTypeTraits<FSkinWeightProfileManagerTickFunction> : public TStructOpsTypeTraitsBase2<FSkinWeightProfileManagerTickFunction>
+{
+	enum
+	{
+		WithCopy = false
+	};
+};
+
+class FSkinWeightProfileManager : public FTickableGameObject
+{
+protected: 
+	friend FSkinWeightProfileManagerAsyncTask;
+
+	static TMap<UWorld*, FSkinWeightProfileManager*> WorldManagers;
+	static void OnWorldInit(UWorld* World, const UWorld::InitializationValues IVS);
+	static void OnWorldCleanup(UWorld* World, bool bSessionEnded, bool bCleanupResources);
+	static void OnPreWorldFinishDestroy(UWorld* World);
+	static void OnWorldBeginTearDown(UWorld* World);
+
+public: 
+	static void OnStartup();
+	static void OnShutdown();
+	static FSkinWeightProfileManager* Get(UWorld* World);
+
+	FSkinWeightProfileManager(UWorld* InWorld);
+	virtual ~FSkinWeightProfileManager() {}
+
+	void RequestSkinWeightProfile(FName InProfileName, USkeletalMesh* Mesh, UObject* Requester, FRequestFinished& Callback, int32 LODIndex = INDEX_NONE);
+	void CancelSkinWeightProfileRequest(UObject* Requester);
+	
+	void DoTick(float DeltaTime, ENamedThreads::Type CurrentThread, const FGraphEventRef& MyCompletionGraphEvent);
+protected:	
+	void CleanupRequest(const FSetProfileRequest& Request);
+protected:
+	TArray<FSetProfileRequest, TInlineAllocator<4>> CanceledRequest;
+	TArray<FSetProfileRequest> PendingSetProfileRequests;
+	TMap<TWeakObjectPtr<USkeletalMesh>, int32> PendingMeshes;
+	FSkinWeightProfileManagerTickFunction TickFunction;
+	int32 LastGamethreadProfileIndex;
+
+	TWeakObjectPtr<UWorld> WeakWorld;
+
+	FGraphEventRef AsyncTask;
+public:
+	virtual bool IsTickableWhenPaused() const override;
+	virtual bool IsTickableInEditor() const override;
+	virtual void Tick(float DeltaTime) override;
+	virtual bool IsTickable() const override;
+	virtual TStatId GetStatId() const override;
+
+};
diff --git a/Engine/Source/Runtime/Engine/Public/Rendering/SkinWeightVertexBuffer.h b/Engine/Source/Runtime/Engine/Public/Rendering/SkinWeightVertexBuffer.h
--- a/Engine/Source/Runtime/Engine/Public/Rendering/SkinWeightVertexBuffer.h
+++ b/Engine/Source/Runtime/Engine/Public/Rendering/SkinWeightVertexBuffer.h
@@ -290,6 +290,8 @@ public:
 	void SetBoneWeight(uint32 VertexWeightOffset, uint32 VertexInfluenceCount, uint32 InfluenceIndex, uint8 BoneWeight);
 	void ResetVertexBoneWeights(uint32 VertexWeightOffset, uint32 VertexInfluenceCount);
 
+	void CopyDataFromBuffer(const TArrayView<const FSkinWeightInfo>& SkinWeightData);
+
 	/** Create an RHI vertex buffer with CPU data. CPU data may be discarded after creation (see TResourceArray::Discard) */
 	FVertexBufferRHIRef CreateRHIBuffer_RenderThread();
 	FVertexBufferRHIRef CreateRHIBuffer_Async();
@@ -319,6 +321,13 @@ public:
 		}
 	}
 
+	bool IsWeightDataValid() const;
+
+	FSkinWeightInfo* GetWeightData() const
+	{
+		return (FSkinWeightInfo*)Data;
+	}
+
 protected:
 	// guaranteed only to be valid if the vertex buffer is valid
 	FShaderResourceViewRHIRef SRVValue;
@@ -391,6 +400,8 @@ public:
 	void GetSkinWeights(TArray<FSkinWeightInfo>& OutVertices) const;
 	FSkinWeightInfo GetVertexSkinWeights(uint32 VertexIndex) const;
 
+	void CopySkinWeightInfoData(const TArrayView<const FSkinWeightInfo>& SkinWeightData);
+
 	friend FArchive& operator<<(FArchive& Ar, FSkinWeightVertexBuffer& VertexBuffer);
 
 	void SerializeMetaData(FArchive& Ar);
