Bump skeletalmesh guid, to force a DDC refresh. There is some asset that has bad DDC
#rnx

diff --git a/Engine/Source/Runtime/Engine/Private/SkeletalMeshRenderData.cpp b/Engine/Source/Runtime/Engine/Private/SkeletalMeshRenderData.cpp
--- a/Engine/Source/Runtime/Engine/Private/SkeletalMeshRenderData.cpp
+++ b/Engine/Source/Runtime/Engine/Private/SkeletalMeshRenderData.cpp
@@ -59,7 +59,7 @@ static void SerializeLODInfoForDDC(USkeletalMesh* SkeletalMesh, FString& KeySuff
 // differences, etc.) replace the version GUID below with a new one.
 // In case of merge conflicts with DDC versions, you MUST generate a new GUID
 // and set this new GUID as the version.
-#define SKELETALMESH_DERIVEDDATA_VER TEXT("D6E4934179CC45EFA985EE626BF24636")
+#define SKELETALMESH_DERIVEDDATA_VER TEXT("EBD3AC7CCBB449049C38BE9816BF7C2E")
 
 static const FString& GetSkeletalMeshDerivedDataVersion()
 {
