Added scope to AActor::RegisterAllComponents and moved IncrementalRegisterComponents outside a verify.

diff --git a/Engine/Source/Runtime/Engine/Private/Actor.cpp b/Engine/Source/Runtime/Engine/Private/Actor.cpp
--- a/Engine/Source/Runtime/Engine/Private/Actor.cpp
+++ b/Engine/Source/Runtime/Engine/Private/Actor.cpp
@@ -4364,10 +4364,13 @@ void AActor::PostUnregisterAllComponents()
 
 void AActor::RegisterAllComponents()
 {
+	QUICK_SCOPE_CYCLE_COUNTER(STAT_AActor_RegisterAllComponents);
+
 	PreRegisterAllComponents();
 
 	// 0 - means register all components
-	verify(IncrementalRegisterComponents(0));
+	bool bAllRegistered = IncrementalRegisterComponents(0);
+	check(bAllRegistered);
 
 	// Clear this flag as it's no longer deferred
 	bHasDeferredComponentRegistration = false;
