Moved RemoveRedundantKeys for CurveTables to cooking to improve loading time

diff --git a/Engine/Source/Runtime/Engine/Private/CurveTable.cpp b/Engine/Source/Runtime/Engine/Private/CurveTable.cpp
--- a/Engine/Source/Runtime/Engine/Private/CurveTable.cpp
+++ b/Engine/Source/Runtime/Engine/Private/CurveTable.cpp
@@ -124,11 +124,6 @@ void UCurveTable::Serialize(FArchive& Ar)
 				FSimpleCurve* NewCurve = new FSimpleCurve();
 				FSimpleCurve::StaticStruct()->SerializeTaggedProperties(Ar, (uint8*)NewCurve, FSimpleCurve::StaticStruct(), nullptr);
 
-				if (!GIsEditor && CVar_CurveTable_RemoveRedundantKeys > 0)
-				{
-					NewCurve->RemoveRedundantKeys(0.f);
-				}
-
 				// Add to map
 				RowMap.Add(RowName, NewCurve);
 			}
@@ -137,11 +132,6 @@ void UCurveTable::Serialize(FArchive& Ar)
 				FRichCurve* NewCurve = new FRichCurve();
 				FRichCurve::StaticStruct()->SerializeTaggedProperties(Ar, (uint8*)NewCurve, FRichCurve::StaticStruct(), nullptr);
 
-				if (!GIsEditor && CVar_CurveTable_RemoveRedundantKeys > 0)
-				{
-					NewCurve->RemoveRedundantKeys(0.f);
-				}
-
 				// Add to map
 				RowMap.Add(RowName, NewCurve);
 
@@ -205,12 +195,20 @@ void UCurveTable::Serialize(FArchive& Ar)
 			if (CurveTableMode == ECurveTableMode::SimpleCurves)
 			{
 				FSimpleCurve* Curve = (FSimpleCurve*)RowIt.Value();
+				if (Ar.IsCooking() && Ar.IsPersistent() && !Ar.IsObjectReferenceCollector() && !Ar.ShouldSkipBulkData() && CVar_CurveTable_RemoveRedundantKeys > 0)
+				{
+					Curve->RemoveRedundantKeys(0.f);
+				}
 				FSimpleCurve::StaticStruct()->SerializeTaggedProperties(Ar, (uint8*)Curve, FSimpleCurve::StaticStruct(), nullptr);
 			}
 			else
 			{
 				check(CurveTableMode == ECurveTableMode::RichCurves);
 				FRichCurve* Curve = (FRichCurve*)RowIt.Value();
+				if (Ar.IsCooking() && Ar.IsPersistent() && !Ar.IsObjectReferenceCollector() && !Ar.ShouldSkipBulkData() && CVar_CurveTable_RemoveRedundantKeys > 0)
+				{
+					Curve->RemoveRedundantKeys(0.f);
+				}
 				FRichCurve::StaticStruct()->SerializeTaggedProperties(Ar, (uint8*)Curve, FRichCurve::StaticStruct(), nullptr);
 			}
 		}
