Fix data corruption where multicast delegates on components for placed actors were being wiped out by recompiling a blueprint

diff --git a/Engine/Source/Runtime/Engine/Private/ComponentInstanceDataCache.cpp b/Engine/Source/Runtime/Engine/Private/ComponentInstanceDataCache.cpp
--- a/Engine/Source/Runtime/Engine/Private/ComponentInstanceDataCache.cpp
+++ b/Engine/Source/Runtime/Engine/Private/ComponentInstanceDataCache.cpp
@@ -62,6 +62,7 @@ public:
 		return (!bPropertyInImmutableStruct
 			&& (InProperty->HasAnyPropertyFlags(CPF_Transient)
 				|| !InProperty->HasAnyPropertyFlags(CPF_Edit | CPF_Interp)
+				|| InProperty->IsA<FMulticastDelegateProperty>()
 				|| PropertiesToSkip.Contains(InProperty)
 				)
 			);
