Always check if modulation plugin is enabled and valid before searching for modulation settings

diff --git a/Engine/Source/Runtime/Engine/Private/ActiveSound.cpp b/Engine/Source/Runtime/Engine/Private/ActiveSound.cpp
--- a/Engine/Source/Runtime/Engine/Private/ActiveSound.cpp
+++ b/Engine/Source/Runtime/Engine/Private/ActiveSound.cpp
@@ -137,26 +137,26 @@ FActiveSound* FActiveSound::CreateVirtualCopy(const FActiveSound& InActiveSoundT
 	return ActiveSound;
 }
 
-FArchive& operator<<( FArchive& Ar, FActiveSound* ActiveSound )
+FArchive& operator<<(FArchive& Ar, FActiveSound* ActiveSound)
 {
-	if( !Ar.IsLoading() && !Ar.IsSaving() )
+	if (!Ar.IsLoading() && !Ar.IsSaving())
 	{
 		Ar << ActiveSound->Sound;
 		Ar << ActiveSound->WaveInstances;
 		Ar << ActiveSound->SoundNodeOffsetMap;
 	}
-	return( Ar );
+	return(Ar);
 }
 
-void FActiveSound::AddReferencedObjects( FReferenceCollector& Collector)
+void FActiveSound::AddReferencedObjects(FReferenceCollector& Collector)
 {
 	for (auto WaveInstanceIt(WaveInstances.CreateConstIterator()); WaveInstanceIt; ++WaveInstanceIt)
 	{
 		FWaveInstance* WaveInstance = WaveInstanceIt.Value();
 		// Avoid recursing back to the wave instance that sourced this active sound
-		if( WaveInstance )
+		if (WaveInstance)
 		{
-			WaveInstance->AddReferencedObjects( Collector );
+			WaveInstance->AddReferencedObjects(Collector);
 		}
 	}
 
@@ -226,7 +226,7 @@ void FActiveSound::SetSound(USoundBase* InSound)
 {
 	Sound = InSound;
 	bApplyInteriorVolumes = (SoundClassOverride && SoundClassOverride->Properties.bApplyAmbientVolumes)
-							|| (Sound && Sound->ShouldApplyInteriorVolumes());
+		|| (Sound && Sound->ShouldApplyInteriorVolumes());
 }
 
 void FActiveSound::SetSourceEffectChain(USoundEffectSourcePresetChain* InSourceEffectChain)
@@ -238,7 +238,7 @@ void FActiveSound::SetSoundClass(USoundClass* SoundClass)
 {
 	SoundClassOverride = SoundClass;
 	bApplyInteriorVolumes = (SoundClassOverride && SoundClassOverride->Properties.bApplyAmbientVolumes)
-							|| (Sound && Sound->ShouldApplyInteriorVolumes());
+		|| (Sound && Sound->ShouldApplyInteriorVolumes());
 }
 
 bool FActiveSound::IsPlayWhenSilent() const
@@ -507,7 +507,7 @@ void FActiveSound::UpdateWaveInstances(TArray<FWaveInstance*> &InWaveInstances,
 	}
 
 	{
-		SCOPE_CYCLE_COUNTER( STAT_AudioFindNearestLocation );
+		SCOPE_CYCLE_COUNTER(STAT_AudioFindNearestLocation);
 		ClosestListenerIndex = AudioDevice->FindClosestListenerIndex(Transform);
 	}
 
@@ -722,6 +722,11 @@ void FActiveSound::UpdateWaveInstances(TArray<FWaveInstance*> &InWaveInstances,
 
 USoundModulationPluginSourceSettingsBase* FActiveSound::FindModulationSettings() const
 {
+	if (!AudioDevice->IsModulationPluginEnabled() || !AudioDevice->ModulationInterface.IsValid())
+	{
+		return nullptr;
+	}
+
 	if (UClass* PluginClass = GetAudioPluginCustomSettingsClass(EAudioPlugin::MODULATION))
 	{
 		for (USoundModulationPluginSourceSettingsBase* Settings : Sound->Modulation.Settings)
@@ -776,12 +781,9 @@ void FActiveSound::MarkPendingDestroy(bool bDestroyNow)
 			}
 		}
 
-		if (AudioDevice->IsModulationPluginEnabled() && AudioDevice->ModulationInterface.IsValid())
+		if (USoundModulationPluginSourceSettingsBase* ModulationSettings = FindModulationSettings())
 		{
-			if (USoundModulationPluginSourceSettingsBase* ModulationSettings = FindModulationSettings())
-			{
-				AudioDevice->ModulationInterface->OnReleaseSound(static_cast<ISoundModulatable&>(*this));
-			}
+			AudioDevice->ModulationInterface->OnReleaseSound(static_cast<ISoundModulatable&>(*this));
 		}
 	}
 
@@ -1102,7 +1104,7 @@ void FActiveSound::HandleInteriorVolumes(FSoundParseParameters& ParseParams)
 	else
 	{
 		// Ambient and listener in different ambient zone
-		if( InteriorSettings.bIsWorldSettings )
+		if (InteriorSettings.bIsWorldSettings)
 		{
 			// The ambient sound is 'outside' - use the listener's exterior volume
 			CurrentInteriorVolume = FMath::Lerp(SourceInteriorVolume, Listener.InteriorSettings.ExteriorVolume, Listener.ExteriorVolumeInterp);
@@ -1143,18 +1145,18 @@ FWaveInstance& FActiveSound::AddWaveInstance(const UPTRINT WaveInstanceHash)
 	return *WaveInstance;
 }
 
-void FActiveSound::ApplyRadioFilter(const FSoundParseParameters& ParseParams )
+void FActiveSound::ApplyRadioFilter(const FSoundParseParameters& ParseParams)
 {
 	check(AudioDevice);
-	if( AudioDevice->GetMixDebugState() != DEBUGSTATE_DisableRadio )
+	if (AudioDevice->GetMixDebugState() != DEBUGSTATE_DisableRadio)
 	{
 		// Make sure the radio filter is requested
-		if( ParseParams.SoundClass)
+		if (ParseParams.SoundClass)
 		{
 			const float RadioFilterVolumeThreshold = ParseParams.VolumeMultiplier * ParseParams.SoundClass->Properties.RadioFilterVolumeThreshold;
-			if (RadioFilterVolumeThreshold > KINDA_SMALL_NUMBER )
+			if (RadioFilterVolumeThreshold > KINDA_SMALL_NUMBER)
 			{
-				bApplyRadioFilter = ( ParseParams.Volume < RadioFilterVolumeThreshold );
+				bApplyRadioFilter = (ParseParams.Volume < RadioFilterVolumeThreshold);
 			}
 		}
 	}
@@ -1166,14 +1168,14 @@ void FActiveSound::ApplyRadioFilter(const FSoundParseParameters& ParseParams )
 	bRadioFilterSelected = true;
 }
 
-bool FActiveSound::GetFloatParameter( const FName InName, float& OutFloat ) const
+bool FActiveSound::GetFloatParameter(const FName InName, float& OutFloat) const
 {
 	// Always fail if we pass in no name.
-	if( InName != NAME_None )
+	if (InName != NAME_None)
 	{
-		for( const FAudioComponentParam& P : InstanceParameters )
+		for (const FAudioComponentParam& P : InstanceParameters)
 		{
-			if( P.ParamName == InName )
+			if (P.ParamName == InName)
 			{
 				OutFloat = P.FloatParam;
 				return true;
@@ -1210,14 +1212,14 @@ void FActiveSound::UpdateConcurrencyVolumeScalars(const float DeltaTime)
 	}
 }
 
-void FActiveSound::SetFloatParameter( const FName InName, const float InFloat )
+void FActiveSound::SetFloatParameter(const FName InName, const float InFloat)
 {
-	if( InName != NAME_None )
+	if (InName != NAME_None)
 	{
 		// First see if an entry for this name already exists
-		for( FAudioComponentParam& P : InstanceParameters )
+		for (FAudioComponentParam& P : InstanceParameters)
 		{
-			if( P.ParamName == InName )
+			if (P.ParamName == InName)
 			{
 				P.FloatParam = InFloat;
 				return;
@@ -1226,19 +1228,19 @@ void FActiveSound::SetFloatParameter( const FName InName, const float InFloat )
 
 		// We didn't find one, so create a new one.
 		const int32 NewParamIndex = InstanceParameters.AddDefaulted();
-		InstanceParameters[ NewParamIndex ].ParamName = InName;
-		InstanceParameters[ NewParamIndex ].FloatParam = InFloat;
+		InstanceParameters[NewParamIndex].ParamName = InName;
+		InstanceParameters[NewParamIndex].FloatParam = InFloat;
 	}
 }
 
-bool FActiveSound::GetWaveParameter( const FName InName, USoundWave*& OutWave ) const
+bool FActiveSound::GetWaveParameter(const FName InName, USoundWave*& OutWave) const
 {
 	// Always fail if we pass in no name.
-	if( InName != NAME_None )
+	if (InName != NAME_None)
 	{
-		for( const FAudioComponentParam& P : InstanceParameters )
+		for (const FAudioComponentParam& P : InstanceParameters)
 		{
-			if( P.ParamName == InName )
+			if (P.ParamName == InName)
 			{
 				OutWave = P.SoundWaveParam;
 				return true;
@@ -1249,14 +1251,14 @@ bool FActiveSound::GetWaveParameter( const FName InName, USoundWave*& OutWave )
 	return false;
 }
 
-void FActiveSound::SetWaveParameter( const FName InName, USoundWave* InWave )
+void FActiveSound::SetWaveParameter(const FName InName, USoundWave* InWave)
 {
-	if( InName != NAME_None )
+	if (InName != NAME_None)
 	{
 		// First see if an entry for this name already exists
-		for( FAudioComponentParam& P : InstanceParameters )
+		for (FAudioComponentParam& P : InstanceParameters)
 		{
-			if( P.ParamName == InName )
+			if (P.ParamName == InName)
 			{
 				P.SoundWaveParam = InWave;
 				return;
@@ -1265,19 +1267,19 @@ void FActiveSound::SetWaveParameter( const FName InName, USoundWave* InWave )
 
 		// We didn't find one, so create a new one.
 		const int32 NewParamIndex = InstanceParameters.AddDefaulted();
-		InstanceParameters[ NewParamIndex ].ParamName = InName;
-		InstanceParameters[ NewParamIndex ].SoundWaveParam = InWave;
+		InstanceParameters[NewParamIndex].ParamName = InName;
+		InstanceParameters[NewParamIndex].SoundWaveParam = InWave;
 	}
 }
 
-bool FActiveSound::GetBoolParameter( const FName InName, bool& OutBool ) const
+bool FActiveSound::GetBoolParameter(const FName InName, bool& OutBool) const
 {
 	// Always fail if we pass in no name.
-	if( InName != NAME_None )
+	if (InName != NAME_None)
 	{
-		for( const FAudioComponentParam& P : InstanceParameters )
+		for (const FAudioComponentParam& P : InstanceParameters)
 		{
-			if( P.ParamName == InName )
+			if (P.ParamName == InName)
 			{
 				OutBool = P.BoolParam;
 				return true;
@@ -1288,14 +1290,14 @@ bool FActiveSound::GetBoolParameter( const FName InName, bool& OutBool ) const
 	return false;
 }
 
-void FActiveSound::SetBoolParameter( const FName InName, const bool InBool )
+void FActiveSound::SetBoolParameter(const FName InName, const bool InBool)
 {
-	if( InName != NAME_None )
+	if (InName != NAME_None)
 	{
 		// First see if an entry for this name already exists
-		for( FAudioComponentParam& P : InstanceParameters )
+		for (FAudioComponentParam& P : InstanceParameters)
 		{
-			if( P.ParamName == InName )
+			if (P.ParamName == InName)
 			{
 				P.BoolParam = InBool;
 				return;
@@ -1304,19 +1306,19 @@ void FActiveSound::SetBoolParameter( const FName InName, const bool InBool )
 
 		// We didn't find one, so create a new one.
 		const int32 NewParamIndex = InstanceParameters.AddDefaulted();
-		InstanceParameters[ NewParamIndex ].ParamName = InName;
-		InstanceParameters[ NewParamIndex ].BoolParam = InBool;
+		InstanceParameters[NewParamIndex].ParamName = InName;
+		InstanceParameters[NewParamIndex].BoolParam = InBool;
 	}
 }
 
-bool FActiveSound::GetIntParameter( const FName InName, int32& OutInt ) const
+bool FActiveSound::GetIntParameter(const FName InName, int32& OutInt) const
 {
 	// Always fail if we pass in no name.
-	if( InName != NAME_None )
+	if (InName != NAME_None)
 	{
-		for( const FAudioComponentParam& P : InstanceParameters )
+		for (const FAudioComponentParam& P : InstanceParameters)
 		{
-			if( P.ParamName == InName )
+			if (P.ParamName == InName)
 			{
 				OutInt = P.IntParam;
 				return true;
@@ -1327,14 +1329,14 @@ bool FActiveSound::GetIntParameter( const FName InName, int32& OutInt ) const
 	return false;
 }
 
-void FActiveSound::SetIntParameter( const FName InName, const int32 InInt )
+void FActiveSound::SetIntParameter(const FName InName, const int32 InInt)
 {
-	if( InName != NAME_None )
+	if (InName != NAME_None)
 	{
 		// First see if an entry for this name already exists
-		for( FAudioComponentParam& P : InstanceParameters )
+		for (FAudioComponentParam& P : InstanceParameters)
 		{
-			if( P.ParamName == InName )
+			if (P.ParamName == InName)
 			{
 				P.IntParam = InInt;
 				return;
@@ -1343,8 +1345,8 @@ void FActiveSound::SetIntParameter( const FName InName, const int32 InInt )
 
 		// We didn't find one, so create a new one.
 		const int32 NewParamIndex = InstanceParameters.AddDefaulted();
-		InstanceParameters[ NewParamIndex ].ParamName = InName;
-		InstanceParameters[ NewParamIndex ].IntParam = InInt;
+		InstanceParameters[NewParamIndex].ParamName = InName;
+		InstanceParameters[NewParamIndex].IntParam = InInt;
 	}
 }
 
@@ -1353,7 +1355,7 @@ void FActiveSound::SetSoundParameter(const FAudioComponentParam& Param)
 	if (Param.ParamName != NAME_None)
 	{
 		// First see if an entry for this name already exists
-		for( FAudioComponentParam& P : InstanceParameters )
+		for (FAudioComponentParam& P : InstanceParameters)
 		{
 			if (P.ParamName == Param.ParamName)
 			{
@@ -1381,7 +1383,7 @@ void FActiveSound::CollectAttenuationShapesForVisualization(TMultiMap<EAttenuati
 	if (SoundCue)
 	{
 		TArray<USoundNodeAttenuation*> AttenuationNodes;
-		SoundCue->RecursiveFindAttenuation( SoundCue->FirstNode, AttenuationNodes );
+		SoundCue->RecursiveFindAttenuation(SoundCue->FirstNode, AttenuationNodes);
 		for (int32 NodeIndex = 0; NodeIndex < AttenuationNodes.Num(); ++NodeIndex)
 		{
 			const FSoundAttenuationSettings* AttenuationSettingsToApply = AttenuationNodes[NodeIndex]->GetAttenuationSettingsToApply();
diff --git a/Engine/Source/Runtime/Engine/Private/AudioDevice.cpp b/Engine/Source/Runtime/Engine/Private/AudioDevice.cpp
--- a/Engine/Source/Runtime/Engine/Private/AudioDevice.cpp
+++ b/Engine/Source/Runtime/Engine/Private/AudioDevice.cpp
@@ -4586,12 +4586,9 @@ void FAudioDevice::AddNewActiveSoundInternal(const FActiveSound& NewActiveSound,
 	}
 	(*PlayCount)++;
 
-	if (ModulationInterface.IsValid())
+	if (USoundModulationPluginSourceSettingsBase* ModulationSettings = ActiveSound->FindModulationSettings())
 	{
-		if (USoundModulationPluginSourceSettingsBase* ModulationSettings = ActiveSound->FindModulationSettings())
-		{
-			ModulationInterface->OnInitSound(static_cast<ISoundModulatable&>(*ActiveSound), *ModulationSettings);
-		}
+		ModulationInterface->OnInitSound(static_cast<ISoundModulatable&>(*ActiveSound), *ModulationSettings);
 	}
 
 #if !(UE_BUILD_SHIPPING || UE_BUILD_TEST)
