Edigrate CL 11296330, 11589183
- Ability to set up distance-based sends to an array of submixes defined in attenuation settings
- Ability to modify submix effect chains from blueprint (add, remove, replace)

diff --git a/Engine/Source/Editor/DetailCustomizations/Private/AttenuationSettingsCustomizations.cpp b/Engine/Source/Editor/DetailCustomizations/Private/AttenuationSettingsCustomizations.cpp
--- a/Engine/Source/Editor/DetailCustomizations/Private/AttenuationSettingsCustomizations.cpp
+++ b/Engine/Source/Editor/DetailCustomizations/Private/AttenuationSettingsCustomizations.cpp
@@ -326,6 +326,7 @@ void FSoundAttenuationSettingsCustomization::CustomizeChildren(TSharedRef<IPrope
 	bIsAirAbsorptionEnabledHandle = PropertyHandles.FindChecked(GET_MEMBER_NAME_CHECKED(FSoundAttenuationSettings, bAttenuateWithLPF)).ToSharedRef();
 	bIsReverbSendEnabledHandle = PropertyHandles.FindChecked(GET_MEMBER_NAME_CHECKED(FSoundAttenuationSettings, bEnableReverbSend)).ToSharedRef();
 	bIsPriorityAttenuationEnabledHandle = PropertyHandles.FindChecked(GET_MEMBER_NAME_CHECKED(FSoundAttenuationSettings, bEnablePriorityAttenuation)).ToSharedRef();
+	bIsSubmixSendAttenuationEnabledHandle = PropertyHandles.FindChecked(GET_MEMBER_NAME_CHECKED(FSoundAttenuationSettings, bEnableSubmixSends)).ToSharedRef();
 	ReverbSendMethodHandle = PropertyHandles.FindChecked(GET_MEMBER_NAME_CHECKED(FSoundAttenuationSettings, ReverbSendMethod)).ToSharedRef();
 	PriorityAttenuationMethodHandle = PropertyHandles.FindChecked(GET_MEMBER_NAME_CHECKED(FSoundAttenuationSettings, PriorityAttenuationMethod)).ToSharedRef();
 	AbsorptionMethodHandle = PropertyHandles.FindChecked(GET_MEMBER_NAME_CHECKED(FSoundAttenuationSettings, AbsorptionMethod)).ToSharedRef();
@@ -521,6 +522,12 @@ void FSoundAttenuationSettingsCustomization::CustomizeChildren(TSharedRef<IPrope
 		.Visibility(TAttribute<EVisibility>(this, &FSoundAttenuationSettingsCustomization::IsManualPriorityAttenuationSelected))
 		.EditCondition(GetIsPriorityAttenuationEnabledAttribute(), nullptr);
 
+	// Add the submix send priority
+	LayoutBuilder.AddPropertyToCategory(bIsSubmixSendAttenuationEnabledHandle)
+		.EditCondition(IsAttenuationOverriddenAttribute(), nullptr);
+
+	LayoutBuilder.AddPropertyToCategory(PropertyHandles.FindChecked(GET_MEMBER_NAME_CHECKED(FSoundAttenuationSettings, SubmixSendSettings)))
+		.EditCondition(GetIsSubmixSendAttenuationEnabledAttribute(), nullptr);
 
 	if (bIsAudioMixerEnabled)
 	{
@@ -528,14 +535,9 @@ void FSoundAttenuationSettingsCustomization::CustomizeChildren(TSharedRef<IPrope
 			.EditCondition(IsAttenuationOverriddenAttribute(), nullptr);
 	}
 
-	if (PropertyHandles.Num() != 61)
+	if (PropertyHandles.Num() != 63)
 	{
-		FString PropertyList;
-		for (auto It(PropertyHandles.CreateConstIterator()); It; ++It)
-		{
-			PropertyList += It.Key().ToString() + TEXT(", ");
-		}
-		ensureMsgf(false, TEXT("Unexpected property handle(s) customizing FSoundAttenuationSettings: %s"), *PropertyList);
+		ensureMsgf(false, TEXT("Unexpected property handle(s) customizing FSoundAttenuationSettings. %d handles found"), PropertyHandles.Num());
 	}
 }
 
@@ -687,6 +689,24 @@ TAttribute<bool> FSoundAttenuationSettingsCustomization::GetIsPriorityAttenuatio
 	return TAttribute<bool>::Create(TAttribute<bool>::FGetter::CreateLambda(Lambda));
 }
 
+TAttribute<bool> FSoundAttenuationSettingsCustomization::GetIsSubmixSendAttenuationEnabledAttribute() const
+{
+	TWeakPtr<IPropertyHandle> bOverrideAttenuationPropertyWeakPtr = bOverrideAttenuationHandle;
+	TWeakPtr<IPropertyHandle> bIsSubmixSendWeakPtr = bIsSubmixSendAttenuationEnabledHandle;
+
+	auto Lambda = [bOverrideAttenuationPropertyWeakPtr, bIsSubmixSendWeakPtr]()
+	{
+		TSharedPtr<IPropertyHandle> bOverrideAttenuationProperty = bOverrideAttenuationPropertyWeakPtr.Pin();
+		TSharedPtr<IPropertyHandle> bIsSubmixSendAttenuationEnabledProperty = bIsSubmixSendWeakPtr.Pin();
+
+		bool Value = GetValue(bOverrideAttenuationProperty);
+		Value &= GetValue(bIsSubmixSendAttenuationEnabledProperty);
+		return Value;
+	};
+
+	return TAttribute<bool>::Create(TAttribute<bool>::FGetter::CreateLambda(Lambda));
+}
+
 EVisibility FSoundAttenuationSettingsCustomization::IsLinearMethodSelected() const
 {
 	uint8 SendMethodValue;
diff --git a/Engine/Source/Editor/DetailCustomizations/Private/AttenuationSettingsCustomizations.h b/Engine/Source/Editor/DetailCustomizations/Private/AttenuationSettingsCustomizations.h
--- a/Engine/Source/Editor/DetailCustomizations/Private/AttenuationSettingsCustomizations.h
+++ b/Engine/Source/Editor/DetailCustomizations/Private/AttenuationSettingsCustomizations.h
@@ -56,6 +56,7 @@ protected:
 	TSharedPtr<IPropertyHandle> bIsOcclusionEnabledHandle;
 	TSharedPtr<IPropertyHandle> bIsReverbSendEnabledHandle;
 	TSharedPtr<IPropertyHandle> bIsPriorityAttenuationEnabledHandle;
+	TSharedPtr<IPropertyHandle> bIsSubmixSendAttenuationEnabledHandle;
 	TSharedPtr<IPropertyHandle> ReverbSendMethodHandle;
 	TSharedPtr<IPropertyHandle> PriorityAttenuationMethodHandle;
 	TSharedPtr<IPropertyHandle> AbsorptionMethodHandle;
@@ -68,6 +69,7 @@ protected:
 	TAttribute<bool> GetIsAirAbsorptionEnabledAttribute() const;
 	TAttribute<bool> GetIsReverbSendEnabledAttribute() const;
 	TAttribute<bool> GetIsPriorityAttenuationEnabledAttribute() const;
+	TAttribute<bool> GetIsSubmixSendAttenuationEnabledAttribute() const;
 
 	EVisibility IsLinearMethodSelected() const;
 	EVisibility IsCustomReverbSendCurveSelected() const;
diff --git a/Engine/Source/Runtime/AudioMixer/Private/AudioMixerBlueprintLibrary.cpp b/Engine/Source/Runtime/AudioMixer/Private/AudioMixerBlueprintLibrary.cpp
--- a/Engine/Source/Runtime/AudioMixer/Private/AudioMixerBlueprintLibrary.cpp
+++ b/Engine/Source/Runtime/AudioMixer/Private/AudioMixerBlueprintLibrary.cpp
@@ -96,6 +96,77 @@ void UAudioMixerBlueprintLibrary::ClearMasterSubmixEffects(const UObject* WorldC
 	}
 }
 
+int32 UAudioMixerBlueprintLibrary::AddSubmixEffect(const UObject* WorldContextObject, USoundSubmix* InSoundSubmix, USoundEffectSubmixPreset* SubmixEffectPreset)
+{
+	if (Audio::FMixerDevice* MixerDevice = GetAudioMixerDeviceFromWorldContext(WorldContextObject))
+	{
+		// Immediately create a new sound effect base here before the object becomes potentially invalidated
+		FSoundEffectSubmixPtr SoundEffectSubmix = MakeShareable(static_cast<FSoundEffectSubmix*>(SubmixEffectPreset->CreateNewEffect()));
+
+		FSoundEffectSubmixInitData InitData;
+		InitData.SampleRate = MixerDevice->GetSampleRate();
+
+		// Initialize and set the preset immediately
+		SoundEffectSubmix->Init(InitData);
+		SoundEffectSubmix->SetPreset(SubmixEffectPreset);
+		SoundEffectSubmix->SetEnabled(true);
+
+		// Get a unique id for the preset object on the game thread. Used to refer to the object on audio render thread.
+		uint32 SubmixPresetUniqueId = SubmixEffectPreset->GetUniqueID();
+
+		return MixerDevice->AddSubmixEffect(InSoundSubmix, SubmixPresetUniqueId, SoundEffectSubmix);
+	}
+	return 0;
+}
+
+void UAudioMixerBlueprintLibrary::RemoveSubmixEffectPreset(const UObject* WorldContextObject, USoundSubmix* InSoundSubmix, USoundEffectSubmixPreset* InSubmixEffectPreset)
+{
+	if (Audio::FMixerDevice* MixerDevice = GetAudioMixerDeviceFromWorldContext(WorldContextObject))
+	{
+		uint32 SubmixPresetUniqueId = InSubmixEffectPreset->GetUniqueID();
+		MixerDevice->RemoveSubmixEffect(InSoundSubmix, SubmixPresetUniqueId);
+	}
+}
+
+void UAudioMixerBlueprintLibrary::RemoveSubmixEffectPresetAtIndex(const UObject* WorldContextObject, USoundSubmix* InSoundSubmix, int32 SubmixChainIndex)
+{
+	if (Audio::FMixerDevice* MixerDevice = GetAudioMixerDeviceFromWorldContext(WorldContextObject))
+	{
+		MixerDevice->RemoveSubmixEffectAtIndex(InSoundSubmix, SubmixChainIndex);
+	}
+}
+
+void UAudioMixerBlueprintLibrary::ReplaceSoundEffectSubmix(const UObject* WorldContextObject, USoundSubmix* InSoundSubmix, int32 SubmixChainIndex, USoundEffectSubmixPreset* SubmixEffectPreset)
+{
+	if (Audio::FMixerDevice* MixerDevice = GetAudioMixerDeviceFromWorldContext(WorldContextObject))
+	{
+		// Immediately create a new sound effect base here before the object becomes potentially invalidated
+		FSoundEffectSubmixPtr SoundEffectSubmix = MakeShareable(static_cast<FSoundEffectSubmix*>(SubmixEffectPreset->CreateNewEffect()));
+
+		FSoundEffectSubmixInitData InitData;
+		InitData.SampleRate = MixerDevice->GetSampleRate();
+
+		// Initialize and set the preset immediately
+		SoundEffectSubmix->Init(InitData);
+		SoundEffectSubmix->SetPreset(SubmixEffectPreset);
+		SoundEffectSubmix->SetEnabled(true);
+
+		// Get a unique id for the preset object on the game thread. Used to refer to the object on audio render thread.
+		uint32 SubmixPresetUniqueId = SubmixEffectPreset->GetUniqueID();
+
+		MixerDevice->ReplaceSoundEffectSubmix(InSoundSubmix, SubmixChainIndex, SubmixPresetUniqueId, SoundEffectSubmix);
+	}
+}
+
+void UAudioMixerBlueprintLibrary::ClearSubmixEffects(const UObject* WorldContextObject, USoundSubmix* InSoundSubmix)
+{
+	if (Audio::FMixerDevice* MixerDevice = GetAudioMixerDeviceFromWorldContext(WorldContextObject))
+	{
+		MixerDevice->ClearSubmixEffects(InSoundSubmix);
+	}
+}
+
+
 void UAudioMixerBlueprintLibrary::StartRecordingOutput(const UObject* WorldContextObject, float ExpectedDuration, USoundSubmix* SubmixToRecord)
 {
 	if (Audio::FMixerDevice* MixerDevice = GetAudioMixerDeviceFromWorldContext(WorldContextObject))
diff --git a/Engine/Source/Runtime/AudioMixer/Private/AudioMixerDevice.cpp b/Engine/Source/Runtime/AudioMixer/Private/AudioMixerDevice.cpp
--- a/Engine/Source/Runtime/AudioMixer/Private/AudioMixerDevice.cpp
+++ b/Engine/Source/Runtime/AudioMixer/Private/AudioMixerDevice.cpp
@@ -943,6 +943,57 @@ namespace Audio
 		});
 	}
 
+	int32 FMixerDevice::AddSubmixEffect(USoundSubmix* InSoundSubmix, uint32 SubmixEffectId, FSoundEffectSubmixPtr SoundEffect)
+	{
+		FMixerSubmixPtr MixerSubmixPtr = GetSubmixInstance(InSoundSubmix).Pin();
+
+		int32 NumEffects = MixerSubmixPtr->GetNumEffects();
+
+		AudioRenderThreadCommand([this, MixerSubmixPtr, SubmixEffectId, SoundEffect]()
+		{
+			MixerSubmixPtr->AddSoundEffectSubmix(SubmixEffectId, SoundEffect);
+		});
+
+		return ++NumEffects;
+	}
+
+	void FMixerDevice::RemoveSubmixEffect(USoundSubmix* InSoundSubmix, uint32 SubmixEffectId)
+	{
+		FMixerSubmixPtr MixerSubmixPtr = GetSubmixInstance(InSoundSubmix).Pin();
+		AudioRenderThreadCommand([MixerSubmixPtr, SubmixEffectId]()
+		{
+			MixerSubmixPtr->RemoveSoundEffectSubmix(SubmixEffectId);
+		});
+	}
+
+	void FMixerDevice::RemoveSubmixEffectAtIndex(USoundSubmix* InSoundSubmix, int32 SubmixChainIndex)
+	{
+		FMixerSubmixPtr MixerSubmixPtr = GetSubmixInstance(InSoundSubmix).Pin();
+		AudioRenderThreadCommand([MixerSubmixPtr, SubmixChainIndex]()
+		{
+			MixerSubmixPtr->RemoveSoundEffectSubmixAtIndex(SubmixChainIndex);
+		});
+	}
+
+	void FMixerDevice::ReplaceSoundEffectSubmix(USoundSubmix* InSoundSubmix, int32 InSubmixChainIndex, int32 SubmixEffectId, FSoundEffectSubmixPtr SoundEffect)
+	{
+		FMixerSubmixPtr MixerSubmixPtr = GetSubmixInstance(InSoundSubmix).Pin();
+		AudioRenderThreadCommand([MixerSubmixPtr, InSubmixChainIndex, SubmixEffectId, SoundEffect]()
+		{
+			MixerSubmixPtr->ReplaceSoundEffectSubmix(InSubmixChainIndex, SubmixEffectId, SoundEffect);
+		});
+	}
+
+	void FMixerDevice::ClearSubmixEffects(USoundSubmix* InSoundSubmix)
+	{
+		FMixerSubmixPtr MixerSubmixPtr = GetSubmixInstance(InSoundSubmix).Pin();
+
+		AudioRenderThreadCommand([MixerSubmixPtr]()
+		{
+			MixerSubmixPtr->ClearSoundEffectSubmixes();
+		});
+	}
+
 	void FMixerDevice::UpdateModulationControls(const uint32 InSourceId, const FSoundModulationControls& InControls)
 	{
 		SourceManager->UpdateModulationControls(InSourceId, InControls);
diff --git a/Engine/Source/Runtime/AudioMixer/Private/AudioMixerSource.cpp b/Engine/Source/Runtime/AudioMixer/Private/AudioMixerSource.cpp
--- a/Engine/Source/Runtime/AudioMixer/Private/AudioMixerSource.cpp
+++ b/Engine/Source/Runtime/AudioMixer/Private/AudioMixerSource.cpp
@@ -983,6 +983,41 @@ namespace Audio
 			MixerSourceVoice->SetSubmixSendInfo(MixerDevice->GetMasterReverbSubmix(), ReverbSendLevel);
 		}
 
+		if (WaveInstance->SubmixSendSettings.Num() > 0)
+		{
+			for (const FAttenuationSubmixSendSettings& SendSettings : WaveInstance->SubmixSendSettings)
+			{
+				if (SendSettings.Submix)
+				{
+					float SubmixSendLevel = 0.0f;
+
+					if (SendSettings.SubmixSendMethod == ESubmixSendMethod::Manual)
+					{
+						SubmixSendLevel = FMath::Clamp(SendSettings.ManualSubmixSendLevel, 0.0f, 1.0f);
+					}
+					else
+					{
+						// The alpha value is determined identically between manual and custom curve methods
+						const float Denom = FMath::Max(SendSettings.SubmixSendDistanceMax - SendSettings.SubmixSendDistanceMin, 1.0f);
+						const float Alpha = FMath::Clamp((WaveInstance->ListenerToSoundDistance - SendSettings.SubmixSendDistanceMin) / Denom, 0.0f, 1.0f);
+
+						if (WaveInstance->ReverbSendMethod == EReverbSendMethod::Linear)
+						{
+							SubmixSendLevel = FMath::Clamp(FMath::Lerp(SendSettings.SubmixSendLevelMin, SendSettings.SubmixSendLevelMax, Alpha), 0.0f, 1.0f);
+						}
+						else
+						{
+							SubmixSendLevel = FMath::Clamp(SendSettings.CustomSubmixSendCurve.GetRichCurveConst()->Eval(Alpha), 0.0f, 1.0f);
+						}
+					}
+
+
+					FMixerSubmixPtr SubmixPtr = MixerDevice->GetSubmixInstance(SendSettings.Submix).Pin();
+					MixerSourceVoice->SetSubmixSendInfo(SubmixPtr, SubmixSendLevel);
+				}
+			}
+		}
+
 		// Update submix send levels
 		for (FSoundSubmixSendInfo& SendInfo : WaveInstance->SoundSubmixSends)
 		{
diff --git a/Engine/Source/Runtime/AudioMixer/Private/AudioMixerSubmix.cpp b/Engine/Source/Runtime/AudioMixer/Private/AudioMixerSubmix.cpp
--- a/Engine/Source/Runtime/AudioMixer/Private/AudioMixerSubmix.cpp
+++ b/Engine/Source/Runtime/AudioMixer/Private/AudioMixerSubmix.cpp
@@ -45,6 +45,7 @@ namespace Audio
 		, CurrentDryLevel(0.0f)
 		, TargetDryLevel(0.0f)
 		, EnvelopeNumChannels(0)
+		, NumSubmixEffects(0)
 		, bIsRecording(false)
 		, bIsBackgroundMuted(false)
 		, OwningSubmixObject(nullptr)
@@ -110,13 +111,17 @@ namespace Audio
 			CurrentDryLevel = FMath::Clamp(SoundSubmix->DryLevel, 0.0f, 1.0f);
 			TargetDryLevel = CurrentDryLevel;
 
+			NumSubmixEffects = 0;
+
 			for (USoundEffectSubmixPreset* EffectPreset : SoundSubmix->SubmixEffectChain)
 			{
 				if (EffectPreset)
 				{
+					++NumSubmixEffects;
+
 					// Create a new effect instance using the preset
 					FSoundEffectSubmix* SubmixEffect = static_cast<FSoundEffectSubmix*>(EffectPreset->CreateNewEffect());
-				
+
 					FSoundEffectSubmixInitData InitData;
 					InitData.DeviceID = MixerDevice->DeviceID;
 					InitData.SampleRate = MixerDevice->GetSampleRate();
@@ -272,7 +277,7 @@ namespace Audio
 
 	void FMixerSubmix::RemoveChildSubmix(TWeakPtr<FMixerSubmix, ESPMode::ThreadSafe> SubmixWeakPtr)
 	{
-		TSharedPtr<FMixerSubmix, ESPMode::ThreadSafe> SubmixStrongPtr  = SubmixWeakPtr.Pin();
+		TSharedPtr<FMixerSubmix, ESPMode::ThreadSafe> SubmixStrongPtr = SubmixWeakPtr.Pin();
 		if (!SubmixStrongPtr.IsValid())
 		{
 			return;
@@ -303,6 +308,11 @@ namespace Audio
 	}
 
 	int32 FMixerSubmix::GetNumEffects() const
+	{
+		return NumSubmixEffects;
+	}
+
+	int32 FMixerSubmix::GetSizeOfSubmixChain() const
 	{
 		return EffectSubmixChain.Num();
 	}
@@ -353,6 +363,8 @@ namespace Audio
 		Info.PresetId = SubmixPresetId;
 		Info.EffectInstance = InSoundEffectSubmix;
 
+		++NumSubmixEffects;
+
 		EffectSubmixChain.Add(Info);
 	}
 
@@ -369,10 +381,26 @@ namespace Audio
 				// Reset reference to the effect
 				Effect.EffectInstance.Reset();
 				Effect.PresetId = INDEX_NONE;
-				return;
+
+				--NumSubmixEffects;
 			}
 		}
+	}
 
+	void FMixerSubmix::RemoveSoundEffectSubmixAtIndex(int32 InIndex)
+	{
+		AUDIO_MIXER_CHECK_AUDIO_PLAT_THREAD(MixerDevice);
+
+		if (InIndex >= 0 && InIndex < EffectSubmixChain.Num())
+		{
+			FSubmixEffectInfo& Effect = EffectSubmixChain[InIndex];
+			if (Effect.EffectInstance.IsValid())
+			{
+				Effect.EffectInstance.Reset();
+				Effect.PresetId = INDEX_NONE;
+				--NumSubmixEffects;
+			}
+		}
 	}
 
 	void FMixerSubmix::ClearSoundEffectSubmixes()
@@ -388,6 +416,7 @@ namespace Audio
 			}
 		}
 
+		NumSubmixEffects = 0;
 		EffectSubmixChain.Reset();
 	}
 
diff --git a/Engine/Source/Runtime/AudioMixer/Public/AudioMixerBlueprintLibrary.h b/Engine/Source/Runtime/AudioMixer/Public/AudioMixerBlueprintLibrary.h
--- a/Engine/Source/Runtime/AudioMixer/Public/AudioMixerBlueprintLibrary.h
+++ b/Engine/Source/Runtime/AudioMixer/Public/AudioMixerBlueprintLibrary.h
@@ -96,6 +96,26 @@ public:
 	UFUNCTION(BlueprintCallable, Category = "Audio|Effects", meta = (WorldContext = "WorldContextObject"))
 	static void ClearMasterSubmixEffects(const UObject* WorldContextObject);
 
+	/** Adds a submix effect preset to the given submix at the end of its submix effect chain. Returns the number of submix effects. */
+	UFUNCTION(BlueprintCallable, Category = "Audio|Effects", meta = (WorldContext = "WorldContextObject"))
+	static int32 AddSubmixEffect(const UObject* WorldContextObject, USoundSubmix* SoundSubmix, USoundEffectSubmixPreset* SubmixEffectPreset);
+
+	/** Removes all instances of a submix effect preset from the given submix. */
+	UFUNCTION(BlueprintCallable, Category = "Audio|Effects", meta = (WorldContext = "WorldContextObject"))
+	static void RemoveSubmixEffectPreset(const UObject* WorldContextObject, USoundSubmix* SoundSubmix, USoundEffectSubmixPreset* SubmixEffectPreset);
+
+	/** Removes the submix effect at the given submix chain index, if there is a submix effect at that index. */
+	UFUNCTION(BlueprintCallable, Category = "Audio|Effects", meta = (WorldContext = "WorldContextObject"))
+	static void RemoveSubmixEffectPresetAtIndex(const UObject* WorldContextObject, USoundSubmix* SoundSubmix, int32 SubmixChainIndex);
+
+	/** Replaces the submix effect at the given submix chain index, adds the effect if there is none at that index. */
+	UFUNCTION(BlueprintCallable, Category = "Audio|Effects", meta = (WorldContext = "WorldContextObject"))
+	static void ReplaceSoundEffectSubmix(const UObject* WorldContextObject, USoundSubmix* InSoundSubmix, int32 SubmixChainIndex, USoundEffectSubmixPreset* SubmixEffectPreset);
+
+	/** Clears all submix effects on the given submix. */
+	UFUNCTION(BlueprintCallable, Category = "Audio|Effects", meta = (WorldContext = "WorldContextObject"))
+	static void ClearSubmixEffects(const UObject* WorldContextObject, USoundSubmix* SoundSubmix);
+
 	/** Start recording audio. By leaving the Submix To Record field blank, you can record the master output of the game. */
 	UFUNCTION(BlueprintCallable, Category = "Audio|Recording", meta = (WorldContext = "WorldContextObject", AdvancedDisplay = 1))
 	static void StartRecordingOutput(const UObject* WorldContextObject, float ExpectedDuration, USoundSubmix* SubmixToRecord = nullptr);
diff --git a/Engine/Source/Runtime/AudioMixer/Public/AudioMixerDevice.h b/Engine/Source/Runtime/AudioMixer/Public/AudioMixerDevice.h
--- a/Engine/Source/Runtime/AudioMixer/Public/AudioMixerDevice.h
+++ b/Engine/Source/Runtime/AudioMixer/Public/AudioMixerDevice.h
@@ -195,6 +195,21 @@ namespace Audio
 		// Clear all submix effects from master submix
 		void ClearMasterSubmixEffects();
 
+		// Add submix effect to given submix
+		int32 AddSubmixEffect(USoundSubmix* InSoundSubmix, uint32 SubmixEffectId, FSoundEffectSubmixPtr SoundEffect);
+
+		// Remove submix effect to given submix
+		void RemoveSubmixEffect(USoundSubmix* InSoundSubmix, uint32 SubmixEffectId);
+
+		// Remove submix effect at the given submix chain index
+		void RemoveSubmixEffectAtIndex(USoundSubmix* InSoundSubmix, int32 SubmixChainIndex);
+
+		// Replace the submix effect of the given submix at the submix chain index with the new submix effect id and submix instance
+		void ReplaceSoundEffectSubmix(USoundSubmix* InSoundSubmix, int32 InSubmixChainIndex, int32 SubmixEffectId, FSoundEffectSubmixPtr SoundEffect);
+
+		// Clear all submix effects from given submix
+		void ClearSubmixEffects(USoundSubmix* InSoundSubmix);
+
 		// Returns the channel array for the given submix channel type
 		const TArray<EAudioMixerChannel::Type>& GetChannelArray() const;
 
diff --git a/Engine/Source/Runtime/AudioMixer/Public/AudioMixerSubmix.h b/Engine/Source/Runtime/AudioMixer/Public/AudioMixerSubmix.h
--- a/Engine/Source/Runtime/AudioMixer/Public/AudioMixerSubmix.h
+++ b/Engine/Source/Runtime/AudioMixer/Public/AudioMixerSubmix.h
@@ -107,6 +107,9 @@ namespace Audio
 		// Returns the number of wet effects in this submix.
 		int32 GetNumEffects() const;
 
+		// Returns the size of the submix chain. 
+		int32 GetSizeOfSubmixChain() const;
+
 		// Add (if not already added) or sets the amount of the source voice's send amount
 		void AddOrSetSourceVoice(FMixerSourceVoice* InSourceVoice, const float SendLevel);
 
@@ -119,6 +122,9 @@ namespace Audio
 		/** Removes the submix effect from the effect submix chain. */
 		void RemoveSoundEffectSubmix(uint32 SubmixPresetId);
 
+		/** Removes the submix effect from the effect submix chain at the given submix index. */
+		void RemoveSoundEffectSubmixAtIndex(int32 InIndex);
+
 		/** Clears all submix effects from the effect submix chain. */
 		void ClearSoundEffectSubmixes();
 
@@ -460,6 +466,9 @@ namespace Audio
 		// This buffer is used for recorded output of the submix.
 		AlignedFloatBuffer RecordingData;
 
+		// Returns the number of submix effects
+		int32 NumSubmixEffects;
+
 		// Bool set to true when this submix is recording data.
 		uint8 bIsRecording : 1;
 
diff --git a/Engine/Source/Runtime/Engine/Classes/Sound/SoundAttenuation.h b/Engine/Source/Runtime/Engine/Classes/Sound/SoundAttenuation.h
--- a/Engine/Source/Runtime/Engine/Classes/Sound/SoundAttenuation.h
+++ b/Engine/Source/Runtime/Engine/Classes/Sound/SoundAttenuation.h
@@ -6,6 +6,8 @@
 #include "IAudioExtensionPlugin.h"
 #include "SoundAttenuation.generated.h"
 
+class USoundSubmix;
+
 // This enumeration is deprecated
 UENUM()
 enum ESoundDistanceCalc
@@ -51,6 +53,19 @@ enum class EReverbSendMethod : uint8
 	Manual,
 };
 
+UENUM(BlueprintType)
+enum class ESubmixSendMethod : uint8
+{
+	// A submix send based on linear interpolation between a distance range and send-level range
+	Linear,
+
+	// A submix send based on a supplied curve
+	CustomCurve,
+
+	// A manual submix send level (Uses the specified constant send level value. Useful for 2D sounds.)
+	Manual,
+};
+
 
 UENUM(BlueprintType)
 enum class EPriorityAttenuationMethod : uint8
@@ -84,6 +99,44 @@ struct ENGINE_API FSoundAttenuationPluginSettings
 	TArray<UReverbPluginSourceSettingsBase*> ReverbPluginSettingsArray;
 };
 
+USTRUCT(BlueprintType)
+struct ENGINE_API FAttenuationSubmixSendSettings
+{
+	GENERATED_USTRUCT_BODY()
+
+	/** Submix to send audio to based on distance. */
+	UPROPERTY(EditAnywhere, BlueprintReadWrite, Category = AttenuationSubmixSend)
+	USoundSubmix* Submix = nullptr;
+
+	/** What method to use to use for submix sends. */
+	UPROPERTY(EditAnywhere, BlueprintReadWrite, Category = AttenuationReverbSend)
+	ESubmixSendMethod SubmixSendMethod = ESubmixSendMethod::Linear;
+
+	/** The amount to send to the Submix when the sound is located at a distance equal to value specified in the submix send distance min. */
+	UPROPERTY(EditAnywhere, BlueprintReadWrite, Category = AttenuationSubmixSend, meta = (DisplayName = "Submix Min Send Level"))
+	float SubmixSendLevelMin = 0.0f;
+
+	/** The amount to send to the Submix when the sound is located at a distance equal to value specified in the reverb max send distance. */
+	UPROPERTY(EditAnywhere, BlueprintReadWrite, Category = AttenuationSubmixSend, meta = (DisplayName = "Submix Max Send Level"))
+	float SubmixSendLevelMax = 1.0f;
+
+	/** The min distance to send to the Submix. */
+	UPROPERTY(EditAnywhere, BlueprintReadWrite, Category = AttenuationSubmixSend, meta = (DisplayName = "Submix Min Send Distance"))
+	float SubmixSendDistanceMin = 400.0f;
+
+	/** The max distance to send to the Submix. */
+	UPROPERTY(EditAnywhere, BlueprintReadWrite, Category = AttenuationSubmixSend, meta = (DisplayName = "Submix Max Send Distance"))
+	float SubmixSendDistanceMax = 6000.0f;
+
+	/* The manual Submix send level to use. Doesn't change as a function of distance. */
+	UPROPERTY(EditAnywhere, BlueprintReadWrite, Category = AttenuationSubmixSend)
+	float ManualSubmixSendLevel = 0.2f;
+
+	/* The custom Submix send curve to use for distance-based send level. */
+	UPROPERTY(EditAnywhere, BlueprintReadWrite, Category = AttenuationSubmixSend)
+	FRuntimeFloatCurve CustomSubmixSendCurve;
+};
+
 /*
 The settings for attenuating.
 */
@@ -136,6 +189,10 @@ struct ENGINE_API FSoundAttenuationSettings : public FBaseAttenuationSettings
 	UPROPERTY(EditAnywhere, BlueprintReadWrite, Category = AttenuationAirAbsorption, meta = (DisplayName = "Enable Log Frequency Scaling"))
 	uint8 bEnableLogFrequencyScaling : 1;
 
+	/** Enables submix sends based on distance. */
+	UPROPERTY(EditAnywhere, BlueprintReadWrite, Category = AttenuationSubmixSend, meta = (DisplayName = "Enable Submix Send"))
+	uint8 bEnableSubmixSends : 1;
+
 	/** What method we use to spatialize the sound. */
 	UPROPERTY(EditAnywhere, BlueprintReadWrite, Category = AttenuationSpatialization, meta = (ClampMin = "0", EditCondition = "bSpatialize", DisplayName = "Spatialization Method"))
 	TEnumAsByte<enum ESoundSpatializationAlgorithm> SpatializationAlgorithm;
@@ -300,6 +357,10 @@ struct ENGINE_API FSoundAttenuationSettings : public FBaseAttenuationSettings
 	UPROPERTY(EditAnywhere, BlueprintReadWrite, Category = AttenuationReverbSend)
 	FRuntimeFloatCurve CustomReverbSendCurve;
 
+	/** Set of submix send settings to use to send audio to submixes as a function of distance. */
+	UPROPERTY(EditAnywhere, BlueprintReadWrite, Category = AttenuationSubmixSend)
+	TArray<FAttenuationSubmixSendSettings> SubmixSendSettings;
+
 	/** The priority attenuation to use when the sound is at the minimum priority attenuation distance. */
 	UPROPERTY(EditAnywhere, BlueprintReadWrite, Category = AttenuationPriority, meta = (DisplayName = "Min Priority Attenuation"))
 	float PriorityAttenuationMin;
@@ -340,6 +401,7 @@ struct ENGINE_API FSoundAttenuationSettings : public FBaseAttenuationSettings
 		, bEnablePriorityAttenuation(false)
 		, bApplyNormalizationToStereoSounds(false)
 		, bEnableLogFrequencyScaling(false)
+		, bEnableSubmixSends(false)
 		, SpatializationAlgorithm(ESoundSpatializationAlgorithm::SPATIALIZATION_Default)
 		, BinauralRadius(0.0f)
 		, AbsorptionMethod(EAirAbsorptionMethod::Linear)
diff --git a/Engine/Source/Runtime/Engine/Private/ActiveSound.cpp b/Engine/Source/Runtime/Engine/Private/ActiveSound.cpp
--- a/Engine/Source/Runtime/Engine/Private/ActiveSound.cpp
+++ b/Engine/Source/Runtime/Engine/Private/ActiveSound.cpp
@@ -1541,6 +1541,18 @@ void FActiveSound::UpdateAttenuation(float DeltaTime, FSoundParseParameters& Par
 		ParseParams.ReverbSendLevelDistanceRange = { Settings->ReverbDistanceMin, Settings->ReverbDistanceMax };
 	}
 
+	if (Settings->bEnableSubmixSends)
+	{
+		ParseParams.SubmixSendSettings.Reset();
+		for (const FAttenuationSubmixSendSettings& SendSettings : Settings->SubmixSendSettings)
+		{
+			if (SendSettings.Submix)
+			{
+				ParseParams.SubmixSendSettings.Add(SendSettings);
+			}
+		}
+	}
+
 	check(AudioDevice);
 	FAttenuationListenerData ListenerData = FAttenuationListenerData::Create(*AudioDevice, ListenerIndex, ParseParams.Transform, *Settings);
 
diff --git a/Engine/Source/Runtime/Engine/Private/Audio.cpp b/Engine/Source/Runtime/Engine/Private/Audio.cpp
--- a/Engine/Source/Runtime/Engine/Private/Audio.cpp
+++ b/Engine/Source/Runtime/Engine/Private/Audio.cpp
@@ -954,6 +954,14 @@ void FWaveInstance::AddReferencedObjects( FReferenceCollector& Collector )
 		}
 	}
 
+	for (FAttenuationSubmixSendSettings& SubmixSend : SubmixSendSettings)
+	{
+		if (SubmixSend.Submix)
+		{
+			Collector.AddReferencedObject(SubmixSend.Submix);
+		}
+	}
+
 	Collector.AddReferencedObject( SoundClass );
 	NotifyBufferFinishedHooks.AddReferencedObjects( Collector );
 }
diff --git a/Engine/Source/Runtime/Engine/Private/SoundAttenuation.cpp b/Engine/Source/Runtime/Engine/Private/SoundAttenuation.cpp
--- a/Engine/Source/Runtime/Engine/Private/SoundAttenuation.cpp
+++ b/Engine/Source/Runtime/Engine/Private/SoundAttenuation.cpp
@@ -107,7 +107,8 @@ bool FSoundAttenuationSettings::operator==(const FSoundAttenuationSettings& Othe
 			&& HPFFrequencyAtMax		== Other.HPFFrequencyAtMax
 			&& HPFFrequencyAtMin		== Other.HPFFrequencyAtMin
 			&& bEnableLogFrequencyScaling == Other.bEnableLogFrequencyScaling
-			&& bEnableListenerFocus == Other.bEnableListenerFocus
+			&& bEnableSubmixSends 		== Other.bEnableSubmixSends
+			&& bEnableListenerFocus 	== Other.bEnableListenerFocus
 			&& FocusAzimuth				== Other.FocusAzimuth
 			&& NonFocusAzimuth			== Other.NonFocusAzimuth
 			&& FocusDistanceScale		== Other.FocusDistanceScale
diff --git a/Engine/Source/Runtime/Engine/Private/SoundWave.cpp b/Engine/Source/Runtime/Engine/Private/SoundWave.cpp
--- a/Engine/Source/Runtime/Engine/Private/SoundWave.cpp
+++ b/Engine/Source/Runtime/Engine/Private/SoundWave.cpp
@@ -1891,6 +1891,9 @@ void USoundWave::Parse(FAudioDevice* AudioDevice, const UPTRINT NodeWaveInstance
 	WaveInstance->ReverbSendLevelRange = ParseParams.ReverbSendLevelRange;
 	WaveInstance->ReverbSendLevelDistanceRange = ParseParams.ReverbSendLevelDistanceRange;
 
+	// Copy the submix send settings
+	WaveInstance->SubmixSendSettings = ParseParams.SubmixSendSettings;
+
 	// Get the envelope follower settings
 	WaveInstance->EnvelopeFollowerAttackTime = ParseParams.EnvelopeFollowerAttackTime;
 	WaveInstance->EnvelopeFollowerReleaseTime = ParseParams.EnvelopeFollowerReleaseTime;
diff --git a/Engine/Source/Runtime/Engine/Public/ActiveSound.h b/Engine/Source/Runtime/Engine/Public/ActiveSound.h
--- a/Engine/Source/Runtime/Engine/Public/ActiveSound.h
+++ b/Engine/Source/Runtime/Engine/Public/ActiveSound.h
@@ -141,6 +141,9 @@ struct FSoundParseParameters
 	float ManualReverbSendLevel;
 	FRuntimeFloatCurve CustomReverbSendCurve;
 
+	// Submix send params to use for this sound
+	TArray<FAttenuationSubmixSendSettings> SubmixSendSettings;
+
 	// The distance between left and right channels when spatializing stereo assets
 	float StereoSpread;
 
diff --git a/Engine/Source/Runtime/Engine/Public/Audio.h b/Engine/Source/Runtime/Engine/Public/Audio.h
--- a/Engine/Source/Runtime/Engine/Public/Audio.h
+++ b/Engine/Source/Runtime/Engine/Public/Audio.h
@@ -365,6 +365,9 @@ public:
 	/** The manual send level to use if the sound is set to use manual send level. */
 	float ManualReverbSendLevel;
 
+	/** The submix send settings to use. */
+	TArray<FAttenuationSubmixSendSettings> SubmixSendSettings;
+
 	/** Cached type hash */
 	uint32 TypeHash;
 
