Bring over crash fix for retained sound cues in editor.

diff --git a/Engine/Source/Runtime/Engine/Private/SoundWave.cpp b/Engine/Source/Runtime/Engine/Private/SoundWave.cpp
--- a/Engine/Source/Runtime/Engine/Private/SoundWave.cpp
+++ b/Engine/Source/Runtime/Engine/Private/SoundWave.cpp
@@ -2475,7 +2475,7 @@ void USoundWave::RetainCompressedAudio(bool bForceSync /*= false*/)
 {
 	// Since the zeroth chunk is always inlined and stored in memory,
 	// early exit if we only have one chunk.
-	if (DisableRetainingCVar || GetNumChunks() <= 1)
+	if (GIsEditor || DisableRetainingCVar || GetNumChunks() <= 1)
 	{
 		return;
 	}
