fix for sound cues with a soundclass set to retain not playing back in preview. Issue was CPL chaining causing audio to be cooked before the RawData bulkdata was properly initialized. Not sure how this occurs, and it might point to something more insidious happening between RawData init and CPL, but avoiding CPL in editor fixes issue, and isn't needed since we don't retain in editor/PIE anyways.

diff --git a/Engine/Source/Runtime/Engine/Private/SoundCue.cpp b/Engine/Source/Runtime/Engine/Private/SoundCue.cpp
--- a/Engine/Source/Runtime/Engine/Private/SoundCue.cpp
+++ b/Engine/Source/Runtime/Engine/Private/SoundCue.cpp
@@ -195,7 +195,7 @@ void USoundCue::PostLoad()
 		CurrentSoundClass = CurrentSoundClass->ParentClass;
 	}
 
-	if (SoundClassLoadingBehavior == ESoundWaveLoadingBehavior::RetainOnLoad)
+	if (!GIsEditor && SoundClassLoadingBehavior == ESoundWaveLoadingBehavior::RetainOnLoad)
 	{
 		RetainSoundCue();
 	}
diff --git a/Engine/Source/Runtime/Engine/Private/SoundNode.cpp b/Engine/Source/Runtime/Engine/Private/SoundNode.cpp
--- a/Engine/Source/Runtime/Engine/Private/SoundNode.cpp
+++ b/Engine/Source/Runtime/Engine/Private/SoundNode.cpp
@@ -122,7 +122,7 @@ void USoundNode::PrimeChildWavePlayers(bool bRecurse)
 
 void USoundNode::RetainChildWavePlayers(bool bRecurse)
 {
-	if (FPlatformCompressionUtilities::IsCurrentPlatformUsingStreamCaching())
+	if (!GIsEditor && FPlatformCompressionUtilities::IsCurrentPlatformUsingStreamCaching())
 	{
 		// Search child nodes for wave players, then prime their waves.
 		for (USoundNode* ChildNode : ChildNodes)
diff --git a/Engine/Source/Runtime/Engine/Private/SoundNodeSoundClass.cpp b/Engine/Source/Runtime/Engine/Private/SoundNodeSoundClass.cpp
--- a/Engine/Source/Runtime/Engine/Private/SoundNodeSoundClass.cpp
+++ b/Engine/Source/Runtime/Engine/Private/SoundNodeSoundClass.cpp
@@ -41,7 +41,7 @@ void USoundNodeSoundClass::PostLoad()
 		CurrentSoundClass = CurrentSoundClass->ParentClass;
 	}
 
-	if (SoundClassLoadingBehavior == ESoundWaveLoadingBehavior::RetainOnLoad)
+	if (!GIsEditor && SoundClassLoadingBehavior == ESoundWaveLoadingBehavior::RetainOnLoad)
 	{
 		RetainChildWavePlayers(true);
 		bRetainingAudioDueToSoundClass = true;
