Fix crash when adding a LOD after reloading an asset
#rnx

diff --git a/Engine/Source/Runtime/Engine/Classes/Engine/SkeletalMeshEditorData.h b/Engine/Source/Runtime/Engine/Classes/Engine/SkeletalMeshEditorData.h
--- a/Engine/Source/Runtime/Engine/Classes/Engine/SkeletalMeshEditorData.h
+++ b/Engine/Source/Runtime/Engine/Classes/Engine/SkeletalMeshEditorData.h
@@ -54,7 +54,7 @@ public:
 
 private:
 	/** Imported raw mesh data. Optional, only the imported mesh LOD has this, generated LOD or old asset will be empty. */
-	TArray<FRawSkeletalMeshBulkData> RawSkeletalMeshBulkDatas;
+	TArray<TSharedRef<FRawSkeletalMeshBulkData>> RawSkeletalMeshBulkDatas;
 
 
 #endif //WITH_EDITORONLY_DATA
diff --git a/Engine/Source/Runtime/Engine/Private/Rendering/SkeletalMeshLODImporterData.cpp b/Engine/Source/Runtime/Engine/Private/Rendering/SkeletalMeshLODImporterData.cpp
--- a/Engine/Source/Runtime/Engine/Private/Rendering/SkeletalMeshLODImporterData.cpp
+++ b/Engine/Source/Runtime/Engine/Private/Rendering/SkeletalMeshLODImporterData.cpp
@@ -688,7 +688,7 @@ FArchive& operator<<(FArchive& Ar, FSkeletalMeshImportData& RawMesh)
 	return Ar;
 }
 
-void FRawSkeletalMeshBulkData::Serialize(FArchive& Ar, TArray<FRawSkeletalMeshBulkData>& RawSkeltalMeshBulkDatas, UObject* Owner)
+void FRawSkeletalMeshBulkData::Serialize(FArchive& Ar, TArray<TSharedRef<FRawSkeletalMeshBulkData>>& RawSkeltalMeshBulkDatas, UObject* Owner)
 {
 	Ar.CountBytes(RawSkeltalMeshBulkDatas.Num() * sizeof(FRawSkeletalMeshBulkData), RawSkeltalMeshBulkDatas.Num() * sizeof(FRawSkeletalMeshBulkData));
 	if (Ar.IsLoading())
@@ -699,9 +699,9 @@ void FRawSkeletalMeshBulkData::Serialize(FArchive& Ar, TArray<FRawSkeletalMeshBu
 		RawSkeltalMeshBulkDatas.Empty(NewNum);
 		for (int32 Index = 0; Index < NewNum; Index++)
 		{
-			int32 NewEntryIndex = RawSkeltalMeshBulkDatas.AddDefaulted();
+			int32 NewEntryIndex = RawSkeltalMeshBulkDatas.Add(MakeShared<FRawSkeletalMeshBulkData>());
 			check(NewEntryIndex == Index);
-			RawSkeltalMeshBulkDatas[Index].Serialize(Ar, Owner);
+			RawSkeltalMeshBulkDatas[Index].Get().Serialize(Ar, Owner);
 		}
 	}
 	else
@@ -711,7 +711,7 @@ void FRawSkeletalMeshBulkData::Serialize(FArchive& Ar, TArray<FRawSkeletalMeshBu
 		Ar << Num;
 		for (int32 Index = 0; Index < Num; Index++)
 		{
-			RawSkeltalMeshBulkDatas[Index].Serialize(Ar, Owner);
+			RawSkeltalMeshBulkDatas[Index].Get().Serialize(Ar, Owner);
 		}
 	}
 }
diff --git a/Engine/Source/Runtime/Engine/Private/SkeletalMeshEditorData.cpp b/Engine/Source/Runtime/Engine/Private/SkeletalMeshEditorData.cpp
--- a/Engine/Source/Runtime/Engine/Private/SkeletalMeshEditorData.cpp
+++ b/Engine/Source/Runtime/Engine/Private/SkeletalMeshEditorData.cpp
@@ -35,13 +35,17 @@ FRawSkeletalMeshBulkData& USkeletalMeshEditorData::GetLODImportedData(int32 LODI
 		//The allocation must be done before going multi thread
 		//TArray is not thread safe when allocating
 		check(IsInGameThread());
-
-		int32 AddItemCount = 1 + (LODIndex - RawSkeletalMeshBulkDatas.Num());
-		RawSkeletalMeshBulkDatas.AddDefaulted(AddItemCount);
+		const int32 AddItemCount = 1 + (LODIndex - RawSkeletalMeshBulkDatas.Num());
+		RawSkeletalMeshBulkDatas.Reserve(AddItemCount);
+		while (LODIndex >= RawSkeletalMeshBulkDatas.Num())
+		{
+			RawSkeletalMeshBulkDatas.Add(MakeShared<FRawSkeletalMeshBulkData>());
+		}
 	}
+
 	check(RawSkeletalMeshBulkDatas.IsValidIndex(LODIndex));
 	//Return the Data
-	return RawSkeletalMeshBulkDatas[LODIndex];
+	return RawSkeletalMeshBulkDatas[LODIndex].Get();
 }
 
 bool USkeletalMeshEditorData::IsLODImportDataValid(int32 LODIndex)
diff --git a/Engine/Source/Runtime/Engine/Public/Rendering/SkeletalMeshLODImporterData.h b/Engine/Source/Runtime/Engine/Public/Rendering/SkeletalMeshLODImporterData.h
--- a/Engine/Source/Runtime/Engine/Public/Rendering/SkeletalMeshLODImporterData.h
+++ b/Engine/Source/Runtime/Engine/Public/Rendering/SkeletalMeshLODImporterData.h
@@ -483,7 +483,7 @@ public:
 	ENGINE_API FRawSkeletalMeshBulkData();
 
 	/*Static function helper to serialize an array of Raw source data. */
-	ENGINE_API static void Serialize(FArchive& Ar, TArray<FRawSkeletalMeshBulkData>& RawSkeltalMeshBulkDatas, UObject* Owner);
+	ENGINE_API static void Serialize(FArchive& Ar, TArray<TSharedRef<FRawSkeletalMeshBulkData>>& RawSkeltalMeshBulkDatas, UObject* Owner);
 	
 	/** Serialization. */
 	ENGINE_API void Serialize(class FArchive& Ar, class UObject* Owner);
