Remove error when using SceneDepth with scene texture node on mobile.
#rnx

diff --git a/Engine/Source/Runtime/Engine/Private/Materials/HLSLMaterialTranslator.cpp b/Engine/Source/Runtime/Engine/Private/Materials/HLSLMaterialTranslator.cpp
--- a/Engine/Source/Runtime/Engine/Private/Materials/HLSLMaterialTranslator.cpp
+++ b/Engine/Source/Runtime/Engine/Private/Materials/HLSLMaterialTranslator.cpp
@@ -4741,17 +4741,6 @@ int32 FHLSLMaterialTranslator::SceneTextureLookup(int32 ViewportUV, uint32 InSce
 	else // mobile
 	{
 		int32 UV = BufferUV;
-		if (Material->GetMaterialDomain() == MD_PostProcess)
-		{
-			int32 BlendableLocation = Material->GetBlendableLocation();
-			if (SceneTextureId == PPI_SceneDepth && !(BlendableLocation == BL_BeforeTranslucency || BlendableLocation == BL_BeforeTonemapping))
-			{
-				// SceneDepth lookups are not available when using MSAA, but we can access depth stored in SceneColor.A channel
-				// SceneColor.A channel holds depth till BeforeTonemapping location, then it's gets overwritten
-				return Errorf(TEXT("SceneDepth lookups are only available when BlendableLocation is BeforeTranslucency or BeforeTonemapping"));
-			}
-		}
-			
 		return AddCodeChunk(MCT_Float4,	TEXT("MobileSceneTextureLookup(Parameters, %d, %s)"), (int32)SceneTextureId, *CoerceParameter(UV, MCT_Float2));
 	}
 }
