This change modifies USkeletalMeshComponent::OnUpdateTransform to disallow deferred kinematic updates for motion which has propagated from a parent.
Here's why:
-> Animation from the skeletal mesh happens during: TG_PrePhysics
-> Deferred kinematic updates are applied during: TG_StartPhysics
-> Propagation from the parent movement happens during: TG_EndPhysics (for physics objects... could in theory be from any tick group)
Therefore, deferred kinematic updates are safe from animation, but not from parent movement.

diff --git a/Engine/Source/Runtime/Engine/Private/SkeletalMeshComponentPhysics.cpp b/Engine/Source/Runtime/Engine/Private/SkeletalMeshComponentPhysics.cpp
--- a/Engine/Source/Runtime/Engine/Private/SkeletalMeshComponentPhysics.cpp
+++ b/Engine/Source/Runtime/Engine/Private/SkeletalMeshComponentPhysics.cpp
@@ -1308,7 +1308,14 @@ void USkeletalMeshComponent::OnUpdateTransform(EUpdateTransformFlags UpdateTrans
 	if(bPhysicsStateCreated && !(UpdateTransformFlags&EUpdateTransformFlags::SkipPhysicsUpdate))
 	{
 #if !OLD_FORCE_UPDATE_BEHAVIOR
-		UpdateKinematicBonesToAnim(GetComponentSpaceTransforms(), Teleport, false);
+		// Animation from the skeletal mesh happens during TG_PrePhysics.
+		// Deferred kinematic updates are applied during TG_StartPhysics.
+		// Propagation from the parent movement happens during TG_EndPhysics (for physics objects... could in theory be from any tick group).
+		// Therefore, deferred kinematic updates are safe from animation, but not from parent movement.
+		const EAllowKinematicDeferral AllowDeferral
+			= !!(UpdateTransformFlags & EUpdateTransformFlags::PropagateFromParent)
+			? EAllowKinematicDeferral::DisallowDeferral : EAllowKinematicDeferral::AllowDeferral;
+		UpdateKinematicBonesToAnim(GetComponentSpaceTransforms(), Teleport, false, AllowDeferral);
 #else
 		UpdateKinematicBonesToAnim(GetComponentSpaceTransforms(), ETeleportType::TeleportPhysics, false);
 #endif
