Better handling projects that leave defaults for EffectsQuality.
Also now looking in ScalabilityGroups and SystemSettings for ini file changes to EffectsQuality.

diff --git a/Engine/Plugins/FX/Niagara/Source/Niagara/Private/NiagaraModule.cpp b/Engine/Plugins/FX/Niagara/Source/Niagara/Private/NiagaraModule.cpp
--- a/Engine/Plugins/FX/Niagara/Source/Niagara/Private/NiagaraModule.cpp
+++ b/Engine/Plugins/FX/Niagara/Source/Niagara/Private/NiagaraModule.cpp
@@ -151,6 +151,7 @@ void INiagaraModule::StartupModule()
 	FNiagaraWorldManager::OnStartup();
 
 	//Force update on module load.
+	EngineEffectsQuality = Scalability::DefaultQualityLevel;
 	static const auto CVarEQ = IConsoleManager::Get().FindTConsoleVariableDataInt(TEXT("sg.EffectsQuality"));
 	OnEffectsQualityChanged(CVarEQ->GetValueOnGameThread(), true);
 
diff --git a/Engine/Plugins/FX/Niagara/Source/Niagara/Private/NiagaraPlatformSet.cpp b/Engine/Plugins/FX/Niagara/Source/Niagara/Private/NiagaraPlatformSet.cpp
--- a/Engine/Plugins/FX/Niagara/Source/Niagara/Private/NiagaraPlatformSet.cpp
+++ b/Engine/Plugins/FX/Niagara/Source/Niagara/Private/NiagaraPlatformSet.cpp
@@ -446,16 +446,21 @@ FNiagaraPlatformSet::FPlatformIniSettings FNiagaraPlatformSet::GetPlatformIniSet
 
 	auto FindCVarValue = [&](const TCHAR* Section, const TCHAR* CVarName, int32& OutVal)
 	{
+		bool bFound = true;
 		if (!ScalabilitySettings.GetInt(Section, CVarName, OutVal))
 		{
 			if (!GameSettings.GetInt(Section, CVarName, OutVal))
 			{
-				EngineSettings.GetInt(Section, CVarName, OutVal);
+				if (!EngineSettings.GetInt(Section, CVarName, OutVal))
+				{
+					bFound = false;
+				}
 			}
 		}
+		return bFound;
 	};
 
-	int32 EffectsQuality = INDEX_NONE;
+	int32 EffectsQuality = Scalability::DefaultQualityLevel;
 	int32 EQMask = INDEX_NONE;
 
 	//If this platform can change quality settings at runtime then we return the full mask.
@@ -469,7 +474,10 @@ FNiagaraPlatformSet::FPlatformIniSettings FNiagaraPlatformSet::GetPlatformIniSet
 	}
 	else
 	{
-		FindCVarValue(TEXT("ScalabilityGroups"), TEXT("sg.EffectsQuality"), EffectsQuality);
+		if (!FindCVarValue(TEXT("ScalabilityGroups"), TEXT("sg.EffectsQuality"), EffectsQuality))
+		{
+			FindCVarValue(TEXT("SystemSettings"), TEXT("sg.EffectsQuality"), EffectsQuality);
+		}
 		EQMask = CreateEQMask(EffectsQuality);
 	}
 
diff --git a/Engine/Source/Runtime/Engine/Private/Scalability.cpp b/Engine/Source/Runtime/Engine/Private/Scalability.cpp
--- a/Engine/Source/Runtime/Engine/Private/Scalability.cpp
+++ b/Engine/Source/Runtime/Engine/Private/Scalability.cpp
@@ -19,42 +19,42 @@ static TAutoConsoleVariable<float> CVarResolutionQuality(
 
 static TAutoConsoleVariable<int32> CVarViewDistanceQuality(
 	TEXT("sg.ViewDistanceQuality"),
-	3,
+	Scalability::DefaultQualityLevel,
 	TEXT("Scalability quality state (internally used by scalability system, ini load/save or using SCALABILITY console command)\n")
 	TEXT(" 0:low, 1:med, 2:high, 3:epic, 4:cinematic, default: 3"),
 	ECVF_ScalabilityGroup);
 
 static TAutoConsoleVariable<int32> CVarAntiAliasingQuality(
 	TEXT("sg.AntiAliasingQuality"),
-	3,
+	Scalability::DefaultQualityLevel,
 	TEXT("Scalability quality state (internally used by scalability system, ini load/save or using SCALABILITY console command)\n")
 	TEXT(" 0:low, 1:med, 2:high, 3:epic, 4:cinematic, default: 3"),
 	ECVF_ScalabilityGroup);
 
 static TAutoConsoleVariable<int32> CVarShadowQuality(
 	TEXT("sg.ShadowQuality"),
-	3,
+	Scalability::DefaultQualityLevel,
 	TEXT("Scalability quality state (internally used by scalability system, ini load/save or using SCALABILITY console command)\n")
 	TEXT(" 0:low, 1:med, 2:high, 3:epic, 4:cinematic, default: 3"),
 	ECVF_ScalabilityGroup);
 
 static TAutoConsoleVariable<int32> CVarPostProcessQuality(
 	TEXT("sg.PostProcessQuality"),
-	3,
+	Scalability::DefaultQualityLevel,
 	TEXT("Scalability quality state (internally used by scalability system, ini load/save or using SCALABILITY console command)\n")
 	TEXT(" 0:low, 1:med, 2:high, 3:epic, 4:cinematic, default: 3"),
 	ECVF_ScalabilityGroup);
 
 static TAutoConsoleVariable<int32> CVarTextureQuality(
 	TEXT("sg.TextureQuality"),
-	3,
+	Scalability::DefaultQualityLevel,
 	TEXT("Scalability quality state (internally used by scalability system, ini load/save or using SCALABILITY console command)\n")
 	TEXT(" 0:low, 1:med, 2:high, 3:epic, 4:cinematic, default: 3"),
 	ECVF_ScalabilityGroup);
 
 static TAutoConsoleVariable<int32> CVarEffectsQuality(
 	TEXT("sg.EffectsQuality"),
-	3,
+	Scalability::DefaultQualityLevel,
 	TEXT("Scalability quality state (internally used by scalability system, ini load/save or using SCALABILITY console command)\n")
 	TEXT(" 0:low, 1:med, 2:high, 3:epic, 4:cinematic, default: 3"),
 	ECVF_ScalabilityGroup);
@@ -68,7 +68,7 @@ static TAutoConsoleVariable<int32> CVarFoliageQuality(
 
 static TAutoConsoleVariable<int32> CVarShadingQuality(
 	TEXT("sg.ShadingQuality"),
-	3,
+	Scalability::DefaultQualityLevel,
 	TEXT("Scalability quality state (internally used by scalability system, ini load/save or using SCALABILITY console command)\n")
 	TEXT(" 0:low, 1:med, 2:high, 3:epic, 4:cinematic, default: 3"),
 	ECVF_ScalabilityGroup);
diff --git a/Engine/Source/Runtime/Engine/Public/Scalability.h b/Engine/Source/Runtime/Engine/Public/Scalability.h
--- a/Engine/Source/Runtime/Engine/Public/Scalability.h
+++ b/Engine/Source/Runtime/Engine/Public/Scalability.h
@@ -16,6 +16,8 @@ namespace Scalability
 		ERelativeToMax,
 	};
 
+	const int32 DefaultQualityLevel = 3;
+
 	/**
 	 * Structure for holding the state of the engine scalability groups
 	 * Actual engine state you can get though GetQualityLevels().
