Deprecate UGameUserSettings::GetSyncInterval() and add GetFramePace()

diff --git a/Engine/Source/Runtime/Engine/Classes/GameFramework/GameUserSettings.h b/Engine/Source/Runtime/Engine/Classes/GameFramework/GameUserSettings.h
--- a/Engine/Source/Runtime/Engine/Classes/GameFramework/GameUserSettings.h
+++ b/Engine/Source/Runtime/Engine/Classes/GameFramework/GameUserSettings.h
@@ -304,9 +304,14 @@ public:
 	static EWindowMode::Type GetDefaultWindowMode();
 
 	/** Gets the current vsync interval setting */
+	UE_DEPRECATED(4.25, "Please use GetFramePace to get the paced frame rate")
 	UFUNCTION(BlueprintPure, Category = Settings)
 	static int32 GetSyncInterval();
 
+	/** Gets the current frame pacing frame rate in fps, or 0 if none */
+	UFUNCTION(BlueprintPure, Category = Settings)
+	static int32 GetFramePace();
+
 	/** Loads the user .ini settings into GConfig */
 	static void LoadConfigIni(bool bForceReload = false);
 
@@ -514,7 +519,8 @@ protected:
 	/** Sets the frame rate limit CVar to the passed in value, 0.0 indicates no limit */
 	static void SetFrameRateLimitCVar(float InLimit);
 
-	/** Sets the sync interval limit CVar forc consoles 0-3 */
+	/** Sets the sync interval limit CVar for consoles 0-3 */
+	UE_DEPRECATED(4.25, "Please use FPlatformRHIFramePacer::SetFramePace to set the desired frame rate")
 	static void SetSyncIntervalCVar(int32 InInterval);
 
 	/** Sets the input latency mode 0 and 2 */
diff --git a/Engine/Source/Runtime/Engine/Private/GameUserSettings.cpp b/Engine/Source/Runtime/Engine/Private/GameUserSettings.cpp
--- a/Engine/Source/Runtime/Engine/Private/GameUserSettings.cpp
+++ b/Engine/Source/Runtime/Engine/Private/GameUserSettings.cpp
@@ -16,6 +16,7 @@
 #include "Sound/SoundCue.h"
 #include "AudioDevice.h"
 #include "DynamicResolutionState.h"
+#include "HAL/PlatformFramePacer.h"
 
 extern EWindowMode::Type GetWindowModeType(EWindowMode::Type WindowMode);
 
@@ -641,6 +642,11 @@ int32 UGameUserSettings::GetSyncInterval()
 	}
 }
 
+int32 UGameUserSettings::GetFramePace()
+{
+	return FPlatformRHIFramePacer::GetFramePace();
+}
+
 void UGameUserSettings::ResetToCurrentSettings()
 {
 	if ( GEngine && GEngine->GameViewport && GEngine->GameViewport->GetWindow().IsValid() )
