Cherry-pick CL, to ensure that work is not lost when UE is integrated:
'Ensure there is always a valid UWorld, in Editor, to request skin weight profiles from

diff --git a/Engine/Source/Runtime/Engine/Private/Animation/SkinWeightProfile.cpp b/Engine/Source/Runtime/Engine/Private/Animation/SkinWeightProfile.cpp
--- a/Engine/Source/Runtime/Engine/Private/Animation/SkinWeightProfile.cpp
+++ b/Engine/Source/Runtime/Engine/Private/Animation/SkinWeightProfile.cpp
@@ -19,7 +19,7 @@
 #include "UObject/WeakObjectPtrTemplates.h"
 #include "Engine/GameEngine.h"
 #include "UObject/ObjectVersion.h"
-
+#include "Engine/World.h"
 
 class ENGINE_API FSkinnedMeshComponentUpdateSkinWeightsContext
 {
@@ -373,15 +373,19 @@ void FSkinWeightProfilesData::SetDynamicDefaultSkinWeightProfile(USkeletalMesh*
 						};
 
 						UWorld* World = nullptr;
+#if WITH_EDITOR
+						World = GWorld;
+#else
 						UGameEngine* GameEngine = Cast<UGameEngine>(GEngine);
 						if (GameEngine)
 						{
 							World = GameEngine->GetGameWorld();
 						}
+#endif
 
 						if (World)
 						{
-							if (FSkinWeightProfileManager * Manager = FSkinWeightProfileManager::Get(World))
+							if (FSkinWeightProfileManager* Manager = FSkinWeightProfileManager::Get(World))
 							{
 								Manager->RequestSkinWeightProfile(Profiles[DefaultProfileIndex].Name, Mesh, Mesh, Callback, LODIndex);
 							}
