fix DebugExtension for generated shader files not being unique

diff --git a/Engine/Source/Runtime/Engine/Private/Materials/MaterialShader.cpp b/Engine/Source/Runtime/Engine/Private/Materials/MaterialShader.cpp
--- a/Engine/Source/Runtime/Engine/Private/Materials/MaterialShader.cpp
+++ b/Engine/Source/Runtime/Engine/Private/Materials/MaterialShader.cpp
@@ -1434,7 +1434,12 @@ void FMaterialShaderMap::Compile(
 			UE_LOG(LogShaders, Display, TEXT("	%s"), *WorkingDebugDescription);
 			NewContent->DebugDescription = *WorkingDebugDescription;
 
-			FString DebugExtension = FString::Printf( TEXT("_%08x%08x"), ShaderMapId.BaseMaterialId.A, ShaderMapId.BaseMaterialId.B);
+			FSHA1 IdParameterSetHash;
+			IdParameterSetHash.Reset();
+			IdParameterSetHash.UpdateWithString(*DebugDescription, DebugDescription.Len());
+			IdParameterSetHash.Final();
+			uint32* Hash = (uint32*)&IdParameterSetHash.m_digest[0];
+			FString DebugExtension = FString::Printf( TEXT("_%08x%08x"), Hash[0], Hash[1]);
 #else
 			FString DebugExtension = "";
 			FString WorkingDebugDescription = "";
