Fix an assert when quitting shooter game with memory mapped file access enabled.

TMaybeMappedAllocator was asserting if the caller tried to resize a memory mapped allocation, which makes sense, only in this case it was attempting to resize to 0 (ie clear it) which is safe to do for memory mapped regions, only it should just delete the region/handle instead.

#rnx

diff --git a/Engine/Source/Runtime/Engine/Public/Animation/AnimCompressionTypes.h b/Engine/Source/Runtime/Engine/Public/Animation/AnimCompressionTypes.h
--- a/Engine/Source/Runtime/Engine/Public/Animation/AnimCompressionTypes.h
+++ b/Engine/Source/Runtime/Engine/Public/Animation/AnimCompressionTypes.h
@@ -539,11 +539,20 @@ public:
 			SIZE_T NumBytesPerElement
 		)
 		{
-			// Avoid calling FMemory::Realloc( nullptr, 0 ) as ANSI C mandates returning a valid pointer which is not what we want.
-			if (Data || NumElements)
+			if (MappedRegion || MappedHandle)
+			{
+				check(NumElements == 0); // Currently we can only support resizing of memory mapped regions to 0 size (ie delete)
+
+				delete MappedRegion;
+				delete MappedHandle;
+				MappedRegion = nullptr;
+				MappedHandle = nullptr;
+				Data = nullptr; // make sure we don't try to free this pointer
+			}
+			else if (Data || NumElements)
 			{
-				check(!MappedHandle && !MappedRegion); // this could be supported, but it probably is never what you want, so we will just assert.
-					//checkSlow(((uint64)NumElements*(uint64)ElementTypeInfo.GetSize() < (uint64)INT_MAX));
+				// Avoid calling FMemory::Realloc( nullptr, 0 ) as ANSI C mandates returning a valid pointer which is not what we want.
+				//checkSlow(((uint64)NumElements*(uint64)ElementTypeInfo.GetSize() < (uint64)INT_MAX));
 				Data = (FScriptContainerElement*)FMemory::Realloc(Data, NumElements*NumBytesPerElement, Alignment);
 			}
 		}
