Adding editor time functionality to FSoundWavePCMLoader and minor bug fixe for IR asset import (Convolution Reverb)

diff --git a/Engine/Plugins/Runtime/Synthesis/Source/SynthesisEditor/Private/ImpulseResponse.cpp b/Engine/Plugins/Runtime/Synthesis/Source/SynthesisEditor/Private/ImpulseResponse.cpp
--- a/Engine/Plugins/Runtime/Synthesis/Source/SynthesisEditor/Private/ImpulseResponse.cpp
+++ b/Engine/Plugins/Runtime/Synthesis/Source/SynthesisEditor/Private/ImpulseResponse.cpp
@@ -40,6 +40,8 @@ UObject* UImpulseResponseFactory::FactoryCreateNew(UClass* Class, UObject* InPar
 	{
 		USoundWave* Wave = StagedSoundWave.Get();
 
+		constexpr bool bForceSynchronousLoad = true;
+
 		Loader.LoadSoundWave(Wave, [&](const USoundWave* SoundWave, const Audio::FSampleBuffer& LoadedSampleBuffer)
 		{
 			NewAsset->NumChannels = LoadedSampleBuffer.GetNumChannels();
@@ -65,7 +67,9 @@ UObject* UImpulseResponseFactory::FactoryCreateNew(UClass* Class, UObject* InPar
 				const int32 CurrOutputIndex = CurrChannel * NumFrames + CurrFrame;
 				OutputBuffer[CurrOutputIndex] = CurrSample;
 			}
-		});
+		}, bForceSynchronousLoad);
+
+		Loader.Update();
 
 		StagedSoundWave.Reset();
 	}
diff --git a/Engine/Source/Runtime/Engine/Classes/Sound/SampleBufferIO.h b/Engine/Source/Runtime/Engine/Classes/Sound/SampleBufferIO.h
--- a/Engine/Source/Runtime/Engine/Classes/Sound/SampleBufferIO.h
+++ b/Engine/Source/Runtime/Engine/Classes/Sound/SampleBufferIO.h
@@ -28,8 +28,8 @@ namespace Audio
 	public:
 		FSoundWavePCMLoader();
 
-		// Loads a USoundWave, call on game thread.
-		void LoadSoundWave(USoundWave* InSoundWave, TFunction<void(const USoundWave* SoundWave, const Audio::FSampleBuffer& OutSampleBuffer)> OnLoaded);
+		// Loads a USoundWave, call on game thread. Unless called with bSynchnous set to true, this class will require Update() to be called on the game thread.
+		void LoadSoundWave(USoundWave* InSoundWave, TFunction<void(const USoundWave* SoundWave, const Audio::FSampleBuffer& OutSampleBuffer)> OnLoaded, bool bSynchrounous = false);
 
 
 		// Update the loading state. 
diff --git a/Engine/Source/Runtime/Engine/Private/SampleBufferIO.cpp b/Engine/Source/Runtime/Engine/Private/SampleBufferIO.cpp
--- a/Engine/Source/Runtime/Engine/Private/SampleBufferIO.cpp
+++ b/Engine/Source/Runtime/Engine/Private/SampleBufferIO.cpp
@@ -15,7 +15,7 @@ namespace Audio
 	{
 	}
 
-	void FSoundWavePCMLoader::LoadSoundWave(USoundWave* InSoundWave, TFunction<void(const USoundWave* SoundWave, const Audio::FSampleBuffer& OutSampleBuffer)> OnLoaded)
+	void FSoundWavePCMLoader::LoadSoundWave(USoundWave* InSoundWave, TFunction<void(const USoundWave* SoundWave, const Audio::FSampleBuffer& OutSampleBuffer)> OnLoaded, bool bSynchronous /** = false */)
 	{
 		FAudioDeviceHandle AudioDevice = FAudioDevice::GetMainAudioDevice();
 
@@ -34,17 +34,31 @@ namespace Audio
 			LoadingSoundWaveInfo.Status = FLoadingSoundWaveInfo::LoadStatus::Loading;
 
 			// Kick off a decompression/precache of the sound wave
-			AudioDevice->Precache(InSoundWave, false, true, true);
+			AudioDevice->Precache(InSoundWave, bSynchronous, true, true);
 		}
 		else
 		{
 			LoadingSoundWaveInfo.Status = FLoadingSoundWaveInfo::LoadStatus::Loaded;
 		}
 
-		LoadingSoundWaveInfo.SoundWave  = InSoundWave;
-		LoadingSoundWaveInfo.OnLoaded   = MoveTemp(OnLoaded);
+		if (bSynchronous)
+		{
+			check(InSoundWave->GetPrecacheState() == ESoundWavePrecacheState::Done);
+
+			const Audio::DefaultUSoundWaveSampleType* RawPCMData = reinterpret_cast<const Audio::DefaultUSoundWaveSampleType*>(InSoundWave->RawPCMData);
+			const int32 NumSamples = InSoundWave->RawPCMDataSize / sizeof(Audio::DefaultUSoundWaveSampleType);
+
+			TSampleBuffer<> SampleBuffer(RawPCMData, NumSamples, InSoundWave->NumChannels, InSoundWave->GetSampleRateForCurrentPlatform());
+			OnLoaded(InSoundWave, SampleBuffer);
+
+		}
+		else
+		{
+			LoadingSoundWaveInfo.SoundWave = InSoundWave;
+			LoadingSoundWaveInfo.OnLoaded = MoveTemp(OnLoaded);
 
-		LoadingSoundWaves.Add(LoadingSoundWaveInfo);
+			LoadingSoundWaves.Add(LoadingSoundWaveInfo);
+		}
 	}
 
 	void FSoundWavePCMLoader::Update()
