Release FXSystem gpu resources onworldcleanup. This frees some lingering render resources that are not free untill the UWorld is released(Which it might never be)

diff --git a/Engine/Source/Runtime/Engine/Private/Particles/FXSystem.cpp b/Engine/Source/Runtime/Engine/Private/Particles/FXSystem.cpp
--- a/Engine/Source/Runtime/Engine/Private/Particles/FXSystem.cpp
+++ b/Engine/Source/Runtime/Engine/Private/Particles/FXSystem.cpp
@@ -51,6 +51,15 @@ FFXSystemInterface* FFXSystemInterface::Create(ERHIFeatureLevel::Type InFeatureL
 		return new FFXSystem(InFeatureLevel, InShaderPlatform, GPUSortManager);
 	}
 }
+void FFXSystemInterface::QueueDestroyGPUSimulation(FFXSystemInterface* FXSystem)
+{
+	check(FXSystem);
+	ENQUEUE_RENDER_COMMAND(FDestroyGPUSimulationCommand)(
+		[FXSystem](FRHICommandList& RHICmdList)
+	{
+		FXSystem->DestroyGPUSimulation();
+	});
+}
 
 void FFXSystemInterface::Destroy( FFXSystemInterface* FXSystem )
 {
diff --git a/Engine/Source/Runtime/Engine/Private/Particles/FXSystemPrivate.h b/Engine/Source/Runtime/Engine/Private/Particles/FXSystemPrivate.h
--- a/Engine/Source/Runtime/Engine/Private/Particles/FXSystemPrivate.h
+++ b/Engine/Source/Runtime/Engine/Private/Particles/FXSystemPrivate.h
@@ -249,7 +249,7 @@ private:
 	/**
 	 * Destroys any resources allocated for GPU simulation for this system.
 	 */
-	void DestroyGPUSimulation();
+	virtual void DestroyGPUSimulation();
 
 	/**
 	 * Initializes GPU resources.
diff --git a/Engine/Source/Runtime/Engine/Private/Particles/FXSystemSet.cpp b/Engine/Source/Runtime/Engine/Private/Particles/FXSystemSet.cpp
--- a/Engine/Source/Runtime/Engine/Private/Particles/FXSystemSet.cpp
+++ b/Engine/Source/Runtime/Engine/Private/Particles/FXSystemSet.cpp
@@ -190,6 +190,16 @@ void FFXSystemSet::OnDestroy()
 	FFXSystemInterface::OnDestroy();
 }
 
+
+void FFXSystemSet::DestroyGPUSimulation()
+{
+	for (FFXSystemInterface*& FXSystem : FXSystems)
+	{
+		check(FXSystem);
+		FXSystem->DestroyGPUSimulation();
+	}
+}
+
 FFXSystemSet::~FFXSystemSet()
 {
 	for (FFXSystemInterface* FXSystem : FXSystems)
diff --git a/Engine/Source/Runtime/Engine/Private/Particles/FXSystemSet.h b/Engine/Source/Runtime/Engine/Private/Particles/FXSystemSet.h
--- a/Engine/Source/Runtime/Engine/Private/Particles/FXSystemSet.h
+++ b/Engine/Source/Runtime/Engine/Private/Particles/FXSystemSet.h
@@ -49,6 +49,7 @@ public:
 		bool bAllowGPUParticleUpdate) override;
 
 	virtual void OnDestroy() override;
+	virtual void DestroyGPUSimulation() override;
 
 	/** Get the shared SortManager, used in the rendering loop to call FGPUSortManager::OnPreRender() and FGPUSortManager::OnPostRenderOpaque() */
 	virtual FGPUSortManager* GetGPUSortManager() const override;
diff --git a/Engine/Source/Runtime/Engine/Private/Particles/ParticleGpuSimulation.cpp b/Engine/Source/Runtime/Engine/Private/Particles/ParticleGpuSimulation.cpp
--- a/Engine/Source/Runtime/Engine/Private/Particles/ParticleGpuSimulation.cpp
+++ b/Engine/Source/Runtime/Engine/Private/Particles/ParticleGpuSimulation.cpp
@@ -4383,8 +4383,10 @@ void FFXSystem::ReleaseGPUResources()
 {
 	if (RHISupportsGPUParticles())
 	{
-		check(ParticleSimulationResources);
-		ParticleSimulationResources->Release();
+		if(ParticleSimulationResources)
+		{
+			ParticleSimulationResources->Release();
+		}
 	}
 }
 
diff --git a/Engine/Source/Runtime/Engine/Private/World.cpp b/Engine/Source/Runtime/Engine/Private/World.cpp
--- a/Engine/Source/Runtime/Engine/Private/World.cpp
+++ b/Engine/Source/Runtime/Engine/Private/World.cpp
@@ -4335,6 +4335,12 @@ void UWorld::CleanupWorldInternal(bool bSessionEnded, bool bCleanupResources, UW
 	FWorldDelegates::OnPostWorldCleanup.Broadcast(this, bSessionEnded, bCleanupResources);
 
 	SubsystemCollection.Deinitialize();
+
+	if (FXSystem)
+	{
+		FFXSystemInterface::QueueDestroyGPUSimulation(FXSystem);
+	}
+
 }
 
 UGameViewportClient* UWorld::GetGameViewport() const
diff --git a/Engine/Source/Runtime/Engine/Public/FXSystem.h b/Engine/Source/Runtime/Engine/Public/FXSystem.h
--- a/Engine/Source/Runtime/Engine/Public/FXSystem.h
+++ b/Engine/Source/Runtime/Engine/Public/FXSystem.h
@@ -133,6 +133,11 @@ public:
 	 */
 	ENGINE_API static void Destroy(FFXSystemInterface* FXSystem);
 
+	/**
+	 * Queue Destroy the gpu simulation on the render thread
+	 */
+	ENGINE_API static void QueueDestroyGPUSimulation(FFXSystemInterface* FXSystem);
+
 	/**
 	 * Register a custom FX system implementation.
 	 */
@@ -153,6 +158,11 @@ public:
 	 */
 	ENGINE_API virtual void OnDestroy() { bIsPendingKill = true; }
 
+	/**
+	 * Gamethread callback when destroy gets called, allows to clean up references.
+	 */
+	ENGINE_API virtual void DestroyGPUSimulation() { }
+
 
 	/**
 	 * Tick the effects system.
