CIS fix after Chaos SQ changes

diff --git a/Engine/Source/Runtime/Engine/Private/Collision/CollisionConversions.cpp b/Engine/Source/Runtime/Engine/Private/Collision/CollisionConversions.cpp
--- a/Engine/Source/Runtime/Engine/Private/Collision/CollisionConversions.cpp
+++ b/Engine/Source/Runtime/Engine/Private/Collision/CollisionConversions.cpp
@@ -171,6 +171,7 @@ static void SetHitResultFromShapeAndFaceIndex(const FPhysicsShape& Shape,  const
 #endif
 	else
 	{
+#if WITH_CHAOS
 		// Currently geom collections are registered with a primitive component user data, but maybe custom should be adapted
 		// to be more general so we can support leaf identification #BGTODO
 		void* UserData = Actor.UserData();
@@ -183,6 +184,7 @@ static void SetHitResultFromShapeAndFaceIndex(const FPhysicsShape& Shape,  const
 			OutResult.BoneName = NAME_None;
 		}
 		else
+#endif
 		{
 			ensureMsgf(false, TEXT("SetHitResultFromShapeAndFaceIndex hit shape with invalid userData"));
 		}
