Move platform cvars warnings.

diff --git a/Engine/Source/Runtime/Engine/Private/UnrealClient.cpp b/Engine/Source/Runtime/Engine/Private/UnrealClient.cpp
--- a/Engine/Source/Runtime/Engine/Private/UnrealClient.cpp
+++ b/Engine/Source/Runtime/Engine/Private/UnrealClient.cpp
@@ -589,36 +589,6 @@ int32 FStatUnitData::DrawStat(FViewport* InViewport, FCanvas* InCanvas, int32 In
 				const FColor GPUMaxColor = GEngine->GetFrameTimeDisplayColor(Max_GPUFrameTime);
 				InCanvas->DrawShadowedString(X3, InY, *FString::Printf(TEXT("%4.2f ms"), Max_GPUFrameTime), Font, GPUMaxColor);
 			}
-			if (GMaxRHIShaderPlatform == SP_PS4)
-			{
-				FString Warnings;
-
-				{
-					static const auto CVar = IConsoleManager::Get().FindConsoleVariable(TEXT("r.PS4ContinuousSubmits"));
-					int32 Value = CVar->GetInt();
-
-					if (!Value)
-					{
-						// good for profiling (avoids bubles) but bad for high fps
-						Warnings += TEXT(" r.PS4ContinuousSubmits");
-					}
-				}
-				{
-					static const auto CVar = IConsoleManager::Get().FindConsoleVariable(TEXT("r.PS4StallsOnMarkers"));
-					int32 Value = CVar->GetInt();
-
-					if (Value)
-					{
-						// good to get Razor aligned GPU profiling but bad for high fps
-						Warnings += TEXT(" r.PS4StallsOnMarkers");
-					}
-				}
-
-				if (!Warnings.IsEmpty())
-				{
-					InCanvas->DrawShadowedString(X3 + 100, InY, *Warnings, Font, FColor::Red);
-				}
-			}
 			InY += RowHeight;
 		}
 		if (IsRunningRHIInSeparateThread())
