Fixes for UE-90323

* VLM shaders in editor are now forced to SM5
* Brick atlas now removes all brick data on ReleaseRHI(). This fixes the render resource deinitialization issue.

diff --git a/Engine/Source/Runtime/Engine/Private/PrecomputedVolumetricLightmap.cpp b/Engine/Source/Runtime/Engine/Private/PrecomputedVolumetricLightmap.cpp
--- a/Engine/Source/Runtime/Engine/Private/PrecomputedVolumetricLightmap.cpp
+++ b/Engine/Source/Runtime/Engine/Private/PrecomputedVolumetricLightmap.cpp
@@ -387,13 +387,13 @@ ENGINE_API void FPrecomputedVolumetricLightmapData::ReleaseRHI()
 {
 	if (GetFeatureLevel() >= ERHIFeatureLevel::SM5)
 	{
+		GVolumetricLightmapBrickAtlas.Remove(this);
+
 		{
 			IndirectionTexture.Texture.SafeRelease();
 			IndirectionTexture.UAV.SafeRelease();
 			BrickData.ReleaseRHI();
 		}
-
-		GVolumetricLightmapBrickAtlas.Remove(this);
 	}
 }
 
@@ -906,10 +906,10 @@ FVolumetricLightmapBrickAtlas::FVolumetricLightmapBrickAtlas()
 }
 
 template<class VolumetricLightmapBrickDataType>
-void CopyDataIntoAtlas(FRHICommandList& RHICmdList, ERHIFeatureLevel::Type FeatureLevel, int32 SrcOffset, int32 DestOffset, int32 NumBricks, const VolumetricLightmapBrickDataType& SrcData, FVolumetricLightmapBrickTextureSet DestTextureSet)
+void CopyDataIntoAtlas(FRHICommandList& RHICmdList, int32 SrcOffset, int32 DestOffset, int32 NumBricks, const VolumetricLightmapBrickDataType& SrcData, FVolumetricLightmapBrickTextureSet DestTextureSet)
 {
 	{
-		FGlobalShaderMap* GlobalShaderMap = GetGlobalShaderMap(FeatureLevel);
+		FGlobalShaderMap* GlobalShaderMap = GetGlobalShaderMap(ERHIFeatureLevel::SM5);
 		{
 			// Transition all the UAVs to writable
 			FRHIUnorderedAccessView* UAVs[UE_ARRAY_COUNT(DestTextureSet.SHCoefficients) + 3] =
@@ -954,7 +954,7 @@ void CopyDataIntoAtlas(FRHICommandList& RHICmdList, ERHIFeatureLevel::Type Featu
 
 	for (int32 i = 0; i < UE_ARRAY_COUNT(SrcData.SHCoefficients); i++)
 	{
-		FGlobalShaderMap* GlobalShaderMap = GetGlobalShaderMap(FeatureLevel);
+		FGlobalShaderMap* GlobalShaderMap = GetGlobalShaderMap(ERHIFeatureLevel::SM5);
 
 		TShaderMapRef<FCopyResidentBrickSHCoefficientsCS> ComputeShader(GlobalShaderMap);
 
@@ -1081,7 +1081,7 @@ void FVolumetricLightmapBrickAtlas::Insert(int32 Index, FPrecomputedVolumetricLi
 		// Copy old allocations
 		for (int32 AllocationIndex = 0; AllocationIndex < Index && AllocationIndex < Allocations.Num(); AllocationIndex++)
 		{
-			CopyDataIntoAtlas(RHICmdList, GetFeatureLevel(), Allocations[AllocationIndex].StartOffset, BrickStartAllocation, Allocations[AllocationIndex].Size, TextureSet, NewTextureSet);
+			CopyDataIntoAtlas(RHICmdList, Allocations[AllocationIndex].StartOffset, BrickStartAllocation, Allocations[AllocationIndex].Size, TextureSet, NewTextureSet);
 
 			NewAllocations.Add(Allocation{ Allocations[AllocationIndex].Data, Allocations[AllocationIndex].Size, BrickStartAllocation });
 			BrickStartAllocation += Allocations[AllocationIndex].Size;
@@ -1091,7 +1091,7 @@ void FVolumetricLightmapBrickAtlas::Insert(int32 Index, FPrecomputedVolumetricLi
 		{
 			int32 NumBricks = Data->BrickDataDimensions.X * Data->BrickDataDimensions.Y * Data->BrickDataDimensions.Z / (PaddedBrickSize * PaddedBrickSize * PaddedBrickSize);
 
-			CopyDataIntoAtlas(RHICmdList, GetFeatureLevel(), 0, BrickStartAllocation, NumBricks, Data->BrickData, NewTextureSet);
+			CopyDataIntoAtlas(RHICmdList, 0, BrickStartAllocation, NumBricks, Data->BrickData, NewTextureSet);
 
 			NewAllocations.Add(Allocation{ Data, NumBricks, BrickStartAllocation });
 			Data->BrickDataBaseOffsetInAtlas = BrickStartAllocation;
@@ -1101,7 +1101,7 @@ void FVolumetricLightmapBrickAtlas::Insert(int32 Index, FPrecomputedVolumetricLi
 		// Copy the rest of allocations
 		for (int32 AllocationIndex = Index; AllocationIndex < Allocations.Num(); AllocationIndex++)
 		{
-			CopyDataIntoAtlas(RHICmdList, GetFeatureLevel(), Allocations[AllocationIndex].StartOffset, BrickStartAllocation, Allocations[AllocationIndex].Size, TextureSet, NewTextureSet);
+			CopyDataIntoAtlas(RHICmdList, Allocations[AllocationIndex].StartOffset, BrickStartAllocation, Allocations[AllocationIndex].Size, TextureSet, NewTextureSet);
 
 			NewAllocations.Add(Allocation{ Allocations[AllocationIndex].Data, Allocations[AllocationIndex].Size, BrickStartAllocation });
 			// Handle the sub level data movements
@@ -1209,7 +1209,7 @@ void FVolumetricLightmapBrickAtlas::Remove(FPrecomputedVolumetricLightmapData* D
 			// Copy old allocations
 			for (int32 AllocationIndex = 0; AllocationIndex < Index && AllocationIndex < Allocations.Num(); AllocationIndex++)
 			{
-				CopyDataIntoAtlas(RHICmdList, GetFeatureLevel(), Allocations[AllocationIndex].StartOffset, BrickStartAllocation, Allocations[AllocationIndex].Size, TextureSet, NewTextureSet);
+				CopyDataIntoAtlas(RHICmdList, Allocations[AllocationIndex].StartOffset, BrickStartAllocation, Allocations[AllocationIndex].Size, TextureSet, NewTextureSet);
 
 				NewAllocations.Add(Allocation{ Allocations[AllocationIndex].Data, Allocations[AllocationIndex].Size, BrickStartAllocation });
 				BrickStartAllocation += Allocations[AllocationIndex].Size;
@@ -1220,7 +1220,7 @@ void FVolumetricLightmapBrickAtlas::Remove(FPrecomputedVolumetricLightmapData* D
 			// Copy the rest of allocations
 			for (int32 AllocationIndex = Index + 1; AllocationIndex < Allocations.Num(); AllocationIndex++)
 			{
-				CopyDataIntoAtlas(RHICmdList, GetFeatureLevel(), Allocations[AllocationIndex].StartOffset, BrickStartAllocation, Allocations[AllocationIndex].Size, TextureSet, NewTextureSet);
+				CopyDataIntoAtlas(RHICmdList, Allocations[AllocationIndex].StartOffset, BrickStartAllocation, Allocations[AllocationIndex].Size, TextureSet, NewTextureSet);
 
 				NewAllocations.Add(Allocation{ Allocations[AllocationIndex].Data, Allocations[AllocationIndex].Size, BrickStartAllocation });
 				Allocations[AllocationIndex].Data->HandleDataMovementInAtlas(Allocations[AllocationIndex].StartOffset, BrickStartAllocation);
@@ -1255,5 +1255,11 @@ void FVolumetricLightmapBrickAtlas::Remove(FPrecomputedVolumetricLightmapData* D
 
 void FVolumetricLightmapBrickAtlas::ReleaseRHI()
 {
+	TArray<Allocation> AllocationList = Allocations;
+	for (Allocation& Alloc : AllocationList)
+	{
+		Remove(Alloc.Data);
+	}
+
 	TextureSet.Release();
 }
