Adding better handling of bad mesh batch submission.
 - Non-shipping builds will validate and reject malformed mesh batches while logging an ensure.
 - Shipping builds compile out the validation.
 - Merged existing validation across different paths into same function.
 - Added validation for index buffer issue.

diff --git a/Engine/Source/Runtime/Engine/Private/PrimitiveSceneProxy.cpp b/Engine/Source/Runtime/Engine/Private/PrimitiveSceneProxy.cpp
--- a/Engine/Source/Runtime/Engine/Private/PrimitiveSceneProxy.cpp
+++ b/Engine/Source/Runtime/Engine/Private/PrimitiveSceneProxy.cpp
@@ -822,11 +822,10 @@ void FPrimitiveSceneProxy::RenderBounds(
 	}
 }
 
-void FPrimitiveSceneProxy::VerifyUsedMaterial(const FMaterialRenderProxy* MaterialRenderProxy) const
+bool FPrimitiveSceneProxy::VerifyUsedMaterial(const FMaterialRenderProxy* MaterialRenderProxy) const
 {
 	// Only verify GetUsedMaterials if uncooked and we can compile shaders, because FShaderCompilingManager::PropagateMaterialChangesToPrimitives is what needs GetUsedMaterials to be accurate
 #if WITH_EDITOR
-
 	if (bVerifyUsedMaterials)
 	{
 		const UMaterialInterface* MaterialInterface = MaterialRenderProxy->GetMaterialInterface();
@@ -837,10 +836,11 @@ void FPrimitiveSceneProxy::VerifyUsedMaterial(const FMaterialRenderProxy* Materi
 		{
 			// Shader compiling uses GetUsedMaterials to detect which components need their scene proxy recreated, so we can only render with materials present in that list
 			ensureMsgf(false, TEXT("PrimitiveComponent tried to render with Material %s, which was not present in the component's GetUsedMaterials results\n    Owner: %s, Resource: %s"), *MaterialInterface->GetName(), *GetOwnerName().ToString(), *GetResourceName().ToString());
+			return false;
 		}
 	}
-	
 #endif
+	return true;
 }
 
 void FPrimitiveSceneProxy::DrawArc(FPrimitiveDrawInterface* PDI, const FVector& Start, const FVector& End, const float Height, const uint32 Segments, const FLinearColor& Color, uint8 DepthPriorityGroup, const float Thickness, const bool bScreenSpace)
diff --git a/Engine/Source/Runtime/Engine/Private/SceneManagement.cpp b/Engine/Source/Runtime/Engine/Private/SceneManagement.cpp
--- a/Engine/Source/Runtime/Engine/Private/SceneManagement.cpp
+++ b/Engine/Source/Runtime/Engine/Private/SceneManagement.cpp
@@ -242,13 +242,6 @@ void FMeshElementCollector::AddMesh(int32 ViewIndex, FMeshBatch& MeshBatch)
 {
 	DEFINE_LOG_CATEGORY_STATIC(FMeshElementCollector_AddMesh, Warning, All);
 
-	checkSlow(MeshBatch.VertexFactory);
-	checkSlow(MeshBatch.VertexFactory->IsInitialized());
-	checkSlow(MeshBatch.MaterialRenderProxy);
-	checkSlow(PrimitiveSceneProxy);
-
-	PrimitiveSceneProxy->VerifyUsedMaterial(MeshBatch.MaterialRenderProxy);
-
 	if (MeshBatch.bCanApplyViewModeOverrides)
 	{
 		FSceneView* View = Views[ViewIndex];
@@ -263,15 +256,13 @@ void FMeshElementCollector::AddMesh(int32 ViewIndex, FMeshBatch& MeshBatch)
 			*this);
 	}
 
-	MeshBatch.PreparePrimitiveUniformBuffer(PrimitiveSceneProxy, FeatureLevel);
-
-	for (int32 Index = 0; Index < MeshBatch.Elements.Num(); ++Index)
+	if (!MeshBatch.Validate(PrimitiveSceneProxy, FeatureLevel))
 	{
-		UE_CLOG(MeshBatch.Elements[Index].IndexBuffer && !MeshBatch.Elements[Index].IndexBuffer->IndexBufferRHI, FMeshElementCollector_AddMesh, Fatal,
-			TEXT("FMeshElementCollector::AddMesh - On MeshBatchElement %d, Material '%s', index buffer object has null RHI resource"),
-			Index, MeshBatch.MaterialRenderProxy ? *MeshBatch.MaterialRenderProxy->GetFriendlyName() : TEXT("null"));
+		return;
 	}
 
+	MeshBatch.PreparePrimitiveUniformBuffer(PrimitiveSceneProxy, FeatureLevel);
+
 	// If we are maintaining primitive scene data on the GPU, copy the primitive uniform buffer data to a unified array so it can be uploaded later
 	if (UseGPUScene(GMaxRHIShaderPlatform, FeatureLevel) && MeshBatch.VertexFactory->GetPrimitiveIdStreamIndex(EVertexInputStreamType::Default) >= 0)
 	{
@@ -528,10 +519,9 @@ int8 ComputeTemporalStaticMeshLOD( const FStaticMeshRenderData* RenderData, cons
 // Ensure we always use the left eye when selecting lods to avoid divergent selections in stereo
 const FSceneView& GetLODView(const FSceneView& InView)
 {
-	uint32 LODViewIndex = IStereoRendering::GetLODViewIndex();
-	if (InView.Family && InView.Family->Views.IsValidIndex(LODViewIndex))
+	if (IStereoRendering::IsASecondaryView(InView) && InView.Family)
 	{
-		return *InView.Family->Views[LODViewIndex];
+		return *InView.Family->Views[0];
 	}
 	else
 	{
@@ -675,8 +665,8 @@ FViewUniformShaderParameters::FViewUniformShaderParameters()
 {
 	FMemory::Memzero(*this);
 
-	FRHITexture* BlackVolume = (GBlackVolumeTexture &&  GBlackVolumeTexture->TextureRHI) ? GBlackVolumeTexture->TextureRHI : GBlackTexture->TextureRHI;
-	FRHITexture* BlackUintVolume = (GBlackUintVolumeTexture &&  GBlackUintVolumeTexture->TextureRHI) ? GBlackUintVolumeTexture->TextureRHI : GBlackTexture->TextureRHI;
+	FRHITexture* BlackVolume = (GBlackVolumeTexture &&  GBlackVolumeTexture->TextureRHI) ? GBlackVolumeTexture->TextureRHI : GBlackTexture->TextureRHI; // for es2, this might need to be 2d
+	FRHITexture* BlackUintVolume = (GBlackUintVolumeTexture &&  GBlackUintVolumeTexture->TextureRHI) ? GBlackUintVolumeTexture->TextureRHI : GBlackTexture->TextureRHI; // for es2, this might need to be 2d
 	check(GBlackVolumeTexture);
 
 	MaterialTextureBilinearClampedSampler = TStaticSamplerState<SF_Bilinear, AM_Clamp, AM_Clamp, AM_Clamp>::GetRHI();
@@ -762,8 +752,6 @@ FViewUniformShaderParameters::FViewUniformShaderParameters()
 	{
 		LightmapSceneData = GBlackTextureWithSRV->ShaderResourceViewRHI;
 	}
-	VTFeedbackBuffer = GEmptyVertexBufferWithUAV->UnorderedAccessViewRHI;
-	QuadOverdraw = GBlackTextureWithUAV->UnorderedAccessViewRHI;
 }
 
 FInstancedViewUniformShaderParameters::FInstancedViewUniformShaderParameters()
@@ -836,16 +824,6 @@ bool FLightCacheInterface::GetVirtualTextureLightmapProducer(ERHIFeatureLevel::T
 
 IMPLEMENT_GLOBAL_SHADER_PARAMETER_STRUCT(FLightmapResourceClusterShaderParameters, "LightmapResourceCluster");
 
-static FRHISamplerState* GetTextureSamplerState(const UTexture* Texture, FRHISamplerState* Default)
-{
-	FRHISamplerState* Result = nullptr;
-	if (Texture && Texture->Resource)
-	{
-		Result = Texture->Resource->SamplerStateRHI;
-	}
-	return Result ? Result : Default;
-}
-
 void GetLightmapClusterResourceParameters(
 	ERHIFeatureLevel::Type FeatureLevel, 
 	const FLightmapClusterResourceInput& Input,
@@ -857,22 +835,6 @@ void GetLightmapClusterResourceParameters(
 	static const auto CVar = IConsoleManager::Get().FindTConsoleVariableDataInt(TEXT("r.VirtualTexturedLightmaps"));
 	const bool bUseVirtualTextures = bAllowHighQualityLightMaps && (CVar->GetValueOnRenderThread() != 0) && UseVirtualTexturing(FeatureLevel);
 
-	Parameters.LightMapTexture = GBlackTexture->TextureRHI;
-	Parameters.SkyOcclusionTexture = GWhiteTexture->TextureRHI;
-	Parameters.AOMaterialMaskTexture = GBlackTexture->TextureRHI;
-	Parameters.StaticShadowTexture = GWhiteTexture->TextureRHI;
-	Parameters.VTLightMapTexture = GBlackTextureWithSRV->ShaderResourceViewRHI;
-	Parameters.VTLightMapTexture_1 = GBlackTextureWithSRV->ShaderResourceViewRHI;
-	Parameters.VTSkyOcclusionTexture = GWhiteTextureWithSRV->ShaderResourceViewRHI;
-	Parameters.VTAOMaterialMaskTexture = GBlackTextureWithSRV->ShaderResourceViewRHI;
-	Parameters.VTStaticShadowTexture = GWhiteTextureWithSRV->ShaderResourceViewRHI;
-	Parameters.LightmapVirtualTexturePageTable0 = GBlackTexture->TextureRHI;
-	Parameters.LightmapVirtualTexturePageTable1 = GBlackTexture->TextureRHI;
-	Parameters.LightMapSampler = GBlackTexture->SamplerStateRHI;
-	Parameters.SkyOcclusionSampler = GWhiteTexture->SamplerStateRHI;
-	Parameters.AOMaterialMaskSampler = GBlackTexture->SamplerStateRHI;
-	Parameters.StaticShadowTextureSampler = GWhiteTexture->SamplerStateRHI;
-
 	if (bUseVirtualTextures)
 	{
 		// this is sometimes called with NULL input to initialize default buffer
@@ -880,34 +842,34 @@ void GetLightmapClusterResourceParameters(
 		if (VirtualTexture && AllocatedVT)
 		{
 			// Bind VT here
-			Parameters.VTLightMapTexture = AllocatedVT->GetPhysicalTextureSRV((uint32)ELightMapVirtualTextureType::HqLayer0, false);
-			Parameters.VTLightMapTexture_1 = AllocatedVT->GetPhysicalTextureSRV((uint32)ELightMapVirtualTextureType::HqLayer1, false);
+			Parameters.LightMapTexture = AllocatedVT->GetPhysicalTexture((uint32)ELightMapVirtualTextureType::HqLayer0);
+			Parameters.LightMapTexture_1 = AllocatedVT->GetPhysicalTexture((uint32)ELightMapVirtualTextureType::HqLayer1);
 
 			if (VirtualTexture->HasLayerForType(ELightMapVirtualTextureType::SkyOcclusion))
 			{
-				Parameters.VTSkyOcclusionTexture = AllocatedVT->GetPhysicalTextureSRV((uint32)ELightMapVirtualTextureType::SkyOcclusion, false);
+				Parameters.SkyOcclusionTexture = AllocatedVT->GetPhysicalTexture((uint32)ELightMapVirtualTextureType::SkyOcclusion);
 			}
 			else
 			{
-				Parameters.VTSkyOcclusionTexture = GWhiteTextureWithSRV->ShaderResourceViewRHI;
+				Parameters.SkyOcclusionTexture = GWhiteTexture->TextureRHI;
 			}
 
 			if (VirtualTexture->HasLayerForType(ELightMapVirtualTextureType::AOMaterialMask))
 			{
-				Parameters.VTAOMaterialMaskTexture = AllocatedVT->GetPhysicalTextureSRV((uint32)ELightMapVirtualTextureType::AOMaterialMask, false);
+				Parameters.AOMaterialMaskTexture = AllocatedVT->GetPhysicalTexture((uint32)ELightMapVirtualTextureType::AOMaterialMask);
 			}
 			else
 			{
-				Parameters.VTAOMaterialMaskTexture = GBlackTextureWithSRV->ShaderResourceViewRHI;
+				Parameters.AOMaterialMaskTexture = GBlackTexture->TextureRHI;
 			}
 
 			if (VirtualTexture->HasLayerForType(ELightMapVirtualTextureType::ShadowMask))
 			{
-				Parameters.VTStaticShadowTexture = AllocatedVT->GetPhysicalTextureSRV((uint32)ELightMapVirtualTextureType::ShadowMask, false);
+				Parameters.StaticShadowTexture = AllocatedVT->GetPhysicalTexture((uint32)ELightMapVirtualTextureType::ShadowMask);
 			}
 			else
 			{
-				Parameters.VTStaticShadowTexture = GWhiteTextureWithSRV->ShaderResourceViewRHI;
+				Parameters.StaticShadowTexture = GWhiteTexture->TextureRHI;
 			}
 
 			FRHITexture* PageTable0 = AllocatedVT->GetPageTableTexture(0u);
@@ -928,21 +890,36 @@ void GetLightmapClusterResourceParameters(
 			Parameters.AOMaterialMaskSampler = TStaticSamplerState<SF_AnisotropicLinear, AM_Clamp, AM_Clamp, AM_Clamp, 0, MaxAniso>::GetRHI();
 			Parameters.StaticShadowTextureSampler = TStaticSamplerState<SF_AnisotropicLinear, AM_Clamp, AM_Clamp, AM_Clamp, 0, MaxAniso>::GetRHI();
 		}
+		else
+		{
+			Parameters.LightMapTexture = GBlackTexture->TextureRHI;
+			Parameters.LightMapTexture_1 = GBlackTexture->TextureRHI;
+			Parameters.SkyOcclusionTexture = GWhiteTexture->TextureRHI;
+			Parameters.AOMaterialMaskTexture = GBlackTexture->TextureRHI;
+			Parameters.StaticShadowTexture = GWhiteTexture->TextureRHI;
+			Parameters.LightmapVirtualTexturePageTable0 = GBlackTexture->TextureRHI;
+			Parameters.LightmapVirtualTexturePageTable1 = GBlackTexture->TextureRHI;
+			Parameters.LightMapSampler = GBlackTexture->SamplerStateRHI;
+			Parameters.SkyOcclusionSampler = GWhiteTexture->SamplerStateRHI;
+			Parameters.AOMaterialMaskSampler = GBlackTexture->SamplerStateRHI;
+			Parameters.StaticShadowTextureSampler = GWhiteTexture->SamplerStateRHI;
+		}
 	}
 	else
 	{
 		const UTexture2D* LightMapTexture = Input.LightMapTextures[bAllowHighQualityLightMaps ? 0 : 1];
 
 		Parameters.LightMapTexture = LightMapTexture ? LightMapTexture->TextureReference.TextureReferenceRHI.GetReference() : GBlackTexture->TextureRHI;
+		Parameters.LightMapTexture_1 = GBlackTexture->TextureRHI;
 		Parameters.SkyOcclusionTexture = Input.SkyOcclusionTexture ? Input.SkyOcclusionTexture->TextureReference.TextureReferenceRHI.GetReference() : GWhiteTexture->TextureRHI;
 		Parameters.AOMaterialMaskTexture = Input.AOMaterialMaskTexture ? Input.AOMaterialMaskTexture->TextureReference.TextureReferenceRHI.GetReference() : GBlackTexture->TextureRHI;
 
-		Parameters.LightMapSampler = GetTextureSamplerState(LightMapTexture, GBlackTexture->SamplerStateRHI);
-		Parameters.SkyOcclusionSampler = GetTextureSamplerState(Input.SkyOcclusionTexture, GWhiteTexture->SamplerStateRHI);
-		Parameters.AOMaterialMaskSampler = GetTextureSamplerState(Input.AOMaterialMaskTexture, GBlackTexture->SamplerStateRHI);
+		Parameters.LightMapSampler = (LightMapTexture && LightMapTexture->Resource) ? LightMapTexture->Resource->SamplerStateRHI : GBlackTexture->SamplerStateRHI;
+		Parameters.SkyOcclusionSampler = (Input.SkyOcclusionTexture && Input.SkyOcclusionTexture->Resource) ? Input.SkyOcclusionTexture->Resource->SamplerStateRHI : GWhiteTexture->SamplerStateRHI;
+		Parameters.AOMaterialMaskSampler = (Input.AOMaterialMaskTexture && Input.AOMaterialMaskTexture->Resource) ? Input.AOMaterialMaskTexture->Resource->SamplerStateRHI : GBlackTexture->SamplerStateRHI;
 
 		Parameters.StaticShadowTexture = Input.ShadowMapTexture ? Input.ShadowMapTexture->TextureReference.TextureReferenceRHI.GetReference() : GWhiteTexture->TextureRHI;
-		Parameters.StaticShadowTextureSampler = GetTextureSamplerState(Input.ShadowMapTexture, GWhiteTexture->SamplerStateRHI);
+		Parameters.StaticShadowTextureSampler = (Input.ShadowMapTexture && Input.ShadowMapTexture->Resource) ? Input.ShadowMapTexture->Resource->SamplerStateRHI : GWhiteTexture->SamplerStateRHI;
 
 		Parameters.LightmapVirtualTexturePageTable0 = GBlackTexture->TextureRHI;
 		Parameters.LightmapVirtualTexturePageTable1 = GBlackTexture->TextureRHI;
@@ -1090,40 +1067,124 @@ void FReadOnlyCVARCache::Init()
 
 void FMeshBatch::PreparePrimitiveUniformBuffer(const FPrimitiveSceneProxy* PrimitiveSceneProxy, ERHIFeatureLevel::Type FeatureLevel)
 {
-	const bool bVFSupportsPrimitiveIdStream = VertexFactory->GetType()->SupportsPrimitiveIdStream();
-	checkf((PrimitiveSceneProxy->DoesVFRequirePrimitiveUniformBuffer() || bVFSupportsPrimitiveIdStream), TEXT("PrimitiveSceneProxy has bVFRequiresPrimitiveUniformBuffer disabled yet tried to draw with a vertex factory (%s) that did not support PrimitiveIdStream."), VertexFactory->GetType()->GetName());
+	// Fallback to using the primitive uniform buffer if GPU scene is disabled.
+	if (VertexFactory->GetType()->SupportsPrimitiveIdStream() && !UseGPUScene(GMaxRHIShaderPlatform, FeatureLevel))
+	{
+		for (int32 ElementIndex = 0; ElementIndex < Elements.Num(); ElementIndex++)
+		{
+			FMeshBatchElement& MeshElement = Elements[ElementIndex];
 
-	const bool bPrimitiveShaderDataComesFromSceneBuffer = VertexFactory->GetPrimitiveIdStreamIndex(EVertexInputStreamType::Default) >= 0;
+			if (!MeshElement.PrimitiveUniformBufferResource)
+			{
+				MeshElement.PrimitiveUniformBuffer = PrimitiveSceneProxy->GetUniformBuffer();
+			}
+		}
+	}
+}
 
-	for (int32 ElementIndex = 0; ElementIndex < Elements.Num(); ElementIndex++)
+#if USE_MESH_BATCH_VALIDATION
+bool FMeshBatch::Validate(const FPrimitiveSceneProxy* PrimitiveSceneProxy, ERHIFeatureLevel::Type FeatureLevel) const
+{
+	check(PrimitiveSceneProxy);
+
+	const auto LogMeshError = [&](const FString& Error) -> bool
+	{
+		const FString VertexFactoryName = VertexFactory ? VertexFactory->GetType()->GetFName().ToString() : TEXT("nullptr");
+		const uint32 bVertexFactoryInitialized = (VertexFactory && VertexFactory->IsInitialized()) ? 1 : 0;
+
+		ensureMsgf(false,
+			TEXT("FMeshBatch was not properly setup. %s.\n")
+			TEXT("\tVertexFactory[Name: %s, Initialized: %u]\n")
+			TEXT("\tPrimitiveSceneProxy[Level: %s, Owner: %s, Resource: %s]"),
+			*Error,
+			*VertexFactoryName,
+			bVertexFactoryInitialized,
+			*PrimitiveSceneProxy->GetLevelName().ToString(),
+			*PrimitiveSceneProxy->GetOwnerName().ToString(),
+			*PrimitiveSceneProxy->GetResourceName().ToString());
+
+		return false;
+	};
+
+	if (!MaterialRenderProxy)
+	{
+		return LogMeshError(TEXT("Mesh has a null material render proxy!"));
+	}
+
+	if (!PrimitiveSceneProxy->VerifyUsedMaterial(MaterialRenderProxy))
+	{
+		return LogMeshError(TEXT("Mesh material is not marked as used by the primitive scene proxy."));
+	}
+
+	if (!VertexFactory)
+	{
+		return LogMeshError(TEXT("Mesh has a null vertex factory!"));
+	}
+
+	if (!VertexFactory->IsInitialized())
+	{
+		return LogMeshError(TEXT("Mesh has an uninitialized vertex factory!"));
+	}
+
+	for (int32 Index = 0; Index < Elements.Num(); ++Index)
 	{
-		FMeshBatchElement& MeshElement = Elements[ElementIndex];
+		const FMeshBatchElement& MeshBatchElement = Elements[Index];
 
-		if (bPrimitiveShaderDataComesFromSceneBuffer)
+		if (MeshBatchElement.IndexBuffer)
 		{
-			checkf(!Elements[ElementIndex].PrimitiveUniformBuffer, 
-				TEXT("FMeshBatch was assigned a PrimitiveUniformBuffer even though Vertex Factory %s fetches primitive shader data through a Scene buffer.  The assigned PrimitiveUniformBuffer cannot be respected.  Use PrimitiveUniformBufferResource instead for dynamic primitive data, or leave both null to get FPrimitiveSceneProxy->UniformBuffer."), VertexFactory->GetType()->GetName());
+			if (const FRHIIndexBuffer* IndexBufferRHI = MeshBatchElement.IndexBuffer->IndexBufferRHI)
+			{
+				const int32 IndexCount = GetVertexCountForPrimitiveCount(MeshBatchElement.NumPrimitives, Type);
+
+				if ((MeshBatchElement.FirstIndex + IndexCount) * IndexBufferRHI->GetStride() > IndexBufferRHI->GetSize())
+				{
+					return LogMeshError(FString::Printf(
+						TEXT("MeshBatchElement %d, Material '%s', index range extends past index buffer bounds: Start %u, Count %u, Buffer Size %u, Buffer stride %u."),
+						Index, MaterialRenderProxy ? *MaterialRenderProxy->GetFriendlyName() : TEXT("nullptr"),
+						MeshBatchElement.FirstIndex, IndexCount, IndexBufferRHI->GetSize(), IndexBufferRHI->GetStride()));
+				}
+			}
+			else
+			{
+				return LogMeshError(FString::Printf(
+					TEXT("FMeshElementCollector::AddMesh - On MeshBatchElement %d, Material '%s', index buffer object has null RHI resource"),
+					Index, MaterialRenderProxy ? *MaterialRenderProxy->GetFriendlyName() : TEXT("nullptr")));
+			}
 		}
+	}
 
-		// If we are not using GPU Scene, draws using vertex factories that do not support an explicit PrimitiveUniformBuffer on the FMeshBatch need to be setup with the FPrimitiveSceneProxy's uniform buffer
-		if (!MeshElement.PrimitiveUniformBufferResource && !UseGPUScene(GMaxRHIShaderPlatform, FeatureLevel) && bVFSupportsPrimitiveIdStream)
+	const bool bVFSupportsPrimitiveIdStream = VertexFactory->GetType()->SupportsPrimitiveIdStream();
+
+	if (!PrimitiveSceneProxy->DoesVFRequirePrimitiveUniformBuffer() && !bVFSupportsPrimitiveIdStream)
+	{
+		return LogMeshError(TEXT("PrimitiveSceneProxy has bVFRequiresPrimitiveUniformBuffer disabled yet tried to draw with a vertex factory that did not support PrimitiveIdStream."));
+	}
+
+	const bool bPrimitiveShaderDataComesFromSceneBuffer = VertexFactory->GetPrimitiveIdStreamIndex(EVertexInputStreamType::Default) >= 0;
+
+	for (int32 ElementIndex = 0; ElementIndex < Elements.Num(); ElementIndex++)
+	{
+		const FMeshBatchElement& MeshElement = Elements[ElementIndex];
+
+		if (bPrimitiveShaderDataComesFromSceneBuffer && Elements[ElementIndex].PrimitiveUniformBuffer)
 		{
-			MeshElement.PrimitiveUniformBuffer = PrimitiveSceneProxy->GetUniformBuffer();
+			return LogMeshError(
+				TEXT("FMeshBatch was assigned a PrimitiveUniformBuffer even though the vertex factory fetches primitive shader data through the GPUScene buffer. ")
+				TEXT("The assigned PrimitiveUniformBuffer cannot be respected. Use PrimitiveUniformBufferResource instead for dynamic primitive data, or leave ")
+				TEXT("both null to get FPrimitiveSceneProxy->UniformBuffer."));
 		}
 
 		const bool bValidPrimitiveData = bPrimitiveShaderDataComesFromSceneBuffer || Elements[ElementIndex].PrimitiveUniformBuffer || Elements[ElementIndex].PrimitiveUniformBufferResource;
 
-		UE_CLOG(!bValidPrimitiveData, LogEngine, Fatal,
-			TEXT("FMeshBatch was not properly setup. No primitive uniform buffer was specified and the vertex factory does not have a valid primitive id stream.\n")
-			TEXT("\tVertexFactory[Name: %s, Initialized: %u]\n")
-			TEXT("\tPrimitiveSceneProxy[Level: %s, Owner: %s, Resource: %s]"),
-			*VertexFactory->GetType()->GetFName().ToString(),
-			VertexFactory->IsInitialized() ? 1 : 0,
-			*PrimitiveSceneProxy->GetLevelName().ToString(),
-			*PrimitiveSceneProxy->GetOwnerName().ToString(),
-			*PrimitiveSceneProxy->GetResourceName().ToString());
+		if (!bValidPrimitiveData)
+		{
+			return LogMeshError(TEXT("No primitive uniform buffer was specified and the vertex factory does not have a valid primitive id stream"));
+		}
 	}
+
+	return true;
 }
+#endif
 
 IMPLEMENT_GLOBAL_SHADER_PARAMETER_STRUCT(FMobileReflectionCaptureShaderParameters, "MobileReflectionCapture");
 
diff --git a/Engine/Source/Runtime/Engine/Public/MeshBatch.h b/Engine/Source/Runtime/Engine/Public/MeshBatch.h
--- a/Engine/Source/Runtime/Engine/Public/MeshBatch.h
+++ b/Engine/Source/Runtime/Engine/Public/MeshBatch.h
@@ -10,6 +10,8 @@
 #include "PrimitiveUniformShaderParameters.h"
 #include "VT/RuntimeVirtualTextureEnum.h"
 
+#define USE_MESH_BATCH_VALIDATION !UE_BUILD_SHIPPING
+
 class FLightCacheInterface;
 
 enum EPrimitiveIdMode
@@ -319,6 +321,12 @@ struct FMeshBatch
 
 	ENGINE_API void PreparePrimitiveUniformBuffer(const FPrimitiveSceneProxy* PrimitiveSceneProxy, ERHIFeatureLevel::Type FeatureLevel);
 
+#if USE_MESH_BATCH_VALIDATION
+	ENGINE_API bool Validate(const FPrimitiveSceneProxy* SceneProxy, ERHIFeatureLevel::Type FeatureLevel) const;
+#else
+	FORCEINLINE bool Validate(const FPrimitiveSceneProxy* SceneProxy, ERHIFeatureLevel::Type FeatureLevel) const { return true; }
+#endif
+
 	/** Default constructor. */
 	FMeshBatch()
 	:	VertexFactory(nullptr)
diff --git a/Engine/Source/Runtime/Engine/Public/PrimitiveSceneProxy.h b/Engine/Source/Runtime/Engine/Public/PrimitiveSceneProxy.h
--- a/Engine/Source/Runtime/Engine/Public/PrimitiveSceneProxy.h
+++ b/Engine/Source/Runtime/Engine/Public/PrimitiveSceneProxy.h
@@ -171,7 +171,7 @@ public:
 	ENGINE_API void RenderBounds(FPrimitiveDrawInterface* PDI, const FEngineShowFlags& EngineShowFlags, const FBoxSphereBounds& Bounds, bool bRenderInEditor) const;
 
 	/** Verifies that a material used for rendering was present in the component's GetUsedMaterials list. */
-	ENGINE_API void VerifyUsedMaterial(const class FMaterialRenderProxy* MaterialRenderProxy) const;
+	ENGINE_API bool VerifyUsedMaterial(const class FMaterialRenderProxy* MaterialRenderProxy) const;
 
 	/** Returns the LOD that the primitive will render at for this view. */
 	virtual int32 GetLOD(const FSceneView* View) const { return INDEX_NONE; }
diff --git a/Engine/Source/Runtime/Renderer/Private/PrimitiveSceneInfo.cpp b/Engine/Source/Runtime/Renderer/Private/PrimitiveSceneInfo.cpp
--- a/Engine/Source/Runtime/Renderer/Private/PrimitiveSceneInfo.cpp
+++ b/Engine/Source/Runtime/Renderer/Private/PrimitiveSceneInfo.cpp
@@ -63,22 +63,23 @@ public:
 	{
 		if (Mesh.HasAnyDrawCalls())
 		{
-			check(Mesh.VertexFactory);
-			check(Mesh.VertexFactory->IsInitialized());
 			checkSlow(IsInParallelRenderingThread());
 
 			FPrimitiveSceneProxy* PrimitiveSceneProxy = PrimitiveSceneInfo->Proxy;
-			PrimitiveSceneProxy->VerifyUsedMaterial(Mesh.MaterialRenderProxy);
+			const ERHIFeatureLevel::Type FeatureLevel = PrimitiveSceneInfo->Scene->GetFeatureLevel();
+
+			if (!Mesh.Validate(PrimitiveSceneProxy, FeatureLevel))
+			{
+				return;
+			}
 
 			FStaticMeshBatch* StaticMesh = new(PrimitiveSceneInfo->StaticMeshes) FStaticMeshBatch(
 				PrimitiveSceneInfo,
 				Mesh,
 				CurrentHitProxy ? CurrentHitProxy->Id : FHitProxyId()
-				);
+			);
 
-			const ERHIFeatureLevel::Type FeatureLevel = PrimitiveSceneInfo->Scene->GetFeatureLevel();
 			StaticMesh->PreparePrimitiveUniformBuffer(PrimitiveSceneProxy, FeatureLevel);
-
 			// Volumetric self shadow mesh commands need to be generated every frame, as they depend on single frame uniform buffers with self shadow data.
 			const bool bSupportsCachingMeshDrawCommands = SupportsCachingMeshDrawCommands(*StaticMesh, FeatureLevel) && !PrimitiveSceneProxy->CastsVolumetricTranslucentShadow();
 
