Copy pasta comment
#rnx

diff --git a/Engine/Source/Runtime/Engine/Public/AudioDeviceManager.h b/Engine/Source/Runtime/Engine/Public/AudioDeviceManager.h
--- a/Engine/Source/Runtime/Engine/Public/AudioDeviceManager.h
+++ b/Engine/Source/Runtime/Engine/Public/AudioDeviceManager.h
@@ -89,7 +89,7 @@ public:
 	
 	~FAudioDeviceHandle();
 
-	// gets a pointer to the compressed chunk.
+	// gets a pointer to the audio device.
 	FAudioDevice* GetAudioDevice() const;
 
 	// Returns the device ID for the audio device referenced by this handle.
