add comment that FSoundWavePCMLoader should be called on the game thread

diff --git a/Engine/Source/Runtime/Engine/Classes/Sound/SampleBufferIO.h b/Engine/Source/Runtime/Engine/Classes/Sound/SampleBufferIO.h
--- a/Engine/Source/Runtime/Engine/Classes/Sound/SampleBufferIO.h
+++ b/Engine/Source/Runtime/Engine/Classes/Sound/SampleBufferIO.h
@@ -32,7 +32,7 @@ namespace Audio
 		void LoadSoundWave(USoundWave* InSoundWave, TFunction<void(const USoundWave* SoundWave, const Audio::FSampleBuffer& OutSampleBuffer)> OnLoaded, bool bSynchrounous = false);
 
 
-		// Update the loading state. 
+		// Update the loading state, should be called on the game thread. 
 		void Update();
 
 		//~ GCObject Interface
