Remove filtering for each individual effect type as it bloats the filter list
Fix OnInit() not being called on PostEditChangeProperty
Add iterator function for Instances (cleaner and making instances shared pointers in subsequent CL, moving towards deprecating direct access of instance array)
#rnx

diff --git a/Engine/Source/Runtime/Engine/Classes/Sound/SoundEffectBase.h b/Engine/Source/Runtime/Engine/Classes/Sound/SoundEffectBase.h
--- a/Engine/Source/Runtime/Engine/Classes/Sound/SoundEffectBase.h
+++ b/Engine/Source/Runtime/Engine/Classes/Sound/SoundEffectBase.h
@@ -27,6 +27,8 @@
 		F##EFFECT_NAME##Settings Settings = _Preset->GetSettings(); \
 
 #define EFFECT_PRESET_METHODS(EFFECT_NAME) \
+		virtual bool CanFilter() const override { return false; } \
+		virtual bool HasAssetActions() const { return true; } \
 		virtual FText GetAssetActionName() const override { return FText::FromString(#EFFECT_NAME); } \
 		virtual UClass* GetSupportedClass() const override { return EFFECT_PRESET_NAME(EFFECT_NAME)::StaticClass(); } \
 		virtual FSoundEffectBase* CreateNewEffect() const override { return new F##EFFECT_NAME; } \
diff --git a/Engine/Source/Runtime/Engine/Classes/Sound/SoundEffectPreset.h b/Engine/Source/Runtime/Engine/Classes/Sound/SoundEffectPreset.h
--- a/Engine/Source/Runtime/Engine/Classes/Sound/SoundEffectPreset.h
+++ b/Engine/Source/Runtime/Engine/Classes/Sound/SoundEffectPreset.h
@@ -21,11 +21,12 @@ public:
 	USoundEffectPreset(const FObjectInitializer& ObjectInitializer);
 	virtual ~USoundEffectPreset();
 
+	virtual bool CanFilter() const { return true; }
 	virtual FText GetAssetActionName() const PURE_VIRTUAL(USoundEffectPreset::GetAssetActionName, return FText(););
 	virtual UClass* GetSupportedClass() const PURE_VIRTUAL(USoundEffectPreset::GetSupportedClass, return nullptr;);
 	virtual USoundEffectPreset* CreateNewPreset(UObject* InParent, FName Name, EObjectFlags Flags) const PURE_VIRTUAL(USoundEffectPreset::CreateNewPreset, return nullptr;);
 	virtual FSoundEffectBase* CreateNewEffect() const PURE_VIRTUAL(USoundEffectPreset::CreateNewEffect, return nullptr;);
-	virtual bool HasAssetActions() const { return true; }
+	virtual bool HasAssetActions() const { return false; }
 	virtual void Init() PURE_VIRTUAL(USoundEffectPreset::Init, );
 	virtual void OnInit() {};
 	virtual FColor GetPresetColor() const { return FColor(200.0f, 100.0f, 100.0f); }
@@ -35,8 +36,7 @@ public:
 	void AddEffectInstance(FSoundEffectBase* InSource);
 	void RemoveEffectInstance(FSoundEffectBase* InSource);
 
-	void AddReferencedEffects(FReferenceCollector& Collector);
-
+	void AddReferencedEffects(FReferenceCollector& InCollector);
 
 	virtual void BeginDestroy() override;
 
@@ -49,4 +49,16 @@ protected:
 	// Array of instances which are using this preset
 	TArray<FSoundEffectBase*> Instances;
 	bool bInitialized;
+
+	template <typename T>
+	void IterateEffects(TFunction<void(T&)> InForEachEffect)
+	{
+		for (FSoundEffectBase* Instance : Instances)
+		{
+			if (Instance)
+			{
+				InForEachEffect(*static_cast<T*>(Instance));
+			}
+		}
+	}
 };
\ No newline at end of file
diff --git a/Engine/Source/Runtime/Engine/Private/SoundEffectPreset.cpp b/Engine/Source/Runtime/Engine/Private/SoundEffectPreset.cpp
--- a/Engine/Source/Runtime/Engine/Private/SoundEffectPreset.cpp
+++ b/Engine/Source/Runtime/Engine/Private/SoundEffectPreset.cpp
@@ -17,16 +17,12 @@ USoundEffectPreset::~USoundEffectPreset()
 {
 }
 
-
 void USoundEffectPreset::EffectCommand(TFunction<void()> Command)
 {
-	for (int32 i = 0; i < Instances.Num(); ++i)
+	IterateEffects<FSoundEffectBase>([Command](FSoundEffectBase& Instance)
 	{
-		if (Instances[i])
-		{
-			Instances[i]->EffectCommand(Command);
-		}
-	}
+		Instance.EffectCommand(Command);
+	});
 }
 
 void USoundEffectPreset::Update()
@@ -58,27 +54,22 @@ void USoundEffectPreset::AddEffectInstance(FSoundEffectBase* InSource)
 	Instances.AddUnique(InSource);
 }
 
-void USoundEffectPreset::AddReferencedEffects(FReferenceCollector& Collector)
+void USoundEffectPreset::AddReferencedEffects(FReferenceCollector& InCollector)
 {
-	for (FSoundEffectBase* Effect : Instances)
+	FReferenceCollector* Collector = &InCollector;
+	IterateEffects<FSoundEffectBase>([Collector](FSoundEffectBase& Instance)
 	{
-		if (Effect)
-		{
-			const USoundEffectPreset* EffectPreset = Effect->GetPreset();
-			Collector.AddReferencedObject(EffectPreset);
-		}
-	}
+		const USoundEffectPreset* EffectPreset = Instance.GetPreset();
+		Collector->AddReferencedObject(EffectPreset);
+	});
 }
 
 void USoundEffectPreset::BeginDestroy()
 {
-	for (int32 i = 0; i < Instances.Num(); ++i)
+	IterateEffects<FSoundEffectBase>([](FSoundEffectBase& Instance)
 	{
-		if (Instances[i])
-		{
-			Instances[i]->ClearPreset(false /* bRemoveFromPreset */);
-		}
-	}
+		Instance.ClearPreset(false /* bRemoveFromPreset */);
+	});
 	Instances.Reset();
 
 	Super::BeginDestroy();
@@ -94,6 +85,7 @@ void USoundEffectPreset::PostEditChangeProperty(FPropertyChangedEvent& PropertyC
 {
 	// Copy the settings to the thread safe version
 	Init();
+	OnInit();
 	Update();
 }
 
