Adding Delegate to PlayerCameraManager, allowing us to control audio fading (on camera fade) in a manner that does not touch the audio device directly.
* note: original audio fading functionality is untouched if we keep this new delegate unbound *

#rnx

diff --git a/Engine/Source/Runtime/Engine/Classes/Camera/PlayerCameraManager.h b/Engine/Source/Runtime/Engine/Classes/Camera/PlayerCameraManager.h
--- a/Engine/Source/Runtime/Engine/Classes/Camera/PlayerCameraManager.h
+++ b/Engine/Source/Runtime/Engine/Classes/Camera/PlayerCameraManager.h
@@ -15,6 +15,8 @@ class AEmitterCameraLensEffectBase;
 class APlayerController;
 class FDebugDisplayInfo;
 
+DECLARE_DYNAMIC_MULTICAST_DELEGATE_TwoParams(FOnAudioFadeChangeSignature, bool, bFadeOut, float, FadeAmount);
+
 /** 
  * Options that define how to blend when changing view targets. 
  * @see FViewTargetTransitionParams, SetViewTarget 
@@ -294,6 +296,10 @@ public:
 	UPROPERTY(EditAnywhere, BlueprintReadWrite, Category = Debug)
 	FVector ViewTargetOffset;
 
+	/** If bound, broadcast instead of manually altering audio device's master volume directly */
+	UPROPERTY(BlueprintAssignable, Transient)
+	FOnAudioFadeChangeSignature OnAudioFadeChangeEvent;
+
 	/** Current camera fade alpha range, where X = starting alpha and Y = final alpha (when bEnableFading == true) */
 	FVector2D FadeAlpha;
 
diff --git a/Engine/Source/Runtime/Engine/Private/PlayerCameraManager.cpp b/Engine/Source/Runtime/Engine/Private/PlayerCameraManager.cpp
--- a/Engine/Source/Runtime/Engine/Private/PlayerCameraManager.cpp
+++ b/Engine/Source/Runtime/Engine/Private/PlayerCameraManager.cpp
@@ -591,7 +591,11 @@ void APlayerCameraManager::UpdateCameraLensEffects(const FTViewTarget& OutVT)
 
 void APlayerCameraManager::ApplyAudioFade()
 {
-	if (GEngine)
+	if (OnAudioFadeChangeEvent.IsBound())
+	{
+		OnAudioFadeChangeEvent.Broadcast(true, 1.f - FadeAmount);
+	}
+	else if (GEngine)
 	{
 		UWorld* World = GetWorld();
 		if (World)
@@ -606,7 +610,11 @@ void APlayerCameraManager::ApplyAudioFade()
 
 void APlayerCameraManager::StopAudioFade()
 {
-	if (GEngine)
+	if (OnAudioFadeChangeEvent.IsBound())
+	{
+		OnAudioFadeChangeEvent.Broadcast(false, 1.f);
+	}
+	else if (GEngine)
 	{
 		UWorld* World = GetWorld();
 		if (World)
