fix bad merge

diff --git a/Engine/Source/Runtime/Engine/Private/Materials/MaterialShader.cpp b/Engine/Source/Runtime/Engine/Private/Materials/MaterialShader.cpp
--- a/Engine/Source/Runtime/Engine/Private/Materials/MaterialShader.cpp
+++ b/Engine/Source/Runtime/Engine/Private/Materials/MaterialShader.cpp
@@ -1436,7 +1436,7 @@ void FMaterialShaderMap::Compile(
 
 			FSHA1 IdParameterSetHash;
 			IdParameterSetHash.Reset();
-			IdParameterSetHash.UpdateWithString(*DebugDescription, DebugDescription.Len());
+			IdParameterSetHash.UpdateWithString(*WorkingDebugDescription, WorkingDebugDescription.Len());
 			IdParameterSetHash.Final();
 			uint32* Hash = (uint32*)&IdParameterSetHash.m_digest[0];
 			FString DebugExtension = FString::Printf( TEXT("_%08x%08x"), Hash[0], Hash[1]);
