Back out and re-integrate CL 12091745 following a bad merge

diff --git a/Engine/Source/Runtime/Engine/Private/SceneManagement.cpp b/Engine/Source/Runtime/Engine/Private/SceneManagement.cpp
--- a/Engine/Source/Runtime/Engine/Private/SceneManagement.cpp
+++ b/Engine/Source/Runtime/Engine/Private/SceneManagement.cpp
@@ -519,9 +519,10 @@ int8 ComputeTemporalStaticMeshLOD( const FStaticMeshRenderData* RenderData, cons
 // Ensure we always use the left eye when selecting lods to avoid divergent selections in stereo
 const FSceneView& GetLODView(const FSceneView& InView)
 {
-	if (IStereoRendering::IsASecondaryView(InView) && InView.Family)
+	uint32 LODViewIndex = IStereoRendering::GetLODViewIndex();
+	if (InView.Family && InView.Family->Views.IsValidIndex(LODViewIndex))
 	{
-		return *InView.Family->Views[0];
+		return *InView.Family->Views[LODViewIndex];
 	}
 	else
 	{
@@ -665,8 +666,8 @@ FViewUniformShaderParameters::FViewUniformShaderParameters()
 {
 	FMemory::Memzero(*this);
 
-	FRHITexture* BlackVolume = (GBlackVolumeTexture &&  GBlackVolumeTexture->TextureRHI) ? GBlackVolumeTexture->TextureRHI : GBlackTexture->TextureRHI; // for es2, this might need to be 2d
-	FRHITexture* BlackUintVolume = (GBlackUintVolumeTexture &&  GBlackUintVolumeTexture->TextureRHI) ? GBlackUintVolumeTexture->TextureRHI : GBlackTexture->TextureRHI; // for es2, this might need to be 2d
+	FRHITexture* BlackVolume = (GBlackVolumeTexture &&  GBlackVolumeTexture->TextureRHI) ? GBlackVolumeTexture->TextureRHI : GBlackTexture->TextureRHI;
+	FRHITexture* BlackUintVolume = (GBlackUintVolumeTexture &&  GBlackUintVolumeTexture->TextureRHI) ? GBlackUintVolumeTexture->TextureRHI : GBlackTexture->TextureRHI;
 	check(GBlackVolumeTexture);
 
 	MaterialTextureBilinearClampedSampler = TStaticSamplerState<SF_Bilinear, AM_Clamp, AM_Clamp, AM_Clamp>::GetRHI();
@@ -752,6 +753,8 @@ FViewUniformShaderParameters::FViewUniformShaderParameters()
 	{
 		LightmapSceneData = GBlackTextureWithSRV->ShaderResourceViewRHI;
 	}
+	VTFeedbackBuffer = GEmptyVertexBufferWithUAV->UnorderedAccessViewRHI;
+	QuadOverdraw = GBlackTextureWithUAV->UnorderedAccessViewRHI;
 }
 
 FInstancedViewUniformShaderParameters::FInstancedViewUniformShaderParameters()
@@ -824,6 +827,16 @@ bool FLightCacheInterface::GetVirtualTextureLightmapProducer(ERHIFeatureLevel::T
 
 IMPLEMENT_GLOBAL_SHADER_PARAMETER_STRUCT(FLightmapResourceClusterShaderParameters, "LightmapResourceCluster");
 
+static FRHISamplerState* GetTextureSamplerState(const UTexture* Texture, FRHISamplerState* Default)
+{
+	FRHISamplerState* Result = nullptr;
+	if (Texture && Texture->Resource)
+	{
+		Result = Texture->Resource->SamplerStateRHI;
+	}
+	return Result ? Result : Default;
+}
+
 void GetLightmapClusterResourceParameters(
 	ERHIFeatureLevel::Type FeatureLevel, 
 	const FLightmapClusterResourceInput& Input,
@@ -835,6 +848,22 @@ void GetLightmapClusterResourceParameters(
 	static const auto CVar = IConsoleManager::Get().FindTConsoleVariableDataInt(TEXT("r.VirtualTexturedLightmaps"));
 	const bool bUseVirtualTextures = bAllowHighQualityLightMaps && (CVar->GetValueOnRenderThread() != 0) && UseVirtualTexturing(FeatureLevel);
 
+	Parameters.LightMapTexture = GBlackTexture->TextureRHI;
+	Parameters.SkyOcclusionTexture = GWhiteTexture->TextureRHI;
+	Parameters.AOMaterialMaskTexture = GBlackTexture->TextureRHI;
+	Parameters.StaticShadowTexture = GWhiteTexture->TextureRHI;
+	Parameters.VTLightMapTexture = GBlackTextureWithSRV->ShaderResourceViewRHI;
+	Parameters.VTLightMapTexture_1 = GBlackTextureWithSRV->ShaderResourceViewRHI;
+	Parameters.VTSkyOcclusionTexture = GWhiteTextureWithSRV->ShaderResourceViewRHI;
+	Parameters.VTAOMaterialMaskTexture = GBlackTextureWithSRV->ShaderResourceViewRHI;
+	Parameters.VTStaticShadowTexture = GWhiteTextureWithSRV->ShaderResourceViewRHI;
+	Parameters.LightmapVirtualTexturePageTable0 = GBlackTexture->TextureRHI;
+	Parameters.LightmapVirtualTexturePageTable1 = GBlackTexture->TextureRHI;
+	Parameters.LightMapSampler = GBlackTexture->SamplerStateRHI;
+	Parameters.SkyOcclusionSampler = GWhiteTexture->SamplerStateRHI;
+	Parameters.AOMaterialMaskSampler = GBlackTexture->SamplerStateRHI;
+	Parameters.StaticShadowTextureSampler = GWhiteTexture->SamplerStateRHI;
+
 	if (bUseVirtualTextures)
 	{
 		// this is sometimes called with NULL input to initialize default buffer
@@ -842,34 +871,34 @@ void GetLightmapClusterResourceParameters(
 		if (VirtualTexture && AllocatedVT)
 		{
 			// Bind VT here
-			Parameters.LightMapTexture = AllocatedVT->GetPhysicalTexture((uint32)ELightMapVirtualTextureType::HqLayer0);
-			Parameters.LightMapTexture_1 = AllocatedVT->GetPhysicalTexture((uint32)ELightMapVirtualTextureType::HqLayer1);
+			Parameters.VTLightMapTexture = AllocatedVT->GetPhysicalTextureSRV((uint32)ELightMapVirtualTextureType::HqLayer0, false);
+			Parameters.VTLightMapTexture_1 = AllocatedVT->GetPhysicalTextureSRV((uint32)ELightMapVirtualTextureType::HqLayer1, false);
 
 			if (VirtualTexture->HasLayerForType(ELightMapVirtualTextureType::SkyOcclusion))
 			{
-				Parameters.SkyOcclusionTexture = AllocatedVT->GetPhysicalTexture((uint32)ELightMapVirtualTextureType::SkyOcclusion);
+				Parameters.VTSkyOcclusionTexture = AllocatedVT->GetPhysicalTextureSRV((uint32)ELightMapVirtualTextureType::SkyOcclusion, false);
 			}
 			else
 			{
-				Parameters.SkyOcclusionTexture = GWhiteTexture->TextureRHI;
+				Parameters.VTSkyOcclusionTexture = GWhiteTextureWithSRV->ShaderResourceViewRHI;
 			}
 
 			if (VirtualTexture->HasLayerForType(ELightMapVirtualTextureType::AOMaterialMask))
 			{
-				Parameters.AOMaterialMaskTexture = AllocatedVT->GetPhysicalTexture((uint32)ELightMapVirtualTextureType::AOMaterialMask);
+				Parameters.VTAOMaterialMaskTexture = AllocatedVT->GetPhysicalTextureSRV((uint32)ELightMapVirtualTextureType::AOMaterialMask, false);
 			}
 			else
 			{
-				Parameters.AOMaterialMaskTexture = GBlackTexture->TextureRHI;
+				Parameters.VTAOMaterialMaskTexture = GBlackTextureWithSRV->ShaderResourceViewRHI;
 			}
 
 			if (VirtualTexture->HasLayerForType(ELightMapVirtualTextureType::ShadowMask))
 			{
-				Parameters.StaticShadowTexture = AllocatedVT->GetPhysicalTexture((uint32)ELightMapVirtualTextureType::ShadowMask);
+				Parameters.VTStaticShadowTexture = AllocatedVT->GetPhysicalTextureSRV((uint32)ELightMapVirtualTextureType::ShadowMask, false);
 			}
 			else
 			{
-				Parameters.StaticShadowTexture = GWhiteTexture->TextureRHI;
+				Parameters.VTStaticShadowTexture = GWhiteTextureWithSRV->ShaderResourceViewRHI;
 			}
 
 			FRHITexture* PageTable0 = AllocatedVT->GetPageTableTexture(0u);
@@ -890,36 +919,21 @@ void GetLightmapClusterResourceParameters(
 			Parameters.AOMaterialMaskSampler = TStaticSamplerState<SF_AnisotropicLinear, AM_Clamp, AM_Clamp, AM_Clamp, 0, MaxAniso>::GetRHI();
 			Parameters.StaticShadowTextureSampler = TStaticSamplerState<SF_AnisotropicLinear, AM_Clamp, AM_Clamp, AM_Clamp, 0, MaxAniso>::GetRHI();
 		}
-		else
-		{
-			Parameters.LightMapTexture = GBlackTexture->TextureRHI;
-			Parameters.LightMapTexture_1 = GBlackTexture->TextureRHI;
-			Parameters.SkyOcclusionTexture = GWhiteTexture->TextureRHI;
-			Parameters.AOMaterialMaskTexture = GBlackTexture->TextureRHI;
-			Parameters.StaticShadowTexture = GWhiteTexture->TextureRHI;
-			Parameters.LightmapVirtualTexturePageTable0 = GBlackTexture->TextureRHI;
-			Parameters.LightmapVirtualTexturePageTable1 = GBlackTexture->TextureRHI;
-			Parameters.LightMapSampler = GBlackTexture->SamplerStateRHI;
-			Parameters.SkyOcclusionSampler = GWhiteTexture->SamplerStateRHI;
-			Parameters.AOMaterialMaskSampler = GBlackTexture->SamplerStateRHI;
-			Parameters.StaticShadowTextureSampler = GWhiteTexture->SamplerStateRHI;
-		}
 	}
 	else
 	{
 		const UTexture2D* LightMapTexture = Input.LightMapTextures[bAllowHighQualityLightMaps ? 0 : 1];
 
 		Parameters.LightMapTexture = LightMapTexture ? LightMapTexture->TextureReference.TextureReferenceRHI.GetReference() : GBlackTexture->TextureRHI;
-		Parameters.LightMapTexture_1 = GBlackTexture->TextureRHI;
 		Parameters.SkyOcclusionTexture = Input.SkyOcclusionTexture ? Input.SkyOcclusionTexture->TextureReference.TextureReferenceRHI.GetReference() : GWhiteTexture->TextureRHI;
 		Parameters.AOMaterialMaskTexture = Input.AOMaterialMaskTexture ? Input.AOMaterialMaskTexture->TextureReference.TextureReferenceRHI.GetReference() : GBlackTexture->TextureRHI;
 
-		Parameters.LightMapSampler = (LightMapTexture && LightMapTexture->Resource) ? LightMapTexture->Resource->SamplerStateRHI : GBlackTexture->SamplerStateRHI;
-		Parameters.SkyOcclusionSampler = (Input.SkyOcclusionTexture && Input.SkyOcclusionTexture->Resource) ? Input.SkyOcclusionTexture->Resource->SamplerStateRHI : GWhiteTexture->SamplerStateRHI;
-		Parameters.AOMaterialMaskSampler = (Input.AOMaterialMaskTexture && Input.AOMaterialMaskTexture->Resource) ? Input.AOMaterialMaskTexture->Resource->SamplerStateRHI : GBlackTexture->SamplerStateRHI;
+		Parameters.LightMapSampler = GetTextureSamplerState(LightMapTexture, GBlackTexture->SamplerStateRHI);
+		Parameters.SkyOcclusionSampler = GetTextureSamplerState(Input.SkyOcclusionTexture, GWhiteTexture->SamplerStateRHI);
+		Parameters.AOMaterialMaskSampler = GetTextureSamplerState(Input.AOMaterialMaskTexture, GBlackTexture->SamplerStateRHI);
 
 		Parameters.StaticShadowTexture = Input.ShadowMapTexture ? Input.ShadowMapTexture->TextureReference.TextureReferenceRHI.GetReference() : GWhiteTexture->TextureRHI;
-		Parameters.StaticShadowTextureSampler = (Input.ShadowMapTexture && Input.ShadowMapTexture->Resource) ? Input.ShadowMapTexture->Resource->SamplerStateRHI : GWhiteTexture->SamplerStateRHI;
+		Parameters.StaticShadowTextureSampler = GetTextureSamplerState(Input.ShadowMapTexture, GWhiteTexture->SamplerStateRHI);
 
 		Parameters.LightmapVirtualTexturePageTable0 = GBlackTexture->TextureRHI;
 		Parameters.LightmapVirtualTexturePageTable1 = GBlackTexture->TextureRHI;
@@ -1070,8 +1084,8 @@ void FMeshBatch::PreparePrimitiveUniformBuffer(const FPrimitiveSceneProxy* Primi
 	// Fallback to using the primitive uniform buffer if GPU scene is disabled.
 	if (VertexFactory->GetType()->SupportsPrimitiveIdStream() && !UseGPUScene(GMaxRHIShaderPlatform, FeatureLevel))
 	{
-		for (int32 ElementIndex = 0; ElementIndex < Elements.Num(); ElementIndex++)
-		{
+	    for (int32 ElementIndex = 0; ElementIndex < Elements.Num(); ElementIndex++)
+	    {
 			FMeshBatchElement& MeshElement = Elements[ElementIndex];
 
 			if (!MeshElement.PrimitiveUniformBufferResource)
@@ -1124,7 +1138,7 @@ bool FMeshBatch::Validate(const FPrimitiveSceneProxy* PrimitiveSceneProxy, ERHIF
 	if (!VertexFactory->IsInitialized())
 	{
 		return LogMeshError(TEXT("Mesh has an uninitialized vertex factory!"));
-	}
+}
 
 	for (int32 Index = 0; Index < Elements.Num(); ++Index)
 	{
