fix for FXSystem being reused across level changes.

diff --git a/Engine/Source/Runtime/Engine/Private/World.cpp b/Engine/Source/Runtime/Engine/Private/World.cpp
--- a/Engine/Source/Runtime/Engine/Private/World.cpp
+++ b/Engine/Source/Runtime/Engine/Private/World.cpp
@@ -4336,7 +4336,7 @@ void UWorld::CleanupWorldInternal(bool bSessionEnded, bool bCleanupResources, UW
 
 	SubsystemCollection.Deinitialize();
 
-	if (FXSystem)
+	if(FXSystem && NewWorld == nullptr)
 	{
 		FFXSystemInterface::QueueDestroyGPUSimulation(FXSystem);
 	}
