Making autoexposure cvar static during PPV serialization upgrade.

diff --git a/Engine/Source/Runtime/Engine/Private/Scene.cpp b/Engine/Source/Runtime/Engine/Private/Scene.cpp
--- a/Engine/Source/Runtime/Engine/Private/Scene.cpp
+++ b/Engine/Source/Runtime/Engine/Private/Scene.cpp
@@ -1119,7 +1119,7 @@ void FPostProcessSettings::PostSerialize(const FArchive& Ar)
 			// Calculate an exposure bias to try and keep the look similar from 4.24 to 4.25
 			if (bIsAnyNonDefault)
 			{
-				const auto VarDefaultAutoExposureExtendDefaultLuminanceRange = IConsoleManager::Get().FindTConsoleVariableDataInt(TEXT("r.DefaultFeature.AutoExposure.ExtendDefaultLuminanceRange"));
+				static const auto VarDefaultAutoExposureExtendDefaultLuminanceRange = IConsoleManager::Get().FindTConsoleVariableDataInt(TEXT("r.DefaultFeature.AutoExposure.ExtendDefaultLuminanceRange"));
 				bool bExtendedLuminanceRange = (VarDefaultAutoExposureExtendDefaultLuminanceRange->GetValueOnAnyThread() == 1);
 
 				const float ExtraAutoExposureBias = CalculateEyeAdaptationExposureVersionUpdate(*this, bExtendedLuminanceRange);
@@ -1175,7 +1175,7 @@ void FPostProcessSettings::PostSerialize(const FArchive& Ar)
 
 			if (bIsAnyNonDefault)
 			{
-				const auto VarDefaultAutoExposureExtendDefaultLuminanceRange = IConsoleManager::Get().FindTConsoleVariableDataInt(TEXT("r.DefaultFeature.AutoExposure.ExtendDefaultLuminanceRange"));
+				static const auto VarDefaultAutoExposureExtendDefaultLuminanceRange = IConsoleManager::Get().FindTConsoleVariableDataInt(TEXT("r.DefaultFeature.AutoExposure.ExtendDefaultLuminanceRange"));
 				bool bExtendedLuminanceRange = (VarDefaultAutoExposureExtendDefaultLuminanceRange->GetValueOnAnyThread() == 1);
 
 				const float ExtraAutoExposureBias = CalculateEyeAdaptationExposureVersionUpdate(*this, bExtendedLuminanceRange);
