*Actually* prevent creating null SRVs and buffer. And when resources are streamed in recreat RHI resources if necessary.

diff --git a/Engine/Source/Runtime/Engine/Private/PositionVertexBuffer.cpp b/Engine/Source/Runtime/Engine/Private/PositionVertexBuffer.cpp
--- a/Engine/Source/Runtime/Engine/Private/PositionVertexBuffer.cpp
+++ b/Engine/Source/Runtime/Engine/Private/PositionVertexBuffer.cpp
@@ -191,9 +191,9 @@ void FPositionVertexBuffer::operator=(const FPositionVertexBuffer &Other)
 template <bool bRenderThread>
 FVertexBufferRHIRef FPositionVertexBuffer::CreateRHIBuffer_Internal()
 {
-	if (NumVertices)
+	if (NumVertices && VertexData)
 	{
-		FResourceArrayInterface* RESTRICT ResourceArray = VertexData ? VertexData->GetResourceArray() : nullptr;
+		FResourceArrayInterface* RESTRICT ResourceArray = VertexData->GetResourceArray();
 		const uint32 SizeInBytes = ResourceArray ? ResourceArray->GetResourceDataSize() : 0;
 		FRHIResourceCreateInfo CreateInfo(ResourceArray);
 		CreateInfo.bWithoutNativeResource = !VertexData;
@@ -233,9 +233,9 @@ void FPositionVertexBuffer::InitRHI()
 		// And it can be the case for CPU *and* GPU access: no differenciation today. That is why we create a SRV in this case.
 		// This also avoid setting lots of states on all the members of all the different buffers used by meshes. Follow up: https://jira.it.epicgames.net/browse/UE-69376.
 		bSRV |= (VertexData && VertexData->GetAllowCPUAccess());
-		if(bSRV)
+		if(bSRV && VertexData && VertexBufferRHI)
 		{
-			PositionComponentSRV = RHICreateShaderResourceView(VertexData ? VertexBufferRHI : nullptr, 4, PF_R32_FLOAT);
+			PositionComponentSRV = RHICreateShaderResourceView(VertexBufferRHI, 4, PF_R32_FLOAT);
 		}
 	}
 }
diff --git a/Engine/Source/Runtime/Engine/Private/Rendering/ColorVertexBuffer.cpp b/Engine/Source/Runtime/Engine/Private/Rendering/ColorVertexBuffer.cpp
--- a/Engine/Source/Runtime/Engine/Private/Rendering/ColorVertexBuffer.cpp
+++ b/Engine/Source/Runtime/Engine/Private/Rendering/ColorVertexBuffer.cpp
@@ -389,9 +389,9 @@ uint32 FColorVertexBuffer::GetAllocatedSize() const
 template <bool bRenderThread>
 FVertexBufferRHIRef FColorVertexBuffer::CreateRHIBuffer_Internal()
 {
-	if (NumVertices)
+	if (NumVertices && VertexData)
 	{
-		FResourceArrayInterface* RESTRICT ResourceArray = VertexData ? VertexData->GetResourceArray() : nullptr;
+		FResourceArrayInterface* RESTRICT ResourceArray = VertexData->GetResourceArray();
 		const uint32 SizeInBytes = ResourceArray ? ResourceArray->GetResourceDataSize() : 0;
 		FRHIResourceCreateInfo CreateInfo(ResourceArray);
 		CreateInfo.bWithoutNativeResource = !VertexData;
@@ -423,9 +423,9 @@ void FColorVertexBuffer::InitRHI()
 
 	VertexBufferRHI = CreateRHIBuffer_RenderThread();
 
-	if (VertexBufferRHI && RHISupportsManualVertexFetch(GMaxRHIShaderPlatform))
+	if (VertexData && VertexBufferRHI && RHISupportsManualVertexFetch(GMaxRHIShaderPlatform))
 	{
-		ColorComponentsSRV = RHICreateShaderResourceView(VertexData ? VertexBufferRHI : nullptr, 4, PF_R8G8B8A8);
+		ColorComponentsSRV = RHICreateShaderResourceView(VertexBufferRHI, 4, PF_R8G8B8A8);
 	}
 }
 
diff --git a/Engine/Source/Runtime/Engine/Private/Rendering/SkeletalMeshVertexClothBuffer.cpp b/Engine/Source/Runtime/Engine/Private/Rendering/SkeletalMeshVertexClothBuffer.cpp
--- a/Engine/Source/Runtime/Engine/Private/Rendering/SkeletalMeshVertexClothBuffer.cpp
+++ b/Engine/Source/Runtime/Engine/Private/Rendering/SkeletalMeshVertexClothBuffer.cpp
@@ -75,9 +75,9 @@ void FSkeletalMeshVertexClothBuffer::ClearMetaData()
 template <bool bRenderThread>
 FVertexBufferRHIRef FSkeletalMeshVertexClothBuffer::CreateRHIBuffer_Internal()
 {
-	if (NumVertices)
+	if (NumVertices && VertexData)
 	{
-		FResourceArrayInterface* ResourceArray = VertexData ? VertexData->GetResourceArray() : nullptr;
+		FResourceArrayInterface* ResourceArray = VertexData->GetResourceArray();
 		const uint32 SizeInBytes = ResourceArray ? ResourceArray->GetResourceDataSize() : 0;
 		const uint32 BuffFlags = BUF_Static | BUF_ShaderResource;
 		FRHIResourceCreateInfo CreateInfo(ResourceArray);
@@ -114,9 +114,9 @@ void FSkeletalMeshVertexClothBuffer::InitRHI()
 
 	VertexBufferRHI = CreateRHIBuffer_RenderThread();
 
-	if (VertexBufferRHI)
+	if (VertexData && VertexBufferRHI)
 	{
-		VertexBufferSRV = RHICreateShaderResourceView(VertexData ? VertexBufferRHI : nullptr, sizeof(FVector4), PF_A32B32G32R32F);
+		VertexBufferSRV = RHICreateShaderResourceView(VertexBufferRHI, sizeof(FVector4), PF_A32B32G32R32F);
 	}
 }
 
diff --git a/Engine/Source/Runtime/Engine/Private/Rendering/StaticMeshVertexBuffer.cpp b/Engine/Source/Runtime/Engine/Private/Rendering/StaticMeshVertexBuffer.cpp
--- a/Engine/Source/Runtime/Engine/Private/Rendering/StaticMeshVertexBuffer.cpp
+++ b/Engine/Source/Runtime/Engine/Private/Rendering/StaticMeshVertexBuffer.cpp
@@ -334,17 +334,18 @@ void FStaticMeshVertexBuffer::InitRHI()
 
 	TangentsVertexBuffer.VertexBufferRHI = CreateTangentsRHIBuffer_RenderThread();
 	TexCoordVertexBuffer.VertexBufferRHI = CreateTexCoordRHIBuffer_RenderThread();
-	if (TangentsVertexBuffer.VertexBufferRHI && (RHISupportsManualVertexFetch(GMaxRHIShaderPlatform) || IsGPUSkinCacheAvailable(GMaxRHIShaderPlatform)))
+	if (TangentsData && TangentsVertexBuffer.VertexBufferRHI && (RHISupportsManualVertexFetch(GMaxRHIShaderPlatform) || IsGPUSkinCacheAvailable(GMaxRHIShaderPlatform)))
 	{
 		TangentsSRV = RHICreateShaderResourceView(
-			TangentsData ? TangentsVertexBuffer.VertexBufferRHI : nullptr,
+			TangentsVertexBuffer.VertexBufferRHI,
 			GetUseHighPrecisionTangentBasis() ? 8 : 4,
 			GetUseHighPrecisionTangentBasis() ? PF_R16G16B16A16_SNORM : PF_R8G8B8A8_SNORM);
 	}
-	if (TexCoordVertexBuffer.VertexBufferRHI && RHISupportsManualVertexFetch(GMaxRHIShaderPlatform))
+
+	if (TexcoordData && TexCoordVertexBuffer.VertexBufferRHI && RHISupportsManualVertexFetch(GMaxRHIShaderPlatform))
 	{
 		TextureCoordinatesSRV = RHICreateShaderResourceView(
-			TexcoordData ? TexCoordVertexBuffer.VertexBufferRHI : nullptr,
+			TexCoordVertexBuffer.VertexBufferRHI,
 			GetUseFullPrecisionUVs() ? 8 : 4,
 			GetUseFullPrecisionUVs() ? PF_G32R32F : PF_G16R16F);
 	}
diff --git a/Engine/Source/Runtime/Engine/Public/RawIndexBuffer.h b/Engine/Source/Runtime/Engine/Public/RawIndexBuffer.h
--- a/Engine/Source/Runtime/Engine/Public/RawIndexBuffer.h
+++ b/Engine/Source/Runtime/Engine/Public/RawIndexBuffer.h
@@ -370,9 +370,9 @@ public:
 	{
 		IndexBufferRHI = CreateRHIBuffer_RenderThread();
 
-		if (IndexBufferRHI && IsSRVNeeded() && Indices.Num())
+		if (Indices.Num() && IndexBufferRHI && IsSRVNeeded() && Indices.Num())
 		{
-			SRVValue = RHICreateShaderResourceView(Indices.Num() ? IndexBufferRHI : nullptr);
+			SRVValue = RHICreateShaderResourceView(IndexBufferRHI);
 		}
 	}
 	
@@ -472,6 +472,12 @@ public:
 	template <uint32 MaxNumUpdates>
 	void InitRHIForStreaming(FRHIIndexBuffer* IntermediateBuffer, TRHIResourceUpdateBatcher<MaxNumUpdates>& Batcher)
 	{
+		if (!SRVValue || !IndexBufferRHI)
+		{
+			ReleaseRHI();
+			InitRHI();
+		}
+
 		if (IndexBufferRHI && IntermediateBuffer)
 		{
 			Batcher.QueueUpdateRequest(IndexBufferRHI, IntermediateBuffer);
diff --git a/Engine/Source/Runtime/Engine/Public/Rendering/ColorVertexBuffer.h b/Engine/Source/Runtime/Engine/Public/Rendering/ColorVertexBuffer.h
--- a/Engine/Source/Runtime/Engine/Public/Rendering/ColorVertexBuffer.h
+++ b/Engine/Source/Runtime/Engine/Public/Rendering/ColorVertexBuffer.h
@@ -145,6 +145,12 @@ public:
 	template <uint32 MaxNumUpdates>
 	void InitRHIForStreaming(FRHIVertexBuffer* IntermediateBuffer, TRHIResourceUpdateBatcher<MaxNumUpdates>& Batcher)
 	{
+		if (!ColorComponentsSRV || !VertexBufferRHI)
+		{
+			ReleaseRHI();
+			InitRHI();
+		}
+
 		if (VertexBufferRHI && IntermediateBuffer)
 		{
 			Batcher.QueueUpdateRequest(VertexBufferRHI, IntermediateBuffer);
diff --git a/Engine/Source/Runtime/Engine/Public/Rendering/PositionVertexBuffer.h b/Engine/Source/Runtime/Engine/Public/Rendering/PositionVertexBuffer.h
--- a/Engine/Source/Runtime/Engine/Public/Rendering/PositionVertexBuffer.h
+++ b/Engine/Source/Runtime/Engine/Public/Rendering/PositionVertexBuffer.h
@@ -102,6 +102,12 @@ public:
 	template <uint32 MaxNumUpdates>
 	void InitRHIForStreaming(FRHIVertexBuffer* IntermediateBuffer, TRHIResourceUpdateBatcher<MaxNumUpdates>& Batcher)
 	{
+		if (!PositionComponentSRV || !VertexBufferRHI)
+		{
+			ReleaseRHI();
+			InitRHI();
+		}
+
 		check(VertexBufferRHI);
 		if (IntermediateBuffer)
 		{
diff --git a/Engine/Source/Runtime/Engine/Public/Rendering/SkeletalMeshVertexClothBuffer.h b/Engine/Source/Runtime/Engine/Public/Rendering/SkeletalMeshVertexClothBuffer.h
--- a/Engine/Source/Runtime/Engine/Public/Rendering/SkeletalMeshVertexClothBuffer.h
+++ b/Engine/Source/Runtime/Engine/Public/Rendering/SkeletalMeshVertexClothBuffer.h
@@ -131,6 +131,12 @@ public:
 	template <uint32 MaxNumUpdates>
 	void InitRHIForStreaming(FRHIVertexBuffer* IntermediateBuffer, TRHIResourceUpdateBatcher<MaxNumUpdates>& Batcher)
 	{
+		if (!VertexBufferSRV || !VertexBufferRHI)
+		{
+			ReleaseRHI();
+			InitRHI();
+		}
+
 		if (VertexBufferRHI && IntermediateBuffer)
 		{
 			check(VertexBufferSRV);
diff --git a/Engine/Source/Runtime/Engine/Public/Rendering/SkinWeightVertexBuffer.h b/Engine/Source/Runtime/Engine/Public/Rendering/SkinWeightVertexBuffer.h
--- a/Engine/Source/Runtime/Engine/Public/Rendering/SkinWeightVertexBuffer.h
+++ b/Engine/Source/Runtime/Engine/Public/Rendering/SkinWeightVertexBuffer.h
@@ -135,6 +135,12 @@ public:
 	template <uint32 MaxNumUpdates>
 	void InitRHIForStreaming(FRHIVertexBuffer* IntermediateBuffer, TRHIResourceUpdateBatcher<MaxNumUpdates>& Batcher)
 	{
+		if (!SRVValue || VertexBufferRHI)
+		{
+			ReleaseRHI();
+			InitRHI();
+		}
+
 		if (VertexBufferRHI && IntermediateBuffer)
 		{
 			check(SRVValue);
@@ -300,6 +306,12 @@ public:
 	template <uint32 MaxNumUpdates>
 	void InitRHIForStreaming(FRHIVertexBuffer* IntermediateBuffer, TRHIResourceUpdateBatcher<MaxNumUpdates>& Batcher)
 	{
+		if (!SRVValue || !VertexBufferRHI)
+		{
+			ReleaseRHI();
+			InitRHI();
+		}
+		
 		if (VertexBufferRHI && IntermediateBuffer)
 		{
 			check(SRVValue);
diff --git a/Engine/Source/Runtime/Engine/Public/Rendering/StaticMeshVertexBuffer.h b/Engine/Source/Runtime/Engine/Public/Rendering/StaticMeshVertexBuffer.h
--- a/Engine/Source/Runtime/Engine/Public/Rendering/StaticMeshVertexBuffer.h
+++ b/Engine/Source/Runtime/Engine/Public/Rendering/StaticMeshVertexBuffer.h
@@ -401,7 +401,12 @@ public:
 		FRHIVertexBuffer* IntermediateTexCoordBuffer,
 		TRHIResourceUpdateBatcher<MaxNumUpdates>& Batcher)
 	{
-		check(TangentsVertexBuffer.VertexBufferRHI && TexCoordVertexBuffer.VertexBufferRHI);
+		if (!TangentsSRV || !TextureCoordinatesSRV || !TangentsVertexBuffer.VertexBufferRHI || !TexCoordVertexBuffer.VertexBufferRHI)
+		{
+			ReleaseRHI();
+			InitRHI();
+		}
+
 		if (IntermediateTangentsBuffer)
 		{
 			Batcher.QueueUpdateRequest(TangentsVertexBuffer.VertexBufferRHI, IntermediateTangentsBuffer);
@@ -414,6 +419,7 @@ public:
 					GetUseHighPrecisionTangentBasis() ? (uint8)PF_R16G16B16A16_SNORM : (uint8)PF_R8G8B8A8_SNORM);;
 			}
 		}
+
 		if (IntermediateTexCoordBuffer)
 		{
 			Batcher.QueueUpdateRequest(TexCoordVertexBuffer.VertexBufferRHI, IntermediateTexCoordBuffer);
