Shader compile log messages now appear during shader compile instead of after all shaders compiled.

diff --git a/Engine/Source/Runtime/Engine/Private/ShaderCompiler/ShaderCompiler.cpp b/Engine/Source/Runtime/Engine/Private/ShaderCompiler/ShaderCompiler.cpp
--- a/Engine/Source/Runtime/Engine/Private/ShaderCompiler/ShaderCompiler.cpp
+++ b/Engine/Source/Runtime/Engine/Private/ShaderCompiler/ShaderCompiler.cpp
@@ -2173,7 +2173,7 @@ void FShaderCompilingManager::BlockOnShaderMapCompletion(const TArray<int32>& Sh
 	if (bAllowAsynchronousShaderCompiling)
 	{
 		int32 NumPendingJobs = 0;
-
+		int32 LogCounter = 0;
 		do 
 		{
 			Thread->CheckHealth();
@@ -2207,6 +2207,13 @@ void FShaderCompilingManager::BlockOnShaderMapCompletion(const TArray<int32>& Sh
 			{
 				// Yield CPU time while waiting
 				FPlatformProcess::Sleep(.01f);
+
+				// Flush threaded logs around every 500ms or so based on Sleep of 0.01f seconds above
+				if (++LogCounter > 50)
+				{
+					LogCounter = 0;
+					GLog->FlushThreadedLogs();
+				}
 			}
 		} 
 		while (NumPendingJobs > 0);
