Cameras: Let people use single-instance shakes with shake source actors if they want.

diff --git a/Engine/Source/Runtime/Engine/Private/Camera/CameraModifier_CameraShake.cpp b/Engine/Source/Runtime/Engine/Private/Camera/CameraModifier_CameraShake.cpp
--- a/Engine/Source/Runtime/Engine/Private/Camera/CameraModifier_CameraShake.cpp
+++ b/Engine/Source/Runtime/Engine/Private/Camera/CameraModifier_CameraShake.cpp
@@ -100,18 +100,14 @@ UCameraShake* UCameraModifier_CameraShake::AddCameraShake(TSubclassOf<UCameraSha
 				UCameraShake* ShakeInst = ShakeInfo.ShakeInstance;
 				if (ShakeInst && (ShakeClass == ShakeInst->GetClass()))
 				{
-					// Just restart the existing shake
+					// Just restart the existing shake, possibly at the new location.
+					// Warning: if the shake source changes, this would "teleport" the shake, which might create a visual
+					// artifact, if the user didn't intend to do this.
+					ShakeInfo.ShakeSource = SourceComponent;
 					ShakeInst->PlayShake(CameraOwner, Scale, Params.PlaySpace, Params.UserPlaySpaceRot);
 					return ShakeInst;
 				}
 			}
-
-			// No existing instance... we'll build a new one, but let's null the
-			// source in case one was specified (we don't support a source because
-			// it would in some cases restart the shake by teleporting it to another
-			// location, which would look bad).
-			ensureMsgf(SourceComponent == nullptr, TEXT("CameraShake assets with SingleInstance enabled shouldn't be located to a source."));
-			SourceComponent = nullptr;
 		}
 
 		// Try to find a shake in the expired pool
