Temporary fix for crashes when exiting the game / editor

#tests 99 launches and exit of cooked PC client

diff --git a/Engine/Source/Runtime/Engine/Private/SoundEffectPreset.cpp b/Engine/Source/Runtime/Engine/Private/SoundEffectPreset.cpp
--- a/Engine/Source/Runtime/Engine/Private/SoundEffectPreset.cpp
+++ b/Engine/Source/Runtime/Engine/Private/SoundEffectPreset.cpp
@@ -5,6 +5,7 @@
 #include "Sound/SoundEffectSource.h"
 #include "Engine/Engine.h"
 #include "AudioDeviceManager.h"
+#include "CoreGlobals.h"
 
 USoundEffectPreset::USoundEffectPreset(const FObjectInitializer& ObjectInitializer)
 	: Super(ObjectInitializer)
@@ -66,10 +67,13 @@ void USoundEffectPreset::AddReferencedEffects(FReferenceCollector& InCollector)
 
 void USoundEffectPreset::BeginDestroy()
 {
-	IterateEffects<FSoundEffectBase>([](FSoundEffectBase& Instance)
+	if (!GExitPurge)
 	{
-		Instance.ClearPreset(false /* bRemoveFromPreset */);
-	});
+		IterateEffects<FSoundEffectBase>([](FSoundEffectBase& Instance)
+		{
+			Instance.ClearPreset(false /* bRemoveFromPreset */);
+		});
+	}
 	Instances.Reset();
 
 	Super::BeginDestroy();
