Fix for crash when OnWorldCleanup is called with NewWorld == this. This happens when loading levels through the content browser and modifying the "Levels". In this case do not clean up the persistent uniformbuffers of the FScene, as it will not be deleted

diff --git a/Engine/Source/Runtime/Engine/Private/World.cpp b/Engine/Source/Runtime/Engine/Private/World.cpp
--- a/Engine/Source/Runtime/Engine/Private/World.cpp
+++ b/Engine/Source/Runtime/Engine/Private/World.cpp
@@ -4201,7 +4201,8 @@ void UWorld::CleanupWorldInternal(bool bSessionEnded, bool bCleanupResources, UW
 	{
 		return;
 	}
-
+	bool bWorldChanged = NewWorld != this;
+	FPlatformMisc::LowLevelOutputDebugStringf(TEXT("XYXYXY XYXYXY Clearnupworld %p -> %p\n"), this, NewWorld);
 	CleanupWorldTag = CleanupWorldGlobalTag;
 
 	UE_LOG(LogWorld, Log, TEXT("UWorld::CleanupWorld for %s, bSessionEnded=%s, bCleanupResources=%s"), *GetName(), bSessionEnded ? TEXT("true") : TEXT("false"), bCleanupResources ? TEXT("true") : TEXT("false"));
@@ -4219,7 +4220,7 @@ void UWorld::CleanupWorldInternal(bool bSessionEnded, bool bCleanupResources, UW
 
 	FWorldDelegates::OnWorldCleanup.Broadcast(this, bSessionEnded, bCleanupResources);
 
-	GetRendererModule().OnWorldCleanup(this, bSessionEnded, bCleanupResources);
+	GetRendererModule().OnWorldCleanup(this, bSessionEnded, bCleanupResources, bWorldChanged);
 
 	if (AISystem != nullptr)
 	{
@@ -4336,7 +4337,7 @@ void UWorld::CleanupWorldInternal(bool bSessionEnded, bool bCleanupResources, UW
 
 	SubsystemCollection.Deinitialize();
 
-	if(FXSystem && NewWorld == nullptr)
+	if(FXSystem && bWorldChanged)
 	{
 		FFXSystemInterface::QueueDestroyGPUSimulation(FXSystem);
 	}
diff --git a/Engine/Source/Runtime/RenderCore/Public/RendererInterface.h b/Engine/Source/Runtime/RenderCore/Public/RendererInterface.h
--- a/Engine/Source/Runtime/RenderCore/Public/RendererInterface.h
+++ b/Engine/Source/Runtime/RenderCore/Public/RendererInterface.h
@@ -676,7 +676,7 @@ public:
 	/** Forces reallocation of scene render targets. */
 	virtual void ReallocateSceneRenderTargets() = 0;
 
-	virtual void OnWorldCleanup(UWorld* World, bool bSessionEnded, bool bCleanupResources) = 0;
+	virtual void OnWorldCleanup(UWorld* World, bool bSessionEnded, bool bCleanupResources, bool bWorldChanged) = 0;
 
 	/** Sets the buffer size of the render targets. */
 	virtual void SceneRenderTargetsSetBufferSize(uint32 SizeX, uint32 SizeY) = 0;
diff --git a/Engine/Source/Runtime/Renderer/Private/Renderer.cpp b/Engine/Source/Runtime/Renderer/Private/Renderer.cpp
--- a/Engine/Source/Runtime/Renderer/Private/Renderer.cpp
+++ b/Engine/Source/Runtime/Renderer/Private/Renderer.cpp
@@ -71,17 +71,17 @@ void FRendererModule::ReallocateSceneRenderTargets()
 	FSceneRenderTargets::GetGlobalUnsafe().UpdateRHI();
 }
 
-void FRendererModule::OnWorldCleanup(UWorld* World, bool bSessionEnded, bool bCleanupResources)
+void FRendererModule::OnWorldCleanup(UWorld* World, bool bSessionEnded, bool bCleanupResources, bool bWorldChanged)
 {
 	FSceneInterface* Scene = World->Scene;
 	ENQUEUE_RENDER_COMMAND(OnWorldCleanup)(
-	[Scene](FRHICommandListImmediate& RHICmdList)
+	[Scene, bWorldChanged](FRHICommandListImmediate& RHICmdList)
 	{
 		if(bFlushRenderTargetsOnWorldCleanup > 0)
 		{
 			GRenderTargetPool.FreeUnusedResources();
 		}
-		if(Scene)
+		if(bWorldChanged && Scene)
 		{
 			Scene->OnWorldCleanup();
 		}
diff --git a/Engine/Source/Runtime/Renderer/Private/RendererModule.h b/Engine/Source/Runtime/Renderer/Private/RendererModule.h
--- a/Engine/Source/Runtime/Renderer/Private/RendererModule.h
+++ b/Engine/Source/Runtime/Renderer/Private/RendererModule.h
@@ -45,7 +45,7 @@ public:
 	virtual FSceneViewStateInterface* AllocateViewState() override;
 	virtual uint32 GetNumDynamicLightsAffectingPrimitive(const FPrimitiveSceneInfo* PrimitiveSceneInfo,const FLightCacheInterface* LCI) override;
 	virtual void ReallocateSceneRenderTargets() override;
-	virtual void OnWorldCleanup(UWorld* World, bool bSessionEnded, bool bCleanupResources) override;
+	virtual void OnWorldCleanup(UWorld* World, bool bSessionEnded, bool bCleanupResources, bool bWorldChanged) override;
 	virtual void SceneRenderTargetsSetBufferSize(uint32 SizeX, uint32 SizeY) override;
 	virtual void InitializeSystemTextures(FRHICommandListImmediate& RHICmdList);
 	virtual void DrawTileMesh(FRHICommandListImmediate& RHICmdList, FMeshPassProcessorRenderState& DrawRenderState, const FSceneView& View, FMeshBatch& Mesh, bool bIsHitTesting, const FHitProxyId& HitProxyId, bool bUse128bitRT = false) override;
