Fix for editor crashing when pressing build.

diff --git a/Engine/Source/Runtime/Engine/Private/World.cpp b/Engine/Source/Runtime/Engine/Private/World.cpp
--- a/Engine/Source/Runtime/Engine/Private/World.cpp
+++ b/Engine/Source/Runtime/Engine/Private/World.cpp
@@ -4201,7 +4201,7 @@ void UWorld::CleanupWorldInternal(bool bSessionEnded, bool bCleanupResources, UW
 	{
 		return;
 	}
-	bool bWorldChanged = NewWorld != this;
+	bool bWorldChanged = NewWorld != this && NewWorld != nullptr;
 	FPlatformMisc::LowLevelOutputDebugStringf(TEXT("XYXYXY XYXYXY Clearnupworld %p -> %p\n"), this, NewWorld);
 	CleanupWorldTag = CleanupWorldGlobalTag;
 
