Making OnAudioFadeChange event less spammy.

#rnx

diff --git a/Engine/Source/Runtime/Engine/Classes/Camera/PlayerCameraManager.h b/Engine/Source/Runtime/Engine/Classes/Camera/PlayerCameraManager.h
--- a/Engine/Source/Runtime/Engine/Classes/Camera/PlayerCameraManager.h
+++ b/Engine/Source/Runtime/Engine/Classes/Camera/PlayerCameraManager.h
@@ -15,7 +15,7 @@ class AEmitterCameraLensEffectBase;
 class APlayerController;
 class FDebugDisplayInfo;
 
-DECLARE_DYNAMIC_MULTICAST_DELEGATE_TwoParams(FOnAudioFadeChangeSignature, bool, bFadeOut, float, FadeAmount);
+DECLARE_DYNAMIC_MULTICAST_DELEGATE_TwoParams(FOnAudioFadeChangeSignature, bool, bFadeOut, float, FadeTime);
 
 /** 
  * Options that define how to blend when changing view targets. 
@@ -296,7 +296,7 @@ public:
 	UPROPERTY(EditAnywhere, BlueprintReadWrite, Category = Debug)
 	FVector ViewTargetOffset;
 
-	/** If bound, broadcast instead of manually altering audio device's master volume directly */
+	/** If bound, broadcast on fade start (with fade time) instead of manually altering audio device's master volume directly */
 	UPROPERTY(BlueprintAssignable, Transient)
 	FOnAudioFadeChangeSignature OnAudioFadeChangeEvent;
 
diff --git a/Engine/Source/Runtime/Engine/Private/PlayerCameraManager.cpp b/Engine/Source/Runtime/Engine/Private/PlayerCameraManager.cpp
--- a/Engine/Source/Runtime/Engine/Private/PlayerCameraManager.cpp
+++ b/Engine/Source/Runtime/Engine/Private/PlayerCameraManager.cpp
@@ -591,11 +591,13 @@ void APlayerCameraManager::UpdateCameraLensEffects(const FTViewTarget& OutVT)
 
 void APlayerCameraManager::ApplyAudioFade()
 {
+	// If an audio fade event has been bound, we'd like it to override the default fade behavior.
 	if (OnAudioFadeChangeEvent.IsBound())
 	{
-		OnAudioFadeChangeEvent.Broadcast(true, 1.f - FadeAmount);
+		return;
 	}
-	else if (GEngine)
+
+	if (GEngine)
 	{
 		UWorld* World = GetWorld();
 		if (World)
@@ -610,11 +612,13 @@ void APlayerCameraManager::ApplyAudioFade()
 
 void APlayerCameraManager::StopAudioFade()
 {
+	// If an audio fade event has been bound, we'd like it to override the default fade behavior.
 	if (OnAudioFadeChangeEvent.IsBound())
 	{
-		OnAudioFadeChangeEvent.Broadcast(false, 1.f);
+		return;
 	}
-	else if (GEngine)
+
+	if (GEngine)
 	{
 		UWorld* World = GetWorld();
 		if (World)
@@ -1394,6 +1398,12 @@ void APlayerCameraManager::StartCameraFade(float FromAlpha, float ToAlpha, float
 	FadeTimeRemaining = InFadeTime;
 	bFadeAudio = bInFadeAudio;
 
+	if (bInFadeAudio && OnAudioFadeChangeEvent.IsBound() && !FMath::IsNearlyEqual(FromAlpha, ToAlpha))
+	{
+		const bool bFadeOut = FromAlpha < ToAlpha;
+		OnAudioFadeChangeEvent.Broadcast(bFadeOut, FadeTime);
+	}
+
 	bAutoAnimateFade = true;
 	bHoldFadeWhenFinished = bInHoldWhenFinished;
 }
