[fix robomerge] If CHAOS_CHECKED, then USE_BODYINSTANCE_DEBUG_NAMES.

diff --git a/Engine/Source/Runtime/Engine/Classes/PhysicsEngine/BodyInstance.h b/Engine/Source/Runtime/Engine/Classes/PhysicsEngine/BodyInstance.h
--- a/Engine/Source/Runtime/Engine/Classes/PhysicsEngine/BodyInstance.h
+++ b/Engine/Source/Runtime/Engine/Classes/PhysicsEngine/BodyInstance.h
@@ -54,7 +54,7 @@ namespace EDOFMode
 
 struct FBodyInstnace;
 
-#define USE_BODYINSTANCE_DEBUG_NAMES ((WITH_EDITORONLY_DATA || UE_BUILD_DEBUG || LOOKING_FOR_PERF_ISSUES) && !(UE_BUILD_SHIPPING || UE_BUILD_TEST) && !NO_LOGGING)
+#define USE_BODYINSTANCE_DEBUG_NAMES ((WITH_EDITORONLY_DATA || UE_BUILD_DEBUG || LOOKING_FOR_PERF_ISSUES || CHAOS_CHECKED) && !(UE_BUILD_SHIPPING || UE_BUILD_TEST) && !NO_LOGGING)
 
 /** Helper struct to specify spawn behavior */
 struct FInitBodySpawnParams
