Resolve robomerge by accepting source.

If CHAOS_CHECKED, then USE_BODYINSTANCE_DEBUG_NAMES.

diff --git a/Engine/Source/Runtime/Engine/Classes/PhysicsEngine/BodyInstance.h b/Engine/Source/Runtime/Engine/Classes/PhysicsEngine/BodyInstance.h
--- a/Engine/Source/Runtime/Engine/Classes/PhysicsEngine/BodyInstance.h
+++ b/Engine/Source/Runtime/Engine/Classes/PhysicsEngine/BodyInstance.h
@@ -52,130 +52,8 @@ namespace EDOFMode
 	};
 }
 
-struct FBodyInstnace;
-
-#define USE_BODYINSTANCE_DEBUG_NAMES ((WITH_EDITORONLY_DATA || UE_BUILD_DEBUG || LOOKING_FOR_PERF_ISSUES || CHAOS_CHECKED) && !(UE_BUILD_SHIPPING || UE_BUILD_TEST) && !NO_LOGGING)
-
-/** Helper struct to specify spawn behavior */
-struct FInitBodySpawnParams
-{
-	ENGINE_API FInitBodySpawnParams(const UPrimitiveComponent* PrimComp);
-
-	/** Whether the created physics actor will be static */
-	bool bStaticPhysics;
-
-	/** Whether to use the BodySetup's PhysicsType to override if the instance simulates*/
-	bool bPhysicsTypeDeterminesSimulation;
-
-	/** An aggregate to place the body into */
-	FPhysicsAggregateHandle Aggregate;
-};
-
-struct FInitBodiesHelperBase
-{
-	ENGINE_API FInitBodiesHelperBase(TArray<FBodyInstance*>& InBodies, TArray<FTransform>& InTransforms, class UBodySetup* InBodySetup, class UPrimitiveComponent* InPrimitiveComp, FPhysScene* InRBScene, const FInitBodySpawnParams& InSpawnParams, FPhysicsAggregateHandle InAggregate);
-
-	FInitBodiesHelperBase(const FInitBodiesHelperBase& InHelper) = delete;
-	FInitBodiesHelperBase(FInitBodiesHelperBase&& InHelper) = delete;
-	FInitBodiesHelperBase& operator=(const FInitBodiesHelperBase& InHelper) = delete;
-	FInitBodiesHelperBase& operator=(FInitBodiesHelperBase&& InHelper) = delete;
-
-	FORCEINLINE bool IsStatic() const { return bStatic; }
-
-	//The arguments passed into InitBodies
-	TArray<FBodyInstance*>& Bodies;   
-	TArray<FTransform>& Transforms;
-	class UBodySetup* BodySetup;
-	class UPrimitiveComponent* PrimitiveComp;
-	FPhysScene* PhysScene;
-	FPhysicsAggregateHandle Aggregate;
-
-#if USE_BODYINSTANCE_DEBUG_NAMES
-	FString DebugName;
-	TSharedPtr<TArray<ANSICHAR>> PhysXName; // Get rid of ANSICHAR in physics
-#endif
-
-	//The constants shared between PhysX and Box2D
-	bool bStatic;
-	bool bInstanceSimulatePhysics;
-	float InstanceBlendWeight;
-
-	const USkeletalMeshComponent* SkelMeshComp;
-
-	const FInitBodySpawnParams& SpawnParams;
-
-	bool DisableQueryOnlyActors;
-
-	// Return to actor ref
-	void CreateActor_AssumesLocked(FBodyInstance* Instance, const FTransform& Transform) const;
-	bool CreateShapes_AssumesLocked(FBodyInstance* Instance) const;
-
-	// Takes actor ref arrays.
-	// #PHYS2 this used to return arrays of low-level physics bodies, which would be added to scene in InitBodies. Should it still do that, rather then later iterate over BodyInstances to get phys actor refs?
-	bool CreateShapesAndActors();
-	void InitBodies();
-
-protected:
-	void UpdateSimulatingAndBlendWeight();
-
-};
-
-template <bool bCompileStatic>
-struct FInitBodiesHelper : public FInitBodiesHelperBase
-{
-	FInitBodiesHelper(TArray<FBodyInstance*>& InBodies, TArray<FTransform>& InTransforms, class UBodySetup* InBodySetup, class UPrimitiveComponent* InPrimitiveComp, FPhysScene* InRBScene, const FInitBodySpawnParams& InSpawnParams, FPhysicsAggregateHandle InAggregate)
-	: FInitBodiesHelperBase(InBodies, InTransforms, InBodySetup, InPrimitiveComp, InRBScene, InSpawnParams, InAggregate)
-	{
-		//Compute all the needed constants
-		bStatic = bCompileStatic || SpawnParams.bStaticPhysics;
-		SkelMeshComp = bCompileStatic ? nullptr : Cast<USkeletalMeshComponent>(PrimitiveComp);
-		if(SpawnParams.bPhysicsTypeDeterminesSimulation)
-		{
-			this->UpdateSimulatingAndBlendWeight();
-		}
-	}
-};
-
-template <bool bCompileStatic>
-struct FInitBodiesHelperWithData : public FInitBodiesHelperBase
-{
-	FInitBodiesHelperWithData() { check(false); }
-	FInitBodiesHelperWithData(TArray<FBodyInstance*>&& InBodies, TArray<FTransform>&& InTransforms, class UBodySetup* InBodySetup, class UPrimitiveComponent* InPrimitiveComp, FPhysScene* InRBScene, const FInitBodySpawnParams& InSpawnParams, FPhysicsAggregateHandle InAggregate)
-	: FInitBodiesHelperBase(OwnedBodies, OwnedTransforms, InBodySetup, InPrimitiveComp, InRBScene, InSpawnParams, InAggregate), OwnedBodies(MoveTemp(InBodies)), OwnedTransforms(MoveTemp(InTransforms))
-	{
-		//Compute all the needed constants
-		bStatic = bCompileStatic || SpawnParams.bStaticPhysics;
-		SkelMeshComp = bCompileStatic ? nullptr : Cast<USkeletalMeshComponent>(PrimitiveComp);
-		if(SpawnParams.bPhysicsTypeDeterminesSimulation)
-		{
-			this->UpdateSimulatingAndBlendWeight();
-		}
-	}
-
-	FInitBodiesHelperWithData(const FInitBodiesHelperWithData& InHelper)
-	: FInitBodiesHelperBase(OwnedBodies, OwnedTransforms, InHelper.BodySetup, InHelper.PrimitiveComp, InHelper.PhysScene, InHelper.SpawnParams, InHelper.Aggregate), OwnedBodies(InHelper.OwnedBodies), OwnedTransforms(InHelper.OwnedTransforms)
-	{
-		ensure(false);
-	}
-
-	FInitBodiesHelperWithData(FInitBodiesHelperWithData&& InHelper)
-	: FInitBodiesHelperBase(OwnedBodies, OwnedTransforms, InHelper.BodySetup, InHelper.PrimitiveComp, InHelper.PhysScene, InHelper.SpawnParams, InHelper.Aggregate), OwnedBodies(MoveTemp(InHelper.OwnedBodies)), OwnedTransforms(MoveTemp(InHelper.OwnedTransforms))
-	{
-		//Compute all the needed constants
-		bStatic = bCompileStatic || SpawnParams.bStaticPhysics;
-		SkelMeshComp = bCompileStatic ? nullptr : Cast<USkeletalMeshComponent>(PrimitiveComp);
-		if(SpawnParams.bPhysicsTypeDeterminesSimulation)
-		{
-			this->UpdateSimulatingAndBlendWeight();
-		}
-	}
-
-	FInitBodiesHelperWithData& operator=(const FInitBodiesHelperWithData& InHelper) = delete;
-	FInitBodiesHelperWithData& operator=(FInitBodiesHelperWithData&& InHelper) = delete;
-
-	TArray<FBodyInstance*> OwnedBodies;
-	TArray<FTransform> OwnedTransforms;
-};
+struct FCollisionNotifyInfo;
+template <bool bCompileStatic> struct FInitBodiesHelper;
 
 USTRUCT()
 struct ENGINE_API FCollisionResponse
@@ -243,6 +121,8 @@ enum class BodyInstanceSceneState : uint8
 	Removed
 };
 
+#define USE_BODYINSTANCE_DEBUG_NAMES ((WITH_EDITORONLY_DATA || UE_BUILD_DEBUG || LOOKING_FOR_PERF_ISSUES || CHAOS_CHECKED) && !(UE_BUILD_SHIPPING || UE_BUILD_TEST) && !NO_LOGGING)
+
 /** Container for a physics representation of an object */
 USTRUCT(BlueprintType)
 struct ENGINE_API FBodyInstance
@@ -546,14 +426,8 @@ public:
 	UPROPERTY()
 	float PhysicsBlendWeight;
 
-private:
-	TArray<FInitBodiesHelperWithData<true>> InitBodiesDeferredListStatic;
-	TArray<FInitBodiesHelperWithData<false>> InitBodiesDeferredListDynamic;
-
 public:
 
-	void InitAllBodies(FPhysScene* PhysScene);
-
 	FPhysicsActorHandle& GetPhysicsActorHandle();
 	const FPhysicsActorHandle& GetPhysicsActorHandle() const;
 	const FPhysicsActorHandle& GetActorReferenceWithWelding() const;
@@ -583,6 +457,21 @@ public:
 	 **/
 	void LoadProfileData(bool bVerifyProfile);
 
+	/** Helper struct to specify spawn behavior */
+	struct FInitBodySpawnParams
+	{
+		ENGINE_API FInitBodySpawnParams(const UPrimitiveComponent* PrimComp);
+
+		/** Whether the created physics actor will be static */
+		bool bStaticPhysics;
+
+		/** Whether to use the BodySetup's PhysicsType to override if the instance simulates*/
+		bool bPhysicsTypeDeterminesSimulation;
+
+		/** An aggregate to place the body into */
+		FPhysicsAggregateHandle Aggregate;
+	};
+
 	void InitBody(UBodySetup* Setup, const FTransform& Transform, UPrimitiveComponent* PrimComp, FPhysScene* InRBScene)
 	{
 		InitBody(Setup, Transform, PrimComp, InRBScene, FInitBodySpawnParams(PrimComp));
@@ -1168,7 +1057,8 @@ private:
 	friend struct FUpdateCollisionResponseHelper;
 	friend class FBodySetupDetails;
 	
-	friend struct FInitBodiesHelperBase;
+	friend struct FInitBodiesHelper<true>;
+	friend struct FInitBodiesHelper<false>;
 	friend class FBodyInstanceCustomizationHelper;
 	friend class FFoliageTypeCustomizationHelpers;
 
