Remove MaxArrayRepSize and MaxArrayRepMemory in favor of using UINT16_MAX everywhere.

diff --git a/Engine/Plugins/Runtime/ReplicationGraph/Source/Private/ReplicationGraph.cpp b/Engine/Plugins/Runtime/ReplicationGraph/Source/Private/ReplicationGraph.cpp
--- a/Engine/Plugins/Runtime/ReplicationGraph/Source/Private/ReplicationGraph.cpp
+++ b/Engine/Plugins/Runtime/ReplicationGraph/Source/Private/ReplicationGraph.cpp
@@ -825,12 +825,19 @@ int32 UReplicationGraph::ServerReplicateActors(float DeltaSeconds)
 
 	bWasConnectionSaturated = false;
 
+	TSet<UNetConnection*> ConnectionsToClose;
+
 	ON_SCOPE_EXIT
 	{
 		// We increment this after our replication has happened. If we increment at the beginning of this function, then we rep with FrameNum X, then start the next game frame with the same FrameNum X. If at the top of that frame,
 		// when processing packets, ticking, etc, we get calls to TearOff, ForceNetUpdate etc which make use of ReplicationGraphFrame, they will be using a stale frame num. So we could replicate, get a server move next frame, ForceNetUpdate, but think we 
 		// already replicated this frame.
 		ReplicationGraphFrame++;
+
+		for (UNetConnection* ConnectionToClose : ConnectionsToClose)
+		{
+			ConnectionToClose->Close();
+		}
 	};
 
 	// -------------------------------------------------------
@@ -1038,7 +1045,11 @@ int32 UReplicationGraph::ServerReplicateActors(float DeltaSeconds)
 				}
 			}
 #endif
-			
+		}
+
+		if (NetConnection->GetPendingCloseDueToReplicationFailure())
+		{
+			ConnectionsToClose.Add(NetConnection);
 		}
 	}
 	
diff --git a/Engine/Source/Runtime/Engine/Classes/Engine/NetConnection.h b/Engine/Source/Runtime/Engine/Classes/Engine/NetConnection.h
--- a/Engine/Source/Runtime/Engine/Classes/Engine/NetConnection.h
+++ b/Engine/Source/Runtime/Engine/Classes/Engine/NetConnection.h
@@ -1278,6 +1278,13 @@ public:
 	/** Sends the NMT_Challenge message based on encryption response */
 	void SendChallengeControlMessage(const FEncryptionKeyResponse& Response);
 
+	bool GetPendingCloseDueToReplicationFailure() const
+	{
+		return bConnectionPendingCloseDueToReplicationFailure;
+	}
+
+	ENGINE_API void SetPendingCloseDueToReplicationFailure();
+
 protected:
 
 	bool GetPendingCloseDueToSocketSendFailure() const
@@ -1408,6 +1415,9 @@ private:
 	/** True if we are pending close due to a socket failure during send */
 	bool bConnectionPendingCloseDueToSocketSendFailure;
 
+	/** True if we are pending close due to unrecoverable replication errors. */
+	bool bConnectionPendingCloseDueToReplicationFailure;
+
 	FNetworkGUID GetActorGUIDFromOpenBunch(FInBunch& Bunch);
 
 	/** Out of order packet tracking/correction */
diff --git a/Engine/Source/Runtime/Engine/Classes/Engine/NetworkSettings.h b/Engine/Source/Runtime/Engine/Classes/Engine/NetworkSettings.h
--- a/Engine/Source/Runtime/Engine/Classes/Engine/NetworkSettings.h
+++ b/Engine/Source/Runtime/Engine/Classes/Engine/NetworkSettings.h
@@ -46,15 +46,17 @@ class ENGINE_API UNetworkSettings : public UDeveloperSettings
 	uint32 bEnableMultiplayerWorldOriginRebasing : 1;
 
 	//! Maximum allowable size for replicated dynamic arrays (in number of elements)
+	UE_DEPRECATED(4.26, "This property is no longer used.")
 	UPROPERTY(config, EditAnywhere, Category = replication, meta = (
-		ConsoleVariable = "net.MaxRepArraySize", DisplayName = "Max Array Size",
+		DisplayName = "Max Array Size",
 		ToolTip = "Maximum allowable size for replicated dynamic arrays (in number of elements). Must be between 1 and 65535.",
 		ClampMin = "1", ClampMax = "65535", UIMin = "1", UIMax = "65535"))
 	int32 MaxRepArraySize = DefaultMaxRepArraySize;
 
 	//! Maximum allowable size for replicated dynamic arrays (in bytes)
+	UE_DEPRECATED(4.26, "This property is no longer used.")
 	UPROPERTY(config, EditAnywhere, Category = replication, meta = (
-		ConsoleVariable = "net.MaxRepArrayMemory", DisplayName = "Max Array Memory",
+		DisplayName = "Max Array Memory",
 		ToolTip = "Maximum allowable size for replicated dynamic arrays (in bytes).  Must be between 1 and 65535.",
 		ClampMin = "1", ClampMax = "65535", UIMin = "1", UIMax = "65535"))
 	int32 MaxRepArrayMemory = DefaultMaxRepArrayMemory;
diff --git a/Engine/Source/Runtime/Engine/Private/DataReplication.cpp b/Engine/Source/Runtime/Engine/Private/DataReplication.cpp
--- a/Engine/Source/Runtime/Engine/Private/DataReplication.cpp
+++ b/Engine/Source/Runtime/Engine/Private/DataReplication.cpp
@@ -176,7 +176,15 @@ public:
 	{
 		UpdateCachedState(Params.DeltaSerializeInfo.Object->GetClass(), Params.DeltaSerializeInfo.Struct);
 		UpdateCachedRepLayout();
-		return CachedRequestState.RepLayout->DeltaSerializeFastArrayProperty(Params, ChangelistMgr.Get());
+
+		const ERepLayoutResult Result = CachedRequestState.RepLayout->DeltaSerializeFastArrayProperty(Params, ChangelistMgr.Get());
+
+		if (ERepLayoutResult::FatalError == Result && Params.DeltaSerializeInfo.Connection)
+		{
+			Params.DeltaSerializeInfo.Connection->SetPendingCloseDueToReplicationFailure();
+		}
+
+		return ERepLayoutResult::Success == Result;
 	}
 
 	virtual void GatherGuidReferencesForFastArray(FFastArrayDeltaSerializeParams& Params) override final
@@ -251,7 +259,7 @@ public:
 		return RepLayout.GetLifetimeCustomDeltaProperty(CustomDeltaPropertyIndex);
 	}
 
-	static void UpdateChangelistMgr(
+	static ERepLayoutResult UpdateChangelistMgr(
 		const FRepLayout& RepLayout,
 		FSendingRepState* RESTRICT RepState,
 		FReplicationChangelistMgr& InChangelistMgr,
@@ -260,7 +268,7 @@ public:
 		const FReplicationFlags& RepFlags,
 		const bool bForceCompare)
 	{
-		RepLayout.UpdateChangelistMgr(RepState, InChangelistMgr, InObject, ReplicationFrame, RepFlags, bForceCompare);
+		return RepLayout.UpdateChangelistMgr(RepState, InChangelistMgr, InObject, ReplicationFrame, RepFlags, bForceCompare);
 	}
 
 	static const ELifetimeCondition GetLifetimeCustomDeltaPropertyCondition(const FRepLayout& RepLayout, const uint16 CustomDeltaPropertyIndex)
@@ -1557,7 +1565,13 @@ bool FObjectReplicator::ReplicateProperties( FOutBunch & Bunch, FReplicationFlag
 
 	// Update change list (this will re-use work done by previous connections)
 	FSendingRepState* SendingRepState = ((Connection->ResendAllDataState == EResendAllDataState::SinceCheckpoint) && CheckpointRepState.IsValid()) ? CheckpointRepState->GetSendingRepState() : RepState->GetSendingRepState();
-	FNetSerializeCB::UpdateChangelistMgr(*RepLayout, SendingRepState, *ChangelistMgr, Object, Connection->Driver->ReplicationFrame, RepFlags, OwningChannel->bForceCompareProperties);
+	const ERepLayoutResult UpdateResult = FNetSerializeCB::UpdateChangelistMgr(*RepLayout, SendingRepState, *ChangelistMgr, Object, Connection->Driver->ReplicationFrame, RepFlags, OwningChannel->bForceCompareProperties);
+
+	if (UNLIKELY(ERepLayoutResult::FatalError == UpdateResult))
+	{
+		Connection->SetPendingCloseDueToReplicationFailure();
+		return false;
+	}
 
 	// Replicate properties in the layout
 	const bool bHasRepLayout = RepLayout->ReplicateProperties(SendingRepState, ChangelistMgr->GetRepChangelistState(), (uint8*)Object, ObjectClass, OwningChannel, Writer, RepFlags);
diff --git a/Engine/Source/Runtime/Engine/Private/NetConnection.cpp b/Engine/Source/Runtime/Engine/Private/NetConnection.cpp
--- a/Engine/Source/Runtime/Engine/Private/NetConnection.cpp
+++ b/Engine/Source/Runtime/Engine/Private/NetConnection.cpp
@@ -3992,6 +3992,10 @@ void UNetConnection::CleanupStaleDormantReplicators()
 	}
 }
 
+void UNetConnection::SetPendingCloseDueToReplicationFailure()
+{
+	bConnectionPendingCloseDueToReplicationFailure = true;
+}
 
 void UNetConnection::SetPendingCloseDueToSocketSendFailure()
 {
diff --git a/Engine/Source/Runtime/Engine/Private/NetworkDriver.cpp b/Engine/Source/Runtime/Engine/Private/NetworkDriver.cpp
--- a/Engine/Source/Runtime/Engine/Private/NetworkDriver.cpp
+++ b/Engine/Source/Runtime/Engine/Private/NetworkDriver.cpp
@@ -4580,6 +4580,8 @@ int32 UNetDriver::ServerReplicateActors(float DeltaSeconds)
 	// Build the consider list (actors that are ready to replicate)
 	ServerReplicateActors_BuildConsiderList( ConsiderList, ServerTickTime );
 
+	TSet<UNetConnection*> ConnectionsToClose;
+
 	FMemMark Mark( FMemStack::Get() );
 
 	for ( int32 i=0; i < ClientConnections.Num(); i++ )
@@ -4715,6 +4717,11 @@ int32 UNetDriver::ServerReplicateActors(float DeltaSeconds)
 			const bool bWasSaturated = GNumSaturatedConnections > LocalNumSaturated;
 			Connection->TrackReplicationForAnalytics(bWasSaturated);
 		}
+
+		if (Connection->GetPendingCloseDueToReplicationFailure())
+		{
+			ConnectionsToClose.Add(Connection);
+		}
 	}
 
 	// shuffle the list of connections if not all connections were ticked
@@ -4743,6 +4750,11 @@ int32 UNetDriver::ServerReplicateActors(float DeltaSeconds)
 		DebugRelevantActors  = false;
 	}
 
+	for (UNetConnection* ConnectionToClose : ConnectionsToClose)
+	{
+		ConnectionToClose->Close();
+	}
+
 	return Updated;
 #else
 	return 0;
diff --git a/Engine/Source/Runtime/Engine/Private/RepLayout.cpp b/Engine/Source/Runtime/Engine/Private/RepLayout.cpp
--- a/Engine/Source/Runtime/Engine/Private/RepLayout.cpp
+++ b/Engine/Source/Runtime/Engine/Private/RepLayout.cpp
@@ -21,6 +21,7 @@
 #include "Net/NetworkGranularMemoryLogging.h"
 #include "Serialization/ArchiveCountMem.h"
 #include "Templates/AndOrNot.h"
+#include "Math/NumericLimits.h"
 #include "PushModelPerNetDriverState.h"
 #include "Net/Core/Trace/NetTrace.h"
 
@@ -73,72 +74,11 @@ constexpr bool GbPushModelValidateProperties = false;
 
 #endif
 
-int32 MaxRepArraySize = UNetworkSettings::DefaultMaxRepArraySize;
-int32 MaxRepArrayMemory = UNetworkSettings::DefaultMaxRepArrayMemory;
-
 extern int32 GNumSharedSerializationHit;
 extern int32 GNumSharedSerializationMiss;
 
 extern TAutoConsoleVariable<int32> CVarNetEnableDetailedScopeCounters;
 
-FConsoleVariableSinkHandle CreateMaxArraySizeCVarAndRegisterSink()
-{
-	static FAutoConsoleVariable CVarMaxArraySize(TEXT("net.MaxRepArraySize"), MaxRepArraySize, TEXT("Maximum allowable size for replicated dynamic arrays (in number of elements). Value must be between 1 and 65535."));
-	static FConsoleCommandDelegate Delegate = FConsoleCommandDelegate::CreateLambda(
-		[]()
-		{
-			const int32 NewMaxRepArraySizeValue = CVarMaxArraySize->GetInt();
-
-			if ((int32)UINT16_MAX < NewMaxRepArraySizeValue || 1 > NewMaxRepArraySizeValue)
-			{
-				UE_LOG(LogRepTraffic, Error,
-					TEXT("SerializeProperties_DynamicArray_r: MaxRepArraySize (%l) must be between 1 and 65535. Cannot accept new value."),
-					NewMaxRepArraySizeValue);
-
-				// Use SetByConsole to guarantee the value gets updated.
-				CVarMaxArraySize->Set(MaxRepArraySize, ECVF_SetByConsole);
-			}
-			else
-			{
-				MaxRepArraySize = NewMaxRepArraySizeValue;
-			}
-		}
-	);
-
-	return IConsoleManager::Get().RegisterConsoleVariableSink_Handle(Delegate);
-}
-
-FConsoleVariableSinkHandle CreateMaxArrayMemoryCVarAndRegisterSink()
-{
-	static FAutoConsoleVariableRef CVarMaxArrayMemory(TEXT("net.MaxRepArrayMemory"), MaxRepArrayMemory, TEXT("Maximum allowable size for replicated dynamic arrays (in bytes). Value must be between 1 and 65535"));
-	static FConsoleCommandDelegate Delegate = FConsoleCommandDelegate::CreateLambda(
-		[]()
-		{
-			const int32 NewMaxRepArrayMemoryValue = CVarMaxArrayMemory->GetInt();
-
-			if ((int32)UINT16_MAX < NewMaxRepArrayMemoryValue || 1 > NewMaxRepArrayMemoryValue)
-			{
-				UE_LOG(LogRepTraffic, Error,
-					TEXT("SerializeProperties_DynamicArray_r: MaxRepArrayMemory (%l) must be between 1 and 65535. Cannot accept new value."),
-					NewMaxRepArrayMemoryValue);
-
-				// Use SetByConsole to guarantee the value gets updated.
-				CVarMaxArrayMemory->Set(MaxRepArrayMemory, ECVF_SetByConsole);
-			}
-			else
-			{
-				MaxRepArrayMemory = NewMaxRepArrayMemoryValue;
-			}
-		}
-	);
-
-	return IConsoleManager::Get().RegisterConsoleVariableSink_Handle(Delegate);
-}
-
-// This just forces the above to get called.
-FConsoleVariableSinkHandle MaxRepArraySizeHandle = CreateMaxArraySizeCVarAndRegisterSink();
-FConsoleVariableSinkHandle MaxRepArrayMemorySink = CreateMaxArrayMemoryCVarAndRegisterSink();
-
 /** 
 * Helper method to allow us to instrument FBitArchive using FNetTraceCollector
 * The rationale behind this variant being declared here is that it makes ugly assumption about FBitArchive when being used to collect stats either is a BitWriter or a BitReader
@@ -285,6 +225,24 @@ namespace UE4_RepLayout_Private
 		return false;
 #endif
 	}
+
+	bool ValidateArraySize(const int32 ArraySize, const FProperty* const Property)
+	{
+		if (UNLIKELY(TNumericLimits<uint16>::Max() <= ArraySize))
+		{
+			FString ErrorString = FString::Printf(TEXT("ValidateArraySize: Replicated arrays must be smaller than %d elements in size. ArraySize = (%d) Property = (%s)."),
+				TNumericLimits<uint16>::Max(), ArraySize, *Property->GetPathName());
+
+			// Keep these separate so that we always log the error even if we only fire the ensure once
+			// or if ensures are disabled.
+			UE_LOG(LogRepTraffic, Error, TEXT("%s"), *ErrorString);
+
+			ensureMsgf(false, TEXT("%s"), *ErrorString);
+			return false;
+		}
+
+		return true;
+	}	
 }
 
 //~ TODO: Consider moving the FastArray members into their own sub-struct to save memory for non fast array
@@ -1079,7 +1037,7 @@ FRepLayout::~FRepLayout()
 {
 }
 
-void FRepLayout::UpdateChangelistMgr(
+ERepLayoutResult FRepLayout::UpdateChangelistMgr(
 	FSendingRepState* RESTRICT RepState,
 	FReplicationChangelistMgr& InChangelistMgr,
 	const UObject* InObject,
@@ -1087,6 +1045,8 @@ void FRepLayout::UpdateChangelistMgr(
 	const FReplicationFlags& RepFlags,
 	const bool bForceCompare) const
 {
+	ERepLayoutResult Result = ERepLayoutResult::Success;
+
 	if (GShareInitialCompareState)
 	{
 		// See if we can re-use the work already done on a previous connection
@@ -1100,11 +1060,11 @@ void FRepLayout::UpdateChangelistMgr(
 			{
 				FReplicationFlags TempFlags = RepFlags;
 				TempFlags.bRolesOnly = true;
-				CompareProperties(RepState, &InChangelistMgr.RepChangelistState, (const uint8*)InObject, TempFlags);
+				Result = CompareProperties(RepState, &InChangelistMgr.RepChangelistState, (const uint8*)InObject, TempFlags);
 			}
 
 			INC_DWORD_STAT_BY(STAT_NetSkippedDynamicProps, 1);
-			return;
+			return Result;
 		}
 	}
 	else
@@ -1118,18 +1078,25 @@ void FRepLayout::UpdateChangelistMgr(
 		if (!bForceCompare && GShareShadowState && !RepFlags.bNetInitial && RepState->LastCompareIndex > 1 && InChangelistMgr.LastReplicationFrame == ReplicationFrame)
 		{
 			INC_DWORD_STAT_BY(STAT_NetSkippedDynamicProps, 1);
-			return;
+			return Result;
 		}
 	}
 
-	CompareProperties(RepState, &InChangelistMgr.RepChangelistState, (const uint8*)InObject, RepFlags);
+	Result = CompareProperties(RepState, &InChangelistMgr.RepChangelistState, (const uint8*)InObject, RepFlags);
 
-	InChangelistMgr.LastReplicationFrame = ReplicationFrame;
-
-	if (GShareInitialCompareState && RepFlags.bNetInitial)
+	// Currently, comparing properties should only result in Success, Empty, or FatalError.
+	// So, don't bother checking for normal errors.
+	if (LIKELY(ERepLayoutResult::FatalError != Result))
 	{
-		InChangelistMgr.LastInitialReplicationFrame = ReplicationFrame;
+		InChangelistMgr.LastReplicationFrame = ReplicationFrame;
+
+		if (GShareInitialCompareState && RepFlags.bNetInitial)
+		{
+			InChangelistMgr.LastInitialReplicationFrame = ReplicationFrame;
+		}
 	}
+
+	return Result;
 }
 
 struct FComparePropertiesSharedParams
@@ -1158,6 +1125,7 @@ struct FComparePropertiesStackParams
 	const FConstRepObjectDataBuffer Data;
 	FRepShadowDataBuffer ShadowData;
 	TArray<uint16>& Changed;
+	ERepLayoutResult& Result;
 };
 
 static uint16 CompareProperties_r(
@@ -1211,18 +1179,18 @@ static bool CompareParentProperty(
 	const FComparePropertiesSharedParams& SharedParams,
 	FComparePropertiesStackParams& StackParams)
 {
-		const FRepParentCmd& Parent = SharedParams.Parents[ParentIndex];
-		const bool bIsLifetime = EnumHasAnyFlags(Parent.Flags, ERepParentFlags::IsLifetime);
+	const FRepParentCmd& Parent = SharedParams.Parents[ParentIndex];
+	const bool bIsLifetime = EnumHasAnyFlags(Parent.Flags, ERepParentFlags::IsLifetime);
 
-		// Active state of a property applies to *all* connections.
-		// If the property is inactive, we can skip comparing it because we know it won't be sent.
-		// Further, this will keep the last active state of the property in the shadow buffer,
-		// meaning the next time the property becomes active it will be sent to all connections.
+	// Active state of a property applies to *all* connections.
+	// If the property is inactive, we can skip comparing it because we know it won't be sent.
+	// Further, this will keep the last active state of the property in the shadow buffer,
+	// meaning the next time the property becomes active it will be sent to all connections.
 	const bool bIsActive = !SharedParams.RepChangedPropertyTracker || SharedParams.RepChangedPropertyTracker->Parents[ParentIndex].Active;
-		const bool bShouldSkip = !bIsLifetime || !bIsActive || (Parent.Condition == COND_InitialOnly && !SharedParams.bIsInitial);
+	const bool bShouldSkip = !bIsLifetime || !bIsActive || (Parent.Condition == COND_InitialOnly && !SharedParams.bIsInitial);
 
-		if (bShouldSkip)
-		{
+	if (bShouldSkip)
+	{
 		return false;
 	}
 
@@ -1230,26 +1198,26 @@ static bool CompareParentProperty(
 	if (SharedParams.bIsNetworkProfilerActive)
 	{
 		const_cast<TBitArray<>&>(SharedParams.PropertiesCompared)[ParentIndex] = true;
-		}
+	}
 #endif
 
 	const FRepLayoutCmd& Cmd = SharedParams.Cmds[Parent.CmdStart];
 
 	if (EnumHasAnyFlags(SharedParams.Flags, ERepLayoutFlags::IsActor))
+	{
+		if (UNLIKELY(ParentIndex == (int32)AActor::ENetFields_Private::Role))
 		{
-			if (UNLIKELY(ParentIndex == (int32)AActor::ENetFields_Private::Role))
-			{
 			return CompareRoleProperty(SharedParams, StackParams, (int32)AActor::ENetFields_Private::Role, SharedParams.RepState->SavedRole);
-			}
+		}
 		if (UNLIKELY(ParentIndex == (int32)AActor::ENetFields_Private::RemoteRole))
-			{
+		{
 			return CompareRoleProperty(SharedParams, StackParams, (int32)AActor::ENetFields_Private::RemoteRole, SharedParams.RepState->SavedRemoteRole);
-			}
 		}
+	}
 		
 	const int32 NumChanges = StackParams.Changed.Num();
 
-		// Note, Handle - 1 to account for CompareProperties_r incrementing handles.
+	// Note, Handle - 1 to account for CompareProperties_r incrementing handles.
 	CompareProperties_r(SharedParams, StackParams, Parent.CmdStart, Parent.CmdEnd, Cmd.RelativeHandle - 1);
 
 	return !!(StackParams.Changed.Num() - NumChanges);
@@ -1386,7 +1354,8 @@ static uint16 CompareProperties_r(
 			FComparePropertiesStackParams NewStackParams{
 				Data,
 				ShadowData,
-				StackParams.Changed
+				StackParams.Changed,
+				StackParams.Result
 			};
 
 			// Once we hit an array, start using a stack based approach
@@ -1415,8 +1384,14 @@ static void CompareProperties_Array_r(
 	FScriptArray* ShadowArray = (FScriptArray*)StackParams.ShadowData.Data;
 	FScriptArray* Array = (FScriptArray*)StackParams.Data.Data;
 
-	const uint16 ArrayNum = Array->Num();
-	const uint16 ShadowArrayNum = ShadowArray->Num();
+	const int32 ArrayNum = Array->Num();
+	const int32 ShadowArrayNum = ShadowArray->Num();
+
+	if (!UE4_RepLayout_Private::ValidateArraySize(ArrayNum, Cmd.Property))
+	{
+		StackParams.Result = ERepLayoutResult::FatalError;
+		return;
+	}
 
 	// Make the shadow state match the actual state at the time of compare
 	FScriptArrayHelper StoredArrayHelper((FArrayProperty*)Cmd.Property, ShadowArray);
@@ -1442,7 +1417,8 @@ static void CompareProperties_Array_r(
 			FComparePropertiesStackParams NewStackParams{
 				ArrayData + ArrayElementOffset,
 				ShadowArrayData + ArrayElementOffset,
-				ChangedLocal
+				ChangedLocal,
+				StackParams.Result
 			};
 
 			LocalHandle = CompareProperties_r(SharedParams, NewStackParams, CmdIndex + 1, Cmd.EndCmd - 1, LocalHandle);
@@ -1451,13 +1427,16 @@ static void CompareProperties_Array_r(
 
 	if (ChangedLocal.Num())
 	{
-		const int32 ChangedLocalNum = ChangedLocal.Num();
-
-		// We do not support nested properties with more than 65k entries
-		check(ChangedLocalNum <= UINT16_MAX);
+		const int32 NumChangedEntries = ChangedLocal.Num();
+		
+		if (!UE4_RepLayout_Private::ValidateArraySize(NumChangedEntries, Cmd.Property))
+		{
+			StackParams.Result = ERepLayoutResult::FatalError;
+			return;
+		}
 
 		StackParams.Changed.Add(Handle);
-		StackParams.Changed.Add(ChangedLocalNum);		// This is so we can jump over the array if we need to
+		StackParams.Changed.Add((uint16)NumChangedEntries);		// This is so we can jump over the array if we need to
 		StackParams.Changed.Append(ChangedLocal);
 		StackParams.Changed.Add(0);
 	}
@@ -1473,7 +1452,7 @@ static void CompareProperties_Array_r(
 	}
 }
 
-bool FRepLayout::CompareProperties(
+ERepLayoutResult FRepLayout::CompareProperties(
 	FSendingRepState* RESTRICT RepState,
 	FRepChangelistState* RESTRICT RepChangelistState,
 	const FConstRepObjectDataBuffer Data,
@@ -1483,7 +1462,7 @@ bool FRepLayout::CompareProperties(
 
 	if (IsEmpty())
 	{
-		return false;
+		return ERepLayoutResult::Empty;
 	}
 
 	RepChangelistState->CompareIndex++;
@@ -1497,10 +1476,12 @@ bool FRepLayout::CompareProperties(
 	Changed.Empty(1);
 
 #if WITH_PUSH_MODEL
-		const TBitArray<>* const LocalPushModelProperties = &PushModelProperties;
+	const TBitArray<>* const LocalPushModelProperties = &PushModelProperties;
 #else
-		const TBitArray<>* const LocalPushModelProperties = nullptr;
-#endif		
+	const TBitArray<>* const LocalPushModelProperties = nullptr;
+#endif
+
+	ERepLayoutResult Result = ERepLayoutResult::Success;
 
 	FComparePropertiesSharedParams SharedParams{
 		/*bIsInitial=*/ !!RepFlags.bNetInitial,
@@ -1521,7 +1502,8 @@ bool FRepLayout::CompareProperties(
 	FComparePropertiesStackParams StackParams{
 		Data,
 		RepChangelistState->StaticBuffer.GetData(),
-		Changed
+		Changed,
+		Result
 	};
 
 	if (RepFlags.bRolesOnly)
@@ -1560,9 +1542,15 @@ bool FRepLayout::CompareProperties(
 		}
 	}
 
+	// If something went wrong, don't touch the changelist history.
+	if (UNLIKELY(ERepLayoutResult::Success != Result))
+	{
+		return Result;
+	}
+
 	if (Changed.Num() == 0)
 	{
-		return false;
+		return ERepLayoutResult::Empty;
 	}
 
 	//
@@ -1593,7 +1581,7 @@ bool FRepLayout::CompareProperties(
 		MergeChangeList(Data, FirstChangelistRef, SecondChangelistCopy, RepChangelistState->ChangeHistory[SecondHistoryIndex].Changed);
 	}
 
-	return true;
+	return Result;
 }
 
 static FORCEINLINE void WritePropertyHandle(
@@ -5764,20 +5752,8 @@ void FRepLayout::SerializeProperties_DynamicArray_r(
 	const int32 ArrayNum = Ar.IsLoading() ? (int32)OutArrayNum : Array->Num();
 
 	// Validate the maximum number of elements.
-	if (ArrayNum> MaxRepArraySize)
-	{
-		UE_LOG(LogRepTraffic, Error, TEXT("SerializeProperties_DynamicArray_r: ArraySize (%d) > net.MaxRepArraySize(%d) (%s). net.MaxRepArraySize can be updated in Project Settings under Network Settings."),
-			ArrayNum, MaxRepArraySize, *Cmd.Property->GetName());
-
-		Ar.SetError();
-	}
-	// Validate the maximum memory.
-	else if (ArrayNum * (int32)Cmd.ElementSize > MaxRepArrayMemory)
+	if (!UE4_RepLayout_Private::ValidateArraySize(ArrayNum, Cmd.Property))
 	{
-		UE_LOG(LogRepTraffic, Error,
-			TEXT("SerializeProperties_DynamicArray_r: ArraySize (%d) * Cmd.ElementSize (%d) > net.MaxRepArrayMemory(%d) (%s). net.MaxRepArrayMemory can be updated in Project Settings under Network Settings."),
-			ArrayNum, (int32)Cmd.ElementSize, MaxRepArrayMemory, *Cmd.Property->GetName());
-
 		Ar.SetError();
 	}
 
@@ -6038,13 +6014,7 @@ void FRepLayout::BuildSharedSerializationForRPC_DynamicArray_r(
 	FScriptArray* Array = (FScriptArray *)Data.Data;	
 	const int32 ArrayNum = Array->Num();
 
-	// Validate the maximum number of elements.
-	if (ArrayNum > MaxRepArraySize)
-	{
-		return;
-	}
-	// Validate the maximum memory.
-	else if (ArrayNum * (int32)Cmd.ElementSize > MaxRepArrayMemory)
+	if (!UE4_RepLayout_Private::ValidateArraySize(ArrayNum, Cmd.Property))
 	{
 		return;
 	}
@@ -6692,7 +6662,7 @@ void FRepLayout::PostSendCustomDeltaProperties(
 {
 }
 
-bool FRepLayout::DeltaSerializeFastArrayProperty(FFastArrayDeltaSerializeParams& Params, FReplicationChangelistMgr* ChangelistMgr) const
+ERepLayoutResult FRepLayout::DeltaSerializeFastArrayProperty(FFastArrayDeltaSerializeParams& Params, FReplicationChangelistMgr* ChangelistMgr) const
 {
 	using namespace UE4_RepLayout_Private;
 
@@ -6716,7 +6686,7 @@ bool FRepLayout::DeltaSerializeFastArrayProperty(FFastArrayDeltaSerializeParams&
 		// This should have already been caught by InitFromClass.
 		// So, log with a lower verbosity.
 		UE_LOG(LogRep, Log, TEXT("FRepLayout::DeltaSerializeFastArrayProperty: Invalid fast array items command index! %s"), *Parent.CachedPropertyName.ToString())
-		return false;
+		return ERepLayoutResult::Error;
 	}
 
 	const FRepLayoutCmd& FastArrayItemCmd = Cmds[CmdIndex];
@@ -6756,7 +6726,7 @@ bool FRepLayout::DeltaSerializeFastArrayProperty(FFastArrayDeltaSerializeParams&
 			if (!bIsWriting)
 			{
 				UE_LOG(LogRep, Error, TEXT("DeltaSerializeFastArrayProperty: Unable to find NetFieldExportGroup during replay playback. Class=%s, Property=%s"), *Owner->GetName(), *Parent.CachedPropertyName.ToString());
-				return false;
+				return ERepLayoutResult::Error;
 			}
 
 			UE_LOG(LogRepProperties, VeryVerbose, TEXT("DeltaSerializeFastArrayProperty: Create Netfield Export Group."))
@@ -7006,14 +6976,22 @@ bool FRepLayout::DeltaSerializeFastArrayProperty(FFastArrayDeltaSerializeParams&
 								nullptr
 							};
 
+							ERepLayoutResult UpdateResult = ERepLayoutResult::Success;
+
 							FComparePropertiesStackParams StackParams{
 								FConstRepObjectDataBuffer(ObjectArrayData + ArrayElementOffset),
 								FRepShadowDataBuffer(ShadowArrayData + ArrayElementOffset),
-								NewChangelist
+								NewChangelist,
+								UpdateResult
 							};
 
 							CompareProperties_r(SharedParams, StackParams, ItemLayoutStart, ItemLayoutEnd, 0);
 
+							if (UNLIKELY(ERepLayoutResult::FatalError == UpdateResult))
+							{
+								return UpdateResult;
+							}
+
 							if (NewChangelist.Num())
 							{
 								NewChangelist.Add(0);
@@ -7176,7 +7154,7 @@ bool FRepLayout::DeltaSerializeFastArrayProperty(FFastArrayDeltaSerializeParams&
 			
 		}
 
-		return !Writer.IsError();
+		return !Writer.IsError() ? ERepLayoutResult::Success : ERepLayoutResult::Error;
 	}
 	else
 	{
@@ -7241,7 +7219,7 @@ bool FRepLayout::DeltaSerializeFastArrayProperty(FFastArrayDeltaSerializeParams&
 				if (!bSuccess)
 				{
 					UE_LOG(LogNetFastTArray, Warning, TEXT("FRepLayout::DeltaSerializeFastArrayProperty: Failed to receive backwards compat properties!"));
-					return false;
+					return ERepLayoutResult::Error;
 				}
 			}
 			else
@@ -7279,27 +7257,27 @@ bool FRepLayout::DeltaSerializeFastArrayProperty(FFastArrayDeltaSerializeParams&
 					if (0 != SharedParams.ReadHandle)
 					{
 						UE_LOG(LogRep, Error, TEXT("ReceiveFastArrayItem: Invalid property terminator handle - Handle=%d"), SharedParams.ReadHandle);
-						return false;
+						return ERepLayoutResult::Error;
 					}
 				}
 				else
 				{
 					UE_LOG(LogNetFastTArray, Warning, TEXT("FRepLayout::DeltaSerializeFastArrayProperty: Failed to received properties"));
-					return false;
+					return ERepLayoutResult::Error;
 				}
 			}
 
 			if (Reader.IsError())
 			{
 				UE_LOG(LogNetFastTArray, Warning, TEXT("FRepLayout::DeltaSerializeFastArrayProperty: Reader.IsError() == true"));
-				return false;
+				return ERepLayoutResult::Error;
 			}
 
 			DeltaSerializeInfo.bGuidListsChanged |= bOutGuidsChanged;
 			DeltaSerializeInfo.bOutHasMoreUnmapped |= bOutHasUnmapped;
 		}
 
-		return true;
+		return ERepLayoutResult::Success;
 	}
 }
 
diff --git a/Engine/Source/Runtime/Engine/Public/Net/RepLayout.h b/Engine/Source/Runtime/Engine/Public/Net/RepLayout.h
--- a/Engine/Source/Runtime/Engine/Public/Net/RepLayout.h
+++ b/Engine/Source/Runtime/Engine/Public/Net/RepLayout.h
@@ -1014,6 +1014,13 @@ enum class ERepLayoutFlags : uint8
 };
 ENUM_CLASS_FLAGS(ERepLayoutFlags);
 
+enum class ERepLayoutResult
+{
+	Success,	// Operation succeeded
+	Empty,		// Operation succeeded, but didn't generate (or consume) any data.
+	Error,		// Operation failed, but may succeed later.
+	FatalError	// Operation failed, and connection should be terminated.
+};
 
 /**
  * This class holds all replicated properties for a given type (either a UClass, UStruct, or UFunction).
@@ -1457,7 +1464,7 @@ private:
 	 * @param Data					The newest Property Data available.
 	 * @param RepFlags				Flags that will be used if the object is replicated.
 	 */
-	bool CompareProperties(
+	ERepLayoutResult CompareProperties(
 		FSendingRepState* RESTRICT RepState,
 		FRepChangelistState* RESTRICT RepChangelistState,
 		const FConstRepObjectDataBuffer Data,
@@ -1696,7 +1703,7 @@ private:
 		const int32 CmdEnd,
 		TArray<FHandleToCmdIndex>& HandleToCmdIndex);
 
-	void UpdateChangelistMgr(
+	ERepLayoutResult UpdateChangelistMgr(
 		FSendingRepState* RESTRICT RepState,
 		FReplicationChangelistMgr& InChangelistMgr,
 		const UObject* InObject,
@@ -1733,7 +1740,7 @@ private:
 		FNetDeltaSerializeInfo& Params,
 		FStructProperty* Property) const;
 
-	bool DeltaSerializeFastArrayProperty(struct FFastArrayDeltaSerializeParams& Params, FReplicationChangelistMgr* ChangelistMgr) const;
+	ERepLayoutResult DeltaSerializeFastArrayProperty(struct FFastArrayDeltaSerializeParams& Params, FReplicationChangelistMgr* ChangelistMgr) const;
 
 	void GatherGuidReferencesForFastArray(struct FFastArrayDeltaSerializeParams& Params) const;
 
