Fixed mesh batch validation to include case where gpu scene is disabled and falls back to the primitive uniform buffer.

diff --git a/Engine/Source/Runtime/Engine/Private/SceneManagement.cpp b/Engine/Source/Runtime/Engine/Private/SceneManagement.cpp
--- a/Engine/Source/Runtime/Engine/Private/SceneManagement.cpp
+++ b/Engine/Source/Runtime/Engine/Private/SceneManagement.cpp
@@ -1176,6 +1176,8 @@ bool FMeshBatch::Validate(const FPrimitiveSceneProxy* PrimitiveSceneProxy, ERHIF
 
 	const bool bPrimitiveShaderDataComesFromSceneBuffer = VertexFactory->GetPrimitiveIdStreamIndex(EVertexInputStreamType::Default) >= 0;
 
+	const bool bPrimitiveHasUniformBuffer = PrimitiveSceneProxy->GetUniformBuffer() != nullptr;
+
 	for (int32 ElementIndex = 0; ElementIndex < Elements.Num(); ElementIndex++)
 	{
 		const FMeshBatchElement& MeshElement = Elements[ElementIndex];
@@ -1188,7 +1190,11 @@ bool FMeshBatch::Validate(const FPrimitiveSceneProxy* PrimitiveSceneProxy, ERHIF
 				TEXT("both null to get FPrimitiveSceneProxy->UniformBuffer."));
 		}
 
-		const bool bValidPrimitiveData = bPrimitiveShaderDataComesFromSceneBuffer || Elements[ElementIndex].PrimitiveUniformBuffer || Elements[ElementIndex].PrimitiveUniformBufferResource;
+		const bool bValidPrimitiveData =
+			   bPrimitiveShaderDataComesFromSceneBuffer
+			|| bPrimitiveHasUniformBuffer
+			|| Elements[ElementIndex].PrimitiveUniformBuffer
+			|| Elements[ElementIndex].PrimitiveUniformBufferResource;
 
 		if (!bValidPrimitiveData)
 		{
