Combined ensure into same text block since it seems to be truncated on non-PC platforms.

diff --git a/Engine/Source/Runtime/Engine/Private/SceneManagement.cpp b/Engine/Source/Runtime/Engine/Private/SceneManagement.cpp
--- a/Engine/Source/Runtime/Engine/Private/SceneManagement.cpp
+++ b/Engine/Source/Runtime/Engine/Private/SceneManagement.cpp
@@ -1107,9 +1107,7 @@ bool FMeshBatch::Validate(const FPrimitiveSceneProxy* PrimitiveSceneProxy, ERHIF
 		const uint32 bVertexFactoryInitialized = (VertexFactory && VertexFactory->IsInitialized()) ? 1 : 0;
 
 		ensureMsgf(false,
-			TEXT("FMeshBatch was not properly setup. %s.\n")
-			TEXT("\tVertexFactory[Name: %s, Initialized: %u]\n")
-			TEXT("\tPrimitiveSceneProxy[Level: %s, Owner: %s, Resource: %s]"),
+			TEXT("FMeshBatch was not properly setup. %s.\n\tVertexFactory[Name: %s, Initialized: %u]\n\tPrimitiveSceneProxy[Level: %s, Owner: %s, Resource: %s]"),
 			*Error,
 			*VertexFactoryName,
 			bVertexFactoryInitialized,
