Fix for camera manager delegate not firing when given a camera fade from 1 to 1.
+ This is a valid use case to 'slam' a black screen transition

#rnx

#deadend

diff --git a/Engine/Source/Runtime/Engine/Private/PlayerCameraManager.cpp b/Engine/Source/Runtime/Engine/Private/PlayerCameraManager.cpp
--- a/Engine/Source/Runtime/Engine/Private/PlayerCameraManager.cpp
+++ b/Engine/Source/Runtime/Engine/Private/PlayerCameraManager.cpp
@@ -615,6 +615,8 @@ void APlayerCameraManager::StopAudioFade()
 	// If an audio fade event has been bound, we'd like it to override the default fade behavior.
 	if (OnAudioFadeChangeEvent.IsBound())
 	{
+		const bool bFadeOut = false;
+		OnAudioFadeChangeEvent.Broadcast(bFadeOut, 0.f);
 		return;
 	}
 
@@ -1398,9 +1400,9 @@ void APlayerCameraManager::StartCameraFade(float FromAlpha, float ToAlpha, float
 	FadeTimeRemaining = InFadeTime;
 	bFadeAudio = bInFadeAudio;
 
-	if (bInFadeAudio && OnAudioFadeChangeEvent.IsBound() && !FMath::IsNearlyEqual(FromAlpha, ToAlpha))
+	if (bInFadeAudio && OnAudioFadeChangeEvent.IsBound())
 	{
-		const bool bFadeOut = FromAlpha < ToAlpha;
+		const bool bFadeOut = FromAlpha < ToAlpha || FMath::IsNearlyEqual(ToAlpha, 1.0f);
 		OnAudioFadeChangeEvent.Broadcast(bFadeOut, FadeTime);
 	}
 
