Fix non-unity build error

diff --git a/Engine/Source/Runtime/Engine/Private/Streaming/LandscapeMeshMobileUpdate.h b/Engine/Source/Runtime/Engine/Private/Streaming/LandscapeMeshMobileUpdate.h
--- a/Engine/Source/Runtime/Engine/Private/Streaming/LandscapeMeshMobileUpdate.h
+++ b/Engine/Source/Runtime/Engine/Private/Streaming/LandscapeMeshMobileUpdate.h
@@ -9,6 +9,7 @@ LandscapeMeshMobileUpdate.h: Helpers to stream in and out mobile landscape verte
 #include "CoreMinimal.h"
 #include "Async/AsyncFileHandle.h"
 
+class UStreamableRenderAsset;
 class ULandscapeLODStreamingProxy;
 struct FLandscapeMobileRenderData;
 
