Fix non-unity build error

diff --git a/Engine/Source/Runtime/Engine/Private/Streaming/LandscapeMeshMobileUpdate.h b/Engine/Source/Runtime/Engine/Private/Streaming/LandscapeMeshMobileUpdate.h
--- a/Engine/Source/Runtime/Engine/Private/Streaming/LandscapeMeshMobileUpdate.h
+++ b/Engine/Source/Runtime/Engine/Private/Streaming/LandscapeMeshMobileUpdate.h
@@ -7,6 +7,7 @@ LandscapeMeshMobileUpdate.h: Helpers to stream in and out mobile landscape verte
 #pragma once
 
 #include "CoreMinimal.h"
+#include "RenderAssetUpdate.h"
 #include "Async/AsyncFileHandle.h"
 
 class UStreamableRenderAsset;
