- Fixed compile errors
#rnx

diff --git a/Engine/Source/Runtime/Engine/Private/Streaming/LandscapeMeshMobileUpdate.cpp b/Engine/Source/Runtime/Engine/Private/Streaming/LandscapeMeshMobileUpdate.cpp
--- a/Engine/Source/Runtime/Engine/Private/Streaming/LandscapeMeshMobileUpdate.cpp
+++ b/Engine/Source/Runtime/Engine/Private/Streaming/LandscapeMeshMobileUpdate.cpp
@@ -257,7 +257,7 @@ FString FLandscapeMeshMobileStreamIn_IO::GetIOFilename(const FContext& Context)
 
 void FLandscapeMeshMobileStreamIn_IO::SetAsyncFileCallback(const FContext& Context)
 {
-	AsyncFileCallback = [this, Context](bool bWasCancelled, IAsyncReadRequest* Req)
+	AsyncFileCallback = [this, Context](bool bWasCancelled, IBulkDataIORequest* Req)
 	{
 		TaskSynchronization.Decrement();
 
diff --git a/Engine/Source/Runtime/Engine/Private/Streaming/LandscapeMeshMobileUpdate.h b/Engine/Source/Runtime/Engine/Private/Streaming/LandscapeMeshMobileUpdate.h
--- a/Engine/Source/Runtime/Engine/Private/Streaming/LandscapeMeshMobileUpdate.h
+++ b/Engine/Source/Runtime/Engine/Private/Streaming/LandscapeMeshMobileUpdate.h
@@ -123,7 +123,7 @@ protected:
 	void CancelIORequest();
 
 	class IBulkDataIORequest* IORequests[MAX_MESH_LOD_COUNT];
-	FAsyncFileCallBack AsyncFileCallback;
+	FBulkDataIORequestCallBack AsyncFileCallback;
 	bool bHighPrioIORequest;
 };
 
