Fix ray tracing crash with static mesh LOD streaming

diff --git a/Engine/Source/Runtime/Engine/Private/StaticMesh.cpp b/Engine/Source/Runtime/Engine/Private/StaticMesh.cpp
--- a/Engine/Source/Runtime/Engine/Private/StaticMesh.cpp
+++ b/Engine/Source/Runtime/Engine/Private/StaticMesh.cpp
@@ -1243,7 +1243,6 @@ void FStaticMeshLODResources::InitResources(UStaticMesh* Parent)
 				Initializer.Segments = GeometrySections;
 
 				RayTracingGeometry.SetInitializer(Initializer);
-				RayTracingGeometry.InitResource();
 			}
 		);
 	}
@@ -1554,6 +1553,25 @@ void FStaticMeshRenderData::InitResources(ERHIFeatureLevel::Type InFeatureLevel,
 		}
 	}
 
+#if RHI_RAYTRACING
+	if (IsRayTracingEnabled())
+	{
+		ENQUEUE_RENDER_COMMAND(InitRayTracingGeometryForInlinedLODs)(
+			[this](FRHICommandListImmediate&)
+			{
+				for (int32 LODIndex = CurrentFirstLODIdx; LODIndex < LODResources.Num(); ++LODIndex)
+				{
+					// Skip LODs that have their render data stripped
+					if (LODResources[LODIndex].VertexBuffers.StaticMeshVertexBuffer.GetNumVertices() > 0)
+					{
+						LODResources[LODIndex].RayTracingGeometry.InitResource();
+					}
+				}
+			}
+		);
+	}
+#endif
+
 	ENQUEUE_RENDER_COMMAND(CmdSetStaticMeshReadyForStreaming)(
 		[this, Owner](FRHICommandListImmediate&)
 	{
diff --git a/Engine/Source/Runtime/Engine/Private/Streaming/StaticMeshUpdate.cpp b/Engine/Source/Runtime/Engine/Private/Streaming/StaticMeshUpdate.cpp
--- a/Engine/Source/Runtime/Engine/Private/Streaming/StaticMeshUpdate.cpp
+++ b/Engine/Source/Runtime/Engine/Private/Streaming/StaticMeshUpdate.cpp
@@ -226,6 +226,22 @@ void FStaticMeshStreamIn::DoFinishUpdate(const FContext& Context)
 				IntermediateBuffersArray[LODIdx].TransferBuffers(LODResource, Batcher);
 			}
 		}
+
+#if RHI_RAYTRACING
+		// Must happen after the batched updates have been flushed
+		if (IsRayTracingEnabled())
+		{
+			for (int32 LODIndex = PendingFirstMip; LODIndex < CurrentFirstLODIdx; ++LODIndex)
+			{
+				// Skip LODs that have their render data stripped
+				if (RenderData->LODResources[LODIndex].VertexBuffers.StaticMeshVertexBuffer.GetNumVertices() > 0)
+				{
+					RenderData->LODResources[LODIndex].RayTracingGeometry.InitResource();
+				}
+			}
+		}
+#endif
+
 		check(Mesh->GetCachedNumResidentLODs() == RenderData->LODResources.Num() - RenderData->CurrentFirstLODIdx);
 		RenderData->CurrentFirstLODIdx = PendingFirstMip;
 		Mesh->SetCachedNumResidentLODs(static_cast<uint8>(RenderData->LODResources.Num() - PendingFirstMip));
@@ -289,6 +305,13 @@ void FStaticMeshStreamOut::DoReleaseBuffers(const FContext& Context)
 				LODResource.AdditionalIndexBuffers->WireframeIndexBuffer.ReleaseRHIForStreaming(Batcher);
 				LODResource.AdditionalIndexBuffers->AdjacencyIndexBuffer.ReleaseRHIForStreaming(Batcher);
 			}
+
+#if RHI_RAYTRACING
+			if (IsRayTracingEnabled())
+			{
+				LODResource.RayTracingGeometry.ReleaseResource();
+			}
+#endif
 		}
 	}
 }
