Fix LOD problems in ray tracing

* Static meshes with WPO were selecting the wrong LOD
* HLOD was not respected in ray tracing. This was hard to fix, but with GatherRayTracingWorldInstances moved to be after InitViews now we can reuse the results directly. HLOD calculation does not depend on frustum culling so it is exactly we want.

diff --git a/Engine/Source/Runtime/Engine/Private/SceneManagement.cpp b/Engine/Source/Runtime/Engine/Private/SceneManagement.cpp
--- a/Engine/Source/Runtime/Engine/Private/SceneManagement.cpp
+++ b/Engine/Source/Runtime/Engine/Private/SceneManagement.cpp
@@ -536,12 +536,15 @@ int8 ComputeStaticMeshLOD( const FStaticMeshRenderData* RenderData, const FVecto
 	{
 		const int32 NumLODs = MAX_STATIC_MESH_LODS;
 		const FSceneView& LODView = GetLODView(View);
-		const float ScreenRadiusSquared = ComputeBoundsScreenRadiusSquared(Origin, SphereRadius, LODView) * FactorScale * FactorScale * LODView.LODDistanceFactor * LODView.LODDistanceFactor;
+		const float ScreenRadiusSquared = ComputeBoundsScreenRadiusSquared(Origin, SphereRadius, LODView);
 
 		// Walk backwards and return the first matching LOD
 		for (int32 LODIndex = NumLODs - 1; LODIndex >= 0; --LODIndex)
 		{
-			if (FMath::Square(RenderData->ScreenSize[LODIndex].GetValueForFeatureLevel(View.GetFeatureLevel()) * 0.5f) > ScreenRadiusSquared)
+			float ScreenSizeScale = FactorScale * LODView.LODDistanceFactor;
+			float MeshScreenSize = RenderData->ScreenSize[LODIndex].GetValueForFeatureLevel(View.GetFeatureLevel()) * ScreenSizeScale;
+
+			if (FMath::Square(MeshScreenSize * 0.5f) > ScreenRadiusSquared)
 			{
 				return FMath::Max(LODIndex, MinLOD);
 			}
diff --git a/Engine/Source/Runtime/Renderer/Private/DeferredShadingRenderer.cpp b/Engine/Source/Runtime/Renderer/Private/DeferredShadingRenderer.cpp
--- a/Engine/Source/Runtime/Renderer/Private/DeferredShadingRenderer.cpp
+++ b/Engine/Source/Runtime/Renderer/Private/DeferredShadingRenderer.cpp
@@ -780,6 +780,15 @@ bool FDeferredShadingSceneRenderer::GatherRayTracingWorldInstances(FRHICommandLi
 				{
 					continue;
 				}
+				
+				FSceneViewState* ViewState = (FSceneViewState*)View.State;
+				const bool bHLODActive = Scene->SceneLODHierarchy.IsActive();
+				const FHLODVisibilityState* const HLODState = bHLODActive && ViewState ? &ViewState->HLODVisibilityState : nullptr;
+
+				if (HLODState && HLODState->IsNodeForcedHidden(PrimitiveIndex))
+				{
+					continue;
+				}
 
 				//#dxr_todo UE-68621  The Raytracing code path does not support ShowFlags since data moved to the SceneInfo. 
 				//Touching the SceneProxy to determine this would simply cost too much
