Fix a missing header file compile error
#rnx

diff --git a/Engine/Source/Runtime/Engine/Private/Streaming/LandscapeMeshMobileUpdate.cpp b/Engine/Source/Runtime/Engine/Private/Streaming/LandscapeMeshMobileUpdate.cpp
--- a/Engine/Source/Runtime/Engine/Private/Streaming/LandscapeMeshMobileUpdate.cpp
+++ b/Engine/Source/Runtime/Engine/Private/Streaming/LandscapeMeshMobileUpdate.cpp
@@ -9,6 +9,7 @@ LandscapeMeshMobileUpdate.cpp: Helpers to stream in and out mobile landscape ver
 #include "LandscapeComponent.h"
 #include "LandscapeRenderMobile.h"
 #include "Streaming/RenderAssetUpdate.inl"
+#include "Streaming/StaticMeshUpdate.h"
 
 template class TRenderAssetUpdate<FLandscapeMeshMobileUpdateContext>;
 
