Fix compile error
#rnx

diff --git a/Engine/Source/Runtime/Engine/Private/Streaming/LandscapeMeshMobileUpdate.h b/Engine/Source/Runtime/Engine/Private/Streaming/LandscapeMeshMobileUpdate.h
--- a/Engine/Source/Runtime/Engine/Private/Streaming/LandscapeMeshMobileUpdate.h
+++ b/Engine/Source/Runtime/Engine/Private/Streaming/LandscapeMeshMobileUpdate.h
@@ -7,7 +7,7 @@ LandscapeMeshMobileUpdate.h: Helpers to stream in and out mobile landscape verte
 #pragma once
 
 #include "CoreMinimal.h"
-#include "RHIResources.h"
+#include "RHI.h"
 #include "RenderAssetUpdate.h"
 #include "Async/AsyncFileHandle.h"
 
