Make 'evaluate WPO in ray tracing' toggleable at runtime

diff --git a/Engine/Source/Runtime/Engine/Classes/Components/StaticMeshComponent.h b/Engine/Source/Runtime/Engine/Classes/Components/StaticMeshComponent.h
--- a/Engine/Source/Runtime/Engine/Classes/Components/StaticMeshComponent.h
+++ b/Engine/Source/Runtime/Engine/Classes/Components/StaticMeshComponent.h
@@ -384,6 +384,9 @@ public:
 	UFUNCTION(BlueprintCallable, Category="Rendering|Lighting")
 	void SetDistanceFieldSelfShadowBias(float NewValue);
 
+	UFUNCTION(BlueprintCallable, Category=RayTracing)
+	void SetEvaluateWorldPositionOffsetInRayTracing(bool NewValue);
+
 	/** 
 	 * Get Local bounds
 	 */
diff --git a/Engine/Source/Runtime/Engine/Private/Components/StaticMeshComponent.cpp b/Engine/Source/Runtime/Engine/Private/Components/StaticMeshComponent.cpp
--- a/Engine/Source/Runtime/Engine/Private/Components/StaticMeshComponent.cpp
+++ b/Engine/Source/Runtime/Engine/Private/Components/StaticMeshComponent.cpp
@@ -1861,6 +1861,31 @@ void UStaticMeshComponent::SetDistanceFieldSelfShadowBias(float NewValue)
 	}
 }
 
+void UStaticMeshComponent::SetEvaluateWorldPositionOffsetInRayTracing(bool NewValue)
+{
+	if (bEvaluateWorldPositionOffset != NewValue && GetScene() != nullptr)
+	{
+		// Update game thread data
+		bEvaluateWorldPositionOffset = NewValue;
+
+		// Skip when this doesn't have a valid static mesh 
+		if (!GetStaticMesh())
+		{
+			return;
+		}
+
+		// Update render thread data
+		ENQUEUE_RENDER_COMMAND(UpdateDFSelfShadowBiasCmd)(
+			[NewValue, PrimitiveSceneProxy = (FStaticMeshSceneProxy*)SceneProxy](FRHICommandList&)
+			{
+				if (PrimitiveSceneProxy)
+				{
+					PrimitiveSceneProxy->SetEvaluateWorldPositionOffsetInRayTracing(NewValue);
+				}
+			});
+	}
+}
+
 void UStaticMeshComponent::SetReverseCulling(bool ReverseCulling)
 {
 	if (ReverseCulling != bReverseCulling)
diff --git a/Engine/Source/Runtime/Engine/Private/StaticMeshRender.cpp b/Engine/Source/Runtime/Engine/Private/StaticMeshRender.cpp
--- a/Engine/Source/Runtime/Engine/Private/StaticMeshRender.cpp
+++ b/Engine/Source/Runtime/Engine/Private/StaticMeshRender.cpp
@@ -42,6 +42,7 @@
 
 #include "UnrealEngine.h"
 #include "RayTracingInstance.h"
+#include "PrimitiveSceneInfo.h"
 
 /** If true, optimized depth-only index buffers are used for shadow rendering. */
 static bool GUseShadowIndexBuffer = true;
@@ -49,6 +50,10 @@ static bool GUseShadowIndexBuffer = true;
 /** If true, reversed index buffer are used for mesh with negative transform determinants. */
 static bool GUseReversedIndexBuffer = true;
 
+DECLARE_STATS_GROUP(TEXT("D3D12RHI: Ray Tracing"), STATGROUP_D3D12RayTracing, STATCAT_Advanced);
+DECLARE_MEMORY_STAT(TEXT("Total Used Video Memory"), STAT_D3D12RayTracingUsedVideoMemory, STATGROUP_D3D12RayTracing);
+DECLARE_MEMORY_STAT(TEXT("Dynamic Vertex Buffer Memory"), STAT_D3D12RayTracingDynamicVertexBufferMemory, STATGROUP_D3D12RayTracing);
+
 static void ToggleShadowIndexBuffers()
 {
 	FlushRenderingCommands();
@@ -198,6 +203,12 @@ FStaticMeshSceneProxy::FStaticMeshSceneProxy(UStaticMeshComponent* InComponent,
 	bDynamicRayTracingGeometry = InComponent->bEvaluateWorldPositionOffset && MaterialRelevance.bUsesWorldPositionOffset;
 	if (IsRayTracingEnabled())
 	{
+		RayTracingGeometries.AddDefaulted(RenderData->LODResources.Num());
+		for (int32 LODIndex = 0; LODIndex < RenderData->LODResources.Num(); LODIndex++)
+		{
+			RayTracingGeometries[LODIndex] = &RenderData->LODResources[LODIndex].RayTracingGeometry;
+		}
+
 		if(bDynamicRayTracingGeometry)
 		{
 			DynamicRayTracingGeometries.AddDefaulted(RenderData->LODResources.Num());
@@ -212,14 +223,6 @@ FStaticMeshSceneProxy::FStaticMeshSceneProxy(UStaticMeshComponent* InComponent,
 				Initializer.bAllowUpdate = true;
 				Initializer.bFastBuild = true;
 			}
-		} 
-		else
-		{
-			RayTracingGeometries.AddDefaulted(RenderData->LODResources.Num());
-			for (int32 LODIndex = 0; LODIndex < RenderData->LODResources.Num(); LODIndex++)
-			{
-				RayTracingGeometries[LODIndex] = &RenderData->LODResources[LODIndex].RayTracingGeometry;
-			}
 		}
 	}
 #endif
@@ -304,11 +307,85 @@ FStaticMeshSceneProxy::FStaticMeshSceneProxy(UStaticMeshComponent* InComponent,
 	AddSpeedTreeWind();
 }
 
+void FStaticMeshSceneProxy::SetEvaluateWorldPositionOffsetInRayTracing(bool NewValue)
+{
+#if RHI_RAYTRACING
+	NewValue &= MaterialRelevance.bUsesWorldPositionOffset;
+	if (NewValue && !bDynamicRayTracingGeometry)
+	{
+		bDynamicRayTracingGeometry = true;
+		if (IsRayTracingEnabled())
+		{
+			DynamicRayTracingGeometries.AddDefaulted(RenderData->LODResources.Num());
+
+			for (int32 LODIndex = 0; LODIndex < RenderData->LODResources.Num(); LODIndex++)
+			{
+				auto& Initializer = DynamicRayTracingGeometries[LODIndex].Initializer;
+				Initializer = RenderData->LODResources[LODIndex].RayTracingGeometry.Initializer;
+				for (FRayTracingGeometrySegment& Segment : Initializer.Segments)
+				{
+					Segment.VertexBuffer = nullptr;
+				}
+				Initializer.bAllowUpdate = true;
+				Initializer.bFastBuild = true;
+			}
+
+			DynamicRayTracingGeometryVertexBuffers.AddDefaulted(DynamicRayTracingGeometries.Num());
+
+			for (int32 i = 0; i < DynamicRayTracingGeometries.Num(); i++)
+			{
+				auto& Geometry = DynamicRayTracingGeometries[i];
+				DynamicRayTracingGeometryVertexBuffers[i].Initialize(4, 256, PF_R32_FLOAT, BUF_UnorderedAccess | BUF_ShaderResource, TEXT("RayTracingDynamicVertexBuffer"));
+
+				INC_MEMORY_STAT_BY(STAT_D3D12RayTracingUsedVideoMemory, DynamicRayTracingGeometryVertexBuffers[i].NumBytes);
+				INC_MEMORY_STAT_BY(STAT_D3D12RayTracingDynamicVertexBufferMemory, DynamicRayTracingGeometryVertexBuffers[i].NumBytes);
+
+				Geometry.InitResource();
+			}
+
+			if (GetPrimitiveSceneInfo())
+			{
+				GetPrimitiveSceneInfo()->bIsRayTracingStaticRelevant = IsRayTracingStaticRelevant();
+			}
+		}
+	}
+	else if (!NewValue && bDynamicRayTracingGeometry)
+	{
+		bDynamicRayTracingGeometry = false;
+		if (IsRayTracingEnabled())
+		{
+			for (auto& Geometry : DynamicRayTracingGeometries)
+			{
+				Geometry.ReleaseResource();
+			}
+
+			DynamicRayTracingGeometries.Empty();
+
+			for (auto& Buffer : DynamicRayTracingGeometryVertexBuffers)
+			{
+				DEC_MEMORY_STAT_BY(STAT_D3D12RayTracingUsedVideoMemory, Buffer.NumBytes);
+				DEC_MEMORY_STAT_BY(STAT_D3D12RayTracingDynamicVertexBufferMemory, Buffer.NumBytes);
+				Buffer.Release();
+			}
+
+			DynamicRayTracingGeometryVertexBuffers.Empty();
+
+			if (GetPrimitiveSceneInfo())
+			{
+				GetPrimitiveSceneInfo()->bIsRayTracingStaticRelevant = IsRayTracingStaticRelevant();
+			}
+		}
+	}
+#endif
+}
+
 FStaticMeshSceneProxy::~FStaticMeshSceneProxy()
 {
 #if RHI_RAYTRACING
 	for (auto& Buffer: DynamicRayTracingGeometryVertexBuffers)
 	{
+		DEC_MEMORY_STAT_BY(STAT_D3D12RayTracingUsedVideoMemory, Buffer.NumBytes);
+		DEC_MEMORY_STAT_BY(STAT_D3D12RayTracingDynamicVertexBufferMemory, Buffer.NumBytes);
 		Buffer.Release();
 	}
 
@@ -565,6 +642,8 @@ void FStaticMeshSceneProxy::CreateRenderThreadResources()
 			auto& Geometry = DynamicRayTracingGeometries[i];
 			DynamicRayTracingGeometryVertexBuffers[i]
 				.Initialize(4, 256, PF_R32_FLOAT, BUF_UnorderedAccess | BUF_ShaderResource, TEXT("FStaticMeshSceneProxy::RayTracingDynamicVertexBuffer"));
+			INC_MEMORY_STAT_BY(STAT_D3D12RayTracingUsedVideoMemory, DynamicRayTracingGeometryVertexBuffers[i].NumBytes);
+			INC_MEMORY_STAT_BY(STAT_D3D12RayTracingDynamicVertexBufferMemory, DynamicRayTracingGeometryVertexBuffers[i].NumBytes);
 			Geometry.InitResource();
 		}
 	}
diff --git a/Engine/Source/Runtime/Engine/Public/StaticMeshResources.h b/Engine/Source/Runtime/Engine/Public/StaticMeshResources.h
--- a/Engine/Source/Runtime/Engine/Public/StaticMeshResources.h
+++ b/Engine/Source/Runtime/Engine/Public/StaticMeshResources.h
@@ -824,6 +824,8 @@ public:
 		const FMaterialRenderProxy* RenderProxy,
 		FMeshBatch& OutMeshBatch) const;
 
+	virtual void SetEvaluateWorldPositionOffsetInRayTracing(bool NewValue);
+
 	virtual uint8 GetCurrentFirstLODIdx_RenderThread() const final override
 	{
 		return GetCurrentFirstLODIdx_Internal();
diff --git a/Engine/Source/Runtime/Renderer/Private/PrimitiveSceneInfo.cpp b/Engine/Source/Runtime/Renderer/Private/PrimitiveSceneInfo.cpp
--- a/Engine/Source/Runtime/Renderer/Private/PrimitiveSceneInfo.cpp
+++ b/Engine/Source/Runtime/Renderer/Private/PrimitiveSceneInfo.cpp
@@ -401,7 +401,7 @@ void FPrimitiveSceneInfo::CacheMeshDrawCommands(FRHICommandListImmediate& RHICmd
 
 		for (FPrimitiveSceneInfo* SceneInfo : SceneInfos)
 		{
-			if (SceneInfo->Proxy->IsRayTracingStaticRelevant() && SceneInfo->StaticMeshes.Num() > 0)
+			if (SceneInfo->RayTracingGeometries.Num() > 0 && SceneInfo->StaticMeshes.Num() > 0)
 			{
 				int32 MaxLOD = -1;
 				for (int32 MeshIndex = 0; MeshIndex < SceneInfo->StaticMeshes.Num(); MeshIndex++)
