Fix PVS warning about testing the same value 2x.

diff --git a/Engine/Source/Runtime/Engine/Public/RawIndexBuffer.h b/Engine/Source/Runtime/Engine/Public/RawIndexBuffer.h
--- a/Engine/Source/Runtime/Engine/Public/RawIndexBuffer.h
+++ b/Engine/Source/Runtime/Engine/Public/RawIndexBuffer.h
@@ -370,7 +370,7 @@ public:
 	{
 		IndexBufferRHI = CreateRHIBuffer_RenderThread();
 
-		if (Indices.Num() && IndexBufferRHI && IsSRVNeeded() && Indices.Num())
+		if (Indices.Num() && IndexBufferRHI && IsSRVNeeded())
 		{
 			SRVValue = RHICreateShaderResourceView(IndexBufferRHI);
 		}
