REL - Crash on exit do to nullptr. Potentially this was guarded by GIsRequestingExit earlier.

diff --git a/Engine/Source/Runtime/Engine/Private/Particles/ParticleGpuSimulation.cpp b/Engine/Source/Runtime/Engine/Private/Particles/ParticleGpuSimulation.cpp
--- a/Engine/Source/Runtime/Engine/Private/Particles/ParticleGpuSimulation.cpp
+++ b/Engine/Source/Runtime/Engine/Private/Particles/ParticleGpuSimulation.cpp
@@ -3858,10 +3858,10 @@ private:
 		{
 			FXSystem->RemoveGPUSimulation( Simulation );
 
+			FParticleSimulationResources* ParticleSimulationResources = FXSystem->GetParticleSimulationResources();
 			// The check for IsEngineExitRequested() is done because at shut down UWorld can be destroyed before particle emitters(?)
-			if (!IsEngineExitRequested())
+			if (!IsEngineExitRequested() && ParticleSimulationResources)
 			{
-				FParticleSimulationResources* ParticleSimulationResources = FXSystem->GetParticleSimulationResources();
 				const int32 TileCount = AllocatedTiles.Num();
 				for ( int32 ActiveTileIndex = 0; ActiveTileIndex < TileCount; ++ActiveTileIndex )
 				{
