Revert 12139705 to fix unity builds

diff --git a/Engine/Plugins/FX/Niagara/Source/Niagara/Private/NiagaraRendererRibbons.cpp b/Engine/Plugins/FX/Niagara/Source/Niagara/Private/NiagaraRendererRibbons.cpp
--- a/Engine/Plugins/FX/Niagara/Source/Niagara/Private/NiagaraRendererRibbons.cpp
+++ b/Engine/Plugins/FX/Niagara/Source/Niagara/Private/NiagaraRendererRibbons.cpp
@@ -230,11 +230,6 @@ FNiagaraRendererRibbons::~FNiagaraRendererRibbons()
 {
 }
 
-
-DECLARE_STATS_GROUP(TEXT("D3D12RHI: Ray Tracing"), STATGROUP_D3D12RayTracing, STATCAT_Advanced);
-DECLARE_MEMORY_STAT(TEXT("Total Used Video Memory"), STAT_D3D12RayTracingUsedVideoMemory, STATGROUP_D3D12RayTracing);
-DECLARE_MEMORY_STAT(TEXT("Dynamic Vertex Buffer Memory"), STAT_D3D12RayTracingDynamicVertexBufferMemory, STATGROUP_D3D12RayTracing);
-
 void FNiagaraRendererRibbons::ReleaseRenderThreadResources()
 {
 	FNiagaraRenderer::ReleaseRenderThreadResources();
@@ -243,9 +238,6 @@ void FNiagaraRendererRibbons::ReleaseRenderThreadResources()
 	{
 		RayTracingGeometry.ReleaseResource();
 		RayTracingDynamicVertexBuffer.Release();
-
-		DEC_MEMORY_STAT_BY(STAT_D3D12RayTracingUsedVideoMemory, RayTracingDynamicVertexBuffer.NumBytes);
-		DEC_MEMORY_STAT_BY(STAT_D3D12RayTracingDynamicVertexBufferMemory, RayTracingDynamicVertexBuffer.NumBytes);
 	}
 #endif
 }
@@ -259,9 +251,6 @@ void FNiagaraRendererRibbons::CreateRenderThreadResources(NiagaraEmitterInstance
 	{
 		RayTracingDynamicVertexBuffer.Initialize(4, 256, PF_R32_FLOAT, BUF_UnorderedAccess | BUF_ShaderResource, TEXT("FNiagaraRendererRibbons::RayTracingDynamicVertexBuffer"));
 
-		INC_MEMORY_STAT_BY(STAT_D3D12RayTracingUsedVideoMemory, RayTracingDynamicVertexBuffer.NumBytes);
-		INC_MEMORY_STAT_BY(STAT_D3D12RayTracingDynamicVertexBufferMemory, RayTracingDynamicVertexBuffer.NumBytes);
-
 		FRayTracingGeometryInitializer Initializer;
 		Initializer.IndexBuffer = nullptr;
 		Initializer.TotalPrimitiveCount = 0;
diff --git a/Engine/Plugins/FX/Niagara/Source/Niagara/Private/NiagaraRendererSprites.cpp b/Engine/Plugins/FX/Niagara/Source/Niagara/Private/NiagaraRendererSprites.cpp
--- a/Engine/Plugins/FX/Niagara/Source/Niagara/Private/NiagaraRendererSprites.cpp
+++ b/Engine/Plugins/FX/Niagara/Source/Niagara/Private/NiagaraRendererSprites.cpp
@@ -154,10 +154,6 @@ FNiagaraRendererSprites::~FNiagaraRendererSprites()
 {
 }
 
-DECLARE_STATS_GROUP(TEXT("D3D12RHI: Ray Tracing"), STATGROUP_D3D12RayTracing, STATCAT_Advanced);
-DECLARE_MEMORY_STAT(TEXT("Total Used Video Memory"), STAT_D3D12RayTracingUsedVideoMemory, STATGROUP_D3D12RayTracing);
-DECLARE_MEMORY_STAT(TEXT("Dynamic Vertex Buffer Memory"), STAT_D3D12RayTracingDynamicVertexBufferMemory, STATGROUP_D3D12RayTracing);
-
 void FNiagaraRendererSprites::ReleaseRenderThreadResources()
 {
 	FNiagaraRenderer::ReleaseRenderThreadResources();
@@ -168,9 +164,6 @@ void FNiagaraRendererSprites::ReleaseRenderThreadResources()
 	{
 		RayTracingGeometry.ReleaseResource();
 		RayTracingDynamicVertexBuffer.Release();
-
-		DEC_MEMORY_STAT_BY(STAT_D3D12RayTracingUsedVideoMemory, RayTracingDynamicVertexBuffer.NumBytes);
-		DEC_MEMORY_STAT_BY(STAT_D3D12RayTracingDynamicVertexBufferMemory, RayTracingDynamicVertexBuffer.NumBytes);
 	}
 #endif
 }
@@ -185,9 +178,6 @@ void FNiagaraRendererSprites::CreateRenderThreadResources(NiagaraEmitterInstance
 	{
 		RayTracingDynamicVertexBuffer.Initialize(4, 256, PF_R32_FLOAT, BUF_UnorderedAccess | BUF_ShaderResource, TEXT("FNiagaraRendererSprites::RayTracingDynamicVertexBuffer"));
 
-		INC_MEMORY_STAT_BY(STAT_D3D12RayTracingUsedVideoMemory, RayTracingDynamicVertexBuffer.NumBytes);
-		INC_MEMORY_STAT_BY(STAT_D3D12RayTracingDynamicVertexBufferMemory, RayTracingDynamicVertexBuffer.NumBytes);
-
 		FRayTracingGeometryInitializer Initializer;
 		Initializer.IndexBuffer = nullptr;
 		Initializer.GeometryType = RTGT_Triangles;
diff --git a/Engine/Source/Runtime/Engine/Private/SkeletalRenderGPUSkin.cpp b/Engine/Source/Runtime/Engine/Private/SkeletalRenderGPUSkin.cpp
--- a/Engine/Source/Runtime/Engine/Private/SkeletalRenderGPUSkin.cpp
+++ b/Engine/Source/Runtime/Engine/Private/SkeletalRenderGPUSkin.cpp
@@ -189,11 +189,6 @@ void FSkeletalMeshObjectGPUSkin::InitResources(USkinnedMeshComponent* InMeshComp
 	BeginInitResource(&RayTracingGeometry);
 #endif
 }
-
-DECLARE_STATS_GROUP(TEXT("D3D12RHI: Ray Tracing"), STATGROUP_D3D12RayTracing, STATCAT_Advanced);
-DECLARE_MEMORY_STAT(TEXT("Total Used Video Memory"), STAT_D3D12RayTracingUsedVideoMemory, STATGROUP_D3D12RayTracing);
-DECLARE_MEMORY_STAT(TEXT("Dynamic Vertex Buffer Memory"), STAT_D3D12RayTracingDynamicVertexBufferMemory, STATGROUP_D3D12RayTracing);
-
 void FSkeletalMeshObjectGPUSkin::ReleaseResources()
 {
 	for( int32 LODIndex=0;LODIndex < LODs.Num();LODIndex++ )
@@ -228,8 +223,6 @@ void FSkeletalMeshObjectGPUSkin::ReleaseResources()
 		[RayTracingDynamicVertexBuffer = MoveTemp(RayTracingDynamicVertexBuffer)](FRHICommandListImmediate& RHICmdList) mutable
 	{
 		RayTracingDynamicVertexBuffer.Release();
-		DEC_MEMORY_STAT_BY(STAT_D3D12RayTracingUsedVideoMemory, RayTracingDynamicVertexBuffer.NumBytes);
-		DEC_MEMORY_STAT_BY(STAT_D3D12RayTracingDynamicVertexBufferMemory, RayTracingDynamicVertexBuffer.NumBytes);
 	});
 
 	check(!RayTracingDynamicVertexBuffer.Buffer.IsValid());
diff --git a/Engine/Source/Runtime/Engine/Private/StaticMeshRender.cpp b/Engine/Source/Runtime/Engine/Private/StaticMeshRender.cpp
--- a/Engine/Source/Runtime/Engine/Private/StaticMeshRender.cpp
+++ b/Engine/Source/Runtime/Engine/Private/StaticMeshRender.cpp
@@ -50,10 +50,6 @@ static bool GUseShadowIndexBuffer = true;
 /** If true, reversed index buffer are used for mesh with negative transform determinants. */
 static bool GUseReversedIndexBuffer = true;
 
-DECLARE_STATS_GROUP(TEXT("D3D12RHI: Ray Tracing"), STATGROUP_D3D12RayTracing, STATCAT_Advanced);
-DECLARE_MEMORY_STAT(TEXT("Total Used Video Memory"), STAT_D3D12RayTracingUsedVideoMemory, STATGROUP_D3D12RayTracing);
-DECLARE_MEMORY_STAT(TEXT("Dynamic Vertex Buffer Memory"), STAT_D3D12RayTracingDynamicVertexBufferMemory, STATGROUP_D3D12RayTracing);
-
 static void ToggleShadowIndexBuffers()
 {
 	FlushRenderingCommands();
@@ -336,10 +332,6 @@ void FStaticMeshSceneProxy::SetEvaluateWorldPositionOffsetInRayTracing(bool NewV
 			{
 				auto& Geometry = DynamicRayTracingGeometries[i];
 				DynamicRayTracingGeometryVertexBuffers[i].Initialize(4, 256, PF_R32_FLOAT, BUF_UnorderedAccess | BUF_ShaderResource, TEXT("RayTracingDynamicVertexBuffer"));
-
-				INC_MEMORY_STAT_BY(STAT_D3D12RayTracingUsedVideoMemory, DynamicRayTracingGeometryVertexBuffers[i].NumBytes);
-				INC_MEMORY_STAT_BY(STAT_D3D12RayTracingDynamicVertexBufferMemory, DynamicRayTracingGeometryVertexBuffers[i].NumBytes);
-
 				Geometry.InitResource();
 			}
 
@@ -363,8 +355,6 @@ void FStaticMeshSceneProxy::SetEvaluateWorldPositionOffsetInRayTracing(bool NewV
 
 			for (auto& Buffer : DynamicRayTracingGeometryVertexBuffers)
 			{
-				DEC_MEMORY_STAT_BY(STAT_D3D12RayTracingUsedVideoMemory, Buffer.NumBytes);
-				DEC_MEMORY_STAT_BY(STAT_D3D12RayTracingDynamicVertexBufferMemory, Buffer.NumBytes);
 				Buffer.Release();
 			}
 
@@ -384,8 +374,6 @@ FStaticMeshSceneProxy::~FStaticMeshSceneProxy()
 #if RHI_RAYTRACING
 	for (auto& Buffer: DynamicRayTracingGeometryVertexBuffers)
 	{
-		DEC_MEMORY_STAT_BY(STAT_D3D12RayTracingUsedVideoMemory, Buffer.NumBytes);
-		DEC_MEMORY_STAT_BY(STAT_D3D12RayTracingDynamicVertexBufferMemory, Buffer.NumBytes);
 		Buffer.Release();
 	}
 
@@ -642,8 +630,6 @@ void FStaticMeshSceneProxy::CreateRenderThreadResources()
 			auto& Geometry = DynamicRayTracingGeometries[i];
 			DynamicRayTracingGeometryVertexBuffers[i]
 				.Initialize(4, 256, PF_R32_FLOAT, BUF_UnorderedAccess | BUF_ShaderResource, TEXT("FStaticMeshSceneProxy::RayTracingDynamicVertexBuffer"));
-			INC_MEMORY_STAT_BY(STAT_D3D12RayTracingUsedVideoMemory, DynamicRayTracingGeometryVertexBuffers[i].NumBytes);
-			INC_MEMORY_STAT_BY(STAT_D3D12RayTracingDynamicVertexBufferMemory, DynamicRayTracingGeometryVertexBuffers[i].NumBytes);
 			Geometry.InitResource();
 		}
 	}
diff --git a/Engine/Source/Runtime/Landscape/Private/LandscapeRender.cpp b/Engine/Source/Runtime/Landscape/Private/LandscapeRender.cpp
--- a/Engine/Source/Runtime/Landscape/Private/LandscapeRender.cpp
+++ b/Engine/Source/Runtime/Landscape/Private/LandscapeRender.cpp
@@ -42,10 +42,6 @@ LandscapeRender.cpp: New terrain rendering
 #include "RayTracingInstance.h"
 #include "ProfilingDebugging/LoadTimeTracker.h"
 
-DECLARE_STATS_GROUP(TEXT("D3D12RHI: Ray Tracing"), STATGROUP_D3D12RayTracing, STATCAT_Advanced);
-DECLARE_MEMORY_STAT(TEXT("Total Used Video Memory"), STAT_D3D12RayTracingUsedVideoMemory, STATGROUP_D3D12RayTracing);
-DECLARE_MEMORY_STAT(TEXT("Dynamic Vertex Buffer Memory"), STAT_D3D12RayTracingDynamicVertexBufferMemory, STATGROUP_D3D12RayTracing);
-
 IMPLEMENT_GLOBAL_SHADER_PARAMETER_STRUCT(FLandscapeUniformShaderParameters, "LandscapeParameters");
 IMPLEMENT_GLOBAL_SHADER_PARAMETER_STRUCT(FLandscapeFixedGridUniformShaderParameters, "LandscapeFixedGrid");
 IMPLEMENT_TYPE_LAYOUT(FLandscapeVertexFactoryPixelShaderParameters);
@@ -1664,9 +1660,6 @@ FLandscapeComponentSceneProxy::~FLandscapeComponentSceneProxy()
 			const int8 SubSectionIdx = SubX + SubY * NumSubsections;
 			SectionRayTracingStates[SubSectionIdx].Geometry.ReleaseResource();
 			SectionRayTracingStates[SubSectionIdx].RayTracingDynamicVertexBuffer.Release();
-
-			DEC_MEMORY_STAT_BY(STAT_D3D12RayTracingUsedVideoMemory, SectionRayTracingStates[SubSectionIdx].RayTracingDynamicVertexBuffer.NumBytes);
-			DEC_MEMORY_STAT_BY(STAT_D3D12RayTracingDynamicVertexBufferMemory, SectionRayTracingStates[SubSectionIdx].RayTracingDynamicVertexBuffer.NumBytes);
 		}
 	}
 #endif
@@ -3709,8 +3702,6 @@ void FLandscapeComponentSceneProxy::GetDynamicRayTracingInstances(FRayTracingMat
 					bNeedsRayTracingGeometryUpdate = true;
 					SectionRayTracingStates[SubSectionIdx].CurrentLOD = CurrentLOD;
 					SectionRayTracingStates[SubSectionIdx].RayTracingDynamicVertexBuffer.Release();
-					DEC_MEMORY_STAT_BY(STAT_D3D12RayTracingUsedVideoMemory, SectionRayTracingStates[SubSectionIdx].RayTracingDynamicVertexBuffer.NumBytes);
-					DEC_MEMORY_STAT_BY(STAT_D3D12RayTracingDynamicVertexBufferMemory, SectionRayTracingStates[SubSectionIdx].RayTracingDynamicVertexBuffer.NumBytes);
 				}
 				if (SectionRayTracingStates[SubSectionIdx].HeightmapLODBias != RenderSystem.GetSectionLODBias(ComponentBase))
 				{
diff --git a/Engine/Source/Runtime/Renderer/Private/RayTracing/RayTracingDynamicGeometry.cpp b/Engine/Source/Runtime/Renderer/Private/RayTracing/RayTracingDynamicGeometry.cpp
--- a/Engine/Source/Runtime/Renderer/Private/RayTracing/RayTracingDynamicGeometry.cpp
+++ b/Engine/Source/Runtime/Renderer/Private/RayTracing/RayTracingDynamicGeometry.cpp
@@ -6,10 +6,6 @@
 
 #if RHI_RAYTRACING
 
-DECLARE_STATS_GROUP(TEXT("D3D12RHI: Ray Tracing"), STATGROUP_D3D12RayTracing, STATCAT_Advanced);
-DECLARE_MEMORY_STAT(TEXT("Total Used Video Memory"), STAT_D3D12RayTracingUsedVideoMemory, STATGROUP_D3D12RayTracing);
-DECLARE_MEMORY_STAT(TEXT("Dynamic Vertex Buffer Memory"), STAT_D3D12RayTracingDynamicVertexBufferMemory, STATGROUP_D3D12RayTracing);
-
 static bool IsSupportedDynamicVertexFactoryType(const FVertexFactoryType* VertexFactoryType)
 {
 	return VertexFactoryType == FindVertexFactoryType(FName(TEXT("FNiagaraSpriteVertexFactory"), FNAME_Find))
@@ -160,12 +156,7 @@ void FRayTracingDynamicGeometryCollection::AddDynamicMeshBatchForGeometryUpdate(
 	uint32 DesiredVertexBufferSize = UpdateParams.VertexBufferSize;
 	if (Buffer.NumBytes != DesiredVertexBufferSize)
 	{
-		int32 OriginalSize = Buffer.NumBytes;
-		DEC_MEMORY_STAT_BY(STAT_D3D12RayTracingUsedVideoMemory, Buffer.NumBytes);
-		DEC_MEMORY_STAT_BY(STAT_D3D12RayTracingDynamicVertexBufferMemory, Buffer.NumBytes);
 		Buffer.Initialize(sizeof(float), DesiredVertexBufferSize / sizeof(float), PF_R32_FLOAT, BUF_UnorderedAccess | BUF_ShaderResource, TEXT("FRayTracingDynamicGeometryCollection::RayTracingDynamicVertexBuffer"));
-		INC_MEMORY_STAT_BY(STAT_D3D12RayTracingUsedVideoMemory, Buffer.NumBytes);
-		INC_MEMORY_STAT_BY(STAT_D3D12RayTracingDynamicVertexBufferMemory, Buffer.NumBytes);
 		bRefit = false;
 	}
 
