Update an incorrect check

diff --git a/Engine/Source/Runtime/Engine/Private/Streaming/LandscapeMeshMobileUpdate.cpp b/Engine/Source/Runtime/Engine/Private/Streaming/LandscapeMeshMobileUpdate.cpp
--- a/Engine/Source/Runtime/Engine/Private/Streaming/LandscapeMeshMobileUpdate.cpp
+++ b/Engine/Source/Runtime/Engine/Private/Streaming/LandscapeMeshMobileUpdate.cpp
@@ -9,7 +9,6 @@ LandscapeMeshMobileUpdate.cpp: Helpers to stream in and out mobile landscape ver
 #include "LandscapeComponent.h"
 #include "LandscapeRenderMobile.h"
 #include "Streaming/RenderAssetUpdate.inl"
-#include "Streaming/StaticMeshUpdate.h"
 
 template class TRenderAssetUpdate<FLandscapeMeshMobileUpdateContext>;
 
@@ -18,7 +17,7 @@ FLandscapeMeshMobileUpdateContext::FLandscapeMeshMobileUpdateContext(ULandscapeL
 	, CurrentThread(InCurrentThread)
 {
 	check(InLandscapeProxy);
-	checkSlow(InCurrentThread != FStaticMeshUpdate::TT_Render || IsInRenderingThread());
+	checkSlow(InCurrentThread != FLandscapeMeshMobileUpdate::TT_Render || IsInRenderingThread());
 	RenderData = InLandscapeProxy->GetRenderData().Get();
 }
 
