Fixed dynamically created USkeletalMesh being configured for streaming by default.

diff --git a/Engine/Source/Runtime/Engine/Private/SkeletalMeshMerge.cpp b/Engine/Source/Runtime/Engine/Private/SkeletalMeshMerge.cpp
--- a/Engine/Source/Runtime/Engine/Private/SkeletalMeshMerge.cpp
+++ b/Engine/Source/Runtime/Engine/Private/SkeletalMeshMerge.cpp
@@ -212,6 +212,13 @@ bool FSkeletalMeshMerge::FinalizeMesh()
 			Result = false;
 		}
 
+		// If in game, streaming must be disabled as there are no files to stream from.
+		// In editor, the engine can stream from the DDC and create the required files on cook.
+		if (!GIsEditor)
+		{
+			MergeMesh->NeverStream = true;
+		}
+
 		// Reinitialize the mesh's render resources.
 		MergeMesh->InitResources();
 	}
diff --git a/Engine/Source/Runtime/Engine/Public/Rendering/SkeletalMeshLODRenderData.h b/Engine/Source/Runtime/Engine/Public/Rendering/SkeletalMeshLODRenderData.h
--- a/Engine/Source/Runtime/Engine/Public/Rendering/SkeletalMeshLODRenderData.h
+++ b/Engine/Source/Runtime/Engine/Public/Rendering/SkeletalMeshLODRenderData.h
@@ -161,7 +161,7 @@ public:
 	/** Constructor (default) */
 	FSkeletalMeshLODRenderData()
 		: BuffersSize(0)
-		, bStreamedDataInlined(false)
+		, bStreamedDataInlined(true)
 		, bIsLODOptional(false)
 	{
 	}
