Removing array property from attenuation settings to fix BP compile error

diff --git a/Engine/Source/Runtime/Engine/Classes/Sound/SoundAttenuation.h b/Engine/Source/Runtime/Engine/Classes/Sound/SoundAttenuation.h
--- a/Engine/Source/Runtime/Engine/Classes/Sound/SoundAttenuation.h
+++ b/Engine/Source/Runtime/Engine/Classes/Sound/SoundAttenuation.h
@@ -359,7 +359,7 @@ struct ENGINE_API FSoundAttenuationSettings : public FBaseAttenuationSettings
 
 	/** Set of submix send settings to use to send audio to submixes as a function of distance. */
 	UPROPERTY(EditAnywhere, BlueprintReadWrite, Category = AttenuationSubmixSend)
-	TArray<FAttenuationSubmixSendSettings> SubmixSendSettings;
+	FAttenuationSubmixSendSettings SubmixSendSettings;
 
 	/** The priority attenuation to use when the sound is at the minimum priority attenuation distance. */
 	UPROPERTY(EditAnywhere, BlueprintReadWrite, Category = AttenuationPriority, meta = (DisplayName = "Min Priority Attenuation"))
diff --git a/Engine/Source/Runtime/Engine/Private/ActiveSound.cpp b/Engine/Source/Runtime/Engine/Private/ActiveSound.cpp
--- a/Engine/Source/Runtime/Engine/Private/ActiveSound.cpp
+++ b/Engine/Source/Runtime/Engine/Private/ActiveSound.cpp
@@ -1544,13 +1544,17 @@ void FActiveSound::UpdateAttenuation(float DeltaTime, FSoundParseParameters& Par
 	if (Settings->bEnableSubmixSends)
 	{
 		ParseParams.SubmixSendSettings.Reset();
-		for (const FAttenuationSubmixSendSettings& SendSettings : Settings->SubmixSendSettings)
+		if (Settings->SubmixSendSettings.Submix)
 		{
-			if (SendSettings.Submix)
-			{
-				ParseParams.SubmixSendSettings.Add(SendSettings);
-			}
+			ParseParams.SubmixSendSettings.Add(Settings->SubmixSendSettings);
 		}
+// 		for (const FAttenuationSubmixSendSettings& SendSettings : Settings->SubmixSendSettings)
+// 		{
+// 			if (SendSettings.Submix)
+// 			{
+// 				ParseParams.SubmixSendSettings.Add(SendSettings);
+// 			}
+// 		}
 	}
 
 	check(AudioDevice);
