Skip resetting the text IDs when instancing a BP component template

diff --git a/Engine/Source/Runtime/Engine/Private/ActorConstruction.cpp b/Engine/Source/Runtime/Engine/Private/ActorConstruction.cpp
--- a/Engine/Source/Runtime/Engine/Private/ActorConstruction.cpp
+++ b/Engine/Source/Runtime/Engine/Private/ActorConstruction.cpp
@@ -982,7 +982,11 @@ UActorComponent* AActor::CreateComponentFromTemplate(UActorComponent* Template,
 		}
 
 		// Note we aren't copying the the RF_ArchetypeObject flag. Also note the result is non-transactional by default.
-		NewActorComp = (UActorComponent*)StaticDuplicateObject(Template, this, InName, RF_AllFlags & ~(RF_ArchetypeObject | RF_Transactional | RF_WasLoaded | RF_Public | RF_InheritableComponentTemplate));
+		FObjectDuplicationParameters DupeActorParameters(Template, this);
+		DupeActorParameters.DestName = InName;
+		DupeActorParameters.FlagMask = RF_AllFlags & ~(RF_ArchetypeObject | RF_Transactional | RF_WasLoaded | RF_Public | RF_InheritableComponentTemplate);
+		DupeActorParameters.PortFlags = PPF_DuplicateVerbatim; // Skip resetting text IDs
+		NewActorComp = (UActorComponent*)StaticDuplicateObjectEx(DupeActorParameters);
 
 		// Handle post-creation tasks.
 		PostCreateBlueprintComponent(NewActorComp);
@@ -1015,6 +1019,9 @@ UActorComponent* AActor::CreateComponentFromTemplateData(const FBlueprintCookedC
 
 			// Set this flag to emulate things that would happen in the SDO case when this flag is set (e.g. - not setting 'bHasBeenCreated').
 			ArPortFlags |= PPF_Duplicate;
+
+			// Skip resetting text IDs
+			ArPortFlags |= PPF_DuplicateVerbatim;
 		}
 	};
 
