#UE4 Disable deprecation temporarily

diff --git a/Engine/Source/Runtime/Engine/Public/AI/NavDataGenerator.h b/Engine/Source/Runtime/Engine/Public/AI/NavDataGenerator.h
--- a/Engine/Source/Runtime/Engine/Public/AI/NavDataGenerator.h
+++ b/Engine/Source/Runtime/Engine/Public/AI/NavDataGenerator.h
@@ -32,7 +32,7 @@ public:
 	virtual void RebuildDirtyAreas(const TArray<FNavigationDirtyArea>& DirtyAreas) {}
 
 	/** determines whether this generator is performing navigation building actions at the moment*/
-	UE_DEPRECATED(4.26, "This function is deprecated. Please use IsBuildInProgressCheckDirty")
+	//UE_DEPRECATED(4.26, "This function is deprecated. Please use IsBuildInProgressCheckDirty")
 	virtual bool IsBuildInProgress(bool bCheckDirtyToo = false) const { return false; }
 
 	/** determines whether this generator is performing navigation building actions at the moment, dirty areas are also checked */
diff --git a/Engine/Source/Runtime/NavigationSystem/Public/NavMesh/RecastNavMeshGenerator.h b/Engine/Source/Runtime/NavigationSystem/Public/NavMesh/RecastNavMeshGenerator.h
--- a/Engine/Source/Runtime/NavigationSystem/Public/NavMesh/RecastNavMeshGenerator.h
+++ b/Engine/Source/Runtime/NavigationSystem/Public/NavMesh/RecastNavMeshGenerator.h
@@ -608,7 +608,7 @@ public:
 	virtual void RebuildDirtyAreas(const TArray<FNavigationDirtyArea>& DirtyAreas) override;
 
 	PRAGMA_DISABLE_DEPRECATION_WARNINGS
-	UE_DEPRECATED(4.26, "This function is deprecated. Please use IsBuildInProgressCheckDirty")
+	//UE_DEPRECATED(4.26, "This function is deprecated. Please use IsBuildInProgressCheckDirty")
 	virtual bool IsBuildInProgress(bool bCheckDirtyToo = false) const override;
 	PRAGMA_ENABLE_DEPRECATION_WARNINGS
 
