#UE4 Added an option to allow you to override the level suffix for better networking in LoadLevelInstance. If the level is loaded on the server and client with the same name then the actors in them will network properly

diff --git a/Engine/Source/Runtime/Engine/Classes/Engine/LevelStreamingDynamic.h b/Engine/Source/Runtime/Engine/Classes/Engine/LevelStreamingDynamic.h
--- a/Engine/Source/Runtime/Engine/Classes/Engine/LevelStreamingDynamic.h
+++ b/Engine/Source/Runtime/Engine/Classes/Engine/LevelStreamingDynamic.h
@@ -35,17 +35,18 @@ class ENGINE_API ULevelStreamingDynamic : public ULevelStreaming
  	*
  	*   Project Settings -> Packaging -> List of Maps to Include in a Packaged Build (you may have to show advanced or type in filter)
  	* 
- 	* @param LevelName - Level package name, ex: /Game/Maps/MyMapName, specifying short name like MyMapName will force very slow search on disk
+ 	* @param LevelName - Level package name to load, ex: /Game/Maps/MyMapName, specifying short name like MyMapName will force very slow search on disk
  	* @param Location - World space location where the level should be spawned
  	* @param Rotation - World space rotation for rotating the entire level
 	* @param bOutSuccess - Whether operation was successful (map was found and added to the sub-levels list)
+	* @param OptionalLevelNameOverride - If set, the loaded level package have this name, which is used by other functions like UnloadStreamLevel. Note this is necessary for server and client networking because the level must have the same name on both.
  	* @return Streaming level object for a level instance
  	*/ 
  	UFUNCTION(BlueprintCallable, Category = LevelStreaming, meta=(DisplayName = "Load Level Instance (by Name)", WorldContext="WorldContextObject"))
- 	static ULevelStreamingDynamic* LoadLevelInstance(UObject* WorldContextObject, FString LevelName, FVector Location, FRotator Rotation, bool& bOutSuccess);
+ 	static ULevelStreamingDynamic* LoadLevelInstance(UObject* WorldContextObject, FString LevelName, FVector Location, FRotator Rotation, bool& bOutSuccess, const FName& OptionalLevelNameOverride = NAME_None);
 
  	UFUNCTION(BlueprintCallable, Category = LevelStreaming, meta=(DisplayName = "Load Level Instance (by Object Reference)", WorldContext="WorldContextObject"))
- 	static ULevelStreamingDynamic* LoadLevelInstanceBySoftObjectPtr(UObject* WorldContextObject, TSoftObjectPtr<UWorld> Level, FVector Location, FRotator Rotation, bool& bOutSuccess);
+ 	static ULevelStreamingDynamic* LoadLevelInstanceBySoftObjectPtr(UObject* WorldContextObject, TSoftObjectPtr<UWorld> Level, FVector Location, FRotator Rotation, bool& bOutSuccess, const FName& OptionalLevelNameOverride = NAME_None);
  	
 	//~ Begin UObject Interface
 	virtual void PostLoad() override;
@@ -62,7 +63,7 @@ private:
 	// Counter used by LoadLevelInstance to create unique level names
 	static int32 UniqueLevelInstanceId;
 
- 	static ULevelStreamingDynamic* LoadLevelInstance_Internal(UWorld* World, const FString& LongPackageName, FVector Location, FRotator Rotation, bool& bOutSuccess);
+ 	static ULevelStreamingDynamic* LoadLevelInstance_Internal(UWorld* World, const FString& LongPackageName, FVector Location, FRotator Rotation, bool& bOutSuccess, const FName& OptionalLevelNameOverride);
 
 };
 
diff --git a/Engine/Source/Runtime/Engine/Private/LevelStreaming.cpp b/Engine/Source/Runtime/Engine/Private/LevelStreaming.cpp
--- a/Engine/Source/Runtime/Engine/Private/LevelStreaming.cpp
+++ b/Engine/Source/Runtime/Engine/Private/LevelStreaming.cpp
@@ -1751,7 +1751,7 @@ void ULevelStreamingDynamic::SetShouldBeLoaded(const bool bInShouldBeLoaded)
 	}
 }
 
-ULevelStreamingDynamic* ULevelStreamingDynamic::LoadLevelInstance(UObject* WorldContextObject, const FString LevelName, const FVector Location, const FRotator Rotation, bool& bOutSuccess)
+ULevelStreamingDynamic* ULevelStreamingDynamic::LoadLevelInstance(UObject* WorldContextObject, const FString LevelName, const FVector Location, const FRotator Rotation, bool& bOutSuccess, const FName& OptionalLevelNameOverride)
 {
 	bOutSuccess = false;
 	UWorld* const World = GEngine->GetWorldFromContextObject(WorldContextObject, EGetWorldErrorMode::LogAndReturnNull);
@@ -1768,10 +1768,10 @@ ULevelStreamingDynamic* ULevelStreamingDynamic::LoadLevelInstance(UObject* World
 		return nullptr;
 	}
 
-	return LoadLevelInstance_Internal(World, LongPackageName, Location, Rotation, bOutSuccess);
+	return LoadLevelInstance_Internal(World, LongPackageName, Location, Rotation, bOutSuccess, OptionalLevelNameOverride);
 }
 
-ULevelStreamingDynamic* ULevelStreamingDynamic::LoadLevelInstanceBySoftObjectPtr(UObject* WorldContextObject, const TSoftObjectPtr<UWorld> Level, const FVector Location, const FRotator Rotation, bool& bOutSuccess)
+ULevelStreamingDynamic* ULevelStreamingDynamic::LoadLevelInstanceBySoftObjectPtr(UObject* WorldContextObject, const TSoftObjectPtr<UWorld> Level, const FVector Location, const FRotator Rotation, bool& bOutSuccess, const FName& OptionalLevelNameOverride)
 {
 	bOutSuccess = false;
 	UWorld* World = GEngine->GetWorldFromContextObject(WorldContextObject, EGetWorldErrorMode::LogAndReturnNull);
@@ -1786,10 +1786,10 @@ ULevelStreamingDynamic* ULevelStreamingDynamic::LoadLevelInstanceBySoftObjectPtr
 		return nullptr;
 	}
 
-	return LoadLevelInstance_Internal(World, Level.GetLongPackageName(), Location, Rotation, bOutSuccess);
+	return LoadLevelInstance_Internal(World, Level.GetLongPackageName(), Location, Rotation, bOutSuccess, OptionalLevelNameOverride);
 }
 
-ULevelStreamingDynamic* ULevelStreamingDynamic::LoadLevelInstance_Internal(UWorld* World, const FString& LongPackageName, const FVector Location, const FRotator Rotation, bool& bOutSuccess)
+ULevelStreamingDynamic* ULevelStreamingDynamic::LoadLevelInstance_Internal(UWorld* World, const FString& LongPackageName, const FVector Location, const FRotator Rotation, bool& bOutSuccess, const FName& OptionalLevelNameOverride)
 {
 	const FString PackagePath = FPackageName::GetLongPackagePath(LongPackageName);
 	FString ShortPackageName = FPackageName::GetShortName(LongPackageName);
@@ -1802,13 +1802,32 @@ ULevelStreamingDynamic* ULevelStreamingDynamic::LoadLevelInstance_Internal(UWorl
 	// Remove PIE prefix if it's there before we actually load the level
 	FString OnDiskPackageName = PackagePath + TEXT("/") + ShortPackageName;
 
-    // Create Unique Name for sub-level package
-	FString UniqueLevelPackageName = PackagePath + TEXT("/") + World->StreamingLevelsPrefix + ShortPackageName;
-	UniqueLevelPackageName += TEXT("_LevelInstance_") + FString::FromInt(++UniqueLevelInstanceId);
+	// Determine loaded package name
+	FString LevelPackageNameStr = PackagePath + TEXT("/") + World->StreamingLevelsPrefix;
+	bool bNeedsUniqueTest = false;
+	if (OptionalLevelNameOverride.IsNone())
+	{
+		// Create Unique Name for sub-level package
+		LevelPackageNameStr += ShortPackageName + TEXT("_LevelInstance_") + FString::FromInt(++UniqueLevelInstanceId);
+	}
+	else
+	{
+		// Use the supplied suffix, which is expected to result in a unique package name but we have to check if it is not.
+		LevelPackageNameStr += OptionalLevelNameOverride.ToString();
+		bNeedsUniqueTest = true;
+	}
+
+	FName LevelPackageName = FName(*LevelPackageNameStr);
+	if (bNeedsUniqueTest && World->GetStreamingLevels().ContainsByPredicate([&LevelPackageName](ULevelStreaming* LS) { return LS && LS->GetWorldAssetPackageFName() == LevelPackageName; }))
+	{
+		// The streaming level already exists, error and return.
+		UE_LOG(LogLevelStreaming, Error, TEXT("LoadLevelInstance called with a name that already exists, returning nullptr. LevelPackageName:%s"), *LevelPackageNameStr);
+		return nullptr;
+	}
     
 	// Setup streaming level object that will load specified map
 	ULevelStreamingDynamic* StreamingLevel = NewObject<ULevelStreamingDynamic>(World, ULevelStreamingDynamic::StaticClass(), NAME_None, RF_Transient, NULL);
-    StreamingLevel->SetWorldAssetByPackageName(FName(*UniqueLevelPackageName));
+    StreamingLevel->SetWorldAssetByPackageName(LevelPackageName);
     StreamingLevel->LevelColor = FColor::MakeRandomColor();
     StreamingLevel->SetShouldBeLoaded(true);
     StreamingLevel->SetShouldBeVisible(true);
