#UE4 Expose Get/SetIsRequestingUnloadAndRemoval to BPs for use with LoadLevelInstance or other dynamic streaming level cases

diff --git a/Engine/Source/Runtime/Engine/Classes/Engine/LevelStreaming.h b/Engine/Source/Runtime/Engine/Classes/Engine/LevelStreaming.h
--- a/Engine/Source/Runtime/Engine/Classes/Engine/LevelStreaming.h
+++ b/Engine/Source/Runtime/Engine/Classes/Engine/LevelStreaming.h
@@ -307,9 +307,11 @@ public:
 	}
 
 	/** Returns if the streaming level has requested to be unloaded and removed. */
+	UFUNCTION(BlueprintPure, Category = LevelStreaming)
 	bool GetIsRequestingUnloadAndRemoval() const { return bIsRequestingUnloadAndRemoval; }
 
 	/** Sets if the streaming level should be unloaded and removed. */
+	UFUNCTION(BlueprintCallable, Category = LevelStreaming)
 	void SetIsRequestingUnloadAndRemoval(bool bInIsRequestingUnloadAndRemoval);
 
 #if WITH_EDITORONLY_DATA
