Added validation case for zero-sized index buffers, as they are used for streaming scenarios.

diff --git a/Engine/Source/Runtime/Engine/Private/SceneManagement.cpp b/Engine/Source/Runtime/Engine/Private/SceneManagement.cpp
--- a/Engine/Source/Runtime/Engine/Private/SceneManagement.cpp
+++ b/Engine/Source/Runtime/Engine/Private/SceneManagement.cpp
@@ -1149,12 +1149,14 @@ bool FMeshBatch::Validate(const FPrimitiveSceneProxy* PrimitiveSceneProxy, ERHIF
 		{
 			if (const FRHIIndexBuffer* IndexBufferRHI = MeshBatchElement.IndexBuffer->IndexBufferRHI)
 			{
-				const int32 IndexCount = GetVertexCountForPrimitiveCount(MeshBatchElement.NumPrimitives, Type);
+				const uint32 IndexCount = GetVertexCountForPrimitiveCount(MeshBatchElement.NumPrimitives, Type);
+				const uint32 IndexBufferSize = IndexBufferRHI->GetSize();
 
-				if ((MeshBatchElement.FirstIndex + IndexCount) * IndexBufferRHI->GetStride() > IndexBufferRHI->GetSize())
+				// A zero-sized index buffer is valid for streaming.
+				if (IndexBufferSize != 0 && (MeshBatchElement.FirstIndex + IndexCount) * IndexBufferRHI->GetStride() > IndexBufferSize)
 				{
 					return LogMeshError(FString::Printf(
-						TEXT("MeshBatchElement %d, Material '%s', index range extends past index buffer bounds: Start %u, Count %u, Buffer Size %u, Buffer stride %u."),
+						TEXT("MeshBatchElement %d, Material '%s', index range extends past index buffer bounds: Start %u, Count %u, Buffer Size %u, Buffer stride %u"),
 						Index, MaterialRenderProxy ? *MaterialRenderProxy->GetFriendlyName() : TEXT("nullptr"),
 						MeshBatchElement.FirstIndex, IndexCount, IndexBufferRHI->GetSize(), IndexBufferRHI->GetStride()));
 				}
@@ -1172,7 +1174,7 @@ bool FMeshBatch::Validate(const FPrimitiveSceneProxy* PrimitiveSceneProxy, ERHIF
 
 	if (!PrimitiveSceneProxy->DoesVFRequirePrimitiveUniformBuffer() && !bVFSupportsPrimitiveIdStream)
 	{
-		return LogMeshError(TEXT("PrimitiveSceneProxy has bVFRequiresPrimitiveUniformBuffer disabled yet tried to draw with a vertex factory that did not support PrimitiveIdStream."));
+		return LogMeshError(TEXT("PrimitiveSceneProxy has bVFRequiresPrimitiveUniformBuffer disabled yet tried to draw with a vertex factory that did not support PrimitiveIdStream"));
 	}
 
 	const bool bPrimitiveShaderDataComesFromSceneBuffer = VertexFactory->GetPrimitiveIdStreamIndex(EVertexInputStreamType::Default) >= 0;
@@ -1188,7 +1190,7 @@ bool FMeshBatch::Validate(const FPrimitiveSceneProxy* PrimitiveSceneProxy, ERHIF
 			return LogMeshError(
 				TEXT("FMeshBatch was assigned a PrimitiveUniformBuffer even though the vertex factory fetches primitive shader data through the GPUScene buffer. ")
 				TEXT("The assigned PrimitiveUniformBuffer cannot be respected. Use PrimitiveUniformBufferResource instead for dynamic primitive data, or leave ")
-				TEXT("both null to get FPrimitiveSceneProxy->UniformBuffer."));
+				TEXT("both null to get FPrimitiveSceneProxy->UniformBuffer"));
 		}
 
 		const bool bValidPrimitiveData =
