changing FAudioChunkCache::GetChunk() so if it is called with bBlockForLoadCompletion it will perform that load synchronously instead of starting and waiting for an Async task

diff --git a/Engine/Source/Developer/Windows/WindowsTargetPlatform/Private/GenericWindowsTargetPlatform.h b/Engine/Source/Developer/Windows/WindowsTargetPlatform/Private/GenericWindowsTargetPlatform.h
--- a/Engine/Source/Developer/Windows/WindowsTargetPlatform/Private/GenericWindowsTargetPlatform.h
+++ b/Engine/Source/Developer/Windows/WindowsTargetPlatform/Private/GenericWindowsTargetPlatform.h
@@ -20,6 +20,7 @@
 
 #define LOCTEXT_NAMESPACE "TGenericWindowsTargetPlatform"
 
+		GConfig->GetBool(CategoryName, TEXT("bInlineStreamedAudioChunks"), OutOverrides.bInlineStreamedAudioChunks, GEngineIni);
 
 /**
  * Template for Windows target platforms
diff --git a/Engine/Source/Runtime/AudioPlatformConfiguration/Public/AudioCompressionSettings.h b/Engine/Source/Runtime/AudioPlatformConfiguration/Public/AudioCompressionSettings.h
--- a/Engine/Source/Runtime/AudioPlatformConfiguration/Public/AudioCompressionSettings.h
+++ b/Engine/Source/Runtime/AudioPlatformConfiguration/Public/AudioCompressionSettings.h
@@ -73,6 +73,9 @@ struct FPlatformAudioCookOverrides
 	// Whether to use the experimental Load on Demand feature, which uses as little memory at runtime as possible.
 	bool bUseStreamCaching;
 
+	// Whether to put streamed audio chunks inline in the Pak file or not (only matters if bUseStreamCaching is true)
+	bool bInlineStreamedAudioChunks;
+
 	// If Load On Demand is enabled, these settings are used to determine chunks and cache sizes.
 	FAudioStreamCachingSettings StreamCachingSettings;
 
@@ -81,6 +84,7 @@ struct FPlatformAudioCookOverrides
 		, CompressionQualityModifier(1.0f)
 		, AutoStreamingThreshold(0.0f)
 		, bUseStreamCaching(false)
+		, bInlineStreamedAudioChunks(false)
 	{
 		PlatformSampleRates.Add(ESoundwaveSampleRateSettings::Max, 48000);
 		PlatformSampleRates.Add(ESoundwaveSampleRateSettings::High, 32000);
diff --git a/Engine/Source/Runtime/Engine/Private/AudioCompressionSettingsUtils.cpp b/Engine/Source/Runtime/Engine/Private/AudioCompressionSettingsUtils.cpp
--- a/Engine/Source/Runtime/Engine/Private/AudioCompressionSettingsUtils.cpp
+++ b/Engine/Source/Runtime/Engine/Private/AudioCompressionSettingsUtils.cpp
@@ -99,6 +99,8 @@ void CacheAudioCookOverrides(FPlatformAudioCookOverrides& OutOverrides, const TC
 
 	PlatformFile.GetBool(*CategoryName, TEXT("bUseAudioStreamCaching"), OutOverrides.bUseStreamCaching);
 
+	GConfig->GetBool(CategoryName, TEXT("bInlineStreamedAudioChunks"), OutOverrides.bInlineStreamedAudioChunks, GEngineIni);
+
 	/** Memory Load On Demand Settings */
 	if (OutOverrides.bUseStreamCaching)
 	{
diff --git a/Engine/Source/Runtime/Engine/Private/AudioStreamingCache.cpp b/Engine/Source/Runtime/Engine/Private/AudioStreamingCache.cpp
--- a/Engine/Source/Runtime/Engine/Private/AudioStreamingCache.cpp
+++ b/Engine/Source/Runtime/Engine/Private/AudioStreamingCache.cpp
@@ -621,31 +621,49 @@ TArrayView<uint8> FAudioChunkCache::GetChunk(const FChunkKey& InKey, bool bBlock
 	{
 		// If we missed it, kick off a new load with it.
 		FoundElement = InsertChunk(InKey);
-
 		if (!FoundElement)
 		{
 			return TArrayView<uint8>();
 		}
-
-		KickOffAsyncLoad(FoundElement, InKey, [](EAudioChunkLoadResult InResult) {}, ENamedThreads::AnyThread, bNeededForPlayback);
-
 		if (bBlockForLoadCompletion)
 		{
-			// If bBlockForLoadCompletion was true and we don't have an element present, we have to load the element into the cache:
-			FoundElement->WaitForAsyncLoadCompletion(false);
+#if DEBUG_STREAM_CACHE
+			FStreamedAudioChunk& Chunk = InKey.SoundWave->RunningPlatformData->Chunks[InKey.ChunkIndex];
+			int32 ChunkAudioDataSize = Chunk.AudioDataSize;
+			FoundElement->DebugInfo.NumTotalChunks = InKey.SoundWave->GetNumChunks() - 1;
+			FoundElement->DebugInfo.TimeLoadStarted = FPlatformTime::Cycles64();
+#endif
+			MemoryCounterBytes -= FoundElement->ChunkData.Num();
+
+			// Reallocate our chunk data This allows us to shrink if possible.
+			// Unfortunately, GetCopy will write out the full zero-padded length of the bulk data,
+			// rather than just the audio data. So we set the array to Chunk.DataSize, then shrink to Chunk.AudioDataSize.
+			FoundElement->ChunkData.SetNumUninitialized(Chunk.DataSize, true);
+			void* DataDestPtr = FoundElement->ChunkData.GetData();
+			Chunk.BulkData.GetCopy(&DataDestPtr, true);
 
-			FoundElement->NumConsumers.Increment();
-			return TArrayView<uint8>(FoundElement->ChunkData.GetData(), FoundElement->ChunkDataSize);
+			FoundElement->ChunkData.SetNumUninitialized(ChunkAudioDataSize);
+			MemoryCounterBytes += FoundElement->ChunkData.Num();
+
+			// Populate key and DataSize. The async read request was set up to write directly into CacheElement->ChunkData.
+			FoundElement->Key = InKey;
+			FoundElement->ChunkDataSize = ChunkAudioDataSize;
+			FoundElement->bIsLoaded = true;
+#if DEBUG_STREAM_CACHE
+			FoundElement->DebugInfo.TimeToLoad = (FPlatformTime::Seconds() - FoundElement->DebugInfo.TimeLoadStarted) * 1000.0f;
+#endif
 		}
-		else if (bNeededForPlayback && (bLogCacheMisses || AlwaysLogCacheMissesCVar))
+		else
 		{
-			// We missed 
+			KickOffAsyncLoad(FoundElement, InKey, [](EAudioChunkLoadResult InResult) {}, ENamedThreads::AnyThread, bNeededForPlayback);
+		}
+		if (bLogCacheMisses && !bBlockForLoadCompletion)
+		{
+			// Chunks missing. Log this as a miss.
 			const uint32 TotalNumChunksInWave = InKey.SoundWave->GetNumChunks();
-
 			FCacheMissInfo CacheMissInfo = { InKey.SoundWaveName, InKey.ChunkIndex, TotalNumChunksInWave, false };
 			CacheMissQueue.Enqueue(MoveTemp(CacheMissInfo));
 		}
-
 		// We missed, return an empty array view.
 		return TArrayView<uint8>();
 	}
diff --git a/Engine/Source/Runtime/Engine/Private/SoundNode.cpp b/Engine/Source/Runtime/Engine/Private/SoundNode.cpp
--- a/Engine/Source/Runtime/Engine/Private/SoundNode.cpp
+++ b/Engine/Source/Runtime/Engine/Private/SoundNode.cpp
@@ -139,7 +139,7 @@ void USoundNode::RetainChildWavePlayers(bool bRecurse)
 				if (WavePlayer != nullptr)
 				{
 					USoundWave* SoundWave = WavePlayer->GetSoundWave();
-					if (SoundWave)
+					if (SoundWave && SoundWave->IsStreaming())
 					{
 						SoundWave->ConditionalPostLoad();
 						SoundWave->RetainCompressedAudio();
@@ -148,7 +148,6 @@ void USoundNode::RetainChildWavePlayers(bool bRecurse)
 			}
 		}
 	}
-
 	bIsRetainingAudio = true;
 }
 
diff --git a/Engine/Source/Runtime/Engine/Private/SoundWave.cpp b/Engine/Source/Runtime/Engine/Private/SoundWave.cpp
--- a/Engine/Source/Runtime/Engine/Private/SoundWave.cpp
+++ b/Engine/Source/Runtime/Engine/Private/SoundWave.cpp
@@ -36,7 +36,7 @@
 #include "Async/Async.h"
 #include "Misc/CommandLine.h"
 
-static int32 SoundWaveDefaultLoadingBehaviorCVar = 0;
+static int32 SoundWaveDefaultLoadingBehaviorCVar = 1;
 FAutoConsoleVariableRef CVarSoundWaveDefaultLoadingBehavior(
 	TEXT("au.streamcache.SoundWaveDefaultLoadingBehavior"),
 	SoundWaveDefaultLoadingBehaviorCVar,
@@ -126,6 +126,14 @@ ITargetPlatform* USoundWave::GetRunningPlatform()
 void FStreamedAudioChunk::Serialize(FArchive& Ar, UObject* Owner, int32 ChunkIndex)
 {
 	DECLARE_SCOPE_CYCLE_COUNTER( TEXT("FStreamedAudioChunk::Serialize"), STAT_StreamedAudioChunk_Serialize, STATGROUP_LoadTime );
+	bool bShouldInlineAudioChunk = false;
+
+	const ITargetPlatform* CookingTarget = Ar.CookingTarget();
+	if (CookingTarget != nullptr)
+	{
+		const FPlatformAudioCookOverrides* Overrides = CookingTarget->GetAudioCompressionSettings();
+		bShouldInlineAudioChunk = Overrides->bInlineStreamedAudioChunks;
+	}
 
 	bool bCooked = Ar.IsCooking();
 	Ar << bCooked;
@@ -133,7 +141,7 @@ void FStreamedAudioChunk::Serialize(FArchive& Ar, UObject* Owner, int32 ChunkInd
 	// ChunkIndex 0 is always inline payload, all other chunks are streamed.
 	if (Ar.IsSaving())
 	{
-		if (ChunkIndex == 0)
+		if (ChunkIndex == 0 || (ChunkIndex == 1 && bShouldInlineAudioChunk))
 		{
 			BulkData.SetBulkDataFlags(BULKDATA_ForceInlinePayload);
 		}
@@ -445,7 +453,7 @@ void USoundWave::Serialize( FArchive& Ar )
 		SampleRate = GetSampleRateForCurrentPlatform();
 
 		// If stream caching is enabled, here we determine if we should retain or prime this wave on load.
-		if (FPlatformCompressionUtilities::IsCurrentPlatformUsingStreamCaching())
+		if (bShouldStreamSound && FPlatformCompressionUtilities::IsCurrentPlatformUsingStreamCaching())
 		{
 			ESoundWaveLoadingBehavior CurrentLoadingBehavior = GetLoadingBehavior(false);
 
@@ -2475,7 +2483,7 @@ void USoundWave::RetainCompressedAudio(bool bForceSync /*= false*/)
 {
 	// Since the zeroth chunk is always inlined and stored in memory,
 	// early exit if we only have one chunk.
-	if (GIsEditor || DisableRetainingCVar || GetNumChunks() <= 1)
+	if (GIsEditor || !IsStreaming() || DisableRetainingCVar || GetNumChunks() <= 1)
 	{
 		return;
 	}
@@ -2487,6 +2495,7 @@ void USoundWave::RetainCompressedAudio(bool bForceSync /*= false*/)
 	else if (bForceSync)
 	{
 		FirstChunk = IStreamingManager::Get().GetAudioStreamingManager().GetLoadedChunk(this, 1, true);
+		ensureAlwaysMsgf(FirstChunk.IsValid(), TEXT("First chunk was invalid after synchronous load in RetainCompressedAudio()!"));
 	}
 	else
 	{
