Add logging for dropped chunks in the cached audio streaming manager.

diff --git a/Engine/Source/Runtime/Engine/Private/AudioStreamingCache.cpp b/Engine/Source/Runtime/Engine/Private/AudioStreamingCache.cpp
--- a/Engine/Source/Runtime/Engine/Private/AudioStreamingCache.cpp
+++ b/Engine/Source/Runtime/Engine/Private/AudioStreamingCache.cpp
@@ -325,6 +325,11 @@ FAudioChunkHandle FCachedAudioStreamingManager::GetLoadedChunk(const USoundWave*
 
 		// The function call below increments the reference count to the internal chunk.
 		TArrayView<uint8> LoadedChunk = Cache->GetChunk(ChunkKey, bBlockForLoad, (bForImmediatePlayback || bBlockForLoad));
+		
+		// Ensure that, if we requested a synchronous load of this chunk, we didn't fail to load said chunk.
+		ensureAlwaysMsgf(!bBlockForLoad || LoadedChunk.GetData(), TEXT("Synchronous load of chunk index %d for SoundWave %s failed to return any data."), ChunkIndex, *SoundWave->GetName());
+
+		UE_CLOG(!bBlockForLoad && !LoadedChunk.GetData(), LogAudio, Display, TEXT("GetLoadedChunk called for chunk index %d of SoundWave %s when audio was not loaded yet. This will result in latency."), ChunkIndex, *SoundWave->GetName());
 
 		// Finally, if there's a chunk after this in the sound, request that it is in the cache.
 		const int32 NextChunk = GetNextChunkIndex(SoundWave, ChunkIndex);
