Fix merge error
#rnx

diff --git a/Engine/Source/Runtime/Engine/Private/AudioCompressionSettingsUtils.cpp b/Engine/Source/Runtime/Engine/Private/AudioCompressionSettingsUtils.cpp
--- a/Engine/Source/Runtime/Engine/Private/AudioCompressionSettingsUtils.cpp
+++ b/Engine/Source/Runtime/Engine/Private/AudioCompressionSettingsUtils.cpp
@@ -99,7 +99,7 @@ void CacheAudioCookOverrides(FPlatformAudioCookOverrides& OutOverrides, const TC
 
 	PlatformFile.GetBool(*CategoryName, TEXT("bUseAudioStreamCaching"), OutOverrides.bUseStreamCaching);
 
-	GConfig->GetBool(CategoryName, TEXT("bInlineStreamedAudioChunks"), OutOverrides.bInlineStreamedAudioChunks, GEngineIni);
+	GConfig->GetBool(*CategoryName, TEXT("bInlineStreamedAudioChunks"), OutOverrides.bInlineStreamedAudioChunks, GEngineIni);
 
 	/** Memory Load On Demand Settings */
 	if (OutOverrides.bUseStreamCaching)
