Added secondary cvar scales to control view distance and hlod from game logic

diff --git a/Engine/Source/Runtime/Core/Private/HAL/ConsoleManager.cpp b/Engine/Source/Runtime/Core/Private/HAL/ConsoleManager.cpp
--- a/Engine/Source/Runtime/Core/Private/HAL/ConsoleManager.cpp
+++ b/Engine/Source/Runtime/Core/Private/HAL/ConsoleManager.cpp
@@ -2527,7 +2527,7 @@ static TAutoConsoleVariable<float> CVarViewDistanceScaleSecondaryScale(
 	1.0f,
 	TEXT("Controls the secondary view distance scale, Default = 1.0.\n")
 	TEXT("This is an optional scale intended to allow some features or gamemodes to opt-in.\n"),
-	ECVF_Scalability | ECVF_RenderThreadSafe);
+	ECVF_RenderThreadSafe);
 
 static TAutoConsoleVariable<float> CVarViewDistanceScale_FieldOfViewMinAngle(
 	TEXT("r.ViewDistanceScale.FieldOfViewMinAngle"),
diff --git a/Engine/Source/Runtime/Engine/Private/LODActor.cpp b/Engine/Source/Runtime/Engine/Private/LODActor.cpp
--- a/Engine/Source/Runtime/Engine/Private/LODActor.cpp
+++ b/Engine/Source/Runtime/Engine/Private/LODActor.cpp
@@ -57,6 +57,13 @@ ENGINE_API TAutoConsoleVariable<FString> CVarHLODDistanceOverride(
 	TEXT("'r.HLOD.DistanceOverride 5000, 10000, 20000' would result in HLOD levels 0, 1 and 2 transitioning at 5000, 1000 and 20000 respectively."),
 	ECVF_Scalability);
 
+static TAutoConsoleVariable<FString> CVarHLODDistanceOverrideScale(
+	TEXT("r.HLOD.DistanceOverrideScale"),
+	"",
+	TEXT("Scales the value in r.HLOD.DistanceOverride, Default off.\n")
+	TEXT("This is an optional scale intended to allow game logic to dynamically modify without impacting scalability.\n")
+);
+
 static TAutoConsoleVariable<int32> CVarHLODForceDisableCastDynamicShadow(
 	TEXT("r.HLOD.ForceDisableCastDynamicShadow"),
 	0,
@@ -365,17 +372,30 @@ void ALODActor::ParseOverrideDistancesCVar()
 {
 	// Parse HLOD override distance cvar into array
 	const FString DistanceOverrideValues = CVarHLODDistanceOverride.GetValueOnAnyThread();
+	const FString DistanceOverrideScaleValues = CVarHLODDistanceOverrideScale.GetValueOnAnyThread();
 	const TCHAR* Delimiters[] = { TEXT(","), TEXT(" ") };
 
 	TArray<FString> Distances;
 	DistanceOverrideValues.ParseIntoArray(/*out*/ Distances, Delimiters, UE_ARRAY_COUNT(Delimiters), /*bCullEmpty=*/ false);
 	Distances.RemoveAll([](const FString& String) { return String.IsEmpty(); });
 
+	TArray<FString> DistanceScales;
+	if (!DistanceOverrideScaleValues.IsEmpty())
+	{
+		DistanceOverrideScaleValues.ParseIntoArray(/*out*/ DistanceScales, Delimiters, UE_ARRAY_COUNT(Delimiters), /*bCullEmpty=*/ false);
+		DistanceScales.RemoveAll([](const FString& String) { return String.IsEmpty(); });
+	}	
+
 	HLODDistances.Empty(Distances.Num());
-	for (const FString& DistanceString : Distances)
+	for (int32 Index = 0; Index < Distances.Num(); ++Index)
 	{
-		const float DistanceForThisLevel = FCString::Atof(*DistanceString);
-		HLODDistances.Add(DistanceForThisLevel);
+		const float DistanceForThisLevel = FCString::Atof(*Distances[Index]);
+		float DistanceScaleForThisLevel = 1.f;
+		if (DistanceScales.IsValidIndex(Index))
+		{
+			DistanceScaleForThisLevel = FCString::Atof(*DistanceScales[Index]);
+		}
+		HLODDistances.Add(DistanceForThisLevel * DistanceScaleForThisLevel);
 	}
 }
 
