compile fix from merge of SoundWave.cpp

diff --git a/Engine/Source/Runtime/Engine/Private/SoundWave.cpp b/Engine/Source/Runtime/Engine/Private/SoundWave.cpp
--- a/Engine/Source/Runtime/Engine/Private/SoundWave.cpp
+++ b/Engine/Source/Runtime/Engine/Private/SoundWave.cpp
@@ -131,7 +131,7 @@ void FStreamedAudioChunk::Serialize(FArchive& Ar, UObject* Owner, int32 ChunkInd
 	const ITargetPlatform* CookingTarget = Ar.CookingTarget();
 	if (CookingTarget != nullptr)
 	{
-		const FPlatformAudioCookOverrides* Overrides = CookingTarget->GetAudioCompressionSettings();
+		const FPlatformAudioCookOverrides* Overrides = FPlatformCompressionUtilities::GetCookOverrides(*CookingTarget->IniPlatformName());
 		bShouldInlineAudioChunk = Overrides->bInlineStreamedAudioChunks;
 	}
 
