Add a more explicit warning when channel names are replicated as full strings and that causes overflows in bunch headers.

diff --git a/Engine/Source/Runtime/Engine/Private/NetConnection.cpp b/Engine/Source/Runtime/Engine/Private/NetConnection.cpp
--- a/Engine/Source/Runtime/Engine/Private/NetConnection.cpp
+++ b/Engine/Source/Runtime/Engine/Private/NetConnection.cpp
@@ -37,6 +37,7 @@
 #include "Net/NetworkGranularMemoryLogging.h"
 #include "SocketSubsystem.h"
 #include "Math/NumericLimits.h"
+#include "UObject/UnrealNames.h"
 
 static TAutoConsoleVariable<int32> CVarPingExcludeFrameTime( TEXT( "net.PingExcludeFrameTime" ), 0, TEXT( "Calculate RTT time between NIC's of server and client." ) );
 
@@ -2933,56 +2934,71 @@ int32 UNetConnection::SendRawBunch(FOutBunch& Bunch, bool InAllowMerge, const FN
 	// Build header.
 	SendBunchHeader.Reset();
 
-	SendBunchHeader.WriteBit( Bunch.bOpen || Bunch.bClose );
-	if( Bunch.bOpen || Bunch.bClose )
+	const bool bIsOpenOrClose = Bunch.bOpen || Bunch.bClose;
+	const bool bIsOpenOrReliable = Bunch.bOpen || Bunch.bReliable;
+
+	SendBunchHeader.WriteBit(bIsOpenOrClose);
+	if (bIsOpenOrClose)
 	{
-		SendBunchHeader.WriteBit( Bunch.bOpen );
-		SendBunchHeader.WriteBit( Bunch.bClose );
-		if( Bunch.bClose )
+		SendBunchHeader.WriteBit(Bunch.bOpen);
+		SendBunchHeader.WriteBit(Bunch.bClose);
+		if (Bunch.bClose)
 		{
 			uint32 Value = (uint32)Bunch.CloseReason;
-			SendBunchHeader.SerializeInt( Value, (uint32)EChannelCloseReason::MAX );
+			SendBunchHeader.SerializeInt(Value, (uint32)EChannelCloseReason::MAX);
 		}
 	}
-	SendBunchHeader.WriteBit( Bunch.bIsReplicationPaused );
-	SendBunchHeader.WriteBit( Bunch.bReliable );
+	SendBunchHeader.WriteBit(Bunch.bIsReplicationPaused);
+	SendBunchHeader.WriteBit(Bunch.bReliable);
 
 	uint32 ChIndex = Bunch.ChIndex;
 	SendBunchHeader.SerializeIntPacked(ChIndex); 
 
-	SendBunchHeader.WriteBit( Bunch.bHasPackageMapExports );
-	SendBunchHeader.WriteBit( Bunch.bHasMustBeMappedGUIDs );
-	SendBunchHeader.WriteBit( Bunch.bPartial );
+	SendBunchHeader.WriteBit(Bunch.bHasPackageMapExports);
+	SendBunchHeader.WriteBit(Bunch.bHasMustBeMappedGUIDs);
+	SendBunchHeader.WriteBit(Bunch.bPartial);
 
-	if ( Bunch.bReliable && !IsInternalAck() )
+	if (Bunch.bReliable && !IsInternalAck())
 	{
 		SendBunchHeader.WriteIntWrapped(Bunch.ChSequence, MAX_CHSEQUENCE);
 	}
 
 	if (Bunch.bPartial)
 	{
-		SendBunchHeader.WriteBit( Bunch.bPartialInitial );
-		SendBunchHeader.WriteBit( Bunch.bPartialFinal );
+		SendBunchHeader.WriteBit(Bunch.bPartialInitial);
+		SendBunchHeader.WriteBit(Bunch.bPartialFinal);
 	}
 
-	if (Bunch.bReliable || Bunch.bOpen)
+	if (bIsOpenOrReliable)
 	{
 		UPackageMap::StaticSerializeName(SendBunchHeader, Bunch.ChName);
 	}
 	
 	SendBunchHeader.WriteIntWrapped(Bunch.GetNumBits(), UNetConnection::MaxPacket * 8);
+
+#if DO_CHECK
+	if (UNLIKELY(SendBunchHeader.IsError()))
+	{
+		const bool bDidReplicateChannelName = bIsOpenOrReliable;
+		const bool bDoesChannelNameReplicateAsString = !Bunch.ChName.ToEName() || !ShouldReplicateAsInteger(*Bunch.ChName.ToEName());
+
+		checkf(false, TEXT("SendBunchHeader Error: Bunch = %s,  Channel Name Serialized As String: %d"),
+			*Bunch.ToString(), !!(bDidReplicateChannelName && bDoesChannelNameReplicateAsString));
+	}
+#endif
+
 	check(!SendBunchHeader.IsError());
 
 	// Remember start position.
-	AllowMerge      = InAllowMerge;
-	Bunch.Time      = Driver->GetElapsedTime();
+	AllowMerge = InAllowMerge;
+	Bunch.Time = Driver->GetElapsedTime();
 
-	if ((Bunch.bClose || Bunch.bOpen) && UE_LOG_ACTIVE(LogNetDormancy,VeryVerbose) )
+	if (bIsOpenOrClose && UE_LOG_ACTIVE(LogNetDormancy,VeryVerbose))
 	{
 		UE_LOG(LogNetDormancy, VeryVerbose, TEXT("Sending: %s"), *Bunch.ToString());
 	}
 
-	if (UE_LOG_ACTIVE(LogNetTraffic,VeryVerbose))
+	if (UE_LOG_ACTIVE(LogNetTraffic, VeryVerbose))
 	{
 		UE_LOG(LogNetTraffic, VeryVerbose, TEXT("Sending: %s"), *Bunch.ToString());
 	}
@@ -3017,14 +3033,14 @@ int32 UNetConnection::SendRawBunch(FOutBunch& Bunch, bool InAllowMerge, const FN
 	// Track channels that wrote data to this packet.
 	FChannelRecordImpl::PushChannelRecord(ChannelRecord, Bunch.PacketId, Bunch.ChIndex);
 
-	UE_LOG(LogNetTraffic, Verbose, TEXT("UNetConnection::SendRawBunch. ChIndex: %d. Bits: %d. PacketId: %d"), Bunch.ChIndex, Bunch.GetNumBits(), Bunch.PacketId );
+	UE_LOG(LogNetTraffic, Verbose, TEXT("UNetConnection::SendRawBunch. ChIndex: %d. Bits: %d. PacketId: %d"), Bunch.ChIndex, Bunch.GetNumBits(), Bunch.PacketId);
 
-	if ( PackageMap && Bunch.bHasPackageMapExports )
+	if (PackageMap && Bunch.bHasPackageMapExports)
 	{
-		PackageMap->NotifyBunchCommit( Bunch.PacketId, &Bunch );
+		PackageMap->NotifyBunchCommit(Bunch.PacketId, &Bunch);
 	}
 
-	if ( Bunch.bHasPackageMapExports )
+	if (Bunch.bHasPackageMapExports)
 	{
 		Driver->NetGUIDOutBytes += (SendBunchHeader.GetNumBits() + Bunch.GetNumBits()) >> 3;
 	}
