Move logic out of UBlueprint::PostLoadSubobjects which runs after archetypes have been used for instancing - move into compilation manager

diff --git a/Engine/Source/Editor/Kismet/Private/BlueprintCompilationManager.cpp b/Engine/Source/Editor/Kismet/Private/BlueprintCompilationManager.cpp
--- a/Engine/Source/Editor/Kismet/Private/BlueprintCompilationManager.cpp
+++ b/Engine/Source/Editor/Kismet/Private/BlueprintCompilationManager.cpp
@@ -847,11 +847,22 @@ void FBlueprintCompilationManagerImpl::FlushCompilationQueueImpl(bool bSuppressB
 			CompilerData.Compiler->ValidateVariableNames();
 		}
 
-		// STAGE VI: Purge null graphs, could be done only on load
+		// STAGE VI: Purge null graphs, misc. data fixup
 		for (FCompilerData& CompilerData : CurrentlyCompilingBPs)
 		{
 			UBlueprint* BP = CompilerData.BP;
-			FBlueprintEditorUtils::PurgeNullGraphs(BP);
+			if(BP->bIsRegeneratingOnLoad)
+			{
+				FBlueprintEditorUtils::PurgeNullGraphs(BP);
+				BP->ConformNativeComponents();
+				if (FLinkerLoad* Linker = BP->GetLinker())
+				{
+					if (Linker->UE4Ver() < VER_UE4_EDITORONLY_BLUEPRINTS)
+					{
+						BP->ChangeOwnerOfTemplates();
+					}
+				}
+			}
 		}
 
 		// STAGE VII: safely throw away old skeleton CDOs:
diff --git a/Engine/Source/Runtime/Engine/Classes/Engine/Blueprint.h b/Engine/Source/Runtime/Engine/Classes/Engine/Blueprint.h
--- a/Engine/Source/Runtime/Engine/Classes/Engine/Blueprint.h
+++ b/Engine/Source/Runtime/Engine/Classes/Engine/Blueprint.h
@@ -832,10 +832,7 @@ public:
 	virtual bool Rename(const TCHAR* NewName = nullptr, UObject* NewOuter = nullptr, ERenameFlags Flags = REN_None) override;
 	virtual UClass* RegenerateClass(UClass* ClassToRegenerate, UObject* PreviousCDO) override;
 	virtual void PostLoad() override;
-	virtual void PostLoadSubobjects( FObjectInstancingGraph* OuterInstanceGraph ) override;
-#if WITH_EDITOR
 	virtual bool Modify(bool bAlwaysMarkDirty = true) override;
-#endif
 	virtual void GetAssetRegistryTags(TArray<FAssetRegistryTag>& OutTags) const override;
 	virtual FPrimaryAssetId GetPrimaryAssetId() const override;
 	virtual void BeginCacheForCookedPlatformData(const ITargetPlatform *TargetPlatform) override;
@@ -885,6 +882,7 @@ public:
 	 */
 	virtual bool FindDiffs(const UBlueprint* OtherBlueprint, FDiffResults& Results) const;
 
+	void ConformNativeComponents();
 #endif	//#if WITH_EDITOR
 
 	//~ Begin UObject Interface
diff --git a/Engine/Source/Runtime/Engine/Private/Blueprint.cpp b/Engine/Source/Runtime/Engine/Private/Blueprint.cpp
--- a/Engine/Source/Runtime/Engine/Private/Blueprint.cpp
+++ b/Engine/Source/Runtime/Engine/Private/Blueprint.cpp
@@ -57,10 +57,9 @@ DEFINE_LOG_CATEGORY(LogBlueprint);
  * 
  * @param  Blueprint	The blueprint whose components you wish to vet.
  */
-static void ConformNativeComponents(UBlueprint* Blueprint)
+void UBlueprint::ConformNativeComponents()
 {
-#if WITH_EDITOR
-	if (UClass* const BlueprintClass = Blueprint->GeneratedClass)
+	if (UClass* const BlueprintClass = GeneratedClass)
 	{
 		if (AActor* BlueprintCDO = Cast<AActor>(BlueprintClass->ClassDefaultObject))
 		{
@@ -117,7 +116,6 @@ static void ConformNativeComponents(UBlueprint* Blueprint)
 			}
 		}
 	}
-#endif // #if WITH_EDITOR
 }
 
 
@@ -1822,8 +1820,6 @@ bool UBlueprint::ChangeOwnerOfTemplates()
 	bool bIsStillStale = false;
 	if (BPGClass)
 	{
-		check(!bIsRegeneratingOnLoad);
-
 		// >>> Backwards Compatibility:  VER_UE4_EDITORONLY_BLUEPRINTS
 		bool bMigratedOwner = false;
 		TSet<class UCurveBase*> Curves;
@@ -1928,14 +1924,6 @@ bool UBlueprint::ChangeOwnerOfTemplates()
 	return !bIsStillStale;
 }
 
-void UBlueprint::PostLoadSubobjects(FObjectInstancingGraph* OuterInstanceGraph)
-{
-	Super::PostLoadSubobjects(OuterInstanceGraph);
-	ChangeOwnerOfTemplates();
-
-	ConformNativeComponents(this);
-}
-
 #if WITH_EDITOR
 bool UBlueprint::Modify(bool bAlwaysMarkDirty)
 {
