Fix static analysis warning.

#rnx

diff --git a/Engine/Source/Runtime/Engine/Private/RepLayout.cpp b/Engine/Source/Runtime/Engine/Private/RepLayout.cpp
--- a/Engine/Source/Runtime/Engine/Private/RepLayout.cpp
+++ b/Engine/Source/Runtime/Engine/Private/RepLayout.cpp
@@ -1543,7 +1543,11 @@ ERepLayoutResult FRepLayout::CompareProperties(
 	}
 
 	// If something went wrong, don't touch the changelist history.
-	if (UNLIKELY(ERepLayoutResult::Success != Result))
+	// NOTE: Currently, SA always throws "Warning: Expression 'ERepLayoutResult::Success != Result' is always false" here.
+	//			It gets tripped up by the indirection / reference semantics of Stack Params and assumes that
+	//			result just remains the same value as it was assigned above.
+	//			Because of this, the error is disabled.
+	if (UNLIKELY(ERepLayoutResult::Success != Result)) // -V547
 	{
 		return Result;
 	}
