Respect bAllowReregistration in UActorComponent::ReregisterComponent
Cleaned up confusing comment about reregister component.

diff --git a/Engine/Source/Runtime/Engine/Classes/Components/ActorComponent.h b/Engine/Source/Runtime/Engine/Classes/Components/ActorComponent.h
--- a/Engine/Source/Runtime/Engine/Classes/Components/ActorComponent.h
+++ b/Engine/Source/Runtime/Engine/Classes/Components/ActorComponent.h
@@ -898,7 +898,7 @@ public:
 	UFUNCTION(BlueprintCallable, Category="Components", meta=(Keywords = "Delete", HidePin="Object", DefaultToSelf="Object", DisplayName = "DestroyComponent", ScriptName = "DestroyComponent"))
 	void K2_DestroyComponent(UObject* Object);
 
-	/** Unregisters and immediately re-registers component.  Handles bWillReregister properly. */
+	/** Unregisters and immediately re-registers component. */
 	void ReregisterComponent();
 
 	/** Changes the ticking group for this component */
diff --git a/Engine/Source/Runtime/Engine/Private/Components/ActorComponent.cpp b/Engine/Source/Runtime/Engine/Private/Components/ActorComponent.cpp
--- a/Engine/Source/Runtime/Engine/Private/Components/ActorComponent.cpp
+++ b/Engine/Source/Runtime/Engine/Private/Components/ActorComponent.cpp
@@ -1483,13 +1483,16 @@ void UActorComponent::ExecuteUnregisterEvents()
 
 void UActorComponent::ReregisterComponent()
 {
-	if(!IsRegistered())
+	if (AllowReregistration())
 	{
-		UE_LOG(LogActorComponent, Log, TEXT("ReregisterComponent: (%s) Not currently registered. Aborting."), *GetPathName());
-		return;
-	}
+		if (!IsRegistered())
+		{
+			UE_LOG(LogActorComponent, Log, TEXT("ReregisterComponent: (%s) Not currently registered. Aborting."), *GetPathName());
+			return;
+		}
 
-	FComponentReregisterContext(this);
+		FComponentReregisterContext(this);
+	}
 }
 
 void UActorComponent::RecreateRenderState_Concurrent()
