Commenting out cook override logic causing cook on Mac to fail.

#rnx

diff --git a/Engine/Source/Runtime/Engine/Private/SoundWave.cpp b/Engine/Source/Runtime/Engine/Private/SoundWave.cpp
--- a/Engine/Source/Runtime/Engine/Private/SoundWave.cpp
+++ b/Engine/Source/Runtime/Engine/Private/SoundWave.cpp
@@ -128,12 +128,12 @@ void FStreamedAudioChunk::Serialize(FArchive& Ar, UObject* Owner, int32 ChunkInd
 	DECLARE_SCOPE_CYCLE_COUNTER( TEXT("FStreamedAudioChunk::Serialize"), STAT_StreamedAudioChunk_Serialize, STATGROUP_LoadTime );
 	bool bShouldInlineAudioChunk = false;
 
-	const ITargetPlatform* CookingTarget = Ar.CookingTarget();
-	if (CookingTarget != nullptr)
-	{
-		const FPlatformAudioCookOverrides* Overrides = FPlatformCompressionUtilities::GetCookOverrides(*CookingTarget->IniPlatformName());
-		bShouldInlineAudioChunk = Overrides->bInlineStreamedAudioChunks;
-	}
+// 	const ITargetPlatform* CookingTarget = Ar.CookingTarget();
+// 	if (CookingTarget != nullptr)
+// 	{
+// 		const FPlatformAudioCookOverrides* Overrides = CookingTarget->GetAudioCompressionSettings();
+// 		bShouldInlineAudioChunk = Overrides->bInlineStreamedAudioChunks;
+// 	}
 
 	bool bCooked = Ar.IsCooking();
 	Ar << bCooked;
