Fix for regression in the way SoloSound console command worked

diff --git a/Engine/Source/Runtime/Engine/Private/Audio.cpp b/Engine/Source/Runtime/Engine/Private/Audio.cpp
--- a/Engine/Source/Runtime/Engine/Private/Audio.cpp
+++ b/Engine/Source/Runtime/Engine/Private/Audio.cpp
@@ -486,12 +486,12 @@ float FSoundSource::GetDebugVolume(const float InVolume)
 		}
 	}
 
-	// SoundCues mutes/solos											
+	// SoundCues mutes/solos (not strictly just cues but any SoundBase)
 	if (OutVolume != 0.0f && WaveInstance->ActiveSound)
 	{						
-		if (USoundCue* SoundCue = Cast<USoundCue>(WaveInstance->ActiveSound->GetSound()))
+		if (USoundBase* ActiveSound= WaveInstance->ActiveSound->GetSound())
 		{
-			Debugger.QuerySoloMuteSoundCue(SoundCue->GetName(), Info.bIsSoloed, Info.bIsMuted, Info.MuteSoloReason);
+			Debugger.QuerySoloMuteSoundCue(ActiveSound->GetName(), Info.bIsSoloed, Info.bIsMuted, Info.MuteSoloReason);
 			if (Info.bIsMuted)
 			{
 				OutVolume = 0.0f;
diff --git a/Engine/Source/Runtime/Engine/Private/AudioDebug.cpp b/Engine/Source/Runtime/Engine/Private/AudioDebug.cpp
--- a/Engine/Source/Runtime/Engine/Private/AudioDebug.cpp
+++ b/Engine/Source/Runtime/Engine/Private/AudioDebug.cpp
@@ -440,7 +440,12 @@ void FAudioDebugger::ToggleNameArray(FName InName, TArray<FName>& InNameArray, b
 			{
 				InNameArray.Empty();
 			}
-			InNameArray.Add(InName);
+
+			// Allow for NAME_None to act as a clear
+			if (InName != NAME_None)
+			{
+				InNameArray.Add(InName);
+			}						
 		}
 	});
 }
diff --git a/Engine/Source/Runtime/Engine/Private/AudioDevice.cpp b/Engine/Source/Runtime/Engine/Private/AudioDevice.cpp
--- a/Engine/Source/Runtime/Engine/Private/AudioDevice.cpp
+++ b/Engine/Source/Runtime/Engine/Private/AudioDevice.cpp
@@ -1388,6 +1388,11 @@ void FAudioDevice::HandleAudioSoloCommon(const TCHAR* Cmd, FOutputDevice& Ar, FT
 		{
 			(DeviceManager->GetDebugger().*FPtr)(*Args[0], bExclusive);
 		}
+		else if(bExclusive)
+		{
+			// If we are exclusive and no argument is passed, pass NAME_None to clear the current state.
+			(DeviceManager->GetDebugger().*FPtr)(NAME_None, true);
+		}
 	}
 }
 
