Use bCullEmpty for HLOD CVars instead of calling RemoveAll

diff --git a/Engine/Source/Runtime/Engine/Private/LODActor.cpp b/Engine/Source/Runtime/Engine/Private/LODActor.cpp
--- a/Engine/Source/Runtime/Engine/Private/LODActor.cpp
+++ b/Engine/Source/Runtime/Engine/Private/LODActor.cpp
@@ -376,14 +376,12 @@ void ALODActor::ParseOverrideDistancesCVar()
 	const TCHAR* Delimiters[] = { TEXT(","), TEXT(" ") };
 
 	TArray<FString> Distances;
-	DistanceOverrideValues.ParseIntoArray(/*out*/ Distances, Delimiters, UE_ARRAY_COUNT(Delimiters), /*bCullEmpty=*/ false);
-	Distances.RemoveAll([](const FString& String) { return String.IsEmpty(); });
+	DistanceOverrideValues.ParseIntoArray(/*out*/ Distances, Delimiters, UE_ARRAY_COUNT(Delimiters), true);
 
 	TArray<FString> DistanceScales;
 	if (!DistanceOverrideScaleValues.IsEmpty())
 	{
-		DistanceOverrideScaleValues.ParseIntoArray(/*out*/ DistanceScales, Delimiters, UE_ARRAY_COUNT(Delimiters), /*bCullEmpty=*/ false);
-		DistanceScales.RemoveAll([](const FString& String) { return String.IsEmpty(); });
+		DistanceOverrideScaleValues.ParseIntoArray(/*out*/ DistanceScales, Delimiters, UE_ARRAY_COUNT(Delimiters), true);
 	}	
 
 	HLODDistances.Empty(Distances.Num());
