#UE4 Make OptionalLevelNameOverride actually optional, and a string

diff --git a/Engine/Source/Runtime/Engine/Classes/Engine/LevelStreamingDynamic.h b/Engine/Source/Runtime/Engine/Classes/Engine/LevelStreamingDynamic.h
--- a/Engine/Source/Runtime/Engine/Classes/Engine/LevelStreamingDynamic.h
+++ b/Engine/Source/Runtime/Engine/Classes/Engine/LevelStreamingDynamic.h
@@ -43,10 +43,10 @@ class ENGINE_API ULevelStreamingDynamic : public ULevelStreaming
  	* @return Streaming level object for a level instance
  	*/ 
  	UFUNCTION(BlueprintCallable, Category = LevelStreaming, meta=(DisplayName = "Load Level Instance (by Name)", WorldContext="WorldContextObject"))
- 	static ULevelStreamingDynamic* LoadLevelInstance(UObject* WorldContextObject, FString LevelName, FVector Location, FRotator Rotation, bool& bOutSuccess, const FName& OptionalLevelNameOverride = NAME_None);
+ 	static ULevelStreamingDynamic* LoadLevelInstance(UObject* WorldContextObject, FString LevelName, FVector Location, FRotator Rotation, bool& bOutSuccess, FString OptionalLevelNameOverride = TEXT(""));
 
  	UFUNCTION(BlueprintCallable, Category = LevelStreaming, meta=(DisplayName = "Load Level Instance (by Object Reference)", WorldContext="WorldContextObject"))
- 	static ULevelStreamingDynamic* LoadLevelInstanceBySoftObjectPtr(UObject* WorldContextObject, TSoftObjectPtr<UWorld> Level, FVector Location, FRotator Rotation, bool& bOutSuccess, const FName& OptionalLevelNameOverride = NAME_None);
+ 	static ULevelStreamingDynamic* LoadLevelInstanceBySoftObjectPtr(UObject* WorldContextObject, TSoftObjectPtr<UWorld> Level, FVector Location, FRotator Rotation, bool& bOutSuccess, FString OptionalLevelNameOverride = TEXT(""));
  	
 	//~ Begin UObject Interface
 	virtual void PostLoad() override;
@@ -63,7 +63,7 @@ private:
 	// Counter used by LoadLevelInstance to create unique level names
 	static int32 UniqueLevelInstanceId;
 
- 	static ULevelStreamingDynamic* LoadLevelInstance_Internal(UWorld* World, const FString& LongPackageName, FVector Location, FRotator Rotation, bool& bOutSuccess, const FName& OptionalLevelNameOverride);
+ 	static ULevelStreamingDynamic* LoadLevelInstance_Internal(UWorld* World, const FString& LongPackageName, FVector Location, FRotator Rotation, bool& bOutSuccess, const FString& OptionalLevelNameOverride);
 
 };
 
diff --git a/Engine/Source/Runtime/Engine/Private/LevelStreaming.cpp b/Engine/Source/Runtime/Engine/Private/LevelStreaming.cpp
--- a/Engine/Source/Runtime/Engine/Private/LevelStreaming.cpp
+++ b/Engine/Source/Runtime/Engine/Private/LevelStreaming.cpp
@@ -1751,7 +1751,7 @@ void ULevelStreamingDynamic::SetShouldBeLoaded(const bool bInShouldBeLoaded)
 	}
 }
 
-ULevelStreamingDynamic* ULevelStreamingDynamic::LoadLevelInstance(UObject* WorldContextObject, const FString LevelName, const FVector Location, const FRotator Rotation, bool& bOutSuccess, const FName& OptionalLevelNameOverride)
+ULevelStreamingDynamic* ULevelStreamingDynamic::LoadLevelInstance(UObject* WorldContextObject, const FString LevelName, const FVector Location, const FRotator Rotation, bool& bOutSuccess, FString OptionalLevelNameOverride)
 {
 	bOutSuccess = false;
 	UWorld* const World = GEngine->GetWorldFromContextObject(WorldContextObject, EGetWorldErrorMode::LogAndReturnNull);
@@ -1771,7 +1771,7 @@ ULevelStreamingDynamic* ULevelStreamingDynamic::LoadLevelInstance(UObject* World
 	return LoadLevelInstance_Internal(World, LongPackageName, Location, Rotation, bOutSuccess, OptionalLevelNameOverride);
 }
 
-ULevelStreamingDynamic* ULevelStreamingDynamic::LoadLevelInstanceBySoftObjectPtr(UObject* WorldContextObject, const TSoftObjectPtr<UWorld> Level, const FVector Location, const FRotator Rotation, bool& bOutSuccess, const FName& OptionalLevelNameOverride)
+ULevelStreamingDynamic* ULevelStreamingDynamic::LoadLevelInstanceBySoftObjectPtr(UObject* WorldContextObject, const TSoftObjectPtr<UWorld> Level, const FVector Location, const FRotator Rotation, bool& bOutSuccess, FString OptionalLevelNameOverride)
 {
 	bOutSuccess = false;
 	UWorld* World = GEngine->GetWorldFromContextObject(WorldContextObject, EGetWorldErrorMode::LogAndReturnNull);
@@ -1789,7 +1789,7 @@ ULevelStreamingDynamic* ULevelStreamingDynamic::LoadLevelInstanceBySoftObjectPtr
 	return LoadLevelInstance_Internal(World, Level.GetLongPackageName(), Location, Rotation, bOutSuccess, OptionalLevelNameOverride);
 }
 
-ULevelStreamingDynamic* ULevelStreamingDynamic::LoadLevelInstance_Internal(UWorld* World, const FString& LongPackageName, const FVector Location, const FRotator Rotation, bool& bOutSuccess, const FName& OptionalLevelNameOverride)
+ULevelStreamingDynamic* ULevelStreamingDynamic::LoadLevelInstance_Internal(UWorld* World, const FString& LongPackageName, const FVector Location, const FRotator Rotation, bool& bOutSuccess, const FString& OptionalLevelNameOverride)
 {
 	const FString PackagePath = FPackageName::GetLongPackagePath(LongPackageName);
 	FString ShortPackageName = FPackageName::GetShortName(LongPackageName);
@@ -1805,7 +1805,7 @@ ULevelStreamingDynamic* ULevelStreamingDynamic::LoadLevelInstance_Internal(UWorl
 	// Determine loaded package name
 	FString LevelPackageNameStr = PackagePath + TEXT("/") + World->StreamingLevelsPrefix;
 	bool bNeedsUniqueTest = false;
-	if (OptionalLevelNameOverride.IsNone())
+	if (OptionalLevelNameOverride.IsEmpty())
 	{
 		// Create Unique Name for sub-level package
 		LevelPackageNameStr += ShortPackageName + TEXT("_LevelInstance_") + FString::FromInt(++UniqueLevelInstanceId);
@@ -1813,7 +1813,7 @@ ULevelStreamingDynamic* ULevelStreamingDynamic::LoadLevelInstance_Internal(UWorl
 	else
 	{
 		// Use the supplied suffix, which is expected to result in a unique package name but we have to check if it is not.
-		LevelPackageNameStr += OptionalLevelNameOverride.ToString();
+		LevelPackageNameStr += OptionalLevelNameOverride;
 		bNeedsUniqueTest = true;
 	}
 
