#UE4 Pass string as const ref. This parameter is still optional for blueprints.

diff --git a/Engine/Source/Runtime/Engine/Classes/Engine/LevelStreamingDynamic.h b/Engine/Source/Runtime/Engine/Classes/Engine/LevelStreamingDynamic.h
--- a/Engine/Source/Runtime/Engine/Classes/Engine/LevelStreamingDynamic.h
+++ b/Engine/Source/Runtime/Engine/Classes/Engine/LevelStreamingDynamic.h
@@ -43,10 +43,10 @@ class ENGINE_API ULevelStreamingDynamic : public ULevelStreaming
  	* @return Streaming level object for a level instance
  	*/ 
  	UFUNCTION(BlueprintCallable, Category = LevelStreaming, meta=(DisplayName = "Load Level Instance (by Name)", WorldContext="WorldContextObject"))
- 	static ULevelStreamingDynamic* LoadLevelInstance(UObject* WorldContextObject, FString LevelName, FVector Location, FRotator Rotation, bool& bOutSuccess, FString OptionalLevelNameOverride = TEXT(""));
+ 	static ULevelStreamingDynamic* LoadLevelInstance(UObject* WorldContextObject, FString LevelName, FVector Location, FRotator Rotation, bool& bOutSuccess, const FString& OptionalLevelNameOverride = TEXT(""));
 
  	UFUNCTION(BlueprintCallable, Category = LevelStreaming, meta=(DisplayName = "Load Level Instance (by Object Reference)", WorldContext="WorldContextObject"))
- 	static ULevelStreamingDynamic* LoadLevelInstanceBySoftObjectPtr(UObject* WorldContextObject, TSoftObjectPtr<UWorld> Level, FVector Location, FRotator Rotation, bool& bOutSuccess, FString OptionalLevelNameOverride = TEXT(""));
+ 	static ULevelStreamingDynamic* LoadLevelInstanceBySoftObjectPtr(UObject* WorldContextObject, TSoftObjectPtr<UWorld> Level, FVector Location, FRotator Rotation, bool& bOutSuccess, const FString& OptionalLevelNameOverride = TEXT(""));
  	
 	//~ Begin UObject Interface
 	virtual void PostLoad() override;
diff --git a/Engine/Source/Runtime/Engine/Private/LevelStreaming.cpp b/Engine/Source/Runtime/Engine/Private/LevelStreaming.cpp
--- a/Engine/Source/Runtime/Engine/Private/LevelStreaming.cpp
+++ b/Engine/Source/Runtime/Engine/Private/LevelStreaming.cpp
@@ -1751,7 +1751,7 @@ void ULevelStreamingDynamic::SetShouldBeLoaded(const bool bInShouldBeLoaded)
 	}
 }
 
-ULevelStreamingDynamic* ULevelStreamingDynamic::LoadLevelInstance(UObject* WorldContextObject, const FString LevelName, const FVector Location, const FRotator Rotation, bool& bOutSuccess, FString OptionalLevelNameOverride)
+ULevelStreamingDynamic* ULevelStreamingDynamic::LoadLevelInstance(UObject* WorldContextObject, const FString LevelName, const FVector Location, const FRotator Rotation, bool& bOutSuccess, const FString& OptionalLevelNameOverride)
 {
 	bOutSuccess = false;
 	UWorld* const World = GEngine->GetWorldFromContextObject(WorldContextObject, EGetWorldErrorMode::LogAndReturnNull);
@@ -1771,7 +1771,7 @@ ULevelStreamingDynamic* ULevelStreamingDynamic::LoadLevelInstance(UObject* World
 	return LoadLevelInstance_Internal(World, LongPackageName, Location, Rotation, bOutSuccess, OptionalLevelNameOverride);
 }
 
-ULevelStreamingDynamic* ULevelStreamingDynamic::LoadLevelInstanceBySoftObjectPtr(UObject* WorldContextObject, const TSoftObjectPtr<UWorld> Level, const FVector Location, const FRotator Rotation, bool& bOutSuccess, FString OptionalLevelNameOverride)
+ULevelStreamingDynamic* ULevelStreamingDynamic::LoadLevelInstanceBySoftObjectPtr(UObject* WorldContextObject, const TSoftObjectPtr<UWorld> Level, const FVector Location, const FRotator Rotation, bool& bOutSuccess, const FString& OptionalLevelNameOverride)
 {
 	bOutSuccess = false;
 	UWorld* World = GEngine->GetWorldFromContextObject(WorldContextObject, EGetWorldErrorMode::LogAndReturnNull);
