Change backwards compatibile class net cache export serialization to include asset path so they will be properly unique
Bump network version for this change

diff --git a/Engine/Source/Runtime/Core/Public/Misc/NetworkVersion.h b/Engine/Source/Runtime/Core/Public/Misc/NetworkVersion.h
--- a/Engine/Source/Runtime/Core/Public/Misc/NetworkVersion.h
+++ b/Engine/Source/Runtime/Core/Public/Misc/NetworkVersion.h
@@ -46,6 +46,7 @@ enum EEngineNetworkVersionHistory
 	HISTORY_FIX_ENUM_SERIALIZATION = 12,			// Bump version to fix enum net serialization issues.
 	HISTORY_OPTIONALLY_QUANTIZE_SPAWN_INFO = 13,	// Bump version to conditionally disable quantization for Scale, Location, and Velocity when spawning network actors.
 	HISTORY_JITTER_IN_HEADER = 14,					// Bump version since we added jitter clock time to packet headers and removed remote saturation
+	HISTORY_CLASSNETCACHE_FULLNAME = 15,			// Bump version to use full paths in GetNetFieldExportGroupForClassNetCache
 	// New history items go above here.
 
 	HISTORY_ENGINENETVERSION_PLUS_ONE,
diff --git a/Engine/Source/Runtime/Engine/Classes/Engine/ActorChannel.h b/Engine/Source/Runtime/Engine/Classes/Engine/ActorChannel.h
--- a/Engine/Source/Runtime/Engine/Classes/Engine/ActorChannel.h
+++ b/Engine/Source/Runtime/Engine/Classes/Engine/ActorChannel.h
@@ -331,4 +331,6 @@ private:
 	// TODO: It would be nice to merge the tracking of these with PendingGuidResolves, to not duplicate memory,
 	// especially since both of these sets should be empty most of the time for most channels.
 	TSet<TSharedRef<struct FQueuedBunchObjectReference>> QueuedBunchObjectReferences;
+
+	static const FString ClassNetCacheSuffix;
 };
diff --git a/Engine/Source/Runtime/Engine/Private/DataChannel.cpp b/Engine/Source/Runtime/Engine/Private/DataChannel.cpp
--- a/Engine/Source/Runtime/Engine/Private/DataChannel.cpp
+++ b/Engine/Source/Runtime/Engine/Private/DataChannel.cpp
@@ -23,6 +23,7 @@
 #include "ProfilingDebugging/CsvProfiler.h"
 #include "Net/NetworkGranularMemoryLogging.h"
 #include "Net/Core/Trace/NetTrace.h"
+#include "Misc/NetworkVersion.h"
 
 DEFINE_LOG_CATEGORY(LogNet);
 DEFINE_LOG_CATEGORY(LogRep);
@@ -841,6 +842,7 @@ void UChannel::AppendMustBeMappedGuids( FOutBunch* Bunch )
 }
 PRAGMA_ENABLE_DEPRECATION_WARNINGS
 
+const FString UActorChannel::ClassNetCacheSuffix = TEXT("_ClassNetCache");
 
 UActorChannel::UActorChannel(const FObjectInitializer& ObjectInitializer)
 	: UChannel(ObjectInitializer)
@@ -3813,11 +3815,6 @@ bool UActorChannel::ReadFieldHeaderAndPayload( UObject* Object, const FClassNetC
 	return true;		// More to read
 }
 
-static FORCEINLINE FString GenerateClassNetCacheNetFieldExportGroupName( const UClass* ObjectClass )
-{
-	return ObjectClass->GetName() + FString( TEXT( "_ClassNetCache" ) );
-}
-
 FNetFieldExportGroup* UActorChannel::GetOrCreateNetFieldExportGroupForClassNetCache( const UObject* Object )
 {
 	if ( !Connection->IsInternalAck() )
@@ -3834,7 +3831,9 @@ FNetFieldExportGroup* UActorChannel::GetOrCreateNetFieldExportGroupForClassNetCa
 
 	UPackageMapClient* PackageMapClient = ( ( UPackageMapClient* )Connection->PackageMap );
 
-	const FString NetFieldExportGroupName = GenerateClassNetCacheNetFieldExportGroupName( ObjectClass );
+	FString NetFieldExportGroupName = ObjectClass->GetPathName();
+	GEngine->NetworkRemapPath(Connection->Driver, NetFieldExportGroupName, false);
+	NetFieldExportGroupName += ClassNetCacheSuffix;
 
 	TSharedPtr< FNetFieldExportGroup > NetFieldExportGroup = PackageMapClient->GetNetFieldExportGroup( NetFieldExportGroupName );
 
@@ -3875,18 +3874,29 @@ FNetFieldExportGroup* UActorChannel::GetOrCreateNetFieldExportGroupForClassNetCa
 	return NetFieldExportGroup.Get();
 }
 
-FNetFieldExportGroup* UActorChannel::GetNetFieldExportGroupForClassNetCache( UClass* ObjectClass )
+FNetFieldExportGroup* UActorChannel::GetNetFieldExportGroupForClassNetCache(UClass* ObjectClass)
 {
-	if ( !Connection->IsInternalAck() )
+	if (!Connection->IsInternalAck())
 	{
 		return nullptr;
 	}	
 
-	const FString NetFieldExportGroupName = GenerateClassNetCacheNetFieldExportGroupName( ObjectClass );
+	FString NetFieldExportGroupName;
+	
+	if (Connection->EngineNetworkProtocolVersion < HISTORY_CLASSNETCACHE_FULLNAME)
+	{
+		NetFieldExportGroupName = ObjectClass->GetName() + ClassNetCacheSuffix;
+	}
+	else
+	{
+		NetFieldExportGroupName = ObjectClass->GetPathName();
+		GEngine->NetworkRemapPath(Connection->Driver, NetFieldExportGroupName, true);
+		NetFieldExportGroupName += ClassNetCacheSuffix;
+	}
 
 	UPackageMapClient* PackageMapClient = CastChecked<UPackageMapClient>(Connection->PackageMap);
 
-	TSharedPtr< FNetFieldExportGroup > NetFieldExportGroup = PackageMapClient->GetNetFieldExportGroup( NetFieldExportGroupName );
+	TSharedPtr<FNetFieldExportGroup> NetFieldExportGroup = PackageMapClient->GetNetFieldExportGroup(NetFieldExportGroupName);
 
 	return NetFieldExportGroup.Get();
 }
