Fix for UE-90683: You can no longer delete conflicting variables

diff --git a/Engine/Source/Runtime/CoreUObject/Public/UObject/Field.h b/Engine/Source/Runtime/CoreUObject/Public/UObject/Field.h
--- a/Engine/Source/Runtime/CoreUObject/Public/UObject/Field.h
+++ b/Engine/Source/Runtime/CoreUObject/Public/UObject/Field.h
@@ -242,7 +242,7 @@ public:
 		return !!Container.Object;
 	}
 	bool IsValidLowLevel() const;
-	operator bool() const
+	inline operator bool() const
 	{
 		return IsValid();
 	}
@@ -917,13 +917,13 @@ COREUOBJECT_API FString GetPathNameSafe(const FField* InField);
 /** 
  * Finds a field given a path to the field (Package.Class[:Subobject:...]:FieldName)
  */
-COREUOBJECT_API FField* FindFieldByPath(const TCHAR* InFieldPath);
+COREUOBJECT_API FField* FindFPropertyByPath(const TCHAR* InFieldPath);
 /**
  * Templated version of FindFieldByPath
  */
 template <typename FieldType>
-inline FieldType* FindField(const TCHAR* InFieldPath)
+inline FieldType* FindFProperty(const TCHAR* InFieldPath)
 {
-	FField* FoundField = FindFieldByPath(InFieldPath);
+	FField* FoundField = FindFPropertyByPath(InFieldPath);
 	return CastField<FieldType>(FoundField);
 }
diff --git a/Engine/Source/Runtime/CoreUObject/Public/UObject/UnrealType.h b/Engine/Source/Runtime/CoreUObject/Public/UObject/UnrealType.h
--- a/Engine/Source/Runtime/CoreUObject/Public/UObject/UnrealType.h
+++ b/Engine/Source/Runtime/CoreUObject/Public/UObject/UnrealType.h
@@ -5538,7 +5538,9 @@ struct TFieldRange
 //
 // Find a typed field in a struct.
 //
-template <class T> T* FindField( const UStruct* Owner, FName FieldName )
+template <class T>
+UE_DEPRECATED(4.25, "FindField will no longer return properties. Use FindProperty instead or FindUField if you want to find functions or enums.")
+ T* FindField( const UStruct* Owner, FName FieldName )
 {
 	// We know that a "none" field won't exist in this Struct
 	if( FieldName.IsNone() )
@@ -5559,21 +5561,91 @@ template <class T> T* FindField( const UStruct* Owner, FName FieldName )
 	return nullptr;
 }
 
-template <class T> T* FindField( const UStruct* Owner, const TCHAR* FieldName )
+template <class T>
+UE_DEPRECATED(4.25, "FindField will no longer return properties. Use FindFProperty instead or FindUField if you want to find UFunctions or UEnums.")
+T* FindField( const UStruct* Owner, const TCHAR* FieldName )
 {
 	// lookup the string name in the Name hash
 	FName Name(FieldName, FNAME_Find);
 	return FindField<T>(Owner, Name);
 }
 
+template <class T> 
+typename TEnableIf<TIsDerivedFrom<T, UField>::IsDerived, T*>::Type FindUField(const UStruct* Owner, FName FieldName)
+{
+	static_assert(sizeof(T) > 0, "T must not be an incomplete type");
+
+	// We know that a "none" field won't exist in this Struct
+	if (FieldName.IsNone())
+	{
+		return nullptr;
+	}
+
+	// Search by comparing FNames (INTs), not strings
+	for (TFieldIterator<T>It(Owner); It; ++It)
+	{
+		if (It->GetFName() == FieldName)
+		{
+			return *It;
+		}
+	}
+
+	// If we didn't find it, return no field
+	return nullptr;
+}
+
+template <class T> 
+typename TEnableIf<TIsDerivedFrom<T, UField>::IsDerived, T*>::Type FindUField(const UStruct* Owner, const TCHAR* FieldName)
+{
+	static_assert(sizeof(T) > 0, "T must not be an incomplete type");
+
+	// lookup the string name in the Name hash
+	FName Name(FieldName, FNAME_Find);
+	return FindUField<T>(Owner, Name);
+}
+
+template <class T>
+typename TEnableIf<TIsDerivedFrom<T, FField>::IsDerived, T*>::Type FindFProperty(const UStruct* Owner, FName FieldName)
+{
+	static_assert(sizeof(T) > 0, "T must not be an incomplete type");
+
+	// We know that a "none" field won't exist in this Struct
+	if (FieldName.IsNone())
+	{
+		return nullptr;
+	}
+
+	// Search by comparing FNames (INTs), not strings
+	for (TFieldIterator<T>It(Owner); It; ++It)
+	{
+		if (It->GetFName() == FieldName)
+		{
+			return *It;
+		}
+	}
+
+	// If we didn't find it, return no field
+	return nullptr;
+}
+
+template <class T>
+typename TEnableIf<TIsDerivedFrom<T, FField>::IsDerived, T*>::Type FindFProperty(const UStruct* Owner, const TCHAR* FieldName)
+{
+	static_assert(sizeof(T) > 0, "T must not be an incomplete type");
+
+	// lookup the string name in the Name hash
+	FName Name(FieldName, FNAME_Find);
+	return FindFProperty<T>(Owner, Name);
+}
+
 /** Finds FProperties or UFunctions and UEnums */
 inline FFieldVariant FindUFieldOrFProperty(const UStruct* Owner, FName FieldName)
 {
 	// Look for properties first as they're most often the runtime thing higher level code wants to find
-	FFieldVariant Result = FindField<FProperty>(Owner, FieldName);
-	if (!Result.IsValid())
+	FFieldVariant Result = FindFProperty<FProperty>(Owner, FieldName);
+	if (!Result)
 	{
-		Result = FindField<UField>(Owner, FieldName);
+		Result = FindUField<UField>(Owner, FieldName);
 	}
 	return Result;
 }
diff --git a/Engine/Source/Runtime/Engine/Classes/Engine/MemberReference.h b/Engine/Source/Runtime/Engine/Classes/Engine/MemberReference.h
--- a/Engine/Source/Runtime/Engine/Classes/Engine/MemberReference.h
+++ b/Engine/Source/Runtime/Engine/Classes/Engine/MemberReference.h
@@ -311,7 +311,7 @@ public:
 	/** Get the scope of this member */
 	UStruct* GetMemberScope(UClass* InMemberParentClass) const
 	{
-		return FindField<UStruct>(InMemberParentClass, *MemberScope);
+		return FindUField<UStruct>(InMemberParentClass, *MemberScope);
 	}
 
 	/** Get the name of the scope of this member */
@@ -369,10 +369,10 @@ public:
 		// Check if the member reference is function scoped
 		if(IsLocalScope())
 		{
-			UStruct* MemberScopeStruct = FindField<UStruct>(SelfScope, *MemberScope);
+			UStruct* MemberScopeStruct = FindUField<UStruct>(SelfScope, *MemberScope);
 
 			// Find in target scope
-			ReturnField = FindField<TFieldType>(MemberScopeStruct, MemberName);
+			ReturnField = FindUFieldOrFProperty(MemberScopeStruct, MemberName).Get<TFieldType>();
 
 #if WITH_EDITOR
 			if(ReturnField == nullptr)
@@ -381,7 +381,7 @@ public:
 				const FName RenamedMemberName = RefreshLocalVariableName(SelfScope);
 				if (RenamedMemberName != NAME_None)
 				{
-					ReturnField = FindField<TFieldType>(MemberScopeStruct, MemberName);
+					ReturnField = FindUFieldOrFProperty(MemberScopeStruct, MemberName).Get<TFieldType>();
 				}
 			}
 #endif
@@ -429,11 +429,11 @@ public:
 				UScriptStruct* SparseClassDataStruct = TargetScope->GetSparseClassDataStruct();
 				if (SparseClassDataStruct)
 				{
-					ReturnField = FindField<TFieldType>(SparseClassDataStruct, MemberName);
+					ReturnField = FindUFieldOrFProperty(SparseClassDataStruct, MemberName).Get<TFieldType>();
 				}
 				if (ReturnField == nullptr)
 				{
-					ReturnField = FindField<TFieldType>(TargetScope, MemberName);
+					ReturnField = FindUFieldOrFProperty(TargetScope, MemberName).Get<TFieldType>();
 				}
 
 #if WITH_EDITOR
@@ -451,7 +451,7 @@ public:
 					if (RenamedMemberName != NAME_None)
 					{
 						MemberName = RenamedMemberName;
-						ReturnField = FindField<TFieldType>(TargetScope, MemberName);
+						ReturnField = FindUFieldOrFProperty(TargetScope, MemberName).Get<TFieldType>();
 					}
 				}
 #endif
@@ -466,7 +466,7 @@ public:
 			// FindObject<>; however this was not reliable (hence the addition 
 			// of GetMemberParentPackage(), etc.)
 			else if (MemberName.ToString().EndsWith(HEADER_GENERATED_DELEGATE_SIGNATURE_SUFFIX))
-					{
+			{
 				ReturnField = ResolveUFunction<TFieldType>();
 				if (ReturnField != nullptr)
 				{
diff --git a/Engine/Source/Runtime/Engine/Classes/Materials/Material.h b/Engine/Source/Runtime/Engine/Classes/Materials/Material.h
--- a/Engine/Source/Runtime/Engine/Classes/Materials/Material.h
+++ b/Engine/Source/Runtime/Engine/Classes/Materials/Material.h
@@ -1839,7 +1839,7 @@ private:
 		{
 			if (Parameter && Parameter->SetParameterValue(InParameterName, Args...))
 			{
-				if (FProperty* ParamProperty = FindField<FProperty>(ParameterType::StaticClass(), "DefaultValue"))
+				if (FProperty* ParamProperty = FindFProperty<FProperty>(ParameterType::StaticClass(), "DefaultValue"))
 				{
 					FPropertyChangedEvent PropertyChangedEvent(ParamProperty);
 					Parameter->PostEditChangeProperty(PropertyChangedEvent);
diff --git a/Engine/Source/Runtime/Engine/Private/AI/Navigation/NavigationModifier.cpp b/Engine/Source/Runtime/Engine/Private/AI/Navigation/NavigationModifier.cpp
--- a/Engine/Source/Runtime/Engine/Private/AI/Navigation/NavigationModifier.cpp
+++ b/Engine/Source/Runtime/Engine/Private/AI/Navigation/NavigationModifier.cpp
@@ -110,7 +110,7 @@ void FNavigationLinkBase::DescribeCustomFlags(const TArray<FString>& EditableFla
 						FString PropName(TEXT("bCustomFlag"));
 						PropName += TTypeToString<int32>::ToString(Idx);
 
-						FProperty* FlagProp = FindField<FProperty>(StructIt, *PropName);
+						FProperty* FlagProp = FindFProperty<FProperty>(StructIt, *PropName);
 						if (FlagProp)
 						{
 							if (Idx < MaxFlags)
diff --git a/Engine/Source/Runtime/Engine/Private/ActorEditor.cpp b/Engine/Source/Runtime/Engine/Private/ActorEditor.cpp
--- a/Engine/Source/Runtime/Engine/Private/ActorEditor.cpp
+++ b/Engine/Source/Runtime/Engine/Private/ActorEditor.cpp
@@ -876,7 +876,7 @@ void AActor::SetActorLabelInternal(const FString& NewActorLabelDirty, bool bMake
 		}
 	}
 
-	FPropertyChangedEvent PropertyEvent( FindField<FProperty>( AActor::StaticClass(), "ActorLabel" ) );
+	FPropertyChangedEvent PropertyEvent( FindFProperty<FProperty>( AActor::StaticClass(), "ActorLabel" ) );
 	PostEditChangeProperty(PropertyEvent);
 
 	FCoreDelegates::OnActorLabelChanged.Broadcast(this);
diff --git a/Engine/Source/Runtime/Engine/Private/Animation/AnimNode_CustomProperty.cpp b/Engine/Source/Runtime/Engine/Private/Animation/AnimNode_CustomProperty.cpp
--- a/Engine/Source/Runtime/Engine/Private/Animation/AnimNode_CustomProperty.cpp
+++ b/Engine/Source/Runtime/Engine/Private/Animation/AnimNode_CustomProperty.cpp
@@ -76,8 +76,8 @@ void FAnimNode_CustomProperty::InitializeProperties(const UObject* InSourceInsta
 			const FName& SourceName = SourcePropertyNames[Idx];
 			const FName& DestName = DestPropertyNames[Idx];
 
-			FProperty* SourceProperty = FindField<FProperty>(SourceClass, SourceName);
-			FProperty* DestProperty = FindField<FProperty>(InTargetClass, DestName);
+			FProperty* SourceProperty = FindFProperty<FProperty>(SourceClass, SourceName);
+			FProperty* DestProperty = FindFProperty<FProperty>(InTargetClass, DestName);
 
 			if (SourceProperty && DestProperty
 #if WITH_EDITOR
diff --git a/Engine/Source/Runtime/Engine/Private/Animation/AnimSequenceBase.cpp b/Engine/Source/Runtime/Engine/Private/Animation/AnimSequenceBase.cpp
--- a/Engine/Source/Runtime/Engine/Private/Animation/AnimSequenceBase.cpp
+++ b/Engine/Source/Runtime/Engine/Private/Animation/AnimSequenceBase.cpp
@@ -641,7 +641,7 @@ uint8* UAnimSequenceBase::FindNotifyPropertyData(int32 NotifyIndex, FArrayProper
 uint8* UAnimSequenceBase::FindArrayProperty(const TCHAR* PropName, FArrayProperty*& ArrayProperty, int32 ArrayIndex)
 {
 	// find Notifies property start point
-	FProperty* Property = FindField<FProperty>(GetClass(), PropName);
+	FProperty* Property = FindFProperty<FProperty>(GetClass(), PropName);
 
 	// found it and if it is array
 	if (Property && Property->IsA(FArrayProperty::StaticClass()))
diff --git a/Engine/Source/Runtime/Engine/Private/Blueprint.cpp b/Engine/Source/Runtime/Engine/Private/Blueprint.cpp
--- a/Engine/Source/Runtime/Engine/Private/Blueprint.cpp
+++ b/Engine/Source/Runtime/Engine/Private/Blueprint.cpp
@@ -1618,7 +1618,7 @@ ETimelineSigType UBlueprint::GetTimelineSignatureForFunctionByName(const FName&
 	// If an object property was specified, find the class of that property instead
 	if(ObjectPropertyName != NAME_None)
 	{
-		FObjectPropertyBase* ObjProperty = FindField<FObjectPropertyBase>(SkeletonGeneratedClass, ObjectPropertyName);
+		FObjectPropertyBase* ObjProperty = FindFProperty<FObjectPropertyBase>(SkeletonGeneratedClass, ObjectPropertyName);
 		if(ObjProperty == NULL)
 		{
 			UE_LOG(LogBlueprint, Log, TEXT("GetTimelineSignatureForFunction: Object Property '%s' not found."), *ObjectPropertyName.ToString());
@@ -1628,7 +1628,7 @@ ETimelineSigType UBlueprint::GetTimelineSignatureForFunctionByName(const FName&
 		UseClass = ObjProperty->PropertyClass;
 	}
 
-	UFunction* Function = FindField<UFunction>(UseClass, FunctionName);
+	UFunction* Function = FindUField<UFunction>(UseClass, FunctionName);
 	if(Function == NULL)
 	{
 		UE_LOG(LogBlueprint, Log, TEXT("GetTimelineSignatureForFunction: Function '%s' not found in class '%s'."), *FunctionName.ToString(), *UseClass->GetName());
diff --git a/Engine/Source/Runtime/Engine/Private/BlueprintGeneratedClass.cpp b/Engine/Source/Runtime/Engine/Private/BlueprintGeneratedClass.cpp
--- a/Engine/Source/Runtime/Engine/Private/BlueprintGeneratedClass.cpp
+++ b/Engine/Source/Runtime/Engine/Private/BlueprintGeneratedClass.cpp
@@ -1063,7 +1063,7 @@ void UBlueprintGeneratedClass::CreateTimelineComponent(AActor* Actor, const UTim
 
 	// Find property with the same name as the template and assign the new Timeline to it
 	UClass* ActorClass = Actor->GetClass();
-	FObjectPropertyBase* Prop = FindField<FObjectPropertyBase>(ActorClass, TimelineTemplate->GetVariableName());
+	FObjectPropertyBase* Prop = FindFProperty<FObjectPropertyBase>(ActorClass, TimelineTemplate->GetVariableName());
 	if (Prop)
 	{
 		Prop->SetObjectPropertyValue_InContainer(Actor, NewTimeline);
@@ -1740,7 +1740,7 @@ void FBlueprintCookedComponentInstancingData::BuildCachedPropertyList(FCustomPro
 		const UStruct* PropertyScope = CurrentScope;
 		while (!Property && PropertyScope)
 		{
-			Property = FindField<FProperty>(PropertyScope, ChangedPropertyInfo.PropertyName);
+			Property = FindFProperty<FProperty>(PropertyScope, ChangedPropertyInfo.PropertyName);
 			PropertyScope = PropertyScope->GetSuperStruct();
 		}
 
diff --git a/Engine/Source/Runtime/Engine/Private/BlueprintMapLibrary.cpp b/Engine/Source/Runtime/Engine/Private/BlueprintMapLibrary.cpp
--- a/Engine/Source/Runtime/Engine/Private/BlueprintMapLibrary.cpp
+++ b/Engine/Source/Runtime/Engine/Private/BlueprintMapLibrary.cpp
@@ -112,7 +112,7 @@ void UBlueprintMapLibrary::GenericMap_SetMapPropertyByName(UObject* OwnerObject,
 {
 	if (OwnerObject)
 	{
-		if (FMapProperty* MapProp = FindField<FMapProperty>(OwnerObject->GetClass(), MapPropertyName))
+		if (FMapProperty* MapProp = FindFProperty<FMapProperty>(OwnerObject->GetClass(), MapPropertyName))
 		{
 			void* Dest = MapProp->ContainerPtrToValuePtr<void>(OwnerObject);
 			MapProp->CopyValuesInternal(Dest, SrcMapAddr, 1);
diff --git a/Engine/Source/Runtime/Engine/Private/BlueprintSetLibrary.cpp b/Engine/Source/Runtime/Engine/Private/BlueprintSetLibrary.cpp
--- a/Engine/Source/Runtime/Engine/Private/BlueprintSetLibrary.cpp
+++ b/Engine/Source/Runtime/Engine/Private/BlueprintSetLibrary.cpp
@@ -189,7 +189,7 @@ void UBlueprintSetLibrary::GenericSet_SetSetPropertyByName(UObject* OwnerObject,
 {
 	if (OwnerObject)
 	{
-		if (FSetProperty* SetProp = FindField<FSetProperty>(OwnerObject->GetClass(), SetPropertyName))
+		if (FSetProperty* SetProp = FindFProperty<FSetProperty>(OwnerObject->GetClass(), SetPropertyName))
 		{
 			void* Dest = SetProp->ContainerPtrToValuePtr<void>(OwnerObject);
 			SetProp->CopyValuesInternal(Dest, SrcSetAddr, 1);
diff --git a/Engine/Source/Runtime/Engine/Private/Collision/CollisionProfile.cpp b/Engine/Source/Runtime/Engine/Private/Collision/CollisionProfile.cpp
--- a/Engine/Source/Runtime/Engine/Private/Collision/CollisionProfile.cpp
+++ b/Engine/Source/Runtime/Engine/Private/Collision/CollisionProfile.cpp
@@ -327,7 +327,7 @@ void UCollisionProfile::LoadProfileConfig(bool bForceInit)
 		{
 			// verify if the Struct name matches
 			// this is to avoid situations where they mismatch and causes random bugs
-			FField* Field = FindField<FField>(Struct, DisplayName);
+			FField* Field = FindFProperty<FField>(Struct, DisplayName);
 
 			if (!Field)
 			{
@@ -451,7 +451,7 @@ void UCollisionProfile::LoadProfileConfig(bool bForceInit)
 				}
 #if WITH_EDITOR
 				// now enum is fixed, so find member variable for the field
-				FField* Field = FindField<FField>(Struct, FName(*VariableName));
+				FField* Field = FindFProperty<FField>(Struct, FName(*VariableName));
 				// I verified up in the class, this can't happen
 				check (Field);
 				Field->SetMetaData(*DisplayNameKey, *DisplayValue);
diff --git a/Engine/Source/Runtime/Engine/Private/ComponentDelegateBinding.cpp b/Engine/Source/Runtime/Engine/Private/ComponentDelegateBinding.cpp
--- a/Engine/Source/Runtime/Engine/Private/ComponentDelegateBinding.cpp
+++ b/Engine/Source/Runtime/Engine/Private/ComponentDelegateBinding.cpp
@@ -16,13 +16,13 @@ void UComponentDelegateBinding::BindDynamicDelegates(UObject* InInstance) const
 		const FBlueprintComponentDelegateBinding& Binding = ComponentDelegateBindings[BindIdx];
 
 	// Get the property that points to the component
-		if (const FObjectProperty* ObjProp = FindField<FObjectProperty>(InInstance->GetClass(), Binding.ComponentPropertyName))
+		if (const FObjectProperty* ObjProp = FindFProperty<FObjectProperty>(InInstance->GetClass(), Binding.ComponentPropertyName))
 		{
 			// ..see if there is actually a component assigned
 			if (UObject* Component = ObjProp->GetObjectPropertyValue_InContainer(InInstance))
 			{
 			// If there is, find and return the delegate property on it
-				if (FMulticastDelegateProperty* MulticastDelegateProp = FindField<FMulticastDelegateProperty>(Component->GetClass(), Binding.DelegatePropertyName))
+				if (FMulticastDelegateProperty* MulticastDelegateProp = FindFProperty<FMulticastDelegateProperty>(Component->GetClass(), Binding.DelegatePropertyName))
 				{
 					// Get the function we want to bind
 					if (UFunction* FunctionToBind = InInstance->GetClass()->FindFunctionByName(Binding.FunctionNameToBind))
@@ -45,13 +45,13 @@ void UComponentDelegateBinding::UnbindDynamicDelegates(UObject* InInstance) cons
 		const FBlueprintComponentDelegateBinding& Binding = ComponentDelegateBindings[BindIdx];
 
 		// Get the property that points to the component
-		if (const FObjectProperty* ObjProp = FindField<FObjectProperty>(InInstance->GetClass(), Binding.ComponentPropertyName))
+		if (const FObjectProperty* ObjProp = FindFProperty<FObjectProperty>(InInstance->GetClass(), Binding.ComponentPropertyName))
 		{
 			// ..see if there is actually a component assigned
 			if (UObject* Component = ObjProp->GetObjectPropertyValue_InContainer(InInstance))
 			{
 				// If there is, find and return the delegate property on it
-				if (FMulticastDelegateProperty* MulticastDelegateProp = FindField<FMulticastDelegateProperty>(Component->GetClass(), Binding.DelegatePropertyName))
+				if (FMulticastDelegateProperty* MulticastDelegateProp = FindFProperty<FMulticastDelegateProperty>(Component->GetClass(), Binding.DelegatePropertyName))
 		{
 					// Unbind function on the instance to this delegate
 					FScriptDelegate Delegate;
@@ -70,14 +70,14 @@ void UComponentDelegateBinding::UnbindDynamicDelegatesForProperty(UObject* InIns
 		const FBlueprintComponentDelegateBinding& Binding = ComponentDelegateBindings[BindIdx];
 		if (InObjectProperty->GetFName() == Binding.ComponentPropertyName)
 		{
-			const FObjectProperty* ObjProp = FindField<FObjectProperty>(InInstance->GetClass(), Binding.ComponentPropertyName);
+			const FObjectProperty* ObjProp = FindFProperty<FObjectProperty>(InInstance->GetClass(), Binding.ComponentPropertyName);
 			if (ObjProp == InObjectProperty)
 			{
 				// ..see if there is actually a component assigned
 				if (UObject* Component = ObjProp->GetObjectPropertyValue_InContainer(InInstance))
 				{
 					// If there is, find and return the delegate property on it
-					if (FMulticastDelegateProperty* MulticastDelegateProp = FindField<FMulticastDelegateProperty>(Component->GetClass(), Binding.DelegatePropertyName))
+					if (FMulticastDelegateProperty* MulticastDelegateProp = FindFProperty<FMulticastDelegateProperty>(Component->GetClass(), Binding.DelegatePropertyName))
 					{
 						// Unbind function on the instance from this delegate
 						FScriptDelegate Delegate;
diff --git a/Engine/Source/Runtime/Engine/Private/Components/SplineComponent.cpp b/Engine/Source/Runtime/Engine/Private/Components/SplineComponent.cpp
--- a/Engine/Source/Runtime/Engine/Private/Components/SplineComponent.cpp
+++ b/Engine/Source/Runtime/Engine/Private/Components/SplineComponent.cpp
@@ -1682,7 +1682,7 @@ void USplineComponent::ApplyComponentInstanceData(FSplineInstanceData* SplineIns
 			// If we are restoring the saved state, unmark the SplineCurves property as 'modified'.
 			// We don't want to consider that these changes have been made through the UCS.
 			TArray<FProperty*> Properties;
-			Properties.Emplace(FindField<FProperty>(USplineComponent::StaticClass(), GET_MEMBER_NAME_CHECKED(USplineComponent, SplineCurves)));
+			Properties.Emplace(FindFProperty<FProperty>(USplineComponent::StaticClass(), GET_MEMBER_NAME_CHECKED(USplineComponent, SplineCurves)));
 			RemoveUCSModifiedProperties(Properties);
 		}
 	}
diff --git a/Engine/Source/Runtime/Engine/Private/DataTable.cpp b/Engine/Source/Runtime/Engine/Private/DataTable.cpp
--- a/Engine/Source/Runtime/Engine/Private/DataTable.cpp
+++ b/Engine/Source/Runtime/Engine/Private/DataTable.cpp
@@ -638,7 +638,7 @@ TArray<FProperty*> UDataTable::GetTablePropertyArray(const TArray<const TCHAR*>&
 			}
 			else
 			{
-				FProperty* ColumnProp = FindField<FProperty>(InRowStruct, PropName);
+				FProperty* ColumnProp = FindFProperty<FProperty>(InRowStruct, PropName);
 
 				for (TFieldIterator<FProperty> It(InRowStruct); It && !ColumnProp; ++It)
 				{
@@ -684,7 +684,7 @@ TArray<FProperty*> UDataTable::GetTablePropertyArray(const TArray<const TCHAR*>&
 	// Generate warning for any properties in struct we are not filling in
 		for (int32 PropIdx = 0; PropIdx < ExpectedPropNames.Num(); PropIdx++)
 	{
-		const FProperty* const ColumnProp = FindField<FProperty>(InRowStruct, ExpectedPropNames[PropIdx]);
+		const FProperty* const ColumnProp = FindFProperty<FProperty>(InRowStruct, ExpectedPropNames[PropIdx]);
 
 #if WITH_EDITOR
 		// If the structure has specified the property as optional for import (gameplay code likely doing a custom fix-up or parse of that property),
diff --git a/Engine/Source/Runtime/Engine/Private/EngineTypes.cpp b/Engine/Source/Runtime/Engine/Private/EngineTypes.cpp
--- a/Engine/Source/Runtime/Engine/Private/EngineTypes.cpp
+++ b/Engine/Source/Runtime/Engine/Private/EngineTypes.cpp
@@ -214,7 +214,7 @@ UActorComponent* FComponentReference::GetComponent(AActor* OwningActor) const
 		{
 			if(ComponentProperty != NAME_None)
 			{
-				FObjectPropertyBase* ObjProp = FindField<FObjectPropertyBase>(SearchActor->GetClass(), ComponentProperty);
+				FObjectPropertyBase* ObjProp = FindFProperty<FObjectPropertyBase>(SearchActor->GetClass(), ComponentProperty);
 				if(ObjProp != NULL)
 				{
 					// .. and return the component that is there
diff --git a/Engine/Source/Runtime/Engine/Private/GameViewportClient.cpp b/Engine/Source/Runtime/Engine/Private/GameViewportClient.cpp
--- a/Engine/Source/Runtime/Engine/Private/GameViewportClient.cpp
+++ b/Engine/Source/Runtime/Engine/Private/GameViewportClient.cpp
@@ -3560,7 +3560,7 @@ bool UGameViewportClient::HandleDisplayCommand( const TCHAR* Cmd, FOutputDevice&
 		if (Obj != nullptr)
 		{
 			FName PropertyName(PropStr, FNAME_Find);
-			if (PropertyName != NAME_None && FindField<FProperty>(Obj->GetClass(), PropertyName) != nullptr)
+			if (PropertyName != NAME_None && FindFProperty<FProperty>(Obj->GetClass(), PropertyName) != nullptr)
 			{
 				AddDebugDisplayProperty(Obj, nullptr, PropertyName);
 			}
@@ -3612,7 +3612,7 @@ bool UGameViewportClient::HandleDisplayAllCommand( const TCHAR* Cmd, FOutputDevi
 			if (Cls != nullptr)
 			{
 				FName PropertyName(PropStr, FNAME_Find);
-				FProperty* Prop = PropertyName != NAME_None ? FindField<FProperty>(Cls, PropertyName) : nullptr;
+				FProperty* Prop = PropertyName != NAME_None ? FindFProperty<FProperty>(Cls, PropertyName) : nullptr;
 				{
 					// add all un-GCable things immediately as that list is static
 					// so then we only have to iterate over dynamic things each frame
diff --git a/Engine/Source/Runtime/Engine/Private/InstancedReferenceSubobjectHelper.cpp b/Engine/Source/Runtime/Engine/Private/InstancedReferenceSubobjectHelper.cpp
--- a/Engine/Source/Runtime/Engine/Private/InstancedReferenceSubobjectHelper.cpp
+++ b/Engine/Source/Runtime/Engine/Private/InstancedReferenceSubobjectHelper.cpp
@@ -10,7 +10,7 @@ UObject* FInstancedPropertyPath::Resolve(const UObject* Container) const
 	auto GetProperty = [&CurrentContainerType, &PropChainRef](int32 ChainIndex)->FProperty*
 	{
 		const FProperty* SrcProperty = PropChainRef[ChainIndex].PropertyPtr;
-		return FindField<FProperty>(CurrentContainerType, SrcProperty->GetFName());
+		return FindFProperty<FProperty>(CurrentContainerType, SrcProperty->GetFName());
 	};
 
 	const FProperty* CurrentProp = GetProperty(0);
diff --git a/Engine/Source/Runtime/Engine/Private/KismetArrayLibrary.cpp b/Engine/Source/Runtime/Engine/Private/KismetArrayLibrary.cpp
--- a/Engine/Source/Runtime/Engine/Private/KismetArrayLibrary.cpp
+++ b/Engine/Source/Runtime/Engine/Private/KismetArrayLibrary.cpp
@@ -390,7 +390,7 @@ void UKismetArrayLibrary::GenericArray_SetArrayPropertyByName(UObject* OwnerObje
 {
 	if (OwnerObject != NULL)
 	{
-		FArrayProperty* ArrayProp = FindField<FArrayProperty>(OwnerObject->GetClass(), ArrayPropertyName);
+		FArrayProperty* ArrayProp = FindFProperty<FArrayProperty>(OwnerObject->GetClass(), ArrayPropertyName);
 		if (ArrayProp != NULL)
 		{
 			void* Dest = ArrayProp->ContainerPtrToValuePtr<void>(OwnerObject);
diff --git a/Engine/Source/Runtime/Engine/Private/KismetSystemLibrary.cpp b/Engine/Source/Runtime/Engine/Private/KismetSystemLibrary.cpp
--- a/Engine/Source/Runtime/Engine/Private/KismetSystemLibrary.cpp
+++ b/Engine/Source/Runtime/Engine/Private/KismetSystemLibrary.cpp
@@ -848,7 +848,7 @@ void UKismetSystemLibrary::SetIntPropertyByName(UObject* Object, FName PropertyN
 {
 	if(Object != NULL)
 	{
-		FIntProperty* IntProp = FindField<FIntProperty>(Object->GetClass(), PropertyName);
+		FIntProperty* IntProp = FindFProperty<FIntProperty>(Object->GetClass(), PropertyName);
 		if(IntProp != NULL)
 		{
 			IntProp->SetPropertyValue_InContainer(Object, Value);
@@ -860,7 +860,7 @@ void UKismetSystemLibrary::SetInt64PropertyByName(UObject* Object, FName Propert
 {
 	if (Object != NULL)
 	{
-		FInt64Property* IntProp = FindField<FInt64Property>(Object->GetClass(), PropertyName);
+		FInt64Property* IntProp = FindFProperty<FInt64Property>(Object->GetClass(), PropertyName);
 		if (IntProp != NULL)
 		{
 			IntProp->SetPropertyValue_InContainer(Object, Value);
@@ -872,11 +872,11 @@ void UKismetSystemLibrary::SetBytePropertyByName(UObject* Object, FName Property
 {
 	if(Object != NULL)
 	{
-		if(FByteProperty* ByteProp = FindField<FByteProperty>(Object->GetClass(), PropertyName))
+		if(FByteProperty* ByteProp = FindFProperty<FByteProperty>(Object->GetClass(), PropertyName))
 		{
 			ByteProp->SetPropertyValue_InContainer(Object, Value);
 		}
-		else if(FEnumProperty* EnumProp = FindField<FEnumProperty>(Object->GetClass(), PropertyName))
+		else if(FEnumProperty* EnumProp = FindFProperty<FEnumProperty>(Object->GetClass(), PropertyName))
 		{
 			void* PropAddr = EnumProp->ContainerPtrToValuePtr<void>(Object);
 			FNumericProperty* UnderlyingProp = EnumProp->GetUnderlyingProperty();
@@ -889,7 +889,7 @@ void UKismetSystemLibrary::SetFloatPropertyByName(UObject* Object, FName Propert
 {
 	if(Object != NULL)
 	{
-		FFloatProperty* FloatProp = FindField<FFloatProperty>(Object->GetClass(), PropertyName);
+		FFloatProperty* FloatProp = FindFProperty<FFloatProperty>(Object->GetClass(), PropertyName);
 		if(FloatProp != NULL)
 		{
 			FloatProp->SetPropertyValue_InContainer(Object, Value);
@@ -901,7 +901,7 @@ void UKismetSystemLibrary::SetBoolPropertyByName(UObject* Object, FName Property
 {
 	if(Object != NULL)
 	{
-		FBoolProperty* BoolProp = FindField<FBoolProperty>(Object->GetClass(), PropertyName);
+		FBoolProperty* BoolProp = FindFProperty<FBoolProperty>(Object->GetClass(), PropertyName);
 		if(BoolProp != NULL)
 		{
 			BoolProp->SetPropertyValue_InContainer(Object, Value );
@@ -913,7 +913,7 @@ void UKismetSystemLibrary::SetObjectPropertyByName(UObject* Object, FName Proper
 {
 	if(Object != NULL && Value != NULL)
 	{
-		FObjectPropertyBase* ObjectProp = FindField<FObjectPropertyBase>(Object->GetClass(), PropertyName);
+		FObjectPropertyBase* ObjectProp = FindFProperty<FObjectPropertyBase>(Object->GetClass(), PropertyName);
 		if(ObjectProp != NULL && Value->IsA(ObjectProp->PropertyClass)) // check it's the right type
 		{
 			ObjectProp->SetObjectPropertyValue_InContainer(Object, Value);
@@ -925,7 +925,7 @@ void UKismetSystemLibrary::SetClassPropertyByName(UObject* Object, FName Propert
 {
 	if (Object && *Value)
 	{
-		FClassProperty* ClassProp = FindField<FClassProperty>(Object->GetClass(), PropertyName);
+		FClassProperty* ClassProp = FindFProperty<FClassProperty>(Object->GetClass(), PropertyName);
 		if (ClassProp != NULL && Value->IsChildOf(ClassProp->MetaClass)) // check it's the right type
 		{
 			ClassProp->SetObjectPropertyValue_InContainer(Object, *Value);
@@ -937,7 +937,7 @@ void UKismetSystemLibrary::SetInterfacePropertyByName(UObject* Object, FName Pro
 {
 	if (Object)
 	{
-		FInterfaceProperty* InterfaceProp = FindField<FInterfaceProperty>(Object->GetClass(), PropertyName);
+		FInterfaceProperty* InterfaceProp = FindFProperty<FInterfaceProperty>(Object->GetClass(), PropertyName);
 		if (InterfaceProp != NULL && Value.GetObject()->GetClass()->ImplementsInterface(InterfaceProp->InterfaceClass)) // check it's the right type
 		{
 			InterfaceProp->SetPropertyValue_InContainer(Object, Value);
@@ -949,7 +949,7 @@ void UKismetSystemLibrary::SetStringPropertyByName(UObject* Object, FName Proper
 {
 	if(Object != NULL)
 	{
-		FStrProperty* StringProp = FindField<FStrProperty>(Object->GetClass(), PropertyName);
+		FStrProperty* StringProp = FindFProperty<FStrProperty>(Object->GetClass(), PropertyName);
 		if(StringProp != NULL)
 		{
 			StringProp->SetPropertyValue_InContainer(Object, Value);
@@ -961,7 +961,7 @@ void UKismetSystemLibrary::SetNamePropertyByName(UObject* Object, FName Property
 {
 	if(Object != NULL)
 	{
-		FNameProperty* NameProp = FindField<FNameProperty>(Object->GetClass(), PropertyName);
+		FNameProperty* NameProp = FindFProperty<FNameProperty>(Object->GetClass(), PropertyName);
 		if(NameProp != NULL)
 		{
 			NameProp->SetPropertyValue_InContainer(Object, Value);
@@ -973,7 +973,7 @@ void UKismetSystemLibrary::SetSoftObjectPropertyByName(UObject* Object, FName Pr
 {
 	if (Object != NULL)
 	{
-		FSoftObjectProperty* ObjectProp = FindField<FSoftObjectProperty>(Object->GetClass(), PropertyName);
+		FSoftObjectProperty* ObjectProp = FindFProperty<FSoftObjectProperty>(Object->GetClass(), PropertyName);
 		const FSoftObjectPtr* SoftObjectPtr = (const FSoftObjectPtr*)(&Value);
 		ObjectProp->SetPropertyValue_InContainer(Object, *SoftObjectPtr);
 	}
@@ -983,7 +983,7 @@ void UKismetSystemLibrary::SetFieldPathPropertyByName(UObject* Object, FName Pro
 {
 	if (Object != NULL)
 	{
-		FFieldPathProperty* FieldProp = FindField<FFieldPathProperty>(Object->GetClass(), PropertyName);
+		FFieldPathProperty* FieldProp = FindFProperty<FFieldPathProperty>(Object->GetClass(), PropertyName);
 		const FFieldPath* FieldPathPtr = (const FFieldPath*)(&Value);
 		FieldProp->SetPropertyValue_InContainer(Object, *FieldPathPtr);
 	}
@@ -993,7 +993,7 @@ void UKismetSystemLibrary::SetSoftClassPropertyByName(UObject* Object, FName Pro
 {
 	if (Object != NULL)
 	{
-		FSoftClassProperty* ObjectProp = FindField<FSoftClassProperty>(Object->GetClass(), PropertyName);
+		FSoftClassProperty* ObjectProp = FindFProperty<FSoftClassProperty>(Object->GetClass(), PropertyName);
 		const FSoftObjectPtr* SoftObjectPtr = (const FSoftObjectPtr*)(&Value);
 		ObjectProp->SetPropertyValue_InContainer(Object, *SoftObjectPtr);
 	}
@@ -1119,7 +1119,7 @@ void UKismetSystemLibrary::SetTextPropertyByName(UObject* Object, FName Property
 {
 	if(Object != NULL)
 	{
-		FTextProperty* TextProp = FindField<FTextProperty>(Object->GetClass(), PropertyName);
+		FTextProperty* TextProp = FindFProperty<FTextProperty>(Object->GetClass(), PropertyName);
 		if(TextProp != NULL)
 		{
 			TextProp->SetPropertyValue_InContainer(Object, Value);
@@ -1132,7 +1132,7 @@ void UKismetSystemLibrary::SetVectorPropertyByName(UObject* Object, FName Proper
 	if(Object != NULL)
 	{
 		UScriptStruct* VectorStruct = TBaseStructure<FVector>::Get();
-		FStructProperty* VectorProp = FindField<FStructProperty>(Object->GetClass(), PropertyName);
+		FStructProperty* VectorProp = FindFProperty<FStructProperty>(Object->GetClass(), PropertyName);
 		if(VectorProp != NULL && VectorProp->Struct == VectorStruct)
 		{
 			*VectorProp->ContainerPtrToValuePtr<FVector>(Object) = Value;
@@ -1145,7 +1145,7 @@ void UKismetSystemLibrary::SetRotatorPropertyByName(UObject* Object, FName Prope
 	if(Object != NULL)
 	{
 		UScriptStruct* RotatorStruct = TBaseStructure<FRotator>::Get();
-		FStructProperty* RotatorProp = FindField<FStructProperty>(Object->GetClass(), PropertyName);
+		FStructProperty* RotatorProp = FindFProperty<FStructProperty>(Object->GetClass(), PropertyName);
 		if(RotatorProp != NULL && RotatorProp->Struct == RotatorStruct)
 		{
 			*RotatorProp->ContainerPtrToValuePtr<FRotator>(Object) = Value;
@@ -1158,7 +1158,7 @@ void UKismetSystemLibrary::SetLinearColorPropertyByName(UObject* Object, FName P
 	if(Object != NULL)
 	{
 		UScriptStruct* ColorStruct = TBaseStructure<FLinearColor>::Get();
-		FStructProperty* ColorProp = FindField<FStructProperty>(Object->GetClass(), PropertyName);
+		FStructProperty* ColorProp = FindFProperty<FStructProperty>(Object->GetClass(), PropertyName);
 		if(ColorProp != NULL && ColorProp->Struct == ColorStruct)
 		{
 			*ColorProp->ContainerPtrToValuePtr<FLinearColor>(Object) = Value;
@@ -1171,7 +1171,7 @@ void UKismetSystemLibrary::SetTransformPropertyByName(UObject* Object, FName Pro
 	if(Object != NULL)
 	{
 		UScriptStruct* TransformStruct = TBaseStructure<FTransform>::Get();
-		FStructProperty* TransformProp = FindField<FStructProperty>(Object->GetClass(), PropertyName);
+		FStructProperty* TransformProp = FindFProperty<FStructProperty>(Object->GetClass(), PropertyName);
 		if(TransformProp != NULL && TransformProp->Struct == TransformStruct)
 		{
 			*TransformProp->ContainerPtrToValuePtr<FTransform>(Object) = Value;
@@ -1189,7 +1189,7 @@ void UKismetSystemLibrary::Generic_SetStructurePropertyByName(UObject* OwnerObje
 {
 	if (OwnerObject != NULL)
 	{
-		FStructProperty* StructProp = FindField<FStructProperty>(OwnerObject->GetClass(), StructPropertyName);
+		FStructProperty* StructProp = FindFProperty<FStructProperty>(OwnerObject->GetClass(), StructPropertyName);
 		if (StructProp != NULL)
 		{
 			void* Dest = StructProp->ContainerPtrToValuePtr<void>(OwnerObject);
diff --git a/Engine/Source/Runtime/Engine/Private/Materials/MaterialExpressions.cpp b/Engine/Source/Runtime/Engine/Private/Materials/MaterialExpressions.cpp
--- a/Engine/Source/Runtime/Engine/Private/Materials/MaterialExpressions.cpp
+++ b/Engine/Source/Runtime/Engine/Private/Materials/MaterialExpressions.cpp
@@ -849,7 +849,7 @@ bool UMaterialExpression::CanEditChange(const FProperty* InProperty) const
 		{
 			const FString& OverridingPropertyName = InProperty->GetMetaData(OverridingInputPropertyMetaData);
 
-			FStructProperty* StructProp = FindField<FStructProperty>(GetClass(), *OverridingPropertyName);
+			FStructProperty* StructProp = FindFProperty<FStructProperty>(GetClass(), *OverridingPropertyName);
 			if (ensure(StructProp != nullptr))
 			{
 				static FName RequiredInputMetaData(TEXT("RequiredInput"));
@@ -2495,7 +2495,7 @@ void UMaterialExpressionRuntimeVirtualTextureSampleParameter::SetValueToMatching
 	if (Material->GetRuntimeVirtualTextureParameterValue(FMaterialParameterInfo(OtherExpression->GetParameterName()), Value))
 	{
 		VirtualTexture = Value;
-		FProperty* ParamProperty = FindField<FProperty>(UMaterialExpressionRuntimeVirtualTextureSampleParameter::StaticClass(), GET_MEMBER_NAME_STRING_CHECKED(UMaterialExpressionRuntimeVirtualTextureSampleParameter, VirtualTexture));
+		FProperty* ParamProperty = FindFProperty<FProperty>(UMaterialExpressionRuntimeVirtualTextureSampleParameter::StaticClass(), GET_MEMBER_NAME_STRING_CHECKED(UMaterialExpressionRuntimeVirtualTextureSampleParameter, VirtualTexture));
 		FPropertyChangedEvent PropertyChangedEvent(ParamProperty);
 		PostEditChangeProperty(PropertyChangedEvent);
 	}
@@ -2717,7 +2717,7 @@ void UMaterialExpressionTextureSampleParameter::SetValueToMatchingExpression(UMa
 	UTexture* ExistingValue;
 	Material->GetTextureParameterValue(OtherExpression->GetParameterName(), ExistingValue);
 	Texture = ExistingValue;
-	FProperty* ParamProperty = FindField<FProperty>(UMaterialExpressionTextureSampleParameter::StaticClass(), GET_MEMBER_NAME_STRING_CHECKED(UMaterialExpressionTextureSampleParameter, Texture));
+	FProperty* ParamProperty = FindFProperty<FProperty>(UMaterialExpressionTextureSampleParameter::StaticClass(), GET_MEMBER_NAME_STRING_CHECKED(UMaterialExpressionTextureSampleParameter, Texture));
 	FPropertyChangedEvent PropertyChangedEvent(ParamProperty);
 	PostEditChangeProperty(PropertyChangedEvent);
 }
@@ -4413,16 +4413,16 @@ void UMaterialExpressionStaticComponentMaskParameter::SetValueToMatchingExpressi
 	DefaultG = G;
 	DefaultB = B;
 	DefaultA = A;
-	FProperty* ParamProperty = FindField<FProperty>(UMaterialExpressionStaticComponentMaskParameter::StaticClass(), GET_MEMBER_NAME_STRING_CHECKED(UMaterialExpressionStaticComponentMaskParameter, DefaultR));
+	FProperty* ParamProperty = FindFProperty<FProperty>(UMaterialExpressionStaticComponentMaskParameter::StaticClass(), GET_MEMBER_NAME_STRING_CHECKED(UMaterialExpressionStaticComponentMaskParameter, DefaultR));
 	FPropertyChangedEvent RChangedEvent(ParamProperty);
 	PostEditChangeProperty(RChangedEvent);
-	ParamProperty = FindField<FProperty>(UMaterialExpressionStaticComponentMaskParameter::StaticClass(), GET_MEMBER_NAME_STRING_CHECKED(UMaterialExpressionStaticComponentMaskParameter, DefaultG));
+	ParamProperty = FindFProperty<FProperty>(UMaterialExpressionStaticComponentMaskParameter::StaticClass(), GET_MEMBER_NAME_STRING_CHECKED(UMaterialExpressionStaticComponentMaskParameter, DefaultG));
 	FPropertyChangedEvent GChangedEvent(ParamProperty);
 	PostEditChangeProperty(GChangedEvent);
-	ParamProperty = FindField<FProperty>(UMaterialExpressionStaticComponentMaskParameter::StaticClass(), GET_MEMBER_NAME_STRING_CHECKED(UMaterialExpressionStaticComponentMaskParameter, DefaultB));
+	ParamProperty = FindFProperty<FProperty>(UMaterialExpressionStaticComponentMaskParameter::StaticClass(), GET_MEMBER_NAME_STRING_CHECKED(UMaterialExpressionStaticComponentMaskParameter, DefaultB));
 	FPropertyChangedEvent BChangedEvent(ParamProperty);
 	PostEditChangeProperty(BChangedEvent);
-	ParamProperty = FindField<FProperty>(UMaterialExpressionStaticComponentMaskParameter::StaticClass(), GET_MEMBER_NAME_STRING_CHECKED(UMaterialExpressionStaticComponentMaskParameter, DefaultA));
+	ParamProperty = FindFProperty<FProperty>(UMaterialExpressionStaticComponentMaskParameter::StaticClass(), GET_MEMBER_NAME_STRING_CHECKED(UMaterialExpressionStaticComponentMaskParameter, DefaultA));
 	FPropertyChangedEvent AChangedEvent(ParamProperty);
 	PostEditChangeProperty(AChangedEvent);
 }
@@ -7359,7 +7359,7 @@ void UMaterialExpressionVectorParameter::SetValueToMatchingExpression(UMaterialE
 	FLinearColor ExistingValue = FLinearColor::Transparent;
 	Material->GetVectorParameterValue(FMaterialParameterInfo(OtherExpression->GetParameterName()), ExistingValue);
 	DefaultValue = ExistingValue;
-	FProperty* ParamProperty = FindField<FProperty>(UMaterialExpressionVectorParameter::StaticClass(), GET_MEMBER_NAME_STRING_CHECKED(UMaterialExpressionVectorParameter, DefaultValue));
+	FProperty* ParamProperty = FindFProperty<FProperty>(UMaterialExpressionVectorParameter::StaticClass(), GET_MEMBER_NAME_STRING_CHECKED(UMaterialExpressionVectorParameter, DefaultValue));
 	FPropertyChangedEvent PropertyChangedEvent(ParamProperty);
 	PostEditChangeProperty(PropertyChangedEvent);
 }
@@ -7660,7 +7660,7 @@ void UMaterialExpressionScalarParameter::SetValueToMatchingExpression(UMaterialE
 	float ExistingValue = 0.0f;
 	Material->GetScalarParameterDefaultValue(FMaterialParameterInfo(OtherExpression->GetParameterName()), ExistingValue);
 	DefaultValue = ExistingValue;
-	FProperty* ParamProperty = FindField<FProperty>(UMaterialExpressionScalarParameter::StaticClass(), GET_MEMBER_NAME_STRING_CHECKED(UMaterialExpressionScalarParameter, DefaultValue));
+	FProperty* ParamProperty = FindFProperty<FProperty>(UMaterialExpressionScalarParameter::StaticClass(), GET_MEMBER_NAME_STRING_CHECKED(UMaterialExpressionScalarParameter, DefaultValue));
 	FPropertyChangedEvent PropertyChangedEvent(ParamProperty);
 	PostEditChangeProperty(PropertyChangedEvent);
 }
@@ -7779,7 +7779,7 @@ void UMaterialExpressionStaticBoolParameter::SetValueToMatchingExpression(UMater
 	FGuid Guid;
 	Material->GetStaticSwitchParameterValue(OtherExpression->GetParameterName(), ExistingValue, Guid);
 	DefaultValue = ExistingValue;
-	FProperty* ParamProperty = FindField<FProperty>(UMaterialExpressionStaticBoolParameter::StaticClass(), GET_MEMBER_NAME_STRING_CHECKED(UMaterialExpressionStaticBoolParameter, DefaultValue));
+	FProperty* ParamProperty = FindFProperty<FProperty>(UMaterialExpressionStaticBoolParameter::StaticClass(), GET_MEMBER_NAME_STRING_CHECKED(UMaterialExpressionStaticBoolParameter, DefaultValue));
 	FPropertyChangedEvent PropertyChangedEvent(ParamProperty);
 	PostEditChangeProperty(PropertyChangedEvent);
 }
@@ -9251,7 +9251,7 @@ FName UMaterialExpressionSceneDepth::GetInputName(int32 InputIndex) const
 	if(InputIndex == 0)
 	{
 		// Display the current InputMode enum's display name.
-		FByteProperty* InputModeProperty = FindField<FByteProperty>( UMaterialExpressionSceneDepth::StaticClass(), "InputMode" );
+		FByteProperty* InputModeProperty = FindFProperty<FByteProperty>( UMaterialExpressionSceneDepth::StaticClass(), "InputMode" );
 		// Can't use GetNameByValue as GetNameStringByValue does name mangling that GetNameByValue does not
 		return *InputModeProperty->Enum->GetNameStringByValue((int64)InputMode.GetValue());
 	}
@@ -9548,7 +9548,7 @@ FName UMaterialExpressionSceneColor::GetInputName(int32 InputIndex) const
 	if(InputIndex == 0)
 	{
 		// Display the current InputMode enum's display name.
-		FByteProperty* InputModeProperty = FindField<FByteProperty>( UMaterialExpressionSceneColor::StaticClass(), "InputMode" );
+		FByteProperty* InputModeProperty = FindFProperty<FByteProperty>( UMaterialExpressionSceneColor::StaticClass(), "InputMode" );
 		// Can't use GetNameByValue as GetNameStringByValue does name mangling that GetNameByValue does not
 		return *InputModeProperty->Enum->GetNameStringByValue((int64)InputMode.GetValue());
 	}
@@ -10274,10 +10274,10 @@ void UMaterialExpressionFontSampleParameter::SetValueToMatchingExpression(UMater
 	Material->GetFontParameterValue(FMaterialParameterInfo(OtherExpression->GetParameterName()), FontValue, FontPage);
 	Font = FontValue;
 	FontTexturePage = FontPage;
-	FProperty* ParamProperty = FindField<FProperty>(UMaterialExpressionFontSampleParameter::StaticClass(), GET_MEMBER_NAME_STRING_CHECKED(UMaterialExpressionFontSampleParameter, Font));
+	FProperty* ParamProperty = FindFProperty<FProperty>(UMaterialExpressionFontSampleParameter::StaticClass(), GET_MEMBER_NAME_STRING_CHECKED(UMaterialExpressionFontSampleParameter, Font));
 	FPropertyChangedEvent PropertyChangedEvent(ParamProperty);
 	PostEditChangeProperty(PropertyChangedEvent);
-	ParamProperty = FindField<FProperty>(UMaterialExpressionFontSampleParameter::StaticClass(), GET_MEMBER_NAME_STRING_CHECKED(UMaterialExpressionFontSampleParameter, FontTexturePage));
+	ParamProperty = FindFProperty<FProperty>(UMaterialExpressionFontSampleParameter::StaticClass(), GET_MEMBER_NAME_STRING_CHECKED(UMaterialExpressionFontSampleParameter, FontTexturePage));
 	FPropertyChangedEvent PageChangedEvent(ParamProperty);
 	PostEditChangeProperty(PageChangedEvent);
 }
diff --git a/Engine/Source/Runtime/Engine/Private/MatineeUtils.cpp b/Engine/Source/Runtime/Engine/Private/MatineeUtils.cpp
--- a/Engine/Source/Runtime/Engine/Private/MatineeUtils.cpp
+++ b/Engine/Source/Runtime/Engine/Private/MatineeUtils.cpp
@@ -342,7 +342,7 @@ namespace FMatineeUtils
 
 			if (InPropertyName.Split(TEXT("."), &CompString, &PropString))
 			{
-				if (FStructProperty* StructProp = FindField<FStructProperty>(InStruct, *CompString))
+				if (FStructProperty* StructProp = FindFProperty<FStructProperty>(InStruct, *CompString))
 				{
 					// The first part of the path was a struct, look inside it
 					void* StructContainer = StructProp->ContainerPtrToValuePtr<void>(InObject);
@@ -375,7 +375,7 @@ namespace FMatineeUtils
 			else
 			{
 				// Look for the property in the current scope
-				if (FProperty* Prop = FindField<FProperty>(InStruct, *InPropertyName))
+				if (FProperty* Prop = FindFProperty<FProperty>(InStruct, *InPropertyName))
 				{
 					OutPropContainer = BasePointer;
 					OutProperty = Prop;
diff --git a/Engine/Source/Runtime/Engine/Private/MemberReference.cpp b/Engine/Source/Runtime/Engine/Private/MemberReference.cpp
--- a/Engine/Source/Runtime/Engine/Private/MemberReference.cpp
+++ b/Engine/Source/Runtime/Engine/Private/MemberReference.cpp
@@ -122,7 +122,7 @@ bool FMemberReference::IsSparseClassData(const UClass* OwningClass) const
 	UScriptStruct* SparseClassDataStruct = OwningClass ? OwningClass->GetSparseClassDataStruct() : nullptr;
 	if (SparseClassDataStruct)
 	{
-		FProperty* VariableProperty = FindField<FProperty>(SparseClassDataStruct, GetMemberName());
+		FProperty* VariableProperty = FindFProperty<FProperty>(SparseClassDataStruct, GetMemberName());
 		bIsSparseClassData = VariableProperty != nullptr;
 	}
 
@@ -282,7 +282,8 @@ TFieldType* FindRemappedFieldImpl(FName FieldClassOutermostName, FName FieldClas
 	FMemberReference::InitFieldRedirectMap();
 
 	// In the case of a bifurcation of a variable (e.g. moved from a parent into certain children), verify that we don't also define the variable in the current scope first
-	if (FindField<TFieldType>(InitialScope, InitialName) != nullptr)
+	FFieldVariant ExistingField = FindUFieldOrFProperty(InitialScope, InitialName);
+	if (ExistingField.Get<TFieldType>())
 	{
 		return nullptr;
 	}
@@ -326,7 +327,7 @@ TFieldType* FindRemappedFieldImpl(FName FieldClassOutermostName, FName FieldClas
 		if (NewFieldName != NAME_None)
 		{
 			// Find the actual field specified by the redirector, so we can return it and update the node that uses it
-			TFieldType* NewField = FindField<TFieldType>(SearchClass, NewFieldName);
+			TFieldType* NewField = FindUFieldOrFProperty(SearchClass, NewFieldName).Get<TFieldType>();
 			if (NewField != nullptr)
 			{
 				if (bInitialScopeMustBeOwnerOfField && !InitialScope->IsChildOf(SearchClass))
diff --git a/Engine/Source/Runtime/Engine/Private/Net/NetPushModelHelpers.cpp b/Engine/Source/Runtime/Engine/Private/Net/NetPushModelHelpers.cpp
--- a/Engine/Source/Runtime/Engine/Private/Net/NetPushModelHelpers.cpp
+++ b/Engine/Source/Runtime/Engine/Private/Net/NetPushModelHelpers.cpp
@@ -16,7 +16,7 @@ void UNetPushModelHelpers::MarkPropertyDirty(UObject* Object, FName PropertyName
 			// If this is too slow to be practical, we might be able build a lookup from ClassName+PropertyName -> RepIndex.
 			// It would be safest to invalidate the lookup on a hot reload / recompile (if possible), blueprint recompile,
 			// and map change.
-			FProperty* Property = FindField<FProperty>(Class, PropertyName);
+			FProperty* Property = FindFProperty<FProperty>(Class, PropertyName);
 			if (Property == nullptr)
 			{
 				UE_LOG(LogNet, Warning, TEXT("UNetPushModelHelpers::MarkPropertyDirty: Unable to find Property %s in Class %s"), *PropertyName.ToString(), *Class->GetPathName());
diff --git a/Engine/Source/Runtime/Engine/Private/NetworkDriver.cpp b/Engine/Source/Runtime/Engine/Private/NetworkDriver.cpp
--- a/Engine/Source/Runtime/Engine/Private/NetworkDriver.cpp
+++ b/Engine/Source/Runtime/Engine/Private/NetworkDriver.cpp
@@ -2798,7 +2798,7 @@ bool UNetDriver::HandleNetDumpServerRPCCommand( const TCHAR* Cmd, FOutputDevice&
 
 bool UNetDriver::HandleNetDumpDormancy(const TCHAR* Cmd, FOutputDevice& Ar)
 {
-	FProperty* Property = FindField<FProperty>(AActor::StaticClass(), TEXT("NetDormancy"));
+	FProperty* Property = FindFProperty<FProperty>(AActor::StaticClass(), TEXT("NetDormancy"));
 	check(Property != nullptr);
 	Ar.Logf(TEXT(""));
 	TArray<AActor*> Actors;
diff --git a/Engine/Source/Runtime/Engine/Private/Particles/ParticleModules.cpp b/Engine/Source/Runtime/Engine/Private/Particles/ParticleModules.cpp
--- a/Engine/Source/Runtime/Engine/Private/Particles/ParticleModules.cpp
+++ b/Engine/Source/Runtime/Engine/Private/Particles/ParticleModules.cpp
@@ -3472,7 +3472,7 @@ void UParticleModuleTypeDataMesh::OnMeshChanged()
 		OuterObj = Emitter->GetOuter();
 	}
 
-	FProperty* MeshProperty = FindField<FProperty>(UParticleModuleTypeDataMesh::StaticClass(), FName(TEXT("Mesh")));
+	FProperty* MeshProperty = FindFProperty<FProperty>(UParticleModuleTypeDataMesh::StaticClass(), FName(TEXT("Mesh")));
 	FPropertyChangedEvent PropertyChangedEvent(MeshProperty);
 
 	UParticleSystem* PartSys = CastChecked<UParticleSystem>(OuterObj);
diff --git a/Engine/Source/Runtime/Engine/Private/PhysicsEngine/PhysicalMaterialMask.cpp b/Engine/Source/Runtime/Engine/Private/PhysicsEngine/PhysicalMaterialMask.cpp
--- a/Engine/Source/Runtime/Engine/Private/PhysicsEngine/PhysicalMaterialMask.cpp
+++ b/Engine/Source/Runtime/Engine/Private/PhysicsEngine/PhysicalMaterialMask.cpp
@@ -155,11 +155,11 @@ void UPhysicalMaterialMask::SetMaskTexture(UTexture* InMaskTexture, const FStrin
 	MaskTexture = InMaskTexture;
 	AssetImportData->AddFileName(InTextureFilename, 0);
 
-	FProperty* ParamProperty = FindField<FProperty>(UPhysicalMaterialMask::StaticClass(), GET_MEMBER_NAME_STRING_CHECKED(UPhysicalMaterialMask, MaskTexture));
+	FProperty* ParamProperty = FindFProperty<FProperty>(UPhysicalMaterialMask::StaticClass(), GET_MEMBER_NAME_STRING_CHECKED(UPhysicalMaterialMask, MaskTexture));
 	FPropertyChangedEvent TextureChangedEvent(ParamProperty);
 	PostEditChangeProperty(TextureChangedEvent);
 
-	ParamProperty = FindField<FProperty>(UPhysicalMaterialMask::StaticClass(), GET_MEMBER_NAME_STRING_CHECKED(UPhysicalMaterialMask, AssetImportData));
+	ParamProperty = FindFProperty<FProperty>(UPhysicalMaterialMask::StaticClass(), GET_MEMBER_NAME_STRING_CHECKED(UPhysicalMaterialMask, AssetImportData));
 	FPropertyChangedEvent AssetImportDataChangedEvent(ParamProperty);
 	PostEditChangeProperty(AssetImportDataChangedEvent);
 }
diff --git a/Engine/Source/Runtime/Engine/Private/SCS_Node.cpp b/Engine/Source/Runtime/Engine/Private/SCS_Node.cpp
--- a/Engine/Source/Runtime/Engine/Private/SCS_Node.cpp
+++ b/Engine/Source/Runtime/Engine/Private/SCS_Node.cpp
@@ -161,7 +161,7 @@ UActorComponent* USCS_Node::ExecuteNodeOnActor(AActor* Actor, USceneComponent* P
 		if (VarName != NAME_None)
 		{
 			UClass* ActorClass = Actor->GetClass();
-			if (FObjectPropertyBase* Prop = FindField<FObjectPropertyBase>(ActorClass, VarName))
+			if (FObjectPropertyBase* Prop = FindFProperty<FObjectPropertyBase>(ActorClass, VarName))
 			{
 				// If it is null we don't really know what's going on, but make it behave as it did before the bug fix
 				if (Prop->PropertyClass == nullptr || NewActorComp->IsA(Prop->PropertyClass))
diff --git a/Engine/Source/Runtime/Engine/Private/SimpleConstructionScript.cpp b/Engine/Source/Runtime/Engine/Private/SimpleConstructionScript.cpp
--- a/Engine/Source/Runtime/Engine/Private/SimpleConstructionScript.cpp
+++ b/Engine/Source/Runtime/Engine/Private/SimpleConstructionScript.cpp
@@ -666,7 +666,7 @@ void USimpleConstructionScript::ExecuteScriptOnActor(AActor* Actor, const TInlin
 					else
 					{
 						// In the non-native case, the SCS node's variable name property is used as the parent identifier
-						FObjectPropertyBase* Property = FindField<FObjectPropertyBase>(ActorClass, RootNode->ParentComponentOrVariableName);
+						FObjectPropertyBase* Property = FindFProperty<FObjectPropertyBase>(ActorClass, RootNode->ParentComponentOrVariableName);
 						if(Property != nullptr)
 						{
 							// If we found a matching property, grab its value and use that as the parent for this node
diff --git a/Engine/Source/Runtime/Engine/Private/StaticMesh.cpp b/Engine/Source/Runtime/Engine/Private/StaticMesh.cpp
--- a/Engine/Source/Runtime/Engine/Private/StaticMesh.cpp
+++ b/Engine/Source/Runtime/Engine/Private/StaticMesh.cpp
@@ -6480,7 +6480,7 @@ void UStaticMesh::SetMaterial(int32 MaterialIndex, UMaterialInterface* NewMateri
 		FScopedTransaction ScopeTransaction(LOCTEXT("StaticMeshMaterialChanged", "StaticMesh: Material changed"));
 
 		// flag the property (Materials) we're modifying so that not all of the object is rebuilt.
-		FProperty* ChangedProperty = FindField<FProperty>(UStaticMesh::StaticClass(), NAME_StaticMaterials);
+		FProperty* ChangedProperty = FindFProperty<FProperty>(UStaticMesh::StaticClass(), NAME_StaticMaterials);
 		check(ChangedProperty);
 		PreEditChange(ChangedProperty);
 		UMaterialInterface* CancelOldMaterial = StaticMaterials[MaterialIndex].MaterialInterface;
diff --git a/Engine/Source/Runtime/Engine/Private/Tests/Internationalization/TextPropertyTest.cpp b/Engine/Source/Runtime/Engine/Private/Tests/Internationalization/TextPropertyTest.cpp
--- a/Engine/Source/Runtime/Engine/Private/Tests/Internationalization/TextPropertyTest.cpp
+++ b/Engine/Source/Runtime/Engine/Private/Tests/Internationalization/TextPropertyTest.cpp
@@ -22,8 +22,8 @@ IMPLEMENT_SIMPLE_AUTOMATION_TEST(FTextPropertyTest, "System.Engine.International
 bool FTextPropertyTest::RunTest (const FString& Parameters)
 {
 	UClass* const TextPropertyTestObjectClass = UTextPropertyTestObject::StaticClass();
-	FTextProperty* const DefaultedTextProperty = FindField<FTextProperty>(TextPropertyTestObjectClass, "DefaultedText");
-	FTextProperty* const UndefaultedTextProperty = FindField<FTextProperty>(TextPropertyTestObjectClass, "UndefaultedText");
+	FTextProperty* const DefaultedTextProperty = FindFProperty<FTextProperty>(TextPropertyTestObjectClass, "DefaultedText");
+	FTextProperty* const UndefaultedTextProperty = FindFProperty<FTextProperty>(TextPropertyTestObjectClass, "UndefaultedText");
 	UTextPropertyTestObject* const TextPropertyTestCDO = Cast<UTextPropertyTestObject>( TextPropertyTestObjectClass->ClassDefaultObject );
 
 	{
diff --git a/Engine/Source/Runtime/Engine/Private/Timeline.cpp b/Engine/Source/Runtime/Engine/Private/Timeline.cpp
--- a/Engine/Source/Runtime/Engine/Private/Timeline.cpp
+++ b/Engine/Source/Runtime/Engine/Private/Timeline.cpp
@@ -233,7 +233,7 @@ void FTimeline::SetPlaybackPosition(float NewPosition, bool bFireEvents, bool bF
 			{
 				if (VecEntry.VectorProperty == NULL)
 				{
-					VecEntry.VectorProperty = FindField<FStructProperty>(PropSetObject->GetClass(), VecEntry.VectorPropertyName);
+					VecEntry.VectorProperty = FindFProperty<FStructProperty>(PropSetObject->GetClass(), VecEntry.VectorPropertyName);
 					if(VecEntry.VectorProperty == NULL)
 					{
 						UE_LOG(LogTimeline, Log, TEXT("SetPlaybackPosition: No vector property '%s' in '%s'"), *VecEntry.VectorPropertyName.ToString(), *PropSetObject->GetName());
@@ -267,7 +267,7 @@ void FTimeline::SetPlaybackPosition(float NewPosition, bool bFireEvents, bool bF
 			{
 				if (FloatEntry.FloatProperty == NULL)
 				{
-					FloatEntry.FloatProperty = FindField<FFloatProperty>(PropSetObject->GetClass(), FloatEntry.FloatPropertyName);
+					FloatEntry.FloatProperty = FindFProperty<FFloatProperty>(PropSetObject->GetClass(), FloatEntry.FloatPropertyName);
 					if(FloatEntry.FloatProperty == NULL)
 					{
 						UE_LOG(LogTimeline, Log, TEXT("SetPlaybackPosition: No float property '%s' in '%s'"), *FloatEntry.FloatPropertyName.ToString(), *PropSetObject->GetName());
@@ -301,7 +301,7 @@ void FTimeline::SetPlaybackPosition(float NewPosition, bool bFireEvents, bool bF
 			{
 				if (ColorEntry.LinearColorProperty == NULL)
 				{
-					ColorEntry.LinearColorProperty = FindField<FStructProperty>(PropSetObject->GetClass(), ColorEntry.LinearColorPropertyName);
+					ColorEntry.LinearColorProperty = FindFProperty<FStructProperty>(PropSetObject->GetClass(), ColorEntry.LinearColorPropertyName);
 					if(ColorEntry.LinearColorProperty == NULL)
 					{
 						UE_LOG(LogTimeline, Log, TEXT("SetPlaybackPosition: No linear color property '%s' in '%s'"), *ColorEntry.LinearColorPropertyName.ToString(), *PropSetObject->GetName());
@@ -324,10 +324,10 @@ void FTimeline::SetPlaybackPosition(float NewPosition, bool bFireEvents, bool bF
 		{
 			if (DirectionProperty == nullptr)
 			{
-				DirectionProperty = FindField<FByteProperty>(PropSetObject->GetClass(), DirectionPropertyName);
+				DirectionProperty = FindFProperty<FByteProperty>(PropSetObject->GetClass(), DirectionPropertyName);
 				if (DirectionProperty == nullptr)
 				{
-					DirectionProperty = FindField<FEnumProperty>(PropSetObject->GetClass(), DirectionPropertyName);
+					DirectionProperty = FindFProperty<FEnumProperty>(PropSetObject->GetClass(), DirectionPropertyName);
 				}
 
 				if (DirectionProperty == nullptr)
diff --git a/Engine/Source/Runtime/Engine/Private/UnrealEngine.cpp b/Engine/Source/Runtime/Engine/Private/UnrealEngine.cpp
--- a/Engine/Source/Runtime/Engine/Private/UnrealEngine.cpp
+++ b/Engine/Source/Runtime/Engine/Private/UnrealEngine.cpp
@@ -10464,7 +10464,7 @@ void DrawStatsHUD( UWorld* World, FViewport* Viewport, FCanvas* Canvas, UCanvas*
 			UClass* Cls = Cast<UClass>(DebugProperties[i].Obj);
 			if (Cls != NULL)
 			{
-				FProperty* Prop = FindField<FProperty>(Cls, DebugProperties[i].PropertyName);
+				FProperty* Prop = FindFProperty<FProperty>(Cls, DebugProperties[i].PropertyName);
 				if (Prop != NULL || DebugProperties[i].bSpecialProperty)
 				{
 					// getall
@@ -10485,7 +10485,7 @@ void DrawStatsHUD( UWorld* World, FViewport* Viewport, FCanvas* Canvas, UCanvas*
 			}
 			else
 			{
-				FProperty* Prop = FindField<FProperty>(DebugProperties[i].Obj->GetClass(), DebugProperties[i].PropertyName);
+				FProperty* Prop = FindFProperty<FProperty>(DebugProperties[i].Obj->GetClass(), DebugProperties[i].PropertyName);
 				if (Prop != NULL || DebugProperties[i].bSpecialProperty)
 				{
 					DrawProperty(CanvasObject, DebugProperties[i].Obj, DebugProperties[i], Prop, X, Y);
@@ -14274,7 +14274,7 @@ void UEngine::CopyPropertiesForUnrelatedObjects(UObject* OldObject, UObject* New
 	FObjectProperty* RootComponentProperty = nullptr;
 	if (NewActor != nullptr && Params.bPreserveRootComponent)
 	{
-		RootComponentProperty = FindField<FObjectProperty>(NewActor->GetClass(), "RootComponent");
+		RootComponentProperty = FindFProperty<FObjectProperty>(NewActor->GetClass(), "RootComponent");
 		if (RootComponentProperty != nullptr)
 		{
 			SavedRootComponent = Cast<USceneComponent>(RootComponentProperty->GetObjectPropertyValue_InContainer(NewActor));
diff --git a/Engine/Source/Runtime/Experimental/GeometryCollectionEngine/Private/GeometryCollection/GeometryCollectionComponent.cpp b/Engine/Source/Runtime/Experimental/GeometryCollectionEngine/Private/GeometryCollection/GeometryCollectionComponent.cpp
--- a/Engine/Source/Runtime/Experimental/GeometryCollectionEngine/Private/GeometryCollection/GeometryCollectionComponent.cpp
+++ b/Engine/Source/Runtime/Experimental/GeometryCollectionEngine/Private/GeometryCollection/GeometryCollectionComponent.cpp
@@ -1306,7 +1306,7 @@ void UGeometryCollectionComponent::OnCreatePhysicsState()
 
 							if (EditorComponent)
 							{
-								EditorComponent->PreEditChange(FindField<FProperty>(EditorComponent->GetClass(), GET_MEMBER_NAME_CHECKED(UGeometryCollectionComponent, CacheParameters)));
+								EditorComponent->PreEditChange(FindFProperty<FProperty>(EditorComponent->GetClass(), GET_MEMBER_NAME_CHECKED(UGeometryCollectionComponent, CacheParameters)));
 								EditorComponent->Modify();
 
 								EditorComponent->CacheParameters.TargetCache = CacheParameters.TargetCache;
diff --git a/Engine/Source/Runtime/GameplayTags/Private/GameplayTagContainer.cpp b/Engine/Source/Runtime/GameplayTags/Private/GameplayTagContainer.cpp
--- a/Engine/Source/Runtime/GameplayTags/Private/GameplayTagContainer.cpp
+++ b/Engine/Source/Runtime/GameplayTags/Private/GameplayTagContainer.cpp
@@ -1675,7 +1675,7 @@ void FGameplayTagQuery::BuildFromEditableQuery(UEditableGameplayTagQuery& Editab
 FString UEditableGameplayTagQuery::GetTagQueryExportText(FGameplayTagQuery const& TagQuery)
 {
 	TagQueryExportText_Helper = TagQuery;
-	FProperty* const TQProperty = FindField<FProperty>(GetClass(), TEXT("TagQueryExportText_Helper"));
+	FProperty* const TQProperty = FindFProperty<FProperty>(GetClass(), TEXT("TagQueryExportText_Helper"));
 
 	FString OutString;
 	TQProperty->ExportTextItem(OutString, (void*)&TagQueryExportText_Helper, (void*)&TagQueryExportText_Helper, this, 0);
diff --git a/Engine/Source/Runtime/GameplayTags/Private/GameplayTagsManager.cpp b/Engine/Source/Runtime/GameplayTags/Private/GameplayTagsManager.cpp
--- a/Engine/Source/Runtime/GameplayTags/Private/GameplayTagsManager.cpp
+++ b/Engine/Source/Runtime/GameplayTags/Private/GameplayTagsManager.cpp
@@ -1358,7 +1358,7 @@ FString UGameplayTagsManager::GetCategoriesMetaFromFunction(const UFunction* Thi
 		// If a param name was specified, check it first for UPARAM metadata
 		if (!ParamName.IsNone())
 		{
-			FProperty* ParamProp = FindField<FProperty>(ThisFunction, ParamName);
+			FProperty* ParamProp = FindFProperty<FProperty>(ThisFunction, ParamName);
 			if (ParamProp)
 			{
 				FilterString = UGameplayTagsManager::Get().GetCategoriesMetaFromField(ParamProp);
diff --git a/Engine/Source/Runtime/MovieScene/Private/MovieSceneCommonHelpers.cpp b/Engine/Source/Runtime/MovieScene/Private/MovieSceneCommonHelpers.cpp
--- a/Engine/Source/Runtime/MovieScene/Private/MovieSceneCommonHelpers.cpp
+++ b/Engine/Source/Runtime/MovieScene/Private/MovieSceneCommonHelpers.cpp
@@ -417,7 +417,7 @@ FPropertyAndIndex FindPropertyAndArrayIndex(UStruct* InStruct, const FString& Pr
 		if (PropertyName.FindLastChar('[', OpenIndex))
 		{
 			FString TruncatedPropertyName(OpenIndex, *PropertyName);
-			PropertyAndIndex.Property = FindField<FProperty>(InStruct, *TruncatedPropertyName);
+			PropertyAndIndex.Property = FindFProperty<FProperty>(InStruct, *TruncatedPropertyName);
 
 			const int32 NumberLength = PropertyName.Len() - OpenIndex - 2;
 			if (NumberLength > 0 && NumberLength <= 10)
@@ -431,7 +431,7 @@ FPropertyAndIndex FindPropertyAndArrayIndex(UStruct* InStruct, const FString& Pr
 		}
 	}
 
-	PropertyAndIndex.Property = FindField<FProperty>(InStruct, *PropertyName);
+	PropertyAndIndex.Property = FindFProperty<FProperty>(InStruct, *PropertyName);
 	return PropertyAndIndex;
 }
 
diff --git a/Engine/Source/Runtime/NavigationSystem/Private/NavAreas/NavAreaMeta_SwitchByAgent.cpp b/Engine/Source/Runtime/NavigationSystem/Private/NavAreas/NavAreaMeta_SwitchByAgent.cpp
--- a/Engine/Source/Runtime/NavigationSystem/Private/NavAreas/NavAreaMeta_SwitchByAgent.cpp
+++ b/Engine/Source/Runtime/NavigationSystem/Private/NavAreas/NavAreaMeta_SwitchByAgent.cpp
@@ -73,7 +73,7 @@ void UNavAreaMeta_SwitchByAgent::UpdateAgentConfig()
 	for (int32 i = 0; i < MaxAllowedAgents; i++)
 	{
 		const FString PropName = FString::Printf(TEXT("Agent%dArea"), i);
-		FProperty* Prop = FindField<FProperty>(UNavAreaMeta_SwitchByAgent::StaticClass(), *PropName);
+		FProperty* Prop = FindFProperty<FProperty>(UNavAreaMeta_SwitchByAgent::StaticClass(), *PropName);
 		check(Prop);
 
 		if (i < NumAgents && NumAgents > 1)
diff --git a/Engine/Source/Runtime/NavigationSystem/Private/NavigationSystem.cpp b/Engine/Source/Runtime/NavigationSystem/Private/NavigationSystem.cpp
--- a/Engine/Source/Runtime/NavigationSystem/Private/NavigationSystem.cpp
+++ b/Engine/Source/Runtime/NavigationSystem/Private/NavigationSystem.cpp
@@ -2512,8 +2512,8 @@ void UNavigationSystemV1::DescribeFilterFlags(const TArray<FString>& FlagsDesc)
 #endif
 
 	// setup properties
-	FStructProperty* StructProp1 = FindField<FStructProperty>(UNavigationQueryFilter::StaticClass(), TEXT("IncludeFlags"));
-	FStructProperty* StructProp2 = FindField<FStructProperty>(UNavigationQueryFilter::StaticClass(), TEXT("ExcludeFlags"));
+	FStructProperty* StructProp1 = FindFProperty<FStructProperty>(UNavigationQueryFilter::StaticClass(), TEXT("IncludeFlags"));
+	FStructProperty* StructProp2 = FindFProperty<FStructProperty>(UNavigationQueryFilter::StaticClass(), TEXT("ExcludeFlags"));
 	check(StructProp1);
 	check(StructProp2);
 
@@ -2525,7 +2525,7 @@ void UNavigationSystemV1::DescribeFilterFlags(const TArray<FString>& FlagsDesc)
 		for (int32 FlagIndex = 0; FlagIndex < MaxFlags; FlagIndex++)
 		{
 			FString PropName = FString::Printf(TEXT("bNavFlag%d"), FlagIndex);
-			FProperty* Prop = FindField<FProperty>(Structs[StructIndex], *PropName);
+			FProperty* Prop = FindFProperty<FProperty>(Structs[StructIndex], *PropName);
 			check(Prop);
 
 			if (UseDesc[FlagIndex].Len())
diff --git a/Engine/Source/Runtime/PropertyPath/Private/PropertyPathHelpers.cpp b/Engine/Source/Runtime/PropertyPath/Private/PropertyPathHelpers.cpp
--- a/Engine/Source/Runtime/PropertyPath/Private/PropertyPathHelpers.cpp
+++ b/Engine/Source/Runtime/PropertyPath/Private/PropertyPathHelpers.cpp
@@ -567,11 +567,7 @@ FFieldVariant FPropertyPathSegment::Resolve(UStruct* InStruct) const
 		if ( InStruct != Struct )
 		{
 			Struct = InStruct;
-			Field = FindField<FProperty>(InStruct, Name);
-			if (!Field.IsValid())
-			{
-				Field = FindField<UField>(InStruct, Name);
-			}
+			Field = FindUFieldOrFProperty(InStruct, Name);
 		}
 
 		return Field;
diff --git a/Engine/Source/Runtime/Serialization/Private/StructDeserializer.cpp b/Engine/Source/Runtime/Serialization/Private/StructDeserializer.cpp
--- a/Engine/Source/Runtime/Serialization/Private/StructDeserializer.cpp
+++ b/Engine/Source/Runtime/Serialization/Private/StructDeserializer.cpp
@@ -124,7 +124,7 @@ bool FStructDeserializer::Deserialize( void* OutStruct, UStruct& TypeInfo, IStru
 			{
 				FReadState NewState;
 
-				NewState.Property = FindField<FProperty>(CurrentState.TypeInfo, *PropertyName);
+				NewState.Property = FindFProperty<FProperty>(CurrentState.TypeInfo, *PropertyName);
 
 				if (NewState.Property != nullptr)
 				{
@@ -247,7 +247,7 @@ bool FStructDeserializer::Deserialize( void* OutStruct, UStruct& TypeInfo, IStru
 				else
 				{
 					// handle scalar property
-					FProperty* Property = FindField<FProperty>(CurrentState.TypeInfo, *PropertyName);
+					FProperty* Property = FindFProperty<FProperty>(CurrentState.TypeInfo, *PropertyName);
 
 					if (Property != nullptr)
 					{
@@ -351,7 +351,7 @@ bool FStructDeserializer::Deserialize( void* OutStruct, UStruct& TypeInfo, IStru
 				}
 				else
 				{
-					NewState.Property = FindField<FProperty>(CurrentState.TypeInfo, *PropertyName);
+					NewState.Property = FindFProperty<FProperty>(CurrentState.TypeInfo, *PropertyName);
 
 					// unrecognized property
 					if (NewState.Property == nullptr)
diff --git a/Engine/Source/Runtime/UMG/Private/Animation/WidgetAnimationDelegateBinding.cpp b/Engine/Source/Runtime/UMG/Private/Animation/WidgetAnimationDelegateBinding.cpp
--- a/Engine/Source/Runtime/UMG/Private/Animation/WidgetAnimationDelegateBinding.cpp
+++ b/Engine/Source/Runtime/UMG/Private/Animation/WidgetAnimationDelegateBinding.cpp
@@ -19,7 +19,7 @@ void UWidgetAnimationDelegateBinding::BindDynamicDelegates(UObject* InInstance)
 		{
 			const FBlueprintWidgetAnimationDelegateBinding& Binding = WidgetAnimationDelegateBindings[BindIndex];
 
-			FObjectProperty* AnimationProp = FindField<FObjectProperty>(InUserWidget->GetClass(), Binding.AnimationToBind);
+			FObjectProperty* AnimationProp = FindFProperty<FObjectProperty>(InUserWidget->GetClass(), Binding.AnimationToBind);
 			if (AnimationProp)
 			{
 				UWidgetAnimation* AnimationPropData = Cast<UWidgetAnimation>(AnimationProp->GetObjectPropertyValue_InContainer(InUserWidget));
diff --git a/Engine/Source/Runtime/UMG/Private/Components/Widget.cpp b/Engine/Source/Runtime/UMG/Private/Components/Widget.cpp
--- a/Engine/Source/Runtime/UMG/Private/Components/Widget.cpp
+++ b/Engine/Source/Runtime/UMG/Private/Components/Widget.cpp
@@ -1472,7 +1472,7 @@ FString UWidget::GetDefaultFontName()
 
 //bool UWidget::BindProperty(const FName& DestinationProperty, UObject* SourceObject, const FName& SourceProperty)
 //{
-//	FDelegateProperty* DelegateProperty = FindField<FDelegateProperty>(GetClass(), FName(*( DestinationProperty.ToString() + TEXT("Delegate") )));
+//	FDelegateProperty* DelegateProperty = FindFProperty<FDelegateProperty>(GetClass(), FName(*( DestinationProperty.ToString() + TEXT("Delegate") )));
 //
 //	if ( DelegateProperty )
 //	{
diff --git a/Engine/Source/Runtime/UMG/Private/UserWidget.cpp b/Engine/Source/Runtime/UMG/Private/UserWidget.cpp
--- a/Engine/Source/Runtime/UMG/Private/UserWidget.cpp
+++ b/Engine/Source/Runtime/UMG/Private/UserWidget.cpp
@@ -138,7 +138,7 @@ void UUserWidget::TemplateInitInner()
 			FString VariableName = Widget->GetName();
 
 			// Find property with the same name as the template and assign the new widget to it.
-			FObjectPropertyBase* Prop = FindField<FObjectPropertyBase>(WidgetClass, *VariableName);
+			FObjectPropertyBase* Prop = FindFProperty<FObjectPropertyBase>(WidgetClass, *VariableName);
 			if ( Prop )
 			{
 				Prop->SetObjectPropertyValue_InContainer(this, Widget);
@@ -263,7 +263,7 @@ bool UUserWidget::VerifyTemplateIntegrity(UUserWidget* TemplateRoot, TArray<FTex
 			FName VariableFName = Widget->GetFName();
 
 			// Find property with the same name as the template and assign the new widget to it.
-			FObjectPropertyBase* Prop = FindField<FObjectPropertyBase>(TemplateClass, VariableFName);
+			FObjectPropertyBase* Prop = FindFProperty<FObjectPropertyBase>(TemplateClass, VariableFName);
 			if ( Prop )
 			{
 				UObject* Value = Prop->GetObjectPropertyValue_InContainer(this);
@@ -378,7 +378,7 @@ void UUserWidget::InitializeNamedSlots(bool bReparentToWidgetTree)
 	{
 		if ( UWidget* BindingContent = Binding.Content )
 		{
-			FObjectPropertyBase* NamedSlotProperty = FindField<FObjectPropertyBase>(GetClass(), Binding.Name);
+			FObjectPropertyBase* NamedSlotProperty = FindFProperty<FObjectPropertyBase>(GetClass(), Binding.Name);
 #if !WITH_EDITOR
 			// In editor, renaming a NamedSlot widget will cause this ensure in UpdatePreviewWidget of widget that use that namedslot
 			ensure(NamedSlotProperty);
diff --git a/Engine/Source/Runtime/UMG/Private/WidgetBlueprintGeneratedClass.cpp b/Engine/Source/Runtime/UMG/Private/WidgetBlueprintGeneratedClass.cpp
--- a/Engine/Source/Runtime/UMG/Private/WidgetBlueprintGeneratedClass.cpp
+++ b/Engine/Source/Runtime/UMG/Private/WidgetBlueprintGeneratedClass.cpp
@@ -140,7 +140,7 @@ void UWidgetBlueprintGeneratedClass::InitializeBindingsStatic(UUserWidget* UserW
 	for (const FDelegateRuntimeBinding& Binding : InBindings)
 	{
 		// If the binding came from a parent class, this will still find it - FindField() searches the super class hierarchy by default.
-		FObjectProperty* WidgetProperty = FindField<FObjectProperty>(UserWidget->GetClass(), *Binding.ObjectName);
+		FObjectProperty* WidgetProperty = FindFProperty<FObjectProperty>(UserWidget->GetClass(), *Binding.ObjectName);
 		if (WidgetProperty == nullptr)
 		{
 			continue;
@@ -150,10 +150,10 @@ void UWidgetBlueprintGeneratedClass::InitializeBindingsStatic(UUserWidget* UserW
 
 		if (Widget)
 		{
-			FDelegateProperty* DelegateProperty = FindField<FDelegateProperty>(Widget->GetClass(), FName(*(Binding.PropertyName.ToString() + TEXT("Delegate"))));
+			FDelegateProperty* DelegateProperty = FindFProperty<FDelegateProperty>(Widget->GetClass(), FName(*(Binding.PropertyName.ToString() + TEXT("Delegate"))));
 			if (!DelegateProperty)
 			{
-				DelegateProperty = FindField<FDelegateProperty>(Widget->GetClass(), Binding.PropertyName);
+				DelegateProperty = FindFProperty<FDelegateProperty>(Widget->GetClass(), Binding.PropertyName);
 			}
 
 			if (DelegateProperty)
@@ -297,7 +297,7 @@ void UWidgetBlueprintGeneratedClass::InitializeWidgetStatic(UUserWidget* UserWid
 			FString VariableName = Widget->GetName();
 
 			// Find property with the same name as the template and assign the new widget to it.
-			FObjectPropertyBase* Prop = FindField<FObjectPropertyBase>(WidgetBlueprintClass, *VariableName);
+			FObjectPropertyBase* Prop = FindFProperty<FObjectPropertyBase>(WidgetBlueprintClass, *VariableName);
 			if (Prop)
 			{
 				Prop->SetObjectPropertyValue_InContainer(UserWidget, Widget);
@@ -335,7 +335,7 @@ void UWidgetBlueprintGeneratedClass::BindAnimations(UUserWidget* Instance, const
 		if (Animation->GetMovieScene())
 		{
 			// Find property with the same name as the animation and assign the animation to it.
-			FObjectPropertyBase* Prop = FindField<FObjectPropertyBase>(Instance->GetClass(), Animation->GetMovieScene()->GetFName());
+			FObjectPropertyBase* Prop = FindFProperty<FObjectPropertyBase>(Instance->GetClass(), Animation->GetMovieScene()->GetFName());
 			if (Prop)
 			{
 				Prop->SetObjectPropertyValue_InContainer(Instance, Animation);
