Temporary fix for crashes when exiting the game / editor

diff --git a/Engine/Source/Runtime/Engine/Private/SoundEffectPreset.cpp b/Engine/Source/Runtime/Engine/Private/SoundEffectPreset.cpp
--- a/Engine/Source/Runtime/Engine/Private/SoundEffectPreset.cpp
+++ b/Engine/Source/Runtime/Engine/Private/SoundEffectPreset.cpp
@@ -58,10 +58,13 @@ void USoundEffectPreset::AddReferencedEffects(FReferenceCollector& InCollector)
 
 void USoundEffectPreset::BeginDestroy()
 {
-	IterateEffects<FSoundEffectBase>([](FSoundEffectBase& Instance)
+	if (!GExitPurge)
 	{
-		Instance.ClearPreset();
-	});
+		IterateEffects<FSoundEffectBase>([](FSoundEffectBase& Instance)
+		{
+			Instance.ClearPreset();
+		});
+	}
 	Instances.Reset();
 
 	Super::BeginDestroy();
