Make sure we avoid changing data for asset not using the new reduction criterion base on percent and absolute
#rnx

diff --git a/Engine/Source/Runtime/Engine/Private/SkeletalMesh.cpp b/Engine/Source/Runtime/Engine/Private/SkeletalMesh.cpp
--- a/Engine/Source/Runtime/Engine/Private/SkeletalMesh.cpp
+++ b/Engine/Source/Runtime/Engine/Private/SkeletalMesh.cpp
@@ -4087,8 +4087,17 @@ static void SerializeReductionSettingsForDDC(FArchive& Ar, FSkeletalMeshOptimiza
 	Ar << ReductionSettings.NumOfVertPercentage;
 	Ar << ReductionSettings.MaxNumOfTriangles;
 	Ar << ReductionSettings.MaxNumOfVerts;
-	Ar << ReductionSettings.MaxNumOfTrianglesPercentage;
-	Ar << ReductionSettings.MaxNumOfVertsPercentage;
+
+	// Keep old DDC keys if these are not set
+	if (ReductionSettings.MaxNumOfTrianglesPercentage != MAX_uint32 || 
+		ReductionSettings.MaxNumOfVertsPercentage != MAX_uint32)
+	{
+		uint32 AvoidCachePoisoningFromOldBug = 0;
+		Ar << AvoidCachePoisoningFromOldBug;
+		Ar << ReductionSettings.MaxNumOfTrianglesPercentage;
+		Ar << ReductionSettings.MaxNumOfVertsPercentage;
+	}
+
 	Ar << ReductionSettings.MaxDeviationPercentage;
 	Ar << ReductionSettings.ReductionMethod;
 	Ar << ReductionSettings.SilhouetteImportance;
diff --git a/Engine/Source/Runtime/Engine/Private/SkeletalMeshRenderData.cpp b/Engine/Source/Runtime/Engine/Private/SkeletalMeshRenderData.cpp
--- a/Engine/Source/Runtime/Engine/Private/SkeletalMeshRenderData.cpp
+++ b/Engine/Source/Runtime/Engine/Private/SkeletalMeshRenderData.cpp
@@ -58,16 +58,12 @@ static void SerializeLODInfoForDDC(USkeletalMesh* SkeletalMesh, FString& KeySuff
 // If skeletal mesh derived data needs to be rebuilt (new format, serialization
 // differences, etc.) replace the version GUID below with a new one.
 // In case of merge conflicts with DDC versions, you MUST generate a new GUID
-// and set this new GUID as the version.
-#define SKELETALMESH_DERIVEDDATA_VER TEXT("EBD3AC7CCBB449049C38BE9816BF7C2E")
+// and set this new GUID as the version.                                       
+#define SKELETALMESH_DERIVEDDATA_VER TEXT("DEE857ABC2804080B6278513ABC358BB")
 
 static const FString& GetSkeletalMeshDerivedDataVersion()
 {
-	static FString CachedVersionString;
-	if (CachedVersionString.IsEmpty())
-	{
-		CachedVersionString = FString::Printf(TEXT("%s"), SKELETALMESH_DERIVEDDATA_VER);
-	}
+	static FString CachedVersionString = SKELETALMESH_DERIVEDDATA_VER;
 	return CachedVersionString;
 }
 
@@ -149,7 +145,7 @@ void FSkeletalMeshRenderData::Cache(const ITargetPlatform* TargetPlatform, USkel
 		if (GetDerivedDataCacheRef().GetSynchronous(*DerivedDataKey, DerivedData, Owner->GetPathName()))
 		{
 			COOK_STAT(Timer.AddHit(DerivedData.Num()));
-
+			
 			FMemoryReader Ar(DerivedData, /*bIsPersistent=*/ true);
 
 			//With skeletal mesh build refactor we serialize the LODModel sections into the DDC
diff --git a/Engine/Source/Runtime/Engine/Public/SkeletalMeshReductionSettings.h b/Engine/Source/Runtime/Engine/Public/SkeletalMeshReductionSettings.h
--- a/Engine/Source/Runtime/Engine/Public/SkeletalMeshReductionSettings.h
+++ b/Engine/Source/Runtime/Engine/Public/SkeletalMeshReductionSettings.h
@@ -79,6 +79,7 @@ struct FSkeletalMeshOptimizationSettings
 	UPROPERTY(EditAnywhere, Category = ReductionMethod, meta = (DisplayName = "Max Vertex Count", ClampMin = 6))
 	uint32 MaxNumOfVerts;
 
+#if WITH_EDITORONLY_DATA
 	/** The maximum number of triangles to retain when using percentage termination criterion. */
 	UPROPERTY(EditAnywhere, Category = ReductionMethod, meta = (DisplayName = "Max Triangle Count", ClampMin = 4, UIMin = "4"))
 	uint32 MaxNumOfTrianglesPercentage;
@@ -86,6 +87,7 @@ struct FSkeletalMeshOptimizationSettings
 	/** The maximum number of vertices to retain when using percentage termination criterion. */
 	UPROPERTY(EditAnywhere, Category = ReductionMethod, meta = (DisplayName = "Max Vertex Count", ClampMin = 6, UIMin = "6"))
 	uint32 MaxNumOfVertsPercentage;
+#endif
 
 	/**If ReductionMethod equals MaxDeviation this value is the maximum deviation from the base mesh as a percentage of the bounding sphere. 
 	 * In code, it ranges from [0, 1]. In the editor UI, it ranges from [0, 100]
@@ -171,8 +173,10 @@ struct FSkeletalMeshOptimizationSettings
 		, NumOfVertPercentage(0.5f)
 		, MaxNumOfTriangles(4)
 		, MaxNumOfVerts(6)
+#if WITH_EDITORONLY_DATA		
 		, MaxNumOfTrianglesPercentage(MAX_uint32)
 		, MaxNumOfVertsPercentage(MAX_uint32)
+#endif
 		, MaxDeviationPercentage(0.5f)
 		, ReductionMethod(SMOT_NumOfTriangles)
 		, SilhouetteImportance(SMOI_Normal)
