Fix for StopDelayed() not working

diff --git a/Engine/Source/Runtime/Engine/Classes/Components/AudioComponent.h b/Engine/Source/Runtime/Engine/Classes/Components/AudioComponent.h
--- a/Engine/Source/Runtime/Engine/Classes/Components/AudioComponent.h
+++ b/Engine/Source/Runtime/Engine/Classes/Components/AudioComponent.h
@@ -435,7 +435,7 @@ public:
 	UFUNCTION(BlueprintCallable, Category="Audio|Components|Audio")
 	virtual void Stop();
 
-	/** Cues request to stop sound after the provided delay, stopping immediately if delay is zero or negative */
+	/** Cues request to stop sound after the provided delay (In Seconds), stopping immediately if delay is zero or negative */
 	UFUNCTION(BlueprintCallable, Category="Audio|Components|Audio")
 	void StopDelayed(float DelayTime);
 
diff --git a/Engine/Source/Runtime/SignalProcessing/Private/VolumeFader.cpp b/Engine/Source/Runtime/SignalProcessing/Private/VolumeFader.cpp
--- a/Engine/Source/Runtime/SignalProcessing/Private/VolumeFader.cpp
+++ b/Engine/Source/Runtime/SignalProcessing/Private/VolumeFader.cpp
@@ -184,6 +184,7 @@ namespace Audio
 		else
 		{
 			ActiveDuration = InDuration;
+			Elapsed = 0.0f;
 		}
 	}
 
@@ -232,7 +233,7 @@ namespace Audio
 		if (FadeCurve == EFaderCurve::Logarithmic)
 		{
 			Alpha = Audio::ConvertToLinear(Alpha);
-		}
+		} 
 		Target = Alpha;
 		FadeCurve = EFaderCurve::Linear;
 		Elapsed = ActiveDuration;
