Deprecating some poorly named functions for management capture mode. CaptureMouseOnClick() -> GetMouseCaptureMode(), SetCaptureMouseOnClick() -> SetMouseCaptureMode().  Added a new function for setting if we should show the mouse cursor.  I've not yet made access to the current bool private, but that's the direction we should go.  Added a new log category for viewport and logging more information about when we transition between different viewport capture states/cusrsor visibility states...etc.

diff --git a/Engine/Source/Runtime/Engine/Classes/Engine/GameViewportClient.h b/Engine/Source/Runtime/Engine/Classes/Engine/GameViewportClient.h
--- a/Engine/Source/Runtime/Engine/Classes/Engine/GameViewportClient.h
+++ b/Engine/Source/Runtime/Engine/Classes/Engine/GameViewportClient.h
@@ -703,20 +703,17 @@ public:
 	}
 
 	/**
-	 * Set the mouse capture behavior when the viewport is clicked
+	 * Set the mouse capture behavior for the viewport.
 	 */
-	void SetCaptureMouseOnClick(EMouseCaptureMode Mode)
-	{
-		MouseCaptureMode = Mode;
-	}
+	void SetMouseCaptureMode(EMouseCaptureMode Mode);
+
+	UE_DEPRECATED(4.26, "Please call UGameViewportClient::SetMouseCaptureMode(EMouseCaptureMode) instead.")
+	void SetCaptureMouseOnClick(EMouseCaptureMode Mode) { SetMouseCaptureMode(Mode); }
 
 	/**
 	 * Gets the mouse capture behavior when the viewport is clicked
 	 */
-	virtual EMouseCaptureMode CaptureMouseOnClick() override
-	{
-		return MouseCaptureMode;
-	}
+	virtual EMouseCaptureMode GetMouseCaptureMode() const override;
 
 	/**
 	 * Gets whether or not the viewport captures the Mouse on launch of the application
@@ -755,23 +752,17 @@ public:
 	/**
 	* Sets the current mouse cursor lock mode when the viewport is clicked
 	*/
-	void SetMouseLockMode(EMouseLockMode InMouseLockMode)
-	{
-		MouseLockMode = InMouseLockMode;
-	}
+	void SetMouseLockMode(EMouseLockMode InMouseLockMode);
 
 	/**
 	 * Sets whether or not the cursor is hidden when the viewport captures the mouse
 	 */
-	void SetHideCursorDuringCapture(bool InHideCursorDuringCapture)
-	{
-		bHideCursorDuringCapture = InHideCursorDuringCapture;
-	}
+	void SetHideCursorDuringCapture(bool InHideCursorDuringCapture);
 
 	/**
 	 * Gets whether or not the cursor is hidden when the viewport captures the mouse
 	 */
-	virtual bool HideCursorDuringCapture() override
+	virtual bool HideCursorDuringCapture() const override
 	{
 		return bHideCursorDuringCapture;
 	}
diff --git a/Engine/Source/Runtime/Engine/Classes/GameFramework/PlayerController.h b/Engine/Source/Runtime/Engine/Classes/GameFramework/PlayerController.h
--- a/Engine/Source/Runtime/Engine/Classes/GameFramework/PlayerController.h
+++ b/Engine/Source/Runtime/Engine/Classes/GameFramework/PlayerController.h
@@ -1897,6 +1897,8 @@ public:
 	UFUNCTION(Reliable, Client)
 	void OnServerStartedVisualLogger(bool bIsLogging);
 
+	void SetShowMouseCursor(bool bShow);
+
 	/** Returns true if the mouse cursor should be shown */
 	virtual bool ShouldShowMouseCursor() const;
 
diff --git a/Engine/Source/Runtime/Engine/Private/GameViewportClient.cpp b/Engine/Source/Runtime/Engine/Private/GameViewportClient.cpp
--- a/Engine/Source/Runtime/Engine/Private/GameViewportClient.cpp
+++ b/Engine/Source/Runtime/Engine/Private/GameViewportClient.cpp
@@ -2717,12 +2717,56 @@ void UGameViewportClient::VerifyPathRenderingComponents()
 	}
 }
 
+void UGameViewportClient::SetMouseCaptureMode(EMouseCaptureMode Mode)
+{
+	if (MouseCaptureMode != Mode)
+	{
+		UE_LOG(LogViewport, Display, TEXT("Viewport MouseCaptureMode Changed, %s -> %s"),
+			*StaticEnum<EMouseCaptureMode>()->GetNameStringByValue((uint64)MouseCaptureMode),
+			*StaticEnum<EMouseCaptureMode>()->GetNameStringByValue((uint64)Mode)
+		);
+
+		MouseCaptureMode = Mode;
+	}
+}
+
+EMouseCaptureMode UGameViewportClient::GetMouseCaptureMode() const
+{
+	return MouseCaptureMode;
+}
+
 bool UGameViewportClient::CaptureMouseOnLaunch()
 {
 	// Capture mouse unless headless
 	return !FApp::CanEverRender() ? false : GetDefault<UInputSettings>()->bCaptureMouseOnLaunch;
 }
 
+void UGameViewportClient::SetMouseLockMode(EMouseLockMode InMouseLockMode)
+{
+	if (MouseLockMode != InMouseLockMode)
+	{
+		UE_LOG(LogViewport, Display, TEXT("Viewport MouseLockMode Changed, %s -> %s"),
+			*StaticEnum<EMouseLockMode>()->GetNameStringByValue((uint64)MouseLockMode),
+			*StaticEnum<EMouseLockMode>()->GetNameStringByValue((uint64)InMouseLockMode)
+		);
+
+		MouseLockMode = InMouseLockMode;
+	}
+}
+
+void UGameViewportClient::SetHideCursorDuringCapture(bool InHideCursorDuringCapture)
+{
+	if (bHideCursorDuringCapture != InHideCursorDuringCapture)
+	{
+		UE_LOG(LogViewport, Display, TEXT("Viewport HideCursorDuringCapture Changed, %s -> %s"),
+			bHideCursorDuringCapture ? TEXT("True") : TEXT("False"),
+			InHideCursorDuringCapture ? TEXT("True") : TEXT("False")
+		);
+
+		bHideCursorDuringCapture = InHideCursorDuringCapture;
+	}
+}
+
 bool UGameViewportClient::Exec( UWorld* InWorld, const TCHAR* Cmd,FOutputDevice& Ar)
 {
 	if ( FParse::Command(&Cmd,TEXT("FORCEFULLSCREEN")) )
diff --git a/Engine/Source/Runtime/Engine/Private/GameplayStatics.cpp b/Engine/Source/Runtime/Engine/Private/GameplayStatics.cpp
--- a/Engine/Source/Runtime/Engine/Private/GameplayStatics.cpp
+++ b/Engine/Source/Runtime/Engine/Private/GameplayStatics.cpp
@@ -430,7 +430,7 @@ EMouseCaptureMode UGameplayStatics::GetViewportMouseCaptureMode(const UObject* W
 		UGameViewportClient* const GameViewportClient = World->GetGameViewport();
 		if (GameViewportClient)
 		{
-			return GameViewportClient->CaptureMouseOnClick();
+			return GameViewportClient->GetMouseCaptureMode();
 		}
 	}
 
@@ -445,7 +445,7 @@ void UGameplayStatics::SetViewportMouseCaptureMode(const UObject* WorldContextOb
 		UGameViewportClient* const GameViewportClient = World->GetGameViewport();
 		if (GameViewportClient)
 		{
-			GameViewportClient->SetCaptureMouseOnClick(MouseCaptureMode);
+			GameViewportClient->SetMouseCaptureMode(MouseCaptureMode);
 		}
 	}
 }
diff --git a/Engine/Source/Runtime/Engine/Private/PlayerController.cpp b/Engine/Source/Runtime/Engine/Private/PlayerController.cpp
--- a/Engine/Source/Runtime/Engine/Private/PlayerController.cpp
+++ b/Engine/Source/Runtime/Engine/Private/PlayerController.cpp
@@ -1504,9 +1504,8 @@ void APlayerController::BeginPlay()
 	//If we are faking touch events show the cursor
 	if (FSlateApplication::IsInitialized() && FSlateApplication::Get().IsFakingTouchEvents())
 	{
-		bShowMouseCursor = true;
+		SetShowMouseCursor(true);
 	}
-
 }
 
 void APlayerController::EndPlay(const EEndPlayReason::Type EndPlayReason)
@@ -2462,6 +2461,19 @@ bool APlayerController::ShouldShowMouseCursor() const
 	return bShowMouseCursor;
 }
 
+void APlayerController::SetShowMouseCursor(bool bShow)
+{
+	if (bShowMouseCursor != bShow)
+	{
+		UE_LOG(LogViewport, Display, TEXT("Player bShowMouseCursor Changed, %s -> %s"),
+			bShowMouseCursor ? TEXT("True") : TEXT("False"),
+			bShow ? TEXT("True") : TEXT("False")
+		);
+
+		bShowMouseCursor = bShow;
+	}
+}
+
 EMouseCursor::Type APlayerController::GetMouseCursor() const
 {
 	if (ShouldShowMouseCursor())
@@ -5366,7 +5378,7 @@ void FInputModeUIOnly::ApplyInputMode(FReply& SlateOperations, class UGameViewpo
 
 		GameViewportClient.SetMouseLockMode(MouseLockMode);
 		GameViewportClient.SetIgnoreInput(true);
-		GameViewportClient.SetCaptureMouseOnClick(EMouseCaptureMode::NoCapture);
+		GameViewportClient.SetMouseCaptureMode(EMouseCaptureMode::NoCapture);
 	}
 }
 
@@ -5384,7 +5396,7 @@ void FInputModeGameAndUI::ApplyInputMode(FReply& SlateOperations, class UGameVie
 		GameViewportClient.SetMouseLockMode(MouseLockMode);
 		GameViewportClient.SetIgnoreInput(false);
 		GameViewportClient.SetHideCursorDuringCapture(bHideCursorDuringCapture);
-		GameViewportClient.SetCaptureMouseOnClick(EMouseCaptureMode::CaptureDuringMouseDown);
+		GameViewportClient.SetMouseCaptureMode(EMouseCaptureMode::CaptureDuringMouseDown);
 	}
 }
 
@@ -5399,7 +5411,7 @@ void FInputModeGameOnly::ApplyInputMode(FReply& SlateOperations, class UGameView
 		SlateOperations.LockMouseToWidget(ViewportWidgetRef);
 		GameViewportClient.SetMouseLockMode(EMouseLockMode::LockOnCapture);
 		GameViewportClient.SetIgnoreInput(false);
-		GameViewportClient.SetCaptureMouseOnClick(bConsumeCaptureMouseDown ? EMouseCaptureMode::CapturePermanently : EMouseCaptureMode::CapturePermanently_IncludingInitialMouseDown);
+		GameViewportClient.SetMouseCaptureMode(bConsumeCaptureMouseDown ? EMouseCaptureMode::CapturePermanently : EMouseCaptureMode::CapturePermanently_IncludingInitialMouseDown);
 	}
 }
 
diff --git a/Engine/Source/Runtime/Engine/Private/Slate/SceneViewport.cpp b/Engine/Source/Runtime/Engine/Private/Slate/SceneViewport.cpp
--- a/Engine/Source/Runtime/Engine/Private/Slate/SceneViewport.cpp
+++ b/Engine/Source/Runtime/Engine/Private/Slate/SceneViewport.cpp
@@ -22,6 +22,8 @@
 #include "IXRTrackingSystem.h"
 #include "StereoRenderTargetManager.h"
 
+DEFINE_LOG_CATEGORY(LogViewport);
+
 extern EWindowMode::Type GetWindowModeType(EWindowMode::Type WindowMode);
 
 const FName NAME_SceneViewport = FName(TEXT("SceneViewport"));
@@ -297,7 +299,7 @@ void FSceneViewport::ProcessAccumulatedPointerInput()
 
 	if ( bCursorHiddenDueToCapture )
 	{
-		switch ( ViewportClient->CaptureMouseOnClick() )
+		switch ( ViewportClient->GetMouseCaptureMode() )
 		{
 		case EMouseCaptureMode::NoCapture:
 		case EMouseCaptureMode::CaptureDuringMouseDown:
@@ -466,11 +468,11 @@ FReply FSceneViewport::OnMouseButtonDown( const FGeometry& InGeometry, const FPo
 		}
 
 		const bool bTemporaryCapture = 
-			ViewportClient->CaptureMouseOnClick() == EMouseCaptureMode::CaptureDuringMouseDown ||
-			(ViewportClient->CaptureMouseOnClick() == EMouseCaptureMode::CaptureDuringRightMouseDown && InMouseEvent.GetEffectingButton() == EKeys::RightMouseButton);
+			ViewportClient->GetMouseCaptureMode() == EMouseCaptureMode::CaptureDuringMouseDown ||
+			(ViewportClient->GetMouseCaptureMode() == EMouseCaptureMode::CaptureDuringRightMouseDown && InMouseEvent.GetEffectingButton() == EKeys::RightMouseButton);
 
 		// Process primary input if we aren't currently a game viewport, we already have capture, or we are permanent capture that doesn't consume the mouse down.
-		const bool bProcessInputPrimary = !IsCurrentlyGameViewport() || HasMouseCapture() || (ViewportClient->CaptureMouseOnClick() == EMouseCaptureMode::CapturePermanently_IncludingInitialMouseDown);
+		const bool bProcessInputPrimary = !IsCurrentlyGameViewport() || HasMouseCapture() || (ViewportClient->GetMouseCaptureMode() == EMouseCaptureMode::CapturePermanently_IncludingInitialMouseDown);
 
 		const bool bAnyMenuWasVisible = FSlateApplication::Get().AnyMenusVisible();
 
@@ -487,8 +489,8 @@ FReply FSceneViewport::OnMouseButtonDown( const FGeometry& InGeometry, const FPo
 		const bool bNewMenuWasOpened = !bAnyMenuWasVisible && FSlateApplication::Get().AnyMenusVisible();
 
 		const bool bPermanentCapture =
-			(ViewportClient->CaptureMouseOnClick() == EMouseCaptureMode::CapturePermanently) ||
-			(ViewportClient->CaptureMouseOnClick() == EMouseCaptureMode::CapturePermanently_IncludingInitialMouseDown);
+			(ViewportClient->GetMouseCaptureMode() == EMouseCaptureMode::CapturePermanently) ||
+			(ViewportClient->GetMouseCaptureMode() == EMouseCaptureMode::CapturePermanently_IncludingInitialMouseDown);
 
 		if (FSlateApplication::Get().IsActive() && !ViewportClient->IgnoreInput() &&
 			!bNewMenuWasOpened && // We should not focus the viewport if a menu was opened as it would close the menu
@@ -579,10 +581,25 @@ FReply FSceneViewport::OnMouseButtonUp( const FGeometry& InGeometry, const FPoin
 
 		bCursorVisible = ViewportClient->GetCursor(this, GetMouseX(), GetMouseY()) != EMouseCursor::None;
 
-		bReleaseMouseCapture =
-			bCursorVisible ||
-			ViewportClient->CaptureMouseOnClick() == EMouseCaptureMode::CaptureDuringMouseDown ||
-			( ViewportClient->CaptureMouseOnClick() == EMouseCaptureMode::CaptureDuringRightMouseDown && InMouseEvent.GetEffectingButton() == EKeys::RightMouseButton );
+		if (bCursorVisible)
+		{
+			bReleaseMouseCapture = true;
+			UE_LOG(LogViewport, Log, TEXT("Releasing Mouse Capture; Cursor is visible"));
+		}
+		else if (ViewportClient->GetMouseCaptureMode() == EMouseCaptureMode::CaptureDuringMouseDown)
+		{
+			bReleaseMouseCapture = true;
+			UE_LOG(LogViewport, Log, TEXT("Releasing Mouse Capture; EMouseCaptureMode::CaptureDuringMouseDown - Mouse Button Released"));
+		}
+		else if (ViewportClient->GetMouseCaptureMode() == EMouseCaptureMode::CaptureDuringRightMouseDown && InMouseEvent.GetEffectingButton() == EKeys::RightMouseButton)
+		{
+			bReleaseMouseCapture = true;
+			UE_LOG(LogViewport, Log, TEXT("Releasing Mouse Capture; EMouseCaptureMode::CaptureDuringRightMouseDown - Right Mouse Button Released"));
+		}
+		else
+		{
+			bReleaseMouseCapture = false;
+		}
 	}
 
 	if ( !IsCurrentlyGameViewport() || bReleaseMouseCapture )
@@ -745,7 +762,7 @@ FReply FSceneViewport::OnTouchStarted( const FGeometry& MyGeometry, const FPoint
 
 		if(ViewportClient->InputTouch( this, TouchEvent.GetUserIndex(), TouchEvent.GetPointerIndex(), ETouchType::Began, TouchPosition, TouchEvent.GetTouchForce(), FDateTime::Now(), TouchEvent.GetTouchpadIndex()) )
 		{
-			const bool bTemporaryCapture = ViewportClient->CaptureMouseOnClick() == EMouseCaptureMode::CaptureDuringMouseDown;
+			const bool bTemporaryCapture = ViewportClient->GetMouseCaptureMode() == EMouseCaptureMode::CaptureDuringMouseDown;
 			if (bTemporaryCapture)
 			{
 				CurrentReplyState = AcquireFocusAndCapture(FIntPoint(TouchEvent.GetScreenSpacePosition().X, TouchEvent.GetScreenSpacePosition().Y), EFocusCause::Mouse);
@@ -814,7 +831,7 @@ FReply FSceneViewport::OnTouchEnded( const FGeometry& MyGeometry, const FPointer
 			CurrentReplyState = FReply::Unhandled(); 
 		}
 
-		if (ViewportClient->CaptureMouseOnClick() == EMouseCaptureMode::CaptureDuringMouseDown)
+		if (ViewportClient->GetMouseCaptureMode() == EMouseCaptureMode::CaptureDuringMouseDown)
 		{
 			CurrentReplyState.ReleaseMouseCapture();
 		}
@@ -1088,8 +1105,8 @@ FReply FSceneViewport::OnFocusReceived(const FFocusEvent& InFocusEvent)
 
 			const bool bPermanentCapture = (!GIsEditor || InFocusEvent.GetCause() == EFocusCause::Mouse)	&&
 										   (ViewportClient != nullptr)										&&
-					(( ViewportClient->CaptureMouseOnClick() == EMouseCaptureMode::CapturePermanently ) ||
-					 ( ViewportClient->CaptureMouseOnClick() == EMouseCaptureMode::CapturePermanently_IncludingInitialMouseDown )
+					(( ViewportClient->GetMouseCaptureMode() == EMouseCaptureMode::CapturePermanently ) ||
+					 ( ViewportClient->GetMouseCaptureMode() == EMouseCaptureMode::CapturePermanently_IncludingInitialMouseDown )
 					);
 
 			// if bPermanentCapture is true, then ViewportClient != nullptr, so it's ok to dereference it.  But the permanent capture must be tested first.
diff --git a/Engine/Source/Runtime/Engine/Public/EngineLogs.h b/Engine/Source/Runtime/Engine/Public/EngineLogs.h
--- a/Engine/Source/Runtime/Engine/Public/EngineLogs.h
+++ b/Engine/Source/Runtime/Engine/Public/EngineLogs.h
@@ -30,4 +30,5 @@ ENGINE_API DECLARE_LOG_CATEGORY_EXTERN(LogSubtitle, Log, All);
 ENGINE_API DECLARE_LOG_CATEGORY_EXTERN(LogTexture, Log, All);
 ENGINE_API DECLARE_LOG_CATEGORY_EXTERN(LogPlayerManagement, Error, All);
 ENGINE_API DECLARE_LOG_CATEGORY_EXTERN(LogSecurity, Warning, All);
-ENGINE_API DECLARE_LOG_CATEGORY_EXTERN(LogEngineSessionManager, Log, All);
\ No newline at end of file
+ENGINE_API DECLARE_LOG_CATEGORY_EXTERN(LogEngineSessionManager, Log, All);
+ENGINE_API DECLARE_LOG_CATEGORY_EXTERN(LogViewport, Log, All);
\ No newline at end of file
diff --git a/Engine/Source/Runtime/Engine/Public/UnrealClient.h b/Engine/Source/Runtime/Engine/Public/UnrealClient.h
--- a/Engine/Source/Runtime/Engine/Public/UnrealClient.h
+++ b/Engine/Source/Runtime/Engine/Public/UnrealClient.h
@@ -911,7 +911,7 @@ public:
 	 *
 	 * @return	the cursor that the OS should display
 	 */
-	virtual EMouseCursor::Type GetCursor(FViewport* Viewport,int32 X,int32 Y) { return EMouseCursor::Default; }
+	virtual EMouseCursor::Type GetCursor(FViewport* Viewport, int32 X,int32 Y) { return EMouseCursor::Default; }
 
 	/**
 	 * Called to map a cursor reply to an actual widget to render.
@@ -936,8 +936,8 @@ public:
 
 	virtual bool IsInPermanentCapture()
 	{ 
-		return  !GIsEditor && ((CaptureMouseOnClick() == EMouseCaptureMode::CapturePermanently) ||
-			(CaptureMouseOnClick() == EMouseCaptureMode::CapturePermanently_IncludingInitialMouseDown));
+		return  !GIsEditor && ((GetMouseCaptureMode() == EMouseCaptureMode::CapturePermanently) ||
+			(GetMouseCaptureMode() == EMouseCaptureMode::CapturePermanently_IncludingInitialMouseDown));
 	}
 
 	/**
@@ -1039,7 +1039,10 @@ public:
 	/**
 	 * Gets the mouse capture behavior when the viewport is clicked
 	 */
-	virtual EMouseCaptureMode CaptureMouseOnClick() { return EMouseCaptureMode::CapturePermanently; }
+	virtual EMouseCaptureMode GetMouseCaptureMode() const { return EMouseCaptureMode::CapturePermanently; }
+
+	UE_DEPRECATED(4.26, "Please call GetMouseCaptureMode() instead.")
+	void CaptureMouseOnClick() { GetMouseCaptureMode(); }
 
 	/**
 	 * Gets whether or not the viewport captures the Mouse on launch of the application
@@ -1062,7 +1065,7 @@ public:
 	/**
 	 * Gets whether or not the cursor is hidden when the viewport captures the mouse
 	 */
-	virtual bool HideCursorDuringCapture() { return false; }
+	virtual bool HideCursorDuringCapture() const { return false; }
 
 	/** 
 	 * Should we make new windows for popups or create an overlay in the current window.
