Fix for XBox compiler only warning in deprecating IsBuildInProgress()

Runtime\NavigationSystem\Public\NavMesh/RecastNavMeshGenerator.h(612): warning C4996: 'FNavDataGenerator::IsBuildInProgress': This function is deprecated. Please use IsBuildInProgressCheckDirty Please update your code to the new API before upgrading to the next release, otherwise your project will no longer compile.
  Runtime\Engine\Public\AI/NavDataGenerator.h(36): note: see declaration of 'FNavDataGenerator::IsBuildInProgress'

diff --git a/Engine/Source/Runtime/Engine/Public/AI/NavDataGenerator.h b/Engine/Source/Runtime/Engine/Public/AI/NavDataGenerator.h
--- a/Engine/Source/Runtime/Engine/Public/AI/NavDataGenerator.h
+++ b/Engine/Source/Runtime/Engine/Public/AI/NavDataGenerator.h
@@ -32,8 +32,8 @@ public:
 	virtual void RebuildDirtyAreas(const TArray<FNavigationDirtyArea>& DirtyAreas) {}
 
 	/** determines whether this generator is performing navigation building actions at the moment*/
-	//UE_DEPRECATED(4.26, "This function is deprecated. Please use IsBuildInProgressCheckDirty")
-	virtual bool IsBuildInProgress(bool bCheckDirtyToo = false) const { return false; }
+	UE_DEPRECATED(4.26, "This function is deprecated. Please use IsBuildInProgressCheckDirty")
+	virtual bool IsBuildInProgress(bool bCheckDirtyToo = false) const { return IsBuildInProgressCheckDirty(); }
 
 	/** determines whether this generator is performing navigation building actions at the moment, dirty areas are also checked */
 	virtual bool IsBuildInProgressCheckDirty() const { return false; }
diff --git a/Engine/Source/Runtime/NavigationSystem/Private/NavMesh/RecastNavMeshGenerator.cpp b/Engine/Source/Runtime/NavigationSystem/Private/NavMesh/RecastNavMeshGenerator.cpp
--- a/Engine/Source/Runtime/NavigationSystem/Private/NavMesh/RecastNavMeshGenerator.cpp
+++ b/Engine/Source/Runtime/NavigationSystem/Private/NavMesh/RecastNavMeshGenerator.cpp
@@ -6174,12 +6174,6 @@ bool FRecastNavMeshGenerator::IsBuildInProgressCheckDirty() const
 		|| SyncTimeSlicedData.TileGeneratorSync.IsValid();
 }
 
-//deprecated
-bool FRecastNavMeshGenerator::IsBuildInProgress(bool bCheckDirtyToo) const
-{
-	return IsBuildInProgressCheckDirty();
-}
-
 #if !RECAST_ASYNC_REBUILDING
 bool FRecastNavMeshGenerator::GetTimeSliceData(int32& OutNumRemainingBuildTasks, double& OutCurrentBuildTaskDuration) const
 {
diff --git a/Engine/Source/Runtime/NavigationSystem/Public/NavMesh/RecastNavMeshGenerator.h b/Engine/Source/Runtime/NavigationSystem/Public/NavMesh/RecastNavMeshGenerator.h
--- a/Engine/Source/Runtime/NavigationSystem/Public/NavMesh/RecastNavMeshGenerator.h
+++ b/Engine/Source/Runtime/NavigationSystem/Public/NavMesh/RecastNavMeshGenerator.h
@@ -607,11 +607,6 @@ public:
 	/** Asks generator to update navigation affected by DirtyAreas */
 	virtual void RebuildDirtyAreas(const TArray<FNavigationDirtyArea>& DirtyAreas) override;
 
-	PRAGMA_DISABLE_DEPRECATION_WARNINGS
-	//UE_DEPRECATED(4.26, "This function is deprecated. Please use IsBuildInProgressCheckDirty")
-	virtual bool IsBuildInProgress(bool bCheckDirtyToo = false) const override;
-	PRAGMA_ENABLE_DEPRECATION_WARNINGS
-
 	/** determines whether this generator is performing navigation building actions at the moment, dirty areas are also checked */
 	virtual bool IsBuildInProgressCheckDirty() const override;
 
