NetworkProfiler
* Added exclusive tracking to NetworkProfiler operations so we can spot how much it costs and also exclude it from the ReplicatActor cost

diff --git a/Engine/Source/Runtime/Engine/Private/NetworkProfiler.cpp b/Engine/Source/Runtime/Engine/Private/NetworkProfiler.cpp
--- a/Engine/Source/Runtime/Engine/Private/NetworkProfiler.cpp
+++ b/Engine/Source/Runtime/Engine/Private/NetworkProfiler.cpp
@@ -15,6 +15,8 @@
 #include "HAL/IConsoleManager.h"
 #include "Engine/Public/TimerManager.h"
 
+CSV_DEFINE_CATEGORY(NetworkProfiler, true);
+
 #if USE_NETWORK_PROFILER
 
 #define SCOPE_LOCK_REF(X) FScopeLock ScopeLock(&X);
diff --git a/Engine/Source/Runtime/Engine/Private/RepLayout.cpp b/Engine/Source/Runtime/Engine/Private/RepLayout.cpp
--- a/Engine/Source/Runtime/Engine/Private/RepLayout.cpp
+++ b/Engine/Source/Runtime/Engine/Private/RepLayout.cpp
@@ -1521,6 +1521,7 @@ ERepLayoutResult FRepLayout::CompareProperties(
 		const uint32 ReplicateParentPropertiesStartTime = SharedParams.bIsNetworkProfilerActive ? FPlatformTime::Cycles() : 0;
 		if (SharedParams.bIsNetworkProfilerActive)
 		{
+			CSV_SCOPED_TIMING_STAT_EXCLUSIVE(NetworkProfiler);
 			SharedParams.PropertiesChanged.Init(false, Parents.Num());
 			SharedParams.PropertiesCompared.Init(false, Parents.Num());
 		}
diff --git a/Engine/Source/Runtime/Engine/Public/Net/NetworkProfiler.h b/Engine/Source/Runtime/Engine/Public/Net/NetworkProfiler.h
--- a/Engine/Source/Runtime/Engine/Public/Net/NetworkProfiler.h
+++ b/Engine/Source/Runtime/Engine/Public/Net/NetworkProfiler.h
@@ -18,9 +18,16 @@ class FOutBunch;
 class UNetConnection;
 struct FURL;
 
+CSV_DECLARE_CATEGORY_EXTERN(NetworkProfiler);
+
 #if USE_NETWORK_PROFILER 
 
-#define NETWORK_PROFILER( x ) if ( GNetworkProfiler.IsTrackingEnabled() ) { x; }
+#define NETWORK_PROFILER( x ) \
+	if ( GNetworkProfiler.IsTrackingEnabled() ) \
+	{ \
+		CSV_SCOPED_TIMING_STAT_EXCLUSIVE(NetworkProfiler); \
+		x; \
+	}
 
 DECLARE_MULTICAST_DELEGATE_OneParam(FOnNetworkProfileFinished, const FString& /*Filename */);
 
