- Fixed analytics to compute session duration and idle time consistenly

diff --git a/Engine/Source/Runtime/Engine/Private/Analytics/EditorSessionSummaryWriter.cpp b/Engine/Source/Runtime/Engine/Private/Analytics/EditorSessionSummaryWriter.cpp
--- a/Engine/Source/Runtime/Engine/Private/Analytics/EditorSessionSummaryWriter.cpp
+++ b/Engine/Source/Runtime/Engine/Private/Analytics/EditorSessionSummaryWriter.cpp
@@ -39,7 +39,7 @@ namespace EditorSessionWriterDefs
 FEditorSessionSummaryWriter::FEditorSessionSummaryWriter() :
 	CurrentSession(nullptr)
 	, StartupSeconds(0.0)
-	, LastUserInteractionTime(0.0f)
+	, IdleSeconds(0.0f)
 	, HeartbeatTimeElapsed(0.0f)
 	, bShutdown(false)
 {
@@ -91,11 +91,7 @@ void FEditorSessionSummaryWriter::InitializeSessions()
 void FEditorSessionSummaryWriter::UpdateTimestamps()
 {
 	CurrentSession->Timestamp = FDateTime::UtcNow();
-
-	const double CurrentSeconds = FPlatformTime::Seconds();
-	CurrentSession->SessionDuration = FMath::FloorToInt(CurrentSeconds - StartupSeconds);
-
-	const double IdleSeconds = CurrentSeconds - LastUserInteractionTime;
+	CurrentSession->SessionDuration = FMath::FloorToInt(static_cast<float>((CurrentSession->Timestamp - CurrentSession->StartupTimestamp).GetTotalSeconds()));
 
 	// 1 + 1 minutes
 	if (IdleSeconds > (60 + 60))
@@ -123,8 +119,9 @@ void FEditorSessionSummaryWriter::Tick(float DeltaTime)
 		return;
 	}
 
-	// cache the last user interaction time so that during a crash we have access to it
-	LastUserInteractionTime = FSlateApplication::Get().GetLastUserInteractionTime();
+	// Note: Update idle time is Tick() because Slate cannot be invoked from any thread and UpdateTimeStamps() can be called from any crashing thread.
+	//       Compute the idle time from Slate point-of view. Note that some tasks blocking the UI (such as importing large assets) may be considered idle time.
+	IdleSeconds = FSlateApplication::Get().GetCurrentTime() - FSlateApplication::Get().GetLastUserInteractionTime();
 
 	HeartbeatTimeElapsed += DeltaTime;
 
diff --git a/Engine/Source/Runtime/Engine/Private/Analytics/EditorSessionSummaryWriter.h b/Engine/Source/Runtime/Engine/Private/Analytics/EditorSessionSummaryWriter.h
--- a/Engine/Source/Runtime/Engine/Private/Analytics/EditorSessionSummaryWriter.h
+++ b/Engine/Source/Runtime/Engine/Private/Analytics/EditorSessionSummaryWriter.h
@@ -38,7 +38,7 @@ private:
 	FEditorAnalyticsSession* CurrentSession;
 	FString CurrentSessionSectionName;
 	double StartupSeconds;
-	double LastUserInteractionTime;
+	double IdleSeconds;
 	float HeartbeatTimeElapsed;
 	bool bShutdown;
 };
