Tentative fix for animation engine test. Make sure the skeletalmesh DDC guid is unique and not an old one set by a robomerge.
#rnx

diff --git a/Engine/Source/Runtime/Engine/Private/SkeletalMeshRenderData.cpp b/Engine/Source/Runtime/Engine/Private/SkeletalMeshRenderData.cpp
--- a/Engine/Source/Runtime/Engine/Private/SkeletalMeshRenderData.cpp
+++ b/Engine/Source/Runtime/Engine/Private/SkeletalMeshRenderData.cpp
@@ -59,7 +59,7 @@ static void SerializeLODInfoForDDC(USkeletalMesh* SkeletalMesh, FString& KeySuff
 // differences, etc.) replace the version GUID below with a new one.
 // In case of merge conflicts with DDC versions, you MUST generate a new GUID
 // and set this new GUID as the version.                                       
-#define SKELETALMESH_DERIVEDDATA_VER TEXT("DEE857ABC2804080B6278513ABC358BB")
+#define SKELETALMESH_DERIVEDDATA_VER TEXT("8972F79737FD445387F6914AC5D6374D")
 
 static const FString& GetSkeletalMeshDerivedDataVersion()
 {
