fix for ensure hit on servers when USoundWave::RetainCompressedAudio is called.

diff --git a/Engine/Source/Runtime/Engine/Private/SoundWave.cpp b/Engine/Source/Runtime/Engine/Private/SoundWave.cpp
--- a/Engine/Source/Runtime/Engine/Private/SoundWave.cpp
+++ b/Engine/Source/Runtime/Engine/Private/SoundWave.cpp
@@ -2483,11 +2483,13 @@ void USoundWave::RetainCompressedAudio(bool bForceSync /*= false*/)
 {
 	// Since the zeroth chunk is always inlined and stored in memory,
 	// early exit if we only have one chunk.
-	if (GIsEditor || !IsStreaming() || DisableRetainingCVar || GetNumChunks() <= 1)
+	if (GIsEditor || IsTemplate() || IsRunningDedicatedServer() || !IsStreaming() || DisableRetainingCVar || GetNumChunks() <= 1)
 	{
 		return;
 	}
 
+	// If the first chunk is already loaded and being retained,
+	// don't kick off another load.
 	if (FirstChunk.IsValid())
 	{
 		return;
