Fix NetworkProfiler compilation issue with non-unity builds

diff --git a/Engine/Source/Runtime/Engine/Public/Net/NetworkProfiler.h b/Engine/Source/Runtime/Engine/Public/Net/NetworkProfiler.h
--- a/Engine/Source/Runtime/Engine/Public/Net/NetworkProfiler.h
+++ b/Engine/Source/Runtime/Engine/Public/Net/NetworkProfiler.h
@@ -11,6 +11,7 @@
 #include "IPAddress.h"
 #include "Containers/BitArray.h"
 #include "Serialization/MemoryWriter.h"
+#include "ProfilingDebugging/CsvProfiler.h"
 #include "Misc/ScopeLock.h"
 
 class AActor;
@@ -18,10 +19,10 @@ class FOutBunch;
 class UNetConnection;
 struct FURL;
 
-CSV_DECLARE_CATEGORY_EXTERN(NetworkProfiler);
-
 #if USE_NETWORK_PROFILER 
 
+CSV_DECLARE_CATEGORY_EXTERN(NetworkProfiler);
+
 #define NETWORK_PROFILER( x ) \
 	if ( GNetworkProfiler.IsTrackingEnabled() ) \
 	{ \
