Update CopyMaterialUniformParametersInternal to work without requiring a FMaterialResource
- Some tests run this during 'Construction Script', which happens at cook-time, when FMaterialResource may not be availiable

diff --git a/Engine/Source/Runtime/Engine/Private/Materials/MaterialInstance.cpp b/Engine/Source/Runtime/Engine/Private/Materials/MaterialInstance.cpp
--- a/Engine/Source/Runtime/Engine/Private/Materials/MaterialInstance.cpp
+++ b/Engine/Source/Runtime/Engine/Private/Materials/MaterialInstance.cpp
@@ -4876,50 +4876,42 @@ void UMaterialInstance::CopyMaterialUniformParametersInternal(UMaterialInterface
 				checkSlow(TextureParameterValues.Num() == 0);
 				checkSlow(RuntimeVirtualTextureParameterValues.Num() == 0);
 
-				const FMaterialResource* MaterialResource = nullptr;
-				if (UWorld* World = AsMaterial->GetWorld())
+				const FMaterialCachedParameters& CachedParameters = AsMaterial->GetCachedExpressionData().Parameters;
+
+				// Scalars
+				ScalarParameterValues.Empty(CachedParameters.GetNumParameters(EMaterialParameterType::Scalar));
+				for (int32 ParameterIndex = 0; ParameterIndex < CachedParameters.GetNumParameters(EMaterialParameterType::Scalar); ++ParameterIndex)
 				{
-					MaterialResource = AsMaterial->GetMaterialResource(World->FeatureLevel);
+					FScalarParameterValue* ParameterValue = new(ScalarParameterValues) FScalarParameterValue;
+					ParameterValue->ParameterInfo.Name = CachedParameters.GetParameterName(EMaterialParameterType::Scalar, ParameterIndex);
+					ParameterValue->ParameterValue = CachedParameters.ScalarValues[ParameterIndex];
 				}
 
-				if (!MaterialResource)
+				// Vectors
+				VectorParameterValues.Empty(CachedParameters.GetNumParameters(EMaterialParameterType::Vector));
+				for (int32 ParameterIndex = 0; ParameterIndex < CachedParameters.GetNumParameters(EMaterialParameterType::Vector); ++ParameterIndex)
 				{
-					MaterialResource = AsMaterial->GetMaterialResource(GMaxRHIFeatureLevel);
+					FVectorParameterValue* ParameterValue = new(VectorParameterValues) FVectorParameterValue;
+					ParameterValue->ParameterInfo.Name = CachedParameters.GetParameterName(EMaterialParameterType::Vector, ParameterIndex);
+					ParameterValue->ParameterValue = CachedParameters.VectorValues[ParameterIndex];
 				}
 
-				if (MaterialResource)
+				// Textures
+				TextureParameterValues.Empty(CachedParameters.GetNumParameters(EMaterialParameterType::Texture));
+				for (int32 ParameterIndex = 0; ParameterIndex < CachedParameters.GetNumParameters(EMaterialParameterType::Texture); ++ParameterIndex)
 				{
-					// Scalars
-					const TArray<FMaterialScalarParameterInfo, FMemoryImageAllocator>& ScalarExpressions = MaterialResource->GetUniformScalarParameterExpressions();
-					for (const FMaterialScalarParameterInfo& Parameter : ScalarExpressions)
-					{
-						FScalarParameterValue* ParameterValue = new(ScalarParameterValues) FScalarParameterValue;
-						ParameterValue->ParameterInfo.Name = *Parameter.ParameterName;
-						Parameter.GetDefaultValue(ParameterValue->ParameterValue);
-					}
-
-					// Vectors
-					const TArray<FMaterialVectorParameterInfo, FMemoryImageAllocator>& VectorExpressions = MaterialResource->GetUniformVectorParameterExpressions();
-					for (const FMaterialVectorParameterInfo& Parameter : VectorExpressions)
-					{
-						FVectorParameterValue* ParameterValue = new(VectorParameterValues) FVectorParameterValue;
-						ParameterValue->ParameterInfo.Name = *Parameter.ParameterName;
-						Parameter.GetDefaultValue(ParameterValue->ParameterValue);
-					}
+					FTextureParameterValue* ParameterValue = new(TextureParameterValues) FTextureParameterValue;
+					ParameterValue->ParameterInfo.Name = CachedParameters.GetParameterName(EMaterialParameterType::Texture, ParameterIndex);
+					ParameterValue->ParameterValue = CachedParameters.TextureValues[ParameterIndex];
+				}
 
-					// Textures
-					for (int32 TypeIndex = 0; TypeIndex < NumMaterialTextureParameterTypes; TypeIndex++)
-					{
-						for (const FMaterialTextureParameterInfo& Parameter : MaterialResource->GetUniformTextureExpressions((EMaterialTextureParameterType)TypeIndex))
-						{
-							if (!Parameter.ParameterInfo.Name.IsNone())
-							{
-								FTextureParameterValue* ParameterValue = new(TextureParameterValues) FTextureParameterValue;
-								ParameterValue->ParameterInfo.Name = *Parameter.ParameterName;
-								Parameter.GetGameThreadTextureValue(AsMaterial, *MaterialResource, ParameterValue->ParameterValue);
-							}
-						}
-					}
+				// Runtime Virtual Textures
+				RuntimeVirtualTextureParameterValues.Empty(CachedParameters.GetNumParameters(EMaterialParameterType::RuntimeVirtualTexture));
+				for (int32 ParameterIndex = 0; ParameterIndex < CachedParameters.GetNumParameters(EMaterialParameterType::RuntimeVirtualTexture); ++ParameterIndex)
+				{
+					FRuntimeVirtualTextureParameterValue* ParameterValue = new(RuntimeVirtualTextureParameterValues) FRuntimeVirtualTextureParameterValue;
+					ParameterValue->ParameterInfo.Name = CachedParameters.GetParameterName(EMaterialParameterType::RuntimeVirtualTexture, ParameterIndex);
+					ParameterValue->ParameterValue = CachedParameters.RuntimeVirtualTextureValues[ParameterIndex];
 				}
 			}
 		}
diff --git a/Engine/Source/Runtime/Engine/Public/MaterialCachedData.h b/Engine/Source/Runtime/Engine/Public/MaterialCachedData.h
--- a/Engine/Source/Runtime/Engine/Public/MaterialCachedData.h
+++ b/Engine/Source/Runtime/Engine/Public/MaterialCachedData.h
@@ -109,6 +109,8 @@ struct FMaterialCachedParameters
 	GENERATED_USTRUCT_BODY()
 
 	inline int32 GetNumParameters(EMaterialParameterType Type) const { return Entries[(int32)Type].ParameterInfos.Num(); }
+	inline const FName& GetParameterName(EMaterialParameterType Type, int32 Index) const { return Entries[(int32)Type].ParameterInfos[Index].Name; }
+
 	int32 FindParameterIndex(EMaterialParameterType Type, const FHashedMaterialParameterInfo& HashedParameterInfo, bool bOveriddenOnly) const;
 	int32 FindParameterIndex(EMaterialParameterType Type, const FHashedMaterialParameterInfo& HashedParameterInfo) const;
 	bool IsParameterValid(EMaterialParameterType Type, int32 Index, bool bOveriddenOnly) const;
