Unshelved from pending changelist '12147787':

PR #6778: Move FAudioCommandFence's WaitTimeMs to ConsoleVariable (Contributed by [REDACTED_PERSON])

diff --git a/Engine/Source/Runtime/Engine/Private/AudioThread.cpp b/Engine/Source/Runtime/Engine/Private/AudioThread.cpp
--- a/Engine/Source/Runtime/Engine/Private/AudioThread.cpp
+++ b/Engine/Source/Runtime/Engine/Private/AudioThread.cpp
@@ -45,6 +45,14 @@ static FAutoConsoleVariableRef CVarBatchAudioAsyncBatchSize(
 	TEXT("When AudioThread.EnableBatchProcessing = 1, controls the number of audio commands grouped together for threading.")
 );
 
+
+static int32 GAudioCommandFenceWaitTimeMs = 35;
+TAutoConsoleVariable<int32> CVarAudioCommandFenceWaitTimeMs(
+	TEXT("AudioCommand.FenceWaitTimeMs"),
+	GAudioCommandFenceWaitTimeMs, 
+	TEXT("Sets number of ms for fence wait"), 
+	ECVF_Default);
+
 struct FAudioThreadInteractor
 {
 	static void UseAudioThreadCVarSinkFunction()
@@ -528,12 +536,12 @@ void FAudioCommandFence::Wait(bool bProcessGameThreadTasks) const
 		QUICK_SCOPE_CYCLE_COUNTER(STAT_FAudioCommandFence_Wait);
 
 		bool bDone = false;
-		const uint32 WaitTimeMs = 35;
+
 		do
 		{
 			if (FenceDoneEvent)
 			{
-				bDone = FenceDoneEvent->Wait(WaitTimeMs);
+				bDone = FenceDoneEvent->Wait(GAudioCommandFenceWaitTimeMs);
 			}
 			else
 			{
@@ -548,7 +556,7 @@ void FAudioCommandFence::Wait(bool bProcessGameThreadTasks) const
 
 			// Log how long we've been waiting for the audio thread:
 			float ThisTime = FPlatformTime::Seconds() - StartTime;
- 			if (ThisTime > static_cast<float>(WaitTimeMs) / 1000.0f + SMALL_NUMBER)
+ 			if (ThisTime > static_cast<float>(GAudioCommandFenceWaitTimeMs) / 1000.0f + SMALL_NUMBER)
 			{
 				if (GCVarEnableAudioCommandLogging == 1)
 				{
