Change to use the new SRV creation API to avoid null resource check. Added comments to convey the intention of creating empty SRVs.

diff --git a/Engine/Source/Runtime/Engine/Private/PositionVertexBuffer.cpp b/Engine/Source/Runtime/Engine/Private/PositionVertexBuffer.cpp
--- a/Engine/Source/Runtime/Engine/Private/PositionVertexBuffer.cpp
+++ b/Engine/Source/Runtime/Engine/Private/PositionVertexBuffer.cpp
@@ -235,7 +235,9 @@ void FPositionVertexBuffer::InitRHI()
 		bSRV |= (VertexData && VertexData->GetAllowCPUAccess());
 		if(bSRV)
 		{
-			PositionComponentSRV = RHICreateShaderResourceView(VertexData ? VertexBufferRHI : nullptr, 4, PF_R32_FLOAT);
+			// When VertexData is null, this buffer hasn't been streamed in yet. We still need to create a FRHIShaderResourceView which will be
+			// cached in a vertex factory uniform buffer later. The nullptr tells the RHI that the SRV doesn't view on anything yet.
+			PositionComponentSRV = RHICreateShaderResourceView(FShaderResourceViewInitializer(VertexData ? VertexBufferRHI : nullptr, PF_R32_FLOAT));
 		}
 	}
 }
diff --git a/Engine/Source/Runtime/Engine/Private/Rendering/ColorVertexBuffer.cpp b/Engine/Source/Runtime/Engine/Private/Rendering/ColorVertexBuffer.cpp
--- a/Engine/Source/Runtime/Engine/Private/Rendering/ColorVertexBuffer.cpp
+++ b/Engine/Source/Runtime/Engine/Private/Rendering/ColorVertexBuffer.cpp
@@ -425,7 +425,9 @@ void FColorVertexBuffer::InitRHI()
 
 	if (VertexBufferRHI && RHISupportsManualVertexFetch(GMaxRHIShaderPlatform))
 	{
-		ColorComponentsSRV = RHICreateShaderResourceView(VertexData ? VertexBufferRHI : nullptr, 4, PF_R8G8B8A8);
+		// When VertexData is null, this buffer hasn't been streamed in yet. We still need to create a FRHIShaderResourceView which will be
+		// cached in a vertex factory uniform buffer later. The nullptr tells the RHI that the SRV doesn't view on anything yet.
+		ColorComponentsSRV = RHICreateShaderResourceView(FShaderResourceViewInitializer(VertexData ? VertexBufferRHI : nullptr, PF_R8G8B8A8));
 	}
 }
 
diff --git a/Engine/Source/Runtime/Engine/Private/Rendering/SkeletalMeshVertexClothBuffer.cpp b/Engine/Source/Runtime/Engine/Private/Rendering/SkeletalMeshVertexClothBuffer.cpp
--- a/Engine/Source/Runtime/Engine/Private/Rendering/SkeletalMeshVertexClothBuffer.cpp
+++ b/Engine/Source/Runtime/Engine/Private/Rendering/SkeletalMeshVertexClothBuffer.cpp
@@ -116,7 +116,9 @@ void FSkeletalMeshVertexClothBuffer::InitRHI()
 
 	if (VertexBufferRHI)
 	{
-		VertexBufferSRV = RHICreateShaderResourceView(VertexData ? VertexBufferRHI : nullptr, sizeof(FVector4), PF_A32B32G32R32F);
+		// When VertexData is null, this buffer hasn't been streamed in yet. We still need to create a FRHIShaderResourceView which will be
+		// cached in a vertex factory uniform buffer later. The nullptr tells the RHI that the SRV doesn't view on anything yet.
+		VertexBufferSRV = RHICreateShaderResourceView(FShaderResourceViewInitializer(VertexData ? VertexBufferRHI : nullptr, PF_A32B32G32R32F));
 	}
 }
 
diff --git a/Engine/Source/Runtime/Engine/Private/Rendering/StaticMeshVertexBuffer.cpp b/Engine/Source/Runtime/Engine/Private/Rendering/StaticMeshVertexBuffer.cpp
--- a/Engine/Source/Runtime/Engine/Private/Rendering/StaticMeshVertexBuffer.cpp
+++ b/Engine/Source/Runtime/Engine/Private/Rendering/StaticMeshVertexBuffer.cpp
@@ -336,17 +336,19 @@ void FStaticMeshVertexBuffer::InitRHI()
 	TexCoordVertexBuffer.VertexBufferRHI = CreateTexCoordRHIBuffer_RenderThread();
 	if (TangentsVertexBuffer.VertexBufferRHI && (RHISupportsManualVertexFetch(GMaxRHIShaderPlatform) || IsGPUSkinCacheAvailable(GMaxRHIShaderPlatform)))
 	{
-		TangentsSRV = RHICreateShaderResourceView(
+		// When TangentsData is null, this buffer hasn't been streamed in yet. We still need to create a FRHIShaderResourceView which will be
+		// cached in a vertex factory uniform buffer later. The nullptr tells the RHI that the SRV doesn't view on anything yet.
+		TangentsSRV = RHICreateShaderResourceView(FShaderResourceViewInitializer(
 			TangentsData ? TangentsVertexBuffer.VertexBufferRHI : nullptr,
-			GetUseHighPrecisionTangentBasis() ? 8 : 4,
-			GetUseHighPrecisionTangentBasis() ? PF_R16G16B16A16_SNORM : PF_R8G8B8A8_SNORM);
+			GetUseHighPrecisionTangentBasis() ? PF_R16G16B16A16_SNORM : PF_R8G8B8A8_SNORM));
 	}
 	if (TexCoordVertexBuffer.VertexBufferRHI && RHISupportsManualVertexFetch(GMaxRHIShaderPlatform))
 	{
-		TextureCoordinatesSRV = RHICreateShaderResourceView(
+		// When TexcoordData is null, this buffer hasn't been streamed in yet. We still need to create a FRHIShaderResourceView which will be
+		// cached in a vertex factory uniform buffer later. The nullptr tells the RHI that the SRV doesn't view on anything yet.
+		TextureCoordinatesSRV = RHICreateShaderResourceView(FShaderResourceViewInitializer(
 			TexcoordData ? TexCoordVertexBuffer.VertexBufferRHI : nullptr,
-			GetUseFullPrecisionUVs() ? 8 : 4,
-			GetUseFullPrecisionUVs() ? PF_G32R32F : PF_G16R16F);
+			GetUseFullPrecisionUVs() ? PF_G32R32F : PF_G16R16F));
 	}
 }
 
