Fixed crash when deleting pose that was prior to base pose index

diff --git a/Engine/Source/Runtime/Engine/Classes/Animation/PoseAsset.h b/Engine/Source/Runtime/Engine/Classes/Animation/PoseAsset.h
--- a/Engine/Source/Runtime/Engine/Classes/Animation/PoseAsset.h
+++ b/Engine/Source/Runtime/Engine/Classes/Animation/PoseAsset.h
@@ -112,7 +112,7 @@ private:
 #if WITH_EDITOR
 	void AddOrUpdatePose(const FSmartName& InPoseName, const TArray<FTransform>& InlocalSpacePose, const TArray<float>& InCurveData);
 	void RenamePose(FSmartName OldPoseName, FSmartName NewPoseName);
-	bool DeletePose(FSmartName PoseName);
+	int32 DeletePose(FSmartName PoseName);
 	bool DeleteCurve(FSmartName CurveName);
 	bool InsertTrack(const FName& InTrackName, USkeleton* InSkeleton, FName& InRetargetSourceName);
 	
diff --git a/Engine/Source/Runtime/Engine/Private/Animation/PoseAsset.cpp b/Engine/Source/Runtime/Engine/Private/Animation/PoseAsset.cpp
--- a/Engine/Source/Runtime/Engine/Private/Animation/PoseAsset.cpp
+++ b/Engine/Source/Runtime/Engine/Private/Animation/PoseAsset.cpp
@@ -193,17 +193,17 @@ void FPoseDataContainer::RenamePose(FSmartName OldPoseName, FSmartName NewPoseNa
 	}
 }
 
-bool FPoseDataContainer::DeletePose(FSmartName PoseName)
+int32 FPoseDataContainer::DeletePose(FSmartName PoseName)
 {
 	int32 PoseIndex = PoseNames.Find(PoseName);
 	if (PoseIndex != INDEX_NONE)
 	{
 		PoseNames.RemoveAt(PoseIndex);
 		Poses.RemoveAt(PoseIndex);
-		return true;
+		return PoseIndex;
 	}
 
-	return false;
+	return INDEX_NONE;
 }
 
 bool FPoseDataContainer::DeleteCurve(FSmartName CurveName)
@@ -1291,10 +1291,23 @@ int32 UPoseAsset::DeletePoses(TArray<FName> PoseNamesToDelete)
 	for (const auto& PoseName : PoseNamesToDelete)
 	{
 		FSmartName PoseSmartName;
-		if (MySkeleton->GetSmartNameByName(USkeleton::AnimCurveMappingName, PoseName, PoseSmartName) 
-			&& 	PoseContainer.DeletePose(PoseSmartName))
+		if (MySkeleton->GetSmartNameByName(USkeleton::AnimCurveMappingName, PoseName, PoseSmartName))
 		{
-			++ItemsDeleted;
+			int32 PoseIndexDeleted = PoseContainer.DeletePose(PoseSmartName);
+			if (PoseIndexDeleted != INDEX_NONE)
+			{
+				++ItemsDeleted;
+				// if base pose index is same as pose index deleted
+				if (BasePoseIndex == PoseIndexDeleted)
+				{
+					BasePoseIndex = INDEX_NONE;
+				}
+				// if base pose index is after this, we reduce the number
+				else if (BasePoseIndex > PoseIndexDeleted)
+				{
+					--BasePoseIndex;
+				}
+			}
 		}
 	}
 
