Discard MaterialAttributeID so we can share the stack internals.
Adopted from UDN request: https://udn.unrealengine.com/questions/534573/nested-material-functions-with-makebreak-material.html

#rnx

diff --git a/Engine/Source/Runtime/Engine/Private/Materials/HLSLMaterialTranslator.cpp b/Engine/Source/Runtime/Engine/Private/Materials/HLSLMaterialTranslator.cpp
--- a/Engine/Source/Runtime/Engine/Private/Materials/HLSLMaterialTranslator.cpp
+++ b/Engine/Source/Runtime/Engine/Private/Materials/HLSLMaterialTranslator.cpp
@@ -2616,7 +2616,8 @@ int32 FHLSLMaterialTranslator::CallExpression(FMaterialExpressionKey ExpressionK
 		if (FunctionCall)
 		{
 			FMaterialExpressionKey ReuseCompileStateExpressionKey = ExpressionKey;
- 			ReuseCompileStateExpressionKey.OutputIndex = INDEX_NONE; // Discard the output so we can share the stack internals
+			ReuseCompileStateExpressionKey.OutputIndex = INDEX_NONE; // Discard the output so we can share the stack internals
+			ReuseCompileStateExpressionKey.MaterialAttributeID = FGuid(0, 0, 0, 0); //Discard the Material Attribute ID so we can share the stack internals
 
 			FMaterialFunctionCompileState* SharedFunctionState = CurrentFunctionState->FindOrAddSharedFunctionState(ReuseCompileStateExpressionKey, FunctionCall);
 			FunctionCall->SetSharedCompileState(SharedFunctionState);
