disconnect soundfield submixes if their format changes to something that is incompatible with a soundfield submix they are currently connected to.

diff --git a/Engine/Source/Editor/AudioEditor/Private/SubmixDetailsCustomization.cpp b/Engine/Source/Editor/AudioEditor/Private/SubmixDetailsCustomization.cpp
--- a/Engine/Source/Editor/AudioEditor/Private/SubmixDetailsCustomization.cpp
+++ b/Engine/Source/Editor/AudioEditor/Private/SubmixDetailsCustomization.cpp
@@ -111,7 +111,7 @@ void FSoundfieldSubmixDetailsCustomization::CustomizeDetails(IDetailLayoutBuilde
 		// void OnNameChanged(FName InName)
 		[EncodingFormatPropertyHandle](FName SelectedName)
 		{
-			if (EncodingFormatPropertyHandle.IsValid())
+			if (EncodingFormatPropertyHandle.IsValid() && EncodingFormatPropertyHandle->IsValidHandle())
 			{
 				EncodingFormatPropertyHandle->SetValue(SelectedName);
 			}
@@ -121,7 +121,7 @@ void FSoundfieldSubmixDetailsCustomization::CustomizeDetails(IDetailLayoutBuilde
 		{
 			FName SelectedName;
 
-			if (EncodingFormatPropertyHandle.IsValid())
+			if (EncodingFormatPropertyHandle.IsValid() && EncodingFormatPropertyHandle->IsValidHandle())
 			{
 				EncodingFormatPropertyHandle->GetValue(SelectedName);
 			}
diff --git a/Engine/Source/Runtime/Engine/Classes/Sound/SoundSubmix.h b/Engine/Source/Runtime/Engine/Classes/Sound/SoundSubmix.h
--- a/Engine/Source/Runtime/Engine/Classes/Sound/SoundSubmix.h
+++ b/Engine/Source/Runtime/Engine/Classes/Sound/SoundSubmix.h
@@ -277,6 +277,15 @@ public:
 
 	// Traverses parent submixes until we find a submix that specifies encoding settings.
 	const USoundfieldEncodingSettingsBase* GetEncodingSettings() const;
+
+	// This function goes through every child submix and the parent submix to ensure that they have a compatible format with this  submix's format.
+	void SanitizeLinks();
+
+protected:
+#if WITH_EDITOR
+	virtual void PostEditChangeProperty(struct FPropertyChangedEvent& PropertyChangedEvent) override;
+#endif
+
 };
 
 /**
@@ -332,6 +341,9 @@ public:
 	*/
 	bool RecurseCheckChild(const USoundSubmix* ChildSoundSubmix) const;
 
+	// This function goes through every child submix and the parent submix to ensure that they have a compatible format with this  submix's format.
+	void SanitizeLinks();
+
 	//TODO: Make this editable only if EndpointType is non-default,
 	// and filter classes based on ISoundfieldFactory::GetCustomSettingsClass().
 	UPROPERTY(EditAnywhere, Category = Endpoint)
@@ -345,4 +357,18 @@ public:
 
 	UPROPERTY(EditAnywhere, Category = Soundfield)
 	TArray<USoundfieldEffectBase*> SoundfieldEffectChain;
+
+protected:
+#if WITH_EDITOR
+	virtual void PostEditChangeProperty(struct FPropertyChangedEvent& PropertyChangedEvent) override;
+#endif
 };
+
+namespace SubmixUtils
+{
+	ENGINE_API bool AreSubmixFormatsCompatible(const USoundSubmixBase* ChildSubmix, const USoundSubmixBase* ParentSubmix);
+
+#if WITH_EDITOR
+	ENGINE_API void RefreshEditorForSubmix(const USoundSubmixBase* InSubmix);
+#endif
+}
diff --git a/Engine/Source/Runtime/Engine/Private/SoundSubmix.cpp b/Engine/Source/Runtime/Engine/Private/SoundSubmix.cpp
--- a/Engine/Source/Runtime/Engine/Private/SoundSubmix.cpp
+++ b/Engine/Source/Runtime/Engine/Private/SoundSubmix.cpp
@@ -12,6 +12,9 @@
 #if WITH_EDITOR
 #include "Framework/Notifications/NotificationManager.h"
 #include "Widgets/Notifications/SNotificationList.h"
+#include "Subsystems/AssetEditorSubsystem.h"
+#include "Editor.h"
+#include "Async/Async.h"
 #endif // WITH_EDITOR
 
 static int32 ClearBrokenSubmixAssetsCVar = 0;
@@ -643,17 +646,17 @@ FName USoundfieldSubmix::GetSubmixFormat() const
 			return ISoundfieldFactory::GetFormatNameForNoEncoding();
 		}
 		else
-			{
+		{
 			return SoundfieldEncodingFormat;
-			}
+		}
 
 	}
 	else if(ParentSoundfieldSubmix)
-			{
+	{
 		// If this submix matches the format of whatever submix it's plugged into, 
 		// Recurse into the submix graph to find it.
 		return ParentSoundfieldSubmix->GetSubmixFormat();
-			}
+	}
 	else
 	{
 		return ISoundfieldFactory::GetFormatNameForNoEncoding();
@@ -671,16 +674,16 @@ const USoundfieldEncodingSettingsBase* USoundfieldSubmix::GetEncodingSettings()
 		return EncodingSettings;
 	}
 	else if (ParentSoundfieldSubmix && SoundfieldEncodingFormat == ISoundfieldFactory::GetFormatNameForInheritedEncoding())
-		{
+	{
 		// If this submix matches the format of whatever it's plugged into,
 		// Recurse into the submix graph to match it's settings.
 		return ParentSoundfieldSubmix->GetEncodingSettings();
-		}
+	}
 	else if (ISoundfieldFactory* Factory = ISoundfieldFactory::Get(SubmixFormatName))
-		{
+	{
 		// If we don't have any encoding settings, use the default.
 		return Factory->GetDefaultEncodingSettings();
-		}
+	}
 	else
 	{
 		// If we don't have anything, exit.
@@ -688,6 +691,62 @@ const USoundfieldEncodingSettingsBase* USoundfieldSubmix::GetEncodingSettings()
 	}
 }
 
+void USoundfieldSubmix::SanitizeLinks()
+{
+	bool bShouldRefreshGraph = false;
+
+	// Iterate through children and check encoding formats.
+	for (int32 Index = ChildSubmixes.Num() - 1; Index >= 0; Index--)
+	{
+		if (!SubmixUtils::AreSubmixFormatsCompatible(ChildSubmixes[Index], this))
+		{
+			CastChecked<USoundSubmixWithParentBase>(ChildSubmixes[Index])->ParentSubmix = nullptr;
+			ChildSubmixes[Index]->Modify();
+			ChildSubmixes.RemoveAtSwap(Index);
+			bShouldRefreshGraph = true;
+		}
+	}
+
+	// If this submix is now incompatible with the parent submix, disconnect it.
+	if (!SubmixUtils::AreSubmixFormatsCompatible(this, ParentSubmix))
+	{
+		ParentSubmix->ChildSubmixes.RemoveSwap(this);
+		ParentSubmix->Modify();
+		ParentSubmix = nullptr;
+		bShouldRefreshGraph = true;
+	}
+
+	if (bShouldRefreshGraph)
+	{
+#if WITH_EDITOR
+		SubmixUtils::RefreshEditorForSubmix(this);
+#endif
+	}
+}
+
+void USoundfieldSubmix::PostEditChangeProperty(struct FPropertyChangedEvent& PropertyChangedEvent)
+{
+	// Whether to clean up now invalid links between submix and refresh the submix graph editor.
+	bool bShouldSanitizeLinks = false;
+
+	if (PropertyChangedEvent.Property != nullptr)
+	{
+		static const FName NAME_SoundfieldFormat(TEXT("SoundfieldEncodingFormat"));
+
+		if (PropertyChangedEvent.Property->GetFName() == NAME_SoundfieldFormat)
+		{
+			bShouldSanitizeLinks = true;
+		}
+	}
+
+	Super::PostEditChangeProperty(PropertyChangedEvent);
+
+	if (bShouldSanitizeLinks)
+	{
+		SanitizeLinks();
+	}
+}
+
 IAudioEndpointFactory* UEndpointSubmix::GetAudioEndpointForSubmix() const
 {
 	return IAudioEndpointFactory::Get(EndpointType);
@@ -717,3 +776,119 @@ TArray<USoundfieldEffectBase*> USoundfieldEndpointSubmix::GetSoundfieldProcessor
 {
 	return SoundfieldEffectChain;
 }
+
+void USoundfieldEndpointSubmix::SanitizeLinks()
+{
+	bool bShouldRefreshEditor = false;
+
+	// Iterate through children and check encoding formats.
+	for (int32 Index = ChildSubmixes.Num() - 1; Index >= 0; Index--)
+	{
+		if (!SubmixUtils::AreSubmixFormatsCompatible(ChildSubmixes[Index], this))
+		{
+			CastChecked<USoundSubmixWithParentBase>(ChildSubmixes[Index])->ParentSubmix = nullptr;
+			ChildSubmixes[Index]->Modify();
+			ChildSubmixes.RemoveAtSwap(Index);
+
+			bShouldRefreshEditor = true;
+		}
+	}
+	
+	if (bShouldRefreshEditor)
+	{
+#if WITH_EDITOR
+		SubmixUtils::RefreshEditorForSubmix(this);
+#endif
+	}
+}
+
+void USoundfieldEndpointSubmix::PostEditChangeProperty(struct FPropertyChangedEvent& PropertyChangedEvent)
+{
+	if (PropertyChangedEvent.Property != nullptr)
+	{
+		static const FName NAME_SoundfieldFormat(TEXT("SoundfieldEndpointType"));
+
+		if (PropertyChangedEvent.Property->GetFName() == NAME_SoundfieldFormat)
+		{
+			// Add this sound class to the parent class if it's not already added
+			SanitizeLinks();
+		}
+	}
+
+	Super::PostEditChangeProperty(PropertyChangedEvent);
+}
+
+ENGINE_API bool SubmixUtils::AreSubmixFormatsCompatible(const USoundSubmixBase* ChildSubmix, const USoundSubmixBase* ParentSubmix)
+{
+	const USoundfieldSubmix* ChildSoundfieldSubmix = Cast<const USoundfieldSubmix>(ChildSubmix);
+	const USoundfieldSubmix* ParentSoundfieldSubmix = Cast<const USoundfieldSubmix>(ParentSubmix);
+
+	// If both the child and parent are soundfield submixes, ensure that their formats are compatible.
+	if (ChildSoundfieldSubmix && ParentSoundfieldSubmix)
+	{
+		ISoundfieldFactory* ChildSoundfieldFactory = ChildSoundfieldSubmix->GetSoundfieldFactoryForSubmix();
+		ISoundfieldFactory* ParentSoundfieldFactory = ParentSoundfieldSubmix->GetSoundfieldFactoryForSubmix();
+
+		if (ChildSoundfieldFactory && ParentSoundfieldFactory)
+		{
+			return ChildSoundfieldFactory->CanTranscodeToSoundfieldFormat(ParentSoundfieldFactory->GetSoundfieldFormatName(), *(ParentSoundfieldSubmix->GetSoundfieldEncodingSettings()->GetProxy()))
+				|| ParentSoundfieldFactory->CanTranscodeFromSoundfieldFormat(ChildSoundfieldFactory->GetSoundfieldFormatName(), *(ChildSoundfieldSubmix->GetSoundfieldEncodingSettings()->GetProxy()));
+		}
+		else
+		{
+			return true;
+		}
+	}
+
+	const USoundfieldEndpointSubmix* ParentSoundfieldEndpointSubmix = Cast<const USoundfieldEndpointSubmix>(ParentSubmix);
+
+	// If the child is a soundfield submix and the parent is a soundfield endpoint submix, ensure that they have compatible formats.
+	if (ChildSoundfieldSubmix && ParentSoundfieldEndpointSubmix)
+	{
+		ISoundfieldFactory* ChildSoundfieldFactory = ChildSoundfieldSubmix->GetSoundfieldFactoryForSubmix();
+		ISoundfieldFactory* ParentSoundfieldFactory = ParentSoundfieldEndpointSubmix->GetSoundfieldEndpointForSubmix();
+
+		if (ChildSoundfieldFactory && ParentSoundfieldFactory)
+		{
+			return ChildSoundfieldFactory->CanTranscodeToSoundfieldFormat(ParentSoundfieldFactory->GetSoundfieldFormatName(), *(ParentSoundfieldSubmix->GetSoundfieldEncodingSettings()->GetProxy()))
+				|| ParentSoundfieldFactory->CanTranscodeFromSoundfieldFormat(ChildSoundfieldFactory->GetSoundfieldFormatName(), *(ChildSoundfieldSubmix->GetSoundfieldEncodingSettings()->GetProxy()));
+		}
+		else
+		{
+			return true;
+		}
+	}
+
+	return true;
+}
+
+#if WITH_EDITOR
+
+ENGINE_API void SubmixUtils::RefreshEditorForSubmix(const USoundSubmixBase* InSubmix)
+{
+	if (!GEditor || !InSubmix)
+	{
+		return;
+	}
+
+	TWeakObjectPtr<USoundSubmixBase> WeakSubmix = TWeakObjectPtr<USoundSubmixBase>(const_cast<USoundSubmixBase*>(InSubmix));
+
+	// Since we may be in the middle of a PostEditProperty call,
+	// Dispatch a command to close and reopen the editor window next tick.
+	AsyncTask(ENamedThreads::GameThread, [WeakSubmix]
+	{
+			if (WeakSubmix.IsValid())
+			{
+				UAssetEditorSubsystem* EditorSubsystem = GEditor->GetEditorSubsystem<UAssetEditorSubsystem>();
+				TArray<IAssetEditorInstance*> SubmixEditors = EditorSubsystem->FindEditorsForAsset(WeakSubmix.Get());
+				for (IAssetEditorInstance* Editor : SubmixEditors)
+				{
+					Editor->CloseWindow();
+				}
+
+				EditorSubsystem->OpenEditorForAsset(WeakSubmix.Get());
+			}
+	});
+}
+
+#endif // WITH_EDITOR
