Compile fix in SoundSubmix.cpp for noneditor builds.

diff --git a/Engine/Source/Runtime/Engine/Private/SoundSubmix.cpp b/Engine/Source/Runtime/Engine/Private/SoundSubmix.cpp
--- a/Engine/Source/Runtime/Engine/Private/SoundSubmix.cpp
+++ b/Engine/Source/Runtime/Engine/Private/SoundSubmix.cpp
@@ -724,6 +724,8 @@ void USoundfieldSubmix::SanitizeLinks()
 	}
 }
 
+#if WITH_EDITOR
+
 void USoundfieldSubmix::PostEditChangeProperty(struct FPropertyChangedEvent& PropertyChangedEvent)
 {
 	// Whether to clean up now invalid links between submix and refresh the submix graph editor.
@@ -747,6 +749,8 @@ void USoundfieldSubmix::PostEditChangeProperty(struct FPropertyChangedEvent& Pro
 	}
 }
 
+#endif // WITH_EDITOR
+
 IAudioEndpointFactory* UEndpointSubmix::GetAudioEndpointForSubmix() const
 {
 	return IAudioEndpointFactory::Get(EndpointType);
@@ -802,6 +806,8 @@ void USoundfieldEndpointSubmix::SanitizeLinks()
 	}
 }
 
+#if WITH_EDITOR
+
 void USoundfieldEndpointSubmix::PostEditChangeProperty(struct FPropertyChangedEvent& PropertyChangedEvent)
 {
 	if (PropertyChangedEvent.Property != nullptr)
@@ -818,6 +824,8 @@ void USoundfieldEndpointSubmix::PostEditChangeProperty(struct FPropertyChangedEv
 	Super::PostEditChangeProperty(PropertyChangedEvent);
 }
 
+#endif // WITH_EDITOR
+
 ENGINE_API bool SubmixUtils::AreSubmixFormatsCompatible(const USoundSubmixBase* ChildSubmix, const USoundSubmixBase* ParentSubmix)
 {
 	const USoundfieldSubmix* ChildSoundfieldSubmix = Cast<const USoundfieldSubmix>(ChildSubmix);
