Fix memory corruption caused by force delete
#rnx

diff --git a/Engine/Source/Runtime/Engine/Classes/Engine/SCS_Node.h b/Engine/Source/Runtime/Engine/Classes/Engine/SCS_Node.h
--- a/Engine/Source/Runtime/Engine/Classes/Engine/SCS_Node.h
+++ b/Engine/Source/Runtime/Engine/Classes/Engine/SCS_Node.h
@@ -89,7 +89,8 @@ class USCS_Node : public UObject
 #if WITH_EDITOR
 
 	/** The scene component constructed for component editing in the SCS editor */
-	class USceneComponent* EditorComponentInstance;
+	UPROPERTY(transient)
+	USceneComponent* EditorComponentInstance;
 
 	/** Make sure, the guid is proper - backward compatibility */
 	ENGINE_API void ValidateGuid();
