Fixed our increased low memory GC interval not accounting for extra dev kit memory size. We want the extra memory to be a transparent safety net.

#rnx

diff --git a/Engine/Source/Runtime/Engine/Private/UnrealEngine.cpp b/Engine/Source/Runtime/Engine/Private/UnrealEngine.cpp
--- a/Engine/Source/Runtime/Engine/Private/UnrealEngine.cpp
+++ b/Engine/Source/Runtime/Engine/Private/UnrealEngine.cpp
@@ -1296,7 +1296,10 @@ float UEngine::GetTimeBetweenGarbageCollectionPasses() const
 	if ( GLowMemoryMemoryThresholdMB > 0.0 )
 	{
 		float MBFree = float(FPlatformMemory::GetStats().AvailablePhysical / 1024 / 1024);
-		if ( MBFree <= GLowMemoryMemoryThresholdMB)
+#if !UE_BUILD_SHIPPING
+		MBFree -= float(FPlatformMemory::GetExtraDevelopmentMemorySize() / 1024 / 1024);
+#endif
+		if (MBFree <= GLowMemoryMemoryThresholdMB)
 		{
 			if (GLowMemoryTimeBetweenPurgingPendingKillObjects < TimeBetweenGC)
 			{
