Fix for sounds failing to cull on priority

diff --git a/Engine/Source/Runtime/Engine/Private/SoundConcurrency.cpp b/Engine/Source/Runtime/Engine/Private/SoundConcurrency.cpp
--- a/Engine/Source/Runtime/Engine/Private/SoundConcurrency.cpp
+++ b/Engine/Source/Runtime/Engine/Private/SoundConcurrency.cpp
@@ -200,7 +200,7 @@ void FConcurrencyGroup::CullSoundsDueToMaxConcurrency()
 {
 	// Nothing to do if our active sound count is less than or equal to our max active sounds
 	// or if eviction is supported
-	if (Settings.IsEvictionSupported() || ActiveSounds.Num() <= Settings.MaxCount)
+	if (ActiveSounds.Num() <= Settings.MaxCount)
 	{
 		return;
 	}
@@ -455,9 +455,9 @@ FConcurrencyGroup* FSoundConcurrencyManager::CanPlaySound(const FActiveSound& Ne
 		return nullptr;
 	}
 
-	// StopQuietest doesn't evict, it culls once we instantiate the sound.  This
-	// is because it is not possible to evaluate sound volumes *before* they play.
-	if (ConcurrencyGroup->GetSettings().ResolutionRule == EMaxConcurrentResolutionRule::StopQuietest)
+	// Some settings don't support immediate eviction, they cull once we instantiate the sound.  This
+	// is because it is not possible to evaluate sound volumes/priorities/etc *before* they play.
+	if ( !ConcurrencyGroup->GetSettings().IsEvictionSupported() )
 	{
 		return ConcurrencyGroup;
 	}
