Skin weight profiles
- Ensured that skin weight buffers are not forcefully released in Editor
- Fixed crash when render asset streaming was disabled

diff --git a/Engine/Source/Runtime/Engine/Private/Animation/SkinWeightProfile.cpp b/Engine/Source/Runtime/Engine/Private/Animation/SkinWeightProfile.cpp
--- a/Engine/Source/Runtime/Engine/Private/Animation/SkinWeightProfile.cpp
+++ b/Engine/Source/Runtime/Engine/Private/Animation/SkinWeightProfile.cpp
@@ -71,7 +71,7 @@ static void OnDefaultProfileCVarsChanged(IConsoleVariable* Variable)
 		if (bClearBuffer || bSetBuffer)
 		{
 			// Make sure no pending skeletal mesh LOD updates
-			if (IStreamingManager::Get_Concurrent())
+			if (IStreamingManager::Get_Concurrent() && IStreamingManager::Get().IsRenderAssetStreamingEnabled())
 			{
 				IStreamingManager::Get().GetRenderAssetStreamingManager().BlockTillAllRequestsFinished();
 			}
@@ -416,8 +416,11 @@ void FSkinWeightProfilesData::ClearDynamicDefaultSkinWeightProfile(USkeletalMesh
 	{
 		if (DefaultOverrideSkinWeightBuffer != nullptr)
 		{
+#if !WITH_EDITOR
+			// Only release when not in Editor, as any other viewport / editor could be relying on this buffer
 			ReleaseBuffer(DefaultProfileName, true);
-			DefaultOverrideSkinWeightBuffer = nullptr;			
+#endif // !WITH_EDITOR
+			DefaultOverrideSkinWeightBuffer = nullptr;
 		}
 
 		bDefaultOverriden = false;
