Follow-up fix for ensure hit on server when stream caching is enabled for windows.

diff --git a/Engine/Source/Runtime/Engine/Private/SoundWave.cpp b/Engine/Source/Runtime/Engine/Private/SoundWave.cpp
--- a/Engine/Source/Runtime/Engine/Private/SoundWave.cpp
+++ b/Engine/Source/Runtime/Engine/Private/SoundWave.cpp
@@ -447,7 +447,7 @@ void USoundWave::Serialize( FArchive& Ar )
 #endif // #if WITH_EDITORONLY_DATA
 	}
 
-	if (!GIsEditor && Ar.IsLoading())
+	if (!GIsEditor && !(IsTemplate() || IsRunningDedicatedServer()) &&  Ar.IsLoading())
 	{
 		// For non-editor builds, we can immediately cache the sample rate.
 		SampleRate = GetSampleRateForCurrentPlatform();
@@ -2041,6 +2041,11 @@ bool USoundWave::ShouldUseStreamCaching() const
 
 TArrayView<const uint8> USoundWave::GetZerothChunk(bool bForImmediatePlayback)
 {
+	if(IsTemplate() || IsRunningDedicatedServer())
+	{
+		return TArrayView<const uint8>();
+	}
+	
 	if (ShouldUseStreamCaching())
 	{
 		// In editor, we actually don't have a zeroth chunk until we try to play an audio file.
