Fix another static analysis warning.

#rnx

diff --git a/Engine/Source/Runtime/Engine/Private/RepLayout.cpp b/Engine/Source/Runtime/Engine/Private/RepLayout.cpp
--- a/Engine/Source/Runtime/Engine/Private/RepLayout.cpp
+++ b/Engine/Source/Runtime/Engine/Private/RepLayout.cpp
@@ -6992,7 +6992,11 @@ ERepLayoutResult FRepLayout::DeltaSerializeFastArrayProperty(FFastArrayDeltaSeri
 
 							CompareProperties_r(SharedParams, StackParams, ItemLayoutStart, ItemLayoutEnd, 0);
 
-							if (UNLIKELY(ERepLayoutResult::FatalError == UpdateResult))
+							// NOTE: Currently, SA always throws "Warning: Expression 'ERepLayoutResult::FatalError == UpdateResult' is always false" here.
+							//			It gets tripped up by the indirection / reference semantics of Stack Params and assumes that
+							//			result just remains the same value as it was assigned above.
+							//			Because of this, the error is disabled.
+							if (UNLIKELY(ERepLayoutResult::FatalError == UpdateResult)) // -V547
 							{
 								return UpdateResult;
 							}
