Further fixes for crashes due to long running async anim compression + optimize TallyErrors to significantly reduce anim compression overhead and worker thread blocking

diff --git a/Engine/Source/Runtime/Engine/Private/Animation/AnimBoneCompressionSettings.cpp b/Engine/Source/Runtime/Engine/Private/Animation/AnimBoneCompressionSettings.cpp
--- a/Engine/Source/Runtime/Engine/Private/Animation/AnimBoneCompressionSettings.cpp
+++ b/Engine/Source/Runtime/Engine/Private/Animation/AnimBoneCompressionSettings.cpp
@@ -102,7 +102,7 @@ static void CompressAnimSequenceImpl(FAnimBoneCompressionContext& Context)
 {
 	Context.bSuccess = Context.Codec->Compress(Context.AnimSeq, Context.Result);
 
-	if (Context.bSuccess)
+	if (Context.bSuccess && !Context.AnimSeq.IsCancelled())
 	{
 		FAnimationUtils::ComputeCompressionError(Context.AnimSeq, Context.Result, Context.Result.AnimData->BoneCompressionErrorStats);
 	}
diff --git a/Engine/Source/Runtime/Engine/Private/Animation/AnimCompress_PerTrackCompression.cpp b/Engine/Source/Runtime/Engine/Private/Animation/AnimCompress_PerTrackCompression.cpp
--- a/Engine/Source/Runtime/Engine/Private/Animation/AnimCompress_PerTrackCompression.cpp
+++ b/Engine/Source/Runtime/Engine/Private/Animation/AnimCompress_PerTrackCompression.cpp
@@ -1018,6 +1018,10 @@ void UAnimCompress_PerTrackCompression::CompressUsingUnderlyingCompressor(
 	// Compress each track independently
 	for (int32 TrackIndex = 0; TrackIndex < NumTracks; ++TrackIndex)
 	{
+		if (CompressibleAnimData.IsCancelled())
+		{
+			return;
+		}
 		// Compression parameters / thresholds
 		FPerTrackParams Params;
 		Params.NumFrames = CompressibleAnimData.NumFrames;
diff --git a/Engine/Source/Runtime/Engine/Private/Animation/AnimCompressionDerivedData.cpp b/Engine/Source/Runtime/Engine/Private/Animation/AnimCompressionDerivedData.cpp
--- a/Engine/Source/Runtime/Engine/Private/Animation/AnimCompressionDerivedData.cpp
+++ b/Engine/Source/Runtime/Engine/Private/Animation/AnimCompressionDerivedData.cpp
@@ -35,7 +35,7 @@ const TCHAR* FDerivedDataAnimationCompression::GetVersionString() const
 	// This is a version string that mimics the old versioning scheme. If you
 	// want to bump this version, generate a new guid using VS->Tools->Create GUID and
 	// return it here. Ex.
-	return TEXT("0F1CECE507FE4F89A374B4C8E7B55052");
+	return TEXT("0B129554D9A4427784EC46E2AFD314B7");
 }
 
 bool FDerivedDataAnimationCompression::Build( TArray<uint8>& OutDataArray )
@@ -144,7 +144,7 @@ void FAsyncCompressedAnimationsManagement::Shutdown()
 		
 		const double PreExitDuration = FPlatformTime::Seconds() - PreExitStartTimer;
 
-		if (PreExitDuration > 10.0)
+		if (PreExitDuration > 20.0)
 		{
 			UE_LOG(LogAnimationCompression, Warning, TEXT("Async Compression Pre Init Timer hit, Active Compressions:%i"), ActiveAsyncCompressionTasks.Num());
 			return;
diff --git a/Engine/Source/Runtime/Engine/Private/Animation/AnimationUtils.cpp b/Engine/Source/Runtime/Engine/Private/Animation/AnimationUtils.cpp
--- a/Engine/Source/Runtime/Engine/Private/Animation/AnimationUtils.cpp
+++ b/Engine/Source/Runtime/Engine/Private/Animation/AnimationUtils.cpp
@@ -586,6 +586,72 @@ bool FAnimationUtils::HasUniformKeySpacing(int32 NumFrames, const TArray<float>&
 	return false;
 }
 
+namespace EPerturbationErrorMode
+{
+	enum Type
+	{
+		Transform,
+		Rotation,
+		Scale,
+	};
+}
+
+template<int32>
+FORCEINLINE void CalcErrorValues(FAnimPerturbationError& TrackError, const FTransform& RawTransform, const FTransform& NewTransform);
+
+template<>
+FORCEINLINE void CalcErrorValues<EPerturbationErrorMode::Transform>(FAnimPerturbationError& TrackError, const FTransform& RawTransform, const FTransform& NewTransform)
+{
+	TrackError.MaxErrorInTransDueToTrans = FMath::Max(TrackError.MaxErrorInTransDueToTrans, (RawTransform.GetLocation() - NewTransform.GetLocation()).SizeSquared());
+	TrackError.MaxErrorInRotDueToTrans = FMath::Max(TrackError.MaxErrorInRotDueToTrans, FQuat::ErrorAutoNormalize(RawTransform.GetRotation(), NewTransform.GetRotation()));
+	//TrackError.MaxErrorInScaleDueToTrans = FMath::Max(TrackError.MaxErrorInScaleDueToTrans, (RawTransform.GetScale3D() - NewTransform.GetScale3D()).SizeSquared());
+}
+
+template<>
+FORCEINLINE void CalcErrorValues<EPerturbationErrorMode::Rotation>(FAnimPerturbationError& TrackError, const FTransform& RawTransform, const FTransform& NewTransform)
+{
+	TrackError.MaxErrorInTransDueToRot = FMath::Max(TrackError.MaxErrorInTransDueToRot, (RawTransform.GetLocation() - NewTransform.GetLocation()).SizeSquared());
+	TrackError.MaxErrorInRotDueToRot = FMath::Max(TrackError.MaxErrorInRotDueToRot, FQuat::ErrorAutoNormalize(RawTransform.GetRotation(), NewTransform.GetRotation()));
+	//TrackError.MaxErrorInScaleDueToRot = FMath::Max(TrackError.MaxErrorInScaleDueToRot, (RawTransform.GetScale3D() - NewTransform.GetScale3D()).SizeSquared());
+}
+
+template<>
+FORCEINLINE void CalcErrorValues<EPerturbationErrorMode::Scale>(FAnimPerturbationError& TrackError, const FTransform& RawTransform, const FTransform& NewTransform)
+{
+	/*TrackError.MaxErrorInTransDueToScale = FMath::Max(TrackError.MaxErrorInTransDueToScale, (RawTransform.GetLocation() - NewTransform.GetLocation()).SizeSquared());
+	TrackError.MaxErrorInRotDueToScale = FMath::Max(TrackError.MaxErrorInRotDueToScale, FQuat::ErrorAutoNormalize(RawTransform.GetRotation(), NewTransform.GetRotation()));
+	TrackError.MaxErrorInScaleDueToScale = FMath::Max(TrackError.MaxErrorInScaleDueToScale, (RawTransform.GetScale3D() - NewTransform.GetScale3D()).SizeSquared());*/
+
+	TrackError.MaxErrorInTransDueToScale = TrackError.MaxErrorInTransDueToRot; //Original tally code was calculating scale errors then using rot regardless.
+	TrackError.MaxErrorInRotDueToScale = TrackError.MaxErrorInRotDueToRot;     
+	//TrackError.MaxErrorInScaleDueToScale = TrackError.MaxErrorInScaleDueToRot;
+}
+
+struct FBoneTestItem
+{
+	int32 BoneIndex;
+	bool  bIsEndEffector;
+	FBoneTestItem(int32 InBoneIndex, const TArray<FBoneData>& BoneData)
+		: BoneIndex(InBoneIndex)
+		, bIsEndEffector(BoneData[BoneIndex].IsEndEffector())
+	{}
+};
+
+template<int32 PERTURBATION_ERROR_MODE>
+void CalcErrorsLoop(const TArray<FBoneTestItem>& BonesToTest, const FCompressibleAnimData& CompressibleAnimData, const TArray<FTransform>& RawAtoms, const TArray<FTransform>& RawTransforms, TArray<FTransform>& NewTransforms, FAnimPerturbationError& ThisBoneError)
+{
+	for (const FBoneTestItem& BoneTest : BonesToTest)
+	{
+		const int32 ParentIndex = CompressibleAnimData.RefSkeleton.GetParentIndex(BoneTest.BoneIndex);
+		NewTransforms[BoneTest.BoneIndex] = RawAtoms[BoneTest.BoneIndex] * NewTransforms[ParentIndex];
+
+		if (BoneTest.bIsEndEffector)
+		{
+			CalcErrorValues<PERTURBATION_ERROR_MODE>(ThisBoneError, RawTransforms[BoneTest.BoneIndex], NewTransforms[BoneTest.BoneIndex]);
+		}
+	}
+}
+
 /**
  * Perturbs the bone(s) associated with each track in turn, measuring the maximum error introduced in end effectors as a result
  */
@@ -597,132 +663,176 @@ void FAnimationUtils::TallyErrorsFromPerturbation(
 	const FVector& ScaleNudge,
 	TArray<FAnimPerturbationError>& InducedErrors)
 {
-	const float TimeStep = (float)CompressibleAnimData.SequenceLength / CompressibleAnimData.NumFrames;
 	const int32 NumBones = CompressibleAnimData.BoneData.Num();
 
-
 	const TArray<FTransform>& RefPose = CompressibleAnimData.RefLocalPoses;
 
 	TArray<FTransform> RawAtoms;
-	TArray<FTransform> NewAtomsT;
-	TArray<FTransform> NewAtomsR;
-	TArray<FTransform> NewAtomsS;
 	TArray<FTransform> RawTransforms;
-	TArray<FTransform> NewTransformsT;
-	TArray<FTransform> NewTransformsR;
-	TArray<FTransform> NewTransformsS;
+	TArray<FTransform> NewTransforms;
 
 	RawAtoms.AddZeroed(NumBones);
-	NewAtomsT.AddZeroed(NumBones);
-	NewAtomsR.AddZeroed(NumBones);
-	NewAtomsS.AddZeroed(NumBones);
 	RawTransforms.AddZeroed(NumBones);
-	NewTransformsT.AddZeroed(NumBones);
-	NewTransformsR.AddZeroed(NumBones);
-	NewTransformsS.AddZeroed(NumBones);
+	NewTransforms.AddZeroed(NumBones);
 
 	InducedErrors.AddUninitialized(NumTracks);
 
 	FTransform Perturbation(RotationNudge, PositionNudge, ScaleNudge);
 
-	for (int32 TrackUnderTest = 0; TrackUnderTest < NumTracks; ++TrackUnderTest)
-	{
-		float MaxErrorT_DueToT = 0.0f;
-		float MaxErrorR_DueToT = 0.0f;
-		float MaxErrorS_DueToT = 0.0f;
-		float MaxErrorT_DueToR = 0.0f;
-		float MaxErrorR_DueToR = 0.0f;
-		float MaxErrorS_DueToR = 0.0f;
-		float MaxErrorT_DueToS = 0.0f;
-		float MaxErrorR_DueToS = 0.0f;
-		float MaxErrorS_DueToS = 0.0f;
-		
-		// for each whole increment of time (frame stepping)
-		for (float Time = 0.0f; Time < CompressibleAnimData.SequenceLength; Time += TimeStep)
+	// Build Track bone mapping for processingh
+	struct FTrackBoneMapping
+	{
+		int32 TrackIndex;
+		int32 BoneIndex;
+		bool  bIsEndEffector;
+		FTrackBoneMapping(int32 InTrackIndex, int32 InBoneIndex, const TArray<FBoneData>& BoneData)
+			: TrackIndex(InTrackIndex)
+			, BoneIndex(InBoneIndex)
+			, bIsEndEffector(BoneData[BoneIndex].IsEndEffector())
+		{}
+	};
+
+	TArray<FTrackBoneMapping> TracksAndBonesToTest;
+	TracksAndBonesToTest.Reserve(NumTracks);
+
+	for (int32 TrackToTest = 0; TrackToTest < NumTracks; TrackToTest++)
+	{
+		TracksAndBonesToTest.Emplace(TrackToTest, CompressibleAnimData.TrackToSkeletonMapTable[TrackToTest].BoneTreeIndex, CompressibleAnimData.BoneData);
+	}
+
+	Algo::SortBy(TracksAndBonesToTest, &FTrackBoneMapping::BoneIndex);
+
+	// Prebuild bone test paths
+	TArray<TArray<FBoneTestItem>> BonesToTestMap;
+	BonesToTestMap.AddDefaulted(NumTracks);
+
+	TArray<int32> TempBonesToTest;
+	TempBonesToTest.Reserve(NumBones);
+
+	for (FTrackBoneMapping& TrackBonePair : TracksAndBonesToTest)
+	{
+		TempBonesToTest.Reset();
+		TempBonesToTest.Append(CompressibleAnimData.BoneData[TrackBonePair.BoneIndex].Children);
+
+		int32 BoneToTestIndex = 0;
+		while (BoneToTestIndex < TempBonesToTest.Num())
 		{
-			// get the raw and compressed atom for each bone
-			for (int32 BoneIndex = 0; BoneIndex < NumBones; ++BoneIndex)
+			const int32 ThisBoneIndex = TempBonesToTest[BoneToTestIndex++];
+			TempBonesToTest.Append(CompressibleAnimData.BoneData[ThisBoneIndex].Children);
+		}
+
+		TArray<FBoneTestItem>& BonesToTest = BonesToTestMap[TrackBonePair.TrackIndex];
+		BonesToTest.Reserve(TempBonesToTest.Num());
+
+		for (const int32 BoneToTest : TempBonesToTest)
+		{
+			BonesToTest.Emplace(BoneToTest, CompressibleAnimData.BoneData);
+		}
+	}
+
+
+	for (int32 Frame = 0; Frame < CompressibleAnimData.NumFrames; ++Frame)
+	{
+		//Build Locals For Frame
+		if (CompressibleAnimData.IsCancelled())
+		{
+			return;
+		}
+		RawAtoms = RefPose;
+		for (const FTrackBoneMapping& TrackAndBone : TracksAndBonesToTest)
+		{
+			const FRawAnimSequenceTrack& RawTrack = CompressibleAnimData.RawAnimationData[TrackAndBone.TrackIndex];
+			ExtractTransformForFrameFromTrackSafe(RawTrack, Frame, RawAtoms[TrackAndBone.BoneIndex]);
+		}
+
+		//Build Reference Component Space for Frame
+		RawTransforms[0] = RawAtoms[0];
+		for (int32 BoneIndex = 1; BoneIndex < NumBones; ++BoneIndex)
+		{
+			const int32 ParentIndex = CompressibleAnimData.RefSkeleton.GetParentIndex(BoneIndex);
+			RawTransforms[BoneIndex] = RawAtoms[BoneIndex] * RawTransforms[ParentIndex];
+		}
+
+		//For each track
+		int32 PreviousCopyPoint = 0;
+
+		for (const FTrackBoneMapping& TrackBonePair : TracksAndBonesToTest)
+		{
+			const int32 BoneToModify = TrackBonePair.BoneIndex;
+
+			TArray<FBoneTestItem>& BonesToTest = BonesToTestMap[TrackBonePair.TrackIndex];
+
+			FAnimPerturbationError& ThisBoneError = InducedErrors[TrackBonePair.TrackIndex];
+
+			// Init unchanged cache
+			for (int32 BoneIndex = PreviousCopyPoint; BoneIndex < BoneToModify; ++BoneIndex)
 			{
-				const int32 TrackIndex = GetAnimTrackIndexForSkeletonBone(BoneIndex, CompressibleAnimData.TrackToSkeletonMapTable);
+				NewTransforms[BoneIndex] = RawTransforms[BoneIndex];
+			}
 
-				if (TrackIndex == INDEX_NONE)
-				{
-					// No track for the bone was found, so use the reference pose.
-					RawAtoms[BoneIndex]	= RefPose[BoneIndex];
-					NewAtomsT[BoneIndex] = RawAtoms[BoneIndex];
-					NewAtomsR[BoneIndex] = RawAtoms[BoneIndex];
-					NewAtomsS[BoneIndex] = RawAtoms[BoneIndex];
-				}
-				else
-				{
-					ExtractTransformFromTrack(Time, CompressibleAnimData.NumFrames, CompressibleAnimData.SequenceLength, CompressibleAnimData.RawAnimationData[TrackIndex], CompressibleAnimData.Interpolation, RawAtoms[BoneIndex]);
+			// Modify test bone
+			PreviousCopyPoint = BoneToModify;
+			
+			// Calc Transform Error
+			NewTransforms[BoneToModify] = RawAtoms[BoneToModify];
+			NewTransforms[BoneToModify].AddToTranslation(PositionNudge);
 
-					NewAtomsT[BoneIndex] = RawAtoms[BoneIndex];
-					NewAtomsR[BoneIndex] = RawAtoms[BoneIndex];
-					NewAtomsS[BoneIndex] = RawAtoms[BoneIndex];
+			const int32 ModifiedParentIndex = CompressibleAnimData.RefSkeleton.GetParentIndex(BoneToModify);
 
-					// Perturb the bone under test
-					if (TrackIndex == TrackUnderTest)
-					{
-						NewAtomsT[BoneIndex].AddToTranslation(PositionNudge);
+			if (ModifiedParentIndex != INDEX_NONE) // Put Modified Bone in Component space
+			{
+				NewTransforms[BoneToModify] = NewTransforms[BoneToModify] * NewTransforms[ModifiedParentIndex];
+			}
 
-						FQuat NewR = NewAtomsR[BoneIndex].GetRotation();
-						NewR += RotationNudge;
-						NewR.Normalize();
-						NewAtomsR[BoneIndex].SetRotation(NewR);
+			if (TrackBonePair.bIsEndEffector)
+			{
+				CalcErrorValues<EPerturbationErrorMode::Transform>(ThisBoneError, RawTransforms[BoneToModify], NewTransforms[BoneToModify]);
+			}
+			CalcErrorsLoop<EPerturbationErrorMode::Transform>(BonesToTest, CompressibleAnimData, RawAtoms, RawTransforms, NewTransforms, ThisBoneError);
 
-						FVector Scale3D = NewAtomsS[BoneIndex].GetScale3D();
-						NewAtomsS[BoneIndex].SetScale3D( Scale3D + ScaleNudge);
-					}
-				}
+			// Calc Rotatin Error
+			NewTransforms[BoneToModify] = RawAtoms[BoneToModify];
+			FQuat NewR = RawAtoms[BoneToModify].GetRotation();
+			NewR += RotationNudge;
+			NewR.Normalize();
+			NewTransforms[BoneToModify].SetRotation(NewR);
 
-				RawTransforms[BoneIndex] = RawAtoms[BoneIndex];
-				NewTransformsT[BoneIndex] = NewAtomsT[BoneIndex];
-				NewTransformsR[BoneIndex] = NewAtomsR[BoneIndex];
-				NewTransformsS[BoneIndex] = NewAtomsS[BoneIndex];
+			if (ModifiedParentIndex != INDEX_NONE) // Put Modified Bone in Component space
+			{
+				NewTransforms[BoneToModify] = NewTransforms[BoneToModify] * NewTransforms[ModifiedParentIndex];
+			}
 
-				// For all bones below the root, final component-space transform is relative transform * component-space transform of parent.
-				if ( BoneIndex > 0 )
-				{
-					const int32 ParentIndex = CompressibleAnimData.RefSkeleton.GetParentIndex(BoneIndex);
+			if (TrackBonePair.bIsEndEffector)
+			{
+				CalcErrorValues<EPerturbationErrorMode::Rotation>(ThisBoneError, RawTransforms[BoneToModify], NewTransforms[BoneToModify]);
+			}
+			CalcErrorsLoop<EPerturbationErrorMode::Rotation>(BonesToTest, CompressibleAnimData, RawAtoms, RawTransforms, NewTransforms, ThisBoneError);
 
-					// Check the precondition that parents occur before children in the RequiredBones array.
-					check( ParentIndex != INDEX_NONE );
-					check( ParentIndex < BoneIndex );
+			//Calc Scale Error
+			NewTransforms[BoneToModify] = RawAtoms[BoneToModify];
+			NewTransforms[BoneToModify].SetScale3D(RawAtoms[BoneToModify].GetScale3D() + ScaleNudge);
 
-					RawTransforms[BoneIndex] *= RawTransforms[ParentIndex];
-					NewTransformsT[BoneIndex] *= NewTransformsT[ParentIndex];
-					NewTransformsR[BoneIndex] *= NewTransformsR[ParentIndex];
-					NewTransformsS[BoneIndex] *= NewTransformsS[ParentIndex];
-				}
+			if (ModifiedParentIndex != INDEX_NONE) // Put Modified Bone in Component space
+			{
+				NewTransforms[BoneToModify] = NewTransforms[BoneToModify] * NewTransforms[ModifiedParentIndex];
+			}
 
-				// Only look at the error that occurs in end effectors
-				if (CompressibleAnimData.BoneData[BoneIndex].IsEndEffector())
-				{
-					MaxErrorT_DueToT = FMath::Max(MaxErrorT_DueToT, (RawTransforms[BoneIndex].GetLocation() - NewTransformsT[BoneIndex].GetLocation()).Size());
-					MaxErrorT_DueToR = FMath::Max(MaxErrorT_DueToR, (RawTransforms[BoneIndex].GetLocation() - NewTransformsR[BoneIndex].GetLocation()).Size());
-					MaxErrorT_DueToS = FMath::Max(MaxErrorT_DueToS, (RawTransforms[BoneIndex].GetLocation() - NewTransformsS[BoneIndex].GetLocation()).Size());
-					MaxErrorR_DueToT = FMath::Max(MaxErrorR_DueToT, FQuat::ErrorAutoNormalize(RawTransforms[BoneIndex].GetRotation(), NewTransformsT[BoneIndex].GetRotation()));
-					MaxErrorR_DueToR = FMath::Max(MaxErrorR_DueToR, FQuat::ErrorAutoNormalize(RawTransforms[BoneIndex].GetRotation(), NewTransformsR[BoneIndex].GetRotation()));
-					MaxErrorR_DueToS = FMath::Max(MaxErrorR_DueToS, FQuat::ErrorAutoNormalize(RawTransforms[BoneIndex].GetRotation(), NewTransformsS[BoneIndex].GetRotation()));
-					MaxErrorS_DueToT = FMath::Max(MaxErrorS_DueToT, (RawTransforms[BoneIndex].GetScale3D() - NewTransformsT[BoneIndex].GetScale3D()).Size());
-					MaxErrorS_DueToR = FMath::Max(MaxErrorS_DueToR, (RawTransforms[BoneIndex].GetScale3D() - NewTransformsR[BoneIndex].GetScale3D()).Size());
-					MaxErrorS_DueToS = FMath::Max(MaxErrorS_DueToS, (RawTransforms[BoneIndex].GetScale3D() - NewTransformsS[BoneIndex].GetScale3D()).Size());
-				}
-			} // for each bone
-		} // for each time
-
-		// Save the worst errors
-		FAnimPerturbationError& TrackError = InducedErrors[TrackUnderTest];
-		TrackError.MaxErrorInTransDueToTrans = MaxErrorT_DueToT;
-		TrackError.MaxErrorInRotDueToTrans = MaxErrorR_DueToT;
-		TrackError.MaxErrorInScaleDueToTrans = MaxErrorS_DueToT;
-		TrackError.MaxErrorInTransDueToRot = MaxErrorT_DueToR;
-		TrackError.MaxErrorInRotDueToRot = MaxErrorR_DueToR;
-		TrackError.MaxErrorInScaleDueToRot = MaxErrorS_DueToR;
-		TrackError.MaxErrorInTransDueToScale = MaxErrorT_DueToR;
-		TrackError.MaxErrorInRotDueToScale = MaxErrorR_DueToR;
-		TrackError.MaxErrorInScaleDueToScale = MaxErrorS_DueToR;
+			if (TrackBonePair.bIsEndEffector)
+			{
+				CalcErrorValues<EPerturbationErrorMode::Scale>(ThisBoneError, RawTransforms[BoneToModify], NewTransforms[BoneToModify]);
+			}
+			CalcErrorsLoop<EPerturbationErrorMode::Scale>(BonesToTest, CompressibleAnimData, RawAtoms, RawTransforms, NewTransforms, ThisBoneError);
+		}
+	}
+
+	for (FAnimPerturbationError& Error : InducedErrors)
+	{
+		Error.MaxErrorInTransDueToTrans = FMath::Sqrt(Error.MaxErrorInTransDueToTrans);
+		Error.MaxErrorInTransDueToRot = FMath::Sqrt(Error.MaxErrorInTransDueToRot);
+		Error.MaxErrorInTransDueToScale = FMath::Sqrt(Error.MaxErrorInTransDueToScale);
+		/*Error.MaxErrorInScaleDueToTrans = FMath::Sqrt(Error.MaxErrorInScaleDueToTrans); //Not used by compression code TODO: Investigate and either use or remove
+		Error.MaxErrorInScaleDueToRot = FMath::Sqrt(Error.MaxErrorInScaleDueToRot);
+		Error.MaxErrorInScaleDueToScale = FMath::Sqrt(Error.MaxErrorInScaleDueToScale);*/
 	}
 }
 
@@ -829,7 +939,37 @@ void FAnimationUtils::EnsureAnimSequenceLoaded(UAnimSequence& AnimSeq)
 	EnsureDependenciesAreLoaded(AnimSeq.CurveCompressionSettings);
 }
 
-void FAnimationUtils::ExtractTransformFromTrack(float Time, int32 NumFrames, float SequenceLength, const struct FRawAnimSequenceTrack& RawTrack, EAnimInterpolationType Interpolation, FTransform &OutAtom)
+void FAnimationUtils::ExtractTransformForFrameFromTrackSafe(const FRawAnimSequenceTrack& RawTrack, int32 Frame, FTransform& OutAtom)
+{
+	// Bail out (with rather wacky data) if data is empty for some reason.
+	if (RawTrack.PosKeys.Num() == 0 || RawTrack.RotKeys.Num() == 0)
+	{
+		OutAtom.SetIdentity();
+		return;
+	}
+
+	ExtractTransformForFrameFromTrack(RawTrack, Frame, OutAtom);
+}
+
+void FAnimationUtils::ExtractTransformForFrameFromTrack(const FRawAnimSequenceTrack& RawTrack, int32 Frame, FTransform& OutAtom)
+{
+	static const FVector DefaultScale3D = FVector(1.f);
+
+	const int32 PosKeyIndex1 = FMath::Min(Frame, RawTrack.PosKeys.Num() - 1);
+	const int32 RotKeyIndex1 = FMath::Min(Frame, RawTrack.RotKeys.Num() - 1);
+
+	if (RawTrack.ScaleKeys.Num() > 0)
+	{
+		const int32 ScaleKeyIndex1 = FMath::Min(Frame, RawTrack.ScaleKeys.Num() - 1);
+		OutAtom = FTransform(RawTrack.RotKeys[RotKeyIndex1], RawTrack.PosKeys[PosKeyIndex1], RawTrack.ScaleKeys[ScaleKeyIndex1]);
+	}
+	else
+	{
+		OutAtom = FTransform(RawTrack.RotKeys[RotKeyIndex1], RawTrack.PosKeys[PosKeyIndex1], DefaultScale3D);
+	}
+}
+
+void FAnimationUtils::ExtractTransformFromTrack(float Time, int32 NumFrames, float SequenceLength, const FRawAnimSequenceTrack& RawTrack, EAnimInterpolationType Interpolation, FTransform &OutAtom)
 {
 	// Bail out (with rather wacky data) if data is empty for some reason.
 	if (RawTrack.PosKeys.Num() == 0 || RawTrack.RotKeys.Num() == 0)
@@ -842,7 +982,7 @@ void FAnimationUtils::ExtractTransformFromTrack(float Time, int32 NumFrames, flo
 	float Alpha;
 	FAnimationRuntime::GetKeyIndicesFromTime(KeyIndex1, KeyIndex2, Alpha, Time, NumFrames, SequenceLength);
 	// @Todo fix me: this change is not good, it has lots of branches. But we'd like to save memory for not saving scale if no scale change exists
-	const bool bHasScaleKey = (RawTrack.ScaleKeys.Num() > 0);
+	
 	static const FVector DefaultScale3D = FVector(1.f);
 
 	if (Interpolation == EAnimInterpolationType::Step)
@@ -852,32 +992,12 @@ void FAnimationUtils::ExtractTransformFromTrack(float Time, int32 NumFrames, flo
 
 	if (Alpha <= 0.f)
 	{
-		const int32 PosKeyIndex1 = FMath::Min(KeyIndex1, RawTrack.PosKeys.Num() - 1);
-		const int32 RotKeyIndex1 = FMath::Min(KeyIndex1, RawTrack.RotKeys.Num() - 1);
-		if (bHasScaleKey)
-		{
-			const int32 ScaleKeyIndex1 = FMath::Min(KeyIndex1, RawTrack.ScaleKeys.Num() - 1);
-			OutAtom = FTransform(RawTrack.RotKeys[RotKeyIndex1], RawTrack.PosKeys[PosKeyIndex1], RawTrack.ScaleKeys[ScaleKeyIndex1]);
-		}
-		else
-		{
-			OutAtom = FTransform(RawTrack.RotKeys[RotKeyIndex1], RawTrack.PosKeys[PosKeyIndex1], DefaultScale3D);
-		}
+		ExtractTransformForFrameFromTrack(RawTrack, KeyIndex1, OutAtom);
 		return;
 	}
 	else if (Alpha >= 1.f)
 	{
-		const int32 PosKeyIndex2 = FMath::Min(KeyIndex2, RawTrack.PosKeys.Num() - 1);
-		const int32 RotKeyIndex2 = FMath::Min(KeyIndex2, RawTrack.RotKeys.Num() - 1);
-		if (bHasScaleKey)
-		{
-			const int32 ScaleKeyIndex2 = FMath::Min(KeyIndex2, RawTrack.ScaleKeys.Num() - 1);
-			OutAtom = FTransform(RawTrack.RotKeys[RotKeyIndex2], RawTrack.PosKeys[PosKeyIndex2], RawTrack.ScaleKeys[ScaleKeyIndex2]);
-		}
-		else
-		{
-			OutAtom = FTransform(RawTrack.RotKeys[RotKeyIndex2], RawTrack.PosKeys[PosKeyIndex2], DefaultScale3D);
-		}
+		ExtractTransformForFrameFromTrack(RawTrack, KeyIndex2, OutAtom);
 		return;
 	}
 
@@ -889,6 +1009,7 @@ void FAnimationUtils::ExtractTransformFromTrack(float Time, int32 NumFrames, flo
 
 	FTransform KeyAtom1, KeyAtom2;
 
+	const bool bHasScaleKey = (RawTrack.ScaleKeys.Num() > 0);
 	if (bHasScaleKey)
 	{
 		const int32 ScaleKeyIndex1 = FMath::Min(KeyIndex1, RawTrack.ScaleKeys.Num() - 1);
diff --git a/Engine/Source/Runtime/Engine/Public/AnimationUtils.h b/Engine/Source/Runtime/Engine/Public/AnimationUtils.h
--- a/Engine/Source/Runtime/Engine/Public/AnimationUtils.h
+++ b/Engine/Source/Runtime/Engine/Public/AnimationUtils.h
@@ -25,13 +25,13 @@ struct FAnimPerturbationError
 {
 	float MaxErrorInTransDueToTrans;
 	float MaxErrorInRotDueToTrans;
-	float MaxErrorInScaleDueToTrans;
+	//float MaxErrorInScaleDueToTrans; // Not used TODO: Either use or remove
 	float MaxErrorInTransDueToRot;
 	float MaxErrorInRotDueToRot;
-	float MaxErrorInScaleDueToRot;
-	float MaxErrorInTransDueToScale;
+	//float MaxErrorInScaleDueToRot; // Not used TODO: Either use or remove
+	float MaxErrorInTransDueToScale; 
 	float MaxErrorInRotDueToScale;
-	float MaxErrorInScaleDueToScale;
+	//float MaxErrorInScaleDueToScale; // Not used TODO: Either use or remove
 };
 
 using AnimationErrorStats = FAnimationErrorStats;
@@ -159,6 +159,12 @@ public:
 	/** Makes sure that the animation sequence and all its dependencies required for compression are loaded. */
 	ENGINE_API static void EnsureAnimSequenceLoaded(UAnimSequence& AnimSeq);
 
+	// Extract specific frame from raw track and place in OutAtom (return identity if track is malformed)
+	ENGINE_API static void ExtractTransformForFrameFromTrackSafe(const FRawAnimSequenceTrack& RawTrack, int32 Frame, FTransform& OutAtom);
+
+	// Extract specific frame from raw track and place in OutAtom
+	ENGINE_API static void ExtractTransformForFrameFromTrack(const FRawAnimSequenceTrack& RawTrack, int32 Frame, FTransform& OutAtom);
+
 	static void ExtractTransformFromTrack(float Time, int32 NumFrames, float SequenceLength, const struct FRawAnimSequenceTrack& RawTrack, EAnimInterpolationType Interpolation, FTransform &OutAtom);
 
 #if WITH_EDITOR
