Chaos - clean up iteration count cvars

diff --git a/Engine/Source/Runtime/Engine/Private/PhysicsEngine/ImmediatePhysics/ImmediatePhysicsChaos/ImmediatePhysicsSimulation_Chaos.cpp b/Engine/Source/Runtime/Engine/Private/PhysicsEngine/ImmediatePhysics/ImmediatePhysicsChaos/ImmediatePhysicsSimulation_Chaos.cpp
--- a/Engine/Source/Runtime/Engine/Private/PhysicsEngine/ImmediatePhysics/ImmediatePhysicsChaos/ImmediatePhysicsSimulation_Chaos.cpp
+++ b/Engine/Source/Runtime/Engine/Private/PhysicsEngine/ImmediatePhysics/ImmediatePhysicsChaos/ImmediatePhysicsSimulation_Chaos.cpp
@@ -105,11 +105,7 @@ FAutoConsoleVariableRef CVarChaosImmPhysJointAngularPositionCorrection(TEXT("p.C
 // Even more temp that the above...
 //
 int32 ChaosImmediate_UsePositionSolver = 1;
-int32 ChaosImmediate_PositionSolverIts = 3;
-int32 ChaosImmediate_PositionSolverPairIts = 2;
 FAutoConsoleVariableRef CVarChaosImmPhysUsePositionSolver(TEXT("p.Chaos.ImmPhys.UsePositionSolver"), ChaosImmediate_UsePositionSolver, TEXT("Use position based collision solver for Immediate Physics (default true)"));
-FAutoConsoleVariableRef CVarChaosImmPhysPositionIts(TEXT("p.Chaos.ImmPhys.PositionSolverIts"), ChaosImmediate_PositionSolverIts, TEXT("Number of iterations to use in position solver mode (default 3)"));
-FAutoConsoleVariableRef CVarChaosImmPhysPositionPairIts(TEXT("p.Chaos.ImmPhys.PositionSolverPairIts"), ChaosImmediate_PositionSolverPairIts, TEXT("Number of pair-iterations to use in position solver mode (default 2)"));
 
 //
 // end remove when finished
@@ -210,7 +206,7 @@ namespace ImmediatePhysics_Chaos
 				{
 					DebugDraw::DrawCollisions(SimulationSpaceTransform, Collisions, 0.3f);
 				}
-				DebugDrawDynamicParticles(4, 4, FColor(128, 0, 0));
+				DebugDrawDynamicParticles(4, 4, FColor(128, 128, 0));
 			});
 		Collisions.SetPostApplyPushOutCallback(
 			[this](const float Dt, const TArray<TPBDCollisionConstraintHandle<float, 3>*>& InConstraintHandles, bool bRequiresAnotherIteration)
@@ -244,7 +240,6 @@ namespace ImmediatePhysics_Chaos
 				{
 					DebugDraw::DrawJointConstraints(SimulationSpaceTransform, InConstraintHandles, 0.6f);
 				}
-				DebugDrawDynamicParticles(4, 4, FColor(0, 0, 128));
 			});
 #endif
 	}
@@ -526,28 +521,20 @@ namespace ImmediatePhysics_Chaos
 
 			if (ChaosImmediate_UsePositionSolver)
 			{
-				SetSolverIterations(
-					ChaosImmediate_PositionSolverIts,
-					ChaosImmediate_PositionSolverPairIts,
-					ChaosImmediate_PositionSolverPairIts,
-					0,
-					0,
-					0);
-
 				Collisions.SetApplyType(ECollisionApplyType::Position);
 			}
 			else
 			{
-				SetSolverIterations(
-					ChaosImmediate_Evolution_Iterations,
-					ChaosImmediate_Joint_PairIterations,
-					ChaosImmediate_Collision_PairIterations,
-					ChaosImmediate_Evolution_PushOutIterations,
-					ChaosImmediate_Joint_PushOutPairIterations,
-					ChaosImmediate_Collision_PushOutPairIterations);
-
 				Collisions.SetApplyType(ECollisionApplyType::Velocity);
 			}
+
+			SetSolverIterations(
+				ChaosImmediate_Evolution_Iterations,
+				ChaosImmediate_Joint_PairIterations,
+				ChaosImmediate_Collision_PairIterations,
+				ChaosImmediate_Evolution_PushOutIterations,
+				ChaosImmediate_Joint_PushOutPairIterations,
+				ChaosImmediate_Collision_PushOutPairIterations);
 		}
 		UE_LOG(LogChaosJoint, Verbose, TEXT("Simulate Dt = %f Steps %d x %f"), DeltaTime, NumSteps, StepTime);
 
diff --git a/Engine/Source/Runtime/Engine/Private/PhysicsEngine/PhysicsAsset.cpp b/Engine/Source/Runtime/Engine/Private/PhysicsEngine/PhysicsAsset.cpp
--- a/Engine/Source/Runtime/Engine/Private/PhysicsEngine/PhysicsAsset.cpp
+++ b/Engine/Source/Runtime/Engine/Private/PhysicsEngine/PhysicsAsset.cpp
@@ -24,10 +24,10 @@
 FSolverIterations::FSolverIterations()
 	: SolverIterations(3)
 	, JointIterations(2)
-	, CollisionIterations(1)
-	, SolverPushOutIterations(1)
-	, JointPushOutIterations(1)
-	, CollisionPushOutIterations(2)
+	, CollisionIterations(2)
+	, SolverPushOutIterations(0)
+	, JointPushOutIterations(0)
+	, CollisionPushOutIterations(0)
 {
 }
 
