Tweaking audio mem report: moving units to be header only (easier to sort it a csv viewer)

diff --git a/Engine/Source/Runtime/Engine/Private/AudioStreamingCache.cpp b/Engine/Source/Runtime/Engine/Private/AudioStreamingCache.cpp
--- a/Engine/Source/Runtime/Engine/Private/AudioStreamingCache.cpp
+++ b/Engine/Source/Runtime/Engine/Private/AudioStreamingCache.cpp
@@ -1507,7 +1507,7 @@ TPair<int, int> FAudioChunkCache::DebugDisplay(UWorld* World, FViewport* Viewpor
 
 		const bool bWasTrimmed = CurrentElement->ChunkData.Num() == 0;
 
-		FString ElementInfo = *FString::Printf(TEXT("%4i. Size: %6.2f KB   Chunk: %d of %d   Request Count: %d    Average Index: %6.2f  Number of Handles Retaining Chunk: %d     Chunk Load Time: %6.4fms      Name: %s Notes: %s %s"),
+		FString ElementInfo = *FString::Printf(TEXT("%4i. Size: %6.2f KB   Chunk: %d of %d   Request Count: %d    Average Index: %6.2f  Number of Handles Retaining Chunk: %d     Chunk Load Time(in ms): %6.4fms      Name: %s Notes: %s %s"),
 			Index,
 			CurrentElement->ChunkData.Num() / 1024.0f,
 			CurrentElement->Key.ChunkIndex,
@@ -1639,7 +1639,7 @@ FString FAudioChunkCache::DebugPrint()
 
 		const bool bWasTrimmed = CurrentElement->ChunkData.Num() == 0;
 
-		FString ElementInfo = *FString::Printf(TEXT("%4i.\t, %6.2f KB\t, %d of %d\t, %d\t, %6.2f\t, %d\t,  %6.4fms\t, %s\t, %s %s %s"),
+		FString ElementInfo = *FString::Printf(TEXT("%4i.\t, %6.2f\t, %d of %d\t, %d\t, %6.2f\t, %d\t,  %6.4f\t, %s\t, %s %s %s"),
 			Index,
 			CurrentElement->ChunkData.Num() / 1024.0f,
 			CurrentElement->Key.ChunkIndex,
