#UE4 ULevelStreamingDynamic::LoadLevelInstance now uses RenameForPIE instead of manually appending the PIE prefix so other networking systems know about the rename and actors in these levels can network properly

diff --git a/Engine/Source/Runtime/Engine/Private/LevelStreaming.cpp b/Engine/Source/Runtime/Engine/Private/LevelStreaming.cpp
--- a/Engine/Source/Runtime/Engine/Private/LevelStreaming.cpp
+++ b/Engine/Source/Runtime/Engine/Private/LevelStreaming.cpp
@@ -1803,7 +1803,7 @@ ULevelStreamingDynamic* ULevelStreamingDynamic::LoadLevelInstance_Internal(UWorl
 	FString OnDiskPackageName = PackagePath + TEXT("/") + ShortPackageName;
 
 	// Determine loaded package name
-	FString LevelPackageNameStr = PackagePath + TEXT("/") + World->StreamingLevelsPrefix;
+	FString LevelPackageNameStr = PackagePath + TEXT("/");
 	bool bNeedsUniqueTest = false;
 	if (OptionalLevelNameOverride.IsEmpty())
 	{
@@ -1817,17 +1817,43 @@ ULevelStreamingDynamic* ULevelStreamingDynamic::LoadLevelInstance_Internal(UWorl
 		bNeedsUniqueTest = true;
 	}
 
-	FName LevelPackageName = FName(*LevelPackageNameStr);
-	if (bNeedsUniqueTest && World->GetStreamingLevels().ContainsByPredicate([&LevelPackageName](ULevelStreaming* LS) { return LS && LS->GetWorldAssetPackageFName() == LevelPackageName; }))
+	FName UnmodifiedLevelPackageName = FName(*LevelPackageNameStr);
+#if WITH_EDITOR
+	const bool bIsPlayInEditor = World->IsPlayInEditor();
+	int32 PIEInstance = INDEX_NONE;
+	if (bIsPlayInEditor)
 	{
-		// The streaming level already exists, error and return.
-		UE_LOG(LogLevelStreaming, Error, TEXT("LoadLevelInstance called with a name that already exists, returning nullptr. LevelPackageName:%s"), *LevelPackageNameStr);
-		return nullptr;
+		const FWorldContext& WorldContext = GEngine->GetWorldContextFromWorldChecked(World);
+		PIEInstance = WorldContext.PIEInstance;
+	}
+#endif
+	if (bNeedsUniqueTest)
+	{
+		FName ModifiedLevelPackageName = UnmodifiedLevelPackageName;
+#if WITH_EDITOR
+		if (bIsPlayInEditor)
+		{
+			ModifiedLevelPackageName = FName(*UWorld::ConvertToPIEPackageName(LevelPackageNameStr, PIEInstance));
+		}
+#endif
+		if (World->GetStreamingLevels().ContainsByPredicate([&ModifiedLevelPackageName](ULevelStreaming* LS) { return LS && LS->GetWorldAssetPackageFName() == ModifiedLevelPackageName; }))
+		{
+			// The streaming level already exists, error and return.
+			UE_LOG(LogLevelStreaming, Error, TEXT("LoadLevelInstance called with a name that already exists, returning nullptr. LevelPackageName:%s"), *ModifiedLevelPackageName.ToString());
+			return nullptr;
+		}
 	}
     
 	// Setup streaming level object that will load specified map
 	ULevelStreamingDynamic* StreamingLevel = NewObject<ULevelStreamingDynamic>(World, ULevelStreamingDynamic::StaticClass(), NAME_None, RF_Transient, NULL);
-    StreamingLevel->SetWorldAssetByPackageName(LevelPackageName);
+    StreamingLevel->SetWorldAssetByPackageName(UnmodifiedLevelPackageName);
+#if WITH_EDITOR
+	if (bIsPlayInEditor)
+	{
+		// Necessary for networking in PIE
+		StreamingLevel->RenameForPIE(PIEInstance);
+	}
+#endif // WITH_EDITOR
     StreamingLevel->LevelColor = FColor::MakeRandomColor();
     StreamingLevel->SetShouldBeLoaded(true);
     StreamingLevel->SetShouldBeVisible(true);
