Stop Apply compression in editor blocking main thread

diff --git a/Engine/Source/Editor/UnrealEd/Private/AnimationEditorUtils.cpp b/Engine/Source/Editor/UnrealEd/Private/AnimationEditorUtils.cpp
--- a/Engine/Source/Editor/UnrealEd/Private/AnimationEditorUtils.cpp
+++ b/Engine/Source/Editor/UnrealEd/Private/AnimationEditorUtils.cpp
@@ -744,7 +744,7 @@ namespace AnimationEditorUtils
 
 					// Clear CompressCommandletVersion so we can recompress these animations later.
 					AnimSeq->CompressCommandletVersion = 0;
-					AnimSeq->RequestAnimCompression(FRequestAnimCompressionParams(false, CompressContext));
+					AnimSeq->RequestAnimCompression(FRequestAnimCompressionParams(true, CompressContext));
 					++CompressContext->AnimIndex;
 				}
 			}
diff --git a/Engine/Source/Runtime/Engine/Classes/Animation/AnimSequence.h b/Engine/Source/Runtime/Engine/Classes/Animation/AnimSequence.h
--- a/Engine/Source/Runtime/Engine/Classes/Animation/AnimSequence.h
+++ b/Engine/Source/Runtime/Engine/Classes/Animation/AnimSequence.h
@@ -581,7 +581,7 @@ public:
 	FString GetDDCCacheKeySuffix(const bool bPerformStripping) const;
 	void ApplyCompressedData(const FString& DataCacheKeySuffix, const bool bPerformFrameStripping, const TArray<uint8>& Data);
 #endif
-	void WaitOnExistingCompression(const bool bCancelIfNotStarted=false);
+	void WaitOnExistingCompression(const bool bWantResults=true);
 	void RequestAnimCompression(FRequestAnimCompressionParams Params);
 	void RequestSyncAnimRecompression(bool bOutput = false) { RequestAnimCompression(FRequestAnimCompressionParams(false, false, bOutput)); }
 	void RequestAsyncAnimRecompression(bool bOutput = false) { RequestAnimCompression(FRequestAnimCompressionParams(true, false, bOutput)); }
diff --git a/Engine/Source/Runtime/Engine/Private/Animation/AnimCompressionDerivedData.cpp b/Engine/Source/Runtime/Engine/Private/Animation/AnimCompressionDerivedData.cpp
--- a/Engine/Source/Runtime/Engine/Private/Animation/AnimCompressionDerivedData.cpp
+++ b/Engine/Source/Runtime/Engine/Private/Animation/AnimCompressionDerivedData.cpp
@@ -336,24 +336,32 @@ void FAsyncCompressedAnimationsManagement::StartAsyncWork(FDerivedDataAnimationC
 	ActiveAsyncCompressionTasks.Emplace(Anim, SourceData, CacheKey, NewTaskSize, AsyncHandle, bPerformFrameStripping);
 }
 
-bool FAsyncCompressedAnimationsManagement::WaitOnActiveCompression(UAnimSequence* Anim)
+bool FAsyncCompressedAnimationsManagement::WaitOnActiveCompression(UAnimSequence* Anim, bool bWantResults)
 {
 	for (int32 ActiveIndex = 0; ActiveIndex < ActiveAsyncCompressionTasks.Num(); ++ActiveIndex)
 	{
 		FActiveAsyncCompressionTask& Task = ActiveAsyncCompressionTasks[ActiveIndex];
 		if (Task.Sequence == Anim)
 		{
-			GetDerivedDataCacheRef().WaitAsynchronousCompletion(Task.AsyncHandle);
-			OnActiveCompressionFinished(ActiveIndex);
-			return true; // Done
+			if(bWantResults)
+			{
+				GetDerivedDataCacheRef().WaitAsynchronousCompletion(Task.AsyncHandle);
+				OnActiveCompressionFinished(ActiveIndex);
+			}
+			else
+			{
+				Task.DataToCompress->IsCancelledSignal.Cancel();
+				Task.Sequence = nullptr;
+			}
+			return true; // Was active
 		}
 	}
 	return false;
 }
 
-bool FAsyncCompressedAnimationsManagement::WaitOnExistingCompression(UAnimSequence* Anim, const bool bCancelIfNotStarted)
+bool FAsyncCompressedAnimationsManagement::WaitOnExistingCompression(UAnimSequence* Anim, const bool bWantResults)
 {
-	if(!WaitOnActiveCompression(Anim))
+	if(!WaitOnActiveCompression(Anim, bWantResults))
 	{
 		//Check if we have a queued task
 		for (int32 QueuedIndex = 0; QueuedIndex < QueuedAsyncCompressionWork.Num(); ++QueuedIndex)
@@ -361,7 +369,7 @@ bool FAsyncCompressedAnimationsManagement::WaitOnExistingCompression(UAnimSequen
 			FQueuedAsyncCompressionWork& Task = QueuedAsyncCompressionWork[QueuedIndex];
 			if (Task.Anim == Anim)
 			{
-				if (!bCancelIfNotStarted)
+				if (bWantResults)
 				{
 					StartAsyncWork(Task.Compressor, Task.Anim, Task.Compressor.GetMemoryUsage(), Task.bPerformFrameStripping);
 				}
@@ -372,7 +380,7 @@ bool FAsyncCompressedAnimationsManagement::WaitOnExistingCompression(UAnimSequen
 				QueuedAsyncCompressionWork.RemoveAtSwap(QueuedIndex, 1, false);
 			}
 		}
-		return bCancelIfNotStarted ? false : WaitOnActiveCompression(Anim);
+		return bWantResults ? WaitOnActiveCompression(Anim, bWantResults) : false;
 	}
 	return true;
 }
diff --git a/Engine/Source/Runtime/Engine/Private/Animation/AnimSequence.cpp b/Engine/Source/Runtime/Engine/Private/Animation/AnimSequence.cpp
--- a/Engine/Source/Runtime/Engine/Private/Animation/AnimSequence.cpp
+++ b/Engine/Source/Runtime/Engine/Private/Animation/AnimSequence.cpp
@@ -2148,14 +2148,13 @@ FString UAnimSequence::GetDDCCacheKeySuffix(const bool bPerformStripping) const
 }
 #endif
 
-void UAnimSequence::WaitOnExistingCompression(const bool bCancelIfNotStarted)
+void UAnimSequence::WaitOnExistingCompression(const bool bWantResults)
 {
 #if WITH_EDITOR
 	check(IsInGameThread());
 	if (bCompressionInProgress)
 	{
-		double WaitTime = FPlatformTime::Seconds();
-		FAsyncCompressedAnimationsManagement::Get().WaitOnExistingCompression(this, bCancelIfNotStarted);
+		FAsyncCompressedAnimationsManagement::Get().WaitOnExistingCompression(this, bWantResults);
 		bCompressionInProgress = false;
 	}
 #endif
@@ -2177,7 +2176,7 @@ void UAnimSequence::RequestAnimCompression(FRequestAnimCompressionParams Params)
 		return;
 	}
 
-	WaitOnExistingCompression(true);
+	WaitOnExistingCompression(false);
 
 	if (BoneCompressionSettings == nullptr || !BoneCompressionSettings->AreSettingsValid())
 	{
diff --git a/Engine/Source/Runtime/Engine/Public/Animation/AnimCompressionDerivedDataPublic.h b/Engine/Source/Runtime/Engine/Public/Animation/AnimCompressionDerivedDataPublic.h
--- a/Engine/Source/Runtime/Engine/Public/Animation/AnimCompressionDerivedDataPublic.h
+++ b/Engine/Source/Runtime/Engine/Public/Animation/AnimCompressionDerivedDataPublic.h
@@ -64,14 +64,14 @@ public:
 	int32 GetNumRemainingJobs() const { return ActiveAsyncCompressionTasks.Num() + QueuedAsyncCompressionWork.Num(); }
 	
 	// Blocks on compression for the supplied animation. Returns false if there was no compression job to wait on. 
-	bool WaitOnExistingCompression(UAnimSequence* Anim, const bool bCancelIfNotStarted);
+	bool WaitOnExistingCompression(UAnimSequence* Anim, const bool bWantResults);
 private:
 
 	void StartAsyncWork(FDerivedDataAnimationCompression& Compressor, UAnimSequence* Anim, const uint64 NewTaskSize, const bool bPerformFrameStripping);
 
 	void StartQueuedTasks(int32 MaxActiveTasks);
 
-	bool WaitOnActiveCompression(UAnimSequence* Anim);
+	bool WaitOnActiveCompression(UAnimSequence* Anim, bool bWantResults);
 
 	void OnActiveCompressionFinished(int32 ActiveAnimIndex);
 
