Change order expressions are visited while updating cached material expression data,
fixes problem with getting correct values from parameters duplicated in multiple functions

diff --git a/Engine/Source/Runtime/Engine/Private/Materials/MaterialCachedData.cpp b/Engine/Source/Runtime/Engine/Private/Materials/MaterialCachedData.cpp
--- a/Engine/Source/Runtime/Engine/Private/Materials/MaterialCachedData.cpp
+++ b/Engine/Source/Runtime/Engine/Private/Materials/MaterialCachedData.cpp
@@ -104,10 +104,32 @@ bool FMaterialCachedExpressionData::UpdateForFunction(const FMaterialCachedExpre
 	}
 
 	bool bResult = true;
+
+	// Update expressions for all dependent functions first, before processing the remaining expressions in this function
+	// This is important so we add parameters in the proper order (parameter values are latched the first time a given parameter name is encountered)
+	FMaterialCachedExpressionContext LocalContext(Context);
+	LocalContext.bUpdateFunctionExpressions = false; // we update functions explicitly
+	{
+		FMaterialCachedExpressionData* Self = this;
+		auto DependentFunctionLamba = [Self, &LocalContext, Association, ParameterIndex, &bResult](UMaterialFunctionInterface* DepFunction) -> bool
+		{
+			const TArray<UMaterialExpression*>* FunctionExpressions = DepFunction->GetFunctionExpressions();
+			if (FunctionExpressions)
+			{
+				if (!Self->UpdateForExpressions(LocalContext, *FunctionExpressions, Association, ParameterIndex))
+				{
+					bResult = false;
+				}
+			}
+			return true;
+		};
+		Function->IterateDependentFunctions(MoveTemp(DependentFunctionLamba));
+	}
+
 	const TArray<UMaterialExpression*>* FunctionExpressions = Function->GetFunctionExpressions();
 	if (FunctionExpressions)
 	{
-		if (!UpdateForExpressions(Context, *FunctionExpressions, Association, ParameterIndex))
+		if (!UpdateForExpressions(LocalContext, *FunctionExpressions, Association, ParameterIndex))
 		{
 			bResult = false;
 		}
@@ -351,29 +373,32 @@ bool FMaterialCachedExpressionData::UpdateForExpressions(const FMaterialCachedEx
 		{
 			bHasSceneColor = true;
 		}
-		else if (UMaterialExpressionMaterialFunctionCall* FunctionCall = Cast<UMaterialExpressionMaterialFunctionCall>(Expression))
+		else if (Context.bUpdateFunctionExpressions)
 		{
-			if (!UpdateForFunction(Context, FunctionCall->MaterialFunction, GlobalParameter, -1))
+			if (UMaterialExpressionMaterialFunctionCall* FunctionCall = Cast<UMaterialExpressionMaterialFunctionCall>(Expression))
 			{
-				bResult = false;
-			}
+				if (!UpdateForFunction(Context, FunctionCall->MaterialFunction, GlobalParameter, -1))
+				{
+					bResult = false;
+				}
 
-			// Update the function call node, so it can relink inputs and outputs as needed
-			// Update even if MaterialFunctionNode->MaterialFunction is NULL, because we need to remove the invalid inputs in that case
-			FunctionCall->UpdateFromFunctionResource();
-		}
-		else if (UMaterialExpressionMaterialAttributeLayers* LayersExpression = Cast<UMaterialExpressionMaterialAttributeLayers>(Expression))
-		{
-			checkf(Association == GlobalParameter, TEXT("UMaterialExpressionMaterialAttributeLayers can't be nested"));
-			if (!UpdateForLayerFunctions(Context, LayersExpression->DefaultLayers))
-			{
-				bResult = false;
+				// Update the function call node, so it can relink inputs and outputs as needed
+				// Update even if MaterialFunctionNode->MaterialFunction is NULL, because we need to remove the invalid inputs in that case
+				FunctionCall->UpdateFromFunctionResource();
 			}
+			else if (UMaterialExpressionMaterialAttributeLayers* LayersExpression = Cast<UMaterialExpressionMaterialAttributeLayers>(Expression))
+			{
+				checkf(Association == GlobalParameter, TEXT("UMaterialExpressionMaterialAttributeLayers can't be nested"));
+				if (!UpdateForLayerFunctions(Context, LayersExpression->DefaultLayers))
+				{
+					bResult = false;
+				}
 
-			DefaultLayers = LayersExpression->DefaultLayers.Layers;
-			DefaultLayerBlends = LayersExpression->DefaultLayers.Blends;
+				DefaultLayers = LayersExpression->DefaultLayers.Layers;
+				DefaultLayerBlends = LayersExpression->DefaultLayers.Blends;
 
-			LayersExpression->RebuildLayerGraph(false);
+				LayersExpression->RebuildLayerGraph(false);
+			}
 		}
 	}
 
diff --git a/Engine/Source/Runtime/Engine/Public/MaterialCachedData.h b/Engine/Source/Runtime/Engine/Public/MaterialCachedData.h
--- a/Engine/Source/Runtime/Engine/Public/MaterialCachedData.h
+++ b/Engine/Source/Runtime/Engine/Public/MaterialCachedData.h
@@ -163,8 +163,10 @@ struct FMaterialCachedParameters
 
 struct FMaterialCachedExpressionContext
 {
-	explicit FMaterialCachedExpressionContext(UMaterialInterface* InParent = nullptr) : Parent(InParent) {}
+	explicit FMaterialCachedExpressionContext(UMaterialInterface* InParent = nullptr) : Parent(InParent), bUpdateFunctionExpressions(true) {}
+
 	UMaterialInterface* Parent;
+	bool bUpdateFunctionExpressions;
 };
 
 USTRUCT()
