fix static analysis warning for null pointer deref. static analysis caught legit error in in the Soundfield Endpoint link sanitization.

diff --git a/Engine/Source/Runtime/Engine/Private/SoundSubmix.cpp b/Engine/Source/Runtime/Engine/Private/SoundSubmix.cpp
--- a/Engine/Source/Runtime/Engine/Private/SoundSubmix.cpp
+++ b/Engine/Source/Runtime/Engine/Private/SoundSubmix.cpp
@@ -829,44 +829,50 @@ void USoundfieldEndpointSubmix::PostEditChangeProperty(struct FPropertyChangedEv
 ENGINE_API bool SubmixUtils::AreSubmixFormatsCompatible(const USoundSubmixBase* ChildSubmix, const USoundSubmixBase* ParentSubmix)
 {
 	const USoundfieldSubmix* ChildSoundfieldSubmix = Cast<const USoundfieldSubmix>(ChildSubmix);
-	const USoundfieldSubmix* ParentSoundfieldSubmix = Cast<const USoundfieldSubmix>(ParentSubmix);
 
 	// If both the child and parent are soundfield submixes, ensure that their formats are compatible.
-	if (ChildSoundfieldSubmix && ParentSoundfieldSubmix)
 	{
-		ISoundfieldFactory* ChildSoundfieldFactory = ChildSoundfieldSubmix->GetSoundfieldFactoryForSubmix();
-		ISoundfieldFactory* ParentSoundfieldFactory = ParentSoundfieldSubmix->GetSoundfieldFactoryForSubmix();
+		const USoundfieldSubmix* ParentSoundfieldSubmix = Cast<const USoundfieldSubmix>(ParentSubmix);
 
-		if (ChildSoundfieldFactory && ParentSoundfieldFactory)
+		if (ChildSoundfieldSubmix && ParentSoundfieldSubmix)
 		{
-			return ChildSoundfieldFactory->CanTranscodeToSoundfieldFormat(ParentSoundfieldFactory->GetSoundfieldFormatName(), *(ParentSoundfieldSubmix->GetSoundfieldEncodingSettings()->GetProxy()))
-				|| ParentSoundfieldFactory->CanTranscodeFromSoundfieldFormat(ChildSoundfieldFactory->GetSoundfieldFormatName(), *(ChildSoundfieldSubmix->GetSoundfieldEncodingSettings()->GetProxy()));
-		}
-		else
-		{
-			return true;
+			ISoundfieldFactory* ChildSoundfieldFactory = ChildSoundfieldSubmix->GetSoundfieldFactoryForSubmix();
+			ISoundfieldFactory* ParentSoundfieldFactory = ParentSoundfieldSubmix->GetSoundfieldFactoryForSubmix();
+
+			if (ChildSoundfieldFactory && ParentSoundfieldFactory)
+			{
+				return ChildSoundfieldFactory->CanTranscodeToSoundfieldFormat(ParentSoundfieldFactory->GetSoundfieldFormatName(), *(ParentSoundfieldSubmix->GetSoundfieldEncodingSettings()->GetProxy()))
+					|| ParentSoundfieldFactory->CanTranscodeFromSoundfieldFormat(ChildSoundfieldFactory->GetSoundfieldFormatName(), *(ChildSoundfieldSubmix->GetSoundfieldEncodingSettings()->GetProxy()));
+			}
+			else
+			{
+				return true;
+			}
 		}
 	}
 
-	const USoundfieldEndpointSubmix* ParentSoundfieldEndpointSubmix = Cast<const USoundfieldEndpointSubmix>(ParentSubmix);
-
 	// If the child is a soundfield submix and the parent is a soundfield endpoint submix, ensure that they have compatible formats.
-	if (ChildSoundfieldSubmix && ParentSoundfieldEndpointSubmix)
 	{
-		ISoundfieldFactory* ChildSoundfieldFactory = ChildSoundfieldSubmix->GetSoundfieldFactoryForSubmix();
-		ISoundfieldFactory* ParentSoundfieldFactory = ParentSoundfieldEndpointSubmix->GetSoundfieldEndpointForSubmix();
-
-		if (ChildSoundfieldFactory && ParentSoundfieldFactory)
-		{
-			return ChildSoundfieldFactory->CanTranscodeToSoundfieldFormat(ParentSoundfieldFactory->GetSoundfieldFormatName(), *(ParentSoundfieldSubmix->GetSoundfieldEncodingSettings()->GetProxy()))
-				|| ParentSoundfieldFactory->CanTranscodeFromSoundfieldFormat(ChildSoundfieldFactory->GetSoundfieldFormatName(), *(ChildSoundfieldSubmix->GetSoundfieldEncodingSettings()->GetProxy()));
-		}
-		else
+		const USoundfieldEndpointSubmix* ParentSoundfieldEndpointSubmix = Cast<const USoundfieldEndpointSubmix>(ParentSubmix);
+		
+		if (ChildSoundfieldSubmix && ParentSoundfieldEndpointSubmix)
 		{
-			return true;
+			ISoundfieldFactory* ChildSoundfieldFactory = ChildSoundfieldSubmix->GetSoundfieldFactoryForSubmix();
+			ISoundfieldFactory* ParentSoundfieldFactory = ParentSoundfieldEndpointSubmix->GetSoundfieldEndpointForSubmix();
+
+			if (ChildSoundfieldFactory && ParentSoundfieldFactory)
+			{
+				return ChildSoundfieldFactory->CanTranscodeToSoundfieldFormat(ParentSoundfieldFactory->GetSoundfieldFormatName(),  *(ParentSoundfieldEndpointSubmix->GetEncodingSettings()->GetProxy()))
+					|| ParentSoundfieldFactory->CanTranscodeFromSoundfieldFormat(ChildSoundfieldFactory->GetSoundfieldFormatName(), *(ChildSoundfieldSubmix->GetSoundfieldEncodingSettings()->GetProxy()));
+			}
+			else
+			{
+				return true;
+			}
 		}
 	}
 
+	// Otherwise, these submixes are compatible.
 	return true;
 }
 
