Editgrate from 4.25:
Get rid of memory spike when stream caching is enabled in the cook commandlet. Chunked compressed data isn't deleted when we're done cooking an asset, causing massive memory usage.

diff --git a/Engine/Source/Runtime/Engine/Private/AudioDerivedData.cpp b/Engine/Source/Runtime/Engine/Private/AudioDerivedData.cpp
--- a/Engine/Source/Runtime/Engine/Private/AudioDerivedData.cpp
+++ b/Engine/Source/Runtime/Engine/Private/AudioDerivedData.cpp
@@ -1629,6 +1629,29 @@ void USoundWave::WillNeverCacheCookedPlatformDataAgain()
 	// TODO: We can clear these arrays if we never need to cook again.
 	RawData.RemoveBulkData();
 	CompressedFormatData.FlushData();
+	
+	// Clear streamed chunks.
+	if (RunningPlatformData && RunningPlatformData->Chunks.Num() > 0)
+	{
+		// Iterate through all compressed streamed chunks and release them. 
+		for (int32 ChunkIndex = 0; ChunkIndex < RunningPlatformData->Chunks.Num(); ChunkIndex++)
+		{
+			RunningPlatformData->Chunks[ChunkIndex].BulkData.RemoveBulkData();
+		}
+	}
+
+	for (auto& PlatformData : CookedPlatformData)
+	{
+		FStreamedAudioPlatformData* PlatformDataPtr = PlatformData.Value;
+		if (PlatformDataPtr && PlatformDataPtr->Chunks.Num() > 0)
+		{
+			// Iterate through all compressed streamed chunks and release them. 
+			for (int32 ChunkIndex = 0; ChunkIndex < PlatformDataPtr->Chunks.Num(); ChunkIndex++)
+			{
+				PlatformDataPtr->Chunks[ChunkIndex].BulkData.RemoveBulkData();
+			}
+		}
+	}
 }
 #endif
 
