Changing the way we shim between the engine anlytics target and allowing studio analytics to capture them as well.

diff --git a/Engine/Source/Runtime/Engine/Private/EngineAnalytics.cpp b/Engine/Source/Runtime/Engine/Private/EngineAnalytics.cpp
--- a/Engine/Source/Runtime/Engine/Private/EngineAnalytics.cpp
+++ b/Engine/Source/Runtime/Engine/Private/EngineAnalytics.cpp
@@ -33,6 +33,60 @@ static TSharedPtr<FEditorSessionSummaryWriter> SessionSummaryWriter;
 static TSharedPtr<FEditorSessionSummarySender> SessionSummarySender;
 #endif
 
+class FStudioAnalyticsShimProvider : public IAnalyticsProvider
+{
+public:
+	FStudioAnalyticsShimProvider(TSharedPtr<IAnalyticsProviderET> InAnalytics)
+		: Analytics(InAnalytics)
+	{
+	}
+
+	virtual bool StartSession(const TArray<FAnalyticsEventAttribute>& Attributes) override
+	{
+		return Analytics->StartSession(Attributes);
+	}
+
+	virtual void EndSession() override
+	{
+		Analytics->EndSession();
+	}
+	
+	virtual FString GetSessionID() const override
+	{
+		return Analytics->GetSessionID();
+	}
+
+	virtual bool SetSessionID(const FString& InSessionID) override
+	{
+		return Analytics->SetSessionID(InSessionID);
+	}
+
+	virtual void FlushEvents() override
+	{
+		Analytics->FlushEvents();
+	}
+
+	virtual void SetUserID(const FString& InUserID) override
+	{
+		Analytics->SetUserID(InUserID);
+	}
+
+	virtual FString GetUserID() const override
+	{
+		return Analytics->GetUserID();
+	}
+
+	virtual void RecordEvent(const FString& EventName, const TArray<FAnalyticsEventAttribute>& Attributes) override
+	{
+		FStudioAnalytics::ReportEvent(EventName, Attributes);
+
+		Analytics->RecordEvent(EventName, Attributes);
+	}
+
+private:
+	TSharedPtr<IAnalyticsProviderET> Analytics;
+};
+
 /**
 * Default config func.
 */
@@ -93,7 +147,8 @@ IAnalyticsProviderET& FEngineAnalytics::GetProvider()
 {
 	checkf(bIsInitialized && IsAvailable(), TEXT("FEngineAnalytics::GetProvider called outside of Initialize/Shutdown."));
 
-	return *Analytics.Get();
+	static FStudioAnalyticsShimProvider StudioAnalyticsShim(Analytics);
+	return StudioAnalyticsShim;
 }
 
 void FEngineAnalytics::Initialize()
@@ -236,16 +291,6 @@ void FEngineAnalytics::Tick(float DeltaTime)
 #endif
 }
 
-void FEngineAnalytics::ReportEvent(const FString& EventName, const TArray<FAnalyticsEventAttribute>& Attributes)
-{
-	if (FEngineAnalytics::IsAvailable())
-	{
-		FEngineAnalytics::GetProvider().RecordEvent(EventName, Attributes);
-	}
-
-	FStudioAnalytics::ReportEvent(EventName, Attributes);
-}
-
 void FEngineAnalytics::LowDriveSpaceDetected()
 {
 #if WITH_EDITOR
diff --git a/Engine/Source/Runtime/Engine/Public/EngineAnalytics.h b/Engine/Source/Runtime/Engine/Public/EngineAnalytics.h
--- a/Engine/Source/Runtime/Engine/Public/EngineAnalytics.h
+++ b/Engine/Source/Runtime/Engine/Public/EngineAnalytics.h
@@ -38,9 +38,6 @@ public:
 	/** Called to shut down the singleton */
 	static ENGINE_API void Shutdown(bool bIsEngineShutdown = false);
 
-	/** General report event function. */
-	static ENGINE_API void ReportEvent(const FString& EventName, const TArray<FAnalyticsEventAttribute>& Attributes);
-
 	static ENGINE_API void Tick(float DeltaTime);
 
 	static ENGINE_API void LowDriveSpaceDetected();
