UE4.25 - Changed a check from an ensure as it's non-fatal
#rnx

diff --git a/Engine/Source/Runtime/Engine/Private/Materials/MaterialShader.cpp b/Engine/Source/Runtime/Engine/Private/Materials/MaterialShader.cpp
--- a/Engine/Source/Runtime/Engine/Private/Materials/MaterialShader.cpp
+++ b/Engine/Source/Runtime/Engine/Private/Materials/MaterialShader.cpp
@@ -282,7 +282,7 @@ bool FStaticParameterSet::Equivalent(const FStaticParameterSet& ReferenceSet) co
 		Temp1.SortForEquivalent();
 		Temp2.SortForEquivalent();
 		bool bResult = (Temp1 == Temp2);
-		check(!bResult || (*this) == ReferenceSet); // if this never fires, then we really didn't need to sort did we?
+		ensure(!bResult || (*this) == ReferenceSet); // if this never fires, then we really didn't need to sort did we?
 		return bResult;
 	}
 	return false;
