Fix rectlight and pointlight shadows on Metal.

#rnx

diff --git a/Engine/Source/Runtime/Engine/Private/Components/PointLightComponent.cpp b/Engine/Source/Runtime/Engine/Private/Components/PointLightComponent.cpp
--- a/Engine/Source/Runtime/Engine/Private/Components/PointLightComponent.cpp
+++ b/Engine/Source/Runtime/Engine/Private/Components/PointLightComponent.cpp
@@ -67,7 +67,7 @@ bool FPointLightSceneProxy::GetWholeSceneProjectedShadowInitializer(const FScene
 		OutInitializer.MinLightW = 0.1f;
 		OutInitializer.MaxDistanceToCastInLightW = Radius;
 
-		bool bSupportsGeometryShaders = RHISupportsGeometryShaders(GShaderPlatformForFeatureLevel[ViewFamily.GetFeatureLevel()]);
+		bool bSupportsGeometryShaders = RHISupportsGeometryShaders(GShaderPlatformForFeatureLevel[ViewFamily.GetFeatureLevel()]) || RHISupportsVertexShaderLayer(ViewFamily.GetShaderPlatform());
 		OutInitializer.bOnePassPointLightShadow = bSupportsGeometryShaders;
 
 		OutInitializer.bRayTracedDistanceField = UseRayTracedDistanceFieldShadows() && DoesPlatformSupportDistanceFieldShadowing(ViewFamily.GetShaderPlatform());
@@ -305,4 +305,4 @@ void UPointLightComponent::PostInterpChange(FProperty* PropertyThatChanged)
 	{
 		Super::PostInterpChange(PropertyThatChanged);
 	}
-}
\ No newline at end of file
+}
diff --git a/Engine/Source/Runtime/Engine/Private/Components/RectLightComponent.cpp b/Engine/Source/Runtime/Engine/Private/Components/RectLightComponent.cpp
--- a/Engine/Source/Runtime/Engine/Private/Components/RectLightComponent.cpp
+++ b/Engine/Source/Runtime/Engine/Private/Components/RectLightComponent.cpp
@@ -259,7 +259,7 @@ bool FRectLightSceneProxy::GetWholeSceneProjectedShadowInitializer(const FSceneV
 		OutInitializer.MinLightW = 0.1f;
 		OutInitializer.MaxDistanceToCastInLightW = Radius;
 		
-		bool bSupportsGeometryShaders = RHISupportsGeometryShaders(GShaderPlatformForFeatureLevel[ViewFamily.GetFeatureLevel()]);
+		bool bSupportsGeometryShaders = RHISupportsGeometryShaders(GShaderPlatformForFeatureLevel[ViewFamily.GetFeatureLevel()]) || RHISupportsVertexShaderLayer(ViewFamily.GetShaderPlatform());
 		OutInitializer.bOnePassPointLightShadow = bSupportsGeometryShaders;
 
 		OutInitializer.bRayTracedDistanceField = UseRayTracedDistanceFieldShadows() && DoesPlatformSupportDistanceFieldShadowing(ViewFamily.GetShaderPlatform());
