Attempt to rename MIDs when running a construction script. Goal is to cook deterministically

diff --git a/Engine/Source/Runtime/Engine/Private/Materials/MaterialInstanceDynamic.cpp b/Engine/Source/Runtime/Engine/Private/Materials/MaterialInstanceDynamic.cpp
--- a/Engine/Source/Runtime/Engine/Private/Materials/MaterialInstanceDynamic.cpp
+++ b/Engine/Source/Runtime/Engine/Private/Materials/MaterialInstanceDynamic.cpp
@@ -31,11 +31,44 @@ UMaterialInstanceDynamic* UMaterialInstanceDynamic::Create(UMaterialInterface* P
 UMaterialInstanceDynamic* UMaterialInstanceDynamic::Create(UMaterialInterface* ParentMaterial, UObject* InOuter, FName Name)
 {
 	LLM_SCOPE(ELLMTag::MaterialInstance);
-	UObject* Outer = InOuter ? InOuter : GetTransientPackage();
-	if (Name != NAME_None && FindObjectFast<UMaterialInstanceDynamic>(Outer, *Name.ToString(), true) != nullptr)
+
+	UPackage* TransientPackage = GetTransientPackage();
+
+	UObject* Outer = InOuter ? InOuter : TransientPackage;
+	if (Name != NAME_None)
 	{
-		// If a MID is made with the same name and outer as another, it will overwrite it. To avoid this we will change the name when there is a collision.
-		Name = MakeUniqueObjectName(Outer, UMaterialInstanceDynamic::StaticClass(), Name);
+		UMaterialInstanceDynamic* ExistingMID = FindObjectFast<UMaterialInstanceDynamic>(Outer, *Name.ToString(), true);
+		if (ExistingMID)
+		{
+			bool bRenamed = false;
+
+#if WITH_EDITOR
+			// try to rename to enhance construction script determinism:
+			AActor* OwningActor = Cast<AActor>(Outer);
+			if (!OwningActor)
+			{
+				OwningActor = Outer->GetTypedOuter<AActor>();
+			}
+
+			if (OwningActor && OwningActor->IsRunningUserConstructionScript())
+			{
+				bRenamed = true;
+				// a collision, we're going to move this existing mid to the transient package and claim the name 
+				// for ourself:
+				MID->Rename(
+					nullptr,
+					TransientPackage,
+					REN_DoNotDirty | REN_DontCreateRedirectors | REN_ForceNoResetLoaders | REN_NonTransactional
+				);
+			}
+#endif
+
+			if (!bRenamed)
+			{
+				// If a MID is made with the same name and outer as another, it will overwrite it. To avoid this we will change the name when there is a collision.
+				Name = MakeUniqueObjectName(Outer, UMaterialInstanceDynamic::StaticClass(), Name);
+			}
+		}
 	}
 	UMaterialInstanceDynamic* MID = NewObject<UMaterialInstanceDynamic>(Outer, Name);
 	MID->SetParentInternal(ParentMaterial, false);
