Fixing the main build upgradint the analytics shim to an ET provider.

diff --git a/Engine/Source/Runtime/Engine/Private/EngineAnalytics.cpp b/Engine/Source/Runtime/Engine/Private/EngineAnalytics.cpp
--- a/Engine/Source/Runtime/Engine/Private/EngineAnalytics.cpp
+++ b/Engine/Source/Runtime/Engine/Private/EngineAnalytics.cpp
@@ -33,7 +33,7 @@ static TSharedPtr<FEditorSessionSummaryWriter> SessionSummaryWriter;
 static TSharedPtr<FEditorSessionSummarySender> SessionSummarySender;
 #endif
 
-class FStudioAnalyticsShimProvider : public IAnalyticsProvider
+class FStudioAnalyticsShimProvider : public IAnalyticsProviderET
 {
 public:
 	FStudioAnalyticsShimProvider(TSharedPtr<IAnalyticsProviderET> InAnalytics)
@@ -41,6 +41,66 @@ public:
 	{
 	}
 
+	virtual void SetAppID(FString&& AppID) override
+	{
+		Analytics->SetAppID(MoveTemp(AppID));
+	}
+
+	virtual const FString& GetAppID() const override
+	{
+		return Analytics->GetAppID();
+	}
+
+	virtual void SetAppVersion(FString&& AppVersion) override
+	{
+		Analytics->SetAppVersion(MoveTemp(AppVersion));
+	}
+
+	virtual const FString& GetAppVersion() const override
+	{
+		return Analytics->GetAppVersion();
+	}
+
+	virtual bool StartSession(TArray<FAnalyticsEventAttribute>&& Attributes) override
+	{
+		return Analytics->StartSession(MoveTemp(Attributes));
+	}
+
+	virtual void RecordEvent(FString EventName, TArray<FAnalyticsEventAttribute>&& Attributes) override
+	{
+		Analytics->RecordEvent(EventName, MoveTemp(Attributes));
+	}
+
+	virtual void RecordEventJson(FString EventName, TArray<FAnalyticsEventAttribute>&& AttributesJson) override
+	{
+		Analytics->RecordEventJson(EventName, MoveTemp(AttributesJson));
+	}
+
+	virtual void SetDefaultEventAttributes(TArray<FAnalyticsEventAttribute>&& Attributes) override
+	{
+		Analytics->SetDefaultEventAttributes(MoveTemp(Attributes));
+	}
+
+	virtual const TArray<FAnalyticsEventAttribute>& GetDefaultEventAttributes() const override
+	{
+		return Analytics->GetDefaultEventAttributes();
+	}
+
+	virtual void SetURLEndpoint(const FString& UrlEndpoint, const TArray<FString>& AltDomains) override
+	{
+		Analytics->SetURLEndpoint(UrlEndpoint, AltDomains);
+	}
+
+	virtual void SetEventCallback(const OnEventRecorded& Callback) override
+	{
+		Analytics->SetEventCallback(Callback);
+	}
+
+	virtual void BlockUntilFlushed(float InTimeoutSec) override
+	{
+		Analytics->BlockUntilFlushed(InTimeoutSec);
+	}
+
 	virtual bool StartSession(const TArray<FAnalyticsEventAttribute>& Attributes) override
 	{
 		return Analytics->StartSession(Attributes);
