Build fix

diff --git a/Engine/Source/Runtime/Engine/Private/EngineAnalytics.cpp b/Engine/Source/Runtime/Engine/Private/EngineAnalytics.cpp
--- a/Engine/Source/Runtime/Engine/Private/EngineAnalytics.cpp
+++ b/Engine/Source/Runtime/Engine/Private/EngineAnalytics.cpp
@@ -101,6 +101,11 @@ public:
 		Analytics->BlockUntilFlushed(InTimeoutSec);
 	}
 
+	virtual const FAnalyticsET::Config& GetConfig() const override
+	{
+		return Analytics->GetConfig();
+	}
+
 	virtual bool StartSession(const TArray<FAnalyticsEventAttribute>& Attributes) override
 	{
 		return Analytics->StartSession(Attributes);
