Don't maintain the name of the child actor through the component instance data cache if the class has changed

diff --git a/Engine/Source/Runtime/Engine/Classes/Components/ChildActorComponent.h b/Engine/Source/Runtime/Engine/Classes/Components/ChildActorComponent.h
--- a/Engine/Source/Runtime/Engine/Classes/Components/ChildActorComponent.h
+++ b/Engine/Source/Runtime/Engine/Classes/Components/ChildActorComponent.h
@@ -40,6 +40,10 @@ public:
 	virtual void ApplyToComponent(UActorComponent* Component, const ECacheApplyPhase CacheApplyPhase) override;
 	virtual void AddReferencedObjects(FReferenceCollector& Collector) override;
 
+	// The class of the child actor when the instance data cache was stored
+	UPROPERTY()
+	TSubclassOf<AActor> ChildActorClass;
+
 	// The name of the spawned child actor so it (attempts to) remain constant across construction script reruns
 	UPROPERTY()
 	FName ChildActorName;
@@ -133,7 +137,7 @@ public:
 	//~ End ActorComponent Interface.
 
 	/** Apply the component instance data to the child actor component */
-	void ApplyComponentInstanceData(struct FChildActorComponentInstanceData* ComponentInstanceData, const ECacheApplyPhase CacheApplyPhase);
+	void ApplyComponentInstanceData(FChildActorComponentInstanceData* ComponentInstanceData, const ECacheApplyPhase CacheApplyPhase);
 
 	/** Create the child actor */
 	virtual void CreateChildActor();
diff --git a/Engine/Source/Runtime/Engine/Private/Components/ChildActorComponent.cpp b/Engine/Source/Runtime/Engine/Private/Components/ChildActorComponent.cpp
--- a/Engine/Source/Runtime/Engine/Private/Components/ChildActorComponent.cpp
+++ b/Engine/Source/Runtime/Engine/Private/Components/ChildActorComponent.cpp
@@ -286,6 +286,7 @@ void UChildActorComponent::OnUnregister()
 
 FChildActorComponentInstanceData::FChildActorComponentInstanceData(const UChildActorComponent* Component)
 	: FSceneComponentInstanceData(Component)
+	, ChildActorClass(Component->GetChildActorClass())
 	, ChildActorName(Component->GetChildActorName())
 	, ComponentInstanceData(nullptr)
 {
@@ -388,7 +389,11 @@ void UChildActorComponent::ApplyComponentInstanceData(FChildActorComponentInstan
 {
 	check(ChildActorInstanceData);
 
-	ChildActorName = ChildActorInstanceData->ChildActorName;
+	if (ChildActorClass == ChildActorInstanceData->ChildActorClass)
+	{
+		ChildActorName = ChildActorInstanceData->ChildActorName;
+	}
+
 	if (ChildActor)
 	{
 		// Only rename if it is safe to
