Don't waste 4MB+ of memory for streaming handle debugstrings on mobile
#rnx

diff --git a/Engine/Source/Runtime/Engine/Private/StreamableManager.cpp b/Engine/Source/Runtime/Engine/Private/StreamableManager.cpp
--- a/Engine/Source/Runtime/Engine/Private/StreamableManager.cpp
+++ b/Engine/Source/Runtime/Engine/Private/StreamableManager.cpp
@@ -1149,7 +1149,9 @@ TSharedPtr<FStreamableHandle> FStreamableManager::RequestAsyncLoad(TArray<FSoftO
 	NewRequest->CompleteDelegate = DelegateToCall;
 	NewRequest->OwningManager = this;
 	NewRequest->RequestedAssets = MoveTemp(TargetsToStream);
+#if (!PLATFORM_IOS && !PLATFORM_ANDROID)
 	NewRequest->DebugName = MoveTemp(DebugName);
+#endif
 	NewRequest->Priority = Priority;
 
 	int32 NumValidRequests = NewRequest->RequestedAssets.Num();
