Second iteration of memory spike cook fix- this time only clear compressed chunks cooked for other platforms, rather than cooked chunks used by the current platform.

diff --git a/Engine/Source/Runtime/Engine/Private/AudioDerivedData.cpp b/Engine/Source/Runtime/Engine/Private/AudioDerivedData.cpp
--- a/Engine/Source/Runtime/Engine/Private/AudioDerivedData.cpp
+++ b/Engine/Source/Runtime/Engine/Private/AudioDerivedData.cpp
@@ -1629,6 +1629,20 @@ void USoundWave::WillNeverCacheCookedPlatformDataAgain()
 	// TODO: We can clear these arrays if we never need to cook again.
 	RawData.RemoveBulkData();
 	CompressedFormatData.FlushData();
+	
+	// Clear streamed chunks cooked specifically for other platforms.
+	for (auto& PlatformData : CookedPlatformData)
+	{
+		FStreamedAudioPlatformData* PlatformDataPtr = PlatformData.Value;
+		if (PlatformDataPtr && PlatformDataPtr->Chunks.Num() > 0)
+		{
+			// Iterate through all compressed streamed chunks and release them. 
+			for (int32 ChunkIndex = 0; ChunkIndex < PlatformDataPtr->Chunks.Num(); ChunkIndex++)
+			{
+				PlatformDataPtr->Chunks[ChunkIndex].BulkData.RemoveBulkData();
+			}
+		}
+	}
 }
 #endif
 
