- Allow attenuation priority to be scaled beyond 1.0f and clarify respective uproperty usage UX.
- Move clamp to final copy of priority value/highest priority getter on ActiveSound
    - Fix edge case where 'always play' highest priority was not checked if placed in concurrent group
    - Fix highest priority check accidentally casting floats to ints on comparison
#rnx

diff --git a/Engine/Source/Runtime/Engine/Classes/Sound/SoundAttenuation.h b/Engine/Source/Runtime/Engine/Classes/Sound/SoundAttenuation.h
--- a/Engine/Source/Runtime/Engine/Classes/Sound/SoundAttenuation.h
+++ b/Engine/Source/Runtime/Engine/Classes/Sound/SoundAttenuation.h
@@ -290,11 +290,11 @@ struct ENGINE_API FSoundAttenuationSettings : public FBaseAttenuationSettings
 	float NonFocusDistanceScale;
 
 	/** Amount to scale the priority of sounds that are in focus. Can be used to boost the priority of sounds that are in focus. */
-	UPROPERTY(EditAnywhere, BlueprintReadWrite, Category = AttenuationListenerFocus, meta = (ClampMin = "0.0", UIMin = "0.0", EditCondition = "bEnableListenerFocus"))
+	UPROPERTY(EditAnywhere, BlueprintReadWrite, Category = AttenuationListenerFocus, meta = (ClampMin = "0.0", UIMin = "0.0", UIMax = "1.0", EditCondition = "bEnableListenerFocus"))
 	float FocusPriorityScale;
 
 	/** Amount to scale the priority of sounds that are not in-focus. Can be used to reduce the priority of sounds that are not in focus. */
-	UPROPERTY(EditAnywhere, BlueprintReadWrite, Category = AttenuationListenerFocus, meta = (ClampMin = "0.0", UIMin = "0.0", EditCondition = "bEnableListenerFocus"))
+	UPROPERTY(EditAnywhere, BlueprintReadWrite, Category = AttenuationListenerFocus, meta = (ClampMin = "0.0", UIMin = "0.0", UIMax = "1.0", EditCondition = "bEnableListenerFocus"))
 	float NonFocusPriorityScale;
 
 	/** Amount to attenuate sounds that are in focus. Can be overridden at the sound-level. */
@@ -361,24 +361,24 @@ struct ENGINE_API FSoundAttenuationSettings : public FBaseAttenuationSettings
 	UPROPERTY(EditAnywhere, BlueprintReadWrite, Category = AttenuationSubmixSend)
 	TArray<FAttenuationSubmixSendSettings> SubmixSendSettings;
 
-	/** The priority attenuation to use when the sound is at the minimum priority attenuation distance. */
-	UPROPERTY(EditAnywhere, BlueprintReadWrite, Category = AttenuationPriority, meta = (DisplayName = "Min Priority Attenuation"))
+	/** Minimum interpolated value to scale final playback priority against when the sound is at the maximum priority attenuation distance. */
+	UPROPERTY(EditAnywhere, BlueprintReadWrite, Category = AttenuationPriority, meta = (ClampMin = "0.0", UIMin = "0.0", UIMax = "1.0", DisplayName = "Min Priority Attenuation"))
 	float PriorityAttenuationMin;
 
-	/** The priority attenuation to use when the sound is at the maximum priority attenuation distance. */
-	UPROPERTY(EditAnywhere, BlueprintReadWrite, Category = AttenuationPriority, meta = (DisplayName = "Max Priority Attenuation"))
+	/** Maximum interpolated value to scale final playback priority against when the sound is at the maximum priority attenuation distance. */
+	UPROPERTY(EditAnywhere, BlueprintReadWrite, Category = AttenuationPriority, meta = (ClampMin = "0.0", UIMin = "0.0", UIMax = "1.0", DisplayName = "Max Priority Attenuation"))
 	float PriorityAttenuationMax;
 
 	/** The min distance to attenuate priority. */
-	UPROPERTY(EditAnywhere, BlueprintReadWrite, Category = AttenuationPriority, meta = (DisplayName = "Priority Attenuation Min Distance"))
+	UPROPERTY(EditAnywhere, BlueprintReadWrite, Category = AttenuationPriority, meta = (ClampMin = "0.0", DisplayName = "Priority Attenuation Min Distance"))
 	float PriorityAttenuationDistanceMin;
 
 	/** The max distance to attenuate priority. */
-	UPROPERTY(EditAnywhere, BlueprintReadWrite, Category = AttenuationPriority, meta = (DisplayName = "Priority Attenuation Max Distance"))
+	UPROPERTY(EditAnywhere, BlueprintReadWrite, Category = AttenuationPriority, meta = (ClampMin = "0.0", DisplayName = "Priority Attenuation Max Distance"))
 	float PriorityAttenuationDistanceMax;
 
-	/* The manual priority attenuation to use. Doesn't change as a function of distance. */
-	UPROPERTY(EditAnywhere, BlueprintReadWrite, Category = AttenuationPriority)
+	/* Static priority scalar to use (doesn't change as a function of distance). */
+	UPROPERTY(EditAnywhere, BlueprintReadWrite, Category = AttenuationPriority, meta = (ClampMin = "0.0", UIMin = "0.0", UIMax = "1.0", DisplayName = "Attenuation Priority"))
 	float ManualPriorityAttenuation;
 
 	/* The custom curve to use for distance-based priority attenuation. */
diff --git a/Engine/Source/Runtime/Engine/Private/ActiveSound.cpp b/Engine/Source/Runtime/Engine/Private/ActiveSound.cpp
--- a/Engine/Source/Runtime/Engine/Private/ActiveSound.cpp
+++ b/Engine/Source/Runtime/Engine/Private/ActiveSound.cpp
@@ -1458,6 +1458,31 @@ float FActiveSound::GetAttenuationFrequency(const FSoundAttenuationSettings* Set
 	return FMath::Clamp<float>(OutputFrequency, MIN_FILTER_FREQUENCY, MAX_FILTER_FREQUENCY);
 }
 
+float FActiveSound::GetHighestPriority() const
+{ 
+	static constexpr float MaxPriority = TNumericLimits<float>::Max();
+	if (bAlwaysPlay)
+	{
+		return MaxPriority;
+	}
+
+	for (const TPair<UPTRINT, FWaveInstance*>& Pair : WaveInstances)
+	{
+		if (Pair.Value && Pair.Value->SoundClass)
+		{
+			const FSoundClassProperties* SoundClassProperties = AudioDevice->GetSoundClassCurrentProperties(Pair.Value->SoundClass);
+			check(SoundClassProperties);
+			if (SoundClassProperties->bAlwaysPlay)
+			{
+				return MaxPriority;
+			}
+		}
+	}
+
+	const float HighestPriority = Priority * FocusData.PriorityHighest * FocusData.PriorityScale;
+	return FMath::Clamp(HighestPriority, 0.0f, 100.0f);
+}
+
 void FActiveSound::UpdateFocusData(float DeltaTime, const FAttenuationListenerData& ListenerData, FAttenuationFocusData* OutFocusData)
 {
 	FAttenuationFocusData* FocusDataToUpdate = OutFocusData ? OutFocusData : &FocusData;
@@ -1566,21 +1591,20 @@ void FActiveSound::UpdateAttenuation(float DeltaTime, FSoundParseParameters& Par
 		}
 		else
 		{
-			const float Denom = FMath::Max(Settings->PriorityAttenuationDistanceMax - Settings->PriorityAttenuationDistanceMin, 1.0f);
+			const float Denom = FMath::Max(Settings->PriorityAttenuationDistanceMax - Settings->PriorityAttenuationDistanceMin, SMALL_NUMBER);
 			const float Alpha = FMath::Clamp((ListenerData.ListenerToSoundDistance - Settings->PriorityAttenuationDistanceMin) / Denom, 0.0f, 1.0f);
 
 			if (Settings->PriorityAttenuationMethod == EPriorityAttenuationMethod::Linear)
 			{
-				PriorityScale = FMath::Clamp(FMath::Lerp(Settings->PriorityAttenuationMin, Settings->PriorityAttenuationMax, Alpha), 0.0f, 1.0f);
+				PriorityScale = FMath::Max(FMath::Lerp(Settings->PriorityAttenuationMin, Settings->PriorityAttenuationMax, Alpha), 0.0f);
 			}
 			else
 			{
-				PriorityScale = FMath::Clamp(Settings->CustomPriorityAttenuationCurve.GetRichCurveConst()->Eval(Alpha), 0.0f, 1.0f);
+				PriorityScale = FMath::Max(Settings->CustomPriorityAttenuationCurve.GetRichCurveConst()->Eval(Alpha), 0.0f);
 			}
 		}
 
 		ParseParams.Priority *= FMath::Max(PriorityScale, 0.0f);
-		ParseParams.Priority = FMath::Clamp(ParseParams.Priority, 0.0f, 100.0f);
 	}
 
 	if (Settings->bSpatialize || Settings->bEnableListenerFocus)
diff --git a/Engine/Source/Runtime/Engine/Private/AudioDebug.cpp b/Engine/Source/Runtime/Engine/Private/AudioDebug.cpp
--- a/Engine/Source/Runtime/Engine/Private/AudioDebug.cpp
+++ b/Engine/Source/Runtime/Engine/Private/AudioDebug.cpp
@@ -1342,7 +1342,14 @@ int32 FAudioDebugger::RenderStatSounds(UWorld* World, FViewport* Viewport, FCanv
 			}
 			case FAudioStats::EDisplaySort::Priority:
 			{
-				DebugValue = FString::Printf(TEXT("%06.2f"), StatSoundInfo.Priority);
+				if (FMath::IsNearlyEqual(StatSoundInfo.Priority, TNumericLimits<float>::Max()))
+				{
+					DebugValue += TEXT("Always  ");
+				}
+				else
+				{
+					DebugValue = FString::Printf(TEXT("%06.2f"), StatSoundInfo.Priority);
+				}
 				break;
 			}
 			case FAudioStats::EDisplaySort::Volume:
diff --git a/Engine/Source/Runtime/Engine/Private/SoundConcurrency.cpp b/Engine/Source/Runtime/Engine/Private/SoundConcurrency.cpp
--- a/Engine/Source/Runtime/Engine/Private/SoundConcurrency.cpp
+++ b/Engine/Source/Runtime/Engine/Private/SoundConcurrency.cpp
@@ -234,8 +234,8 @@ void FConcurrencyGroup::CullSoundsDueToMaxConcurrency()
 				case EMaxConcurrentResolutionRule::StopLowestPriority:
 				case EMaxConcurrentResolutionRule::StopLowestPriorityThenPreventNew:
 				{
-					const int32 APriority = A.GetHighestPriority();
-					const int32 BPriority = B.GetHighestPriority();
+					const float APriority = A.GetHighestPriority();
+					const float BPriority = B.GetHighestPriority();
 					if (!FMath::IsNearlyEqual(APriority, BPriority, KINDA_SMALL_NUMBER))
 					{
 						return APriority < BPriority;
diff --git a/Engine/Source/Runtime/Engine/Private/SoundWave.cpp b/Engine/Source/Runtime/Engine/Private/SoundWave.cpp
--- a/Engine/Source/Runtime/Engine/Private/SoundWave.cpp
+++ b/Engine/Source/Runtime/Engine/Private/SoundWave.cpp
@@ -1854,7 +1854,7 @@ void USoundWave::Parse(FAudioDevice* AudioDevice, const UPTRINT NodeWaveInstance
 	}
 	else
 	{
-		WaveInstance->Priority = ParseParams.Priority;
+		WaveInstance->Priority = FMath::Clamp(ParseParams.Priority, 0.0f, 100.0f);
 	}
 
 	WaveInstance->Location = ParseParams.Transform.GetTranslation();
diff --git a/Engine/Source/Runtime/Engine/Public/ActiveSound.h b/Engine/Source/Runtime/Engine/Public/ActiveSound.h
--- a/Engine/Source/Runtime/Engine/Public/ActiveSound.h
+++ b/Engine/Source/Runtime/Engine/Public/ActiveSound.h
@@ -669,7 +669,7 @@ public:
 	void ParseAttenuation(FSoundParseParameters& OutParseParams, int32 ListenerIndex, const FSoundAttenuationSettings& InAttenuationSettings);
 
 	/** Returns the highest effective priority of the child wave instances */
-	float GetHighestPriority() const { return Priority * FocusData.PriorityHighest * FocusData.PriorityScale; }
+	float GetHighestPriority() const;
 
 	/** Sets the amount of audio from this active sound to send to the submix. */
 	void SetSubmixSend(const FSoundSubmixSendInfo& SubmixSendInfo);
