Fixed assert in PreparePrimitiveUniformBuffer with skeletal meshes when ES3.1 Preview mode is enabled.

diff --git a/Engine/Source/Runtime/Engine/Private/SceneManagement.cpp b/Engine/Source/Runtime/Engine/Private/SceneManagement.cpp
--- a/Engine/Source/Runtime/Engine/Private/SceneManagement.cpp
+++ b/Engine/Source/Runtime/Engine/Private/SceneManagement.cpp
@@ -1085,13 +1085,13 @@ void FReadOnlyCVARCache::Init()
 void FMeshBatch::PreparePrimitiveUniformBuffer(const FPrimitiveSceneProxy* PrimitiveSceneProxy, ERHIFeatureLevel::Type FeatureLevel)
 {
 	// Fallback to using the primitive uniform buffer if GPU scene is disabled.
-	if (VertexFactory->GetType()->SupportsPrimitiveIdStream() && !UseGPUScene(GMaxRHIShaderPlatform, FeatureLevel))
+	if (!UseGPUScene(GMaxRHIShaderPlatform, FeatureLevel))
 	{
 	    for (int32 ElementIndex = 0; ElementIndex < Elements.Num(); ElementIndex++)
 	    {
 			FMeshBatchElement& MeshElement = Elements[ElementIndex];
 
-			if (!MeshElement.PrimitiveUniformBufferResource)
+			if (!MeshElement.PrimitiveUniformBuffer && !MeshElement.PrimitiveUniformBufferResource)
 			{
 				MeshElement.PrimitiveUniformBuffer = PrimitiveSceneProxy->GetUniformBuffer();
 			}
@@ -1177,7 +1177,9 @@ bool FMeshBatch::Validate(const FPrimitiveSceneProxy* PrimitiveSceneProxy, ERHIF
 		return LogMeshError(TEXT("PrimitiveSceneProxy has bVFRequiresPrimitiveUniformBuffer disabled yet tried to draw with a vertex factory that did not support PrimitiveIdStream"));
 	}
 
-	const bool bPrimitiveShaderDataComesFromSceneBuffer = VertexFactory->GetPrimitiveIdStreamIndex(EVertexInputStreamType::Default) >= 0;
+	const bool bUseGPUScene = UseGPUScene(GMaxRHIShaderPlatform, FeatureLevel);
+	
+	const bool bPrimitiveShaderDataComesFromSceneBuffer = bUseGPUScene && VertexFactory->GetPrimitiveIdStreamIndex(EVertexInputStreamType::Default) >= 0;
 
 	const bool bPrimitiveHasUniformBuffer = PrimitiveSceneProxy->GetUniformBuffer() != nullptr;
 
@@ -1187,7 +1189,8 @@ bool FMeshBatch::Validate(const FPrimitiveSceneProxy* PrimitiveSceneProxy, ERHIF
 
 		if (bPrimitiveShaderDataComesFromSceneBuffer && Elements[ElementIndex].PrimitiveUniformBuffer)
 		{
-			return LogMeshError(
+			// This is a non-fatal error.
+			LogMeshError(
 				TEXT("FMeshBatch was assigned a PrimitiveUniformBuffer even though the vertex factory fetches primitive shader data through the GPUScene buffer. ")
 				TEXT("The assigned PrimitiveUniformBuffer cannot be respected. Use PrimitiveUniformBufferResource instead for dynamic primitive data, or leave ")
 				TEXT("both null to get FPrimitiveSceneProxy->UniformBuffer"));
