Fixing an ensure holding onto the engine analytics privder in the shim, and updating some conventions.

diff --git a/Engine/Plugins/Editor/AssetSearch/Source/Private/AssetSearchManager.cpp b/Engine/Plugins/Editor/AssetSearch/Source/Private/AssetSearchManager.cpp
--- a/Engine/Plugins/Editor/AssetSearch/Source/Private/AssetSearchManager.cpp
+++ b/Engine/Plugins/Editor/AssetSearch/Source/Private/AssetSearchManager.cpp
@@ -821,7 +821,7 @@ void FAssetSearchManager::Search(const FSearchQuery& Query, TFunction<void(TArra
 {
 	check(IsInGameThread());
 
-	FStudioAnalytics::ReportEvent(TEXT("AssetSearch"), {
+	FStudioAnalytics::RecordEvent(TEXT("AssetSearch"), {
 		FAnalyticsEventAttribute(TEXT("QueryString"), Query.Query)
 	});
 
diff --git a/Engine/Source/Runtime/Engine/Private/Analytics/StudioAnalytics.cpp b/Engine/Source/Runtime/Engine/Private/Analytics/StudioAnalytics.cpp
--- a/Engine/Source/Runtime/Engine/Private/Analytics/StudioAnalytics.cpp
+++ b/Engine/Source/Runtime/Engine/Private/Analytics/StudioAnalytics.cpp
@@ -22,9 +22,9 @@
 bool FStudioAnalytics::bInitialized = false;
 volatile double FStudioAnalytics::TimeEstimation = 0;
 FThread FStudioAnalytics::TimerThread;
-TSharedPtr<IAnalyticsProvider> FStudioAnalytics::Analytics;
+TSharedPtr<IAnalyticsProviderET> FStudioAnalytics::Analytics;
 
-void FStudioAnalytics::SetProvider(TSharedRef<IAnalyticsProvider> InAnalytics)
+void FStudioAnalytics::SetProvider(TSharedRef<IAnalyticsProviderET> InAnalytics)
 {
 	checkf(!Analytics.IsValid(), TEXT("FStudioAnalytics::SetProvider called more than once."));
 
@@ -40,7 +40,7 @@ void FStudioAnalytics::SetProvider(TSharedRef<IAnalyticsProvider> InAnalytics)
 	}
 }
 
-IAnalyticsProvider& FStudioAnalytics::GetProvider()
+IAnalyticsProviderET& FStudioAnalytics::GetProvider()
 {
 	checkf(IsAvailable(), TEXT("FStudioAnalytics::GetProvider called outside of Initialize/Shutdown."));
 
@@ -95,12 +95,12 @@ double FStudioAnalytics::GetAnalyticSeconds()
 	return bInitialized ? TimeEstimation : FPlatformTime::Seconds();
 }
 
-void FStudioAnalytics::ReportEvent(const FString& EventName)
+void FStudioAnalytics::RecordEvent(const FString& EventName)
 {
-	ReportEvent(EventName, TArray<FAnalyticsEventAttribute>());
+	RecordEvent(EventName, TArray<FAnalyticsEventAttribute>());
 }
 
-void FStudioAnalytics::ReportEvent(const FString& EventName, const TArray<FAnalyticsEventAttribute>& Attributes)
+void FStudioAnalytics::RecordEvent(const FString& EventName, const TArray<FAnalyticsEventAttribute>& Attributes)
 {
 	if (FStudioAnalytics::IsAvailable())
 	{
@@ -108,6 +108,22 @@ void FStudioAnalytics::ReportEvent(const FString& EventName, const TArray<FAnaly
 	}
 }
 
+void FStudioAnalytics::RecordEvent(const FString& EventName, TArray<FAnalyticsEventAttribute>&& Attributes)
+{
+	if (FStudioAnalytics::IsAvailable())
+	{
+		FStudioAnalytics::GetProvider().RecordEvent(EventName, MoveTemp(Attributes));
+	}
+}
+
+void FStudioAnalytics::RecordEventJson(const FString& EventName, TArray<FAnalyticsEventAttribute>&& AttributesJson)
+{
+	if (FStudioAnalytics::IsAvailable())
+	{
+		FStudioAnalytics::GetProvider().RecordEventJson(EventName, MoveTemp(AttributesJson));
+	}
+}
+
 void FStudioAnalytics::FireEvent_Loading(const FString& LoadingName, double SecondsSpentLoading, const TArray<FAnalyticsEventAttribute>& InAttributes)
 {
 	// Ignore anything less than a 1/4th a second.
diff --git a/Engine/Source/Runtime/Engine/Private/EngineAnalytics.cpp b/Engine/Source/Runtime/Engine/Private/EngineAnalytics.cpp
--- a/Engine/Source/Runtime/Engine/Private/EngineAnalytics.cpp
+++ b/Engine/Source/Runtime/Engine/Private/EngineAnalytics.cpp
@@ -43,62 +43,52 @@ public:
 
 	virtual void SetAppID(FString&& AppID) override
 	{
-		Analytics->SetAppID(MoveTemp(AppID));
+		Analytics.Pin()->SetAppID(MoveTemp(AppID));
 	}
 
 	virtual const FString& GetAppID() const override
 	{
-		return Analytics->GetAppID();
+		return Analytics.Pin()->GetAppID();
 	}
 
 	virtual void SetAppVersion(FString&& AppVersion) override
 	{
-		Analytics->SetAppVersion(MoveTemp(AppVersion));
+		Analytics.Pin()->SetAppVersion(MoveTemp(AppVersion));
 	}
 
 	virtual const FString& GetAppVersion() const override
 	{
-		return Analytics->GetAppVersion();
+		return Analytics.Pin()->GetAppVersion();
 	}
 
 	virtual bool StartSession(TArray<FAnalyticsEventAttribute>&& Attributes) override
 	{
-		return Analytics->StartSession(MoveTemp(Attributes));
-	}
-
-	virtual void RecordEvent(FString EventName, TArray<FAnalyticsEventAttribute>&& Attributes) override
-	{
-		Analytics->RecordEvent(EventName, MoveTemp(Attributes));
-	}
-
-	virtual void RecordEventJson(FString EventName, TArray<FAnalyticsEventAttribute>&& AttributesJson) override
-	{
-		Analytics->RecordEventJson(EventName, MoveTemp(AttributesJson));
+		return Analytics.Pin()->StartSession(MoveTemp(Attributes));
 	}
 
 	virtual void SetDefaultEventAttributes(TArray<FAnalyticsEventAttribute>&& Attributes) override
 	{
-		Analytics->SetDefaultEventAttributes(MoveTemp(Attributes));
+		Analytics.Pin()->SetDefaultEventAttributes(MoveTemp(Attributes));
 	}
 
 	virtual const TArray<FAnalyticsEventAttribute>& GetDefaultEventAttributes() const override
 	{
-		return Analytics->GetDefaultEventAttributes();
+		return Analytics.Pin()->GetDefaultEventAttributes();
 	}
 
 	virtual void SetURLEndpoint(const FString& UrlEndpoint, const TArray<FString>& AltDomains) override
 	{
-		Analytics->SetURLEndpoint(UrlEndpoint, AltDomains);
+		Analytics.Pin()->SetURLEndpoint(UrlEndpoint, AltDomains);
 	}
 
 	virtual void SetEventCallback(const OnEventRecorded& Callback) override
 	{
-		Analytics->SetEventCallback(Callback);
+		Analytics.Pin()->SetEventCallback(Callback);
 	}
 
 	virtual void BlockUntilFlushed(float InTimeoutSec) override
 	{
-		Analytics->BlockUntilFlushed(InTimeoutSec);
+		Analytics.Pin()->BlockUntilFlushed(InTimeoutSec);
 	}
 
 	virtual const FAnalyticsET::Config& GetConfig() const override
@@ -108,48 +98,62 @@ public:
 
 	virtual bool StartSession(const TArray<FAnalyticsEventAttribute>& Attributes) override
 	{
-		return Analytics->StartSession(Attributes);
+		return Analytics.Pin()->StartSession(Attributes);
 	}
 
 	virtual void EndSession() override
 	{
-		Analytics->EndSession();
+		Analytics.Pin()->EndSession();
 	}
 	
 	virtual FString GetSessionID() const override
 	{
-		return Analytics->GetSessionID();
+		return Analytics.Pin()->GetSessionID();
 	}
 
 	virtual bool SetSessionID(const FString& InSessionID) override
 	{
-		return Analytics->SetSessionID(InSessionID);
+		return Analytics.Pin()->SetSessionID(InSessionID);
 	}
 
 	virtual void FlushEvents() override
 	{
-		Analytics->FlushEvents();
+		Analytics.Pin()->FlushEvents();
 	}
 
 	virtual void SetUserID(const FString& InUserID) override
 	{
-		Analytics->SetUserID(InUserID);
+		Analytics.Pin()->SetUserID(InUserID);
 	}
 
 	virtual FString GetUserID() const override
 	{
-		return Analytics->GetUserID();
+		return Analytics.Pin()->GetUserID();
 	}
 
 	virtual void RecordEvent(const FString& EventName, const TArray<FAnalyticsEventAttribute>& Attributes) override
 	{
-		FStudioAnalytics::ReportEvent(EventName, Attributes);
+		FStudioAnalytics::RecordEvent(EventName, Attributes);
+
+		Analytics.Pin()->RecordEvent(EventName, Attributes);
+	}
+
+	virtual void RecordEvent(FString EventName, TArray<FAnalyticsEventAttribute>&& Attributes) override
+	{
+		FStudioAnalytics::RecordEvent(EventName, MoveTemp(Attributes));
+
+		Analytics.Pin()->RecordEvent(EventName, MoveTemp(Attributes));
+	}
+
+	virtual void RecordEventJson(FString EventName, TArray<FAnalyticsEventAttribute>&& AttributesJson) override
+	{
+		FStudioAnalytics::RecordEventJson(EventName, MoveTemp(AttributesJson));
 
-		Analytics->RecordEvent(EventName, Attributes);
+		Analytics.Pin()->RecordEventJson(EventName, MoveTemp(AttributesJson));
 	}
 
 private:
-	TSharedPtr<IAnalyticsProviderET> Analytics;
+	TWeakPtr<IAnalyticsProviderET> Analytics;
 };
 
 /**
diff --git a/Engine/Source/Runtime/Engine/Public/StudioAnalytics.h b/Engine/Source/Runtime/Engine/Public/StudioAnalytics.h
--- a/Engine/Source/Runtime/Engine/Public/StudioAnalytics.h
+++ b/Engine/Source/Runtime/Engine/Public/StudioAnalytics.h
@@ -6,6 +6,7 @@
 #include "HAL/Thread.h"
 
 class IAnalyticsProvider;
+class IAnalyticsProviderET;
 struct FAnalyticsEventAttribute;
 
 /**
@@ -14,13 +15,13 @@ struct FAnalyticsEventAttribute;
 class FStudioAnalytics : FNoncopyable
 {
 public:
-	static ENGINE_API void SetProvider(TSharedRef<IAnalyticsProvider> InAnalytics);
+	static ENGINE_API void SetProvider(TSharedRef<IAnalyticsProviderET> InAnalytics);
 
 	/**
 	 * Return the provider instance. Not valid outside of Initialize/Shutdown calls.
 	 * Note: must check IsAvailable() first else this code will assert if the provider is not valid.
 	 */
-	static ENGINE_API IAnalyticsProvider& GetProvider();
+	static ENGINE_API IAnalyticsProviderET& GetProvider();
 
 	/** Helper function to determine if the provider is valid. */
 	static ENGINE_API bool IsAvailable() { return Analytics.IsValid(); }
@@ -32,8 +33,10 @@ public:
 	static ENGINE_API void Shutdown();
 
 	/** General report event function. */
-	static ENGINE_API void ReportEvent(const FString& EventName);
-	static ENGINE_API void ReportEvent(const FString& EventName, const TArray<FAnalyticsEventAttribute>& Attributes);
+	static ENGINE_API void RecordEvent(const FString& EventName);
+	static ENGINE_API void RecordEvent(const FString& EventName, const TArray<FAnalyticsEventAttribute>& Attributes);
+	static ENGINE_API void RecordEvent(const FString& EventName, TArray<FAnalyticsEventAttribute>&& Attributes);
+	static ENGINE_API void RecordEventJson(const FString& EventName, TArray<FAnalyticsEventAttribute>&& AttributesJson);
 
 	/** An event for reporting load time that blocks the editor. */
 	static ENGINE_API void FireEvent_Loading(const FString& LoadingName, double SecondsSpentLoading, const TArray<FAnalyticsEventAttribute>& Attributes = TArray<FAnalyticsEventAttribute>());
@@ -43,7 +46,7 @@ private:
 
 private:
 	static bool bInitialized;
-	static TSharedPtr<IAnalyticsProvider> Analytics;
+	static TSharedPtr<IAnalyticsProviderET> Analytics;
 	static FThread TimerThread;
 	static volatile double TimeEstimation;
 };
