Fixed up analytics WeakPtr reference using Pin()

diff --git a/Engine/Source/Runtime/Engine/Private/EngineAnalytics.cpp b/Engine/Source/Runtime/Engine/Private/EngineAnalytics.cpp
--- a/Engine/Source/Runtime/Engine/Private/EngineAnalytics.cpp
+++ b/Engine/Source/Runtime/Engine/Private/EngineAnalytics.cpp
@@ -93,7 +93,7 @@ public:
 
 	virtual const FAnalyticsET::Config& GetConfig() const override
 	{
-		return Analytics->GetConfig();
+		return Analytics.Pin()->GetConfig();
 	}
 
 	virtual bool StartSession(const TArray<FAnalyticsEventAttribute>& Attributes) override
