Adding LoadingBehavior category to audio stream caching section in audiomemreport, as well as on-screen debug print.

diff --git a/Engine/Source/Runtime/Engine/Classes/Sound/SoundWaveLoadingBehavior.cpp b/Engine/Source/Runtime/Engine/Classes/Sound/SoundWaveLoadingBehavior.cpp
new file mode 100644
--- /dev/null
+++ b/Engine/Source/Runtime/Engine/Classes/Sound/SoundWaveLoadingBehavior.cpp
@@ -0,0 +1,25 @@
+// Copyright Epic Games, Inc. All Rights Reserved.
+#include "SoundWaveLoadingBehavior.h"
+#include "CoreMinimal.h"
+
+
+const TCHAR* EnumToString(ESoundWaveLoadingBehavior InCurrentState)
+{
+	switch (InCurrentState)
+	{
+	case ESoundWaveLoadingBehavior::Inherited:
+		return TEXT("Inherited");
+	case ESoundWaveLoadingBehavior::RetainOnLoad:
+		return TEXT("RetainOnLoad");
+	case ESoundWaveLoadingBehavior::PrimeOnLoad:
+		return TEXT("PrimeOnLoad");
+	case ESoundWaveLoadingBehavior::LoadOnDemand:
+		return TEXT("LoadOnDemand");
+	case ESoundWaveLoadingBehavior::ForceInline:
+		return TEXT("ForceInline");
+	case ESoundWaveLoadingBehavior::Uninitialized:
+		return TEXT("Uninitialized");
+	}
+	ensure(false);
+	return TEXT("Unknown");
+}
diff --git a/Engine/Source/Runtime/Engine/Classes/Sound/SoundWaveLoadingBehavior.h b/Engine/Source/Runtime/Engine/Classes/Sound/SoundWaveLoadingBehavior.h
--- a/Engine/Source/Runtime/Engine/Classes/Sound/SoundWaveLoadingBehavior.h
+++ b/Engine/Source/Runtime/Engine/Classes/Sound/SoundWaveLoadingBehavior.h
@@ -26,3 +26,5 @@ enum class ESoundWaveLoadingBehavior : uint8
 	// This value is used to delineate when the value of ESoundWaveLoadingBehavior hasn't been cached on a USoundWave yet.
 	Uninitialized = 0xff UMETA(Hidden)
 };
+
+const TCHAR* EnumToString(ESoundWaveLoadingBehavior InCurrentState);
\ No newline at end of file
diff --git a/Engine/Source/Runtime/Engine/Private/AudioStreamingCache.cpp b/Engine/Source/Runtime/Engine/Private/AudioStreamingCache.cpp
--- a/Engine/Source/Runtime/Engine/Private/AudioStreamingCache.cpp
+++ b/Engine/Source/Runtime/Engine/Private/AudioStreamingCache.cpp
@@ -552,10 +552,16 @@ bool FAudioChunkCache::AddOrTouchChunk(const FChunkKey& InKey, TFunction<void(EA
 		if (FoundElement->bIsLoaded)
 		{
 			ExecuteOnLoadCompleteCallback(EAudioChunkLoadResult::AlreadyLoaded, OnLoadCompleted, CallbackThread);
+#if DEBUG_STREAM_CACHE
+			checkf(FoundElement->DebugInfo.LoadingBehavior == InKey.SoundWave->GetLoadingBehavior(true), TEXT("Cached element had different LoadingBehavior {%i} than the new request {%i}")
+				, EnumToString(FoundElement->DebugInfo.LoadingBehavior), EnumToString(InKey.SoundWave->GetLoadingBehavior(true)))
+#endif
+
 		}
 
 #if DEBUG_STREAM_CACHE
 		FoundElement->DebugInfo.NumTimesTouched++;
+		FoundElement->DebugInfo.LoadingBehavior = InKey.SoundWave->GetLoadingBehavior(true);
 #endif
 
 		return true;
@@ -572,6 +578,7 @@ bool FAudioChunkCache::AddOrTouchChunk(const FChunkKey& InKey, TFunction<void(EA
 
 #if DEBUG_STREAM_CACHE
 		CacheElement->DebugInfo.bWasCacheMiss = bNeededForPlayback;
+		CacheElement->DebugInfo.LoadingBehavior = InKey.SoundWave->GetLoadingBehavior(true);
 #endif
 		const FStreamedAudioChunk& Chunk = InKey.SoundWave->RunningPlatformData->Chunks[InKey.ChunkIndex];
 		int32 ChunkDataSize = Chunk.AudioDataSize;
@@ -1146,6 +1153,7 @@ void FAudioChunkCache::KickOffAsyncLoad(FCacheElement* CacheElement, const FChun
 
 #if DEBUG_STREAM_CACHE
 	CacheElement->DebugInfo.NumTotalChunks = InKey.SoundWave->GetNumChunks() - 1;
+	CacheElement->DebugInfo.LoadingBehavior = InKey.SoundWave->GetLoadingBehavior(true);
 #endif
 
 	// In editor, we retrieve from the DDC. In non-editor situations, we read the chunk async from the pak file.
@@ -1489,6 +1497,7 @@ TPair<int, int> FAudioChunkCache::DebugDisplay(UWorld* World, FViewport* Viewpor
 		int32 NumTimesTouched = -1;
 		double TimeToLoad = -1.0;
 		float AveragePlaceInCache = -1.0f;
+		ESoundWaveLoadingBehavior LoadingBehavior = ESoundWaveLoadingBehavior::Uninitialized;
 		bool bWasCacheMiss = false;
 		bool bIsStaleChunk = false;
 
@@ -1497,6 +1506,7 @@ TPair<int, int> FAudioChunkCache::DebugDisplay(UWorld* World, FViewport* Viewpor
 		NumTimesTouched = CurrentElement->DebugInfo.NumTimesTouched;
 		TimeToLoad = CurrentElement->DebugInfo.TimeToLoad;
 		AveragePlaceInCache = CurrentElement->DebugInfo.AverageLocationInCacheWhenNeeded;
+		LoadingBehavior = CurrentElement->DebugInfo.LoadingBehavior;
 		bWasCacheMiss = CurrentElement->DebugInfo.bWasCacheMiss;
 #endif
 
@@ -1507,7 +1517,7 @@ TPair<int, int> FAudioChunkCache::DebugDisplay(UWorld* World, FViewport* Viewpor
 
 		const bool bWasTrimmed = CurrentElement->ChunkData.Num() == 0;
 
-		FString ElementInfo = *FString::Printf(TEXT("%4i. Size: %6.2f KB   Chunk: %d of %d   Request Count: %d    Average Index: %6.2f  Number of Handles Retaining Chunk: %d     Chunk Load Time(in ms): %6.4fms      Name: %s Notes: %s %s"),
+		FString ElementInfo = *FString::Printf(TEXT("%4i. Size: %6.2f KB   Chunk: %d of %d   Request Count: %d    Average Index: %6.2f  Number of Handles Retaining Chunk: %d     Chunk Load Time(in ms): %6.4fms      Loading Behavior: %s      Name: %s Notes: %s %s"),
 			Index,
 			CurrentElement->ChunkData.Num() / 1024.0f,
 			CurrentElement->Key.ChunkIndex,
@@ -1516,6 +1526,7 @@ TPair<int, int> FAudioChunkCache::DebugDisplay(UWorld* World, FViewport* Viewpor
 			AveragePlaceInCache,
 			CurrentElement->NumConsumers.GetValue(),
 			TimeToLoad,
+			EnumToString(LoadingBehavior),
 			bWasTrimmed ? TEXT("TRIMMED CHUNK") : *CurrentElement->Key.SoundWaveName.ToString(),
 			bWasCacheMiss ? TEXT("(Cache Miss!)") : TEXT(""),
 			bIsStaleChunk ? TEXT("(Stale Chunk)") : TEXT("")
@@ -1612,7 +1623,7 @@ FString FAudioChunkCache::DebugPrint()
 	CurrentElement = MostRecentElement;
 	int32 Index = 0;
 
-	OutputString += TEXT("Index:\t, Size (KB):\t, Chunk:\t, Request Count:\t, Average Index:\t, Number of Handles Retaining Chunk:\t, Chunk Load Time:\t, Name: \t, Notes:\t, \n");
+	OutputString += TEXT("Index:\t, Size (KB):\t, Chunk:\t, Request Count:\t, Average Index:\t, Number of Handles Retaining Chunk:\t, Chunk Load Time:\t, Name: \t, LoadingBehavior: \t, Notes:\t, \n");
 
 	// More detailed info about individual chunks here:
 	while (CurrentElement != nullptr)
@@ -1621,6 +1632,7 @@ FString FAudioChunkCache::DebugPrint()
 		int32 NumTimesTouched = -1;
 		double TimeToLoad = -1.0;
 		float AveragePlaceInCache = -1.0f;
+		ESoundWaveLoadingBehavior LoadingBehavior = ESoundWaveLoadingBehavior::Uninitialized;
 		bool bWasCacheMiss = false;
 		bool bIsStaleChunk = false;
 
@@ -1629,6 +1641,7 @@ FString FAudioChunkCache::DebugPrint()
 		NumTimesTouched = CurrentElement->DebugInfo.NumTimesTouched;
 		TimeToLoad = CurrentElement->DebugInfo.TimeToLoad;
 		AveragePlaceInCache = CurrentElement->DebugInfo.AverageLocationInCacheWhenNeeded;
+		LoadingBehavior = CurrentElement->DebugInfo.LoadingBehavior;
 		bWasCacheMiss = CurrentElement->DebugInfo.bWasCacheMiss;
 #endif
 
@@ -1639,7 +1652,7 @@ FString FAudioChunkCache::DebugPrint()
 
 		const bool bWasTrimmed = CurrentElement->ChunkData.Num() == 0;
 
-		FString ElementInfo = *FString::Printf(TEXT("%4i.\t, %6.2f\t, %d of %d\t, %d\t, %6.2f\t, %d\t,  %6.4f\t, %s\t, %s %s %s"),
+		FString ElementInfo = *FString::Printf(TEXT("%4i.\t, %6.2f\t, %d of %d\t, %d\t, %6.2f\t, %d\t,  %6.4f\t, %s\t, %s, %s %s %s"),
 			Index,
 			CurrentElement->ChunkData.Num() / 1024.0f,
 			CurrentElement->Key.ChunkIndex,
@@ -1649,6 +1662,7 @@ FString FAudioChunkCache::DebugPrint()
 			CurrentElement->NumConsumers.GetValue(),
 			TimeToLoad,
 			bWasTrimmed ? TEXT("TRIMMED CHUNK") : *CurrentElement->Key.SoundWaveName.ToString(),
+			EnumToString(LoadingBehavior),
 			bWasCacheMiss ? TEXT("(Cache Miss!)") : TEXT(""),
 			bIsStaleChunk ? TEXT("(Stale Chunk)") : TEXT(""),
 			CurrentElement->IsLoadInProgress() ? TEXT("(Loading In Progress)") : TEXT("")
diff --git a/Engine/Source/Runtime/Engine/Public/AudioStreamingCache.h b/Engine/Source/Runtime/Engine/Public/AudioStreamingCache.h
--- a/Engine/Source/Runtime/Engine/Public/AudioStreamingCache.h
+++ b/Engine/Source/Runtime/Engine/Public/AudioStreamingCache.h
@@ -120,6 +120,9 @@ private:
 		// This is a cumulative moving average of a chunks location before it was 
 		float AverageLocationInCacheWhenNeeded;
 
+		// Note the loading behavior of the sound wave that inserted this element into the cache
+		ESoundWaveLoadingBehavior LoadingBehavior;
+
 		// if true, 
 		bool bWasCacheMiss;
 
@@ -129,6 +132,7 @@ private:
 			, TimeLoadStarted(0.0)
 			, TimeToLoad(0.0)
 			, AverageLocationInCacheWhenNeeded(0.0f)
+			, LoadingBehavior(ESoundWaveLoadingBehavior::Uninitialized)
 			, bWasCacheMiss(false)
 		{
 		}
@@ -139,6 +143,7 @@ private:
 			NumTimesTouched = 0;
 			TimeLoadStarted = 0;
 			TimeToLoad = 0.0f;
+			LoadingBehavior = ESoundWaveLoadingBehavior::Uninitialized;
 			bWasCacheMiss = false;
 			AverageLocationInCacheWhenNeeded = 0.0f;
 		}
