Fix SetSimulationShape.
Fixes updating mass properties on actors that disabled collision.

diff --git a/Engine/Source/Runtime/Engine/Private/PhysicsEngine/Experimental/PhysInterface_Chaos.cpp b/Engine/Source/Runtime/Engine/Private/PhysicsEngine/Experimental/PhysInterface_Chaos.cpp
--- a/Engine/Source/Runtime/Engine/Private/PhysicsEngine/Experimental/PhysInterface_Chaos.cpp
+++ b/Engine/Source/Runtime/Engine/Private/PhysicsEngine/Experimental/PhysInterface_Chaos.cpp
@@ -919,6 +919,13 @@ void FPhysInterface_Chaos::SetComLocalPose_AssumesLocked(const FPhysicsActorHand
 	}
 }
 
+void FPhysInterface_Chaos::SetIsSimulationShape(const FPhysicsShapeHandle& InShape, bool bIsSimShape)
+{
+	FPhysicsShapeHandle& ShapeHandle = const_cast<FPhysicsShapeHandle&>(InShape);
+	ShapeHandle.Shape->SetSimulate(bIsSimShape);
+	ShapeHandle.bSimulation = bIsSimShape;
+}
+
 float FPhysInterface_Chaos::GetStabilizationEnergyThreshold_AssumesLocked(const FPhysicsActorHandle& InHandle)
 {
 	// #todo : Implement
@@ -1604,6 +1611,8 @@ bool FPhysInterface_Chaos::IsSimulationShape(const FPhysicsShapeHandle& InShape)
 
 bool FPhysInterface_Chaos::IsQueryShape(const FPhysicsShapeHandle& InShape)
 {
+	// This data is not stored on concrete shape. TODO: Remove ensure if we actually use this flag when constructing shape handles.
+	CHAOS_ENSURE(false);
     return InShape.bQuery;
 }
 
diff --git a/Engine/Source/Runtime/Engine/Public/Physics/Experimental/PhysInterface_Chaos.h b/Engine/Source/Runtime/Engine/Public/Physics/Experimental/PhysInterface_Chaos.h
--- a/Engine/Source/Runtime/Engine/Public/Physics/Experimental/PhysInterface_Chaos.h
+++ b/Engine/Source/Runtime/Engine/Public/Physics/Experimental/PhysInterface_Chaos.h
@@ -402,8 +402,8 @@ public:
     static void SetMaskFilter(const FPhysicsShapeHandle& InShape, FMaskFilter InFilter) {}
 	static void SetSimulationFilter(const FPhysicsShapeHandle& InShape, const FCollisionFilterData& InFilter);
 	static void SetQueryFilter(const FPhysicsShapeHandle& InShape, const FCollisionFilterData& InFilter);
-    static void SetIsSimulationShape(const FPhysicsShapeHandle& InShape, bool bIsSimShape) { const_cast<FPhysicsShapeHandle&>(InShape).bSimulation = bIsSimShape; }
-    static void SetIsQueryShape(const FPhysicsShapeHandle& InShape, bool bIsQueryShape) { const_cast<FPhysicsShapeHandle&>(InShape).bSimulation = bIsQueryShape; }
+	static void SetIsSimulationShape(const FPhysicsShapeHandle& InShape, bool bIsSimShape);
+    static void SetIsQueryShape(const FPhysicsShapeHandle& InShape, bool bIsQueryShape) { const_cast<FPhysicsShapeHandle&>(InShape).bQuery = bIsQueryShape; }
     static void SetUserData(const FPhysicsShapeHandle& InShape, void* InUserData);
     static void SetGeometry(const FPhysicsShapeHandle& InShape, physx::PxGeometry& InGeom) {}
 	static void SetLocalTransform(const FPhysicsShapeHandle& InShape, const FTransform& NewLocalTransform);
