CIS error, Non-Unity (needed to move include from implementation file to header)

diff --git a/Engine/Source/Runtime/Engine/Private/AudioStreamingCache.cpp b/Engine/Source/Runtime/Engine/Private/AudioStreamingCache.cpp
--- a/Engine/Source/Runtime/Engine/Private/AudioStreamingCache.cpp
+++ b/Engine/Source/Runtime/Engine/Private/AudioStreamingCache.cpp
@@ -6,7 +6,6 @@ AudioStreaming.cpp: Implementation of audio streaming classes.
 
 #include "AudioStreamingCache.h"
 #include "Misc/CoreStats.h"
-#include "Sound/SoundWave.h"
 #include "Sound/AudioSettings.h"
 #include "DerivedDataCacheInterface.h"
 #include "Serialization/MemoryReader.h"
diff --git a/Engine/Source/Runtime/Engine/Public/AudioStreamingCache.h b/Engine/Source/Runtime/Engine/Public/AudioStreamingCache.h
--- a/Engine/Source/Runtime/Engine/Public/AudioStreamingCache.h
+++ b/Engine/Source/Runtime/Engine/Public/AudioStreamingCache.h
@@ -16,6 +16,7 @@ AudioStreaming.h: Definitions of classes used for audio streaming.
 #include "Async/AsyncFileHandle.h"
 #include "HAL/ThreadSafeBool.h"
 #include "AudioStreaming.h"
+#include "Sound/SoundWave.h"
 
 
 #define DEBUG_STREAM_CACHE !UE_BUILD_SHIPPING
