Fix startup crash when -disable-platform-chunk-install is enabled

diff --git a/Engine/Source/Runtime/Engine/Private/Streaming/StreamingTexture.cpp b/Engine/Source/Runtime/Engine/Private/Streaming/StreamingTexture.cpp
--- a/Engine/Source/Runtime/Engine/Private/Streaming/StreamingTexture.cpp
+++ b/Engine/Source/Runtime/Engine/Private/Streaming/StreamingTexture.cpp
@@ -159,17 +159,20 @@ void FStreamingRenderAsset::UpdateOptionalMipsState_Async()
 		OptionalMipsState = IFileManager::Get().FileExists(*OptionalBulkDataFilename) ? EOptionalMipsState::OMS_HasOptionalMips : EOptionalMipsState::OMS_NoOptionalMips;
 	}
 #else
-	if (OptionalMipsState == EOptionalMipsState::OMS_NotCached && OptionalMipIndex != INDEX_NONE)
+	if (RenderAsset)
 	{
-		if (RenderAsset->DoesMipDataExist(OptionalMipIndex))
-		{
-			OptionalMipsState = EOptionalMipsState::OMS_HasOptionalMips;
-		}
-		else
+		if (OptionalMipsState == EOptionalMipsState::OMS_NotCached && OptionalMipIndex != INDEX_NONE)
 		{
-			OptionalMipsState = EOptionalMipsState::OMS_NoOptionalMips;
-		}
-	}	
+			if (RenderAsset->DoesMipDataExist(OptionalMipIndex))
+			{
+				OptionalMipsState = EOptionalMipsState::OMS_HasOptionalMips;
+			}
+			else
+			{
+				OptionalMipsState = EOptionalMipsState::OMS_NoOptionalMips;
+			}
+		}	
+	}
 #endif
 }
 
