Improved reliability of the Editor analytics.
  - Modified the Editor session to support lockless logging and corresponding analysis for concurrent events.
  - Added null check EditorSessionSummaryWriter::LowDriveSpaceDetected() to prevent accessing a null session.
  - Fixed CrashReportClient that could use the analytic provider outside of init/shutdow
  - Fixed IdleTime reported that could be buggy if Slate did not register any interaction before a crash occurs.

diff --git a/Engine/Source/Developer/EditorAnalyticsSession/Private/EditorAnalyticsSession.cpp b/Engine/Source/Developer/EditorAnalyticsSession/Private/EditorAnalyticsSession.cpp
--- a/Engine/Source/Developer/EditorAnalyticsSession/Private/EditorAnalyticsSession.cpp
+++ b/Engine/Source/Developer/EditorAnalyticsSession/Private/EditorAnalyticsSession.cpp
@@ -2,6 +2,7 @@
 
 #include "EditorAnalyticsSession.h"
 #include "Modules/ModuleManager.h"
+#include "Internationalization/Regex.h"
 
 IMPLEMENT_MODULE(FEditorAnalyticsSessionModule, EditorAnalyticsSession);
 
@@ -118,6 +119,105 @@ namespace EditorAnalyticsUtils
 		return EditorAnalyticsDefs::SessionSummarySection + TEXT("/") + SessionID;
 	}
 
+	static FString GetSessionEventLogDir()
+	{
+		return FString::Printf(TEXT("%sAnalytics"), FPlatformProcess::ApplicationSettingsDir());
+	}
+
+	static void LogSessionEvent(FEditorAnalyticsSession& Session, FEditorAnalyticsSession::EEventType InEventType, const FDateTime& InTimestamp)
+	{
+		// The logger uses the robustness of the file system to log events. It creates a file (a directory) for each event and encodes the event payload in the file name.
+		// This doesn't require any fancy synchronization or complicated concurrent file IO implementation. Since the number of events is low (0 to 5 per session), that's
+		// a straight forward working solution. The files are deleted when the session is deleted. Also avoid memory allocation, this can be called when the heap is corrupted.
+		TCHAR TimestampStr[256];
+		FCString::Sprintf(TimestampStr, TFormatSpecifier<decltype(InTimestamp.ToUnixTimestamp())>::GetFormatSpecifier(), InTimestamp.ToUnixTimestamp());
+
+		TCHAR Pathname[512];
+		FCString::Sprintf(Pathname, TEXT("%s/%s_%d_%d_%d_%d_%s"), 
+			*EditorAnalyticsUtils::GetSessionEventLogDir(),
+			*Session.SessionId,
+			static_cast<int>(InEventType),
+			FPlatformAtomics::AtomicRead(&Session.Idle1Min),
+			FPlatformAtomics::AtomicRead(&Session.Idle5Min),
+			FPlatformAtomics::AtomicRead(&Session.Idle30Min),
+			TimestampStr);
+
+		IFileManager::Get().MakeDirectory(Pathname, /*Tree*/true);
+	}
+
+	/** Analyze the events logged with LogEvent() and update the session fields to reflect the last state of the session. */
+	static void UpdateSessionFromLogAnalysis(FEditorAnalyticsSession& Session)
+	{
+		// Read and aggregate the log events. The event data is encoded in the directory names created by the logger
+		FRegexPattern Pattern(TEXT(R"((^[a-fA-F0-9-]+)_([0-9]+)_([0-9]+)_([0-9]+)_([0-9]+)_([0-9]+))")); // Need help with regex? Try https://regex101.com/
+		IFileManager::Get().IterateDirectoryRecursively(*EditorAnalyticsUtils::GetSessionEventLogDir(), [&Session, &Pattern](const TCHAR* Pathname, bool bIsDir)
+		{
+			if (bIsDir)
+			{
+				FRegexMatcher Matcher(Pattern, FPaths::GetCleanFilename(Pathname));
+				if (Matcher.FindNext() && Matcher.GetCaptureGroup(1) == Session.SessionId)
+				{
+					FEditorAnalyticsSession::EEventType EventType = static_cast<FEditorAnalyticsSession::EEventType>(FCString::Atoi(*Matcher.GetCaptureGroup(2))); // Event
+					switch (EventType)
+					{
+						case FEditorAnalyticsSession::EEventType::Crashed:     Session.bCrashed = true;       break;
+						case FEditorAnalyticsSession::EEventType::GpuCrashed : Session.bGPUCrashed = true;    break;
+						case FEditorAnalyticsSession::EEventType::Terminated:  Session.bIsTerminating = true; break;
+						default: break;
+					}
+
+					int32 ParsedIdle = FCString::Atoi(*Matcher.GetCaptureGroup(3)); // Idle1Min.
+					if (ParsedIdle > Session.Idle1Min)
+					{
+						Session.Idle1Min = ParsedIdle; // No concurrency expected when reloading (no need for atomic compare exchange)
+					}
+
+					ParsedIdle = FCString::Atoi(*Matcher.GetCaptureGroup(4)); // Idle5Min.
+					if (ParsedIdle > Session.Idle5Min)
+					{
+						Session.Idle5Min = ParsedIdle;
+					}
+
+					ParsedIdle = FCString::Atoi(*Matcher.GetCaptureGroup(5)); // Idle30Min.
+					if (ParsedIdle > Session.Idle30Min)
+					{
+						Session.Idle30Min = ParsedIdle;
+					}
+
+					FDateTime ParsedTimestamp = FDateTime::FromUnixTimestamp(FCString::Atoi64(*Matcher.GetCaptureGroup(6))); // Unix timestamp (UTC)
+					if (ParsedTimestamp > Session.Timestamp)
+					{
+						Session.Timestamp = ParsedTimestamp;
+					}
+				}
+			}
+			return true;
+		});
+	}
+
+	static void DeleteLogEvents(const FEditorAnalyticsSession& Session)
+	{
+		// Gather the list of files
+		TArray<FString> SessionEventPaths;
+		IFileManager::Get().IterateDirectoryRecursively(*EditorAnalyticsUtils::GetSessionEventLogDir(), [&Session, &SessionEventPaths](const TCHAR* Pathname, bool bIsDir)
+		{
+			if (bIsDir)
+			{
+				if (FPaths::GetCleanFilename(Pathname).StartsWith(Session.SessionId))
+				{
+					SessionEventPaths.Emplace(Pathname);
+				}
+			}
+			return true; // Continue
+		});
+
+		// Delete the session files.
+		for (const FString& EventPathname : SessionEventPaths)
+		{
+			IFileManager::Get().DeleteDirectory(*EventPathname, /*RequiredExist*/false, /*Tree*/false);
+		}
+	}
+
 // Utility macros to make it easier to check that all fields are being written to.
 #define GET_STORED_STRING(FieldName) FPlatformMisc::GetStoredValue(EditorAnalyticsDefs::StoreId, SectionName, EditorAnalyticsDefs:: FieldName ## StoreKey, Session.FieldName)
 #define GET_STORED_INT(FieldName) \
@@ -126,7 +226,6 @@ namespace EditorAnalyticsUtils
 		FPlatformMisc::GetStoredValue(EditorAnalyticsDefs::StoreId, SectionName, EditorAnalyticsDefs:: FieldName ## StoreKey, FieldName ## Temp); \
 		Session.FieldName = FCString::Atoi(*FieldName ## Temp); \
 	}
-	
 
 	static void LoadInternal(FEditorAnalyticsSession& Session, const FString& InSessionId)
 	{
@@ -159,7 +258,6 @@ namespace EditorAnalyticsUtils
 			Session.Timestamp = EditorAnalyticsUtils::StringToTimestamp(TimestampString);
 		}
 
-		GET_STORED_INT(SessionDuration);
 		GET_STORED_INT(Idle1Min);
 		GET_STORED_INT(Idle5Min);
 		GET_STORED_INT(Idle30Min);
@@ -215,6 +313,9 @@ namespace EditorAnalyticsUtils
 		Session.bIsInPIE = EditorAnalyticsUtils::GetStoredBool(SectionName, EditorAnalyticsDefs::IsInPIEStoreKey);
 		Session.bIsInVRMode = EditorAnalyticsUtils::GetStoredBool(SectionName, EditorAnalyticsDefs::IsInVRModeStoreKey);
 		Session.bIsInEnterprise = EditorAnalyticsUtils::GetStoredBool(SectionName, EditorAnalyticsDefs::IsInEnterpriseStoreKey);
+
+		// Analyze the logged events and update corresponding fields in the session.
+		UpdateSessionFromLogAnalysis(Session);
 	}
 
 #undef GET_STORED_INT
@@ -244,7 +345,6 @@ FEditorAnalyticsSession::FEditorAnalyticsSession()
 	PlatformProcessID = 0;
 	StartupTimestamp = FDateTime::MinValue();
 	Timestamp = FDateTime::MinValue();
-	SessionDuration = 0;
 	Idle1Min = 0;
 	Idle5Min = 0;
 	Idle30Min = 0;
@@ -359,18 +459,14 @@ bool FEditorAnalyticsSession::Save()
 			{EditorAnalyticsDefs::ProjectDescriptionStoreKey,  ProjectDescription},
 			{EditorAnalyticsDefs::ProjectVersionStoreKey,  ProjectVersion},
 			{EditorAnalyticsDefs::TimestampStoreKey,       EditorAnalyticsUtils::TimestampToString(Timestamp)},
-			{EditorAnalyticsDefs::SessionDurationStoreKey, FString::FromInt(SessionDuration)},
 			{EditorAnalyticsDefs::Idle1MinStoreKey,  FString::FromInt(Idle1Min)},
 			{EditorAnalyticsDefs::Idle5MinStoreKey,  FString::FromInt(Idle5Min)},
 			{EditorAnalyticsDefs::Idle30MinStoreKey, FString::FromInt(Idle30Min)},
 			{EditorAnalyticsDefs::CurrentUserActivityStoreKey, CurrentUserActivity},
 			{EditorAnalyticsDefs::AverageFPSStoreKey,     FString::SanitizeFloat(AverageFPS)},
-			{EditorAnalyticsDefs::IsCrashStoreKey,        EditorAnalyticsUtils::BoolToStoredString(bCrashed)},
-			{EditorAnalyticsDefs::IsGPUCrashStoreKey,     EditorAnalyticsUtils::BoolToStoredString(bGPUCrashed)},
 			{EditorAnalyticsDefs::IsDebuggerStoreKey,     EditorAnalyticsUtils::BoolToStoredString(bIsDebugger)},
 			{EditorAnalyticsDefs::WasDebuggerStoreKey,    EditorAnalyticsUtils::BoolToStoredString(bWasEverDebugger)},
 			{EditorAnalyticsDefs::IsVanillaStoreKey,      EditorAnalyticsUtils::BoolToStoredString(bIsVanilla)},
-			{EditorAnalyticsDefs::IsTerminatingKey,       EditorAnalyticsUtils::BoolToStoredString(bIsTerminating)},
 			{EditorAnalyticsDefs::WasShutdownStoreKey,    EditorAnalyticsUtils::BoolToStoredString(bWasShutdown)},
 			{EditorAnalyticsDefs::IsInPIEStoreKey,        EditorAnalyticsUtils::BoolToStoredString(bIsInPIE)    },
 			{EditorAnalyticsDefs::IsInEnterpriseStoreKey, EditorAnalyticsUtils::BoolToStoredString(bIsInEnterprise)},
@@ -383,29 +479,6 @@ bool FEditorAnalyticsSession::Save()
 	return true;
 }
 
-bool FEditorAnalyticsSession::SaveForCrash()
-{
-	if (!ensure(IsLocked()))
-	{
-		return false;
-	}
-
-	const FString StorageLocation = EditorAnalyticsUtils::GetSessionStorageLocation(SessionId);
-
-	TMap<FString, FString> KeyValues = {
-		{EditorAnalyticsDefs::IsCrashStoreKey,         EditorAnalyticsUtils::BoolToStoredString(bCrashed)},
-		{EditorAnalyticsDefs::IsGPUCrashStoreKey,      EditorAnalyticsUtils::BoolToStoredString(bGPUCrashed)},
-		{EditorAnalyticsDefs::IsTerminatingKey,        EditorAnalyticsUtils::BoolToStoredString(bIsTerminating)},
-		{EditorAnalyticsDefs::WasShutdownStoreKey,     EditorAnalyticsUtils::BoolToStoredString(bWasShutdown)},
-		{EditorAnalyticsDefs::TimestampStoreKey,       EditorAnalyticsUtils::TimestampToString(Timestamp)},
-		{EditorAnalyticsDefs::SessionDurationStoreKey, FString::FromInt(SessionDuration)},
-	};
-
-	FPlatformMisc::SetStoredValues(EditorAnalyticsDefs::StoreId, StorageLocation, KeyValues);
-
-	return true;
-}
-
 bool FEditorAnalyticsSession::Load(const FString& InSessionID)
 {
 	if (!ensure(IsLocked()))
@@ -479,6 +552,9 @@ bool FEditorAnalyticsSession::Delete() const
 	FPlatformMisc::DeleteStoredValue(EditorAnalyticsDefs::StoreId, SectionName, EditorAnalyticsDefs::IsInEnterpriseStoreKey);
 	FPlatformMisc::DeleteStoredValue(EditorAnalyticsDefs::StoreId, SectionName, EditorAnalyticsDefs::IsInVRModeStoreKey);
 
+	// Delete the log files.
+	EditorAnalyticsUtils::DeleteLogEvents(*this);
+
 	return true;
 }
 
@@ -536,3 +612,8 @@ bool FEditorAnalyticsSession::SaveStoredSessionIDs(const TArray<FString>& InSess
 	FPlatformMisc::SetStoredValue(EditorAnalyticsDefs::StoreId, EditorAnalyticsDefs::SessionSummarySection, EditorAnalyticsDefs::SessionListStoreKey, SessionListString);
 	return true;
 }
+
+void FEditorAnalyticsSession::LogEvent(EEventType InEventType, const FDateTime& InTimestamp)
+{
+	EditorAnalyticsUtils::LogSessionEvent(*this, InEventType, InTimestamp);
+}
diff --git a/Engine/Source/Developer/EditorAnalyticsSession/Private/EditorSessionSummarySender.cpp b/Engine/Source/Developer/EditorAnalyticsSession/Private/EditorSessionSummarySender.cpp
--- a/Engine/Source/Developer/EditorAnalyticsSession/Private/EditorSessionSummarySender.cpp
+++ b/Engine/Source/Developer/EditorAnalyticsSession/Private/EditorSessionSummarySender.cpp
@@ -186,7 +186,7 @@ void FEditorSessionSummarySender::SendSessionSummaryEvent(const FEditorAnalytics
 	AnalyticsAttributes.Emplace(TEXT("ShutdownType"), ShutdownTypeString);
 	AnalyticsAttributes.Emplace(TEXT("StartupTimestamp"), Session.StartupTimestamp.ToIso8601());
 	AnalyticsAttributes.Emplace(TEXT("Timestamp"), Session.Timestamp.ToIso8601());
-	AnalyticsAttributes.Emplace(TEXT("SessionDuration"), Session.SessionDuration);
+	AnalyticsAttributes.Emplace(TEXT("SessionDuration"), FMath::FloorToInt(static_cast<float>((Session.Timestamp - Session.StartupTimestamp).GetTotalSeconds())));
 	AnalyticsAttributes.Emplace(TEXT("1MinIdle"), Session.Idle1Min);
 	AnalyticsAttributes.Emplace(TEXT("5MinIdle"), Session.Idle5Min);
 	AnalyticsAttributes.Emplace(TEXT("30MinIdle"), Session.Idle30Min);
diff --git a/Engine/Source/Developer/EditorAnalyticsSession/Public/EditorAnalyticsSession.h b/Engine/Source/Developer/EditorAnalyticsSession/Public/EditorAnalyticsSession.h
--- a/Engine/Source/Developer/EditorAnalyticsSession/Public/EditorAnalyticsSession.h
+++ b/Engine/Source/Developer/EditorAnalyticsSession/Public/EditorAnalyticsSession.h
@@ -6,9 +6,17 @@
 #include "HAL/CriticalSection.h"
 #include "Misc/DateTime.h"
 #include "Modules/ModuleInterface.h"
+#include "Templates/Atomic.h"
 
 struct EDITORANALYTICSSESSION_API FEditorAnalyticsSession
 {
+	enum class EEventType
+	{
+		Crashed = 0,
+		GpuCrashed,
+		Terminated,
+	};
+
 	FString SessionId;
 
 	FString AppId;
@@ -24,10 +32,10 @@ struct EDITORANALYTICSSESSION_API FEditorAnalyticsSession
 
 	FDateTime StartupTimestamp;
 	FDateTime Timestamp;
-	int32 SessionDuration;
-	int32 Idle1Min;
-	int32 Idle5Min;
-	int32 Idle30Min;
+	volatile int32 IdleSeconds; // Can be updated from concurrent threads.
+	volatile int32 Idle1Min;
+	volatile int32 Idle5Min;
+	volatile int32 Idle30Min;
 	FString CurrentUserActivity;
 	TArray<FString> Plugins;
 	float AverageFPS;
@@ -49,8 +57,8 @@ struct EDITORANALYTICSSESSION_API FEditorAnalyticsSession
 	FString OSMajor;
 	FString OSMinor;
 	FString OSVersion;
-	bool bIs64BitOS : 1;
 
+	bool bIs64BitOS : 1;
 	bool bCrashed : 1;
 	bool bGPUCrashed : 1;
 	bool bIsDebugger : 1;
@@ -71,12 +79,6 @@ struct EDITORANALYTICSSESSION_API FEditorAnalyticsSession
 	 */
 	bool Save();
 
-	/**
-	 * Save only the bare minimum of fields that may have changed during a crash.
-	 * @returns true if the session was successfully saved.
-	 */
-	bool SaveForCrash();
-
 	/**
 	 *  Load a session with the given session ID from stored values.
 	 * @returns true if the session was found and successfully loaded.
@@ -122,8 +124,19 @@ struct EDITORANALYTICSSESSION_API FEditorAnalyticsSession
 	/** Is the local storage already locked? */
 	static bool IsLocked();
 
-private:
+	/**
+	 * Append an event to the session log. The function is meant to record concurrent events, especially during a crash
+	 * with minimum contention. The logger appends and persists the events of interest locklessly on spot as opposed to
+	 * overriding existing values in the key-store. Appending is better because it prevent dealing with event ordering
+	 * on the spot (no synchronization needed) and will preserve more information.
+	 *
+	 * @note They key-store is not easily usable in a 'lockless' fashion. On Windows, the OS provides thread safe API
+	 *       to modify the registry (add/update). On Mac/Linux, the key-store is a simple file and without synchronization,
+	 *       concurrent writes will likely corrupt the file.
+	 */
+	void LogEvent(EEventType EventTpe, const FDateTime& Timestamp);
 
+private:
 	static FSystemWideCriticalSection* StoredValuesLock;
 
 	/** 
diff --git a/Engine/Source/Programs/CrashReportClient/Private/CrashReportClientApp.cpp b/Engine/Source/Programs/CrashReportClient/Private/CrashReportClientApp.cpp
--- a/Engine/Source/Programs/CrashReportClient/Private/CrashReportClientApp.cpp
+++ b/Engine/Source/Programs/CrashReportClient/Private/CrashReportClientApp.cpp
@@ -940,10 +940,10 @@ void RunCrashReportClient(const TCHAR* CommandLine)
 			bool bAbnormalShutdownFromEditorPov = false;
 			bool bNormalShutdownFromOsPov = false; // Can only be true if the Editor process exit code is known and equal to zero.
 
+			// Send the editor summary event(s) first, maximizing chance of sucessfull transmission (less opportunities for bugs to prevent it).
+			FCrashReportAnalytics::Initialize();
+			if (FCrashReportAnalytics::IsAvailable())
 			{
-				// Send the editor summary event(s) first, maximizing chance of sucessfull transmission (less opportunities for bugs to prevent it).
-				FCrashReportAnalytics::Initialize();
-				if (FCrashReportAnalytics::IsAvailable())
 				{
 					// NOTE: The Editor doesn't create summary events if analytics are disabled (not permitted to send). It may still send pending events
 					//       accmulated while 'Send Data' was true, but will not send any newer.
@@ -975,21 +975,21 @@ void RunCrashReportClient(const TCHAR* CommandLine)
 					// Send summary session event(s).
 					EditorSessionSummarySender.Shutdown();
 				}
-				FCrashReportAnalytics::Shutdown();
-			}
 
-			// If the Editor hasn't called all its crash/exit handlers properly (session summary flags are not set) and the Editor exit code isn't known or different than zero.
-			if (bAbnormalShutdownFromEditorPov && !bNormalShutdownFromOsPov)
-			{
-				// Load our temporary crash context file.
-				FSharedCrashContext TempCrashContext;
-				FMemory::Memzero(TempCrashContext);
-				if (LoadTempCrashContextFromFile(TempCrashContext, MonitorPid) && TempCrashContext.UserSettings.bSendUsageData && TempCrashContext.UserSettings.bSendUnattendedBugReports)
+				// If the Editor hasn't called all its crash/exit handlers properly (session summary flags are not set) and the Editor exit code isn't known or different than zero.
+				if (bAbnormalShutdownFromEditorPov && !bNormalShutdownFromOsPov)
 				{
-					// Send a spoofed crash report in the case that we detect an abnormal shutdown has occurred
-					HandleAbnormalShutdown(TempCrashContext, MonitorPid, MonitorWritePipe, RecoveryServicePtr);
+					// Load our temporary crash context file.
+					FSharedCrashContext TempCrashContext;
+					FMemory::Memzero(TempCrashContext);
+					if (LoadTempCrashContextFromFile(TempCrashContext, MonitorPid) && TempCrashContext.UserSettings.bSendUsageData && TempCrashContext.UserSettings.bSendUnattendedBugReports)
+					{
+						// Send a spoofed crash report in the case that we detect an abnormal shutdown has occurred
+						HandleAbnormalShutdown(TempCrashContext, MonitorPid, MonitorWritePipe, RecoveryServicePtr);
+					}
 				}
 			}
+			FCrashReportAnalytics::Shutdown();
 		}
 #endif
 		// clean up the context file
diff --git a/Engine/Source/Runtime/Engine/Private/Analytics/EditorSessionSummaryWriter.cpp b/Engine/Source/Runtime/Engine/Private/Analytics/EditorSessionSummaryWriter.cpp
--- a/Engine/Source/Runtime/Engine/Private/Analytics/EditorSessionSummaryWriter.cpp
+++ b/Engine/Source/Runtime/Engine/Private/Analytics/EditorSessionSummaryWriter.cpp
@@ -33,33 +33,22 @@ DEFINE_LOG_CATEGORY_STATIC(LogEditorSessionSummary, Verbose, All);
 
 namespace EditorSessionWriterDefs
 {
+	// Number of seconds to wait between each update of the mutable metrics.
 	static const float HeartbeatPeriodSeconds = 60;
 }
 
-FEditorSessionSummaryWriter::FEditorSessionSummaryWriter() :
-	CurrentSession(nullptr)
-	, StartupSeconds(0.0)
-	, IdleSeconds(0.0f)
-	, HeartbeatTimeElapsed(0.0f)
+FEditorSessionSummaryWriter::FEditorSessionSummaryWriter()
+	: HeartbeatTimeElapsed(0.0f)
 	, bShutdown(false)
 {
 }
 
-void FEditorSessionSummaryWriter::Initialize()
+FEditorSessionSummaryWriter::~FEditorSessionSummaryWriter()
 {
-	// Register for crash and app state callbacks
-	FCoreDelegates::OnHandleSystemError.AddRaw(this, &FEditorSessionSummaryWriter::OnCrashing);
-	FCoreDelegates::ApplicationWillTerminateDelegate.AddRaw(this, &FEditorSessionSummaryWriter::OnTerminate);
-	FCoreDelegates::IsVanillaProductChanged.AddRaw(this, &FEditorSessionSummaryWriter::OnVanillaStateChanged);
-	FUserActivityTracking::OnActivityChanged.AddRaw(this, &FEditorSessionSummaryWriter::OnUserActivity);
-	FSlateApplication::Get().GetOnModalLoopTickEvent().AddRaw(this, &FEditorSessionSummaryWriter::Tick);
-
-	StartupSeconds = FPlatformTime::Seconds();
-
-	InitializeSessions();
+	Shutdown(); // In case it wasn't already called.
 }
 
-void FEditorSessionSummaryWriter::InitializeSessions()
+void FEditorSessionSummaryWriter::Initialize()
 {
 	if (!FEngineAnalytics::IsAvailable() || CurrentSession != nullptr)
 	{
@@ -68,7 +57,7 @@ void FEditorSessionSummaryWriter::InitializeSessions()
 
 	UE_LOG(LogEditorSessionSummary, Verbose, TEXT("Initializing EditorSessionSummaryWriter for editor session tracking"));
 
-	if (FEditorAnalyticsSession::Lock())
+	if (FEditorAnalyticsSession::Lock()) // System wide locking to write the session file/registry
 	{
 		// Create a session Session for this session
 		CurrentSession = CreateCurrentSession();
@@ -86,29 +75,44 @@ void FEditorSessionSummaryWriter::InitializeSessions()
 
 		FEditorAnalyticsSession::Unlock();
 	}
+
+	if (CurrentSession)
+	{
+		// Register for crash and app state callbacks
+		FCoreDelegates::OnHandleSystemError.AddRaw(this, &FEditorSessionSummaryWriter::OnCrashing); // WARNING: Don't assume this function is only called from game thread.
+		FCoreDelegates::ApplicationWillTerminateDelegate.AddRaw(this, &FEditorSessionSummaryWriter::OnTerminate); // WARNING: Don't assume this function is only called from game thread.
+		FCoreDelegates::IsVanillaProductChanged.AddRaw(this, &FEditorSessionSummaryWriter::OnVanillaStateChanged);
+		FUserActivityTracking::OnActivityChanged.AddRaw(this, &FEditorSessionSummaryWriter::OnUserActivity);
+		FSlateApplication::Get().GetOnModalLoopTickEvent().AddRaw(this, &FEditorSessionSummaryWriter::Tick);
+	}
 }
 
 void FEditorSessionSummaryWriter::UpdateTimestamps()
 {
 	CurrentSession->Timestamp = FDateTime::UtcNow();
-	CurrentSession->SessionDuration = FMath::FloorToInt(static_cast<float>((CurrentSession->Timestamp - CurrentSession->StartupTimestamp).GetTotalSeconds()));
+	UpdateIdleTimes();
+}
+
+void FEditorSessionSummaryWriter::UpdateIdleTimes()
+{
+	int32 IdleSecondsTmp = FPlatformAtomics::AtomicRead(&CurrentSession->IdleSeconds); // Atomically load only once.
 
 	// 1 + 1 minutes
-	if (IdleSeconds > (60 + 60))
+	if (IdleSecondsTmp > (60 + 60))
 	{
-		CurrentSession->Idle1Min += 1;
+		FPlatformAtomics::InterlockedIncrement(&CurrentSession->Idle1Min); // User spent one more minute as idle > 1 min
 	}
 
 	// 5 + 1 minutes
-	if (IdleSeconds > (5 * 60 + 60))
+	if (IdleSecondsTmp > (5 * 60 + 60))
 	{
-		CurrentSession->Idle5Min += 1;
+		FPlatformAtomics::InterlockedIncrement(&CurrentSession->Idle5Min); // User spent one more minute as idle > 5 min
 	}
 
 	// 30 + 1 minutes
-	if (IdleSeconds > (30 * 60 + 60))
+	if (IdleSecondsTmp > (30 * 60 + 60))
 	{
-		CurrentSession->Idle30Min += 1;
+		FPlatformAtomics::InterlockedIncrement(&CurrentSession->Idle30Min); // User spent one more minute as idle > 30 min
 	}
 }
 
@@ -119,9 +123,11 @@ void FEditorSessionSummaryWriter::Tick(float DeltaTime)
 		return;
 	}
 
-	// Note: Update idle time is Tick() because Slate cannot be invoked from any thread and UpdateTimeStamps() can be called from any crashing thread.
+	// Note: Update idle time in Tick() because Slate cannot be invoked from any thread and UpdateTimeStamps() can be called from a crashing thread.
 	//       Compute the idle time from Slate point-of view. Note that some tasks blocking the UI (such as importing large assets) may be considered idle time.
-	IdleSeconds = FSlateApplication::Get().GetCurrentTime() - FSlateApplication::Get().GetLastUserInteractionTime();
+	CurrentSession->IdleSeconds = FSlateApplication::Get().GetLastUserInteractionTime() != 0 ? // In case Slate did not register any interaction yet (ex the user just launches the Editor and goes away)
+		FMath::FloorToInt(static_cast<float>(FSlateApplication::Get().GetCurrentTime() - FSlateApplication::Get().GetLastUserInteractionTime())) :
+		FMath::FloorToInt(static_cast<float>((FDateTime::UtcNow() - CurrentSession->StartupTimestamp).GetTotalSeconds()));
 
 	HeartbeatTimeElapsed += DeltaTime;
 
@@ -129,8 +135,11 @@ void FEditorSessionSummaryWriter::Tick(float DeltaTime)
 	{
 		HeartbeatTimeElapsed = 0.0f;
 
-		// Try late initialization
-		InitializeSessions();
+		// Try late initialization (in case user toggled the analytics 'Send Data' on during the session).
+		if (CurrentSession == nullptr)
+		{
+			Initialize();
+		}
 
 		if (CurrentSession != nullptr)
 		{
@@ -164,43 +173,41 @@ void FEditorSessionSummaryWriter::Tick(float DeltaTime)
 
 void FEditorSessionSummaryWriter::LowDriveSpaceDetected()
 {
-	CurrentSession->bIsLowDriveSpace = true;
+	if (CurrentSession)
+	{
+		CurrentSession->bIsLowDriveSpace = true;
 
-	TrySaveCurrentSession();
+		TrySaveCurrentSession();
+	}
 }
 
 void FEditorSessionSummaryWriter::Shutdown()
 {
-	FCoreDelegates::OnHandleSystemError.RemoveAll(this);
-	FCoreDelegates::ApplicationWillTerminateDelegate.RemoveAll(this);
-	FCoreDelegates::IsVanillaProductChanged.RemoveAll(this);
-
-	FUserActivityTracking::OnActivityChanged.RemoveAll(this);
-
-	if (CurrentSession != nullptr)
+	// NOTE: Initialize(), Shutdown() and ~FEditorSessionSummaryWriter() are expected to be called from the game thread only.
+	if (CurrentSession && !bShutdown)
 	{
-		if (!CurrentSession->bIsTerminating && !CurrentSession->bCrashed)
-		{
-			FSlateApplication::Get().GetOnModalLoopTickEvent().RemoveAll(this);
-
-			CurrentSession->bWasShutdown = true;
-		}
-
+		// NOTE: Shutdown() may crash if a delegate is broadcasted from another thread at the same time (that's a bug in 4.25) the delegate are modified.
+		FCoreDelegates::ApplicationWillTerminateDelegate.RemoveAll(this);
+		FCoreDelegates::IsVanillaProductChanged.RemoveAll(this);
+		FUserActivityTracking::OnActivityChanged.RemoveAll(this);
+		FSlateApplication::Get().GetOnModalLoopTickEvent().RemoveAll(this);
+		FCoreDelegates::OnHandleSystemError.RemoveAll(this);
+
+		CurrentSession->bWasShutdown = true;
 		UpdateTimestamps();
-
 		TrySaveCurrentSession();
 
-		delete CurrentSession;
-		CurrentSession = nullptr;
-		bShutdown = true;
+		CurrentSession.Reset();
 	}
+
+	bShutdown = true;
 }
 
-FEditorAnalyticsSession* FEditorSessionSummaryWriter::CreateCurrentSession() const
+TUniquePtr<FEditorAnalyticsSession> FEditorSessionSummaryWriter::CreateCurrentSession()
 {
 	check(FEngineAnalytics::IsAvailable()); // The function assumes the caller checked it before calling.
 
-	FEditorAnalyticsSession* Session = new FEditorAnalyticsSession();
+	TUniquePtr<FEditorAnalyticsSession> Session = MakeUnique<FEditorAnalyticsSession>();
 	IAnalyticsProviderET& AnalyticProvider = FEngineAnalytics::GetProvider();
 
 	FGuid SessionId;
@@ -228,7 +235,6 @@ FEditorAnalyticsSession* FEditorSessionSummaryWriter::CreateCurrentSession() con
 	Session->ProjectVersion = ProjectSettings.ProjectVersion;
 	Session->EngineVersion = FEngineVersion::Current().ToString(EVersionComponent::Changelist);
 	Session->Timestamp = Session->StartupTimestamp = FDateTime::UtcNow();
-	Session->SessionDuration = 0;
 	Session->bIsDebugger = FPlatformMisc::IsDebuggerPresent();
 	Session->bWasEverDebugger = FPlatformMisc::IsDebuggerPresent();
 	Session->CurrentUserActivity = GetUserActivityString();
@@ -270,31 +276,31 @@ FEditorAnalyticsSession* FEditorSessionSummaryWriter::CreateCurrentSession() con
 extern CORE_API bool GIsGPUCrashed;
 void FEditorSessionSummaryWriter::OnCrashing()
 {
+	// NOTE: This function is called from the crashing thread or the a crash processing thread and is concurrent with other functions such as Tick(), Initialize() or Shutdown() running on the game thread.
 	if (CurrentSession != nullptr)
 	{
-		UpdateTimestamps();
+		UpdateIdleTimes();
+		CurrentSession->LogEvent(FEditorAnalyticsSession::EEventType::Crashed, FDateTime::UtcNow());
 
-		CurrentSession->bCrashed = true;
-		CurrentSession->bGPUCrashed = GIsGPUCrashed;
+		if (GIsGPUCrashed)
+		{
+			CurrentSession->LogEvent(FEditorAnalyticsSession::EEventType::GpuCrashed, FDateTime::UtcNow());
+		}
 
-		TrySaveCurrentSession();
+		// NOTE: Don't try to save the session, we don't know if the lock used to save the key-store is corrupted (or held by the crashing thread) when OnCrashing() is called from the crash handler thread.
 	}
 }
 
 void FEditorSessionSummaryWriter::OnTerminate()
 {
+	// NOTE: This function can be called from any thread (from the crashing thread too) and is likely concurrent with other functions such as Tick(), Initiallize() or Shutdown() running on the game thread.
 	if (CurrentSession != nullptr)
 	{
-		UpdateTimestamps();
-
-		CurrentSession->bIsTerminating = true;
-
-		TrySaveCurrentSession();
+		UpdateIdleTimes();
+		CurrentSession->LogEvent(FEditorAnalyticsSession::EEventType::Terminated, FDateTime::UtcNow());
 
-		if (IsEngineExitRequested())
-		{
-			Shutdown();
-		}
+		// NOTE: Don't try to save the session, we don't know if this is called from a crash handler (and if the crashing thread corrupted (or held) the lock to save the key-store.)
+		// NOTE: Don't explicitely Shutdown(), it is expected to be called on game thread to prevent unregistered delegate from a random thread. Just let the normal flow call Shutdown() or not. Destructor will do in last resort.
 	}
 }
 
@@ -318,7 +324,7 @@ void FEditorSessionSummaryWriter::OnUserActivity(const FUserActivity& UserActivi
 	}
 }
 
-FString FEditorSessionSummaryWriter::GetUserActivityString() const
+FString FEditorSessionSummaryWriter::GetUserActivityString()
 {
 	const FUserActivity& UserActivity = FUserActivityTracking::GetUserActivity();
 
@@ -332,10 +338,10 @@ FString FEditorSessionSummaryWriter::GetUserActivityString() const
 
 void FEditorSessionSummaryWriter::TrySaveCurrentSession()
 {
-	if (FEditorAnalyticsSession::Lock())
+	if (FEditorAnalyticsSession::Lock()) // Inter-process lock to grant this process exclusive access to the key-store file/registry.
 	{
+		FScopeLock ScopedLock(&SaveSessionLock); // Intra-process lock to grant the calling thread exclusive access to the key-store file/registry.
 		CurrentSession->Save();
-
 		FEditorAnalyticsSession::Unlock();
 	}
 }
diff --git a/Engine/Source/Runtime/Engine/Private/Analytics/EditorSessionSummaryWriter.h b/Engine/Source/Runtime/Engine/Private/Analytics/EditorSessionSummaryWriter.h
--- a/Engine/Source/Runtime/Engine/Private/Analytics/EditorSessionSummaryWriter.h
+++ b/Engine/Source/Runtime/Engine/Private/Analytics/EditorSessionSummaryWriter.h
@@ -14,6 +14,7 @@ class FEditorSessionSummaryWriter
 {
 public:
 	FEditorSessionSummaryWriter();
+	~FEditorSessionSummaryWriter();
 
 	void Initialize();
 	void Tick(float DeltaTime);
@@ -21,24 +22,21 @@ public:
 	void Shutdown();
 
 private:
-	void InitializeSessions();
-
-	FEditorAnalyticsSession* CreateCurrentSession() const;
-
 	void OnCrashing();
 	void OnTerminate();
 	void OnUserActivity(const FUserActivity& UserActivity);
 	void OnVanillaStateChanged(bool bIsVanilla);
 
-	FString GetUserActivityString() const;
+	static TUniquePtr<FEditorAnalyticsSession> CreateCurrentSession();
+	static FString GetUserActivityString();
 	void UpdateTimestamps();
+	void UpdateIdleTimes();
 	void TrySaveCurrentSession();
 
 private:
-	FEditorAnalyticsSession* CurrentSession;
+	TUniquePtr<FEditorAnalyticsSession> CurrentSession;
 	FString CurrentSessionSectionName;
-	double StartupSeconds;
-	double IdleSeconds;
+	FCriticalSection SaveSessionLock;
 	float HeartbeatTimeElapsed;
 	bool bShutdown;
 };
