Disable adjacency information for FStaticMeshSceneProxy when selection outline is enabled.

#rnx

diff --git a/Engine/Source/Runtime/Engine/Private/StaticMeshRender.cpp b/Engine/Source/Runtime/Engine/Private/StaticMeshRender.cpp
--- a/Engine/Source/Runtime/Engine/Private/StaticMeshRender.cpp
+++ b/Engine/Source/Runtime/Engine/Private/StaticMeshRender.cpp
@@ -564,7 +564,9 @@ bool FStaticMeshSceneProxy::GetMeshElement(
 	}
 
 	const bool bWireframe = false;
-	const bool bRequiresAdjacencyInformation = RequiresAdjacencyInformation(MaterialInterface, VertexFactory->GetType(), FeatureLevel);
+
+	// Disable adjacency information when the selection outline is enabled, since tessellation won't be used.
+	const bool bRequiresAdjacencyInformation = !bUseSelectionOutline && RequiresAdjacencyInformation(MaterialInterface, VertexFactory->GetType(), FeatureLevel);
 
 	// Two sided material use bIsFrontFace which is wrong with Reversed Indices. AdjacencyInformation use another index buffer.
 	const bool bUseReversedIndices = GUseReversedIndexBuffer && IsLocalToWorldDeterminantNegative() && (LOD.bHasReversedIndices != 0) && !bRequiresAdjacencyInformation && !Material->IsTwoSided();
