#ue4 - Fix regression from 8671782 causing actor overlaps to not properly trigger during initial level load even if bGenerateOverlapEventsDuringLevelStreaming was enabled. Fixed some incorrect related comments in the header.

diff --git a/Engine/Source/Runtime/Engine/Classes/GameFramework/Actor.h b/Engine/Source/Runtime/Engine/Classes/GameFramework/Actor.h
--- a/Engine/Source/Runtime/Engine/Classes/GameFramework/Actor.h
+++ b/Engine/Source/Runtime/Engine/Classes/GameFramework/Actor.h
@@ -280,8 +280,8 @@ public:
 	uint8 bFindCameraComponentWhenViewTarget:1;
 	
     /**
-	 * If true, this actor will generate overlap Begin/End events when spawned as part of level streaming.
-	 * You might enable this is in the case where a streaming level loads around an actor and you want Begin/End overlaps to trigger.
+	 * If true, this actor will generate overlap Begin/End events when spawned as part of level streaming, which includes initial level load.
+	 * You might enable this is in the case where a streaming level loads around an actor and you want Begin/End overlap events to trigger.
 	 * @see UpdateOverlapsMethodDuringLevelStreaming
 	 */
 	UPROPERTY(EditAnywhere, BlueprintReadWrite, Category=Collision)
@@ -411,7 +411,7 @@ protected:
 	 * If overlaps are not initialized, this actor and attached components will not have an initial state of what objects are touching it,
 	 * and overlap events may only come in once one of those objects update overlaps themselves (for example when moving).
 	 * However if an object touching it *does* initialize state, both objects will know about their touching state with each other.
-	 * This can be a potentially large performance savings during level streaming, and is safe if the object and others initially
+	 * This can be a potentially large performance savings during level loading and streaming, and is safe if the object and others initially
 	 * overlapping it do not need the overlap state because they will not trigger overlap notifications.
 	 * 
 	 * Note that if 'bGenerateOverlapEventsDuringLevelStreaming' is true, overlaps are always updated in this case, but that flag
@@ -1426,7 +1426,7 @@ public:
 	/** Returns whether an actor has had BeginPlay called on it (and not subsequently had EndPlay called) */
 	bool HasActorBegunPlay() const { return ActorHasBegunPlay == EActorBeginPlayState::HasBegunPlay; }
 
-	/** Returns whether an actor has called BeginPlay(). */
+	/** Returns whether an actor is beginning play in DispatchBeginPlay() during level streaming (which includes initial level load). */
 	bool IsActorBeginningPlayFromLevelStreaming() const { return bActorBeginningPlayFromLevelStreaming; }
 
 	/** Returns true if this actor is currently being destroyed, some gameplay events may be unsafe */
diff --git a/Engine/Source/Runtime/Engine/Private/Components/PrimitiveComponent.cpp b/Engine/Source/Runtime/Engine/Private/Components/PrimitiveComponent.cpp
--- a/Engine/Source/Runtime/Engine/Private/Components/PrimitiveComponent.cpp
+++ b/Engine/Source/Runtime/Engine/Private/Components/PrimitiveComponent.cpp
@@ -2656,7 +2656,8 @@ void UPrimitiveComponent::BeginComponentOverlap(const FOverlapInfo& OtherOverlap
 			AddUniqueOverlapFast(OtherComp->OverlappingComponents, FOverlapInfo(this, INDEX_NONE));	// uniqueness unverified, so addunique
 			
 			const UWorld* World = GetWorld();
-			if (bDoNotifies && World && World->HasBegunPlay())
+			const bool bLevelStreamingOverlap = (MyActor && MyActor->bGenerateOverlapEventsDuringLevelStreaming && MyActor->IsActorBeginningPlayFromLevelStreaming());
+			if (bDoNotifies && ((World && World->HasBegunPlay()) || bLevelStreamingOverlap))
 			{
 				// first execute component delegates
 				if (!IsPendingKill())
