Add ElementIndex to HitResult. By default this will contain the shape index, but is named generically so that if we ever do away with the concept of shapes it can be used to represent the index of an element in an aggregate (or whatever we call single objects with multiple "shapes" in the future).

diff --git a/Engine/Source/Runtime/Engine/Classes/Engine/EngineTypes.h b/Engine/Source/Runtime/Engine/Classes/Engine/EngineTypes.h
--- a/Engine/Source/Runtime/Engine/Classes/Engine/EngineTypes.h
+++ b/Engine/Source/Runtime/Engine/Classes/Engine/EngineTypes.h
@@ -2016,6 +2016,10 @@ struct ENGINE_API FHitResult
 	UPROPERTY()
 	int32 Item;
 
+	/** Index to item that was hit, also hit primitive specific. */
+	UPROPERTY()
+	uint8 ElementIndex;
+
 	/**
 	 * Physical material that was hit.
 	 * @note Must set bReturnPhysicalMaterial on the swept PrimitiveComponent or in the query params for this to be returned.
diff --git a/Engine/Source/Runtime/Engine/Private/Collision/Collision.cpp b/Engine/Source/Runtime/Engine/Private/Collision/Collision.cpp
--- a/Engine/Source/Runtime/Engine/Private/Collision/Collision.cpp
+++ b/Engine/Source/Runtime/Engine/Private/Collision/Collision.cpp
@@ -45,7 +45,7 @@ bool FHitResult::NetSerialize(FArchive& Ar, class UPackageMap* Map, bool& bOutSu
 	bool bImpactPointEqualsLocation = 0, bImpactNormalEqualsNormal = 0;
 
 	// Often times the indexes are invalid, use that as an optimization
-	bool bInvalidItem = 0, bInvalidFaceIndex = 0, bNoPenetrationDepth = 0;
+	bool bInvalidItem = 0, bInvalidFaceIndex = 0, bNoPenetrationDepth = 0, bInvalidElementIndex = 0;
 
 	if (Ar.IsSaving())
 	{
@@ -54,11 +54,12 @@ bool FHitResult::NetSerialize(FArchive& Ar, class UPackageMap* Map, bool& bOutSu
 		bInvalidItem = (Item == INDEX_NONE);
 		bInvalidFaceIndex = (FaceIndex == INDEX_NONE);
 		bNoPenetrationDepth = (PenetrationDepth == 0.0f);
+		bInvalidElementIndex = (ElementIndex == INDEX_NONE);
 	}
 
 	// pack bitfield with flags
-	uint8 Flags = (bBlockingHit << 0) | (bStartPenetrating << 1) | (bImpactPointEqualsLocation << 2) | (bImpactNormalEqualsNormal << 3) | (bInvalidItem << 4) | (bInvalidFaceIndex << 5) | (bNoPenetrationDepth << 6);
-	Ar.SerializeBits(&Flags, 7); 
+	uint8 Flags = (bBlockingHit << 0) | (bStartPenetrating << 1) | (bImpactPointEqualsLocation << 2) | (bImpactNormalEqualsNormal << 3) | (bInvalidItem << 4) | (bInvalidFaceIndex << 5) | (bNoPenetrationDepth << 6) | (bInvalidElementIndex << 7);
+	Ar.SerializeBits(&Flags, 8); 
 	bBlockingHit = (Flags & (1 << 0)) ? 1 : 0;
 	bStartPenetrating = (Flags & (1 << 1)) ? 1 : 0;
 	bImpactPointEqualsLocation = (Flags & (1 << 2)) ? 1 : 0;
@@ -66,6 +67,9 @@ bool FHitResult::NetSerialize(FArchive& Ar, class UPackageMap* Map, bool& bOutSu
 	bInvalidItem = (Flags & (1 << 4)) ? 1 : 0;
 	bInvalidFaceIndex = (Flags & (1 << 5)) ? 1 : 0;
 	bNoPenetrationDepth = (Flags & (1 << 6)) ? 1 : 0;
+	bInvalidElementIndex = (Flags & (1 << 7)) ? 1 : 0;
+	// NOTE: Every bit in Flags is being used. If any more bits are added,
+	// Flags will need to be upgraded to a uint16.
 
 	Ar << Time;
 
@@ -138,6 +142,14 @@ bool FHitResult::NetSerialize(FArchive& Ar, class UPackageMap* Map, bool& bOutSu
 		FaceIndex = INDEX_NONE;
 	}
 	
+	if (!bInvalidElementIndex)
+	{
+		Ar << ElementIndex;
+	}
+	else if (Ar.IsLoading())
+	{
+		ElementIndex = INDEX_NONE;
+	}
 
 	return true;
 }
diff --git a/Engine/Source/Runtime/Engine/Private/Collision/CollisionConversions.cpp b/Engine/Source/Runtime/Engine/Private/Collision/CollisionConversions.cpp
--- a/Engine/Source/Runtime/Engine/Private/Collision/CollisionConversions.cpp
+++ b/Engine/Source/Runtime/Engine/Private/Collision/CollisionConversions.cpp
@@ -134,6 +134,12 @@ static FVector FindGeomOpposingNormal(ECollisionShapeType QueryGeomType, const F
 static void SetHitResultFromShapeAndFaceIndex(const FPhysicsShape& Shape,  const FPhysicsActor& Actor, const uint32 FaceIndex, const FVector& HitLocation, FHitResult& OutResult, bool bReturnPhysMat)
 {
 	SCOPE_CYCLE_COUNTER(STAT_CollisionSetHitResultFromShapeAndFaceIndex);
+
+#if WITH_CHAOS
+		const int32 ShapeIndex = Shape.GetShapeIndex();
+		CHAOS_CHECK(ShapeIndex < (int32)TNumericLimits<uint8>::Max()); // I could just write < 256, but this makes it more clear *why*
+		OutResult.ElementIndex = (uint8)ShapeIndex;
+#endif
 	
 	UPrimitiveComponent* OwningComponent = nullptr;
 	if(const FBodyInstance* BodyInst = GetUserData(Actor))
diff --git a/Engine/Source/Runtime/Experimental/Chaos/Public/Chaos/GeometryParticles.h b/Engine/Source/Runtime/Experimental/Chaos/Public/Chaos/GeometryParticles.h
--- a/Engine/Source/Runtime/Experimental/Chaos/Public/Chaos/GeometryParticles.h
+++ b/Engine/Source/Runtime/Experimental/Chaos/Public/Chaos/GeometryParticles.h
@@ -200,6 +200,11 @@ namespace Chaos
 			}
 		}
 
+		int32 GetShapeIndex() const
+		{
+			return ShapeIdx;
+		}
+
 	private:
 
 		class IPhysicsProxyBase* Proxy;
diff --git a/Engine/Source/Runtime/PhysicsCore/Public/ChaosSQTypes.h b/Engine/Source/Runtime/PhysicsCore/Public/ChaosSQTypes.h
--- a/Engine/Source/Runtime/PhysicsCore/Public/ChaosSQTypes.h
+++ b/Engine/Source/Runtime/PhysicsCore/Public/ChaosSQTypes.h
@@ -29,7 +29,7 @@ namespace ChaosInterface
 
 	struct FQueryHit : public FActorShape
 	{
-		FQueryHit() : FaceIndex(-1) {} 
+		FQueryHit() : FaceIndex(-1), ElementIndex(-1) {} 
 
 		/**
 		Face index of touched triangle, for triangle meshes, convex meshes and height fields. Defaults to -1 if face index is not available
@@ -37,6 +37,8 @@ namespace ChaosInterface
 
 		int32 FaceIndex; // Signed int to match TArray's size type, and so INDEX_NONE/-1 doesn't underflow.
 
+		int32 ElementIndex; // Currently used to indicate which shape was hit for a particle with multiple shapes.
+
 		void Serialize(Chaos::FChaosArchive& Ar);
 	};
 
