* Add ToggleableCollisionShapes to vehicle configs
* Add replicated list of disabled collision shapes to vehicle
* Add code for toggling collision shapes and associating collision shapes with skeletal mesh bones in vehicles

diff --git a/Engine/Source/Runtime/Engine/Private/PhysicsEngine/Experimental/PhysInterface_Chaos.cpp b/Engine/Source/Runtime/Engine/Private/PhysicsEngine/Experimental/PhysInterface_Chaos.cpp
--- a/Engine/Source/Runtime/Engine/Private/PhysicsEngine/Experimental/PhysInterface_Chaos.cpp
+++ b/Engine/Source/Runtime/Engine/Private/PhysicsEngine/Experimental/PhysInterface_Chaos.cpp
@@ -2207,8 +2207,8 @@ uint32 GetTriangleMeshExternalFaceIndex(const FPhysicsShape& Shape, uint32 Inter
 template<typename AllocatorType>
 int32 GetAllShapesInternal_AssumedLocked(const FPhysicsActorHandle& InActorHandle, TArray<FPhysicsShapeReference_Chaos, AllocatorType>& OutShapes)
 {
-	OutShapes.Reset();
 	const Chaos::FShapesArray& ShapesArray = InActorHandle->ShapesArray();
+	OutShapes.Reset(ShapesArray.Num());
 	//todo: can we avoid this construction?
 	for (const TUniquePtr<Chaos::FPerShapeData>& Shape : ShapesArray)
 	{
