Fixed rotating movement component not working in standalone

diff --git a/Engine/Source/Runtime/Engine/Private/Components/MovementComponent.cpp b/Engine/Source/Runtime/Engine/Private/Components/MovementComponent.cpp
--- a/Engine/Source/Runtime/Engine/Private/Components/MovementComponent.cpp
+++ b/Engine/Source/Runtime/Engine/Private/Components/MovementComponent.cpp
@@ -114,7 +114,7 @@ void UMovementComponent::InitializeComponent()
 	Super::InitializeComponent();
 
 	// RootComponent is null in OnRegister for blueprint (non-native) root components.
-	if (!UpdatedComponent && bAutoRegisterUpdatedComponent)
+	if (bAutoRegisterUpdatedComponent && !IsValid(UpdatedComponent))
 	{
 		// Auto-register owner's root component if found.
 		if (AActor* MyActor = GetOwner())
