#ue4 - Fix SA warning.

diff --git a/Engine/Source/Runtime/Engine/Private/Components/PrimitiveComponent.cpp b/Engine/Source/Runtime/Engine/Private/Components/PrimitiveComponent.cpp
--- a/Engine/Source/Runtime/Engine/Private/Components/PrimitiveComponent.cpp
+++ b/Engine/Source/Runtime/Engine/Private/Components/PrimitiveComponent.cpp
@@ -2656,7 +2656,7 @@ void UPrimitiveComponent::BeginComponentOverlap(const FOverlapInfo& OtherOverlap
 			AddUniqueOverlapFast(OtherComp->OverlappingComponents, FOverlapInfo(this, INDEX_NONE));	// uniqueness unverified, so addunique
 			
 			const UWorld* World = GetWorld();
-			const bool bLevelStreamingOverlap = (MyActor && MyActor->bGenerateOverlapEventsDuringLevelStreaming && MyActor->IsActorBeginningPlayFromLevelStreaming());
+			const bool bLevelStreamingOverlap = (bDoNotifies && MyActor->bGenerateOverlapEventsDuringLevelStreaming && MyActor->IsActorBeginningPlayFromLevelStreaming());
 			if (bDoNotifies && ((World && World->HasBegunPlay()) || bLevelStreamingOverlap))
 			{
 				// first execute component delegates
