Consolidate single-pass stereo options under a Multi-View option
Fixes invalid combinations of single-pass techniques resulting in shader compiling errors

diff --git a/Engine/Plugins/Runtime/Oculus/OculusVR/Source/OculusEditor/Private/OculusToolWidget.cpp b/Engine/Plugins/Runtime/Oculus/OculusVR/Source/OculusEditor/Private/OculusToolWidget.cpp
--- a/Engine/Plugins/Runtime/Oculus/OculusVR/Source/OculusEditor/Private/OculusToolWidget.cpp
+++ b/Engine/Plugins/Runtime/Oculus/OculusVR/Source/OculusEditor/Private/OculusToolWidget.cpp
@@ -363,16 +363,16 @@ void SOculusToolWidget::Construct(const FArguments& InArgs)
 		&SOculusToolWidget::ForwardShadingEnable }
 	);
 
-	SimpleSettings.Add(FName("InstancedStereo"), {
-		FName("InstancedStereo"),
-		LOCTEXT("InstancedStereoDescription", "Instanced stereo is not enabled for this project. Instanced stereo substantially reduces draw calls, and improves rendering performance."),
-		&SOculusToolWidget::InstancedStereoVisibility,
+	SimpleSettings.Add(FName("MultiView"), {
+		FName("MultiView"),
+		LOCTEXT("MultiViewDescription", "Multi-view is not enabled for this project. Instanced stereo substantially reduces draw calls, and improves rendering performance."),
+		&SOculusToolWidget::MultiViewVisibility,
 		TArray<SimpleSettingAction>(),
 		(int)SupportFlags::SupportPC
 	});
-	SimpleSettings.Find(FName("InstancedStereo"))->actions.Add(
-		{ LOCTEXT("InstancedStereoButtonText", "Enable Instanced Stereo"),
-		&SOculusToolWidget::InstancedStereoEnable }
+	SimpleSettings.Find(FName("MultiView"))->actions.Add(
+		{ LOCTEXT("MultiViewButtonText", "Enable Multi-View"),
+		&SOculusToolWidget::MultiViewEnable }
 	);
 
 	SimpleSettings.Add(FName("MobileMultiView"), {
@@ -387,16 +387,16 @@ void SOculusToolWidget::Construct(const FArguments& InArgs)
 		&SOculusToolWidget::MobileMultiViewEnable }
 	);
 
-	SimpleSettings.Add(FName("MobileHDR"), {
-		FName("MobileHDR"),
-		LOCTEXT("MobileHDRDescription", "Mobile HDR has performance and stability issues in VR. We strongly recommend disabling it."),
-		&SOculusToolWidget::MobileHDRVisibility,
+	SimpleSettings.Add(FName("MobilePostProcessing"), {
+		FName("MobilePostProcessing"),
+		LOCTEXT("MobilePostProcessingDescription", "Mobile Post-Processing has performance and stability issues in VR. We strongly recommend disabling it."),
+		&SOculusToolWidget::MobilePostProcessingVisibility,
 		TArray<SimpleSettingAction>(),
 		(int)SupportFlags::SupportMobile
 	});
-	SimpleSettings.Find(FName("MobileHDR"))->actions.Add(
-		{ LOCTEXT("MobileHDRButton", "Disable Mobile HDR"),
-		&SOculusToolWidget::MobileHDRDisable }
+	SimpleSettings.Find(FName("MobilePostProcessing"))->actions.Add(
+		{ LOCTEXT("MobileHDRButton", "Disable Mobile Post-Processing"),
+		&SOculusToolWidget::MobilePostProcessingDisable }
 	);
 
 	SimpleSettings.Add(FName("AndroidManifest"), {
@@ -576,30 +576,28 @@ EVisibility SOculusToolWidget::ForwardShadingVisibility(FName tag) const
 	return UsingForwardShading() ? EVisibility::Collapsed : EVisibility::Visible;
 }
 
-FReply SOculusToolWidget::InstancedStereoEnable(bool text)
+FReply SOculusToolWidget::MultiViewEnable(bool text)
 {
 	URendererSettings* Settings = GetMutableDefault<URendererSettings>();
-	Settings->bInstancedStereo = 1;
-	Settings->UpdateSinglePropertyInConfigFile(Settings->GetClass()->FindPropertyByName(GET_MEMBER_NAME_CHECKED(URendererSettings, bInstancedStereo)), Settings->GetDefaultConfigFilename());
+	Settings->bMultiView = 1;
+	Settings->UpdateSinglePropertyInConfigFile(Settings->GetClass()->FindPropertyByName(GET_MEMBER_NAME_CHECKED(URendererSettings, bMultiView)), Settings->GetDefaultConfigFilename());
 	SuggestRestart();
 	return FReply::Handled();
 }
 
-EVisibility SOculusToolWidget::InstancedStereoVisibility(FName tag) const
+EVisibility SOculusToolWidget::MultiViewVisibility(FName tag) const
 {
 	URendererSettings* Settings = GetMutableDefault<URendererSettings>();
-	const bool bInstancedStereo = Settings->bInstancedStereo != 0;
+	const bool bMultiView = Settings->bMultiView != 0;
 
-	return bInstancedStereo ? EVisibility::Collapsed : EVisibility::Visible;
+	return bMultiView ? EVisibility::Collapsed : EVisibility::Visible;
 }
 
 FReply SOculusToolWidget::MobileMultiViewEnable(bool text)
 {
 	URendererSettings* Settings = GetMutableDefault<URendererSettings>();
 	Settings->bMobileMultiView = 1;
-	Settings->bMobileMultiViewDirect = 1;
 	Settings->UpdateSinglePropertyInConfigFile(Settings->GetClass()->FindPropertyByName(GET_MEMBER_NAME_CHECKED(URendererSettings, bMobileMultiView)), Settings->GetDefaultConfigFilename());
-	Settings->UpdateSinglePropertyInConfigFile(Settings->GetClass()->FindPropertyByName(GET_MEMBER_NAME_CHECKED(URendererSettings, bMobileMultiViewDirect)), Settings->GetDefaultConfigFilename());
 	SuggestRestart();
 	return FReply::Handled();
 }
@@ -608,24 +606,23 @@ EVisibility SOculusToolWidget::MobileMultiViewVisibility(FName tag) const
 {
 	URendererSettings* Settings = GetMutableDefault<URendererSettings>();
 	const bool bMMV = Settings->bMobileMultiView != 0;
-	const bool bMMVD = Settings->bMobileMultiViewDirect != 0;
 
-	return (bMMV && bMMVD) ? EVisibility::Collapsed : EVisibility::Visible;
+	return bMMV ? EVisibility::Collapsed : EVisibility::Visible;
 }
 
-FReply SOculusToolWidget::MobileHDRDisable(bool text)
+FReply SOculusToolWidget::MobilePostProcessingDisable(bool text)
 {
 	URendererSettings* Settings = GetMutableDefault<URendererSettings>();
-	Settings->bMobileHDR = 0;
-	Settings->UpdateSinglePropertyInConfigFile(Settings->GetClass()->FindPropertyByName(GET_MEMBER_NAME_CHECKED(URendererSettings, bMobileHDR)), Settings->GetDefaultConfigFilename());
+	Settings->bMobilePostProcessing = 0;
+	Settings->UpdateSinglePropertyInConfigFile(Settings->GetClass()->FindPropertyByName(GET_MEMBER_NAME_CHECKED(URendererSettings, bMobilePostProcessing)), Settings->GetDefaultConfigFilename());
 	SuggestRestart();
 	return FReply::Handled();
 }
 
-EVisibility SOculusToolWidget::MobileHDRVisibility(FName tag) const
+EVisibility SOculusToolWidget::MobilePostProcessingVisibility(FName tag) const
 {
 	URendererSettings* Settings = GetMutableDefault<URendererSettings>();
-	return Settings->bMobileHDR == 0 ? EVisibility::Collapsed : EVisibility::Visible;
+	return Settings->bMobilePostProcessing == 0 ? EVisibility::Collapsed : EVisibility::Visible;
 }
 
 FString SOculusToolWidget::GetConfigPath() const
diff --git a/Engine/Plugins/Runtime/Oculus/OculusVR/Source/OculusEditor/Private/OculusToolWidget.h b/Engine/Plugins/Runtime/Oculus/OculusVR/Source/OculusEditor/Private/OculusToolWidget.h
--- a/Engine/Plugins/Runtime/Oculus/OculusVR/Source/OculusEditor/Private/OculusToolWidget.h
+++ b/Engine/Plugins/Runtime/Oculus/OculusVR/Source/OculusEditor/Private/OculusToolWidget.h
@@ -77,8 +77,8 @@ protected:
 	FReply ForwardShadingEnable(bool text);
 	EVisibility ForwardShadingVisibility(FName tag) const;
 
-	FReply InstancedStereoEnable(bool text);
-	EVisibility InstancedStereoVisibility(FName tag) const;
+	FReply MultiViewEnable(bool text);
+	EVisibility MultiViewVisibility(FName tag) const;
 
 	FReply StartInVREnable(bool text);
 	EVisibility StartInVRVisibility(FName tag) const;
@@ -92,8 +92,8 @@ protected:
 	FReply MobileMultiViewEnable(bool text);
 	EVisibility MobileMultiViewVisibility(FName tag) const;
 
-	FReply MobileHDRDisable(bool text);
-	EVisibility MobileHDRVisibility(FName tag) const;
+	FReply MobilePostProcessingDisable(bool text);
+	EVisibility MobilePostProcessingVisibility(FName tag) const;
 
 	FReply AndroidManifestGearGo(bool text);
 	FReply AndroidManifestQuest(bool text);
diff --git a/Engine/Shaders/Public/ShaderVersion.ush b/Engine/Shaders/Public/ShaderVersion.ush
--- a/Engine/Shaders/Public/ShaderVersion.ush
+++ b/Engine/Shaders/Public/ShaderVersion.ush
@@ -3,4 +3,4 @@
 // in Platform.ush (which should be included in any shader) it allows to invalidate the shader DDC.
 // 
 // If you are merging streams and there is a conflict with this GUID you should make a new GUID rather than taking one or the other.
-// GUID = 7002869D0B0F4343B804036352C03423
+// GUID = 8F43812D-41AA-4930-8F86-2AE388575C56
\ No newline at end of file
diff --git a/Engine/Source/Editor/HardwareTargeting/Private/HardwareTargetingModule.cpp b/Engine/Source/Editor/HardwareTargeting/Private/HardwareTargetingModule.cpp
--- a/Engine/Source/Editor/HardwareTargeting/Private/HardwareTargetingModule.cpp
+++ b/Engine/Source/Editor/HardwareTargeting/Private/HardwareTargetingModule.cpp
@@ -222,9 +222,6 @@ void FHardwareTargetingModule::GatherSettings(FMetaSettingGatherer& Builder)
 	const bool bAnyScalable = Settings->DefaultGraphicsPerformance == EGraphicsPreset::Scalable;
 
 	{
-		// Based roughly on https://docs.unrealengine.com/latest/INT/Platforms/Mobile/PostProcessEffects/index.html
-		UE_META_SETTING_ENTRY(Builder, URendererSettings, bMobileHDR, !bLowEndMobile);
-
 		// Bloom works and isn't terribly expensive on anything beyond low-end
 		UE_META_SETTING_ENTRY(Builder, URendererSettings, bDefaultFeatureBloom, !bLowEndMobile);
 
diff --git a/Engine/Source/Runtime/Engine/Classes/Engine/RendererSettings.h b/Engine/Source/Runtime/Engine/Classes/Engine/RendererSettings.h
--- a/Engine/Source/Runtime/Engine/Classes/Engine/RendererSettings.h
+++ b/Engine/Source/Runtime/Engine/Classes/Engine/RendererSettings.h
@@ -170,12 +170,6 @@ class ENGINE_API URendererSettings : public UDeveloperSettings
 {
 	GENERATED_UCLASS_BODY()
 
-	UPROPERTY(config, EditAnywhere, Category=Mobile, meta=(
-		ConsoleVariable="r.MobileHDR",DisplayName="Mobile HDR",
-		ToolTip="If true, mobile renders in full HDR. Disable this setting for games that do not require lighting features for better performance on slow devices. Changing this setting requires restarting the editor.",
-		ConfigRestartRequired = true))
-	uint32 bMobileHDR:1;
-
 	UPROPERTY(config, EditAnywhere, Category = Mobile, meta = (
 		ConsoleVariable = "r.Mobile.DisableVertexFog", DisplayName = "Disable vertex fogging in mobile shaders",
 		ToolTip = "If true, vertex fog will be omitted from all mobile shaders. If your game does not use fog, you should choose this setting to increase shading performance.",
@@ -199,11 +193,6 @@ class ENGINE_API URendererSettings : public UDeveloperSettings
 		ConfigRestartRequired = true))
 		uint32 bMobileUseLegacyShadingModel : 1;
 
-	UPROPERTY(config, meta = (
-		ConsoleVariable = "r.Mobile.UseHWsRGBEncoding", DisplayName = "Single-pass linear rendering",
-		ToolTip = "If true then mobile single-pass (non mobile HDR) rendering will use HW accelerated sRGB encoding/decoding. Available only on Oculus for now."))
-		uint32 bMobileUseHWsRGBEncoding : 1;
-
 	UPROPERTY(config, EditAnywhere, Category = Mobile, meta=(
 		ConsoleVariable="r.Mobile.AllowDitheredLODTransition", DisplayName="Allow Dithered LOD Transition",
 		ToolTip="Whether to support 'Dithered LOD Transition' material option on mobile platforms. Enabling this may degrade performance as rendering will not benefit from Early-Z optimization.",
@@ -222,7 +211,7 @@ class ENGINE_API URendererSettings : public UDeveloperSettings
 		ToolTip = "Whether to support Virtual Textures on mobile. Requires general Virtual Texturing option enabled as well. Changing this setting requires restarting the editor.",
 		ConfigRestartRequired = true))
 	uint32 bMobileVirtualTextures : 1;
-	
+
 	UPROPERTY(config, EditAnywhere, Category = Materials, meta = (
 		ConsoleVariable = "r.DiscardUnusedQuality", DisplayName = "Game Discards Unused Material Quality Levels",
 		ToolTip = "When running in game mode, whether to keep shaders for all quality levels in memory or only those needed for the current quality level.\nUnchecked: Keep all quality levels in memory allowing a runtime quality level change. (default)\nChecked: Discard unused quality levels when loading content for the game, saving some memory."))
@@ -247,7 +236,7 @@ class ENGINE_API URendererSettings : public UDeveloperSettings
 		ConsoleVariable="r.MinScreenRadiusForCSMDepth",DisplayName="Min Screen Radius for Cascaded Shadow Maps",
 		ToolTip="Screen radius at which objects are culled for cascaded shadow map depth passes. Larger values can improve performance but can cause artifacts as objects stop casting shadows."))
 	float MinScreenRadiusForCSMdepth;
-	
+
 	UPROPERTY(config, EditAnywhere, Category=Culling, meta=(
 		ConsoleVariable="r.PrecomputedVisibilityWarning",DisplayName="Warn about no precomputed visibility",
 		ToolTip="Displays a warning when no precomputed visibility data is available for the current camera location. This can be helpful if you are making a game that relies on precomputed visibility, e.g. a first person mobile game."))
@@ -323,7 +312,7 @@ class ENGINE_API URendererSettings : public UDeveloperSettings
 
 	UPROPERTY(config, EditAnywhere, Category = Materials, meta = (
 		DisplayName = "Use anisotropic BRDF (Beta)",
-		ConsoleVariable = "r.AnisotropicBRDF", 
+		ConsoleVariable = "r.AnisotropicBRDF",
 		EditCondition = "!bBasePassOutputsVelocity || bAnisotropicBRDF && bBasePassOutputsVelocity",
 		ConfigRestartRequired = true,
 		ToolTip = "Uses an anisotropic BRDF for default lit and clear coat surface materials. Changing this setting requires restarting the editor.\n\nNote: can only be enabled if 'Output velocities during base pass' (r.BasePassOutputsVelocity) is disabled."))
@@ -444,11 +433,11 @@ class ENGINE_API URendererSettings : public UDeveloperSettings
 		ConsoleVariable = "r.DefaultFeature.AutoExposure", DisplayName = "Auto Exposure",
 		ToolTip = "Whether the default for AutoExposure is enabled or not (postprocess volume/camera/game setting can still override and enable or disable it independently)"))
 	uint32 bDefaultFeatureAutoExposure : 1;
-	
+
 	UPROPERTY(config, EditAnywhere, Category = DefaultSettings, meta = (
 		ConsoleVariable = "r.DefaultFeature.AutoExposure.Method", DisplayName = "Auto Exposure",
 		ToolTip = "The default method for AutoExposure(postprocess volume/camera/game setting can still override and enable or disable it independently)"))
-	TEnumAsByte<EAutoExposureMethodUI::Type> DefaultFeatureAutoExposure; 
+	TEnumAsByte<EAutoExposureMethodUI::Type> DefaultFeatureAutoExposure;
 
 	UPROPERTY(config, EditAnywhere, Category = DefaultSettings, meta = (
 		ConsoleVariable = "r.DefaultFeature.AutoExposure.Bias", DisplayName = "Auto Exposure Bias",
@@ -552,7 +541,7 @@ class ENGINE_API URendererSettings : public UDeveloperSettings
 		ToolTip="Enables emitting velocity during Base Pass rendering. Changing this setting requires restarting the editor.\n\nNote: can only be enabled if 'Use anisotropic BRDF' (r.AnisotropicBRDF) is disabled.\nNote: enabling this behaves as if 'Output velocities due to vertex deformation' (r.VertexDeformationOutputsVelocity) is also enabled."
 		))
 	uint32 bBasePassOutputsVelocity:1;
-	
+
 	UPROPERTY(config, EditAnywhere, Category=Optimizations, meta=(
 		DisplayName="Output velocities due to vertex deformation",
 		ConsoleVariable="r.VertexDeformationOutputsVelocity",
@@ -578,7 +567,7 @@ class ENGINE_API URendererSettings : public UDeveloperSettings
 		ToolTip = "The X size of the GPU simulation texture size. SizeX*SizeY determines the maximum number of GPU simulated particles in an emitter. Potentially overridden by CVar settings in BaseDeviceProfile.ini.",
 		ConfigRestartRequired = true))
 	int32 GPUSimulationTextureSizeX;
-	
+
 	UPROPERTY(config, EditAnywhere, Category = Optimizations, meta = (
 		ConsoleVariable = "fx.GPUSimulationTextureSizeY",
 		DisplayName = "GPU Particle simulation texture size - Y",
@@ -601,38 +590,37 @@ class ENGINE_API URendererSettings : public UDeveloperSettings
 		ConsoleVariable = "r.MorphTarget.Mode", DisplayName = "Use GPU for computing morph targets",
 		ToolTip = "Whether to use original CPU method (loop per morph then by vertex) or use a GPU-based method on Shader Model 5 hardware."))
 	uint32 bUseGPUMorphTargets : 1;
-	
+
 	UPROPERTY(config, EditAnywhere, Category = Debugging, meta = (
 		ConsoleVariable = "r.GPUCrashDebugging", DisplayName = "Enable vendor specific GPU crash analysis tools",
 		ToolTip = "Enables vendor specific GPU crash analysis tools.  Currently only supports NVIDIA Aftermath on DX11.",
 		ConfigRestartRequired = true))
 		uint32 bNvidiaAftermathEnabled : 1;
 
-	UPROPERTY(config, EditAnywhere, Category=VR, meta=(
-		ConsoleVariable="vr.InstancedStereo", DisplayName="Instanced Stereo",
-		ToolTip="Enable instanced stereo rendering (only available for D3D SM5 or PS4).",
-		ConfigRestartRequired=true))
-	uint32 bInstancedStereo:1;
-
 	UPROPERTY(config, EditAnywhere, Category = VR, meta = (
-		EditCondition = "bInstancedStereo",
-		ConsoleVariable = "vr.MultiView", DisplayName = "Multi-View",
-		ToolTip = "Enable multi-view for instanced stereo rendering (only available on the PS4).",
+		ConsoleVariable = "vr.InstancedStereo", DisplayName = "Multi-View",
+		ToolTip = "Enable single-pass stereoscopic rendering through view instancing or draw call instancing.",
 		ConfigRestartRequired = true))
 	uint32 bMultiView : 1;
 
+	UPROPERTY(config, EditAnywhere, Category = VR, meta=(
+		ConsoleVariable="r.MobileHDR",DisplayName="Mobile Post-Processing",
+		ToolTip="If true, mobile pipelines include a full post-processing pass with tonemapping. Disable this setting for a performance boost and to enable stereoscopic rendering optimizations. Changing this setting requires restarting the editor.",
+		ConfigRestartRequired = true))
+	uint32 bMobilePostProcessing:1;
+
 	UPROPERTY(config, EditAnywhere, Category = VR, meta = (
+		EditCondition = "!bMobilePostProcessing",
 		ConsoleVariable = "vr.MobileMultiView", DisplayName = "Mobile Multi-View",
-		ToolTip = "Enable mobile multi-view rendering (only available on Oculus Mobile and some Gear VR Android devices).",
+		ToolTip = "Enable single-pass stereoscopic rendering on mobile platforms.",
 		ConfigRestartRequired = true))
 		uint32 bMobileMultiView : 1;
-
+	
 	UPROPERTY(config, EditAnywhere, Category = VR, meta = (
-		EditCondition = "bMobileMultiView",
-		ConsoleVariable = "vr.MobileMultiView.Direct", DisplayName = "Mobile Multi-View Direct",
-		ToolTip = "Enable direct mobile multi-view rendering (only available on multi-view enabled Oculus Mobile, Gear VR, and Daydream Android devices).",
-		ConfigRestartRequired = true))
-		uint32 bMobileMultiViewDirect : 1;
+		EditCondition = "!bMobilePostProcessing",
+		ConsoleVariable = "r.Mobile.UseHWsRGBEncoding", DisplayName = "Single-pass linear rendering",
+		ToolTip = "If true then mobile single-pass (without post-processing) rendering will use HW accelerated sRGB encoding/decoding. Available only on Oculus for now."))
+		uint32 bMobileUseHWsRGBEncoding : 1;
 
 	UPROPERTY(config, EditAnywhere, Category = VR, meta = (
 		ConsoleVariable = "vr.RoundRobinOcclusion", DisplayName = "Round Robin Occlusion Queries",
@@ -690,16 +678,16 @@ class ENGINE_API URendererSettings : public UDeveloperSettings
 		ConfigRestartRequired = true))
 		uint32 bSupportPointLightWholeSceneShadows : 1;
 
-	/** 
+	/**
 	"Atmospheric fog requires permutations of the basepass shaders.  Disabling will reduce the number of shader permutations required per material. Changing this setting requires restarting the editor."
 	*/
 	UPROPERTY(config, EditAnywhere, Category = ShaderPermutationReduction, meta = (
-		ConsoleVariable = "r.SupportAtmosphericFog", DisplayName = "Support Atmospheric Fog",	
+		ConsoleVariable = "r.SupportAtmosphericFog", DisplayName = "Support Atmospheric Fog",
 		ConfigRestartRequired = true))
 		uint32 bSupportAtmosphericFog : 1;
 
 	/**
-	"The sky atmosphere component requires extra samplers/textures to be bound to apply aerial perspective on transparent surfaces (and all surfaces on mobile via per vertex evaluation)." 
+	"The sky atmosphere component requires extra samplers/textures to be bound to apply aerial perspective on transparent surfaces (and all surfaces on mobile via per vertex evaluation)."
 	*/
 	UPROPERTY(config, EditAnywhere, Category = ShaderPermutationReduction, meta = (
 		ConsoleVariable = "r.SupportSkyAtmosphere", DisplayName = "Support Sky Atmosphere",
diff --git a/Engine/Source/Runtime/Engine/Private/SceneView.cpp b/Engine/Source/Runtime/Engine/Private/SceneView.cpp
--- a/Engine/Source/Runtime/Engine/Private/SceneView.cpp
+++ b/Engine/Source/Runtime/Engine/Private/SceneView.cpp
@@ -17,6 +17,7 @@
 #include "Interfaces/Interface_PostProcessVolume.h"
 #include "Engine/TextureCube.h"
 #include "StereoRendering.h"
+#include "StereoRenderTargetManager.h"
 #include "IHeadMountedDisplay.h"
 #include "IXRTrackingSystem.h"
 #include "Engine/RendererSettings.h"
@@ -718,7 +719,8 @@ FSceneView::FSceneView(const FSceneViewInitOptions& InitOptions)
 	bool bPlatformRequiresReverseCulling = ((IsOpenGLPlatform(ShaderPlatform) || IsSwitchPlatform(ShaderPlatform)) && bUsingMobileRenderer && !IsPCPlatform(ShaderPlatform) && !IsVulkanMobilePlatform(ShaderPlatform));
 	static auto* MobileHDRCvar = IConsoleManager::Get().FindTConsoleVariableDataInt(TEXT("r.MobileHDR"));
 	check(MobileHDRCvar);
-	bReverseCulling = (bPlatformRequiresReverseCulling && MobileHDRCvar->GetValueOnAnyThread() == 0) ? !bReverseCulling : bReverseCulling;
+	const bool bSkipPostprocessing = MobileHDRCvar->GetValueOnAnyThread() == 0;
+	bReverseCulling = (bPlatformRequiresReverseCulling && bSkipPostprocessing) ? !bReverseCulling : bReverseCulling;
 
 	// Setup transformation constants to be used by the graphics hardware to transform device normalized depth samples
 	// into world oriented z.
@@ -753,23 +755,32 @@ FSceneView::FSceneView(const FSceneViewInitOptions& InitOptions)
 
 	// Query instanced stereo and multi-view state
 	static const auto CVar = IConsoleManager::Get().FindTConsoleVariableDataInt(TEXT("vr.InstancedStereo"));
-	bIsInstancedStereoEnabled = RHISupportsInstancedStereo(ShaderPlatform) ? (CVar ? (CVar->GetValueOnAnyThread() != false) : false) : false;
+	bIsInstancedStereoEnabled = !bUsingMobileRenderer && RHISupportsInstancedStereo(ShaderPlatform) && (CVar && CVar->GetValueOnAnyThread() != 0);
 
-	static const auto MultiViewCVar = IConsoleManager::Get().FindTConsoleVariableDataInt(TEXT("vr.MultiView"));
-	bIsMultiViewEnabled = RHISupportsMultiView(ShaderPlatform) && (MultiViewCVar && MultiViewCVar->GetValueOnAnyThread() != 0);
+	// TODO: Should be renamed to multi-viewport
+	bIsMultiViewEnabled = RHISupportsMultiView(ShaderPlatform) && bIsInstancedStereoEnabled;
 
 	static const auto MobileMultiViewCVar = IConsoleManager::Get().FindTConsoleVariableDataInt(TEXT("vr.MobileMultiView"));
-	bIsMobileMultiViewEnabled = RHISupportsMobileMultiView(ShaderPlatform) && (MobileMultiViewCVar && MobileMultiViewCVar->GetValueOnAnyThread() != 0);
+	bIsMobileMultiViewEnabled = bUsingMobileRenderer && bSkipPostprocessing && (MobileMultiViewCVar && MobileMultiViewCVar->GetValueOnAnyThread() != 0);
+	if (bIsMobileMultiViewEnabled && !RHISupportsMobileMultiView(ShaderPlatform))
+	{
+		// Native mobile multi-view is not supported, attempt to fall back to instancing
+		if (GRHISupportsArrayIndexFromAnyShader)
+		{
+			bIsInstancedStereoEnabled = true;
+		}
+		else
+		{
+			bIsMobileMultiViewEnabled = false;
+		}
+	}
 
-	// TODO: Test platform support for direct
-	static const auto MobileMultiViewDirectCVar = IConsoleManager::Get().FindTConsoleVariableDataInt(TEXT("vr.MobileMultiView.Direct"));
-	bIsMobileMultiViewDirectEnabled = (MobileMultiViewDirectCVar && MobileMultiViewDirectCVar->GetValueOnAnyThread() != 0);
+	// If the plugin uses separate render targets it is required to support mobile multi-view direct
+	IStereoRenderTargetManager* const StereoRenderTargetManager = GEngine->StereoRenderingDevice.IsValid() ? GEngine->StereoRenderingDevice->GetRenderTargetManager() : nullptr;
+	bIsMobileMultiViewDirectEnabled = bIsMobileMultiViewEnabled && StereoRenderTargetManager && StereoRenderTargetManager->ShouldUseSeparateRenderTarget();
 
 	bShouldBindInstancedViewUB = bIsInstancedStereoEnabled || bIsMobileMultiViewEnabled;
 
-	// If the device doesn't support mobile multi-view, disable it.
-	bIsMobileMultiViewEnabled = bIsMobileMultiViewEnabled && (GSupportsMobileMultiView || GRHISupportsArrayIndexFromAnyShader) && IStereoRendering::IsStereoEyePass(StereoPass);
-
 	SetupAntiAliasingMethod();
 
 	if (CVarEnableTemporalUpsample.GetValueOnAnyThread() && AntiAliasingMethod == AAM_TemporalAA)
diff --git a/Engine/Source/Runtime/Engine/Private/ShaderCompiler/ShaderCompiler.cpp b/Engine/Source/Runtime/Engine/Private/ShaderCompiler/ShaderCompiler.cpp
--- a/Engine/Source/Runtime/Engine/Private/ShaderCompiler/ShaderCompiler.cpp
+++ b/Engine/Source/Runtime/Engine/Private/ShaderCompiler/ShaderCompiler.cpp
@@ -3397,7 +3397,8 @@ void GlobalBeginCompileShader(
 		Input.Environment.SetDefine(TEXT("SHADING_PATH_DEFERRED"), 1);
 	}
 
-	if (FSceneInterface::GetShadingPath(GetMaxSupportedFeatureLevel((EShaderPlatform)Target.Platform)) == EShadingPath::Mobile)
+	const bool bUsingMobileRenderer = FSceneInterface::GetShadingPath(GetMaxSupportedFeatureLevel((EShaderPlatform)Target.Platform)) == EShadingPath::Mobile;
+	if (bUsingMobileRenderer)
 	{
 		Input.Environment.SetDefine(TEXT("SHADING_PATH_MOBILE"), 1);
 	}
@@ -3405,23 +3406,35 @@ void GlobalBeginCompileShader(
 	// Set VR definitions
 	{
 		static const auto CVarInstancedStereo = IConsoleManager::Get().FindTConsoleVariableDataInt(TEXT("vr.InstancedStereo"));
-		static const auto CVarMultiView = IConsoleManager::Get().FindTConsoleVariableDataInt(TEXT("vr.MultiView"));
 		static const auto CVarMobileMultiView = IConsoleManager::Get().FindTConsoleVariableDataInt(TEXT("vr.MobileMultiView"));
 		static const auto CVarODSCapture = IConsoleManager::Get().FindTConsoleVariableDataInt(TEXT("vr.ODSCapture"));
+		static const auto CVarMobileHDR = IConsoleManager::Get().FindTConsoleVariableDataInt(TEXT("r.MobileHDR"));
 
 		const bool bIsInstancedStereoCVar = CVarInstancedStereo ? (CVarInstancedStereo->GetValueOnGameThread() != 0) : false;
-		const bool bIsMultiViewCVar = CVarMultiView ? (CVarMultiView->GetValueOnGameThread() != 0) : false;
-		const bool bIsMobileMultiViewCVar = CVarMobileMultiView ? (CVarMobileMultiView->GetValueOnGameThread() != 0) : false;
+		const bool bIsMobileMultiViewCVar = CVarMobileMultiView && CVarMobileHDR ?
+			(CVarMobileMultiView->GetValueOnGameThread() != 0 && CVarMobileHDR->GetValueOnGameThread() == 0) : false;
 		const bool bIsODSCapture = CVarODSCapture && (CVarODSCapture->GetValueOnGameThread() != 0);
 
 		const EShaderPlatform ShaderPlatform = static_cast<EShaderPlatform>(Target.Platform);
-		
-		const bool bIsInstancedStereo = bIsInstancedStereoCVar && RHISupportsInstancedStereo(ShaderPlatform);
-		Input.Environment.SetDefine(TEXT("INSTANCED_STEREO"), bIsInstancedStereo);
-		Input.Environment.SetDefine(TEXT("MULTI_VIEW"), bIsInstancedStereo && bIsMultiViewCVar && RHISupportsMultiView(ShaderPlatform));
 
-		const bool bSupportsMobileMultiview = RHISupportsMobileMultiView(ShaderPlatform);
-		Input.Environment.SetDefine(TEXT("MOBILE_MULTI_VIEW"), bIsMobileMultiViewCVar && bSupportsMobileMultiview);
+		bool bIsInstancedStereo = !bUsingMobileRenderer && bIsInstancedStereoCVar && RHISupportsInstancedStereo(ShaderPlatform);
+		bool bIsMobileMultiview = bUsingMobileRenderer && bIsMobileMultiViewCVar;
+		if (bIsMobileMultiview && !RHISupportsMobileMultiView(ShaderPlatform))
+		{
+			// Native mobile multi-view is not supported, attempt to fall back to instancing
+			if (GRHISupportsArrayIndexFromAnyShader)
+			{
+				bIsInstancedStereo = true;
+			}
+			else
+			{
+				bIsMobileMultiview = false;
+			}
+		}
+
+		Input.Environment.SetDefine(TEXT("INSTANCED_STEREO"), bIsInstancedStereo);
+		Input.Environment.SetDefine(TEXT("MULTI_VIEW"), bIsInstancedStereo && RHISupportsMultiView(ShaderPlatform));
+		Input.Environment.SetDefine(TEXT("MOBILE_MULTI_VIEW"), bIsMobileMultiview);
 
 		// Throw a warning if we are silently disabling ISR due to missing platform support.
 		if (bIsInstancedStereoCVar && !bIsInstancedStereo && !GShaderCompilingManager->AreWarningsSuppressed(ShaderPlatform))
diff --git a/Engine/Source/Runtime/RHI/Public/RHIDefinitions.h b/Engine/Source/Runtime/RHI/Public/RHIDefinitions.h
--- a/Engine/Source/Runtime/RHI/Public/RHIDefinitions.h
+++ b/Engine/Source/Runtime/RHI/Public/RHIDefinitions.h
@@ -1596,7 +1596,7 @@ inline bool RHIHasTiledGPU(const FStaticShaderPlatform Platform)
 inline bool RHISupportsMobileMultiView(const FStaticShaderPlatform Platform)
 {
 	return Platform == EShaderPlatform::SP_OPENGL_ES3_1_ANDROID || IsVulkanMobilePlatform(Platform)
-		|| Platform == EShaderPlatform::SP_PCD3D_ES3_1 || FDataDrivenShaderPlatformInfo::GetSupportsMobileMultiView(Platform);
+		|| FDataDrivenShaderPlatformInfo::GetSupportsMobileMultiView(Platform);
 }
 
 inline bool RHISupportsNativeShaderLibraries(const FStaticShaderPlatform Platform)
diff --git a/Engine/Source/Runtime/RenderCore/Private/Shader.cpp b/Engine/Source/Runtime/RenderCore/Private/Shader.cpp
--- a/Engine/Source/Runtime/RenderCore/Private/Shader.cpp
+++ b/Engine/Source/Runtime/RenderCore/Private/Shader.cpp
@@ -1297,12 +1297,11 @@ void ShaderMapAppendKeyString(EShaderPlatform Platform, FString& KeyString)
 
 	{
 		static const auto CVarInstancedStereo = IConsoleManager::Get().FindTConsoleVariableDataInt(TEXT("vr.InstancedStereo"));
-		static const auto CVarMultiView = IConsoleManager::Get().FindTConsoleVariableDataInt(TEXT("vr.MultiView"));
 		static const auto CVarMobileMultiView = IConsoleManager::Get().FindTConsoleVariableDataInt(TEXT("vr.MobileMultiView"));
 		static const auto CVarODSCapture = IConsoleManager::Get().FindTConsoleVariableDataInt(TEXT("vr.ODSCapture"));
 
 		const bool bIsInstancedStereo = (RHISupportsInstancedStereo(Platform) && (CVarInstancedStereo && CVarInstancedStereo->GetValueOnGameThread() != 0));
-		const bool bIsMultiView = (RHISupportsMultiView(Platform) && (CVarMultiView && CVarMultiView->GetValueOnGameThread() != 0));
+		const bool bIsMultiView = (RHISupportsMultiView(Platform) && bIsInstancedStereo);
 
 		const bool bIsAndroidGLES = RHISupportsMobileMultiView(Platform);
 		const bool bIsMobileMultiView = (bIsAndroidGLES && (CVarMobileMultiView && CVarMobileMultiView->GetValueOnGameThread() != 0));
diff --git a/Engine/Source/Runtime/Renderer/Private/SceneRendering.cpp b/Engine/Source/Runtime/Renderer/Private/SceneRendering.cpp
--- a/Engine/Source/Runtime/Renderer/Private/SceneRendering.cpp
+++ b/Engine/Source/Runtime/Renderer/Private/SceneRendering.cpp
@@ -150,26 +150,12 @@ static TAutoConsoleVariable<int32> CVarInstancedStereo(
 	TEXT("0 to disable instanced stereo (default), 1 to enable."),
 	ECVF_ReadOnly | ECVF_RenderThreadSafe);
 
-static TAutoConsoleVariable<int32> CVarMultiView(
-	TEXT("vr.MultiView"),
-	0,
-	TEXT("0 to disable multi-view instanced stereo, 1 to enable.\n")
-	TEXT("Currently only supported by the PS4 & Metal RHIs."),
-	ECVF_ReadOnly | ECVF_RenderThreadSafe);
-
 static TAutoConsoleVariable<int32> CVarMobileMultiView(
 	TEXT("vr.MobileMultiView"),
 	0,
 	TEXT("0 to disable mobile multi-view, 1 to enable.\n"),
 	ECVF_ReadOnly | ECVF_RenderThreadSafe);
 
-static TAutoConsoleVariable<int32> CVarMobileMultiViewDirect(
-	TEXT("vr.MobileMultiView.Direct"),
-	0,
-	TEXT("0 to disable mobile multi-view direct, 1 to enable.\n")
-	TEXT("When enabled the scene color render target array is provided by the hmd plugin so we can skip the blit.\n"),
-	ECVF_ReadOnly | ECVF_RenderThreadSafe);
-
 static TAutoConsoleVariable<int32> CVarRoundRobinOcclusion(
 	TEXT("vr.RoundRobinOcclusion"),
 	0,
@@ -4237,4 +4223,4 @@ void RunGPUSkinCacheTransition(FRHICommandList& RHICmdList, FScene* Scene, EGPUS
 			GPUSkinCache->TransitionAllToReadable(RHICmdList);
 		}
 	}
-}
\ No newline at end of file
+}
