Fix issue where runtime virtual texture streaming textures didn't init with packed page tables.

diff --git a/Engine/Source/Runtime/Engine/Private/Texture2D.cpp b/Engine/Source/Runtime/Engine/Private/Texture2D.cpp
--- a/Engine/Source/Runtime/Engine/Private/Texture2D.cpp
+++ b/Engine/Source/Runtime/Engine/Private/Texture2D.cpp
@@ -2139,6 +2139,8 @@ class IAllocatedVirtualTexture* FVirtualTexture2DResource::AcquireAllocatedVT()
 		VTDesc.TileSize = VTData->TileSize;
 		VTDesc.TileBorderSize = VTData->TileBorderSize;
 		VTDesc.NumTextureLayers = VTData->GetNumLayers();
+		VTDesc.bShareDuplicateLayers = TextureOwner->IsVirtualTexturedWithSinglePhysicalSpace();
+
 		for (uint32 LayerIndex = 0u; LayerIndex < VTDesc.NumTextureLayers; ++LayerIndex)
 		{
 			VTDesc.ProducerHandle[LayerIndex] = ProducerHandle; // use the same producer for each layer
diff --git a/Engine/Source/Runtime/Engine/Private/VT/RuntimeVirtualTexture.cpp b/Engine/Source/Runtime/Engine/Private/VT/RuntimeVirtualTexture.cpp
--- a/Engine/Source/Runtime/Engine/Private/VT/RuntimeVirtualTexture.cpp
+++ b/Engine/Source/Runtime/Engine/Private/VT/RuntimeVirtualTexture.cpp
@@ -91,17 +91,17 @@ public:
 	}
 
 	/** Queues up render thread work to create resources and also releases any old resources. */
-	void Init(FVTProducerDescription const& InDesc, IVirtualTexture* InVirtualTextureProducer, bool InPrivateSpace)
+	void Init(FVTProducerDescription const& InDesc, IVirtualTexture* InVirtualTextureProducer, bool InSinglePhysicalSpace, bool InPrivateSpace)
 	{
 		FRuntimeVirtualTextureRenderResource* Resource = this;
 		
 		ENQUEUE_RENDER_COMMAND(FRuntimeVirtualTextureRenderResource_Init)(
-			[Resource, InDesc, InVirtualTextureProducer, InPrivateSpace](FRHICommandList& RHICmdList)
+			[Resource, InDesc, InVirtualTextureProducer, InSinglePhysicalSpace, InPrivateSpace](FRHICommandList& RHICmdList)
 		{
 			FVirtualTextureProducerHandle OldProducerHandle = Resource->ProducerHandle;
 			ReleaseVirtualTexture(Resource->AllocatedVirtualTexture);
 			Resource->ProducerHandle = GetRendererModule().RegisterVirtualTextureProducer(InDesc, InVirtualTextureProducer);
-			Resource->AllocatedVirtualTexture = AllocateVirtualTexture(InDesc, Resource->ProducerHandle, InPrivateSpace);
+			Resource->AllocatedVirtualTexture = AllocateVirtualTexture(InDesc, Resource->ProducerHandle, InSinglePhysicalSpace, InPrivateSpace);
 			// Release old producer after new one is created so that any destroy callbacks can access the new producer
 			GetRendererModule().ReleaseVirtualTextureProducer(OldProducerHandle);
 		});
@@ -125,7 +125,7 @@ public:
 
 protected:
 	/** Allocate in the global virtual texture system. */
-	static IAllocatedVirtualTexture* AllocateVirtualTexture(FVTProducerDescription const& InDesc, FVirtualTextureProducerHandle const& InProducerHandle, bool InPrivateSpace)
+	static IAllocatedVirtualTexture* AllocateVirtualTexture(FVTProducerDescription const& InDesc, FVirtualTextureProducerHandle const& InProducerHandle, bool InSinglePhysicalSpace, bool InPrivateSpace)
 	{
 		IAllocatedVirtualTexture* OutAllocatedVirtualTexture = nullptr;
 
@@ -138,7 +138,7 @@ protected:
 			VTDesc.TileBorderSize = InDesc.TileBorderSize;
 			VTDesc.NumTextureLayers = InDesc.NumTextureLayers;
 			VTDesc.bPrivateSpace = InPrivateSpace;
-			VTDesc.bShareDuplicateLayers = true;
+			VTDesc.bShareDuplicateLayers = InSinglePhysicalSpace;
 
 			for (uint32 LayerIndex = 0u; LayerIndex < VTDesc.NumTextureLayers; ++LayerIndex)
 			{
@@ -435,7 +435,7 @@ void URuntimeVirtualTexture::InitResource(IVirtualTexture* InProducer, FTransfor
 {
 	FVTProducerDescription Desc;
 	GetProducerDescription(Desc, VolumeToWorld);
-	Resource->Init(Desc, InProducer, bPrivateSpace);
+	Resource->Init(Desc, InProducer, bSinglePhysicalSpace, bPrivateSpace);
 }
 
 void URuntimeVirtualTexture::InitNullResource()
@@ -443,7 +443,7 @@ void URuntimeVirtualTexture::InitNullResource()
 	FVTProducerDescription Desc;
 	FNullVirtualTextureProducer::GetNullProducerDescription(Desc);
 	FNullVirtualTextureProducer* Producer = new FNullVirtualTextureProducer;
-	Resource->Init(Desc, Producer, false);
+	Resource->Init(Desc, Producer, false, false);
 }
 
 void URuntimeVirtualTexture::GetAssetRegistryTags(TArray<FAssetRegistryTag>& OutTags) const
