Fixed: Niagara CPU shaders don't compile when -noshaderworker is used on the command line

Niagara sets the shader format to VVM_1_0 for CPU emitters, but FShaderCompileUtilities::ExecuteShaderCompileJob was trying to deduce it again from the target platform. This is SP_PCD3D_SM5 for Niagara, so we ended up with the regular D3D shader compiler, instead of the VM compiler.

#rnx

diff --git a/Engine/Source/Runtime/Engine/Private/ShaderCompiler/ShaderCompiler.cpp b/Engine/Source/Runtime/Engine/Private/ShaderCompiler/ShaderCompiler.cpp
--- a/Engine/Source/Runtime/Engine/Private/ShaderCompiler/ShaderCompiler.cpp
+++ b/Engine/Source/Runtime/Engine/Private/ShaderCompiler/ShaderCompiler.cpp
@@ -1494,7 +1494,7 @@ void FShaderCompileUtilities::ExecuteShaderCompileJob(FShaderCommonCompileJob& J
 	auto* SingleJob = Job.GetSingleShaderJob();
 	if (SingleJob)
 	{
-		const FName Format = LegacyShaderPlatformToShaderFormat(EShaderPlatform(SingleJob->Input.Target.Platform));
+		const FName Format = SingleJob->Input.ShaderFormat != NAME_None ? SingleJob->Input.ShaderFormat : LegacyShaderPlatformToShaderFormat(EShaderPlatform(SingleJob->Input.Target.Platform));
 		const IShaderFormat* Compiler = TPM.FindShaderFormat(Format);
 
 		if (!Compiler)
