Allow non-threading builds to FlushLevelStreaming. Level streaming does not depend on threads.

Prevent FlushLevelStreaming reentrancy from causing issues.

diff --git a/Engine/Source/Runtime/Engine/Private/World.cpp b/Engine/Source/Runtime/Engine/Private/World.cpp
--- a/Engine/Source/Runtime/Engine/Private/World.cpp
+++ b/Engine/Source/Runtime/Engine/Private/World.cpp
@@ -3477,8 +3477,19 @@ void UWorld::UpdateStreamingLevelPriority(ULevelStreaming* StreamingLevel)
 
 void UWorld::FlushLevelStreaming(EFlushLevelStreamingType FlushType)
 {
-	if (!FPlatformProcess::SupportsMultithreading())
+	if (FlushType == FlushLevelStreamingType
+		|| FlushLevelStreamingType == EFlushLevelStreamingType::Full
+		|| FlushType == EFlushLevelStreamingType::None)
+	{
+		// We're already flushing at the correct level, or none was passed in
+		// No need to flush
+		return;
+	}
+	else if (FlushType == EFlushLevelStreamingType::Full && FlushLevelStreamingType == EFlushLevelStreamingType::Visibility)
 	{
+		// If FlushLevelStreaming is called for a full update while we are already doing a visibility update, 
+		// upgrade the stored type and let it go back to the loop that was already running
+		FlushLevelStreamingType = EFlushLevelStreamingType::Full;
 		return;
 	}
 
@@ -3496,24 +3507,26 @@ void UWorld::FlushLevelStreaming(EFlushLevelStreamingType FlushType)
 	UpdateLevelStreaming();
 
 	// Making levels visible is spread across several frames so we simply loop till it is done.
-	bool bLevelsPendingVisibility = true;
-	while( bLevelsPendingVisibility )
+	bool bLevelsPendingVisibility = IsVisibilityRequestPending();
+	while (bLevelsPendingVisibility)
 	{
-		bLevelsPendingVisibility = IsVisibilityRequestPending();
-
 		// Tick level streaming to make levels visible.
-		if( bLevelsPendingVisibility )
+
+		const EFlushLevelStreamingType LastStreamingType = FlushLevelStreamingType;
+
+		// Only flush async loading if we're performing a full flush.
+		if (FlushLevelStreamingType == EFlushLevelStreamingType::Full)
 		{
-			// Only flush async loading if we're performing a full flush.
-			if (FlushLevelStreamingType == EFlushLevelStreamingType::Full)
-			{
-				// Make sure all outstanding loads are taken care of...
-				FlushAsyncLoading();
-			}
-	
-			// Update level streaming.
-			UpdateLevelStreaming();
+			// Make sure all outstanding loads are taken care of...
+			FlushAsyncLoading();
 		}
+
+		// Update level streaming.
+		UpdateLevelStreaming();
+
+		// If FlushLevelStreaming reentered as a result of FlushAsyncLoading or UpdateLevelStreaming and upgraded
+		// the flush type, we'll need to do at least one additional loop to be certain all processing is complete
+		bLevelsPendingVisibility = ((LastStreamingType != FlushLevelStreamingType) || IsVisibilityRequestPending());
 	}
 	
 	check(!IsVisibilityRequestPending());
