Fix IMPLEMENT_VERTEX_FACTORY_PARAMETER_TYPE for FInstancedStaticMeshVertexFactory

diff --git a/Engine/Source/Runtime/Engine/Private/InstancedStaticMesh.cpp b/Engine/Source/Runtime/Engine/Private/InstancedStaticMesh.cpp
--- a/Engine/Source/Runtime/Engine/Private/InstancedStaticMesh.cpp
+++ b/Engine/Source/Runtime/Engine/Private/InstancedStaticMesh.cpp
@@ -778,18 +778,9 @@ void FInstancedStaticMeshVertexFactory::InitRHI()
 }
 
 IMPLEMENT_VERTEX_FACTORY_PARAMETER_TYPE(FInstancedStaticMeshVertexFactory, SF_Vertex, FInstancedStaticMeshVertexFactoryShaderParameters);
-
-FVertexFactoryShaderParameters* FInstancedStaticMeshVertexFactory::ConstructShaderParameters(EShaderFrequency ShaderFrequency)
-{
-	switch (ShaderFrequency)
-	{
-	case SF_Vertex:
-	case SF_RayHitGroup:
-		return new FInstancedStaticMeshVertexFactoryShaderParameters();
-	default:
-		return nullptr;
-	}
-}
+#if RHI_RAYTRACING
+IMPLEMENT_VERTEX_FACTORY_PARAMETER_TYPE(FInstancedStaticMeshVertexFactory, SF_RayHitGroup, FInstancedStaticMeshVertexFactoryShaderParameters);
+#endif
 
 IMPLEMENT_VERTEX_FACTORY_TYPE_EX(FInstancedStaticMeshVertexFactory,"/Engine/Private/LocalVertexFactory.ush",true,true,true,true,true,true,false);
 
diff --git a/Engine/Source/Runtime/Engine/Private/InstancedStaticMesh.h b/Engine/Source/Runtime/Engine/Private/InstancedStaticMesh.h
--- a/Engine/Source/Runtime/Engine/Private/InstancedStaticMesh.h
+++ b/Engine/Source/Runtime/Engine/Private/InstancedStaticMesh.h
@@ -239,8 +239,6 @@ public:
 	// FRenderResource interface.
 	virtual void InitRHI() override;
 
-	static FVertexFactoryShaderParameters* ConstructShaderParameters(EShaderFrequency ShaderFrequency);
-
 	/** Make sure we account for changes in the signature of GetStaticBatchElementVisibility() */
 	static CONSTEXPR uint32 NumBitsForVisibilityMask()
 	{		
