Remove incorrect multithreading check in IsVisibilityRequestPending
#rnx

diff --git a/Engine/Source/Runtime/Engine/Private/World.cpp b/Engine/Source/Runtime/Engine/Private/World.cpp
--- a/Engine/Source/Runtime/Engine/Private/World.cpp
+++ b/Engine/Source/Runtime/Engine/Private/World.cpp
@@ -3581,10 +3581,6 @@ void UWorld::ConditionallyBuildStreamingData()
 
 bool UWorld::IsVisibilityRequestPending() const
 {
-	if (!FPlatformProcess::SupportsMultithreading())
-	{
-		return false;
-	}
 	return (CurrentLevelPendingVisibility != nullptr || CurrentLevelPendingInvisibility != nullptr);
 }
 
