Remove unused GetAllocatedSize() method...this memory is now accounted for as part of the larger frozen shader memory image

diff --git a/Engine/Source/Runtime/Engine/Public/MaterialShared.h b/Engine/Source/Runtime/Engine/Public/MaterialShared.h
--- a/Engine/Source/Runtime/Engine/Public/MaterialShared.h
+++ b/Engine/Source/Runtime/Engine/Public/MaterialShared.h
@@ -541,21 +541,6 @@ public:
 		return UniformBufferLayout;
 	}
 
-	uint32 GetAllocatedSize() const
-	{
-		/*return UniformVectorExpressions.GetAllocatedSize()
-			+ UniformScalarExpressions.GetAllocatedSize()
-			+ Uniform2DTextureExpressions.GetAllocatedSize()
-			+ UniformCubeTextureExpressions.GetAllocatedSize()
-			+ Uniform2DArrayTextureExpressions.GetAllocatedSize()
-			+ UniformVolumeTextureExpressions.GetAllocatedSize()
-			+ UniformVirtualTextureExpressions.GetAllocatedSize()
-			+ UniformExternalTextureExpressions.GetAllocatedSize()
-			+ VTStacks.GetAllocatedSize()
-			+ ParameterCollections.GetAllocatedSize();*/
-		return 0u;
-	}
-
 	inline const FMaterialVectorParameterInfo& GetVectorParameter(uint32 Index) const { return UniformVectorParameters[Index]; }
 	inline const FMaterialScalarParameterInfo& GetScalarParameter(uint32 Index) const { return UniformScalarParameters[Index]; }
 	inline const FMaterialTextureParameterInfo& GetTextureParameter(EMaterialTextureParameterType Type, int32 Index) const { return UniformTextureParameters[(uint32)Type][Index]; }
