Temporary fix for crashes when exiting the game / editor

diff --git a/Engine/Source/Runtime/Engine/Private/SoundEffectPreset.cpp b/Engine/Source/Runtime/Engine/Private/SoundEffectPreset.cpp
--- a/Engine/Source/Runtime/Engine/Private/SoundEffectPreset.cpp
+++ b/Engine/Source/Runtime/Engine/Private/SoundEffectPreset.cpp
@@ -5,6 +5,7 @@
 #include "Sound/SoundEffectSource.h"
 #include "Engine/Engine.h"
 #include "AudioDeviceManager.h"
+#include "CoreGlobals.h"
 
 USoundEffectPreset::USoundEffectPreset(const FObjectInitializer& ObjectInitializer)
 	: Super(ObjectInitializer)
@@ -72,11 +73,14 @@ void USoundEffectPreset::AddReferencedEffects(FReferenceCollector& Collector)
 
 void USoundEffectPreset::BeginDestroy()
 {
-	for (int32 i = 0; i < Instances.Num(); ++i)
+	if (!GExitPurge)
 	{
-		if (Instances[i])
+		for (int32 i = 0; i < Instances.Num(); ++i)
 		{
-			Instances[i]->ClearPreset(false /* bRemoveFromPreset */);
+			if (Instances[i])
+			{
+				Instances[i]->ClearPreset(false /* bRemoveFromPreset */);
+			}
 		}
 	}
 	Instances.Reset();
