Fixed assert in PreparePrimitiveUniformBuffer with skeletal meshes when ES3.1 Preview mode is enabled.

diff --git a/Engine/Source/Runtime/Engine/Private/SceneManagement.cpp b/Engine/Source/Runtime/Engine/Private/SceneManagement.cpp
--- a/Engine/Source/Runtime/Engine/Private/SceneManagement.cpp
+++ b/Engine/Source/Runtime/Engine/Private/SceneManagement.cpp
@@ -1096,21 +1096,24 @@ void FMeshBatch::PreparePrimitiveUniformBuffer(const FPrimitiveSceneProxy* Primi
 	const bool bVFSupportsPrimitiveIdStream = VertexFactory->GetType()->SupportsPrimitiveIdStream();
 	checkf((PrimitiveSceneProxy->DoesVFRequirePrimitiveUniformBuffer() || bVFSupportsPrimitiveIdStream), TEXT("PrimitiveSceneProxy has bVFRequiresPrimitiveUniformBuffer disabled yet tried to draw with a vertex factory (%s) that did not support PrimitiveIdStream."), VertexFactory->GetType()->GetName());
 
+	const bool bUseGPUScene = UseGPUScene(GMaxRHIShaderPlatform, FeatureLevel);
 	const bool bPrimitiveShaderDataComesFromSceneBuffer = VertexFactory->GetPrimitiveIdStreamIndex(EVertexInputStreamType::Default) >= 0;
 
 	for (int32 ElementIndex = 0; ElementIndex < Elements.Num(); ElementIndex++)
 	{
 		FMeshBatchElement& MeshElement = Elements[ElementIndex];
 
-		if (bPrimitiveShaderDataComesFromSceneBuffer)
+		// When GPU scene is enabled, primitive data is required to come from the scene buffer OR the uniform buffer; not both.
+		// However, it is possible that our scene feature level is lower and doesn't actually support GPU scene, in which case the vertex
+		// factory stream is unused. This should really only happen in mobile preview.
+		if (bUseGPUScene)
 		{
-			checkf(!Elements[ElementIndex].PrimitiveUniformBuffer, 
-				TEXT("FMeshBatch was assigned a PrimitiveUniformBuffer even though Vertex Factory %s fetches primitive shader data through a Scene buffer.  The assigned PrimitiveUniformBuffer cannot be respected.  Use PrimitiveUniformBufferResource instead for dynamic primitive data, or leave both null to get FPrimitiveSceneProxy->UniformBuffer."), VertexFactory->GetType()->GetName());
+			checkf(!bPrimitiveShaderDataComesFromSceneBuffer || !MeshElement.PrimitiveUniformBuffer,
+				TEXT("FMeshBatch was assigned a PrimitiveUniformBuffer even though Vertex Factory %s fetches primitive shader data through a Scene buffer. The assigned PrimitiveUniformBuffer cannot be respected. Use PrimitiveUniformBufferResource instead for dynamic primitive data, or leave both null to get FPrimitiveSceneProxy->UniformBuffer."), VertexFactory->GetType()->GetName());
 		}
-
-		// If we are not using GPU Scene, draws using vertex factories that do not support an explicit PrimitiveUniformBuffer on the FMeshBatch need to be setup with the FPrimitiveSceneProxy's uniform buffer
-		if (!MeshElement.PrimitiveUniformBufferResource && !UseGPUScene(GMaxRHIShaderPlatform, FeatureLevel) && bVFSupportsPrimitiveIdStream)
+		else if (!MeshElement.PrimitiveUniformBufferResource)
 		{
+			// If not using GPU scene, fall back to the primitive uniform buffer.
 			MeshElement.PrimitiveUniformBuffer = PrimitiveSceneProxy->GetUniformBuffer();
 		}
 
