- Wrapped the InitAllBodies ParallelFor in a check for the cvar to make sure it's enabled before spending the time to crank up a ParallelFor for no reason. It's disabled, so the check will return false always, but this is an okay fix for the short term (in case someone does enable it), while a better fix is discovered (ie, UseDeferredPhysicsBodyCreation should return true only if there are any outstanding PhysicsBodies to deal with, not checking the cvar)

diff --git a/Engine/Source/Runtime/Engine/Classes/PhysicsEngine/BodyInstance.h b/Engine/Source/Runtime/Engine/Classes/PhysicsEngine/BodyInstance.h
--- a/Engine/Source/Runtime/Engine/Classes/PhysicsEngine/BodyInstance.h
+++ b/Engine/Source/Runtime/Engine/Classes/PhysicsEngine/BodyInstance.h
@@ -546,6 +546,9 @@ public:
 	UPROPERTY()
 	float PhysicsBlendWeight;
 
+	/** True if we want to use deferred body instance creation */
+	static bool UseDeferredPhysicsBodyCreation();
+
 private:
 	TArray<FInitBodiesHelperWithData<true>> InitBodiesDeferredListStatic;
 	TArray<FInitBodiesHelperWithData<false>> InitBodiesDeferredListDynamic;
diff --git a/Engine/Source/Runtime/Engine/Private/Level.cpp b/Engine/Source/Runtime/Engine/Private/Level.cpp
--- a/Engine/Source/Runtime/Engine/Private/Level.cpp
+++ b/Engine/Source/Runtime/Engine/Private/Level.cpp
@@ -941,6 +941,7 @@ void ULevel::IncrementalUpdateComponents(int32 NumComponentsToUpdate, bool bReru
 
 	int32 PreviousIndex = CurrentActorIndexForUpdateComponents;
 	// Find next valid actor to process components registration
+
 	while (CurrentActorIndexForUpdateComponents < Actors.Num())
 	{
 		AActor* Actor = Actors[CurrentActorIndexForUpdateComponents];
@@ -1022,6 +1023,7 @@ void ULevel::IncrementalUpdateComponents(int32 NumComponentsToUpdate, bool bReru
 		check(OwningWorld->IsGameWorld());
 	}
 
+	if (FBodyInstance::UseDeferredPhysicsBodyCreation())
 	{
 		SCOPE_CYCLE_COUNTER(STAT_DeferredUpdateBodies);
 		// Init Bodies
diff --git a/Engine/Source/Runtime/Engine/Private/PhysicsEngine/BodyInstance.cpp b/Engine/Source/Runtime/Engine/Private/PhysicsEngine/BodyInstance.cpp
--- a/Engine/Source/Runtime/Engine/Private/PhysicsEngine/BodyInstance.cpp
+++ b/Engine/Source/Runtime/Engine/Private/PhysicsEngine/BodyInstance.cpp
@@ -80,6 +80,8 @@ DECLARE_CYCLE_STAT(TEXT("CreatePhysicsActor"), STAT_CreatePhysicsActor, STATGROU
 DECLARE_CYCLE_STAT(TEXT("BodyInstance SetCollisionProfileName"), STAT_BodyInst_SetCollisionProfileName, STATGROUP_Physics);
 DECLARE_CYCLE_STAT(TEXT("Phys SetBodyTransform"), STAT_SetBodyTransform, STATGROUP_Physics);
 
+static int32 GUseDeferredPhysicsBodyCreation = 0;
+
 // @HACK Guard to better encapsulate game related hacks introduced into UpdatePhysicsFilterData()
 TAutoConsoleVariable<int32> CVarEnableDynamicPerBodyFilterHacks(
 	TEXT("p.EnableDynamicPerBodyFilterHacks"), 
@@ -88,13 +90,21 @@ TAutoConsoleVariable<int32> CVarEnableDynamicPerBodyFilterHacks(
 	ECVF_ReadOnly
 );
 
-TAutoConsoleVariable<int32> CVarEnableDeferredPhysicsCreation(
+FAutoConsoleVariableRef CVarEnableDeferredPhysicsCreation(
 	TEXT("p.EnableDeferredPhysicsCreation"), 
-	0, 
+	GUseDeferredPhysicsBodyCreation,
 	TEXT("Enables/Disables deferred physics creation."),
 	ECVF_Default
 );
 
+bool FBodyInstance::UseDeferredPhysicsBodyCreation()
+{
+	// @todo: this could be improved by tracking all oustanding delayed creations, and return true if there are any at all. 
+	// or maybe keep a list of instances that have them, so we don't loop over all actors, etc
+	return GUseDeferredPhysicsBodyCreation != 0;
+}
+
+
 using namespace PhysicsInterfaceTypes;
 
 bool IsRigidBodyKinematic_AssumesLocked(const FPhysicsActorHandle& InActorRef)
@@ -1353,7 +1363,7 @@ void FBodyInstance::InitBody(class UBodySetup* Setup, const FTransform& Transfor
 	bool bIsStatic = SpawnParams.bStaticPhysics;
 	if(bIsStatic)
 	{
-		if (CVarEnableDeferredPhysicsCreation.GetValueOnGameThread())
+		if (GUseDeferredPhysicsBodyCreation)
 		{
 			InitBodiesDeferredListStatic.Add(FInitBodiesHelperWithData<true>(MoveTemp(Bodies), MoveTemp(Transforms), Setup, PrimComp, InRBScene, SpawnParams, SpawnParams.Aggregate));
 		}
@@ -1365,7 +1375,7 @@ void FBodyInstance::InitBody(class UBodySetup* Setup, const FTransform& Transfor
 	}
 	else
 	{
-		if (CVarEnableDeferredPhysicsCreation.GetValueOnGameThread())
+		if (GUseDeferredPhysicsBodyCreation)
 		{
 			InitBodiesDeferredListDynamic.Add(FInitBodiesHelperWithData<false>(MoveTemp(Bodies), MoveTemp(Transforms), Setup, PrimComp, InRBScene, SpawnParams, SpawnParams.Aggregate));
 		}
