Adding ability to query the number of controllers

diff --git a/Engine/Source/Runtime/Engine/Classes/Engine/World.h b/Engine/Source/Runtime/Engine/Classes/Engine/World.h
--- a/Engine/Source/Runtime/Engine/Classes/Engine/World.h
+++ b/Engine/Source/Runtime/Engine/Classes/Engine/World.h
@@ -2131,6 +2131,9 @@ public:
 	/** Returns an iterator for the controller list. */
 	FConstControllerIterator GetControllerIterator() const;
 
+	/** @return Returns the number of Controllers. */
+	int32 GetNumControllers() const;
+	
 	/** @return Returns an iterator for the pawn list. */
 	UE_DEPRECATED(4.24, "The PawnIterator is an inefficient mechanism for iterating pawns. Please use TActorIterator<PawnType> instead.")
 	FConstPawnIterator GetPawnIterator() const;
diff --git a/Engine/Source/Runtime/Engine/Private/World.cpp b/Engine/Source/Runtime/Engine/Private/World.cpp
--- a/Engine/Source/Runtime/Engine/Private/World.cpp
+++ b/Engine/Source/Runtime/Engine/Private/World.cpp
@@ -4374,6 +4374,11 @@ FConstControllerIterator UWorld::GetControllerIterator() const
 	return ControllerList.CreateConstIterator();
 }
 
+int32 UWorld::GetNumControllers() const
+{
+	return ControllerList.Num();
+}
+
 FConstPlayerControllerIterator UWorld::GetPlayerControllerIterator() const
 {
 	return PlayerControllerList.CreateConstIterator();
