Adding LoadingBehavior column to Audio Stream Caching section of audio mem report

diff --git a/Engine/Source/Runtime/Engine/Classes/Sound/SoundWaveLoadingBehavior.cpp b/Engine/Source/Runtime/Engine/Classes/Sound/SoundWaveLoadingBehavior.cpp
new file mode 100644
--- /dev/null
+++ b/Engine/Source/Runtime/Engine/Classes/Sound/SoundWaveLoadingBehavior.cpp
@@ -0,0 +1,25 @@
+// Copyright Epic Games, Inc. All Rights Reserved.
+#include "SoundWaveLoadingBehavior.h"
+#include "CoreMinimal.h"
+
+
+const TCHAR* EnumToString(ESoundWaveLoadingBehavior InCurrentState)
+{
+	switch (InCurrentState)
+	{
+	case ESoundWaveLoadingBehavior::Inherited:
+		return TEXT("Inherited");
+	case ESoundWaveLoadingBehavior::RetainOnLoad:
+		return TEXT("RetainOnLoad");
+	case ESoundWaveLoadingBehavior::PrimeOnLoad:
+		return TEXT("PrimeOnLoad");
+	case ESoundWaveLoadingBehavior::LoadOnDemand:
+		return TEXT("LoadOnDemand");
+	case ESoundWaveLoadingBehavior::ForceInline:
+		return TEXT("ForceInline");
+	case ESoundWaveLoadingBehavior::Uninitialized:
+		return TEXT("Uninitialized");
+	}
+	ensure(false);
+	return TEXT("Unknown");
+}
diff --git a/Engine/Source/Runtime/Engine/Classes/Sound/SoundWaveLoadingBehavior.h b/Engine/Source/Runtime/Engine/Classes/Sound/SoundWaveLoadingBehavior.h
--- a/Engine/Source/Runtime/Engine/Classes/Sound/SoundWaveLoadingBehavior.h
+++ b/Engine/Source/Runtime/Engine/Classes/Sound/SoundWaveLoadingBehavior.h
@@ -26,3 +26,5 @@ enum class ESoundWaveLoadingBehavior : uint8
 	// This value is used to delineate when the value of ESoundWaveLoadingBehavior hasn't been cached on a USoundWave yet.
 	Uninitialized = 0xff UMETA(Hidden)
 };
+
+const TCHAR* EnumToString(ESoundWaveLoadingBehavior InCurrentState);
\ No newline at end of file
diff --git a/Engine/Source/Runtime/Engine/Private/AudioStreamingCache.cpp b/Engine/Source/Runtime/Engine/Private/AudioStreamingCache.cpp
--- a/Engine/Source/Runtime/Engine/Private/AudioStreamingCache.cpp
+++ b/Engine/Source/Runtime/Engine/Private/AudioStreamingCache.cpp
@@ -6,7 +6,6 @@ AudioStreaming.cpp: Implementation of audio streaming classes.
 
 #include "AudioStreamingCache.h"
 #include "Misc/CoreStats.h"
-#include "Sound/SoundWave.h"
 #include "Sound/AudioSettings.h"
 #include "DerivedDataCacheInterface.h"
 #include "Serialization/MemoryReader.h"
@@ -552,10 +551,16 @@ bool FAudioChunkCache::AddOrTouchChunk(const FChunkKey& InKey, TFunction<void(EA
 		if (FoundElement->bIsLoaded)
 		{
 			ExecuteOnLoadCompleteCallback(EAudioChunkLoadResult::AlreadyLoaded, OnLoadCompleted, CallbackThread);
+#if DEBUG_STREAM_CACHE
+			checkf(FoundElement->DebugInfo.LoadingBehavior == InKey.SoundWave->GetLoadingBehavior(true), TEXT("Cached element had different LoadingBehavior {%i} than the new request {%i}")
+				, EnumToString(FoundElement->DebugInfo.LoadingBehavior), EnumToString(InKey.SoundWave->GetLoadingBehavior(true)))
+#endif
+
 		}
 
 #if DEBUG_STREAM_CACHE
 		FoundElement->DebugInfo.NumTimesTouched++;
+		FoundElement->DebugInfo.LoadingBehavior = InKey.SoundWave->GetLoadingBehavior(true);
 #endif
 
 		return true;
@@ -572,6 +577,7 @@ bool FAudioChunkCache::AddOrTouchChunk(const FChunkKey& InKey, TFunction<void(EA
 
 #if DEBUG_STREAM_CACHE
 		CacheElement->DebugInfo.bWasCacheMiss = bNeededForPlayback;
+		CacheElement->DebugInfo.LoadingBehavior = InKey.SoundWave->GetLoadingBehavior(true);
 #endif
 		const FStreamedAudioChunk& Chunk = InKey.SoundWave->RunningPlatformData->Chunks[InKey.ChunkIndex];
 		int32 ChunkDataSize = Chunk.AudioDataSize;
@@ -615,7 +621,7 @@ TArrayView<uint8> FAudioChunkCache::GetChunk(const FChunkKey& InKey, bool bBlock
 				return TArrayView<uint8>();
 			}
 		}
-		
+
 
 		// If this value is ever negative, it means that we're decrementing more than we're incrementing:
 		check(FoundElement->NumConsumers.GetValue() >= 0);
@@ -762,7 +768,7 @@ void FAudioChunkCache::BlockForAllPendingLoads() const
 
 	float TimeStarted = FPlatformTime::Seconds();
 
-	do 
+	do
 	{
 		// If we did find an in flight async load,
 		// sleep to let other threads complete this task.
@@ -820,7 +826,7 @@ FString FAudioChunkCache::FlushCacheMissLog()
 	FString ConcatenatedCacheMisses;
 	ConcatenatedCacheMisses.Append(TEXT("All Cache Misses:\nSoundWave:\t, ChunkIndex\n"));
 
-	struct FMissedChunk 
+	struct FMissedChunk
 	{
 		FName SoundWaveName;
 		uint32 ChunkIndex;
@@ -848,10 +854,10 @@ FString FAudioChunkCache::FlushCacheMissLog()
 		ConcatenatedCacheMisses.AppendInt(CacheMissInfo.ChunkIndex);
 		ConcatenatedCacheMisses.Append(TEXT("\n"));
 
-		FChunkKey Chunk = 
-		{ 
-			nullptr, 
-			CacheMissInfo.SoundWaveName, 
+		FChunkKey Chunk =
+		{
+			nullptr,
+			CacheMissInfo.SoundWaveName,
 			CacheMissInfo.ChunkIndex,
 #if WITH_EDITOR
 			0
@@ -922,7 +928,7 @@ FAudioChunkCache::FCacheElement* FAudioChunkCache::FindElementForKey(const FChun
 	{
 		if (InKey == CurrentElement->Key)
 		{
-			
+
 #if DEBUG_STREAM_CACHE
 			float& CMA = CurrentElement->DebugInfo.AverageLocationInCacheWhenNeeded;
 			CMA += ((ElementPosition - CMA) / (CurrentElement->DebugInfo.NumTimesTouched + 1));
@@ -1009,7 +1015,7 @@ bool FAudioChunkCache::ShouldAddNewChunk() const
 FAudioChunkCache::FCacheElement* FAudioChunkCache::InsertChunk(const FChunkKey& InKey)
 {
 	FCacheElement* CacheElement = nullptr;
-	
+
 	{
 		FScopeLock ScopeLock(&CacheMutationCriticalSection);
 
@@ -1027,7 +1033,7 @@ FAudioChunkCache::FCacheElement* FAudioChunkCache::InsertChunk(const FChunkKey&
 				UE_LOG(LogAudio, Display, TEXT("Audio Stream Cache: Using %d of %d chunks.."), ChunksInUse, CachePool.Num());
 				bLoggedCacheSaturated = true;
 			}
-			
+
 			// The pools filled, so we're going to need to evict.
 			CacheElement = EvictLeastRecentChunk();
 
@@ -1096,7 +1102,7 @@ FAudioChunkCache::FCacheElement* FAudioChunkCache::EvictLeastRecentChunk()
 				{
 					CacheElement->MoreRecentElement->LessRecentElement = CacheElement->LessRecentElement;
 				}
-				
+
 				// If we ever hit this while loop it means that CacheElement is not the least recently used element.
 				check(CacheElement->LessRecentElement);
 				CacheElement->LessRecentElement->MoreRecentElement = CacheElement->MoreRecentElement;
@@ -1127,7 +1133,7 @@ FAudioChunkCache::FCacheElement* FAudioChunkCache::EvictLeastRecentChunk()
 void FAudioChunkCache::KickOffAsyncLoad(FCacheElement* CacheElement, const FChunkKey& InKey, TFunction<void(EAudioChunkLoadResult)> OnLoadCompleted, ENamedThreads::Type CallbackThread, bool bNeededForPlayback)
 {
 	check(CacheElement);
-	
+
 	const FStreamedAudioChunk& Chunk = InKey.SoundWave->RunningPlatformData->Chunks[InKey.ChunkIndex];
 	int32 ChunkDataSize = Chunk.AudioDataSize;
 
@@ -1146,6 +1152,7 @@ void FAudioChunkCache::KickOffAsyncLoad(FCacheElement* CacheElement, const FChun
 
 #if DEBUG_STREAM_CACHE
 	CacheElement->DebugInfo.NumTotalChunks = InKey.SoundWave->GetNumChunks() - 1;
+	CacheElement->DebugInfo.LoadingBehavior = InKey.SoundWave->GetLoadingBehavior(true);
 #endif
 
 	// In editor, we retrieve from the DDC. In non-editor situations, we read the chunk async from the pak file.
@@ -1161,7 +1168,7 @@ void FAudioChunkCache::KickOffAsyncLoad(FCacheElement* CacheElement, const FChun
 		{
 			check(CacheElement->DDCTask->IsDone());
 		}
-		
+
 #if DEBUG_STREAM_CACHE
 		CacheElement->DebugInfo.TimeLoadStarted = FPlatformTime::Cycles64();
 #endif
@@ -1208,7 +1215,7 @@ void FAudioChunkCache::KickOffAsyncLoad(FCacheElement* CacheElement, const FChun
 
 		// If we ever want to eliminate zero-padding in chunks, that could be verified here:
 		//ensureAlwaysMsgf(AudioChunkSize == ChunkBulkDataSize, TEXT("For memory load on demand, we do not zero-pad to page sizes."));
-		
+
 		NumberOfLoadsInFlight.Increment();
 
 		FBulkDataIORequestCallBack AsyncFileCallBack = [this, OnLoadCompleted, CacheElement, InKey, ChunkDataSize, CallbackThread](bool bWasCancelled, IBulkDataIORequest*)
@@ -1231,7 +1238,7 @@ void FAudioChunkCache::KickOffAsyncLoad(FCacheElement* CacheElement, const FChun
 		CacheElement->DebugInfo.TimeLoadStarted = FPlatformTime::Seconds();
 #endif
 
-		CacheElement->ReadRequest.Reset(Chunk.BulkData.CreateStreamingRequest( 0, ChunkDataSize, AsyncIOPriority, &AsyncFileCallBack, CacheElement->ChunkData.GetData()));
+		CacheElement->ReadRequest.Reset(Chunk.BulkData.CreateStreamingRequest(0, ChunkDataSize, AsyncIOPriority, &AsyncFileCallBack, CacheElement->ChunkData.GetData()));
 		if (!CacheElement->ReadRequest.IsValid())
 		{
 			UE_LOG(LogAudio, Error, TEXT("Chunk load in audio LRU cache failed."));
@@ -1250,57 +1257,57 @@ EAsyncIOPriorityAndFlags FAudioChunkCache::GetAsyncPriorityForChunk(const FChunk
 	{
 		switch (PlaybackRequestPriorityCVar)
 		{
-			case 4:
-			{
-				return AIOP_MIN;
-			}
-			case 3:
-			{
-				return AIOP_Low;
-			}
-			case 2:
-			{
-				return AIOP_BelowNormal;
-			}
-			case 1:
-			{
-				return AIOP_Normal;
-			}
-			case 0:
-			default:
-			{
-				return AIOP_High;
-			}
+		case 4:
+		{
+			return AIOP_MIN;
+		}
+		case 3:
+		{
+			return AIOP_Low;
+		}
+		case 2:
+		{
+			return AIOP_BelowNormal;
+		}
+		case 1:
+		{
+			return AIOP_Normal;
+		}
+		case 0:
+		default:
+		{
+			return AIOP_High;
+		}
 		}
 	}
 	else
 	{
 		switch (ReadRequestPriorityCVar)
 		{
-			case 4:
-			{
-				return AIOP_MIN;
-			}
-			case 3:
-			{
-				return AIOP_Low;
-			}
-			case 2:
-			{
-				return AIOP_BelowNormal;
-			}
-			case 1:
-			{
-				return AIOP_Normal;
-			}
-			case 0:
-			default:
-			{
-				return AIOP_High;
-			}
+		case 4:
+		{
+			return AIOP_MIN;
+		}
+		case 3:
+		{
+			return AIOP_Low;
+		}
+		case 2:
+		{
+			return AIOP_BelowNormal;
+		}
+		case 1:
+		{
+			return AIOP_Normal;
+		}
+		case 0:
+		default:
+		{
+			return AIOP_High;
+		}
 		}
 	}
-	
+
 }
 
 void FAudioChunkCache::ExecuteOnLoadCompleteCallback(EAudioChunkLoadResult Result, const TFunction<void(EAudioChunkLoadResult)>& OnLoadCompleted, const ENamedThreads::Type& CallbackThread)
@@ -1312,7 +1319,7 @@ void FAudioChunkCache::ExecuteOnLoadCompleteCallback(EAudioChunkLoadResult Resul
 	else
 	{
 		// Dispatch an async notify.
-		AsyncTask(CallbackThread, [OnLoadCompleted, Result]() 
+		AsyncTask(CallbackThread, [OnLoadCompleted, Result]()
 		{
 			OnLoadCompleted(Result);
 		});
@@ -1340,7 +1347,7 @@ int32 FCachedAudioStreamingManager::RenderStatAudioStreaming(UWorld* World, FVie
 		Y += 12;
 
 		TPair<int, int> Size = Cache.DebugDisplay(World, Viewport, Canvas, X, Y, ViewLocation, ViewRotation);
-		
+
 		// Separate caches are laid out horizontally across the screen, so the total height is equal to our tallest cache panel:
 		X += Size.Key;
 		Height = FMath::Max(Height, Size.Value);
@@ -1381,7 +1388,7 @@ void FCachedAudioStreamingManager::SetProfilingMode(bool bEnabled)
 uint64 FCachedAudioStreamingManager::TrimMemory(uint64 NumBytesToFree)
 {
 	uint64 NumBytesLeftToFree = NumBytesToFree;
-	
+
 	// TODO: When we support multiple caches, it's probably best to do this in reverse,
 	// since the caches are sorted from shortest sounds to longest.
 	// Freeing longer chunks will get us bigger gains and (presumably) have lower churn.
@@ -1430,7 +1437,7 @@ TPair<int, int> FAudioChunkCache::DebugDisplay(UWorld* World, FViewport* Viewpor
 
 	// In editor builds, this is a chunk that was built in a previous version of the cook quality settings.
 	const FLinearColor StaleChunkColor(143 / ColorMax, 73 / ColorMax, 70 / ColorMax);
-	
+
 	// A chunk that currently has an async load in flight.
 	const FLinearColor CurrentlyLoadingChunkColor = FLinearColor::Yellow;
 
@@ -1458,12 +1465,12 @@ TPair<int, int> FAudioChunkCache::DebugDisplay(UWorld* World, FViewport* Viewpor
 		NumBytesCounter += CurrentElement->ChunkData.Num();
 		CurrentElement = CurrentElement->LessRecentElement;
 	}
-	
+
 	// Convert to megabytes and print the total size:
 	const double NumMegabytesInUse = (double)NumBytesCounter / (1024 * 1024);
 	const double MaxCacheSizeMB = ((double)MemoryLimitBytes) / (1024 * 1024);
 
-	FString CacheMemoryUsage = *FString::Printf(TEXT("Using: %.4f Megabytes (%lu bytes). Max Potential Usage: %.4f Megabytes."), NumMegabytesInUse, MemoryCounterBytes.Load(),  MaxCacheSizeMB);
+	FString CacheMemoryUsage = *FString::Printf(TEXT("Using: %.4f Megabytes (%lu bytes). Max Potential Usage: %.4f Megabytes."), NumMegabytesInUse, MemoryCounterBytes.Load(), MaxCacheSizeMB);
 
 	// We're going to align this horizontally with the number of elements right above it.
 	Canvas->DrawShadowedString(X + CacheTitleOffsetX, Y, *CacheMemoryUsage, UEngine::GetSmallFont(), FLinearColor::Green);
@@ -1489,6 +1496,7 @@ TPair<int, int> FAudioChunkCache::DebugDisplay(UWorld* World, FViewport* Viewpor
 		int32 NumTimesTouched = -1;
 		double TimeToLoad = -1.0;
 		float AveragePlaceInCache = -1.0f;
+		ESoundWaveLoadingBehavior LoadingBehavior = ESoundWaveLoadingBehavior::Uninitialized;
 		bool bWasCacheMiss = false;
 		bool bIsStaleChunk = false;
 
@@ -1497,6 +1505,7 @@ TPair<int, int> FAudioChunkCache::DebugDisplay(UWorld* World, FViewport* Viewpor
 		NumTimesTouched = CurrentElement->DebugInfo.NumTimesTouched;
 		TimeToLoad = CurrentElement->DebugInfo.TimeToLoad;
 		AveragePlaceInCache = CurrentElement->DebugInfo.AverageLocationInCacheWhenNeeded;
+		LoadingBehavior = CurrentElement->DebugInfo.LoadingBehavior;
 		bWasCacheMiss = CurrentElement->DebugInfo.bWasCacheMiss;
 #endif
 
@@ -1507,7 +1516,7 @@ TPair<int, int> FAudioChunkCache::DebugDisplay(UWorld* World, FViewport* Viewpor
 
 		const bool bWasTrimmed = CurrentElement->ChunkData.Num() == 0;
 
-		FString ElementInfo = *FString::Printf(TEXT("%4i. Size: %6.2f KB   Chunk: %d of %d   Request Count: %d    Average Index: %6.2f  Number of Handles Retaining Chunk: %d     Chunk Load Time(in ms): %6.4fms      Name: %s Notes: %s %s"),
+		FString ElementInfo = *FString::Printf(TEXT("%4i. Size: %6.2f KB   Chunk: %d of %d   Request Count: %d    Average Index: %6.2f  Number of Handles Retaining Chunk: %d     Chunk Load Time(in ms): %6.4fms      Loading Behavior: %s      Name: %s Notes: %s %s"),
 			Index,
 			CurrentElement->ChunkData.Num() / 1024.0f,
 			CurrentElement->Key.ChunkIndex,
@@ -1516,10 +1525,11 @@ TPair<int, int> FAudioChunkCache::DebugDisplay(UWorld* World, FViewport* Viewpor
 			AveragePlaceInCache,
 			CurrentElement->NumConsumers.GetValue(),
 			TimeToLoad,
+			EnumToString(LoadingBehavior),
 			bWasTrimmed ? TEXT("TRIMMED CHUNK") : *CurrentElement->Key.SoundWaveName.ToString(),
 			bWasCacheMiss ? TEXT("(Cache Miss!)") : TEXT(""),
 			bIsStaleChunk ? TEXT("(Stale Chunk)") : TEXT("")
-			);
+		);
 
 		// Since there's a lot of info here,
 		// Subtly fading the chunk info to gray seems to help as a visual indicator of how far down on the list things are.
@@ -1599,7 +1609,7 @@ FString FAudioChunkCache::DebugPrint()
 	// Convert to megabytes and print the total size:
 	const double NumMegabytesInUse = (double)NumBytesCounter / (1024 * 1024);
 	const double NumMegabytesRetained = (double)NumBytesRetained / (1024 * 1024);
-	
+
 	const double MaxCacheSizeMB = ((double)MemoryLimitBytes) / (1024 * 1024);
 	const double PercentageOfCacheRetained = NumMegabytesRetained / MaxCacheSizeMB;
 
@@ -1612,7 +1622,7 @@ FString FAudioChunkCache::DebugPrint()
 	CurrentElement = MostRecentElement;
 	int32 Index = 0;
 
-	OutputString += TEXT("Index:\t, Size (KB):\t, Chunk:\t, Request Count:\t, Average Index:\t, Number of Handles Retaining Chunk:\t, Chunk Load Time:\t, Name: \t, Notes:\t, \n");
+	OutputString += TEXT("Index:\t, Size (KB):\t, Chunk:\t, Request Count:\t, Average Index:\t, Number of Handles Retaining Chunk:\t, Chunk Load Time:\t, Name: \t, LoadingBehavior: \t, Notes:\t, \n");
 
 	// More detailed info about individual chunks here:
 	while (CurrentElement != nullptr)
@@ -1621,6 +1631,7 @@ FString FAudioChunkCache::DebugPrint()
 		int32 NumTimesTouched = -1;
 		double TimeToLoad = -1.0;
 		float AveragePlaceInCache = -1.0f;
+		ESoundWaveLoadingBehavior LoadingBehavior = ESoundWaveLoadingBehavior::Uninitialized;
 		bool bWasCacheMiss = false;
 		bool bIsStaleChunk = false;
 
@@ -1629,6 +1640,7 @@ FString FAudioChunkCache::DebugPrint()
 		NumTimesTouched = CurrentElement->DebugInfo.NumTimesTouched;
 		TimeToLoad = CurrentElement->DebugInfo.TimeToLoad;
 		AveragePlaceInCache = CurrentElement->DebugInfo.AverageLocationInCacheWhenNeeded;
+		LoadingBehavior = CurrentElement->DebugInfo.LoadingBehavior;
 		bWasCacheMiss = CurrentElement->DebugInfo.bWasCacheMiss;
 #endif
 
@@ -1639,7 +1651,7 @@ FString FAudioChunkCache::DebugPrint()
 
 		const bool bWasTrimmed = CurrentElement->ChunkData.Num() == 0;
 
-		FString ElementInfo = *FString::Printf(TEXT("%4i.\t, %6.2f\t, %d of %d\t, %d\t, %6.2f\t, %d\t,  %6.4f\t, %s\t, %s %s %s"),
+		FString ElementInfo = *FString::Printf(TEXT("%4i.\t, %6.2f\t, %d of %d\t, %d\t, %6.2f\t, %d\t,  %6.4f\t, %s\t, %s, %s %s %s"),
 			Index,
 			CurrentElement->ChunkData.Num() / 1024.0f,
 			CurrentElement->Key.ChunkIndex,
@@ -1649,6 +1661,7 @@ FString FAudioChunkCache::DebugPrint()
 			CurrentElement->NumConsumers.GetValue(),
 			TimeToLoad,
 			bWasTrimmed ? TEXT("TRIMMED CHUNK") : *CurrentElement->Key.SoundWaveName.ToString(),
+			EnumToString(LoadingBehavior),
 			bWasCacheMiss ? TEXT("(Cache Miss!)") : TEXT(""),
 			bIsStaleChunk ? TEXT("(Stale Chunk)") : TEXT(""),
 			CurrentElement->IsLoadInProgress() ? TEXT("(Loading In Progress)") : TEXT("")
diff --git a/Engine/Source/Runtime/Engine/Public/AudioStreamingCache.h b/Engine/Source/Runtime/Engine/Public/AudioStreamingCache.h
--- a/Engine/Source/Runtime/Engine/Public/AudioStreamingCache.h
+++ b/Engine/Source/Runtime/Engine/Public/AudioStreamingCache.h
@@ -16,6 +16,8 @@ AudioStreaming.h: Definitions of classes used for audio streaming.
 #include "Async/AsyncFileHandle.h"
 #include "HAL/ThreadSafeBool.h"
 #include "AudioStreaming.h"
+#include "Sound/SoundWave.h"
+#include "Sound/SoundWaveLoadingBehavior.h"
 
 
 #define DEBUG_STREAM_CACHE !UE_BUILD_SHIPPING
@@ -120,6 +122,9 @@ private:
 		// This is a cumulative moving average of a chunks location before it was 
 		float AverageLocationInCacheWhenNeeded;
 
+		// Note the loading behavior of the sound wave that inserted this element into the cache
+		ESoundWaveLoadingBehavior LoadingBehavior;
+
 		// if true, 
 		bool bWasCacheMiss;
 
@@ -129,6 +134,7 @@ private:
 			, TimeLoadStarted(0.0)
 			, TimeToLoad(0.0)
 			, AverageLocationInCacheWhenNeeded(0.0f)
+			, LoadingBehavior(ESoundWaveLoadingBehavior::Uninitialized)
 			, bWasCacheMiss(false)
 		{
 		}
@@ -139,6 +145,7 @@ private:
 			NumTimesTouched = 0;
 			TimeLoadStarted = 0;
 			TimeToLoad = 0.0f;
+			LoadingBehavior = ESoundWaveLoadingBehavior::Uninitialized;
 			bWasCacheMiss = false;
 			AverageLocationInCacheWhenNeeded = 0.0f;
 		}
