Add fast path for lookups into the chunk table in the audio stream cache. Once a chunk is loaded or kicked off in the stream cache, the index of that chunk in the cache stable is stored on the USoundWave's list of chunks. If the chunk pointed to has been modified/evicted, we fallback to linear search behavior.

diff --git a/Engine/Source/Runtime/Engine/Classes/Sound/SoundWave.h b/Engine/Source/Runtime/Engine/Classes/Sound/SoundWave.h
--- a/Engine/Source/Runtime/Engine/Classes/Sound/SoundWave.h
+++ b/Engine/Source/Runtime/Engine/Classes/Sound/SoundWave.h
@@ -49,6 +49,8 @@ enum class ESoundWavePrecacheState
 	Done
 };
 
+constexpr uint64 InvalidAudioStreamCacheLookupID = TNumericLimits<uint64>::Max();
+
 /**
  * A chunk of streamed audio.
  */
@@ -63,8 +65,14 @@ struct FStreamedAudioChunk
 	/** Bulk data if stored in the package. */
 	FByteBulkData BulkData;
 
+	/** This is set by the audio stream cache to speed up lookup. Maps directly to the location into the cache chunk table that we loaded this chunk into. */
+	uint64 CacheLookupID;
+
 	/** Default constructor. */
 	FStreamedAudioChunk()
+		: DataSize(0)
+		, AudioDataSize(0)
+		, CacheLookupID(InvalidAudioStreamCacheLookupID)
 	{
 	}
 
@@ -485,6 +493,19 @@ public:
 
 	bool IsRetainingAudio();
 
+	/**
+	 * This is called by the audio stream cache once we put a chunk of compressed data in the chunk table. 
+	 * It allows us to look up the chunk directly, rather than searching the cache linearly.
+	 */
+	void SetCacheLookupIDForChunk(uint32 InChunkIndex, uint64 InCacheLookupID);
+
+	/**
+	 * This is called by the audio stream cache when we try to get a chunk of compressed data.
+	 * It allows us to look up the chunk directly, rather than searching the cache linearly.
+	 * Returns InvalidAudioStreamCacheLookupID if the chunk hasn't been added to the cache.
+	 */
+	uint64 GetCacheLookupIDForChunk(uint32 InChunkIndex);
+
 	/** Returns the loading behavior we should use for this sound wave.
 	 *  If this is called within Serialize(), this should be called with bCheckSoundClasses = false,
 	 *  Since there is no guarantee that the deserialized USoundClasses have been resolved yet.
diff --git a/Engine/Source/Runtime/Engine/Private/AudioStreaming.cpp b/Engine/Source/Runtime/Engine/Private/AudioStreaming.cpp
--- a/Engine/Source/Runtime/Engine/Private/AudioStreaming.cpp
+++ b/Engine/Source/Runtime/Engine/Private/AudioStreaming.cpp
@@ -920,7 +920,7 @@ FAudioChunkHandle FLegacyAudioStreamingManager::GetLoadedChunk(const USoundWave*
 				if (WaveData->LoadedChunks[Index].Index == ChunkIndex)
 				{
 					CriticalSection.Unlock();
-					return BuildChunkHandle(WaveData->LoadedChunks[Index].Data, WaveData->LoadedChunks[Index].AudioDataSize, SoundWave, SoundWave->GetFName(), ChunkIndex);
+					return BuildChunkHandle(WaveData->LoadedChunks[Index].Data, WaveData->LoadedChunks[Index].AudioDataSize, SoundWave, SoundWave->GetFName(), ChunkIndex, InvalidAudioStreamCacheLookupID);
 				}
 			}
 		}
diff --git a/Engine/Source/Runtime/Engine/Private/AudioStreamingCache.cpp b/Engine/Source/Runtime/Engine/Private/AudioStreamingCache.cpp
--- a/Engine/Source/Runtime/Engine/Private/AudioStreamingCache.cpp
+++ b/Engine/Source/Runtime/Engine/Private/AudioStreamingCache.cpp
@@ -323,11 +323,15 @@ FAudioChunkHandle FCachedAudioStreamingManager::GetLoadedChunk(const USoundWave*
 		}
 
 		// The function call below increments the reference count to the internal chunk.
-		TArrayView<uint8> LoadedChunk = Cache->GetChunk(ChunkKey, bBlockForLoad, (bForImmediatePlayback || bBlockForLoad));
+		uint64 LookupIDForChunk = ChunkKey.SoundWave->GetCacheLookupIDForChunk(ChunkKey.ChunkIndex);
+		TArrayView<uint8> LoadedChunk = Cache->GetChunk(ChunkKey, bBlockForLoad, (bForImmediatePlayback || bBlockForLoad), LookupIDForChunk);
 		
 		// Ensure that, if we requested a synchronous load of this chunk, we didn't fail to load said chunk.
 		ensureAlwaysMsgf(!bBlockForLoad || LoadedChunk.GetData(), TEXT("Synchronous load of chunk index %d for SoundWave %s failed to return any data."), ChunkIndex, *SoundWave->GetName());
 
+		// Set the updated cache offset for this chunk index.
+		ChunkKey.SoundWave->SetCacheLookupIDForChunk(ChunkIndex, LookupIDForChunk);
+
 		UE_CLOG(!bBlockForLoad && !LoadedChunk.GetData(), LogAudio, Display, TEXT("GetLoadedChunk called for chunk index %d of SoundWave %s when audio was not loaded yet. This will result in latency."), ChunkIndex, *SoundWave->GetName());
 
 		// Finally, if there's a chunk after this in the sound, request that it is in the cache.
@@ -345,8 +349,8 @@ FAudioChunkHandle FCachedAudioStreamingManager::GetLoadedChunk(const USoundWave*
 #endif
 			};
 
-			bool bIsValidChunk = Cache->AddOrTouchChunk(NextChunkKey, [](EAudioChunkLoadResult) {}, ENamedThreads::AnyThread, false);
-			if (!bIsValidChunk)
+			uint64 LookupIDForNextChunk = Cache->AddOrTouchChunk(NextChunkKey, [](EAudioChunkLoadResult) {}, ENamedThreads::AnyThread, false);
+			if (LookupIDForNextChunk == InvalidAudioStreamCacheLookupID)
 			{
 				UE_LOG(LogAudio, Warning, TEXT("Cache overflow!!! couldn't load chunk %d for sound %s!"), ChunkIndex, *SoundWave->GetName());
 
@@ -373,9 +377,13 @@ FAudioChunkHandle FCachedAudioStreamingManager::GetLoadedChunk(const USoundWave*
 					UE_LOG(LogAudio, Warning, TEXT("Removed %d retained sounds from the stream cache."), NumChunksReleased);
 				});
 			}
+			else
+			{
+				NextChunkKey.SoundWave->SetCacheLookupIDForChunk(NextChunkKey.ChunkIndex, LookupIDForNextChunk);
+			}
 		}
 
-		return BuildChunkHandle(LoadedChunk.GetData(), LoadedChunk.Num(), SoundWave, SoundWave->GetFName(), ChunkIndex);
+		return BuildChunkHandle(LoadedChunk.GetData(), LoadedChunk.Num(), SoundWave, SoundWave->GetFName(), ChunkIndex, LookupIDForChunk);
 	}
 	else
 	{
@@ -425,7 +433,8 @@ int32 FCachedAudioStreamingManager::GetNextChunkIndex(const USoundWave* InSoundW
 {
 	LLM_SCOPE(ELLMTag::AudioStreamCache);
 	check(InSoundWave);
-	// TODO: Figure out a way to tell whether this wave is looping or not.
+	// TODO: Figure out a way to tell whether this wave is looping or not. For now we always prime the first chunk
+	// during the playback of the last chunk.
 	// if(bNotLooping) return ((int32) CurrentChunkIndex) < (InSoundWave->RunningPlatformData->Chunks.Num() - 1);
 	
 	const int32 NumChunksTotal = InSoundWave->GetNumChunks();
@@ -463,7 +472,7 @@ void FCachedAudioStreamingManager::AddReferenceToChunk(const FAudioChunkHandle&
 #endif
 	};
 
-	Cache->AddNewReferenceToChunk(ChunkKey);
+	Cache->AddNewReferenceToChunk(ChunkKey, InHandle.CacheLookupID);
 }
 
 void FCachedAudioStreamingManager::RemoveReferenceToChunk(const FAudioChunkHandle& InHandle)
@@ -482,7 +491,7 @@ void FCachedAudioStreamingManager::RemoveReferenceToChunk(const FAudioChunkHandl
 #endif
 	};
 
-	Cache->RemoveReferenceToChunk(ChunkKey);
+	Cache->RemoveReferenceToChunk(ChunkKey, InHandle.CacheLookupID);
 }
 
 bool FCachedAudioStreamingManager::RequestChunk(USoundWave* SoundWave, uint32 ChunkIndex, TFunction<void(EAudioChunkLoadResult)> OnLoadCompleted, ENamedThreads::Type ThreadToCallOnLoadCompletedOn, bool bForImmediatePlayback)
@@ -492,12 +501,15 @@ bool FCachedAudioStreamingManager::RequestChunk(USoundWave* SoundWave, uint32 Ch
 	if (Cache)
 	{
 		FAudioChunkCache::FChunkKey ChunkKey = { SoundWave, SoundWave->GetFName(), ChunkIndex };
-		return Cache->AddOrTouchChunk(ChunkKey, OnLoadCompleted, ThreadToCallOnLoadCompletedOn, bForImmediatePlayback);
+		uint64 LookupIDForChunk = Cache->AddOrTouchChunk(ChunkKey, OnLoadCompleted, ThreadToCallOnLoadCompletedOn, bForImmediatePlayback);
+		SoundWave->SetCacheLookupIDForChunk(ChunkIndex, LookupIDForChunk);
+		return LookupIDForChunk != InvalidAudioStreamCacheLookupID;
 	}
 	else
 	{
 		// This can hit if an out of bounds chunk was requested, or the zeroth chunk was requested from the streaming manager.
 		ensureMsgf(false, TEXT("GetCacheForWave failed for SoundWave %s!"), *SoundWave->GetName());
+		SoundWave->SetCacheLookupIDForChunk(ChunkIndex, InvalidAudioStreamCacheLookupID);
 		return false;
 	}
 }
@@ -526,7 +538,7 @@ FAudioChunkCache::~FAudioChunkCache()
 	check(NumberOfLoadsInFlight.GetValue() == 0);
 }
 
-bool FAudioChunkCache::AddOrTouchChunk(const FChunkKey& InKey, TFunction<void(EAudioChunkLoadResult)> OnLoadCompleted, ENamedThreads::Type CallbackThread, bool bNeededForPlayback)
+uint64 FAudioChunkCache::AddOrTouchChunk(const FChunkKey& InKey, TFunction<void(EAudioChunkLoadResult)> OnLoadCompleted, ENamedThreads::Type CallbackThread, bool bNeededForPlayback)
 {
 	// Update cache limit if needed.
 	if (!FMath::IsNearlyZero(StreamCacheSizeOverrideMBCVar) && StreamCacheSizeOverrideMBCVar > 0.0f)
@@ -538,12 +550,13 @@ bool FAudioChunkCache::AddOrTouchChunk(const FChunkKey& InKey, TFunction<void(EA
 	{
 		ensure(false);
 		ExecuteOnLoadCompleteCallback(EAudioChunkLoadResult::ChunkOutOfBounds, OnLoadCompleted, CallbackThread);
-		return false;
+		return InvalidAudioStreamCacheLookupID;
 	}
 
 	FScopeLock ScopeLock(&CacheMutationCriticalSection);
 
-	FCacheElement* FoundElement = FindElementForKey(InKey);
+	const uint64 LookupIDForChunk = InKey.SoundWave->GetCacheLookupIDForChunk(InKey.ChunkIndex);
+	FCacheElement* FoundElement = FindElementForKey(InKey, LookupIDForChunk);
 	
 	if (FoundElement)
 	{
@@ -563,7 +576,7 @@ bool FAudioChunkCache::AddOrTouchChunk(const FChunkKey& InKey, TFunction<void(EA
 		FoundElement->DebugInfo.LoadingBehavior = InKey.SoundWave->GetLoadingBehavior(true);
 #endif
 
-		return true;
+		return FoundElement->CacheLookupID;
 	}
 	else
 	{
@@ -572,7 +585,7 @@ bool FAudioChunkCache::AddOrTouchChunk(const FChunkKey& InKey, TFunction<void(EA
 		if (!CacheElement)
 		{
 			ExecuteOnLoadCompleteCallback(EAudioChunkLoadResult::CacheBlown, OnLoadCompleted, CallbackThread);
-			return false;
+			return InvalidAudioStreamCacheLookupID;
 		}
 
 #if DEBUG_STREAM_CACHE
@@ -598,17 +611,17 @@ bool FAudioChunkCache::AddOrTouchChunk(const FChunkKey& InKey, TFunction<void(EA
 			CacheMissQueue.Enqueue(MoveTemp(CacheMissInfo));
 		}
 
-		return true;
+		return CacheElement->CacheLookupID;
 	}
 }
 
-TArrayView<uint8> FAudioChunkCache::GetChunk(const FChunkKey& InKey, bool bBlockForLoadCompletion, bool bNeededForPlayback)
+TArrayView<uint8> FAudioChunkCache::GetChunk(const FChunkKey& InKey, bool bBlockForLoadCompletion, bool bNeededForPlayback, uint64& OutCacheOffset)
 {
 	FScopeLock ScopeLock(&CacheMutationCriticalSection);
-
-	FCacheElement* FoundElement = FindElementForKey(InKey);
+	FCacheElement* FoundElement = FindElementForKey(InKey, OutCacheOffset);
 	if (FoundElement)
 	{
+		OutCacheOffset = FoundElement->CacheLookupID;
 		TouchElement(FoundElement);
 		if (FoundElement->IsLoadInProgress())
 		{
@@ -634,8 +647,12 @@ TArrayView<uint8> FAudioChunkCache::GetChunk(const FChunkKey& InKey, bool bBlock
 		FoundElement = InsertChunk(InKey);
 		if (!FoundElement)
 		{
+			OutCacheOffset = InvalidAudioStreamCacheLookupID;
 			return TArrayView<uint8>();
 		}
+		
+		OutCacheOffset = FoundElement->CacheLookupID;
+
 		if (bBlockForLoadCompletion)
 		{
 			FStreamedAudioChunk& Chunk = InKey.SoundWave->RunningPlatformData->Chunks[InKey.ChunkIndex];
@@ -646,6 +663,7 @@ TArrayView<uint8> FAudioChunkCache::GetChunk(const FChunkKey& InKey, bool bBlock
 #endif
 			MemoryCounterBytes -= FoundElement->ChunkData.Num();
 
+		
 			// Reallocate our chunk data This allows us to shrink if possible.
 			// Unfortunately, GetCopy will write out the full zero-padded length of the bulk data,
 			// rather than just the audio data. So we set the array to Chunk.DataSize, then shrink to Chunk.AudioDataSize.
@@ -680,9 +698,9 @@ TArrayView<uint8> FAudioChunkCache::GetChunk(const FChunkKey& InKey, bool bBlock
 	}
 }
 
-void FAudioChunkCache::AddNewReferenceToChunk(const FChunkKey& InKey)
+void FAudioChunkCache::AddNewReferenceToChunk(const FChunkKey& InKey, uint64 ChunkOffset)
 {
-	FCacheElement* FoundElement = FindElementForKey(InKey);
+	FCacheElement* FoundElement = FindElementForKey(InKey, ChunkOffset);
 	check(FoundElement);
 
 	// If this value is ever negative, it means that we're decrementing more than we're incrementing:
@@ -690,9 +708,9 @@ void FAudioChunkCache::AddNewReferenceToChunk(const FChunkKey& InKey)
 	FoundElement->NumConsumers.Increment();
 }
 
-void FAudioChunkCache::RemoveReferenceToChunk(const FChunkKey& InKey)
+void FAudioChunkCache::RemoveReferenceToChunk(const FChunkKey& InKey, uint64 ChunkOffset)
 {
-	FCacheElement* FoundElement = FindElementForKey(InKey);
+	FCacheElement* FoundElement = FindElementForKey(InKey, ChunkOffset);
 	check(FoundElement);
 
 	// If this value is ever less than 1 when we hit this code, it means that we're decrementing more than we're incrementing:
@@ -914,9 +932,22 @@ FString FAudioChunkCache::FlushCacheMissLog()
 	return TopChunkMissesLog + TEXT("\n") + ConcatenatedCacheMisses;
 }
 
-FAudioChunkCache::FCacheElement* FAudioChunkCache::FindElementForKey(const FChunkKey& InKey)
+FAudioChunkCache::FCacheElement* FAudioChunkCache::FindElementForKey(const FChunkKey& InKey, uint64 CacheOffset)
 {
 	FScopeLock ScopeLock(&CacheMutationCriticalSection);
+
+	// If we have a known cache offset, access that chunk directly.
+	if (CacheOffset != InvalidAudioStreamCacheLookupID)
+	{
+		check(CacheOffset < CachePool.Num());
+		// Finally, sanity check that the key is still the same.
+		if (CachePool[CacheOffset].Key == InKey)
+		{
+			return &CachePool[CacheOffset];
+		}
+	}
+
+	//Otherwise, linearly search the cache.
 	FCacheElement* CurrentElement = MostRecentElement;
 
 #if DEBUG_STREAM_CACHE
@@ -1023,6 +1054,7 @@ FAudioChunkCache::FCacheElement* FAudioChunkCache::InsertChunk(const FChunkKey&
 		{
 			// We haven't filled up the pool yet, so we don't need to evict anything.
 			CacheElement = &CachePool[ChunksInUse];
+			CacheElement->CacheLookupID = ChunksInUse;
 			ChunksInUse++;
 		}
 		else
@@ -1056,6 +1088,7 @@ FAudioChunkCache::FCacheElement* FAudioChunkCache::InsertChunk(const FChunkKey&
 		}
 	}
 
+	InKey.SoundWave->SetCacheLookupIDForChunk(InKey.ChunkIndex, CacheElement->CacheLookupID);
 	return CacheElement;
 }
 
diff --git a/Engine/Source/Runtime/Engine/Private/ContentStreaming.cpp b/Engine/Source/Runtime/Engine/Private/ContentStreaming.cpp
--- a/Engine/Source/Runtime/Engine/Private/ContentStreaming.cpp
+++ b/Engine/Source/Runtime/Engine/Private/ContentStreaming.cpp
@@ -1403,18 +1403,20 @@ FAudioChunkHandle::FAudioChunkHandle()
 	, CorrespondingWave(nullptr)
 	, CorrespondingWaveName()
 	, ChunkIndex(INDEX_NONE)
+	, CacheLookupID(InvalidAudioStreamCacheLookupID)
 #if WITH_EDITOR
 	, ChunkGeneration(INDEX_NONE)
 #endif
 {
 }
 
-FAudioChunkHandle::FAudioChunkHandle(const uint8* InData, uint32 NumBytes, const USoundWave* InSoundWave, const FName& SoundWaveName, uint32 InChunkIndex)
+FAudioChunkHandle::FAudioChunkHandle(const uint8* InData, uint32 NumBytes, const USoundWave* InSoundWave, const FName& SoundWaveName, uint32 InChunkIndex, uint64 InCacheLookupID)
 	: CachedData(InData)
 	, CachedDataNumBytes(NumBytes)
 	, CorrespondingWave(InSoundWave)
 	, CorrespondingWaveName(SoundWaveName)
 	, ChunkIndex(InChunkIndex)
+	, CacheLookupID(InCacheLookupID)
 #if WITH_EDITOR
 	, ChunkGeneration(InSoundWave->CurrentChunkRevision.GetValue())
 #endif
@@ -1440,6 +1442,7 @@ FAudioChunkHandle& FAudioChunkHandle::operator=(const FAudioChunkHandle& Other)
 	CorrespondingWave = Other.CorrespondingWave;
 	CorrespondingWaveName = Other.CorrespondingWaveName;
 	ChunkIndex = Other.ChunkIndex;
+	CacheLookupID = Other.CacheLookupID;
 #if WITH_EDITOR
 	ChunkGeneration = Other.ChunkGeneration;
 #endif
@@ -1491,7 +1494,7 @@ bool FAudioChunkHandle::IsStale() const
 }
 #endif
 
-FAudioChunkHandle IAudioStreamingManager::BuildChunkHandle(const uint8* InData, uint32 NumBytes, const USoundWave* InSoundWave, const FName& SoundWaveName, uint32 InChunkIndex)
+FAudioChunkHandle IAudioStreamingManager::BuildChunkHandle(const uint8* InData, uint32 NumBytes, const USoundWave* InSoundWave, const FName& SoundWaveName, uint32 InChunkIndex, uint64 InCacheLookupID)
 {
-	return FAudioChunkHandle(InData, NumBytes, InSoundWave, SoundWaveName, InChunkIndex);
+	return FAudioChunkHandle(InData, NumBytes, InSoundWave, SoundWaveName, InChunkIndex, InCacheLookupID);
 }
diff --git a/Engine/Source/Runtime/Engine/Private/SoundWave.cpp b/Engine/Source/Runtime/Engine/Private/SoundWave.cpp
--- a/Engine/Source/Runtime/Engine/Private/SoundWave.cpp
+++ b/Engine/Source/Runtime/Engine/Private/SoundWave.cpp
@@ -2535,6 +2535,22 @@ bool USoundWave::IsRetainingAudio()
 	return FirstChunk.IsValid();
 }
 
+void USoundWave::SetCacheLookupIDForChunk(uint32 InChunkIndex, uint64 InCacheLookupID)
+{
+	check(FPlatformCompressionUtilities::IsCurrentPlatformUsingStreamCaching());
+	check(RunningPlatformData && InChunkIndex < ((uint32)RunningPlatformData->Chunks.Num()));
+
+	RunningPlatformData->Chunks[InChunkIndex].CacheLookupID = InCacheLookupID;
+}
+
+uint64 USoundWave::GetCacheLookupIDForChunk(uint32 InChunkIndex)
+{
+	check(FPlatformCompressionUtilities::IsCurrentPlatformUsingStreamCaching());
+	check(RunningPlatformData && (InChunkIndex < (uint32)RunningPlatformData->Chunks.Num()));
+
+	return RunningPlatformData->Chunks[InChunkIndex].CacheLookupID;
+}
+
 void USoundWave::CacheInheritedLoadingBehavior()
 {
 	check(IsInGameThread());
diff --git a/Engine/Source/Runtime/Engine/Public/AudioStreamingCache.h b/Engine/Source/Runtime/Engine/Public/AudioStreamingCache.h
--- a/Engine/Source/Runtime/Engine/Public/AudioStreamingCache.h
+++ b/Engine/Source/Runtime/Engine/Public/AudioStreamingCache.h
@@ -53,15 +53,18 @@ public:
 	
 	~FAudioChunkCache();
 
-	// Places chunk in cache, or puts this chunk back at the top of the cache if it's already loaded. Returns false on failure.
-	bool AddOrTouchChunk(const FChunkKey& InKey, TFunction<void(EAudioChunkLoadResult) > OnLoadCompleted, ENamedThreads::Type CallbackThread, bool bNeededForPlayback);
+	// Places chunk in cache, or puts this chunk back at the top of the cache if it's already loaded. Returns the static lookup ID of the chunk in the cache on success,
+	// or InvalidAudioStreamCacheLookupID on failiure.
+	uint64 AddOrTouchChunk(const FChunkKey& InKey, TFunction<void(EAudioChunkLoadResult) > OnLoadCompleted, ENamedThreads::Type CallbackThread, bool bNeededForPlayback);
 
 	// Returns the chunk asked for, or an empty TArrayView if that chunk is not loaded.
-	TArrayView<uint8> GetChunk(const FChunkKey& InKey, bool bBlockForLoadCompletion, bool bNeededForPlayback);
+	// InOutCacheLookupID can optionally be set as a cache offset to use directly rather than searching the cache for a matching chunk.
+	// InOutCacheLookupID will be set to the offset the chunk is in the cache, which can be used for faster lookup in the future.
+	TArrayView<uint8> GetChunk(const FChunkKey& InKey, bool bBlockForLoadCompletion, bool bNeededForPlayback, uint64& InOutCacheLookupID);
 
 	// add an additional reference for a chunk.
-	void AddNewReferenceToChunk(const FChunkKey& InKey);
-	void RemoveReferenceToChunk(const FChunkKey& InKey);
+	void AddNewReferenceToChunk(const FChunkKey& InKey, uint64 InCacheLookupID);
+	void RemoveReferenceToChunk(const FChunkKey& InKey, uint64 InCacheLookupID);
 
 	// Evict all sounds from the cache.
 	void ClearCache();
@@ -160,7 +163,7 @@ private:
 		uint32 ChunkDataSize;
 		FCacheElement* MoreRecentElement;
 		FCacheElement* LessRecentElement;
-		uint32 CacheIndex;
+		uint64 CacheLookupID;
 
 		FThreadSafeBool bIsLoaded;
 		
@@ -182,7 +185,7 @@ private:
 			: ChunkDataSize(0)
 			, MoreRecentElement(nullptr)
 			, LessRecentElement(nullptr)
-			, CacheIndex(InCacheIndex)
+			, CacheLookupID(InCacheIndex)
 			, bIsLoaded(false)
 		{
 		}
@@ -270,7 +273,8 @@ private:
 	bool bLogCacheMisses;
 
 	// Returns cached element if it exists in our cache, nullptr otherwise.
-	FCacheElement* FindElementForKey(const FChunkKey& InKey);
+	// If the index of the element is already known, it can be used here to avoid searching the cache.
+	FCacheElement* FindElementForKey(const FChunkKey& InKey, uint64 CacheLookupID = InvalidAudioStreamCacheLookupID);
 
 	// Puts this element at the front of the linked list.
 	void TouchElement(FCacheElement* InElement);
@@ -350,7 +354,7 @@ public:
 	virtual void NotifyLevelOffset( class ULevel* Level, const FVector& Offset ) override;
 	// End IStreamingManager interface
 
-	// IAudioStreamingManager interface
+	// IAudioStreamingManager interface (unused functions)
 	virtual void AddStreamingSoundWave(USoundWave* SoundWave) override;
 	virtual void RemoveStreamingSoundWave(USoundWave* SoundWave) override;
 	virtual void AddDecoder(ICompressedAudioInfo* CompressedAudioInfo) override;
@@ -361,6 +365,9 @@ public:
 	virtual void AddStreamingSoundSource(FSoundSource* SoundSource) override;
 	virtual void RemoveStreamingSoundSource(FSoundSource* SoundSource) override;
 	virtual bool IsManagedStreamingSoundSource(const FSoundSource* SoundSource) const override;
+	// End IAudioStreamingManager interface (unused)
+
+	// IAudioStreamingManager interface (used functions)
 	virtual bool RequestChunk(USoundWave* SoundWave, uint32 ChunkIndex, TFunction<void(EAudioChunkLoadResult)> OnLoadCompleted, ENamedThreads::Type ThreadToCallOnLoadCompletedOn, bool bForImmediatePlayback = false) override;
 	virtual FAudioChunkHandle GetLoadedChunk(const USoundWave* SoundWave, uint32 ChunkIndex, bool bBlockForLoad = false, bool bForImmediatePlayback = false) const override;
 	virtual uint64 TrimMemory(uint64 NumBytesToFree) override;
diff --git a/Engine/Source/Runtime/Engine/Public/ContentStreaming.h b/Engine/Source/Runtime/Engine/Public/ContentStreaming.h
--- a/Engine/Source/Runtime/Engine/Public/ContentStreaming.h
+++ b/Engine/Source/Runtime/Engine/Public/ContentStreaming.h
@@ -115,15 +115,21 @@ public:
 
 private:
 	// This constructor should only be called by an implementation of IAudioStreamingManager.
-	FAudioChunkHandle(const uint8* InData, uint32 NumBytes, const USoundWave* InSoundWave, const FName& SoundWaveName, uint32 InChunkIndex);
+	FAudioChunkHandle(const uint8* InData, uint32 NumBytes, const USoundWave* InSoundWave, const FName& SoundWaveName, uint32 InChunkIndex, uint64 InCacheLookupID);
 
 	const uint8*  CachedData;
 	int32 CachedDataNumBytes;
 
 	const USoundWave* CorrespondingWave;
 	FName CorrespondingWaveName;
+
+	// The index of this chunk in the sound wave's full set of chunks of compressed audio.
 	int32 ChunkIndex;
 
+	// This ID can be used to access the element this handle is for directly,
+	// rather than linearly searching the cache. This should only be used by the stream cache itself.
+	uint64 CacheLookupID;
+
 #if WITH_EDITOR
 	uint32 ChunkGeneration;
 #endif
@@ -556,7 +562,7 @@ protected:
 	friend FAudioChunkHandle;
 
 	/** This can be called by implementers of IAudioStreamingManager to construct an FAudioChunkHandle using an otherwise inaccessible constructor. */
-	static FAudioChunkHandle BuildChunkHandle(const uint8* InData, uint32 NumBytes, const USoundWave* InSoundWave, const FName& SoundWaveName, uint32 InChunkIndex);
+	static FAudioChunkHandle BuildChunkHandle(const uint8* InData, uint32 NumBytes, const USoundWave* InSoundWave, const FName& SoundWaveName, uint32 InChunkIndex, uint64 CacheLookupID);
 
 	/**
 	 * This can be used to increment reference counted handles to audio chunks. Called by the copy constructor of FAudioChunkHandle.
