Restore incorrectly deleted file
#rnx

diff --git a/Engine/Source/Runtime/Engine/Classes/Sound/SoundWaveLoadingBehavior.cpp b/Engine/Source/Runtime/Engine/Classes/Sound/SoundWaveLoadingBehavior.cpp
new file mode 100644
--- /dev/null
+++ b/Engine/Source/Runtime/Engine/Classes/Sound/SoundWaveLoadingBehavior.cpp
@@ -0,0 +1,24 @@
+// Copyright Epic Games, Inc. All Rights Reserved.
+#include "SoundWaveLoadingBehavior.h"
+
+
+const TCHAR* EnumToString(ESoundWaveLoadingBehavior InCurrentState)
+{
+	switch (InCurrentState)
+	{
+	case ESoundWaveLoadingBehavior::Inherited:
+		return TEXT("Inherited");
+	case ESoundWaveLoadingBehavior::RetainOnLoad:
+		return TEXT("RetainOnLoad");
+	case ESoundWaveLoadingBehavior::PrimeOnLoad:
+		return TEXT("PrimeOnLoad");
+	case ESoundWaveLoadingBehavior::LoadOnDemand:
+		return TEXT("LoadOnDemand");
+	case ESoundWaveLoadingBehavior::ForceInline:
+		return TEXT("ForceInline");
+	case ESoundWaveLoadingBehavior::Uninitialized:
+		return TEXT("Uninitialized");
+	}
+	ensure(false);
+	return TEXT("Unknown");
+}
