CIS error, file include order issue

diff --git a/Engine/Source/Runtime/Engine/Classes/Sound/SoundWaveLoadingBehavior.cpp b/Engine/Source/Runtime/Engine/Classes/Sound/SoundWaveLoadingBehavior.cpp
--- a/Engine/Source/Runtime/Engine/Classes/Sound/SoundWaveLoadingBehavior.cpp
+++ b/Engine/Source/Runtime/Engine/Classes/Sound/SoundWaveLoadingBehavior.cpp
@@ -1,6 +1,5 @@
 // Copyright Epic Games, Inc. All Rights Reserved.
 #include "SoundWaveLoadingBehavior.h"
-#include "CoreMinimal.h"
 
 
 const TCHAR* EnumToString(ESoundWaveLoadingBehavior InCurrentState)
diff --git a/Engine/Source/Runtime/Engine/Classes/Sound/SoundWaveLoadingBehavior.h b/Engine/Source/Runtime/Engine/Classes/Sound/SoundWaveLoadingBehavior.h
--- a/Engine/Source/Runtime/Engine/Classes/Sound/SoundWaveLoadingBehavior.h
+++ b/Engine/Source/Runtime/Engine/Classes/Sound/SoundWaveLoadingBehavior.h
@@ -1,5 +1,6 @@
 // Copyright Epic Games, Inc. All Rights Reserved.
 #pragma once
+#include "CoreMinimal.h"
 
 /**
  * Only used when stream caching is enabled. Determines how we are going to load or retain a given audio asset.
