convert bulkdata size check to a conditional log.

diff --git a/Engine/Source/Runtime/Engine/Private/SoundWave.cpp b/Engine/Source/Runtime/Engine/Private/SoundWave.cpp
--- a/Engine/Source/Runtime/Engine/Private/SoundWave.cpp
+++ b/Engine/Source/Runtime/Engine/Private/SoundWave.cpp
@@ -930,7 +930,7 @@ void USoundWave::EnsureZerothChunkIsLoaded()
 	check(RunningPlatformData && RunningPlatformData->Chunks.Num() > 0);
 	FStreamedAudioChunk& ZerothChunk = RunningPlatformData->Chunks[0];
 	// Some sanity checks to ensure that the bulk size set up
-	check(ZerothChunk.BulkData.GetBulkDataSize() == ZerothChunk.DataSize);
+	UE_CLOG(ZerothChunk.BulkData.GetBulkDataSize() == ZerothChunk.DataSize, LogAudio, Warning, TEXT("Bulk data serialized out had a mismatched size with the DataSize field. Soundwave: %s Bulk Data Reported Size: %d Bulk Data Actual Size: %ld"), *GetFullName(), ZerothChunk.DataSize, ZerothChunk.BulkData.GetBulkDataSize());
 
 	ZerothChunkData = ZerothChunk.BulkData.GetCopyAsBuffer(ZerothChunk.AudioDataSize, true);
 #endif // WITH_EDITOR
