Add  LLM tag for spatialization plugins.

diff --git a/Engine/Source/Runtime/AudioMixer/Private/AudioMixerDevice.cpp b/Engine/Source/Runtime/AudioMixer/Private/AudioMixerDevice.cpp
--- a/Engine/Source/Runtime/AudioMixer/Private/AudioMixerDevice.cpp
+++ b/Engine/Source/Runtime/AudioMixer/Private/AudioMixerDevice.cpp
@@ -310,20 +310,24 @@ namespace Audio
 				PluginInitializationParams.BufferLength = OpenStreamParams.NumFrames;
 				PluginInitializationParams.AudioDevicePtr = this;
 
-				// Initialize any plugins if they exist
-				if (SpatializationPluginInterface.IsValid())
 				{
-					SpatializationPluginInterface->Initialize(PluginInitializationParams);
-				}
+					LLM_SCOPE(ELLMTag::AudioMixerPlugins);
 
-				if (OcclusionInterface.IsValid())
-				{
-					OcclusionInterface->Initialize(PluginInitializationParams);
-				}
+					// Initialize any plugins if they exist
+					if (SpatializationPluginInterface.IsValid())
+					{
+						SpatializationPluginInterface->Initialize(PluginInitializationParams);
+					}
 
-				if (ReverbPluginInterface.IsValid())
-				{
-					ReverbPluginInterface->Initialize(PluginInitializationParams);
+					if (OcclusionInterface.IsValid())
+					{
+						OcclusionInterface->Initialize(PluginInitializationParams);
+					}
+
+					if (ReverbPluginInterface.IsValid())
+					{
+						ReverbPluginInterface->Initialize(PluginInitializationParams);
+					}
 				}
 
 				// Need to set these up before we start the audio stream.
@@ -753,6 +757,7 @@ namespace Audio
 
 		if (IsReverbPluginEnabled() && ReverbPluginInterface)
 		{
+			LLM_SCOPE(ELLMTag::AudioMixerPlugins);
 			USoundSubmix* ReverbPluginSubmix = ReverbPluginInterface->GetSubmix();
 			check(ReverbPluginSubmix);
 			ReverbPluginSubmix->AddToRoot();
diff --git a/Engine/Source/Runtime/AudioMixer/Private/AudioMixerSourceManager.cpp b/Engine/Source/Runtime/AudioMixer/Private/AudioMixerSourceManager.cpp
--- a/Engine/Source/Runtime/AudioMixer/Private/AudioMixerSourceManager.cpp
+++ b/Engine/Source/Runtime/AudioMixer/Private/AudioMixerSourceManager.cpp
@@ -472,6 +472,7 @@ namespace Audio
 		if (SourceInfo.bUseHRTFSpatializer)
 		{
 			AUDIO_MIXER_CHECK(bUsingSpatializationPlugin);
+			LLM_SCOPE(ELLMTag::AudioMixerPlugins);
 			SpatializationPlugin->OnReleaseSource(SourceId);
 		}
 
@@ -697,6 +698,7 @@ namespace Audio
 			if (InitParams.bUseHRTFSpatialization)
 			{
 				AUDIO_MIXER_CHECK(bUsingSpatializationPlugin);
+				LLM_SCOPE(ELLMTag::AudioMixerPlugins);
 				SpatializationPlugin->OnInitSource(SourceId, InitParams.AudioComponentUserID, InitParams.SpatializationPluginSettings);
 			}
 
@@ -1829,7 +1831,10 @@ namespace Audio
 				SourceInfo.AudioPluginOutputData.AudioBuffer.AddZeroed(2 * NumOutputFrames);
 			}
 
-			SpatializationPlugin->ProcessAudio(AudioPluginInputData, SourceInfo.AudioPluginOutputData);
+			{
+				LLM_SCOPE(ELLMTag::AudioMixerPlugins);
+				SpatializationPlugin->ProcessAudio(AudioPluginInputData, SourceInfo.AudioPluginOutputData);
+			}
 
 			// If this is an external send, we treat this source audio as if it was still a mono source
 			// This will allow it to traditionally pan in the ComputeOutputBuffers function and be
@@ -2672,6 +2677,7 @@ namespace Audio
 		if (bUsingSpatializationPlugin)
 		{
 			AUDIO_MIXER_CHECK(SpatializationPlugin.IsValid());
+			LLM_SCOPE(ELLMTag::AudioMixerPlugins);
 			SpatializationPlugin->OnAllSourcesProcessed();
 		}
 
diff --git a/Engine/Source/Runtime/Core/Private/HAL/LowLevelMemTracker.cpp b/Engine/Source/Runtime/Core/Private/HAL/LowLevelMemTracker.cpp
--- a/Engine/Source/Runtime/Core/Private/HAL/LowLevelMemTracker.cpp
+++ b/Engine/Source/Runtime/Core/Private/HAL/LowLevelMemTracker.cpp
@@ -58,6 +58,7 @@ DECLARE_LLM_MEMORY_STAT(TEXT("Audio"), STAT_AudioLLM, STATGROUP_LLMFULL);
 DECLARE_LLM_MEMORY_STAT(TEXT("AudioMisc"), STAT_AudioMiscLLM, STATGROUP_LLMFULL);
 DECLARE_LLM_MEMORY_STAT(TEXT("AudioSoundWaves"), STAT_AudioSoundWavesLLM, STATGROUP_LLMFULL);
 DECLARE_LLM_MEMORY_STAT(TEXT("AudioMixer"), STAT_AudioMixerLLM, STATGROUP_LLMFULL);
+DECLARE_LLM_MEMORY_STAT(TEXT("AudioMixerPlugins"), STAT_AudioMixerPluginsLLM, STATGROUP_LLMFULL);
 DECLARE_LLM_MEMORY_STAT(TEXT("AudioPrecache"), STAT_AudioPrecacheLLM, STATGROUP_LLMFULL);
 DECLARE_LLM_MEMORY_STAT(TEXT("AudioDecompress"), STAT_AudioDecompressLLM, STATGROUP_LLMFULL);
 DECLARE_LLM_MEMORY_STAT(TEXT("AudioRealtimePrecache"), STAT_AudioRealtimePrecacheLLM, STATGROUP_LLMFULL);
diff --git a/Engine/Source/Runtime/Core/Public/HAL/LowLevelMemTracker.h b/Engine/Source/Runtime/Core/Public/HAL/LowLevelMemTracker.h
--- a/Engine/Source/Runtime/Core/Public/HAL/LowLevelMemTracker.h
+++ b/Engine/Source/Runtime/Core/Public/HAL/LowLevelMemTracker.h
@@ -159,6 +159,7 @@ enum class ELLMTagSet : uint8
 	macro(AudioMisc,							"AudioMisc",					GET_STATFNAME(STAT_AudioMiscLLM),							GET_STATFNAME(STAT_AudioSummaryLLM),			ELLMTag::Audio)\
 	macro(AudioSoundWaves,						"AudioSoundWaves",				GET_STATFNAME(STAT_AudioSoundWavesLLM),						GET_STATFNAME(STAT_AudioSummaryLLM),			ELLMTag::Audio)\
 	macro(AudioMixer,							"AudioMixer",					GET_STATFNAME(STAT_AudioMixerLLM),							GET_STATFNAME(STAT_AudioSummaryLLM),			ELLMTag::Audio)\
+	macro(AudioMixerPlugins,					"AudioMixerPlugins",			GET_STATFNAME(STAT_AudioMixerPluginsLLM),					GET_STATFNAME(STAT_AudioSummaryLLM),			ELLMTag::Audio)\
 	macro(AudioPrecache,						"AudioPrecache",				GET_STATFNAME(STAT_AudioPrecacheLLM),						GET_STATFNAME(STAT_AudioSummaryLLM),			ELLMTag::Audio)\
 	macro(AudioDecompress,						"AudioDecompress",				GET_STATFNAME(STAT_AudioDecompressLLM),						GET_STATFNAME(STAT_AudioSummaryLLM),			ELLMTag::Audio)\
 	macro(AudioRealtimePrecache,				"AudioRealtimePrecache",		GET_STATFNAME(STAT_AudioRealtimePrecacheLLM),				GET_STATFNAME(STAT_AudioSummaryLLM),			ELLMTag::Audio)\
diff --git a/Engine/Source/Runtime/Engine/Private/AudioDevice.cpp b/Engine/Source/Runtime/Engine/Private/AudioDevice.cpp
--- a/Engine/Source/Runtime/Engine/Private/AudioDevice.cpp
+++ b/Engine/Source/Runtime/Engine/Private/AudioDevice.cpp
@@ -420,79 +420,83 @@ bool FAudioDevice::Init(Audio::FDeviceId InDeviceID, int32 InMaxSources)
 		NumStoppingSources = 0;
 	}
 
-	// Cache any plugin settings objects we have loaded
-	UpdateAudioPluginSettingsObjectCache();
-
-	//Get the requested spatialization plugin and set it up.
-	IAudioSpatializationFactory* SpatializationPluginFactory = AudioPluginUtilities::GetDesiredSpatializationPlugin();
-	if (SpatializationPluginFactory != nullptr)
 	{
-		SpatializationPluginInterface = SpatializationPluginFactory->CreateNewSpatializationPlugin(this);
-		if (!IsAudioMixerEnabled())
+		LLM_SCOPE(ELLMTag::AudioMixerPlugins);
+
+		// Cache any plugin settings objects we have loaded
+		UpdateAudioPluginSettingsObjectCache();
+
+		//Get the requested spatialization plugin and set it up.
+		IAudioSpatializationFactory* SpatializationPluginFactory = AudioPluginUtilities::GetDesiredSpatializationPlugin();
+		if (SpatializationPluginFactory != nullptr)
+		{
+			SpatializationPluginInterface = SpatializationPluginFactory->CreateNewSpatializationPlugin(this);
+			if (!IsAudioMixerEnabled())
+			{
+				//Set up initialization parameters for system level effect plugins:
+				FAudioPluginInitializationParams PluginInitializationParams;
+				PluginInitializationParams.SampleRate = SampleRate;
+				PluginInitializationParams.NumSources = GetMaxSources();
+				PluginInitializationParams.BufferLength = PlatformSettings.CallbackBufferFrameSize;
+				PluginInitializationParams.AudioDevicePtr = this;
+
+				SpatializationPluginInterface->Initialize(PluginInitializationParams);
+			}
+
+			bSpatializationInterfaceEnabled = true;
+			bSpatializationIsExternalSend = SpatializationPluginFactory->IsExternalSend();
+			MaxChannelsSupportedBySpatializationPlugin = SpatializationPluginFactory->GetMaxSupportedChannels();
+			UE_LOG(LogAudio, Log, TEXT("Using Audio Spatialization Plugin: %s is external send: %d"), *(SpatializationPluginFactory->GetDisplayName()), bSpatializationIsExternalSend);
+		}
+		else
+		{
+			UE_LOG(LogAudio, Log, TEXT("Using built-in audio spatialization."));
+		}
+
+		//Get the requested reverb plugin and set it up:
+		IAudioReverbFactory* ReverbPluginFactory = AudioPluginUtilities::GetDesiredReverbPlugin();
+		if (ReverbPluginFactory != nullptr)
+		{
+			ReverbPluginInterface = ReverbPluginFactory->CreateNewReverbPlugin(this);
+			bReverbInterfaceEnabled = true;
+			bReverbIsExternalSend = ReverbPluginFactory->IsExternalSend();
+			UE_LOG(LogAudio, Log, TEXT("Audio Reverb Plugin: %s"), *(ReverbPluginFactory->GetDisplayName()));
+		}
+		else
+		{
+			UE_LOG(LogAudio, Log, TEXT("Using built-in audio reverb."));
+		}
+
+		//Get the requested occlusion plugin and set it up.
+		IAudioOcclusionFactory* OcclusionPluginFactory = AudioPluginUtilities::GetDesiredOcclusionPlugin();
+		if (OcclusionPluginFactory != nullptr)
+		{
+			OcclusionInterface = OcclusionPluginFactory->CreateNewOcclusionPlugin(this);
+			bOcclusionInterfaceEnabled = true;
+			bOcclusionIsExternalSend = OcclusionPluginFactory->IsExternalSend();
+			UE_LOG(LogAudio, Display, TEXT("Audio Occlusion Plugin: %s"), *(OcclusionPluginFactory->GetDisplayName()));
+		}
+		else
+		{
+			UE_LOG(LogAudio, Display, TEXT("Using built-in audio occlusion."));
+		}
+
+		//Get the requested modulation plugin and set it up.
+		if (IAudioModulationFactory* ModulationPluginFactory = AudioPluginUtilities::GetDesiredModulationPlugin())
 		{
+			ModulationInterface = ModulationPluginFactory->CreateNewModulationPlugin(this);
+
 			//Set up initialization parameters for system level effect plugins:
 			FAudioPluginInitializationParams PluginInitializationParams;
 			PluginInitializationParams.SampleRate = SampleRate;
 			PluginInitializationParams.NumSources = GetMaxSources();
 			PluginInitializationParams.BufferLength = PlatformSettings.CallbackBufferFrameSize;
 			PluginInitializationParams.AudioDevicePtr = this;
+			ModulationInterface->Initialize(PluginInitializationParams);
 
-			SpatializationPluginInterface->Initialize(PluginInitializationParams);
+			bModulationInterfaceEnabled = true;
+			UE_LOG(LogAudio, Display, TEXT("Audio Modulation Plugin: %s"), *(ModulationPluginFactory->GetDisplayName().ToString()));
 		}
-
-		bSpatializationInterfaceEnabled = true;
-		bSpatializationIsExternalSend = SpatializationPluginFactory->IsExternalSend();
-		MaxChannelsSupportedBySpatializationPlugin = SpatializationPluginFactory->GetMaxSupportedChannels();
-		UE_LOG(LogAudio, Log, TEXT("Using Audio Spatialization Plugin: %s is external send: %d"), *(SpatializationPluginFactory->GetDisplayName()), bSpatializationIsExternalSend);
-	}
-	else
-	{
-		UE_LOG(LogAudio, Log, TEXT("Using built-in audio spatialization."));
-	}
-
-	//Get the requested reverb plugin and set it up:
-	IAudioReverbFactory* ReverbPluginFactory = AudioPluginUtilities::GetDesiredReverbPlugin();
-	if (ReverbPluginFactory != nullptr)
-	{
-		ReverbPluginInterface = ReverbPluginFactory->CreateNewReverbPlugin(this);
-		bReverbInterfaceEnabled = true;
-		bReverbIsExternalSend = ReverbPluginFactory->IsExternalSend();
-		UE_LOG(LogAudio, Log, TEXT("Audio Reverb Plugin: %s"), *(ReverbPluginFactory->GetDisplayName()));
-	}
-	else
-	{
-		UE_LOG(LogAudio, Log, TEXT("Using built-in audio reverb."));
-	}
-
-	//Get the requested occlusion plugin and set it up.
-	IAudioOcclusionFactory* OcclusionPluginFactory = AudioPluginUtilities::GetDesiredOcclusionPlugin();
-	if (OcclusionPluginFactory != nullptr)
-	{
-		OcclusionInterface = OcclusionPluginFactory->CreateNewOcclusionPlugin(this);
-		bOcclusionInterfaceEnabled = true;
-		bOcclusionIsExternalSend = OcclusionPluginFactory->IsExternalSend();
-		UE_LOG(LogAudio, Display, TEXT("Audio Occlusion Plugin: %s"), *(OcclusionPluginFactory->GetDisplayName()));
-	}
-	else
-	{
-		UE_LOG(LogAudio, Display, TEXT("Using built-in audio occlusion."));
-	}
-
-	//Get the requested modulation plugin and set it up.
-	if (IAudioModulationFactory* ModulationPluginFactory = AudioPluginUtilities::GetDesiredModulationPlugin())
-	{
-		ModulationInterface = ModulationPluginFactory->CreateNewModulationPlugin(this);
-
-		//Set up initialization parameters for system level effect plugins:
-		FAudioPluginInitializationParams PluginInitializationParams;
-		PluginInitializationParams.SampleRate = SampleRate;
-		PluginInitializationParams.NumSources = GetMaxSources();
-		PluginInitializationParams.BufferLength = PlatformSettings.CallbackBufferFrameSize;
-		PluginInitializationParams.AudioDevicePtr = this;
-		ModulationInterface->Initialize(PluginInitializationParams);
-
-		bModulationInterfaceEnabled = true;
-		UE_LOG(LogAudio, Display, TEXT("Audio Modulation Plugin: %s"), *(ModulationPluginFactory->GetDisplayName().ToString()));
 	}
 
 	// allow the platform to startup
@@ -690,6 +694,8 @@ void FAudioDevice::Teardown()
 	Sources.Reset();
 	FreeSources.Reset();
 
+	LLM_SCOPE(ELLMTag::AudioMixerPlugins);
+
 	if (SpatializationPluginInterface.IsValid())
 	{
 		SpatializationPluginInterface->Shutdown();
