Stop shadow fade-out transition at the end of CSM shadows for directional lights when there is static shadowing taking over (as blending is then done in the lighting shader)

diff --git a/Engine/Source/Runtime/Engine/Private/Components/DirectionalLightComponent.cpp b/Engine/Source/Runtime/Engine/Private/Components/DirectionalLightComponent.cpp
--- a/Engine/Source/Runtime/Engine/Private/Components/DirectionalLightComponent.cpp
+++ b/Engine/Source/Runtime/Engine/Private/Components/DirectionalLightComponent.cpp
@@ -787,7 +787,8 @@ private:
 		{
 			SplitFar += FadeExtension;
 		}
-		else 
+		// Only do this if there is no static shadowing taking over after the fade (i.e., the intention is to fade to nothing, when only dynamic shadowing is used)
+		else if (!(bPrecomputedLightingIsValid && bStaticShadowing))
 		{
 			if (bIsRayTracedCascade)
 			{
