#ue4 - Move overlap initialization used within AActor::DispatchBeginPlay() to its own function for clarity and for possibly future adjustments.

diff --git a/Engine/Source/Runtime/Engine/Classes/GameFramework/Actor.h b/Engine/Source/Runtime/Engine/Classes/GameFramework/Actor.h
--- a/Engine/Source/Runtime/Engine/Classes/GameFramework/Actor.h
+++ b/Engine/Source/Runtime/Engine/Classes/GameFramework/Actor.h
@@ -447,6 +447,9 @@ private:
 	UPROPERTY(Config, Category = Collision, VisibleAnywhere)
 	EActorUpdateOverlapsMethod DefaultUpdateOverlapsMethodDuringLevelStreaming;
 
+	/** Internal helper to update Overlaps during Actor initialization/BeginPlay correctly based on the UpdateOverlapsMethodDuringLevelStreaming and bGenerateOverlapEventsDuringLevelStreaming settings. */
+	void UpdateInitialOverlaps(bool bFromLevelStreaming);
+
 	/** Describes how much control the remote machine has over the actor. */
 	UPROPERTY(Replicated, Transient)
 	TEnumAsByte<enum ENetRole> RemoteRole;	
diff --git a/Engine/Source/Runtime/Engine/Private/Actor.cpp b/Engine/Source/Runtime/Engine/Private/Actor.cpp
--- a/Engine/Source/Runtime/Engine/Private/Actor.cpp
+++ b/Engine/Source/Runtime/Engine/Private/Actor.cpp
@@ -3405,44 +3405,7 @@ void AActor::DispatchBeginPlay(bool bFromLevelStreaming)
 		if (!IsPendingKill())
 		{
 			// Initialize overlap state
-			if (!bFromLevelStreaming)
-			{
-				UpdateOverlaps();
-			}
-			else
-			{
-				// Note: Conditionally doing notifies here since loading or streaming in isn't actually conceptually beginning a touch.
-				//	     Rather, it was always touching and the mechanics of loading is just an implementation detail.
-				if (bGenerateOverlapEventsDuringLevelStreaming)
-				{
-					UpdateOverlaps(bGenerateOverlapEventsDuringLevelStreaming);
-				}
-				else
-				{
-					bool bUpdateOverlaps = true;
-					const EActorUpdateOverlapsMethod UpdateMethod = GetUpdateOverlapsMethodDuringLevelStreaming();
-					switch (UpdateMethod)
-					{
-					case EActorUpdateOverlapsMethod::OnlyUpdateMovable:
-						bUpdateOverlaps = IsRootComponentMovable();
-						break;
-
-					case EActorUpdateOverlapsMethod::NeverUpdate:
-						bUpdateOverlaps = false;
-						break;
-
-					case EActorUpdateOverlapsMethod::AlwaysUpdate:
-					default:
-						bUpdateOverlaps = true;
-						break;
-					}
-
-					if (bUpdateOverlaps)
-					{
-						UpdateOverlaps(bGenerateOverlapEventsDuringLevelStreaming);
-					}
-				}
-			}
+			UpdateInitialOverlaps(bFromLevelStreaming);
 		}
 
 		bActorBeginningPlayFromLevelStreaming = false;
@@ -3492,6 +3455,48 @@ void AActor::BeginPlay()
 	ActorHasBegunPlay = EActorBeginPlayState::HasBegunPlay;
 }
 
+void AActor::UpdateInitialOverlaps(bool bFromLevelStreaming)
+{
+	if (!bFromLevelStreaming)
+	{
+		UpdateOverlaps();
+	}
+	else
+	{
+		// Note: Conditionally doing notifies here since loading or streaming in isn't actually conceptually beginning a touch.
+		//	     Rather, it was always touching and the mechanics of loading is just an implementation detail.
+		if (bGenerateOverlapEventsDuringLevelStreaming)
+		{
+			UpdateOverlaps(bGenerateOverlapEventsDuringLevelStreaming);
+		}
+		else
+		{
+			bool bUpdateOverlaps = true;
+			const EActorUpdateOverlapsMethod UpdateMethod = GetUpdateOverlapsMethodDuringLevelStreaming();
+			switch (UpdateMethod)
+			{
+				case EActorUpdateOverlapsMethod::OnlyUpdateMovable:
+					bUpdateOverlaps = IsRootComponentMovable();
+					break;
+
+				case EActorUpdateOverlapsMethod::NeverUpdate:
+					bUpdateOverlaps = false;
+					break;
+
+				case EActorUpdateOverlapsMethod::AlwaysUpdate:
+				default:
+					bUpdateOverlaps = true;
+					break;
+			}
+
+			if (bUpdateOverlaps)
+			{
+				UpdateOverlaps(bGenerateOverlapEventsDuringLevelStreaming);
+			}
+		}
+	}
+}
+
 void AActor::EnableInput(APlayerController* PlayerController)
 {
 	if (PlayerController)
