Added OnMeshChanged broadcast to UStaticMesh::PostEditChangeProperty()
This fixes a crash in Niagara when using the mesh renderer.

diff --git a/Engine/Source/Runtime/Engine/Private/StaticMesh.cpp b/Engine/Source/Runtime/Engine/Private/StaticMesh.cpp
--- a/Engine/Source/Runtime/Engine/Private/StaticMesh.cpp
+++ b/Engine/Source/Runtime/Engine/Private/StaticMesh.cpp
@@ -6054,6 +6054,7 @@ ENGINE_API void UStaticMesh::RemoveVertexColors()
 	{
 		Build();
 		MarkPackageDirty();
+		OnMeshChanged.Broadcast();
 	}
 #endif
 }
