Fixed crash when assigning a Source2DTexture to a volume texture.
Added auto tile computation when setting a Source2DTexture for the first time.

diff --git a/Engine/Source/Editor/UnrealEd/Private/Factories/VolumeTextureFactory.cpp b/Engine/Source/Editor/UnrealEd/Private/Factories/VolumeTextureFactory.cpp
--- a/Engine/Source/Editor/UnrealEd/Private/Factories/VolumeTextureFactory.cpp
+++ b/Engine/Source/Editor/UnrealEd/Private/Factories/VolumeTextureFactory.cpp
@@ -35,29 +35,10 @@ UObject* UVolumeTextureFactory::FactoryCreateNew(UClass* Class,UObject* InParent
 		NewVolumeTexture->MipGenSettings = TMGS_FromTextureGroup;
 		NewVolumeTexture->NeverStream = true;
 		NewVolumeTexture->CompressionNone = false;
-
-		const FTextureSource& Source = InitialTexture->Source;
-		const int32 NumPixels = Source.GetSizeX() * Source.GetSizeY();
-
-		if (NumPixels > 0)
-		{
-			// By default, assume that the 2d texture holds a square images.
-			int32 TileSize = FMath::RoundToInt(FMath::Pow((float)NumPixels, 1.f / 3.f));
-
-			// Remap the tiles as if there were streached accross the texture length.
-			int32 NumTilesBySide = FMath::RoundToInt(FMath::Sqrt((float)(Source.GetSizeX() / TileSize) * (Source.GetSizeY() / TileSize)));
-
-			if (NumTilesBySide > 0)
-			{
-				NewVolumeTexture->Source2DTexture = InitialTexture;
-				NewVolumeTexture->Source2DTileSizeX = Source.GetSizeX() / NumTilesBySide;
-				NewVolumeTexture->Source2DTileSizeY = Source.GetSizeY() / NumTilesBySide;
-
-				NewVolumeTexture->UpdateSourceFromSourceTexture();
-				NewVolumeTexture->UpdateResource();
-			}
-
-		}
+		NewVolumeTexture->Source2DTexture = InitialTexture;
+		NewVolumeTexture->SetDefaultSource2DTileSize();
+		NewVolumeTexture->UpdateSourceFromSourceTexture();
+		NewVolumeTexture->UpdateResource();
 	}
 
 	return NewVolumeTexture;
diff --git a/Engine/Source/Runtime/Engine/Classes/Engine/VolumeTexture.h b/Engine/Source/Runtime/Engine/Classes/Engine/VolumeTexture.h
--- a/Engine/Source/Runtime/Engine/Classes/Engine/VolumeTexture.h
+++ b/Engine/Source/Runtime/Engine/Classes/Engine/VolumeTexture.h
@@ -84,6 +84,7 @@ public:
 	virtual float GetSurfaceHeight() const override { return GetSizeY(); }
 	virtual FTextureResource* CreateResource() override;
 #if WITH_EDITOR
+	ENGINE_API void SetDefaultSource2DTileSize();
 	ENGINE_API virtual void PostEditChangeProperty(FPropertyChangedEvent& PropertyChangedEvent) override;
 #endif // WITH_EDITOR
 	virtual void UpdateResource() override;
diff --git a/Engine/Source/Runtime/Engine/Private/VolumeTexture.cpp b/Engine/Source/Runtime/Engine/Private/VolumeTexture.cpp
--- a/Engine/Source/Runtime/Engine/Private/VolumeTexture.cpp
+++ b/Engine/Source/Runtime/Engine/Private/VolumeTexture.cpp
@@ -232,15 +232,18 @@ uint32 UVolumeTexture::CalcTextureMemorySize(int32 MipCount) const
 	if (PlatformData)
 	{
 		const EPixelFormat Format = GetPixelFormat();
-		const uint32 Flags = (SRGB ? TexCreate_SRGB : 0)  | TexCreate_OfflineProcessed | (bNoTiling ? TexCreate_NoTiling : 0);
+		if (Format != PF_Unknown)
+		{
+			const uint32 Flags = (SRGB ? TexCreate_SRGB : 0)  | TexCreate_OfflineProcessed | (bNoTiling ? TexCreate_NoTiling : 0);
 
-		uint32 SizeX = 0;
-		uint32 SizeY = 0;
-		uint32 SizeZ = 0;
-		CalcMipMapExtent3D(GetSizeX(), GetSizeY(), GetSizeZ(), Format, FMath::Max<int32>(0, GetNumMips() - MipCount), SizeX, SizeY, SizeZ);
+			uint32 SizeX = 0;
+			uint32 SizeY = 0;
+			uint32 SizeZ = 0;
+			CalcMipMapExtent3D(GetSizeX(), GetSizeY(), GetSizeZ(), Format, FMath::Max<int32>(0, GetNumMips() - MipCount), SizeX, SizeY, SizeZ);
 
-		uint32 TextureAlign = 0;
-		Size = (uint32)RHICalcTexture3DPlatformSize(SizeX, SizeY, SizeZ, Format, FMath::Max(1, MipCount), Flags, FRHIResourceCreateInfo(PlatformData->GetExtData()), TextureAlign);
+			uint32 TextureAlign = 0;
+			Size = (uint32)RHICalcTexture3DPlatformSize(SizeX, SizeY, SizeZ, Format, FMath::Max(1, MipCount), Flags, FRHIResourceCreateInfo(PlatformData->GetExtData()), TextureAlign);
+		}
 	}
 	return Size;
 }
@@ -566,6 +569,26 @@ void UVolumeTexture::GetResourceSizeEx(FResourceSizeEx& CumulativeResourceSize)
 
 #if WITH_EDITOR
 
+void UVolumeTexture::SetDefaultSource2DTileSize()
+{
+	Source2DTileSizeX = 0;
+	Source2DTileSizeY = 0;
+
+	if (Source2DTexture)
+	{
+		const int32 SourceSizeX = Source2DTexture->Source.GetSizeX();
+		const int32 SourceSizeY = Source2DTexture->Source.GetSizeY();
+		if (SourceSizeX > 0 && SourceSizeY > 0)
+		{
+			const int32 NumPixels = SourceSizeX * SourceSizeY;
+			const int32 TileSize = FMath::RoundToInt(FMath::Pow((float)NumPixels, 1.f / 3.f));
+			const int32 NumTilesBySide = FMath::RoundToInt(FMath::Sqrt((float)(SourceSizeX / TileSize) * (SourceSizeY / TileSize)));
+			Source2DTileSizeX = SourceSizeX / NumTilesBySide;
+			Source2DTileSizeY = SourceSizeY / NumTilesBySide;
+		}
+	}
+}
+
 void UVolumeTexture::PostEditChangeProperty(FPropertyChangedEvent& PropertyChangedEvent)
 {
  	FProperty* PropertyThatChanged = PropertyChangedEvent.Property;
@@ -576,6 +599,13 @@ void UVolumeTexture::PostEditChangeProperty(FPropertyChangedEvent& PropertyChang
 		static const FName TileSizeYName("Source2DTileSizeY");
 
 		const FName PropertyName = PropertyChangedEvent.Property->GetFName();
+
+		// Set default tile size if none is currently specified.
+		if (PropertyName == SourceTextureName && !Source2DTileSizeX && !Source2DTileSizeY)
+		{
+			SetDefaultSource2DTileSize();
+		}
+		// Update the content of the volume texture
 		if (PropertyName == SourceTextureName || PropertyName == TileSizeXName || PropertyName == TileSizeYName)
 		{
 			UpdateSourceFromSourceTexture();
