Property name/comment clarification
#rnx

diff --git a/Engine/Source/Runtime/Engine/Classes/Sound/SoundAttenuation.h b/Engine/Source/Runtime/Engine/Classes/Sound/SoundAttenuation.h
--- a/Engine/Source/Runtime/Engine/Classes/Sound/SoundAttenuation.h
+++ b/Engine/Source/Runtime/Engine/Classes/Sound/SoundAttenuation.h
@@ -361,12 +361,12 @@ struct ENGINE_API FSoundAttenuationSettings : public FBaseAttenuationSettings
 	UPROPERTY(EditAnywhere, BlueprintReadWrite, Category = AttenuationSubmixSend)
 	TArray<FAttenuationSubmixSendSettings> SubmixSendSettings;
 
-	/** Minimum interpolated value to scale final playback priority against when the sound is at the maximum priority attenuation distance. */
-	UPROPERTY(EditAnywhere, BlueprintReadWrite, Category = AttenuationPriority, meta = (ClampMin = "0.0", UIMin = "0.0", UIMax = "1.0", DisplayName = "Min Priority Attenuation"))
+	/** Interpolated value to scale priority against when the sound is at the minimum priority attenuation distance from the closest listener. */
+	UPROPERTY(EditAnywhere, BlueprintReadWrite, Category = AttenuationPriority, meta = (ClampMin = "0.0", UIMin = "0.0", UIMax = "1.0", DisplayName = "Priority Attenuation At Min Distance"))
 	float PriorityAttenuationMin;
 
-	/** Maximum interpolated value to scale final playback priority against when the sound is at the maximum priority attenuation distance. */
-	UPROPERTY(EditAnywhere, BlueprintReadWrite, Category = AttenuationPriority, meta = (ClampMin = "0.0", UIMin = "0.0", UIMax = "1.0", DisplayName = "Max Priority Attenuation"))
+	/** Interpolated value to scale priority against when the sound is at the maximum priority attenuation distance from the closest listener. */
+	UPROPERTY(EditAnywhere, BlueprintReadWrite, Category = AttenuationPriority, meta = (ClampMin = "0.0", UIMin = "0.0", UIMax = "1.0", DisplayName = "Priority Attenuation At Max Distance"))
 	float PriorityAttenuationMax;
 
 	/** The min distance to attenuate priority. */
