Convert ChunkData to a raw pointer rather than a TArray. saves roughly 1.7 megs for a 32 meg cache with default settings, and scales linearly

diff --git a/Engine/Source/Runtime/Engine/Private/AudioStreamingCache.cpp b/Engine/Source/Runtime/Engine/Private/AudioStreamingCache.cpp
--- a/Engine/Source/Runtime/Engine/Private/AudioStreamingCache.cpp
+++ b/Engine/Source/Runtime/Engine/Private/AudioStreamingCache.cpp
@@ -639,7 +639,7 @@ TArrayView<uint8> FAudioChunkCache::GetChunk(const FChunkKey& InKey, bool bBlock
 		// If this value is ever negative, it means that we're decrementing more than we're incrementing:
 		check(FoundElement->NumConsumers.GetValue() >= 0);
 		FoundElement->NumConsumers.Increment();
-		return TArrayView<uint8>(FoundElement->ChunkData.GetData(), FoundElement->ChunkDataSize);
+		return TArrayView<uint8>(FoundElement->ChunkData, FoundElement->ChunkDataSize);
 	}
 	else
 	{
@@ -661,18 +661,29 @@ TArrayView<uint8> FAudioChunkCache::GetChunk(const FChunkKey& InKey, bool bBlock
 			FoundElement->DebugInfo.NumTotalChunks = InKey.SoundWave->GetNumChunks() - 1;
 			FoundElement->DebugInfo.TimeLoadStarted = FPlatformTime::Cycles64();
 #endif
-			MemoryCounterBytes -= FoundElement->ChunkData.Num();
+			MemoryCounterBytes -= FoundElement->ChunkDataSize;
 
-		
 			// Reallocate our chunk data This allows us to shrink if possible.
-			// Unfortunately, GetCopy will write out the full zero-padded length of the bulk data,
-			// rather than just the audio data. So we set the array to Chunk.DataSize, then shrink to Chunk.AudioDataSize.
-			FoundElement->ChunkData.SetNumUninitialized(Chunk.DataSize, true);
-			void* DataDestPtr = FoundElement->ChunkData.GetData();
-			Chunk.BulkData.GetCopy(&DataDestPtr, true);
+			FoundElement->ChunkData = (uint8*)FMemory::Realloc(FoundElement->ChunkData, ChunkAudioDataSize);
+
+			if (Chunk.DataSize != ChunkAudioDataSize)
+			{
+				// Unfortunately, GetCopy will write out the full zero-padded length of the bulk data,
+				// rather than just the audio data. So we set the array to Chunk.DataSize, then shrink to Chunk.AudioDataSize.
+				TArray<uint8> TempChunkBuffer;
+				TempChunkBuffer.AddUninitialized(Chunk.DataSize);
+				void* DataDestPtr = TempChunkBuffer.GetData();
+				Chunk.BulkData.GetCopy(&DataDestPtr, true);
+
+				FMemory::Memcpy(FoundElement->ChunkData, TempChunkBuffer.GetData(), ChunkAudioDataSize);
+			}
+			else
+			{
+				void* DataDestPtr = FoundElement->ChunkData;
+				Chunk.BulkData.GetCopy(&DataDestPtr, true);
+			}
 
-			FoundElement->ChunkData.SetNumUninitialized(ChunkAudioDataSize);
-			MemoryCounterBytes += FoundElement->ChunkData.Num();
+			MemoryCounterBytes += ChunkAudioDataSize;
 
 			// Populate key and DataSize. The async read request was set up to write directly into CacheElement->ChunkData.
 			FoundElement->Key = InKey;
@@ -755,12 +766,14 @@ uint64 FAudioChunkCache::TrimMemory(uint64 BytesToFree)
 	{
 		if (CurrentElement->CanEvictChunk())
 		{
-			BytesFreed += CurrentElement->ChunkData.Num();
-			MemoryCounterBytes -= CurrentElement->ChunkData.Num();
+			BytesFreed += CurrentElement->ChunkDataSize;
+			MemoryCounterBytes -= CurrentElement->ChunkDataSize;
 			// Empty the chunk data and invalidate the key.
+			if(CurrentElement->ChunkData)
 			{
 				LLM_SCOPE(ELLMTag::AudioStreamCacheCompressedData);
-				CurrentElement->ChunkData.Empty();
+				FMemory::Free(CurrentElement->ChunkData);
+				CurrentElement->ChunkData = nullptr;
 			}
 
 			CurrentElement->ChunkDataSize = 0;
@@ -1172,16 +1185,17 @@ void FAudioChunkCache::KickOffAsyncLoad(FCacheElement* CacheElement, const FChun
 
 	EAsyncIOPriorityAndFlags AsyncIOPriority = GetAsyncPriorityForChunk(InKey, bNeededForPlayback);
 
-	MemoryCounterBytes -= CacheElement->ChunkData.Num();
+	MemoryCounterBytes -= CacheElement->ChunkDataSize;
 
 	{
 		LLM_SCOPE(ELLMTag::AudioStreamCacheCompressedData);
 
 		// Reallocate our chunk data This allows us to shrink if possible.
-		CacheElement->ChunkData.SetNumUninitialized(Chunk.AudioDataSize, true);
+		CacheElement->ChunkData = (uint8*) FMemory::Realloc(CacheElement->ChunkData, Chunk.AudioDataSize);
+		CacheElement->ChunkDataSize = Chunk.AudioDataSize;
 	}
 
-	MemoryCounterBytes += CacheElement->ChunkData.Num();
+	MemoryCounterBytes += CacheElement->ChunkDataSize;
 
 #if DEBUG_STREAM_CACHE
 	CacheElement->DebugInfo.NumTotalChunks = InKey.SoundWave->GetNumChunks() - 1;
@@ -1225,7 +1239,7 @@ void FAudioChunkCache::KickOffAsyncLoad(FCacheElement* CacheElement, const FChun
 
 		CacheElement->DDCTask.Reset(new FAsyncStreamDerivedChunkTask(
 			Chunk.DerivedDataKey,
-			CacheElement->ChunkData.GetData(),
+			CacheElement->ChunkData,
 			ChunkDataSize,
 			&NumberOfLoadsInFlight,
 			MoveTemp(OnLoadComplete)
@@ -1241,10 +1255,10 @@ void FAudioChunkCache::KickOffAsyncLoad(FCacheElement* CacheElement, const FChun
 			CacheElement->WaitForAsyncLoadCompletion(true);
 		}
 
-		// Sanity check our bulk data against
+		// Sanity check our bulk data against our currently allocated chunk size in the cache.
 		const int32 ChunkBulkDataSize = Chunk.BulkData.GetBulkDataSize();
 		check(ChunkDataSize <= ChunkBulkDataSize);
-		check(((int32)ChunkDataSize) <= CacheElement->ChunkData.Num());
+		check(((uint32)ChunkDataSize) <= CacheElement->ChunkDataSize);
 
 		// If we ever want to eliminate zero-padding in chunks, that could be verified here:
 		//ensureAlwaysMsgf(AudioChunkSize == ChunkBulkDataSize, TEXT("For memory load on demand, we do not zero-pad to page sizes."));
@@ -1271,7 +1285,7 @@ void FAudioChunkCache::KickOffAsyncLoad(FCacheElement* CacheElement, const FChun
 		CacheElement->DebugInfo.TimeLoadStarted = FPlatformTime::Seconds();
 #endif
 
-		CacheElement->ReadRequest.Reset(Chunk.BulkData.CreateStreamingRequest(0, ChunkDataSize, AsyncIOPriority, &AsyncFileCallBack, CacheElement->ChunkData.GetData()));
+		CacheElement->ReadRequest.Reset(Chunk.BulkData.CreateStreamingRequest(0, ChunkDataSize, AsyncIOPriority, &AsyncFileCallBack, CacheElement->ChunkData));
 		if (!CacheElement->ReadRequest.IsValid())
 		{
 			UE_LOG(LogAudio, Error, TEXT("Chunk load in audio LRU cache failed."));
@@ -1495,7 +1509,7 @@ TPair<int, int> FAudioChunkCache::DebugDisplay(UWorld* World, FViewport* Viewpor
 	while (CurrentElement != nullptr)
 	{
 		// Note: this is potentially a stale value if we're in the middle of FCacheElement::KickOffAsyncLoad.
-		NumBytesCounter += CurrentElement->ChunkData.Num();
+		NumBytesCounter += CurrentElement->ChunkDataSize;
 		CurrentElement = CurrentElement->LessRecentElement;
 	}
 
@@ -1547,11 +1561,11 @@ TPair<int, int> FAudioChunkCache::DebugDisplay(UWorld* World, FViewport* Viewpor
 		bIsStaleChunk = (CurrentElement->Key.SoundWave == nullptr) || (CurrentElement->Key.SoundWave->CurrentChunkRevision.GetValue() != CurrentElement->Key.ChunkRevision);
 #endif
 
-		const bool bWasTrimmed = CurrentElement->ChunkData.Num() == 0;
+		const bool bWasTrimmed = CurrentElement->ChunkDataSize == 0;
 
 		FString ElementInfo = *FString::Printf(TEXT("%4i. Size: %6.2f KB   Chunk: %d of %d   Request Count: %d    Average Index: %6.2f  Number of Handles Retaining Chunk: %d     Chunk Load Time(in ms): %6.4fms      Loading Behavior: %s      Name: %s Notes: %s %s"),
 			Index,
-			CurrentElement->ChunkData.Num() / 1024.0f,
+			CurrentElement->ChunkDataSize / 1024.0f,
 			CurrentElement->Key.ChunkIndex,
 			NumTotalChunks,
 			NumTimesTouched,
@@ -1629,11 +1643,11 @@ FString FAudioChunkCache::DebugPrint()
 	while (CurrentElement != nullptr)
 	{
 		// Note: this is potentially a stale value if we're in the middle of FCacheElement::KickOffAsyncLoad.
-		NumBytesCounter += CurrentElement->ChunkData.Num();
+		NumBytesCounter += CurrentElement->ChunkDataSize;
 
 		if (CurrentElement->IsInUse())
 		{
-			NumBytesRetained += CurrentElement->ChunkData.Num();
+			NumBytesRetained += CurrentElement->ChunkDataSize;
 		}
 
 		CurrentElement = CurrentElement->LessRecentElement;
@@ -1682,11 +1696,11 @@ FString FAudioChunkCache::DebugPrint()
 		bIsStaleChunk = (CurrentElement->Key.SoundWave == nullptr) || (CurrentElement->Key.SoundWave->CurrentChunkRevision.GetValue() != CurrentElement->Key.ChunkRevision);
 #endif
 
-		const bool bWasTrimmed = CurrentElement->ChunkData.Num() == 0;
+		const bool bWasTrimmed = CurrentElement->ChunkDataSize == 0;
 
 		FString ElementInfo = *FString::Printf(TEXT("%4i.\t, %6.2f\t, %d of %d\t, %d\t, %6.2f\t, %d\t,  %6.4f\t, %s\t, %s, %s %s %s"),
 			Index,
-			CurrentElement->ChunkData.Num() / 1024.0f,
+			CurrentElement->ChunkDataSize / 1024.0f,
 			CurrentElement->Key.ChunkIndex,
 			NumTotalChunks,
 			NumTimesTouched,
diff --git a/Engine/Source/Runtime/Engine/Public/AudioStreamingCache.h b/Engine/Source/Runtime/Engine/Public/AudioStreamingCache.h
--- a/Engine/Source/Runtime/Engine/Public/AudioStreamingCache.h
+++ b/Engine/Source/Runtime/Engine/Public/AudioStreamingCache.h
@@ -159,7 +159,7 @@ private:
 	struct FCacheElement
 	{
 		FChunkKey Key;
-		TArray<uint8> ChunkData;
+		uint8* ChunkData;
 		uint32 ChunkDataSize;
 		FCacheElement* MoreRecentElement;
 		FCacheElement* LessRecentElement;
@@ -182,7 +182,8 @@ private:
 #endif
 
 		FCacheElement(uint32 MaxChunkSize, uint32 InCacheIndex)
-			: ChunkDataSize(0)
+			: ChunkData(nullptr)
+			, ChunkDataSize(0)
 			, MoreRecentElement(nullptr)
 			, LessRecentElement(nullptr)
 			, CacheLookupID(InCacheIndex)
@@ -235,6 +236,12 @@ private:
 		{
 			WaitForAsyncLoadCompletion(true);
 			checkf(NumConsumers.GetValue() == 0, TEXT("Tried to destroy streaming cache while the cached data was in use!"));
+			if (ChunkData)
+			{
+				FMemory::Free(ChunkData);
+			}
+
+			ChunkData = nullptr;
 		}
 	};
 
