Converted a few usages of render fences in material code to use an atomic bool instead. Purpose is to remove dependency on task graph to allow batching of material commands in a future change.

#rnx

diff --git a/Engine/Source/Runtime/Engine/Classes/Materials/Material.h b/Engine/Source/Runtime/Engine/Classes/Materials/Material.h
--- a/Engine/Source/Runtime/Engine/Classes/Materials/Material.h
+++ b/Engine/Source/Runtime/Engine/Classes/Materials/Material.h
@@ -6,6 +6,7 @@
 #include "UObject/ObjectMacros.h"
 #include "Misc/Guid.h"
 #include "Engine/EngineTypes.h"
+#include "HAL/ThreadSafeBool.h"
 #include "RenderCommandFence.h"
 #include "Materials/MaterialInterface.h"
 #include "MaterialShared.h"
@@ -947,8 +948,8 @@ private:
 	/** Material resources being cached for cooking. */
 	TMap<const class ITargetPlatform*, TArray<FMaterialResource*>> CachedMaterialResourcesForCooking;
 #endif
-	/** Fence used to guarantee that the RT is finished using various resources in this UMaterial before cleanup. */
-	FRenderCommandFence ReleaseFence;
+	/** Flag used to guarantee that the RT is finished using various resources in this UMaterial before cleanup. */
+	FThreadSafeBool ReleasedByRT;
 
 	UPROPERTY()
 	FMaterialCachedExpressionData CachedExpressionData;
diff --git a/Engine/Source/Runtime/Engine/Classes/Materials/MaterialInstance.h b/Engine/Source/Runtime/Engine/Classes/Materials/MaterialInstance.h
--- a/Engine/Source/Runtime/Engine/Classes/Materials/MaterialInstance.h
+++ b/Engine/Source/Runtime/Engine/Classes/Materials/MaterialInstance.h
@@ -6,6 +6,7 @@
 #include "UObject/ObjectMacros.h"
 #include "Misc/Guid.h"
 #include "RenderCommandFence.h"
+#include "HAL/ThreadSafeBool.h"
 #include "Materials/MaterialInterface.h"
 #include "StaticParameterSet.h"
 #include "MaterialShared.h"
@@ -470,8 +471,8 @@ private:
 	/** Material resources being cached for cooking. */
 	TMap<const class ITargetPlatform*, TArray<FMaterialResource*>> CachedMaterialResourcesForCooking;
 #endif
-	/** Fence used to guarantee that the RT is finished using various resources in this UMaterial before cleanup. */
-	FRenderCommandFence ReleaseFence;
+	/** Flag used to guarantee that the RT is finished using various resources in this UMaterial before cleanup. */
+	FThreadSafeBool ReleasedByRT;
 
 public:
 	// Begin UMaterialInterface interface.
diff --git a/Engine/Source/Runtime/Engine/Classes/Materials/MaterialParameterCollection.h b/Engine/Source/Runtime/Engine/Classes/Materials/MaterialParameterCollection.h
--- a/Engine/Source/Runtime/Engine/Classes/Materials/MaterialParameterCollection.h
+++ b/Engine/Source/Runtime/Engine/Classes/Materials/MaterialParameterCollection.h
@@ -9,6 +9,7 @@
 #include "CoreMinimal.h"
 #include "UObject/ObjectMacros.h"
 #include "UObject/Object.h"
+#include "HAL/ThreadSafeBool.h"
 #include "Misc/Guid.h"
 #include "UniformBuffer.h"
 #include "Templates/UniquePtr.h"
@@ -129,8 +130,8 @@ private:
 	virtual ENGINE_API void FinishDestroy() override;
 	virtual ENGINE_API bool IsReadyForFinishDestroy() override;
 
-	/** Fence used to guarantee that the RT is finished using various resources in this UMaterial before cleanup. */
-	FRenderCommandFence ReleaseFence;
+	/** Flag used to guarantee that the RT is finished using various resources in this UMaterial before cleanup. */
+	FThreadSafeBool ReleasedByRT;
 
 	/** Default resource used when no instance is available. */
 	class FMaterialParameterCollectionInstanceResource* DefaultResource;
diff --git a/Engine/Source/Runtime/Engine/Private/Materials/Material.cpp b/Engine/Source/Runtime/Engine/Private/Materials/Material.cpp
--- a/Engine/Source/Runtime/Engine/Private/Materials/Material.cpp
+++ b/Engine/Source/Runtime/Engine/Private/Materials/Material.cpp
@@ -887,6 +887,7 @@ bool UMaterial::IsDefaultMaterial() const
 
 UMaterial::UMaterial(const FObjectInitializer& ObjectInitializer)
 	: Super(ObjectInitializer)
+	, ReleasedByRT(true)
 {
 	BlendMode = BLEND_Opaque;
 	ShadingModel = MSM_DefaultLit;
@@ -4413,23 +4414,26 @@ void UMaterial::BeginDestroy()
 
 	if (DefaultMaterialInstance)
 	{
+		ReleasedByRT = false;
+
 		FMaterialRenderProxy* LocalResource = DefaultMaterialInstance;
+		FThreadSafeBool* Released = &ReleasedByRT;
 		ENQUEUE_RENDER_COMMAND(BeginDestroyCommand)(
-		[LocalResource](FRHICommandList& RHICmdList)
+		[LocalResource, Released](FRHICommandListImmediate& RHICmdList)
 		{
 			LocalResource->MarkForGarbageCollection();
 			LocalResource->ReleaseResource();
+
+			*Released = true;
 		});
 	}
-
-	ReleaseFence.BeginFence();
 }
 
 bool UMaterial::IsReadyForFinishDestroy()
 {
 	bool bReady = Super::IsReadyForFinishDestroy();
 
-	return bReady && ReleaseFence.IsFenceComplete();
+	return bReady && ReleasedByRT;
 }
 
 void UMaterial::ReleaseResources()
diff --git a/Engine/Source/Runtime/Engine/Private/Materials/MaterialInstance.cpp b/Engine/Source/Runtime/Engine/Private/Materials/MaterialInstance.cpp
--- a/Engine/Source/Runtime/Engine/Private/Materials/MaterialInstance.cpp
+++ b/Engine/Source/Runtime/Engine/Private/Materials/MaterialInstance.cpp
@@ -541,6 +541,7 @@ bool UMaterialInstance::UpdateParameters()
 
 UMaterialInstance::UMaterialInstance(const FObjectInitializer& ObjectInitializer)
 	: Super(ObjectInitializer)
+	, ReleasedByRT(true)
 {
 	bHasStaticPermutationResource = false;
 #if WITH_EDITOR
@@ -3598,23 +3599,26 @@ void UMaterialInstance::BeginDestroy()
 
 	if (!HasAnyFlags(RF_ClassDefaultObject))
 	{
+		ReleasedByRT = false;
+
 		FMaterialRenderProxy* LocalResource = Resource;
+		FThreadSafeBool* Released = &ReleasedByRT;
 		ENQUEUE_RENDER_COMMAND(BeginDestroyCommand)(
-		[LocalResource](FRHICommandList& RHICmdList)
+		[LocalResource, Released](FRHICommandListImmediate& RHICmdList)
 		{
 			LocalResource->MarkForGarbageCollection();
 			LocalResource->ReleaseResource();
+
+			*Released = true;
 		});		
 	}
-
-	ReleaseFence.BeginFence();
 }
 
 bool UMaterialInstance::IsReadyForFinishDestroy()
 {
 	bool bIsReady = Super::IsReadyForFinishDestroy();
 
-	return bIsReady && ReleaseFence.IsFenceComplete();;
+	return bIsReady && ReleasedByRT;
 }
 
 void UMaterialInstance::FinishDestroy()
diff --git a/Engine/Source/Runtime/Engine/Private/Materials/ParameterCollection.cpp b/Engine/Source/Runtime/Engine/Private/Materials/ParameterCollection.cpp
--- a/Engine/Source/Runtime/Engine/Private/Materials/ParameterCollection.cpp
+++ b/Engine/Source/Runtime/Engine/Private/Materials/ParameterCollection.cpp
@@ -24,6 +24,7 @@ TMultiMap<FGuid, FMaterialParameterCollectionInstanceResource*> GDefaultMaterial
 
 UMaterialParameterCollection::UMaterialParameterCollection(const FObjectInitializer& ObjectInitializer)
 	: Super(ObjectInitializer)
+	, ReleasedByRT(true)
 {
 	DefaultResource = nullptr;
 }
@@ -63,24 +64,27 @@ void UMaterialParameterCollection::BeginDestroy()
 {
 	if (DefaultResource)
 	{
+		ReleasedByRT = false;
+
 		FMaterialParameterCollectionInstanceResource* Resource = DefaultResource;
 		FGuid Id = StateId;
+		FThreadSafeBool* Released = &ReleasedByRT;
 		ENQUEUE_RENDER_COMMAND(RemoveDefaultResourceCommand)(
-			[Resource, Id](FRHICommandListImmediate& RHICmdList)
+			[Resource, Id, Released](FRHICommandListImmediate& RHICmdList)
 			{	
 				GDefaultMaterialParameterCollectionInstances.RemoveSingle(Id, Resource);
+				*Released = true;
 			}
 		);
 	}
 
-	ReleaseFence.BeginFence();
 	Super::BeginDestroy();
 }
 
 bool UMaterialParameterCollection::IsReadyForFinishDestroy()
 {
 	bool bIsReady = Super::IsReadyForFinishDestroy();
-	return bIsReady && ReleaseFence.IsFenceComplete();
+	return bIsReady && ReleasedByRT;
 }
 
 void UMaterialParameterCollection::FinishDestroy()
