Clear UWorld's PersistentLevel from its CachedLevelCollection when initializing said PersistentLevel's UWorld - prevents inconsistent results from GetActorsOfClass when cooking

diff --git a/Engine/Source/Runtime/Engine/Private/World.cpp b/Engine/Source/Runtime/Engine/Private/World.cpp
--- a/Engine/Source/Runtime/Engine/Private/World.cpp
+++ b/Engine/Source/Runtime/Engine/Private/World.cpp
@@ -1419,6 +1419,10 @@ void UWorld::InitWorld(const InitializationValues IVS)
 
 	Levels.Empty(1);
 	Levels.Add( PersistentLevel );
+	if (FLevelCollection* Collection = PersistentLevel->GetCachedLevelCollection())
+	{
+		Collection->RemoveLevel(PersistentLevel);
+	}
 	PersistentLevel->OwningWorld = this;
 	PersistentLevel->bIsVisible = true;
 
