Restoring change to handle the case where ChangeBundleStateForPrimaryAsset gets a new bundle state that's a superset of the current one.  This will create a combined handle rather than cancel the pending stream.

diff --git a/Engine/Source/Runtime/Engine/Private/AssetManager.cpp b/Engine/Source/Runtime/Engine/Private/AssetManager.cpp
--- a/Engine/Source/Runtime/Engine/Private/AssetManager.cpp
+++ b/Engine/Source/Runtime/Engine/Private/AssetManager.cpp
@@ -1205,8 +1205,30 @@ TSharedPtr<FStreamableHandle> UAssetManager::ChangeBundleStateForPrimaryAssets(c
 					ExistingHandles.Add(NameData->PendingState.Handle);
 					continue;
 				}
-				// Clear pending state
-				NameData->PendingState.Reset(true);
+
+				bool bIsSuperSet = false;
+				if (NewBundles->Num() > NameData->PendingState.BundleNames.Num())
+				{
+					bIsSuperSet = true;
+					for (int i = 0; i < NameData->PendingState.BundleNames.Num(); i++)
+					{
+						if (!NewBundles->Contains(NameData->PendingState.BundleNames[i]))
+						{
+							bIsSuperSet = false;
+							break;
+						}
+					}
+				}
+
+				if (bIsSuperSet)
+				{
+					ExistingHandles.Add(NameData->PendingState.Handle);
+				}
+				else
+				{
+					// Clear pending state
+					NameData->PendingState.Reset(true);
+				}
 			}
 			else if (NameData->CurrentState.IsValid() && NameData->CurrentState.BundleNames == *NewBundles)
 			{
