Fixed crash when moving/renaming anim blueprints

The FProperty refactor converted all UStructProperty* values into TFieldPaths, but the data is not serialized except in the case of UAnimClassData (for nativizaton). This reverts the various property arrays back to using FStructProperty*, which not only fixes some performance issues, but also prevents relocated data from pointing at incorrect locations when loaded assets are reloacated byt Link() is not subsequently called.

Fixed UAnimClassData (used in nativized builds) to correctly store and patch up anim blueprint functions.

diff --git a/Engine/Source/Developer/BlueprintCompilerCppBackend/Private/BlueprintCompilerCppBackendAnim.cpp b/Engine/Source/Developer/BlueprintCompilerCppBackend/Private/BlueprintCompilerCppBackendAnim.cpp
--- a/Engine/Source/Developer/BlueprintCompilerCppBackend/Private/BlueprintCompilerCppBackendAnim.cpp
+++ b/Engine/Source/Developer/BlueprintCompilerCppBackend/Private/BlueprintCompilerCppBackendAnim.cpp
@@ -37,6 +37,7 @@ void FBackendHelperAnim::CreateAnimClassData(FEmitterLocalContext& Context)
 				, FEmitDefaultValueHelper::EPropertyGenerationControlFlags::AllowTransient);
 		}
 
+		Context.AddLine(FString::Printf(TEXT("%s->ResolvePropertyPaths();"), *LocalNativeName));
 		Context.AddLine(FString::Printf(TEXT("InDynamicClass->%s = %s;"), GET_MEMBER_NAME_STRING_CHECKED(UDynamicClass, AnimClassImplementation), *LocalNativeName));
 	}
 }
diff --git a/Engine/Source/Editor/AnimationBlueprintEditor/Private/AnimGraphConnectionDrawingPolicy.cpp b/Engine/Source/Editor/AnimationBlueprintEditor/Private/AnimGraphConnectionDrawingPolicy.cpp
--- a/Engine/Source/Editor/AnimationBlueprintEditor/Private/AnimGraphConnectionDrawingPolicy.cpp
+++ b/Engine/Source/Editor/AnimationBlueprintEditor/Private/AnimGraphConnectionDrawingPolicy.cpp
@@ -43,9 +43,9 @@ void FAnimGraphConnectionDrawingPolicy::BuildExecutionRoadmap()
 
 		if ((VisitRecord.SourceID >= 0) && (VisitRecord.SourceID < AnimBlueprintClass->AnimNodeProperties.Num()) && (VisitRecord.TargetID >= 0) && (VisitRecord.TargetID < AnimBlueprintClass->AnimNodeProperties.Num()))
 		{
-			if (UAnimGraphNode_Base* SourceNode = Cast<UAnimGraphNode_Base>(PropertySourceMap.FindRef(AnimBlueprintClass->AnimNodeProperties[VisitRecord.SourceID].Get())))
+			if (UAnimGraphNode_Base* SourceNode = Cast<UAnimGraphNode_Base>(PropertySourceMap.FindRef(AnimBlueprintClass->AnimNodeProperties[VisitRecord.SourceID])))
 			{
-				if (UAnimGraphNode_Base* TargetNode = Cast<UAnimGraphNode_Base>(PropertySourceMap.FindRef(AnimBlueprintClass->AnimNodeProperties[VisitRecord.TargetID].Get())))
+				if (UAnimGraphNode_Base* TargetNode = Cast<UAnimGraphNode_Base>(PropertySourceMap.FindRef(AnimBlueprintClass->AnimNodeProperties[VisitRecord.TargetID])))
 				{
 					UEdGraphPin* PoseNet = NULL;
 
diff --git a/Engine/Source/Editor/KismetCompiler/Private/AnimBlueprintCompiler.cpp b/Engine/Source/Editor/KismetCompiler/Private/AnimBlueprintCompiler.cpp
--- a/Engine/Source/Editor/KismetCompiler/Private/AnimBlueprintCompiler.cpp
+++ b/Engine/Source/Editor/KismetCompiler/Private/AnimBlueprintCompiler.cpp
@@ -2294,6 +2294,10 @@ void FAnimBlueprintCompilerContext::CleanAndSanitizeClass(UBlueprintGeneratedCla
 	NewAnimBlueprintClass->AnimNodeProperties.Empty();
 	NewAnimBlueprintClass->LinkedAnimGraphNodeProperties.Empty();
 	NewAnimBlueprintClass->LinkedAnimLayerNodeProperties.Empty();
+	NewAnimBlueprintClass->PreUpdateNodeProperties.Empty();
+	NewAnimBlueprintClass->DynamicResetNodeProperties.Empty();
+	NewAnimBlueprintClass->StateMachineNodeProperties.Empty();
+	NewAnimBlueprintClass->InitializationNodeProperties.Empty();
 	NewAnimBlueprintClass->EvaluateGraphExposedInputs.Empty();
 	NewAnimBlueprintClass->GraphAssetPlayerInformation.Empty();
 	NewAnimBlueprintClass->GraphBlendOptions.Empty();
diff --git a/Engine/Source/Runtime/Engine/Classes/Animation/AnimBlueprintGeneratedClass.h b/Engine/Source/Runtime/Engine/Classes/Animation/AnimBlueprintGeneratedClass.h
--- a/Engine/Source/Runtime/Engine/Classes/Animation/AnimBlueprintGeneratedClass.h
+++ b/Engine/Source/Runtime/Engine/Classes/Animation/AnimBlueprintGeneratedClass.h
@@ -303,13 +303,13 @@ class ENGINE_API UAnimBlueprintGeneratedClass : public UBlueprintGeneratedClass,
 	TArray<FAnimBlueprintFunction> AnimBlueprintFunctions;
 
 	// The arrays of anim nodes; this is transient generated data (created during Link)
-	TArray<FStructPropertyPath> AnimNodeProperties;
-	TArray<FStructPropertyPath> LinkedAnimGraphNodeProperties;
-	TArray<FStructPropertyPath> LinkedAnimLayerNodeProperties;
-	TArray<FStructPropertyPath> PreUpdateNodeProperties;
-	TArray<FStructPropertyPath> DynamicResetNodeProperties;
-	TArray<FStructPropertyPath> StateMachineNodeProperties;
-	TArray<FStructPropertyPath> InitializationNodeProperties;
+	TArray<FStructProperty*> AnimNodeProperties;
+	TArray<FStructProperty*> LinkedAnimGraphNodeProperties;
+	TArray<FStructProperty*> LinkedAnimLayerNodeProperties;
+	TArray<FStructProperty*> PreUpdateNodeProperties;
+	TArray<FStructProperty*> DynamicResetNodeProperties;
+	TArray<FStructProperty*> StateMachineNodeProperties;
+	TArray<FStructProperty*> InitializationNodeProperties;
 
 	// Array of sync group names in the order that they are requested during compile
 	UPROPERTY()
@@ -332,13 +332,13 @@ public:
 	virtual const TArray<FBakedAnimationStateMachine>& GetBakedStateMachines() const override { return BakedStateMachines; }
 	virtual USkeleton* GetTargetSkeleton() const override { return TargetSkeleton; }
 	virtual const TArray<FAnimNotifyEvent>& GetAnimNotifies() const override { return AnimNotifies; }
-	virtual const TArray<FStructPropertyPath>& GetAnimNodeProperties() const override { return AnimNodeProperties; }
-	virtual const TArray<FStructPropertyPath>& GetLinkedAnimGraphNodeProperties() const override { return LinkedAnimGraphNodeProperties; }
-	virtual const TArray<FStructPropertyPath>& GetLinkedAnimLayerNodeProperties() const override { return LinkedAnimLayerNodeProperties; }
-	virtual const TArray<FStructPropertyPath>& GetPreUpdateNodeProperties() const override { return PreUpdateNodeProperties; }
-	virtual const TArray<FStructPropertyPath>& GetDynamicResetNodeProperties() const override { return DynamicResetNodeProperties; }
-	virtual const TArray<FStructPropertyPath>& GetStateMachineNodeProperties() const override { return StateMachineNodeProperties; }
-	virtual const TArray<FStructPropertyPath>& GetInitializationNodeProperties() const override { return InitializationNodeProperties; }
+	virtual const TArray<FStructProperty*>& GetAnimNodeProperties() const override { return AnimNodeProperties; }
+	virtual const TArray<FStructProperty*>& GetLinkedAnimGraphNodeProperties() const override { return LinkedAnimGraphNodeProperties; }
+	virtual const TArray<FStructProperty*>& GetLinkedAnimLayerNodeProperties() const override { return LinkedAnimLayerNodeProperties; }
+	virtual const TArray<FStructProperty*>& GetPreUpdateNodeProperties() const override { return PreUpdateNodeProperties; }
+	virtual const TArray<FStructProperty*>& GetDynamicResetNodeProperties() const override { return DynamicResetNodeProperties; }
+	virtual const TArray<FStructProperty*>& GetStateMachineNodeProperties() const override { return StateMachineNodeProperties; }
+	virtual const TArray<FStructProperty*>& GetInitializationNodeProperties() const override { return InitializationNodeProperties; }
 	virtual const TArray<FName>& GetSyncGroupNames() const override { return SyncGroupNames; }
 	virtual const TMap<FName, FCachedPoseIndices>& GetOrderedSavedPoseNodeIndicesMap() const override { return OrderedSavedPoseIndicesMap; }
 	virtual int32 GetSyncGroupIndex(FName SyncGroupName) const override { return SyncGroupNames.IndexOfByKey(SyncGroupName); }
@@ -400,7 +400,7 @@ public:
 		const int32* pIndex = GetNodePropertyIndex<StructType>(Node, SearchMode);
 		if (pIndex)
 		{
-			if (FStructProperty* AnimationProperty = AnimNodeProperties[AnimNodeProperties.Num() - 1 - *pIndex].Get())
+			if (FStructProperty* AnimationProperty = AnimNodeProperties[AnimNodeProperties.Num() - 1 - *pIndex])
 			{
 				if (AnimationProperty->Struct->IsChildOf(StructType::StaticStruct()))
 				{
@@ -430,7 +430,7 @@ public:
 		const int32* pIndex = GetNodePropertyIndexFromGuid(NodeGuid, SearchMode);
 		if (pIndex)
 		{
-			if (FStructProperty* AnimProperty = AnimNodeProperties[AnimNodeProperties.Num() - 1 - *pIndex].Get())
+			if (FStructProperty* AnimProperty = AnimNodeProperties[AnimNodeProperties.Num() - 1 - *pIndex])
 			{
 				if (AnimProperty->Struct->IsChildOf(StructType::StaticStruct()))
 				{
@@ -446,7 +446,7 @@ public:
 	StructType& GetPropertyInstanceChecked(UObject* Object, UAnimGraphNode_Base* Node, EPropertySearchMode::Type SearchMode = EPropertySearchMode::OnlyThis)
 	{
 		const int32 Index = AnimBlueprintDebugData.NodePropertyToIndexMap.FindChecked(Node);
-		FStructProperty* AnimationProperty = AnimNodeProperties[AnimNodeProperties.Num() - 1 - Index].Get();
+		FStructProperty* AnimationProperty = AnimNodeProperties[AnimNodeProperties.Num() - 1 - Index];
 		check(AnimationProperty);
 		check(AnimationProperty->Struct->IsChildOf(StructType::StaticStruct()));
 		return AnimationProperty->ContainerPtrToValuePtr<StructType>((void*)Object);
@@ -484,7 +484,7 @@ NodeType* GetNodeFromPropertyIndex(UObject* AnimInstanceObject, const IAnimClass
 {
 	if (PropertyIndex != INDEX_NONE)
 	{
-		FStructProperty* NodeProperty = AnimBlueprintClass->GetAnimNodeProperties()[AnimBlueprintClass->GetAnimNodeProperties().Num() - 1 - PropertyIndex].Get(); //@TODO: Crazysauce
+		FStructProperty* NodeProperty = AnimBlueprintClass->GetAnimNodeProperties()[AnimBlueprintClass->GetAnimNodeProperties().Num() - 1 - PropertyIndex]; //@TODO: Crazysauce
 		check(NodeProperty->Struct == NodeType::StaticStruct());
 		return NodeProperty->ContainerPtrToValuePtr<NodeType>(AnimInstanceObject);
 	}
diff --git a/Engine/Source/Runtime/Engine/Classes/Animation/AnimClassData.h b/Engine/Source/Runtime/Engine/Classes/Animation/AnimClassData.h
--- a/Engine/Source/Runtime/Engine/Classes/Animation/AnimClassData.h
+++ b/Engine/Source/Runtime/Engine/Classes/Animation/AnimClassData.h
@@ -9,10 +9,29 @@
 #include "Animation/AnimTypes.h"
 #include "Animation/AnimStateMachineTypes.h"
 #include "Animation/AnimClassInterface.h"
+#include "Algo/Transform.h"
 #include "AnimClassData.generated.h"
 
 class USkeleton;
 
+/** Serialized anim BP function data */
+USTRUCT()
+struct FAnimBlueprintFunctionData
+{
+	GENERATED_BODY()
+
+	UPROPERTY()
+	TFieldPath<FStructProperty> OutputPoseNodeProperty;
+
+	/** The properties of the input nodes, patched up during link */
+	UPROPERTY()
+	TArray<TFieldPath<FStructProperty>> InputPoseNodeProperties;
+
+	/** The input properties themselves */
+	UPROPERTY()
+	TArray<TFieldPath<FProperty>> InputProperties;
+};
+
 UCLASS()
 class ENGINE_API UAnimClassData : public UObject, public IAnimClassInterface
 {
@@ -38,33 +57,44 @@ public:
 	UPROPERTY()
 	TArray<FAnimBlueprintFunction> AnimBlueprintFunctions;
 
+	// Serialized function data, used to patch up transient data in AnimBlueprintFunctions
+	UPROPERTY()
+	TArray<FAnimBlueprintFunctionData> AnimBlueprintFunctionData;
+
 	// The array of anim nodes
 	UPROPERTY()
 	TArray< TFieldPath<FStructProperty> > AnimNodeProperties;
+	TArray< FStructProperty* > ResolvedAnimNodeProperties;
 
 	// The array of linked anim graph nodes
 	UPROPERTY()
 	TArray< TFieldPath<FStructProperty> > LinkedAnimGraphNodeProperties;
+	TArray< FStructProperty* > ResolvedLinkedAnimGraphNodeProperties;
 
 	// The array of linked anim layer nodes
 	UPROPERTY()
 	TArray< TFieldPath<FStructProperty> > LinkedAnimLayerNodeProperties;
+	TArray< FStructProperty* > ResolvedLinkedAnimLayerNodeProperties;
 
 	// Array of nodes that need a PreUpdate() call
 	UPROPERTY()
 	TArray< TFieldPath<FStructProperty> > PreUpdateNodeProperties;
+	TArray< FStructProperty* > ResolvedPreUpdateNodeProperties;
 
 	// Array of nodes that need a DynamicReset() call
 	UPROPERTY()
 	TArray< TFieldPath<FStructProperty> > DynamicResetNodeProperties;
+	TArray< FStructProperty* > ResolvedDynamicResetNodeProperties;
 
 	// Array of state machine nodes
 	UPROPERTY()
 	TArray< TFieldPath<FStructProperty> > StateMachineNodeProperties;
+	TArray< FStructProperty* > ResolvedStateMachineNodeProperties;
 
 	// Array of nodes that need an OnInitializeAnimInstance call
 	UPROPERTY()
 	TArray< TFieldPath<FStructProperty> > InitializationNodeProperties;
+	TArray< FStructProperty* > ResolvedInitializationNodeProperties;
 
 	// Indices for any Asset Player found within a specific (named) Anim Layer Graph, or implemented Anim Interface Graph
 	UPROPERTY()
@@ -89,19 +119,22 @@ public:
 	virtual const TArray<FAnimNotifyEvent>& GetAnimNotifies() const override { return AnimNotifies; }
 	virtual const TArray<FAnimBlueprintFunction>& GetAnimBlueprintFunctions() const override { return AnimBlueprintFunctions; }
 	virtual const TMap<FName, FCachedPoseIndices>& GetOrderedSavedPoseNodeIndicesMap() const override { return OrderedSavedPoseIndicesMap; }
-	virtual const TArray<FStructPropertyPath>& GetAnimNodeProperties() const override { return AnimNodeProperties; }
-	virtual const TArray<FStructPropertyPath>& GetLinkedAnimGraphNodeProperties() const override { return LinkedAnimGraphNodeProperties; }
-	virtual const TArray<FStructPropertyPath>& GetLinkedAnimLayerNodeProperties() const override { return LinkedAnimLayerNodeProperties; }
-	virtual const TArray<FStructPropertyPath>& GetPreUpdateNodeProperties() const override { return PreUpdateNodeProperties; }
-	virtual const TArray<FStructPropertyPath>& GetDynamicResetNodeProperties() const override { return DynamicResetNodeProperties; }
-	virtual const TArray<FStructPropertyPath>& GetStateMachineNodeProperties() const override { return StateMachineNodeProperties; }
-	virtual const TArray<FStructPropertyPath>& GetInitializationNodeProperties() const override { return InitializationNodeProperties; }
+	virtual const TArray<FStructProperty*>& GetAnimNodeProperties() const override { return ResolvedAnimNodeProperties; }
+	virtual const TArray<FStructProperty*>& GetLinkedAnimGraphNodeProperties() const override { return ResolvedLinkedAnimGraphNodeProperties; }
+	virtual const TArray<FStructProperty*>& GetLinkedAnimLayerNodeProperties() const override { return ResolvedLinkedAnimLayerNodeProperties; }
+	virtual const TArray<FStructProperty*>& GetPreUpdateNodeProperties() const override { return ResolvedPreUpdateNodeProperties; }
+	virtual const TArray<FStructProperty*>& GetDynamicResetNodeProperties() const override { return ResolvedDynamicResetNodeProperties; }
+	virtual const TArray<FStructProperty*>& GetStateMachineNodeProperties() const override { return ResolvedStateMachineNodeProperties; }
+	virtual const TArray<FStructProperty*>& GetInitializationNodeProperties() const override { return ResolvedInitializationNodeProperties; }
 	virtual const TArray<FName>& GetSyncGroupNames() const override { return SyncGroupNames; }
 	virtual int32 GetSyncGroupIndex(FName SyncGroupName) const override { return SyncGroupNames.IndexOfByKey(SyncGroupName); }
 	virtual const TArray<FExposedValueHandler>& GetExposedValueHandlers() const { return EvaluateGraphExposedInputs; }
 	virtual const TMap<FName, FGraphAssetPlayerInformation>& GetGraphAssetPlayerInformation() const { return GraphNameAssetPlayers; }
 	virtual const TMap<FName, FAnimGraphBlendOptions>& GetGraphBlendOptions() const { return GraphBlendOptions; }
 
+	// Resolve TFieldPaths to FStructPropertys
+	void ResolvePropertyPaths();
+
 #if WITH_EDITOR
 	void CopyFrom(IAnimClassInterface* AnimClass)
 	{
@@ -110,14 +143,38 @@ public:
 		TargetSkeleton = AnimClass->GetTargetSkeleton();
 		AnimNotifies = AnimClass->GetAnimNotifies();
 		AnimBlueprintFunctions = AnimClass->GetAnimBlueprintFunctions();
+		AnimBlueprintFunctionData.Empty(AnimBlueprintFunctions.Num());
+
+		for(const FAnimBlueprintFunction& AnimBlueprintFunction : AnimBlueprintFunctions)
+		{
+			FAnimBlueprintFunctionData& NewAnimBlueprintFunctionData = AnimBlueprintFunctionData.AddDefaulted_GetRef();
+			NewAnimBlueprintFunctionData.OutputPoseNodeProperty = AnimBlueprintFunction.OutputPoseNodeProperty;
+			Algo::Transform(AnimBlueprintFunction.InputProperties, NewAnimBlueprintFunctionData.InputProperties, [](FProperty* InProperty){ return TFieldPath<FProperty>(InProperty); });
+			Algo::Transform(AnimBlueprintFunction.InputPoseNodeProperties, NewAnimBlueprintFunctionData.InputPoseNodeProperties, [](FStructProperty* InProperty){ return TFieldPath<FStructProperty>(InProperty); });
+		}
+
 		OrderedSavedPoseIndicesMap = AnimClass->GetOrderedSavedPoseNodeIndicesMap();
-		AnimNodeProperties = AnimClass->GetAnimNodeProperties();
-		LinkedAnimGraphNodeProperties = AnimClass->GetLinkedAnimGraphNodeProperties();
-		LinkedAnimLayerNodeProperties = AnimClass->GetLinkedAnimLayerNodeProperties();
-		PreUpdateNodeProperties = AnimClass->GetPreUpdateNodeProperties();
-		DynamicResetNodeProperties = AnimClass->GetDynamicResetNodeProperties();
-		StateMachineNodeProperties = AnimClass->GetStateMachineNodeProperties();
-		InitializationNodeProperties = AnimClass->GetInitializationNodeProperties();
+
+		auto MakePropertyPath = [](FStructProperty* InProperty)
+		{ 
+			return TFieldPath<FStructProperty>(InProperty); 
+		};
+
+		Algo::Transform(AnimClass->GetAnimNodeProperties(), AnimNodeProperties, MakePropertyPath);
+		ResolvedAnimNodeProperties = AnimClass->GetAnimNodeProperties();
+		Algo::Transform(AnimClass->GetLinkedAnimGraphNodeProperties(), LinkedAnimGraphNodeProperties, MakePropertyPath);
+		ResolvedLinkedAnimGraphNodeProperties = AnimClass->GetLinkedAnimGraphNodeProperties();
+		Algo::Transform(AnimClass->GetLinkedAnimLayerNodeProperties(), LinkedAnimLayerNodeProperties, MakePropertyPath);
+		ResolvedLinkedAnimLayerNodeProperties = AnimClass->GetLinkedAnimLayerNodeProperties();
+		Algo::Transform(AnimClass->GetPreUpdateNodeProperties(), PreUpdateNodeProperties, MakePropertyPath);
+		ResolvedPreUpdateNodeProperties = AnimClass->GetPreUpdateNodeProperties();
+		Algo::Transform(AnimClass->GetDynamicResetNodeProperties(), DynamicResetNodeProperties, MakePropertyPath);
+		ResolvedDynamicResetNodeProperties = AnimClass->GetDynamicResetNodeProperties();
+		Algo::Transform(AnimClass->GetStateMachineNodeProperties(), StateMachineNodeProperties, MakePropertyPath);
+		ResolvedStateMachineNodeProperties = AnimClass->GetStateMachineNodeProperties();
+		Algo::Transform(AnimClass->GetInitializationNodeProperties(), InitializationNodeProperties, MakePropertyPath);
+		ResolvedInitializationNodeProperties = AnimClass->GetInitializationNodeProperties();
+
 		SyncGroupNames = AnimClass->GetSyncGroupNames();
 		EvaluateGraphExposedInputs = AnimClass->GetExposedValueHandlers();
 		GraphNameAssetPlayers = AnimClass->GetGraphAssetPlayerInformation();
diff --git a/Engine/Source/Runtime/Engine/Classes/Animation/AnimClassInterface.h b/Engine/Source/Runtime/Engine/Classes/Animation/AnimClassInterface.h
--- a/Engine/Source/Runtime/Engine/Classes/Animation/AnimClassInterface.h
+++ b/Engine/Source/Runtime/Engine/Classes/Animation/AnimClassInterface.h
@@ -65,16 +65,13 @@ struct FAnimBlueprintFunction
 	TArray<int32> InputPoseNodeIndices;
 
 	/** The property of the output node, patched up during link */
-	UPROPERTY(transient)
-	TFieldPath<FStructProperty> OutputPoseNodeProperty;
+	FStructProperty* OutputPoseNodeProperty;
 
 	/** The properties of the input nodes, patched up during link */
-	UPROPERTY(transient)
-	TArray< TFieldPath<FStructProperty> > InputPoseNodeProperties;
+	TArray< FStructProperty* > InputPoseNodeProperties;
 
 	/** The input properties themselves */
-	UPROPERTY(transient)
-	TArray< TFieldPath<FProperty> > InputProperties;
+	TArray< FProperty* > InputProperties;
 
 	/** Whether this function is actually implemented by this class - it could just be a stub */
 	UPROPERTY(transient)
@@ -147,17 +144,17 @@ class ENGINE_API IAnimClassInterface
 public:
 	virtual const TArray<FBakedAnimationStateMachine>& GetBakedStateMachines() const = 0;
 	virtual const TArray<FAnimNotifyEvent>& GetAnimNotifies() const = 0;
-	virtual const TArray<FStructPropertyPath>& GetAnimNodeProperties() const = 0;
+	virtual const TArray<FStructProperty*>& GetAnimNodeProperties() const = 0;
 	UE_DEPRECATED(4.24, "Function has been renamed, please use GetLinkedAnimGraphNodeProperties")
-	virtual const TArray<FStructPropertyPath>& GetSubInstanceNodeProperties() const { return GetLinkedAnimGraphNodeProperties(); }
-	virtual const TArray<FStructPropertyPath>& GetLinkedAnimGraphNodeProperties() const = 0;
+	virtual const TArray<FStructProperty*>& GetSubInstanceNodeProperties() const { return GetLinkedAnimGraphNodeProperties(); }
+	virtual const TArray<FStructProperty*>& GetLinkedAnimGraphNodeProperties() const = 0;
 	UE_DEPRECATED(4.24, "Function has been renamed, please use GetLinkedLayerNodeProperties")
-	virtual const TArray<FStructPropertyPath>& GetLayerNodeProperties() const { return GetLinkedAnimLayerNodeProperties(); }
-	virtual const TArray<FStructPropertyPath>& GetLinkedAnimLayerNodeProperties() const = 0;
-	virtual const TArray<FStructPropertyPath>& GetPreUpdateNodeProperties() const = 0;
-	virtual const TArray<FStructPropertyPath>& GetDynamicResetNodeProperties() const = 0;
-	virtual const TArray<FStructPropertyPath>& GetStateMachineNodeProperties() const = 0;
-	virtual const TArray<FStructPropertyPath>& GetInitializationNodeProperties() const = 0;
+	virtual const TArray<FStructProperty*>& GetLayerNodeProperties() const { return GetLinkedAnimLayerNodeProperties(); }
+	virtual const TArray<FStructProperty*>& GetLinkedAnimLayerNodeProperties() const = 0;
+	virtual const TArray<FStructProperty*>& GetPreUpdateNodeProperties() const = 0;
+	virtual const TArray<FStructProperty*>& GetDynamicResetNodeProperties() const = 0;
+	virtual const TArray<FStructProperty*>& GetStateMachineNodeProperties() const = 0;
+	virtual const TArray<FStructProperty*>& GetInitializationNodeProperties() const = 0;
 	virtual const TArray<FExposedValueHandler>& GetExposedValueHandlers() const = 0;
 	virtual const TArray<FName>& GetSyncGroupNames() const = 0;
 	virtual const TMap<FName, FCachedPoseIndices>& GetOrderedSavedPoseNodeIndicesMap() const = 0;
diff --git a/Engine/Source/Runtime/Engine/Private/Animation/AnimBlueprintGeneratedClass.cpp b/Engine/Source/Runtime/Engine/Private/Animation/AnimBlueprintGeneratedClass.cpp
--- a/Engine/Source/Runtime/Engine/Private/Animation/AnimBlueprintGeneratedClass.cpp
+++ b/Engine/Source/Runtime/Engine/Private/Animation/AnimBlueprintGeneratedClass.cpp
@@ -492,7 +492,7 @@ void UAnimBlueprintGeneratedClass::LinkFunctionsToDefaultObjectNodes(UObject* De
 	// Link functions to their nodes
 	for(int32 AnimNodeIndex = 0; AnimNodeIndex < AnimNodeProperties.Num(); ++AnimNodeIndex)
 	{
-		FStructProperty* StructProperty = AnimNodeProperties[AnimNodeIndex].Get();
+		FStructProperty* StructProperty = AnimNodeProperties[AnimNodeIndex];
 		if (StructProperty->Struct->IsChildOf(FAnimNode_Root::StaticStruct()))
 		{
 			FAnimNode_Root* RootNode = StructProperty->ContainerPtrToValuePtr<FAnimNode_Root>(DefaultObject);
@@ -541,7 +541,7 @@ void UAnimBlueprintGeneratedClass::LinkFunctionsToDefaultObjectNodes(UObject* De
 	// Determine whether functions are 'implemented'
 	for(FAnimBlueprintFunction& AnimBlueprintFunction : AnimBlueprintFunctions)
 	{
-		if(AnimBlueprintFunction.OutputPoseNodeProperty.Get())
+		if(AnimBlueprintFunction.OutputPoseNodeProperty)
 		{
 			FAnimNode_Root* RootNode = AnimBlueprintFunction.OutputPoseNodeProperty->ContainerPtrToValuePtr<FAnimNode_Root>(DefaultObject);
 			if(RootNode->Result.LinkID != INDEX_NONE)
diff --git a/Engine/Source/Runtime/Engine/Private/Animation/AnimClassData.cpp b/Engine/Source/Runtime/Engine/Private/Animation/AnimClassData.cpp
new file mode 100644
--- /dev/null
+++ b/Engine/Source/Runtime/Engine/Private/Animation/AnimClassData.cpp
@@ -0,0 +1,31 @@
+// Copyright Epic Games, Inc. All Rights Reserved.
+
+#include "Animation/AnimClassData.h"
+#include "Animation/AnimNode_Root.h"
+#include "Animation/AnimNode_LinkedInputPose.h"
+
+void UAnimClassData::ResolvePropertyPaths()
+{
+	auto ResolveTransform = [](const TFieldPath<FStructProperty>& InPropertyPath)
+	{ 
+		return InPropertyPath.Get(); 
+	};
+
+	// Copy serialized property paths to resolved paths
+	Algo::Transform(AnimNodeProperties, ResolvedAnimNodeProperties, ResolveTransform);
+	Algo::Transform(LinkedAnimGraphNodeProperties, ResolvedLinkedAnimGraphNodeProperties, ResolveTransform);
+	Algo::Transform(LinkedAnimLayerNodeProperties, ResolvedLinkedAnimLayerNodeProperties, ResolveTransform);
+	Algo::Transform(PreUpdateNodeProperties, ResolvedPreUpdateNodeProperties, ResolveTransform);
+	Algo::Transform(DynamicResetNodeProperties, ResolvedDynamicResetNodeProperties, ResolveTransform);
+	Algo::Transform(StateMachineNodeProperties, ResolvedStateMachineNodeProperties, ResolveTransform);
+	Algo::Transform(InitializationNodeProperties, ResolvedInitializationNodeProperties, ResolveTransform);
+
+	check(AnimBlueprintFunctions.Num() == AnimBlueprintFunctionData.Num());
+
+	for(int32 FunctionIndex = 0; FunctionIndex < AnimBlueprintFunctions.Num(); ++FunctionIndex)
+	{
+		AnimBlueprintFunctions[FunctionIndex].OutputPoseNodeProperty = AnimBlueprintFunctionData[FunctionIndex].OutputPoseNodeProperty.Get();
+		Algo::Transform(AnimBlueprintFunctionData[FunctionIndex].InputProperties, AnimBlueprintFunctions[FunctionIndex].InputProperties, [](const TFieldPath<FProperty>& InPropertyPath){ return InPropertyPath.Get(); });
+		Algo::Transform(AnimBlueprintFunctionData[FunctionIndex].InputPoseNodeProperties, AnimBlueprintFunctions[FunctionIndex].InputPoseNodeProperties, ResolveTransform);
+	}
+}
\ No newline at end of file
diff --git a/Engine/Source/Runtime/Engine/Private/Animation/AnimInstance.cpp b/Engine/Source/Runtime/Engine/Private/Animation/AnimInstance.cpp
--- a/Engine/Source/Runtime/Engine/Private/Animation/AnimInstance.cpp
+++ b/Engine/Source/Runtime/Engine/Private/Animation/AnimInstance.cpp
@@ -1535,12 +1535,12 @@ FName UAnimInstance::GetCurrentStateName(int32 MachineIndex)
 {
 	if (IAnimClassInterface* AnimBlueprintClass = IAnimClassInterface::GetFromClass(GetClass()))
 	{
-		const TArray<FStructPropertyPath>& AnimNodeProperties = AnimBlueprintClass->GetAnimNodeProperties();
+		const TArray<FStructProperty*>& AnimNodeProperties = AnimBlueprintClass->GetAnimNodeProperties();
 		if ((MachineIndex >= 0) && (MachineIndex < AnimNodeProperties.Num()))
 		{
 			const int32 InstancePropertyIndex = AnimNodeProperties.Num() - 1 - MachineIndex; //@TODO: ANIMREFACTOR: Reverse indexing
 
-			FStructProperty* MachineInstanceProperty = AnimNodeProperties[InstancePropertyIndex].Get();
+			FStructProperty* MachineInstanceProperty = AnimNodeProperties[InstancePropertyIndex];
 			checkSlow(MachineInstanceProperty->Struct->IsChildOf(FAnimNode_StateMachine::StaticStruct()));
 
 			FAnimNode_StateMachine* MachineInstance = MachineInstanceProperty->ContainerPtrToValuePtr<FAnimNode_StateMachine>(this);
@@ -2536,8 +2536,8 @@ UAnimInstance* UAnimInstance::GetLinkedAnimGraphInstanceByTag(FName InTag) const
 {
 	if (IAnimClassInterface* AnimBlueprintClass = IAnimClassInterface::GetFromClass(GetClass()))
 	{
-		const TArray<FStructPropertyPath>& LinkedAnimGraphNodeProperties = AnimBlueprintClass->GetLinkedAnimGraphNodeProperties();
-		for(const FStructPropertyPath& LinkedAnimGraphNodeProperty : LinkedAnimGraphNodeProperties)
+		const TArray<FStructProperty*>& LinkedAnimGraphNodeProperties = AnimBlueprintClass->GetLinkedAnimGraphNodeProperties();
+		for(const FStructProperty* LinkedAnimGraphNodeProperty : LinkedAnimGraphNodeProperties)
 		{
 			const FAnimNode_LinkedAnimGraph* LinkedAnimGraph = LinkedAnimGraphNodeProperty->ContainerPtrToValuePtr<FAnimNode_LinkedAnimGraph>(this);
 			if(LinkedAnimGraph && LinkedAnimGraph->Tag == InTag)
@@ -2554,8 +2554,8 @@ void UAnimInstance::GetLinkedAnimGraphInstancesByTag(FName InTag, TArray<UAnimIn
 {
 	if (IAnimClassInterface* AnimBlueprintClass = IAnimClassInterface::GetFromClass(GetClass()))
 	{
-		const TArray<FStructPropertyPath>& LinkedAnimGraphNodeProperties = AnimBlueprintClass->GetLinkedAnimGraphNodeProperties();
-		for (const FStructPropertyPath& LinkedAnimGraphNodeProperty : LinkedAnimGraphNodeProperties)
+		const TArray<FStructProperty*>& LinkedAnimGraphNodeProperties = AnimBlueprintClass->GetLinkedAnimGraphNodeProperties();
+		for (const FStructProperty* LinkedAnimGraphNodeProperty : LinkedAnimGraphNodeProperties)
 		{
 			const FAnimNode_LinkedAnimGraph* LinkedAnimGraph = LinkedAnimGraphNodeProperty->ContainerPtrToValuePtr<FAnimNode_LinkedAnimGraph>(this);
 			if(LinkedAnimGraph && LinkedAnimGraph->Tag == InTag)
@@ -2570,8 +2570,8 @@ void UAnimInstance::LinkAnimGraphByTag(FName InTag, TSubclassOf<UAnimInstance> I
 {
 	if (IAnimClassInterface* AnimBlueprintClass = IAnimClassInterface::GetFromClass(GetClass()))
 	{
-		const TArray<FStructPropertyPath>& LinkedAnimGraphNodeProperties = AnimBlueprintClass->GetLinkedAnimGraphNodeProperties();
-		for (const FStructPropertyPath& LinkedAnimGraphNodeProperty : LinkedAnimGraphNodeProperties)
+		const TArray<FStructProperty*>& LinkedAnimGraphNodeProperties = AnimBlueprintClass->GetLinkedAnimGraphNodeProperties();
+		for (const FStructProperty* LinkedAnimGraphNodeProperty : LinkedAnimGraphNodeProperties)
 		{
 			FAnimNode_LinkedAnimGraph* LinkedAnimGraph = LinkedAnimGraphNodeProperty->ContainerPtrToValuePtr<FAnimNode_LinkedAnimGraph>(this);
 			if(LinkedAnimGraph && LinkedAnimGraph->Tag == InTag)
@@ -2606,7 +2606,7 @@ void UAnimInstance::PerformLinkedLayerOverlayOperation(TSubclassOf<UAnimInstance
 		// Map of group name->nodes to run under that group instance
 		TMap<FName, TArray<FAnimNode_LinkedAnimLayer*, TInlineAllocator<4>>, TInlineSetAllocator<4>> LayerNodesToSet;
 
-		for(const FStructPropertyPath& LayerNodeProperty : AnimBlueprintClass->GetLinkedAnimLayerNodeProperties())
+		for(const FStructProperty* LayerNodeProperty : AnimBlueprintClass->GetLinkedAnimLayerNodeProperties())
 		{
 			FAnimNode_LinkedAnimLayer* Layer = LayerNodeProperty->ContainerPtrToValuePtr<FAnimNode_LinkedAnimLayer>(this);
 
@@ -2639,7 +2639,7 @@ void UAnimInstance::PerformLinkedLayerOverlayOperation(TSubclassOf<UAnimInstance
 		auto UnlinkLayerNodesInInstance = [](UAnimInstance* InAnimInstance)
 		{
 			const IAnimClassInterface* NewLinkedInstanceClass = IAnimClassInterface::GetFromClass(InAnimInstance->GetClass());
-			for(const FStructPropertyPath& LayerNodeProperty : NewLinkedInstanceClass->GetLinkedAnimLayerNodeProperties())
+			for(const FStructProperty* LayerNodeProperty : NewLinkedInstanceClass->GetLinkedAnimLayerNodeProperties())
 			{
 				FAnimNode_LinkedAnimLayer* LinkedAnimLayerNode = LayerNodeProperty->ContainerPtrToValuePtr<FAnimNode_LinkedAnimLayer>(InAnimInstance);
 				LinkedAnimLayerNode->DynamicUnlink(InAnimInstance);
@@ -2850,7 +2850,7 @@ UAnimInstance* UAnimInstance::GetLinkedAnimLayerInstanceByGroup(FName InGroup) c
 {
 	if (IAnimClassInterface* AnimBlueprintClass = IAnimClassInterface::GetFromClass(GetClass()))
 	{
-		for(const FStructPropertyPath& LayerNodeProperty : AnimBlueprintClass->GetLinkedAnimLayerNodeProperties())
+		for(const FStructProperty* LayerNodeProperty : AnimBlueprintClass->GetLinkedAnimLayerNodeProperties())
 		{
 			const FAnimNode_LinkedAnimLayer* Layer = LayerNodeProperty->ContainerPtrToValuePtr<FAnimNode_LinkedAnimLayer>(this);
 
@@ -2882,7 +2882,7 @@ UAnimInstance* UAnimInstance::GetLinkedAnimLayerInstanceByClass(TSubclassOf<UAni
 {
 	if (IAnimClassInterface* AnimBlueprintClass = IAnimClassInterface::GetFromClass(GetClass()))
 	{
-		for(const FStructPropertyPath& LayerNodeProperty : AnimBlueprintClass->GetLinkedAnimLayerNodeProperties())
+		for(const FStructProperty* LayerNodeProperty : AnimBlueprintClass->GetLinkedAnimLayerNodeProperties())
 		{
 			const FAnimNode_LinkedAnimLayer* Layer = LayerNodeProperty->ContainerPtrToValuePtr<FAnimNode_LinkedAnimLayer>(this);
 			UAnimInstance* TargetInstance = Layer->GetTargetInstance<UAnimInstance>();
diff --git a/Engine/Source/Runtime/Engine/Private/Animation/AnimInstanceProxy.cpp b/Engine/Source/Runtime/Engine/Private/Animation/AnimInstanceProxy.cpp
--- a/Engine/Source/Runtime/Engine/Private/Animation/AnimInstanceProxy.cpp
+++ b/Engine/Source/Runtime/Engine/Private/Animation/AnimInstanceProxy.cpp
@@ -90,13 +90,13 @@ void FAnimInstanceProxy::Initialize(UAnimInstance* InAnimInstance)
 
 	if (AnimClassInterface)
 	{
-		const TArray<FStructPropertyPath>& AnimNodeProperties = AnimClassInterface->GetAnimNodeProperties();
+		const TArray<FStructProperty*>& AnimNodeProperties = AnimClassInterface->GetAnimNodeProperties();
 
 		// Grab a pointer to the default root node, if any
 		RootNode = nullptr;
 		if(AnimClassInterface->GetAnimBlueprintFunctions().Num() > 0)
 		{
-			if(FProperty* RootNodeProperty = AnimClassInterface->GetAnimBlueprintFunctions()[0].OutputPoseNodeProperty.Get())
+			if(FProperty* RootNodeProperty = AnimClassInterface->GetAnimBlueprintFunctions()[0].OutputPoseNodeProperty)
 			{
 				RootNode = RootNodeProperty->ContainerPtrToValuePtr<FAnimNode_Root>(InAnimInstance);
 			}
@@ -196,7 +196,7 @@ void FAnimInstanceProxy::InitializeRootNode(bool bInDeferRootNodeInitialization)
 	if(AnimClassInterface)
 	{
 		// cache any state machine descriptions we have
-		for (const FStructPropertyPath& Property : AnimClassInterface->GetStateMachineNodeProperties())
+		for (const FStructProperty* Property : AnimClassInterface->GetStateMachineNodeProperties())
 		{
 			FAnimNode_StateMachine* StateMachine = Property->ContainerPtrToValuePtr<FAnimNode_StateMachine>(AnimInstanceObject);
 			StateMachine->CacheMachineDescription(AnimClassInterface);
@@ -204,21 +204,21 @@ void FAnimInstanceProxy::InitializeRootNode(bool bInDeferRootNodeInitialization)
 
 		// Init any nodes that need non-relevancy based initialization
 		UAnimInstance* AnimInstance = CastChecked<UAnimInstance>(GetAnimInstanceObject());
-		for (const FStructPropertyPath& Property : AnimClassInterface->GetInitializationNodeProperties())
+		for (const FStructProperty* Property : AnimClassInterface->GetInitializationNodeProperties())
 		{
 			FAnimNode_Base* AnimNode = Property->ContainerPtrToValuePtr<FAnimNode_Base>(AnimInstanceObject);
 			AnimNode->OnInitializeAnimInstance(this, AnimInstance);
 		}
 
 		// Cache any preupdate nodes
-		for (const FStructPropertyPath& Property : AnimClassInterface->GetPreUpdateNodeProperties())
+		for (const FStructProperty* Property : AnimClassInterface->GetPreUpdateNodeProperties())
 		{
 			FAnimNode_Base* AnimNode = Property->ContainerPtrToValuePtr<FAnimNode_Base>(AnimInstanceObject);
 			GameThreadPreUpdateNodes.Add(AnimNode);
 		}
 
 		// Cache any dynamic reset nodes
-		for (const FStructPropertyPath& Property : AnimClassInterface->GetDynamicResetNodeProperties())
+		for (const FStructProperty* Property : AnimClassInterface->GetDynamicResetNodeProperties())
 		{
 			FAnimNode_Base* AnimNode = Property->ContainerPtrToValuePtr<FAnimNode_Base>(AnimInstanceObject);
 			DynamicResetNodes.Add(AnimNode);
@@ -1096,12 +1096,12 @@ FAnimNode_Base* FAnimInstanceProxy::GetCheckedNodeFromIndexUntyped(int32 NodeIdx
 	FAnimNode_Base* NodePtr = nullptr;
 	if (AnimClassInterface)
 	{
-		const TArray<FStructPropertyPath>& AnimNodeProperties = AnimClassInterface->GetAnimNodeProperties();
+		const TArray<FStructProperty*>& AnimNodeProperties = AnimClassInterface->GetAnimNodeProperties();
 		const int32 InstanceIdx = AnimNodeProperties.Num() - 1 - NodeIdx;
 
 		if (AnimNodeProperties.IsValidIndex(InstanceIdx))
 		{
-			FStructProperty* NodeProperty = AnimNodeProperties[InstanceIdx].Get();
+			FStructProperty* NodeProperty = AnimNodeProperties[InstanceIdx];
 
 			if (NodeProperty->Struct->IsChildOf(RequiredStructType))
 			{
@@ -1128,12 +1128,12 @@ FAnimNode_Base* FAnimInstanceProxy::GetNodeFromIndexUntyped(int32 NodeIdx, UScri
 	FAnimNode_Base* NodePtr = nullptr;
 	if (AnimClassInterface)
 	{
-		const TArray<FStructPropertyPath>& AnimNodeProperties = AnimClassInterface->GetAnimNodeProperties();
+		const TArray<FStructProperty*>& AnimNodeProperties = AnimClassInterface->GetAnimNodeProperties();
 		const int32 InstanceIdx = AnimNodeProperties.Num() - 1 - NodeIdx;
 
 		if (AnimNodeProperties.IsValidIndex(InstanceIdx))
 		{
-			FStructProperty* NodeProperty = AnimNodeProperties[InstanceIdx].Get();
+			FStructProperty* NodeProperty = AnimNodeProperties[InstanceIdx];
 
 			if (NodeProperty->Struct->IsChildOf(RequiredStructType))
 			{
@@ -1236,7 +1236,7 @@ void FAnimInstanceProxy::UpdateAnimation_WithRoot(const FAnimationUpdateContext&
 			if(AnimClassInterface)
 			{
 				// Initialize linked sub graphs
-				for(const FStructPropertyPath& LayerNodeProperty : AnimClassInterface->GetLinkedAnimLayerNodeProperties())
+				for(const FStructProperty* LayerNodeProperty : AnimClassInterface->GetLinkedAnimLayerNodeProperties())
 				{
 					if(FAnimNode_LinkedAnimLayer* LayerNode = LayerNodeProperty->ContainerPtrToValuePtr<FAnimNode_LinkedAnimLayer>(AnimInstanceObject))
 					{
@@ -2033,12 +2033,12 @@ FAnimNode_StateMachine* FAnimInstanceProxy::GetStateMachineInstance(int32 Machin
 {
 	if (AnimClassInterface)
 	{
-		const TArray<FStructPropertyPath>& AnimNodeProperties = AnimClassInterface->GetAnimNodeProperties();
+		const TArray<FStructProperty*>& AnimNodeProperties = AnimClassInterface->GetAnimNodeProperties();
 		if ((MachineIndex >= 0) && (MachineIndex < AnimNodeProperties.Num()))
 		{
 			const int32 InstancePropertyIndex = AnimNodeProperties.Num() - 1 - MachineIndex;
 
-			FStructProperty* MachineInstanceProperty = AnimNodeProperties[InstancePropertyIndex].Get();
+			FStructProperty* MachineInstanceProperty = AnimNodeProperties[InstancePropertyIndex];
 			checkSlow(MachineInstanceProperty->Struct->IsChildOf(FAnimNode_StateMachine::StaticStruct()));
 
 			return MachineInstanceProperty->ContainerPtrToValuePtr<FAnimNode_StateMachine>(AnimInstanceObject);
@@ -2129,9 +2129,9 @@ void FAnimInstanceProxy::BindNativeDelegates()
 
 	auto ForEachStateLambda = [&](IAnimClassInterface* InAnimClassInterface, const FName& MachineName, const FName& StateName, TFunctionRef<void(FAnimNode_StateMachine*, const FBakedAnimationState&, int32)> Predicate)
 	{
-		for (const FStructPropertyPath& Property : InAnimClassInterface->GetAnimNodeProperties())
+		for (const FStructProperty* Property : InAnimClassInterface->GetAnimNodeProperties())
 		{
-			if(Property.Get() && Property->Struct->IsChildOf(FAnimNode_StateMachine::StaticStruct()))
+			if(Property && Property->Struct->IsChildOf(FAnimNode_StateMachine::StaticStruct()))
 			{
 				FAnimNode_StateMachine* StateMachine = Property->ContainerPtrToValuePtr<FAnimNode_StateMachine>(AnimInstanceObject);
 				if(StateMachine)
@@ -2221,10 +2221,10 @@ FAnimNode_StateMachine* FAnimInstanceProxy::GetStateMachineInstanceFromName(FNam
 {
 	if (AnimClassInterface)
 	{
-		const TArray<FStructPropertyPath>& AnimNodeProperties = AnimClassInterface->GetAnimNodeProperties();
+		const TArray<FStructProperty*>& AnimNodeProperties = AnimClassInterface->GetAnimNodeProperties();
 		for (int32 MachineIndex = 0; MachineIndex < AnimNodeProperties.Num(); MachineIndex++)
 		{
-			FStructProperty* Property = AnimNodeProperties[AnimNodeProperties.Num() - 1 - MachineIndex].Get();
+			FStructProperty* Property = AnimNodeProperties[AnimNodeProperties.Num() - 1 - MachineIndex];
 			if (Property && Property->Struct->IsChildOf(FAnimNode_StateMachine::StaticStruct()))
 			{
 				FAnimNode_StateMachine* StateMachine = Property->ContainerPtrToValuePtr<FAnimNode_StateMachine>(AnimInstanceObject);
@@ -2249,10 +2249,10 @@ const FBakedAnimationStateMachine* FAnimInstanceProxy::GetStateMachineInstanceDe
 {
 	if (AnimClassInterface)
 	{
-		const TArray<FStructPropertyPath>& AnimNodeProperties = AnimClassInterface->GetAnimNodeProperties();
+		const TArray<FStructProperty*>& AnimNodeProperties = AnimClassInterface->GetAnimNodeProperties();
 		for (int32 MachineIndex = 0; MachineIndex < AnimNodeProperties.Num(); MachineIndex++)
 		{
-			FStructProperty* Property = AnimNodeProperties[AnimNodeProperties.Num() - 1 - MachineIndex].Get();
+			FStructProperty* Property = AnimNodeProperties[AnimNodeProperties.Num() - 1 - MachineIndex];
 			if(Property && Property->Struct->IsChildOf(FAnimNode_StateMachine::StaticStruct()))
 			{
 				FAnimNode_StateMachine* StateMachine = Property->ContainerPtrToValuePtr<FAnimNode_StateMachine>(AnimInstanceObject);
@@ -2277,10 +2277,10 @@ int32 FAnimInstanceProxy::GetStateMachineIndex(FName MachineName)
 {
 	if (AnimClassInterface)
 	{
-		const TArray<FStructPropertyPath>& AnimNodeProperties = AnimClassInterface->GetAnimNodeProperties();
+		const TArray<FStructProperty*>& AnimNodeProperties = AnimClassInterface->GetAnimNodeProperties();
 		for (int32 MachineIndex = 0; MachineIndex < AnimNodeProperties.Num(); MachineIndex++)
 		{
-			FStructProperty* Property = AnimNodeProperties[AnimNodeProperties.Num() - 1 - MachineIndex].Get();
+			FStructProperty* Property = AnimNodeProperties[AnimNodeProperties.Num() - 1 - MachineIndex];
 			if(Property && Property->Struct->IsChildOf(FAnimNode_StateMachine::StaticStruct()))
 			{
 				FAnimNode_StateMachine* StateMachine = Property->ContainerPtrToValuePtr<FAnimNode_StateMachine>(AnimInstanceObject);
@@ -2305,10 +2305,10 @@ void FAnimInstanceProxy::GetStateMachineIndexAndDescription(FName InMachineName,
 {
 	if (AnimClassInterface)
 	{
-		const TArray<FStructPropertyPath>& AnimNodeProperties = AnimClassInterface->GetAnimNodeProperties();
+		const TArray<FStructProperty*>& AnimNodeProperties = AnimClassInterface->GetAnimNodeProperties();
 		for (int32 MachineIndex = 0; MachineIndex < AnimNodeProperties.Num(); MachineIndex++)
 		{
-			FStructProperty* Property = AnimNodeProperties[AnimNodeProperties.Num() - 1 - MachineIndex].Get();
+			FStructProperty* Property = AnimNodeProperties[AnimNodeProperties.Num() - 1 - MachineIndex];
 			if (Property && Property->Struct->IsChildOf(FAnimNode_StateMachine::StaticStruct()))
 			{
 				FAnimNode_StateMachine* StateMachine = Property->ContainerPtrToValuePtr<FAnimNode_StateMachine>(AnimInstanceObject);
@@ -2344,7 +2344,7 @@ int32 FAnimInstanceProxy::GetInstanceAssetPlayerIndex(FName MachineName, FName S
 	{
 		if(const FBakedAnimationStateMachine* MachineDescription = GetStateMachineInstanceDesc(MachineName))
 		{
-			const TArray<FStructPropertyPath>& AnimNodeProperties = AnimClassInterface->GetAnimNodeProperties();
+			const TArray<FStructProperty*>& AnimNodeProperties = AnimClassInterface->GetAnimNodeProperties();
 			for(int32 StateIndex = 0; StateIndex < MachineDescription->States.Num(); StateIndex++)
 			{
 				const FBakedAnimationState& State = MachineDescription->States[StateIndex];
@@ -2353,7 +2353,7 @@ int32 FAnimInstanceProxy::GetInstanceAssetPlayerIndex(FName MachineName, FName S
 					for(int32 PlayerIndex = 0; PlayerIndex < State.PlayerNodeIndices.Num(); PlayerIndex++)
 					{
 						checkSlow(State.PlayerNodeIndices[PlayerIndex] < AnimNodeProperties.Num());
-						FStructProperty* AssetPlayerProperty = AnimNodeProperties[AnimNodeProperties.Num() - 1 - State.PlayerNodeIndices[PlayerIndex]].Get();
+						FStructProperty* AssetPlayerProperty = AnimNodeProperties[AnimNodeProperties.Num() - 1 - State.PlayerNodeIndices[PlayerIndex]];
 						if(AssetPlayerProperty && AssetPlayerProperty->Struct->IsChildOf(FAnimNode_AssetPlayerBase::StaticStruct()))
 						{
 							FAnimNode_AssetPlayerBase* AssetPlayer = AssetPlayerProperty->ContainerPtrToValuePtr<FAnimNode_AssetPlayerBase>(AnimInstanceObject);
diff --git a/Engine/Source/Runtime/Engine/Private/Animation/AnimNodeBase.cpp b/Engine/Source/Runtime/Engine/Private/Animation/AnimNodeBase.cpp
--- a/Engine/Source/Runtime/Engine/Private/Animation/AnimNodeBase.cpp
+++ b/Engine/Source/Runtime/Engine/Private/Animation/AnimNodeBase.cpp
@@ -206,10 +206,10 @@ void FPoseLinkBase::AttemptRelink(const FAnimationBaseContext& Context)
 		check(AnimBlueprintClass);
 
 		// adding ensure. We had a crash on here
-		const TArray<FStructPropertyPath>& AnimNodeProperties = AnimBlueprintClass->GetAnimNodeProperties();
+		const TArray<FStructProperty*>& AnimNodeProperties = AnimBlueprintClass->GetAnimNodeProperties();
 		if ( ensure(AnimNodeProperties.IsValidIndex(LinkID)) )
 		{
-			FStructPropertyPath LinkedProperty = AnimNodeProperties[LinkID];
+			FStructProperty* LinkedProperty = AnimNodeProperties[LinkID];
 			void* LinkedNodePtr = LinkedProperty->ContainerPtrToValuePtr<void>(Context.AnimInstanceProxy->GetAnimInstanceObject());
 			LinkedNode = (FAnimNode_Base*)LinkedNodePtr;
 		}
diff --git a/Engine/Source/Runtime/Engine/Private/Animation/AnimNode_LinkedAnimGraph.cpp b/Engine/Source/Runtime/Engine/Private/Animation/AnimNode_LinkedAnimGraph.cpp
--- a/Engine/Source/Runtime/Engine/Private/Animation/AnimNode_LinkedAnimGraph.cpp
+++ b/Engine/Source/Runtime/Engine/Private/Animation/AnimNode_LinkedAnimGraph.cpp
@@ -351,7 +351,7 @@ void FAnimNode_LinkedAnimGraph::DynamicLink(UAnimInstance* InOwningAnimInstance)
 						InputPoses[InputPoseIndex].AttemptRelink(Context);
 
 						int32 InputPropertyIndex = FindFunctionInputIndex(AnimBlueprintFunction, AnimBlueprintFunction.InputPoseNames[InputPoseIndex]);
-						if(InputPropertyIndex != INDEX_NONE && AnimBlueprintFunction.InputPoseNodeProperties[InputPropertyIndex].Get())
+						if(InputPropertyIndex != INDEX_NONE && AnimBlueprintFunction.InputPoseNodeProperties[InputPropertyIndex])
 						{
 							FAnimNode_LinkedInputPose* LinkedInputPoseNode = AnimBlueprintFunction.InputPoseNodeProperties[InputPropertyIndex]->ContainerPtrToValuePtr<FAnimNode_LinkedInputPose>(LinkTargetInstance);
 							check(LinkedInputPoseNode->Name == AnimBlueprintFunction.InputPoseNames[InputPoseIndex]);
@@ -363,7 +363,7 @@ void FAnimNode_LinkedAnimGraph::DynamicLink(UAnimInstance* InOwningAnimInstance)
 						}
 					}
 
-					if(AnimBlueprintFunction.OutputPoseNodeProperty.Get())
+					if(AnimBlueprintFunction.OutputPoseNodeProperty)
 					{
 						LinkedRoot = AnimBlueprintFunction.OutputPoseNodeProperty->ContainerPtrToValuePtr<FAnimNode_Root>(LinkTargetInstance);
 					}
@@ -401,7 +401,7 @@ void FAnimNode_LinkedAnimGraph::DynamicUnlink(UAnimInstance* InOwningAnimInstanc
 					for(int32 InputPoseIndex = 0; InputPoseIndex < AnimBlueprintFunction.InputPoseNames.Num() && InputPoseIndex < InputPoses.Num(); ++InputPoseIndex)
 					{
 						int32 InputPropertyIndex = FindFunctionInputIndex(AnimBlueprintFunction, AnimBlueprintFunction.InputPoseNames[InputPoseIndex]);
-						if(InputPropertyIndex != INDEX_NONE && AnimBlueprintFunction.InputPoseNodeProperties[InputPropertyIndex].Get())
+						if(InputPropertyIndex != INDEX_NONE && AnimBlueprintFunction.InputPoseNodeProperties[InputPropertyIndex])
 						{
 							FAnimNode_LinkedInputPose* LinkedInputPoseNode = AnimBlueprintFunction.InputPoseNodeProperties[InputPropertyIndex]->ContainerPtrToValuePtr<FAnimNode_LinkedInputPose>(LinkTargetInstance);
 							check(LinkedInputPoseNode->Name == AnimBlueprintFunction.InputPoseNames[InputPoseIndex]);
diff --git a/Engine/Source/Runtime/Engine/Private/Animation/AnimTrace.cpp b/Engine/Source/Runtime/Engine/Private/Animation/AnimTrace.cpp
--- a/Engine/Source/Runtime/Engine/Private/Animation/AnimTrace.cpp
+++ b/Engine/Source/Runtime/Engine/Private/Animation/AnimTrace.cpp
@@ -607,9 +607,9 @@ void FAnimTrace::OutputAnimNodeStart(const FAnimationBaseContext& InContext, uin
 	IAnimClassInterface* AnimBlueprintClass = InContext.GetAnimClass();
 	if(AnimBlueprintClass)
 	{
-		const TArray<FStructPropertyPath>& AnimNodeProperties = AnimBlueprintClass->GetAnimNodeProperties();
+		const TArray<FStructProperty*>& AnimNodeProperties = AnimBlueprintClass->GetAnimNodeProperties();
 		check(AnimNodeProperties.IsValidIndex(InNodeId));
-		FStructProperty* LinkedProperty = AnimNodeProperties[InNodeId].Get();
+		FStructProperty* LinkedProperty = AnimNodeProperties[InNodeId];
 		check(LinkedProperty->Struct);
 
 #if WITH_EDITOR
