Fix CIS.

diff --git a/Engine/Source/Runtime/Engine/Public/Physics/Experimental/PhysInterface_Chaos.h b/Engine/Source/Runtime/Engine/Public/Physics/Experimental/PhysInterface_Chaos.h
--- a/Engine/Source/Runtime/Engine/Public/Physics/Experimental/PhysInterface_Chaos.h
+++ b/Engine/Source/Runtime/Engine/Public/Physics/Experimental/PhysInterface_Chaos.h
@@ -14,7 +14,6 @@
 #include "Physics/GenericPhysicsInterface.h"
 #include "Physics/Experimental/PhysicsUserData_Chaos.h"
 #include "Chaos/PhysicalMaterials.h"
-#include "PhysicalMaterials/PhysicalMaterialMask.h"
 
 //NOTE: Do not include Chaos headers directly as it means recompiling all of engine. This should be reworked to avoid allocations
 
