Fix CIS.

diff --git a/Engine/Source/Runtime/Engine/Private/PhysicsEngine/BodyInstance.cpp b/Engine/Source/Runtime/Engine/Private/PhysicsEngine/BodyInstance.cpp
--- a/Engine/Source/Runtime/Engine/Private/PhysicsEngine/BodyInstance.cpp
+++ b/Engine/Source/Runtime/Engine/Private/PhysicsEngine/BodyInstance.cpp
@@ -38,6 +38,7 @@
 #include "Chaos/TriangleMeshImplicitObject.h"
 #include "Chaos/MassProperties.h"
 #include "Chaos/Utilities.h"
+#include "Physics/Experimental/ChaosInterfaceUtils.h"
 
 #if WITH_CHAOS
 #include "Chaos/ParticleHandle.h"
@@ -1178,25 +1179,6 @@ bool FInitBodiesHelperBase::CreateShapesAndActors()
 	return true;
 }
 
-#if WITH_CHAOS
-	Chaos::EChaosCollisionTraceFlag ConvertCollisionTraceFlag(ECollisionTraceFlag Flag)
-	{
-		if (Flag == ECollisionTraceFlag::CTF_UseDefault)
-			return Chaos::EChaosCollisionTraceFlag::Chaos_CTF_UseDefault;
-		if (Flag == ECollisionTraceFlag::CTF_UseSimpleAndComplex)
-			return Chaos::EChaosCollisionTraceFlag::Chaos_CTF_UseSimpleAndComplex;
-		if (Flag == ECollisionTraceFlag::CTF_UseSimpleAsComplex)
-			return Chaos::EChaosCollisionTraceFlag::Chaos_CTF_UseSimpleAsComplex;
-		if (Flag == ECollisionTraceFlag::CTF_UseComplexAsSimple)
-			return Chaos::EChaosCollisionTraceFlag::Chaos_CTF_UseComplexAsSimple;
-		if (Flag == ECollisionTraceFlag::CTF_MAX)
-			return Chaos::EChaosCollisionTraceFlag::Chaos_CTF_MAX;
-		ensure(false);
-		return Chaos::EChaosCollisionTraceFlag::Chaos_CTF_UseDefault;
-	}
-#endif
-
-
 void FInitBodiesHelperBase::InitBodies()
 {
 #if WITH_CHAOS
@@ -1253,7 +1235,7 @@ void FInitBodiesHelperBase::InitBodies()
 							{
 								for (int32 ShapeIndex = 0; ShapeIndex < NumShapes; ++ShapeIndex)
 								{
-									ActorHandle->SetShapeCollisionTraceType(ShapeIndex, ConvertCollisionTraceFlag(BI->BodySetup->CollisionTraceFlag)) ;
+									ActorHandle->SetShapeCollisionTraceType(ShapeIndex, ChaosInterface::ConvertCollisionTraceFlag(BI->BodySetup->CollisionTraceFlag)) ;
 								}
 							}
 
diff --git a/Engine/Source/Runtime/Engine/Public/Physics/Experimental/ChaosInterfaceUtils.h b/Engine/Source/Runtime/Engine/Public/Physics/Experimental/ChaosInterfaceUtils.h
--- a/Engine/Source/Runtime/Engine/Public/Physics/Experimental/ChaosInterfaceUtils.h
+++ b/Engine/Source/Runtime/Engine/Public/Physics/Experimental/ChaosInterfaceUtils.h
@@ -10,6 +10,8 @@ struct FGeometryAddParams;
 
 namespace ChaosInterface
 {
+	Chaos::EChaosCollisionTraceFlag ConvertCollisionTraceFlag(ECollisionTraceFlag Flag);
+	
 	/**
 	 * Create the Chaos Geometry based on the geometry parameters.
 	 */
