Fix insert frame leaving raw/source animationdata in invalid state

diff --git a/Engine/Source/Runtime/Engine/Classes/Animation/AnimSequence.h b/Engine/Source/Runtime/Engine/Classes/Animation/AnimSequence.h
--- a/Engine/Source/Runtime/Engine/Classes/Animation/AnimSequence.h
+++ b/Engine/Source/Runtime/Engine/Classes/Animation/AnimSequence.h
@@ -684,6 +684,17 @@ public:
 	 * Return true if it contains transform curves
 	 */
 	bool DoesContainTransformCurves() const;
+
+	/**
+	 * Returns whether this animation has baked transform curves (i.e. has the raw data been modified)
+	 */
+	bool HasBakedTransformCurves() const;
+
+	/**
+	 * Restore the pre baked transform curve raw data
+	 */
+	void RestoreSourceData();
+
 	/**
 	* Return true if compressed data is out of date / missing and so animation needs to use raw data
 	*/
diff --git a/Engine/Source/Runtime/Engine/Private/Animation/AnimSequence.cpp b/Engine/Source/Runtime/Engine/Private/Animation/AnimSequence.cpp
--- a/Engine/Source/Runtime/Engine/Private/Animation/AnimSequence.cpp
+++ b/Engine/Source/Runtime/Engine/Private/Animation/AnimSequence.cpp
@@ -539,6 +539,33 @@ float GetIntervalPerKey(int32 NumFrames, float SequenceLength)
 	return (NumFrames > 1) ? (SequenceLength / (NumFrames-1)) : MINIMUM_ANIMATION_LENGTH;
 }
 
+// Handles keeping source raw data in sync when modifying raw data
+struct FModifyRawDataSourceGuard
+{
+private:
+	UAnimSequence* ModifyingSequence;
+
+public:
+	FModifyRawDataSourceGuard(UAnimSequence* AnimToModify)
+		: ModifyingSequence(nullptr)
+	{
+		check(AnimToModify)
+		if (AnimToModify->HasBakedTransformCurves())
+		{
+			ModifyingSequence = AnimToModify;
+			ModifyingSequence->RestoreSourceData();
+		}
+	}
+
+	~FModifyRawDataSourceGuard()
+	{
+		if (ModifyingSequence)
+		{
+			ModifyingSequence->BakeTrackCurvesToRawAnimation();
+		}
+	}
+};
+
 /////////////////////////////////////////////////////
 // UAnimSequence
 
@@ -1948,6 +1975,8 @@ bool UAnimSequence::InsertFramesToRawAnimData( int32 StartFrame, int32 EndFrame,
 	int32 NumFramesToInsert = EndFrame-StartFrame;
 	if ((CopyFrame>=0 && CopyFrame<NumFrames) && (StartFrame >= 0 && StartFrame <=NumFrames) && NumFramesToInsert > 0)
 	{
+		FModifyRawDataSourceGuard Modify(this);
+
 		for (auto& RawData : RawAnimationData)
 		{
 			if (RawData.PosKeys.Num() > 1 && RawData.PosKeys.IsValidIndex(CopyFrame))
@@ -3654,6 +3683,8 @@ int32 FindFirstChildTrack(const USkeleton* MySkeleton, const FReferenceSkeleton&
 
 int32 UAnimSequence::InsertTrack(const FName& BoneName)
 {
+	FModifyRawDataSourceGuard Modify(this);
+
 	// first verify if it doesn't exists, if it does, return
 	int32 CurrentTrackIndex = AnimationTrackNames.Find(BoneName);
 	if (CurrentTrackIndex != INDEX_NONE)
@@ -4394,6 +4425,20 @@ bool UAnimSequence::DoesContainTransformCurves() const
 	return (RawCurveData.TransformCurves.Num() > 0);
 }
 
+bool  UAnimSequence::HasBakedTransformCurves() const
+{
+	return DoesContainTransformCurves() && SourceRawAnimationData.Num() > 0;
+}
+
+void  UAnimSequence::RestoreSourceData() 
+{
+	if (HasBakedTransformCurves())
+	{
+		RawAnimationData = MoveTemp(SourceRawAnimationData);
+		bNeedsRebake = true;
+	}
+}
+
 void UAnimSequence::AddKeyToSequence(float Time, const FName& BoneName, const FTransform& AdditiveTransform)
 {
 	// if source animation exists, but doesn't match with raw animation number, it's possible this has been retargetted
