Fixup default CVar value, fallout from P4 merge

diff --git a/Engine/Source/Runtime/Engine/Private/Animation/SkinWeightProfileManager.cpp b/Engine/Source/Runtime/Engine/Private/Animation/SkinWeightProfileManager.cpp
--- a/Engine/Source/Runtime/Engine/Private/Animation/SkinWeightProfileManager.cpp
+++ b/Engine/Source/Runtime/Engine/Private/Animation/SkinWeightProfileManager.cpp
@@ -21,7 +21,7 @@ DEFINE_LOG_CATEGORY_STATIC(LogSkinWeightProfileManager, Warning, Display);
 
 #define PROFILE_SKIN_WEIGHT_MANAGER (!UE_BUILD_SHIPPING)
 
-static int32 GAllowCPU = 0;
+static int32 GAllowCPU = 1;
 static FAutoConsoleVariableRef CVarAllowCPU (
 	TEXT("SkinWeightProfileManager.AllowCPU"),
 	GAllowCPU,
