addition of more detailed stats for physical and virtual memory

diff --git a/Engine/Source/Runtime/Engine/Private/ChartCreation.cpp b/Engine/Source/Runtime/Engine/Private/ChartCreation.cpp
--- a/Engine/Source/Runtime/Engine/Private/ChartCreation.cpp
+++ b/Engine/Source/Runtime/Engine/Private/ChartCreation.cpp
@@ -632,6 +632,12 @@ void FPerformanceTrackingChart::Reset(const FDateTime& InStartTime)
 	AccumulatedChartTime = 0.0;
 	TimeDisregarded = 0.0;
 	FramesDisregarded = 0;
+	MaxPhysicalMemory = 0;
+	MaxVirtualMemory = 0;
+	MinPhysicalMemory = ULONG_MAX;
+	MinVirtualMemory = ULONG_MAX;
+	TotalPhysicalMemoryUsed = 0;
+	TotalVirtualMemoryUsed = 0;
 
 	FrametimeHistogram.InitFromArray({ 0.0/1000.0, 15.0/1000.0, 20.0/1000.0, 30.0/1000.0, 35.0/1000.0, 60.0/1000.0, 100.0/1000.0, 1000.0/1000.0 });
 
@@ -706,6 +712,12 @@ void FPerformanceTrackingChart::AccumulateWith(const FPerformanceTrackingChart&
 	TimeDisregarded += Chart.TimeDisregarded;
 	FramesDisregarded += Chart.FramesDisregarded;
 	CaptureStartTime = FMath::Min(CaptureStartTime, Chart.CaptureStartTime);
+	MaxPhysicalMemory = FMath::Max(MaxPhysicalMemory, Chart.MaxPhysicalMemory);
+	MaxVirtualMemory = FMath::Min(MaxVirtualMemory, Chart.MaxVirtualMemory);
+	MinPhysicalMemory = FMath::Min(MinPhysicalMemory, Chart.MinPhysicalMemory);
+	MinVirtualMemory = FMath::Min(MinVirtualMemory, Chart.MinVirtualMemory);
+	TotalPhysicalMemoryUsed += Chart.TotalPhysicalMemoryUsed;
+	TotalVirtualMemoryUsed += Chart.TotalVirtualMemoryUsed;
 }
 
 void FPerformanceTrackingChart::StartCharting()
@@ -799,6 +811,15 @@ void FPerformanceTrackingChart::ProcessFrame(const FFrameData& FrameData)
 		MinDrawnPrimitives = FMath::Min(MinDrawnPrimitives, GNumPrimitivesDrawnRHI);
 		TotalDrawnPrimitives += GNumPrimitivesDrawnRHI;
 
+		// track memory
+		FPlatformMemoryStats MemoryStats = FPlatformMemory::GetStats();
+		MaxPhysicalMemory = FMath::Max(MaxPhysicalMemory, static_cast<uint64>(MemoryStats.UsedPhysical));
+		MaxVirtualMemory = FMath::Max(MaxVirtualMemory, static_cast<uint64>(MemoryStats.UsedVirtual));
+		MinPhysicalMemory = FMath::Min(MinPhysicalMemory, static_cast<uint64>(MemoryStats.UsedPhysical));
+		MinVirtualMemory = FMath::Min(MinVirtualMemory, static_cast<uint64>(MemoryStats.UsedVirtual));
+		TotalPhysicalMemoryUsed += MemoryStats.UsedPhysical;
+		TotalVirtualMemoryUsed += MemoryStats.UsedVirtual;
+
 		// Handle hitching
 		if (FrameData.HitchStatus != EFrameHitchType::NoHitch)
 		{
diff --git a/Engine/Source/Runtime/Engine/Private/ProfilingDebugging/HealthSnapshot.cpp b/Engine/Source/Runtime/Engine/Private/ProfilingDebugging/HealthSnapshot.cpp
--- a/Engine/Source/Runtime/Engine/Private/ProfilingDebugging/HealthSnapshot.cpp
+++ b/Engine/Source/Runtime/Engine/Private/ProfilingDebugging/HealthSnapshot.cpp
@@ -126,6 +126,16 @@ void FHealthSnapshot::CapturePerformanceStats(const FPerformanceTrackingChart* G
 		DynamicResolution.Max = GameplayFPSChart->DynamicResHistogram.GetMaxOfAllMeasures();
 		DynamicResolution.Avg = GameplayFPSChart->DynamicResHistogram.GetAverageOfAllMeasures();
 
+		PhysicalMemory.Max = GameplayFPSChart->MaxPhysicalMemory;
+		PhysicalMemory.Min = GameplayFPSChart->MinPhysicalMemory;
+		VirtualMemory.Max = GameplayFPSChart->MaxVirtualMemory;
+		VirtualMemory.Min = GameplayFPSChart->MinVirtualMemory;
+		if (FramesCounted > 0)
+		{
+			PhysicalMemory.Avg = GameplayFPSChart->TotalPhysicalMemoryUsed / static_cast<uint64>(FramesCounted);
+			VirtualMemory.Avg = GameplayFPSChart->TotalVirtualMemoryUsed / static_cast<uint64>(FramesCounted);
+		}
+
 		// What % of frames were bound
 		GameThread.PercentFramesBound = FramesCounted > 0 ? (GameplayFPSChart->NumFramesBound_GameThread * 100.0) / FramesCounted : 0;
 
@@ -181,6 +191,8 @@ void FHealthSnapshot::DumpStats(FOutputDevice& Ar, FName CategoryName)
 		Ar.CategorizedLogf(CategoryName, ELogVerbosity::Log, TEXT("GPU: Avg %.02fms, Hitches/Min: %.02f, Bound Frames: %.02f%%"), GPU.AvgTime * 1000, GPU.HitchesPerMinute, GPU.PercentFramesBound);
 		Ar.CategorizedLogf(CategoryName, ELogVerbosity::Log, TEXT("DrawCalls: Avg: %d, Max: %d, Min: %d"), DrawCalls.Avg, DrawCalls.Max, DrawCalls.Min);
 		Ar.CategorizedLogf(CategoryName, ELogVerbosity::Log, TEXT("DrawnPrims: Avg: %d, Max: %d, Min: %d"), PrimitivesDrawn.Avg, PrimitivesDrawn.Max, PrimitivesDrawn.Min);
+		Ar.CategorizedLogf(CategoryName, ELogVerbosity::Log, TEXT("PhysicalMemoryMB: Avg: %d, Max: %d, Min: %d"), PhysicalMemory.Avg, PhysicalMemory.Max, PhysicalMemory.Min);
+		Ar.CategorizedLogf(CategoryName, ELogVerbosity::Log, TEXT("VirtualMemoryMB: Avg: %d, Max: %d, Min: %d"), VirtualMemory.Avg, VirtualMemory.Max, VirtualMemory.Min);
 	}
 
 	Ar.CategorizedLogf(CategoryName, ELogVerbosity::Log, TEXT("CPU Memory: Used %.2fMB, Peak %.2fMB"), CPUMemoryMB.Used, CPUMemoryMB.Peak);
diff --git a/Engine/Source/Runtime/Engine/Public/ChartCreation.h b/Engine/Source/Runtime/Engine/Public/ChartCreation.h
--- a/Engine/Source/Runtime/Engine/Public/ChartCreation.h
+++ b/Engine/Source/Runtime/Engine/Public/ChartCreation.h
@@ -216,6 +216,14 @@ public:
 	FDelegateHandle DeviceProfilesUpdatedDelegateHandle;
 
 	bool bIsChartingPaused;
+	
+	// memory stats
+	uint64 MaxPhysicalMemory;
+	uint64 MaxVirtualMemory;
+	uint64 MinPhysicalMemory;
+	uint64 MinVirtualMemory;
+	uint64 TotalPhysicalMemoryUsed;
+	uint64 TotalVirtualMemoryUsed;
 
 public:
 	FPerformanceTrackingChart();
diff --git a/Engine/Source/Runtime/Engine/Public/ProfilingDebugging/HealthSnapshot.h b/Engine/Source/Runtime/Engine/Public/ProfilingDebugging/HealthSnapshot.h
--- a/Engine/Source/Runtime/Engine/Public/ProfilingDebugging/HealthSnapshot.h
+++ b/Engine/Source/Runtime/Engine/Public/ProfilingDebugging/HealthSnapshot.h
@@ -145,6 +145,8 @@ public:
 	FMemoryStat<float> OnionMemoryMB;
 #endif //PLATFORM_PS4
 	float LLMTotalMemoryMB;
+	FMmaStat<uint64> PhysicalMemory; // more detailed physical memory stats obtained more frequently
+	FMmaStat<uint64> VirtualMemory;
 
 	/** Performance data */
 	double MeasuredPerfTime; // Duration of time the following performance values came from
