Add missing #if WITH_EDITOR

diff --git a/Engine/Source/Runtime/Engine/Private/Animation/AnimSequence.cpp b/Engine/Source/Runtime/Engine/Private/Animation/AnimSequence.cpp
--- a/Engine/Source/Runtime/Engine/Private/Animation/AnimSequence.cpp
+++ b/Engine/Source/Runtime/Engine/Private/Animation/AnimSequence.cpp
@@ -539,6 +539,7 @@ float GetIntervalPerKey(int32 NumFrames, float SequenceLength)
 	return (NumFrames > 1) ? (SequenceLength / (NumFrames-1)) : MINIMUM_ANIMATION_LENGTH;
 }
 
+#if WITH_EDITOR
 // Handles keeping source raw data in sync when modifying raw data
 struct FModifyRawDataSourceGuard
 {
@@ -565,6 +566,7 @@ public:
 		}
 	}
 };
+#endif
 
 /////////////////////////////////////////////////////
 // UAnimSequence
@@ -1975,7 +1977,9 @@ bool UAnimSequence::InsertFramesToRawAnimData( int32 StartFrame, int32 EndFrame,
 	int32 NumFramesToInsert = EndFrame-StartFrame;
 	if ((CopyFrame>=0 && CopyFrame<NumFrames) && (StartFrame >= 0 && StartFrame <=NumFrames) && NumFramesToInsert > 0)
 	{
+#if WITH_EDITOR
 		FModifyRawDataSourceGuard Modify(this);
+#endif
 
 		for (auto& RawData : RawAnimationData)
 		{
@@ -3683,7 +3687,9 @@ int32 FindFirstChildTrack(const USkeleton* MySkeleton, const FReferenceSkeleton&
 
 int32 UAnimSequence::InsertTrack(const FName& BoneName)
 {
+#if WITH_EDITOR
 	FModifyRawDataSourceGuard Modify(this);
+#endif
 
 	// first verify if it doesn't exists, if it does, return
 	int32 CurrentTrackIndex = AnimationTrackNames.Find(BoneName);
@@ -4425,6 +4431,7 @@ bool UAnimSequence::DoesContainTransformCurves() const
 	return (RawCurveData.TransformCurves.Num() > 0);
 }
 
+#if WITH_EDITOR
 bool  UAnimSequence::HasBakedTransformCurves() const
 {
 	return DoesContainTransformCurves() && SourceRawAnimationData.Num() > 0;
@@ -4438,6 +4445,7 @@ void  UAnimSequence::RestoreSourceData()
 		bNeedsRebake = true;
 	}
 }
+#endif
 
 void UAnimSequence::AddKeyToSequence(float Time, const FName& BoneName, const FTransform& AdditiveTransform)
 {
