[NO NON-PII COMMIT MESSAGE]

diff --git a/Engine/Source/Runtime/AudioMixer/Private/AudioMixerDevice.cpp b/Engine/Source/Runtime/AudioMixer/Private/AudioMixerDevice.cpp
--- a/Engine/Source/Runtime/AudioMixer/Private/AudioMixerDevice.cpp
+++ b/Engine/Source/Runtime/AudioMixer/Private/AudioMixerDevice.cpp
@@ -2065,7 +2065,7 @@ namespace Audio
 
 	bool FMixerDevice::IsMainAudioDevice() const
 	{
-		bool bIsMain = (this == GEngine->GetMainAudioDeviceRaw());
+		bool bIsMain = (this == FAudioDeviceManager::Get()->GetMainAudioDeviceRaw());
 		return bIsMain;
 	}
 
diff --git a/Engine/Source/Runtime/Engine/Classes/Engine/Engine.h b/Engine/Source/Runtime/Engine/Classes/Engine/Engine.h
--- a/Engine/Source/Runtime/Engine/Classes/Engine/Engine.h
+++ b/Engine/Source/Runtime/Engine/Classes/Engine/Engine.h
@@ -1559,11 +1559,6 @@ public:
 	UPROPERTY(globalconfig)
 	uint32 bEnableVisualLogRecordingOnStart;
 
-protected:
-
-	/** Whether the engine should be playing sounds.  If false at initialization time the AudioDevice will not be created */
-	uint32 bUseSound:1;
-
 private:
 	
 	/** Semaphore to control screen saver inhibitor thread access. */
@@ -1997,7 +1992,7 @@ public:
 
 	/** Called at shutdown, just before the exit purge.	 */
 	virtual void PreExit();
-	virtual void ShutdownAudioDeviceManager();
+	virtual void ReleaseAudioDeviceManager();
 	
 	void ShutdownHMD();
 
diff --git a/Engine/Source/Runtime/Engine/Private/AudioDevice.cpp b/Engine/Source/Runtime/Engine/Private/AudioDevice.cpp
--- a/Engine/Source/Runtime/Engine/Private/AudioDevice.cpp
+++ b/Engine/Source/Runtime/Engine/Private/AudioDevice.cpp
@@ -6344,7 +6344,7 @@ bool FAudioDevice::CanUseVRAudioDevice()
 	if (GIsEditor)
 	{
 		UEditorEngine* EdEngine = Cast<UEditorEngine>(GEngine);
-		return EdEngine->IsVRPreviewActive();
+		return EdEngine ? EdEngine->IsVRPreviewActive() : false;
 	}
 	else
 #endif
diff --git a/Engine/Source/Runtime/Engine/Private/AudioDeviceManager.cpp b/Engine/Source/Runtime/Engine/Private/AudioDeviceManager.cpp
--- a/Engine/Source/Runtime/Engine/Private/AudioDeviceManager.cpp
+++ b/Engine/Source/Runtime/Engine/Private/AudioDeviceManager.cpp
@@ -76,6 +76,8 @@ static FAutoConsoleCommand GReportAudioDevicesCommand(
 	})
 );
 
+FAudioDeviceManager* FAudioDeviceManager::Singleton = nullptr;
+
 // Some stress tests:
 #if INSTRUMENT_AUDIODEVICE_HANDLES
 static TArray<FAudioDeviceHandle> IntentionallyLeakedHandles;
@@ -146,6 +148,11 @@ FAudioDeviceManager::FAudioDeviceManager()
 
 FAudioDeviceManager::~FAudioDeviceManager()
 {
+	FAudioCommandFence Fence;
+	Fence.BeginFence();
+	Fence.Wait();
+	FAudioThread::StopAudioThread();
+
 	MainAudioDeviceHandle.Reset();
 
 	// Notify anyone listening to the device manager that we are about to destroy the audio device.
@@ -390,7 +397,7 @@ FAudioDeviceHandle FAudioDeviceManager::RequestAudioDevice(const FAudioDevicePar
 	}
 }
 
-bool FAudioDeviceManager::Initialize()
+bool FAudioDeviceManager::InitializeManager()
 {
 	if (LoadDefaultAudioDeviceModule())
 	{
@@ -639,14 +646,51 @@ FAudioDevice* FAudioDeviceManager::GetAudioDeviceRaw(Audio::FDeviceId Handle)
 	return AudioDevice;
 }
 
-FAudioDeviceManager* FAudioDeviceManager::Get()
+bool FAudioDeviceManager::Initialize()
 {
-	if (GEngine)
+	if (!Singleton)
 	{
-		return GEngine->GetAudioDeviceManager();
+		bool bUseSound = true;
+
+		// If dedicated server, the -nosound, or -benchmark parameters are used, disable sound.
+		if (FParse::Param(FCommandLine::Get(), TEXT("nosound")) || FApp::IsBenchmarking() || IsRunningDedicatedServer() || (IsRunningCommandlet() && !IsAllowCommandletAudio()))
+		{
+			bUseSound = false;
+		}
+
+		if (FParse::Param(FCommandLine::Get(), TEXT("enablesound")))
+		{
+			bUseSound = true;
+		}
+
+		if (bUseSound)
+		{
+			Singleton = new FAudioDeviceManager();
+
+			if (!Singleton->InitializeManager())
+			{
+				UE_LOG(LogAudio, Warning, TEXT("Initializing the Audio Device Manager Failed!"));
+				delete Singleton;
+				Singleton = nullptr;
+			}
+		}
 	}
 
-	return nullptr;
+	return Singleton != nullptr;
+}
+
+FAudioDeviceManager* FAudioDeviceManager::Get()
+{
+	return Singleton;
+}
+
+void FAudioDeviceManager::Shutdown()
+{
+	if (Singleton)
+	{
+		delete Singleton;
+		Singleton = nullptr;
+	}
 }
 
 FAudioDeviceHandle FAudioDeviceManager::GetActiveAudioDevice()
@@ -1292,7 +1336,7 @@ FAudioDeviceHandle::~FAudioDeviceHandle()
 {
 	if (IsValid())
 	{
-		FAudioDeviceManager* AudioDeviceManager = GEngine->GetAudioDeviceManager();
+		FAudioDeviceManager* AudioDeviceManager = FAudioDeviceManager::Get();
 		if (AudioDeviceManager)
 		{
 			AudioDeviceManager->DecrementDevice(DeviceId, World);
@@ -1317,7 +1361,7 @@ Audio::FDeviceId FAudioDeviceHandle::GetDeviceID() const
 
 bool FAudioDeviceHandle::IsValid() const
 {
-	return GEngine && GEngine->GetAudioDeviceManager() && Device != nullptr;
+	return Device != nullptr;
 }
 
 void FAudioDeviceHandle::Reset()
diff --git a/Engine/Source/Runtime/Engine/Private/GameEngine.cpp b/Engine/Source/Runtime/Engine/Private/GameEngine.cpp
--- a/Engine/Source/Runtime/Engine/Private/GameEngine.cpp
+++ b/Engine/Source/Runtime/Engine/Private/GameEngine.cpp
@@ -710,8 +710,6 @@ UEngine::UEngine(const FObjectInitializer& ObjectInitializer)
 
 	SelectionHighlightIntensityBillboards = 0.25f;
 
-	bUseSound = true;
-
 	bHardwareSurveyEnabled_DEPRECATED = false;
 	bIsInitialized = false;
 
diff --git a/Engine/Source/Runtime/Engine/Private/UnrealEngine.cpp b/Engine/Source/Runtime/Engine/Private/UnrealEngine.cpp
--- a/Engine/Source/Runtime/Engine/Private/UnrealEngine.cpp
+++ b/Engine/Source/Runtime/Engine/Private/UnrealEngine.cpp
@@ -1777,21 +1777,11 @@ void UEngine::OnExternalUIChange(bool bInIsOpening)
 	FSlateApplication::Get().ExternalUIChange(bInIsOpening);
 }
 
-void UEngine::ShutdownAudioDeviceManager()
+void UEngine::ReleaseAudioDeviceManager()
 {
-	// Shutdown the main audio device in the UEEngine
-	if (AudioDeviceManager)
-	{
-		FAudioCommandFence Fence;
-		Fence.BeginFence();
-		Fence.Wait();
-
-		FAudioThread::StopAudioThread();
-
-		MainAudioDeviceHandle.Reset();
-		delete AudioDeviceManager;
-		AudioDeviceManager = NULL;
-	}
+	// Release the main audio device handle.
+	MainAudioDeviceHandle.Reset();
+	AudioDeviceManager = nullptr;
 }
 
 void UEngine::PreExit()
@@ -2285,17 +2275,6 @@ void UEngine::UpdateTimecode()
 
 void UEngine::ParseCommandline()
 {
-	// If dedicated server, the -nosound, or -benchmark parameters are used, disable sound.
-	if(FParse::Param(FCommandLine::Get(),TEXT("nosound")) || FApp::IsBenchmarking() || IsRunningDedicatedServer() || (IsRunningCommandlet() && !IsAllowCommandletAudio()))
-	{
-		bUseSound = false;
-	}
-
-	if (FParse::Param(FCommandLine::Get(), TEXT("enablesound")))
-	{
-		bUseSound = true;
-	}
-
 	if( FParse::Param( FCommandLine::Get(), TEXT("noailogging")) )
 	{
 		bDisableAILogging = true;
@@ -2772,7 +2751,7 @@ void UEngine::FinishDestroy()
 	{
 		// shut down all subsystems.
 		GEngine = NULL;
-		ShutdownAudioDeviceManager();
+		ReleaseAudioDeviceManager();
 
 		FURL::StaticExit();
 	}
@@ -2929,23 +2908,19 @@ FAudioDeviceHandle UEngine::GetActiveAudioDevice()
 
 void UEngine::InitializeAudioDeviceManager()
 {
-	if (AudioDeviceManager == nullptr && bUseSound)
+	FAudioDeviceManager::Initialize();
+
+	AudioDeviceManager = FAudioDeviceManager::Get();
+
+	if (AudioDeviceManager)
 	{
-		AudioDeviceManager = new FAudioDeviceManager();
-		if (AudioDeviceManager->Initialize())
-		{
-			MainAudioDeviceHandle = AudioDeviceManager->GetMainAudioDeviceHandle();
-		}
-		else
-		{
-			ShutdownAudioDeviceManager();
-		}
+		MainAudioDeviceHandle = AudioDeviceManager->GetMainAudioDeviceHandle();
 	}
 }
 
 bool UEngine::UseSound() const
 {
-	return (bUseSound && AudioDeviceManager != nullptr);
+	return AudioDeviceManager != nullptr;
 }
 /**
 * A fake stereo rendering device used to test stereo rendering without an attached device.
diff --git a/Engine/Source/Runtime/Engine/Public/AudioDeviceManager.h b/Engine/Source/Runtime/Engine/Public/AudioDeviceManager.h
--- a/Engine/Source/Runtime/Engine/Public/AudioDeviceManager.h
+++ b/Engine/Source/Runtime/Engine/Public/AudioDeviceManager.h
@@ -218,14 +218,9 @@ public:
 	*/
 	~FAudioDeviceManager();
 
-	/**
-	* Initialize the audio device manager.
-	* Return true if successfully initialized.
-	**/
-	bool Initialize();
-
 	/** Returns the handle to the main audio device. */
 	FAudioDeviceHandle GetMainAudioDeviceHandle() const { return MainAudioDeviceHandle; }
+	FAudioDevice* GetMainAudioDeviceRaw() const { return MainAudioDeviceHandle.GetAudioDevice(); }
 	Audio::FDeviceId GetMainAudioDeviceID() const { return MainAudioDeviceHandle.GetDeviceID(); }
 
 	/** Returns true if we're currently using the audio mixer. */
@@ -264,7 +259,13 @@ public:
 	*/
 	FAudioDevice* GetAudioDeviceRaw(Audio::FDeviceId DeviceID);
 
+	/**
+	* Initialize the audio device manager.
+	* Return true if successfully initialized.
+	**/
+	static bool Initialize();
 	static FAudioDeviceManager* Get();
+	static void Shutdown();
 
 	/**
 	* Returns a ptr to the active audio device. If there is no active
@@ -397,6 +398,7 @@ private:
 	TUniquePtr<FAudioDebugger> AudioDebugger;
 #endif // ENABLE_AUDIO_DEBUG
 
+	bool InitializeManager();
 
 	/** Creates a handle given the index and a generation value. */
 	uint32 GetNewDeviceID();
@@ -530,4 +532,6 @@ private:
 	FAudioCommandFence SyncFence;
 
 	friend class FAudioDeviceHandle;
+
+	static FAudioDeviceManager* Singleton;
 };
diff --git a/Engine/Source/Runtime/Launch/Private/LaunchEngineLoop.cpp b/Engine/Source/Runtime/Launch/Private/LaunchEngineLoop.cpp
--- a/Engine/Source/Runtime/Launch/Private/LaunchEngineLoop.cpp
+++ b/Engine/Source/Runtime/Launch/Private/LaunchEngineLoop.cpp
@@ -4223,7 +4223,7 @@ void FEngineLoop::Exit()
 
 	if ( GEngine != nullptr )
 	{
-		GEngine->ShutdownAudioDeviceManager();
+		GEngine->ReleaseAudioDeviceManager();
 	}
 
 	if ( GEngine != nullptr )
@@ -4231,6 +4231,8 @@ void FEngineLoop::Exit()
 		GEngine->PreExit();
 	}
 
+	FAudioDeviceManager::Shutdown();
+
 	// close all windows
 	FSlateApplication::Shutdown();
 
