Fix NonUnity.

diff --git a/Engine/Source/Runtime/Engine/Public/Physics/Experimental/ChaosInterfaceUtils.h b/Engine/Source/Runtime/Engine/Public/Physics/Experimental/ChaosInterfaceUtils.h
--- a/Engine/Source/Runtime/Engine/Public/Physics/Experimental/ChaosInterfaceUtils.h
+++ b/Engine/Source/Runtime/Engine/Public/Physics/Experimental/ChaosInterfaceUtils.h
@@ -5,9 +5,15 @@
 #include "Containers/ContainersFwd.h"
 #include "Containers/ContainerAllocationPolicies.h"
 #include "Templates/UniquePtr.h"
+#include "PhysicsEngine/BodySetupEnums.h"
 
 struct FGeometryAddParams;
 
+namespace Chaos
+{
+	enum EChaosCollisionTraceFlag;
+}
+
 namespace ChaosInterface
 {
 	Chaos::EChaosCollisionTraceFlag ConvertCollisionTraceFlag(ECollisionTraceFlag Flag);
