Chaos - RBAN support for "contribute to mass" flag in physics asset

diff --git a/Engine/Source/Runtime/Engine/Private/PhysicsEngine/ImmediatePhysics/ImmediatePhysicsChaos/ImmediatePhysicsActorHandle_Chaos.cpp b/Engine/Source/Runtime/Engine/Private/PhysicsEngine/ImmediatePhysics/ImmediatePhysicsChaos/ImmediatePhysicsActorHandle_Chaos.cpp
--- a/Engine/Source/Runtime/Engine/Private/PhysicsEngine/ImmediatePhysics/ImmediatePhysicsChaos/ImmediatePhysicsActorHandle_Chaos.cpp
+++ b/Engine/Source/Runtime/Engine/Private/PhysicsEngine/ImmediatePhysics/ImmediatePhysicsChaos/ImmediatePhysicsActorHandle_Chaos.cpp
@@ -66,56 +66,65 @@ namespace ImmediatePhysics_Chaos
 
 		for (uint32 i = 0; i < static_cast<uint32>(AggGeom.SphereElems.Num()); ++i)
 		{
-			const FKSphereElem ScaledSphereElem = AggGeom.SphereElems[i].GetFinalScaled(Scale, LocalTransform);
+			const FKSphereElem& SphereElem = AggGeom.SphereElems[i];
+			if (SphereElem.GetContributeToMass())
+			{
+				const FKSphereElem ScaledSphereElem = SphereElem.GetFinalScaled(Scale, LocalTransform);
 
-			TMassProperties<FReal, 3> MassProperties;
-			MassProperties.CenterOfMass = LocalTransform.GetTranslation() + ScaledSphereElem.Center;
-			MassProperties.RotationOfMass = FRotation3::FromIdentity();
-			MassProperties.Volume = ScaledSphereElem.GetVolume(FVector::OneVector);
-			MassProperties.InertiaTensor = CalculateInertia_Solid(MassProperties.Volume, ScaledSphereElem);
+				TMassProperties<FReal, 3> MassProperties;
+				MassProperties.CenterOfMass = LocalTransform.GetTranslation() + ScaledSphereElem.Center;
+				MassProperties.RotationOfMass = FRotation3::FromIdentity();
+				MassProperties.Volume = ScaledSphereElem.GetVolume(FVector::OneVector);
+				MassProperties.InertiaTensor = CalculateInertia_Solid(MassProperties.Volume, ScaledSphereElem);
 
-			AllMassProperties.Add(MassProperties);
+				AllMassProperties.Add(MassProperties);
+			}
 		}
 		for (uint32 i = 0; i < static_cast<uint32>(AggGeom.BoxElems.Num()); ++i)
 		{
 			const auto& BoxElem = AggGeom.BoxElems[i];
+			if (BoxElem.GetContributeToMass())
+			{
+				TMassProperties<FReal, 3> MassProperties;
+				MassProperties.CenterOfMass = LocalTransform.GetTranslation() + BoxElem.Center;
+				MassProperties.RotationOfMass = LocalTransform.GetRotation() * FRotation3(FQuat(BoxElem.Rotation));
+				MassProperties.Volume = BoxElem.GetVolume(Scale);
+				MassProperties.InertiaTensor = CalculateInertia_Solid(MassProperties.Volume, BoxElem);
 
-			TMassProperties<FReal, 3> MassProperties;
-			MassProperties.CenterOfMass = LocalTransform.GetTranslation() + BoxElem.Center;
-			MassProperties.RotationOfMass = LocalTransform.GetRotation() * FRotation3(FQuat(BoxElem.Rotation));
-			MassProperties.Volume = BoxElem.GetVolume(Scale);
-			MassProperties.InertiaTensor = CalculateInertia_Solid(MassProperties.Volume, BoxElem);
-
-			AllMassProperties.Add(MassProperties);
+				AllMassProperties.Add(MassProperties);
+			}
 		}
 		for (uint32 i = 0; i < static_cast<uint32>(AggGeom.SphylElems.Num()); ++i)
 		{
 			const FKSphylElem& UnscaledSphyl = AggGeom.SphylElems[i];
-			const FKSphylElem ScaledSphylElem = UnscaledSphyl.GetFinalScaled(Scale, LocalTransform);
-			float HalfHeight = FMath::Max(ScaledSphylElem.Length * 0.5f, KINDA_SMALL_NUMBER);
-			const float Radius = FMath::Max(ScaledSphylElem.Radius, KINDA_SMALL_NUMBER);
-			if (HalfHeight < KINDA_SMALL_NUMBER)
+			if (UnscaledSphyl.GetContributeToMass())
 			{
-				//not a capsule just use a sphere
-				const FKSphereElem ScaledSphereElem = FKSphereElem(Radius);
+				const FKSphylElem ScaledSphylElem = UnscaledSphyl.GetFinalScaled(Scale, LocalTransform);
+				float HalfHeight = FMath::Max(ScaledSphylElem.Length * 0.5f, KINDA_SMALL_NUMBER);
+				const float Radius = FMath::Max(ScaledSphylElem.Radius, KINDA_SMALL_NUMBER);
+				if (HalfHeight < KINDA_SMALL_NUMBER)
+				{
+					//not a capsule just use a sphere
+					const FKSphereElem ScaledSphereElem = FKSphereElem(Radius);
 
-				TMassProperties<FReal, 3> MassProperties;
-				MassProperties.CenterOfMass = LocalTransform.GetTranslation() + ScaledSphereElem.Center;
-				MassProperties.RotationOfMass = FRotation3::FromIdentity();
-				MassProperties.Volume = ScaledSphereElem.GetVolume(FVector::OneVector);
-				MassProperties.InertiaTensor = CalculateInertia_Solid(MassProperties.Volume, ScaledSphereElem);
+					TMassProperties<FReal, 3> MassProperties;
+					MassProperties.CenterOfMass = LocalTransform.GetTranslation() + ScaledSphereElem.Center;
+					MassProperties.RotationOfMass = FRotation3::FromIdentity();
+					MassProperties.Volume = ScaledSphereElem.GetVolume(FVector::OneVector);
+					MassProperties.InertiaTensor = CalculateInertia_Solid(MassProperties.Volume, ScaledSphereElem);
 
-				AllMassProperties.Add(MassProperties);
-			}
-			else
-			{
-				TMassProperties<FReal, 3> MassProperties;
-				MassProperties.CenterOfMass = LocalTransform.GetTranslation() + ScaledSphylElem.Center;
-				MassProperties.RotationOfMass = LocalTransform.GetRotation() * FRotation3(FQuat(ScaledSphylElem.Rotation));
-				MassProperties.Volume = ScaledSphylElem.GetVolume(FVector::OneVector);
-				MassProperties.InertiaTensor = CalculateInertia_Solid(MassProperties.Volume, ScaledSphylElem);
+					AllMassProperties.Add(MassProperties);
+				}
+				else
+				{
+					TMassProperties<FReal, 3> MassProperties;
+					MassProperties.CenterOfMass = LocalTransform.GetTranslation() + ScaledSphylElem.Center;
+					MassProperties.RotationOfMass = LocalTransform.GetRotation() * FRotation3(FQuat(ScaledSphylElem.Rotation));
+					MassProperties.Volume = ScaledSphylElem.GetVolume(FVector::OneVector);
+					MassProperties.InertiaTensor = CalculateInertia_Solid(MassProperties.Volume, ScaledSphylElem);
 
-				AllMassProperties.Add(MassProperties);
+					AllMassProperties.Add(MassProperties);
+				}
 			}
 		}
 #if WITH_CHAOS && !PHYSICS_INTERFACE_PHYSX
