Back out changelist 12450069. Possible cause of EdTest ensure

diff --git a/Engine/Source/Runtime/Engine/Private/AssetManager.cpp b/Engine/Source/Runtime/Engine/Private/AssetManager.cpp
--- a/Engine/Source/Runtime/Engine/Private/AssetManager.cpp
+++ b/Engine/Source/Runtime/Engine/Private/AssetManager.cpp
@@ -1205,30 +1205,8 @@ TSharedPtr<FStreamableHandle> UAssetManager::ChangeBundleStateForPrimaryAssets(c
 					ExistingHandles.Add(NameData->PendingState.Handle);
 					continue;
 				}
-
-				bool bIsSuperSet = false;
-				if (NewBundles->Num() > NameData->PendingState.BundleNames.Num())
-				{
-					bIsSuperSet = true;
-					for (int i = 0; i < NameData->PendingState.BundleNames.Num(); i++)
-					{
-						if (!NewBundles->Contains(NameData->PendingState.BundleNames[i]))
-						{
-							bIsSuperSet = false;
-							break;
-						}
-					}
-				}
-
-				if (bIsSuperSet)
-				{
-					ExistingHandles.Add(NameData->PendingState.Handle);
-				}
-				else
-				{
-					// Clear pending state
-					NameData->PendingState.Reset(true);
-				}
+				// Clear pending state
+				NameData->PendingState.Reset(true);
 			}
 			else if (NameData->CurrentState.IsValid() && NameData->CurrentState.BundleNames == *NewBundles)
 			{
