Add ability to scale volume for sound set to a certain concurrency group by distance or priority as well as elapsed play time.

diff --git a/Engine/Source/Runtime/Engine/Classes/Sound/SoundConcurrency.h b/Engine/Source/Runtime/Engine/Classes/Sound/SoundConcurrency.h
--- a/Engine/Source/Runtime/Engine/Classes/Sound/SoundConcurrency.h
+++ b/Engine/Source/Runtime/Engine/Classes/Sound/SoundConcurrency.h
@@ -55,6 +55,18 @@ namespace EMaxConcurrentResolutionRule
 	};
 }
 
+UENUM()
+enum class EConcurrencyVolumeScaleMode
+{
+	/* Scales volume of older sounds more than newer sounds (default) */
+	Default = 0,
+
+	/* Scales distant sounds by volume scalar more than closer sounds */
+	Distance,
+
+	/* Scales lower priority sounds by volume scalar more than closer sounds */
+	Priority
+};
 
 USTRUCT(BlueprintType)
 struct ENGINE_API FSoundConcurrencySettings
@@ -79,7 +91,7 @@ struct ENGINE_API FSoundConcurrencySettings
 
 private:
 	/**
-	 * Ducking factor to apply per older voice instance (generation), which compounds as new voices play
+	 * Ducking factor to apply per older voice instance (generation), which compounds based on scaling mode
 	 * and (optionally) revives them as they stop according to the provided attack/release times.
 	 *
 	 * AppliedVolumeScale = Math.Pow(DuckingScale, VoiceGeneration)
@@ -88,22 +100,26 @@ private:
 	float VolumeScale;
 
 public:
+	/** Volume Scale mode designating how to scale voice volume based on number of member sounds active in group. */
+	UPROPERTY(EditAnywhere, BlueprintReadWrite, Category = "Volume Scaling")
+	EConcurrencyVolumeScaleMode VolumeScaleMode;
+
 	/**
 	 * Time taken to apply duck using volume scalar.
 	 */
-	UPROPERTY(EditAnywhere, BlueprintReadWrite, Category = "Volume Scaling", meta = (DisplayName = "Attack Time", UIMin = "0.0", ClampMin = "0.0", UIMax = "10.0", ClampMax = "1000000.0"))
+	UPROPERTY(EditAnywhere, BlueprintReadWrite, Category = "Volume Scaling", meta = (DisplayName = "Duck Time", UIMin = "0.0", ClampMin = "0.0", UIMax = "10.0", ClampMax = "1000000.0"))
 	float VolumeScaleAttackTime;
 
 	/**
-	 * Whether or not volume scaling can release (i.e. "recover") volume ducking behavior when concurrency group sounds stop.
+	 * Whether or not volume scaling can recover volume ducking behavior when concurrency group sounds stop (default scale mode only).
 	 */
-	UPROPERTY(EditAnywhere, BlueprintReadWrite, Category = "Volume Scaling", meta = (DisplayName = "Can Release"))
+	UPROPERTY(EditAnywhere, BlueprintReadWrite, Category = "Volume Scaling", meta = (DisplayName = "Can Recover", EditCondition = "VolumeScaleMode == EConcurrencyVolumeScaleMode::Default"))
 	uint32 bVolumeScaleCanRelease:1;
 
 	/**
-	 * Time taken to remove duck using volume scalar.
+	 * Time taken to recover volume scalar duck.
 	 */
-	UPROPERTY(EditAnywhere, BlueprintReadWrite, Category = "Volume Scaling", meta = (DisplayName = "Release Time", EditCondition = "bVolumeScaleCanRelease", UIMin = "0.0", ClampMin = "0.0", UIMax = "10.0", ClampMax="1000000.0"))
+	UPROPERTY(EditAnywhere, BlueprintReadWrite, Category = "Volume Scaling", meta = (DisplayName = "Recover Time", EditCondition = "bVolumeScaleCanRelease && VolumeScaleMode == EConcurrencyVolumeScaleMode::Default", UIMin = "0.0", ClampMin = "0.0", UIMax = "10.0", ClampMax="1000000.0"))
 	float VolumeScaleReleaseTime;
 
 	/**
@@ -119,6 +135,7 @@ public:
 		, ResolutionRule(EMaxConcurrentResolutionRule::StopFarthestThenOldest)
 		, RetriggerTime(0.0f)
 		, VolumeScale(1.0f)
+		, VolumeScaleMode(EConcurrencyVolumeScaleMode::Default)
 		, VolumeScaleAttackTime(0.01f)
 		, bVolumeScaleCanRelease(0)
 		, VolumeScaleReleaseTime(0.5f)
@@ -196,6 +213,7 @@ public:
 
 	void Update(float InElapsed);
 
+	float GetLerpTime() const;
 	float GetVolume(bool bInDecibels = false) const;
 	float GetTargetVolume(bool bInDecibels = false) const;
 
@@ -246,6 +264,9 @@ public:
 	/** Removes an active sound from the active sound array. */
 	void RemoveActiveSound(FActiveSound& ActiveSound);
 
+	/** Updates volume based on distance generation if set as VolumeScaleMode */
+	void UpdateGeneration(FActiveSound* NewActiveSound = nullptr);
+
 	/** Sorts the active sound if concurrency settings require culling post playback */
 	void CullSoundsDueToMaxConcurrency();
 };
@@ -303,6 +324,9 @@ public:
 	/** Stops sound, applying concurrency rules for how to stop. */
 	void StopDueToVoiceStealing(FActiveSound& ActiveSound);
 
+	/** Updates generations for concurrency groups set to scale active sound volumes by distance or priority */
+	void UpdateVolumeScaleGenerations();
+
 	/** Culls any active sounds due to max concurrency sound resolution rule constraints being met */
 	void UpdateSoundsToCull();
 
diff --git a/Engine/Source/Runtime/Engine/Private/ActiveSound.cpp b/Engine/Source/Runtime/Engine/Private/ActiveSound.cpp
--- a/Engine/Source/Runtime/Engine/Private/ActiveSound.cpp
+++ b/Engine/Source/Runtime/Engine/Private/ActiveSound.cpp
@@ -1458,12 +1458,11 @@ float FActiveSound::GetAttenuationFrequency(const FSoundAttenuationSettings* Set
 	return FMath::Clamp<float>(OutputFrequency, MIN_FILTER_FREQUENCY, MAX_FILTER_FREQUENCY);
 }
 
-float FActiveSound::GetHighestPriority() const
-{ 
-	static constexpr float MaxPriority = TNumericLimits<float>::Max();
+bool FActiveSound::GetAlwaysPlay() const
+{
 	if (bAlwaysPlay)
 	{
-		return MaxPriority;
+		return true;
 	}
 
 	for (const TPair<UPTRINT, FWaveInstance*>& Pair : WaveInstances)
@@ -1474,11 +1473,25 @@ float FActiveSound::GetHighestPriority() const
 			check(SoundClassProperties);
 			if (SoundClassProperties->bAlwaysPlay)
 			{
-				return MaxPriority;
+				return true;
 			}
 		}
 	}
 
+	return false;
+}
+
+float FActiveSound::GetHighestPriority(bool bIgnoreAlwaysPlay) const
+{ 
+	if (!bIgnoreAlwaysPlay)
+	{
+		if (GetAlwaysPlay())
+		{
+			static constexpr float AlwaysPlayPriority = TNumericLimits<float>::Max();
+			return AlwaysPlayPriority;
+		}
+	}
+
 	const float HighestPriority = Priority * FocusData.PriorityHighest * FocusData.PriorityScale;
 	return FMath::Clamp(HighestPriority, 0.0f, 100.0f);
 }
diff --git a/Engine/Source/Runtime/Engine/Private/AudioDevice.cpp b/Engine/Source/Runtime/Engine/Private/AudioDevice.cpp
--- a/Engine/Source/Runtime/Engine/Private/AudioDevice.cpp
+++ b/Engine/Source/Runtime/Engine/Private/AudioDevice.cpp
@@ -2283,12 +2283,13 @@ void FAudioDevice::RecursiveApplyAdjuster(const FSoundClassAdjuster& InAdjuster,
 	}
 }
 
-void FAudioDevice::CullSoundsDueToMaxConcurrency(TArray<FWaveInstance*>& WaveInstances, TArray<FActiveSound *>& ActiveSoundsCopy)
+void FAudioDevice::UpdateConcurrency(TArray<FWaveInstance*>& WaveInstances, TArray<FActiveSound *>& ActiveSoundsCopy)
 {
 	// Now stop any sounds that are active that are in concurrency resolution groups that resolve by stopping quietest
 	{
 		SCOPE_CYCLE_COUNTER(STAT_AudioEvaluateConcurrency);
 		ConcurrencyManager.UpdateSoundsToCull();
+		ConcurrencyManager.UpdateVolumeScaleGenerations();
 	}
 
 	for (int32 i = ActiveSoundsCopy.Num() - 1; i >= 0; --i)
@@ -2324,6 +2325,17 @@ void FAudioDevice::CullSoundsDueToMaxConcurrency(TArray<FWaveInstance*>& WaveIns
 			WaveInstances.RemoveAtSwap(i, 1, false);
 		}
 	}
+
+	// Must be completed after removing wave instances as it avoids an issue
+	// where quiet loops can wrongfully scale concurrency ducking improperly if they continue
+	// to attempt to be evaluated while being periodically realized to check volumes from virtualized.
+	for (int32 i = 0; i < ActiveSoundsCopy.Num(); ++i)
+	{
+		if (FActiveSound* ActiveSound = ActiveSoundsCopy[i])
+		{
+			ActiveSound->UpdateConcurrencyVolumeScalars(GetGameDeltaTime());
+		}
+	}
 }
 
 bool FAudioDevice::ApplySoundMix(USoundMix* NewMix, FSoundMixState* SoundMixState)
@@ -3717,18 +3729,7 @@ int32 FAudioDevice::GetSortedActiveWaveInstances(TArray<FWaveInstance*>& WaveIns
 
 	if (GetType != ESortedActiveWaveGetType::QueryOnly)
 	{
-		CullSoundsDueToMaxConcurrency(WaveInstances, ActiveSoundsCopy);
-	}
-
-	// Must be completed after StopQuietSoundsDueToMaxConcurrency as it avoids an issue
-	// where quiet loops can wrongfully scale concurrency ducking improperly if they continue
-	// to attempt to be evaluated while being periodically realized to check volumes from virtualized.
-	for (int32 i = 0; i < ActiveSoundsCopy.Num(); ++i)
-	{
-		if (FActiveSound* ActiveSound = ActiveSoundsCopy[i])
-		{
-			ActiveSound->UpdateConcurrencyVolumeScalars(GetGameDeltaTime());
-		}
+		UpdateConcurrency(WaveInstances, ActiveSoundsCopy);
 	}
 
 	int32 FirstActiveIndex = 0;
@@ -5014,19 +5015,13 @@ bool FAudioDevice::LocationIsAudible(const FVector& Location, int32 ListenerInde
 
 float FAudioDevice::GetDistanceToNearestListener(const FVector& Location) const
 {
-	const bool bInAudioThread = IsInAudioThread();
-	const bool bInGameThread = IsInGameThread();
-
-	check(bInAudioThread || bInGameThread);
-
-	float DistSquared;
-	const bool bAllowAttenuationOverrides = true;
-	if (FindClosestListenerIndex(Location, DistSquared, bAllowAttenuationOverrides) == INDEX_NONE)
+	float DistSquared = 0.0f;
+	if (GetDistanceSquaredToNearestListener(Location, DistSquared))
 	{
-		return WORLD_MAX;
+		return FMath::Sqrt(DistSquared);
 	}
 
-	return FMath::Sqrt(DistSquared);
+	return WORLD_MAX;
 }
 
 float FAudioDevice::GetSquaredDistanceToListener(const FVector& Location, const FTransform& ListenerTransform) const
@@ -5057,6 +5052,26 @@ bool FAudioDevice::GetDistanceSquaredToListener(const FVector& Location, int32 L
 	return true;
 }
 
+bool FAudioDevice::GetDistanceSquaredToNearestListener(const FVector& Location, float& OutSqDistance) const
+{
+	OutSqDistance = TNumericLimits<float>::Max();
+	const bool bInAudioThread = IsInAudioThread();
+	const bool bInGameThread = IsInGameThread();
+
+	check(bInAudioThread || bInGameThread);
+
+	float DistSquared;
+	const bool bAllowAttenuationOverrides = true;
+	if (FindClosestListenerIndex(Location, DistSquared, bAllowAttenuationOverrides) == INDEX_NONE)
+	{
+		OutSqDistance = WORLD_MAX;
+		return false;
+	}
+
+	OutSqDistance = DistSquared;
+	return true;
+}
+
 bool FAudioDevice::GetListenerPosition(int32 ListenerIndex, FVector& OutPosition, bool bAllowOverride) const
 {
 	OutPosition = FVector::ZeroVector;
diff --git a/Engine/Source/Runtime/Engine/Private/SoundConcurrency.cpp b/Engine/Source/Runtime/Engine/Private/SoundConcurrency.cpp
--- a/Engine/Source/Runtime/Engine/Private/SoundConcurrency.cpp
+++ b/Engine/Source/Runtime/Engine/Private/SoundConcurrency.cpp
@@ -3,6 +3,7 @@
 #include "Sound/SoundConcurrency.h"
 #include "Components/AudioComponent.h"
 #include "ActiveSound.h"
+#include "AudioDefines.h"
 #include "AudioDevice.h"
 #include "AudioVirtualLoop.h"
 #include "DSP/VolumeFader.h"
@@ -12,7 +13,7 @@
 // Forward Declarations
 struct FListener;
 
-namespace
+namespace SoundConcurrency
 {
 	void SetSoundDataTarget(const FActiveSound& ActiveSound, FConcurrencySoundData& SoundData, float InTargetVolume, float InLerpTime)
 	{
@@ -36,7 +37,7 @@ namespace
 		}
 #endif // UE_BUILD_SHIPPING
 	}
-} // namespace <>
+} // namespace SoundConcurrency
 
 
 USoundConcurrency::USoundConcurrency(const FObjectInitializer& ObjectInitializer)
@@ -105,6 +106,11 @@ void FConcurrencySoundData::Update(float InElapsed)
 	Elapsed += InElapsed;
 }
 
+float FConcurrencySoundData::GetLerpTime() const
+{
+	return LerpTime;
+}
+
 float FConcurrencySoundData::GetVolume(bool bInDecibels) const
 {
 	if (FMath::IsNearlyZero(LerpTime) || Elapsed > LerpTime || FMath::IsNearlyEqual(DbTargetVolume, DbStartVolume))
@@ -125,11 +131,13 @@ float FConcurrencySoundData::GetTargetVolume(bool bInDecibels) const
 void FConcurrencySoundData::SetTarget(float InTargetVolume, float InLerpTime)
 {
 	DbStartVolume = GetVolume(true);
-
+	DbTargetVolume = Audio::ConvertToDecibels(InTargetVolume, KINDA_SMALL_NUMBER);
 	LerpTime = FMath::Max(InLerpTime, 0.0f);
-	Elapsed = 0.0f;
 
-	DbTargetVolume = Audio::ConvertToDecibels(InTargetVolume, KINDA_SMALL_NUMBER);
+	if (InLerpTime != LerpTime && InTargetVolume != DbTargetVolume)
+	{
+		Elapsed = 0.0f;
+	}
 }
 
 FConcurrencyGroup::FConcurrencyGroup(FConcurrencyGroupID InGroupID, const FConcurrencyHandle& ConcurrencyHandle)
@@ -156,11 +164,17 @@ void FConcurrencyGroup::AddActiveSound(FActiveSound& ActiveSound)
 	}
 
 	FConcurrencySoundData SoundData;
-	SoundData.Generation = ActiveSounds.Num();
-	SetSoundDataTarget(ActiveSound, SoundData, 1.0f, 0.0f);
+
+	if (Settings.VolumeScaleMode == EConcurrencyVolumeScaleMode::Default)
+	{
+		SoundData.Generation = ActiveSounds.Num();
+		SoundConcurrency::SetSoundDataTarget(ActiveSound, SoundData, 1.0f, 0.0f);
+	}
 
 	ActiveSounds.Add(&ActiveSound);
 	ActiveSound.ConcurrencyGroupData.Add(GroupID, MoveTemp(SoundData));
+
+	UpdateGeneration(&ActiveSound);
 }
 
 void FConcurrencyGroup::RemoveActiveSound(FActiveSound& ActiveSound)
@@ -187,11 +201,72 @@ void FConcurrencyGroup::RemoveActiveSound(FActiveSound& ActiveSound)
 			OtherSoundData.Generation--;
 		}
 
-		if (Settings.bVolumeScaleCanRelease)
+		if (Settings.bVolumeScaleCanRelease && Settings.VolumeScaleMode == EConcurrencyVolumeScaleMode::Default)
 		{
 			const float GenerationDelta = static_cast<float>(ActiveSounds.Num() - OtherSoundData.Generation) - 1.0f;
 			const float NewTargetVolume = FMath::Clamp(FMath::Pow(Settings.GetVolumeScale(), GenerationDelta), 0.0f, 1.0f);
-			SetSoundDataTarget(*OtherSound, OtherSoundData, NewTargetVolume, Settings.VolumeScaleReleaseTime);
+			SoundConcurrency::SetSoundDataTarget(*OtherSound, OtherSoundData, NewTargetVolume, Settings.VolumeScaleReleaseTime);
+		}
+	}
+}
+
+void FConcurrencyGroup::UpdateGeneration(FActiveSound* NewActiveSound)
+{
+	// Default uses value set at time of adding sound to group, so no sort required
+	if (Settings.VolumeScaleMode == EConcurrencyVolumeScaleMode::Default)
+	{
+		return;
+	}
+
+	struct FCompareActiveSounds
+	{
+		EConcurrencyVolumeScaleMode ScaleMode;
+
+		FCompareActiveSounds(EConcurrencyVolumeScaleMode InScaleMode)
+			: ScaleMode(InScaleMode)
+		{
+		}
+
+		FORCEINLINE bool operator()(const FActiveSound& A, const FActiveSound& B) const
+		{
+			switch(ScaleMode)
+			{
+				case EConcurrencyVolumeScaleMode::Distance:
+				{
+					float DistSqA = 0.0f;
+					float DistSqB = 0.0f;
+					A.AudioDevice->GetDistanceSquaredToNearestListener(A.LastLocation, DistSqA);
+					B.AudioDevice->GetDistanceSquaredToNearestListener(B.LastLocation, DistSqB);
+					return DistSqA > DistSqB;
+				}
+
+				case EConcurrencyVolumeScaleMode::Priority:
+				{
+					// Ensures sounds set to always play are sorted above those that aren't, but are sorted appropriately between one another
+					const float APriority = A.GetAlwaysPlay() ? A.GetHighestPriority(true /* bIgnoreAlwaysPlay */) + MAX_SOUND_PRIORITY + 1.0f : A.GetHighestPriority();
+					const float BPriority = B.GetAlwaysPlay() ? B.GetHighestPriority(true /* bIgnoreAlwaysPlay */) + MAX_SOUND_PRIORITY + 1.0f : B.GetHighestPriority();
+					return APriority < BPriority;
+				}
+			}
+		}
+	};
+
+	ActiveSounds.Sort(FCompareActiveSounds(Settings.VolumeScaleMode));
+
+	for (int32 i = 0; i < ActiveSounds.Num(); ++i)
+	{
+		FActiveSound* ActiveSound = ActiveSounds[i];
+		
+		if (FConcurrencySoundData* SoundData = ActiveSound->ConcurrencyGroupData.Find(GetGroupID()))
+		{
+			SoundData->Generation = i;
+			const float GenerationDelta = static_cast<float>(ActiveSounds.Num() - SoundData->Generation) - 1.0f;
+			const float NewTargetVolume = FMath::Clamp(FMath::Pow(Settings.GetVolumeScale(), GenerationDelta), 0.0f, 1.0f);
+
+			// If new sound added, immediately lerp to the generation volume.
+			const float AttackTime = NewActiveSound == ActiveSound ? 0.0f : Settings.VolumeScaleAttackTime;
+
+			SoundConcurrency::SetSoundDataTarget(*ActiveSound, *SoundData, NewTargetVolume, AttackTime);
 		}
 	}
 }
@@ -234,8 +309,9 @@ void FConcurrencyGroup::CullSoundsDueToMaxConcurrency()
 				case EMaxConcurrentResolutionRule::StopLowestPriority:
 				case EMaxConcurrentResolutionRule::StopLowestPriorityThenPreventNew:
 				{
-					const float APriority = A.GetHighestPriority();
-					const float BPriority = B.GetHighestPriority();
+					// Ensures sounds set to always play are sorted above those that aren't, but are sorted appropriately between one another
+					const float APriority = A.GetAlwaysPlay() ? A.GetHighestPriority(true /* bIgnoreAlwaysPlay */) + MAX_SOUND_PRIORITY + 1.0f : A.GetHighestPriority();
+					const float BPriority = B.GetAlwaysPlay() ? B.GetHighestPriority(true /* bIgnoreAlwaysPlay */) + MAX_SOUND_PRIORITY + 1.0f : B.GetHighestPriority();
 					if (!FMath::IsNearlyEqual(APriority, BPriority, KINDA_SMALL_NUMBER))
 					{
 						return APriority < BPriority;
@@ -694,31 +770,34 @@ FActiveSound* FSoundConcurrencyManager::CreateAndEvictActiveSounds(const FActive
 		check(ConcurrencyGroup);
 
 		const FSoundConcurrencySettings& Settings = ConcurrencyGroup->GetSettings();
-		const float Volume = Settings.GetVolumeScale();
-		if (!FMath::IsNearlyEqual(Volume, 1.0f))
+		if (Settings.VolumeScaleMode == EConcurrencyVolumeScaleMode::Default)
 		{
-			check(Volume >= 0.0f);
-			const int32 NextGeneration = ConcurrencyGroup->GetNextGeneration();
-
-			// If we're ducking older sounds in the concurrency group, then loop through each sound in the concurrency group
-			// and update their duck amount based on each sound's generation and the next generation count. The older the sound, the more ducking.
-			const TArray<FActiveSound*>& ActiveSounds = ConcurrencyGroup->GetActiveSounds();
-			for (FActiveSound* OtherSound : ActiveSounds)
+			const float Volume = Settings.GetVolumeScale();
+			if (!FMath::IsNearlyEqual(Volume, 1.0f))
 			{
-				check(OtherSound);
-
-				FConcurrencySoundData& OtherSoundData = OtherSound->ConcurrencyGroupData.FindChecked(ConcurrencyGroup->GetGroupID());
-				const float ActiveSoundGeneration = static_cast<float>(OtherSoundData.Generation);
-				const float GenerationDelta = NextGeneration - ActiveSoundGeneration;
-				const float NewTargetVolume = FMath::Pow(Volume, GenerationDelta);
+				check(Volume >= 0.0f);
+				const int32 NextGeneration = ConcurrencyGroup->GetNextGeneration();
 
-				// Don't allow volume to recover if release time is negative
-				if (Settings.bVolumeScaleCanRelease && OtherSoundData.GetTargetVolume() < NewTargetVolume)
+				// If we're ducking older sounds in the concurrency group, then loop through each sound in the concurrency group
+				// and update their duck amount based on each sound's generation and the next generation count. The older the sound, the more ducking.
+				const TArray<FActiveSound*>& ActiveSounds = ConcurrencyGroup->GetActiveSounds();
+				for (FActiveSound* OtherSound : ActiveSounds)
 				{
-					continue;
-				}
+					check(OtherSound);
+
+					FConcurrencySoundData& OtherSoundData = OtherSound->ConcurrencyGroupData.FindChecked(ConcurrencyGroup->GetGroupID());
+					const float ActiveSoundGeneration = static_cast<float>(OtherSoundData.Generation);
+					const float GenerationDelta = NextGeneration - ActiveSoundGeneration;
+					const float NewTargetVolume = FMath::Pow(Volume, GenerationDelta);
 
-				SetSoundDataTarget(*OtherSound, OtherSoundData, NewTargetVolume, Settings.VolumeScaleAttackTime);
+					// Don't allow volume to recover if release time is negative
+					if (Settings.bVolumeScaleCanRelease && OtherSoundData.GetTargetVolume() < NewTargetVolume)
+					{
+						continue;
+					}
+
+					SoundConcurrency::SetSoundDataTarget(*OtherSound, OtherSoundData, NewTargetVolume, Settings.VolumeScaleAttackTime);
+				}
 			}
 		}
 
@@ -871,6 +950,16 @@ void FSoundConcurrencyManager::StopDueToVoiceStealing(FActiveSound& ActiveSound)
 	}
 }
 
+void FSoundConcurrencyManager::UpdateVolumeScaleGenerations()
+{
+	check(IsInAudioThread());
+
+	for (auto& ConcurrencyGroupEntry : ConcurrencyGroups)
+	{
+		ConcurrencyGroupEntry.Value->UpdateGeneration();
+	}
+}
+
 void FSoundConcurrencyManager::UpdateSoundsToCull()
 {
 	check(IsInAudioThread());
diff --git a/Engine/Source/Runtime/Engine/Public/ActiveSound.h b/Engine/Source/Runtime/Engine/Public/ActiveSound.h
--- a/Engine/Source/Runtime/Engine/Public/ActiveSound.h
+++ b/Engine/Source/Runtime/Engine/Public/ActiveSound.h
@@ -668,8 +668,13 @@ public:
 	/** Applies the active sound's attenuation settings to the input parse params using the given listener */
 	void ParseAttenuation(FSoundParseParameters& OutParseParams, int32 ListenerIndex, const FSoundAttenuationSettings& InAttenuationSettings);
 
-	/** Returns the highest effective priority of the child wave instances */
-	float GetHighestPriority() const;
+	/** Returns whether or not sound or any active wave instances it manages are set to always play */
+	bool GetAlwaysPlay() const;
+
+	/** Returns the highest effective priority of the child wave instances. If bIgnoreAlwaysPlay set to true, gives highest
+	  * priority disregarding always play priority override.
+	  */
+	float GetHighestPriority(bool bIgnoreAlwaysPlay = false) const;
 
 	/** Sets the amount of audio from this active sound to send to the submix. */
 	void SetSubmixSend(const FSoundSubmixSendInfo& SubmixSendInfo);
diff --git a/Engine/Source/Runtime/Engine/Public/AudioDevice.h b/Engine/Source/Runtime/Engine/Public/AudioDevice.h
--- a/Engine/Source/Runtime/Engine/Public/AudioDevice.h
+++ b/Engine/Source/Runtime/Engine/Public/AudioDevice.h
@@ -883,11 +883,18 @@ public:
 	float GetSquaredDistanceToListener(const FVector& Location, const FTransform& ListenerTransform) const;
 
 	/**
-	* Returns the distance from location to the appropriate listener representation, depending on calling thread
+	* Sets OutSqDistance to the distance from location to the appropriate listener representation, depending on calling thread.
+	* Returns true if listener position is valid, false if not (in which case, OutSqDistance is undefined).
 	*/
 	bool GetDistanceSquaredToListener(const FVector& Location, int32 ListenerIndex, float& OutSqDistance) const;
 
 	/**
+	* Sets OutSqDistance to the distance from location the closest listener, depending on calling thread.
+	* Returns true if listener position is valid, false if not (in which case, OutSqDistance is undefined).
+	*/
+	bool GetDistanceSquaredToNearestListener(const FVector& Location, float& OutSqDistance) const;
+		
+		/**
 	* Returns a position from the appropriate listener representation, depending on calling thread.
 	*
 	* @param	ListenerIndex	index of the listener or proxy
@@ -1362,7 +1369,7 @@ private:
 
 	/** Stops quiet/low priority sounds due to being evaluated as not fulfilling concurrency requirements
 	 */
-	void CullSoundsDueToMaxConcurrency(TArray<FWaveInstance*>& WaveInstances, TArray<FActiveSound*>& ActiveSoundsCopy);
+	void UpdateConcurrency(TArray<FWaveInstance*>& WaveInstances, TArray<FActiveSound*>& ActiveSoundsCopy);
 
 	/**
 	 * Checks if the given sound would be audible.
