Fix missing return coverage
#rnx

diff --git a/Engine/Source/Runtime/Engine/Private/SoundConcurrency.cpp b/Engine/Source/Runtime/Engine/Private/SoundConcurrency.cpp
--- a/Engine/Source/Runtime/Engine/Private/SoundConcurrency.cpp
+++ b/Engine/Source/Runtime/Engine/Private/SoundConcurrency.cpp
@@ -248,6 +248,8 @@ void FConcurrencyGroup::UpdateGeneration(FActiveSound* NewActiveSound)
 					return APriority < BPriority;
 				}
 			}
+
+			return true;
 		}
 	};
 
