Parallelize deferred kinematic update.

diff --git a/Engine/Source/Runtime/Engine/Classes/Components/SkeletalMeshComponent.h b/Engine/Source/Runtime/Engine/Classes/Components/SkeletalMeshComponent.h
--- a/Engine/Source/Runtime/Engine/Classes/Components/SkeletalMeshComponent.h
+++ b/Engine/Source/Runtime/Engine/Classes/Components/SkeletalMeshComponent.h
@@ -451,8 +451,6 @@ public:
 	UPROPERTY(transient)
 	uint8 bHasValidBodies:1;
 
-	/** Indicates that this SkeletalMeshComponent has deferred kinematic bone updates until next physics sim.  */
-	uint8 bDeferredKinematicUpdate:1;
 
 	/** Enables blending in of physics bodies whether Simulate or not*/
 	UPROPERTY(transient)
@@ -477,6 +475,9 @@ public:
 	UPROPERTY(EditAnywhere, BlueprintReadWrite, Category=Clothing)
 	uint8 bDisableClothSimulation:1;
 
+	/** Indicates that this SkeletalMeshComponent has deferred kinematic bone updates until next physics sim if not INDEX_NONE. */
+	int32 DeferredKinematicUpdateIndex;
+
 private:
 	/** Disable rigid body animation nodes and play original animation without simulation */
 	UPROPERTY(EditAnywhere, Category=Clothing)
diff --git a/Engine/Source/Runtime/Engine/Private/Components/SkeletalMeshComponent.cpp b/Engine/Source/Runtime/Engine/Private/Components/SkeletalMeshComponent.cpp
--- a/Engine/Source/Runtime/Engine/Private/Components/SkeletalMeshComponent.cpp
+++ b/Engine/Source/Runtime/Engine/Private/Components/SkeletalMeshComponent.cpp
@@ -254,6 +254,8 @@ USkeletalMeshComponent::USkeletalMeshComponent(const FObjectInitializer& ObjectI
 
 	bSkipKinematicUpdateWhenInterpolating = false;
 	bSkipBoundsUpdateWhenInterpolating = false;
+
+	DeferredKinematicUpdateIndex = INDEX_NONE;
 }
 
 void USkeletalMeshComponent::Serialize(FArchive& Ar)
@@ -578,7 +580,7 @@ void USkeletalMeshComponent::OnUnregister()
 		ClothingSimulationContext = nullptr;
 	}
 
-	if (bDeferredKinematicUpdate != 0)
+	if (DeferredKinematicUpdateIndex != INDEX_NONE)
 	{
 		UWorld* World = GetWorld();
 		FPhysScene* PhysScene = World ? World->GetPhysicsScene() : nullptr;
diff --git a/Engine/Source/Runtime/Engine/Private/PhysicsEngine/Experimental/PhysScene_Chaos.cpp b/Engine/Source/Runtime/Engine/Private/PhysicsEngine/Experimental/PhysScene_Chaos.cpp
--- a/Engine/Source/Runtime/Engine/Private/PhysicsEngine/Experimental/PhysScene_Chaos.cpp
+++ b/Engine/Source/Runtime/Engine/Private/PhysicsEngine/Experimental/PhysScene_Chaos.cpp
@@ -681,17 +681,19 @@ void FPhysScene_Chaos::UpdateActorsInAccelerationStructure(const TArrayView<FPhy
 			for (int32 ActorIndex = 0; ActorIndex < NumActors; ++ActorIndex)
 			{
 				const FPhysicsActorHandle& Actor = Actors[ActorIndex];
-
-				// @todo(chaos): dedupe code in UpdateActorInAccelerationStructure
-				TAABB<FReal, 3> WorldBounds;
-				const bool bHasBounds = Actor->Geometry()->HasBoundingBox();
-				if (bHasBounds)
+				if (Actor != nullptr)
 				{
-					WorldBounds = Actor->Geometry()->BoundingBox().TransformedAABB(TRigidTransform<FReal, 3>(Actor->X(), Actor->R()));
-				}
+					// @todo(chaos): dedupe code in UpdateActorInAccelerationStructure
+					TAABB<FReal, 3> WorldBounds;
+					const bool bHasBounds = Actor->Geometry()->HasBoundingBox();
+					if (bHasBounds)
+					{
+						WorldBounds = Actor->Geometry()->BoundingBox().TransformedAABB(TRigidTransform<FReal, 3>(Actor->X(), Actor->R()));
+					}
 
-				Chaos::TAccelerationStructureHandle<float, 3> AccelerationHandle(Actor);
-				SpatialAcceleration->UpdateElementIn(AccelerationHandle, WorldBounds, bHasBounds, Actor->SpatialIdx());
+					Chaos::TAccelerationStructureHandle<float, 3> AccelerationHandle(Actor);
+					SpatialAcceleration->UpdateElementIn(AccelerationHandle, WorldBounds, bHasBounds, Actor->SpatialIdx());
+				}
 			}
 		}
 
@@ -1589,19 +1591,9 @@ bool FPhysScene_ChaosInterface::MarkForPreSimKinematicUpdate(USkeletalMeshCompon
 	if (InSkelComp != nullptr && !InSkelComp->IsPendingKill())
 	{
 		// If we are already flagged, just need to update info
-		if (InSkelComp->bDeferredKinematicUpdate)
+		if (InSkelComp->DeferredKinematicUpdateIndex != INDEX_NONE)
 		{
-			TPair<USkeletalMeshComponent*, FDeferredKinematicUpdateInfo>* FoundItem = DeferredKinematicUpdateSkelMeshes.FindByPredicate([InSkelComp](const TPair<USkeletalMeshComponent*, FDeferredKinematicUpdateInfo>& InItem) { return InSkelComp == InItem.Key; });
-			if (!ensure(FoundItem != nullptr))// If the bool was set, we must be in the array!
-			{
-				FDeferredKinematicUpdateInfo Info;
-				Info.TeleportType = InTeleport;
-				Info.bNeedsSkinning = bNeedsSkinning;
-				DeferredKinematicUpdateSkelMeshes.Emplace(InSkelComp, Info);
-				return true;
-			}
-
-			FDeferredKinematicUpdateInfo& Info = FoundItem->Value;
+			FDeferredKinematicUpdateInfo& Info = DeferredKinematicUpdateSkelMeshes[InSkelComp->DeferredKinematicUpdateIndex].Value;
 
 			// If we are currently not going to teleport physics, but this update wants to, we 'upgrade' it
 			if (Info.TeleportType == ETeleportType::None && InTeleport == ETeleportType::TeleportPhysics)
@@ -1622,10 +1614,8 @@ bool FPhysScene_ChaosInterface::MarkForPreSimKinematicUpdate(USkeletalMeshCompon
 			FDeferredKinematicUpdateInfo Info;
 			Info.TeleportType = InTeleport;
 			Info.bNeedsSkinning = bNeedsSkinning;
+			InSkelComp->DeferredKinematicUpdateIndex = DeferredKinematicUpdateSkelMeshes.Num();
 			DeferredKinematicUpdateSkelMeshes.Emplace(InSkelComp, Info);
-
-			// Set flag on component
-			InSkelComp->bDeferredKinematicUpdate = true;
 		}
 	}
 
@@ -1635,15 +1625,13 @@ bool FPhysScene_ChaosInterface::MarkForPreSimKinematicUpdate(USkeletalMeshCompon
 void FPhysScene_ChaosInterface::ClearPreSimKinematicUpdate(USkeletalMeshComponent* InSkelComp)
 {
 	// If non-null, and flagged for deferred update..
-	if (InSkelComp != nullptr && InSkelComp->bDeferredKinematicUpdate)
+	if (InSkelComp != nullptr && InSkelComp->DeferredKinematicUpdateIndex != INDEX_NONE)
 	{
 		// Remove from map
-		int32 NumRemoved = DeferredKinematicUpdateSkelMeshes.RemoveAll([InSkelComp](const TPair<USkeletalMeshComponent*, FDeferredKinematicUpdateInfo>& InItem) { return InSkelComp == InItem.Key; });
-
-		ensure(NumRemoved == 1); // Should be in array if flag was set!
+		DeferredKinematicUpdateSkelMeshes.RemoveAtSwap(InSkelComp->DeferredKinematicUpdateIndex);
 
-		// Clear flag
-		InSkelComp->bDeferredKinematicUpdate = false;
+		// Clear Index
+		InSkelComp->DeferredKinematicUpdateIndex = INDEX_NONE;
 	}
 }
 
@@ -1734,37 +1722,133 @@ void FPhysScene_ChaosInterface::UpdateKinematicsOnDeferredSkelMeshes()
 {
 	SCOPE_CYCLE_COUNTER(STAT_UpdateKinematicsOnDeferredSkelMeshesChaos);
 
-	using FActorsArray = TArray<FPhysicsActorHandle, TInlineAllocator<64>>;
-	using FTransformsArray = TArray<FTransform, TInlineAllocator<64>>;
+	// Holds start index in actor pool for each skeletal mesh.
+	TArray<int32, TInlineAllocator<64>> SkeletalMeshStartIndexArray;
 
-	FActorsArray KinematicUpdateActors;
-	FTransformsArray KinematicUpdateTransforms;
-	FActorsArray TeleportActors;
-	FTransformsArray TeleportTransforms;
-
-	for (const TPair<USkeletalMeshComponent*, FDeferredKinematicUpdateInfo>& DeferredKinematicUpdate : DeferredKinematicUpdateSkelMeshes)
+	TArray<FPhysicsActorHandle, TInlineAllocator<64>>TeleportActorsPool;
+	TArray<IPhysicsProxyBase*, TInlineAllocator<64>> ProxiesToDirty;
+	
+	// Count max number of bodies to determine actor pool size.
 	{
-		USkeletalMeshComponent* SkelComp = DeferredKinematicUpdate.Key;
-		const FDeferredKinematicUpdateInfo& Info = DeferredKinematicUpdate.Value;
+		SkeletalMeshStartIndexArray.Reserve(DeferredKinematicUpdateSkelMeshes.Num());
 
-		ensure(SkelComp->bDeferredKinematicUpdate); // Should be true if in map!
-
-		if (!SkelComp->bEnablePerPolyCollision)
+		int32 TotalBodies = 0;
+		for (const TPair<USkeletalMeshComponent*, FDeferredKinematicUpdateInfo>& DeferredKinematicUpdate : DeferredKinematicUpdateSkelMeshes)
 		{
-			GatherActorsAndTransforms(SkelComp, SkelComp->GetComponentSpaceTransforms(), Info.TeleportType, Info.bNeedsSkinning, KinematicUpdateActors, KinematicUpdateTransforms, TeleportActors, TeleportTransforms);
+			SkeletalMeshStartIndexArray.Add(TotalBodies);
+
+			USkeletalMeshComponent* SkelComp = DeferredKinematicUpdate.Key;
+			if (!SkelComp->bEnablePerPolyCollision)
+			{
+				TotalBodies += SkelComp->Bodies.Num();
+			}
 		}
-		else
+
+		// Actors pool is spare, initialize to nullptr.
+		TeleportActorsPool.AddZeroed(TotalBodies);
+		ProxiesToDirty.Reserve(TotalBodies);
+	}
+
+	// Gather proxies that need to be dirtied before paralell loop, and update any per poly collision skeletal meshes.
+	{
+		for (const TPair<USkeletalMeshComponent*, FDeferredKinematicUpdateInfo>& DeferredKinematicUpdate : DeferredKinematicUpdateSkelMeshes)
 		{
-			// TODO: acceleration for per-poly collision
-			SkelComp->UpdateKinematicBonesToAnim(SkelComp->GetComponentSpaceTransforms(), Info.TeleportType, Info.bNeedsSkinning, EAllowKinematicDeferral::DisallowDeferral);
+			USkeletalMeshComponent* SkelComp = DeferredKinematicUpdate.Key;
+			const FDeferredKinematicUpdateInfo& Info = DeferredKinematicUpdate.Value;
+
+			if (!SkelComp->bEnablePerPolyCollision)
+			{
+				const int32 NumBodies = SkelComp->Bodies.Num();
+				for (int32 i = 0; i < NumBodies; i++)
+				{
+					FBodyInstance* BodyInst = SkelComp->Bodies[i];
+					FPhysicsActorHandle& ActorHandle = BodyInst->ActorHandle;
+					if (!BodyInst->IsInstanceSimulatingPhysics())
+					{
+						const int32 BoneIndex = BodyInst->InstanceBoneIndex;
+						if (BoneIndex != INDEX_NONE)
+						{
+							IPhysicsProxyBase* Proxy = ActorHandle->GetProxy();
+							if (Proxy && Proxy->GetDirtyIdx() == INDEX_NONE)
+							{
+								ProxiesToDirty.Add(Proxy);
+							}
+						}
+					}
+				}
+			}
+			else
+			{
+				// TODO: acceleration for per-poly collision
+				SkelComp->UpdateKinematicBonesToAnim(SkelComp->GetComponentSpaceTransforms(), Info.TeleportType, Info.bNeedsSkinning, EAllowKinematicDeferral::DisallowDeferral);
+			}
 		}
+	}
+
+	// Mark all body's proxies as dirty, as this is not threadsafe and cannot be done in parallel loop.
+	if (ProxiesToDirty.Num() > 0)
+	{
+		// Assumes all particles have the same solver, safe for now, maybe not in the future.
+		IPhysicsProxyBase* Proxy = ProxiesToDirty[0];
+		Chaos::FPhysicsSolverBase* Solver = Proxy->GetSolver();
+		Solver->AddDirtyProxiesUnsafe(ProxiesToDirty);
+	}
 
-		// Clear deferred flag
-		SkelComp->bDeferredKinematicUpdate = false;
+	{
+		Chaos::PhysicsParallelFor(DeferredKinematicUpdateSkelMeshes.Num(), [&](int32 Index)
+		{
+			const TPair<USkeletalMeshComponent*, FDeferredKinematicUpdateInfo>& DeferredKinematicUpdate = DeferredKinematicUpdateSkelMeshes[Index];
+			USkeletalMeshComponent* SkelComp = DeferredKinematicUpdate.Key;
+			const FDeferredKinematicUpdateInfo& Info = DeferredKinematicUpdate.Value;
+
+			SkelComp->DeferredKinematicUpdateIndex = INDEX_NONE;
+
+			if (!SkelComp->bEnablePerPolyCollision)
+			{
+				const UPhysicsAsset* PhysicsAsset = SkelComp->GetPhysicsAsset();
+				const FTransform& CurrentLocalToWorld = SkelComp->GetComponentTransform();
+				const int32 NumBodies = SkelComp->Bodies.Num();
+				const TArray<FTransform>& ComponentSpaceTransforms = SkelComp->GetComponentSpaceTransforms();
+				
+				const int32 ActorPoolStartIndex = SkeletalMeshStartIndexArray[Index];
+				for (int32 i = 0; i < NumBodies; i++)
+				{
+					FBodyInstance* BodyInst = SkelComp->Bodies[i];
+					FPhysicsActorHandle& ActorHandle = BodyInst->ActorHandle;
+					if (!BodyInst->IsInstanceSimulatingPhysics())
+					{
+						const int32 BoneIndex = BodyInst->InstanceBoneIndex;
+						if (BoneIndex != INDEX_NONE)
+						{
+							const FTransform BoneTransform = ComponentSpaceTransforms[BoneIndex] * CurrentLocalToWorld;
+
+							TeleportActorsPool[ActorPoolStartIndex + i] = ActorHandle;
+
+							// TODO: Kinematic targets. Check Teleport type on FDeferredKinematicUpdateInfo and don't always teleport.
+							ActorHandle->SetX(BoneTransform.GetLocation(), false);	// only set dirty once in SetR
+							ActorHandle->SetR(BoneTransform.GetRotation());
+							ActorHandle->UpdateShapeBounds(BoneTransform);
+
+							if (!PhysicsAsset->SkeletalBodySetups[i]->bSkipScaleFromAnimation)
+							{
+								const FVector& MeshScale3D = CurrentLocalToWorld.GetScale3D();
+								if (MeshScale3D.IsUniform())
+								{
+									BodyInst->UpdateBodyScale(BoneTransform.GetScale3D());
+								}
+								else
+								{
+									BodyInst->UpdateBodyScale(MeshScale3D);
+								}
+							}
+						}
+					}
+				}
+			}
+		});
 	}
 
-	ProcessKinematicTargetActors(Scene, KinematicUpdateActors, KinematicUpdateTransforms);
-	ProcessTeleportActors(Scene, TeleportActors, TeleportTransforms);
+	Scene.UpdateActorsInAccelerationStructure(TeleportActorsPool);
 
 	DeferredKinematicUpdateSkelMeshes.Reset();
 }
diff --git a/Engine/Source/Runtime/Experimental/Chaos/Public/Chaos/Framework/PhysicsSolverBase.h b/Engine/Source/Runtime/Experimental/Chaos/Public/Chaos/Framework/PhysicsSolverBase.h
--- a/Engine/Source/Runtime/Experimental/Chaos/Public/Chaos/Framework/PhysicsSolverBase.h
+++ b/Engine/Source/Runtime/Experimental/Chaos/Public/Chaos/Framework/PhysicsSolverBase.h
@@ -55,6 +55,20 @@ namespace Chaos
 			}
 		}
 
+		// Batch proxy insertion, does not check DirtyIdx.
+		template< typename TProxiesArray>
+		void AddMultipleUnsafe(TProxiesArray& ProxiesArray)
+		{
+			int32 Idx = ProxiesData.Num();
+			ProxiesData.Append(ProxiesArray);
+
+			for(IPhysicsProxyBase* Proxy : ProxiesArray)
+			{
+				Proxy->SetDirtyIdx(Idx++);
+			}
+		}
+
+
 		void Remove(IPhysicsProxyBase* Base)
 		{
 			const int32 Idx = Base->GetDirtyIdx();
@@ -173,6 +187,13 @@ namespace Chaos
 			DirtyProxiesDataBuffer.AccessProducerBuffer()->Remove(ProxyBaseIn);
 		}
 
+		// Batch dirty proxies without checking DirtyIdx.
+		template <typename TProxiesArray>
+		void AddDirtyProxiesUnsafe(TProxiesArray& ProxiesArray)
+		{
+			DirtyProxiesDataBuffer.AccessProducerBuffer()->AddMultipleUnsafe(ProxiesArray);
+		}
+
 		void AddDirtyProxyShape(IPhysicsProxyBase* ProxyBaseIn, int32 ShapeIdx)
 		{
 			DirtyProxiesDataBuffer.AccessProducerBuffer()->AddShape(ProxyBaseIn,ShapeIdx);
diff --git a/Engine/Source/Runtime/Experimental/Chaos/Public/Chaos/ParticleHandle.h b/Engine/Source/Runtime/Experimental/Chaos/Public/Chaos/ParticleHandle.h
--- a/Engine/Source/Runtime/Experimental/Chaos/Public/Chaos/ParticleHandle.h
+++ b/Engine/Source/Runtime/Experimental/Chaos/Public/Chaos/ParticleHandle.h
@@ -1446,12 +1446,17 @@ public:
 	}
 
 	void UpdateShapeBounds()
+	{
+		UpdateShapeBounds(FRigidTransform3(X(), R()));
+	}
+
+	void UpdateShapeBounds(const FTransform& Transform)
 	{
 		if (MNonFrequentData.Read().Geometry->HasBoundingBox())
 		{
 			for (auto& Shape : MShapesArray)
 			{
-				Shape->UpdateShapeBounds(FRigidTransform3(X(), R()));
+				Shape->UpdateShapeBounds(Transform);
 			}
 		}
 	}
