Back out changelist 12473520 - crash in quick smoke.

diff --git a/Engine/Source/Runtime/Engine/Classes/Components/SkeletalMeshComponent.h b/Engine/Source/Runtime/Engine/Classes/Components/SkeletalMeshComponent.h
--- a/Engine/Source/Runtime/Engine/Classes/Components/SkeletalMeshComponent.h
+++ b/Engine/Source/Runtime/Engine/Classes/Components/SkeletalMeshComponent.h
@@ -451,6 +451,8 @@ public:
 	UPROPERTY(transient)
 	uint8 bHasValidBodies:1;
 
+	/** Indicates that this SkeletalMeshComponent has deferred kinematic bone updates until next physics sim.  */
+	uint8 bDeferredKinematicUpdate:1;
 
 	/** Enables blending in of physics bodies whether Simulate or not*/
 	UPROPERTY(transient)
@@ -475,9 +477,6 @@ public:
 	UPROPERTY(EditAnywhere, BlueprintReadWrite, Category=Clothing)
 	uint8 bDisableClothSimulation:1;
 
-	/** Indicates that this SkeletalMeshComponent has deferred kinematic bone updates until next physics sim if not INDEX_NONE. */
-	int32 DeferredKinematicUpdateIndex;
-
 private:
 	/** Disable rigid body animation nodes and play original animation without simulation */
 	UPROPERTY(EditAnywhere, Category=Clothing)
diff --git a/Engine/Source/Runtime/Engine/Private/Components/SkeletalMeshComponent.cpp b/Engine/Source/Runtime/Engine/Private/Components/SkeletalMeshComponent.cpp
--- a/Engine/Source/Runtime/Engine/Private/Components/SkeletalMeshComponent.cpp
+++ b/Engine/Source/Runtime/Engine/Private/Components/SkeletalMeshComponent.cpp
@@ -254,8 +254,6 @@ USkeletalMeshComponent::USkeletalMeshComponent(const FObjectInitializer& ObjectI
 
 	bSkipKinematicUpdateWhenInterpolating = false;
 	bSkipBoundsUpdateWhenInterpolating = false;
-
-	DeferredKinematicUpdateIndex = INDEX_NONE;
 }
 
 void USkeletalMeshComponent::Serialize(FArchive& Ar)
@@ -580,7 +578,7 @@ void USkeletalMeshComponent::OnUnregister()
 		ClothingSimulationContext = nullptr;
 	}
 
-	if (DeferredKinematicUpdateIndex != INDEX_NONE)
+	if (bDeferredKinematicUpdate != 0)
 	{
 		UWorld* World = GetWorld();
 		FPhysScene* PhysScene = World ? World->GetPhysicsScene() : nullptr;
diff --git a/Engine/Source/Runtime/Engine/Private/PhysicsEngine/Experimental/PhysScene_Chaos.cpp b/Engine/Source/Runtime/Engine/Private/PhysicsEngine/Experimental/PhysScene_Chaos.cpp
--- a/Engine/Source/Runtime/Engine/Private/PhysicsEngine/Experimental/PhysScene_Chaos.cpp
+++ b/Engine/Source/Runtime/Engine/Private/PhysicsEngine/Experimental/PhysScene_Chaos.cpp
@@ -681,19 +681,17 @@ void FPhysScene_Chaos::UpdateActorsInAccelerationStructure(const TArrayView<FPhy
 			for (int32 ActorIndex = 0; ActorIndex < NumActors; ++ActorIndex)
 			{
 				const FPhysicsActorHandle& Actor = Actors[ActorIndex];
-				if (Actor != nullptr)
-				{
-					// @todo(chaos): dedupe code in UpdateActorInAccelerationStructure
-					TAABB<FReal, 3> WorldBounds;
-					const bool bHasBounds = Actor->Geometry()->HasBoundingBox();
-					if (bHasBounds)
-					{
-						WorldBounds = Actor->Geometry()->BoundingBox().TransformedAABB(TRigidTransform<FReal, 3>(Actor->X(), Actor->R()));
-					}
 
-					Chaos::TAccelerationStructureHandle<float, 3> AccelerationHandle(Actor);
-					SpatialAcceleration->UpdateElementIn(AccelerationHandle, WorldBounds, bHasBounds, Actor->SpatialIdx());
+				// @todo(chaos): dedupe code in UpdateActorInAccelerationStructure
+				TAABB<FReal, 3> WorldBounds;
+				const bool bHasBounds = Actor->Geometry()->HasBoundingBox();
+				if (bHasBounds)
+				{
+					WorldBounds = Actor->Geometry()->BoundingBox().TransformedAABB(TRigidTransform<FReal, 3>(Actor->X(), Actor->R()));
 				}
+
+				Chaos::TAccelerationStructureHandle<float, 3> AccelerationHandle(Actor);
+				SpatialAcceleration->UpdateElementIn(AccelerationHandle, WorldBounds, bHasBounds, Actor->SpatialIdx());
 			}
 		}
 
@@ -1591,9 +1589,19 @@ bool FPhysScene_ChaosInterface::MarkForPreSimKinematicUpdate(USkeletalMeshCompon
 	if (InSkelComp != nullptr && !InSkelComp->IsPendingKill())
 	{
 		// If we are already flagged, just need to update info
-		if (InSkelComp->DeferredKinematicUpdateIndex != INDEX_NONE)
+		if (InSkelComp->bDeferredKinematicUpdate)
 		{
-			FDeferredKinematicUpdateInfo& Info = DeferredKinematicUpdateSkelMeshes[InSkelComp->DeferredKinematicUpdateIndex].Value;
+			TPair<USkeletalMeshComponent*, FDeferredKinematicUpdateInfo>* FoundItem = DeferredKinematicUpdateSkelMeshes.FindByPredicate([InSkelComp](const TPair<USkeletalMeshComponent*, FDeferredKinematicUpdateInfo>& InItem) { return InSkelComp == InItem.Key; });
+			if (!ensure(FoundItem != nullptr))// If the bool was set, we must be in the array!
+			{
+				FDeferredKinematicUpdateInfo Info;
+				Info.TeleportType = InTeleport;
+				Info.bNeedsSkinning = bNeedsSkinning;
+				DeferredKinematicUpdateSkelMeshes.Emplace(InSkelComp, Info);
+				return true;
+			}
+
+			FDeferredKinematicUpdateInfo& Info = FoundItem->Value;
 
 			// If we are currently not going to teleport physics, but this update wants to, we 'upgrade' it
 			if (Info.TeleportType == ETeleportType::None && InTeleport == ETeleportType::TeleportPhysics)
@@ -1614,8 +1622,10 @@ bool FPhysScene_ChaosInterface::MarkForPreSimKinematicUpdate(USkeletalMeshCompon
 			FDeferredKinematicUpdateInfo Info;
 			Info.TeleportType = InTeleport;
 			Info.bNeedsSkinning = bNeedsSkinning;
-			InSkelComp->DeferredKinematicUpdateIndex = DeferredKinematicUpdateSkelMeshes.Num();
 			DeferredKinematicUpdateSkelMeshes.Emplace(InSkelComp, Info);
+
+			// Set flag on component
+			InSkelComp->bDeferredKinematicUpdate = true;
 		}
 	}
 
@@ -1625,13 +1635,15 @@ bool FPhysScene_ChaosInterface::MarkForPreSimKinematicUpdate(USkeletalMeshCompon
 void FPhysScene_ChaosInterface::ClearPreSimKinematicUpdate(USkeletalMeshComponent* InSkelComp)
 {
 	// If non-null, and flagged for deferred update..
-	if (InSkelComp != nullptr && InSkelComp->DeferredKinematicUpdateIndex != INDEX_NONE)
+	if (InSkelComp != nullptr && InSkelComp->bDeferredKinematicUpdate)
 	{
 		// Remove from map
-		DeferredKinematicUpdateSkelMeshes.RemoveAtSwap(InSkelComp->DeferredKinematicUpdateIndex);
+		int32 NumRemoved = DeferredKinematicUpdateSkelMeshes.RemoveAll([InSkelComp](const TPair<USkeletalMeshComponent*, FDeferredKinematicUpdateInfo>& InItem) { return InSkelComp == InItem.Key; });
+
+		ensure(NumRemoved == 1); // Should be in array if flag was set!
 
-		// Clear Index
-		InSkelComp->DeferredKinematicUpdateIndex = INDEX_NONE;
+		// Clear flag
+		InSkelComp->bDeferredKinematicUpdate = false;
 	}
 }
 
@@ -1722,133 +1734,37 @@ void FPhysScene_ChaosInterface::UpdateKinematicsOnDeferredSkelMeshes()
 {
 	SCOPE_CYCLE_COUNTER(STAT_UpdateKinematicsOnDeferredSkelMeshesChaos);
 
-	// Holds start index in actor pool for each skeletal mesh.
-	TArray<int32, TInlineAllocator<64>> SkeletalMeshStartIndexArray;
-
-	TArray<FPhysicsActorHandle, TInlineAllocator<64>>TeleportActorsPool;
-	TArray<IPhysicsProxyBase*, TInlineAllocator<64>> ProxiesToDirty;
-	
-	// Count max number of bodies to determine actor pool size.
-	{
-		SkeletalMeshStartIndexArray.Reserve(DeferredKinematicUpdateSkelMeshes.Num());
+	using FActorsArray = TArray<FPhysicsActorHandle, TInlineAllocator<64>>;
+	using FTransformsArray = TArray<FTransform, TInlineAllocator<64>>;
 
-		int32 TotalBodies = 0;
-		for (const TPair<USkeletalMeshComponent*, FDeferredKinematicUpdateInfo>& DeferredKinematicUpdate : DeferredKinematicUpdateSkelMeshes)
-		{
-			SkeletalMeshStartIndexArray.Add(TotalBodies);
+	FActorsArray KinematicUpdateActors;
+	FTransformsArray KinematicUpdateTransforms;
+	FActorsArray TeleportActors;
+	FTransformsArray TeleportTransforms;
 
-			USkeletalMeshComponent* SkelComp = DeferredKinematicUpdate.Key;
-			if (!SkelComp->bEnablePerPolyCollision)
-			{
-				TotalBodies += SkelComp->Bodies.Num();
-			}
-		}
+	for (const TPair<USkeletalMeshComponent*, FDeferredKinematicUpdateInfo>& DeferredKinematicUpdate : DeferredKinematicUpdateSkelMeshes)
+	{
+		USkeletalMeshComponent* SkelComp = DeferredKinematicUpdate.Key;
+		const FDeferredKinematicUpdateInfo& Info = DeferredKinematicUpdate.Value;
 
-		// Actors pool is spare, initialize to nullptr.
-		TeleportActorsPool.AddZeroed(TotalBodies);
-		ProxiesToDirty.Reserve(TotalBodies);
-	}
+		ensure(SkelComp->bDeferredKinematicUpdate); // Should be true if in map!
 
-	// Gather proxies that need to be dirtied before paralell loop, and update any per poly collision skeletal meshes.
-	{
-		for (const TPair<USkeletalMeshComponent*, FDeferredKinematicUpdateInfo>& DeferredKinematicUpdate : DeferredKinematicUpdateSkelMeshes)
+		if (!SkelComp->bEnablePerPolyCollision)
 		{
-			USkeletalMeshComponent* SkelComp = DeferredKinematicUpdate.Key;
-			const FDeferredKinematicUpdateInfo& Info = DeferredKinematicUpdate.Value;
-
-			if (!SkelComp->bEnablePerPolyCollision)
-			{
-				const int32 NumBodies = SkelComp->Bodies.Num();
-				for (int32 i = 0; i < NumBodies; i++)
-				{
-					FBodyInstance* BodyInst = SkelComp->Bodies[i];
-					FPhysicsActorHandle& ActorHandle = BodyInst->ActorHandle;
-					if (!BodyInst->IsInstanceSimulatingPhysics())
-					{
-						const int32 BoneIndex = BodyInst->InstanceBoneIndex;
-						if (BoneIndex != INDEX_NONE)
-						{
-							IPhysicsProxyBase* Proxy = ActorHandle->GetProxy();
-							if (Proxy && Proxy->GetDirtyIdx() == INDEX_NONE)
-							{
-								ProxiesToDirty.Add(Proxy);
-							}
-						}
-					}
-				}
-			}
-			else
-			{
-				// TODO: acceleration for per-poly collision
-				SkelComp->UpdateKinematicBonesToAnim(SkelComp->GetComponentSpaceTransforms(), Info.TeleportType, Info.bNeedsSkinning, EAllowKinematicDeferral::DisallowDeferral);
-			}
+			GatherActorsAndTransforms(SkelComp, SkelComp->GetComponentSpaceTransforms(), Info.TeleportType, Info.bNeedsSkinning, KinematicUpdateActors, KinematicUpdateTransforms, TeleportActors, TeleportTransforms);
 		}
-	}
-
-	// Mark all body's proxies as dirty, as this is not threadsafe and cannot be done in parallel loop.
-	if (ProxiesToDirty.Num() > 0)
-	{
-		// Assumes all particles have the same solver, safe for now, maybe not in the future.
-		IPhysicsProxyBase* Proxy = ProxiesToDirty[0];
-		Chaos::FPhysicsSolverBase* Solver = Proxy->GetSolver();
-		Solver->AddDirtyProxiesUnsafe(ProxiesToDirty);
-	}
-
-	{
-		Chaos::PhysicsParallelFor(DeferredKinematicUpdateSkelMeshes.Num(), [&](int32 Index)
+		else
 		{
-			const TPair<USkeletalMeshComponent*, FDeferredKinematicUpdateInfo>& DeferredKinematicUpdate = DeferredKinematicUpdateSkelMeshes[Index];
-			USkeletalMeshComponent* SkelComp = DeferredKinematicUpdate.Key;
-			const FDeferredKinematicUpdateInfo& Info = DeferredKinematicUpdate.Value;
-
-			SkelComp->DeferredKinematicUpdateIndex = INDEX_NONE;
-
-			if (!SkelComp->bEnablePerPolyCollision)
-			{
-				const UPhysicsAsset* PhysicsAsset = SkelComp->GetPhysicsAsset();
-				const FTransform& CurrentLocalToWorld = SkelComp->GetComponentTransform();
-				const int32 NumBodies = SkelComp->Bodies.Num();
-				const TArray<FTransform>& ComponentSpaceTransforms = SkelComp->GetComponentSpaceTransforms();
-				
-				const int32 ActorPoolStartIndex = SkeletalMeshStartIndexArray[Index];
-				for (int32 i = 0; i < NumBodies; i++)
-				{
-					FBodyInstance* BodyInst = SkelComp->Bodies[i];
-					FPhysicsActorHandle& ActorHandle = BodyInst->ActorHandle;
-					if (!BodyInst->IsInstanceSimulatingPhysics())
-					{
-						const int32 BoneIndex = BodyInst->InstanceBoneIndex;
-						if (BoneIndex != INDEX_NONE)
-						{
-							const FTransform BoneTransform = ComponentSpaceTransforms[BoneIndex] * CurrentLocalToWorld;
-
-							TeleportActorsPool[ActorPoolStartIndex + i] = ActorHandle;
-
-							// TODO: Kinematic targets. Check Teleport type on FDeferredKinematicUpdateInfo and don't always teleport.
-							ActorHandle->SetX(BoneTransform.GetLocation(), false);	// only set dirty once in SetR
-							ActorHandle->SetR(BoneTransform.GetRotation());
-							ActorHandle->UpdateShapeBounds(BoneTransform);
+			// TODO: acceleration for per-poly collision
+			SkelComp->UpdateKinematicBonesToAnim(SkelComp->GetComponentSpaceTransforms(), Info.TeleportType, Info.bNeedsSkinning, EAllowKinematicDeferral::DisallowDeferral);
+		}
 
-							if (!PhysicsAsset->SkeletalBodySetups[i]->bSkipScaleFromAnimation)
-							{
-								const FVector& MeshScale3D = CurrentLocalToWorld.GetScale3D();
-								if (MeshScale3D.IsUniform())
-								{
-									BodyInst->UpdateBodyScale(BoneTransform.GetScale3D());
-								}
-								else
-								{
-									BodyInst->UpdateBodyScale(MeshScale3D);
-								}
-							}
-						}
-					}
-				}
-			}
-		});
+		// Clear deferred flag
+		SkelComp->bDeferredKinematicUpdate = false;
 	}
 
-	Scene.UpdateActorsInAccelerationStructure(TeleportActorsPool);
+	ProcessKinematicTargetActors(Scene, KinematicUpdateActors, KinematicUpdateTransforms);
+	ProcessTeleportActors(Scene, TeleportActors, TeleportTransforms);
 
 	DeferredKinematicUpdateSkelMeshes.Reset();
 }
diff --git a/Engine/Source/Runtime/Experimental/Chaos/Public/Chaos/Framework/PhysicsSolverBase.h b/Engine/Source/Runtime/Experimental/Chaos/Public/Chaos/Framework/PhysicsSolverBase.h
--- a/Engine/Source/Runtime/Experimental/Chaos/Public/Chaos/Framework/PhysicsSolverBase.h
+++ b/Engine/Source/Runtime/Experimental/Chaos/Public/Chaos/Framework/PhysicsSolverBase.h
@@ -55,20 +55,6 @@ namespace Chaos
 			}
 		}
 
-		// Batch proxy insertion, does not check DirtyIdx.
-		template< typename TProxiesArray>
-		void AddMultipleUnsafe(TProxiesArray& ProxiesArray)
-		{
-			int32 Idx = ProxiesData.Num();
-			ProxiesData.Append(ProxiesArray);
-
-			for(IPhysicsProxyBase* Proxy : ProxiesArray)
-			{
-				Proxy->SetDirtyIdx(Idx++);
-			}
-		}
-
-
 		void Remove(IPhysicsProxyBase* Base)
 		{
 			const int32 Idx = Base->GetDirtyIdx();
@@ -187,13 +173,6 @@ namespace Chaos
 			DirtyProxiesDataBuffer.AccessProducerBuffer()->Remove(ProxyBaseIn);
 		}
 
-		// Batch dirty proxies without checking DirtyIdx.
-		template <typename TProxiesArray>
-		void AddDirtyProxiesUnsafe(TProxiesArray& ProxiesArray)
-		{
-			DirtyProxiesDataBuffer.AccessProducerBuffer()->AddMultipleUnsafe(ProxiesArray);
-		}
-
 		void AddDirtyProxyShape(IPhysicsProxyBase* ProxyBaseIn, int32 ShapeIdx)
 		{
 			DirtyProxiesDataBuffer.AccessProducerBuffer()->AddShape(ProxyBaseIn,ShapeIdx);
diff --git a/Engine/Source/Runtime/Experimental/Chaos/Public/Chaos/ParticleHandle.h b/Engine/Source/Runtime/Experimental/Chaos/Public/Chaos/ParticleHandle.h
--- a/Engine/Source/Runtime/Experimental/Chaos/Public/Chaos/ParticleHandle.h
+++ b/Engine/Source/Runtime/Experimental/Chaos/Public/Chaos/ParticleHandle.h
@@ -1446,17 +1446,12 @@ public:
 	}
 
 	void UpdateShapeBounds()
-	{
-		UpdateShapeBounds(FRigidTransform3(X(), R()));
-	}
-
-	void UpdateShapeBounds(const FTransform& Transform)
 	{
 		if (MNonFrequentData.Read().Geometry->HasBoundingBox())
 		{
 			for (auto& Shape : MShapesArray)
 			{
-				Shape->UpdateShapeBounds(Transform);
+				Shape->UpdateShapeBounds(FRigidTransform3(X(), R()));
 			}
 		}
 	}
