Added missing DDC stats for a few types

#rnx

diff --git a/Engine/Source/Runtime/Engine/Private/Animation/AnimSequence.cpp b/Engine/Source/Runtime/Engine/Private/Animation/AnimSequence.cpp
--- a/Engine/Source/Runtime/Engine/Private/Animation/AnimSequence.cpp
+++ b/Engine/Source/Runtime/Engine/Private/Animation/AnimSequence.cpp
@@ -44,6 +44,7 @@
 #include "DeviceProfiles/DeviceProfile.h"
 #include "Animation/AnimStreamable.h"
 #include "Modules/ModuleManager.h"
+#include "ProfilingDebugging/CookStats.h"
 
 #define USE_SLERP 0
 #define LOCTEXT_NAMESPACE "AnimSequence"
@@ -51,6 +52,17 @@
 DECLARE_CYCLE_STAT(TEXT("AnimSeq GetBonePose"), STAT_AnimSeq_GetBonePose, STATGROUP_Anim);
 DECLARE_CYCLE_STAT(TEXT("AnimSeq EvalCurveData"), STAT_AnimSeq_EvalCurveData, STATGROUP_Anim);
 
+#if ENABLE_COOK_STATS
+namespace AnimSequenceCookStats
+{
+	static FCookStats::FDDCResourceUsageStats UsageStats;
+	static FCookStatsManager::FAutoRegisterCallback RegisterCookStats([](FCookStatsManager::AddStatFuncRef AddStat)
+	{
+		UsageStats.LogStats(AddStat, TEXT("AnimSequence.Usage"), TEXT(""));
+	});
+}
+#endif
+
 CSV_DECLARE_CATEGORY_MODULE_EXTERN(ENGINE_API, Animation);
 
 int32 GPerformFrameStripping = 0;
@@ -2187,8 +2199,10 @@ void UAnimSequence::WaitOnExistingCompression(const bool bWantResults)
 	check(IsInGameThread());
 	if (bCompressionInProgress)
 	{
+		COOK_STAT(auto Timer = AnimSequenceCookStats::UsageStats.TimeAsyncWait());
 		FAsyncCompressedAnimationsManagement::Get().WaitOnExistingCompression(this, bWantResults);
 		bCompressionInProgress = false;
+		COOK_STAT(Timer.TrackCyclesOnly()); // Need to get hit/miss and size from WaitOnExistingCompression!
 	}
 #endif
 }
@@ -2249,6 +2263,8 @@ void UAnimSequence::RequestAnimCompression(FRequestAnimCompressionParams Params)
 
 	TArray<uint8> OutData;
 	{
+		COOK_STAT(auto Timer = AnimSequenceCookStats::UsageStats.TimeSyncWork());
+
 		FDerivedDataAnimationCompression* AnimCompressor = new FDerivedDataAnimationCompression(TEXT("AnimSeq"), AssetDDCKey, Params.CompressContext);
 
 		const FString FinalDDCKey = FDerivedDataCacheInterface::BuildCacheKey(AnimCompressor->GetPluginName(), AnimCompressor->GetVersionString(), *AnimCompressor->GetPluginSpecificCacheKeySuffix());
@@ -2256,7 +2272,11 @@ void UAnimSequence::RequestAnimCompression(FRequestAnimCompressionParams Params)
 		// For debugging DDC/Compression issues		
 		const bool bSkipDDC = false;
 
-		if (bSkipDDC || !GetDerivedDataCacheRef().GetSynchronous(*FinalDDCKey, OutData, AnimCompressor->GetDebugContextString()))
+		if (!bSkipDDC && GetDerivedDataCacheRef().GetSynchronous(*FinalDDCKey, OutData, AnimCompressor->GetDebugContextString()))
+		{
+			COOK_STAT(Timer.AddHit(OutData.Num()));
+		}
+		else
 		{
 			// Data does not exist, need to build it.
 			// Filter RAW data to get rid of mismatched tracks (translation/rotation data with a different number of keys than there are frames)
@@ -2269,19 +2289,27 @@ void UAnimSequence::RequestAnimCompression(FRequestAnimCompressionParams Params)
 			if (bSkipDDC || (CompressCommandletVersion == INDEX_NONE))
 			{
 				AnimCompressor->Build(OutData);
+				COOK_STAT(Timer.AddMiss(OutData.Num()));
 			}
 			else if (AnimCompressor->CanBuild())
 			{
 				if (Params.bAsyncCompression)
 				{
 					FAsyncCompressedAnimationsManagement::Get().RequestAsyncCompression(*AnimCompressor, this, bPerformStripping, OutData);
+					COOK_STAT(Timer.TrackCyclesOnly());
 				}
 				else
 				{
-					GetDerivedDataCacheRef().GetSynchronous(AnimCompressor, OutData);
+					bool bBuilt = false;
+					const bool bSuccess = GetDerivedDataCacheRef().GetSynchronous(AnimCompressor, OutData, &bBuilt);
+					COOK_STAT(Timer.AddHitOrMiss(!bSuccess || bBuilt ? FCookStats::CallStats::EHitOrMiss::Miss : FCookStats::CallStats::EHitOrMiss::Hit, OutData.Num()));
 				}
 				AnimCompressor = nullptr;
 			}
+			else
+			{
+				COOK_STAT(Timer.TrackCyclesOnly());
+			}
 		}
 		else
 		{
diff --git a/Engine/Source/Runtime/Engine/Private/Animation/AnimStreamable.cpp b/Engine/Source/Runtime/Engine/Private/Animation/AnimStreamable.cpp
--- a/Engine/Source/Runtime/Engine/Private/Animation/AnimStreamable.cpp
+++ b/Engine/Source/Runtime/Engine/Private/Animation/AnimStreamable.cpp
@@ -18,11 +18,23 @@
 #include "Animation/AnimCurveCompressionCodec.h"
 #include "BonePose.h"
 #include "ContentStreaming.h"
+#include "ProfilingDebugging/CookStats.h"
 
 CSV_DECLARE_CATEGORY_MODULE_EXTERN(ENGINE_API, Animation);
 
 DECLARE_CYCLE_STAT(TEXT("AnimStreamable GetAnimationPose"), STAT_AnimStreamable_GetAnimationPose, STATGROUP_Anim);
 
+#if ENABLE_COOK_STATS
+namespace AnimStreamableCookStats
+{
+	static FCookStats::FDDCResourceUsageStats UsageStats;
+	static FCookStatsManager::FAutoRegisterCallback RegisterCookStats([](FCookStatsManager::AddStatFuncRef AddStat)
+	{
+		UsageStats.LogStats(AddStat, TEXT("AnimStreamable.Usage"), TEXT(""));
+	});
+}
+#endif
+
 // This is a version string for the streaming anim chunk logic
 // If you want to bump this version, generate a new guid using VS->Tools->Create GUID and
 // set it here
@@ -597,6 +609,8 @@ void UAnimStreamable::RequestCompressedDataForChunk(const FString& ChunkDDCKey,
 
 	TArray<uint8> OutData;
 	{
+		COOK_STAT(auto Timer = AnimStreamableCookStats::UsageStats.TimeSyncWork());
+
 		bool bNeedToCleanUpAnimCompressor = true;
 		FDerivedDataAnimationCompression* AnimCompressor = new FDerivedDataAnimationCompression(TEXT("StreamAnim"), ChunkDDCKey, CompressContext);
 
@@ -610,7 +624,11 @@ void UAnimStreamable::RequestCompressedDataForChunk(const FString& ChunkDDCKey,
 		Chunk.StartTime = FrameStart * FrameLength;
 		Chunk.SequenceLength = ChunkNumFrames * FrameLength;
 
-		if (bSkipDDC || !GetDerivedDataCacheRef().GetSynchronous(*FinalDDCKey, OutData, GetPathName()))
+		if (!bSkipDDC && GetDerivedDataCacheRef().GetSynchronous(*FinalDDCKey, OutData, GetPathName()))
+		{
+			COOK_STAT(Timer.AddHit(OutData.Num()));
+		}
+		else
 		{
 			FCompressibleAnimRef CompressibleData = MakeShared<FCompressibleAnimData, ESPMode::ThreadSafe>(BoneCompressionSettings, CurveCompressionSettings, GetSkeleton(), Interpolation, Chunk.SequenceLength, ChunkNumFrames+1);
 
@@ -654,11 +672,18 @@ void UAnimStreamable::RequestCompressedDataForChunk(const FString& ChunkDDCKey,
 			if (bSkipDDC)
 			{
 				AnimCompressor->Build(OutData);
+				COOK_STAT(Timer.AddMiss(OutData.Num()));
 			}
 			else if (AnimCompressor->CanBuild())
 			{
 				bNeedToCleanUpAnimCompressor = false; // GetSynchronous will handle this
-				GetDerivedDataCacheRef().GetSynchronous(AnimCompressor, OutData);
+				bool bBuilt = false;
+				const bool bSuccess = GetDerivedDataCacheRef().GetSynchronous(AnimCompressor, OutData, &bBuilt);
+				COOK_STAT(Timer.AddHitOrMiss(!bSuccess || bBuilt ? FCookStats::CallStats::EHitOrMiss::Miss : FCookStats::CallStats::EHitOrMiss::Hit, OutData.Num()));
+			}
+			else
+			{
+				COOK_STAT(Timer.TrackCyclesOnly());
 			}
 		}
 
diff --git a/Engine/Source/Runtime/Engine/Private/SoundWave.cpp b/Engine/Source/Runtime/Engine/Private/SoundWave.cpp
--- a/Engine/Source/Runtime/Engine/Private/SoundWave.cpp
+++ b/Engine/Source/Runtime/Engine/Private/SoundWave.cpp
@@ -683,6 +683,8 @@ void USoundWave::BeginGetCompressedData(FName Format, const FPlatformAudioCookOv
 	{
 		if (GetDerivedDataCache())
 		{
+			COOK_STAT(auto Timer = SoundWaveCookStats::UsageStats.TimeSyncWork());
+			COOK_STAT(Timer.TrackCyclesOnly());
 			FDerivedAudioDataCompressor* DeriveAudioData = new FDerivedAudioDataCompressor(this, Format, PlatformSpecificFormat, CompressionOverrides);
 			uint32 GetHandle = GetDerivedDataCacheRef().GetAsynchronous(DeriveAudioData);
 			AsyncLoadingDataFormats.Add(PlatformSpecificFormat, GetHandle);
@@ -716,9 +718,15 @@ FByteBulkData* USoundWave::GetCompressedData(FName Format, const FPlatformAudioC
 			bool bDataWasBuilt = false;
 			bool bGetSuccessful = false;
 
-			COOK_STAT(auto Timer = SoundWaveCookStats::UsageStats.TimeSyncWork());
 #if WITH_EDITOR
 			uint32* AsyncHandle = AsyncLoadingDataFormats.Find(PlatformSpecificFormat);
+#else
+			uint32* AsyncHandle = nullptr;
+#endif
+
+			COOK_STAT(auto Timer = AsyncHandle ? SoundWaveCookStats::UsageStats.TimeAsyncWait() : SoundWaveCookStats::UsageStats.TimeSyncWork());
+
+#if WITH_EDITOR
 			if (AsyncHandle)
 			{
 				GetDerivedDataCacheRef().WaitAsynchronousCompletion(*AsyncHandle);
diff --git a/Engine/Source/Runtime/Engine/Private/StaticMesh.cpp b/Engine/Source/Runtime/Engine/Private/StaticMesh.cpp
--- a/Engine/Source/Runtime/Engine/Private/StaticMesh.cpp
+++ b/Engine/Source/Runtime/Engine/Private/StaticMesh.cpp
@@ -3931,9 +3931,11 @@ bool UStaticMesh::LoadMeshDescription(int32 LodIndex, FMeshDescription& OutMeshD
 	FString MeshDataKey;
 	if (GetMeshDataKey(LodIndex, MeshDataKey))
 	{
+		COOK_STAT(auto Timer = StaticMeshCookStats::UsageStats.TimeSyncWork());
 		TArray<uint8> DerivedData;
 		if (GetDerivedDataCacheRef().GetSynchronous(*MeshDataKey, DerivedData, GetPathName()))
 		{
+			COOK_STAT(Timer.AddHit(DerivedData.Num()));
 			// If there was valid DDC data, we assume this is because the asset is an old one with valid RawMeshBulkData
 			check(!SourceModel.RawMeshBulkData->IsEmpty());
 
@@ -3949,6 +3951,10 @@ bool UStaticMesh::LoadMeshDescription(int32 LodIndex, FMeshDescription& OutMeshD
 			MeshDescriptionBulkData.LoadMeshDescription(OutMeshDescription);
 			return true;
 		}
+		else
+		{
+			COOK_STAT(Timer.TrackCyclesOnly());
+		}
 	}
 
 	// If after all this we *still* don't have a valid MeshDescription, but there's a valid RawMesh, convert that to a MeshDescription.
@@ -4226,8 +4232,13 @@ void UStaticMesh::CacheMeshData()
 				FString MeshDataKey;
 				if (GetMeshDataKey(LodIndex, MeshDataKey))
 				{
+					COOK_STAT(auto Timer = StaticMeshCookStats::UsageStats.TimeSyncWork());
 					// If the DDC key doesn't exist, convert the data and save it to DDC
-					if (!GetDerivedDataCacheRef().CachedDataProbablyExists(*MeshDataKey))
+					if (GetDerivedDataCacheRef().CachedDataProbablyExists(*MeshDataKey))
+					{
+						COOK_STAT(Timer.TrackCyclesOnly());
+					}
+					else
 					{
 						// Get the RawMesh for this LOD
 						FRawMesh TempRawMesh;
@@ -4252,6 +4263,7 @@ void UStaticMesh::CacheMeshData()
 						FMemoryWriter Ar(DerivedData, bIsPersistent);
 						MeshDescriptionBulkData.Serialize(Ar, this);
 						GetDerivedDataCacheRef().Put(*MeshDataKey, DerivedData, GetPathName());
+						COOK_STAT(Timer.AddMiss(DerivedData.Num()));
 					}
 				}
 			}
diff --git a/Engine/Source/Runtime/Experimental/GeometryCollectionEngine/Private/GeometryCollection/GeometryCollectionObject.cpp b/Engine/Source/Runtime/Experimental/GeometryCollectionEngine/Private/GeometryCollection/GeometryCollectionObject.cpp
--- a/Engine/Source/Runtime/Experimental/GeometryCollectionEngine/Private/GeometryCollection/GeometryCollectionObject.cpp
+++ b/Engine/Source/Runtime/Experimental/GeometryCollectionEngine/Private/GeometryCollection/GeometryCollectionObject.cpp
@@ -12,6 +12,7 @@
 #include "HAL/IConsoleManager.h"
 #include "UObject/Package.h"
 #include "Materials/MaterialInstance.h"
+#include "ProfilingDebugging/CookStats.h"
 
 #if WITH_EDITOR
 #include "GeometryCollection/DerivedDataGeometryCollectionCooker.h"
@@ -25,6 +26,17 @@
 
 DEFINE_LOG_CATEGORY_STATIC(UGeometryCollectionLogging, NoLogging, All);
 
+#if ENABLE_COOK_STATS
+namespace GeometryCollectionCookStats
+{
+	static FCookStats::FDDCResourceUsageStats UsageStats;
+	static FCookStatsManager::FAutoRegisterCallback RegisterCookStats([](FCookStatsManager::AddStatFuncRef AddStat)
+	{
+		UsageStats.LogStats(AddStat, TEXT("GeometryCollection.Usage"), TEXT(""));
+	});
+}
+#endif
+
 UGeometryCollection::UGeometryCollection(const FObjectInitializer& ObjectInitializer)
 	: Super(ObjectInitializer)
 	, CollisionType(ECollisionTypeEnum::Chaos_Surface_Volumetric)
@@ -352,6 +364,8 @@ void UGeometryCollection::Serialize(FArchive& Ar)
 #if WITH_EDITOR
 void UGeometryCollection::CreateSimulationDataImp(bool bCopyFromDDC, const TCHAR* OverrideVersion)
 {
+	COOK_STAT(auto Timer = GeometryCollectionCookStats::UsageStats.TimeSyncWork());
+
 	// Skips the DDC fetch entirely for testing the builder without adding to the DDC
 	const static bool bSkipDDC = false;
 
@@ -362,16 +376,19 @@ void UGeometryCollection::CreateSimulationDataImp(bool bCopyFromDDC, const TCHAR
 
 	if (GeometryCollectionCooker->CanBuild())
 	{
-		if(bSkipDDC)
+		if (bSkipDDC)
 		{
 			GeometryCollectionCooker->Build(DDCData);
+			COOK_STAT(Timer.AddMiss(DDCData.Num()));
 		}
 		else
 		{
-			GetDerivedDataCacheRef().GetSynchronous(GeometryCollectionCooker, DDCData);
+			bool bBuilt = false;
+			const bool bSuccess = GetDerivedDataCacheRef().GetSynchronous(GeometryCollectionCooker, DDCData, &bBuilt);
+			COOK_STAT(Timer.AddHitOrMiss(!bSuccess || bBuilt ? FCookStats::CallStats::EHitOrMiss::Miss : FCookStats::CallStats::EHitOrMiss::Hit, DDCData.Num()));
 		}
 
-		if(bCopyFromDDC)
+		if (bCopyFromDDC)
 		{
 			FMemoryReader Ar(DDCData);
 			Chaos::FChaosArchive ChaosAr(Ar);
