Fix a merge error in AnimSequence.cpp

#rnx

diff --git a/Engine/Source/Runtime/Engine/Private/Animation/AnimSequence.cpp b/Engine/Source/Runtime/Engine/Private/Animation/AnimSequence.cpp
--- a/Engine/Source/Runtime/Engine/Private/Animation/AnimSequence.cpp
+++ b/Engine/Source/Runtime/Engine/Private/Animation/AnimSequence.cpp
@@ -2275,6 +2275,7 @@ void UAnimSequence::RequestAnimCompression(FRequestAnimCompressionParams Params)
 		if (!bSkipDDC && GetDerivedDataCacheRef().GetSynchronous(*FinalDDCKey, OutData, AnimCompressor->GetDebugContextString()))
 		{
 			COOK_STAT(Timer.AddHit(OutData.Num()));
+			bCompressedDataFromDDC = true;
 		}
 		else
 		{
@@ -2311,10 +2312,6 @@ void UAnimSequence::RequestAnimCompression(FRequestAnimCompressionParams Params)
 				COOK_STAT(Timer.TrackCyclesOnly());
 			}
 		}
-		else
-		{
-			bCompressedDataFromDDC = true;
-		}
 
 		if (AnimCompressor)
 		{
