Fixed a staticmesh HLOD rebuild crash with some recent code

diff --git a/Engine/Source/Runtime/Engine/Private/StaticMesh.cpp b/Engine/Source/Runtime/Engine/Private/StaticMesh.cpp
--- a/Engine/Source/Runtime/Engine/Private/StaticMesh.cpp
+++ b/Engine/Source/Runtime/Engine/Private/StaticMesh.cpp
@@ -2719,7 +2719,7 @@ void UStaticMesh::InitResources()
 	{
 		LinkStreaming();
 	}
-	else if (GIsEditor)
+	else if (GIsEditor && RenderData)
 	{
 		RenderData->CurrentFirstLODIdx = RenderData->GetFirstValidLODIdx(0);
 		SetCachedNumResidentLODs(RenderData->LODResources.Num() - RenderData->CurrentFirstLODIdx);
