Fix RandomSeed difference seen on different devices
Function argument order evaluation is undefined

diff --git a/Engine/Source/Runtime/Engine/Private/Distributions.cpp b/Engine/Source/Runtime/Engine/Private/Distributions.cpp
--- a/Engine/Source/Runtime/Engine/Private/Distributions.cpp
+++ b/Engine/Source/Runtime/Engine/Private/Distributions.cpp
@@ -770,12 +770,11 @@ void FRawDistribution::GetValue3Random(float Time, float* InValue, struct FRando
 	const float* Entry1;
 	const float* Entry2;
 	float LerpAlpha = 0.0f;
-	FVector RandValues(
-		DIST_GET_RANDOM_VALUE(InRandomStream),
-		DIST_GET_RANDOM_VALUE(InRandomStream),
-		DIST_GET_RANDOM_VALUE(InRandomStream)
-		);
+	FVector RandValues;
 
+	RandValues[0] = DIST_GET_RANDOM_VALUE(InRandomStream);
+	RandValues[1] = DIST_GET_RANDOM_VALUE(InRandomStream);
+	RandValues[2] = DIST_GET_RANDOM_VALUE(InRandomStream);
 	switch(LookupTable.LockFlag)
 	{
 	case EDVLF_XY:
