Recache referenced textures when material instance rebuilds

diff --git a/Engine/Source/Runtime/Engine/Private/Materials/MaterialInstance.cpp b/Engine/Source/Runtime/Engine/Private/Materials/MaterialInstance.cpp
--- a/Engine/Source/Runtime/Engine/Private/Materials/MaterialInstance.cpp
+++ b/Engine/Source/Runtime/Engine/Private/Materials/MaterialInstance.cpp
@@ -2929,6 +2929,10 @@ void UMaterialInstance::CacheShadersForResources(EShaderPlatform ShaderPlatform,
 	check(!BaseMaterial->HasAnyFlags(RF_NeedPostLoad));
 #endif
 
+#if WITH_EDITOR
+	UpdateCachedLayerParameters();
+#endif
+
 	for (int32 ResourceIndex = 0; ResourceIndex < ResourcesToCache.Num(); ResourceIndex++)
 	{
 		FMaterialResource* CurrentResource = ResourcesToCache[ResourceIndex];
