UE4.25 - Compile fix
#rnx

diff --git a/Engine/Source/Runtime/Engine/Private/Streaming/StaticMeshUpdate.cpp b/Engine/Source/Runtime/Engine/Private/Streaming/StaticMeshUpdate.cpp
--- a/Engine/Source/Runtime/Engine/Private/Streaming/StaticMeshUpdate.cpp
+++ b/Engine/Source/Runtime/Engine/Private/Streaming/StaticMeshUpdate.cpp
@@ -13,7 +13,7 @@ StaticMeshUpdate.cpp: Helpers to stream in and out static mesh LODs.
 #include "Streaming/RenderAssetUpdate.inl"
 
 int32 GStreamingMaxReferenceChecks = 3;
-static FAutoConsoleVariableRef CVarGPUSortFrameCountBeforeBufferShrinking(
+static FAutoConsoleVariableRef CVarStreamingMaxReferenceChecksBeforeStreamOut(
 	TEXT("r.Streaming.MaxReferenceChecksBeforeStreamOut"),
 	GStreamingMaxReferenceChecks,
 	TEXT("Number of times the engine wait for references to be released before forcing streamout. (default=3)"),
