Fix mesh loading dependency for static mesh using a ComplexCollisionMesh
Fix phys meshes to be created before creating the navcollision
Rotate the navcollision DDC guid
Those are to fix missing collisions when creating navcoll for meshes using ComplexCollisionMesh

diff --git a/Engine/Source/Runtime/Engine/Private/StaticMesh.cpp b/Engine/Source/Runtime/Engine/Private/StaticMesh.cpp
--- a/Engine/Source/Runtime/Engine/Private/StaticMesh.cpp
+++ b/Engine/Source/Runtime/Engine/Private/StaticMesh.cpp
@@ -5136,6 +5136,12 @@ void UStaticMesh::PostLoad()
 	ClearMeshDescriptions();
 #endif
 
+	// If any, make sure the ComplexCollisionMesh is loaded before creating the NavCollision
+	if (ComplexCollisionMesh && ComplexCollisionMesh != this)
+	{
+		ComplexCollisionMesh->ConditionalPostLoad();
+	}
+
 	CreateNavCollision();
 }
 
@@ -6001,6 +6007,9 @@ void UStaticMesh::CreateNavCollision(const bool bIsUpdate)
 				NavCollision->InvalidateCollision();
 			}
 #endif // WITH_EDITOR
+
+			// Physics meshes need to be ready to gather the collision in Setup().
+			BodySetup->CreatePhysicsMeshes();
 			NavCollision->Setup(BodySetup);
 		}
 	}
diff --git a/Engine/Source/Runtime/NavigationSystem/Private/NavCollision.cpp b/Engine/Source/Runtime/NavigationSystem/Private/NavCollision.cpp
--- a/Engine/Source/Runtime/NavigationSystem/Private/NavCollision.cpp
+++ b/Engine/Source/Runtime/NavigationSystem/Private/NavCollision.cpp
@@ -94,7 +94,7 @@ public:
 
 	virtual const TCHAR* GetVersionString() const override
 	{
-		return TEXT("2EE6412E1E884D1AAA7F79FF760363BC");
+		return TEXT("2983273BE5D6425391E8C0B3D1C0A55D");
 	}
 
 	virtual FString GetPluginSpecificCacheKeySuffix() const override
