Build fix, missing exclusion

diff --git a/Engine/Source/Runtime/Engine/Private/StaticMesh.cpp b/Engine/Source/Runtime/Engine/Private/StaticMesh.cpp
--- a/Engine/Source/Runtime/Engine/Private/StaticMesh.cpp
+++ b/Engine/Source/Runtime/Engine/Private/StaticMesh.cpp
@@ -5136,11 +5136,13 @@ void UStaticMesh::PostLoad()
 	ClearMeshDescriptions();
 #endif
 
+#if WITH_EDITORONLY_DATA
 	// If any, make sure the ComplexCollisionMesh is loaded before creating the NavCollision
 	if (ComplexCollisionMesh && ComplexCollisionMesh != this)
 	{
 		ComplexCollisionMesh->ConditionalPostLoad();
 	}
+#endif //WITH_EDITORONLY_DATA
 
 	CreateNavCollision();
 }
