Added support for G8 format when converting a RenderTarget2D to a Texture2D.
Added a warning message when trying to convert and unsupported format.

diff --git a/Engine/Source/Runtime/Engine/Private/TextureRenderTarget2D.cpp b/Engine/Source/Runtime/Engine/Private/TextureRenderTarget2D.cpp
--- a/Engine/Source/Runtime/Engine/Private/TextureRenderTarget2D.cpp
+++ b/Engine/Source/Runtime/Engine/Private/TextureRenderTarget2D.cpp
@@ -298,6 +298,14 @@ UTexture2D* UTextureRenderTarget2D::ConstructTexture2D(UObject* Outer, const FSt
 	case PF_FloatRGBA:
 		TextureFormat = TSF_RGBA16F;
 		break;
+	case PF_G8:
+		TextureFormat = TSF_G8;
+		break;
+	default:
+	{
+		FText InvalidFormatMessage = NSLOCTEXT("TextureRenderTarget2D", "UnsupportedFormatRenderTarget2DWarning", "Unsupported format when creating Texture2D from TextureRenderTarget2D. Supported formats are B8G8R8A8, FloatRGBA and G8.");
+		FMessageDialog::Open(EAppMsgType::Ok, InvalidFormatMessage);
+	}
 	}
 
 	// exit if source is not compatible.
@@ -353,7 +361,7 @@ void UTextureRenderTarget2D::UpdateTexture2D(UTexture2D* InTexture2D, ETextureSo
 	// init to the same size as the 2d texture
 	InTexture2D->Source.Init(SizeX, SizeY, 1, 1, InTextureFormat);
 
-	uint32* TextureData = (uint32*)InTexture2D->Source.LockMip(0);
+	uint8* TextureData = (uint8*)InTexture2D->Source.LockMip(0);
 	const int32 TextureDataSize = InTexture2D->Source.CalcMipSize(0);
 
 	// read the 2d surface
@@ -423,6 +431,16 @@ void UTextureRenderTarget2D::UpdateTexture2D(UTexture2D* InTexture2D, ETextureSo
 		check(TextureDataSize == SurfData.Num() * sizeof(FFloat16Color));
 		FMemory::Memcpy(TextureData, SurfData.GetData(), TextureDataSize);
 	}
+	else if (InTextureFormat == TSF_G8)
+	{
+		TArray<FColor> SurfData;
+		RenderTarget->ReadPixels(SurfData);
+		check(TextureDataSize == SurfData.Num() * sizeof(uint8));
+		for (int32 Pixel = 0; Pixel < SurfData.Num(); Pixel++)
+		{
+			TextureData[Pixel] = SurfData[Pixel].R;
+		}
+	}
 
 	InTexture2D->Source.UnlockMip(0);
 
