Exposing some parameters in SoundClass/SoundMix/SoundCue to BlueprintReadOnly, etc so that we can perform some asset actions on them

diff --git a/Engine/Source/Runtime/Engine/Classes/Sound/SoundClass.h b/Engine/Source/Runtime/Engine/Classes/Sound/SoundClass.h
--- a/Engine/Source/Runtime/Engine/Classes/Sound/SoundClass.h
+++ b/Engine/Source/Runtime/Engine/Classes/Sound/SoundClass.h
@@ -224,7 +224,7 @@ class ENGINE_API USoundClass : public UObject
 	FSoundModulation Modulation;
 
 public:
-	UPROPERTY()
+	UPROPERTY(BlueprintReadOnly, Category = General)
 	USoundClass* ParentClass;
 
 #if WITH_EDITORONLY_DATA
diff --git a/Engine/Source/Runtime/Engine/Classes/Sound/SoundCue.h b/Engine/Source/Runtime/Engine/Classes/Sound/SoundCue.h
--- a/Engine/Source/Runtime/Engine/Classes/Sound/SoundCue.h
+++ b/Engine/Source/Runtime/Engine/Classes/Sound/SoundCue.h
@@ -92,11 +92,11 @@ class ENGINE_API USoundCue : public USoundBase
 	USoundNode* FirstNode;
 
 	/* Volume multiplier for the Sound Cue */
-	UPROPERTY(EditAnywhere, Category=Sound, AssetRegistrySearchable)
+	UPROPERTY(EditAnywhere, BlueprintReadOnly, Category=Sound, AssetRegistrySearchable)
 	float VolumeMultiplier;
 
 	/* Pitch multiplier for the Sound Cue */
-	UPROPERTY(EditAnywhere, Category=Sound, AssetRegistrySearchable)
+	UPROPERTY(EditAnywhere, BlueprintReadOnly, Category=Sound, AssetRegistrySearchable)
 	float PitchMultiplier;
 
 	/* Attenuation settings to use if Override Attenuation is set to true */
diff --git a/Engine/Source/Runtime/Engine/Classes/Sound/SoundMix.h b/Engine/Source/Runtime/Engine/Classes/Sound/SoundMix.h
--- a/Engine/Source/Runtime/Engine/Classes/Sound/SoundMix.h
+++ b/Engine/Source/Runtime/Engine/Classes/Sound/SoundMix.h
@@ -127,33 +127,33 @@ struct FAudioEQEffect : public FAudioEffectParameters
 /**
  * Elements of data for sound group volume control
  */
-USTRUCT()
+USTRUCT(BlueprintType)
 struct FSoundClassAdjuster
 {
 	GENERATED_USTRUCT_BODY()
 
 	/* The sound class this adjuster affects. */
-	UPROPERTY(EditAnywhere, Category=SoundClassAdjuster, DisplayName = "Sound Class" )
+	UPROPERTY(EditAnywhere, BlueprintReadOnly, Category=SoundClassAdjuster, DisplayName = "Sound Class" )
 	USoundClass* SoundClassObject;
 
 	/* A multiplier applied to the volume. */
-	UPROPERTY(EditAnywhere, Category=SoundClassAdjuster, meta = (ClampMin = "0.0", UIMin = "0.0", UIMax = "4.0"))
+	UPROPERTY(EditAnywhere, BlueprintReadOnly, Category=SoundClassAdjuster, meta = (ClampMin = "0.0", UIMin = "0.0", UIMax = "4.0"))
 	float VolumeAdjuster;
 
 	/* A multiplier applied to the pitch. */
-	UPROPERTY(EditAnywhere, Category=SoundClassAdjuster, meta = (ClampMin = "0.0", ClampMax = "8.0", UIMin = "0.0", UIMax = "8.0"))
+	UPROPERTY(EditAnywhere, BlueprintReadOnly, Category=SoundClassAdjuster, meta = (ClampMin = "0.0", ClampMax = "8.0", UIMin = "0.0", UIMax = "8.0"))
 	float PitchAdjuster;
 
 	/* Lowpass filter cutoff frequency to apply to sound sources. */
-	UPROPERTY(EditAnywhere, Category = SoundClassAdjuster, meta = (ClampMin = "0.0", ClampMax = "20000.0", UIMin = "0.0", UIMax = "20000.0"))
+	UPROPERTY(EditAnywhere, BlueprintReadOnly, Category = SoundClassAdjuster, meta = (ClampMin = "0.0", ClampMax = "20000.0", UIMin = "0.0", UIMax = "20000.0"))
 	float LowPassFilterFrequency;
 
 	/* Set to true to apply this adjuster to all children of the sound class. */
-	UPROPERTY(EditAnywhere, Category=SoundClassAdjuster )
+	UPROPERTY(EditAnywhere, BlueprintReadOnly, Category=SoundClassAdjuster )
 	uint32 bApplyToChildren:1;
 
 	/* A multiplier applied to VoiceCenterChannelVolume. */
-	UPROPERTY(EditAnywhere, Category=SoundClassAdjuster, meta = (ClampMin = "0.0", UIMin = "0.0", UIMax = "4.0"))
+	UPROPERTY(EditAnywhere, BlueprintReadOnly, Category=SoundClassAdjuster, meta = (ClampMin = "0.0", UIMin = "0.0", UIMax = "4.0"))
 	float VoiceCenterChannelVolumeAdjuster;
 
 	FSoundClassAdjuster()
@@ -184,7 +184,7 @@ class USoundMix : public UObject
 	struct FAudioEQEffect EQSettings;
 
 	/* Array of changes to be applied to groups. */
-	UPROPERTY(EditAnywhere, Category=SoundClasses)
+	UPROPERTY(EditAnywhere, BlueprintReadOnly, Category=SoundClasses)
 	TArray<struct FSoundClassAdjuster> SoundClassEffects;
 
 	/* Initial delay in seconds before the the mix is applied. */
