Adds HoldWhenFinished parameter to PlayerController::ClientSetCameraFade and adds support for fading from the PlayerCameraManager's current FadeAmount
#rnx

diff --git a/Engine/Source/Runtime/Engine/Classes/GameFramework/PlayerController.h b/Engine/Source/Runtime/Engine/Classes/GameFramework/PlayerController.h
--- a/Engine/Source/Runtime/Engine/Classes/GameFramework/PlayerController.h
+++ b/Engine/Source/Runtime/Engine/Classes/GameFramework/PlayerController.h
@@ -1002,12 +1002,13 @@ public:
 	 * Tell client to fade camera
 	 * @Param bEnableFading - true if we should apply FadeColor/FadeAmount to the screen
 	 * @Param FadeColor - Color to fade to
-	 * @Param FadeAlpha - Amount of fading to apply
+	 * @Param FadeAlpha - Contains the start fade (X) and end fade (Y) values to apply. A start fade of less than 0 will use the screen's current fade value
 	 * @Param FadeTime - length of time for fade to occur over
 	 * @Param bFadeAudio - true to apply fading of audio alongside the video
+	 * @param bHoldWhenFinished - True for fade to hold at the ToAlpha until fade is disabled
 	 */
 	UFUNCTION(Reliable, Client)
-	void ClientSetCameraFade(bool bEnableFading, FColor FadeColor = FColor(ForceInit), FVector2D FadeAlpha = FVector2D(ForceInit), float FadeTime = 0, bool bFadeAudio = false);
+	void ClientSetCameraFade(bool bEnableFading, FColor FadeColor = FColor(ForceInit), FVector2D FadeAlpha = FVector2D(-1.0f, 0.0f), float FadeTime = 0, bool bFadeAudio = false, bool bHoldWhenFinished = false);
 
 	/**
 	 * Replicated function to set camera style on client
diff --git a/Engine/Source/Runtime/Engine/Private/PlayerController.cpp b/Engine/Source/Runtime/Engine/Private/PlayerController.cpp
--- a/Engine/Source/Runtime/Engine/Private/PlayerController.cpp
+++ b/Engine/Source/Runtime/Engine/Private/PlayerController.cpp
@@ -1653,13 +1653,15 @@ void APlayerController::ResetCameraMode()
 
 /// @cond DOXYGEN_WARNINGS
 
-void APlayerController::ClientSetCameraFade_Implementation(bool bEnableFading, FColor FadeColor, FVector2D FadeAlpha, float FadeTime, bool bFadeAudio)
+void APlayerController::ClientSetCameraFade_Implementation(bool bEnableFading, FColor FadeColor, FVector2D FadeAlpha, float FadeTime, bool bFadeAudio, bool bHoldWhenFinished)
 {
 	if (PlayerCameraManager != nullptr)
 	{
 		if (bEnableFading)
 		{
-			PlayerCameraManager->StartCameraFade(FadeAlpha.X, FadeAlpha.Y, FadeTime, FadeColor.ReinterpretAsLinear(), bFadeAudio);
+			// Allow fading from the current FadeAmount to allow for smooth transitions into new fades
+			const float FadeStart = FadeAlpha.X >= 0.f ? FadeAlpha.X : PlayerCameraManager->FadeAmount;
+			PlayerCameraManager->StartCameraFade(FadeStart, FadeAlpha.Y, FadeTime, FadeColor.ReinterpretAsLinear(), bFadeAudio, bHoldWhenFinished);
 		}
 		else
 		{
