Fixed potential crash when ray tracing ISM with invalid render data

diff --git a/Engine/Source/Runtime/Engine/Private/InstancedStaticMesh.cpp b/Engine/Source/Runtime/Engine/Private/InstancedStaticMesh.cpp
--- a/Engine/Source/Runtime/Engine/Private/InstancedStaticMesh.cpp
+++ b/Engine/Source/Runtime/Engine/Private/InstancedStaticMesh.cpp
@@ -1134,6 +1134,11 @@ void FInstancedStaticMeshSceneProxy::GetDynamicRayTracingInstances(struct FRayTr
 	{
 		return;
 	}
+	
+	if (!InstancedRenderData.PerInstanceRenderData.IsValid())
+	{
+		return;
+	}
 
 	uint32 LOD = GetCurrentFirstLODIdx_RenderThread();
 	const int32 InstanceCount = InstancedRenderData.PerInstanceRenderData->InstanceBuffer.GetNumInstances();
@@ -1249,6 +1254,11 @@ void FInstancedStaticMeshSceneProxy::SetupRayTracingCullClusters()
 	{
 		return;
 	}
+	
+	if (!InstancedRenderData.PerInstanceRenderData.IsValid())
+	{
+		return;
+	}
 
 	//#dxr_todo: select the appropriate LOD depending on Context.View
 	int32 LOD = 0;
