Expose AllowShaderWorkers to cvars as well.

diff --git a/Engine/Config/ConsoleVariables.ini b/Engine/Config/ConsoleVariables.ini
--- a/Engine/Config/ConsoleVariables.ini
+++ b/Engine/Config/ConsoleVariables.ini
@@ -59,6 +59,8 @@
 
 ; r.XGEShaderCompile is now enabled by default in source. Uncomment to disable XGE shader compilation.
 ;r.XGEShaderCompile = 0
+;Uncomment to compile shaders in proc of the running UE process. Useful for debugging the shader compilation pipeline.
+;r.Shaders.AllowCompilingThroughWorkers=0
 ; Uncomment when running with a graphical debugger (but not when profiling)
 ;r.Shaders.Optimize=0
 ; When this is enabled, shaders will have extra debugging info. This could change patch sizes, uniqueness, etc and will recompile the shaders
diff --git a/Engine/Source/Runtime/Engine/Private/ShaderCompiler/ShaderCompiler.cpp b/Engine/Source/Runtime/Engine/Private/ShaderCompiler/ShaderCompiler.cpp
--- a/Engine/Source/Runtime/Engine/Private/ShaderCompiler/ShaderCompiler.cpp
+++ b/Engine/Source/Runtime/Engine/Private/ShaderCompiler/ShaderCompiler.cpp
@@ -1825,7 +1825,8 @@ FShaderCompilingManager::FShaderCompilingManager() :
 	verify(GConfig->GetBool( TEXT("DevOptions.Shaders"), TEXT("bAllowAsynchronousShaderCompiling"), bAllowAsynchronousShaderCompiling, GEngineIni ));
 
 	// override the use of workers, can be helpful for debugging shader compiler code
-	if (!FPlatformProcess::SupportsMultithreading() || FParse::Param(FCommandLine::Get(),TEXT("noshaderworker")))
+	static const IConsoleVariable* CVarAllowCompilingThroughWorkers = IConsoleManager::Get().FindConsoleVariable(TEXT("r.Shaders.AllowCompilingThroughWorkers"), false);
+	if (!FPlatformProcess::SupportsMultithreading() || FParse::Param(FCommandLine::Get(), TEXT("noshaderworker")) || (CVarAllowCompilingThroughWorkers && CVarAllowCompilingThroughWorkers->GetInt() == 0))
 	{
 		bAllowCompilingThroughWorkers = false;
 	}
diff --git a/Engine/Source/Runtime/Engine/Private/ShaderCompiler/ShaderCompilerXGE.cpp b/Engine/Source/Runtime/Engine/Private/ShaderCompiler/ShaderCompilerXGE.cpp
--- a/Engine/Source/Runtime/Engine/Private/ShaderCompiler/ShaderCompilerXGE.cpp
+++ b/Engine/Source/Runtime/Engine/Private/ShaderCompiler/ShaderCompilerXGE.cpp
@@ -76,7 +76,9 @@ namespace XGEShaderCompilerVariables
 			{
 				XGEShaderCompilerVariables::Enabled = 1;
 			}
-			if (FParse::Param(FCommandLine::Get(), TEXT("noxgeshadercompile")) || FParse::Param(FCommandLine::Get(), TEXT("noshaderworker")))
+
+			static const IConsoleVariable* CVarAllowCompilingThroughWorkers = IConsoleManager::Get().FindConsoleVariable(TEXT("r.Shaders.AllowCompilingThroughWorkers"), false);
+			if (FParse::Param(FCommandLine::Get(), TEXT("noxgeshadercompile")) || FParse::Param(FCommandLine::Get(), TEXT("noshaderworker")) || (CVarAllowCompilingThroughWorkers && CVarAllowCompilingThroughWorkers->GetInt() == 0))
 			{
 				XGEShaderCompilerVariables::Enabled = 0;
 			}
diff --git a/Engine/Source/Runtime/RenderCore/Private/Shader.cpp b/Engine/Source/Runtime/RenderCore/Private/Shader.cpp
--- a/Engine/Source/Runtime/RenderCore/Private/Shader.cpp
+++ b/Engine/Source/Runtime/RenderCore/Private/Shader.cpp
@@ -54,6 +54,15 @@ static TAutoConsoleVariable<int32> CVarSkipShaderCompression(
 	ECVF_ReadOnly | ECVF_Cheat
 	);
 
+static TAutoConsoleVariable<int32> CVarAllowCompilingThroughWorkers(
+	TEXT("r.Shaders.AllowCompilingThroughWorkers"),
+	1,
+	TEXT("Allows shader compilation through external ShaderCompileWorker processes.\n")
+	TEXT("1 - (Default) Allows external shader compiler workers\n") 
+	TEXT("0 - Disallows external shader compiler workers. Will run shader compilation in proc of UE process."),
+	ECVF_ReadOnly
+	);
+
 static TAutoConsoleVariable<int32> CVarShaderCompilerEmitWarningsOnLoad(
 	TEXT("r.ShaderCompiler.EmitWarningsOnLoad"),
 	0,
