F - More debug info for material failure
#rnx

diff --git a/Engine/Source/Runtime/Engine/Private/Materials/MaterialShader.cpp b/Engine/Source/Runtime/Engine/Private/Materials/MaterialShader.cpp
--- a/Engine/Source/Runtime/Engine/Private/Materials/MaterialShader.cpp
+++ b/Engine/Source/Runtime/Engine/Private/Materials/MaterialShader.cpp
@@ -1733,7 +1733,13 @@ bool FMaterialShaderMap::ProcessCompilationResults(const TArray<TSharedRef<FShad
 								{
 									FMeshMaterialShaderType* ShaderType = ((FMeshMaterialShaderType*)(StageTypes[Index]))->GetMeshMaterialShaderType();
 									FShader* Shader = MeshShaderMap->GetShader(ShaderType, kUniqueShaderPermutationId);
-									check(Shader);
+#if DO_CHECK
+#if ALLOW_SHADERMAP_DEBUG_DATA
+									UE_CLOG(Shader == nullptr, LogShaders, Fatal, TEXT("Failed to get ShaderType %s Permutation %d from MeshMaterial ShaderPipeline %s, FriendlyName %s, MaterialPath %s"), ShaderType->GetName(), kUniqueShaderPermutationId, ShaderPipelineType->GetName(), *FriendlyName, *MaterialPath);
+#else
+									UE_CLOG(Shader == nullptr, LogShaders, Fatal, TEXT("Failed to get ShaderType %s Permutation %d from MeshMaterial ShaderPipeline %s"), ShaderType->GetName(), kUniqueShaderPermutationId, ShaderPipelineType->GetName());
+#endif
+#endif	// DO_CHECK
 									ShaderStages.Add(Shader);
 								}
 
