Add "Res=1280x720wf" type command line argument to set window resolution easier.

diff --git a/Engine/Source/Runtime/Engine/Private/GameEngine.cpp b/Engine/Source/Runtime/Engine/Private/GameEngine.cpp
--- a/Engine/Source/Runtime/Engine/Private/GameEngine.cpp
+++ b/Engine/Source/Runtime/Engine/Private/GameEngine.cpp
@@ -97,6 +97,8 @@ static FAutoConsoleVariableRef CVarDoAsyncEndOfFrameTasks(
 	TEXT("Experimental option to run various things concurrently with the HUD render.")
 	);
 
+bool ParseResolution(const TCHAR* InResolution, uint32& OutX, uint32& OutY, int32& WindowMode);
+
 /** Benchmark results to the log */
 static void RunSynthBenchmark(const TArray<FString>& Args)
 {
@@ -314,12 +316,29 @@ void UGameEngine::ConditionallyOverrideSettings(int32& ResolutionX, int32& Resol
 
 void UGameEngine::DetermineGameWindowResolution( int32& ResolutionX, int32& ResolutionY, EWindowMode::Type& WindowMode, bool bUseWorkAreaForWindowed )
 {
+	FString ResolutionStr;;
+	if (FParse::Value(FCommandLine::Get(), TEXT("Res="), ResolutionStr))
+	{
+		uint32 ResX = 0;
+		uint32 ResY = 0;
+		int32 WinMode = EWindowMode::Windowed;
+
+		if (ParseResolution(*ResolutionStr, ResX, ResY, WinMode))
+		{
+			ResolutionX = ResX;
+			ResolutionY = ResY;
+			WindowMode = EWindowMode::ConvertIntToWindowMode(WinMode);
+		}
+	}
+	else
+	{
+		FParse::Value(FCommandLine::Get(), TEXT("ResX="), ResolutionX);
+		FParse::Value(FCommandLine::Get(), TEXT("ResY="), ResolutionY);
+	}
+
 	//fullscreen is always supported, but don't allow windowed mode on platforms that dont' support it.
 	WindowMode = (!FPlatformProperties::SupportsWindowedMode() && (WindowMode == EWindowMode::Windowed || WindowMode == EWindowMode::WindowedFullscreen)) ? EWindowMode::Fullscreen : WindowMode;
 
-	FParse::Value(FCommandLine::Get(), TEXT("ResX="), ResolutionX);
-	FParse::Value(FCommandLine::Get(), TEXT("ResY="), ResolutionY);
-
 	// consume available desktop area
 	FDisplayMetrics DisplayMetrics;
 	if (FSlateApplication::IsInitialized())
