Disable retaining on iOS due to cache pressure. Fix issue where disabling retaining would not work for assets loaded before the CVar is set.

diff --git a/Engine/Source/Runtime/Engine/Private/SoundWave.cpp b/Engine/Source/Runtime/Engine/Private/SoundWave.cpp
--- a/Engine/Source/Runtime/Engine/Private/SoundWave.cpp
+++ b/Engine/Source/Runtime/Engine/Private/SoundWave.cpp
@@ -824,8 +824,9 @@ void USoundWave::PostLoad()
 
 		if (ShouldUseStreamCaching() && ActualLoadingBehavior != GetLoadingBehavior(false))
 		{
-			if (ActualLoadingBehavior == ESoundWaveLoadingBehavior::RetainOnLoad)
+			if (!DisableRetainingCVar && ActualLoadingBehavior == ESoundWaveLoadingBehavior::RetainOnLoad)
 			{
+				UE_LOG(LogAudio, Display, TEXT("Sound Wave '%s' will have to load its compressed audio data async."), *GetName());
 				RetainCompressedAudio(false);
 			}
 			else
@@ -840,6 +841,12 @@ void USoundWave::PostLoad()
 			}
 		}
 
+		// If DisableRetainingCVar was set after this USoundWave was loaded by the ALT, release its compressed audio here.
+		if (DisableRetainingCVar)
+		{
+			ReleaseCompressedAudio();
+		}
+
 		// In case any code accesses bStreaming directly, we fix up bStreaming based on the current platform's cook overrides.
 		bStreaming = IsStreaming(nullptr);
 	}
