Fix for ensure hit in windows server.

diff --git a/Engine/Source/Runtime/Engine/Private/SoundWave.cpp b/Engine/Source/Runtime/Engine/Private/SoundWave.cpp
--- a/Engine/Source/Runtime/Engine/Private/SoundWave.cpp
+++ b/Engine/Source/Runtime/Engine/Private/SoundWave.cpp
@@ -957,6 +957,10 @@ uint32 USoundWave::GetNumChunks() const
 	{
 		return RunningPlatformData->Chunks.Num();
 	}
+	else if (IsTemplate() || IsRunningDedicatedServer())
+	{
+		return 0;
+	}
 	else
 	{
 		ensureAlwaysMsgf(false, TEXT("Call CachePlatformData(false) before calling this function in editor."));
