[Navigation] Fix for InstancedStaticMesh components that were not registered to the navigation octree if their associated mesh was assigned post registration. Also updated StaticMeshComponent to set bNavigationRelevant before notifying the navigation system to perform the proper operation (register/unregister).

diff --git a/Engine/Source/Runtime/Engine/Private/InstancedStaticMesh.cpp b/Engine/Source/Runtime/Engine/Private/InstancedStaticMesh.cpp
--- a/Engine/Source/Runtime/Engine/Private/InstancedStaticMesh.cpp
+++ b/Engine/Source/Runtime/Engine/Private/InstancedStaticMesh.cpp
@@ -1708,6 +1708,11 @@ void UInstancedStaticMeshComponent::OnCreatePhysicsState()
 	CreateAllInstanceBodies();
 
 	USceneComponent::OnCreatePhysicsState();
+
+	// Since StaticMeshComponent was not called
+	// Navigation relevancy needs to be handled here
+	bNavigationRelevant = IsNavigationRelevant();
+	FNavigationSystem::UpdateComponentData(*this);
 }
 
 void UInstancedStaticMeshComponent::OnDestroyPhysicsState()
@@ -1716,6 +1721,11 @@ void UInstancedStaticMeshComponent::OnDestroyPhysicsState()
 
 	// Release all physics representations
 	ClearAllInstanceBodies();
+
+	// Since StaticMeshComponent was not called
+	// Navigation relevancy needs to be handled here
+	bNavigationRelevant = IsNavigationRelevant();
+	FNavigationSystem::UpdateComponentData(*this);
 }
 
 bool UInstancedStaticMeshComponent::CanEditSimulatePhysics()
