Fix static analysis.

diff --git a/Engine/Source/Runtime/Engine/Private/PhysicsEngine/Experimental/PhysScene_Chaos.cpp b/Engine/Source/Runtime/Engine/Private/PhysicsEngine/Experimental/PhysScene_Chaos.cpp
--- a/Engine/Source/Runtime/Engine/Private/PhysicsEngine/Experimental/PhysScene_Chaos.cpp
+++ b/Engine/Source/Runtime/Engine/Private/PhysicsEngine/Experimental/PhysScene_Chaos.cpp
@@ -1624,13 +1624,15 @@ bool FPhysScene_ChaosInterface::MarkForPreSimKinematicUpdate(USkeletalMeshCompon
 
 void FPhysScene_ChaosInterface::ClearPreSimKinematicUpdate(USkeletalMeshComponent* InSkelComp)
 {
-	// If non-null, and flagged for deferred update..
-	const int32 DeferredKinematicUpdateIndex = InSkelComp->DeferredKinematicUpdateIndex;
-	if (InSkelComp != nullptr && DeferredKinematicUpdateIndex != INDEX_NONE)
+	if (InSkelComp != nullptr)
 	{
-		DeferredKinematicUpdateSkelMeshes.Last().Key->DeferredKinematicUpdateIndex = DeferredKinematicUpdateIndex;
-		DeferredKinematicUpdateSkelMeshes.RemoveAtSwap(InSkelComp->DeferredKinematicUpdateIndex);
-		InSkelComp->DeferredKinematicUpdateIndex = INDEX_NONE;
+		const int32 DeferredKinematicUpdateIndex = InSkelComp->DeferredKinematicUpdateIndex;
+		if (DeferredKinematicUpdateIndex != INDEX_NONE)
+		{
+			DeferredKinematicUpdateSkelMeshes.Last().Key->DeferredKinematicUpdateIndex = DeferredKinematicUpdateIndex;
+			DeferredKinematicUpdateSkelMeshes.RemoveAtSwap(InSkelComp->DeferredKinematicUpdateIndex);
+			InSkelComp->DeferredKinematicUpdateIndex = INDEX_NONE;
+		}
 	}
 }
 
