Support VolumeTexture in TextureProperty node. Also added "GetSizeZ" function to FTexture.

#rnx
#hardlock Nick.Penwarden

diff --git a/Engine/Source/Runtime/Engine/Private/Materials/HLSLMaterialTranslator.cpp b/Engine/Source/Runtime/Engine/Private/Materials/HLSLMaterialTranslator.cpp
--- a/Engine/Source/Runtime/Engine/Private/Materials/HLSLMaterialTranslator.cpp
+++ b/Engine/Source/Runtime/Engine/Private/Materials/HLSLMaterialTranslator.cpp
@@ -4558,14 +4558,14 @@ int32 FHLSLMaterialTranslator::TextureProperty(int32 TextureIndex, EMaterialExpo
 {
 	EMaterialValueType TextureType = GetParameterType(TextureIndex);
 
-	if(TextureType != MCT_Texture2D && TextureType != MCT_TextureVirtual)
+	if(TextureType != MCT_Texture2D && TextureType != MCT_TextureVirtual && TextureType != MCT_VolumeTexture)
 	{
-		return Errorf(TEXT("Texture size only available for Texture2D, not %s"),DescribeType(TextureType));
+		return Errorf(TEXT("Texture size only available for Texture2D, TextureVirtual, and VolumeTexture, not %s"),DescribeType(TextureType));
 	}
 		
 	auto TextureExpression = (FMaterialUniformExpressionTexture*) (*CurrentScopeChunks)[TextureIndex].UniformExpression.GetReference();
 
-	return AddUniformExpression(new FMaterialUniformExpressionTextureProperty(TextureExpression, Property), MCT_Float2, TEXT(""));
+	return AddUniformExpression(new FMaterialUniformExpressionTextureProperty(TextureExpression, Property), (TextureType == MCT_VolumeTexture ? MCT_Float3 : MCT_Float2), TEXT(""));
 }
 
 int32 FHLSLMaterialTranslator::TextureDecalMipmapLevel(int32 TextureSizeInput)
diff --git a/Engine/Source/Runtime/Engine/Private/Materials/MaterialUniformExpressions.cpp b/Engine/Source/Runtime/Engine/Private/Materials/MaterialUniformExpressions.cpp
--- a/Engine/Source/Runtime/Engine/Private/Materials/MaterialUniformExpressions.cpp
+++ b/Engine/Source/Runtime/Engine/Private/Materials/MaterialUniformExpressions.cpp
@@ -735,7 +735,8 @@ static void EvaluateTextureSize(const FMaterialRenderContext& Context, FPreshade
 	{
 		const uint32 SizeX = Texture->Resource->GetSizeX();
 		const uint32 SizeY = Texture->Resource->GetSizeY();
-		Stack.Add(FLinearColor((float)SizeX, (float)SizeY, 0.0f, 0.0f));
+		const uint32 SizeZ = Texture->Resource->GetSizeZ();
+		Stack.Add(FLinearColor((float)SizeX, (float)SizeY, (float)SizeZ, 0.0f));
 	}
 }
 
@@ -746,7 +747,8 @@ static void EvaluateTexelSize(const FMaterialRenderContext& Context, FPreshaderS
 	{
 		const uint32 SizeX = Texture->Resource->GetSizeX();
 		const uint32 SizeY = Texture->Resource->GetSizeY();
-		Stack.Add(FLinearColor(1.0f / (float)SizeX, 1.0f / (float)SizeY, 0.0f, 0.0f));
+		const uint32 SizeZ = Texture->Resource->GetSizeZ();
+		Stack.Add(FLinearColor(1.0f / (float)SizeX, 1.0f / (float)SizeY, (SizeZ > 0 ? 1.0f / (float)SizeZ : 0.0f), 0.0f));
 	}
 }
 
diff --git a/Engine/Source/Runtime/Engine/Private/VolumeTexture.cpp b/Engine/Source/Runtime/Engine/Private/VolumeTexture.cpp
--- a/Engine/Source/Runtime/Engine/Private/VolumeTexture.cpp
+++ b/Engine/Source/Runtime/Engine/Private/VolumeTexture.cpp
@@ -494,14 +494,19 @@ public:
 	{
 		return FMath::Max<uint32>(SizeY >> CurrentFirstMip, 1);
 	}
+	/** Returns the depth of the texture in pixels. */
+	uint32 GetSizeZ() const override
+	{
+		return FMath::Max<uint32>(SizeZ >> CurrentFirstMip, 1);
+	}
 
 private:
 
-	/** The UTexture2D which this resource represents.														*/
+	/** The UVolumeTexture which this resource represents */
 	UVolumeTexture*	Owner;
 
 #if STATS
-	/** The FName of the LODGroup-specific stat	*/
+	/** The FName of the LODGroup-specific stat */
 	FName LODGroupStatName;
 #endif
 	/** The FName of the texture asset */
diff --git a/Engine/Source/Runtime/RenderCore/Public/RenderResource.h b/Engine/Source/Runtime/RenderCore/Public/RenderResource.h
--- a/Engine/Source/Runtime/RenderCore/Public/RenderResource.h
+++ b/Engine/Source/Runtime/RenderCore/Public/RenderResource.h
@@ -411,6 +411,11 @@ public:
 	{
 		return 0;
 	}
+	/** Returns the depth of the texture in pixels. */
+	virtual uint32 GetSizeZ() const
+	{
+		return 0;
+	}
 
 	// FRenderResource interface.
 	virtual void ReleaseRHI() override
