Run WorldCollisionAsync on worker threads in multithread forked servers

diff --git a/Engine/Source/Runtime/Engine/Private/Collision/WorldCollisionAsync.cpp b/Engine/Source/Runtime/Engine/Private/Collision/WorldCollisionAsync.cpp
--- a/Engine/Source/Runtime/Engine/Private/Collision/WorldCollisionAsync.cpp
+++ b/Engine/Source/Runtime/Engine/Private/Collision/WorldCollisionAsync.cpp
@@ -12,6 +12,7 @@
 #include "CollisionQueryParams.h"
 #include "WorldCollision.h"
 #include "Engine/World.h"
+#include "Misc/Fork.h"
 #include "PhysicsEngine/BodyInstance.h"
 #include "Physics/PhysicsInterfaceCore.h"
 #include "ProfilingDebugging/CsvProfiler.h"
@@ -43,6 +44,11 @@ namespace AsyncTraceCVars
 		TEXT("Whether to use worker thread for async trace functionality. This works if FApp::ShouldUseThreadingForPerformance is true. Otherwise it will always use game thread. \n")
 		TEXT("0: Use game thread, 1: User worker thread"),
 		ECVF_Default);
+
+	bool IsAsyncTraceOnWorkerThreads()
+	{
+		return RunAsyncTraceOnWorkerThread != 0 && (FApp::ShouldUseThreadingForPerformance() || FForkProcessHelper::IsForkedMultithreadInstance());
+	}
 }
 
 namespace
@@ -306,7 +312,7 @@ namespace
 		}
 
 		auto* Datum      = GetTraceContainer<DatumType>(DataBuffer)[Next.Block]->Buffer;
-		const bool bRunAsyncTraceOnWorkerThread = !!AsyncTraceCVars::RunAsyncTraceOnWorkerThread && FApp::ShouldUseThreadingForPerformance();
+		const bool bRunAsyncTraceOnWorkerThread = AsyncTraceCVars::IsAsyncTraceOnWorkerThreads();
 		if (bRunAsyncTraceOnWorkerThread)
 		{
 			DataBuffer.AsyncTraceCompletionEvent.Emplace(TGraphTask<FAsyncTraceTask>::CreateTask(NULL, ENamedThreads::GameThread).ConstructAndDispatchWhenReady(Datum, Next.Index));
@@ -452,7 +458,7 @@ bool UWorld::QueryOverlapData(const FTraceHandle& Handle, FOverlapDatum& OutData
 
 void UWorld::WaitForAllAsyncTraceTasks()
 {
-	const bool bRunAsyncTraceOnWorkerThread = !!AsyncTraceCVars::RunAsyncTraceOnWorkerThread && FApp::ShouldUseThreadingForPerformance();
+	const bool bRunAsyncTraceOnWorkerThread = AsyncTraceCVars::IsAsyncTraceOnWorkerThreads();
 	if (bRunAsyncTraceOnWorkerThread)
 	{
 		// if running thread, wait until all threads finishes, if we don't do this, there might be more thread running
