Renamed FStaticMeshStreamOut_DDC to FStaticMeshStreamOut.
Renamed FStaticMeshStreamOut_IO to FStaticMeshStreamOut_Swap.
UStaticMesh::StreamOut now only use FStaticMeshStreamOut while we investigate reported errors with the new path.

diff --git a/Engine/Source/Runtime/Engine/Private/StaticMesh.cpp b/Engine/Source/Runtime/Engine/Private/StaticMesh.cpp
--- a/Engine/Source/Runtime/Engine/Private/StaticMesh.cpp
+++ b/Engine/Source/Runtime/Engine/Private/StaticMesh.cpp
@@ -5615,15 +5615,13 @@ bool UStaticMesh::StreamOut(int32 NewMipCount)
 	check(IsInGameThread());
 	if (bIsStreamable && !PendingUpdate && RenderData.IsValid() && RenderData->bReadyForStreaming && NewMipCount < GetNumResidentMips())
 	{
-#if WITH_EDITOR
-		if (FPlatformProperties::HasEditorOnlyData())
+		if (FPlatformProperties::HasEditorOnlyData() || true) // Temporary disabling the Swap path.
 		{
-			PendingUpdate = new FStaticMeshStreamOut_DDC(this, NewMipCount);
+			PendingUpdate = new FStaticMeshStreamOut(this, NewMipCount);
 		}
 		else
-#endif
 		{
-			PendingUpdate = new FStaticMeshStreamOut_IO(this, NewMipCount);
+			PendingUpdate = new FStaticMeshStreamOut_Swap(this, NewMipCount);
 		}
 		return !PendingUpdate->IsCancelled();
 	}
diff --git a/Engine/Source/Runtime/Engine/Private/Streaming/StaticMeshUpdate.cpp b/Engine/Source/Runtime/Engine/Private/Streaming/StaticMeshUpdate.cpp
--- a/Engine/Source/Runtime/Engine/Private/Streaming/StaticMeshUpdate.cpp
+++ b/Engine/Source/Runtime/Engine/Private/Streaming/StaticMeshUpdate.cpp
@@ -273,14 +273,13 @@ void FStaticMeshStreamIn::DoCancel(const FContext& Context)
 	DoFinishUpdate(Context);
 }
 
-#if WITH_EDITOR
-FStaticMeshStreamOut_DDC::FStaticMeshStreamOut_DDC(UStaticMesh* InMesh, int32 InRequestedMips)
+FStaticMeshStreamOut::FStaticMeshStreamOut(UStaticMesh* InMesh, int32 InRequestedMips)
 	: FStaticMeshUpdate(InMesh, InRequestedMips)
 {
 	PushTask(FContext(InMesh, TT_None), TT_Render, SRA_UPDATE_CALLBACK(DoReleaseBuffers), TT_None, nullptr);
 }
 
-void FStaticMeshStreamOut_DDC ::DoReleaseBuffers(const FContext& Context)
+void FStaticMeshStreamOut ::DoReleaseBuffers(const FContext& Context)
 {
 	check(Context.CurrentThread == TT_Render);
 	UStaticMesh* Mesh = Context.Mesh;
@@ -308,16 +307,15 @@ void FStaticMeshStreamOut_DDC ::DoReleaseBuffers(const FContext& Context)
 		}
 	}
 }
-#endif // WITH_EDITOR
 
 
-FStaticMeshStreamOut_IO::FStaticMeshStreamOut_IO(UStaticMesh* InMesh, int32 InRequestedMips)
+FStaticMeshStreamOut_Swap::FStaticMeshStreamOut_Swap(UStaticMesh* InMesh, int32 InRequestedMips)
 	: FStaticMeshUpdate(InMesh, InRequestedMips)
 {
 	PushTask(FContext(InMesh, TT_None), TT_Render, SRA_UPDATE_CALLBACK(DoReleaseBuffers), TT_None, nullptr);
 }
 
-void FStaticMeshStreamOut_IO::DoReleaseBuffers(const FContext& Context)
+void FStaticMeshStreamOut_Swap::DoReleaseBuffers(const FContext& Context)
 {
 	check(Context.CurrentThread == TT_Render);
 	UStaticMesh* Mesh = Context.Mesh;
@@ -345,7 +343,7 @@ void FStaticMeshStreamOut_IO::DoReleaseBuffers(const FContext& Context)
 	}
 }
 
-void FStaticMeshStreamOut_IO::DoReleaseLODResources(const FContext& Context)
+void FStaticMeshStreamOut_Swap::DoReleaseLODResources(const FContext& Context)
 {
 	check(Context.CurrentThread == TT_Render);
 
@@ -373,7 +371,7 @@ void FStaticMeshStreamOut_IO::DoReleaseLODResources(const FContext& Context)
 	}
 }
 
-void FStaticMeshStreamOut_IO::ReleaseLODResources()
+void FStaticMeshStreamOut_Swap::ReleaseLODResources()
 {
 	if (LODResourcesToRelease.Num())
 	{
diff --git a/Engine/Source/Runtime/Engine/Private/Streaming/StaticMeshUpdate.h b/Engine/Source/Runtime/Engine/Private/Streaming/StaticMeshUpdate.h
--- a/Engine/Source/Runtime/Engine/Private/Streaming/StaticMeshUpdate.h
+++ b/Engine/Source/Runtime/Engine/Private/Streaming/StaticMeshUpdate.h
@@ -122,25 +122,23 @@ private:
 	void CreateBuffers_Internal(const FContext& Context);
 };
 
-#if WITH_EDITOR
 /** A streamout that doesn't actually touches the CPU data. Required because DDC stream in doesn't reset. */
-class FStaticMeshStreamOut_DDC : public FStaticMeshUpdate
+class FStaticMeshStreamOut : public FStaticMeshUpdate
 {
 public:
-	FStaticMeshStreamOut_DDC(UStaticMesh* InMesh, int32 InRequestedMips);
+	FStaticMeshStreamOut(UStaticMesh* InMesh, int32 InRequestedMips);
 
 private:
 	/** Release RHI buffers and update SRVs */
 	void DoReleaseBuffers(const FContext& Context);
 };
-#endif
 
-class FStaticMeshStreamOut_IO : public FStaticMeshUpdate
+class FStaticMeshStreamOut_Swap : public FStaticMeshUpdate
 {
 public:
-	FStaticMeshStreamOut_IO(UStaticMesh* InMesh, int32 InRequestedMips);
+	FStaticMeshStreamOut_Swap(UStaticMesh* InMesh, int32 InRequestedMips);
 
-	virtual ~FStaticMeshStreamOut_IO() { ReleaseLODResources(); }
+	virtual ~FStaticMeshStreamOut_Swap() { ReleaseLODResources(); }
 
 private:
 
diff --git a/Engine/Source/Runtime/Engine/Public/StaticMeshResources.h b/Engine/Source/Runtime/Engine/Public/StaticMeshResources.h
--- a/Engine/Source/Runtime/Engine/Public/StaticMeshResources.h
+++ b/Engine/Source/Runtime/Engine/Public/StaticMeshResources.h
@@ -561,7 +561,7 @@ private:
 	friend class FStaticMeshRenderData;
 	friend class FStaticMeshStreamIn;
 	friend class FStaticMeshStreamIn_IO;
-	friend class FStaticMeshStreamOut_DDC;
+	friend class FStaticMeshStreamOut;
 };
 
 struct ENGINE_API FStaticMeshVertexFactories
