Removed duplicate uobject stat scope when calling DoDeferredRenderUpdates_Concurrent from UWorld::SendAllEndOfFrameUpdates.

#rnx

diff --git a/Engine/Source/Runtime/Engine/Private/Components/ActorComponent.cpp b/Engine/Source/Runtime/Engine/Private/Components/ActorComponent.cpp
--- a/Engine/Source/Runtime/Engine/Private/Components/ActorComponent.cpp
+++ b/Engine/Source/Runtime/Engine/Private/Components/ActorComponent.cpp
@@ -1568,6 +1568,7 @@ void UActorComponent::DoDeferredRenderUpdates_Concurrent()
 	checkf(!IsPendingKill(), TEXT("%s"), *GetFullName());
 
 	FScopeCycleCounterUObject ContextScope(this);
+	FScopeCycleCounterUObject AdditionalScope(STATS ? AdditionalStatObject() : nullptr);
 
 	if(!IsRegistered())
 	{
diff --git a/Engine/Source/Runtime/Engine/Private/LevelTick.cpp b/Engine/Source/Runtime/Engine/Private/LevelTick.cpp
--- a/Engine/Source/Runtime/Engine/Private/LevelTick.cpp
+++ b/Engine/Source/Runtime/Engine/Private/LevelTick.cpp
@@ -1030,8 +1030,6 @@ void UWorld::SendAllEndOfFrameUpdates()
 			{
 				if (NextComponent->IsRegistered() && !NextComponent->IsTemplate() && !NextComponent->IsPendingKill())
 				{
-					FScopeCycleCounterUObject ComponentScope(NextComponent);
-					FScopeCycleCounterUObject AdditionalScope(STATS ? NextComponent->AdditionalStatObject() : nullptr);
 					NextComponent->DoDeferredRenderUpdates_Concurrent();
 				}
 				check(NextComponent->IsPendingKill() || NextComponent->GetMarkedForEndOfFrameUpdateState() == EComponentMarkedForEndOfFrameUpdateState::Marked);
@@ -1048,8 +1046,6 @@ void UWorld::SendAllEndOfFrameUpdates()
 				{
 					if (Component->IsRegistered() && !Component->IsTemplate() && !Component->IsPendingKill())
 					{
-						FScopeCycleCounterUObject ComponentScope(Component);
-						FScopeCycleCounterUObject AdditionalScope(STATS ? Component->AdditionalStatObject() : nullptr);
 						Component->DoDeferredRenderUpdates_Concurrent();
 					}
 
