Minor optimization for SetAllowRigidBodyNode to not recreate the anim script instance

diff --git a/Engine/Source/Runtime/Engine/Private/Components/SkeletalMeshComponent.cpp b/Engine/Source/Runtime/Engine/Private/Components/SkeletalMeshComponent.cpp
--- a/Engine/Source/Runtime/Engine/Private/Components/SkeletalMeshComponent.cpp
+++ b/Engine/Source/Runtime/Engine/Private/Components/SkeletalMeshComponent.cpp
@@ -3734,10 +3734,17 @@ void USkeletalMeshComponent::SetAllowRigidBodyAnimNode(bool bInAllow, bool bRein
 	{
 		bDisableRigidBodyAnimNode = !bInAllow;
 
-		if(bReinitAnim && bRegistered && SkeletalMesh)
+		if(bReinitAnim && bRegistered)
 		{
 			// need to reinitialize rigid body nodes for new setting to take effect
-			InitializeAnimScriptInstance(true);
+			if (AnimScriptInstance)
+			{
+				AnimScriptInstance->InitializeAnimation();
+			}
+			if (PostProcessAnimInstance)
+			{
+				PostProcessAnimInstance->InitializeAnimation();
+			}
 		}
 	}
 }
