- Fixed concurrency volume scaling popping in/not fading correctly
- Fixed priority cycling per frame when playtime is the same and priorities are the same
- Fixed GetHighestPriority double counting ActiveSound's Priority field (field seeds ParseParams version which is cached in FocusData.HighestPriority value)
#rnx

diff --git a/Engine/Source/Runtime/Engine/Private/ActiveSound.cpp b/Engine/Source/Runtime/Engine/Private/ActiveSound.cpp
--- a/Engine/Source/Runtime/Engine/Private/ActiveSound.cpp
+++ b/Engine/Source/Runtime/Engine/Private/ActiveSound.cpp
@@ -99,6 +99,9 @@ FActiveSound::FActiveSound()
 	, DebugColor(FColor::Black)
 #endif // ENABLE_AUDIO_DEBUG
 {
+	static uint32 TotalPlayOrder = 0;
+	PlayOrder = TotalPlayOrder++;
+
 	if (!ActiveSoundTraceDelegate.IsBound())
 	{
 		ActiveSoundTraceDelegate.BindStatic(&OcclusionTraceDone);
@@ -1492,7 +1495,7 @@ float FActiveSound::GetHighestPriority(bool bIgnoreAlwaysPlay) const
 		}
 	}
 
-	const float HighestPriority = Priority * FocusData.PriorityHighest * FocusData.PriorityScale;
+	const float HighestPriority = FocusData.PriorityHighest * FocusData.PriorityScale;
 	return FMath::Clamp(HighestPriority, 0.0f, 100.0f);
 }
 
diff --git a/Engine/Source/Runtime/Engine/Private/SoundConcurrency.cpp b/Engine/Source/Runtime/Engine/Private/SoundConcurrency.cpp
--- a/Engine/Source/Runtime/Engine/Private/SoundConcurrency.cpp
+++ b/Engine/Source/Runtime/Engine/Private/SoundConcurrency.cpp
@@ -1,17 +1,24 @@
 // Copyright Epic Games, Inc. All Rights Reserved.
 
 #include "Sound/SoundConcurrency.h"
-#include "Components/AudioComponent.h"
+
 #include "ActiveSound.h"
 #include "AudioDefines.h"
 #include "AudioDevice.h"
 #include "AudioVirtualLoop.h"
+#include "Components/AudioComponent.h"
 #include "DSP/VolumeFader.h"
+#include "HAL/IConsoleManager.h"
 #include "Sound/SoundBase.h"
 
 
-// Forward Declarations
-struct FListener;
+static float ConcurrencyMinVolumeScaleCVar = 1.e-3f;
+FAutoConsoleVariableRef CVarConcurrencyMinVolumeScale(
+	TEXT("au.Concurrency.MinVolumeScale"),
+	ConcurrencyMinVolumeScaleCVar,
+	TEXT("Volume threshold considered silent for volume scaling (linear scale).\n"),
+	ECVF_Default);
+
 
 namespace SoundConcurrency
 {
@@ -113,14 +120,28 @@ float FConcurrencySoundData::GetLerpTime() const
 
 float FConcurrencySoundData::GetVolume(bool bInDecibels) const
 {
-	if (FMath::IsNearlyZero(LerpTime) || Elapsed > LerpTime || FMath::IsNearlyEqual(DbTargetVolume, DbStartVolume))
+	float VolumeDb = 0.0f;
+	if (FMath::IsNearlyZero(LerpTime) || Elapsed >= LerpTime || FMath::IsNearlyEqual(DbTargetVolume, DbStartVolume, ConcurrencyMinVolumeScaleCVar))
+	{
+		VolumeDb = DbTargetVolume;
+	}
+	else
 	{
-		return bInDecibels ? DbTargetVolume : Audio::ConvertToLinear(DbTargetVolume);
+		const float Alpha = FMath::Clamp(Elapsed / LerpTime, 0.0f, 1.0f);
+		VolumeDb = FMath::Lerp(DbStartVolume, DbTargetVolume, Alpha);
 	}
 
-	float Alpha = Elapsed / LerpTime;
-	const float DbCurrentVolume = FMath::Lerp(DbStartVolume, DbTargetVolume, Alpha);
-	return bInDecibels ? DbCurrentVolume : Audio::ConvertToLinear(DbCurrentVolume);
+	if (bInDecibels)
+	{
+		return VolumeDb;
+	}
+	
+	float VolumeLin = Audio::ConvertToLinear(VolumeDb);
+	if (VolumeLin <= ConcurrencyMinVolumeScaleCVar)
+	{
+		return 0.0f;
+	}
+	return VolumeLin;
 }
 
 float FConcurrencySoundData::GetTargetVolume(bool bInDecibels) const
@@ -130,14 +151,23 @@ float FConcurrencySoundData::GetTargetVolume(bool bInDecibels) const
 
 void FConcurrencySoundData::SetTarget(float InTargetVolume, float InLerpTime)
 {
-	DbStartVolume = GetVolume(true);
-	DbTargetVolume = Audio::ConvertToDecibels(InTargetVolume, KINDA_SMALL_NUMBER);
+	const float DbInitVolume = DbTargetVolume;
+	const float InitLerpTime = LerpTime;
+
+	DbTargetVolume = Audio::ConvertToDecibels(InTargetVolume, ConcurrencyMinVolumeScaleCVar);
 	LerpTime = FMath::Max(InLerpTime, 0.0f);
 
-	if (InLerpTime != LerpTime && InTargetVolume != DbTargetVolume)
+	if (LerpTime > 0.0f)
 	{
-		Elapsed = 0.0f;
+		const float Alpha = FMath::Clamp(Elapsed / LerpTime, 0.0f, 1.0f);
+		DbStartVolume = FMath::Lerp(DbStartVolume, DbTargetVolume, Alpha);
 	}
+	else
+	{
+		DbStartVolume = DbTargetVolume;
+	}
+
+	Elapsed = 0.0f;
 }
 
 FConcurrencyGroup::FConcurrencyGroup(FConcurrencyGroupID InGroupID, const FConcurrencyHandle& ConcurrencyHandle)
@@ -237,6 +267,13 @@ void FConcurrencyGroup::UpdateGeneration(FActiveSound* NewActiveSound)
 					float DistSqB = 0.0f;
 					A.AudioDevice->GetDistanceSquaredToNearestListener(A.LastLocation, DistSqA);
 					B.AudioDevice->GetDistanceSquaredToNearestListener(B.LastLocation, DistSqB);
+
+					// If sounds share the same distance, newer sounds will be sorted first to avoid volume ping-ponging 
+					if (FMath::IsNearlyEqual(DistSqA, DistSqB, KINDA_SMALL_NUMBER))
+					{
+						return A.GetPlayOrder() < B.GetPlayOrder();
+					}
+
 					return DistSqA > DistSqB;
 				}
 
@@ -245,6 +282,13 @@ void FConcurrencyGroup::UpdateGeneration(FActiveSound* NewActiveSound)
 					// Ensures sounds set to always play are sorted above those that aren't, but are sorted appropriately between one another
 					const float APriority = A.GetAlwaysPlay() ? A.GetHighestPriority(true /* bIgnoreAlwaysPlay */) + MAX_SOUND_PRIORITY + 1.0f : A.GetHighestPriority();
 					const float BPriority = B.GetAlwaysPlay() ? B.GetHighestPriority(true /* bIgnoreAlwaysPlay */) + MAX_SOUND_PRIORITY + 1.0f : B.GetHighestPriority();
+
+					// If sounds share the same priority, newer sounds will be sorted first to avoid volume ping-ponging 
+					if (FMath::IsNearlyEqual(APriority, BPriority, KINDA_SMALL_NUMBER))
+					{
+						return A.GetPlayOrder() < B.GetPlayOrder();
+					}
+
 					return APriority < BPriority;
 				}
 			}
@@ -268,7 +312,10 @@ void FConcurrencyGroup::UpdateGeneration(FActiveSound* NewActiveSound)
 			// If new sound added, immediately lerp to the generation volume.
 			const float AttackTime = NewActiveSound == ActiveSound ? 0.0f : Settings.VolumeScaleAttackTime;
 
-			SoundConcurrency::SetSoundDataTarget(*ActiveSound, *SoundData, NewTargetVolume, AttackTime);
+			if (!FMath::IsNearlyEqual(AttackTime, SoundData->GetLerpTime()) || !FMath::IsNearlyEqual(SoundData->GetTargetVolume(), NewTargetVolume, ConcurrencyMinVolumeScaleCVar))
+			{
+				SoundConcurrency::SetSoundDataTarget(*ActiveSound, *SoundData, NewTargetVolume, AttackTime);
+			}
 		}
 	}
 }
@@ -302,7 +349,7 @@ void FConcurrencyGroup::CullSoundsDueToMaxConcurrency()
 					// If sounds share the same volume, newer sounds will be sorted first to avoid loop realization ping-ponging 
 					if (FMath::IsNearlyEqual(A.VolumeConcurrency, B.VolumeConcurrency, KINDA_SMALL_NUMBER))
 					{
-						return A.PlaybackTime > B.PlaybackTime;
+						return A.GetPlayOrder() > B.GetPlayOrder();
 					}
 					return A.VolumeConcurrency < B.VolumeConcurrency;
 				}
@@ -321,8 +368,8 @@ void FConcurrencyGroup::CullSoundsDueToMaxConcurrency()
 
 					// Newer sounds pushed forward in sort to make them more likely to be culled if PreventNew
 					return ResolutionRule == EMaxConcurrentResolutionRule::StopLowestPriorityThenPreventNew
-						? A.PlaybackTime < B.PlaybackTime
-						: A.PlaybackTime > B.PlaybackTime;
+						? A.GetPlayOrder() < B.GetPlayOrder()
+						: A.GetPlayOrder() > B.GetPlayOrder();
 				}
 				break;
 
diff --git a/Engine/Source/Runtime/Engine/Public/ActiveSound.h b/Engine/Source/Runtime/Engine/Public/ActiveSound.h
--- a/Engine/Source/Runtime/Engine/Public/ActiveSound.h
+++ b/Engine/Source/Runtime/Engine/Public/ActiveSound.h
@@ -272,12 +272,14 @@ private:
 	FName AudioComponentName;
 	FName OwnerName;
 
+	uint32 PlayOrder;
 
 public:
 	// ISoundModulatable Implementation
 	USoundModulationPluginSourceSettingsBase* FindModulationSettings() const override;
 	uint32 GetObjectId() const override { return Sound ? Sound->GetUniqueID() : INDEX_NONE; }
 	int32 GetPlayCount() const override;
+	uint32 GetPlayOrder() const { return PlayOrder; }
 	bool IsPreviewSound() const override { return bIsPreviewSound; }
 	void Stop() override;
 
