Fix CIS.

diff --git a/Engine/Source/Runtime/Engine/Public/Physics/Experimental/ChaosInterfaceUtils.h b/Engine/Source/Runtime/Engine/Public/Physics/Experimental/ChaosInterfaceUtils.h
--- a/Engine/Source/Runtime/Engine/Public/Physics/Experimental/ChaosInterfaceUtils.h
+++ b/Engine/Source/Runtime/Engine/Public/Physics/Experimental/ChaosInterfaceUtils.h
@@ -11,7 +11,7 @@ struct FGeometryAddParams;
 
 namespace Chaos
 {
-	enum EChaosCollisionTraceFlag;
+	enum EChaosCollisionTraceFlag : int;
 }
 
 namespace ChaosInterface
