Undo CL 12220856
- Doesn't work with complex materials

diff --git a/Engine/Source/Runtime/Engine/Private/Materials/MaterialInstance.cpp b/Engine/Source/Runtime/Engine/Private/Materials/MaterialInstance.cpp
--- a/Engine/Source/Runtime/Engine/Private/Materials/MaterialInstance.cpp
+++ b/Engine/Source/Runtime/Engine/Private/Materials/MaterialInstance.cpp
@@ -4884,42 +4884,50 @@ void UMaterialInstance::CopyMaterialUniformParametersInternal(UMaterialInterface
 				checkSlow(TextureParameterValues.Num() == 0);
 				checkSlow(RuntimeVirtualTextureParameterValues.Num() == 0);
 
-				const FMaterialCachedParameters& CachedParameters = AsMaterial->GetCachedExpressionData().Parameters;
-
-				// Scalars
-				ScalarParameterValues.Empty(CachedParameters.GetNumParameters(EMaterialParameterType::Scalar));
-				for (int32 ParameterIndex = 0; ParameterIndex < CachedParameters.GetNumParameters(EMaterialParameterType::Scalar); ++ParameterIndex)
+				const FMaterialResource* MaterialResource = nullptr;
+				if (UWorld* World = AsMaterial->GetWorld())
 				{
-					FScalarParameterValue* ParameterValue = new(ScalarParameterValues) FScalarParameterValue;
-					ParameterValue->ParameterInfo.Name = CachedParameters.GetParameterName(EMaterialParameterType::Scalar, ParameterIndex);
-					ParameterValue->ParameterValue = CachedParameters.ScalarValues[ParameterIndex];
+					MaterialResource = AsMaterial->GetMaterialResource(World->FeatureLevel);
 				}
 
-				// Vectors
-				VectorParameterValues.Empty(CachedParameters.GetNumParameters(EMaterialParameterType::Vector));
-				for (int32 ParameterIndex = 0; ParameterIndex < CachedParameters.GetNumParameters(EMaterialParameterType::Vector); ++ParameterIndex)
+				if (!MaterialResource)
 				{
-					FVectorParameterValue* ParameterValue = new(VectorParameterValues) FVectorParameterValue;
-					ParameterValue->ParameterInfo.Name = CachedParameters.GetParameterName(EMaterialParameterType::Vector, ParameterIndex);
-					ParameterValue->ParameterValue = CachedParameters.VectorValues[ParameterIndex];
+					MaterialResource = AsMaterial->GetMaterialResource(GMaxRHIFeatureLevel);
 				}
 
-				// Textures
-				TextureParameterValues.Empty(CachedParameters.GetNumParameters(EMaterialParameterType::Texture));
-				for (int32 ParameterIndex = 0; ParameterIndex < CachedParameters.GetNumParameters(EMaterialParameterType::Texture); ++ParameterIndex)
+				if (MaterialResource)
 				{
-					FTextureParameterValue* ParameterValue = new(TextureParameterValues) FTextureParameterValue;
-					ParameterValue->ParameterInfo.Name = CachedParameters.GetParameterName(EMaterialParameterType::Texture, ParameterIndex);
-					ParameterValue->ParameterValue = CachedParameters.TextureValues[ParameterIndex];
-				}
+					// Scalars
+					const TArray<FMaterialScalarParameterInfo, FMemoryImageAllocator>& ScalarExpressions = MaterialResource->GetUniformScalarParameterExpressions();
+					for (const FMaterialScalarParameterInfo& Parameter : ScalarExpressions)
+					{
+						FScalarParameterValue* ParameterValue = new(ScalarParameterValues) FScalarParameterValue;
+						ParameterValue->ParameterInfo.Name = *Parameter.ParameterName;
+						Parameter.GetDefaultValue(ParameterValue->ParameterValue);
+					}
 
-				// Runtime Virtual Textures
-				RuntimeVirtualTextureParameterValues.Empty(CachedParameters.GetNumParameters(EMaterialParameterType::RuntimeVirtualTexture));
-				for (int32 ParameterIndex = 0; ParameterIndex < CachedParameters.GetNumParameters(EMaterialParameterType::RuntimeVirtualTexture); ++ParameterIndex)
-				{
-					FRuntimeVirtualTextureParameterValue* ParameterValue = new(RuntimeVirtualTextureParameterValues) FRuntimeVirtualTextureParameterValue;
-					ParameterValue->ParameterInfo.Name = CachedParameters.GetParameterName(EMaterialParameterType::RuntimeVirtualTexture, ParameterIndex);
-					ParameterValue->ParameterValue = CachedParameters.RuntimeVirtualTextureValues[ParameterIndex];
+					// Vectors
+					const TArray<FMaterialVectorParameterInfo, FMemoryImageAllocator>& VectorExpressions = MaterialResource->GetUniformVectorParameterExpressions();
+					for (const FMaterialVectorParameterInfo& Parameter : VectorExpressions)
+					{
+						FVectorParameterValue* ParameterValue = new(VectorParameterValues) FVectorParameterValue;
+						ParameterValue->ParameterInfo.Name = *Parameter.ParameterName;
+						Parameter.GetDefaultValue(ParameterValue->ParameterValue);
+					}
+
+					// Textures
+					for (int32 TypeIndex = 0; TypeIndex < NumMaterialTextureParameterTypes; TypeIndex++)
+					{
+						for (const FMaterialTextureParameterInfo& Parameter : MaterialResource->GetUniformTextureExpressions((EMaterialTextureParameterType)TypeIndex))
+						{
+							if (!Parameter.ParameterInfo.Name.IsNone())
+							{
+								FTextureParameterValue* ParameterValue = new(TextureParameterValues) FTextureParameterValue;
+								ParameterValue->ParameterInfo.Name = *Parameter.ParameterName;
+								Parameter.GetGameThreadTextureValue(AsMaterial, *MaterialResource, ParameterValue->ParameterValue);
+							}
+						}
+					}
 				}
 			}
 		}
