#UE4 Fix problem where handles would incorrectly become children of other handles

diff --git a/Engine/Source/Runtime/Engine/Private/AssetManager.cpp b/Engine/Source/Runtime/Engine/Private/AssetManager.cpp
--- a/Engine/Source/Runtime/Engine/Private/AssetManager.cpp
+++ b/Engine/Source/Runtime/Engine/Private/AssetManager.cpp
@@ -1206,29 +1206,8 @@ TSharedPtr<FStreamableHandle> UAssetManager::ChangeBundleStateForPrimaryAssets(c
 					continue;
 				}
 
-				bool bIsSuperSet = false;
-				if (NewBundles->Num() > NameData->PendingState.BundleNames.Num())
-				{
-					bIsSuperSet = true;
-					for (int i = 0; i < NameData->PendingState.BundleNames.Num(); i++)
-					{
-						if (!NewBundles->Contains(NameData->PendingState.BundleNames[i]))
-						{
-							bIsSuperSet = false;
-							break;
-						}
-					}
-				}
-
-				if (bIsSuperSet)
-				{
-					ExistingHandles.Add(NameData->PendingState.Handle);
-				}
-				else
-				{
-					// Clear pending state
-					NameData->PendingState.Reset(true);
-				}
+				// Clear pending state
+				NameData->PendingState.Reset(true);
 			}
 			else if (NameData->CurrentState.IsValid() && NameData->CurrentState.BundleNames == *NewBundles)
 			{
