Temporary backout of texture resouce relase change while I fix a crash it triggers.
#rnx

diff --git a/Engine/Source/Runtime/Engine/Private/Texture.cpp b/Engine/Source/Runtime/Engine/Private/Texture.cpp
--- a/Engine/Source/Runtime/Engine/Private/Texture.cpp
+++ b/Engine/Source/Runtime/Engine/Private/Texture.cpp
@@ -112,12 +112,9 @@ void UTexture::ReleaseResource()
 		check( !Texture2D  || !Texture2D->HasPendingUpdate() );
 
 		// Free the resource.
-		BeginReleaseResource(Resource);
-		ENQUEUE_RENDER_COMMAND(DeleteResource)([ToDelete = Resource](FRHICommandListImmediate& RHICmdList)
-		{
-			delete ToDelete;
-		});
-		Resource = nullptr;
+		ReleaseResourceAndFlush(Resource);
+		delete Resource;
+		Resource = NULL;
 	}
 }
 
