Fix bug with volume scale fade not being respected when lerping after initial fade time elapsed
#rnx

diff --git a/Engine/Source/Runtime/Engine/Private/SoundConcurrency.cpp b/Engine/Source/Runtime/Engine/Private/SoundConcurrency.cpp
--- a/Engine/Source/Runtime/Engine/Private/SoundConcurrency.cpp
+++ b/Engine/Source/Runtime/Engine/Private/SoundConcurrency.cpp
@@ -120,16 +120,8 @@ float FConcurrencySoundData::GetLerpTime() const
 
 float FConcurrencySoundData::GetVolume(bool bInDecibels) const
 {
-	float VolumeDb = 0.0f;
-	if (FMath::IsNearlyZero(LerpTime) || Elapsed >= LerpTime || FMath::IsNearlyEqual(DbTargetVolume, DbStartVolume, ConcurrencyMinVolumeScaleCVar))
-	{
-		VolumeDb = DbTargetVolume;
-	}
-	else
-	{
-		const float Alpha = FMath::Clamp(Elapsed / LerpTime, 0.0f, 1.0f);
-		VolumeDb = FMath::Lerp(DbStartVolume, DbTargetVolume, Alpha);
-	}
+	const float Alpha = FMath::Clamp(Elapsed / FMath::Max(LerpTime, SMALL_NUMBER), 0.0f, 1.0f);
+	const float VolumeDb = FMath::Lerp(DbStartVolume, DbTargetVolume, Alpha);
 
 	if (bInDecibels)
 	{
@@ -137,7 +129,7 @@ float FConcurrencySoundData::GetVolume(bool bInDecibels) const
 	}
 	
 	float VolumeLin = Audio::ConvertToLinear(VolumeDb);
-	if (VolumeLin <= ConcurrencyMinVolumeScaleCVar)
+	if (VolumeLin < ConcurrencyMinVolumeScaleCVar || FMath::IsNearlyEqual(VolumeLin, ConcurrencyMinVolumeScaleCVar))
 	{
 		return 0.0f;
 	}
@@ -146,27 +138,25 @@ float FConcurrencySoundData::GetVolume(bool bInDecibels) const
 
 float FConcurrencySoundData::GetTargetVolume(bool bInDecibels) const
 {
-	return bInDecibels ? DbTargetVolume : Audio::ConvertToLinear(DbTargetVolume);
+	if (bInDecibels)
+	{
+		return DbTargetVolume;
+	}
+
+	const float VolumeLin = Audio::ConvertToLinear(DbTargetVolume);
+	if (VolumeLin < ConcurrencyMinVolumeScaleCVar || FMath::IsNearlyEqual(VolumeLin, ConcurrencyMinVolumeScaleCVar))
+	{
+		return 0.0f;
+	}
+
+	return VolumeLin;
 }
 
 void FConcurrencySoundData::SetTarget(float InTargetVolume, float InLerpTime)
 {
-	const float DbInitVolume = DbTargetVolume;
-	const float InitLerpTime = LerpTime;
-
+	DbStartVolume = GetVolume(true /* bInDecibels */);
 	DbTargetVolume = Audio::ConvertToDecibels(InTargetVolume, ConcurrencyMinVolumeScaleCVar);
 	LerpTime = FMath::Max(InLerpTime, 0.0f);
-
-	if (LerpTime > 0.0f)
-	{
-		const float Alpha = FMath::Clamp(Elapsed / LerpTime, 0.0f, 1.0f);
-		DbStartVolume = FMath::Lerp(DbStartVolume, DbTargetVolume, Alpha);
-	}
-	else
-	{
-		DbStartVolume = DbTargetVolume;
-	}
-
 	Elapsed = 0.0f;
 }
 
