Adding a child connection check to UNetConnection::CreateChannelByName since it should never be called on a UChildConnection

diff --git a/Engine/Source/Runtime/Engine/Private/NetConnection.cpp b/Engine/Source/Runtime/Engine/Private/NetConnection.cpp
--- a/Engine/Source/Runtime/Engine/Private/NetConnection.cpp
+++ b/Engine/Source/Runtime/Engine/Private/NetConnection.cpp
@@ -3051,6 +3051,7 @@ int32 UNetConnection::SendRawBunch(FOutBunch& Bunch, bool InAllowMerge, const FN
 UChannel* UNetConnection::CreateChannelByName( const FName& ChName, EChannelCreateFlags CreateFlags, int32 ChIndex )
 {
 	check(Driver->IsKnownChannelName(ChName));
+	checkf(GetUChildConnection() == nullptr, TEXT("Creating channels on child connections is unsupported: ChName: %s ChIndex: %d Connection: %s"), *ChName.ToString(), ChIndex, *Describe());
 	AssertValid();
 
 	// If no channel index was specified, find the first available.
