Fix CIS

diff --git a/Engine/Source/Runtime/Engine/Public/Physics/Experimental/ChaosInterfaceUtils.h b/Engine/Source/Runtime/Engine/Public/Physics/Experimental/ChaosInterfaceUtils.h
--- a/Engine/Source/Runtime/Engine/Public/Physics/Experimental/ChaosInterfaceUtils.h
+++ b/Engine/Source/Runtime/Engine/Public/Physics/Experimental/ChaosInterfaceUtils.h
@@ -6,14 +6,10 @@
 #include "Containers/ContainerAllocationPolicies.h"
 #include "Templates/UniquePtr.h"
 #include "PhysicsEngine/BodySetupEnums.h"
+#include "Chaos/GeometryParticles.h"
 
 struct FGeometryAddParams;
 
-namespace Chaos
-{
-	enum EChaosCollisionTraceFlag : int;
-}
-
 namespace ChaosInterface
 {
 	Chaos::EChaosCollisionTraceFlag ConvertCollisionTraceFlag(ECollisionTraceFlag Flag);
