Fix non-unity build

diff --git a/Engine/Source/Runtime/Engine/Public/Physics/ImmediatePhysics/ImmediatePhysicsChaos/ImmediatePhysicsCore_Chaos.h b/Engine/Source/Runtime/Engine/Public/Physics/ImmediatePhysics/ImmediatePhysicsChaos/ImmediatePhysicsCore_Chaos.h
--- a/Engine/Source/Runtime/Engine/Public/Physics/ImmediatePhysics/ImmediatePhysicsChaos/ImmediatePhysicsCore_Chaos.h
+++ b/Engine/Source/Runtime/Engine/Public/Physics/ImmediatePhysics/ImmediatePhysicsChaos/ImmediatePhysicsCore_Chaos.h
@@ -14,6 +14,7 @@ namespace Chaos
 	class FPBDJointConstraintHandle;
 	class FPBDJointConstraints;
 	class FPBDRigidsEvolutionGBF;
+	template<class T> class TArrayCollectionArray;
 	template<typename T, int D> struct TKinematicGeometryParticleParameters;
 	template<typename T, int D> class TKinematicTarget;
 	template<typename P> class TPBDConstraintIslandRule;
