Only update cached material data on client, as expressions are not loaded on server
- Only validate legacy parameter behaivor on client, since cached data isn't valid on server

diff --git a/Engine/Source/Runtime/Engine/Private/Materials/Material.cpp b/Engine/Source/Runtime/Engine/Private/Materials/Material.cpp
--- a/Engine/Source/Runtime/Engine/Private/Materials/Material.cpp
+++ b/Engine/Source/Runtime/Engine/Private/Materials/Material.cpp
@@ -2039,6 +2039,12 @@ bool UMaterial::GetParameterDesc(const FHashedMaterialParameterInfo& ParameterIn
 }
 void UMaterial::UpdateCachedExpressionData()
 {
+	if (!GIsClient)
+	{
+		// Expressions may not be loaded on server, so only rebuild the data if we're the client
+		return;
+	}
+
 	{
 		FMaterialCachedExpressionData UpdatedCachedExpressionData;
 		UpdatedCachedExpressionData.Reset();
@@ -2078,11 +2084,15 @@ bool UMaterial::GetScalarParameterValue(const FHashedMaterialParameterInfo& Para
 {
 	const bool bResult = GetScalarParameterValue_New(ParameterInfo, OutValue, bOveriddenOnly);
 #if WITH_EDITOR
-	float OldValue = 0.0f;
-	const bool bOldResult = GetScalarParameterValue_Legacy(ParameterInfo, OldValue, bOveriddenOnly);
+	if (GIsClient)
+	{
+		// Expressions may not be loaded on server, and cached data is not updated...results will be incorrect
+		float OldValue = 0.0f;
+		const bool bOldResult = GetScalarParameterValue_Legacy(ParameterInfo, OldValue, bOveriddenOnly);
 
-	check(bOldResult == bResult);
-	check(!bResult || OutValue == OldValue);
+		check(bOldResult == bResult);
+		check(!bResult || OutValue == OldValue);
+	}
 #endif
 	return bResult;
 }
@@ -2218,11 +2228,15 @@ bool UMaterial::GetVectorParameterValue(const FHashedMaterialParameterInfo& Para
 {
 	const bool bResult = GetVectorParameterValue_New(ParameterInfo, OutValue, bOveriddenOnly);
 #if WITH_EDITOR
-	FLinearColor OldValue(ForceInitToZero);
-	const bool bOldResult = GetVectorParameterValue_Legacy(ParameterInfo, OldValue, bOveriddenOnly);
+	if (GIsClient)
+	{
+		// Expressions may not be loaded on server, and cached data is not updated...results will be incorrect
+		FLinearColor OldValue(ForceInitToZero);
+		const bool bOldResult = GetVectorParameterValue_Legacy(ParameterInfo, OldValue, bOveriddenOnly);
 
-	check(bOldResult == bResult);
-	check(!bResult || OutValue == OldValue);
+		check(bOldResult == bResult);
+		check(!bResult || OutValue == OldValue);
+	}
 #endif
 	return bResult;
 }
@@ -2342,11 +2356,15 @@ bool UMaterial::GetTextureParameterValue(const FHashedMaterialParameterInfo& Par
 {
 	const bool bResult = GetTextureParameterValue_New(ParameterInfo, OutValue, bOveriddenOnly);
 #if WITH_EDITOR
-	UTexture* OldValue = nullptr;
-	const bool bOldResult = GetTextureParameterValue_Legacy(ParameterInfo, OldValue, bOveriddenOnly);
+	if (GIsClient)
+	{
+		// Expressions may not be loaded on server, and cached data is not updated...results will be incorrect
+		UTexture* OldValue = nullptr;
+		const bool bOldResult = GetTextureParameterValue_Legacy(ParameterInfo, OldValue, bOveriddenOnly);
 
-	check(bOldResult == bResult);
-	check(!bOldResult || OutValue == OldValue); // if result is false, texture value is undefined
+		check(bOldResult == bResult);
+		check(!bOldResult || OutValue == OldValue); // if result is false, texture value is undefined
+	}
 #endif
 	return bResult;
 }
@@ -2438,11 +2456,15 @@ bool UMaterial::GetRuntimeVirtualTextureParameterValue(const FHashedMaterialPara
 {
 	const bool bResult = GetRuntimeVirtualTextureParameterValue_New(ParameterInfo, OutValue, bOveriddenOnly);
 #if WITH_EDITOR
-	URuntimeVirtualTexture* OldValue = nullptr;
-	const bool bOldResult = GetRuntimeVirtualTextureParameterValue_Legacy(ParameterInfo, OldValue, bOveriddenOnly);
+	if (GIsClient)
+	{
+		// Expressions may not be loaded on server, and cached data is not updated...results will be incorrect
+		URuntimeVirtualTexture* OldValue = nullptr;
+		const bool bOldResult = GetRuntimeVirtualTextureParameterValue_Legacy(ParameterInfo, OldValue, bOveriddenOnly);
 
-	check(bOldResult == bResult);
-	check(!bResult || OutValue == OldValue);
+		check(bOldResult == bResult);
+		check(!bResult || OutValue == OldValue);
+	}
 #endif
 	return bResult;
 }
@@ -2547,12 +2569,16 @@ bool UMaterial::GetFontParameterValue(const FHashedMaterialParameterInfo& Parame
 {
 	const bool bResult = GetFontParameterValue_New(ParameterInfo, OutFontValue, OutFontPage, bOveriddenOnly);
 #if WITH_EDITOR
-	UFont* OldValue = nullptr;
-	int32 OldPage = INDEX_NONE;
-	const bool bOldResult = GetFontParameterValue_Legacy(ParameterInfo, OldValue, OldPage, bOveriddenOnly);
+	if (GIsClient)
+	{
+		// Expressions may not be loaded on server, and cached data is not updated...results will be incorrect
+		UFont* OldValue = nullptr;
+		int32 OldPage = INDEX_NONE;
+		const bool bOldResult = GetFontParameterValue_Legacy(ParameterInfo, OldValue, OldPage, bOveriddenOnly);
 
-	check(bOldResult == bResult);
-	check(!bResult || (OutFontValue == OldValue && OldPage == OutFontPage));
+		check(bOldResult == bResult);
+		check(!bResult || (OutFontValue == OldValue && OldPage == OutFontPage));
+	}
 #endif
 	return bResult;
 }
diff --git a/Engine/Source/Runtime/Engine/Private/Materials/MaterialInstance.cpp b/Engine/Source/Runtime/Engine/Private/Materials/MaterialInstance.cpp
--- a/Engine/Source/Runtime/Engine/Private/Materials/MaterialInstance.cpp
+++ b/Engine/Source/Runtime/Engine/Private/Materials/MaterialInstance.cpp
@@ -4008,6 +4008,12 @@ void UMaterialInstance::UpdateStaticPermutation(const FStaticParameterSet& NewPa
 
 void UMaterialInstance::UpdateCachedLayerParameters()
 {
+	if (!GIsClient)
+	{
+		// Expressions may not be loaded on server, so only rebuild the data if we're the client
+		return;
+	}
+
 	FMaterialCachedExpressionData CachedExpressionData;
 	CachedExpressionData.Reset();
 	if (Parent)
