fix for assert when FinishRecordingOutput was given an invalid path argument.

diff --git a/Engine/Source/Runtime/Engine/Private/SampleBufferIO.cpp b/Engine/Source/Runtime/Engine/Private/SampleBufferIO.cpp
--- a/Engine/Source/Runtime/Engine/Private/SampleBufferIO.cpp
+++ b/Engine/Source/Runtime/Engine/Private/SampleBufferIO.cpp
@@ -183,7 +183,12 @@ namespace Audio
 		FText InvalidPathReason;
 		bool const bValidPackageName = FPackageName::IsValidLongPackageName(AbsoluteFilePath, false, &InvalidPathReason);
 
-		check(bValidPackageName);
+		if (!bValidPackageName)
+		{
+			UE_LOG(LogAudio, Warning, TEXT("File path given was not valid (%s): the recording output will be saved to a USoundWave at the root of the Content folder."), *InvalidPathReason.ToString());
+			AbsoluteFilePath = TEXT("/Game/") + FString(TEXT("/")) + FileName;
+			AbsoluteFilePath = AbsoluteFilePath.Replace(TEXT("//"), TEXT("/"), ESearchCase::CaseSensitive);
+		}
 
 		// Set up Package.
 		CurrentPackage = CreatePackage(nullptr, *AbsoluteFilePath);
@@ -287,7 +292,13 @@ namespace Audio
 			FText InvalidPathReason;
 			bool const bValidPackageName = FPackageName::IsValidLongPackageName(AbsoluteFilePath, false, &InvalidPathReason);
 
-			check(bValidPackageName);
+			if (!bValidPackageName)
+			{
+				UE_LOG(LogAudio, Warning, TEXT("File path given was not valid (%s): the recording output will be saved to a USoundWave at the root of the Content folder."), *InvalidPathReason.ToString());
+				AbsoluteFilePath = TEXT("/Game/") + *FileName;
+				FPaths::NormalizeDirectoryName(AbsoluteFilePath);
+				AbsoluteFilePath = AbsoluteFilePath.Replace(TEXT("//"), TEXT("/"), ESearchCase::CaseSensitive);
+			}
 
 			// Set up Package.
 			CurrentPackage = CreatePackage(nullptr, *AbsoluteFilePath);
