Async UTexture resource release instead of flushing render thread.

diff --git a/Engine/Source/Runtime/Engine/Private/Texture.cpp b/Engine/Source/Runtime/Engine/Private/Texture.cpp
--- a/Engine/Source/Runtime/Engine/Private/Texture.cpp
+++ b/Engine/Source/Runtime/Engine/Private/Texture.cpp
@@ -112,9 +112,12 @@ void UTexture::ReleaseResource()
 		check( !Texture2D  || !Texture2D->HasPendingUpdate() );
 
 		// Free the resource.
-		ReleaseResourceAndFlush(Resource);
-		delete Resource;
-		Resource = NULL;
+		BeginReleaseResource(Resource);
+		ENQUEUE_RENDER_COMMAND(DeleteResource)([ToDelete = Resource](FRHICommandListImmediate& RHICmdList)
+		{
+			delete ToDelete;
+		});
+		Resource = nullptr;
 	}
 }
 
diff --git a/Engine/Source/Runtime/Engine/Private/Texture2D.cpp b/Engine/Source/Runtime/Engine/Private/Texture2D.cpp
--- a/Engine/Source/Runtime/Engine/Private/Texture2D.cpp
+++ b/Engine/Source/Runtime/Engine/Private/Texture2D.cpp
@@ -1781,16 +1781,6 @@ void FTexture2DResource::InitRHI()
  */
 void FTexture2DResource::ReleaseRHI()
 {
-	const TIndirectArray<FTexture2DMipMap>& OwnerMips = Owner->GetPlatformMips();
-
-	// It should be safe to release the texture.
-	check(!Owner->HasPendingUpdate());
-
-	if ( (Texture2DRHI->GetFlags() & TexCreate_Virtual) != TexCreate_Virtual )
-	{
-		check(OwnerMips[CurrentFirstMip].SizeX == Texture2DRHI->GetSizeX() && OwnerMips[CurrentFirstMip].SizeY == Texture2DRHI->GetSizeY());
-	}
-
 	DEC_DWORD_STAT_BY( STAT_TextureMemory, TextureSize );
 	DEC_DWORD_STAT_FNAME_BY( LODGroupStatName, TextureSize );
 
