Enable masked bending to be used with SingleLayerWater ShadingModel

diff --git a/Engine/Shaders/Private/SingleLayerWaterShading.ush b/Engine/Shaders/Private/SingleLayerWaterShading.ush
--- a/Engine/Shaders/Private/SingleLayerWaterShading.ush
+++ b/Engine/Shaders/Private/SingleLayerWaterShading.ush
@@ -8,6 +8,12 @@
 #include "DistortionCommon.ush"
 #include "SingleLayerWaterCommon.ush"
 
+#if MATERIAL_SHADINGMODEL_SINGLELAYERWATER
+	// We need the water base pass to be allowed to use masking (if blend mode is Masked) because it doesn't rely on a depth pre-pass (where clipping would normally occur)
+	#undef EARLY_Z_PASS_ONLY_MATERIAL_MASKING
+	#define EARLY_Z_PASS_ONLY_MATERIAL_MASKING 0
+#endif
+
 struct WaterVolumeLightingOutput
 {
 	float3 Luminance;
diff --git a/Engine/Source/Runtime/Engine/Private/Materials/HLSLMaterialTranslator.cpp b/Engine/Source/Runtime/Engine/Private/Materials/HLSLMaterialTranslator.cpp
--- a/Engine/Source/Runtime/Engine/Private/Materials/HLSLMaterialTranslator.cpp
+++ b/Engine/Source/Runtime/Engine/Private/Materials/HLSLMaterialTranslator.cpp
@@ -890,9 +890,9 @@ bool FHLSLMaterialTranslator::Translate()
 
 		if (MaterialShadingModels.HasShadingModel(MSM_SingleLayerWater))
 		{
-			if (BlendMode != BLEND_Opaque)
+			if (BlendMode != BLEND_Opaque && BlendMode != BLEND_Masked)
 			{
-				Errorf(TEXT("SingleLayerWater materials must be opaque."));
+				Errorf(TEXT("SingleLayerWater materials must be opaque or masked."));
 			}
 			if (!MaterialShadingModels.HasOnlyShadingModel(MSM_SingleLayerWater))
 			{
