Optimization for ray tracing ISMs to only add instances if they have not been culled.

diff --git a/Engine/Source/Runtime/Engine/Private/InstancedStaticMesh.cpp b/Engine/Source/Runtime/Engine/Private/InstancedStaticMesh.cpp
--- a/Engine/Source/Runtime/Engine/Private/InstancedStaticMesh.cpp
+++ b/Engine/Source/Runtime/Engine/Private/InstancedStaticMesh.cpp
@@ -1155,20 +1155,7 @@ void FInstancedStaticMeshSceneProxy::GetDynamicRayTracingInstances(struct FRayTr
 	//preallocate the worst-case to prevent an explosion of reallocs
 	//#dxr_todo: possibly track used instances and reserve based on previous behavior
 	RayTracingInstanceTemplate.InstanceTransforms.Reserve(InstanceCount);
-
-
-	int32 SectionCount = InstancedRenderData.LODModels[LOD].Sections.Num();
-
-	for (int32 SectionIdx = 0; SectionIdx < SectionCount; ++SectionIdx)
-	{
-		//#dxr_todo: so far we use the parent static mesh path to get material data
-		FMeshBatch MeshBatch;
-		GetMeshElement(LOD, 0, SectionIdx, 0, false, false, MeshBatch);
-
-		RayTracingInstanceTemplate.Materials.Add(MeshBatch);
-	}
-	RayTracingInstanceTemplate.BuildInstanceMaskAndFlags();
-
+	
 	if (CVarRayTracingRenderInstancesCulling.GetValueOnRenderThread() > 0 && RayTracingCullClusterInstances.Num() > 0)
 	{
 		const float BVHCullRadius = CVarRayTracingInstancesCullClusterRadius.GetValueOnRenderThread();
@@ -1245,7 +1232,23 @@ void FInstancedStaticMeshSceneProxy::GetDynamicRayTracingInstances(struct FRayTr
 		}
 	}
 
-	OutRayTracingInstances.Add(RayTracingInstanceTemplate);
+
+	if (RayTracingInstanceTemplate.InstanceTransforms.Num() > 0)
+	{
+		int32 SectionCount = InstancedRenderData.LODModels[LOD].Sections.Num();
+
+		for (int32 SectionIdx = 0; SectionIdx < SectionCount; ++SectionIdx)
+		{
+			//#dxr_todo: so far we use the parent static mesh path to get material data
+			FMeshBatch MeshBatch;
+			FStaticMeshSceneProxy::GetMeshElement(LOD, 0, SectionIdx, 0, false, false, MeshBatch);
+
+			RayTracingInstanceTemplate.Materials.Add(MeshBatch);
+		}
+		RayTracingInstanceTemplate.BuildInstanceMaskAndFlags();
+
+		OutRayTracingInstances.Add(RayTracingInstanceTemplate);
+	}
 }
 
 void FInstancedStaticMeshSceneProxy::SetupRayTracingCullClusters()
