Optimization in ISM culling structure to avoid transforming them at each frame. Also it replaces a linked list structure with a more efficient AOS

diff --git a/Engine/Source/Runtime/Engine/Private/InstancedStaticMesh.cpp b/Engine/Source/Runtime/Engine/Private/InstancedStaticMesh.cpp
--- a/Engine/Source/Runtime/Engine/Private/InstancedStaticMesh.cpp
+++ b/Engine/Source/Runtime/Engine/Private/InstancedStaticMesh.cpp
@@ -74,7 +74,7 @@ static TAutoConsoleVariable<int32> CVarRayTracingRenderInstancesCulling(
 static TAutoConsoleVariable<float> CVarRayTracingInstancesCullClusterMaxRadiusMultiplier(
 	TEXT("r.RayTracing.InstancedStaticMeshes.CullClusterMaxRadiusMultiplier"),
 	20.0f, 
-	TEXT("Multiplier for the maximum instance size (default = 20cm)"));
+	TEXT("Multiplier for the maximum instance size (default = 20)"));
 
 static TAutoConsoleVariable<float> CVarRayTracingInstancesCullClusterRadius(
 	TEXT("r.RayTracing.InstancedStaticMeshes.CullClusterRadius"),
@@ -1155,8 +1155,8 @@ void FInstancedStaticMeshSceneProxy::GetDynamicRayTracingInstances(struct FRayTr
 	//preallocate the worst-case to prevent an explosion of reallocs
 	//#dxr_todo: possibly track used instances and reserve based on previous behavior
 	RayTracingInstanceTemplate.InstanceTransforms.Reserve(InstanceCount);
-	
-	if (CVarRayTracingRenderInstancesCulling.GetValueOnRenderThread() > 0 && RayTracingCullClusterInstances.Num() > 0)
+
+	if (CVarRayTracingRenderInstancesCulling.GetValueOnRenderThread() > 0 && RayTracingCullClusters.Num() > 0)
 	{
 		const float BVHCullRadius = CVarRayTracingInstancesCullClusterRadius.GetValueOnRenderThread();
 		const float BVHLowScaleThreshold = CVarRayTracingInstancesLowScaleThreshold.GetValueOnRenderThread();
@@ -1166,10 +1166,12 @@ void FInstancedStaticMeshSceneProxy::GetDynamicRayTracingInstances(struct FRayTr
 		FMatrix ToWorld = GetLocalToWorld();
 
 		// Iterate over all culling clusters
-		for (int32 ClusterIdx = 0; ClusterIdx < RayTracingCullClusterBoundsMin.Num(); ++ClusterIdx)
+		for (int32 ClusterIdx = 0; ClusterIdx < RayTracingCullClusters.Num(); ++ClusterIdx)
 		{
-			FVector VClusterBBoxSize = RayTracingCullClusterBoundsMax[ClusterIdx] - RayTracingCullClusterBoundsMin[ClusterIdx];
-			FVector VClusterCenter = 0.5f * (RayTracingCullClusterBoundsMax[ClusterIdx] + RayTracingCullClusterBoundsMin[ClusterIdx]);
+			FRayTracingCullCluster& Cluster = RayTracingCullClusters[ClusterIdx];
+
+			FVector VClusterBBoxSize = Cluster.BoundsMax - Cluster.BoundsMin;
+			FVector VClusterCenter = 0.5f * (Cluster.BoundsMax + Cluster.BoundsMin);
 			FVector VToClusterCenter = VClusterCenter - Context.ReferenceView->ViewLocation;
 			float ClusterRadius = 0.5f * VClusterBBoxSize.Size();
 			float DistToClusterCenter = VToClusterCenter.Size();
@@ -1180,31 +1182,19 @@ void FInstancedStaticMeshSceneProxy::GetDynamicRayTracingInstances(struct FRayTr
 				continue;
 			}
 
-			TDoubleLinkedList< uint32 >* InstanceList = RayTracingCullClusterInstances[ClusterIdx];
-
-			// Unroll instances in the current cluster into the array
-			for (TDoubleLinkedList<uint32>::TDoubleLinkedListNode* InstancePtr = InstanceList->GetHead(); InstancePtr != nullptr; InstancePtr = InstancePtr->GetNextNode())
+			// Cull instances in the cluster
+			for (FCullNode& Node : Cluster.Nodes)
 			{
-				const uint32 Instance = InstancePtr->GetValue();
+				const uint32 InstanceIdx = Node.Instance;
 
-				FMatrix Transform;
-				InstancedRenderData.PerInstanceRenderData->InstanceBuffer.GetInstanceTransform(Instance, Transform);
-				Transform.M[3][3] = 1.0f;
-				FMatrix InstanceTransform = Transform * GetLocalToWorld();
-				FVector InstanceLocation = InstanceTransform.TransformPosition({ 0.0f,0.0f,0.0f });
+				FVector InstanceLocation = Node.Center;
 				FVector VToInstanceCenter = Context.ReferenceView->ViewLocation - InstanceLocation;
-				float   DistanceToInstanceCenter = VToInstanceCenter.Size();
-
-				FVector VMin, VMax, VDiag;
-				VMin = InstanceTransform.TransformPosition(GetLocalBounds().GetBox().Min);
-				VMax = InstanceTransform.TransformPosition(GetLocalBounds().GetBox().Max);
-				VDiag = VMax - VMin;
-
-				float InstanceRadius = 0.5f * VDiag.Size();
+				float DistanceToInstanceCenter = VToInstanceCenter.Size();
+				float InstanceRadius = Node.Radius;
 				float DistanceToInstanceStart = DistanceToInstanceCenter - InstanceRadius;
 
 				// Cull instance based on distance
-				if (DistanceToInstanceStart > BVHCullRadius&& ApplyGeneralCulling)
+				if (DistanceToInstanceStart > BVHCullRadius && ApplyGeneralCulling)
 					continue;
 
 				// Special culling for small scale objects
@@ -1214,6 +1204,11 @@ void FInstancedStaticMeshSceneProxy::GetDynamicRayTracingInstances(struct FRayTr
 						continue;
 				}
 
+				FMatrix Transform;
+				InstancedRenderData.PerInstanceRenderData->InstanceBuffer.GetInstanceTransform(InstanceIdx, Transform);
+				Transform.M[3][3] = 1.0f;
+				FMatrix InstanceTransform = Transform * GetLocalToWorld();
+
 				RayTracingInstanceTemplate.InstanceTransforms.Add(InstanceTransform);
 			}
 		}
@@ -1276,26 +1271,39 @@ void FInstancedStaticMeshSceneProxy::SetupRayTracingCullClusters()
 		float MaxInstanceRadius = 0.0f;
 
 		// Init first cluster
-		RayTracingCullClusterInstances.Add(new TDoubleLinkedList< uint32>());
-		RayTracingCullClusterBoundsMin.Add(FVector(MAX_FLT, MAX_FLT, MAX_FLT));
-		RayTracingCullClusterBoundsMax.Add(FVector(-MAX_FLT, -MAX_FLT, -MAX_FLT));
+		{
+			FRayTracingCullCluster &Cluster = RayTracingCullClusters.Emplace_GetRef();
+			Cluster.BoundsMin = FVector(MAX_FLT, MAX_FLT, MAX_FLT);
+			Cluster.BoundsMax = FVector(-MAX_FLT, -MAX_FLT, -MAX_FLT);
+		}
 
-		// Traverse instances to find maximum rarius
+		TArray<float> InstanceRadii;
+		InstanceRadii.Reserve(InstanceCount);
+
+		// Traverse instances to find maximum radius
 		for (int32 Instance = 0; Instance < InstanceCount; ++Instance)
 		{
-			FMatrix Transform;
-			InstancedRenderData.PerInstanceRenderData->InstanceBuffer.GetInstanceTransform(Instance, Transform);
-			Transform.M[3][3] = 1.0f;
-			FMatrix InstanceTransform = Transform * ComponentLocalToWorld;
+			if (InstancedRenderData.Component->PerInstanceSMData.IsValidIndex(Instance))
+			{
+				FMatrix Transform;
+				InstancedRenderData.PerInstanceRenderData->InstanceBuffer.GetInstanceTransform(Instance, Transform);
+				Transform.M[3][3] = 1.0f;
+				FMatrix InstanceTransform = Transform * ComponentLocalToWorld;
 
-			FVector VMin, VMax;
-			InstancedRenderData.Component->GetLocalBounds(VMin, VMax);
-			VMin = InstanceTransform.TransformPosition(VMin);
-			VMax = InstanceTransform.TransformPosition(VMax);
+				FVector VMin, VMax;
+				InstancedRenderData.Component->GetLocalBounds(VMin, VMax);
+				VMin = InstanceTransform.TransformPosition(VMin);
+				VMax = InstanceTransform.TransformPosition(VMax);
 
-			FVector VBBoxSize = VMax - VMin;
+				FVector VBBoxSize = VMax - VMin;
 
-			MaxInstanceRadius = FMath::Max(0.5f * VBBoxSize.Size(), MaxInstanceRadius);
+				MaxInstanceRadius = FMath::Max(0.5f * VBBoxSize.Size(), MaxInstanceRadius);
+				InstanceRadii.Add(0.5f * VBBoxSize.Size());
+			}
+			else
+			{
+				InstanceRadii.Add(0.0f);
+			}
 		}
 
 		float MaxClusterRadius = MaxInstanceRadius * MaxClusterRadiusMultiplier;
@@ -1316,8 +1324,8 @@ void FInstancedStaticMeshSceneProxy::SetupRayTracingCullClusters()
 			for (int32 CandidateCluster = 0; CandidateCluster <= ClusterIndex; ++CandidateCluster)
 			{
 				// Build new candidate cluster bounds
-				FVector VCandidateMin = VMin.ComponentMin(RayTracingCullClusterBoundsMin[CandidateCluster]);
-				FVector VCandidateMax = VMax.ComponentMax(RayTracingCullClusterBoundsMax[CandidateCluster]);
+				FVector VCandidateMin = VMin.ComponentMin(RayTracingCullClusters[CandidateCluster].BoundsMin);
+				FVector VCandidateMax = VMax.ComponentMax(RayTracingCullClusters[CandidateCluster].BoundsMax);
 
 				FVector VCandidateBBoxSize = VCandidateMax - VCandidateMin;
 				float MaxCandidateRadius = 0.5f * VCandidateBBoxSize.Size();
@@ -1325,9 +1333,12 @@ void FInstancedStaticMeshSceneProxy::SetupRayTracingCullClusters()
 				// If new candidate is still small enough, update current cluster
 				if (MaxCandidateRadius <= MaxClusterRadius)
 				{
-					RayTracingCullClusterInstances[CandidateCluster]->AddTail(Instance);
-					RayTracingCullClusterBoundsMin[CandidateCluster] = VCandidateMin;
-					RayTracingCullClusterBoundsMax[CandidateCluster] = VCandidateMax;
+					RayTracingCullClusters[CandidateCluster].BoundsMin = VCandidateMin;
+					RayTracingCullClusters[CandidateCluster].BoundsMax = VCandidateMax;
+					FCullNode& Node = RayTracingCullClusters[CandidateCluster].Nodes.Emplace_GetRef();
+					Node.Radius = InstanceRadii[Instance];
+					Node.Center = InstanceLocation;
+					Node.Instance = Instance;
 					bClusterFound = true;
 					break;
 				}
@@ -1337,10 +1348,16 @@ void FInstancedStaticMeshSceneProxy::SetupRayTracingCullClusters()
 			if (!bClusterFound)
 			{
 				++ClusterIndex;
-				RayTracingCullClusterInstances.Add(new TDoubleLinkedList< uint32>());
-				RayTracingCullClusterInstances[ClusterIndex]->AddTail(Instance);
-				RayTracingCullClusterBoundsMin.Add(VMin);
-				RayTracingCullClusterBoundsMax.Add(VMax);
+
+				FRayTracingCullCluster &Cluster = RayTracingCullClusters.Emplace_GetRef();
+				Cluster.BoundsMin = VMin;
+				Cluster.BoundsMax = VMax;
+				FCullNode Node;
+
+				Node.Radius = InstanceRadii[Instance];
+				Node.Center = InstanceLocation;
+				Node.Instance = Instance;
+				Cluster.Nodes.Add(Node);
 			}
 		}
 	}
diff --git a/Engine/Source/Runtime/Engine/Private/InstancedStaticMesh.h b/Engine/Source/Runtime/Engine/Private/InstancedStaticMesh.h
--- a/Engine/Source/Runtime/Engine/Private/InstancedStaticMesh.h
+++ b/Engine/Source/Runtime/Engine/Private/InstancedStaticMesh.h
@@ -468,12 +468,6 @@ public:
 
 	~FInstancedStaticMeshSceneProxy()
 	{
-#if RHI_RAYTRACING
-		for (int32 i = 0; i < RayTracingCullClusterInstances.Num(); ++i)
-		{
-			delete RayTracingCullClusterInstances[i];
-		}
-#endif
 	}
 
 	// FPrimitiveSceneProxy interface.
@@ -564,9 +558,23 @@ protected:
 	FInstancingUserData UserData_DeselectedInstances;
 
 #if RHI_RAYTRACING
-	TArray< FVector >						RayTracingCullClusterBoundsMin;
-	TArray< FVector >						RayTracingCullClusterBoundsMax;
-	TArray< TDoubleLinkedList< uint32 >* >	RayTracingCullClusterInstances;
+	
+	/* Precomputed bounding spheres for culling */
+	struct FCullNode
+	{
+		FVector Center;
+		float Radius;
+		uint32 Instance;
+	};
+
+	struct FRayTracingCullCluster
+	{
+		FVector BoundsMin;
+		FVector BoundsMax;
+		TArray<FCullNode> Nodes;
+	};
+
+	TArray<FRayTracingCullCluster> RayTracingCullClusters;
 #endif
 
 	/** Common path for the Get*MeshElement functions */
