Fix CIS

diff --git a/Engine/Source/Runtime/Engine/Classes/PhysicsEngine/ConvexElem.h b/Engine/Source/Runtime/Engine/Classes/PhysicsEngine/ConvexElem.h
--- a/Engine/Source/Runtime/Engine/Classes/PhysicsEngine/ConvexElem.h
+++ b/Engine/Source/Runtime/Engine/Classes/PhysicsEngine/ConvexElem.h
@@ -6,7 +6,7 @@
 #include "UObject/ObjectMacros.h"
 #include "PhysicsEngine/ShapeElem.h"
 
-#define TEMP_HEADER_CHAOS_LEVEL_1 1
+#define TEMP_HEADER_CHAOS_LEVEL_1
 #include "Chaos/Convex.h"
 #undef TEMP_HEADER_CHAOS_LEVEL_1
 
diff --git a/Engine/Source/Runtime/Engine/Public/Physics/PhysicsInterfaceCore.h b/Engine/Source/Runtime/Engine/Public/Physics/PhysicsInterfaceCore.h
--- a/Engine/Source/Runtime/Engine/Public/Physics/PhysicsInterfaceCore.h
+++ b/Engine/Source/Runtime/Engine/Public/Physics/PhysicsInterfaceCore.h
@@ -15,7 +15,7 @@
 
 #elif WITH_CHAOS
 
-#define TEMP_HEADER_CHAOS_LEVEL_1 1
+#define TEMP_HEADER_CHAOS_LEVEL_1
 #include "Physics/Experimental/PhysInterface_Chaos.h"
 #include "Physics/Experimental/PhysScene_Chaos.h"
 #undef TEMP_HEADER_CHAOS_LEVEL_1
