#UE4 Minor logging improvement

diff --git a/Engine/Source/Runtime/Engine/Private/AssetManager.cpp b/Engine/Source/Runtime/Engine/Private/AssetManager.cpp
--- a/Engine/Source/Runtime/Engine/Private/AssetManager.cpp
+++ b/Engine/Source/Runtime/Engine/Private/AssetManager.cpp
@@ -3075,8 +3075,15 @@ void UAssetManager::UpdateManagementDatabase(bool bForceRefresh)
 			{
 				FName PackageName = FName(*AssetRef.GetLongPackageName());
 
-				AssetPackagesReferenced.AddUnique(PackageName);
-				PackagesToUpdateChunksFor.Add(PackageName);
+				if (PackageName == NAME_None)
+				{
+					UE_LOG(LogAssetManager, Warning, TEXT("Ignoring 'None' reference originating from %s from NameData"), *PrimaryAssetId.ToString());
+				}
+				else
+				{
+					AssetPackagesReferenced.AddUnique(PackageName);
+					PackagesToUpdateChunksFor.Add(PackageName);
+				}
 			}
 
 			TMap<FName, FAssetBundleEntry>* BundleMap = CachedAssetBundles.Find(PrimaryAssetId);
@@ -3090,20 +3097,21 @@ void UAssetManager::UpdateManagementDatabase(bool bForceRefresh)
 					{
 						FName PackageName = FName(*BundleAssetRef.GetLongPackageName());
 
-						AssetPackagesReferenced.AddUnique(PackageName);
-						PackagesToUpdateChunksFor.Add(PackageName);
+						if (PackageName == NAME_None)
+						{
+							UE_LOG(LogAssetManager, Warning, TEXT("Ignoring 'None' reference originating from %s from Bundle %s"), *PrimaryAssetId.ToString(), *BundlePair.Key.ToString());
+						}
+						else
+						{
+							AssetPackagesReferenced.AddUnique(PackageName);
+							PackagesToUpdateChunksFor.Add(PackageName);
+						}
 					}
 				}
 			}
 
 			for (const FName& AssetPackage : AssetPackagesReferenced)
 			{
-				if (AssetPackage == NAME_None)
-				{
-					UE_LOG(LogAssetManager, Warning, TEXT("Ignoring 'None' reference originating from %s"), *PrimaryAssetId.ToString());
-					continue;
-				}
-
 				TMultiMap<FAssetIdentifier, FAssetIdentifier>& ManagerMap = Rules.bApplyRecursively ? PriorityManagementMap.FindOrAdd(Rules.Priority) : NoReferenceManagementMap;
 
 				ManagerMap.Add(PrimaryAssetId, AssetPackage);
