- Remvoed some debug code that was left in accidentally
#rnx

diff --git a/Engine/Source/Runtime/Engine/Private/Materials/MaterialShared.cpp b/Engine/Source/Runtime/Engine/Private/Materials/MaterialShared.cpp
--- a/Engine/Source/Runtime/Engine/Private/Materials/MaterialShared.cpp
+++ b/Engine/Source/Runtime/Engine/Private/Materials/MaterialShared.cpp
@@ -2574,13 +2574,6 @@ void FMaterialRenderProxy::UpdateDeferredCachedUniformExpressions()
 			UE_LOG(LogMaterial, Fatal, TEXT("FMaterialRenderProxy deleted and GC mark was: %i"), MaterialProxy->IsMarkedForGarbageCollection());
 		}
 
-		// @todo fnhack
-		if (MaterialProxy->GetFriendlyName() == TEXT("FootSplash_Particle"))
-		{
-			continue;
-		}
-
-
 		UMaterialInterface::IterateOverActiveFeatureLevels([&](ERHIFeatureLevel::Type InFeatureLevel)
 		{
 			const FMaterialRenderProxy* FallbackMaterialProxy = nullptr;
