Fix PhysX build.

diff --git a/Engine/Source/Runtime/Engine/Classes/Components/SkeletalMeshComponent.h b/Engine/Source/Runtime/Engine/Classes/Components/SkeletalMeshComponent.h
--- a/Engine/Source/Runtime/Engine/Classes/Components/SkeletalMeshComponent.h
+++ b/Engine/Source/Runtime/Engine/Classes/Components/SkeletalMeshComponent.h
@@ -477,6 +477,10 @@ public:
 
 	/** Indicates that this SkeletalMeshComponent has deferred kinematic bone updates until next physics sim if not INDEX_NONE. */
 	int32 DeferredKinematicUpdateIndex;
+#if PHYSICS_INTERFACE_PHYSX
+	/** Indicates that this SkeletalMeshComponent has deferred kinematic bone updates until next physics sim.  */
+	uint8 bDeferredKinematicUpdate:1;
+#endif
 
 private:
 	/** Disable rigid body animation nodes and play original animation without simulation */
