* Add CollisionEnabled setting to FKShapeElem to configure initial shape sim and query enable states.
* Change bDisable on PerShapeData to !bSimEnabled and add bQueryEnabled.
* Check bQueryEnabled in shape pre-filter callback.

diff --git a/Engine/Source/Runtime/Core/Public/UObject/ExternalPhysicsCustomObjectVersion.h b/Engine/Source/Runtime/Core/Public/UObject/ExternalPhysicsCustomObjectVersion.h
--- a/Engine/Source/Runtime/Core/Public/UObject/ExternalPhysicsCustomObjectVersion.h
+++ b/Engine/Source/Runtime/Core/Public/UObject/ExternalPhysicsCustomObjectVersion.h
@@ -123,6 +123,9 @@ struct CORE_API FExternalPhysicsCustomObjectVersion
 		// Added serialization for the physics material sleep counter threshold
 		PhysicsMaterialSleepCounterThreshold,
 
+		// Added ability to remove shapes from sim and/or query separately
+		AddShapeSimAndQueryCollisionEnabled,
+
 		// -----<new versions can be added above this line>-------------------------------------------------
 		VersionPlusOne,
 		LatestVersion = VersionPlusOne - 1
diff --git a/Engine/Source/Runtime/Engine/Classes/PhysicsEngine/ShapeElem.h b/Engine/Source/Runtime/Engine/Classes/PhysicsEngine/ShapeElem.h
--- a/Engine/Source/Runtime/Engine/Classes/PhysicsEngine/ShapeElem.h
+++ b/Engine/Source/Runtime/Engine/Classes/PhysicsEngine/ShapeElem.h
@@ -33,6 +33,7 @@ struct FKShapeElem
 	: RestOffset(0.f)
 	, ShapeType(EAggCollisionShape::Unknown)
 	, bContributeToMass(true)
+	, CollisionEnabled(ECollisionEnabled::QueryAndPhysics)
 	, UserData(this)
 	{}
 
@@ -40,6 +41,7 @@ struct FKShapeElem
 	: RestOffset(0.f)
 	, ShapeType(InShapeType)
 	, bContributeToMass(true)
+	, CollisionEnabled(ECollisionEnabled::QueryAndPhysics)
 	, UserData(this)
 	{}
 
@@ -48,6 +50,7 @@ struct FKShapeElem
 	, Name(Copy.Name)
 	, ShapeType(Copy.ShapeType)
 	, bContributeToMass(Copy.bContributeToMass)
+	, CollisionEnabled(ECollisionEnabled::QueryAndPhysics)
 	, UserData(this)
 	{
 	}
@@ -83,6 +86,14 @@ struct FKShapeElem
 	/** Set whether this shape will contribute to the mass of the body */
 	ENGINE_API void SetContributeToMass(bool bInContributeToMass) { bContributeToMass = bInContributeToMass; }
 
+#if WITH_CHAOS
+	/** Set whether this shape should be considered for query or sim collision */
+	ENGINE_API void SetCollisionEnabled(ECollisionEnabled::Type InCollisionEnabled) { CollisionEnabled = InCollisionEnabled; }
+#endif
+
+	/** Get whether this shape should be considered for query or sim collision */
+	ENGINE_API ECollisionEnabled::Type GetCollisionEnabled() const { return CollisionEnabled; }
+
 	/** Offset used when generating contact points. This allows you to smooth out
 		the Minkowski sum by radius R. Useful for making objects slide smoothly
 		on top of irregularities  */
@@ -93,9 +104,11 @@ protected:
 	/** Helper function to safely clone instances of this shape element */
 	void CloneElem(const FKShapeElem& Other)
 	{
+		RestOffset = Other.RestOffset;
 		ShapeType = Other.ShapeType;
 		Name = Other.Name;
 		bContributeToMass = Other.bContributeToMass;
+		CollisionEnabled = Other.CollisionEnabled;
 	}
 
 private:
@@ -111,5 +124,10 @@ private:
 	UPROPERTY(Category=Shape, EditAnywhere)
 	uint8 bContributeToMass : 1;
 
+	/** Course per-primitive collision filtering. This allows for individual primitives to
+		be toggled in and out of sim and query collision without changing filtering details. */
+	UPROPERTY(Category=Shape, EditAnywhere)
+	TEnumAsByte<ECollisionEnabled::Type> CollisionEnabled;
+
 	FUserData UserData;
 };
diff --git a/Engine/Source/Runtime/Engine/Private/PhysicsEngine/BodyInstance.cpp b/Engine/Source/Runtime/Engine/Private/PhysicsEngine/BodyInstance.cpp
--- a/Engine/Source/Runtime/Engine/Private/PhysicsEngine/BodyInstance.cpp
+++ b/Engine/Source/Runtime/Engine/Private/PhysicsEngine/BodyInstance.cpp
@@ -908,7 +908,7 @@ void FBodyInstance::UpdatePhysicsFilterData()
 
 #if WITH_CHAOS
 				// If this shape shouldn't collide in the sim we disable it here until we have more support.
-				Shape.Shape->SetDisable(!CollisionEnabledHasPhysics(GetCollisionEnabled()));
+				Shape.Shape->SetSimEnabled(CollisionEnabledHasPhysics(GetCollisionEnabled()));
 #endif
 
 				// Apply new collision settings to this shape
@@ -1238,7 +1238,7 @@ void FInitBodiesHelperBase::InitBodies()
 							{
 								for(int32 ShapeIndex = 0; ShapeIndex < NumShapes; ++ShapeIndex)
 								{
-									ActorHandle->SetShapeCollisionDisable(ShapeIndex, true);
+									ActorHandle->SetShapeSimCollisionEnabled(ShapeIndex, false);
 								}
 							}
 							if (BI->BodySetup.IsValid())
diff --git a/Engine/Source/Runtime/Engine/Private/PhysicsEngine/CollisionQueryFilterCallback.cpp b/Engine/Source/Runtime/Engine/Private/PhysicsEngine/CollisionQueryFilterCallback.cpp
--- a/Engine/Source/Runtime/Engine/Private/PhysicsEngine/CollisionQueryFilterCallback.cpp
+++ b/Engine/Source/Runtime/Engine/Private/PhysicsEngine/CollisionQueryFilterCallback.cpp
@@ -111,6 +111,12 @@ ECollisionQueryHitType FCollisionQueryFilterCallback::PreFilterImp(const FCollis
 ECollisionQueryHitType FCollisionQueryFilterCallback::PreFilterImp(const FCollisionFilterData& FilterData, const Chaos::FPerShapeData& Shape, const Chaos::TGeometryParticle<float,3>& Actor)
 {
 	//SCOPE_CYCLE_COUNTER(STAT_Collision_PreFilter);
+
+	if (!Shape.GetQueryEnabled())
+	{
+		return ECollisionQueryHitType::None;
+	}
+
 	FCollisionFilterData ShapeFilter = ChaosInterface::GetQueryFilterData(Shape);
 
 	// We usually don't have ignore components so we try to avoid the virtual getSimulationFilterData() call below. 'word2' of shape sim filter data is componentID.
diff --git a/Engine/Source/Runtime/Engine/Private/PhysicsEngine/Experimental/ChaosInterfaceUtils.cpp b/Engine/Source/Runtime/Engine/Private/PhysicsEngine/Experimental/ChaosInterfaceUtils.cpp
--- a/Engine/Source/Runtime/Engine/Private/PhysicsEngine/Experimental/ChaosInterfaceUtils.cpp
+++ b/Engine/Source/Runtime/Engine/Private/PhysicsEngine/Experimental/ChaosInterfaceUtils.cpp
@@ -177,9 +177,9 @@ namespace ChaosInterface
 
 		ensure(bMakeComplexGeometry || bMakeSimpleGeometry);
 
-		auto NewShapeHelper = [&InParams, &CollisionTraceType](Chaos::TSerializablePtr<Chaos::FImplicitObject> InGeom, int32 ShapeIdx, void* UserData, bool bComplexShape = false)
+		auto NewShapeHelper = [&InParams, &CollisionTraceType](Chaos::TSerializablePtr<Chaos::FImplicitObject> InGeom, int32 ShapeIdx, void* UserData, ECollisionEnabled::Type ShapeCollisionEnabled, bool bComplexShape = false)
 		{
-			auto NewShape = Chaos::FPerShapeData::CreatePerShapeData(ShapeIdx);
+			TUniquePtr<Chaos::FPerShapeData> NewShape = Chaos::FPerShapeData::CreatePerShapeData(ShapeIdx);
 			NewShape->SetGeometry(InGeom);
 			NewShape->SetQueryData(bComplexShape ? InParams.CollisionData.CollisionFilterData.QueryComplexFilter : InParams.CollisionData.CollisionFilterData.QuerySimpleFilter);
 			NewShape->SetSimData(InParams.CollisionData.CollisionFilterData.SimFilter);
@@ -187,6 +187,10 @@ namespace ChaosInterface
 			NewShape->UpdateShapeBounds(InParams.WorldTransform);
 			NewShape->SetUserData(UserData);
 			NewShape->SetSimulate(bComplexShape ? InParams.CollisionData.CollisionFlags.bEnableSimCollisionComplex : InParams.CollisionData.CollisionFlags.bEnableSimCollisionSimple);
+#if WITH_CHAOS
+			NewShape->SetSimEnabled(ShapeCollisionEnabled == ECollisionEnabled::QueryAndPhysics || ECollisionEnabled::PhysicsOnly);
+			NewShape->SetQueryEnabled(ShapeCollisionEnabled == ECollisionEnabled::QueryAndPhysics || ECollisionEnabled::QueryOnly);
+#endif
 			return NewShape;
 		};
 
@@ -198,7 +202,7 @@ namespace ChaosInterface
 				const FKSphereElem ScaledSphereElem = SphereElem.GetFinalScaled(Scale, InParams.LocalTransform);
 				const float UseRadius = FMath::Max(ScaledSphereElem.Radius, KINDA_SMALL_NUMBER);
 				auto ImplicitSphere = MakeUnique<Chaos::TSphere<float, 3>>(ScaledSphereElem.Center, UseRadius);
-				auto NewShape = NewShapeHelper(MakeSerializable(ImplicitSphere), Shapes.Num(), (void*)SphereElem.GetUserData());
+				TUniquePtr<Chaos::FPerShapeData> NewShape = NewShapeHelper(MakeSerializable(ImplicitSphere), Shapes.Num(), (void*)SphereElem.GetUserData(), SphereElem.GetCollisionEnabled());
 				Shapes.Emplace(MoveTemp(NewShape));
 				Geoms.Add(MoveTemp(ImplicitSphere));
 			}
@@ -226,7 +230,7 @@ namespace ChaosInterface
 					Implicit = MakeUnique<Chaos::TBox<float, 3>>(BoxTransform.GetTranslation() - HalfExtents, BoxTransform.GetTranslation() + HalfExtents);
 				}
 
-				auto NewShape = NewShapeHelper(MakeSerializable(Implicit),Shapes.Num(), (void*)BoxElem.GetUserData());
+				TUniquePtr<Chaos::FPerShapeData> NewShape = NewShapeHelper(MakeSerializable(Implicit),Shapes.Num(), (void*)BoxElem.GetUserData(), BoxElem.GetCollisionEnabled());
 				Shapes.Emplace(MoveTemp(NewShape));
 				Geoms.Add(MoveTemp(Implicit));
 			}
@@ -241,7 +245,7 @@ namespace ChaosInterface
 				{
 					//not a capsule just use a sphere
 					auto ImplicitSphere = MakeUnique<Chaos::TSphere<float, 3>>(ScaledSphylElem.Center, Radius);
-					auto NewShape = NewShapeHelper(MakeSerializable(ImplicitSphere),Shapes.Num(), (void*)UnscaledSphyl.GetUserData());
+					TUniquePtr<Chaos::FPerShapeData> NewShape = NewShapeHelper(MakeSerializable(ImplicitSphere),Shapes.Num(), (void*)UnscaledSphyl.GetUserData(), UnscaledSphyl.GetCollisionEnabled());
 					Shapes.Emplace(MoveTemp(NewShape));
 					Geoms.Add(MoveTemp(ImplicitSphere));
 
@@ -251,7 +255,7 @@ namespace ChaosInterface
 					Chaos::TVector<float, 3> HalfExtents = ScaledSphylElem.Rotation.RotateVector(Chaos::TVector<float, 3>(0, 0, HalfHeight));
 
 					auto ImplicitCapsule = MakeUnique<Chaos::TCapsule<float>>(ScaledSphylElem.Center - HalfExtents, ScaledSphylElem.Center + HalfExtents, Radius);
-					auto NewShape = NewShapeHelper(MakeSerializable(ImplicitCapsule),Shapes.Num(), (void*)UnscaledSphyl.GetUserData());
+					TUniquePtr<Chaos::FPerShapeData> NewShape = NewShapeHelper(MakeSerializable(ImplicitCapsule),Shapes.Num(), (void*)UnscaledSphyl.GetUserData(), UnscaledSphyl.GetCollisionEnabled());
 					Shapes.Emplace(MoveTemp(NewShape));
 					Geoms.Add(MoveTemp(ImplicitCapsule));
 				}
@@ -311,7 +315,7 @@ namespace ChaosInterface
 							Implicit = TUniquePtr<Chaos::FImplicitObject>(new Chaos::TImplicitObjectScaled<Chaos::FConvex>(ConvexImplicit, Scale));
 						}
 
-						auto NewShape = NewShapeHelper(MakeSerializable(Implicit),Shapes.Num(), (void*)CollisionBody.GetUserData());
+						TUniquePtr<Chaos::FPerShapeData> NewShape = NewShapeHelper(MakeSerializable(Implicit),Shapes.Num(), (void*)CollisionBody.GetUserData(), CollisionBody.GetCollisionEnabled());
 						Shapes.Emplace(MoveTemp(NewShape));
 						Geoms.Add(MoveTemp(Implicit));
 					}
@@ -335,7 +339,7 @@ namespace ChaosInterface
 
 				ChaosTriMesh->SetCullsBackFaceRaycast(!InParams.bDoubleSided);
 
-				auto NewShape = NewShapeHelper(MakeSerializable(Implicit),Shapes.Num(), nullptr, true);
+				TUniquePtr<Chaos::FPerShapeData> NewShape = NewShapeHelper(MakeSerializable(Implicit),Shapes.Num(), nullptr, ECollisionEnabled::QueryAndPhysics, true);
 				Shapes.Emplace(MoveTemp(NewShape));
 				Geoms.Add(MoveTemp(Implicit));
 			}
@@ -345,7 +349,7 @@ namespace ChaosInterface
 		for (const auto& PhysXMesh : InParams.TriMeshes)
 		{
 			auto Implicit = ConvertPhysXMeshToLevelset(PhysXMesh, Scale);
-			auto NewShape = NewShapeHelper(MakeSerializable(Implicit),Shapes.Num(), nullptr, true);
+			auto NewShape = NewShapeHelper(MakeSerializable(Implicit), Shapes.Num(), nullptr, ECollisionEnabled::QueryAndPhysics, true);
 			Shapes.Emplace(MoveTemp(NewShape));
 			Geoms.Add(MoveTemp(Implicit));
 
diff --git a/Engine/Source/Runtime/Experimental/Chaos/Private/Chaos/CollisionResolution.cpp b/Engine/Source/Runtime/Experimental/Chaos/Private/Chaos/CollisionResolution.cpp
--- a/Engine/Source/Runtime/Experimental/Chaos/Private/Chaos/CollisionResolution.cpp
+++ b/Engine/Source/Runtime/Experimental/Chaos/Private/Chaos/CollisionResolution.cpp
@@ -2456,8 +2456,8 @@ namespace Chaos
 			//
 			// Disabled shapes do not collide
 			//
-			if (Shape0 && (Shape0->GetDisable() || !IsValid(Shape0->GetSimData()) ) ) return false;
-			if (Shape1 && (Shape1->GetDisable ()|| !IsValid(Shape1->GetSimData()) ) ) return false;
+			if (Shape0 && (!Shape0->GetSimEnabled() || !IsValid(Shape0->GetSimData()) ) ) return false;
+			if (Shape1 && (!Shape1->GetSimEnabled() || !IsValid(Shape1->GetSimData()) ) ) return false;
 
 			//
 			// Triangle Mesh geometry is only used if the shape specifies UseComplexAsSimple
diff --git a/Engine/Source/Runtime/Experimental/Chaos/Public/Chaos/GeometryParticles.h b/Engine/Source/Runtime/Experimental/Chaos/Public/Chaos/GeometryParticles.h
--- a/Engine/Source/Runtime/Experimental/Chaos/Public/Chaos/GeometryParticles.h
+++ b/Engine/Source/Runtime/Experimental/Chaos/Public/Chaos/GeometryParticles.h
@@ -150,11 +150,26 @@ namespace Chaos
 			});
 		}
 
-		bool GetDisable() const { return CollisionData.Read().bDisable; }
+		bool GetQueryEnabled() const { return CollisionData.Read().bQueryCollision; }
+		void SetQueryEnabled(const bool bEnable)
+		{
+			CollisionData.Modify(true, DirtyFlags, Proxy, ShapeIdx, [bEnable](FCollisionData& Data){ Data.bQueryCollision = bEnable; });
+		}
+
+		bool GetSimEnabled() const { return CollisionData.Read().bSimCollision; }
+		void SetSimEnabled(const bool bEnable)
+		{
+			CollisionData.Modify(true, DirtyFlags, Proxy, ShapeIdx, [bEnable](FCollisionData& Data){ Data.bSimCollision = bEnable; });
+		}
+
+		/*
+		// TODO: Deprecate GetDisable() and SetDisable()!
+		bool GetDisable() const { return !CollisionData.Read().bSimCollision; }
 		void SetDisable(const bool bDisable)
 		{
-			CollisionData.Modify(true, DirtyFlags, Proxy, ShapeIdx, [bDisable](FCollisionData& Data){ Data.bDisable = bDisable; });
+			CollisionData.Modify(true, DirtyFlags, Proxy, ShapeIdx, [bDisable](FCollisionData& Data){ Data.bSimCollision = !bDisable; });
 		}
+		*/
 
 		bool GetSimulate() const { return CollisionData.Read().bSimulate; }
 		void SetSimulate(const bool bSimulate)
diff --git a/Engine/Source/Runtime/Experimental/Chaos/Public/Chaos/ParticleDirtyFlags.h b/Engine/Source/Runtime/Experimental/Chaos/Public/Chaos/ParticleDirtyFlags.h
--- a/Engine/Source/Runtime/Experimental/Chaos/Public/Chaos/ParticleDirtyFlags.h
+++ b/Engine/Source/Runtime/Experimental/Chaos/Public/Chaos/ParticleDirtyFlags.h
@@ -391,13 +391,15 @@ struct FCollisionData
 	FCollisionFilterData SimData;
 	void* UserData;
 	EChaosCollisionTraceFlag CollisionTraceType;
-	uint8 bDisable : 1;
+	uint8 bSimCollision : 1;
+	uint8 bQueryCollision : 1;
 	uint8 bSimulate : 1;
 
 	FCollisionData()
 	: UserData(nullptr)
 	, CollisionTraceType(EChaosCollisionTraceFlag::Chaos_CTF_UseDefault)
-	, bDisable(false)
+	, bSimCollision(true)
+	, bQueryCollision(true)
 	, bSimulate(true)
 	{
 	}
@@ -410,11 +412,20 @@ struct FCollisionData
 		Ar << QueryData;
 		Ar << SimData;
 
-		if(Ar.CustomVer(FExternalPhysicsCustomObjectVersion::GUID) >= FExternalPhysicsCustomObjectVersion::AddShapeCollisionDisable)
+		if (Ar.CustomVer(FExternalPhysicsCustomObjectVersion::GUID) >= FExternalPhysicsCustomObjectVersion::AddShapeSimAndQueryCollisionEnabled)
 		{
-			bool Disable = bDisable;
+			int8 EnableSim = bSimCollision;
+			int8 EnableQuery = bQueryCollision;
+			Ar << EnableSim;
+			Ar << EnableQuery;
+			bSimCollision = EnableSim;
+			bQueryCollision = EnableQuery;
+		}
+		else if (Ar.CustomVer(FExternalPhysicsCustomObjectVersion::GUID) >= FExternalPhysicsCustomObjectVersion::AddShapeCollisionDisable)
+		{
+			bool Disable = !bSimCollision;
 			Ar << Disable;
-			bDisable = Disable;
+			bSimCollision = !Disable;
 		}
 
 		if(Ar.CustomVer(FExternalPhysicsCustomObjectVersion::GUID) >= FExternalPhysicsCustomObjectVersion::SerializePerShapeDataSimulateFlag)
diff --git a/Engine/Source/Runtime/Experimental/ChaosSolvers/Private/PhysicsProxy/GeometryCollectionPhysicsProxy.cpp b/Engine/Source/Runtime/Experimental/ChaosSolvers/Private/PhysicsProxy/GeometryCollectionPhysicsProxy.cpp
--- a/Engine/Source/Runtime/Experimental/ChaosSolvers/Private/PhysicsProxy/GeometryCollectionPhysicsProxy.cpp
+++ b/Engine/Source/Runtime/Experimental/ChaosSolvers/Private/PhysicsProxy/GeometryCollectionPhysicsProxy.cpp
@@ -326,7 +326,7 @@ void PopulateSimulatedParticle(
 		FilterData.Word3 = 0xFFFF; // collision candidate channels
 		for (const TUniquePtr<Chaos::FPerShapeData>& Shape : Handle->ShapesArray())
 		{
-			Shape->SetDisable(false);
+			Shape->SetSimEnabled(true);
 			Shape->SetCollisionTraceType(Chaos::EChaosCollisionTraceFlag::Chaos_CTF_UseDefault);
 			//Shape->CollisionTraceType = Chaos::EChaosCollisionTraceFlag::Chaos_CTF_UseSimpleAndComplex;
 			Shape->SetSimData(FilterData);
