Chaos
- clean up includes
- ImmediatePhysicsSimulation_Chaos.h

diff --git a/Engine/Source/Runtime/Engine/Private/PhysicsEngine/ImmediatePhysics/ImmediatePhysicsChaos/ImmediatePhysicsSimulation_Chaos.cpp b/Engine/Source/Runtime/Engine/Private/PhysicsEngine/ImmediatePhysics/ImmediatePhysicsChaos/ImmediatePhysicsSimulation_Chaos.cpp
--- a/Engine/Source/Runtime/Engine/Private/PhysicsEngine/ImmediatePhysics/ImmediatePhysicsChaos/ImmediatePhysicsSimulation_Chaos.cpp
+++ b/Engine/Source/Runtime/Engine/Private/PhysicsEngine/ImmediatePhysics/ImmediatePhysicsChaos/ImmediatePhysicsSimulation_Chaos.cpp
@@ -5,19 +5,27 @@
 #include "Physics/ImmediatePhysics/ImmediatePhysicsChaos/ImmediatePhysicsJointHandle_Chaos.h"
 #include "Physics/ImmediatePhysics/ImmediatePhysicsStats.h"
 
+#include "PhysicalMaterials/PhysicalMaterial.h"
+#include "PhysicsEngine/BodyInstance.h"
+#include "PhysicsEngine/ConstraintInstance.h"
+
+#define CHAOS_INCLUDE_LEVEL_1
 #include "Chaos/ChaosDebugDraw.h"
+#include "Chaos/Collision/NarrowPhase.h"
+#include "Chaos/Collision/ParticlePairBroadPhase.h"
+#include "Chaos/Collision/ParticlePairCollisionDetector.h"
 #include "Chaos/DebugDrawQueue.h"
+#include "Chaos/Evolution/PBDMinEvolution.h"
 #include "Chaos/Joint/ChaosJointLog.h"
 #include "Chaos/ParticleHandle.h"
 #include "Chaos/PBDCollisionConstraints.h"
+#include "Chaos/PBDConstraintRule.h"
 #include "Chaos/PBDJointConstraints.h"
 #include "Chaos/PBDRigidParticles.h"
 #include "Chaos/PBDRigidsEvolutionGBF.h"
 #include "ChaosLog.h"
+#undef CHAOS_INCLUDE_LEVEL_1
 
-#include "PhysicalMaterials/PhysicalMaterial.h"
-#include "PhysicsEngine/BodyInstance.h"
-#include "PhysicsEngine/ConstraintInstance.h"
 
 //PRAGMA_DISABLE_OPTIMIZATION
 
@@ -152,109 +160,164 @@ FAutoConsoleVariableRef CVarChaosImmPhysDebugDrawJointFeatures(TEXT("p.Chaos.Imm
 
 namespace ImmediatePhysics_Chaos
 {
+	struct FSimulation::FImplementation
+	{
+	public:
+		using FCollisionConstraints = Chaos::FPBDCollisionConstraints;
+		using FCollisionDetector = Chaos::FParticlePairCollisionDetector;
+		using FParticleHandle = Chaos::TGeometryParticleHandle<FReal, Dimensions>;
+		using FParticlePair = Chaos::TVector<Chaos::TGeometryParticleHandle<Chaos::FReal, 3>*, 2>;
+		using FRigidParticleSOAs = Chaos::TPBDRigidsSOAs<FReal, 3>;
+
+		FImplementation()
+			: Particles()
+			, Joints()
+			, Collisions(Particles, CollidedParticles, ParticleMaterials, 0, 0, ChaosImmediate_Collision_CullDistance, ChaosImmediate_Collision_ShapePadding)
+			, BroadPhase(ActivePotentiallyCollidingPairs, ChaosImmediate_Collision_CullDistance)
+			, NarrowPhase()
+			, CollisionDetector(BroadPhase, NarrowPhase, Collisions)
+			, JointsRule(0, Joints)
+			, CollisionsRule(1, Collisions)
+			, Evolution(Particles, ParticlePrevXs, ParticlePrevRs, CollisionDetector, ChaosImmediate_Evolution_BoundsExtension)
+			, NumActiveDynamicActorHandles(0)
+			, SimulationSpaceTransform(FTransform::Identity)
+			, RollingAverageStepTime(ChaosImmediate_Evolution_InitialStepTime)
+			, NumRollingAverageStepTimes(1)
+			, MaxNumRollingAverageStepTimes(ChaosImmediate_Evolution_DeltaTimeCount)
+			, bActorsDirty(false)
+		{
+			Particles.GetParticleHandles().AddArray(&ParticleMaterials);
+			Particles.GetParticleHandles().AddArray(&PerParticleMaterials);
+			Particles.GetParticleHandles().AddArray(&ParticlePrevXs);
+			Particles.GetParticleHandles().AddArray(&ParticlePrevRs);
+
+			Evolution.AddConstraintRule(&CollisionsRule);
+			Evolution.AddConstraintRule(&JointsRule);
+		}
+
+		~FImplementation()
+		{
+		}
+
+		// @todo(ccaulfield): Look into these...
+		TArray<FParticlePair> ActivePotentiallyCollidingPairs;
+		Chaos::TArrayCollectionArray<bool> CollidedParticles;
+		Chaos::TArrayCollectionArray<Chaos::TSerializablePtr<Chaos::FChaosPhysicsMaterial>> ParticleMaterials;
+		Chaos::TArrayCollectionArray<TUniquePtr<Chaos::FChaosPhysicsMaterial>> PerParticleMaterials;
+		Chaos::TArrayCollectionArray<Chaos::FVec3> ParticlePrevXs;
+		Chaos::TArrayCollectionArray<Chaos::FRotation3> ParticlePrevRs;
+
+		FRigidParticleSOAs Particles;
+		Chaos::FPBDJointConstraints Joints;
+		FCollisionConstraints Collisions;
+		Chaos::FParticlePairBroadPhase BroadPhase;
+		Chaos::FNarrowPhase NarrowPhase;
+		FCollisionDetector CollisionDetector;
+		Chaos::TSimpleConstraintRule<Chaos::FPBDJointConstraints> JointsRule;
+		Chaos::TSimpleConstraintRule<FCollisionConstraints> CollisionsRule;
+		Chaos::FPBDMinEvolution Evolution;
+
+		/** Mapping from entity index to handle */
+		// @todo(ccaulfield): we now have handles pointing to handles which is inefficient - we can do better than this, but don't want to change API yet
+		TArray<FActorHandle*> ActorHandles;
+		int32 NumActiveDynamicActorHandles;
+
+		/** Mapping from constraint index to handle */
+		TArray<FJointHandle*> JointHandles;
+
+		/** Slow to access. */
+		// @todo(ccaulfield): Optimize
+		TMap<const FParticleHandle*, TSet<const FParticleHandle*>> IgnoreCollisionParticlePairTable;
+
+		TArray<FParticlePair> PotentiallyCollidingPairs;
+
+		FTransform SimulationSpaceTransform;
+
+		FReal FixedStepTime;
+		FReal RollingAverageStepTime;
+		int32 NumRollingAverageStepTimes;
+		int32 MaxNumRollingAverageStepTimes;
+
+		bool bActorsDirty;
+	};
+
 	//
 	//
 	//
 
 	FSimulation::FSimulation()
-		: ActivePotentiallyCollidingPairs()
-		, CollidedParticles()
-		, ParticleMaterials()
-		, PerParticleMaterials()
-		, ParticlePrevXs()
-		, ParticlePrevRs()
-		, Particles()
-		, Joints()
-		, Collisions(Particles, CollidedParticles, ParticleMaterials, 0, 0, ChaosImmediate_Collision_CullDistance, ChaosImmediate_Collision_ShapePadding)
-		, BroadPhase(ActivePotentiallyCollidingPairs, ChaosImmediate_Collision_CullDistance)
-		, NarrowPhase()
-		, CollisionDetector(BroadPhase, NarrowPhase, Collisions)
-		, JointsRule(0, Joints)
-		, CollisionsRule(1, Collisions)
-		, Evolution(Particles, ParticlePrevXs, ParticlePrevRs, CollisionDetector, ChaosImmediate_Evolution_BoundsExtension)
-		, NumActiveDynamicActorHandles(0)
-		, SimulationSpaceTransform(FTransform::Identity)
-		, RollingAverageStepTime(ChaosImmediate_Evolution_InitialStepTime)
-		, NumRollingAverageStepTimes(1)
-		, MaxNumRollingAverageStepTimes(ChaosImmediate_Evolution_DeltaTimeCount)
-		, bActorsDirty(false)
 	{
 		using namespace Chaos;
 
-		Particles.GetParticleHandles().AddArray(&ParticleMaterials);
-		Particles.GetParticleHandles().AddArray(&PerParticleMaterials);
-		Particles.GetParticleHandles().AddArray(&ParticlePrevXs);
-		Particles.GetParticleHandles().AddArray(&ParticlePrevRs);
-
-		Evolution.AddConstraintRule(&CollisionsRule);
-		Evolution.AddConstraintRule(&JointsRule);
+		Implementation = MakeUnique<FImplementation>();
 
 		// RBAN collision customization
-		Collisions.DisableHandles();
-		Collisions.SetApplyType(ECollisionApplyType::Position);
-		NarrowPhase.GetContext().bFilteringEnabled = false;
-		NarrowPhase.GetContext().bDeferUpdate = true;
-		NarrowPhase.GetContext().bAllowManifolds = false;
+		Implementation->Collisions.DisableHandles();
+		Implementation->Collisions.SetApplyType(ECollisionApplyType::Position);
+		Implementation->NarrowPhase.GetContext().bFilteringEnabled = false;
+		Implementation->NarrowPhase.GetContext().bDeferUpdate = true;
+		Implementation->NarrowPhase.GetContext().bAllowManifolds = false;
+
 
 #if CHAOS_DEBUG_DRAW
-		Evolution.SetPostIntegrateCallback(
+		Implementation->Evolution.SetPostIntegrateCallback(
 			[this]()
 			{
 				DebugDrawDynamicParticles(3, 4, FColor(32, 32, 0));
 				DebugDrawConstraints(3, 4, 0.3f);
 			});
-		Evolution.SetPostApplyCallback(
+		Implementation->Evolution.SetPostApplyCallback(
 			[this]()
 			{
 				DebugDrawDynamicParticles(3, 3, FColor(128, 128, 0));
 				DebugDrawConstraints(3, 3, 0.6f);
 			});
-		Evolution.SetPostApplyPushOutCallback(
+		Implementation->Evolution.SetPostApplyPushOutCallback(
 			[this]()
 			{
 				DebugDrawDynamicParticles(3, 3, FColor(255, 255, 0));
 				DebugDrawConstraints(3, 3, 1.0f);
 			});
-		Collisions.SetPostApplyCallback(
+		Implementation->Collisions.SetPostApplyCallback(
 			[this](const float Dt, const TArray<FPBDCollisionConstraintHandle*>& InConstraintHandles)
 			{
 				if (ChaosImmediate_DebugDrawCollisions == 4)
 				{
-					DebugDraw::DrawCollisions(SimulationSpaceTransform, Collisions, 0.3f);
+					DebugDraw::DrawCollisions(Implementation->SimulationSpaceTransform, Implementation->Collisions, 0.3f);
 				}
 				DebugDrawDynamicParticles(4, 4, FColor(128, 128, 0));
 			});
-		Collisions.SetPostApplyPushOutCallback(
+		Implementation->Collisions.SetPostApplyPushOutCallback(
 			[this](const float Dt, const TArray<FPBDCollisionConstraintHandle*>& InConstraintHandles, bool bRequiresAnotherIteration)
 			{
 				if (ChaosImmediate_DebugDrawCollisions == 4)
 				{
-					DebugDraw::DrawCollisions(SimulationSpaceTransform, Collisions, 0.6f);
+					DebugDraw::DrawCollisions(Implementation->SimulationSpaceTransform, Implementation->Collisions, 0.6f);
 				}
 			});
-		Joints.SetPreApplyCallback(
+		Implementation->Joints.SetPreApplyCallback(
 			[this](const float Dt, const TArray<FPBDJointConstraintHandle*>& InConstraintHandles)
 			{
 				if (ChaosImmediate_DebugDrawJoints == 4)
 				{
-					DebugDraw::DrawJointConstraints(SimulationSpaceTransform, InConstraintHandles, 0.3f);
+					DebugDraw::DrawJointConstraints(Implementation->SimulationSpaceTransform, InConstraintHandles, 0.3f);
 				}
 			});
-		Joints.SetPostApplyCallback(
+		Implementation->Joints.SetPostApplyCallback(
 			[this](const float Dt, const TArray<FPBDJointConstraintHandle*>& InConstraintHandles)
 			{
 				if (ChaosImmediate_DebugDrawJoints == 4)
 				{
-					DebugDraw::DrawJointConstraints(SimulationSpaceTransform, InConstraintHandles, 0.6f);
+					DebugDraw::DrawJointConstraints(Implementation->SimulationSpaceTransform, InConstraintHandles, 0.6f);
 				}
 				DebugDrawDynamicParticles(4, 4, FColor(0, 128, 0));
 			});
-		Joints.SetPostProjectCallback(
+		Implementation->Joints.SetPostProjectCallback(
 			[this](const float Dt, const TArray<FPBDJointConstraintHandle*>& InConstraintHandles)
 			{
 				if (ChaosImmediate_DebugDrawJoints == 4)
 				{
-					DebugDraw::DrawJointConstraints(SimulationSpaceTransform, InConstraintHandles, 0.6f);
+					DebugDraw::DrawJointConstraints(Implementation->SimulationSpaceTransform, InConstraintHandles, 0.6f);
 				}
 			});
 #endif
@@ -264,19 +327,33 @@ namespace ImmediatePhysics_Chaos
 	{
 		using namespace Chaos;
 
-		for (FActorHandle* ActorHandle : ActorHandles)
+		for (FActorHandle* ActorHandle : Implementation->ActorHandles)
 		{
 			delete ActorHandle;
 		}
-		ActorHandles.Empty();
+		Implementation->ActorHandles.Empty();
 
-		for (FJointHandle* JointHandle : JointHandles)
+		for (FJointHandle* JointHandle : Implementation->JointHandles)
 		{
 			delete JointHandle;
 		}
-		JointHandles.Empty();
+		Implementation->JointHandles.Empty();
+	}
+
+	int32 FSimulation::NumActors() const
+	{
+		return Implementation->ActorHandles.Num();
 	}
 
+	FActorHandle* FSimulation::GetActorHandle(int32 ActorHandleIndex)
+	{
+		return Implementation->ActorHandles[ActorHandleIndex];
+	}
+	
+	const FActorHandle* FSimulation::GetActorHandle(int32 ActorHandleIndex) const
+	{
+		return Implementation->ActorHandles[ActorHandleIndex];
+	}
 
 	FActorHandle* FSimulation::CreateStaticActor(FBodyInstance* BodyInstance)
 	{
@@ -300,8 +377,8 @@ namespace ImmediatePhysics_Chaos
 
 		using namespace Chaos;
 
-		FActorHandle* ActorHandle = new FActorHandle(Particles, ParticlePrevXs, ParticlePrevRs, ActorType, BodyInstance, Transform);
-		int ActorIndex = ActorHandles.Add(ActorHandle);
+		FActorHandle* ActorHandle = new FActorHandle(Implementation->Particles, Implementation->ParticlePrevXs, Implementation->ParticlePrevRs, ActorType, BodyInstance, Transform);
+		int ActorIndex = Implementation->ActorHandles.Add(ActorHandle);
 
 		TUniquePtr<FChaosPhysicsMaterial> Material = MakeUnique<FChaosPhysicsMaterial>();
 		if (BodyInstance != nullptr)
@@ -314,13 +391,13 @@ namespace ImmediatePhysics_Chaos
 			}
 		}
 
-		ParticleMaterials.Add(MakeSerializable(Material));
-		PerParticleMaterials.Add(MoveTemp(Material));
-		CollidedParticles.Add(false);
-		ParticlePrevXs.Add(ActorHandle->GetParticle()->X());
-		ParticlePrevRs.Add(ActorHandle->GetParticle()->R());
+		Implementation->ParticleMaterials.Add(MakeSerializable(Material));
+		Implementation->PerParticleMaterials.Add(MoveTemp(Material));
+		Implementation->CollidedParticles.Add(false);
+		Implementation->ParticlePrevXs.Add(ActorHandle->GetParticle()->X());
+		Implementation->ParticlePrevRs.Add(ActorHandle->GetParticle()->R());
 
-		bActorsDirty = true;
+		Implementation->bActorsDirty = true;
 
 		return ActorHandle;
 	}
@@ -330,43 +407,43 @@ namespace ImmediatePhysics_Chaos
 		// @todo(ccaulfield): FActorHandle could remember its index to optimize this
 		// @todo(ccaulfield): Remove colliding particle pairs
 
-		int32 Index = ActorHandles.Remove(ActorHandle);
+		int32 Index = Implementation->ActorHandles.Remove(ActorHandle);
 		delete ActorHandle;
 
-		ParticleMaterials.RemoveAt(Index, 1);
-		PerParticleMaterials.RemoveAt(Index, 1);
-		CollidedParticles.RemoveAt(Index, 1);
-		ParticlePrevXs.RemoveAt(Index, 1);
-		ParticlePrevRs.RemoveAt(Index, 1);
+		Implementation->ParticleMaterials.RemoveAt(Index, 1);
+		Implementation->PerParticleMaterials.RemoveAt(Index, 1);
+		Implementation->CollidedParticles.RemoveAt(Index, 1);
+		Implementation->ParticlePrevXs.RemoveAt(Index, 1);
+		Implementation->ParticlePrevRs.RemoveAt(Index, 1);
 
-		bActorsDirty = true;
+		Implementation->bActorsDirty = true;
 	}
 
 	FJointHandle* FSimulation::CreateJoint(FConstraintInstance* ConstraintInstance, FActorHandle* Body1, FActorHandle* Body2)
 	{
-		FJointHandle* JointHandle = new FJointHandle(&Joints, ConstraintInstance, Body1, Body2);
-		JointHandles.Add(JointHandle);
+		FJointHandle* JointHandle = new FJointHandle(&Implementation->Joints, ConstraintInstance, Body1, Body2);
+		Implementation->JointHandles.Add(JointHandle);
 		return JointHandle;
 	}
 
 	void FSimulation::DestroyJoint(FJointHandle* JointHandle)
 	{
 		// @todo(ccaulfield): FJointHandle could remember its index to optimize this
-		JointHandles.Remove(JointHandle);
+		Implementation->JointHandles.Remove(JointHandle);
 		delete JointHandle;
 	}
 
 	void FSimulation::SetNumActiveBodies(int32 InNumActiveActorHandles)
 	{
-		if (InNumActiveActorHandles == NumActiveDynamicActorHandles)
+		if (InNumActiveActorHandles == Implementation->NumActiveDynamicActorHandles)
 		{
 			return;
 		}
 
 		// @todo(ccaulfield): can be optimized, but I think we end up with kinematic at the start and statics at the end of the 
 		// list - maybe split them up or put kinematics at the end (in a way that does not impact particle order).
-		NumActiveDynamicActorHandles = 0;
-		for (int32 ActorHandleIndex = 0; ActorHandleIndex < ActorHandles.Num(); ++ActorHandleIndex)
+		Implementation->NumActiveDynamicActorHandles = 0;
+		for (int32 ActorHandleIndex = 0; ActorHandleIndex < Implementation->ActorHandles.Num(); ++ActorHandleIndex)
 		{
 			FActorHandle* Handle = GetActorHandle(ActorHandleIndex);
 			if (Handle->GetIsKinematic())
@@ -374,10 +451,10 @@ namespace ImmediatePhysics_Chaos
 				continue;
 			}
 
-			if (NumActiveDynamicActorHandles < InNumActiveActorHandles)
+			if (Implementation->NumActiveDynamicActorHandles < InNumActiveActorHandles)
 			{
 				Handle->SetEnabled(true);
-				++NumActiveDynamicActorHandles;
+				++Implementation->NumActiveDynamicActorHandles;
 			}
 			else
 			{
@@ -385,14 +462,14 @@ namespace ImmediatePhysics_Chaos
 			}
 		}
 
-		bActorsDirty = true;
+		Implementation->bActorsDirty = true;
 	}
 
 	void FSimulation::SetIgnoreCollisionPairTable(const TArray<FIgnorePair>& InIgnoreCollisionPairTable)
 	{
 		using namespace Chaos;
 
-		IgnoreCollisionParticlePairTable.Reset();
+		Implementation->IgnoreCollisionParticlePairTable.Reset();
 		for (const FIgnorePair& IgnorePair : InIgnoreCollisionPairTable)
 		{
 			if(!IgnorePair.A || !IgnorePair.B)
@@ -403,35 +480,35 @@ namespace ImmediatePhysics_Chaos
 			const TGeometryParticleHandle<FReal, Dimensions>* ParticleA = IgnorePair.A->GetParticle();
 			const TGeometryParticleHandle<FReal, Dimensions>* ParticleB = IgnorePair.B->GetParticle();
 
-			IgnoreCollisionParticlePairTable.FindOrAdd(ParticleA).Add(ParticleB); 
-			IgnoreCollisionParticlePairTable.FindOrAdd(ParticleB).Add(ParticleA);
+			Implementation->IgnoreCollisionParticlePairTable.FindOrAdd(ParticleA).Add(ParticleB); 
+			Implementation->IgnoreCollisionParticlePairTable.FindOrAdd(ParticleB).Add(ParticleA);
 		}
 
-		PotentiallyCollidingPairs.Empty();
-		int NumActorHandles = ActorHandles.Num();
+		Implementation->PotentiallyCollidingPairs.Empty();
+		int NumActorHandles = Implementation->ActorHandles.Num();
 		for (int ActorHandleIndex0 = 0; ActorHandleIndex0 < NumActorHandles; ++ActorHandleIndex0)
 		{
-			FActorHandle* ActorHandle0 = ActorHandles[ActorHandleIndex0];
+			FActorHandle* ActorHandle0 = Implementation->ActorHandles[ActorHandleIndex0];
 			TGeometryParticleHandle<FReal, Dimensions>* Particle0 = ActorHandle0->GetParticle();
 			//TPBDRigidParticleHandle<FReal, Dimensions>* Rigid0 = Particle0->CastToRigidParticle();
 			//if (Rigid0 && (Rigid0->ObjectState() == EObjectStateType::Dynamic))
 			{
 				for (int ActorHandleIndex1 = ActorHandleIndex0 + 1; ActorHandleIndex1 < NumActorHandles; ++ActorHandleIndex1)
 				{
-					FActorHandle* ActorHandle1 = ActorHandles[ActorHandleIndex1];
+					FActorHandle* ActorHandle1 = Implementation->ActorHandles[ActorHandleIndex1];
 					TGeometryParticleHandle<FReal, Dimensions>* Particle1 = ActorHandle1->GetParticle();
 
-					const TSet<const FParticleHandle*>* Particle0IgnoreSet = IgnoreCollisionParticlePairTable.Find(Particle0);
+					const TSet<const typename FImplementation::FParticleHandle*>* Particle0IgnoreSet = Implementation->IgnoreCollisionParticlePairTable.Find(Particle0);
 					bool bIgnoreActorHandlePair = (Particle0IgnoreSet != nullptr) && Particle0IgnoreSet->Contains(Particle1);
 					if (!bIgnoreActorHandlePair)
 					{
-						PotentiallyCollidingPairs.Emplace(FParticlePair(Particle0, Particle1));
+						Implementation->PotentiallyCollidingPairs.Emplace(typename FImplementation::FParticlePair(Particle0, Particle1));
 					}
 				}
 			}
 		}
 
-		bActorsDirty = true;
+		Implementation->bActorsDirty = true;
 	}
 
 	void FSimulation::SetIgnoreCollisionActors(const TArray<FActorHandle*>& InIgnoreCollisionActors)
@@ -447,38 +524,38 @@ namespace ImmediatePhysics_Chaos
 			}
 		}
 
-		bActorsDirty = true;
+		Implementation->bActorsDirty = true;
 	}
 
 	void FSimulation::UpdateActivePotentiallyCollidingPairs()
 	{
 		using namespace Chaos;
 
-		ActivePotentiallyCollidingPairs.Reset();
-		for (const FParticlePair& ParticlePair : PotentiallyCollidingPairs)
+		Implementation->ActivePotentiallyCollidingPairs.Reset();
+		for (const typename FImplementation::FParticlePair& ParticlePair : Implementation->PotentiallyCollidingPairs)
 		{
 			bool bAnyDisabled = TGenericParticleHandle<FReal, 3>(ParticlePair[0])->Disabled() || TGenericParticleHandle<FReal, 3>(ParticlePair[1])->Disabled();
 			bool bAnyDynamic = TGenericParticleHandle<FReal, 3>(ParticlePair[0])->IsDynamic() || TGenericParticleHandle<FReal, 3>(ParticlePair[1])->IsDynamic();
 			if (bAnyDynamic && !bAnyDisabled)
 			{
-				ActivePotentiallyCollidingPairs.Add(ParticlePair);
+				Implementation->ActivePotentiallyCollidingPairs.Add(ParticlePair);
 			}
 		}
 	}
 
 	void FSimulation::SetSimulationSpaceTransform(const FTransform& Transform)
 	{ 
-		SimulationSpaceTransform = Transform;
-		NarrowPhase.GetContext().SpaceTransform = Transform;	// @todo(chaos): remove when manifolds are fixed or removed
+		Implementation->SimulationSpaceTransform = Transform;
+		Implementation->NarrowPhase.GetContext().SpaceTransform = Transform;	// @todo(chaos): remove when manifolds are fixed or removed
 	}
 
 
 	FReal FSimulation::UpdateStepTime(const FReal DeltaTime, const FReal MaxStepTime)
 	{
-		RollingAverageStepTime = RollingAverageStepTime + (DeltaTime - RollingAverageStepTime) / NumRollingAverageStepTimes;
-		RollingAverageStepTime = FMath::Min(RollingAverageStepTime, MaxStepTime);
-		NumRollingAverageStepTimes = FMath::Min(NumRollingAverageStepTimes + 1, MaxNumRollingAverageStepTimes);
-		return RollingAverageStepTime;
+		Implementation->RollingAverageStepTime = Implementation->RollingAverageStepTime + (DeltaTime - Implementation->RollingAverageStepTime) / Implementation->NumRollingAverageStepTimes;
+		Implementation->RollingAverageStepTime = FMath::Min(Implementation->RollingAverageStepTime, MaxStepTime);
+		Implementation->NumRollingAverageStepTimes = FMath::Min(Implementation->NumRollingAverageStepTimes + 1, Implementation->MaxNumRollingAverageStepTimes);
+		return Implementation->RollingAverageStepTime;
 	}
 
 	void FSimulation::Simulate(float InDeltaTime, float MaxStepTime, int32 MaxSubSteps, const FVector& InGravity)
@@ -498,7 +575,7 @@ namespace ImmediatePhysics_Chaos
 
 		// TEMP: overrides
 		{
-			FPBDJointSolverSettings JointsSettings = Joints.GetSettings();
+			FPBDJointSolverSettings JointsSettings = Implementation->Joints.GetSettings();
 			JointsSettings.SwingTwistAngleTolerance = ChaosImmediate_Joint_SwingTwistAngleTolerance;
 			JointsSettings.PositionTolerance = ChaosImmediate_Joint_PositionTolerance;
 			JointsSettings.AngleTolerance = ChaosImmediate_Joint_AngleTolerance;
@@ -520,26 +597,26 @@ namespace ImmediatePhysics_Chaos
 			JointsSettings.LinearDriveDamping = ChaosImmediate_Joint_LinearDriveDamping;
 			JointsSettings.AngularDriveStiffness = ChaosImmediate_Joint_AngularDriveStiffness;
 			JointsSettings.AngularDriveDamping = ChaosImmediate_Joint_AngularDriveDamping;
-			Joints.SetSettings(JointsSettings);
+			Implementation->Joints.SetSettings(JointsSettings);
 
-			Collisions.SetShapePadding(ChaosImmediate_Collision_ShapePadding);
-			Collisions.SetCollisionsEnabled(ChaosImmediate_Collision_Enabled != 0);
-			CollisionsRule.SetPriority(ChaosImmediate_Collision_Priority);
+			Implementation->Collisions.SetShapePadding(ChaosImmediate_Collision_ShapePadding);
+			Implementation->Collisions.SetCollisionsEnabled(ChaosImmediate_Collision_Enabled != 0);
+			Implementation->CollisionsRule.SetPriority(ChaosImmediate_Collision_Priority);
 
-			Collisions.SetCullDistance(ChaosImmediate_Collision_CullDistance);
-			BroadPhase.SetCullDustance(ChaosImmediate_Collision_CullDistance);
-			Evolution.SetBoundsExtension(ChaosImmediate_Evolution_BoundsExtension);
+			Implementation->Collisions.SetCullDistance(ChaosImmediate_Collision_CullDistance);
+			Implementation->BroadPhase.SetCullDustance(ChaosImmediate_Collision_CullDistance);
+			Implementation->Evolution.SetBoundsExtension(ChaosImmediate_Evolution_BoundsExtension);
 
-			NarrowPhase.GetContext().bDeferUpdate = (ChaosImmediate_Collision_DeferNarrowPhase != 0);
-			NarrowPhase.GetContext().bAllowManifolds = (ChaosImmediate_Collision_UseManifolds != 0);
+			Implementation->NarrowPhase.GetContext().bDeferUpdate = (ChaosImmediate_Collision_DeferNarrowPhase != 0);
+			Implementation->NarrowPhase.GetContext().bAllowManifolds = (ChaosImmediate_Collision_UseManifolds != 0);
 
 			if (ChaosImmediate_UsePositionSolver)
 			{
-				Collisions.SetApplyType(ECollisionApplyType::Position);
+				Implementation->Collisions.SetApplyType(ECollisionApplyType::Position);
 			}
 			else
 			{
-				Collisions.SetApplyType(ECollisionApplyType::Velocity);
+				Implementation->Collisions.SetApplyType(ECollisionApplyType::Velocity);
 			}
 
 			if (ChaosImmediate_Evolution_StepTime > 0)
@@ -568,9 +645,9 @@ namespace ImmediatePhysics_Chaos
 		// Fixed timestep mode
 		// @todo(ccaulfield): integrate this better with the non-fixed step calculate above
 		FReal RewindTime = 0.0f;
-		if (FixedStepTime > 0)
+		if (Implementation->FixedStepTime > 0)
 		{
-			StepTime = FixedStepTime;
+			StepTime = Implementation->FixedStepTime;
 			NumSteps = FMath::FloorToInt(DeltaTime / StepTime);
 			FReal RemainderTime = DeltaTime - NumSteps * StepTime;
 			if (RemainderTime > ChaosImmediate_Evolution_FixedStepTolerance* StepTime)
@@ -582,15 +659,15 @@ namespace ImmediatePhysics_Chaos
 		}
 
 		// Handle new or deleted particles
-		if (bActorsDirty)
+		if (Implementation->bActorsDirty)
 		{
 			UpdateActivePotentiallyCollidingPairs();
-			bActorsDirty = false;
+			Implementation->bActorsDirty = false;
 		}
 
 		UE_LOG(LogChaosJoint, Verbose, TEXT("Simulate Dt = %f Steps %d x %f (Rewind %f)"), DeltaTime, NumSteps, StepTime, RewindTime);
-		Evolution.SetGravity(InGravity);
-		Evolution.Advance(StepTime, NumSteps, RewindTime);
+		Implementation->Evolution.SetGravity(InGravity);
+		Implementation->Evolution.Advance(StepTime, NumSteps, RewindTime);
 
 		DebugDrawKinematicParticles(1, 4, FColor(128, 0, 0));
 		DebugDrawDynamicParticles(1, 3, FColor(255, 255, 0));
@@ -601,34 +678,34 @@ namespace ImmediatePhysics_Chaos
 	{
 		if (InFixedDt >= 0.0f)
 		{
-			FixedStepTime = InFixedDt;
+			Implementation->FixedStepTime = InFixedDt;
 		}
 
 		if (SolverIts >= 0)
 		{
-			Evolution.SetNumIterations(SolverIts);
+			Implementation->Evolution.SetNumIterations(SolverIts);
 		}
 		if (SolverPushOutIts >= 0)
 		{
-			Evolution.SetNumPushOutIterations(SolverPushOutIts);
+			Implementation->Evolution.SetNumPushOutIterations(SolverPushOutIts);
 		}
 
 		if (JointIts >= 0)
 		{
-			Joints.SetNumPairIterations(JointIts);
+			Implementation->Joints.SetNumPairIterations(JointIts);
 		}
 		if (JointPushOutIts >= 0)
 		{
-			Joints.SetNumPushOutPairIterations(JointPushOutIts);
+			Implementation->Joints.SetNumPushOutPairIterations(JointPushOutIts);
 		}
 
 		if (CollisionIts >= 0)
 		{
-			Collisions.SetPairIterations(CollisionIts);
+			Implementation->Collisions.SetPairIterations(CollisionIts);
 		}
 		if (CollisionPushOutIts >= 0)
 		{
-			Collisions.SetPushOutPairIterations(CollisionPushOutIts);
+			Implementation->Collisions.SetPushOutPairIterations(CollisionPushOutIts);
 		}
 	}
 
@@ -640,15 +717,15 @@ namespace ImmediatePhysics_Chaos
 		{
 			if ((ChaosImmediate_DebugDrawParticles >= MinDebugLevel) && (ChaosImmediate_DebugDrawParticles <= MaxDebugLevel))
 			{
-				DebugDraw::DrawParticleTransforms(SimulationSpaceTransform, Particles.GetActiveKinematicParticlesView());
+				DebugDraw::DrawParticleTransforms(Implementation->SimulationSpaceTransform, Implementation->Particles.GetActiveKinematicParticlesView());
 			}
 			if ((ChaosImmediate_DebugDrawShapes >= MinDebugLevel) && (ChaosImmediate_DebugDrawShapes <= MaxDebugLevel))
 			{
-				DebugDraw::DrawParticleShapes(SimulationSpaceTransform, Particles.GetActiveKinematicParticlesView(), Color);
+				DebugDraw::DrawParticleShapes(Implementation->SimulationSpaceTransform, Implementation->Particles.GetActiveKinematicParticlesView(), Color);
 			}
 			if ((ChaosImmediate_DebugDrawBounds >= MinDebugLevel) && (ChaosImmediate_DebugDrawBounds <= MaxDebugLevel))
 			{
-				DebugDraw::DrawParticleBounds(SimulationSpaceTransform, Particles.GetActiveKinematicParticlesView(), Color);
+				DebugDraw::DrawParticleBounds(Implementation->SimulationSpaceTransform, Implementation->Particles.GetActiveKinematicParticlesView(), Color);
 			}
 		}
 #endif
@@ -663,15 +740,15 @@ namespace ImmediatePhysics_Chaos
 		{
 			if ((ChaosImmediate_DebugDrawParticles >= MinDebugLevel) && (ChaosImmediate_DebugDrawParticles <= MaxDebugLevel))
 			{
-				DebugDraw::DrawParticleTransforms(SimulationSpaceTransform, Particles.GetActiveParticlesView());
+				DebugDraw::DrawParticleTransforms(Implementation->SimulationSpaceTransform, Implementation->Particles.GetActiveParticlesView());
 			}
 			if ((ChaosImmediate_DebugDrawShapes >= MinDebugLevel) && (ChaosImmediate_DebugDrawShapes <= MaxDebugLevel))
 			{
-				DebugDraw::DrawParticleShapes(SimulationSpaceTransform, Particles.GetActiveParticlesView(), Color);
+				DebugDraw::DrawParticleShapes(Implementation->SimulationSpaceTransform, Implementation->Particles.GetActiveParticlesView(), Color);
 			}
 			if ((ChaosImmediate_DebugDrawBounds >= MinDebugLevel) && (ChaosImmediate_DebugDrawBounds <= MaxDebugLevel))
 			{
-				DebugDraw::DrawParticleBounds(SimulationSpaceTransform, Particles.GetActiveParticlesView(), Color);
+				DebugDraw::DrawParticleBounds(Implementation->SimulationSpaceTransform, Implementation->Particles.GetActiveParticlesView(), Color);
 			}
 		}
 #endif
@@ -685,11 +762,11 @@ namespace ImmediatePhysics_Chaos
 		{
 			if ((ChaosImmediate_DebugDrawCollisions >= MinDebugLevel) && (ChaosImmediate_DebugDrawCollisions <= MaxDebugLevel))
 			{
-				DebugDraw::DrawCollisions(SimulationSpaceTransform, Collisions, ColorScale);
+				DebugDraw::DrawCollisions(Implementation->SimulationSpaceTransform, Implementation->Collisions, ColorScale);
 			}
 			if ((ChaosImmediate_DebugDrawJoints >= MinDebugLevel) && (ChaosImmediate_DebugDrawJoints <= MaxDebugLevel))
 			{
-				DebugDraw::DrawJointConstraints(SimulationSpaceTransform, Joints, ColorScale, (uint32)ChaosImmediate_DebugDrawJointFeatures);
+				DebugDraw::DrawJointConstraints(Implementation->SimulationSpaceTransform, Implementation->Joints, ColorScale, (uint32)ChaosImmediate_DebugDrawJointFeatures);
 			}
 		}
 #endif
diff --git a/Engine/Source/Runtime/Engine/Public/Physics/ImmediatePhysics/ImmediatePhysicsChaos/ImmediatePhysicsCore_Chaos.h b/Engine/Source/Runtime/Engine/Public/Physics/ImmediatePhysics/ImmediatePhysicsChaos/ImmediatePhysicsCore_Chaos.h
--- a/Engine/Source/Runtime/Engine/Public/Physics/ImmediatePhysics/ImmediatePhysicsChaos/ImmediatePhysicsCore_Chaos.h
+++ b/Engine/Source/Runtime/Engine/Public/Physics/ImmediatePhysics/ImmediatePhysicsChaos/ImmediatePhysicsCore_Chaos.h
@@ -11,16 +11,21 @@
 namespace Chaos
 {
 	class FImplicitObject;
+	class FNarrowPhase;
+	class FParticlePairBroadPhase;
+	class FParticlePairCollisionDetector;
 	class FPBDJointConstraintHandle;
 	class FPBDJointConstraints;
 	class FPBDRigidsEvolutionGBF;
+	class FPerShapeData;
 	template<class T> class TArrayCollectionArray;
 	template<typename T, int D> struct TKinematicGeometryParticleParameters;
 	template<typename T, int D> class TKinematicTarget;
-	template<typename P> class TPBDConstraintIslandRule;
+	template<typename T> class TPBDConstraintIslandRule;
 	template<typename T, int D> struct TPBDRigidParticleParameters;
 	template<typename T, int D> class TPBDRigidsSOAs;
-	class FPerShapeData;
+	template<typename T> class TSimpleConstraintRule;
+
 }
 
 namespace ImmediatePhysics_Chaos
diff --git a/Engine/Source/Runtime/Engine/Public/Physics/ImmediatePhysics/ImmediatePhysicsChaos/ImmediatePhysicsSimulation_Chaos.h b/Engine/Source/Runtime/Engine/Public/Physics/ImmediatePhysics/ImmediatePhysicsChaos/ImmediatePhysicsSimulation_Chaos.h
--- a/Engine/Source/Runtime/Engine/Public/Physics/ImmediatePhysics/ImmediatePhysicsChaos/ImmediatePhysicsSimulation_Chaos.h
+++ b/Engine/Source/Runtime/Engine/Public/Physics/ImmediatePhysics/ImmediatePhysicsChaos/ImmediatePhysicsSimulation_Chaos.h
@@ -4,17 +4,6 @@
 
 #include "Physics/ImmediatePhysics/ImmediatePhysicsChaos/ImmediatePhysicsCore_Chaos.h"
 
-#define TEMP_HEADER_CHAOS_LEVEL_1
-#include "Chaos/ChaosDebugDrawDeclares.h"
-#include "Chaos/Collision/NarrowPhase.h"
-#include "Chaos/Collision/ParticlePairBroadPhase.h"
-#include "Chaos/Collision/ParticlePairCollisionDetector.h"
-#include "Chaos/Evolution/PBDMinEvolution.h"
-#include "Chaos/PBDCollisionConstraints.h"
-#include "Chaos/PBDConstraintRule.h"
-#include "Chaos/PBDJointConstraints.h"
-#undef TEMP_HEADER_CHAOS_LEVEL_1
-
 #include "Engine/EngineTypes.h"
 #include "Templates/UniquePtr.h"
 
@@ -27,10 +16,10 @@ namespace ImmediatePhysics_Chaos
 		FSimulation();
 		~FSimulation();
 
-		int32 NumActors() const { return ActorHandles.Num(); }
+		int32 NumActors() const;
 
-		FActorHandle* GetActorHandle(int32 ActorHandleIndex) { return ActorHandles[ActorHandleIndex]; }
-		const FActorHandle* GetActorHandle(int32 ActorHandleIndex) const { return ActorHandles[ActorHandleIndex]; }
+		FActorHandle* GetActorHandle(int32 ActorHandleIndex);
+		const FActorHandle* GetActorHandle(int32 ActorHandleIndex) const;
 
 		/** Create a static body and add it to the simulation */
 		FActorHandle* CreateStaticActor(FBodyInstance* BodyInstance);
@@ -81,53 +70,8 @@ namespace ImmediatePhysics_Chaos
 		void DebugDrawDynamicParticles(const int32 MinDebugLevel, const int32 MaxDebugLevel, const FColor& Color);
 		void DebugDrawConstraints(const int32 MinDebugLevel, const int32 MaxDebugLevel, const float ColorScale);
 
-		using FCollisionConstraints = Chaos::FPBDCollisionConstraints;
-		using FCollisionDetector = Chaos::FParticlePairCollisionDetector;
-		using FRigidParticleSOAs = Chaos::TPBDRigidsSOAs<FReal, 3>;
-		using FParticleHandle = Chaos::TGeometryParticleHandle<FReal, Dimensions>;
-		using FParticlePair = Chaos::TVector<Chaos::TGeometryParticleHandle<Chaos::FReal, 3>*, 2>;
-
-		// @todo(ccaulfield): Look into these...
-		TArray<FParticlePair> ActivePotentiallyCollidingPairs;
-		Chaos::TArrayCollectionArray<bool> CollidedParticles;
-		Chaos::TArrayCollectionArray<Chaos::TSerializablePtr<Chaos::FChaosPhysicsMaterial>> ParticleMaterials;
-		Chaos::TArrayCollectionArray<TUniquePtr<Chaos::FChaosPhysicsMaterial>> PerParticleMaterials;
-		Chaos::TArrayCollectionArray<Chaos::FVec3> ParticlePrevXs;
-		Chaos::TArrayCollectionArray<Chaos::FRotation3> ParticlePrevRs;
-
-		FRigidParticleSOAs Particles;
-		Chaos::FPBDJointConstraints Joints;
-		FCollisionConstraints Collisions;
-		Chaos::FParticlePairBroadPhase BroadPhase;
-		Chaos::FNarrowPhase NarrowPhase;
-		FCollisionDetector CollisionDetector;
-		Chaos::TSimpleConstraintRule<Chaos::FPBDJointConstraints> JointsRule;
-		Chaos::TSimpleConstraintRule<FCollisionConstraints> CollisionsRule;
-		Chaos::FPBDMinEvolution Evolution;
-
-
-		/** Mapping from entity index to handle */
-		// @todo(ccaulfield): we now have handles pointing to handles which is inefficient - we can do better than this, but don't want to change API yet
-		TArray<FActorHandle*> ActorHandles;
-		int32 NumActiveDynamicActorHandles;
-
-		/** Mapping from constraint index to handle */
-		TArray<FJointHandle*> JointHandles;
-
-		/** Slow to access. */
-		// @todo(ccaulfield): Optimize
-		TMap<const FParticleHandle*, TSet<const FParticleHandle*>> IgnoreCollisionParticlePairTable;
-
-		TArray<FParticlePair> PotentiallyCollidingPairs;
-
-		FTransform SimulationSpaceTransform;
-
-		FReal FixedStepTime;
-		FReal RollingAverageStepTime;
-		int32 NumRollingAverageStepTimes;
-		int32 MaxNumRollingAverageStepTimes;
-
-		bool bActorsDirty;
+		struct FImplementation;
+		TUniquePtr<FImplementation> Implementation;
 	};
 
 }
