Retry shader compile through workers after a worker was abandoned, but increase iteration for retries after every failed attempt.

#rnx

diff --git a/Engine/Source/Runtime/Engine/Private/ShaderCompiler/ShaderCompiler.cpp b/Engine/Source/Runtime/Engine/Private/ShaderCompiler/ShaderCompiler.cpp
--- a/Engine/Source/Runtime/Engine/Private/ShaderCompiler/ShaderCompiler.cpp
+++ b/Engine/Source/Runtime/Engine/Private/ShaderCompiler/ShaderCompiler.cpp
@@ -95,6 +95,12 @@ uint32 FShaderCommonCompileJob::GetNextJobId()
 static float GRegularWorkerTimeToLive = 20.0f;
 static float GBuildWorkerTimeToLive = 600.0f;
 
+// Configuration to retry shader compile through wrokers after a worker has been abandoned
+static constexpr int32 GSingleThreadedRunsIdle = -1;
+static constexpr int32 GSingleThreadedRunsDisabled = -2;
+static constexpr int32 GSingleThreadedRunsIncreaseFactor = 8;
+static constexpr int32 GSingleThreadedRunsMaxCount = (1 << 24);
+
 static void ModalErrorOrLog(const FString& Text, int64 CurrentFilePos = 0, int64 ExpectedFileSize = 0)
 {
 	FString BadFile;
@@ -1354,7 +1360,7 @@ bool FShaderCompileThreadRunnable::LaunchWorkersIfNeeded()
 	for (int32 WorkerIndex = 0; WorkerIndex < WorkerInfos.Num(); WorkerIndex++)
 	{
 		FShaderCompileWorkerInfo& CurrentWorkerInfo = *WorkerInfos[WorkerIndex];
-		if (CurrentWorkerInfo.QueuedJobs.Num() == 0 )
+		if (CurrentWorkerInfo.QueuedJobs.Num() == 0)
 		{
 			// Skip if nothing to do
 			// Also, use the opportunity to free OS resources by cleaning up handles of no more running processes
@@ -1584,6 +1590,23 @@ int32 FShaderCompileThreadRunnable::CompilingLoop()
 			// Fall back to local compiles if the SCW crashed.
 			// This is nasty but needed to work around issues where message passing through files to SCW is unreliable on random PCs
 			Manager->bAllowCompilingThroughWorkers = false;
+
+			// Try to recover from abandoned workers after a certain amount of single-threaded compilations
+			if (Manager->NumSingleThreadedRunsBeforeRetry == GSingleThreadedRunsIdle)
+			{
+				// First try to recover, only run single-threaded approach once
+				Manager->NumSingleThreadedRunsBeforeRetry = 1;
+			}
+			else if (Manager->NumSingleThreadedRunsBeforeRetry > GSingleThreadedRunsMaxCount)
+			{
+				// Stop retry approach after too many retries have failed
+				Manager->NumSingleThreadedRunsBeforeRetry = GSingleThreadedRunsDisabled;
+			}
+			else
+			{
+				// Next time increase runs by factor X
+				Manager->NumSingleThreadedRunsBeforeRetry *= GSingleThreadedRunsIncreaseFactor;
+			}
 		}
 		else
 		{
@@ -1593,7 +1616,19 @@ int32 FShaderCompileThreadRunnable::CompilingLoop()
 	}
 	else
 	{
+		// Execute all pending worker tasks single-threaded
 		CompileDirectlyThroughDll();
+
+		// If single-threaded mode was enabled by an abandoned worker, try to recover after the given amount of runs
+		if (Manager->NumSingleThreadedRunsBeforeRetry > 0)
+		{
+			Manager->NumSingleThreadedRunsBeforeRetry--;
+			if (Manager->NumSingleThreadedRunsBeforeRetry == 0)
+			{
+				UE_LOG(LogShaderCompilers, Display, TEXT("Retry shader compiling through workers."));
+				Manager->bAllowCompilingThroughWorkers = true;
+			}
+		}
 	}
 
 	return NumActiveThreads;
@@ -1806,6 +1841,7 @@ FShaderCompilingManager::FShaderCompilingManager() :
 	bCompilingDuringGame(false),
 	NumOutstandingJobs(0),
 	NumExternalJobs(0),
+	NumSingleThreadedRunsBeforeRetry(GSingleThreadedRunsIdle),
 #if PLATFORM_MAC
 	ShaderCompileWorkerName(FPaths::EngineDir() / TEXT("Binaries/Mac/ShaderCompileWorker")),
 #elif PLATFORM_LINUX
diff --git a/Engine/Source/Runtime/Engine/Public/ShaderCompiler.h b/Engine/Source/Runtime/Engine/Public/ShaderCompiler.h
--- a/Engine/Source/Runtime/Engine/Public/ShaderCompiler.h
+++ b/Engine/Source/Runtime/Engine/Public/ShaderCompiler.h
@@ -471,6 +471,8 @@ private:
 	uint32 NumShaderCompilingThreadsDuringGame;
 	/** Largest number of jobs that can be put in the same batch. */
 	int32 MaxShaderJobBatchSize;
+	/** Number of runs through single-threaded compiling before we can retry to compile through workers. -1 if not used. */
+	int32 NumSingleThreadedRunsBeforeRetry;
 	/** Process Id of UE4. */
 	uint32 ProcessId;
 	/** Whether to allow compiling shaders through the worker application, which allows multiple cores to be used. */
