Disable ISPC for fill component space code in skeletalmesh due to negative scaling

diff --git a/Engine/Source/Runtime/Engine/Private/Components/SkeletalMeshComponent.cpp b/Engine/Source/Runtime/Engine/Private/Components/SkeletalMeshComponent.cpp
--- a/Engine/Source/Runtime/Engine/Private/Components/SkeletalMeshComponent.cpp
+++ b/Engine/Source/Runtime/Engine/Private/Components/SkeletalMeshComponent.cpp
@@ -1395,6 +1395,7 @@ void USkeletalMeshComponent::FillComponentSpaceTransforms(const USkeletalMesh* I
 #endif
 	}
 
+#if 0 		 // temporarily disabling ISPC code due to negative scale issue
 	if (INTEL_ISPC)
 	{
 #if INTEL_ISPC
@@ -1409,6 +1410,7 @@ void USkeletalMeshComponent::FillComponentSpaceTransforms(const USkeletalMesh* I
 #endif
 	}
 	else
+#endif // 0
 	{
 		for (int32 i = 1; i < FillComponentSpaceTransformsRequiredBones.Num(); i++)
 		{
