Guard mass properties update against invalid inertia tensor (temporary fix).

diff --git a/Engine/Source/Runtime/Engine/Private/PhysicsEngine/BodyInstance.cpp b/Engine/Source/Runtime/Engine/Private/PhysicsEngine/BodyInstance.cpp
--- a/Engine/Source/Runtime/Engine/Private/PhysicsEngine/BodyInstance.cpp
+++ b/Engine/Source/Runtime/Engine/Private/PhysicsEngine/BodyInstance.cpp
@@ -3145,15 +3145,19 @@ void FBodyInstance::UpdateMassProperties()
 
 				// #PHYS2 Refactor out PxMassProperties (Our own impl?)
 #if WITH_CHAOS
-				const Chaos::TRotation<float, 3> Rotation = Chaos::TransformToLocalSpace(TotalMassProperties.InertiaTensor);
-				const FVector MassSpaceInertiaTensor(TotalMassProperties.InertiaTensor.M[0][0], TotalMassProperties.InertiaTensor.M[1][1], TotalMassProperties.InertiaTensor.M[2][2]);
-
-				FPhysicsInterface::SetMassSpaceInertiaTensor_AssumesLocked(Actor, MassSpaceInertiaTensor);
+				// Only set mass properties if inertia tensor is valid. TODO Remove this once we track down cause of empty tensors.
+				const float InertiaTensorTrace = (TotalMassProperties.InertiaTensor.M[0][0] + TotalMassProperties.InertiaTensor.M[1][1] + TotalMassProperties.InertiaTensor.M[2][2]) / 3;
+				if (CHAOS_ENSURE(InertiaTensorTrace > SMALL_NUMBER))
+				{
+					const Chaos::TRotation<float, 3> Rotation = Chaos::TransformToLocalSpace(TotalMassProperties.InertiaTensor);
+					const FVector MassSpaceInertiaTensor(TotalMassProperties.InertiaTensor.M[0][0], TotalMassProperties.InertiaTensor.M[1][1], TotalMassProperties.InertiaTensor.M[2][2]);
+					FPhysicsInterface::SetMassSpaceInertiaTensor_AssumesLocked(Actor, MassSpaceInertiaTensor);
 
-				FPhysicsInterface::SetMass_AssumesLocked(Actor, TotalMassProperties.Mass);
+					FPhysicsInterface::SetMass_AssumesLocked(Actor, TotalMassProperties.Mass);
 
-				FTransform Com(Rotation, TotalMassProperties.CenterOfMass);
-				FPhysicsInterface::SetComLocalPose_AssumesLocked(Actor, Com);
+					FTransform Com(Rotation, TotalMassProperties.CenterOfMass);
+					FPhysicsInterface::SetComLocalPose_AssumesLocked(Actor, Com);
+				}
 #else
 				PxQuat MassOrientation;
 				const FVector MassSpaceInertiaTensor = P2UVector(PxMassProperties::getMassSpaceInertia(TotalMassProperties.inertiaTensor, MassOrientation));
