Fix crash with `obj list` when run on dedicated servers.

#rnx
None

diff --git a/Engine/Source/Runtime/Engine/Private/RepLayout.cpp b/Engine/Source/Runtime/Engine/Private/RepLayout.cpp
--- a/Engine/Source/Runtime/Engine/Private/RepLayout.cpp
+++ b/Engine/Source/Runtime/Engine/Private/RepLayout.cpp
@@ -945,7 +945,13 @@ void FRepStateStaticBuffer::CountBytes(FArchive& Ar) const
 	const uint64 StaticTotalMemory = LocalAr.GetMax();
 
 	FCountBytesHelper CountBytesHelper(LocalAr, Buffer.GetData(), StaticTotalMemory, RepLayout->Parents, RepLayout->Cmds);
-	CountBytesHelper.CountBytes();
+
+	// Buffers can be empty for ReceivingRepStates that are being tracked on servers.
+	// TODO: Do we actually need to allocate receiving rep states on servers at all?
+	if (Buffer.Num() > 0)
+	{
+		CountBytesHelper.CountBytes();
+	}
 
 	const uint64 StaticOnRepMemory = StaticTotalMemory - CountBytesHelper.NonRepStaticMemory;
 
