Adding option to use control channel for destruction infos (net.ControlChannelDestructionInfo, disabled by default)

diff --git a/Engine/Plugins/Runtime/ReplicationGraph/Source/Private/ReplicationGraph.cpp b/Engine/Plugins/Runtime/ReplicationGraph/Source/Private/ReplicationGraph.cpp
--- a/Engine/Plugins/Runtime/ReplicationGraph/Source/Private/ReplicationGraph.cpp
+++ b/Engine/Plugins/Runtime/ReplicationGraph/Source/Private/ReplicationGraph.cpp
@@ -2592,14 +2592,13 @@ int64 UNetReplicationGraphConnection::ReplicateDestructionInfos(const FNetViewer
 
 		if (bSendDestructionInfo || DestructInfo->bIgnoreDistanceCulling)
 		{
-			UActorChannel* Channel = (UActorChannel*)NetConnection->CreateChannelByName(NAME_Actor, EChannelCreateFlags::OpenedLocally);
-			if (Channel)
+			if (NetConnection && NetConnection->Driver)
 			{
-				NumBits += Channel->SetChannelActorForDestroy(DestructInfo);
-			}
+				NumBits += NetConnection->Driver->SendDestructionInfo(NetConnection, DestructInfo);
 
-			PendingDestructInfoList.RemoveAtSwap(idx, 1, false);
-			TrackedDestructionInfoPtrs.Remove(DestructInfo);
+				PendingDestructInfoList.RemoveAtSwap(idx, 1, false);
+				TrackedDestructionInfoPtrs.Remove(DestructInfo);
+			}
 		}
 		else
 		{
@@ -2618,25 +2617,24 @@ int64 UNetReplicationGraphConnection::ReplicateDormantDestructionInfos()
 
 	int64 NumBits = 0;
 
-	for (const FCachedDormantDestructInfo& Info : PendingDormantDestructList)
+	if (NetConnection && NetConnection->Driver)
 	{
-		FActorDestructionInfo DestructInfo;
-		DestructInfo.DestroyedPosition = FVector::ZeroVector;
-		DestructInfo.NetGUID = Info.NetGUID;
-		DestructInfo.Level = Info.Level;
-		DestructInfo.ObjOuter = Info.ObjOuter;
-		DestructInfo.PathName = Info.PathName;
-		DestructInfo.StreamingLevelName = NAME_None;			// currently unused by SetChannelActorForDestroy
-		DestructInfo.Reason = EChannelCloseReason::Relevancy;
-
-		UActorChannel* Channel = (UActorChannel*)NetConnection->CreateChannelByName(NAME_Actor, EChannelCreateFlags::OpenedLocally);
-		if (Channel)
+		for (const FCachedDormantDestructInfo& Info : PendingDormantDestructList)
 		{
-			NumBits += Channel->SetChannelActorForDestroy(&DestructInfo);
+			FActorDestructionInfo DestructInfo;
+			DestructInfo.DestroyedPosition = FVector::ZeroVector;
+			DestructInfo.NetGUID = Info.NetGUID;
+			DestructInfo.Level = Info.Level;
+			DestructInfo.ObjOuter = Info.ObjOuter;
+			DestructInfo.PathName = Info.PathName;
+			DestructInfo.StreamingLevelName = NAME_None;			// currently unused by SetChannelActorForDestroy
+			DestructInfo.Reason = EChannelCloseReason::Relevancy;
+
+			NumBits += NetConnection->Driver->SendDestructionInfo(NetConnection, &DestructInfo);
 		}
-	}
 
-	PendingDormantDestructList.Empty();
+		PendingDormantDestructList.Empty();
+	}
 
 	return NumBits;
 }
diff --git a/Engine/Source/Runtime/Engine/Classes/Engine/ActorChannel.h b/Engine/Source/Runtime/Engine/Classes/Engine/ActorChannel.h
--- a/Engine/Source/Runtime/Engine/Classes/Engine/ActorChannel.h
+++ b/Engine/Source/Runtime/Engine/Classes/Engine/ActorChannel.h
@@ -162,6 +162,7 @@ public:
 
 	virtual void NotifyActorChannelOpen(AActor* InActor, FInBunch& InBunch);
 
+	UE_DEPRECATED(4.26, "Use UNetworkDriver::SendDestructionInfo instead.")
 	int64 SetChannelActorForDestroy( struct FActorDestructionInfo *DestructInfo );
 
 	/** Append any export bunches */
@@ -317,7 +318,7 @@ protected:
 	bool ObjectHasReplicator(const TWeakObjectPtr<UObject>& Obj) const;	// returns whether we have already created a replicator for this object or not
 
 	/** Unmap all references to this object, so that if later we receive this object again, we can remap the original references */
-	void MoveMappedObjectToUnmapped( const UObject* Object );
+	void MoveMappedObjectToUnmapped(const UObject* Object);
 
 	void DestroyActorAndComponents();
 
diff --git a/Engine/Source/Runtime/Engine/Classes/Engine/ControlChannel.h b/Engine/Source/Runtime/Engine/Classes/Engine/ControlChannel.h
--- a/Engine/Source/Runtime/Engine/Classes/Engine/ControlChannel.h
+++ b/Engine/Source/Runtime/Engine/Classes/Engine/ControlChannel.h
@@ -12,6 +12,7 @@
 class FInBunch;
 class FOutBunch;
 class UNetConnection;
+struct FActorDestructionInfo;
 
 /**
  * A queued control channel message
@@ -105,4 +106,10 @@ class ENGINE_API UControlChannel
 
 	/** Describe the text channel. */
 	virtual FString Describe() override;
+
+	/** Sends a message to destroy a specific actor without creating an actor channel. */
+	int64 SendDestructionInfo(FActorDestructionInfo* DestructionInfo);
+
+	/** Handle receiving a destruction message for an actor outside of an actor channel. */
+	void ReceiveDestructionInfo(FInBunch& Bunch);
 };
diff --git a/Engine/Source/Runtime/Engine/Classes/Engine/NetDriver.h b/Engine/Source/Runtime/Engine/Classes/Engine/NetDriver.h
--- a/Engine/Source/Runtime/Engine/Classes/Engine/NetDriver.h
+++ b/Engine/Source/Runtime/Engine/Classes/Engine/NetDriver.h
@@ -530,6 +530,31 @@ struct ENGINE_API FPacketSimulationSettings
 	bool ConfigHelperBool(const TCHAR* Name, bool& Value, const TCHAR* OptionalQualifier);
 };
 
+struct ENGINE_API FActorDestructionInfo
+{
+public:
+	FActorDestructionInfo()
+		: Reason(EChannelCloseReason::Destroyed)
+		, bIgnoreDistanceCulling(false) 
+	{}
+
+	TWeakObjectPtr<ULevel> Level; 
+	TWeakObjectPtr<UObject> ObjOuter;
+	FVector DestroyedPosition;
+	FNetworkGUID NetGUID;
+	FString PathName;
+	FName StreamingLevelName;
+	EChannelCloseReason Reason;
+
+	/** When true the destruction info data will be sent even if the viewers are not close to the actor */
+	bool bIgnoreDistanceCulling;
+
+	void CountBytes(FArchive& Ar)
+	{
+		PathName.CountBytes(Ar);
+	}
+};
+
 //
 // Priority sortable list.
 //
@@ -540,14 +565,14 @@ struct FActorPriority
 	FNetworkObjectInfo*			ActorInfo;	// Actor info.
 	class UActorChannel*		Channel;	// Actor channel.
 
-	struct FActorDestructionInfo *	DestructionInfo;	// Destroy an actor
+	FActorDestructionInfo *	DestructionInfo;	// Destroy an actor
 
 	FActorPriority() : 
 		Priority(0), ActorInfo(NULL), Channel(NULL), DestructionInfo(NULL)
 	{}
 
 	FActorPriority(class UNetConnection* InConnection, class UActorChannel* InChannel, FNetworkObjectInfo* InActorInfo, const TArray<struct FNetViewer>& Viewers, bool bLowBandwidth);
-	FActorPriority(class UNetConnection* InConnection, struct FActorDestructionInfo * DestructInfo, const TArray<struct FNetViewer>& Viewers );
+	FActorPriority(class UNetConnection* InConnection, FActorDestructionInfo * DestructInfo, const TArray<struct FNetViewer>& Viewers );
 };
 
 struct FCompareFActorPriority
@@ -558,28 +583,6 @@ struct FCompareFActorPriority
 	}
 };
 
-struct FActorDestructionInfo
-{
-public:
-	FActorDestructionInfo() : Reason(EChannelCloseReason::Destroyed) {}
-
-	TWeakObjectPtr<ULevel>		Level;
-	TWeakObjectPtr<UObject>		ObjOuter;
-	FVector			DestroyedPosition;
-	FNetworkGUID	NetGUID;
-	FString			PathName;
-	FName			StreamingLevelName;
-	EChannelCloseReason Reason;
-
-	/** When true the destruction info data will be sent even if the viewers are not close to the actor */
-	bool bIgnoreDistanceCulling = false;
-
-	void CountBytes(FArchive& Ar)
-	{
-		PathName.CountBytes(Ar);
-	}
-};
-
 /** Used to specify properties of a channel type */
 USTRUCT()
 struct ENGINE_API FChannelDefinition
@@ -1544,6 +1547,9 @@ public:
 	/** Returns identifier used for NetTrace */
 	inline uint32 GetNetTraceId() const { return NetTraceId; }
 
+	/** Sends a message to a client to destroy an actor to the client.  The actor may already be destroyed locally. */
+	ENGINE_API int64 SendDestructionInfo(UNetConnection* Connection, FActorDestructionInfo* DestructionInfo);
+
 protected:
 
 	/** Register all TickDispatch, TickFlush, PostTickFlush to tick in World */
@@ -1604,6 +1610,9 @@ public:
 
 	static bool IsDormInitialStartupActor(AActor* Actor);
 
+	/** Unmap all references to this object, so that if later we receive this object again, we can remap the original references */
+	void MoveMappedObjectToUnmapped(const UObject* Object);
+
 protected:
 
 	bool bMaySendProperties;
diff --git a/Engine/Source/Runtime/Engine/Private/DataChannel.cpp b/Engine/Source/Runtime/Engine/Private/DataChannel.cpp
--- a/Engine/Source/Runtime/Engine/Private/DataChannel.cpp
+++ b/Engine/Source/Runtime/Engine/Private/DataChannel.cpp
@@ -1391,6 +1391,7 @@ IMPLEMENT_CONTROL_CHANNEL_MESSAGE(BeaconJoin);
 IMPLEMENT_CONTROL_CHANNEL_MESSAGE(BeaconAssignGUID);
 IMPLEMENT_CONTROL_CHANNEL_MESSAGE(BeaconNetGUIDAck);
 IMPLEMENT_CONTROL_CHANNEL_MESSAGE(EncryptionAck);
+IMPLEMENT_CONTROL_CHANNEL_MESSAGE(DestructionInfo);
 
 void UControlChannel::Init( UNetConnection* InConnection, int32 InChannelIndex, EChannelCreateFlags CreateFlags )
 {
@@ -1550,7 +1551,7 @@ void UControlChannel::ReceivedBunch( FInBunch& Bunch )
 				}
 			}
 		}
-		else if(MessageType == NMT_SecurityViolation)
+		else if (MessageType == NMT_SecurityViolation)
 		{
 			FString DebugMessage;
 			if (FNetControlMessage<NMT_SecurityViolation>::Receive(Bunch, DebugMessage))
@@ -1559,6 +1560,10 @@ void UControlChannel::ReceivedBunch( FInBunch& Bunch )
 				break;
 			}
 		}
+		else if (MessageType == NMT_DestructionInfo)
+		{
+			ReceiveDestructionInfo(Bunch);
+		}
 		else
 		{
 			// Process control message on client/server connection
@@ -1782,6 +1787,121 @@ FString UControlChannel::Describe()
 	return UChannel::Describe();
 }
 
+int64 UControlChannel::SendDestructionInfo(FActorDestructionInfo* DestructionInfo)
+{
+	int64 NumBits = 0;
+
+	checkf(Connection && Connection->PackageMap, TEXT("SendDestructionInfo requires a valid connection and package map: %s"), *Describe());
+	checkf(DestructionInfo, TEXT("SendDestructionInfo was passed an invalid desctruction info: %s"), *Describe());
+
+	if (!Closing && (Connection->State == USOCK_Open || Connection->State == USOCK_Pending))
+	{
+		// Outer must be valid to call PackageMap->WriteObject. In the case of streaming out levels, this can go null out of from underneath us. In that case, just skip the destruct info.
+		// We assume that if server unloads a level that clients will to and this will implicitly destroy all actors in it, so not worried about leaking actors client side here.
+		if (UObject* ObjOuter = DestructionInfo->ObjOuter.Get())
+		{
+			FOutBunch InfoBunch(Connection->PackageMap, false);
+			check(!InfoBunch.IsError());
+			InfoBunch.bReliable = 1;
+
+			uint8 MessageType = NMT_DestructionInfo;
+			InfoBunch << MessageType;
+
+			EChannelCloseReason Reason = DestructionInfo->Reason;
+			InfoBunch << Reason;
+
+			Connection->PackageMap->WriteObject(InfoBunch, ObjOuter, DestructionInfo->NetGUID, DestructionInfo->PathName);
+
+			UE_LOG(LogNetTraffic, Log, TEXT("SendDestructionInfo: NetGUID <%s> Path: %s. Bits: %d"), *DestructionInfo->NetGUID.ToString(), *DestructionInfo->PathName, InfoBunch.GetNumBits());
+			UE_LOG(LogNetDormancy, Verbose, TEXT("SendDestructionInfo: NetGUID <%s> Path: %s. Bits: %d"), *DestructionInfo->NetGUID.ToString(), *DestructionInfo->PathName, InfoBunch.GetNumBits());
+
+			SendBunch(&InfoBunch, false);
+
+			NumBits = InfoBunch.GetNumBits();
+		}
+	}
+
+	return NumBits;
+}
+
+void UControlChannel::ReceiveDestructionInfo(FInBunch& Bunch)
+{
+	checkf(Connection && Connection->PackageMap && Connection->Driver, TEXT("UControlChannel::ReceiveDestructionInfo requires a valid connection, package map, and driver: %s"), *Describe());
+
+	EChannelCloseReason CloseReason = EChannelCloseReason::Destroyed;
+	Bunch << CloseReason;
+
+	FNetworkGUID NetGUID;
+	UObject* Object = nullptr;
+
+	if (Connection->PackageMap->SerializeObject(Bunch, UObject::StaticClass(), Object, &NetGUID))
+	{
+		if (AActor* TheActor = Cast<AActor>(Object))
+		{
+			// If we're the client, destroy this actor.
+			if (!Connection->Driver->IsServer())
+			{
+				checkf(TheActor->IsValidLowLevel(), TEXT("ReceiveDestructionInfo serialized an invalid actor: %s"), *Describe());
+				checkSlow(Connection->IsValidLowLevel());
+				checkSlow(Connection->Driver->IsValidLowLevel());
+
+				if (TheActor->GetTearOff() && !Connection->Driver->ShouldClientDestroyTearOffActors())
+				{
+					if (!bTornOff)
+					{
+						TheActor->SetRole(ROLE_Authority);
+						TheActor->SetReplicates(false);
+						bTornOff = true;
+						if (TheActor->GetWorld() != nullptr && !IsEngineExitRequested())
+						{
+							TheActor->TornOff();
+						}
+
+						Connection->Driver->NotifyActorTornOff(TheActor);
+					}
+				}
+				else if (Dormant && (CloseReason == EChannelCloseReason::Dormancy) && !TheActor->GetTearOff())
+				{
+					TheActor->NetDormancy = DORM_DormantAll;
+
+					Connection->Driver->NotifyActorFullyDormantForConnection(TheActor, Connection);
+				}
+				else if (!TheActor->bNetTemporary && TheActor->GetWorld() != nullptr && !IsEngineExitRequested() && Connection->Driver->ShouldClientDestroyActor(TheActor))
+				{
+					// Destroy the actor
+
+					// Unmap any components in this actor. This will make sure that once the Actor is remapped
+					// any references to components will be remapped as well.
+					for (UActorComponent* Component : TheActor->GetComponents())
+					{
+						Connection->Driver->MoveMappedObjectToUnmapped(Component);
+					}
+
+					// Unmap this object so we can remap it if it becomes relevant again in the future
+					Connection->Driver->MoveMappedObjectToUnmapped(TheActor);
+
+					TheActor->PreDestroyFromReplication();
+					TheActor->Destroy(true);
+
+					if (CVarFilterGuidRemapping.GetValueOnAnyThread() > 0)
+					{
+						// Remove this actor's NetGUID from the list of unmapped values, it will be added back if it replicates again
+						if (NetGUID.IsValid() && Connection != nullptr && Connection->Driver != nullptr && Connection->Driver->GuidCache.IsValid())
+						{
+							Connection->Driver->GuidCache->ImportedNetGuids.Remove(NetGUID);
+						}
+					}
+
+					if (UPackageMapClient* PackageMapClient = Cast<UPackageMapClient>(Connection->PackageMap))
+					{
+						PackageMapClient->SetHasQueuedBunches(NetGUID, false);
+					}
+				}
+			}
+		}
+	}
+}
+
 /*-----------------------------------------------------------------------------
 	UActorChannel.
 -----------------------------------------------------------------------------*/
@@ -1877,49 +1997,11 @@ void UActorChannel::CleanupReplicators( const bool bKeepReplicators )
 	ActorReplicator.Reset();
 }
 
-static TAutoConsoleVariable<int32> CVarRelinkMappedReferences( TEXT( "net.RelinkMappedReferences" ), 1, TEXT( "" ) );
-
-void UActorChannel::MoveMappedObjectToUnmapped( const UObject* Object )
+void UActorChannel::MoveMappedObjectToUnmapped(const UObject* Object)
 {
-	if ( !Object )
+	if (Connection && Connection->Driver)
 	{
-		return;
-	}
-
-	if ( !CVarRelinkMappedReferences.GetValueOnGameThread() )
-	{
-		return;
-	}
-
-	UNetDriver* Driver = Connection ? Connection->Driver : nullptr;
-
-	if ( !Driver || Driver->IsServer() )
-	{
-		return;
-	}
-
-	// Find all replicators that are referencing this object, and make sure to mark the references as unmapped
-	// This is so when/if this object is instantiated again (using same network guid), we can re-establish the old references
-	FNetworkGUID NetGuid = Driver->GuidCache->NetGUIDLookup.FindRef( const_cast<UObject*>(Object) );
-
-	if ( NetGuid.IsValid() )
-	{
-		TSet< FObjectReplicator* >* Replicators = Driver->GuidToReplicatorMap.Find( NetGuid );
-
-		if ( Replicators != nullptr )
-		{
-			for ( FObjectReplicator* Replicator : *Replicators )
-			{
-				if ( Replicator->MoveMappedObjectToUnmapped( NetGuid ) )
-				{
-					Driver->UnmappedReplicators.Add( Replicator );
-				}
-				else if ( !Driver->UnmappedReplicators.Contains( Replicator ) )
-				{
-					UE_LOG( LogNet, Warning, TEXT( "UActorChannel::MoveMappedObjectToUnmapped: MoveMappedObjectToUnmapped didn't find object: %s" ), *GetPathNameSafe( Replicator->GetObject() ) );
-				}
-			}
-		}
+		Connection->Driver->MoveMappedObjectToUnmapped(Object);
 	}
 }
 
diff --git a/Engine/Source/Runtime/Engine/Private/DemoNetDriver.cpp b/Engine/Source/Runtime/Engine/Private/DemoNetDriver.cpp
--- a/Engine/Source/Runtime/Engine/Private/DemoNetDriver.cpp
+++ b/Engine/Source/Runtime/Engine/Private/DemoNetDriver.cpp
@@ -3067,7 +3067,7 @@ void UDemoNetDriver::TickDemoRecordFrame(float DeltaSeconds)
 	LastReplayFrameFidelity = ReplicatedPercent;
 }
 
-bool UDemoNetDriver::ReplicatePrioritizedActor(const FActorPriority& ActorPriority, const class FRepActorsParams& Params)
+bool UDemoNetDriver::ReplicatePrioritizedActor(const FActorPriority& ActorPriority, const FRepActorsParams& Params)
 {
 	FNetworkObjectInfo* ActorInfo = ActorPriority.ActorInfo;
 	FActorDestructionInfo* DestructionInfo = ActorPriority.DestructionInfo;
@@ -3088,10 +3088,12 @@ bool UDemoNetDriver::ReplicatePrioritizedActor(const FActorPriority& ActorPriori
 			UE_LOG(LogDemo, Verbose, TEXT("TickDemoRecord creating destroy channel for NetGUID <%s,%s> Priority: %d"), *DestructionInfo->NetGUID.ToString(), *DestructionInfo->PathName, ActorPriority.Priority);
 
 			// Send a close bunch on the new channel
+			PRAGMA_DISABLE_DEPRECATION_WARNINGS
 			Channel->SetChannelActorForDestroy(DestructionInfo);
+			PRAGMA_ENABLE_DEPRECATION_WARNINGS
 
 			// Remove from connection's to-be-destroyed list (close bunch is reliable, so it will make it there)
-			Connection->GetDestroyedStartupOrDormantActorGUIDs().Remove(DestructionInfo->NetGUID);
+			Connection->RemoveDestructionInfo(DestructionInfo);
 		}
 	}
 	else if (ActorInfo != nullptr && DestructionInfo == nullptr)
diff --git a/Engine/Source/Runtime/Engine/Private/NetworkDriver.cpp b/Engine/Source/Runtime/Engine/Private/NetworkDriver.cpp
--- a/Engine/Source/Runtime/Engine/Private/NetworkDriver.cpp
+++ b/Engine/Source/Runtime/Engine/Private/NetworkDriver.cpp
@@ -241,6 +241,14 @@ static TAutoConsoleVariable<int32> CVarAllowReliableMulticastToNonRelevantChanne
 	1,
 	TEXT("Allow Reliable Server Multicasts to be sent to non-Relevant Actors, as long as their is an existing ActorChannel."));
 
+static int32 GNetControlChannelDestructionInfo = 0;
+static FAutoConsoleVariableRef CVarNetControlChannelDestructionInfo(
+	TEXT("net.ControlChannelDestructionInfo"),
+	GNetControlChannelDestructionInfo,
+	TEXT("If enabled, send destruction info updates via the control channel instead of creating a new actor channel.")
+	TEXT("0: Old behavior, use an actor channel. 1: New behavior, use the control channel"),
+	ECVF_Default);
+
 
 /*-----------------------------------------------------------------------------
 	UNetDriver implementation.
@@ -3675,7 +3683,7 @@ FActorPriority::FActorPriority(UNetConnection* InConnection, UActorChannel* InCh
 	}
 }
 
-FActorPriority::FActorPriority(class UNetConnection* InConnection, struct FActorDestructionInfo * Info, const TArray<struct FNetViewer>& Viewers )
+FActorPriority::FActorPriority(UNetConnection* InConnection, FActorDestructionInfo * Info, const TArray<FNetViewer>& Viewers )
 	: ActorInfo(NULL), Channel(NULL), DestructionInfo(Info)
 {
 	
@@ -4197,15 +4205,12 @@ int32 UNetDriver::ServerReplicateActors_ProcessPrioritizedActors( UNetConnection
 				continue;
 			}
 
-			UActorChannel* Channel = (UActorChannel*)Connection->CreateChannelByName( NAME_Actor, EChannelCreateFlags::OpenedLocally );
-			if ( Channel )
-			{
-				FinalRelevantCount++;
-				UE_LOG( LogNetTraffic, Log, TEXT( "Server replicate actor creating destroy channel for NetGUID <%s,%s> Priority: %d" ), *PriorityActors[j]->DestructionInfo->NetGUID.ToString(), *PriorityActors[j]->DestructionInfo->PathName, PriorityActors[j]->Priority );
+			FinalRelevantCount++;
+			UE_LOG( LogNetTraffic, Log, TEXT( "Server replicate actor creating destroy channel for NetGUID <%s,%s> Priority: %d" ), *PriorityActors[j]->DestructionInfo->NetGUID.ToString(), *PriorityActors[j]->DestructionInfo->PathName, PriorityActors[j]->Priority );
 
-				Channel->SetChannelActorForDestroy( PriorityActors[j]->DestructionInfo );		// Send a close bunch on the new channel
-				Connection->RemoveDestructionInfo( PriorityActors[j]->DestructionInfo );		// Remove from connections to-be-destroyed list (close bunch of reliable, so it will make it there)
-			}
+			SendDestructionInfo(Connection, PriorityActors[j]->DestructionInfo);
+
+			Connection->RemoveDestructionInfo( PriorityActors[j]->DestructionInfo );		// Remove from connections to-be-destroyed list (close bunch of reliable, so it will make it there)
 			continue;
 		}
 
@@ -4375,6 +4380,38 @@ int32 UNetDriver::ServerReplicateActors_ProcessPrioritizedActors( UNetConnection
 }
 #endif
 
+int64 UNetDriver::SendDestructionInfo(UNetConnection* Connection, FActorDestructionInfo* DestructionInfo)
+{
+	QUICK_SCOPE_CYCLE_COUNTER(STAT_NetSendDestructionInfo);
+	CSV_SCOPED_TIMING_STAT_EXCLUSIVE(SendDestructionInfo);
+
+	int64 NumBits = 0;
+
+	checkf(Connection, TEXT("SendDestructionInfo called with invalid connection: %s"), *GetDescription());
+	checkf(DestructionInfo, TEXT("SendDestructionInfo called with invalid destruction info: %s"), *GetDescription());
+
+	if (GNetControlChannelDestructionInfo != 0)
+	{
+		if (UControlChannel* ControlChan = Cast<UControlChannel>(Connection->Channels[0]))
+		{
+			NumBits = ControlChan->SendDestructionInfo(DestructionInfo);
+		}
+	}
+	else
+	{
+		UActorChannel* Channel = (UActorChannel*)Connection->CreateChannelByName(NAME_Actor, EChannelCreateFlags::OpenedLocally);
+		if (Channel)
+		{
+			// Send a close bunch on the new channel
+			PRAGMA_DISABLE_DEPRECATION_WARNINGS
+				NumBits = Channel->SetChannelActorForDestroy(DestructionInfo);
+			PRAGMA_ENABLE_DEPRECATION_WARNINGS
+		}
+	}
+
+	return NumBits;
+}
+
 // -------------------------------------------------------------------------------------------------------------------------
 //	Replication profiling (server cpu) helpers.
 //		-CSV stat cateogry "Replication". Not enabled by default. (use -csvCategories=Replication to enable from cmd line)
@@ -5764,6 +5801,50 @@ ECreateReplicationChangelistMgrFlags UNetDriver::GetCreateReplicationChangelistM
 	return ECreateReplicationChangelistMgrFlags::None;
 }
 
+static TAutoConsoleVariable<int32> CVarRelinkMappedReferences(TEXT("net.RelinkMappedReferences"), 1, TEXT(""));
+
+void UNetDriver::MoveMappedObjectToUnmapped(const UObject* Object)
+{
+	if (!Object)
+	{
+		return;
+	}
+
+	if (!CVarRelinkMappedReferences.GetValueOnGameThread())
+	{
+		return;
+	}
+
+	if (IsServer())
+	{
+		return;
+	}
+
+	// Find all replicators that are referencing this object, and make sure to mark the references as unmapped
+	// This is so when/if this object is instantiated again (using same network guid), we can re-establish the old references
+	FNetworkGUID NetGuid = GuidCache->NetGUIDLookup.FindRef(const_cast<UObject*>(Object));
+
+	if (NetGuid.IsValid())
+	{
+		TSet<FObjectReplicator*>* Replicators = GuidToReplicatorMap.Find(NetGuid);
+
+		if (Replicators != nullptr)
+		{
+			for (FObjectReplicator* Replicator : *Replicators)
+			{
+				if (Replicator->MoveMappedObjectToUnmapped(NetGuid))
+				{
+					UnmappedReplicators.Add(Replicator);
+				}
+				else if (!UnmappedReplicators.Contains(Replicator))
+				{
+					UE_LOG(LogNet, Warning, TEXT("UActorChannel::MoveMappedObjectToUnmapped: MoveMappedObjectToUnmapped didn't find object: %s"), *GetPathNameSafe(Replicator->GetObject()));
+				}
+			}
+		}
+	}
+}
+
 FAutoConsoleCommandWithWorld	DumpRelevantActorsCommand(
 	TEXT("net.DumpRelevantActors"), 
 	TEXT( "Dumps information on relevant actors during next network update" ), 
diff --git a/Engine/Source/Runtime/Engine/Public/Net/DataChannel.h b/Engine/Source/Runtime/Engine/Public/Net/DataChannel.h
--- a/Engine/Source/Runtime/Engine/Public/Net/DataChannel.h
+++ b/Engine/Source/Runtime/Engine/Public/Net/DataChannel.h
@@ -166,6 +166,7 @@ DEFINE_CONTROL_CHANNEL_MESSAGE(NetGUIDAssign, 18, FNetworkGUID, FString); // Exp
 DEFINE_CONTROL_CHANNEL_MESSAGE(SecurityViolation, 19, FString); // server tells client that it has violated security and has been disconnected
 DEFINE_CONTROL_CHANNEL_MESSAGE(GameSpecific, 20, uint8, FString); // custom game-specific message routed to UGameInstance for processing
 DEFINE_CONTROL_CHANNEL_MESSAGE(EncryptionAck, 21);
+DEFINE_CONTROL_CHANNEL_MESSAGE(DestructionInfo, 22);
 
 // 			Beacon control channel flow
 // Client												Server
