Guard the body of UpdatePlayInEditorWorldDebugString behind a GIsEditor check so it doesn't do anything in -game or -server runs

diff --git a/Engine/Source/Runtime/Engine/Private/UnrealEngine.cpp b/Engine/Source/Runtime/Engine/Private/UnrealEngine.cpp
--- a/Engine/Source/Runtime/Engine/Private/UnrealEngine.cpp
+++ b/Engine/Source/Runtime/Engine/Private/UnrealEngine.cpp
@@ -395,6 +395,11 @@ ENGINE_API FString GPlayInEditorContextString(TEXT("invalid"));
 ENGINE_API void UpdatePlayInEditorWorldDebugString(const FWorldContext* WorldContext)
 {
 #if WITH_EDITOR
+	if (!GIsEditor)
+	{
+		return;
+	}
+
 	if (WorldContext == nullptr)
 	{
 		ensure(GPlayInEditorID == INDEX_NONE);
