Changed static mesh stream out to avoid swaping the LOD resources ptr.
The new implementation assumes that no new references to streamed out LOD resources will be created once mesh stream out is triggered.
This is enforced by updating FStaticMeshRenderData::CurrentFirstLODIdx immediately on the gamethread, as soon as the streamout is triggered.
Code that need to access valid FStaticMeshLODResources CPU data must call FStaticMeshRenderData::GetCurrentFirstLOD() and keep a
Fixed issue mesh streaming issues when using -onethread, which ended up running problematic rendercommands on async thread
TRefCountPtr<const FStaticMeshLODResources> while accessing the LODs.
FStaticMeshFilteredAreaWeightedSectionSampler now derives from FWeightedRandomSampler instead of FStaticMeshAreaWeightedSectionSampler.
Prevented UNiagaraDataInterfaceStaticMesh from creating GPU data unless there are some GPU emitters in the owner system.

diff --git a/Engine/Plugins/FX/Niagara/Source/Niagara/Classes/NiagaraDataInterfaceStaticMesh.h b/Engine/Plugins/FX/Niagara/Source/Niagara/Classes/NiagaraDataInterfaceStaticMesh.h
--- a/Engine/Plugins/FX/Niagara/Source/Niagara/Classes/NiagaraDataInterfaceStaticMesh.h
+++ b/Engine/Plugins/FX/Niagara/Source/Niagara/Classes/NiagaraDataInterfaceStaticMesh.h
@@ -10,13 +10,15 @@ struct FNDIStaticMesh_InstanceData;
 struct FNDIStaticMeshSectionFilter;
 
 /** Allows uniform random sampling of a number of mesh sections filtered by an FNDIStaticMeshSectionFilter */
-struct FStaticMeshFilteredAreaWeightedSectionSampler : FStaticMeshAreaWeightedSectionSampler
+struct FStaticMeshFilteredAreaWeightedSectionSampler : FWeightedRandomSampler
 {
 	FStaticMeshFilteredAreaWeightedSectionSampler();
 	void Init(const FStaticMeshLODResources* InRes, FNDIStaticMesh_InstanceData* InOwner);
-	virtual float GetWeights(TArray<float>& OutWeights)override;
 
 protected:
+
+	virtual float GetWeights(TArray<float>& OutWeights)override;
+
 	TRefCountPtr<const FStaticMeshLODResources> Res;
 	FNDIStaticMesh_InstanceData* Owner;
 };
@@ -78,8 +80,8 @@ protected:
 		uint32 Alias;
 	};
 
-	// Cached pointer to Section render data used for initialization only.
-	TRefCountPtr<const FStaticMeshLODResources> SectionRenderData = nullptr;
+	// Cached pointer to Section render data used for initialization only. This doesn't need to be ref counted since it doesn't reference CPU data.
+	const FStaticMeshLODResources* SectionRenderData = nullptr;
 
 	TArray<SectionInfo> ValidSections;					// Only the section we want to spawn from
 
@@ -153,7 +155,7 @@ struct FNDIStaticMesh_InstanceData
 	FORCEINLINE_DEBUGGABLE bool ResetRequired(UNiagaraDataInterfaceStaticMesh* Interface)const;
 
 	FORCEINLINE const TArray<int32>& GetValidSections()const { return ValidSections; }
-	FORCEINLINE const FStaticMeshAreaWeightedSectionSampler& GetAreaWeightedSampler() const { return Sampler; }
+	FORCEINLINE const FWeightedRandomSampler& GetAreaWeightedSampler() const { return Sampler; }
 
 	void InitVertexColorFiltering();
 
diff --git a/Engine/Plugins/FX/Niagara/Source/Niagara/Private/NiagaraDataInterfaceStaticMesh.cpp b/Engine/Plugins/FX/Niagara/Source/Niagara/Private/NiagaraDataInterfaceStaticMesh.cpp
--- a/Engine/Plugins/FX/Niagara/Source/Niagara/Private/NiagaraDataInterfaceStaticMesh.cpp
+++ b/Engine/Plugins/FX/Niagara/Source/Niagara/Private/NiagaraDataInterfaceStaticMesh.cpp
@@ -54,34 +54,38 @@ void FStaticMeshFilteredAreaWeightedSectionSampler::Init(const FStaticMeshLODRes
 	Res = InRes;
 	Owner = InOwner;
 
-	FStaticMeshAreaWeightedSectionSampler::Init(Res);
+	Initialize();
 }
 
 float FStaticMeshFilteredAreaWeightedSectionSampler::GetWeights(TArray<float>& OutWeights)
 {
-	check(Owner && Owner->Mesh);
 	float Total = 0.0f;
-	OutWeights.Empty(Owner->GetValidSections().Num());
-	if (Owner->Mesh->bSupportUniformlyDistributedSampling && Res->AreaWeightedSectionSamplers.Num() > 0)
+	if (Owner && Owner->Mesh)
 	{
-		for (int32 i = 0; i < Owner->GetValidSections().Num(); ++i)
+		OutWeights.Empty(Owner->GetValidSections().Num());
+		if (Owner->Mesh->bSupportUniformlyDistributedSampling && Res->AreaWeightedSectionSamplers.Num() > 0)
 		{
-			int32 SecIdx = Owner->GetValidSections()[i];
-			float T = Res->AreaWeightedSectionSamplers[SecIdx].GetTotalWeight();
-			OutWeights.Add(T);
-			Total += T;
+			for (int32 i = 0; i < Owner->GetValidSections().Num(); ++i)
+			{
+				int32 SecIdx = Owner->GetValidSections()[i];
+				float T = Res->AreaWeightedSectionSamplers[SecIdx].GetTotalWeight();
+				OutWeights.Add(T);
+				Total += T;
+			}
 		}
-	}
-	else
-	{
-		for (int32 i = 0; i < Owner->GetValidSections().Num(); ++i)
+		else
 		{
-			int32 SecIdx = Owner->GetValidSections()[i];
-			float T = 1.0f;
-			OutWeights.Add(T);
-			Total += T;
+			for (int32 i = 0; i < Owner->GetValidSections().Num(); ++i)
+			{
+				int32 SecIdx = Owner->GetValidSections()[i];
+				float T = 1.0f;
+				OutWeights.Add(T);
+				Total += T;
+			}
 		}
 
+		// Release the reference to the LODresource to avoid blocking stream out operations.
+		Res.SafeRelease();
 	}
 	return Total;
 }
@@ -117,7 +121,9 @@ void FStaticMeshGpuSpawnBuffer::Initialise(const FStaticMeshLODResources* Res, c
 	}
 
 	if (bIsGpuUniformlyDistributedSampling)
+	{
 		BufferUniformTriangleSamplingSRV = Res->AreaWeightedSectionSamplersBuffer.GetBufferSRV(); // Cache that SRV for later
+	}
 }
 
 void FStaticMeshGpuSpawnBuffer::InitRHI()
@@ -1230,7 +1236,7 @@ bool UNiagaraDataInterfaceStaticMesh::InitPerInstanceData(void* PerInstanceData,
 	if (bSuccess)
 	{
 		FStaticMeshGpuSpawnBuffer* MeshGpuSpawnBuffer = nullptr;
-		if (Inst->Mesh)
+		if (Inst->Mesh && SystemInstance->HasGPUEmitters())
 		{
 			// Always allocate when bAllowCPUAccess (index buffer can only have SRV created in this case as of today)
 			// We do not know if this interface is allocated for CPU or GPU so we allocate for both case... TODO: have some cached data created in case a GPU version is needed?
diff --git a/Engine/Source/Runtime/Engine/Private/StaticMesh.cpp b/Engine/Source/Runtime/Engine/Private/StaticMesh.cpp
--- a/Engine/Source/Runtime/Engine/Private/StaticMesh.cpp
+++ b/Engine/Source/Runtime/Engine/Private/StaticMesh.cpp
@@ -932,16 +932,23 @@ void FStaticMeshAreaWeightedSectionSampler::Init(const FStaticMeshLODResources*
 
 float FStaticMeshAreaWeightedSectionSampler::GetWeights(TArray<float>& OutWeights)
 {
-	//If this hits, you're trying to get weights on a sampler that's not been initialized.
-	check(Owner);
 	float Total = 0.0f;
-	OutWeights.Empty(Owner->Sections.Num());
-	for (int32 i = 0; i < Owner->Sections.Num(); ++i)
+
+	if (Owner)
 	{
-		float T = Owner->AreaWeightedSectionSamplers[i].GetTotalWeight();
-		OutWeights.Add(T);
-		Total += T;
+		//If this hits, you're trying to get weights on a sampler that's not been initialized.
+		OutWeights.Empty(Owner->Sections.Num());
+		for (int32 i = 0; i < Owner->Sections.Num(); ++i)
+		{
+			float T = Owner->AreaWeightedSectionSamplers[i].GetTotalWeight();
+			OutWeights.Add(T);
+			Total += T;
+		}
+
+		// Release the reference to the LODresource to avoid blocking stream out operations.
+		Owner.SafeRelease();
 	}
+
 	return Total;
 }
 
@@ -1096,22 +1103,7 @@ void FStaticMeshVertexBuffers::InitFromDynamicVertex(FLocalVertexFactory* Vertex
 		});
 };
 
-void FStaticMeshVertexBuffers::CopyRHIForStreaming(const FStaticMeshVertexBuffers& Other, bool InAllowCPUAccess)
-{
-	StaticMeshVertexBuffer.CopyRHIForStreaming(Other.StaticMeshVertexBuffer, InAllowCPUAccess);
-	PositionVertexBuffer.CopyRHIForStreaming(Other.PositionVertexBuffer, InAllowCPUAccess);
-	ColorVertexBuffer.CopyRHIForStreaming(Other.ColorVertexBuffer, InAllowCPUAccess);
-}
-
-void FAdditionalStaticMeshIndexBuffers::CopyRHIForStreaming(const FAdditionalStaticMeshIndexBuffers& Other, bool InAllowCPUAccess)
-{
-	ReversedIndexBuffer.CopyRHIForStreaming(Other.ReversedIndexBuffer, InAllowCPUAccess);
-	ReversedDepthOnlyIndexBuffer.CopyRHIForStreaming(Other.ReversedDepthOnlyIndexBuffer, InAllowCPUAccess);
-	WireframeIndexBuffer.CopyRHIForStreaming(Other.WireframeIndexBuffer, InAllowCPUAccess);
-	AdjacencyIndexBuffer.CopyRHIForStreaming(Other.AdjacencyIndexBuffer, InAllowCPUAccess);
-}
-
-FStaticMeshLODData::FStaticMeshLODData()
+FStaticMeshLODResources::FStaticMeshLODResources(bool bAddRef)
 	: MaxDeviation(0.0f)
 	, bHasAdjacencyInfo(false)
 	, bHasDepthOnlyIndices(false)
@@ -1128,6 +1120,10 @@ FStaticMeshLODData::FStaticMeshLODData()
 	, StaticMeshIndexMemory(0)
 #endif
 {
+	if (bAddRef)
+	{
+		AddRef();
+	}
 }
 
 FStaticMeshLODResources::~FStaticMeshLODResources()
@@ -1137,26 +1133,6 @@ FStaticMeshLODResources::~FStaticMeshLODResources()
 	delete AdditionalIndexBuffers;
 }
 
-void FStaticMeshLODResources::CopyRHIForStreaming(const FStaticMeshLODResources& Other, bool InAllowCPUAccess)
-{
-	VertexBuffers.CopyRHIForStreaming(Other.VertexBuffers, InAllowCPUAccess);
-	IndexBuffer.CopyRHIForStreaming(Other.IndexBuffer, InAllowCPUAccess);
-	DepthOnlyIndexBuffer.CopyRHIForStreaming(Other.DepthOnlyIndexBuffer, InAllowCPUAccess);
-	if (Other.AdditionalIndexBuffers)
-	{
-		if (!AdditionalIndexBuffers)
-		{
-			AdditionalIndexBuffers = new FAdditionalStaticMeshIndexBuffers;
-		}
-		AdditionalIndexBuffers->CopyRHIForStreaming(*Other.AdditionalIndexBuffers, InAllowCPUAccess);
-	}
-	else if (AdditionalIndexBuffers)
-	{
-		delete AdditionalIndexBuffers;
-		AdditionalIndexBuffers = nullptr;
-	}
-}
-
 template <bool bIncrement>
 void FStaticMeshLODResources::UpdateIndexMemoryStats()
 {
@@ -5615,14 +5591,8 @@ bool UStaticMesh::StreamOut(int32 NewMipCount)
 	check(IsInGameThread());
 	if (bIsStreamable && !PendingUpdate && RenderData.IsValid() && RenderData->bReadyForStreaming && NewMipCount < GetNumResidentMips())
 	{
-		if (FPlatformProperties::HasEditorOnlyData() || true) // Temporary disabling the Swap path.
-		{
-			PendingUpdate = new FStaticMeshStreamOut(this, NewMipCount);
-		}
-		else
-		{
-			PendingUpdate = new FStaticMeshStreamOut_Swap(this, NewMipCount);
-		}
+		// We need to keep the CPU data in non cook in order to be able for tools to work correctly.
+		PendingUpdate = new FStaticMeshStreamOut(this, NewMipCount, bAllowCPUAccess && !FPlatformProperties::HasEditorOnlyData());
 		return !PendingUpdate->IsCancelled();
 	}
 	return false;
diff --git a/Engine/Source/Runtime/Engine/Private/Streaming/StaticMeshUpdate.cpp b/Engine/Source/Runtime/Engine/Private/Streaming/StaticMeshUpdate.cpp
--- a/Engine/Source/Runtime/Engine/Private/Streaming/StaticMeshUpdate.cpp
+++ b/Engine/Source/Runtime/Engine/Private/Streaming/StaticMeshUpdate.cpp
@@ -12,11 +12,11 @@ StaticMeshUpdate.cpp: Helpers to stream in and out static mesh LODs.
 #include "Serialization/MemoryReader.h"
 #include "Streaming/RenderAssetUpdate.inl"
 
-int32 GStreamingMaxReferenceChecks = 3;
+int32 GStreamingMaxReferenceChecks = 2;
 static FAutoConsoleVariableRef CVarStreamingMaxReferenceChecksBeforeStreamOut(
 	TEXT("r.Streaming.MaxReferenceChecksBeforeStreamOut"),
 	GStreamingMaxReferenceChecks,
-	TEXT("Number of times the engine wait for references to be released before forcing streamout. (default=3)"),
+	TEXT("Number of times the engine wait for references to be released before forcing streamout. (default=2)"),
 	ECVF_Default
 );
 
@@ -273,121 +273,117 @@ void FStaticMeshStreamIn::DoCancel(const FContext& Context)
 	DoFinishUpdate(Context);
 }
 
-FStaticMeshStreamOut::FStaticMeshStreamOut(UStaticMesh* InMesh, int32 InRequestedMips)
+FStaticMeshStreamOut::FStaticMeshStreamOut(UStaticMesh* InMesh, int32 InRequestedMips, bool InDiscardCPUData)
 	: FStaticMeshUpdate(InMesh, InRequestedMips)
+	, bDiscardCPUData(InDiscardCPUData)
 {
-	PushTask(FContext(InMesh, TT_None), TT_Render, SRA_UPDATE_CALLBACK(DoReleaseBuffers), TT_None, nullptr);
-}
+	check(InMesh);
 
-void FStaticMeshStreamOut ::DoReleaseBuffers(const FContext& Context)
-{
-	check(Context.CurrentThread == TT_Render);
-	UStaticMesh* Mesh = Context.Mesh;
-	FStaticMeshRenderData* RenderData = Context.RenderData;
-	if (!IsCancelled() && Mesh && RenderData)
+	// Immediately change CurrentFirstLOdIdx to prevent new references from being made to streamed out lods.
+	FStaticMeshRenderData* RenderData = InMesh->RenderData.Get();
+	if (RenderData)
 	{
-		check(CurrentFirstLODIdx == RenderData->CurrentFirstLODIdx && PendingFirstMip > CurrentFirstLODIdx);
-		check(Mesh->GetCachedNumResidentLODs() == RenderData->LODResources.Num() - RenderData->CurrentFirstLODIdx);
-
+		InitialFirstLOD = RenderData->CurrentFirstLODIdx;
 		RenderData->CurrentFirstLODIdx = PendingFirstMip;
-		Mesh->SetCachedNumResidentLODs(static_cast<uint8>(RenderData->LODResources.Num() - PendingFirstMip));
+		InMesh->SetCachedNumResidentLODs(static_cast<uint8>(RenderData->LODResources.Num() - PendingFirstMip));
+	}
 
-		TRHIResourceUpdateBatcher<GStaticMeshMaxNumResourceUpdatesPerBatch> Batcher;
-		for (int32 LODIdx = CurrentFirstLODIdx; LODIdx < PendingFirstMip; ++LODIdx)
-		{
-			FStaticMeshLODResources& LODResource = RenderData->LODResources[LODIdx];
-			LODResource.DecrementMemoryStats();
-			LODResource.ReleaseRHIForStreaming(Batcher);
-#if RHI_RAYTRACING
-			if (IsRayTracingEnabled())
-			{
-				LODResource.RayTracingGeometry.ReleaseResource();
-			}
-#endif
-		}
+	if (InDiscardCPUData)
+	{
+		PushTask(FContext(InMesh, TT_None), TT_Async, SRA_UPDATE_CALLBACK(CheckReferencesAndDiscardCPUData), TT_Async, SRA_UPDATE_CALLBACK(Cancel));
+	}
+	else
+	{
+		PushTask(FContext(InMesh, TT_None), TT_Render, SRA_UPDATE_CALLBACK(ReleaseRHIBuffers), TT_Async, SRA_UPDATE_CALLBACK(Cancel));
 	}
 }
 
-
-FStaticMeshStreamOut_Swap::FStaticMeshStreamOut_Swap(UStaticMesh* InMesh, int32 InRequestedMips)
-	: FStaticMeshUpdate(InMesh, InRequestedMips)
+void FStaticMeshStreamOut::CheckReferencesAndDiscardCPUData(const FContext& Context)
 {
-	PushTask(FContext(InMesh, TT_None), TT_Render, SRA_UPDATE_CALLBACK(DoReleaseBuffers), TT_None, nullptr);
-}
+	DECLARE_SCOPE_CYCLE_COUNTER(TEXT("FStaticMeshStreamOut::CheckReferencesAndDiscardCPUData"), STAT_StaticMeshStreamOut_CheckReferencesAndDiscardCPUData, STATGROUP_StreamingDetails);
+	check(Context.CurrentThread == TT_Async);
 
-void FStaticMeshStreamOut_Swap::DoReleaseBuffers(const FContext& Context)
-{
-	check(Context.CurrentThread == TT_Render);
 	UStaticMesh* Mesh = Context.Mesh;
 	FStaticMeshRenderData* RenderData = Context.RenderData;
-	if (!IsCancelled() && Mesh && RenderData)
-	{
-		RenderData->CurrentFirstLODIdx = PendingFirstMip;
-		Mesh->SetCachedNumResidentLODs(static_cast<uint8>(RenderData->LODResources.Num() - PendingFirstMip));
+	uint32 NumExternalReferences = 0;
 
+	if (Mesh && RenderData)
+	{
 		for (int32 LODIdx = CurrentFirstLODIdx; LODIdx < PendingFirstMip; ++LODIdx)
 		{
-			FStaticMeshLODResources* PreviousLODResources = RenderData->LODResources.GetData()[LODIdx];
-			FStaticMeshLODResources* NextLODResources = new FStaticMeshLODResources;
-
-			// Copy everything but the CPU data.
-			static_cast<FStaticMeshLODData&>(*NextLODResources) = static_cast<FStaticMeshLODData&>(*PreviousLODResources);
-			NextLODResources->CopyRHIForStreaming(*PreviousLODResources, false);
+			// Minus 1 since the LODResources reference is not considered external
+			NumExternalReferences += RenderData->LODResources[LODIdx].GetRefCount() - 1;
+		}
 
-			// Change which LODResource the static mesh render data points to.
-			// This will prevent new references from being made.
-			RenderData->LODResources.GetData()[LODIdx] = NextLODResources; 
-			LODResourcesToRelease.Add(PreviousLODResources);
+		if (NumExternalReferences > PreviousNumberOfExternalReferences && NumReferenceChecks > 0)
+		{
+			PreviousNumberOfExternalReferences = NumExternalReferences;
+			UE_LOG(LogStaticMesh, Warning, TEXT("[%s] Streamed out LODResources got referenced while in pending stream out."), *Mesh->GetName());
 		}
-		DoReleaseLODResources(Context);
 	}
-}
-
-void FStaticMeshStreamOut_Swap::DoReleaseLODResources(const FContext& Context)
-{
-	check(Context.CurrentThread == TT_Render);
 
-	bool bHasExternalReferences = false;
-	for (FStaticMeshLODResources& LODResource : LODResourcesToRelease)
+	if (!NumExternalReferences || NumReferenceChecks >= GStreamingMaxReferenceChecks)
 	{
-		if (LODResource.GetRefCount() != 1)
+		for (int32 LODIdx = CurrentFirstLODIdx; LODIdx < PendingFirstMip; ++LODIdx)
 		{
-			bHasExternalReferences = true;
-			++NumReferenceChecks;
-			break;
+			RenderData->LODResources[LODIdx].DiscardCPUData();
 		}
-	}
 
-	// It is expected that systems that keep a persistent reference to streamed out LODs will detect the streamout
-	// and recache their loddata. This is why we still streamout here after X attemps, assuming there is no active access on the loddata.
-	if (!bHasExternalReferences || NumReferenceChecks >= GStreamingMaxReferenceChecks)
-	{
-		ReleaseLODResources();
+		// Because we discarded the CPU data, the stream out can not be cancelled anymore.
+		PushTask(Context, TT_Render, SRA_UPDATE_CALLBACK(ReleaseRHIBuffers), TT_Render, SRA_UPDATE_CALLBACK(ReleaseRHIBuffers));
 	}
 	else
 	{
+		++NumReferenceChecks;
+		if (NumReferenceChecks >= GStreamingMaxReferenceChecks)
+		{
+			UE_LOG(LogStaticMesh, Log, TEXT("[%s] Streamed out LODResources references are not getting released."), *Mesh->GetName());
+		}
+
 		bDeferExecution = true;
-		PushTask(Context, TT_Render, SRA_UPDATE_CALLBACK(DoReleaseLODResources), TT_Render, SRA_UPDATE_CALLBACK(DoReleaseLODResources));
+		PushTask(Context, TT_Async, SRA_UPDATE_CALLBACK(CheckReferencesAndDiscardCPUData), TT_Async, SRA_UPDATE_CALLBACK(Cancel));
+	}
+}
+
+void FStaticMeshStreamOut::ReleaseRHIBuffers(const FContext& Context)
+{
+	DECLARE_SCOPE_CYCLE_COUNTER(TEXT("FStaticMeshStreamOut::ReleaseRHIBuffers"), STAT_StaticMeshStreamOut_ReleaseRHIBuffers, STATGROUP_StreamingDetails);
+	check(Context.CurrentThread == TT_Render);
+
+	UStaticMesh* Mesh = Context.Mesh;
+	FStaticMeshRenderData* RenderData = Context.RenderData;
+	if (Mesh && RenderData)
+	{
+		TRHIResourceUpdateBatcher<GStaticMeshMaxNumResourceUpdatesPerBatch> Batcher;
+		for (int32 LODIdx = CurrentFirstLODIdx; LODIdx < PendingFirstMip; ++LODIdx)
+		{
+			FStaticMeshLODResources& LODResource = RenderData->LODResources[LODIdx];
+			LODResource.DecrementMemoryStats();
+			LODResource.ReleaseRHIForStreaming(Batcher);
+			
+#if RHI_RAYTRACING
+			if (IsRayTracingEnabled())
+			{
+				LODResource.RayTracingGeometry.ReleaseResource();
+			}
+#endif
+		}
 	}
 }
 
-void FStaticMeshStreamOut_Swap::ReleaseLODResources()
+void FStaticMeshStreamOut::Cancel(const FContext& Context)
 {
-	if (LODResourcesToRelease.Num())
+	DECLARE_SCOPE_CYCLE_COUNTER(TEXT("FStaticMeshStreamOut::Cancel"), STAT_StaticMeshStreamOut_Cancel, STATGROUP_StreamingDetails);
+	check(Context.CurrentThread == TT_Async);
+
+	UStaticMesh* Mesh = Context.Mesh;
+	FStaticMeshRenderData* RenderData = Context.RenderData;
+	if (Mesh && RenderData)
 	{
-	    TRHIResourceUpdateBatcher<GStaticMeshMaxNumResourceUpdatesPerBatch> Batcher;
-	    FStaticMeshLODResources** LODResourcesArray = LODResourcesToRelease.GetData();
-	    for (int32 LODIndex = 0; LODIndex < LODResourcesToRelease.Num(); ++LODIndex)
-	    {
-		    FStaticMeshLODResources* LODResources = LODResourcesArray[LODIndex];
-		    LODResourcesArray[LODIndex] = nullptr;
-    
-		    LODResources->ReleaseRHIForStreaming(Batcher);
-		    LODResources->ReleaseResources(); // Also decrement memory stats.
-		    LODResources->Release();
-	    }
-	    LODResourcesToRelease.Empty();
+		RenderData->CurrentFirstLODIdx = InitialFirstLOD;
+		Mesh->SetCachedNumResidentLODs(static_cast<uint8>(RenderData->LODResources.Num() - InitialFirstLOD));
 	}
+
 }
 
 void FStaticMeshStreamIn_IO::FCancelIORequestsTask::DoWork()
diff --git a/Engine/Source/Runtime/Engine/Private/Streaming/StaticMeshUpdate.h b/Engine/Source/Runtime/Engine/Private/Streaming/StaticMeshUpdate.h
--- a/Engine/Source/Runtime/Engine/Private/Streaming/StaticMeshUpdate.h
+++ b/Engine/Source/Runtime/Engine/Private/Streaming/StaticMeshUpdate.h
@@ -126,35 +126,19 @@ private:
 class FStaticMeshStreamOut : public FStaticMeshUpdate
 {
 public:
-	FStaticMeshStreamOut(UStaticMesh* InMesh, int32 InRequestedMips);
+	FStaticMeshStreamOut(UStaticMesh* InMesh, int32 InRequestedMips, bool InDiscardCPUData);
 
 private:
-	/** Release RHI buffers and update SRVs */
-	void DoReleaseBuffers(const FContext& Context);
-};
-
-class FStaticMeshStreamOut_Swap : public FStaticMeshUpdate
-{
-public:
-	FStaticMeshStreamOut_Swap(UStaticMesh* InMesh, int32 InRequestedMips);
-
-	virtual ~FStaticMeshStreamOut_Swap() { ReleaseLODResources(); }
 
-private:
+	void CheckReferencesAndDiscardCPUData(const FContext& Context);
+	void ReleaseRHIBuffers(const FContext& Context);
+	/** Restore */
+	void Cancel(const FContext& Context);
 
+	uint8 InitialFirstLOD = 0;
+	bool bDiscardCPUData = false;
 	int32 NumReferenceChecks = 0;
-	TIndirectArray<FStaticMeshLODResources>  LODResourcesToRelease;
-
-	/** Release RHI buffers and update SRVs */
-	void DoReleaseBuffers(const FContext& Context);
-
-	/** Check reference and possibly release LODResourcesToRelease */
-	void DoReleaseLODResources(const FContext& Context);
-
-	/** Release all LODs in LODResourcesToRelease */
-	void ReleaseLODResources();
-
-
+	uint32 PreviousNumberOfExternalReferences = 0;
 };
 
 class FStaticMeshStreamIn_IO : public FStaticMeshStreamIn
diff --git a/Engine/Source/Runtime/Engine/Public/StaticMeshResources.h b/Engine/Source/Runtime/Engine/Public/StaticMeshResources.h
--- a/Engine/Source/Runtime/Engine/Public/StaticMeshResources.h
+++ b/Engine/Source/Runtime/Engine/Public/StaticMeshResources.h
@@ -244,9 +244,11 @@ struct ENGINE_API FStaticMeshAreaWeightedSectionSampler : FWeightedRandomSampler
 {
 	FStaticMeshAreaWeightedSectionSampler();
 	void Init(const FStaticMeshLODResources* InOwner);
-	virtual float GetWeights(TArray<float>& OutWeights)override;
 
 protected:
+
+	virtual float GetWeights(TArray<float>& OutWeights)override;
+
 	TRefCountPtr<const FStaticMeshLODResources> Owner;
 };
 
@@ -308,9 +310,6 @@ struct FStaticMeshVertexBuffers
 
 	/* This is a temporary function to refactor and convert old code, do not copy this as is and try to build your data as SoA from the beginning.*/
 	void ENGINE_API InitModelVF(FLocalVertexFactory* VertexFactory);
-
-	/** Copy everything, keeping reference to the same RHI resources. */
-	void CopyRHIForStreaming(const FStaticMeshVertexBuffers& Other, bool InAllowCPUAccess);
 };
 
 struct FAdditionalStaticMeshIndexBuffers
@@ -323,15 +322,16 @@ struct FAdditionalStaticMeshIndexBuffers
 	FRawStaticIndexBuffer WireframeIndexBuffer;
 	/** Index buffer containing adjacency information required by tessellation. */
 	FRawStaticIndexBuffer AdjacencyIndexBuffer;
-
-	/** Copy everything, keeping reference to the same RHI resources. */
-	void CopyRHIForStreaming(const FAdditionalStaticMeshIndexBuffers& Other, bool InAllowCPUAccess);
 };
 
-/** The part of FStaticMeshLODResources which is not streamable. */
-struct FStaticMeshLODData
+/** 
+ * Rendering resources needed to render an individual static mesh LOD.
+ * This structure is ref counted to allow the LOD streamer to evaluate the number of readers to it (readers that could access the CPU data).
+ * Because the stream out clears the CPU readable data, CPU code that samples it must ensure to only reference LODs above CurrentFirstLODIdx.
+ */
+struct FStaticMeshLODResources : public FRefCountBase
 {
-	FStaticMeshLODData();
+public:
 
 	/** Sections for this LOD. */
 	using FStaticMeshSectionArray = TArray<FStaticMeshSection, TInlineAllocator<1>>;
@@ -387,22 +387,6 @@ struct FStaticMeshLODData
 	/** Map of wedge index to vertex index. Each LOD need one*/
 	TArray<int32> WedgeMap;
 #endif
-};
-
-/** 
- * Rendering resources needed to render an individual static mesh LOD.
- * They are ref counted to prevent streaming out LODs while they are being accessed for read.
- * The stream out implies swapping the LODResoures entry by a dummy empty one. This is always executed within a render command.
- * This means that any direct reference to a LODResources is safe within the scope of any render commands.
- * If the reference is required on the gamethread, async threads or RHI threads (not synchronized at the end of the render command that created the tasks),
- * then the ref count must be incremented / decremented through the use of a TRefCountPtr. 
- * The stream out logic in FStaticMeshStreamOut will postpone the operation if the ref count is not 1.
- * "1" accounting for the reference in the FStaticMeshRenderData::LODResources. In order to prevent postponing indefinitely the streamout, 
- * it is required that anyone using a TRefCountPtr also check for CurrentFirstLODIdx for things that may run every frame.
- */
-struct FStaticMeshLODResources : public FRefCountBase, public FStaticMeshLODData
-{
-public:
 
 	FStaticMeshVertexBuffers VertexBuffers;
 
@@ -425,19 +409,10 @@ public:
 	FStaticMeshSectionAreaWeightedTriangleSamplerBuffer AreaWeightedSectionSamplersBuffer;
 	
 	/** Default constructor. Add a reference if not stored with a TRefCountPtr */
-	FORCEINLINE FStaticMeshLODResources(bool bAddRef = true)
-	{
-		if (bAddRef)
-		{
-			AddRef();
-		}
-	}
+	ENGINE_API FStaticMeshLODResources(bool bAddRef = true);
 
 	ENGINE_API ~FStaticMeshLODResources();
 
-	/** Copy everything, keeping reference to the same RHI resources. */
-	void CopyRHIForStreaming(const FStaticMeshLODResources& Other, bool InAllowCPUAccess);
-
 	template <typename TBatcher>
 	void ReleaseRHIForStreaming(TBatcher& Batcher)
 	{
