Prevented FDistanceFieldAsyncQueue from recreating renderstates on staticmeshes that are being updated at the same time.

diff --git a/Engine/Source/Runtime/Engine/Private/DistanceFieldAtlas.cpp b/Engine/Source/Runtime/Engine/Private/DistanceFieldAtlas.cpp
--- a/Engine/Source/Runtime/Engine/Private/DistanceFieldAtlas.cpp
+++ b/Engine/Source/Runtime/Engine/Private/DistanceFieldAtlas.cpp
@@ -1191,6 +1191,8 @@ void FDistanceFieldAsyncQueue::ProcessAsyncTasks()
 			Task->GeneratedVolumeData->VolumeTexture.Initialize(Task->StaticMesh);
 			FDistanceFieldVolumeData* OldVolumeData = Task->StaticMesh->RenderData->LODResources[0].DistanceFieldData;
 
+			// Renderstates are not initialized between UStaticMesh::PreEditChange() and UStaticMesh::PostEditChange()
+			if (Task->StaticMesh->RenderData->IsInitialized())
 			{
 				// Cause all components using this static mesh to get re-registered, which will recreate their proxies and primitive uniform buffers
 				FStaticMeshComponentRecreateRenderStateContext RecreateRenderStateContext(Task->StaticMesh, false);
@@ -1198,6 +1200,11 @@ void FDistanceFieldAsyncQueue::ProcessAsyncTasks()
 				// Assign the new volume data
 				Task->StaticMesh->RenderData->LODResources[0].DistanceFieldData = Task->GeneratedVolumeData;
 			}
+			else
+			{
+				// Assign the new volume data
+				Task->StaticMesh->RenderData->LODResources[0].DistanceFieldData = Task->GeneratedVolumeData;
+			}
 
 			OldVolumeData->VolumeTexture.Release();
 
