Add missing include in ShapeElem.h to fix CIS engine compile

diff --git a/Engine/Source/Runtime/Engine/Classes/PhysicsEngine/ShapeElem.h b/Engine/Source/Runtime/Engine/Classes/PhysicsEngine/ShapeElem.h
--- a/Engine/Source/Runtime/Engine/Classes/PhysicsEngine/ShapeElem.h
+++ b/Engine/Source/Runtime/Engine/Classes/PhysicsEngine/ShapeElem.h
@@ -7,6 +7,7 @@
 #include "EngineDefines.h"
 #include "PhysxUserData.h"
 #include "PhysicsCore/Public/PhysicsInterfaceTypesCore.h"
+#include "Engine/EngineTypes.h"
 #include "ShapeElem.generated.h"
 
 namespace EAggCollisionShape
