Fix typo in CreateGeometry causing filter bug

diff --git a/Engine/Source/Runtime/Engine/Private/PhysicsEngine/Experimental/ChaosInterfaceUtils.cpp b/Engine/Source/Runtime/Engine/Private/PhysicsEngine/Experimental/ChaosInterfaceUtils.cpp
--- a/Engine/Source/Runtime/Engine/Private/PhysicsEngine/Experimental/ChaosInterfaceUtils.cpp
+++ b/Engine/Source/Runtime/Engine/Private/PhysicsEngine/Experimental/ChaosInterfaceUtils.cpp
@@ -188,8 +188,8 @@ namespace ChaosInterface
 			NewShape->SetUserData(UserData);
 			NewShape->SetSimulate(bComplexShape ? InParams.CollisionData.CollisionFlags.bEnableSimCollisionComplex : InParams.CollisionData.CollisionFlags.bEnableSimCollisionSimple);
 #if WITH_CHAOS
-			NewShape->SetSimEnabled(ShapeCollisionEnabled == ECollisionEnabled::QueryAndPhysics || ECollisionEnabled::PhysicsOnly);
-			NewShape->SetQueryEnabled(ShapeCollisionEnabled == ECollisionEnabled::QueryAndPhysics || ECollisionEnabled::QueryOnly);
+			NewShape->SetSimEnabled(ShapeCollisionEnabled == ECollisionEnabled::QueryAndPhysics || ShapeCollisionEnabled == ECollisionEnabled::PhysicsOnly);
+			NewShape->SetQueryEnabled(ShapeCollisionEnabled == ECollisionEnabled::QueryAndPhysics || ShapeCollisionEnabled == ECollisionEnabled::QueryOnly);
 #endif
 			return NewShape;
 		};
