Added maximum number of debug strings per actor for HUD's AddDebugText to avoid overly long frame when possessing the player after having ejected from it for a long time in PIE (the HUD strings are not flushed during that time)

diff --git a/Engine/Source/Runtime/Engine/Classes/GameFramework/HUD.h b/Engine/Source/Runtime/Engine/Classes/GameFramework/HUD.h
--- a/Engine/Source/Runtime/Engine/Classes/GameFramework/HUD.h
+++ b/Engine/Source/Runtime/Engine/Classes/GameFramework/HUD.h
@@ -111,6 +111,7 @@ protected:
 	UPROPERTY()
 	UCanvas* DebugCanvas;
 
+	/** List of debug strings attached to actors, sorted by actor first, then by order of addition */
 	UPROPERTY()
 	TArray<struct FDebugTextInfo> DebugTextList;
 
@@ -588,6 +589,11 @@ protected:
 private:
 	// Helper function to deal with screen offset and splitscreen mapping of coordinates to HUD
 	FVector2D GetCoordinateOffset() const;
+
+	/** Helper function to find the first and last index (if any) of elements from DebugTextList with the matching actor : 
+	* @return number of elements found.
+	*/
+	int32 FindDebugTextListIntervalForActor(AActor* InSrcActor, int32& OutFirstIdx, int32& OutLastIdx) const;
 };
 
 
diff --git a/Engine/Source/Runtime/Engine/Private/HUD.cpp b/Engine/Source/Runtime/Engine/Private/HUD.cpp
--- a/Engine/Source/Runtime/Engine/Private/HUD.cpp
+++ b/Engine/Source/Runtime/Engine/Private/HUD.cpp
@@ -56,6 +56,11 @@ TAutoConsoleVariable<float> GSafeZoneVisualizationAlphaCVar(
 	TEXT("The alpha value of the safe zone overlay (0..1)\n")
 	TEXT(" default: 0.2"));
 
+TAutoConsoleVariable<int32> GMaxDebugTextStringsPerActorCVar(
+	TEXT("r.DebugSafeZone.MaxDebugTextStringsPerActor"),
+	128,
+	TEXT("The maximum number of debug strings that can be attached to a given actor (<=0 : no limit)"));
+
 const FColor AHUD::WhiteColor(255, 255, 255, 255);
 const FColor AHUD::GreenColor(0, 255, 0, 255);
 const FColor AHUD::RedColor(255, 0, 0, 255);
@@ -135,8 +140,8 @@ FVector2D AHUD::GetCoordinateOffset() const
 
 		if (SceneView)
 		{
-			Offset.X = - SceneView->UnscaledViewRect.Min.X; // And this will deal with the viewport offset if its a split screen
-			Offset.Y = - SceneView->UnscaledViewRect.Min.Y;
+			Offset.X = -SceneView->UnscaledViewRect.Min.X; // And this will deal with the viewport offset if its a split screen
+			Offset.Y = -SceneView->UnscaledViewRect.Min.Y;
 		}
 	}
 
@@ -146,20 +151,20 @@ FVector2D AHUD::GetCoordinateOffset() const
 void AHUD::PostRender()
 {
 	// Theres nothing we can really do without a canvas or a world - so leave now in that case
-	if ( (GetWorld() == nullptr) || (Canvas == nullptr))
+	if ((GetWorld() == nullptr) || (Canvas == nullptr))
 	{
 		return;
 	}
 	// Set up delta time
 	RenderDelta = GetWorld()->TimeSeconds - LastHUDRenderTime;
 
-	if ( PlayerOwner != NULL )
+	if (PlayerOwner != nullptr)
 	{
 		// draw any debug text in real-time
 		DrawDebugTextList();
 	}
 
-	if ( bShowDebugInfo )
+	if (bShowDebugInfo)
 	{
 		if (DebugCanvas)
 		{
@@ -167,14 +172,14 @@ void AHUD::PostRender()
 			ShowDebugInfo(DebugCanvas->DisplayDebugManager.GetMaxCharHeightRef(), DebugCanvas->DisplayDebugManager.GetYPosRef());
 		}
 	}
-	else if ( bShowHUD && FApp::CanEverRender() )
+	else if (bShowHUD && FApp::CanEverRender())
 	{
 		DrawHUD();
-		
+
 		// No need to do work to determine hit box candidates if there will never be any
 		if (HitBoxMap.Num() > 0)
 		{
-			ULocalPlayer* LocalPlayer = GetOwningPlayerController() ? Cast<ULocalPlayer>(GetOwningPlayerController()->Player) : NULL;
+			ULocalPlayer* LocalPlayer = GetOwningPlayerController() ? Cast<ULocalPlayer>(GetOwningPlayerController()->Player) : nullptr;
 
 			if (LocalPlayer && LocalPlayer->ViewportClient)
 			{
@@ -211,7 +216,7 @@ void AHUD::PostRender()
 						ContactPoint += ContactPointOffset;
 					}
 				}
-				UpdateHitBoxCandidates( MoveTemp(ContactPoints) );
+				UpdateHitBoxCandidates(MoveTemp(ContactPoints));
 			}
 		}
 		else if (HitBoxesOver.Num() > 0)
@@ -224,10 +229,10 @@ void AHUD::PostRender()
 			HitBoxesOver.Reset();
 		}
 	}
-	
-	if( bShowHitBoxDebugInfo )
+
+	if (bShowHitBoxDebugInfo)
 	{
-		RenderHitBoxes( Canvas->Canvas );
+		RenderHitBoxes(Canvas->Canvas);
 	}
 
 	DrawSafeZoneOverlay();
@@ -244,14 +249,14 @@ void AHUD::DrawActorOverlays(FVector Viewpoint, FRotator ViewRotation)
 	int32 i = 0;
 	while (i < PostRenderedActors.Num())
 	{
-		if ( PostRenderedActors[i] != NULL )
+		if (PostRenderedActors[i] != nullptr)
 		{
-			PostRenderedActors[i]->PostRenderFor(PlayerOwner,Canvas,Viewpoint,ViewDir);
+			PostRenderedActors[i]->PostRenderFor(PlayerOwner, Canvas, Viewpoint, ViewDir);
 			i++;
 		}
 		else
 		{
-			PostRenderedActors.RemoveAt(i,1);
+			PostRenderedActors.RemoveAt(i, 1);
 		}
 	}
 }
@@ -301,11 +306,11 @@ void AHUD::DrawSafeZoneOverlay()
 
 void AHUD::RemovePostRenderedActor(AActor* A)
 {
-	for ( int32 i=0; i<PostRenderedActors.Num(); i++ )
+	for (int32 i = 0; i < PostRenderedActors.Num(); i++)
 	{
-		if ( PostRenderedActors[i] == A )
+		if (PostRenderedActors[i] == A)
 		{
-			PostRenderedActors[i] = NULL;
+			PostRenderedActors[i] = nullptr;
 			return;
 		}
 	}
@@ -314,18 +319,18 @@ void AHUD::RemovePostRenderedActor(AActor* A)
 void AHUD::AddPostRenderedActor(AActor* A)
 {
 	// make sure that A is not already in list
-	for ( int32 i=0; i<PostRenderedActors.Num(); i++ )
+	for (int32 i = 0; i < PostRenderedActors.Num(); i++)
 	{
-		if ( PostRenderedActors[i] == A )
+		if (PostRenderedActors[i] == A)
 		{
 			return;
 		}
 	}
 
 	// add A at first empty slot
-	for ( int32 i=0; i<PostRenderedActors.Num(); i++ )
+	for (int32 i = 0; i < PostRenderedActors.Num(); i++)
 	{
-		if ( PostRenderedActors[i] == NULL )
+		if (PostRenderedActors[i] == nullptr)
 		{
 			PostRenderedActors[i] = A;
 			return;
@@ -502,7 +507,7 @@ AActor* AHUD::GetCurrentDebugTargetActor()
 		// Otherwise fall back to our last successful hit.
 		return DebugTargetActor ? DebugTargetActor : ShowDebugTargetActor;
 	}
-	else 
+	else
 	{
 		// Otherwise we use our Cached DebugTargetActor.
 		DebugTargetActor = ShowDebugTargetActor;
@@ -635,7 +640,7 @@ void AHUD::DrawHUD()
 {
 	HitBoxMap.Reset();
 	HitBoxHits.Reset();
-	if ( bShowOverlays && (PlayerOwner != NULL) )
+	if (bShowOverlays && (PlayerOwner != nullptr))
 	{
 		FVector ViewPoint;
 		FRotator ViewRotation;
@@ -662,10 +667,10 @@ UFont* AHUD::GetFontFromSizeIndex(int32 FontSizeIndex) const
 
 void AHUD::OnLostFocusPause(bool bEnable)
 {
-	if ( bLostFocusPaused == bEnable )
+	if (bLostFocusPaused == bEnable)
 		return;
 
-	if ( GetNetMode() != NM_Client )
+	if (GetNetMode() != NM_Client)
 	{
 		bLostFocusPaused = bEnable;
 		PlayerOwner->SetPause(bEnable);
@@ -674,22 +679,22 @@ void AHUD::OnLostFocusPause(bool bEnable)
 
 void AHUD::DrawDebugTextList()
 {
-	if( (DebugTextList.Num() > 0) && (DebugCanvas != nullptr ) )
+	if ((DebugTextList.Num() > 0) && (DebugCanvas != nullptr))
 	{
 		FRotator CameraRot;
 		FVector CameraLoc;
 		PlayerOwner->GetPlayerViewPoint(CameraLoc, CameraRot);
 
-		FCanvasTextItem TextItem( FVector2D::ZeroVector, FText::GetEmpty(), GEngine->GetSmallFont(), FLinearColor::White );
+		FCanvasTextItem TextItem(FVector2D::ZeroVector, FText::GetEmpty(), GEngine->GetSmallFont(), FLinearColor::White);
 		for (int32 Idx = 0; Idx < DebugTextList.Num(); Idx++)
 		{
-			if (DebugTextList[Idx].SrcActor == NULL)
+			if (DebugTextList[Idx].SrcActor == nullptr)
 			{
-				DebugTextList.RemoveAt(Idx--,1);
+				DebugTextList.RemoveAt(Idx--, 1);
 				continue;
 			}
 
-			if( DebugTextList[Idx].Font != NULL )
+			if (DebugTextList[Idx].Font != nullptr)
 			{
 				TextItem.Font = DebugTextList[Idx].Font;
 			}
@@ -708,7 +713,7 @@ void AHUD::DrawDebugTextList()
 			{
 				FVector Offset = FMath::Lerp(DebugTextList[Idx].SrcActorOffset, DebugTextList[Idx].SrcActorDesiredOffset, Alpha);
 
-				if( DebugTextList[Idx].bKeepAttachedToActor )
+				if (DebugTextList[Idx].bKeepAttachedToActor)
 				{
 					WorldTextLoc = DebugTextList[Idx].SrcActor->GetActorLocation() + Offset;
 				}
@@ -726,15 +731,15 @@ void AHUD::DrawDebugTextList()
 			{
 				TextItem.DisableShadow();
 			}
-				
+
 			// don't draw text behind the camera
-			if ( ((WorldTextLoc - CameraLoc) | CameraRot.Vector()) > 0.f )
+			if (((WorldTextLoc - CameraLoc) | CameraRot.Vector()) > 0.f)
 			{
 				FVector ScreenLoc = Canvas->Project(WorldTextLoc);
-				TextItem.SetColor( DebugTextList[Idx].TextColor );
-				TextItem.Text = FText::FromString( DebugTextList[Idx].DebugText );
-				TextItem.Scale = FVector2D( DebugTextList[Idx].FontScale, DebugTextList[Idx].FontScale);
-				DebugCanvas->DrawItem( TextItem, FVector2D( FMath::CeilToFloat(ScreenLoc.X), FMath::CeilToFloat(ScreenLoc.Y) ) );
+				TextItem.SetColor(DebugTextList[Idx].TextColor);
+				TextItem.Text = FText::FromString(DebugTextList[Idx].DebugText);
+				TextItem.Scale = FVector2D(DebugTextList[Idx].FontScale, DebugTextList[Idx].FontScale);
+				DebugCanvas->DrawItem(TextItem, FVector2D(FMath::CeilToFloat(ScreenLoc.X), FMath::CeilToFloat(ScreenLoc.Y)));
 			}
 
 			// do this at the end so even small durations get at least one frame
@@ -760,28 +765,46 @@ void AHUD::DrawDebugTextList()
 
 /// @cond DOXYGEN_WARNINGS
 
+int32 AHUD::FindDebugTextListIntervalForActor(AActor* InSrcActor, int32& OutFirstIdx, int32& OutLastIdx) const
+{
+	OutFirstIdx = DebugTextList.IndexOfByPredicate([=](const FDebugTextInfo& Element) { return (Element.SrcActor == InSrcActor); });
+	OutLastIdx = OutFirstIdx;
+	if (OutFirstIdx == INDEX_NONE)
+	{
+		return 0;
+	}
+
+	const int32 NumElements = DebugTextList.Num();
+	while (((OutLastIdx + 1) < NumElements) && (DebugTextList[OutLastIdx + 1].SrcActor == InSrcActor))
+	{
+		++OutLastIdx;
+	}
+
+	return (OutLastIdx - OutFirstIdx + 1);
+}
+
 void AHUD::AddDebugText_Implementation(const FString& DebugText,
-										 AActor* SrcActor,
-										 float Duration,
-										 FVector Offset,
-										 FVector DesiredOffset,
-										 FColor TextColor,
-										 bool bSkipOverwriteCheck,
-										 bool bAbsoluteLocation,
-										 bool bKeepAttachedToActor,
-										 UFont* InFont,
-										 float FontScale,
-										 bool bDrawShadow
-										 )
+	AActor* SrcActor,
+	float Duration,
+	FVector Offset,
+	FVector DesiredOffset,
+	FColor TextColor,
+	bool bSkipOverwriteCheck,
+	bool bAbsoluteLocation,
+	bool bKeepAttachedToActor,
+	UFont* InFont,
+	float FontScale,
+	bool bDrawShadow
+)
 {
 	// set a default color
 	if (TextColor == FColor::Transparent)
 	{
 		TextColor = FColor::White;
 	}
-	
+
 	// and a default source actor of our pawn
-	if (SrcActor != NULL)
+	if (SrcActor != nullptr)
 	{
 		if (DebugText.Len() == 0)
 		{
@@ -789,60 +812,78 @@ void AHUD::AddDebugText_Implementation(const FString& DebugText,
 		}
 		else
 		{
-			// search for an existing entry
-			int32 Idx = 0;
-			if (!bSkipOverwriteCheck)
+			FDebugTextInfo* DebugTextInfo = nullptr;
+
+			// find existing elements corresponding to this actor if any :
+			int32 FirstIdx, LastIdx;
+			int32 NumElements = FindDebugTextListIntervalForActor(SrcActor, FirstIdx, LastIdx);
+			if (NumElements > 0)
 			{
-				Idx = INDEX_NONE;
-				for (int32 i = 0; i < DebugTextList.Num() && Idx == INDEX_NONE; ++i)
+				if (!bSkipOverwriteCheck)
 				{
-					if (DebugTextList[i].SrcActor == SrcActor)
-					{
-						Idx = i;
-					}
+					// just replace the existing entry :
+					DebugTextInfo = &DebugTextList[FirstIdx];
 				}
-				if (Idx == INDEX_NONE)
+				else
 				{
-					// manually grow the array one struct element
-					Idx = DebugTextList.Add(FDebugTextInfo());
+					// if there's a max number per actor, respect it : 
+					int32 MaxNumStrings = GMaxDebugTextStringsPerActorCVar.GetValueOnGameThread();
+					if ((MaxNumStrings <= 0) || (NumElements < MaxNumStrings))
+					{
+						DebugTextInfo = &DebugTextList.EmplaceAt_GetRef(LastIdx + 1);
+					}
 				}
 			}
 			else
 			{
-				Idx = DebugTextList.Add(FDebugTextInfo());
+				DebugTextInfo = &DebugTextList.Emplace_GetRef();
+			}
+
+			if (DebugTextInfo != nullptr)
+			{
+				// assign the new text and actor
+				DebugTextInfo->SrcActor = SrcActor;
+				DebugTextInfo->SrcActorOffset = Offset;
+				DebugTextInfo->SrcActorDesiredOffset = DesiredOffset;
+				DebugTextInfo->DebugText = DebugText;
+				DebugTextInfo->TimeRemaining = Duration;
+				DebugTextInfo->Duration = Duration;
+				DebugTextInfo->TextColor = TextColor;
+				DebugTextInfo->bAbsoluteLocation = bAbsoluteLocation;
+				DebugTextInfo->bKeepAttachedToActor = bKeepAttachedToActor;
+				DebugTextInfo->OrigActorLocation = SrcActor->GetActorLocation();
+				DebugTextInfo->Font = InFont;
+				DebugTextInfo->FontScale = FontScale;
+				DebugTextInfo->bDrawShadow = bDrawShadow;
 			}
-			// assign the new text and actor
-			DebugTextList[Idx].SrcActor = SrcActor;
-			DebugTextList[Idx].SrcActorOffset = Offset;
-			DebugTextList[Idx].SrcActorDesiredOffset = DesiredOffset;
-			DebugTextList[Idx].DebugText = DebugText;
-			DebugTextList[Idx].TimeRemaining = Duration;
-			DebugTextList[Idx].Duration = Duration;
-			DebugTextList[Idx].TextColor = TextColor;
-			DebugTextList[Idx].bAbsoluteLocation = bAbsoluteLocation;
-			DebugTextList[Idx].bKeepAttachedToActor = bKeepAttachedToActor;
-			DebugTextList[Idx].OrigActorLocation = SrcActor->GetActorLocation();
-			DebugTextList[Idx].Font = InFont;
-			DebugTextList[Idx].FontScale = FontScale;
-			DebugTextList[Idx].bDrawShadow = bDrawShadow;
 		}
 	}
 }
 
-/** Remove debug text for the specific actor. */
+/** Remove all debug text for the specific actor. */
 void AHUD::RemoveDebugText_Implementation(AActor* SrcActor, bool bLeaveDurationText)
 {
-	int32 Idx = INDEX_NONE;
-	for (int32 i = 0; i < DebugTextList.Num() && Idx == INDEX_NONE; ++i)
+	// remove from the back to reduce copying the remaining portion of the list as it is shrunk, since order matters :
+	int32 FirstIdx, LastIdx;
+	int32 NumElements = FindDebugTextListIntervalForActor(SrcActor, FirstIdx, LastIdx);
+	if (NumElements > 0)
 	{
-		if (DebugTextList[i].SrcActor == SrcActor && (!bLeaveDurationText || DebugTextList[i].TimeRemaining == -1.f))
+		if (!bLeaveDurationText)
 		{
-			Idx = i;
+			// immediate removal : remove the list altogether : 
+			DebugTextList.RemoveAt(FirstIdx, NumElements);
+		}
+		else
+		{
+			// remove those with no expiry date 
+			for (int32 Idx = LastIdx; Idx >= FirstIdx; --Idx)
+			{
+				if (DebugTextList[Idx].TimeRemaining == -1.f)
+				{
+					DebugTextList.RemoveAt(Idx, 1);
+				}
+			}
 		}
-	}
-	if (Idx != INDEX_NONE)
-	{
-		DebugTextList.RemoveAt(Idx,1);
 	}
 }
 
@@ -895,9 +936,9 @@ void AHUD::DrawText(FString const& Text, FLinearColor Color, float ScreenX, floa
 			ScreenX *= Scale;
 			ScreenY *= Scale;
 		}
-		FCanvasTextItem TextItem( FVector2D( ScreenX, ScreenY ), FText::FromString(Text), Font ? Font : GEngine->GetMediumFont(), Color );
-		TextItem.Scale = FVector2D( Scale, Scale );
-		Canvas->DrawItem( TextItem );
+		FCanvasTextItem TextItem(FVector2D(ScreenX, ScreenY), FText::FromString(Text), Font ? Font : GEngine->GetMediumFont(), Color);
+		TextItem.Scale = FVector2D(Scale, Scale);
+		Canvas->DrawItem(TextItem);
 	}
 }
 
@@ -905,14 +946,14 @@ void AHUD::DrawMaterial(UMaterialInterface* Material, float ScreenX, float Scree
 {
 	if (IsCanvasValid_WarnIfNot() && Material)
 	{
-		FCanvasTileItem TileItem( FVector2D( ScreenX, ScreenY ), Material->GetRenderProxy(), FVector2D( ScreenW, ScreenH ) * Scale, FVector2D( MaterialU, MaterialV ), FVector2D( MaterialU+MaterialUWidth, MaterialV +MaterialVHeight) );
+		FCanvasTileItem TileItem(FVector2D(ScreenX, ScreenY), Material->GetRenderProxy(), FVector2D(ScreenW, ScreenH) * Scale, FVector2D(MaterialU, MaterialV), FVector2D(MaterialU + MaterialUWidth, MaterialV + MaterialVHeight));
 		TileItem.Rotation = FRotator(0, Rotation, 0);
 		TileItem.PivotPoint = RotPivot;
 		if (bScalePosition)
 		{
 			TileItem.Position *= Scale;
 		}
-		Canvas->DrawItem( TileItem );
+		Canvas->DrawItem(TileItem);
 	}
 }
 
@@ -920,12 +961,12 @@ void AHUD::DrawMaterialSimple(UMaterialInterface* Material, float ScreenX, float
 {
 	if (IsCanvasValid_WarnIfNot() && Material)
 	{
-		FCanvasTileItem TileItem( FVector2D( ScreenX, ScreenY ), Material->GetRenderProxy(), FVector2D( ScreenW, ScreenH ) * Scale );
+		FCanvasTileItem TileItem(FVector2D(ScreenX, ScreenY), Material->GetRenderProxy(), FVector2D(ScreenW, ScreenH) * Scale);
 		if (bScalePosition)
 		{
 			TileItem.Position *= Scale;
 		}
-		Canvas->DrawItem( TileItem );
+		Canvas->DrawItem(TileItem);
 	}
 }
 
@@ -933,15 +974,15 @@ void AHUD::DrawTexture(UTexture* Texture, float ScreenX, float ScreenY, float Sc
 {
 	if (IsCanvasValid_WarnIfNot() && Texture)
 	{
-		FCanvasTileItem TileItem( FVector2D( ScreenX, ScreenY ), Texture->Resource, FVector2D( ScreenW, ScreenH ) * Scale, FVector2D( TextureU, TextureV ), FVector2D( TextureU + TextureUWidth, TextureV + TextureVHeight ), Color );
+		FCanvasTileItem TileItem(FVector2D(ScreenX, ScreenY), Texture->Resource, FVector2D(ScreenW, ScreenH) * Scale, FVector2D(TextureU, TextureV), FVector2D(TextureU + TextureUWidth, TextureV + TextureVHeight), Color);
 		TileItem.Rotation = FRotator(0, Rotation, 0);
 		TileItem.PivotPoint = RotPivot;
 		if (bScalePosition)
 		{
 			TileItem.Position *= Scale;
 		}
-		TileItem.BlendMode = FCanvas::BlendToSimpleElementBlend( BlendMode );
-		Canvas->DrawItem( TileItem );
+		TileItem.BlendMode = FCanvas::BlendToSimpleElementBlend(BlendMode);
+		Canvas->DrawItem(TileItem);
 	}
 }
 
@@ -949,7 +990,7 @@ void AHUD::DrawTextureSimple(UTexture* Texture, float ScreenX, float ScreenY, fl
 {
 	if (IsCanvasValid_WarnIfNot() && Texture)
 	{
-		FCanvasTileItem TileItem( FVector2D( ScreenX, ScreenY ), Texture->Resource, FLinearColor::White );
+		FCanvasTileItem TileItem(FVector2D(ScreenX, ScreenY), Texture->Resource, FLinearColor::White);
 		if (bScalePosition)
 		{
 			TileItem.Position *= Scale;
@@ -957,7 +998,7 @@ void AHUD::DrawTextureSimple(UTexture* Texture, float ScreenX, float ScreenY, fl
 		// Apply the scale to the size (which will have been setup from the texture in the constructor).
 		TileItem.Size *= Scale;
 		TileItem.BlendMode = SE_BLEND_Translucent;
-		Canvas->DrawItem( TileItem );
+		Canvas->DrawItem(TileItem);
 	}
 }
 
@@ -965,7 +1006,7 @@ void AHUD::DrawMaterialTriangle(UMaterialInterface* Material, FVector2D V0_Pos,
 {
 	if (IsCanvasValid_WarnIfNot() && Material)
 	{
-		FCanvasTriangleItem TriangleItem(V0_Pos, V1_Pos, V2_Pos, V0_UV, V1_UV, V2_UV, NULL);
+		FCanvasTriangleItem TriangleItem(V0_Pos, V1_Pos, V2_Pos, V0_UV, V1_UV, V2_UV, nullptr);
 		TriangleItem.TriangleList[0].V0_Color = V0_Color;
 		TriangleItem.TriangleList[0].V1_Color = V1_Color;
 		TriangleItem.TriangleList[0].V2_Color = V2_Color;
@@ -979,12 +1020,12 @@ FVector AHUD::Project(FVector Location) const
 	{
 		return Canvas->Project(Location);
 	}
-	return FVector(0,0,0);
+	return FVector(0, 0, 0);
 }
 
 void AHUD::Deproject(float ScreenX, float ScreenY, FVector& WorldPosition, FVector& WorldDirection) const
 {
-	WorldPosition = WorldDirection = FVector(0,0,0);
+	WorldPosition = WorldDirection = FVector(0, 0, 0);
 	if (IsCanvasValid_WarnIfNot())
 	{
 		Canvas->Deproject(FVector2D(ScreenX, ScreenY), WorldPosition, WorldDirection);
@@ -1017,7 +1058,7 @@ void AHUD::GetActorsInSelectionRectangle(TSubclassOf<class AActor> ClassFilter,
 		FVector(-1.f, 1.f, 1.f),
 		FVector(-1.f, 1.f, -1.f),
 		FVector(-1.f, -1.f, 1.f),
-		FVector(-1.f, -1.f, -1.f)	};
+		FVector(-1.f, -1.f, -1.f) };
 
 	//~~~
 
@@ -1048,7 +1089,7 @@ void AHUD::GetActorsInSelectionRectangle(TSubclassOf<class AActor> ClassFilter,
 		//Selection Box must fully enclose the Projected Actor Bounds
 		if (bActorMustBeFullyEnclosed)
 		{
-			if(SelectionRectangle.IsInside(ActorBox2D))
+			if (SelectionRectangle.IsInside(ActorBox2D))
 			{
 				OutActors.Add(EachActor);
 			}
@@ -1066,11 +1107,11 @@ void AHUD::GetActorsInSelectionRectangle(TSubclassOf<class AActor> ClassFilter,
 
 void AHUD::DrawRect(FLinearColor Color, float ScreenX, float ScreenY, float Width, float Height)
 {
-	if (IsCanvasValid_WarnIfNot())	
+	if (IsCanvasValid_WarnIfNot())
 	{
-		FCanvasTileItem TileItem( FVector2D(ScreenX, ScreenY), GWhiteTexture, FVector2D( Width, Height ), Color );
+		FCanvasTileItem TileItem(FVector2D(ScreenX, ScreenY), GWhiteTexture, FVector2D(Width, Height), Color);
 		TileItem.BlendMode = SE_BLEND_Translucent;
-		Canvas->DrawItem( TileItem );		
+		Canvas->DrawItem(TileItem);
 	}
 }
 
@@ -1078,10 +1119,10 @@ void AHUD::DrawLine(float StartScreenX, float StartScreenY, float EndScreenX, fl
 {
 	if (IsCanvasValid_WarnIfNot())
 	{
-		FCanvasLineItem LineItem( FVector2D(StartScreenX, StartScreenY), FVector2D(EndScreenX, EndScreenY) );
-		LineItem.SetColor( LineColor );
+		FCanvasLineItem LineItem(FVector2D(StartScreenX, StartScreenY), FVector2D(EndScreenX, EndScreenY));
+		LineItem.SetColor(LineColor);
 		LineItem.LineThickness = LineThickness;
-		Canvas->DrawItem( LineItem );
+		Canvas->DrawItem(LineItem);
 	}
 }
 
@@ -1093,23 +1134,23 @@ APlayerController* AHUD::GetOwningPlayerController() const
 
 APawn* AHUD::GetOwningPawn() const
 {
-	return PlayerOwner ? PlayerOwner->GetPawn() : NULL;
+	return PlayerOwner ? PlayerOwner->GetPawn() : nullptr;
 }
 
-void AHUD::RenderHitBoxes( FCanvas* InCanvas )
+void AHUD::RenderHitBoxes(FCanvas* InCanvas)
 {
 	for (const FHUDHitBox& HitBox : HitBoxMap)
 	{
 		FLinearColor BoxColor = FLinearColor::White;
-		if( HitBoxHits.Contains(const_cast<FHUDHitBox*>(&HitBox)))
+		if (HitBoxHits.Contains(const_cast<FHUDHitBox*>(&HitBox)))
 		{
 			BoxColor = FLinearColor::Red;
 		}
-		HitBox.Draw( InCanvas, BoxColor );
+		HitBox.Draw(InCanvas, BoxColor);
 	}
 }
 
-void AHUD::UpdateHitBoxCandidates( TArray<FVector2D> InContactPoints )
+void AHUD::UpdateHitBoxCandidates(TArray<FVector2D> InContactPoints)
 {
 	HitBoxHits.Reset();
 	for (FHUDHitBox& HitBox : HitBoxMap)
@@ -1172,7 +1213,7 @@ void AHUD::UpdateHitBoxCandidates( TArray<FVector2D> InContactPoints )
 	}
 }
 
-const FHUDHitBox* AHUD::GetHitBoxAtCoordinates( FVector2D InHitLocation, const bool bIsConsumingInput ) const
+const FHUDHitBox* AHUD::GetHitBoxAtCoordinates(FVector2D InHitLocation, const bool bIsConsumingInput) const
 {
 	if (HitBoxMap.Num() > 0)
 	{
@@ -1180,7 +1221,7 @@ const FHUDHitBox* AHUD::GetHitBoxAtCoordinates( FVector2D InHitLocation, const b
 
 		for (const FHUDHitBox& HitBox : HitBoxMap)
 		{
-			if( (!bIsConsumingInput || HitBox.ConsumesInput()) && HitBox.Contains( InHitLocation ) )
+			if ((!bIsConsumingInput || HitBox.ConsumesInput()) && HitBox.Contains(InHitLocation))
 			{
 				return &HitBox;
 			}
@@ -1207,11 +1248,11 @@ void AHUD::GetHitBoxesAtCoordinates(FVector2D InHitLocation, TArray<const FHUDHi
 	}
 }
 
-const FHUDHitBox* AHUD::GetHitBoxWithName( const FName InName ) const
+const FHUDHitBox* AHUD::GetHitBoxWithName(const FName InName) const
 {
 	for (const FHUDHitBox& HitBox : HitBoxMap)
 	{
-		if( HitBox.GetName() == InName )
+		if (HitBox.GetName() == InName)
 		{
 			return &HitBox;
 		}
@@ -1280,8 +1321,8 @@ bool AHUD::UpdateAndDispatchHitBoxClickEvents(FVector2D ClickLocation, const EIn
 }
 
 void AHUD::AddHitBox(FVector2D Position, FVector2D Size, FName Name, bool bConsumesInput, int32 Priority)
-{	
-	if( GetHitBoxWithName(Name) == nullptr )
+{
+	if (GetHitBoxWithName(Name) == nullptr)
 	{
 		bool bAdded = false;
 		for (int32 Index = 0; Index < HitBoxMap.Num(); ++Index)
@@ -1306,12 +1347,12 @@ void AHUD::AddHitBox(FVector2D Position, FVector2D Size, FName Name, bool bConsu
 
 bool AHUD::IsCanvasValid_WarnIfNot() const
 {
-	const bool bIsValid = Canvas != NULL;
+	const bool bIsValid = Canvas != nullptr;
 	if (!bIsValid)
 	{
 		FMessageLog("PIE").Warning()
 			->AddToken(FUObjectToken::Create(const_cast<AHUD*>(this)))
-			->AddToken(FTextToken::Create(LOCTEXT( "PIE_Warning_Message_CanvasCallOutsideOfDrawCanvas", "Canvas Draw functions may only be called during the handling of the DrawHUD event" )));
+			->AddToken(FTextToken::Create(LOCTEXT("PIE_Warning_Message_CanvasCallOutsideOfDrawCanvas", "Canvas Draw functions may only be called during the handling of the DrawHUD event")));
 	}
 
 	return bIsValid;
