UE4.25 - Static analysis fix

diff --git a/Engine/Source/Runtime/Engine/Private/Streaming/StaticMeshUpdate.cpp b/Engine/Source/Runtime/Engine/Private/Streaming/StaticMeshUpdate.cpp
--- a/Engine/Source/Runtime/Engine/Private/Streaming/StaticMeshUpdate.cpp
+++ b/Engine/Source/Runtime/Engine/Private/Streaming/StaticMeshUpdate.cpp
@@ -324,9 +324,12 @@ void FStaticMeshStreamOut::CheckReferencesAndDiscardCPUData(const FContext& Cont
 
 	if (!NumExternalReferences || NumReferenceChecks >= GStreamingMaxReferenceChecks)
 	{
-		for (int32 LODIdx = CurrentFirstLODIdx; LODIdx < PendingFirstMip; ++LODIdx)
+		if (RenderData)
 		{
-			RenderData->LODResources[LODIdx].DiscardCPUData();
+			for (int32 LODIdx = CurrentFirstLODIdx; LODIdx < PendingFirstMip; ++LODIdx)
+			{
+				RenderData->LODResources[LODIdx].DiscardCPUData();
+			}
 		}
 
 		// Because we discarded the CPU data, the stream out can not be cancelled anymore.
