Added lvl1 include guards for chaos

diff --git a/Engine/Plugins/Experimental/ChaosCloth/Source/Chaos/ChaosCloth.Build.cs b/Engine/Plugins/Experimental/ChaosCloth/Source/Chaos/ChaosCloth.Build.cs
--- a/Engine/Plugins/Experimental/ChaosCloth/Source/Chaos/ChaosCloth.Build.cs
+++ b/Engine/Plugins/Experimental/ChaosCloth/Source/Chaos/ChaosCloth.Build.cs
@@ -21,5 +21,6 @@ public class ChaosCloth : ModuleRules
 		);
 
         SetupModulePhysicsSupport(Target);
-    }
+		PrivateDefinitions.Add("CHAOS_INCLUDE_LEVEL_1=1");
+	}
 }
diff --git a/Engine/Plugins/Experimental/ChaosNiagara/Source/ChaosNiagara/ChaosNiagara.Build.cs b/Engine/Plugins/Experimental/ChaosNiagara/Source/ChaosNiagara/ChaosNiagara.Build.cs
--- a/Engine/Plugins/Experimental/ChaosNiagara/Source/ChaosNiagara/ChaosNiagara.Build.cs
+++ b/Engine/Plugins/Experimental/ChaosNiagara/Source/ChaosNiagara/ChaosNiagara.Build.cs
@@ -45,5 +45,7 @@ public class ChaosNiagara : ModuleRules
                 "GeometryCollectionCore"
             }
         );
-    }
+
+		PrivateDefinitions.Add("CHAOS_INCLUDE_LEVEL_1=1");
+	}
 }
diff --git a/Engine/Source/Editor/UnrealEd/Private/EditorEngine.cpp b/Engine/Source/Editor/UnrealEd/Private/EditorEngine.cpp
--- a/Engine/Source/Editor/UnrealEd/Private/EditorEngine.cpp
+++ b/Engine/Source/Editor/UnrealEd/Private/EditorEngine.cpp
@@ -222,7 +222,9 @@
 #include "UObject/UnrealType.h"
 
 #if WITH_CHAOS
+#define CHAOS_INCLUDE_LEVEL_1
 #include "ChaosSolversModule.h"
+#undef CHAOS_INCLUDE_LEVEL_1
 #endif
 
 
diff --git a/Engine/Source/Runtime/Engine/Private/Collision/KAggregateGeom.cpp b/Engine/Source/Runtime/Engine/Private/Collision/KAggregateGeom.cpp
--- a/Engine/Source/Runtime/Engine/Private/Collision/KAggregateGeom.cpp
+++ b/Engine/Source/Runtime/Engine/Private/Collision/KAggregateGeom.cpp
@@ -20,7 +20,9 @@
 #endif
 
 #if WITH_CHAOS
+#define CHAOS_INCLUDE_LEVEL_1
 #include "Chaos/ImplicitObject.h"
+#undef CHAOS_INCLUDE_LEVEL_1
 #endif
 
 
diff --git a/Engine/Source/Runtime/Engine/Private/Collision/SceneQuery.cpp b/Engine/Source/Runtime/Engine/Private/Collision/SceneQuery.cpp
--- a/Engine/Source/Runtime/Engine/Private/Collision/SceneQuery.cpp
+++ b/Engine/Source/Runtime/Engine/Private/Collision/SceneQuery.cpp
@@ -25,7 +25,9 @@ float DebugLineLifetime = 2.f;
 #include "PhysicsEngine/PhysXSupport.h"
 #include "PhysicsEngine/CollisionAnalyzerCapture.h"
 
+#define CHAOS_INCLUDE_LEVEL_1
 #include "ChaosSolversModule.h"
+#undef CHAOS_INCLUDE_LEVEL_1
 #include "Physics/Experimental/ChaosInterfaceWrapper.h"
 
 CSV_DEFINE_CATEGORY(SceneQuery, false);
diff --git a/Engine/Source/Runtime/Engine/Private/Collision/SceneQueryLowLevel.cpp b/Engine/Source/Runtime/Engine/Private/Collision/SceneQueryLowLevel.cpp
--- a/Engine/Source/Runtime/Engine/Private/Collision/SceneQueryLowLevel.cpp
+++ b/Engine/Source/Runtime/Engine/Private/Collision/SceneQueryLowLevel.cpp
@@ -12,12 +12,16 @@
 #include "PhysicsEngine/CollisionQueryFilterCallback.h"
 #include "PhysicsCore.h"
 
-#include "PhysTestSerializer.h"
 
 #include "SQAccelerator.h"
+
+#define CHAOS_INCLUDE_LEVEL_1
 #include "SQVerifier.h"
+#include "PhysTestSerializer.h"
+
 #include "PBDRigidsSolver.h"
 #include "Chaos/PBDRigidsEvolutionGBF.h"
+#undef CHAOS_INCLUDE_LEVEL_1
 
 int32 ForceStandardSQ = 0;
 FAutoConsoleVariableRef CVarForceStandardSQ(TEXT("p.ForceStandardSQ"), ForceStandardSQ, TEXT("If enabled, we force the standard scene query even if custom SQ structure is enabled"));
diff --git a/Engine/Source/Runtime/Engine/Private/Collision/WorldCollision.cpp b/Engine/Source/Runtime/Engine/Private/Collision/WorldCollision.cpp
--- a/Engine/Source/Runtime/Engine/Private/Collision/WorldCollision.cpp
+++ b/Engine/Source/Runtime/Engine/Private/Collision/WorldCollision.cpp
@@ -17,7 +17,9 @@
 #include "Physics/PhysicsInterfaceTypes.h"
 
 #if WITH_CHAOS
+#define CHAOS_INCLUDE_LEVEL_1
 #include "Chaos/ImplicitObject.h"
+#undef CHAOS_INCLUDE_LEVEL_1
 #endif
 
 using namespace PhysicsInterfaceTypes;
diff --git a/Engine/Source/Runtime/Engine/Private/GameEngine.cpp b/Engine/Source/Runtime/Engine/Private/GameEngine.cpp
--- a/Engine/Source/Runtime/Engine/Private/GameEngine.cpp
+++ b/Engine/Source/Runtime/Engine/Private/GameEngine.cpp
@@ -73,7 +73,9 @@
 #endif
 
 #if WITH_CHAOS
+#define CHAOS_INCLUDE_LEVEL_1
 #include "ChaosSolversModule.h"
+#undef CHAOS_INCLUDE_LEVEL_1
 #endif
 
 CSV_DECLARE_CATEGORY_MODULE_EXTERN(CORE_API, Basic);
diff --git a/Engine/Source/Runtime/Engine/Private/PhysicsEngine/BodyInstance.cpp b/Engine/Source/Runtime/Engine/Private/PhysicsEngine/BodyInstance.cpp
--- a/Engine/Source/Runtime/Engine/Private/PhysicsEngine/BodyInstance.cpp
+++ b/Engine/Source/Runtime/Engine/Private/PhysicsEngine/BodyInstance.cpp
@@ -28,6 +28,8 @@
 #include "Logging/TokenizedMessage.h"
 #include "Logging/MessageLog.h"
 #include "Misc/UObjectToken.h"
+
+#define CHAOS_INCLUDE_LEVEL_1
 #include "ChaosCheck.h"
 #include "Chaos/Capsule.h"
 #include "Chaos/Convex.h"
@@ -43,6 +45,7 @@
 #if WITH_CHAOS
 #include "Chaos/ParticleHandle.h"
 #endif // WITH_CHAOS
+#undef CHAOS_INCLUDE_LEVEL_1
 
 #if PHYSICS_INTERFACE_PHYSX
 	#include "PhysXPublic.h"
diff --git a/Engine/Source/Runtime/Engine/Private/PhysicsEngine/BodySetup.cpp b/Engine/Source/Runtime/Engine/Private/PhysicsEngine/BodySetup.cpp
--- a/Engine/Source/Runtime/Engine/Private/PhysicsEngine/BodySetup.cpp
+++ b/Engine/Source/Runtime/Engine/Private/PhysicsEngine/BodySetup.cpp
@@ -21,8 +21,11 @@
 #include "UObject/UObjectIterator.h"
 #include "UObject/PropertyPortFlags.h"
 #include "Components/SplineMeshComponent.h"
+
+#define CHAOS_INCLUDE_LEVEL_1
 #include "ChaosCheck.h"
 #include "Chaos/Convex.h"
+#undef CHAOS_INCLUDE_LEVEL_1
 
 #include "PhysXCookHelper.h"
 
@@ -44,12 +47,16 @@
 #include "ProfilingDebugging/CookStats.h"
 #include "UObject/AnimPhysObjectVersion.h"
 
+#define CHAOS_INCLUDE_LEVEL_1
 #include "Chaos/TriangleMeshImplicitObject.h"
+#undef CHAOS_INCLUDE_LEVEL_1
 
 #if WITH_CHAOS
 	#include "Experimental/ChaosDerivedData.h"
 	#include "Physics/Experimental/ChaosDerivedDataReader.h"
+#define CHAOS_INCLUDE_LEVEL_1
 	#include "Chaos/CollisionConvexMesh.h"
+#undef CHAOS_INCLUDE_LEVEL_1
 #endif
 
 /** Helper for enum output... */
diff --git a/Engine/Source/Runtime/Engine/Private/PhysicsEngine/Experimental/ChaosDerivedData.cpp b/Engine/Source/Runtime/Engine/Private/PhysicsEngine/Experimental/ChaosDerivedData.cpp
--- a/Engine/Source/Runtime/Engine/Private/PhysicsEngine/Experimental/ChaosDerivedData.cpp
+++ b/Engine/Source/Runtime/Engine/Private/PhysicsEngine/Experimental/ChaosDerivedData.cpp
@@ -4,6 +4,7 @@
 
 #if WITH_CHAOS
 
+#define CHAOS_INCLUDE_LEVEL_1
 #include "Chaos/Core.h"
 #include "Chaos/CollisionConvexMesh.h"
 #include "Chaos/ChaosArchive.h"
@@ -11,6 +12,8 @@
 #include "Chaos/Particles.h"
 #include "Chaos/TriangleMeshImplicitObject.h"
 #include "Chaos/Convex.h"
+#undef CHAOS_INCLUDE_LEVEL_1
+
 #include "Serialization/Archive.h"
 #include "PhysicsEngine/BodySetup.h"
 #include "Serialization/MemoryWriter.h"
diff --git a/Engine/Source/Runtime/Engine/Private/PhysicsEngine/Experimental/ChaosDerivedDataReader.cpp b/Engine/Source/Runtime/Engine/Private/PhysicsEngine/Experimental/ChaosDerivedDataReader.cpp
--- a/Engine/Source/Runtime/Engine/Private/PhysicsEngine/Experimental/ChaosDerivedDataReader.cpp
+++ b/Engine/Source/Runtime/Engine/Private/PhysicsEngine/Experimental/ChaosDerivedDataReader.cpp
@@ -1,9 +1,12 @@
 // Copyright Epic Games, Inc. All Rights Reserved.
 
 #include "Physics/Experimental/ChaosDerivedDataReader.h"
+
+#define CHAOS_INCLUDE_LEVEL_1
 #include "Chaos/ChaosArchive.h"
 #include "Chaos/TriangleMeshImplicitObject.h"
 #include "Chaos/Convex.h"
+#undef CHAOS_INCLUDE_LEVEL_1
 
 template<typename T, int d>
 FChaosDerivedDataReader<T, d>::FChaosDerivedDataReader(FBulkDataInterface* InBulkData)
diff --git a/Engine/Source/Runtime/Engine/Private/PhysicsEngine/Experimental/ChaosInterfaceUtils.cpp b/Engine/Source/Runtime/Engine/Private/PhysicsEngine/Experimental/ChaosInterfaceUtils.cpp
--- a/Engine/Source/Runtime/Engine/Private/PhysicsEngine/Experimental/ChaosInterfaceUtils.cpp
+++ b/Engine/Source/Runtime/Engine/Private/PhysicsEngine/Experimental/ChaosInterfaceUtils.cpp
@@ -2,6 +2,7 @@
 
 #include "Physics/Experimental/ChaosInterfaceUtils.h"
 
+#define CHAOS_INCLUDE_LEVEL_1
 #include "Chaos/Box.h"
 #include "Chaos/Capsule.h"
 #include "Chaos/GeometryParticles.h"
@@ -13,6 +14,7 @@
 #include "Chaos/TriangleMeshImplicitObject.h"
 #include "Chaos/UniformGrid.h"
 #include "Chaos/Convex.h"
+#undef CHAOS_INCLUDE_LEVEL_1
 
 #include "PhysicsEngine/BoxElem.h"
 #include "PhysicsEngine/ConvexElem.h"
diff --git a/Engine/Source/Runtime/Engine/Private/PhysicsEngine/Experimental/PhysInterface_Chaos.cpp b/Engine/Source/Runtime/Engine/Private/PhysicsEngine/Experimental/PhysInterface_Chaos.cpp
--- a/Engine/Source/Runtime/Engine/Private/PhysicsEngine/Experimental/PhysInterface_Chaos.cpp
+++ b/Engine/Source/Runtime/Engine/Private/PhysicsEngine/Experimental/PhysInterface_Chaos.cpp
@@ -7,7 +7,10 @@
 #include "Physics/Experimental/ChaosInterfaceUtils.h"
 #include "Physics/PhysicsInterfaceTypes.h"
 #include "PhysicsEngine/AggregateGeom.h"
+#include "Templates/UniquePtr.h"
 
+#define CHAOS_INCLUDE_LEVEL_1
+#include "PhysicsSolver.h"
 #include "Chaos/Box.h"
 #include "Chaos/Cylinder.h"
 #include "Chaos/TaperedCylinder.h"
@@ -20,17 +23,16 @@
 #include "Chaos/Sphere.h"
 #include "Chaos/Matrix.h"
 #include "Chaos/MassProperties.h"
-#include "PhysicsSolver.h"
-#include "Templates/UniquePtr.h"
 #include "ChaosSolversModule.h"
 #include "Chaos/ErrorReporter.h"
 #include "Chaos/ImplicitObjectScaled.h"
-#include "Chaos/TriangleMeshImplicitObject.h"
 #include "Chaos/Convex.h"
 #include "Chaos/GeometryQueries.h"
 #include "Chaos/Plane.h"
 #include "ChaosCheck.h"
 #include "Chaos/Particle/ParticleUtilities.h"
+#undef CHAOS_INCLUDE_LEVEL_1
+
 #include "Async/ParallelFor.h"
 #include "Components/PrimitiveComponent.h"
 #include "Physics/PhysicsFiltering.h"
diff --git a/Engine/Source/Runtime/Engine/Private/PhysicsEngine/Experimental/PhysScene_Chaos.cpp b/Engine/Source/Runtime/Engine/Private/PhysicsEngine/Experimental/PhysScene_Chaos.cpp
--- a/Engine/Source/Runtime/Engine/Private/PhysicsEngine/Experimental/PhysScene_Chaos.cpp
+++ b/Engine/Source/Runtime/Engine/Private/PhysicsEngine/Experimental/PhysScene_Chaos.cpp
@@ -2,10 +2,6 @@
 
 #include "Physics/Experimental/PhysScene_Chaos.h"
 
-#include "PhysicsSolver.h"
-#include "ChaosSolversModule.h"
-#include "ChaosLog.h"
-#include "ChaosStats.h"
 
 #include "CoreMinimal.h"
 #include "GameDelegates.h"
@@ -13,16 +9,27 @@
 #include "Async/AsyncWork.h"
 #include "Async/ParallelFor.h"
 #include "Engine/Engine.h"
-#include "Field/FieldSystem.h"
-#include "Framework/Dispatcher.h"
-#include "Framework/PersistentTask.h"
-#include "Framework/PhysicsTickTask.h"
+
 #include "Misc/CoreDelegates.h"
 #include "Misc/ScopeLock.h"
 #include "PhysicsEngine/BodySetup.h"
 #include "PhysicsEngine/PhysicsSettings.h"
 #include "Components/PrimitiveComponent.h"
 #include "Components/SkeletalMeshComponent.h"
+#include "PhysicsReplication.h"
+#include "Physics/Experimental/PhysicsUserData_Chaos.h"
+#include "ProfilingDebugging/CsvProfiler.h"
+
+#define CHAOS_INCLUDE_LEVEL_1
+#include "PhysicsSolver.h"
+#include "ChaosSolversModule.h"
+#include "ChaosLog.h"
+#include "ChaosStats.h"
+
+#include "Field/FieldSystem.h"
+#include "Framework/Dispatcher.h"
+#include "Framework/PersistentTask.h"
+#include "Framework/PhysicsTickTask.h"
 
 #include "PhysicsProxy/FieldSystemPhysicsProxy.h"
 #include "PhysicsProxy/GeometryCollectionPhysicsProxy.h"
@@ -37,11 +44,10 @@
 #include "PBDRigidActiveParticlesBuffer.h"
 #include "Chaos/GeometryParticlesfwd.h"
 #include "Chaos/Box.h"
-#include "PhysicsReplication.h"
 #include "ChaosSolvers/Public/EventsData.h"
 #include "ChaosSolvers/Public/EventManager.h"
-#include "Physics/Experimental/PhysicsUserData_Chaos.h"
-#include "ProfilingDebugging/CsvProfiler.h"
+#undef CHAOS_INCLUDE_LEVEL_1
+
 
 #if !UE_BUILD_SHIPPING
 #include "Engine/World.h"
diff --git a/Engine/Source/Runtime/Engine/Private/PhysicsEngine/ImmediatePhysics/ImmediatePhysicsChaos/ImmediatePhysicsActorHandle_Chaos.cpp b/Engine/Source/Runtime/Engine/Private/PhysicsEngine/ImmediatePhysics/ImmediatePhysicsChaos/ImmediatePhysicsActorHandle_Chaos.cpp
--- a/Engine/Source/Runtime/Engine/Private/PhysicsEngine/ImmediatePhysics/ImmediatePhysicsChaos/ImmediatePhysicsActorHandle_Chaos.cpp
+++ b/Engine/Source/Runtime/Engine/Private/PhysicsEngine/ImmediatePhysics/ImmediatePhysicsChaos/ImmediatePhysicsActorHandle_Chaos.cpp
@@ -2,6 +2,7 @@
 
 #include "Physics/ImmediatePhysics/ImmediatePhysicsChaos/ImmediatePhysicsActorHandle_Chaos.h"
 
+#define CHAOS_INCLUDE_LEVEL_1
 #include "Chaos/Box.h"
 #include "Chaos/Capsule.h"
 #include "Chaos/Evolution/PBDMinEvolution.h"
@@ -12,6 +13,7 @@
 #include "Chaos/PBDRigidsEvolutionGBF.h"
 #include "Chaos/Sphere.h"
 #include "Chaos/Utilities.h"
+#undef CHAOS_INCLUDE_LEVEL_1
 
 #include "Physics/Experimental/ChaosInterfaceUtils.h"
 
diff --git a/Engine/Source/Runtime/Engine/Private/PhysicsEngine/ImmediatePhysics/ImmediatePhysicsChaos/ImmediatePhysicsJointHandle_Chaos.cpp b/Engine/Source/Runtime/Engine/Private/PhysicsEngine/ImmediatePhysics/ImmediatePhysicsChaos/ImmediatePhysicsJointHandle_Chaos.cpp
--- a/Engine/Source/Runtime/Engine/Private/PhysicsEngine/ImmediatePhysics/ImmediatePhysicsChaos/ImmediatePhysicsJointHandle_Chaos.cpp
+++ b/Engine/Source/Runtime/Engine/Private/PhysicsEngine/ImmediatePhysics/ImmediatePhysicsChaos/ImmediatePhysicsJointHandle_Chaos.cpp
@@ -3,8 +3,10 @@
 #include "Physics/ImmediatePhysics/ImmediatePhysicsChaos/ImmediatePhysicsJointHandle_Chaos.h"
 #include "Physics/ImmediatePhysics/ImmediatePhysicsChaos/ImmediatePhysicsActorHandle_Chaos.h"
 
+#define CHAOS_INCLUDE_LEVEL_1
 #include "Chaos/Particle/ParticleUtilities.h"
 #include "Chaos/PBDJointConstraints.h"
+#undef CHAOS_INCLUDE_LEVEL_1
 
 #include "PhysicsEngine/ConstraintInstance.h"
 
diff --git a/Engine/Source/Runtime/Engine/Private/PhysicsEngine/PhysLevel.cpp b/Engine/Source/Runtime/Engine/Private/PhysicsEngine/PhysLevel.cpp
--- a/Engine/Source/Runtime/Engine/Private/PhysicsEngine/PhysLevel.cpp
+++ b/Engine/Source/Runtime/Engine/Private/PhysicsEngine/PhysLevel.cpp
@@ -15,7 +15,10 @@
 #include "Modules/ModuleManager.h"
 #include "PhysicsInitialization.h"
 #include "PhysicsEngine/PhysicsSettings.h"
+
+#define CHAOS_INCLUDE_LEVEL_1
 #include "ChaosSolversModule.h"
+#undef CHAOS_INCLUDE_LEVEL_1
 
 #if PHYSICS_INTERFACE_PHYSX
 	#include "PhysicsEngine/PhysXSupport.h"
diff --git a/Engine/Source/Runtime/Engine/Private/PhysicsEngine/PhysScene_PhysX.cpp b/Engine/Source/Runtime/Engine/Private/PhysicsEngine/PhysScene_PhysX.cpp
--- a/Engine/Source/Runtime/Engine/Private/PhysicsEngine/PhysScene_PhysX.cpp
+++ b/Engine/Source/Runtime/Engine/Private/PhysicsEngine/PhysScene_PhysX.cpp
@@ -40,7 +40,10 @@
 #include "PhysicsEngine/ConstraintInstance.h"
 #include "PhysicsReplication.h"
 #include "ProfilingDebugging/CsvProfiler.h"
+
+#define CHAOS_INCLUDE_LEVEL_1
 #include "PhysTestSerializer.h"
+#undef CHAOS_INCLUDE_LEVEL_1
 
 /** Physics stats **/
 
diff --git a/Engine/Source/Runtime/Engine/Private/PhysicsEngine/PhysUtils.cpp b/Engine/Source/Runtime/Engine/Private/PhysicsEngine/PhysUtils.cpp
--- a/Engine/Source/Runtime/Engine/Private/PhysicsEngine/PhysUtils.cpp
+++ b/Engine/Source/Runtime/Engine/Private/PhysicsEngine/PhysUtils.cpp
@@ -16,11 +16,13 @@
 #include "PhysicsEngine/BodySetup.h"
 #include "Physics/PhysicsInterfaceCore.h"
 #include "PhysXSupportCore.h"
+#define CHAOS_INCLUDE_LEVEL_1
 #include "PhysicsSolver.h"
 #include "Chaos/PBDRigidsEvolutionGBF.h"
 #include "Chaos/ChaosArchive.h"
 #include "Chaos/TrackedGeometryManager.h"
 #include "RewindData.h"
+#undef CHAOS_INCLUDE_LEVEL_1
 
 /** Returns false if ModelToHulls operation should halt because of vertex count overflow. */
 static bool AddConvexPrim(FKAggregateGeom* OutGeom, TArray<FPlane> &Planes, UModel* InModel)
diff --git a/Engine/Source/Runtime/Engine/Private/PhysicsEngine/PhysicsSettings.cpp b/Engine/Source/Runtime/Engine/Private/PhysicsEngine/PhysicsSettings.cpp
--- a/Engine/Source/Runtime/Engine/Private/PhysicsEngine/PhysicsSettings.cpp
+++ b/Engine/Source/Runtime/Engine/Private/PhysicsEngine/PhysicsSettings.cpp
@@ -6,7 +6,10 @@
 #include "UObject/Package.h"
 
 #include "Framework/Threading.h"
+
+#define CHAOS_INCLUDE_LEVEL_1
 #include "ChaosSolversModule.h"
+#undef CHAOS_INCLUDE_LEVEL_1
 
 UPhysicsSettings::UPhysicsSettings(const FObjectInitializer& ObjectInitializer)
 	: Super(ObjectInitializer)
diff --git a/Engine/Source/Runtime/Engine/Private/UnrealEngine.cpp b/Engine/Source/Runtime/Engine/Private/UnrealEngine.cpp
--- a/Engine/Source/Runtime/Engine/Private/UnrealEngine.cpp
+++ b/Engine/Source/Runtime/Engine/Private/UnrealEngine.cpp
@@ -129,7 +129,10 @@ UnrealEngine.cpp: Implements the UEngine class and helpers.
 #include "HAL/LowLevelMemTracker.h"
 #include "HAL/PlatformApplicationMisc.h"
 #include "DynamicResolutionState.h"
+
+#define CHAOS_INCLUDE_LEVEL_1
 #include "Chaos/TriangleMeshImplicitObject.h"
+#undef CHAOS_INCLUDE_LEVEL_1
 
 #include "Particles/Spawn/ParticleModuleSpawn.h"
 #include "Particles/TypeData/ParticleModuleTypeDataMesh.h"
diff --git a/Engine/Source/Runtime/Engine/Private/World.cpp b/Engine/Source/Runtime/Engine/Private/World.cpp
--- a/Engine/Source/Runtime/Engine/Private/World.cpp
+++ b/Engine/Source/Runtime/Engine/Private/World.cpp
@@ -126,7 +126,9 @@
 #include "ObjectTrace.h"
 
 #if INCLUDE_CHAOS
+#define CHAOS_INCLUDE_LEVEL_1
 #include "ChaosSolversModule.h"
+#undef CHAOS_INCLUDE_LEVEL_1
 #endif
 
 DEFINE_LOG_CATEGORY_STATIC(LogWorld, Log, All);
diff --git a/Engine/Source/Runtime/NavigationSystem/Private/NavMesh/RecastNavMeshGenerator.cpp b/Engine/Source/Runtime/NavigationSystem/Private/NavMesh/RecastNavMeshGenerator.cpp
--- a/Engine/Source/Runtime/NavigationSystem/Private/NavMesh/RecastNavMeshGenerator.cpp
+++ b/Engine/Source/Runtime/NavigationSystem/Private/NavMesh/RecastNavMeshGenerator.cpp
@@ -17,8 +17,10 @@
 	#include "PhysXPublic.h"
 #endif
 #if WITH_CHAOS
+#define CHAOS_INCLUDE_LEVEL_1
 	#include "Chaos/HeightField.h"
 	#include "Chaos/TriangleMeshImplicitObject.h"
+#undef CHAOS_INCLUDE_LEVEL_1
 #endif
 #include "NavMesh/PImplRecastNavMesh.h"
 
diff --git a/Engine/Source/Runtime/PhysicsCore/Public/PhysTestSerializer.h b/Engine/Source/Runtime/PhysicsCore/Public/PhysTestSerializer.h
--- a/Engine/Source/Runtime/PhysicsCore/Public/PhysTestSerializer.h
+++ b/Engine/Source/Runtime/PhysicsCore/Public/PhysTestSerializer.h
@@ -2,6 +2,9 @@
 
 #pragma once
 
+#include "Chaos/IncludeLvl1.inl"
+CHAOS_LEVEL_CHECK
+
 #include "CoreMinimal.h"
 #include "PhysicsInterfaceDeclaresCore.h"
 
diff --git a/Engine/Source/Runtime/PhysicsSQ/PhysicsSQ.Build.cs b/Engine/Source/Runtime/PhysicsSQ/PhysicsSQ.Build.cs
--- a/Engine/Source/Runtime/PhysicsSQ/PhysicsSQ.Build.cs
+++ b/Engine/Source/Runtime/PhysicsSQ/PhysicsSQ.Build.cs
@@ -14,5 +14,6 @@ public class PhysicsSQ: ModuleRules
 		);
 
 		SetupModulePhysicsSupport(Target);
+		PrivateDefinitions.Add("CHAOS_INCLUDE_LEVEL_1=1");
 	}
 }
diff --git a/Engine/Source/Runtime/PhysicsSQ/Public/SQVerifier.h b/Engine/Source/Runtime/PhysicsSQ/Public/SQVerifier.h
--- a/Engine/Source/Runtime/PhysicsSQ/Public/SQVerifier.h
+++ b/Engine/Source/Runtime/PhysicsSQ/Public/SQVerifier.h
@@ -5,9 +5,11 @@
 #include "ProfilingDebugging/ScopedTimers.h"
 
 #if WITH_PHYSX
-
+#include "Chaos/IncludeLvl1.inl"
+CHAOS_LEVEL_CHECK
 #include "Chaos/PBDRigidsEvolutionGBF.h"
 #include "PhysicsCore/Public/PhysicsInterfaceUtilsCore.h"
+#include "PhysTestSerializer.h"
 
 #ifndef SQ_REPLAY_TEST
 #define SQ_REPLAY_TEST(cond) bEnsureOnMismatch ? ensure(cond) : (cond)
