In Editor analytics, increased the chances correctly differenciating when the Editor is started with or without a game. Lot of projects has the invalid GUID 00000000-0000-0000-0000-000000000000 and and empty name (same as if the Editor started without a game) and 25% of abnormal terminations are for for those projects.
  - When UGeneralProjectSettings doesn't have a project name set, use FApp::GetProjectName().

diff --git a/Engine/Source/Runtime/Engine/Private/Analytics/EditorSessionSummaryWriter.cpp b/Engine/Source/Runtime/Engine/Private/Analytics/EditorSessionSummaryWriter.cpp
--- a/Engine/Source/Runtime/Engine/Private/Analytics/EditorSessionSummaryWriter.cpp
+++ b/Engine/Source/Runtime/Engine/Private/Analytics/EditorSessionSummaryWriter.cpp
@@ -234,7 +234,7 @@ TUniquePtr<FEditorAnalyticsSession> FEditorSessionSummaryWriter::CreateCurrentSe
 
 	Session->PlatformProcessID = FPlatformProcess::GetCurrentProcessId();
 	Session->MonitorProcessID = OutOfProcessMonitorProcessId;
-	Session->ProjectName = ProjectSettings.ProjectName;
+	Session->ProjectName = ProjectSettings.ProjectName.Len() ? ProjectSettings.ProjectName : FApp::GetProjectName();
 	Session->ProjectID = ProjectSettings.ProjectID.ToString(EGuidFormats::DigitsWithHyphens);
 	Session->ProjectDescription = ProjectSettings.Description;
 	Session->ProjectVersion = ProjectSettings.ProjectVersion;
