Remove Chaos include from CollisionDebugDraw.h

diff --git a/Engine/Source/Runtime/Engine/Private/Collision/CollisionDebugDrawing.cpp b/Engine/Source/Runtime/Engine/Private/Collision/CollisionDebugDrawing.cpp
--- a/Engine/Source/Runtime/Engine/Private/Collision/CollisionDebugDrawing.cpp
+++ b/Engine/Source/Runtime/Engine/Private/Collision/CollisionDebugDrawing.cpp
@@ -6,9 +6,12 @@
 #include "DrawDebugHelpers.h"
 #include "PhysXPublic.h"
 
-#include "Chaos/Sphere.h"
+#define CHAOS_INCLUDE_LEVEL_1
 #include "Chaos/Box.h"
 #include "Chaos/Capsule.h"
+#include "Chaos/ImplicitObject.h"
+#include "Chaos/Sphere.h"
+#undef CHAOS_INCLUDE_LEVEL_1
 
 
 #if ENABLE_DRAW_DEBUG
diff --git a/Engine/Source/Runtime/Engine/Private/Collision/CollisionDebugDrawing.h b/Engine/Source/Runtime/Engine/Private/Collision/CollisionDebugDrawing.h
--- a/Engine/Source/Runtime/Engine/Private/Collision/CollisionDebugDrawing.h
+++ b/Engine/Source/Runtime/Engine/Private/Collision/CollisionDebugDrawing.h
@@ -9,9 +9,11 @@
 #include "EngineDefines.h"
 #include "PhysicsPublic.h"
 #include "PhysXIncludes.h"
-#define TEMP_HEADER_CHAOS_LEVEL_1
-#include "Chaos/ImplicitObject.h"
-#undef TEMP_HEADER_CHAOS_LEVEL_1
+
+namespace Chaos
+{
+	class FImplicitObject;
+}
 
 #if PHYSICS_INTERFACE_PHYSX
 /** Draw PhysX geom with overlaps */
