Fixed crash loading into WeapGun_Impacts in Shooter game as the auto destroy subsystem was null.

diff --git a/Engine/Source/Runtime/Engine/Private/Actor.cpp b/Engine/Source/Runtime/Engine/Private/Actor.cpp
--- a/Engine/Source/Runtime/Engine/Private/Actor.cpp
+++ b/Engine/Source/Runtime/Engine/Private/Actor.cpp
@@ -1642,14 +1642,16 @@ void AActor::SetAutoDestroyWhenFinished(bool bVal)
 	bAutoDestroyWhenFinished = bVal;
 	if (UWorld* MyWorld = GetWorld())
 	{
-		UAutoDestroySubsystem* AutoDestroySub = MyWorld->GetSubsystem<UAutoDestroySubsystem>();
-		if (bAutoDestroyWhenFinished && (HasActorBegunPlay() || IsActorBeginningPlay()))
+		if(UAutoDestroySubsystem* AutoDestroySub = MyWorld->GetSubsystem<UAutoDestroySubsystem>())
 		{
-			AutoDestroySub->RegisterActor(this);
-		}
-		else
-		{
-			AutoDestroySub->UnregisterActor(this);
+			if (bAutoDestroyWhenFinished && (HasActorBegunPlay() || IsActorBeginningPlay()))
+			{
+				AutoDestroySub->RegisterActor(this);
+			}
+			else
+			{
+				AutoDestroySub->UnregisterActor(this);
+			}
 		}
 	}
 }
