Fixed crash in ray tracing when Show Static Meshes flag was unchecked.

diff --git a/Engine/Source/Runtime/Engine/Private/StaticMeshRender.cpp b/Engine/Source/Runtime/Engine/Private/StaticMeshRender.cpp
--- a/Engine/Source/Runtime/Engine/Private/StaticMeshRender.cpp
+++ b/Engine/Source/Runtime/Engine/Private/StaticMeshRender.cpp
@@ -1637,6 +1637,11 @@ void FStaticMeshSceneProxy::GetDynamicMeshElements(const TArray<const FSceneView
 #if RHI_RAYTRACING
 void FStaticMeshSceneProxy::GetDynamicRayTracingInstances(FRayTracingMaterialGatheringContext& Context, TArray<FRayTracingInstance>& OutRayTracingInstances )
 {
+	if (DynamicRayTracingGeometries.Num() <= 0)
+	{
+		return;
+	}
+
 	uint8 PrimitiveDPG = GetStaticDepthPriorityGroup();
 	const uint32 LODIndex = GetLOD(Context.ReferenceView);
 	const FStaticMeshLODResources& LODModel = RenderData->LODResources[LODIndex];
