Make sure the ddc key is not changing during the build of the skeletalmesh

#rnx

diff --git a/Engine/Source/Runtime/Engine/Classes/Engine/SkeletalMesh.h b/Engine/Source/Runtime/Engine/Classes/Engine/SkeletalMesh.h
--- a/Engine/Source/Runtime/Engine/Classes/Engine/SkeletalMesh.h
+++ b/Engine/Source/Runtime/Engine/Classes/Engine/SkeletalMesh.h
@@ -1049,6 +1049,11 @@ private:
 	//This function should be call only in the PostLoad
 	void CreateUserSectionsDataForLegacyAssets();
 
+	/*
+	 * This function unbind any incorrect cloth asset that is bind.
+	 */
+	void PostLoadValidateClothingData();
+
 public:
 	//We want to avoid calling post edit change multiple time during import and build process.
 
diff --git a/Engine/Source/Runtime/Engine/Private/SkeletalMesh.cpp b/Engine/Source/Runtime/Engine/Private/SkeletalMesh.cpp
--- a/Engine/Source/Runtime/Engine/Private/SkeletalMesh.cpp
+++ b/Engine/Source/Runtime/Engine/Private/SkeletalMesh.cpp
@@ -776,7 +776,7 @@ void USkeletalMesh::InitResources()
 
 					if (!bImportDataInSync || !bRenderDataInSync)
 					{
-						UE_ASSET_LOG(LogSkeletalMesh, Fatal, this, TEXT("Data out of sync in lod %d. bImportDataInSync=%d, bRenderDataInSync=%d"), LODIndex, bImportDataInSync, bRenderDataInSync);
+						UE_ASSET_LOG(LogSkeletalMesh, Error, this, TEXT("Data out of sync in lod %d. bImportDataInSync=%d, bRenderDataInSync=%d. This happen when DDC cache has corrupted data (Key has change during the skeletalmesh build)"), LODIndex, bImportDataInSync, bRenderDataInSync);
 					}
 				}
 			}
@@ -2318,6 +2318,37 @@ void USkeletalMesh::CreateUserSectionsDataForLegacyAssets()
 	}
 }
 
+void USkeletalMesh::PostLoadValidateClothingData()
+{
+	//Make sure PostEditChange is not call when we validate the clothing data during the PostLoad
+	FScopedSkeletalMeshPostEditChange ScopedPostEditChange(this, false, false);
+	for (int32 LodIndex = 0; LodIndex < GetLODNum(); LodIndex++)
+	{
+		FSkeletalMeshLODModel& ThisLODModel = ImportedModel->LODModels[LodIndex];
+
+		if (ThisLODModel.RawSkeletalMeshBulkData.IsEmpty() || !ThisLODModel.RawSkeletalMeshBulkData.IsBuildDataAvailable())
+		{
+			//Invalid clothing asset will not be unbind
+			continue;
+		}
+
+		int32 SectionNum = ThisLODModel.Sections.Num();
+		for (int32 SectionIndex = 0; SectionIndex < SectionNum; ++SectionIndex)
+		{
+			FSkelMeshSection& Section = ThisLODModel.Sections[SectionIndex];
+			if (Section.HasClothingData())
+			{
+				UClothingAssetBase* ClothingAsset = GetClothingAsset(Section.ClothingData.AssetGuid);
+				if(!ClothingAsset->IsValidLod(LodIndex))
+				{
+					//Unbind invalid cloth asset
+					ClothingAsset->UnbindFromSkeletalMesh(this, LodIndex);
+				}
+			}
+		}
+	}
+}
+
 #endif // WITH_EDITOR
 
 void USkeletalMesh::PostLoad()
@@ -2473,6 +2504,10 @@ void USkeletalMesh::PostLoad()
 			CreateUserSectionsDataForLegacyAssets();
 		}
 
+		//We must unbind incorrect cloth data before computing the ddc key (CacheDerivedData)
+		//Incorrect cloth data will not be rebind after the build, which will change the DDC key, this situation lead to a corrupted DDC cache
+		PostLoadValidateClothingData();
+
 		if (GetResourceForRendering() == nullptr)
 		{
 			CacheDerivedData();
diff --git a/Engine/Source/Runtime/Engine/Private/SkeletalMeshRenderData.cpp b/Engine/Source/Runtime/Engine/Private/SkeletalMeshRenderData.cpp
--- a/Engine/Source/Runtime/Engine/Private/SkeletalMeshRenderData.cpp
+++ b/Engine/Source/Runtime/Engine/Private/SkeletalMeshRenderData.cpp
@@ -335,7 +335,14 @@ void FSkeletalMeshRenderData::Cache(const ITargetPlatform* TargetPlatform, USkel
 				const bool bNeedsCPUAccess = FSkeletalMeshLODRenderData::ShouldKeepCPUResources(Owner, LODIndex, bForceKeepCPUResources);
 				LODData.SerializeStreamedData(Ar, Owner, LODIndex, LODStripFlags, bNeedsCPUAccess, bForceKeepCPUResources);
 			}
-			GetDerivedDataCacheRef().Put(*DerivedDataKey, DerivedData, Owner->GetPathName());
+
+			//Recompute the derived data key in case there was some data correction during the build process, this make sure the DDC key is always representing the correct build result.
+			//There should never be correction of the data during the build, the data has to be corrected in the post load before calling this function.
+			FString BuiltDerivedDataKey = BuildSkeletalMeshDerivedDataKey(Owner);
+			ensureMsgf(BuiltDerivedDataKey == DerivedDataKey, TEXT("Skeletal mesh [%s] build has change the source data. The derived data key before and after the build is different."), *Owner->GetPathName());
+
+			//Store the data using the built key to avoid DDC corruption
+			GetDerivedDataCacheRef().Put(*BuiltDerivedDataKey, DerivedData, Owner->GetPathName());
 
 			int32 T1 = FPlatformTime::Cycles();
 			UE_LOG(LogSkeletalMesh, Log, TEXT("Built Skeletal Mesh [%.2fs] %s"), FPlatformTime::ToMilliseconds(T1 - T0) / 1000.0f, *Owner->GetPathName());
