Remove chaos material from public headers

Benn.Gallagher

diff --git a/Engine/Source/Runtime/Engine/Classes/PhysicalMaterials/PhysicalMaterialMask.h b/Engine/Source/Runtime/Engine/Classes/PhysicalMaterials/PhysicalMaterialMask.h
--- a/Engine/Source/Runtime/Engine/Classes/PhysicalMaterials/PhysicalMaterialMask.h
+++ b/Engine/Source/Runtime/Engine/Classes/PhysicalMaterials/PhysicalMaterialMask.h
@@ -72,7 +72,7 @@ public:
 #endif // WITH_EDITOR
 
 #if WITH_CHAOS
-	FPhysicsMaterialMaskHandle MaterialMaskHandle;
+	TUniquePtr<FPhysicsMaterialMaskHandle> MaterialMaskHandle;
 
 	/** Get the physics-interface derived version of this material */
 	FPhysicsMaterialMaskHandle& GetPhysicsMaterialMask();
diff --git a/Engine/Source/Runtime/Engine/Private/PhysicsEngine/PhysicalMaterialMask.cpp b/Engine/Source/Runtime/Engine/Private/PhysicsEngine/PhysicalMaterialMask.cpp
--- a/Engine/Source/Runtime/Engine/Private/PhysicsEngine/PhysicalMaterialMask.cpp
+++ b/Engine/Source/Runtime/Engine/Private/PhysicsEngine/PhysicalMaterialMask.cpp
@@ -98,9 +98,9 @@ UPhysicalMaterialMask::UPhysicalMaterialMask(const FObjectInitializer& ObjectIni
 void UPhysicalMaterialMask::PostEditChangeProperty(FPropertyChangedEvent& PropertyChangedEvent)
 {
 #if WITH_CHAOS
-	if (MaterialMaskHandle.IsValid())
+	if (MaterialMaskHandle && MaterialMaskHandle->IsValid())
 	{
-		FPhysicsInterface::UpdateMaterialMask(MaterialMaskHandle, this);
+		FPhysicsInterface::UpdateMaterialMask(*MaterialMaskHandle, this);
 	}
 #endif
 
@@ -136,7 +136,10 @@ void UPhysicalMaterialMask::PostLoad()
 void UPhysicalMaterialMask::FinishDestroy()
 {
 #if WITH_CHAOS
-	FPhysicsInterface::ReleaseMaterialMask(MaterialMaskHandle);
+	if(MaterialMaskHandle)
+	{
+		FPhysicsInterface::ReleaseMaterialMask(*MaterialMaskHandle);
+	}
 #endif
 	Super::FinishDestroy();
 }
@@ -228,15 +231,21 @@ void UPhysicalMaterialMask::DumpMaskData()
 
 FPhysicsMaterialMaskHandle& UPhysicalMaterialMask::GetPhysicsMaterialMask()
 {
-	if (!MaterialMaskHandle.IsValid())
+	if(!MaterialMaskHandle)
+	{
+		//need to use ptr to avoid polluting engine headers
+		MaterialMaskHandle = MakeUnique<FPhysicsMaterialMaskHandle>();
+	}
+
+	if (!MaterialMaskHandle->IsValid())
 	{
-		MaterialMaskHandle = FPhysicsInterface::CreateMaterialMask(this);
-		check(MaterialMaskHandle.IsValid());
+		*MaterialMaskHandle = FPhysicsInterface::CreateMaterialMask(this);
+		check(MaterialMaskHandle->IsValid());
 
-		FPhysicsInterface::UpdateMaterialMask(MaterialMaskHandle, this);
+		FPhysicsInterface::UpdateMaterialMask(*MaterialMaskHandle, this);
 	}
 
-	return MaterialMaskHandle;
+	return *MaterialMaskHandle;
 }
 
 #endif
diff --git a/Engine/Source/Runtime/Engine/Public/Physics/Experimental/PhysInterface_Chaos.h b/Engine/Source/Runtime/Engine/Public/Physics/Experimental/PhysInterface_Chaos.h
--- a/Engine/Source/Runtime/Engine/Public/Physics/Experimental/PhysInterface_Chaos.h
+++ b/Engine/Source/Runtime/Engine/Public/Physics/Experimental/PhysInterface_Chaos.h
@@ -13,7 +13,6 @@
 #include "Physics/PhysicsInterfaceTypes.h"
 #include "Physics/GenericPhysicsInterface.h"
 #include "Physics/Experimental/PhysicsUserData_Chaos.h"
-#include "Chaos/PhysicalMaterials.h"
 
 //NOTE: Do not include Chaos headers directly as it means recompiling all of engine. This should be reworked to avoid allocations
 
