Added more forward declares to physical material and chaos contact modification

Benn.Gallagher

diff --git a/Engine/Source/Runtime/Engine/Classes/PhysicalMaterials/PhysicalMaterial.h b/Engine/Source/Runtime/Engine/Classes/PhysicalMaterials/PhysicalMaterial.h
--- a/Engine/Source/Runtime/Engine/Classes/PhysicalMaterials/PhysicalMaterial.h
+++ b/Engine/Source/Runtime/Engine/Classes/PhysicalMaterials/PhysicalMaterial.h
@@ -10,6 +10,7 @@
 #include "Vehicles/TireType.h"
 #include "PhysicsEngine/PhysicsSettingsEnums.h"
 #include "Physics/PhysicsInterfaceCore.h"
+#include "PhysicsInterfaceDeclaresCore.h"
 #include "PhysicalMaterial.generated.h"
 
 struct FPropertyChangedEvent;
@@ -21,6 +22,7 @@ namespace physx
 }
 #endif
 
+
 /** DEPRECATED Pairs desired tire friction scale with tire type */
 USTRUCT()
 struct FTireFrictionScalePair
@@ -129,7 +131,7 @@ class ENGINE_API UPhysicalMaterial : public UObject
 
 public:
 
-	FPhysicsMaterialHandle MaterialHandle;
+	TUniquePtr<FPhysicsMaterialHandle> MaterialHandle;
 
 	FChaosUserData UserData;
 
diff --git a/Engine/Source/Runtime/Engine/Classes/PhysicalMaterials/PhysicalMaterialMask.h b/Engine/Source/Runtime/Engine/Classes/PhysicalMaterials/PhysicalMaterialMask.h
--- a/Engine/Source/Runtime/Engine/Classes/PhysicalMaterials/PhysicalMaterialMask.h
+++ b/Engine/Source/Runtime/Engine/Classes/PhysicalMaterials/PhysicalMaterialMask.h
@@ -63,6 +63,8 @@ public:
 	virtual void FinishDestroy() override;
 	//~ End UObject Interface
 
+	~UPhysicalMaterialMask();
+
 #if WITH_EDITOR
 	// Helper method to set mask texture
 	void SetMaskTexture(UTexture* InMaskTexture, const FString& InTextureFilename);
diff --git a/Engine/Source/Runtime/Engine/Private/PhysicsEngine/Experimental/PhysScene_Chaos.cpp b/Engine/Source/Runtime/Engine/Private/PhysicsEngine/Experimental/PhysScene_Chaos.cpp
--- a/Engine/Source/Runtime/Engine/Private/PhysicsEngine/Experimental/PhysScene_Chaos.cpp
+++ b/Engine/Source/Runtime/Engine/Private/PhysicsEngine/Experimental/PhysScene_Chaos.cpp
@@ -1177,7 +1177,8 @@ FPhysScene_ChaosInterface::FPhysScene_ChaosInterface(const AWorldSettings* InSet
 {
 	//Initialize unique ptrs that are just here to allow forward declare. This should be reworked todo(ocohen)
 #if TODO_FIX_REFERENCES_TO_ADDARRAY
-	Scene.GetSolver()->GetEvolution()->GetParticles().AddArray(&BodyInstances);
+	BodyInstances = MakeUnique<Chaos::TArrayCollectionArray<FBodyInstance*>>();
+	Scene.GetSolver()->GetEvolution()->GetParticles().AddArray(BodyInstances.Get());
 #endif
 
 	// Create replication manager
diff --git a/Engine/Source/Runtime/Engine/Private/PhysicsEngine/PhysicalMaterial.cpp b/Engine/Source/Runtime/Engine/Private/PhysicsEngine/PhysicalMaterial.cpp
--- a/Engine/Source/Runtime/Engine/Private/PhysicsEngine/PhysicalMaterial.cpp
+++ b/Engine/Source/Runtime/Engine/Private/PhysicsEngine/PhysicalMaterial.cpp
@@ -14,8 +14,13 @@
 
 #if PHYSICS_INTERFACE_PHYSX
 	#include "PhysicsEngine/PhysXSupport.h"
+	#include "Physics/PhysicsInterfacePhysX.h"
 #endif // WITH_PHYSX
 
+#if WITH_CHAOS
+	#include "Chaos/PhysicalMaterials.h"
+#endif
+
 UDEPRECATED_PhysicalMaterialPropertyBase::UDEPRECATED_PhysicalMaterialPropertyBase(const FObjectInitializer& ObjectInitializer)
 	: Super(ObjectInitializer)
 {
@@ -40,8 +45,12 @@ UPhysicalMaterial::UPhysicalMaterial(const FObjectInitializer& ObjectInitializer
 #if WITH_EDITOR
 void UPhysicalMaterial::PostEditChangeProperty(FPropertyChangedEvent& PropertyChangedEvent)
 {
+	if(!MaterialHandle)
+	{
+		MaterialHandle = MakeUnique<FPhysicsMaterialHandle>();
+	}
 	// Update PhysX material last so we have a valid Parent
-	FPhysicsInterface::UpdateMaterial(MaterialHandle, this);
+	FPhysicsInterface::UpdateMaterial(*MaterialHandle, this);
 
 	Super::PostEditChangeProperty(PropertyChangedEvent);
 }
@@ -53,7 +62,11 @@ void UPhysicalMaterial::RebuildPhysicalMaterials()
 	{
 		if (UPhysicalMaterial * PhysicalMaterial = Cast<UPhysicalMaterial>(*Iter))
 		{
-			FPhysicsInterface::UpdateMaterial(PhysicalMaterial->MaterialHandle, PhysicalMaterial);
+			if(!PhysicalMaterial->MaterialHandle)
+			{
+				PhysicalMaterial->MaterialHandle = MakeUnique<FPhysicsMaterialHandle>();
+			}
+			FPhysicsInterface::UpdateMaterial(*PhysicalMaterial->MaterialHandle, PhysicalMaterial);
 		}
 	}
 }
@@ -76,22 +89,29 @@ void UPhysicalMaterial::PostLoad()
 
 void UPhysicalMaterial::FinishDestroy()
 {
-	FPhysicsInterface::ReleaseMaterial(MaterialHandle);
+	if(MaterialHandle)
+	{
+		FPhysicsInterface::ReleaseMaterial(*MaterialHandle);
+	}
 	Super::FinishDestroy();
 }
 
 FPhysicsMaterialHandle& UPhysicalMaterial::GetPhysicsMaterial()
 {
-	if(!MaterialHandle.IsValid())
+	if(!MaterialHandle)
+	{
+		MaterialHandle = MakeUnique<FPhysicsMaterialHandle>();
+	}
+	if(!MaterialHandle->IsValid())
 	{
-		MaterialHandle = FPhysicsInterface::CreateMaterial(this);
-		check(MaterialHandle.IsValid());
+		*MaterialHandle = FPhysicsInterface::CreateMaterial(this);
+		check(MaterialHandle->IsValid());
 
-		FPhysicsInterface::SetUserData(MaterialHandle, &UserData);
-		FPhysicsInterface::UpdateMaterial(MaterialHandle, this);
+		FPhysicsInterface::SetUserData(*MaterialHandle, &UserData);
+		FPhysicsInterface::UpdateMaterial(*MaterialHandle, this);
 	}
 
-	return MaterialHandle;
+	return *MaterialHandle;
 }
 
 EPhysicalSurface UPhysicalMaterial::DetermineSurfaceType(UPhysicalMaterial const* PhysicalMaterial)
diff --git a/Engine/Source/Runtime/Engine/Private/PhysicsEngine/PhysicalMaterialMask.cpp b/Engine/Source/Runtime/Engine/Private/PhysicsEngine/PhysicalMaterialMask.cpp
--- a/Engine/Source/Runtime/Engine/Private/PhysicsEngine/PhysicalMaterialMask.cpp
+++ b/Engine/Source/Runtime/Engine/Private/PhysicsEngine/PhysicalMaterialMask.cpp
@@ -93,6 +93,8 @@ UPhysicalMaterialMask::UPhysicalMaterialMask(const FObjectInitializer& ObjectIni
 {
 }
 
+UPhysicalMaterialMask::~UPhysicalMaterialMask() = default;
+
 #if WITH_EDITOR
 
 void UPhysicalMaterialMask::PostEditChangeProperty(FPropertyChangedEvent& PropertyChangedEvent)
diff --git a/Engine/Source/Runtime/Engine/Public/Physics/Experimental/PhysScene_Chaos.h b/Engine/Source/Runtime/Engine/Public/Physics/Experimental/PhysScene_Chaos.h
--- a/Engine/Source/Runtime/Engine/Public/Physics/Experimental/PhysScene_Chaos.h
+++ b/Engine/Source/Runtime/Engine/Public/Physics/Experimental/PhysScene_Chaos.h
@@ -11,12 +11,13 @@
 #include "PhysicsPublic.h"
 #include "PhysInterface_Chaos.h"
 #include "Physics/PhysicsInterfaceUtils.h"
-#include "Chaos/Transform.h"
-#include "Chaos/ArrayCollectionArray.h"
+
+//Chaos includes. Todo: move to chaos core so we can include for all of engine
 #include "Chaos/Declares.h"
 #include "PhysicsProxy/SingleParticlePhysicsProxyFwd.h"
 #include "Framework/Threading.h"
-#include "Chaos/PBDCollisionConstraints.h"
+#include "Chaos/Core.h"
+#include "Chaos/CollisionResolutionTypes.h"
 
 #ifndef CHAOS_WITH_PAUSABLE_SOLVER
 #define CHAOS_WITH_PAUSABLE_SOLVER 1
@@ -62,6 +63,10 @@ namespace Chaos
 	class ISpatialAccelerationCollection;
 
 	class IDispatcher;
+
+	template <typename T>
+	class TArrayCollectionArray;
+
 }
 
 /**
@@ -507,7 +512,7 @@ private:
 
 	float MDeltaTime;
 	//Body Instances
-	Chaos::TArrayCollectionArray<FBodyInstance*> BodyInstances;
+	TUniquePtr<Chaos::TArrayCollectionArray<FBodyInstance*>> BodyInstances;
 	// Temp Interface
 	UWorld* MOwningWorld;
 	TArray<FCollisionNotifyInfo> MNotifies;
diff --git a/Engine/Source/Runtime/Experimental/Chaos/Public/Chaos/CollisionResolutionTypes.h b/Engine/Source/Runtime/Experimental/Chaos/Public/Chaos/CollisionResolutionTypes.h
--- a/Engine/Source/Runtime/Experimental/Chaos/Public/Chaos/CollisionResolutionTypes.h
+++ b/Engine/Source/Runtime/Experimental/Chaos/Public/Chaos/CollisionResolutionTypes.h
@@ -68,4 +68,8 @@ namespace Chaos
 		FVec3 AccumulatedImpulse;
 	};
 
+
+	//TODO: move into a better forward declare location
+	class FPBDCollisionConstraintHandle;
+	using FCollisionModifierCallback = TFunction<ECollisionModifierResult(FPBDCollisionConstraintHandle*)>;
 }
diff --git a/Engine/Source/Runtime/Experimental/Chaos/Public/Chaos/PBDCollisionConstraints.h b/Engine/Source/Runtime/Experimental/Chaos/Public/Chaos/PBDCollisionConstraints.h
--- a/Engine/Source/Runtime/Experimental/Chaos/Public/Chaos/PBDCollisionConstraints.h
+++ b/Engine/Source/Runtime/Experimental/Chaos/Public/Chaos/PBDCollisionConstraints.h
@@ -31,7 +31,6 @@ class TPBDRigidsSOAs;
 using FRigidBodyContactConstraintsPostComputeCallback = TFunction<void()>;
 using FRigidBodyContactConstraintsPostApplyCallback = TFunction<void(const FReal Dt, const TArray<FPBDCollisionConstraintHandle*>&)>;
 using FRigidBodyContactConstraintsPostApplyPushOutCallback = TFunction<void(const FReal Dt, const TArray<FPBDCollisionConstraintHandle*>&, bool)>;
-using FCollisionModifierCallback = TFunction<ECollisionModifierResult(FPBDCollisionConstraintHandle*)>;
 
 /**
  * A container and solver for collision constraints.
