Fix build from robomerge into a refactor
#rnx

diff --git a/Engine/Source/Runtime/Engine/Private/SkeletalMesh.cpp b/Engine/Source/Runtime/Engine/Private/SkeletalMesh.cpp
--- a/Engine/Source/Runtime/Engine/Private/SkeletalMesh.cpp
+++ b/Engine/Source/Runtime/Engine/Private/SkeletalMesh.cpp
@@ -2325,10 +2325,9 @@ void USkeletalMesh::PostLoadValidateClothingData()
 	for (int32 LodIndex = 0; LodIndex < GetLODNum(); LodIndex++)
 	{
 		FSkeletalMeshLODModel& ThisLODModel = ImportedModel->LODModels[LodIndex];
-
-		if (ThisLODModel.RawSkeletalMeshBulkData.IsEmpty() || !ThisLODModel.RawSkeletalMeshBulkData.IsBuildDataAvailable())
+		if (IsLODImportedDataEmpty(LodIndex) || !IsLODImportedDataBuildAvailable(LodIndex))
 		{
-			//Invalid clothing asset will not be unbind
+			//Invalid clothing asset will not be unbind in case we do not build the asset
 			continue;
 		}
 
diff --git a/Engine/Source/Runtime/Engine/Private/SkeletalMeshRenderData.cpp b/Engine/Source/Runtime/Engine/Private/SkeletalMeshRenderData.cpp
--- a/Engine/Source/Runtime/Engine/Private/SkeletalMeshRenderData.cpp
+++ b/Engine/Source/Runtime/Engine/Private/SkeletalMeshRenderData.cpp
@@ -338,7 +338,7 @@ void FSkeletalMeshRenderData::Cache(const ITargetPlatform* TargetPlatform, USkel
 
 			//Recompute the derived data key in case there was some data correction during the build process, this make sure the DDC key is always representing the correct build result.
 			//There should never be correction of the data during the build, the data has to be corrected in the post load before calling this function.
-			FString BuiltDerivedDataKey = BuildSkeletalMeshDerivedDataKey(Owner);
+			FString BuiltDerivedDataKey = BuildSkeletalMeshDerivedDataKey(TargetPlatform, Owner);
 			ensureMsgf(BuiltDerivedDataKey == DerivedDataKey, TEXT("Skeletal mesh [%s] build has change the source data. The derived data key before and after the build is different."), *Owner->GetPathName());
 
 			//Store the data using the built key to avoid DDC corruption
