Remove convex header and fix non-unity build

diff --git a/Engine/Source/Editor/UnrealEd/Private/Fbx/FbxMainExport.cpp b/Engine/Source/Editor/UnrealEd/Private/Fbx/FbxMainExport.cpp
--- a/Engine/Source/Editor/UnrealEd/Private/Fbx/FbxMainExport.cpp
+++ b/Engine/Source/Editor/UnrealEd/Private/Fbx/FbxMainExport.cpp
@@ -105,6 +105,7 @@
 #include "Framework/Application/SlateApplication.h"
 #include "Interfaces/IMainFrameModule.h"
 #include "UObject/MetaData.h"
+#include "ChaosCheck.h"
 
 namespace UnFbx
 {
diff --git a/Engine/Source/Runtime/Engine/Classes/PhysicalMaterials/PhysicalMaterial.h b/Engine/Source/Runtime/Engine/Classes/PhysicalMaterials/PhysicalMaterial.h
--- a/Engine/Source/Runtime/Engine/Classes/PhysicalMaterials/PhysicalMaterial.h
+++ b/Engine/Source/Runtime/Engine/Classes/PhysicalMaterials/PhysicalMaterial.h
@@ -52,6 +52,8 @@ class ENGINE_API UPhysicalMaterial : public UObject
 {
 	GENERATED_UCLASS_BODY()
 
+	virtual ~UPhysicalMaterial();
+	UPhysicalMaterial(FVTableHelper& Helper);
 	//
 	// Surface properties.
 	//
diff --git a/Engine/Source/Runtime/Engine/Classes/PhysicalMaterials/PhysicalMaterialMask.h b/Engine/Source/Runtime/Engine/Classes/PhysicalMaterials/PhysicalMaterialMask.h
--- a/Engine/Source/Runtime/Engine/Classes/PhysicalMaterials/PhysicalMaterialMask.h
+++ b/Engine/Source/Runtime/Engine/Classes/PhysicalMaterials/PhysicalMaterialMask.h
@@ -63,7 +63,8 @@ public:
 	virtual void FinishDestroy() override;
 	//~ End UObject Interface
 
-	~UPhysicalMaterialMask();
+	virtual ~UPhysicalMaterialMask();
+	UPhysicalMaterialMask(FVTableHelper& Helper);
 
 #if WITH_EDITOR
 	// Helper method to set mask texture
diff --git a/Engine/Source/Runtime/Engine/Classes/PhysicsEngine/ConvexElem.h b/Engine/Source/Runtime/Engine/Classes/PhysicsEngine/ConvexElem.h
--- a/Engine/Source/Runtime/Engine/Classes/PhysicsEngine/ConvexElem.h
+++ b/Engine/Source/Runtime/Engine/Classes/PhysicsEngine/ConvexElem.h
@@ -6,10 +6,6 @@
 #include "UObject/ObjectMacros.h"
 #include "PhysicsEngine/ShapeElem.h"
 
-#define TEMP_HEADER_CHAOS_LEVEL_1
-#include "Chaos/Convex.h"
-#undef TEMP_HEADER_CHAOS_LEVEL_1
-
 #if WITH_CHAOS
 #include "Chaos/Serializable.h"
 #endif
@@ -29,8 +25,7 @@ namespace Chaos
 {
 	class FImplicitObject;
 
-	template <typename T, int d>
-	class TConvex;
+	class FConvex;
 }
 
 /** One convex hull, used for simplified collision. */
diff --git a/Engine/Source/Runtime/Engine/Private/PhysicsEngine/PhysicalMaterial.cpp b/Engine/Source/Runtime/Engine/Private/PhysicsEngine/PhysicalMaterial.cpp
--- a/Engine/Source/Runtime/Engine/Private/PhysicsEngine/PhysicalMaterial.cpp
+++ b/Engine/Source/Runtime/Engine/Private/PhysicsEngine/PhysicalMaterial.cpp
@@ -42,6 +42,13 @@ UPhysicalMaterial::UPhysicalMaterial(const FObjectInitializer& ObjectInitializer
 	UserData = FChaosUserData(this);
 }
 
+UPhysicalMaterial::UPhysicalMaterial(FVTableHelper& Helper)
+	: Super(Helper)
+{
+}
+
+UPhysicalMaterial::~UPhysicalMaterial() = default;
+
 #if WITH_EDITOR
 void UPhysicalMaterial::PostEditChangeProperty(FPropertyChangedEvent& PropertyChangedEvent)
 {
diff --git a/Engine/Source/Runtime/Engine/Private/PhysicsEngine/PhysicalMaterialMask.cpp b/Engine/Source/Runtime/Engine/Private/PhysicsEngine/PhysicalMaterialMask.cpp
--- a/Engine/Source/Runtime/Engine/Private/PhysicsEngine/PhysicalMaterialMask.cpp
+++ b/Engine/Source/Runtime/Engine/Private/PhysicsEngine/PhysicalMaterialMask.cpp
@@ -27,6 +27,10 @@
 #include "Misc/Paths.h"
 #include "Misc/PackageName.h"
 
+#if WITH_CHAOS
+#include "Chaos/PhysicalMaterials.h"
+#endif
+
 #endif
 
 DEFINE_LOG_CATEGORY_STATIC(LogPhysicalMaterialMask, Log, All);
@@ -93,6 +97,11 @@ UPhysicalMaterialMask::UPhysicalMaterialMask(const FObjectInitializer& ObjectIni
 {
 }
 
+UPhysicalMaterialMask::UPhysicalMaterialMask(FVTableHelper& Helper)
+	: Super(Helper)
+{
+}
+
 UPhysicalMaterialMask::~UPhysicalMaterialMask() = default;
 
 #if WITH_EDITOR
diff --git a/Engine/Source/Runtime/Engine/Private/PhysicsEngine/PhysicsHandleComponent.cpp b/Engine/Source/Runtime/Engine/Private/PhysicsEngine/PhysicsHandleComponent.cpp
--- a/Engine/Source/Runtime/Engine/Private/PhysicsEngine/PhysicsHandleComponent.cpp
+++ b/Engine/Source/Runtime/Engine/Private/PhysicsEngine/PhysicsHandleComponent.cpp
@@ -11,6 +11,8 @@
 	#include "PhysXPublic.h"
 #endif // WITH_PHYSX
 
+#include "ChaosCheck.h"
+
 UPhysicsHandleComponent::UPhysicsHandleComponent(const FObjectInitializer& ObjectInitializer)
 	: Super(ObjectInitializer)
 {
diff --git a/Engine/Source/Runtime/Landscape/Classes/LandscapeHeightfieldCollisionComponent.h b/Engine/Source/Runtime/Landscape/Classes/LandscapeHeightfieldCollisionComponent.h
--- a/Engine/Source/Runtime/Landscape/Classes/LandscapeHeightfieldCollisionComponent.h
+++ b/Engine/Source/Runtime/Landscape/Classes/LandscapeHeightfieldCollisionComponent.h
@@ -42,6 +42,9 @@ class ULandscapeHeightfieldCollisionComponent : public UPrimitiveComponent
 {
 	GENERATED_UCLASS_BODY()
 
+	ULandscapeHeightfieldCollisionComponent(FVTableHelper& Helper);
+	virtual ~ULandscapeHeightfieldCollisionComponent();
+
 	/** List of layers painted on this component. Matches the WeightmapLayerAllocations array in the LandscapeComponent. */
 	UPROPERTY()
 	TArray<ULandscapeLayerInfoObject*> ComponentLayerInfos;
diff --git a/Engine/Source/Runtime/Landscape/Private/LandscapeCollision.cpp b/Engine/Source/Runtime/Landscape/Private/LandscapeCollision.cpp
--- a/Engine/Source/Runtime/Landscape/Private/LandscapeCollision.cpp
+++ b/Engine/Source/Runtime/Landscape/Private/LandscapeCollision.cpp
@@ -2612,6 +2612,14 @@ ULandscapeHeightfieldCollisionComponent::ULandscapeHeightfieldCollisionComponent
 	SetNetAddressable();
 }
 
+ULandscapeHeightfieldCollisionComponent::ULandscapeHeightfieldCollisionComponent(FVTableHelper& Helper)
+	: Super(Helper)
+{
+
+}
+
+ULandscapeHeightfieldCollisionComponent::~ULandscapeHeightfieldCollisionComponent() = default;
+
 ULandscapeComponent* ULandscapeHeightfieldCollisionComponent::GetRenderComponent() const
 {
 	return RenderComponent.Get();
