Make sure the skeletalmesh DDC key is always the same
#rnx

diff --git a/Engine/Source/Runtime/Engine/Public/Rendering/SkeletalMeshLODModel.h b/Engine/Source/Runtime/Engine/Public/Rendering/SkeletalMeshLODModel.h
--- a/Engine/Source/Runtime/Engine/Public/Rendering/SkeletalMeshLODModel.h
+++ b/Engine/Source/Runtime/Engine/Public/Rendering/SkeletalMeshLODModel.h
@@ -363,6 +363,8 @@ public:
 		, bIsRawSkeletalMeshBulkDataEmpty(true)
 		, BuildStringID(TEXT(""))
 	{
+		//Sice this ID is part of the DDC Key, we have to set it to an empty GUID not an empty string
+		RawSkeletalMeshBulkDataID = FGuid().ToString();
 		//Allocate the private mutex
 		BulkDataReadMutex = new FCriticalSection();
 	}
