submitting for Ethan. fix for check that would hit when seeking uncompressed audio to the last frame in a chunk.

diff --git a/Engine/Source/Runtime/Engine/Private/ADPCMAudioInfo.cpp b/Engine/Source/Runtime/Engine/Private/ADPCMAudioInfo.cpp
--- a/Engine/Source/Runtime/Engine/Private/ADPCMAudioInfo.cpp
+++ b/Engine/Source/Runtime/Engine/Private/ADPCMAudioInfo.cpp
@@ -745,12 +745,16 @@ bool FADPCMAudioInfo::StreamCompressedData(uint8* Destination, bool bLooping, ui
 
 				bSeekPending = false;
 			}
-
+			
 			uint32 DecompressedSamplesToCopy = FMath::Min<uint32>(
 				(CurrentChunkDataSize - CurrentChunkBufferOffset) / ChannelSampleSize,
 				BufferSize / ChannelSampleSize);
 
-			check(DecompressedSamplesToCopy > 0);
+			if (DecompressedSamplesToCopy == 0)
+			{
+				CurrentChunkBufferOffset = CurrentChunkDataSize;
+				continue;
+			}
 
 			// Ensure we don't go over the number of samples left in the audio data
 			if(DecompressedSamplesToCopy > TotalSamplesPerChannel - TotalSamplesStreamed)
