Submitting for Ethan:
Fix for seekable streaming being off for seeking very far into very long files, at a rate of one ADPCM block for every streaming chunk we are seeking through.

diff --git a/Engine/Source/Runtime/Engine/Private/ADPCMAudioInfo.cpp b/Engine/Source/Runtime/Engine/Private/ADPCMAudioInfo.cpp
--- a/Engine/Source/Runtime/Engine/Private/ADPCMAudioInfo.cpp
+++ b/Engine/Source/Runtime/Engine/Private/ADPCMAudioInfo.cpp
@@ -154,14 +154,19 @@ void FADPCMAudioInfo::SeekToTimeInternal(const float InSeekTime)
 			const int32 ChannelBlockSize = BlockSize * NumChannels;
 			for (uint32 BlockIndex = 0; BlockIndex < CurrentCompressedBlockIndex; ++BlockIndex)
 			{
-				if (CurrentChunkBufferOffset + ChannelBlockSize >= StreamingSoundWave->GetSizeOfChunk(CurrentChunkIndex))
+				const uint32 SizeOfChunk = StreamingSoundWave->GetSizeOfChunk(CurrentChunkIndex);
+				if (CurrentChunkBufferOffset + ChannelBlockSize >= SizeOfChunk)
 				{
+					const uint32 RemainderAfterEndOfChunk = CurrentChunkBufferOffset + ChannelBlockSize - SizeOfChunk;
+					ensureMsgf(RemainderAfterEndOfChunk == 0, TEXT("Found partial ADPCM block of %u samples- Please check FADPCMAudioFormat::SplitDataForStreaming for errors."), RemainderAfterEndOfChunk);
 					++CurrentChunkIndex;
 					CurrentChunkBufferOffset = 0;
 				}
-
-				// Always add chunks in NumChannels pairs
-				CurrentChunkBufferOffset += ChannelBlockSize;
+				else
+				{
+					// Always add chunks in NumChannels pairs
+					CurrentChunkBufferOffset += ChannelBlockSize;
+				}
 				
 				if (CurrentChunkIndex >= TotalStreamingChunks)
 				{
