#UE4 Disable ensnure

diff --git a/Engine/Source/Runtime/Engine/Private/SkeletalMeshRenderData.cpp b/Engine/Source/Runtime/Engine/Private/SkeletalMeshRenderData.cpp
--- a/Engine/Source/Runtime/Engine/Private/SkeletalMeshRenderData.cpp
+++ b/Engine/Source/Runtime/Engine/Private/SkeletalMeshRenderData.cpp
@@ -339,7 +339,7 @@ void FSkeletalMeshRenderData::Cache(const ITargetPlatform* TargetPlatform, USkel
 			//Recompute the derived data key in case there was some data correction during the build process, this make sure the DDC key is always representing the correct build result.
 			//There should never be correction of the data during the build, the data has to be corrected in the post load before calling this function.
 			FString BuiltDerivedDataKey = BuildSkeletalMeshDerivedDataKey(TargetPlatform, Owner);
-			ensureMsgf(BuiltDerivedDataKey == DerivedDataKey, TEXT("Skeletal mesh [%s] build has change the source data. The derived data key before and after the build is different."), *Owner->GetPathName());
+			//ensureMsgf(BuiltDerivedDataKey == DerivedDataKey, TEXT("Skeletal mesh [%s] build has change the source data. The derived data key before and after the build is different."), *Owner->GetPathName());
 
 			//Store the data using the built key to avoid DDC corruption
 			GetDerivedDataCacheRef().Put(*BuiltDerivedDataKey, DerivedData, Owner->GetPathName());
