Added missing include

diff --git a/Engine/Source/Runtime/Engine/Private/PhysicsEngine/CollisionQueryFilterCallback.cpp b/Engine/Source/Runtime/Engine/Private/PhysicsEngine/CollisionQueryFilterCallback.cpp
--- a/Engine/Source/Runtime/Engine/Private/PhysicsEngine/CollisionQueryFilterCallback.cpp
+++ b/Engine/Source/Runtime/Engine/Private/PhysicsEngine/CollisionQueryFilterCallback.cpp
@@ -14,6 +14,7 @@
 
 #include "ChaosInterfaceWrapperCore.h"
 #include "Physics/Experimental/ChaosInterfaceWrapper.h"
+#include "Chaos/GeometryParticles.h"
 
 ECollisionQueryHitType FCollisionQueryFilterCallback::CalcQueryHitType(const FCollisionFilterData& QueryFilter, const FCollisionFilterData& ShapeFilter, bool bPreFilter)
 {
