Remove chaos include guard macros

diff --git a/Engine/Plugins/Experimental/ChaosCloth/Source/Chaos/Public/ChaosCloth/ChaosClothingSimulation.h b/Engine/Plugins/Experimental/ChaosCloth/Source/Chaos/Public/ChaosCloth/ChaosClothingSimulation.h
--- a/Engine/Plugins/Experimental/ChaosCloth/Source/Chaos/Public/ChaosCloth/ChaosClothingSimulation.h
+++ b/Engine/Plugins/Experimental/ChaosCloth/Source/Chaos/Public/ChaosCloth/ChaosClothingSimulation.h
@@ -3,13 +3,11 @@
 
 #include "ClothingSimulation.h"
 #include "ClothingAsset.h"
-#define TEMP_HEADER_CHAOS_LEVEL_1
 #include "Chaos/ArrayCollectionArray.h"
 #include "Chaos/PBDEvolution.h"
 #include "Chaos/Transform.h"
 #include "Chaos/TriangleMesh.h"
 #include "Chaos/ChaosDebugDrawDeclares.h"
-#undef TEMP_HEADER_CHAOS_LEVEL_1
 #include "ChaosCloth/ChaosClothConfig.h"
 #include "Components/SkeletalMeshComponent.h"
 
diff --git a/Engine/Source/Editor/UnrealEd/Private/EditorEngine.cpp b/Engine/Source/Editor/UnrealEd/Private/EditorEngine.cpp
--- a/Engine/Source/Editor/UnrealEd/Private/EditorEngine.cpp
+++ b/Engine/Source/Editor/UnrealEd/Private/EditorEngine.cpp
@@ -222,9 +222,7 @@
 #include "UObject/UnrealType.h"
 
 #if WITH_CHAOS
-#define CHAOS_INCLUDE_LEVEL_1
 #include "ChaosSolversModule.h"
-#undef CHAOS_INCLUDE_LEVEL_1
 #endif
 
 
diff --git a/Engine/Source/Runtime/Engine/Private/Collision/CollisionDebugDrawing.cpp b/Engine/Source/Runtime/Engine/Private/Collision/CollisionDebugDrawing.cpp
--- a/Engine/Source/Runtime/Engine/Private/Collision/CollisionDebugDrawing.cpp
+++ b/Engine/Source/Runtime/Engine/Private/Collision/CollisionDebugDrawing.cpp
@@ -6,12 +6,10 @@
 #include "DrawDebugHelpers.h"
 #include "PhysXPublic.h"
 
-#define CHAOS_INCLUDE_LEVEL_1
 #include "Chaos/Box.h"
 #include "Chaos/Capsule.h"
 #include "Chaos/ImplicitObject.h"
 #include "Chaos/Sphere.h"
-#undef CHAOS_INCLUDE_LEVEL_1
 
 
 #if ENABLE_DRAW_DEBUG
diff --git a/Engine/Source/Runtime/Engine/Private/PhysicsEngine/BodySetup.cpp b/Engine/Source/Runtime/Engine/Private/PhysicsEngine/BodySetup.cpp
--- a/Engine/Source/Runtime/Engine/Private/PhysicsEngine/BodySetup.cpp
+++ b/Engine/Source/Runtime/Engine/Private/PhysicsEngine/BodySetup.cpp
@@ -22,10 +22,8 @@
 #include "UObject/PropertyPortFlags.h"
 #include "Components/SplineMeshComponent.h"
 
-#define CHAOS_INCLUDE_LEVEL_1
 #include "ChaosCheck.h"
 #include "Chaos/Convex.h"
-#undef CHAOS_INCLUDE_LEVEL_1
 
 #include "PhysXCookHelper.h"
 
@@ -47,16 +45,12 @@
 #include "ProfilingDebugging/CookStats.h"
 #include "UObject/AnimPhysObjectVersion.h"
 
-#define CHAOS_INCLUDE_LEVEL_1
 #include "Chaos/TriangleMeshImplicitObject.h"
-#undef CHAOS_INCLUDE_LEVEL_1
 
 #if WITH_CHAOS
 	#include "Experimental/ChaosDerivedData.h"
 	#include "Physics/Experimental/ChaosDerivedDataReader.h"
-#define CHAOS_INCLUDE_LEVEL_1
 	#include "Chaos/CollisionConvexMesh.h"
-#undef CHAOS_INCLUDE_LEVEL_1
 #endif
 
 /** Helper for enum output... */
diff --git a/Engine/Source/Runtime/Engine/Public/Physics/PhysicsInterfaceCore.h b/Engine/Source/Runtime/Engine/Public/Physics/PhysicsInterfaceCore.h
--- a/Engine/Source/Runtime/Engine/Public/Physics/PhysicsInterfaceCore.h
+++ b/Engine/Source/Runtime/Engine/Public/Physics/PhysicsInterfaceCore.h
@@ -15,10 +15,8 @@
 
 #elif WITH_CHAOS
 
-#define TEMP_HEADER_CHAOS_LEVEL_1
 #include "Physics/Experimental/PhysInterface_Chaos.h"
 #include "Physics/Experimental/PhysScene_Chaos.h"
-#undef TEMP_HEADER_CHAOS_LEVEL_1
 
 #else
 
diff --git a/Engine/Source/Runtime/Experimental/Chaos/Private/kDOP.h b/Engine/Source/Runtime/Experimental/Chaos/Private/kDOP.h
--- a/Engine/Source/Runtime/Experimental/Chaos/Private/kDOP.h
+++ b/Engine/Source/Runtime/Experimental/Chaos/Private/kDOP.h
@@ -2,7 +2,6 @@
 
 #pragma once
 
-#include "Chaos/IncludeLvl1.inl"
 #include "CoreMinimal.h"
 #include "Containers/ChunkedArray.h"
 
diff --git a/Engine/Source/Runtime/Experimental/Chaos/Public/Chaos/AABB.h b/Engine/Source/Runtime/Experimental/Chaos/Public/Chaos/AABB.h
--- a/Engine/Source/Runtime/Experimental/Chaos/Public/Chaos/AABB.h
+++ b/Engine/Source/Runtime/Experimental/Chaos/Public/Chaos/AABB.h
@@ -1,7 +1,6 @@
 // Copyright Epic Games, Inc. All Rights Reserved.
 #pragma once
 
-#include "Chaos/IncludeLvl1.inl"
 #include "Chaos/Vector.h"
 #include "Chaos/Matrix.h"
 #include "Chaos/Plane.h"
diff --git a/Engine/Source/Runtime/PhysicsCore/Public/PhysTestSerializer.h b/Engine/Source/Runtime/PhysicsCore/Public/PhysTestSerializer.h
--- a/Engine/Source/Runtime/PhysicsCore/Public/PhysTestSerializer.h
+++ b/Engine/Source/Runtime/PhysicsCore/Public/PhysTestSerializer.h
@@ -2,9 +2,6 @@
 
 #pragma once
 
-#include "Chaos/IncludeLvl1.inl"
-CHAOS_LEVEL_CHECK
-
 #include "CoreMinimal.h"
 #include "PhysicsInterfaceDeclaresCore.h"
 
diff --git a/Engine/Source/Runtime/PhysicsSQ/Public/SQVerifier.h b/Engine/Source/Runtime/PhysicsSQ/Public/SQVerifier.h
--- a/Engine/Source/Runtime/PhysicsSQ/Public/SQVerifier.h
+++ b/Engine/Source/Runtime/PhysicsSQ/Public/SQVerifier.h
@@ -5,8 +5,6 @@
 #include "ProfilingDebugging/ScopedTimers.h"
 
 #if WITH_PHYSX
-#include "Chaos/IncludeLvl1.inl"
-CHAOS_LEVEL_CHECK
 #include "Chaos/PBDRigidsEvolutionGBF.h"
 #include "PhysicsCore/Public/PhysicsInterfaceUtilsCore.h"
 #include "PhysTestSerializer.h"
