Make sure if SetReplicates is called on an already replicating Actor, we force property comparisons to capture the role change.

diff --git a/Engine/Source/Runtime/Engine/Private/Actor.cpp b/Engine/Source/Runtime/Engine/Private/Actor.cpp
--- a/Engine/Source/Runtime/Engine/Private/Actor.cpp
+++ b/Engine/Source/Runtime/Engine/Private/Actor.cpp
@@ -3264,14 +3264,18 @@ void AActor::SetReplicates(bool bInReplicates)
 		RemoteRole = bInReplicates ? ROLE_SimulatedProxy : ROLE_None;
 		bReplicates = bInReplicates;
 
-		// Only call into net driver if we just started replicating changed
-		// This actor should already be in the Network Actors List if it was already replicating.
-		if (bNewlyReplicates)
+		if (bReplicates)
 		{
 			// GetWorld will return nullptr on CDO, FYI
-			if (UWorld* MyWorld = GetWorld())		
+			if (UWorld* MyWorld = GetWorld())
 			{
-				MyWorld->AddNetworkActor(this);
+				// Only call into net driver if we just started replicating changed
+				// This actor should already be in the Network Actors List if it was already replicating.
+				if (bNewlyReplicates)
+				{
+					MyWorld->AddNetworkActor(this);
+				}
+
 				ForcePropertyCompare();
 			}
 		}
