Remove UnrealWatchdog (simplifying NotForLicensees dependencies)

diff --git a/Engine/Build/InstalledEngineBuild.xml b/Engine/Build/InstalledEngineBuild.xml
--- a/Engine/Build/InstalledEngineBuild.xml
+++ b/Engine/Build/InstalledEngineBuild.xml
@@ -488,7 +488,6 @@
 			<Compile Target="CrashReportClient" Configuration="Shipping" Platform="Win64" Tag="#Build Tools Win64" Arguments="$(CrashReporterCompileArgs)"/>
 			<Compile Target="CrashReportClientEditor" Configuration="Shipping" Platform="Win64" Tag="#Build Tools Win64" Arguments="$(CrashReporterCompileArgs)"/>
 			<Tag Files="#NotForLicensee Build Tools Win64" With="#Build Tools Win64"/>
-			<Compile Target="UnrealWatchdog" Configuration="Shipping" Platform="Win64" Tag="#Build Tools Win64"/>
 			<Compile Target="ShaderCompileWorker" Configuration="Development" Platform="Win64" Tag="#Build Tools Win64"/>
 			<Compile Target="UnrealCEFSubProcess" Configuration="Development" Platform="Win64" Tag="#Build Tools Win64"/>
 			<Compile Target="UnrealInsights" Configuration="Development" Platform="Win64" Tag="#Build Tools Win64"/>
diff --git a/Engine/Config/BaseEngine.ini b/Engine/Config/BaseEngine.ini
--- a/Engine/Config/BaseEngine.ini
+++ b/Engine/Config/BaseEngine.ini
@@ -2370,12 +2370,6 @@ r.MeshReductionModule="QuadricMeshReduction"
 +InternalOnlyPaths=(Path="/Engine/Sequencer")
 +InternalOnlyPaths=(Path="/Engine/NotForLicensees")
 
-[EngineSessionManager]
-UseWatchdogMTBF=false
-AllowWatchdogDialogs=false
-AllowWatchdogDetectHangs=false
-WatchdogMinimumHangSeconds=120
-
 [/Script/LevelSequence.LevelSequence]
 DefaultCompletionMode=RestoreState
 
diff --git a/Engine/Source/Programs/UnrealWatchdog/Private/Linux/UnrealWatchdogMainLinux.cpp b/Engine/Source/Programs/UnrealWatchdog/Private/Linux/UnrealWatchdogMainLinux.cpp
deleted file mode 100644
diff --git a/Engine/Source/Programs/UnrealWatchdog/Private/Mac/UnrealWatchdogMainMac.cpp b/Engine/Source/Programs/UnrealWatchdog/Private/Mac/UnrealWatchdogMainMac.cpp
deleted file mode 100644
diff --git a/Engine/Source/Programs/UnrealWatchdog/Private/UnrealWatchdog.cpp b/Engine/Source/Programs/UnrealWatchdog/Private/UnrealWatchdog.cpp
deleted file mode 100644
diff --git a/Engine/Source/Programs/UnrealWatchdog/Private/UnrealWatchdog.h b/Engine/Source/Programs/UnrealWatchdog/Private/UnrealWatchdog.h
deleted file mode 100644
diff --git a/Engine/Source/Programs/UnrealWatchdog/Private/WatchdogAnalytics.cpp b/Engine/Source/Programs/UnrealWatchdog/Private/WatchdogAnalytics.cpp
deleted file mode 100644
diff --git a/Engine/Source/Programs/UnrealWatchdog/Private/WatchdogAnalytics.h b/Engine/Source/Programs/UnrealWatchdog/Private/WatchdogAnalytics.h
deleted file mode 100644
diff --git a/Engine/Source/Programs/UnrealWatchdog/Private/Windows/UnrealWatchdogMainWindows.cpp b/Engine/Source/Programs/UnrealWatchdog/Private/Windows/UnrealWatchdogMainWindows.cpp
deleted file mode 100644
diff --git a/Engine/Source/Programs/UnrealWatchdog/UnrealWatchdog.Build.cs b/Engine/Source/Programs/UnrealWatchdog/UnrealWatchdog.Build.cs
deleted file mode 100644
diff --git a/Engine/Source/Programs/UnrealWatchdog/UnrealWatchdog.Target.cs b/Engine/Source/Programs/UnrealWatchdog/UnrealWatchdog.Target.cs
deleted file mode 100644
diff --git a/Engine/Source/Runtime/Engine/Private/EngineSessionManager.cpp b/Engine/Source/Runtime/Engine/Private/EngineSessionManager.cpp
--- a/Engine/Source/Runtime/Engine/Private/EngineSessionManager.cpp
+++ b/Engine/Source/Runtime/Engine/Private/EngineSessionManager.cpp
@@ -49,9 +49,7 @@ namespace SessionManagerDefs
 	static const FString SessionRecordListSection(TEXT("List"));
 	static const FString EditorSessionRecordSectionPrefix(TEXT("Unreal Engine/Editor Sessions/"));
 	static const FString GameSessionRecordSectionPrefix(TEXT("Unreal Engine/Game Sessions/"));
-	static const FString WatchdogRecordSectionPrefix(TEXT("Unreal Engine/Watchdog/"));
 	static const FString SessionsVersionString(TEXT("1_3"));
-	static const FString WatchdogVersionString(TEXT("1_0"));
 	static const FString ModeStoreKey(TEXT("Mode"));
 	static const FString ProjectNameStoreKey(TEXT("ProjectName"));
 	static const FString CommandLineStoreKey(TEXT("CommandLine"));
@@ -214,25 +212,11 @@ void FEngineSessionManager::Tick(float DeltaTime)
 					CurrentSession.bWasEverDebugger = true;
 
 					FPlatformMisc::SetStoredValue(SessionManagerDefs::StoreId, CurrentSessionSectionName, SessionManagerDefs::WasDebuggerStoreKey, SessionManagerDefs::TrueValueString);
-
-#if PLATFORM_SUPPORTS_WATCHDOG
-					if (!WatchdogSectionName.IsEmpty())
-					{
-						FPlatformMisc::SetStoredValue(SessionManagerDefs::StoreId, WatchdogSectionName, SessionManagerDefs::WasDebuggerStoreKey, SessionManagerDefs::TrueValueString);
-					}
-#endif
 				}
 			}
 
 			const FString TimestampString = EngineSessionManagerUtils::TimestampToString(FDateTime::UtcNow());
 			FPlatformMisc::SetStoredValue(SessionManagerDefs::StoreId, CurrentSessionSectionName, SessionManagerDefs::TimestampStoreKey, TimestampString);
-
-#if PLATFORM_SUPPORTS_WATCHDOG
-			if (!WatchdogSectionName.IsEmpty())
-			{
-				FPlatformMisc::SetStoredValue(SessionManagerDefs::StoreId, WatchdogSectionName, SessionManagerDefs::TimestampStoreKey, TimestampString);
-			}
-#endif
 		}
 	}
 }
@@ -260,16 +244,6 @@ void FEngineSessionManager::Shutdown()
 		if (!CurrentSession.bCrashed)
 		{
 			DeleteStoredRecordValues(CurrentSessionSectionName);
-
-#if PLATFORM_SUPPORTS_WATCHDOG
-			if (!WatchdogSectionName.IsEmpty())
-			{
-				const FString& ShutdownType = CurrentSession.bIsTerminating ? SessionManagerDefs::TerminatedSessionToken : SessionManagerDefs::ShutdownSessionToken;
-				FPlatformMisc::SetStoredValue(SessionManagerDefs::StoreId, WatchdogSectionName, SessionManagerDefs::StatusStoreKey, ShutdownType);
-				FPlatformMisc::SetStoredValue(SessionManagerDefs::StoreId, WatchdogSectionName, SessionManagerDefs::TimestampStoreKey, EngineSessionManagerUtils::TimestampToString(FDateTime::UtcNow()));
-				WatchdogSectionName.Empty();
-			}
-#endif
 		}
 
 		bInitializedRecords = false;
@@ -574,15 +548,6 @@ void FEngineSessionManager::CreateAndWriteRecordForSession()
 	FPlatformMisc::SetStoredValue(SessionManagerDefs::StoreId, CurrentSessionSectionName, SessionManagerDefs::PlatformProcessIDKey, CurrentProcessIDString);
 
 	SessionRecords.Add(CurrentSession);
-
-#if PLATFORM_SUPPORTS_WATCHDOG
-	bool bUseWatchdog = false;
-	GConfig->GetBool(TEXT("EngineSessionManager"), TEXT("UseWatchdogMTBF"), bUseWatchdog, GEngineIni);
-	if ((!CurrentSession.bWasEverDebugger && bUseWatchdog && !FParse::Param(FCommandLine::Get(), TEXT("NoWatchdog"))) || FParse::Param(FCommandLine::Get(), TEXT("ForceWatchdog")))
-	{
-		StartWatchdog(ModeString, CurrentSession.ProjectName, FPlatformProperties::PlatformName(), CurrentSession.SessionId, CurrentSession.EngineVersion);
-	}
-#endif
 }
 
 extern CORE_API bool GIsGPUCrashed;
@@ -594,14 +559,6 @@ void FEngineSessionManager::OnCrashing()
 		CurrentSession.bGPUCrashed = GIsGPUCrashed;
 		FPlatformMisc::SetStoredValue(SessionManagerDefs::StoreId, CurrentSessionSectionName, SessionManagerDefs::CrashStoreKey, SessionManagerDefs::TrueValueString);
 		FPlatformMisc::SetStoredValue(SessionManagerDefs::StoreId, CurrentSessionSectionName, SessionManagerDefs::GPUCrashStoreKey, CurrentSession.bGPUCrashed ? SessionManagerDefs::TrueValueString : SessionManagerDefs::FalseValueString);
-
-#if PLATFORM_SUPPORTS_WATCHDOG
-		if (!WatchdogSectionName.IsEmpty())
-		{
-			FPlatformMisc::SetStoredValue(SessionManagerDefs::StoreId, WatchdogSectionName, SessionManagerDefs::StatusStoreKey, SessionManagerDefs::CrashSessionToken);
-			FPlatformMisc::SetStoredValue(SessionManagerDefs::StoreId, WatchdogSectionName, SessionManagerDefs::TimestampStoreKey, EngineSessionManagerUtils::TimestampToString(FDateTime::UtcNow()));
-		}
-#endif
 	}
 }
 
@@ -654,13 +611,6 @@ void FEngineSessionManager::OnTerminate()
 			// In these cases, shut down the engine session so it won't send an abnormal shutdown report.
 			Shutdown();
 		}
-#if PLATFORM_SUPPORTS_WATCHDOG
-		else if (!WatchdogSectionName.IsEmpty())
-		{
-			FPlatformMisc::SetStoredValue(SessionManagerDefs::StoreId, WatchdogSectionName, SessionManagerDefs::StatusStoreKey, SessionManagerDefs::TerminatedSessionToken);
-			FPlatformMisc::SetStoredValue(SessionManagerDefs::StoreId, WatchdogSectionName, SessionManagerDefs::TimestampStoreKey, EngineSessionManagerUtils::TimestampToString(FDateTime::UtcNow()));
-		}
-#endif
 	}
 }
 
@@ -695,14 +645,6 @@ void FEngineSessionManager::OnUserActivity(const FUserActivity& UserActivity)
 	{
 		CurrentSession.CurrentUserActivity = GetUserActivityString();
 		FPlatformMisc::SetStoredValue(SessionManagerDefs::StoreId, CurrentSessionSectionName, SessionManagerDefs::UserActivityStoreKey, CurrentSession.CurrentUserActivity);
-
-#if PLATFORM_SUPPORTS_WATCHDOG
-		if (!WatchdogSectionName.IsEmpty())
-		{
-			FPlatformMisc::SetStoredValue(SessionManagerDefs::StoreId, WatchdogSectionName, SessionManagerDefs::UserActivityStoreKey, CurrentSession.CurrentUserActivity);
-			FPlatformMisc::SetStoredValue(SessionManagerDefs::StoreId, WatchdogSectionName, SessionManagerDefs::TimestampStoreKey, EngineSessionManagerUtils::TimestampToString(FDateTime::UtcNow()));
-		}
-#endif
 	}
 }
 
@@ -718,115 +660,4 @@ FString FEngineSessionManager::GetUserActivityString() const
 	return UserActivity.ActionName;
 }
 
-#if PLATFORM_SUPPORTS_WATCHDOG
-
-/**
- * @EventName Engine.StartWatchdog
- *
- * @Trigger Event raised by EngineSessionManager as part of MTBF tracking. Records an attempt to start the UnrealWatchdog process.
- *
- * @Type Client
- * @Owner Chris.Wood
- *
- * @EventParam RunType - Editor or Game
- * @EventParam ProjectName - Project for the session.
- * @EventParam Platform - Windows, Mac, Linux
- * @EventParam SessionId - Analytics SessionID of the session.
- * @EventParam EngineVersion - EngineVersion of the session.
- * @EventParam IsInternalBuild - internal Epic build environment or not? Calls FEngineBuildSettings::IsInternalBuild(). Value is Yes or No.
- * @EventParam Outcome - Whether the watchdog was started successfully. One of Succeeded, CreateProcFailed or MissingBinaryFailed.
- *
- * @Comments Currently only runs Watchdog when MTBF is enabled, we aren't debugging, we're a DESKTOP platform and watchdog is specifically enabled via config or command line arg.
- */
-void FEngineSessionManager::StartWatchdog(const FString& RunType, const FString& ProjectName, const FString& PlatformName, const FString& SessionId, const FString& EngineVersion)
-{
-	uint32 ProcessId =  FPlatformProcess::GetCurrentProcessId();
-	const int SuccessfulRtnCode = 0;	// hardcode this for now, zero might not always be correct
-
-	FString LogFilePath = FPaths::ConvertRelativePathToFull(FPlatformOutputDevices::GetAbsoluteLogFilename());
-
-	FString WatchdogClientArguments =
-		FString::Printf(TEXT(
-			"-PID=%u -RunType=%s -ProjectName=\"%s\" -Platform=%s -SessionId=%s -EngineVersion=%s -SuccessfulRtnCode=%d -LogPath=\"%s\""),
-			ProcessId, *RunType, *ProjectName, *PlatformName, *SessionId, *EngineVersion, SuccessfulRtnCode, *LogFilePath);
-
-	bool bAllowWatchdogDetectHangs = false;
-	GConfig->GetBool(TEXT("EngineSessionManager"), TEXT("AllowWatchdogDetectHangs"), bAllowWatchdogDetectHangs, GEngineIni);
-
-	if (bAllowWatchdogDetectHangs)
-	{
-		int HangSeconds = 120;
-		GConfig->GetInt(TEXT("EngineSessionManager"), TEXT("WatchdogMinimumHangSeconds"), HangSeconds, GEngineIni);
-
-		WatchdogClientArguments.Append(FString::Printf(TEXT(" -DetectHangs -HangSeconds=%d"), HangSeconds));
-	}
-
-	if (FEngineBuildSettings::IsInternalBuild())
-	{
-		// Suppress the watchdog dialogs if this engine session should never show interactive UI
-		if (!FApp::IsUnattended() && !IsRunningDedicatedServer() && FApp::CanEverRender())
-		{
-			// Only show watchdog dialogs if it's set in config
-			bool bAllowWatchdogDialogs = false;
-			GConfig->GetBool(TEXT("EngineSessionManager"), TEXT("AllowWatchdogDialogs"), bAllowWatchdogDialogs, GEngineIni);
-
-			if (bAllowWatchdogDialogs)
-			{
-				WatchdogClientArguments.Append(TEXT(" -AllowDialogs"));
-			}
-		}
-	}
-
-	FString WatchdogPath = FPaths::ConvertRelativePathToFull(FPlatformProcess::GenerateApplicationPath(TEXT("UnrealWatchdog"), EBuildConfiguration::Development));
-
-	TArray< FAnalyticsEventAttribute > WatchdogStartedAttributes;
-	WatchdogStartedAttributes.Add(FAnalyticsEventAttribute(TEXT("RunType"), RunType));
-	WatchdogStartedAttributes.Add(FAnalyticsEventAttribute(TEXT("ProjectName"), ProjectName));
-	WatchdogStartedAttributes.Add(FAnalyticsEventAttribute(TEXT("Platform"), PlatformName));
-	WatchdogStartedAttributes.Add(FAnalyticsEventAttribute(TEXT("SessionId"), SessionId));
-	WatchdogStartedAttributes.Add(FAnalyticsEventAttribute(TEXT("IsInternalBuild"), FEngineBuildSettings::IsInternalBuild() ? TEXT("Yes") : TEXT("No")));
-
-	if (FPaths::FileExists(WatchdogPath))
-	{
-		FProcHandle WatchdogProcessHandle = FPlatformProcess::CreateProc(*WatchdogPath, *WatchdogClientArguments, true, true, false, NULL, 0, NULL, NULL);
-
-		if (WatchdogProcessHandle.IsValid())
-		{
-			FString WatchdogStartTimeString = EngineSessionManagerUtils::TimestampToString(FDateTime::UtcNow());
-			FString WasDebuggerString = CurrentSession.bWasEverDebugger ? SessionManagerDefs::TrueValueString : SessionManagerDefs::FalseValueString;
-
-			WatchdogStartedAttributes.Add(FAnalyticsEventAttribute(TEXT("Outcome"), TEXT("Succeeded")));
-			UE_LOG(LogEngineSessionManager, Log, TEXT("Started UnrealWatchdog for process id %u"), ProcessId);
-
-			WatchdogSectionName = GetWatchdogStoreSectionString(ProcessId);
-
-			FPlatformMisc::SetStoredValue(SessionManagerDefs::StoreId, WatchdogSectionName, SessionManagerDefs::CommandLineStoreKey, FCommandLine::GetOriginalForLogging());
-			FPlatformMisc::SetStoredValue(SessionManagerDefs::StoreId, WatchdogSectionName, SessionManagerDefs::StartupTimeStoreKey, WatchdogStartTimeString);
-			FPlatformMisc::SetStoredValue(SessionManagerDefs::StoreId, WatchdogSectionName, SessionManagerDefs::TimestampStoreKey, WatchdogStartTimeString);
-			FPlatformMisc::SetStoredValue(SessionManagerDefs::StoreId, WatchdogSectionName, SessionManagerDefs::StatusStoreKey, SessionManagerDefs::RunningSessionToken);
-			FPlatformMisc::SetStoredValue(SessionManagerDefs::StoreId, WatchdogSectionName, SessionManagerDefs::UserActivityStoreKey, CurrentSession.CurrentUserActivity);
-			FPlatformMisc::SetStoredValue(SessionManagerDefs::StoreId, WatchdogSectionName, SessionManagerDefs::WasDebuggerStoreKey, WasDebuggerString);
-		}
-		else
-		{
-			WatchdogStartedAttributes.Add(FAnalyticsEventAttribute(TEXT("Outcome"), TEXT("CreateProcFailed")));
-			UE_LOG(LogEngineSessionManager, Warning, TEXT("Unable to start UnrealWatchdog.exe. CreateProc failed."));
-		}
-	}
-	else
-	{
-		WatchdogStartedAttributes.Add(FAnalyticsEventAttribute(TEXT("Outcome"), TEXT("MissingBinaryFailed")));
-		UE_LOG(LogEngineSessionManager, Log, TEXT("Unable to start UnrealWatchdog.exe. File not found."));
-	}
-
-	FEngineAnalytics::GetProvider().RecordEvent(TEXT("Engine.StartWatchdog"), WatchdogStartedAttributes);
-}
-
-FString FEngineSessionManager::GetWatchdogStoreSectionString(uint32 InPID)
-{
-	return FString::Printf(TEXT("%s%s/%u"), *SessionManagerDefs::WatchdogRecordSectionPrefix, *SessionManagerDefs::WatchdogVersionString, InPID);
-}
-
-#endif // PLATFORM_SUPPORTS_WATCHDOG
-
 #undef LOCTEXT_NAMESPACE
diff --git a/Engine/Source/Runtime/Engine/Private/EngineSessionManager.h b/Engine/Source/Runtime/Engine/Private/EngineSessionManager.h
--- a/Engine/Source/Runtime/Engine/Private/EngineSessionManager.h
+++ b/Engine/Source/Runtime/Engine/Private/EngineSessionManager.h
@@ -4,10 +4,6 @@
 
 #include "CoreMinimal.h"
 
-// TODO: implement Watchdog on Mac and Linux and replace Windows check with desktop check below
-//#define PLATFORM_SUPPORTS_WATCHDOG		PLATFORM_DESKTOP
-#define PLATFORM_SUPPORTS_WATCHDOG		PLATFORM_WINDOWS
-
 struct FUserActivity;
 
 enum class EEngineSessionManagerMode
@@ -82,11 +78,6 @@ private:
 	void OnVanillaStateChanged(bool bIsVanilla);
 	FString GetUserActivityString() const;
 
-#if PLATFORM_SUPPORTS_WATCHDOG
-	void StartWatchdog(const FString& RunType, const FString& ProjectName, const FString& PlatformName, const FString& SessionId, const FString& EngineVersion);
-	FString GetWatchdogStoreSectionString(uint32 InPID);
-#endif
-
 private:
 	EEngineSessionManagerMode Mode;
 	bool bInitializedRecords;
@@ -95,8 +86,4 @@ private:
 	FSessionRecord CurrentSession;
 	FString CurrentSessionSectionName;
 	TArray<FSessionRecord> SessionRecords;
-
-#if PLATFORM_SUPPORTS_WATCHDOG
-	FString WatchdogSectionName;
-#endif
 };
