Editor: Allow debug drawing on the dedicated server to be optionally shown on connected clients when playing in the editor
- Introduces a new show flag ServerDrawDebug which can be toggled per client to control whether or not a DrawDebugXYZ call originating on the dedicated server will show up on that client
- Editor Preferences..Level Editor..Play has new options to control the default value of the ServerDrawDebug flag and tint color/strength to differentiate server vs client debug draw colors (e.g., blend server draws towards black or white)

diff --git a/Engine/Source/Editor/UnrealEd/Classes/Settings/LevelEditorPlaySettings.h b/Engine/Source/Editor/UnrealEd/Classes/Settings/LevelEditorPlaySettings.h
--- a/Engine/Source/Editor/UnrealEd/Classes/Settings/LevelEditorPlaySettings.h
+++ b/Engine/Source/Editor/UnrealEd/Classes/Settings/LevelEditorPlaySettings.h
@@ -365,7 +365,27 @@ private:
 	UPROPERTY(config)
 	FString AdditionalLaunchOptions;
 
+	/** Controls the default value of the show flag ServerDrawDebug */
+	UPROPERTY(config, EditAnywhere, Category = MultiplayerOptions)
+	bool bShowServerDebugDrawingByDefault;
+
+	/** How strongly debug drawing originating from the server will be biased towards the tint color */
+	UPROPERTY(config, EditAnywhere, Category=MultiplayerOptions, meta=(ClampMin=0, ClampMax=1, UIMin=0, UIMax=1))
+	float ServerDebugDrawingColorTintStrength;
+
+	/** Debug drawing originating from the server will be biased towards this color */
+	UPROPERTY(config, EditAnywhere, Category=MultiplayerOptions)
+	FLinearColor ServerDebugDrawingColorTint;
+
+private:
+	void PushDebugDrawingSettings();
+
 public:
+	bool ShowServerDebugDrawingByDefault() const
+	{
+		return bShowServerDebugDrawingByDefault;
+	}
+
 	/** Additional options that will be passed to the server as arguments, for example -debug. Only works with separate process servers. */
 	UPROPERTY(config, EditAnywhere, Category = "Multiplayer Options|Server")
 	FString AdditionalServerLaunchParameters;
diff --git a/Engine/Source/Editor/UnrealEd/Private/PlayLevel.cpp b/Engine/Source/Editor/UnrealEd/Private/PlayLevel.cpp
--- a/Engine/Source/Editor/UnrealEd/Private/PlayLevel.cpp
+++ b/Engine/Source/Editor/UnrealEd/Private/PlayLevel.cpp
@@ -2697,6 +2697,8 @@ UGameInstance* UEditorEngine::CreateInnerProcessPIEGameInstance(FRequestPlaySess
 		ViewportClient = NewObject<UGameViewportClient>(this, GameViewportClientClass);
 		ViewportClient->Init(*PieWorldContext, GameInstance, bCreateNewAudioDevice);
 
+		ViewportClient->EngineShowFlags.SetServerDrawDebug(PlayInSettings->ShowServerDebugDrawingByDefault());
+
 		if (!InParams.EditorPlaySettings->EnableGameSound)
 		{
 			if (FAudioDeviceHandle GameInstanceAudioDevice = GameInstance->GetWorld()->GetAudioDevice())
diff --git a/Engine/Source/Editor/UnrealEd/Private/Settings/SettingsClasses.cpp b/Engine/Source/Editor/UnrealEd/Private/Settings/SettingsClasses.cpp
--- a/Engine/Source/Editor/UnrealEd/Private/Settings/SettingsClasses.cpp
+++ b/Engine/Source/Editor/UnrealEd/Private/Settings/SettingsClasses.cpp
@@ -38,6 +38,7 @@
 #include "DeviceProfiles/DeviceProfileManager.h"
 #include "DesktopPlatformModule.h"
 #include "InstalledPlatformInfo.h"
+#include "DrawDebugHelpers.h"
 
 #define LOCTEXT_NAMESPACE "SettingsClasses"
 
@@ -478,6 +479,21 @@ PRAGMA_ENABLE_DEPRECATION_WARNINGS
 	CenterNewWindow = true;
 
 	EnablePIEEnterAndExitSounds = false;
+
+	bShowServerDebugDrawingByDefault = true;
+	ServerDebugDrawingColorTintStrength = 0.0f;
+	ServerDebugDrawingColorTint = FLinearColor(0.0f, 0.0f, 0.0f, 1.0f);
+}
+
+void ULevelEditorPlaySettings::PushDebugDrawingSettings()
+{
+#if ENABLE_DRAW_DEBUG
+	extern ENGINE_API float GServerDrawDebugColorTintStrength;
+	extern ENGINE_API FLinearColor GServerDrawDebugColorTint;
+
+	GServerDrawDebugColorTintStrength = ServerDebugDrawingColorTintStrength;
+	GServerDrawDebugColorTint = ServerDebugDrawingColorTint;
+#endif
 }
 
 void ULevelEditorPlaySettings::PostEditChangeProperty(struct FPropertyChangedEvent& PropertyChangedEvent)
@@ -500,6 +516,22 @@ void ULevelEditorPlaySettings::PostEditChangeProperty(struct FPropertyChangedEve
 		NetworkEmulationSettings.OnPostEditChange(PropertyChangedEvent);
 	}
 
+	PushDebugDrawingSettings();
+	if (PropertyChangedEvent.MemberProperty && PropertyChangedEvent.MemberProperty->GetFName() == GET_MEMBER_NAME_CHECKED(ULevelEditorPlaySettings, bShowServerDebugDrawingByDefault))
+	{
+		// If the show option is turned on or off, force it on or off in any active PIE instances too as a QOL aid so they don't have to stop and restart PIE again for it to take effect
+		for (const FWorldContext& WorldContext : GEngine->GetWorldContexts())
+		{
+			if ((WorldContext.WorldType == EWorldType::PIE) &&
+				(WorldContext.World() != nullptr) &&
+				(WorldContext.World()->GetNetMode() == NM_Client) &&
+				(WorldContext.GameViewport != nullptr))
+			{
+				WorldContext.GameViewport->EngineShowFlags.SetServerDrawDebug(bShowServerDebugDrawingByDefault);
+			}
+		}
+	}
+
 	Super::PostEditChangeProperty(PropertyChangedEvent);
 }
 
@@ -515,6 +547,8 @@ void ULevelEditorPlaySettings::PostInitProperties()
 #if WITH_EDITOR
 	FCoreDelegates::OnSafeFrameChangedEvent.AddUObject(this, &ULevelEditorPlaySettings::UpdateCustomSafeZones);
 #endif
+
+	PushDebugDrawingSettings();
 }
 
 bool ULevelEditorPlaySettings::CanEditChange(const FProperty* InProperty) const
diff --git a/Engine/Source/Runtime/Engine/Private/DrawDebugHelpers.cpp b/Engine/Source/Runtime/Engine/Private/DrawDebugHelpers.cpp
--- a/Engine/Source/Runtime/Engine/Private/DrawDebugHelpers.cpp
+++ b/Engine/Source/Runtime/Engine/Private/DrawDebugHelpers.cpp
@@ -16,6 +16,52 @@
 
 #if ENABLE_DRAW_DEBUG
 
+ENGINE_API float GServerDrawDebugColorTintStrength = 0.75f;
+ENGINE_API FLinearColor GServerDrawDebugColorTint(0.0f, 0.0f, 0.0f, 1.0f);
+
+#if WITH_EDITOR
+
+FColor AdjustColorForServer(const FColor InColor)
+{
+	if (GServerDrawDebugColorTintStrength > 0.0f)
+	{
+		return FMath::Lerp(FLinearColor::FromSRGBColor(InColor), GServerDrawDebugColorTint, GServerDrawDebugColorTintStrength).ToFColor(/*bSRGB=*/ true);
+	}
+	else
+	{
+		return InColor;
+	}
+}
+
+bool CanDrawServerDebugInContext(const FWorldContext& WorldContext)
+{
+	return
+		(WorldContext.WorldType == EWorldType::PIE) &&
+		(WorldContext.World() != nullptr) &&
+		(WorldContext.World()->GetNetMode() == NM_Client) &&
+		(WorldContext.GameViewport != nullptr) &&
+		(WorldContext.GameViewport->EngineShowFlags.ServerDrawDebug);
+}
+
+#define UE_DRAW_SERVER_DEBUG_ON_EACH_CLIENT(FunctionName, ...) \
+		if (GIsEditor) \
+		{ \
+			for (const FWorldContext& WorldContext : GEngine->GetWorldContexts()) \
+			{ \
+				if (CanDrawServerDebugInContext(WorldContext)) \
+				{ \
+					FunctionName(WorldContext.World(), __VA_ARGS__); \
+				} \
+			} \
+		}
+
+#else
+
+#define UE_DRAW_SERVER_DEBUG_ON_EACH_CLIENT(FunctionName, ...)
+
+#endif
+
+
 void FlushPersistentDebugLines( const UWorld* InWorld )
 {
 	if(InWorld && InWorld->PersistentLineBatcher)
@@ -26,7 +72,7 @@ void FlushPersistentDebugLines( const UWorld* InWorld )
 
 ULineBatchComponent* GetDebugLineBatcher( const UWorld* InWorld, bool bPersistentLines, float LifeTime, bool bDepthIsForeground )
 {
-	return (InWorld ? (bDepthIsForeground ? InWorld->ForegroundLineBatcher : (( bPersistentLines || (LifeTime > 0.f) ) ? InWorld->PersistentLineBatcher : InWorld->LineBatcher)) : NULL);
+	return (InWorld ? (bDepthIsForeground ? InWorld->ForegroundLineBatcher : (( bPersistentLines || (LifeTime > 0.f) ) ? InWorld->PersistentLineBatcher : InWorld->LineBatcher)) : nullptr);
 }
 
 static float GetDebugLineLifeTime(ULineBatchComponent* LineBatcher, float LifeTime, bool bPersistent)
@@ -36,32 +82,36 @@ static float GetDebugLineLifeTime(ULineBatchComponent* LineBatcher, float LifeTi
 
 void DrawDebugLine(const UWorld* InWorld, FVector const& LineStart, FVector const& LineEnd, FColor const& Color, bool bPersistentLines, float LifeTime, uint8 DepthPriority, float Thickness)
 {
-	// no debug line drawing on dedicated server
 	if (GEngine->GetNetMode(InWorld) != NM_DedicatedServer)
 	{
 		// this means foreground lines can't be persistent 
-		ULineBatchComponent* const LineBatcher = GetDebugLineBatcher( InWorld, bPersistentLines, LifeTime, (DepthPriority == SDPG_Foreground) );
-		if(LineBatcher != NULL)
+		if (ULineBatchComponent* const LineBatcher = GetDebugLineBatcher(InWorld, bPersistentLines, LifeTime, (DepthPriority == SDPG_Foreground)))
 		{
 			float const LineLifeTime = GetDebugLineLifeTime(LineBatcher, LifeTime, bPersistentLines);
 			LineBatcher->DrawLine(LineStart, LineEnd, Color, DepthPriority, Thickness, LineLifeTime);
 		}
 	}
+	else
+	{
+		UE_DRAW_SERVER_DEBUG_ON_EACH_CLIENT(DrawDebugLine, LineStart, LineEnd, AdjustColorForServer(Color), bPersistentLines, LifeTime, DepthPriority, Thickness);
+	}
 }
 
 void DrawDebugPoint(const UWorld* InWorld, FVector const& Position, float Size, FColor const& Color, bool bPersistentLines, float LifeTime, uint8 DepthPriority)
 {
-	// no debug line drawing on dedicated server
 	if (GEngine->GetNetMode(InWorld) != NM_DedicatedServer)
 	{
 		// this means foreground lines can't be persistent 
-		ULineBatchComponent* const LineBatcher = GetDebugLineBatcher( InWorld, bPersistentLines, LifeTime, (DepthPriority == SDPG_Foreground) );
-		if(LineBatcher != NULL)
+		if (ULineBatchComponent* const LineBatcher = GetDebugLineBatcher(InWorld, bPersistentLines, LifeTime, (DepthPriority == SDPG_Foreground)))
 		{
 			const float PointLifeTime = GetDebugLineLifeTime(LineBatcher, LifeTime, bPersistentLines);
 			LineBatcher->DrawPoint(Position, Color.ReinterpretAsLinear(), Size, DepthPriority, PointLifeTime);
 		}
 	}
+	else
+	{
+		UE_DRAW_SERVER_DEBUG_ON_EACH_CLIENT(DrawDebugPoint, Position, Size, AdjustColorForServer(Color), bPersistentLines, LifeTime, DepthPriority);
+	}
 }
 
 void DrawDebugDirectionalArrow(const UWorld* InWorld, FVector const& LineStart, FVector const& LineEnd, float ArrowSize, FColor const& Color, bool bPersistentLines, float LifeTime, uint8 DepthPriority, float Thickness)
@@ -95,6 +145,10 @@ void DrawDebugDirectionalArrow(const UWorld* InWorld, FVector const& LineStart,
 		DrawDebugLine(InWorld, LineEnd, LineEnd + TM.TransformPosition(FVector(-ArrowSqrt, ArrowSqrt, 0)), Color, bPersistentLines, LifeTime, DepthPriority, Thickness);
 		DrawDebugLine(InWorld, LineEnd, LineEnd + TM.TransformPosition(FVector(-ArrowSqrt, -ArrowSqrt, 0)), Color, bPersistentLines, LifeTime, DepthPriority, Thickness);
 	}
+	else
+	{
+		UE_DRAW_SERVER_DEBUG_ON_EACH_CLIENT(DrawDebugDirectionalArrow, LineStart, LineEnd, ArrowSize, AdjustColorForServer(Color), bPersistentLines, LifeTime, DepthPriority, Thickness);
+	}
 }
 
 void DrawDebugBox(const UWorld* InWorld, FVector const& Center, FVector const& Box, FColor const& Color, bool bPersistentLines, float LifeTime, uint8 DepthPriority, float Thickness)
@@ -103,8 +157,7 @@ void DrawDebugBox(const UWorld* InWorld, FVector const& Center, FVector const& B
 	if (GEngine->GetNetMode(InWorld) != NM_DedicatedServer)
 	{
 		// this means foreground lines can't be persistent 
-		ULineBatchComponent* const LineBatcher = GetDebugLineBatcher( InWorld, bPersistentLines, LifeTime, (DepthPriority == SDPG_Foreground) );		
-		if(LineBatcher != NULL)
+		if (ULineBatchComponent* const LineBatcher = GetDebugLineBatcher(InWorld, bPersistentLines, LifeTime, (DepthPriority == SDPG_Foreground)))
 		{
 			float LineLifeTime = GetDebugLineLifeTime(LineBatcher, LifeTime, bPersistentLines);
 
@@ -124,6 +177,10 @@ void DrawDebugBox(const UWorld* InWorld, FVector const& Center, FVector const& B
 			LineBatcher->DrawLine(Center + FVector(-Box.X,  Box.Y,  Box.Z), Center + FVector(-Box.X,  Box.Y, -Box.Z), Color, DepthPriority, Thickness, LineLifeTime);
 		}
 	}
+	else
+	{
+		UE_DRAW_SERVER_DEBUG_ON_EACH_CLIENT(DrawDebugBox, Center, Box, AdjustColorForServer(Color), bPersistentLines, LifeTime, DepthPriority, Thickness);
+	}
 }
 
 void DrawDebugBox(const UWorld* InWorld, FVector const& Center, FVector const& Box, const FQuat& Rotation, FColor const& Color, bool bPersistentLines, float LifeTime, uint8 DepthPriority, float Thickness)
@@ -132,8 +189,7 @@ void DrawDebugBox(const UWorld* InWorld, FVector const& Center, FVector const& B
 	if (GEngine->GetNetMode(InWorld) != NM_DedicatedServer)
 	{
 		// this means foreground lines can't be persistent 
-		ULineBatchComponent* const LineBatcher = GetDebugLineBatcher( InWorld, bPersistentLines, LifeTime, (DepthPriority == SDPG_Foreground) );
-		if(LineBatcher != NULL)
+		if (ULineBatchComponent* const LineBatcher = GetDebugLineBatcher(InWorld, bPersistentLines, LifeTime, (DepthPriority == SDPG_Foreground)))
 		{
 			float const LineLifeTime = GetDebugLineLifeTime(LineBatcher, LifeTime, bPersistentLines);
 			TArray<struct FBatchedLine> Lines;
@@ -190,6 +246,10 @@ void DrawDebugBox(const UWorld* InWorld, FVector const& Center, FVector const& B
 			LineBatcher->DrawLines(Lines);
 		}
 	}
+	else
+	{
+		UE_DRAW_SERVER_DEBUG_ON_EACH_CLIENT(DrawDebugBox, Center, Box, Rotation, AdjustColorForServer(Color), bPersistentLines, LifeTime, DepthPriority, Thickness);
+	}
 }
 
 
@@ -198,13 +258,16 @@ void DrawDebugMesh(const UWorld* InWorld, TArray<FVector> const& Verts, TArray<i
 	// no debug line drawing on dedicated server
 	if (GEngine->GetNetMode(InWorld) != NM_DedicatedServer)
 	{
-		ULineBatchComponent* const LineBatcher = GetDebugLineBatcher( InWorld, bPersistent, LifeTime, false );
-		if(LineBatcher != NULL)
+		if (ULineBatchComponent* const LineBatcher = GetDebugLineBatcher(InWorld, bPersistent, LifeTime, false))
 		{
 			float const ActualLifetime = GetDebugLineLifeTime(LineBatcher, LifeTime, bPersistent);
 			LineBatcher->DrawMesh(Verts, Indices, Color, DepthPriority, ActualLifetime);
 		}
 	}
+	else
+	{
+		UE_DRAW_SERVER_DEBUG_ON_EACH_CLIENT(DrawDebugMesh, Verts, Indices, AdjustColorForServer(Color), bPersistent, LifeTime, DepthPriority);
+	}
 }
 
 void DrawDebugSolidBox(const UWorld* InWorld, FBox const& Box, FColor const& Color, const FTransform& Transform, bool bPersistent, float LifeTime, uint8 DepthPriority)
@@ -212,13 +275,16 @@ void DrawDebugSolidBox(const UWorld* InWorld, FBox const& Box, FColor const& Col
 	// no debug line drawing on dedicated server
 	if (GEngine->GetNetMode(InWorld) != NM_DedicatedServer)
 	{
-		ULineBatchComponent* const LineBatcher = GetDebugLineBatcher( InWorld, bPersistent, LifeTime, false );
-		if(LineBatcher != NULL)
+		if (ULineBatchComponent* const LineBatcher = GetDebugLineBatcher(InWorld, bPersistent, LifeTime, false))
 		{
 			float const ActualLifetime = GetDebugLineLifeTime(LineBatcher, LifeTime, bPersistent);
 			LineBatcher->DrawSolidBox(Box, Transform, Color, DepthPriority, ActualLifetime);
 		}
 	}
+	else
+	{
+		UE_DRAW_SERVER_DEBUG_ON_EACH_CLIENT(DrawDebugSolidBox, Box, AdjustColorForServer(Color), Transform, bPersistent, LifeTime, DepthPriority);
+	}
 }
 
 void DrawDebugSolidBox(const UWorld* InWorld, FVector const& Center, FVector const& Extent, FColor const& Color, bool bPersistent, float LifeTime, uint8 DepthPriority)
@@ -239,34 +305,7 @@ void DrawDebugSolidBox(const UWorld* InWorld, FVector const& Center, FVector con
 /** Loc is an anchor point in the world to guide which part of the infinite plane to draw. */
 void DrawDebugSolidPlane(const UWorld* InWorld, FPlane const& P, FVector const& Loc, float Size, FColor const& Color, bool bPersistent, float LifeTime, uint8 DepthPriority)
 {
-	// no debug line drawing on dedicated server
-	if (GEngine->GetNetMode(InWorld) != NM_DedicatedServer)
-	{
-		FVector const ClosestPtOnPlane = Loc - P.PlaneDot(Loc) * P;
-
-		FVector U, V;
-		P.FindBestAxisVectors(U, V);
-		U *= Size;
-		V *= Size;
-
-		TArray<FVector> Verts;
-		Verts.AddUninitialized(4);
-		Verts[0] = ClosestPtOnPlane + U + V;
-		Verts[1] = ClosestPtOnPlane - U + V;
-		Verts[2] = ClosestPtOnPlane + U - V;
-		Verts[3] = ClosestPtOnPlane - U - V;
-
-		TArray<int32> Indices;
-		Indices.AddUninitialized(6);
-		Indices[0] = 0; Indices[1] = 2; Indices[2] = 1;
-		Indices[3] = 1; Indices[4] = 2; Indices[5] = 3;
-
-		// plane quad
-		DrawDebugMesh(InWorld, Verts, Indices, Color, bPersistent, LifeTime, DepthPriority);
-
-		// arrow indicating normal
-		DrawDebugDirectionalArrow(InWorld, ClosestPtOnPlane, ClosestPtOnPlane + P * 16.f, 8.f, FColor::White, bPersistent, LifeTime, DepthPriority);
-	}
+	DrawDebugSolidPlane(InWorld, P, Loc, FVector2D(Size, Size), Color, bPersistent, LifeTime, DepthPriority);
 }
 
 ENGINE_API void DrawDebugSolidPlane(const UWorld* InWorld, FPlane const& P, FVector const& Loc, FVector2D const& Extents, FColor const& Color, bool bPersistent/*=false*/, float LifeTime/*=-1*/, uint8 DepthPriority /*= 0*/)
@@ -299,6 +338,10 @@ ENGINE_API void DrawDebugSolidPlane(const UWorld* InWorld, FPlane const& P, FVec
 		// arrow indicating normal
 		DrawDebugDirectionalArrow(InWorld, ClosestPtOnPlane, ClosestPtOnPlane + P * 16.f, 8.f, FColor::White, bPersistent, LifeTime, DepthPriority);
 	}
+	else
+	{
+		UE_DRAW_SERVER_DEBUG_ON_EACH_CLIENT(DrawDebugSolidPlane, P, Loc, Extents, AdjustColorForServer(Color), bPersistent, LifeTime, DepthPriority);
+	}
 }
 
 void DrawDebugCoordinateSystem(const UWorld* InWorld, FVector const& AxisLoc, FRotator const& AxisRot, float Scale, bool bPersistentLines, float LifeTime, uint8 DepthPriority, float Thickness)
@@ -312,8 +355,7 @@ void DrawDebugCoordinateSystem(const UWorld* InWorld, FVector const& AxisLoc, FR
 		FVector const Z = R.GetScaledAxis( EAxis::Z );
 
 		// this means foreground lines can't be persistent 
-		ULineBatchComponent* const LineBatcher = GetDebugLineBatcher( InWorld, bPersistentLines, LifeTime, (DepthPriority == SDPG_Foreground) );
-		if(LineBatcher != NULL)
+		if (ULineBatchComponent* const LineBatcher = GetDebugLineBatcher(InWorld, bPersistentLines, LifeTime, (DepthPriority == SDPG_Foreground)))
 		{
 			const float LineLifeTime = GetDebugLineLifeTime(LineBatcher, LifeTime, bPersistentLines);
 			LineBatcher->DrawLine(AxisLoc, AxisLoc + X*Scale, FColor::Red, DepthPriority, Thickness, LineLifeTime );
@@ -321,6 +363,10 @@ void DrawDebugCoordinateSystem(const UWorld* InWorld, FVector const& AxisLoc, FR
 			LineBatcher->DrawLine(AxisLoc, AxisLoc + Z*Scale, FColor::Blue, DepthPriority, Thickness, LineLifeTime );
 		}
 	}
+	else
+	{
+		UE_DRAW_SERVER_DEBUG_ON_EACH_CLIENT(DrawDebugCoordinateSystem, AxisLoc, AxisRot, Scale, bPersistentLines, LifeTime, DepthPriority, Thickness);
+	}
 }
 
 ENGINE_API void DrawDebugCrosshairs(const UWorld* InWorld, FVector const& AxisLoc, FRotator const& AxisRot, float Scale, const FColor& Color, bool bPersistentLines, float LifeTime, uint8 DepthPriority)
@@ -334,8 +380,7 @@ ENGINE_API void DrawDebugCrosshairs(const UWorld* InWorld, FVector const& AxisLo
 		FVector const Z = 0.5f * R.GetScaledAxis(EAxis::Z);
 
 		// this means foreground lines can't be persistent 
-		ULineBatchComponent* const LineBatcher = GetDebugLineBatcher(InWorld, bPersistentLines, LifeTime, (DepthPriority == SDPG_Foreground));
-		if (LineBatcher != NULL)
+		if (ULineBatchComponent* const LineBatcher = GetDebugLineBatcher(InWorld, bPersistentLines, LifeTime, (DepthPriority == SDPG_Foreground)))
 		{
 			const float LineLifeTime = GetDebugLineLifeTime(LineBatcher, LifeTime, bPersistentLines);
 			LineBatcher->DrawLine(AxisLoc - X*Scale, AxisLoc + X*Scale, Color, DepthPriority, 0.f, LineLifeTime);
@@ -343,50 +388,48 @@ ENGINE_API void DrawDebugCrosshairs(const UWorld* InWorld, FVector const& AxisLo
 			LineBatcher->DrawLine(AxisLoc - Z*Scale, AxisLoc + Z*Scale, Color, DepthPriority, 0.f, LineLifeTime);
 		}
 	}
+	else
+	{
+		UE_DRAW_SERVER_DEBUG_ON_EACH_CLIENT(DrawDebugCrosshairs, AxisLoc, AxisRot, Scale, AdjustColorForServer(Color), bPersistentLines, LifeTime, DepthPriority);
+	}
 }
 
 
 static void InternalDrawDebugCircle(const UWorld* InWorld, const FMatrix& TransformMatrix, float Radius, int32 Segments, const FColor& Color, bool bPersistentLines, float LifeTime, uint8 DepthPriority, float Thickness = 0.f)
 {
-	// no debug line drawing on dedicated server
-	if( GEngine->GetNetMode(InWorld) != NM_DedicatedServer )
+	if (ULineBatchComponent* const LineBatcher = GetDebugLineBatcher(InWorld, bPersistentLines, LifeTime, (DepthPriority == SDPG_Foreground)))
 	{
-		ULineBatchComponent* const LineBatcher = GetDebugLineBatcher( InWorld, bPersistentLines, LifeTime, (DepthPriority == SDPG_Foreground) );
-		if(LineBatcher != NULL)
-		{
-			const float LineLifeTime = GetDebugLineLifeTime(LineBatcher, LifeTime, bPersistentLines);
+		const float LineLifeTime = GetDebugLineLifeTime(LineBatcher, LifeTime, bPersistentLines);
 
-			// Need at least 4 segments
-			Segments = FMath::Max(Segments, 4);
-			const float AngleStep = 2.f * PI / float(Segments);
+		// Need at least 4 segments
+		Segments = FMath::Max(Segments, 4);
+		const float AngleStep = 2.f * PI / float(Segments);
 
-			const FVector Center = TransformMatrix.GetOrigin();
-			const FVector AxisY = TransformMatrix.GetScaledAxis( EAxis::Y );
-			const FVector AxisZ = TransformMatrix.GetScaledAxis( EAxis::Z );
+		const FVector Center = TransformMatrix.GetOrigin();
+		const FVector AxisY = TransformMatrix.GetScaledAxis(EAxis::Y);
+		const FVector AxisZ = TransformMatrix.GetScaledAxis(EAxis::Z);
 
-			TArray<FBatchedLine> Lines;
-			Lines.Empty(Segments);
+		TArray<FBatchedLine> Lines;
+		Lines.Empty(Segments);
 
-			float Angle = 0.f;
-			while( Segments-- )
-			{
-				const FVector Vertex1 = Center + Radius * (AxisY * FMath::Cos(Angle) + AxisZ * FMath::Sin(Angle));
-				Angle += AngleStep;
-				const FVector Vertex2 = Center + Radius * (AxisY * FMath::Cos(Angle) + AxisZ * FMath::Sin(Angle));
-				Lines.Add(FBatchedLine(Vertex1, Vertex2, Color, LineLifeTime, Thickness, DepthPriority));
-			}
-			LineBatcher->DrawLines(Lines);
+		float Angle = 0.f;
+		while (Segments--)
+		{
+			const FVector Vertex1 = Center + Radius * (AxisY * FMath::Cos(Angle) + AxisZ * FMath::Sin(Angle));
+			Angle += AngleStep;
+			const FVector Vertex2 = Center + Radius * (AxisY * FMath::Cos(Angle) + AxisZ * FMath::Sin(Angle));
+			Lines.Add(FBatchedLine(Vertex1, Vertex2, Color, LineLifeTime, Thickness, DepthPriority));
 		}
+		LineBatcher->DrawLines(Lines);
 	}
 }
 
 void DrawDebugCircle(const UWorld* InWorld, const FMatrix& TransformMatrix, float Radius, int32 Segments, const FColor& Color, bool bPersistentLines, float LifeTime, uint8 DepthPriority, float Thickness, bool bDrawAxis)
 {
 	// no debug line drawing on dedicated server
-	if( GEngine->GetNetMode(InWorld) != NM_DedicatedServer )
+	if (GEngine->GetNetMode(InWorld) != NM_DedicatedServer)
 	{
-		ULineBatchComponent* const LineBatcher = GetDebugLineBatcher( InWorld, bPersistentLines, LifeTime, (DepthPriority == SDPG_Foreground) );
-		if(LineBatcher != NULL)
+		if (ULineBatchComponent* const LineBatcher = GetDebugLineBatcher(InWorld, bPersistentLines, LifeTime, (DepthPriority == SDPG_Foreground)))
 		{
 			const float LineLifeTime = GetDebugLineLifeTime(LineBatcher, LifeTime, bPersistentLines);
 
@@ -408,6 +451,10 @@ void DrawDebugCircle(const UWorld* InWorld, const FMatrix& TransformMatrix, floa
 			}
 		}
 	}
+	else
+	{
+		UE_DRAW_SERVER_DEBUG_ON_EACH_CLIENT(DrawDebugCircle, TransformMatrix, Radius, Segments, AdjustColorForServer(Color), bPersistentLines, LifeTime, DepthPriority, Thickness, bDrawAxis);
+	}
 }
 
 void DrawDebugCircle(const UWorld* InWorld, FVector Center, float Radius, int32 Segments, const FColor& Color, bool PersistentLines, float LifeTime, uint8 DepthPriority, float Thickness, FVector YAxis, FVector ZAxis, bool bDrawAxis)
@@ -435,10 +482,9 @@ void DrawDebugCircle(const UWorld* InWorld, FVector Center, float Radius, int32
 void DrawDebug2DDonut(const UWorld* InWorld, const FMatrix& TransformMatrix, float InnerRadius, float OuterRadius, int32 Segments, const FColor& Color, bool bPersistentLines, float LifeTime, uint8 DepthPriority, float Thickness)
 {
 	// no debug line drawing on dedicated server
-	if( GEngine->GetNetMode(InWorld) != NM_DedicatedServer )
+	if (GEngine->GetNetMode(InWorld) != NM_DedicatedServer)
 	{
-		ULineBatchComponent* const LineBatcher = GetDebugLineBatcher( InWorld, bPersistentLines, LifeTime, (DepthPriority == SDPG_Foreground) );
-		if(LineBatcher != NULL)
+		if (ULineBatchComponent* const LineBatcher = GetDebugLineBatcher(InWorld, bPersistentLines, LifeTime, (DepthPriority == SDPG_Foreground)))
 		{
 			const float LineLifeTime = GetDebugLineLifeTime(LineBatcher, LifeTime, bPersistentLines);
 
@@ -460,6 +506,10 @@ void DrawDebug2DDonut(const UWorld* InWorld, const FMatrix& TransformMatrix, flo
 			LineBatcher->DrawLines(Lines);
 		}
 	}
+	else
+	{
+		UE_DRAW_SERVER_DEBUG_ON_EACH_CLIENT(DrawDebug2DDonut, TransformMatrix, InnerRadius, OuterRadius, Segments, AdjustColorForServer(Color), bPersistentLines, LifeTime, DepthPriority, Thickness);
+	}
 }
 
 void DrawDebugSphere(const UWorld* InWorld, FVector const& Center, float Radius, int32 Segments, FColor const& Color, bool bPersistentLines, float LifeTime, uint8 DepthPriority, float Thickness)
@@ -468,8 +518,7 @@ void DrawDebugSphere(const UWorld* InWorld, FVector const& Center, float Radius,
 	if (GEngine->GetNetMode(InWorld) != NM_DedicatedServer)
 	{
 		// this means foreground lines can't be persistent 
-		ULineBatchComponent* const LineBatcher = GetDebugLineBatcher( InWorld, bPersistentLines, LifeTime, (DepthPriority == SDPG_Foreground) );
-		if (LineBatcher != NULL)
+		if (ULineBatchComponent* const LineBatcher = GetDebugLineBatcher(InWorld, bPersistentLines, LifeTime, (DepthPriority == SDPG_Foreground)))
 		{
 			float LineLifeTime = GetDebugLineLifeTime(LineBatcher, LifeTime, bPersistentLines);
 
@@ -519,6 +568,10 @@ void DrawDebugSphere(const UWorld* InWorld, FVector const& Center, float Radius,
 			LineBatcher->DrawLines(Lines);
 		}
 	}
+	else
+	{
+		UE_DRAW_SERVER_DEBUG_ON_EACH_CLIENT(DrawDebugSphere, Center, Radius, Segments, AdjustColorForServer(Color), bPersistentLines, LifeTime, DepthPriority, Thickness);
+	}
 }
 
 void DrawDebugCylinder(const UWorld* InWorld, FVector const& Start, FVector const& End, float Radius, int32 Segments, FColor const& Color, bool bPersistentLines, float LifeTime, uint8 DepthPriority, float Thickness)
@@ -526,35 +579,34 @@ void DrawDebugCylinder(const UWorld* InWorld, FVector const& Start, FVector cons
 	// no debug line drawing on dedicated server
 	if (GEngine->GetNetMode(InWorld) != NM_DedicatedServer)
 	{
-		// Need at least 4 segments
-		Segments = FMath::Max(Segments, 4);
+		// this means foreground lines can't be persistent 
+		if (ULineBatchComponent* const LineBatcher = GetDebugLineBatcher(InWorld, bPersistentLines, LifeTime, (DepthPriority == SDPG_Foreground)))
+		{
+			// Need at least 4 segments
+			Segments = FMath::Max(Segments, 4);
 
-		// Rotate a point around axis to form cylinder segments
-		FVector Segment;
-		FVector P1, P2, P3, P4;
-		const float AngleInc = 360.f / Segments;
-		float Angle = AngleInc;
+			// Rotate a point around axis to form cylinder segments
+			FVector Segment;
+			FVector P1, P2, P3, P4;
+			const float AngleInc = 360.f / Segments;
+			float Angle = AngleInc;
 
-		// Default for Axis is up
-		FVector Axis = (End - Start).GetSafeNormal();
-		if( Axis.IsZero() )
-		{
-			Axis = FVector(0.f, 0.f, 1.f);
-		}
+			// Default for Axis is up
+			FVector Axis = (End - Start).GetSafeNormal();
+			if( Axis.IsZero() )
+			{
+				Axis = FVector(0.f, 0.f, 1.f);
+			}
 
-		FVector Perpendicular;
-		FVector Dummy;
+			FVector Perpendicular;
+			FVector Dummy;
 
-		Axis.FindBestAxisVectors(Perpendicular, Dummy);
+			Axis.FindBestAxisVectors(Perpendicular, Dummy);
 		
-		Segment = Perpendicular.RotateAngleAxis(0, Axis) * Radius;
-		P1 = Segment + Start;
-		P3 = Segment + End;
+			Segment = Perpendicular.RotateAngleAxis(0, Axis) * Radius;
+			P1 = Segment + Start;
+			P3 = Segment + End;
 
-		// this means foreground lines can't be persistent 
-		ULineBatchComponent* const LineBatcher = GetDebugLineBatcher( InWorld, bPersistentLines, LifeTime, (DepthPriority == SDPG_Foreground) );
-		if(LineBatcher != NULL)
-		{
 			const float LineLifeTime = GetDebugLineLifeTime(LineBatcher, LifeTime, bPersistentLines);
 			while( Segments-- )
 			{
@@ -572,69 +624,78 @@ void DrawDebugCylinder(const UWorld* InWorld, FVector const& Start, FVector cons
 			}
 		}
 	}
+	else
+	{
+		UE_DRAW_SERVER_DEBUG_ON_EACH_CLIENT(DrawDebugCylinder, Start, End, Radius, Segments, AdjustColorForServer(Color), bPersistentLines, LifeTime, DepthPriority, Thickness);
+	}
 }
 
 /** Used by gameplay when defining a cone by a vertical and horizontal dot products. */
 void DrawDebugAltCone(const UWorld* InWorld, FVector const& Origin, FRotator const& Rotation, float Length, float AngleWidth, float AngleHeight, FColor const& DrawColor, bool bPersistentLines, float LifeTime, uint8 DepthPriority, float Thickness)
 {
-	FRotationMatrix const RM( Rotation );
-	FVector const AxisX = RM.GetScaledAxis( EAxis::X );
-	FVector const AxisY = RM.GetScaledAxis( EAxis::Y );
-	FVector const AxisZ = RM.GetScaledAxis( EAxis::Z );
-
-	// this means foreground lines can't be persistent 
-	ULineBatchComponent* const LineBatcher = GetDebugLineBatcher( InWorld, bPersistentLines, LifeTime, (DepthPriority == SDPG_Foreground) );
-	if (LineBatcher != NULL)
+	if (GEngine->GetNetMode(InWorld) != NM_DedicatedServer)
 	{
+		// this means foreground lines can't be persistent 
+		if (ULineBatchComponent* const LineBatcher = GetDebugLineBatcher(InWorld, bPersistentLines, LifeTime, (DepthPriority == SDPG_Foreground)))
+		{
+			FRotationMatrix const RM(Rotation);
+			FVector const AxisX = RM.GetScaledAxis(EAxis::X);
+			FVector const AxisY = RM.GetScaledAxis(EAxis::Y);
+			FVector const AxisZ = RM.GetScaledAxis(EAxis::Z);
 
-		float const LineLifeTime = GetDebugLineLifeTime(LineBatcher, LifeTime, bPersistentLines);
+			float const LineLifeTime = GetDebugLineLifeTime(LineBatcher, LifeTime, bPersistentLines);
 
-		FVector const EndPoint = Origin + AxisX * Length;
-		FVector const Up = FMath::Tan(AngleHeight * 0.5f) * AxisZ * Length;
-		FVector const Right = FMath::Tan(AngleWidth * 0.5f) * AxisY * Length;
-		FVector const HalfUp = Up * 0.5f;
-		FVector const HalfRight = Right * 0.5f;
+			FVector const EndPoint = Origin + AxisX * Length;
+			FVector const Up = FMath::Tan(AngleHeight * 0.5f) * AxisZ * Length;
+			FVector const Right = FMath::Tan(AngleWidth * 0.5f) * AxisY * Length;
+			FVector const HalfUp = Up * 0.5f;
+			FVector const HalfRight = Right * 0.5f;
 
-		TArray<FBatchedLine> Lines;
-		Lines.Empty();
+			TArray<FBatchedLine> Lines;
+			Lines.Empty();
 
-		FVector A = EndPoint + Up - Right;
-		FVector B = EndPoint + Up + Right;
-		FVector C = EndPoint - Up + Right;
-		FVector D = EndPoint - Up - Right;
+			FVector A = EndPoint + Up - Right;
+			FVector B = EndPoint + Up + Right;
+			FVector C = EndPoint - Up + Right;
+			FVector D = EndPoint - Up - Right;
 
-		// Corners
-		Lines.Add(FBatchedLine(Origin, A, DrawColor, LineLifeTime, Thickness, DepthPriority));
-		Lines.Add(FBatchedLine(Origin, B, DrawColor, LineLifeTime, Thickness, DepthPriority));
-		Lines.Add(FBatchedLine(Origin, C, DrawColor, LineLifeTime, Thickness, DepthPriority));
-		Lines.Add(FBatchedLine(Origin, D, DrawColor, LineLifeTime, Thickness, DepthPriority));
+			// Corners
+			Lines.Add(FBatchedLine(Origin, A, DrawColor, LineLifeTime, Thickness, DepthPriority));
+			Lines.Add(FBatchedLine(Origin, B, DrawColor, LineLifeTime, Thickness, DepthPriority));
+			Lines.Add(FBatchedLine(Origin, C, DrawColor, LineLifeTime, Thickness, DepthPriority));
+			Lines.Add(FBatchedLine(Origin, D, DrawColor, LineLifeTime, Thickness, DepthPriority));
 
-		// Further most plane/frame
-		Lines.Add(FBatchedLine(A, B, DrawColor, LineLifeTime, Thickness, DepthPriority));
-		Lines.Add(FBatchedLine(B, C, DrawColor, LineLifeTime, Thickness, DepthPriority));
-		Lines.Add(FBatchedLine(C, D, DrawColor, LineLifeTime, Thickness, DepthPriority));
-		Lines.Add(FBatchedLine(D, A, DrawColor, LineLifeTime, Thickness, DepthPriority));
+			// Further most plane/frame
+			Lines.Add(FBatchedLine(A, B, DrawColor, LineLifeTime, Thickness, DepthPriority));
+			Lines.Add(FBatchedLine(B, C, DrawColor, LineLifeTime, Thickness, DepthPriority));
+			Lines.Add(FBatchedLine(C, D, DrawColor, LineLifeTime, Thickness, DepthPriority));
+			Lines.Add(FBatchedLine(D, A, DrawColor, LineLifeTime, Thickness, DepthPriority));
 
-		// Mid points
-		Lines.Add(FBatchedLine(Origin, EndPoint + Up, DrawColor, LineLifeTime, Thickness, DepthPriority));
-		Lines.Add(FBatchedLine(Origin, EndPoint - Up, DrawColor, LineLifeTime, Thickness, DepthPriority));
-		Lines.Add(FBatchedLine(Origin, EndPoint + Right, DrawColor, LineLifeTime, Thickness, DepthPriority));
-		Lines.Add(FBatchedLine(Origin, EndPoint - Right, DrawColor, LineLifeTime, Thickness, DepthPriority));
+			// Mid points
+			Lines.Add(FBatchedLine(Origin, EndPoint + Up, DrawColor, LineLifeTime, Thickness, DepthPriority));
+			Lines.Add(FBatchedLine(Origin, EndPoint - Up, DrawColor, LineLifeTime, Thickness, DepthPriority));
+			Lines.Add(FBatchedLine(Origin, EndPoint + Right, DrawColor, LineLifeTime, Thickness, DepthPriority));
+			Lines.Add(FBatchedLine(Origin, EndPoint - Right, DrawColor, LineLifeTime, Thickness, DepthPriority));
 
-		// Inbetween
-		Lines.Add(FBatchedLine(Origin, EndPoint + Up - HalfRight, DrawColor, LineLifeTime, Thickness, DepthPriority));
-		Lines.Add(FBatchedLine(Origin, EndPoint + Up + HalfRight, DrawColor, LineLifeTime, Thickness, DepthPriority));
+			// Inbetween
+			Lines.Add(FBatchedLine(Origin, EndPoint + Up - HalfRight, DrawColor, LineLifeTime, Thickness, DepthPriority));
+			Lines.Add(FBatchedLine(Origin, EndPoint + Up + HalfRight, DrawColor, LineLifeTime, Thickness, DepthPriority));
 
-		Lines.Add(FBatchedLine(Origin, EndPoint - Up - HalfRight, DrawColor, LineLifeTime, Thickness, DepthPriority));
-		Lines.Add(FBatchedLine(Origin, EndPoint - Up + HalfRight, DrawColor, LineLifeTime, Thickness, DepthPriority));
+			Lines.Add(FBatchedLine(Origin, EndPoint - Up - HalfRight, DrawColor, LineLifeTime, Thickness, DepthPriority));
+			Lines.Add(FBatchedLine(Origin, EndPoint - Up + HalfRight, DrawColor, LineLifeTime, Thickness, DepthPriority));
 
-		Lines.Add(FBatchedLine(Origin, EndPoint + Right - HalfUp, DrawColor, LineLifeTime, Thickness, DepthPriority));
-		Lines.Add(FBatchedLine(Origin, EndPoint + Right + HalfUp, DrawColor, LineLifeTime, Thickness, DepthPriority));
+			Lines.Add(FBatchedLine(Origin, EndPoint + Right - HalfUp, DrawColor, LineLifeTime, Thickness, DepthPriority));
+			Lines.Add(FBatchedLine(Origin, EndPoint + Right + HalfUp, DrawColor, LineLifeTime, Thickness, DepthPriority));
 
-		Lines.Add(FBatchedLine(Origin, EndPoint - Right - HalfUp, DrawColor, LineLifeTime, Thickness, DepthPriority));
-		Lines.Add(FBatchedLine(Origin, EndPoint - Right + HalfUp, DrawColor, LineLifeTime, Thickness, DepthPriority));
+			Lines.Add(FBatchedLine(Origin, EndPoint - Right - HalfUp, DrawColor, LineLifeTime, Thickness, DepthPriority));
+			Lines.Add(FBatchedLine(Origin, EndPoint - Right + HalfUp, DrawColor, LineLifeTime, Thickness, DepthPriority));
 
-		LineBatcher->DrawLines(Lines);
+			LineBatcher->DrawLines(Lines);
+		}
+	}
+	else
+	{
+		UE_DRAW_SERVER_DEBUG_ON_EACH_CLIENT(DrawDebugAltCone, Origin, Rotation, Length, AngleWidth, AngleHeight, AdjustColorForServer(DrawColor), bPersistentLines, LifeTime, DepthPriority, Thickness);
 	}
 }
 
@@ -689,8 +750,7 @@ void DrawDebugCone(const UWorld* InWorld, FVector const& Origin, FVector const&
 		const FMatrix ConeToWorld = FScaleMatrix(FVector(Length)) * FMatrix(DirectionNorm, YAxis, ZAxis, Origin);
 
 		// this means foreground lines can't be persistent 
-		ULineBatchComponent* const LineBatcher = GetDebugLineBatcher( InWorld, bPersistentLines, LifeTime, (DepthPriority == SDPG_Foreground) );
-		if(LineBatcher != NULL)
+		if (ULineBatchComponent* const LineBatcher = GetDebugLineBatcher(InWorld, bPersistentLines, LifeTime, (DepthPriority == SDPG_Foreground)))
 		{
 			float const LineLifeTime = GetDebugLineLifeTime(LineBatcher, LifeTime, bPersistentLines);
 
@@ -721,6 +781,10 @@ void DrawDebugCone(const UWorld* InWorld, FVector const& Origin, FVector const&
 			LineBatcher->DrawLines(Lines);
 		}
 	}
+	else
+	{
+		UE_DRAW_SERVER_DEBUG_ON_EACH_CLIENT(DrawDebugCone, Origin, Direction, Length, AngleWidth, AngleHeight, NumSides, AdjustColorForServer(DrawColor), bPersistentLines, LifeTime, DepthPriority, Thickness);
+	}
 }
 
 void DrawDebugString(const UWorld* InWorld, FVector const& TextLocation, const FString& Text, class AActor* TestBaseActor, FColor const& TextColor, float Duration, bool bDrawShadow, float FontScale)
@@ -728,8 +792,8 @@ void DrawDebugString(const UWorld* InWorld, FVector const& TextLocation, const F
 	// no debug line drawing on dedicated server
 	if (GEngine->GetNetMode(InWorld) != NM_DedicatedServer)
 	{
-		check(TestBaseActor == NULL || TestBaseActor->GetWorld() == InWorld);
-		AActor* BaseAct = (TestBaseActor != NULL) ? TestBaseActor : InWorld->GetWorldSettings();
+		check((TestBaseActor == nullptr) || (TestBaseActor->GetWorld() == InWorld));
+		AActor* BaseAct = (TestBaseActor != nullptr) ? TestBaseActor : InWorld->GetWorldSettings();
 
 		// iterate through the player controller list
 		for( FConstPlayerControllerIterator Iterator = InWorld->GetPlayerControllerIterator(); Iterator; ++Iterator )
@@ -737,10 +801,16 @@ void DrawDebugString(const UWorld* InWorld, FVector const& TextLocation, const F
 			APlayerController* PlayerController = Iterator->Get();
 			if (PlayerController && PlayerController->MyHUD && PlayerController->Player)
 			{
-				PlayerController->MyHUD->AddDebugText(Text, BaseAct, Duration, TextLocation, TextLocation, TextColor, true, (TestBaseActor==NULL), false, nullptr, FontScale, bDrawShadow);
+				PlayerController->MyHUD->AddDebugText(Text, BaseAct, Duration, TextLocation, TextLocation, TextColor, /*bSkipOverwriteCheck=*/ true, /*bAbsoluteLocation=*/ (TestBaseActor==nullptr), /*bKeepAttachedToActor=*/ false, nullptr, FontScale, bDrawShadow);
 			}
 		}
 	}
+	else
+	{
+		// We do a bit of converting here if the original call was relative, as there's a check() that the base actor is
+		// in the same world as being rendered to (and it might be in a different position on client vs server anyways)
+		UE_DRAW_SERVER_DEBUG_ON_EACH_CLIENT(DrawDebugString, (TestBaseActor != nullptr) ? (TextLocation + TestBaseActor->GetActorLocation()) : TextLocation, Text, /*TestBaseActor=*/ nullptr, AdjustColorForServer(TextColor), Duration, bDrawShadow, FontScale);
+	}
 }
 
 void FlushDebugStrings( const UWorld* InWorld )
@@ -797,6 +867,10 @@ void DrawDebugFrustum(const UWorld* InWorld, const FMatrix& FrustumToWorld, FCol
 		DrawDebugLine(InWorld, Vertices[0][0][1], Vertices[1][0][1],Color,  bPersistentLines, LifeTime, DepthPriority, Thickness);
 		DrawDebugLine(InWorld, Vertices[0][1][1], Vertices[1][1][1],Color,  bPersistentLines, LifeTime, DepthPriority, Thickness);
 	}
+	else
+	{
+		UE_DRAW_SERVER_DEBUG_ON_EACH_CLIENT(DrawDebugFrustum, FrustumToWorld, AdjustColorForServer(Color), bPersistentLines, LifeTime, DepthPriority, Thickness);
+	}
 }
 
 
@@ -863,6 +937,10 @@ void DrawDebugCapsule(const UWorld* InWorld, FVector const& Center, float HalfHe
 		DrawDebugLine(InWorld, TopEnd + Radius*YAxis, BottomEnd + Radius*YAxis, Color, bPersistentLines, LifeTime, DepthPriority, Thickness);
 		DrawDebugLine(InWorld, TopEnd - Radius*YAxis, BottomEnd - Radius*YAxis, Color, bPersistentLines, LifeTime, DepthPriority, Thickness);
 	}
+	else
+	{
+		UE_DRAW_SERVER_DEBUG_ON_EACH_CLIENT(DrawDebugCapsule, Center, HalfHeight, Radius, Rotation, AdjustColorForServer(Color), bPersistentLines, LifeTime, DepthPriority, Thickness);
+	}
 }
 
 
@@ -905,6 +983,10 @@ void DrawDebugCamera(const UWorld* InWorld, FVector const& Location, FRotator co
 		DrawDebugLine(InWorld, Corners[2], Corners[3], Color, bPersistentLines, LifeTime, DepthPriority);
 		DrawDebugLine(InWorld, Corners[3], Corners[0], Color, bPersistentLines, LifeTime, DepthPriority);
 	}
+	else
+	{
+		UE_DRAW_SERVER_DEBUG_ON_EACH_CLIENT(DrawDebugCamera, Location, Rotation, FOVDeg, Scale, AdjustColorForServer(Color), bPersistentLines, LifeTime, DepthPriority);
+	}
 }
 
 
diff --git a/Engine/Source/Runtime/Engine/Public/ShowFlags.h b/Engine/Source/Runtime/Engine/Public/ShowFlags.h
--- a/Engine/Source/Runtime/Engine/Public/ShowFlags.h
+++ b/Engine/Source/Runtime/Engine/Public/ShowFlags.h
@@ -393,6 +393,7 @@ private:
 		SetRequiredTextureResolution(false);
 		SetMotionBlur(InitMode != ESFIM_Editor && InitMode != ESFIM_VREditing);
 		SetBones(false);
+		SetServerDrawDebug(false);
 		SetScreenPercentage(InitMode != ESFIM_Editor && InitMode != ESFIM_VREditing);
 		SetVREditing(InitMode == ESFIM_VREditing);
 		SetOcclusionMeshes(false);
diff --git a/Engine/Source/Runtime/Engine/Public/ShowFlagsValues.inl b/Engine/Source/Runtime/Engine/Public/ShowFlagsValues.inl
--- a/Engine/Source/Runtime/Engine/Public/ShowFlagsValues.inl
+++ b/Engine/Source/Runtime/Engine/Public/ShowFlagsValues.inl
@@ -334,6 +334,8 @@ SHOWFLAG_FIXED_IN_SHIPPING(0, RequiredTextureResolution, SFG_Hidden, NSLOCTEXT("
 SHOWFLAG_ALWAYS_ACCESSIBLE(WidgetComponents, SFG_Normal, NSLOCTEXT("UnrealEd", "WidgetComponentsSF", "Widget Components"))
 /** Draw the bones of all skeletal meshes */
 SHOWFLAG_FIXED_IN_SHIPPING(0, Bones, SFG_Developer, NSLOCTEXT("UnrealEd", "BoneSF", "Bones"))
+/** Draw debug rendering from the Dedicated Server during Play In Editor (with Use Dedicated Server checked) in this client viewport */
+SHOWFLAG_FIXED_IN_SHIPPING(0, ServerDrawDebug, SFG_Developer, NSLOCTEXT("UnrealEd", "ServerDrawDebugSF", "Dedicated Server Debug Drawing"))
 /** If media planes should be shown */
 SHOWFLAG_ALWAYS_ACCESSIBLE(MediaPlanes, SFG_Normal, NSLOCTEXT("UnrealEd", "MediaPlanesSF", "Media Planes"))
 /** if this is a vr editing viewport, needed? */
