Remove external EngineAnalytics configuration options in EngineAnalytics. It was not being used anymore.

diff --git a/Engine/Source/Runtime/Engine/Private/EngineAnalytics.cpp b/Engine/Source/Runtime/Engine/Private/EngineAnalytics.cpp
--- a/Engine/Source/Runtime/Engine/Private/EngineAnalytics.cpp
+++ b/Engine/Source/Runtime/Engine/Private/EngineAnalytics.cpp
@@ -161,31 +161,9 @@ private:
 	TWeakPtr<IAnalyticsProviderET> Analytics;
 };
 
-/**
-* Default config func.
-*/
-FAnalyticsET::Config DefaultEngineAnalyticsConfigFunc()
-{
-	return FAnalyticsET::Config();
-}
-
-/**
-* Engine analytics config to initialize the analytics provider.
-* External code should bind this delegate if engine analytics are desired,
-* preferably in private code that won't be redistributed.
-*/
-TFunction<FAnalyticsET::Config()>& GetEngineAnalyticsConfigFunc()
-{
-	static TFunction<FAnalyticsET::Config()> Config = &DefaultEngineAnalyticsConfigFunc;
-	return Config;
-}
-
 static TSharedPtr<IAnalyticsProviderET> CreateEpicAnalyticsProvider()
 {
-	// Get the default config.
-	FAnalyticsET::Config Config = GetEngineAnalyticsConfigFunc()();
-	// Set any fields that weren't set by default.
-	if (Config.APIKeyET.IsEmpty())
+	FAnalyticsET::Config Config;
 	{
 		// We always use the "Release" analytics account unless we're running in analytics test mode (usually with
 		// a command-line parameter), or we're an internal Epic build
@@ -200,18 +178,9 @@ static TSharedPtr<IAnalyticsProviderET> CreateEpicAnalyticsProvider()
 		const TCHAR* UE4TypeStr = bHasOverride ? *UE4TypeOverride : FEngineBuildSettings::IsPerforceBuild() ? TEXT("Perforce") : TEXT("UnrealEngine");
 		Config.APIKeyET = FString::Printf(TEXT("UEEditor.%s.%s"), UE4TypeStr, BuildTypeStr);
 	}
-	if (Config.APIServerET.IsEmpty())
-	{
-		Config.APIServerET = TEXT("https://datarouter.ol.epicgames.com/");
-	}
-	if (Config.AppEnvironment.IsEmpty())
-	{
-		Config.AppEnvironment = TEXT("datacollector-source");
-	}
-	if (Config.AppVersionET.IsEmpty())
-	{
-		Config.AppVersionET = FEngineVersion::Current().ToString();
-	}
+	Config.APIServerET = TEXT("https://datarouter.ol.epicgames.com/");
+	Config.AppEnvironment = TEXT("datacollector-binary");
+	Config.AppVersionET = FEngineVersion::Current().ToString();
 
 	// Connect the engine analytics provider (if there is a configuration delegate installed)
 	return FAnalyticsET::Get().CreateAnalyticsProvider(Config);
