Added an experimental project setting to enable MeshStreaming or not. False by default.

diff --git a/Engine/Source/Runtime/Engine/Classes/Engine/RendererSettings.h b/Engine/Source/Runtime/Engine/Classes/Engine/RendererSettings.h
--- a/Engine/Source/Runtime/Engine/Classes/Engine/RendererSettings.h
+++ b/Engine/Source/Runtime/Engine/Classes/Engine/RendererSettings.h
@@ -634,6 +634,12 @@ class ENGINE_API URendererSettings : public UDeveloperSettings
 		ConfigRestartRequired = true))
 		uint32 bODSCapture : 1;
 
+	UPROPERTY(config, EditAnywhere, Category = Experimental, meta = (
+		ConsoleVariable="r.MeshStreaming",DisplayName="Mesh Streaming",
+		ToolTip="When enabled mesh will stream in based on what is visible on screen.",
+		ConfigRestartRequired = true))
+		uint32 bMeshStreaming : 1;
+
 	UPROPERTY(config, EditAnywhere, Category=Editor, meta=(
 		ConsoleVariable="r.WireframeCullThreshold",DisplayName="Wireframe Cull Threshold",
 		ToolTip="Screen radius at which wireframe objects are culled. Larger values can improve performance when viewing a scene in wireframe."))
diff --git a/Engine/Source/Runtime/Engine/Private/ContentStreaming.cpp b/Engine/Source/Runtime/Engine/Private/ContentStreaming.cpp
--- a/Engine/Source/Runtime/Engine/Private/ContentStreaming.cpp
+++ b/Engine/Source/Runtime/Engine/Private/ContentStreaming.cpp
@@ -27,6 +27,14 @@
 	Globals.
 -----------------------------------------------------------------------------*/
 
+static TAutoConsoleVariable<int32> CVarMeshStreaming(
+	TEXT("r.MeshStreaming"),
+	0,
+	TEXT("Experimental - ")
+	TEXT("When non zero, enables mesh stremaing.\n"),
+	ECVF_ReadOnly | ECVF_RenderThreadSafe);
+
+
 /** Collection of views that need to be taken into account for streaming. */
 TArray<FStreamingViewInfo> IStreamingManager::CurrentViewInfos;
 
@@ -974,6 +982,11 @@ bool FStreamingManagerCollection::IsTextureStreamingEnabled() const
 	return IsRenderAssetStreamingEnabled();
 }
 
+bool FStreamingManagerCollection::IsMeshStreamingEnabled() const
+{
+	return CVarMeshStreaming.GetValueOnAnyThread() != 0 && IsRenderAssetStreamingEnabled();
+}
+
 bool FStreamingManagerCollection::IsRenderAssetStreamingEnabled() const
 {
 	return TextureStreamingManager != 0;
diff --git a/Engine/Source/Runtime/Engine/Private/SkeletalMesh.cpp b/Engine/Source/Runtime/Engine/Private/SkeletalMesh.cpp
--- a/Engine/Source/Runtime/Engine/Private/SkeletalMesh.cpp
+++ b/Engine/Source/Runtime/Engine/Private/SkeletalMesh.cpp
@@ -1165,7 +1165,7 @@ bool USkeletalMesh::UpdateStreamingStatus(bool bWaitForMipFading)
 
 void USkeletalMesh::LinkStreaming()
 {
-	if (!IsTemplate() && IStreamingManager::Get().IsTextureStreamingEnabled() && IsStreamingRenderAsset(this))
+	if (!IsTemplate() && IStreamingManager::Get().IsMeshStreamingEnabled() && IsStreamingRenderAsset(this))
 	{
 		IStreamingManager::Get().GetTextureStreamingManager().AddStreamingRenderAsset(this);
 	}
@@ -1177,7 +1177,7 @@ void USkeletalMesh::LinkStreaming()
 
 void USkeletalMesh::UnlinkStreaming()
 {
-	if (!IsTemplate() && IStreamingManager::Get().IsTextureStreamingEnabled())
+	if (!IsTemplate() && StreamingIndex != INDEX_NONE)
 	{
 		IStreamingManager::Get().GetTextureStreamingManager().RemoveStreamingRenderAsset(this);
 	}
diff --git a/Engine/Source/Runtime/Engine/Private/SkeletalMeshLODRenderData.cpp b/Engine/Source/Runtime/Engine/Private/SkeletalMeshLODRenderData.cpp
--- a/Engine/Source/Runtime/Engine/Private/SkeletalMeshLODRenderData.cpp
+++ b/Engine/Source/Runtime/Engine/Private/SkeletalMeshLODRenderData.cpp
@@ -751,11 +751,14 @@ bool FSkeletalMeshLODRenderData::IsLODCookedOut(const ITargetPlatform* TargetPla
 	}
 	check(TargetPlatform);
 
-	const bool bSupportLODStreaming = SkeletalMesh->bSupportLODStreaming.GetValueForPlatformIdentifiers(
+	const bool bSupportLODStreaming = !SkeletalMesh->NeverStream && SkeletalMesh->bSupportLODStreaming.GetValueForPlatformIdentifiers(
 		TargetPlatform->GetPlatformInfo().PlatformGroupName,
 		TargetPlatform->GetPlatformInfo().VanillaPlatformName);
+
+	static auto* VarMeshStreaming = IConsoleManager::Get().FindConsoleVariable(TEXT("r.MeshStreaming"));
+	const bool bMeshStreamingEnabled = !VarMeshStreaming || VarMeshStreaming->GetInt() != 0;
 	
-	return !TargetPlatform->SupportsFeature(ETargetPlatformFeatures::MeshLODStreaming) || !bSupportLODStreaming;
+	return !bMeshStreamingEnabled || !TargetPlatform->SupportsFeature(ETargetPlatformFeatures::MeshLODStreaming) || !bSupportLODStreaming;
 #else
 	return false;
 #endif
@@ -773,9 +776,12 @@ bool FSkeletalMeshLODRenderData::IsLODInlined(const ITargetPlatform* TargetPlatf
 
 	const FName PlatformGroupName = TargetPlatform->GetPlatformInfo().PlatformGroupName;
 	const FName VanillaPlatformName = TargetPlatform->GetPlatformInfo().VanillaPlatformName;
-	const bool bSupportLODStreaming = SkeletalMesh->bSupportLODStreaming.GetValueForPlatformIdentifiers(PlatformGroupName, VanillaPlatformName);
+	const bool bSupportLODStreaming = !SkeletalMesh->NeverStream && SkeletalMesh->bSupportLODStreaming.GetValueForPlatformIdentifiers(PlatformGroupName, VanillaPlatformName);
+
+	static auto* VarMeshStreaming = IConsoleManager::Get().FindConsoleVariable(TEXT("r.MeshStreaming"));
+	const bool bMeshStreamingEnabled = !VarMeshStreaming || VarMeshStreaming->GetInt() != 0;
 
-	if (!TargetPlatform->SupportsFeature(ETargetPlatformFeatures::MeshLODStreaming) || !bSupportLODStreaming)
+	if (!bMeshStreamingEnabled || !TargetPlatform->SupportsFeature(ETargetPlatformFeatures::MeshLODStreaming) || !bSupportLODStreaming)
 	{
 		return true;
 	}
diff --git a/Engine/Source/Runtime/Engine/Private/SkeletalMeshRenderData.cpp b/Engine/Source/Runtime/Engine/Private/SkeletalMeshRenderData.cpp
--- a/Engine/Source/Runtime/Engine/Private/SkeletalMeshRenderData.cpp
+++ b/Engine/Source/Runtime/Engine/Private/SkeletalMeshRenderData.cpp
@@ -104,8 +104,11 @@ static FString BuildSkeletalMeshDerivedDataKey(const ITargetPlatform* TargetPlat
 	const FName PlatformGroupName = TargetPlatform->GetPlatformInfo().PlatformGroupName;
 	const FName VanillaPlatformName = TargetPlatform->GetPlatformInfo().VanillaPlatformName;
 
-	const bool bSupportLODStreaming = SkelMesh->bSupportLODStreaming.GetValueForPlatformIdentifiers(PlatformGroupName, VanillaPlatformName);
-	if (bSupportLODStreaming)
+	static auto* VarMeshStreaming = IConsoleManager::Get().FindConsoleVariable(TEXT("r.MeshStreaming"));
+	const bool bMeshStreamingEnabled = !VarMeshStreaming || VarMeshStreaming->GetInt() != 0;
+
+	const bool bSupportLODStreaming = !SkelMesh->NeverStream && SkelMesh->bSupportLODStreaming.GetValueForPlatformIdentifiers(PlatformGroupName, VanillaPlatformName);
+	if (bMeshStreamingEnabled && TargetPlatform->SupportsFeature(ETargetPlatformFeatures::MeshLODStreaming) && bSupportLODStreaming)
 	{
 		const int32 MaxNumStreamedLODs = SkelMesh->MaxNumStreamedLODs.GetValueForPlatformIdentifiers(PlatformGroupName, VanillaPlatformName);
 		const int32 MaxNumOptionalLODs = SkelMesh->MaxNumOptionalLODs.GetValueForPlatformIdentifiers(PlatformGroupName, VanillaPlatformName);
@@ -490,3 +493,14 @@ int32 FSkeletalMeshRenderData::GetMaxBonesPerSection() const
 	}
 	return MaxBonesPerSection;
 }
+
+int32 FSkeletalMeshRenderData::GetFirstValidLODIdx(int32 MinIdx) const
+{
+	const int32 LODCount = LODRenderData.Num();
+	int32 LODIndex = FMath::Clamp<int32>(MinIdx, 0, LODCount - 1);
+	while (LODIndex < LODCount && !LODRenderData[LODIndex].GetNumVertices())
+	{
+		++LODIndex;
+	}
+	return LODIndex;
+}
diff --git a/Engine/Source/Runtime/Engine/Private/StaticMesh.cpp b/Engine/Source/Runtime/Engine/Private/StaticMesh.cpp
--- a/Engine/Source/Runtime/Engine/Private/StaticMesh.cpp
+++ b/Engine/Source/Runtime/Engine/Private/StaticMesh.cpp
@@ -317,9 +317,12 @@ bool FStaticMeshLODResources::IsLODCookedOut(const ITargetPlatform* TargetPlatfo
 	}
 	check(TargetPlatform);
 
+	static auto* VarMeshStreaming = IConsoleManager::Get().FindConsoleVariable(TEXT("r.MeshStreaming"));
+	const bool bMeshStreamingEnabled = !VarMeshStreaming || VarMeshStreaming->GetInt() != 0;
+
 	// If LOD streaming is supported, LODs below MinLOD are stored to optional paks and thus never cooked out
 	const FStaticMeshLODGroup& LODGroupSettings = TargetPlatform->GetStaticMeshLODSettings().GetLODGroup(StaticMesh->LODGroup);
-	return !TargetPlatform->SupportsFeature(ETargetPlatformFeatures::MeshLODStreaming) || !LODGroupSettings.IsLODStreamingSupported();
+	return !bMeshStreamingEnabled || !TargetPlatform->SupportsFeature(ETargetPlatformFeatures::MeshLODStreaming) || StaticMesh->NeverStream || !LODGroupSettings.IsLODStreamingSupported();
 #else
 	return false;
 #endif
@@ -335,8 +338,11 @@ bool FStaticMeshLODResources::IsLODInlined(const ITargetPlatform* TargetPlatform
 	}
 	check(TargetPlatform);
 
+	static auto* VarMeshStreaming = IConsoleManager::Get().FindConsoleVariable(TEXT("r.MeshStreaming"));
+	const bool bMeshStreamingEnabled = !VarMeshStreaming || VarMeshStreaming->GetInt() != 0;
+
 	const FStaticMeshLODGroup& LODGroupSettings = TargetPlatform->GetStaticMeshLODSettings().GetLODGroup(StaticMesh->LODGroup);
-	if (!TargetPlatform->SupportsFeature(ETargetPlatformFeatures::MeshLODStreaming) || !LODGroupSettings.IsLODStreamingSupported())
+	if (!bMeshStreamingEnabled || !TargetPlatform->SupportsFeature(ETargetPlatformFeatures::MeshLODStreaming) || StaticMesh->NeverStream || !LODGroupSettings.IsLODStreamingSupported())
 	{
 		return true;
 	}
@@ -2671,10 +2677,7 @@ void UStaticMesh::InitResources()
 	// Determine whether or not this mesh can be streamed.
 	const int32 NumLODs = GetNumLODs();
 
-	// Most tools assume that LOD0 is available and wouldn't work correctly if not. For example lighting build.
-	// Because of this, mesh streaming is in game mode (cooked or not).
-	bIsStreamable = !GIsEditor 
-		&& !NeverStream
+	bIsStreamable = !NeverStream
 		&& NumLODs > 1
 		&& !RenderData->LODResources[0].bBuffersInlined;
 		//&& !bTemporarilyDisableStreaming;
@@ -2695,7 +2698,9 @@ void UStaticMesh::InitResources()
 	{
 		LinkStreaming();
 	}
-	else if (GIsEditor && RenderData)
+	
+	// Make sure the first LOD is actually the lowest LOD.
+	if (RenderData)
 	{
 		RenderData->CurrentFirstLODIdx = RenderData->GetFirstValidLODIdx(0);
 		SetCachedNumResidentLODs(RenderData->LODResources.Num() - RenderData->CurrentFirstLODIdx);
@@ -5618,7 +5623,8 @@ bool UStaticMesh::StreamIn(int32 NewMipCount, bool bHighPrio)
 		else
 #endif
 		{
-			if (GRHISupportsAsyncTextureCreation)
+			// When not using threaded rendering, rendercommands get executed on async thread which create issues on some RHI. See EnqueueUniqueRenderCommand() and IsInRenderingThread().
+			if (GRHISupportsAsyncTextureCreation && GIsThreadedRendering)
 			{
 				PendingUpdate = new FStaticMeshStreamIn_IO_Async(this, NewMipCount, bHighPrio);
 			}
@@ -5696,7 +5702,7 @@ bool UStaticMesh::UpdateStreamingStatus(bool bWaitForMipFading)
 
 void UStaticMesh::LinkStreaming()
 {
-	if (!IsTemplate() && IStreamingManager::Get().IsTextureStreamingEnabled() && IsStreamingRenderAsset(this))
+	if (!IsTemplate() && IStreamingManager::Get().IsMeshStreamingEnabled() && IsStreamingRenderAsset(this))
 	{
 		IStreamingManager::Get().GetTextureStreamingManager().AddStreamingRenderAsset(this);
 	}
@@ -5708,7 +5714,7 @@ void UStaticMesh::LinkStreaming()
 
 void UStaticMesh::UnlinkStreaming()
 {
-	if (!IsTemplate() && IStreamingManager::Get().IsTextureStreamingEnabled())
+	if (!IsTemplate() && StreamingIndex != INDEX_NONE)
 	{
 		IStreamingManager::Get().GetTextureStreamingManager().RemoveStreamingRenderAsset(this);
 	}
diff --git a/Engine/Source/Runtime/Engine/Public/ContentStreaming.h b/Engine/Source/Runtime/Engine/Public/ContentStreaming.h
--- a/Engine/Source/Runtime/Engine/Public/ContentStreaming.h
+++ b/Engine/Source/Runtime/Engine/Public/ContentStreaming.h
@@ -669,12 +669,17 @@ struct FStreamingManagerCollection : public IStreamingManager
 	ENGINE_API bool IsStreamingEnabled() const;
 
 	/**
-	 * Checks whether texture streaming is active
+	 * Checks whether texture streaming is enabled
 	 */
 	ENGINE_API bool IsTextureStreamingEnabled() const;
 
 	/**
-	 * Checks whether texture/mesh streaming is active
+	 * Checks whether mesh streaming is enabled
+	 */
+	ENGINE_API bool IsMeshStreamingEnabled() const;
+
+	/**
+	 * Checks whether texture/mesh streaming is enabled
 	 */
 	virtual bool IsRenderAssetStreamingEnabled() const;
 
diff --git a/Engine/Source/Runtime/Engine/Public/Rendering/SkeletalMeshRenderData.h b/Engine/Source/Runtime/Engine/Public/Rendering/SkeletalMeshRenderData.h
--- a/Engine/Source/Runtime/Engine/Public/Rendering/SkeletalMeshRenderData.h
+++ b/Engine/Source/Runtime/Engine/Public/Rendering/SkeletalMeshRenderData.h
@@ -69,6 +69,15 @@ public:
 	*/
 	ENGINE_API int32 GetMaxBonesPerSection() const;
 
+	/** Return first valid LOD index starting at MinLODIdx. */
+	ENGINE_API int32 GetFirstValidLODIdx(int32 MinLODIdx) const;
+
+	/** Return the current first LODIdx that can be used. */
+	FORCEINLINE int32 GetCurrentFirstLODIdx(int32 MinLODIdx) const
+	{
+		return GetFirstValidLODIdx(FMath::Max<int32>(CurrentFirstLODIdx, MinLODIdx));
+	}
+
 private:
 	/** True if the resource has been initialized. */
 	bool bInitialized = false;
