Optimize EvaluatePreshader to reduce number of TMemoryImagePtr accesses

diff --git a/Engine/Source/Runtime/Engine/Private/Materials/MaterialUniformExpressions.cpp b/Engine/Source/Runtime/Engine/Private/Materials/MaterialUniformExpressions.cpp
--- a/Engine/Source/Runtime/Engine/Private/Materials/MaterialUniformExpressions.cpp
+++ b/Engine/Source/Runtime/Engine/Private/Materials/MaterialUniformExpressions.cpp
@@ -527,9 +527,28 @@ void FMaterialPreshaderData::WriteData(const void* Value, uint32 Size)
 	Data.Append((uint8*)Value, Size);
 }
 
-void FMaterialPreshaderData::ReadData(void* Result, uint32 Offset, uint32 Size) const
+struct FPreshaderDataPtr
 {
-	FMemory::Memcpy(Result, &Data[Offset], Size);
+	FPreshaderDataPtr(const uint8* InPtr = nullptr) : Ptr(InPtr) {}
+	const uint8* RESTRICT Ptr;
+};
+
+namespace
+{
+	template<typename T>
+	inline T ReadPreshaderValue(FPreshaderDataPtr& RESTRICT Data)
+	{
+		T Result;
+		FMemory::Memcpy(&Result, Data.Ptr, sizeof(T));
+		Data.Ptr += sizeof(T);
+		return Result;
+	}
+
+	template<>
+	inline uint8 ReadPreshaderValue<uint8>(FPreshaderDataPtr& RESTRICT Data)
+	{
+		return *Data.Ptr++;
+	}
 }
 
 static void GetVectorParameter(const FUniformExpressionSet& UniformExpressionSet, uint32 ParameterIndex, const FMaterialRenderContext& Context, FLinearColor& OutValue)
@@ -608,14 +627,14 @@ static void GetScalarParameter(const FUniformExpressionSet& UniformExpressionSet
 using FPreshaderStack = TArray<FLinearColor, TInlineAllocator<64u>>;
 
 template<typename Operation>
-static void EvaluateUnaryOp(FPreshaderStack& Stack, const Operation& Op)
+static inline void EvaluateUnaryOp(FPreshaderStack& Stack, const Operation& Op)
 {
 	const FLinearColor Value = Stack.Pop(false);
 	Stack.Add(FLinearColor(Op(Value.R), Op(Value.G), Op(Value.B), Op(Value.A)));
 }
 
 template<typename Operation>
-static void EvaluateBinaryOp(FPreshaderStack& Stack, const Operation& Op)
+static inline void EvaluateBinaryOp(FPreshaderStack& Stack, const Operation& Op)
 {
 	const FLinearColor Value1 = Stack.Pop(false);
 	const FLinearColor Value0 = Stack.Pop(false);
@@ -623,7 +642,7 @@ static void EvaluateBinaryOp(FPreshaderStack& Stack, const Operation& Op)
 }
 
 template<typename Operation>
-static void EvaluateTernaryOp(FPreshaderStack& Stack, const Operation& Op)
+static inline void EvaluateTernaryOp(FPreshaderStack& Stack, const Operation& Op)
 {
 	const FLinearColor Value2 = Stack.Pop(false);
 	const FLinearColor Value1 = Stack.Pop(false);
@@ -631,9 +650,9 @@ static void EvaluateTernaryOp(FPreshaderStack& Stack, const Operation& Op)
 	Stack.Add(FLinearColor(Op(Value0.R, Value1.R, Value2.R), Op(Value0.G, Value1.G, Value2.G), Op(Value0.B, Value1.B, Value2.B), Op(Value0.A, Value1.A, Value2.A)));
 }
 
-static void EvaluateDot(FPreshaderStack& Stack, const FMaterialPreshaderData& Data, uint32& Offset)
+static void EvaluateDot(FPreshaderStack& Stack, FPreshaderDataPtr& RESTRICT Data)
 {
-	const uint8 ValueType = Data.Read<uint8>(Offset);
+	const uint8 ValueType = ReadPreshaderValue<uint8>(Data);
 	const FLinearColor Value1 = Stack.Pop(false);
 	const FLinearColor Value0 = Stack.Pop(false);
 	float Result = Value0.R * Value1.R;
@@ -643,9 +662,9 @@ static void EvaluateDot(FPreshaderStack& Stack, const FMaterialPreshaderData& Da
 	Stack.Add(FLinearColor(Result, Result, Result, Result));
 }
 
-static void EvaluateCross(FPreshaderStack& Stack, const FMaterialPreshaderData& Data, uint32& Offset)
+static void EvaluateCross(FPreshaderStack& Stack, FPreshaderDataPtr& RESTRICT Data)
 {
-	const uint8 ValueType = Data.Read<uint8>(Offset);
+	const uint8 ValueType = ReadPreshaderValue<uint8>(Data);
 	FLinearColor ValueB = Stack.Pop(false);
 	FLinearColor ValueA = Stack.Pop(false);
 	
@@ -670,13 +689,13 @@ static void EvaluateCross(FPreshaderStack& Stack, const FMaterialPreshaderData&
 	Stack.Add(FLinearColor(Cross.X, Cross.Y, Cross.Z, 0.0f));
 }
 
-static void EvaluateComponentSwizzle(FPreshaderStack& Stack, const FMaterialPreshaderData& Data, uint32& Offset)
+static void EvaluateComponentSwizzle(FPreshaderStack& Stack, FPreshaderDataPtr& RESTRICT Data)
 {
-	const uint8 NumElements = Data.Read<uint8>(Offset);
-	const uint8 IndexR = Data.Read<uint8>(Offset);
-	const uint8 IndexG = Data.Read<uint8>(Offset);
-	const uint8 IndexB = Data.Read<uint8>(Offset);
-	const uint8 IndexA = Data.Read<uint8>(Offset);
+	const uint8 NumElements = ReadPreshaderValue<uint8>(Data);
+	const uint8 IndexR = ReadPreshaderValue<uint8>(Data);
+	const uint8 IndexG = ReadPreshaderValue<uint8>(Data);
+	const uint8 IndexB = ReadPreshaderValue<uint8>(Data);
+	const uint8 IndexA = ReadPreshaderValue<uint8>(Data);
 
 	FLinearColor Value = Stack.Pop(false);
 	FLinearColor Result(0.0f, 0.0f, 0.0f, 0.0f);
@@ -703,9 +722,9 @@ static void EvaluateComponentSwizzle(FPreshaderStack& Stack, const FMaterialPres
 	Stack.Add(Result);
 }
 
-static void EvaluateAppenedVector(FPreshaderStack& Stack, const FMaterialPreshaderData& Data, uint32& Offset)
+static void EvaluateAppenedVector(FPreshaderStack& Stack, FPreshaderDataPtr& RESTRICT Data)
 {
-	const uint8 NumComponentsA = Data.Read<uint8>(Offset);
+	const uint8 NumComponentsA = ReadPreshaderValue<uint8>(Data);
 
 	const FLinearColor ValueB = Stack.Pop(false);
 	const FLinearColor ValueA = Stack.Pop(false);
@@ -718,19 +737,19 @@ static void EvaluateAppenedVector(FPreshaderStack& Stack, const FMaterialPreshad
 	Stack.Add(Result);
 }
 
-static const UTexture* GetTextureParameter(const FMaterialRenderContext& Context, const FMaterialPreshaderData& Data, uint32& Offset)
+static const UTexture* GetTextureParameter(const FMaterialRenderContext& Context, FPreshaderDataPtr& RESTRICT Data)
 {
-	const FHashedMaterialParameterInfo ParameterInfo = Data.Read<FHashedMaterialParameterInfo>(Offset);
-	const int32 TextureIndex = Data.Read<int32>(Offset);
+	const FHashedMaterialParameterInfo ParameterInfo = ReadPreshaderValue<FHashedMaterialParameterInfo>(Data);
+	const int32 TextureIndex = ReadPreshaderValue<int32>(Data);
 	
 	const UTexture* Texture = nullptr;
 	GetTextureParameterValue(ParameterInfo, TextureIndex, Context, Texture);
 	return Texture;
 }
 
-static void EvaluateTextureSize(const FMaterialRenderContext& Context, FPreshaderStack& Stack, const FMaterialPreshaderData& Data, uint32& Offset)
+static void EvaluateTextureSize(const FMaterialRenderContext& Context, FPreshaderStack& Stack, FPreshaderDataPtr& RESTRICT Data)
 {
-	const UTexture* Texture = GetTextureParameter(Context, Data, Offset);
+	const UTexture* Texture = GetTextureParameter(Context, Data);
 	if (Texture)
 	{
 		const uint32 SizeX = Texture->Resource->GetSizeX();
@@ -740,9 +759,9 @@ static void EvaluateTextureSize(const FMaterialRenderContext& Context, FPreshade
 	}
 }
 
-static void EvaluateTexelSize(const FMaterialRenderContext& Context, FPreshaderStack& Stack, const FMaterialPreshaderData& Data, uint32& Offset)
+static void EvaluateTexelSize(const FMaterialRenderContext& Context, FPreshaderStack& Stack, FPreshaderDataPtr& RESTRICT Data)
 {
-	const UTexture* Texture = GetTextureParameter(Context, Data, Offset);
+	const UTexture* Texture = GetTextureParameter(Context, Data);
 	if (Texture)
 	{
 		const uint32 SizeX = Texture->Resource->GetSizeX();
@@ -752,17 +771,17 @@ static void EvaluateTexelSize(const FMaterialRenderContext& Context, FPreshaderS
 	}
 }
 
-static FGuid GetExternalTextureGuid(const FMaterialRenderContext& Context, const FMaterialPreshaderData& Data, uint32& Offset)
+static FGuid GetExternalTextureGuid(const FMaterialRenderContext& Context, FPreshaderDataPtr& RESTRICT Data)
 {
-	const FHashedName ParameterName = Data.Read<FHashedName>(Offset);
-	const FGuid ExternalTextureGuid = Data.Read<FGuid>(Offset);
-	const int32 TextureIndex = Data.Read<int32>(Offset);
+	const FHashedName ParameterName = ReadPreshaderValue<FHashedName>(Data);
+	const FGuid ExternalTextureGuid = ReadPreshaderValue<FGuid>(Data);
+	const int32 TextureIndex = ReadPreshaderValue<int32>(Data);
 	return GetExternalTextureGuid(Context, ExternalTextureGuid, ParameterName, TextureIndex);
 }
 
-static void EvaluateExternalTextureCoordinateScaleRotation(const FMaterialRenderContext& Context, FPreshaderStack& Stack, const FMaterialPreshaderData& Data, uint32& Offset)
+static void EvaluateExternalTextureCoordinateScaleRotation(const FMaterialRenderContext& Context, FPreshaderStack& Stack, FPreshaderDataPtr& RESTRICT Data)
 {
-	const FGuid GuidToLookup = GetExternalTextureGuid(Context, Data, Offset);
+	const FGuid GuidToLookup = GetExternalTextureGuid(Context, Data);
 	FLinearColor Result(1.f, 0.f, 0.f, 1.f);
 	if (GuidToLookup.IsValid())
 	{
@@ -771,9 +790,9 @@ static void EvaluateExternalTextureCoordinateScaleRotation(const FMaterialRender
 	Stack.Add(Result);
 }
 
-static void EvaluateExternalTextureCoordinateOffset(const FMaterialRenderContext& Context, FPreshaderStack& Stack, const FMaterialPreshaderData& Data, uint32& Offset)
+static void EvaluateExternalTextureCoordinateOffset(const FMaterialRenderContext& Context, FPreshaderStack& Stack, FPreshaderDataPtr& RESTRICT Data)
 {
-	const FGuid GuidToLookup = GetExternalTextureGuid(Context, Data, Offset);
+	const FGuid GuidToLookup = GetExternalTextureGuid(Context, Data);
 	FLinearColor Result(0.f, 0.f, 0.f, 0.f);
 	if (GuidToLookup.IsValid())
 	{
@@ -782,11 +801,11 @@ static void EvaluateExternalTextureCoordinateOffset(const FMaterialRenderContext
 	Stack.Add(Result);
 }
 
-static void EvaluateRuntimeVirtualTextureUniform(const FMaterialRenderContext& Context, FPreshaderStack& Stack, const FMaterialPreshaderData& Data, uint32& Offset)
+static void EvaluateRuntimeVirtualTextureUniform(const FMaterialRenderContext& Context, FPreshaderStack& Stack, FPreshaderDataPtr& RESTRICT Data)
 {
-	const FHashedMaterialParameterInfo ParameterInfo = Data.Read<FHashedMaterialParameterInfo>(Offset);
-	const int32 TextureIndex = Data.Read<int32>(Offset);
-	const int32 VectorIndex = Data.Read<int32>(Offset);
+	const FHashedMaterialParameterInfo ParameterInfo = ReadPreshaderValue<FHashedMaterialParameterInfo>(Data);
+	const int32 TextureIndex = ReadPreshaderValue<int32>(Data);
+	const int32 VectorIndex = ReadPreshaderValue<int32>(Data);
 
 	const URuntimeVirtualTexture* Texture = nullptr;
 	if (ParameterInfo.Name.IsNone() || !Context.MaterialRenderProxy || !Context.MaterialRenderProxy->GetTextureValue(ParameterInfo, &Texture, Context))
@@ -833,32 +852,31 @@ FORCENOINLINE static float DivideComponent(float A, float B)
 	return A / GetSafeDivisor(B);
 }
 
-static void EvaluatePreshader(const FUniformExpressionSet* UniformExpressionSet, const FMaterialRenderContext& Context, FPreshaderStack& Stack, const FMaterialPreshaderData& Data, const FMaterialUniformPreshaderHeader& Preshader, FLinearColor& OutValue)
+static void EvaluatePreshader(const FUniformExpressionSet* UniformExpressionSet, const FMaterialRenderContext& Context, FPreshaderStack& Stack, const uint8* BaseData, uint32 Size, FLinearColor& OutValue)
 {
 	static const float LogToLog10 = 1.0f / FMath::Loge(10.f);
-
-	uint32 Offset = Preshader.OpcodeOffset;
-	const uint32 OffsetEnd = Offset + Preshader.OpcodeSize;
+	FPreshaderDataPtr Data(BaseData);
+	uint8 const* const DataEnd = BaseData + Size;
 
 	Stack.Reset();
-	while (Offset < OffsetEnd)
+	while (Data.Ptr < DataEnd)
 	{
-		const EMaterialPreshaderOpcode Opcode = Data.ReadOpcode(Offset);
+		const EMaterialPreshaderOpcode Opcode = (EMaterialPreshaderOpcode)ReadPreshaderValue<uint8>(Data);
 		switch (Opcode)
 		{
 		case EMaterialPreshaderOpcode::ConstantZero:
 			Stack.Add(FLinearColor(0.0f, 0.0f, 0.0f, 0.0f));
 			break;
 		case EMaterialPreshaderOpcode::Constant:
-			Stack.Add(Data.Read<FLinearColor>(Offset));
+			Stack.Add(ReadPreshaderValue<FLinearColor>(Data));
 			break;
 		case EMaterialPreshaderOpcode::VectorParameter:
 			check(UniformExpressionSet);
-			GetVectorParameter(*UniformExpressionSet, Data.Read<uint16>(Offset), Context, Stack.AddDefaulted_GetRef());
+			GetVectorParameter(*UniformExpressionSet, ReadPreshaderValue<uint16>(Data), Context, Stack.AddDefaulted_GetRef());
 			break;
 		case EMaterialPreshaderOpcode::ScalarParameter:
 			check(UniformExpressionSet);
-			GetScalarParameter(*UniformExpressionSet, Data.Read<uint16>(Offset), Context, Stack.AddDefaulted_GetRef());
+			GetScalarParameter(*UniformExpressionSet, ReadPreshaderValue<uint16>(Data), Context, Stack.AddDefaulted_GetRef());
 			break;
 		case EMaterialPreshaderOpcode::Add: EvaluateBinaryOp(Stack, [](float Lhs, float Rhs) { return Lhs + Rhs; }); break;
 		case EMaterialPreshaderOpcode::Sub: EvaluateBinaryOp(Stack, [](float Lhs, float Rhs) { return Lhs - Rhs; }); break;
@@ -868,8 +886,8 @@ static void EvaluatePreshader(const FUniformExpressionSet* UniformExpressionSet,
 		case EMaterialPreshaderOpcode::Min: EvaluateBinaryOp(Stack, [](float Lhs, float Rhs) { return FMath::Min(Lhs, Rhs); }); break;
 		case EMaterialPreshaderOpcode::Max: EvaluateBinaryOp(Stack, [](float Lhs, float Rhs) { return FMath::Max(Lhs, Rhs); }); break;
 		case EMaterialPreshaderOpcode::Clamp: EvaluateTernaryOp(Stack, [](float A, float B, float C) { return FMath::Clamp(A, B, C); }); break;
-		case EMaterialPreshaderOpcode::Dot: EvaluateDot(Stack, Data, Offset); break;
-		case EMaterialPreshaderOpcode::Cross: EvaluateCross(Stack, Data, Offset); break;
+		case EMaterialPreshaderOpcode::Dot: EvaluateDot(Stack, Data); break;
+		case EMaterialPreshaderOpcode::Cross: EvaluateCross(Stack, Data); break;
 		case EMaterialPreshaderOpcode::Sqrt: EvaluateUnaryOp(Stack, [](float V) { return FMath::Sqrt(V); }); break;
 		case EMaterialPreshaderOpcode::Sin: EvaluateUnaryOp(Stack, [](float V) { return FMath::Sin(V); }); break;
 		case EMaterialPreshaderOpcode::Cos: EvaluateUnaryOp(Stack, [](float V) { return FMath::Cos(V); }); break;
@@ -889,19 +907,19 @@ static void EvaluatePreshader(const FUniformExpressionSet* UniformExpressionSet,
 		case EMaterialPreshaderOpcode::Fractional: EvaluateUnaryOp(Stack, [](float V) { return FMath::Fractional(V); }); break;
 		case EMaterialPreshaderOpcode::Log2: EvaluateUnaryOp(Stack, [](float V) { return FMath::Log2(V); }); break;
 		case EMaterialPreshaderOpcode::Log10: EvaluateUnaryOp(Stack, [](float V) { return FMath::Loge(V) * LogToLog10; }); break;
-		case EMaterialPreshaderOpcode::ComponentSwizzle: EvaluateComponentSwizzle(Stack, Data, Offset); break;
-		case EMaterialPreshaderOpcode::AppendVector: EvaluateAppenedVector(Stack, Data, Offset); break;
-		case EMaterialPreshaderOpcode::TextureSize: EvaluateTextureSize(Context, Stack, Data, Offset); break;
-		case EMaterialPreshaderOpcode::TexelSize: EvaluateTexelSize(Context, Stack, Data, Offset); break;
-		case EMaterialPreshaderOpcode::ExternalTextureCoordinateScaleRotation: EvaluateExternalTextureCoordinateScaleRotation(Context, Stack, Data, Offset); break;
-		case EMaterialPreshaderOpcode::ExternalTextureCoordinateOffset: EvaluateExternalTextureCoordinateOffset(Context, Stack, Data, Offset); break;
-		case EMaterialPreshaderOpcode::RuntimeVirtualTextureUniform: EvaluateRuntimeVirtualTextureUniform(Context, Stack, Data, Offset); break;
+		case EMaterialPreshaderOpcode::ComponentSwizzle: EvaluateComponentSwizzle(Stack, Data); break;
+		case EMaterialPreshaderOpcode::AppendVector: EvaluateAppenedVector(Stack, Data); break;
+		case EMaterialPreshaderOpcode::TextureSize: EvaluateTextureSize(Context, Stack, Data); break;
+		case EMaterialPreshaderOpcode::TexelSize: EvaluateTexelSize(Context, Stack, Data); break;
+		case EMaterialPreshaderOpcode::ExternalTextureCoordinateScaleRotation: EvaluateExternalTextureCoordinateScaleRotation(Context, Stack, Data); break;
+		case EMaterialPreshaderOpcode::ExternalTextureCoordinateOffset: EvaluateExternalTextureCoordinateOffset(Context, Stack, Data); break;
+		case EMaterialPreshaderOpcode::RuntimeVirtualTextureUniform: EvaluateRuntimeVirtualTextureUniform(Context, Stack, Data); break;
 		default:
 			UE_LOG(LogMaterial, Fatal, TEXT("Unknown preshader opcode %d"), (uint8)Opcode);
 			break;
 		}
 	}
-	check(Offset == OffsetEnd);
+	check(Data.Ptr == DataEnd);
 
 	ensure(Stack.Num() <= 1);
 	if (Stack.Num() > 0)
@@ -915,12 +933,8 @@ void FMaterialUniformExpression::GetNumberValue(const struct FMaterialRenderCont
 	FMaterialPreshaderData PreshaderData;
 	WriteNumberOpcodes(PreshaderData);
 
-	FMaterialUniformPreshaderHeader PreshaderHeader;
-	PreshaderHeader.OpcodeOffset = 0u;
-	PreshaderHeader.OpcodeSize = PreshaderData.Num();
-	
 	FPreshaderStack Stack;
-	EvaluatePreshader(nullptr, Context, Stack, PreshaderData, PreshaderHeader, OutValue);
+	EvaluatePreshader(nullptr, Context, Stack, PreshaderData.Data.GetData(), PreshaderData.Num(), OutValue);
 }
 
 const FMaterialVectorParameterInfo* FUniformExpressionSet::FindVectorParameter(const FHashedMaterialParameterInfo& ParameterInfo) const
@@ -1066,12 +1080,13 @@ void FUniformExpressionSet::FillUniformBuffer(const FMaterialRenderContext& Mate
 
 		// Dump vector expression into the buffer.
 		FPreshaderStack PreshaderStack;
+		const uint8* PreshaderData = UniformPreshaderData.Data.GetData();
 		for(int32 VectorIndex = 0;VectorIndex < UniformVectorPreshaders.Num();++VectorIndex)
 		{
 			FLinearColor VectorValue(0, 0, 0, 0);
 
 			const FMaterialUniformPreshaderHeader& Preshader = UniformVectorPreshaders[VectorIndex];
-			EvaluatePreshader(this, MaterialRenderContext, PreshaderStack, UniformPreshaderData, Preshader, VectorValue);
+			EvaluatePreshader(this, MaterialRenderContext, PreshaderStack, PreshaderData + Preshader.OpcodeOffset, Preshader.OpcodeSize, VectorValue);
 	
 			FLinearColor* DestAddress = (FLinearColor*)BufferCursor;
 			*DestAddress = VectorValue;
@@ -1085,7 +1100,7 @@ void FUniformExpressionSet::FillUniformBuffer(const FMaterialRenderContext& Mate
 			FLinearColor VectorValue(0,0,0,0);
 
 			const FMaterialUniformPreshaderHeader& Preshader = UniformScalarPreshaders[ScalarIndex];
-			EvaluatePreshader(this, MaterialRenderContext, PreshaderStack, UniformPreshaderData, Preshader, VectorValue);
+			EvaluatePreshader(this, MaterialRenderContext, PreshaderStack, PreshaderData + Preshader.OpcodeOffset, Preshader.OpcodeSize, VectorValue);
 
 			float* DestAddress = (float*)BufferCursor;
 			*DestAddress = VectorValue.R;
diff --git a/Engine/Source/Runtime/Engine/Public/MaterialShared.h b/Engine/Source/Runtime/Engine/Public/MaterialShared.h
--- a/Engine/Source/Runtime/Engine/Public/MaterialShared.h
+++ b/Engine/Source/Runtime/Engine/Public/MaterialShared.h
@@ -274,22 +274,11 @@ public:
 	inline int32 Num() const { return Data.Num(); }
 
 	void WriteData(const void* Value, uint32 Size);
-	void ReadData(void* Result, uint32 Offset, uint32 Size) const;
 
 	template<typename T>
 	FMaterialPreshaderData& Write(const T& Value) { WriteData(&Value, sizeof(T)); return *this; }
 
-	template<typename T>
-	T Read(uint32& Offset) const
-	{
-		T Result;
-		ReadData(&Result, Offset, sizeof(T));
-		Offset += sizeof(T);
-		return Result;
-	}
-
 	inline FMaterialPreshaderData& WriteOpcode(EMaterialPreshaderOpcode Op) { return Write<uint8>((uint8)Op); }
-	inline EMaterialPreshaderOpcode ReadOpcode(uint32& Offset) const { return (EMaterialPreshaderOpcode)Read<uint8>(Offset); }
 
 	LAYOUT_FIELD(TMemoryImageArray<uint8>, Data);
 };
