Add a option for each platform to enable skeletal mesh LOD streaming by default. Artists can still override the default settings on a per-asset or per-group basis.

diff --git a/Engine/Platforms/XXX/Source/Developer/XXXTargetPlatform/Private/XXXTargetPlatform.cpp b/Engine/Platforms/XXX/Source/Developer/XXXTargetPlatform/Private/XXXTargetPlatform.cpp
--- a/Engine/Platforms/XXX/Source/Developer/XXXTargetPlatform/Private/XXXTargetPlatform.cpp
+++ b/Engine/Platforms/XXX/Source/Developer/XXXTargetPlatform/Private/XXXTargetPlatform.cpp
@@ -35,6 +35,7 @@ FXXXTargetPlatform::FXXXTargetPlatform()
 	// load up texture settings from the config file
 	XXXLODSettings = nullptr; // These are registered by the device profile system.
 	StaticMeshLODSettings.Initialize(XXXEngineSettings);
+	SkeletalMeshDefaultLODStreamingSettings.Initialize(XXXEngineSettings);
 #endif
 }
 
@@ -150,6 +151,12 @@ const class FStaticMeshLODSettings& FXXXTargetPlatform::GetStaticMeshLODSettings
 }
 
 
+const FSkeletalMeshDefaultLODStreamingSettings& FXXXTargetPlatform::GetSkeletalMeshDefaultLODStreamingSettings() const
+{
+	return SkeletalMeshDefaultLODStreamingSettings;
+}
+
+
 void FXXXTargetPlatform::GetTextureFormats( const UTexture* InTexture, TArray< TArray<FName> >& OutFormats ) const
 {
 	GetDefaultTextureFormatNamePerLayer(OutFormats.AddDefaulted_GetRef(), this, InTexture, XXXEngineSettings, true);
diff --git a/Engine/Platforms/XXX/Source/Developer/XXXTargetPlatform/Private/XXXTargetPlatform.h b/Engine/Platforms/XXX/Source/Developer/XXXTargetPlatform/Private/XXXTargetPlatform.h
--- a/Engine/Platforms/XXX/Source/Developer/XXXTargetPlatform/Private/XXXTargetPlatform.h
+++ b/Engine/Platforms/XXX/Source/Developer/XXXTargetPlatform/Private/XXXTargetPlatform.h
@@ -16,6 +16,7 @@
 
 #if WITH_ENGINE
 #include "StaticMeshResources.h"
+#include "SkeletalMeshDefaultLODStreamingSettings.h"
 #endif // WITH_ENGINE
 
 /**
@@ -63,6 +64,8 @@ public:
 
 	virtual const class FStaticMeshLODSettings& GetStaticMeshLODSettings() const override;
 
+	virtual const FSkeletalMeshDefaultLODStreamingSettings& GetSkeletalMeshDefaultLODStreamingSettings() const override;
+
 	virtual void GetTextureFormats( const UTexture* InTexture, TArray< TArray<FName> >& OutFormats ) const override;
 
 	virtual void GetAllTextureFormats(TArray<FName>& OutFormats) const override;
@@ -109,6 +112,8 @@ private:
 
 	// Holds the static mesh LOD settings.
 	FStaticMeshLODSettings StaticMeshLODSettings;
+
+	FSkeletalMeshDefaultLODStreamingSettings SkeletalMeshDefaultLODStreamingSettings;
 #endif
 
 private:
diff --git a/Engine/Plugins/Runtime/AR/Microsoft/HoloLensAR/Source/HoloLensTargetPlatform/Private/HoloLensTargetPlatform.cpp b/Engine/Plugins/Runtime/AR/Microsoft/HoloLensAR/Source/HoloLensTargetPlatform/Private/HoloLensTargetPlatform.cpp
--- a/Engine/Plugins/Runtime/AR/Microsoft/HoloLensAR/Source/HoloLensTargetPlatform/Private/HoloLensTargetPlatform.cpp
+++ b/Engine/Plugins/Runtime/AR/Microsoft/HoloLensAR/Source/HoloLensTargetPlatform/Private/HoloLensTargetPlatform.cpp
@@ -24,6 +24,7 @@ FHoloLensTargetPlatform::FHoloLensTargetPlatform()
 	FConfigCacheIni::LoadLocalIniFile(EngineSettings, TEXT("Engine"), true, *PlatformName());
 	TextureLODSettings = nullptr; // These are registered by the device profile system.
 	StaticMeshLODSettings.Initialize(EngineSettings);
+	SkeletalMeshDefaultLODStreamingSettings.Initialize(EngineSettings);
 #endif
 
 	DeviceDetector = IHoloLensDeviceDetector::Create();
diff --git a/Engine/Plugins/Runtime/AR/Microsoft/HoloLensAR/Source/HoloLensTargetPlatform/Private/HoloLensTargetPlatform.h b/Engine/Plugins/Runtime/AR/Microsoft/HoloLensAR/Source/HoloLensTargetPlatform/Private/HoloLensTargetPlatform.h
--- a/Engine/Plugins/Runtime/AR/Microsoft/HoloLensAR/Source/HoloLensTargetPlatform/Private/HoloLensTargetPlatform.h
+++ b/Engine/Plugins/Runtime/AR/Microsoft/HoloLensAR/Source/HoloLensTargetPlatform/Private/HoloLensTargetPlatform.h
@@ -19,6 +19,7 @@
 
 #if WITH_ENGINE
 #include "StaticMeshResources.h"
+#include "SkeletalMeshDefaultLODStreamingSettings.h"
 #endif // WITH_ENGINE
 
 #define LOCTEXT_NAMESPACE "HoloLensTargetPlatform"
@@ -70,6 +71,8 @@ public:
 
 	virtual const class FStaticMeshLODSettings& GetStaticMeshLODSettings() const override { return StaticMeshLODSettings; }
 
+	virtual FSkeletalMeshDefaultLODStreamingSettings& GetSkeletalMeshDefaultLODStreamingSettings() const override { return SkeletalMeshDefaultLODStreamingSettings; }
+
 	virtual void GetTextureFormats(const UTexture* InTexture, TArray< TArray<FName> >& OutFormats) const override;
 
 	virtual const UTextureLODSettings& GetTextureLODSettings() const override { return *TextureLODSettings; }
@@ -170,6 +173,8 @@ private:
 
 	// Holds the static mesh LOD settings.
 	FStaticMeshLODSettings StaticMeshLODSettings;
+
+	FSkeletalMeshDefaultLODStreamingSettings SkeletalMeshDefaultLODStreamingSettings;
 #endif //WITH_ENGINE
 
 private:
diff --git a/Engine/Source/Developer/AllDesktopTargetPlatform/Private/AllDesktopTargetPlatform.cpp b/Engine/Source/Developer/AllDesktopTargetPlatform/Private/AllDesktopTargetPlatform.cpp
--- a/Engine/Source/Developer/AllDesktopTargetPlatform/Private/AllDesktopTargetPlatform.cpp
+++ b/Engine/Source/Developer/AllDesktopTargetPlatform/Private/AllDesktopTargetPlatform.cpp
@@ -22,6 +22,7 @@ FAllDesktopTargetPlatform::FAllDesktopTargetPlatform()
 	// use non-platform specific settings
 	FConfigCacheIni::LoadLocalIniFile(EngineSettings, TEXT("Engine"), true, NULL);
 	StaticMeshLODSettings.Initialize(EngineSettings);
+	SkeletalMeshDefaultLODStreamingSettings.Initialize(EngineSettings);
 #endif // #if WITH_ENGINE
 
 }
diff --git a/Engine/Source/Developer/AllDesktopTargetPlatform/Private/AllDesktopTargetPlatform.h b/Engine/Source/Developer/AllDesktopTargetPlatform/Private/AllDesktopTargetPlatform.h
--- a/Engine/Source/Developer/AllDesktopTargetPlatform/Private/AllDesktopTargetPlatform.h
+++ b/Engine/Source/Developer/AllDesktopTargetPlatform/Private/AllDesktopTargetPlatform.h
@@ -15,6 +15,7 @@
 #if WITH_ENGINE
 #include "AudioCompressionSettings.h"
 #include "StaticMeshResources.h"
+#include "SkeletalMeshDefaultLODStreamingSettings.h"
 #endif // WITH_ENGINE
 
 class UTextureLODSettings;
@@ -133,6 +134,11 @@ public:
 		return StaticMeshLODSettings;
 	}
 
+	virtual const FSkeletalMeshDefaultLODStreamingSettings& GetSkeletalMeshDefaultLODStreamingSettings() const override
+	{
+		return SkeletalMeshDefaultLODStreamingSettings;
+	}
+
 	virtual void GetTextureFormats(const UTexture* Texture, TArray< TArray<FName> >& OutFormats) const override;
 
 	virtual void GetAllTextureFormats(TArray<FName>& OutFormats) const override;
@@ -177,6 +183,8 @@ private:
 
 	// Holds the static mesh LOD settings.
 	FStaticMeshLODSettings StaticMeshLODSettings;
+
+	FSkeletalMeshDefaultLODStreamingSettings SkeletalMeshDefaultLODStreamingSettings;
 #endif // WITH_ENGINE
 
 	// Holds an event delegate that is executed when a new target device has been discovered.
diff --git a/Engine/Source/Developer/Android/AndroidTargetPlatform/Private/AndroidTargetPlatform.cpp b/Engine/Source/Developer/Android/AndroidTargetPlatform/Private/AndroidTargetPlatform.cpp
--- a/Engine/Source/Developer/Android/AndroidTargetPlatform/Private/AndroidTargetPlatform.cpp
+++ b/Engine/Source/Developer/Android/AndroidTargetPlatform/Private/AndroidTargetPlatform.cpp
@@ -199,6 +199,7 @@ FAndroidTargetPlatform::FAndroidTargetPlatform(bool bInIsClient )
 		FConfigCacheIni::LoadLocalIniFile(EngineSettings, TEXT("Engine"), true, *IniPlatformName());
 			TextureLODSettings = nullptr; // These are registered by the device profile system.
 		StaticMeshLODSettings.Initialize(EngineSettings);
+		SkeletalMeshDefaultLODStreamingSettings.Initialize(EngineSettings);
 	#endif
 
 	TickDelegate = FTickerDelegate::CreateRaw(this, &FAndroidTargetPlatform::HandleTicker);
@@ -426,6 +427,12 @@ const FStaticMeshLODSettings& FAndroidTargetPlatform::GetStaticMeshLODSettings(
 }
 
 
+const FSkeletalMeshDefaultLODStreamingSettings& FAndroidTargetPlatform::GetSkeletalMeshDefaultLODStreamingSettings() const
+{
+	return SkeletalMeshDefaultLODStreamingSettings;
+}
+
+
 void FAndroidTargetPlatform::GetTextureFormats( const UTexture* InTexture, TArray< TArray<FName> >& OutFormats) const
 {
 #if WITH_EDITOR
diff --git a/Engine/Source/Developer/Android/AndroidTargetPlatform/Private/AndroidTargetPlatform.h b/Engine/Source/Developer/Android/AndroidTargetPlatform/Private/AndroidTargetPlatform.h
--- a/Engine/Source/Developer/Android/AndroidTargetPlatform/Private/AndroidTargetPlatform.h
+++ b/Engine/Source/Developer/Android/AndroidTargetPlatform/Private/AndroidTargetPlatform.h
@@ -24,6 +24,7 @@
 #if WITH_ENGINE
 #include "Internationalization/Text.h"
 #include "StaticMeshResources.h"
+#include "SkeletalMeshDefaultLODStreamingSettings.h"
 #endif // WITH_ENGINE
 
 #define LOCTEXT_NAMESPACE "FAndroidTargetPlatform" 
@@ -196,6 +197,8 @@ public:
 	
 	virtual const class FStaticMeshLODSettings& GetStaticMeshLODSettings() const override;
 
+	virtual const FSkeletalMeshDefaultLODStreamingSettings& GetSkeletalMeshDefaultLODStreamingSettings() const override;
+
 	virtual void GetTextureFormats( const UTexture* InTexture, TArray< TArray<FName> >& OutFormats) const override;
 
 	virtual void GetAllTextureFormats(TArray<FName>& OutFormats) const override;
@@ -308,6 +311,8 @@ protected:
 	// Holds the static mesh LOD settings.
 	FStaticMeshLODSettings StaticMeshLODSettings;
 
+	FSkeletalMeshDefaultLODStreamingSettings SkeletalMeshDefaultLODStreamingSettings;
+
 	ITargetDevicePtr DefaultDevice;
 #endif //WITH_ENGINE
 
diff --git a/Engine/Source/Developer/IOS/IOSTargetPlatform/Private/IOSTargetPlatform.cpp b/Engine/Source/Developer/IOS/IOSTargetPlatform/Private/IOSTargetPlatform.cpp
--- a/Engine/Source/Developer/IOS/IOSTargetPlatform/Private/IOSTargetPlatform.cpp
+++ b/Engine/Source/Developer/IOS/IOSTargetPlatform/Private/IOSTargetPlatform.cpp
@@ -37,6 +37,7 @@ FIOSTargetPlatform::FIOSTargetPlatform(bool bInIsTVOS, bool bInIsClientOnly)
 	FConfigCacheIni::LoadLocalIniFile(EngineSettings, TEXT("Engine"), true, *PlatformName());
 	TextureLODSettings = nullptr; // TextureLODSettings are registered by the device profile.
 	StaticMeshLODSettings.Initialize(EngineSettings);
+	SkeletalMeshDefaultLODStreamingSettings.Initialize(EngineSettings);
 #endif // #if WITH_ENGINE
 
 	// Initialize Ticker for device discovery
diff --git a/Engine/Source/Developer/IOS/IOSTargetPlatform/Private/IOSTargetPlatform.h b/Engine/Source/Developer/IOS/IOSTargetPlatform/Private/IOSTargetPlatform.h
--- a/Engine/Source/Developer/IOS/IOSTargetPlatform/Private/IOSTargetPlatform.h
+++ b/Engine/Source/Developer/IOS/IOSTargetPlatform/Private/IOSTargetPlatform.h
@@ -20,6 +20,7 @@
 #if WITH_ENGINE
 #include "AudioCompressionSettings.h"
 #include "StaticMeshResources.h"
+#include "SkeletalMeshDefaultLODStreamingSettings.h"
 #endif // WITH_ENGINE
 
 /**
@@ -115,6 +116,11 @@ public:
 		return StaticMeshLODSettings;
 	}
 
+	virtual const FSkeletalMeshDefaultLODStreamingSettings& GetSkeletalMeshDefaultLODStreamingSettings() const override
+	{
+		return SkeletalMeshDefaultLODStreamingSettings;
+	}
+
 	virtual void GetTextureFormats( const UTexture* Texture, TArray< TArray<FName> >& OutFormats) const override;
 
 	virtual void GetAllTextureFormats( TArray<FName>& OutFormats) const override;
@@ -212,6 +218,8 @@ private:
 
 	// Holds the static mesh LOD settings.
 	FStaticMeshLODSettings StaticMeshLODSettings;
+
+	FSkeletalMeshDefaultLODStreamingSettings SkeletalMeshDefaultLODStreamingSettings;
 #endif // WITH_ENGINE
 
     // holds usb device helper
diff --git a/Engine/Source/Developer/Linux/LinuxTargetPlatform/Private/LinuxTargetPlatform.h b/Engine/Source/Developer/Linux/LinuxTargetPlatform/Private/LinuxTargetPlatform.h
--- a/Engine/Source/Developer/Linux/LinuxTargetPlatform/Private/LinuxTargetPlatform.h
+++ b/Engine/Source/Developer/Linux/LinuxTargetPlatform/Private/LinuxTargetPlatform.h
@@ -18,6 +18,7 @@
 #include "AudioCompressionSettings.h"
 #include "Sound/SoundWave.h"
 #include "StaticMeshResources.h"
+#include "SkeletalMeshDefaultLODStreamingSettings.h"
 #endif // WITH_ENGINE
 #include "Interfaces/IProjectManager.h"
 #include "InstalledPlatformInfo.h"
@@ -59,6 +60,7 @@ public:
 		FConfigCacheIni::LoadLocalIniFile(EngineSettings, TEXT("Engine"), true, *this->PlatformName());
 		TextureLODSettings = nullptr;
 		StaticMeshLODSettings.Initialize(EngineSettings);
+		SkeletalMeshDefaultLODStreamingSettings.Initialize(EngineSettings);
 
 		InitDevicesFromConfig();
 
@@ -291,6 +293,10 @@ public:
 		return StaticMeshLODSettings;
 	}
 
+	virtual const FSkeletalMeshDefaultLODStreamingSettings& GetSkeletalMeshDefaultLODStreamingSettings() const override
+	{
+		return SkeletalMeshDefaultLODStreamingSettings;
+	}
 
 	virtual void GetTextureFormats( const UTexture* InTexture, TArray< TArray<FName> >& OutFormats) const override
 	{
@@ -517,6 +523,8 @@ protected:
 	// Holds static mesh LOD settings.
 	FStaticMeshLODSettings StaticMeshLODSettings;
 
+	FSkeletalMeshDefaultLODStreamingSettings SkeletalMeshDefaultLODStreamingSettings;
+
 	// True if the project requires encoded HDR reflection captures
 	bool bRequiresEncodedHDRReflectionCaptures;
 #endif // WITH_ENGINE
diff --git a/Engine/Source/Developer/Mac/MacTargetPlatform/Private/GenericMacTargetPlatform.h b/Engine/Source/Developer/Mac/MacTargetPlatform/Private/GenericMacTargetPlatform.h
--- a/Engine/Source/Developer/Mac/MacTargetPlatform/Private/GenericMacTargetPlatform.h
+++ b/Engine/Source/Developer/Mac/MacTargetPlatform/Private/GenericMacTargetPlatform.h
@@ -18,6 +18,7 @@
 #include "Sound/SoundWave.h"
 #include "TextureResource.h"
 #include "StaticMeshResources.h"
+#include "SkeletalMeshDefaultLODStreamingSettings.h"
 #endif // WITH_ENGINE
 
 #define LOCTEXT_NAMESPACE "TGenericMacTargetPlatform"
@@ -48,6 +49,7 @@ public:
 			FConfigCacheIni::LoadLocalIniFile(EngineSettings, TEXT("Engine"), true, *this->PlatformName());
 			TextureLODSettings = nullptr;
 			StaticMeshLODSettings.Initialize(EngineSettings);
+			SkeletalMeshDefaultLODStreamingSettings.Initialize(EngineSettings);
 		#endif
 	}
 
@@ -154,6 +156,11 @@ return TSuper::SupportsFeature(Feature);
 		return StaticMeshLODSettings;
 	}
 
+	virtual const FSkeletalMeshDefaultLODStreamingSettings& GetSkeletalMeshDefaultLODStreamingSettings() const override
+	{
+		return SkeletalMeshDefaultLODStreamingSettings;
+	}
+
 	virtual void GetTextureFormats( const UTexture* Texture, TArray< TArray<FName> >& OutFormats) const override
 	{
 		if (!IS_DEDICATED_SERVER)
@@ -410,6 +417,7 @@ private:
 	// Holds the static mesh LOD settings.
 	FStaticMeshLODSettings StaticMeshLODSettings;
 
+	FSkeletalMeshDefaultLODStreamingSettings SkeletalMeshDefaultLODStreamingSettings;
 #endif // WITH_ENGINE
 
 private:
diff --git a/Engine/Source/Developer/TargetPlatform/Public/Interfaces/ITargetPlatform.h b/Engine/Source/Developer/TargetPlatform/Public/Interfaces/ITargetPlatform.h
--- a/Engine/Source/Developer/TargetPlatform/Public/Interfaces/ITargetPlatform.h
+++ b/Engine/Source/Developer/TargetPlatform/Public/Interfaces/ITargetPlatform.h
@@ -517,6 +517,13 @@ public:
 	 */
 	virtual const class FStaticMeshLODSettings& GetStaticMeshLODSettings() const = 0;
 
+	/**
+	 * Gets the skeletal mesh default LOD streaming settings used by this platform.
+	 *
+	 * @return A skeletal mesh default LOD streaming settings structure.
+	 */
+	virtual const class FSkeletalMeshDefaultLODStreamingSettings& GetSkeletalMeshDefaultLODStreamingSettings() const = 0;
+
 	/** 
 	 * Gets the name of the mesh builder module that is responsible for 
 	 * building static and skeletal mesh derived data for this platform.
diff --git a/Engine/Source/Developer/Windows/WindowsTargetPlatform/Private/GenericWindowsTargetPlatform.h b/Engine/Source/Developer/Windows/WindowsTargetPlatform/Private/GenericWindowsTargetPlatform.h
--- a/Engine/Source/Developer/Windows/WindowsTargetPlatform/Private/GenericWindowsTargetPlatform.h
+++ b/Engine/Source/Developer/Windows/WindowsTargetPlatform/Private/GenericWindowsTargetPlatform.h
@@ -15,6 +15,7 @@
 	#include "StaticMeshResources.h"
 	#include "RHI.h"
 	#include "AudioCompressionSettings.h"
+	#include "SkeletalMeshDefaultLODStreamingSettings.h"
 #endif // WITH_ENGINE
 
 
@@ -45,7 +46,7 @@ public:
 		FConfigCacheIni::LoadLocalIniFile(EngineSettings, TEXT("Engine"), true, *this->PlatformName());
 		TextureLODSettings = nullptr; // These are registered by the device profile system.
 		StaticMeshLODSettings.Initialize(EngineSettings);
-
+		SkeletalMeshDefaultLODStreamingSettings.Initialize(EngineSettings);
 
 		// Get the Target RHIs for this platform, we do not always want all those that are supported.
 		TArray<FName> TargetedShaderFormats;
@@ -247,6 +248,11 @@ public:
 		return StaticMeshLODSettings;
 	}
 
+	virtual const FSkeletalMeshDefaultLODStreamingSettings& GetSkeletalMeshDefaultLODStreamingSettings() const override
+	{
+		return SkeletalMeshDefaultLODStreamingSettings;
+	}
+
 	virtual void GetTextureFormats( const UTexture* InTexture, TArray< TArray<FName> >& OutFormats) const override
 	{
 		if (!IS_DEDICATED_SERVER)
@@ -505,6 +511,8 @@ private:
 	// Holds static mesh LOD settings.
 	FStaticMeshLODSettings StaticMeshLODSettings;
 
+	FSkeletalMeshDefaultLODStreamingSettings SkeletalMeshDefaultLODStreamingSettings;
+
 	// True if the project supports non-DX11 texture formats.
 	bool bSupportDX11TextureFormats;
 
diff --git a/Engine/Source/Editor/Persona/Private/PersonaMeshDetails.cpp b/Engine/Source/Editor/Persona/Private/PersonaMeshDetails.cpp
--- a/Engine/Source/Editor/Persona/Private/PersonaMeshDetails.cpp
+++ b/Engine/Source/Editor/Persona/Private/PersonaMeshDetails.cpp
@@ -3115,6 +3115,11 @@ void FPersonaMeshDetails::CustomizeLODSettingsCategories(IDetailLayoutBuilder& D
 	DisableBelowMinLodStrippingRow.IsEnabled(TAttribute<bool>::Create(TAttribute<bool>::FGetter::CreateSP(this, &FPersonaMeshDetails::IsLODInfoEditingEnabled, -1)));
 	DetailLayout.HideProperty(DisableBelowMinLodStrippingPropertyHandle);
 
+	TSharedPtr<IPropertyHandle> bOverrideLODStreamingSettingsHandle = DetailLayout.GetProperty(GET_MEMBER_NAME_CHECKED(USkeletalMesh, bOverrideLODStreamingSettings), USkeletalMesh::StaticClass());
+	IDetailPropertyRow& bOverrideLODStreamingSettingsRow = LODSettingsCategory.AddProperty(bOverrideLODStreamingSettingsHandle);
+	bOverrideLODStreamingSettingsRow.IsEnabled(TAttribute<bool>::Create(TAttribute<bool>::FGetter::CreateSP(this, &FPersonaMeshDetails::IsLODInfoEditingEnabled, -1)));
+	DetailLayout.HideProperty(bOverrideLODStreamingSettingsHandle);
+
 	TSharedPtr<IPropertyHandle> bSupportLODStreamingPropertyHandle = DetailLayout.GetProperty(GET_MEMBER_NAME_CHECKED(USkeletalMesh, bSupportLODStreaming), USkeletalMesh::StaticClass());
 	IDetailPropertyRow& bSupportLODStreamingRow = LODSettingsCategory.AddProperty(bSupportLODStreamingPropertyHandle);
 	bSupportLODStreamingRow.IsEnabled(TAttribute<bool>::Create(TAttribute<bool>::FGetter::CreateSP(this, &FPersonaMeshDetails::IsLODInfoEditingEnabled, -1)));
diff --git a/Engine/Source/Editor/UnrealEd/Private/Factories/SkeletalMeshImport.cpp b/Engine/Source/Editor/UnrealEd/Private/Factories/SkeletalMeshImport.cpp
--- a/Engine/Source/Editor/UnrealEd/Private/Factories/SkeletalMeshImport.cpp
+++ b/Engine/Source/Editor/UnrealEd/Private/Factories/SkeletalMeshImport.cpp
@@ -247,6 +247,7 @@ ExistingSkelMeshData* SaveExistingSkelMeshData(USkeletalMesh* ExistingSkelMesh,
 	ExistingMeshDataPtr->UseMaterialNameSlotWorkflow = SkeletalMeshIsUsingMaterialSlotNameWorkflow(ExistingSkelMesh->AssetImportData);
 	ExistingMeshDataPtr->MinLOD = ExistingSkelMesh->MinLod;
 	ExistingMeshDataPtr->DisableBelowMinLodStripping = ExistingSkelMesh->DisableBelowMinLodStripping;
+	ExistingMeshDataPtr->bOverrideLODStreamingSettings = ExistingSkelMesh->bOverrideLODStreamingSettings;
 	ExistingMeshDataPtr->bSupportLODStreaming = ExistingSkelMesh->bSupportLODStreaming;
 	ExistingMeshDataPtr->MaxNumStreamedLODs = ExistingSkelMesh->MaxNumStreamedLODs;
 	ExistingMeshDataPtr->MaxNumOptionalLODs = ExistingSkelMesh->MaxNumOptionalLODs;
@@ -652,6 +653,7 @@ void RestoreExistingSkelMeshData(ExistingSkelMeshData* MeshData, USkeletalMesh*
 	int32 SafeReimportLODIndex = ReimportLODIndex < 0 ? 0 : ReimportLODIndex;
 	SkeletalMesh->MinLod = MeshData->MinLOD;
 	SkeletalMesh->DisableBelowMinLodStripping = MeshData->DisableBelowMinLodStripping;
+	SkeletalMesh->bOverrideLODStreamingSettings = MeshData->bOverrideLODStreamingSettings;
 	SkeletalMesh->bSupportLODStreaming = MeshData->bSupportLODStreaming;
 	SkeletalMesh->MaxNumStreamedLODs = MeshData->MaxNumStreamedLODs;
 	SkeletalMesh->MaxNumOptionalLODs = MeshData->MaxNumOptionalLODs;
diff --git a/Engine/Source/Editor/UnrealEd/Public/SkelImport.h b/Engine/Source/Editor/UnrealEd/Public/SkelImport.h
--- a/Engine/Source/Editor/UnrealEd/Public/SkelImport.h
+++ b/Engine/Source/Editor/UnrealEd/Public/SkelImport.h
@@ -84,6 +84,7 @@ struct ExistingSkelMeshData
 	FSkeletalMeshSamplingInfo				ExistingSamplingInfo;
 	FPerPlatformInt							MinLOD;
 	FPerPlatformBool						DisableBelowMinLodStripping;
+	bool									bOverrideLODStreamingSettings;
 	FPerPlatformBool						bSupportLODStreaming;
 	FPerPlatformInt							MaxNumStreamedLODs;
 	FPerPlatformInt							MaxNumOptionalLODs;
diff --git a/Engine/Source/Runtime/Engine/Classes/Engine/SkeletalMesh.h b/Engine/Source/Runtime/Engine/Classes/Engine/SkeletalMesh.h
--- a/Engine/Source/Runtime/Engine/Classes/Engine/SkeletalMesh.h
+++ b/Engine/Source/Runtime/Engine/Classes/Engine/SkeletalMesh.h
@@ -714,23 +714,37 @@ public:
 	FPerPlatformBool DisableBelowMinLodStripping;
 
 #if WITH_EDITORONLY_DATA
+	/** Whether this skeletal mesh overrides default LOD streaming settings. */
+	UPROPERTY(EditAnywhere, Category=LODSettings)
+	bool bOverrideLODStreamingSettings;
+
 	/** Whether we can stream the LODs of this mesh */
-	UPROPERTY(EditAnywhere, Category=LODSettings, meta=(DisplayName="Stream LODs"))
+	UPROPERTY(EditAnywhere, Category=LODSettings, meta=(DisplayName="Stream LODs", EditCondition="bOverrideLODStreamingSettings"))
 	FPerPlatformBool bSupportLODStreaming;
 
 	/** Maximum number of LODs that can be streamed */
-	UPROPERTY(EditAnywhere, Category=LODSettings)
+	UPROPERTY(EditAnywhere, Category=LODSettings, meta=(EditCondition="bOverrideLODStreamingSettings"))
 	FPerPlatformInt MaxNumStreamedLODs;
 
 	/** Maximum number of LODs below min LOD level that can be saved to optional pak (currently, need to be either 0 or > num of LODs below MinLod) */
-	UPROPERTY(EditAnywhere, Category=LODSettings)
+	UPROPERTY(EditAnywhere, Category=LODSettings, meta=(EditCondition="bOverrideLODStreamingSettings"))
 	FPerPlatformInt MaxNumOptionalLODs;
 
 	UPROPERTY(EditAnywhere, BlueprintReadWrite, AssetRegistrySearchable, BlueprintSetter = SetLODSettings, Category = LODSettings)
 	USkeletalMeshLODSettings* LODSettings;
-
 #endif // WITH_EDITORONLY_DATA
 
+#if WITH_EDITOR
+	/** Get whether this mesh use LOD streaming. Do not use bSupportLODStreaming directly. Call this method instead. */
+	bool GetSupportsLODStreaming(const class ITargetPlatform* TargetPlatform) const;
+
+	/** Get the maximum number of LODs that can be streamed. Do not use MaxNumStreamedLODs directly. Call this method instead. */
+	int32 GetMaxNumStreamedLODs(const class ITargetPlatform* TargetPlatform) const;
+
+	/** Get the maximum number of optional LODs. Do not use MaxNumOptionalLODs directly. Call this method instead. */
+	int32 GetMaxNumOptionalLODs(const class ITargetPlatform* TargetPlatform) const;
+#endif
+
 	UFUNCTION(BlueprintSetter)
 	void SetLODSettings(USkeletalMeshLODSettings* InLODSettings);
 
@@ -1451,14 +1465,8 @@ public:
 	 * Returns total number of LOD
 	 */
 	int32 GetLODNum() const 
-	{ 
-#if WITH_EDITOR
-		if (bSupportLODStreaming.Default || bSupportLODStreaming.PerPlatform.FindKey(true))
-		{
-			check(LODInfo.Num() <= MAX_MESH_LOD_COUNT); 
-		}
-#endif
-		return LODInfo.Num();  
+	{
+		return LODInfo.Num();
 	}
 
 public:
diff --git a/Engine/Source/Runtime/Engine/Classes/Engine/SkeletalMeshLODSettings.h b/Engine/Source/Runtime/Engine/Classes/Engine/SkeletalMeshLODSettings.h
--- a/Engine/Source/Runtime/Engine/Classes/Engine/SkeletalMeshLODSettings.h
+++ b/Engine/Source/Runtime/Engine/Classes/Engine/SkeletalMeshLODSettings.h
@@ -124,16 +124,20 @@ protected:
 	UPROPERTY(globalconfig, EditAnywhere, Category = LODGroups)
 	FPerPlatformBool DisableBelowMinLodStripping;
 
+	/** Whether meshes in this group override default LOD streaming settings. */
+	UPROPERTY(globalconfig, EditAnywhere, Category=LODGroups)
+	bool bOverrideLODStreamingSettings;
+
 	/** Whether meshes in this group stream LODs by default */
-	UPROPERTY(globalconfig, EditAnywhere, Category=LODGroups, meta=(DisplayName="Stream LODs"))
+	UPROPERTY(globalconfig, EditAnywhere, Category=LODGroups, meta=(DisplayName="Stream LODs", EditCondition="bOverrideLODStreamingSettings"))
 	FPerPlatformBool bSupportLODStreaming;
 
 	/** Default maximum number of streamed LODs for meshes in this group */
-	UPROPERTY(globalconfig, EditAnywhere, Category=LODGroups)
+	UPROPERTY(globalconfig, EditAnywhere, Category=LODGroups, meta=(EditCondition="bOverrideLODStreamingSettings"))
 	FPerPlatformInt MaxNumStreamedLODs;
 
 	/** Default maximum number of optional LODs for meshes in this group (currently, need to be either 0 or > num of LODs below MinLod) */
-	UPROPERTY(globalconfig, EditAnywhere, Category=LODGroups)
+	UPROPERTY(globalconfig, EditAnywhere, Category=LODGroups, meta=(EditCondition="bOverrideLODStreamingSettings"))
 	FPerPlatformInt MaxNumOptionalLODs;
 
 	UPROPERTY(globalconfig, EditAnywhere, Category=LODGroups)
diff --git a/Engine/Source/Runtime/Engine/Private/Rendering/SkinWeightVertexBuffer.cpp b/Engine/Source/Runtime/Engine/Private/Rendering/SkinWeightVertexBuffer.cpp
--- a/Engine/Source/Runtime/Engine/Private/Rendering/SkinWeightVertexBuffer.cpp
+++ b/Engine/Source/Runtime/Engine/Private/Rendering/SkinWeightVertexBuffer.cpp
@@ -137,9 +137,9 @@ void FSkinWeightLookupVertexBuffer::InitRHI()
 		// This also avoid setting lots of states on all the members of all the different buffers used by meshes. Follow up: https://jira.it.epicgames.net/browse/UE-69376.
 		bSRV |= GetNeedsCPUAccess();
 
-		if (bSRV && LookupData && VertexBufferRHI)
+		if (bSRV)
 		{
-			SRVValue = RHICreateShaderResourceView(VertexBufferRHI, PixelFormatStride, PixelFormat);
+			SRVValue = RHICreateShaderResourceView(FShaderResourceViewInitializer(LookupData ? VertexBufferRHI : nullptr, PixelFormat));
 		}
 	}
 }
@@ -439,9 +439,9 @@ void FSkinWeightDataVertexBuffer::InitRHI()
 	// This also avoid setting lots of states on all the members of all the different buffers used by meshes. Follow up: https://jira.it.epicgames.net/browse/UE-69376.
 	bSRV |= GetNeedsCPUAccess();
 
-	if (bSRV && WeightData && VertexBufferRHI)
+	if (bSRV)
 	{
-		SRVValue = RHICreateShaderResourceView(VertexBufferRHI, GetPixelFormatStride(), GetPixelFormat());
+		SRVValue = RHICreateShaderResourceView(FShaderResourceViewInitializer(WeightData ? VertexBufferRHI : nullptr, GetPixelFormat()));
 	}
 }
 
diff --git a/Engine/Source/Runtime/Engine/Private/SkeletalMesh.cpp b/Engine/Source/Runtime/Engine/Private/SkeletalMesh.cpp
--- a/Engine/Source/Runtime/Engine/Private/SkeletalMesh.cpp
+++ b/Engine/Source/Runtime/Engine/Private/SkeletalMesh.cpp
@@ -69,6 +69,7 @@
 
 #include "IMeshReductionManagerModule.h"
 #include "SkeletalMeshReductionSettings.h"
+#include "SkeletalMeshDefaultLODStreamingSettings.h"
 
 #endif // #if WITH_EDITOR
 
@@ -4055,6 +4056,53 @@ void USkeletalMesh::ResetLODInfo()
 	LODInfo.Reset();
 }
 
+#if WITH_EDITOR
+bool USkeletalMesh::GetSupportsLODStreaming(const ITargetPlatform* TargetPlatform) const
+{
+	check(TargetPlatform);
+	if (bOverrideLODStreamingSettings)
+	{
+		return bSupportLODStreaming.GetValueForPlatformIdentifiers(
+			TargetPlatform->GetPlatformInfo().PlatformGroupName,
+			TargetPlatform->GetPlatformInfo().VanillaPlatformName);
+	}
+	else
+	{
+		return TargetPlatform->GetSkeletalMeshDefaultLODStreamingSettings().IsLODStreamingEnabled();
+	}
+}
+
+int32 USkeletalMesh::GetMaxNumStreamedLODs(const ITargetPlatform* TargetPlatform) const
+{
+	check(TargetPlatform);
+	if (bOverrideLODStreamingSettings)
+	{
+		return MaxNumStreamedLODs.GetValueForPlatformIdentifiers(
+			TargetPlatform->GetPlatformInfo().PlatformGroupName,
+			TargetPlatform->GetPlatformInfo().VanillaPlatformName);
+	}
+	else
+	{
+		return TargetPlatform->GetSkeletalMeshDefaultLODStreamingSettings().GetMaxNumStreamedLODs();
+	}
+}
+
+int32 USkeletalMesh::GetMaxNumOptionalLODs(const ITargetPlatform* TargetPlatform) const
+{
+	check(TargetPlatform);
+	if (bOverrideLODStreamingSettings)
+	{
+		return MaxNumOptionalLODs.GetValueForPlatformIdentifiers(
+			TargetPlatform->GetPlatformInfo().PlatformGroupName,
+			TargetPlatform->GetPlatformInfo().VanillaPlatformName);
+	}
+	else
+	{
+		return TargetPlatform->GetSkeletalMeshDefaultLODStreamingSettings().GetMaxNumOptionalLODs();
+	}
+}
+#endif
+
 void USkeletalMesh::SetLODSettings(USkeletalMeshLODSettings* InLODSettings)
 {
 #if WITH_EDITORONLY_DATA
diff --git a/Engine/Source/Runtime/Engine/Private/SkeletalMeshLODRenderData.cpp b/Engine/Source/Runtime/Engine/Private/SkeletalMeshLODRenderData.cpp
--- a/Engine/Source/Runtime/Engine/Private/SkeletalMeshLODRenderData.cpp
+++ b/Engine/Source/Runtime/Engine/Private/SkeletalMeshLODRenderData.cpp
@@ -751,14 +751,10 @@ bool FSkeletalMeshLODRenderData::IsLODCookedOut(const ITargetPlatform* TargetPla
 	}
 	check(TargetPlatform);
 
-	const bool bSupportLODStreaming = !SkeletalMesh->NeverStream && SkeletalMesh->bSupportLODStreaming.GetValueForPlatformIdentifiers(
-		TargetPlatform->GetPlatformInfo().PlatformGroupName,
-		TargetPlatform->GetPlatformInfo().VanillaPlatformName);
-
 	static auto* VarMeshStreaming = IConsoleManager::Get().FindConsoleVariable(TEXT("r.MeshStreaming"));
 	const bool bMeshStreamingEnabled = !VarMeshStreaming || VarMeshStreaming->GetInt() != 0;
 	
-	return !bMeshStreamingEnabled || !TargetPlatform->SupportsFeature(ETargetPlatformFeatures::MeshLODStreaming) || !bSupportLODStreaming;
+	return !bMeshStreamingEnabled || !TargetPlatform->SupportsFeature(ETargetPlatformFeatures::MeshLODStreaming) || !SkeletalMesh->GetSupportsLODStreaming(TargetPlatform);
 #else
 	return false;
 #endif
@@ -774,14 +770,10 @@ bool FSkeletalMeshLODRenderData::IsLODInlined(const ITargetPlatform* TargetPlatf
 	}
 	check(TargetPlatform);
 
-	const FName PlatformGroupName = TargetPlatform->GetPlatformInfo().PlatformGroupName;
-	const FName VanillaPlatformName = TargetPlatform->GetPlatformInfo().VanillaPlatformName;
-	const bool bSupportLODStreaming = !SkeletalMesh->NeverStream && SkeletalMesh->bSupportLODStreaming.GetValueForPlatformIdentifiers(PlatformGroupName, VanillaPlatformName);
-
 	static auto* VarMeshStreaming = IConsoleManager::Get().FindConsoleVariable(TEXT("r.MeshStreaming"));
 	const bool bMeshStreamingEnabled = !VarMeshStreaming || VarMeshStreaming->GetInt() != 0;
-
-	if (!bMeshStreamingEnabled || !TargetPlatform->SupportsFeature(ETargetPlatformFeatures::MeshLODStreaming) || !bSupportLODStreaming)
+	
+	if (!bMeshStreamingEnabled || !TargetPlatform->SupportsFeature(ETargetPlatformFeatures::MeshLODStreaming) || !SkeletalMesh->GetSupportsLODStreaming(TargetPlatform))
 	{
 		return true;
 	}
@@ -791,7 +783,7 @@ bool FSkeletalMeshLODRenderData::IsLODInlined(const ITargetPlatform* TargetPlatf
 		return false;
 	}
 
-	const int32 MaxNumStreamedLODs = SkeletalMesh->MaxNumStreamedLODs.GetValueForPlatformIdentifiers(PlatformGroupName, VanillaPlatformName);
+	const int32 MaxNumStreamedLODs = SkeletalMesh->GetMaxNumStreamedLODs(TargetPlatform);
 	const int32 NumLODs = SkeletalMesh->GetLODNum();
 	const int32 NumStreamedLODs = FMath::Min(MaxNumStreamedLODs, NumLODs - 1);
 	const int32 InlinedLODStartIdx = NumStreamedLODs;
@@ -805,9 +797,7 @@ int32 FSkeletalMeshLODRenderData::GetNumOptionalLODsAllowed(const ITargetPlatfor
 {
 #if WITH_EDITOR
 	check(TargetPlatform && SkeletalMesh);
-	const FName PlatformGroupName = TargetPlatform->GetPlatformInfo().PlatformGroupName;
-	const FName VanillaPlatformName = TargetPlatform->GetPlatformInfo().VanillaPlatformName;
-	return  SkeletalMesh->MaxNumOptionalLODs.GetValueForPlatformIdentifiers(PlatformGroupName, VanillaPlatformName);
+	return SkeletalMesh->GetMaxNumOptionalLODs(TargetPlatform);
 #else
 	return 0;
 #endif
diff --git a/Engine/Source/Runtime/Engine/Private/SkeletalMeshLODSettings.cpp b/Engine/Source/Runtime/Engine/Private/SkeletalMeshLODSettings.cpp
--- a/Engine/Source/Runtime/Engine/Private/SkeletalMeshLODSettings.cpp
+++ b/Engine/Source/Runtime/Engine/Private/SkeletalMeshLODSettings.cpp
@@ -7,6 +7,8 @@
 #include "Animation/AnimSequence.h"
 #include "UObject/FortniteMainBranchObjectVersion.h"
 #include "UObject/EditorObjectVersion.h"
+#include "SkeletalMeshDefaultLODStreamingSettings.h"
+#include "Misc/ConfigCacheIni.h"
 
 DEFINE_LOG_CATEGORY_STATIC(LogSkeletalMeshLODSettings, Warning, All)
 
@@ -168,6 +170,7 @@ int32 USkeletalMeshLODSettings::SetLODSettingsToMesh(USkeletalMesh* InMesh) cons
 		InMesh->MinLod = MinLod;
 		InMesh->DisableBelowMinLodStripping = DisableBelowMinLodStripping;
 #if WITH_EDITORONLY_DATA
+		InMesh->bOverrideLODStreamingSettings = bOverrideLODStreamingSettings;
 		InMesh->bSupportLODStreaming = bSupportLODStreaming;
 		InMesh->MaxNumStreamedLODs = MaxNumStreamedLODs;
 		InMesh->MaxNumOptionalLODs = MaxNumOptionalLODs;
@@ -193,6 +196,7 @@ int32 USkeletalMeshLODSettings::SetLODSettingsFromMesh(USkeletalMesh* InMesh)
 		MinLod = InMesh->MinLod;
 		DisableBelowMinLodStripping = InMesh->DisableBelowMinLodStripping;
 #if WITH_EDITORONLY_DATA
+		bOverrideLODStreamingSettings = InMesh->bOverrideLODStreamingSettings;
 		bSupportLODStreaming = InMesh->bSupportLODStreaming;
 		MaxNumStreamedLODs = InMesh->MaxNumStreamedLODs;
 		MaxNumOptionalLODs = InMesh->MaxNumOptionalLODs;
@@ -292,3 +296,31 @@ const float FSkeletalMeshLODGroupSettings::GetScreenSize() const
 	return ScreenSize.Default;
 }
 
+/////////////////////////////////////////////////////////////
+// FSkeletalMeshDefaultLODStreamingSettings 
+/////////////////////////////////////////////////////////////
+FSkeletalMeshDefaultLODStreamingSettings::FSkeletalMeshDefaultLODStreamingSettings()
+	: bSupportLODStreaming(false)
+	, MaxNumStreamedLODs(0)
+	, MaxNumOptionalLODs(0)
+{}
+
+void FSkeletalMeshDefaultLODStreamingSettings::Initialize(const FConfigFile& IniFile)
+{
+	const TCHAR* IniSection = TEXT("SkeletalMeshDefaultLODStreamingSettings");
+	
+	if (!IniFile.GetBool(IniSection, TEXT("bSupportLODStreaming"), bSupportLODStreaming))
+	{
+		bSupportLODStreaming = false;
+	}
+
+	if (!IniFile.GetInt(IniSection, TEXT("MaxNumStreamedLODs"), MaxNumStreamedLODs))
+	{
+		MaxNumStreamedLODs = 0;
+	}
+
+	if (!IniFile.GetInt(IniSection, TEXT("MaxNumOptionalLODs"), MaxNumOptionalLODs))
+	{
+		MaxNumOptionalLODs = 0;
+	}
+}
diff --git a/Engine/Source/Runtime/Engine/Private/SkeletalMeshRenderData.cpp b/Engine/Source/Runtime/Engine/Private/SkeletalMeshRenderData.cpp
--- a/Engine/Source/Runtime/Engine/Private/SkeletalMeshRenderData.cpp
+++ b/Engine/Source/Runtime/Engine/Private/SkeletalMeshRenderData.cpp
@@ -101,17 +101,14 @@ static FString BuildSkeletalMeshDerivedDataKey(const ITargetPlatform* TargetPlat
 	KeySuffix += SkelMesh->bHasVertexColors ? "1" : "0";
 	KeySuffix += SkelMesh->VertexColorGuid.ToString(EGuidFormats::Digits);
 
-	const FName PlatformGroupName = TargetPlatform->GetPlatformInfo().PlatformGroupName;
-	const FName VanillaPlatformName = TargetPlatform->GetPlatformInfo().VanillaPlatformName;
-
 	static auto* VarMeshStreaming = IConsoleManager::Get().FindConsoleVariable(TEXT("r.MeshStreaming"));
 	const bool bMeshStreamingEnabled = !VarMeshStreaming || VarMeshStreaming->GetInt() != 0;
-
-	const bool bSupportLODStreaming = !SkelMesh->NeverStream && SkelMesh->bSupportLODStreaming.GetValueForPlatformIdentifiers(PlatformGroupName, VanillaPlatformName);
+	const bool bSupportLODStreaming = !SkelMesh->NeverStream && SkelMesh->GetSupportsLODStreaming(TargetPlatform);
+	
 	if (bMeshStreamingEnabled && TargetPlatform->SupportsFeature(ETargetPlatformFeatures::MeshLODStreaming) && bSupportLODStreaming)
 	{
-		const int32 MaxNumStreamedLODs = SkelMesh->MaxNumStreamedLODs.GetValueForPlatformIdentifiers(PlatformGroupName, VanillaPlatformName);
-		const int32 MaxNumOptionalLODs = SkelMesh->MaxNumOptionalLODs.GetValueForPlatformIdentifiers(PlatformGroupName, VanillaPlatformName);
+		const int32 MaxNumStreamedLODs = SkelMesh->GetMaxNumStreamedLODs(TargetPlatform);
+		const int32 MaxNumOptionalLODs = SkelMesh->GetMaxNumOptionalLODs(TargetPlatform);
 		KeySuffix += *FString::Printf(TEXT("1%08x%08x"), MaxNumStreamedLODs, MaxNumOptionalLODs);
 	}
 	else
diff --git a/Engine/Source/Runtime/Engine/Public/SkeletalMeshDefaultLODStreamingSettings.h b/Engine/Source/Runtime/Engine/Public/SkeletalMeshDefaultLODStreamingSettings.h
new file mode 100644
--- /dev/null
+++ b/Engine/Source/Runtime/Engine/Public/SkeletalMeshDefaultLODStreamingSettings.h
@@ -0,0 +1,37 @@
+// Copyright Epic Games, Inc. All Rights Reserved.
+
+#pragma once
+
+#include "CoreMinimal.h"
+
+class FSkeletalMeshDefaultLODStreamingSettings
+{
+public:
+	ENGINE_API FSkeletalMeshDefaultLODStreamingSettings();
+
+	/**
+	 * Initializes default LOD streaming settings by reading them from the passed in config file section.
+	 * @param IniFile Preloaded ini file object to load from
+	 */
+	ENGINE_API void Initialize(const class FConfigFile& IniFile);
+
+	bool IsLODStreamingEnabled() const
+	{
+		return bSupportLODStreaming;
+	}
+
+	int32 GetMaxNumStreamedLODs() const
+	{
+		return MaxNumStreamedLODs;
+	}
+
+	int32 GetMaxNumOptionalLODs() const
+	{
+		return MaxNumOptionalLODs;
+	}
+
+private:
+	bool bSupportLODStreaming;
+	int32 MaxNumStreamedLODs;
+	int32 MaxNumOptionalLODs;
+};
