Submitting for Ethan: fix for ADPCM seeking

diff --git a/Engine/Source/Runtime/Engine/Private/ADPCMAudioInfo.cpp b/Engine/Source/Runtime/Engine/Private/ADPCMAudioInfo.cpp
--- a/Engine/Source/Runtime/Engine/Private/ADPCMAudioInfo.cpp
+++ b/Engine/Source/Runtime/Engine/Private/ADPCMAudioInfo.cpp
@@ -30,6 +30,13 @@ FAutoConsoleVariableRef CVarADPCMReadFailiureTimeout(
 	TEXT("Sets the number of ADPCM decode attempts we'll try before stopping the sound wave altogether.\n"),
 	ECVF_Default);
 
+static int32 ADPCMDisableSeekForwardOnChunkMissesCVar = 0;
+FAutoConsoleVariableRef CVarADPCMDisableSeekForwardOnChunkMisses(
+	TEXT("au.adpcm.DisableSeekForwardOnReadMisses"),
+	ADPCMDisableSeekForwardOnChunkMissesCVar,
+	TEXT("When there is a seek pending and this CVar is set to 0, we will scan forward in the file.\n"),
+	ECVF_Default);
+
 #define WAVE_FORMAT_LPCM  1
 #define WAVE_FORMAT_ADPCM 2
 
@@ -628,6 +635,25 @@ bool FADPCMAudioInfo::StreamCompressedData(uint8* Destination, bool bLooping, ui
 							// possible a game thread stall can cause this.
 							UE_LOG(LogAudio, Verbose, TEXT("Missed Deadline chunk %d"), CurrentChunkIndex);
 						}
+						else if (!ADPCMDisableSeekForwardOnChunkMissesCVar)
+						{
+							// If we have a seek pending, seek forward in the chunk offset to where we would have been.
+							const uint32 NumBlocksForCallback = BufferSize / (UncompressedBlockSize * NumChannels);
+							const uint32 NumCompressedBytesForCallback = NumBlocksForCallback * CompressedBlockSize * NumChannels;
+							
+							CurrentChunkBufferOffset += NumCompressedBytesForCallback;
+
+							const uint32 SizeOfCurrentChunk = StreamingSoundWave->GetSizeOfChunk(CurrentChunkIndex);
+							if (CurrentChunkBufferOffset >= SizeOfCurrentChunk)
+							{
+								++CurrentChunkIndex;
+								CurrentChunkBufferOffset %= SizeOfCurrentChunk;
+								if (!ensureMsgf(CurrentChunkBufferOffset % CompressedBlockSize == 0, TEXT("Error: seeked partway into an ADPCM block. Please check the above code, as well as FAudioFormatADPCM::SplitDataForStreaming for accuracy.")))
+								{
+									CurrentChunkBufferOffset = AlignDown(CurrentChunkBufferOffset, CompressedBlockSize);
+								}
+							}
+						}
 
 						// zero out remaining data and bail
 						FMemory::Memset(OutData, 0, BufferSize);
