Make sure skeletalmesh reduction restore all data so LOD DDC key do not change during reduction.

diff --git a/Engine/Plugins/Experimental/SkeletalReduction/Source/Private/SkeletalMeshReductionPlugin.cpp b/Engine/Plugins/Experimental/SkeletalReduction/Source/Private/SkeletalMeshReductionPlugin.cpp
--- a/Engine/Plugins/Experimental/SkeletalReduction/Source/Private/SkeletalMeshReductionPlugin.cpp
+++ b/Engine/Plugins/Experimental/SkeletalReduction/Source/Private/SkeletalMeshReductionPlugin.cpp
@@ -1890,6 +1890,7 @@ void FQuadricSkeletalMeshReduction::ReduceSkeletalMesh(USkeletalMesh& SkeletalMe
 	};
 
 	FString BackupLodModelBuildStringID = TEXT("");
+	FString BackupRawSkeletalMeshBulkDataID = TEXT("");
 
 	// Unbind any existing clothing assets before we reimport the geometry
 	TArray<ClothingAssetUtils::FClothingAssetMeshBinding> ClothingBindings;
@@ -1899,6 +1900,7 @@ void FQuadricSkeletalMeshReduction::ReduceSkeletalMesh(USkeletalMesh& SkeletalMe
 		FLODUtilities::UnbindClothingAndBackup(&SkeletalMesh, ClothingBindings, LODIndex);
 		//We have to put back the exact UserSectionsData to not invalidate the DDC key
 		BackupLodModelBuildStringID = SkeletalMeshResource.LODModels[LODIndex].BuildStringID;
+		BackupRawSkeletalMeshBulkDataID = SkeletalMeshResource.LODModels[LODIndex].RawSkeletalMeshBulkDataID;
 	}
 
 	bool bReducingSourceModel = false;
@@ -2093,6 +2095,7 @@ void FQuadricSkeletalMeshReduction::ReduceSkeletalMesh(USkeletalMesh& SkeletalMe
 				ImportedModelLOD.SyncronizeUserSectionsDataArray();
 				//If its an existing LOD put back the buildStringID
 				ImportedModelLOD.BuildStringID = BackupLodModelBuildStringID;
+				ImportedModelLOD.RawSkeletalMeshBulkDataID = BackupRawSkeletalMeshBulkDataID;
 			}
 		}
 	}
diff --git a/Engine/Source/Runtime/Engine/Classes/Engine/SkeletalMesh.h b/Engine/Source/Runtime/Engine/Classes/Engine/SkeletalMesh.h
--- a/Engine/Source/Runtime/Engine/Classes/Engine/SkeletalMesh.h
+++ b/Engine/Source/Runtime/Engine/Classes/Engine/SkeletalMesh.h
@@ -1068,6 +1068,11 @@ private:
 	 */
 	void PostLoadValidateClothingData();
 
+	/*
+	 * This function will enforce the user section data is coherent with the sections.
+	 */
+	void PostLoadValidateUserSectionData();
+
 public:
 	//We want to avoid calling post edit change multiple time during import and build process.
 
diff --git a/Engine/Source/Runtime/Engine/Private/SkeletalMesh.cpp b/Engine/Source/Runtime/Engine/Private/SkeletalMesh.cpp
--- a/Engine/Source/Runtime/Engine/Private/SkeletalMesh.cpp
+++ b/Engine/Source/Runtime/Engine/Private/SkeletalMesh.cpp
@@ -2349,6 +2349,79 @@ void USkeletalMesh::PostLoadValidateClothingData()
 	}
 }
 
+void USkeletalMesh::PostLoadValidateUserSectionData()
+{
+	for (int32 LodIndex = 0; LodIndex < GetLODNum(); LodIndex++)
+	{
+		const FSkeletalMeshLODInfo* LODInfoPtr = GetLODInfo(LodIndex);
+		if (!LODInfoPtr || !LODInfoPtr->bHasBeenSimplified)
+		{
+			//We validate only generated LOD from a base LOD
+			continue;
+		}
+
+		FSkeletalMeshLODModel& ThisLODModel = ImportedModel->LODModels[LodIndex];
+		const int32 SectionNum = ThisLODModel.Sections.Num();
+		//See if more then one section use the same UserSectionData
+		bool bLODHaveSectionIssue = false;
+		TBitArray<> AvailableUserSectionData;
+		AvailableUserSectionData.Init(true, ThisLODModel.UserSectionsData.Num());
+		for (int32 SectionIndex = 0; SectionIndex < SectionNum; ++SectionIndex)
+		{
+			FSkelMeshSection& Section = ThisLODModel.Sections[SectionIndex];
+			if (Section.ChunkedParentSectionIndex != INDEX_NONE)
+			{
+				continue;
+			}
+			if(!AvailableUserSectionData.IsValidIndex(Section.OriginalDataSectionIndex) || !AvailableUserSectionData[Section.OriginalDataSectionIndex])
+			{
+				bLODHaveSectionIssue = true;
+				break;
+			}
+			AvailableUserSectionData[Section.OriginalDataSectionIndex] = false;
+		}
+		if(!bLODHaveSectionIssue)
+		{
+			//Everything is good nothing to fix
+			continue;
+		}
+
+		//Force the source UserSectionData, then restore the UserSectionData value each section was using
+		//We use the source section user data entry in case we do not have any override
+		const FSkeletalMeshLODModel& BaseLODModel = ImportedModel->LODModels[LODInfoPtr->ReductionSettings.BaseLOD];
+		TMap<int32, FSkelMeshSourceSectionUserData> NewUserSectionsData;
+
+		int32 CurrentOriginalSectionIndex = 0;
+		for (int32 SectionIndex = 0; SectionIndex < SectionNum; ++SectionIndex)
+		{
+			FSkelMeshSection& Section = ThisLODModel.Sections[SectionIndex];
+			if (Section.ChunkedParentSectionIndex != INDEX_NONE)
+			{
+				//We do not restore user section data for chunked section, the parent has already fix it
+				Section.OriginalDataSectionIndex = CurrentOriginalSectionIndex;
+				continue;
+			}
+
+			FSkelMeshSourceSectionUserData& SectionUserData = NewUserSectionsData.FindOrAdd(CurrentOriginalSectionIndex);
+			if(const FSkelMeshSourceSectionUserData* BackupSectionUserData = ThisLODModel.UserSectionsData.Find(Section.OriginalDataSectionIndex))
+			{
+				SectionUserData = *BackupSectionUserData;
+			}
+			else if(const FSkelMeshSourceSectionUserData* BaseSectionUserData = BaseLODModel.UserSectionsData.Find(CurrentOriginalSectionIndex))
+			{
+				SectionUserData = *BaseSectionUserData;
+			}
+
+			Section.OriginalDataSectionIndex = CurrentOriginalSectionIndex;
+			//Parent (non chunked) section must increment the index
+			CurrentOriginalSectionIndex++;
+		}
+		ThisLODModel.UserSectionsData = NewUserSectionsData;
+
+		UE_ASSET_LOG(LogSkeletalMesh, Display, this, TEXT("Fix some section data of this asset for lod %d. Verify all sections of this mesh are ok and save the asset to fix this issue."), LodIndex);
+	}
+}
+
 #endif // WITH_EDITOR
 
 void USkeletalMesh::PostLoad()
@@ -2507,6 +2580,7 @@ void USkeletalMesh::PostLoad()
 		//We must unbind incorrect cloth data before computing the ddc key (CacheDerivedData)
 		//Incorrect cloth data will not be rebind after the build, which will change the DDC key, this situation lead to a corrupted DDC cache
 		PostLoadValidateClothingData();
+		PostLoadValidateUserSectionData();
 
 		if (GetResourceForRendering() == nullptr)
 		{
diff --git a/Engine/Source/Runtime/Engine/Private/SkeletalMeshRenderData.cpp b/Engine/Source/Runtime/Engine/Private/SkeletalMeshRenderData.cpp
--- a/Engine/Source/Runtime/Engine/Private/SkeletalMeshRenderData.cpp
+++ b/Engine/Source/Runtime/Engine/Private/SkeletalMeshRenderData.cpp
@@ -339,7 +339,13 @@ void FSkeletalMeshRenderData::Cache(const ITargetPlatform* TargetPlatform, USkel
 			//Recompute the derived data key in case there was some data correction during the build process, this make sure the DDC key is always representing the correct build result.
 			//There should never be correction of the data during the build, the data has to be corrected in the post load before calling this function.
 			FString BuiltDerivedDataKey = BuildSkeletalMeshDerivedDataKey(TargetPlatform, Owner);
-			//ensureMsgf(BuiltDerivedDataKey == DerivedDataKey, TEXT("Skeletal mesh [%s] build has change the source data. The derived data key before and after the build is different."), *Owner->GetPathName());
+			if(BuiltDerivedDataKey != DerivedDataKey)
+			{
+				//If we are in this case we should resave the asset so the source data will be the same and we can use this DDC. Reduction can change the number of sections and the user section data is in the DDC key.
+				//So if we change the reduction algorithm, its possible we fall in this situation.
+				//We save the real data key which force the asset to always rebuild when the editor is loading it until the user save it
+				UE_LOG(LogSkeletalMesh, Log, TEXT("Skeletal mesh [%s]: The derived data key is different after the build. Resave the asset to avoid rebuilding it everytime the editor load it."), *Owner->GetPathName());
+			}
 
 			//Store the data using the built key to avoid DDC corruption
 			GetDerivedDataCacheRef().Put(*BuiltDerivedDataKey, DerivedData, Owner->GetPathName());
