Submitting for Ethan: Fix for ADPCM crash when seeking to the end of a file

diff --git a/Engine/Source/Runtime/Engine/Private/ADPCMAudioInfo.cpp b/Engine/Source/Runtime/Engine/Private/ADPCMAudioInfo.cpp
--- a/Engine/Source/Runtime/Engine/Private/ADPCMAudioInfo.cpp
+++ b/Engine/Source/Runtime/Engine/Private/ADPCMAudioInfo.cpp
@@ -648,10 +648,28 @@ bool FADPCMAudioInfo::StreamCompressedData(uint8* Destination, bool bLooping, ui
 							{
 								++CurrentChunkIndex;
 								CurrentChunkBufferOffset %= SizeOfCurrentChunk;
+
 								if (!ensureMsgf(CurrentChunkBufferOffset % CompressedBlockSize == 0, TEXT("Error: seeked partway into an ADPCM block. Please check the above code, as well as FAudioFormatADPCM::SplitDataForStreaming for accuracy.")))
 								{
 									CurrentChunkBufferOffset = AlignDown(CurrentChunkBufferOffset, CompressedBlockSize);
 								}
+
+								if (CurrentChunkIndex >= StreamingSoundWave->GetNumChunks())
+								{
+									// If this is the case, we've seeked to the end of the file.
+									ReachedEndOfSamples = true;
+									CurrentUncompressedBlockSampleIndex = 0;
+									CurrentChunkIndex = FirstChunkSampleDataIndex;
+									CurrentChunkBufferOffset = 0;
+									TotalSamplesStreamed = 0;
+									CurCompressedChunkData = nullptr;
+									if (!bLooping)
+									{
+										// Set the remaining buffer to 0
+										FMemory::Memset(OutData, 0, BufferSize);
+										return true;
+									}
+								}
 							}
 						}
 
