There was a fix in cloth that allow to retrieve the deprecated data when binding the cloth even if the cloth is load after the skeletalmesh. This allow us to remove the temporary validate cloth binding in the skeletalmesh postload.
#rnx

diff --git a/Engine/Source/Runtime/Engine/Classes/Engine/SkeletalMesh.h b/Engine/Source/Runtime/Engine/Classes/Engine/SkeletalMesh.h
--- a/Engine/Source/Runtime/Engine/Classes/Engine/SkeletalMesh.h
+++ b/Engine/Source/Runtime/Engine/Classes/Engine/SkeletalMesh.h
@@ -1063,10 +1063,6 @@ private:
 	//This function should be call only in the PostLoad
 	void CreateUserSectionsDataForLegacyAssets();
 
-	/*
-	 * This function unbind any incorrect cloth asset that is bind.
-	 */
-	void PostLoadValidateClothingData();
 
 	/*
 	 * This function will enforce the user section data is coherent with the sections.
diff --git a/Engine/Source/Runtime/Engine/Private/SkeletalMesh.cpp b/Engine/Source/Runtime/Engine/Private/SkeletalMesh.cpp
--- a/Engine/Source/Runtime/Engine/Private/SkeletalMesh.cpp
+++ b/Engine/Source/Runtime/Engine/Private/SkeletalMesh.cpp
@@ -2319,36 +2319,6 @@ void USkeletalMesh::CreateUserSectionsDataForLegacyAssets()
 	}
 }
 
-void USkeletalMesh::PostLoadValidateClothingData()
-{
-	//Make sure PostEditChange is not call when we validate the clothing data during the PostLoad
-	FScopedSkeletalMeshPostEditChange ScopedPostEditChange(this, false, false);
-	for (int32 LodIndex = 0; LodIndex < GetLODNum(); LodIndex++)
-	{
-		FSkeletalMeshLODModel& ThisLODModel = ImportedModel->LODModels[LodIndex];
-		if (IsLODImportedDataEmpty(LodIndex) || !IsLODImportedDataBuildAvailable(LodIndex))
-		{
-			//Invalid clothing asset will not be unbind in case we do not build the asset
-			continue;
-		}
-
-		int32 SectionNum = ThisLODModel.Sections.Num();
-		for (int32 SectionIndex = 0; SectionIndex < SectionNum; ++SectionIndex)
-		{
-			FSkelMeshSection& Section = ThisLODModel.Sections[SectionIndex];
-			if (Section.HasClothingData())
-			{
-				UClothingAssetBase* ClothingAsset = GetClothingAsset(Section.ClothingData.AssetGuid);
-				if(!ClothingAsset->IsValidLod(LodIndex))
-				{
-					//Unbind invalid cloth asset
-					ClothingAsset->UnbindFromSkeletalMesh(this, LodIndex);
-				}
-			}
-		}
-	}
-}
-
 void USkeletalMesh::PostLoadValidateUserSectionData()
 {
 	for (int32 LodIndex = 0; LodIndex < GetLODNum(); LodIndex++)
@@ -2577,9 +2547,6 @@ void USkeletalMesh::PostLoad()
 			CreateUserSectionsDataForLegacyAssets();
 		}
 
-		//We must unbind incorrect cloth data before computing the ddc key (CacheDerivedData)
-		//Incorrect cloth data will not be rebind after the build, which will change the DDC key, this situation lead to a corrupted DDC cache
-		PostLoadValidateClothingData();
 		PostLoadValidateUserSectionData();
 
 		if (GetResourceForRendering() == nullptr)
