Submit for Ethan: Don't do sync behavior if we're waiting for a chunk of audio to load during playback

diff --git a/Engine/Source/Runtime/Engine/Private/ADPCMAudioInfo.cpp b/Engine/Source/Runtime/Engine/Private/ADPCMAudioInfo.cpp
--- a/Engine/Source/Runtime/Engine/Private/ADPCMAudioInfo.cpp
+++ b/Engine/Source/Runtime/Engine/Private/ADPCMAudioInfo.cpp
@@ -622,6 +622,23 @@ bool FADPCMAudioInfo::StreamCompressedData(uint8* Destination, bool bLooping, ui
 						++CurrentChunkIndex;
 					}
 
+					if (CurrentChunkIndex >= StreamingSoundWave->GetNumChunks())
+					{
+						// If this is the case, we've seeked to the end of the file.
+						ReachedEndOfSamples = true;
+						CurrentUncompressedBlockSampleIndex = 0;
+						CurrentChunkIndex = FirstChunkSampleDataIndex;
+						CurrentChunkBufferOffset = 0;
+						TotalSamplesStreamed = 0;
+						CurCompressedChunkData = nullptr;
+						if (!bLooping)
+						{
+							// Set the remaining buffer to 0
+							FMemory::Memset(OutData, 0, BufferSize);
+							return true;
+						}
+					}
+
 					// Request the next chunk of data from the streaming engine
 					CurCompressedChunkData = GetLoadedChunk(StreamingSoundWave, CurrentChunkIndex, CurrentChunkDataSize);
 
@@ -635,9 +652,12 @@ bool FADPCMAudioInfo::StreamCompressedData(uint8* Destination, bool bLooping, ui
 							// possible a game thread stall can cause this.
 							UE_LOG(LogAudio, Verbose, TEXT("Missed Deadline chunk %d"), CurrentChunkIndex);
 						}
-						else if (!ADPCMDisableSeekForwardOnChunkMissesCVar)
+
+						// If we have a seek pending, or we're in the middle of playing back the file, seek forward in the chunk offset to where we would have been.
+						const bool bShouldSeekForwardOnMissedChunk = !ADPCMDisableSeekForwardOnChunkMissesCVar && (bSeekPending || CurrentChunkIndex > FirstChunkSampleDataIndex);
+
+						if (bShouldSeekForwardOnMissedChunk)
 						{
-							// If we have a seek pending, seek forward in the chunk offset to where we would have been.
 							const uint32 NumBlocksForCallback = BufferSize / (UncompressedBlockSize * NumChannels);
 							const uint32 NumCompressedBytesForCallback = NumBlocksForCallback * CompressedBlockSize * NumChannels;
 							
