Fix for client crash on launch due to race between AudioDeviceManager init and call to FAudioDeviceManager::GetAudioDeviceRaw in the audio render thread in FSubmixEffectDynamicsProcessor::OnPresetChanged

diff --git a/Engine/Source/Runtime/Engine/Private/AudioDeviceManager.cpp b/Engine/Source/Runtime/Engine/Private/AudioDeviceManager.cpp
--- a/Engine/Source/Runtime/Engine/Private/AudioDeviceManager.cpp
+++ b/Engine/Source/Runtime/Engine/Private/AudioDeviceManager.cpp
@@ -621,6 +621,7 @@ uint32 FAudioDeviceManager::CreateUniqueStackWalkID()
 
 FAudioDeviceHandle FAudioDeviceManager::GetAudioDevice(Audio::FDeviceId Handle)
 {
+	FScopeLock ScopeLock(&DeviceMapCriticalSection);
 	FAudioDeviceContainer* Container = Devices.Find(Handle);
 	if (Container)
 	{
@@ -635,6 +636,7 @@ FAudioDeviceHandle FAudioDeviceManager::GetAudioDevice(Audio::FDeviceId Handle)
 
 FAudioDevice* FAudioDeviceManager::GetAudioDeviceRaw(Audio::FDeviceId Handle)
 {
+	FScopeLock ScopeLock(&DeviceMapCriticalSection);
 	if (!IsValidAudioDevice(Handle))
 	{
 		return nullptr;
@@ -743,10 +745,12 @@ void FAudioDeviceManager::UpdateActiveAudioDevices(bool bGameTicking)
 	}
 
 
-	for (auto& DeviceContainer : Devices)
-	{
-		DeviceContainer.Value.Device->Update(bGameTicking);
-	}
+	IterateOverAllDevices(
+		[&bGameTicking](Audio::FDeviceId, FAudioDevice* InDevice)
+		{
+			InDevice->Update(bGameTicking);
+		}
+	);
 
 	if (GCVarEnableAudioThreadWait)
 	{
@@ -777,102 +781,112 @@ void FAudioDeviceManager::IterateOverAllDevices(TFunction<void(Audio::FDeviceId,
 
 void FAudioDeviceManager::AddReferencedObjects(FReferenceCollector& Collector)
 {
-	FScopeLock ScopeLock(&DeviceMapCriticalSection);
-	for (auto& DeviceContainer : Devices)
-	{
-		check(DeviceContainer.Value.Device);
-		DeviceContainer.Value.Device->AddReferencedObjects(Collector);
-	}
+	IterateOverAllDevices(
+		[&Collector](Audio::FDeviceId, FAudioDevice* InDevice)
+		{
+			InDevice->AddReferencedObjects(Collector);
+		}
+	);
 }
 
 void FAudioDeviceManager::StopSoundsUsingResource(USoundWave* InSoundWave, TArray<UAudioComponent*>* StoppedComponents)
 {
-	FScopeLock ScopeLock(&DeviceMapCriticalSection);
-	for (auto& DeviceContainer : Devices)
-	{
-		DeviceContainer.Value.Device->StopSoundsUsingResource(InSoundWave, StoppedComponents);
-	}
+	IterateOverAllDevices(
+		[&InSoundWave, &StoppedComponents](Audio::FDeviceId, FAudioDevice* InDevice)
+		{
+			InDevice->StopSoundsUsingResource(InSoundWave, StoppedComponents);
+		}
+	);
 }
 
 void FAudioDeviceManager::RegisterSoundClass(USoundClass* SoundClass)
 {
-	FScopeLock ScopeLock(&DeviceMapCriticalSection);
-	for (auto& DeviceContainer : Devices)
-	{
-		DeviceContainer.Value.Device->RegisterSoundClass(SoundClass);
-	}
+	IterateOverAllDevices(
+		[&SoundClass](Audio::FDeviceId, FAudioDevice* InDevice)
+		{
+			InDevice->RegisterSoundClass(SoundClass);
+		}
+	);
 }
 
 void FAudioDeviceManager::UnregisterSoundClass(USoundClass* SoundClass)
 {
-	FScopeLock ScopeLock(&DeviceMapCriticalSection);
-	for (auto& DeviceContainer : Devices)
-	{
-		DeviceContainer.Value.Device->UnregisterSoundClass(SoundClass);
-	}
+	IterateOverAllDevices(
+		[&SoundClass](Audio::FDeviceId, FAudioDevice* InDevice)
+		{
+			InDevice->UnregisterSoundClass(SoundClass);
+		}
+	);
 }
 
 void FAudioDeviceManager::InitSoundClasses()
 {
-	FScopeLock ScopeLock(&DeviceMapCriticalSection);
-	for (auto& DeviceContainer : Devices)
-	{
-		DeviceContainer.Value.Device->InitSoundClasses();
-	}
+	IterateOverAllDevices(
+		[](Audio::FDeviceId, FAudioDevice* InDevice)
+		{
+			InDevice->InitSoundClasses();
+		}
+	);
 }
 
 void FAudioDeviceManager::RegisterSoundSubmix(const USoundSubmixBase* SoundSubmix)
 {
-	FScopeLock ScopeLock(&DeviceMapCriticalSection);
-	for (auto& DeviceContainer : Devices)
-	{
-		DeviceContainer.Value.Device->RegisterSoundSubmix(SoundSubmix, true);
-	}
+	IterateOverAllDevices(
+		[&SoundSubmix](Audio::FDeviceId, FAudioDevice* InDevice)
+		{
+			InDevice->RegisterSoundSubmix(SoundSubmix, true);
+		}
+	);
 }
 
 void FAudioDeviceManager::UnregisterSoundSubmix(const USoundSubmixBase* SoundSubmix)
 {
-	FScopeLock ScopeLock(&DeviceMapCriticalSection);
-	for (auto& DeviceContainer : Devices)
-	{
-		DeviceContainer.Value.Device->UnregisterSoundSubmix(SoundSubmix);
-	}
+	IterateOverAllDevices(
+		[&SoundSubmix](Audio::FDeviceId, FAudioDevice* InDevice)
+		{
+			InDevice->UnregisterSoundSubmix(SoundSubmix);
+		}
+	);
 }
 
 void FAudioDeviceManager::InitSoundSubmixes()
 {
-	FScopeLock ScopeLock(&DeviceMapCriticalSection);
-	for (auto& DeviceContainer : Devices)
-	{
-		DeviceContainer.Value.Device->InitSoundSubmixes();
-	}
+	IterateOverAllDevices(
+		[](Audio::FDeviceId, FAudioDevice* InDevice)
+		{
+			InDevice->InitSoundSubmixes();
+		}
+	);
 }
 
 void FAudioDeviceManager::InitSoundEffectPresets()
 {
-	FScopeLock ScopeLock(&DeviceMapCriticalSection);
-	for (auto& DeviceContainer : Devices)
-	{
-		DeviceContainer.Value.Device->InitSoundEffectPresets();
-	}
+	IterateOverAllDevices(
+		[](Audio::FDeviceId, FAudioDevice* InDevice)
+		{
+			InDevice->InitSoundEffectPresets();
+		}
+	);
 }
 
 void FAudioDeviceManager::UpdateSourceEffectChain(const uint32 SourceEffectChainId, const TArray<FSourceEffectChainEntry>& SourceEffectChain, const bool bPlayEffectChainTails)
 {
-	FScopeLock ScopeLock(&DeviceMapCriticalSection);
-	for (auto& DeviceContainer : Devices)
-	{
-		DeviceContainer.Value.Device->UpdateSourceEffectChain(SourceEffectChainId, SourceEffectChain, bPlayEffectChainTails);
-	}
+	IterateOverAllDevices(
+		[&SourceEffectChainId, &SourceEffectChain, &bPlayEffectChainTails](Audio::FDeviceId, FAudioDevice* InDevice)
+		{
+			InDevice->UpdateSourceEffectChain(SourceEffectChainId, SourceEffectChain, bPlayEffectChainTails);
+		}
+	);
 }
 
 void FAudioDeviceManager::UpdateSubmix(USoundSubmixBase* SoundSubmix)
 {
-	FScopeLock ScopeLock(&DeviceMapCriticalSection);
-	for (auto& DeviceContainer : Devices)
-	{
-		DeviceContainer.Value.Device->UpdateSubmixProperties(SoundSubmix);
-	}
+	IterateOverAllDevices(
+		[&SoundSubmix](Audio::FDeviceId, FAudioDevice* InDevice)
+		{
+			InDevice->UpdateSubmixProperties(SoundSubmix);
+		}
+	);
 }
 
 void FAudioDeviceManager::SetActiveDevice(uint32 InAudioDeviceHandle)
@@ -929,11 +943,16 @@ void FAudioDeviceManager::SetSoloDevice(Audio::FDeviceId InAudioDeviceHandle)
 
 uint8 FAudioDeviceManager::GetNumActiveAudioDevices() const
 {
+	FCriticalSection* ConstCastCritSection = const_cast<FCriticalSection*>(&DeviceMapCriticalSection);
+	FScopeLock ScopeLock(ConstCastCritSection);
 	return Devices.Num();
 }
 
 uint8 FAudioDeviceManager::GetNumMainAudioDeviceWorlds() const
 {
+	FCriticalSection* ConstCastCritSection = const_cast<FCriticalSection*>(&DeviceMapCriticalSection);
+	FScopeLock ScopeLock(ConstCastCritSection);
+
 	const Audio::FDeviceId MainDeviceID = MainAudioDeviceHandle.GetDeviceID();
 	if (Devices.Contains(MainDeviceID))
 	{
@@ -959,6 +978,7 @@ TArray<FAudioDevice*> FAudioDeviceManager::GetAudioDevices()
 
 TArray<UWorld*> FAudioDeviceManager::GetWorldsUsingAudioDevice(const Audio::FDeviceId& InID)
 {
+	FScopeLock ScopeLock(&DeviceMapCriticalSection);
 	if (Devices.Contains(InID))
 	{
 		return Devices[InID].WorldsUsingThisDevice;
