Fix FCollisionQueryParams ctor to not skip narrow phase by default.

diff --git a/Engine/Source/Runtime/Engine/Private/Collision/Collision.cpp b/Engine/Source/Runtime/Engine/Private/Collision/Collision.cpp
--- a/Engine/Source/Runtime/Engine/Private/Collision/Collision.cpp
+++ b/Engine/Source/Runtime/Engine/Private/Collision/Collision.cpp
@@ -184,7 +184,7 @@ FCollisionQueryParams::FCollisionQueryParams(FName InTraceTag, const TStatId& In
 	IgnoreMask = 0;
 	bIgnoreBlocks = false;
 	bIgnoreTouches = false;
-
+	bSkipNarrowPhase = false;
 	AddIgnoredActor(InIgnoreActor);
 	if (InIgnoreActor != NULL)
 	{
