Submitting for Ethan: ADPCM crash / loop fix.

diff --git a/Engine/Source/Runtime/Engine/Private/ADPCMAudioInfo.cpp b/Engine/Source/Runtime/Engine/Private/ADPCMAudioInfo.cpp
--- a/Engine/Source/Runtime/Engine/Private/ADPCMAudioInfo.cpp
+++ b/Engine/Source/Runtime/Engine/Private/ADPCMAudioInfo.cpp
@@ -628,7 +628,7 @@ bool FADPCMAudioInfo::StreamCompressedData(uint8* Destination, bool bLooping, ui
 						ReachedEndOfSamples = true;
 						CurrentUncompressedBlockSampleIndex = 0;
 						CurrentChunkIndex = FirstChunkSampleDataIndex;
-						CurrentChunkBufferOffset = 0;
+						CurrentChunkBufferOffset = FirstChunkSampleDataOffset;
 						TotalSamplesStreamed = 0;
 						CurCompressedChunkData = nullptr;
 						if (!bLooping)
@@ -661,7 +661,27 @@ bool FADPCMAudioInfo::StreamCompressedData(uint8* Destination, bool bLooping, ui
 							const uint32 NumBlocksForCallback = BufferSize / (UncompressedBlockSize * NumChannels);
 							const uint32 NumCompressedBytesForCallback = NumBlocksForCallback * CompressedBlockSize * NumChannels;
 							
+							uint32 NumDecompressedSamplesForCallback = BufferSize / (ChannelSampleSize);
+
+							if (NumDecompressedSamplesForCallback > TotalSamplesPerChannel - TotalSamplesStreamed)
+							{
+								// If this is the case, we've seeked to the end of the file.
+								ReachedEndOfSamples = true;
+								CurrentUncompressedBlockSampleIndex = 0;
+								CurrentChunkIndex = FirstChunkSampleDataIndex;
+								CurrentChunkBufferOffset = FirstChunkSampleDataOffset;
+								TotalSamplesStreamed = 0;
+								CurCompressedChunkData = nullptr;
+								if (!bLooping)
+								{
+									// Set the remaining buffer to 0
+									FMemory::Memset(OutData, 0, BufferSize);
+									return true;
+								}
+							}
+
 							CurrentChunkBufferOffset += NumCompressedBytesForCallback;
+							TotalSamplesStreamed += NumDecompressedSamplesForCallback;
 
 							const uint32 SizeOfCurrentChunk = StreamingSoundWave->GetSizeOfChunk(CurrentChunkIndex);
 							if (CurrentChunkBufferOffset >= SizeOfCurrentChunk)
@@ -680,7 +700,7 @@ bool FADPCMAudioInfo::StreamCompressedData(uint8* Destination, bool bLooping, ui
 									ReachedEndOfSamples = true;
 									CurrentUncompressedBlockSampleIndex = 0;
 									CurrentChunkIndex = FirstChunkSampleDataIndex;
-									CurrentChunkBufferOffset = 0;
+									CurrentChunkBufferOffset = FirstChunkSampleDataOffset;
 									TotalSamplesStreamed = 0;
 									CurCompressedChunkData = nullptr;
 									if (!bLooping)
