IWYU fix
#rnx

diff --git a/Engine/Source/Runtime/Engine/Classes/Sound/SoundClass.h b/Engine/Source/Runtime/Engine/Classes/Sound/SoundClass.h
--- a/Engine/Source/Runtime/Engine/Classes/Sound/SoundClass.h
+++ b/Engine/Source/Runtime/Engine/Classes/Sound/SoundClass.h
@@ -6,6 +6,7 @@
 #include "AudioDefines.h"
 #include "AudioDynamicParameter.h"
 #include "IAudioExtensionPlugin.h"
+#include "IAudioModulation.h"
 #include "Sound/AudioOutputTarget.h"
 #include "UObject/ObjectMacros.h"
 #include "UObject/Object.h"
