IWYU fix
#rnx

diff --git a/Engine/Source/Runtime/Engine/Classes/Sound/SoundModulationParameter.h b/Engine/Source/Runtime/Engine/Classes/Sound/SoundModulationParameter.h
--- a/Engine/Source/Runtime/Engine/Classes/Sound/SoundModulationParameter.h
+++ b/Engine/Source/Runtime/Engine/Classes/Sound/SoundModulationParameter.h
@@ -1,6 +1,7 @@
 // Copyright Epic Games, Inc. All Rights Reserved.
 #pragma once
 
+#include "Audio.h"
 #include "CoreMinimal.h"
 #include "DSP/BufferVectorOperations.h"
 #include "IAudioModulation.h"
