Mirror 12732403 from Release-4.25
Blueprint debugging information now uses a safer FProperty reference type (TFieldPath<FProperty>), preventing crashes while debugging after properties have been removed

diff --git a/Engine/Source/Runtime/Engine/Classes/Engine/BlueprintGeneratedClass.h b/Engine/Source/Runtime/Engine/Classes/Engine/BlueprintGeneratedClass.h
--- a/Engine/Source/Runtime/Engine/Classes/Engine/BlueprintGeneratedClass.h
+++ b/Engine/Source/Runtime/Engine/Classes/Engine/BlueprintGeneratedClass.h
@@ -227,10 +227,10 @@ protected:
 	TMap<TWeakObjectPtr<UFunction>, FDebuggingInfoForSingleFunction> PerFunctionLineNumbers;
 
 	// Map from objects to class properties they created
-	TMap<TWeakObjectPtr<UObject>, FProperty*> DebugObjectToPropertyMap;
+	TMap<TWeakObjectPtr<UObject>, TFieldPath<FProperty> > DebugObjectToPropertyMap;
 
 	// Map from pins or nodes to class properties they created
-	TMap<FEdGraphPinReference, FProperty*> DebugPinToPropertyMap;
+	TMap<FEdGraphPinReference, TFieldPath<FProperty> > DebugPinToPropertyMap;
 
 public:
 
@@ -383,20 +383,20 @@ public:
 			return nullptr;
 		}
 
-		FProperty* PropertyPtr = DebugPinToPropertyMap.FindRef(Pin);
+		TFieldPath<FProperty> PropertyPtr = DebugPinToPropertyMap.FindRef(Pin);
 		if ((PropertyPtr == nullptr) && (Pin->LinkedTo.Num() > 0))
 		{
 			// Try checking the other side of the connection
 			PropertyPtr = DebugPinToPropertyMap.FindRef(Pin->LinkedTo[0]);
 		}
 
-		return PropertyPtr;
+		return *PropertyPtr;
 	}
 
 	// Looks thru the debugging data for any class variables associated with the node (e.g., temporary variables or timelines)
 	FProperty* FindClassPropertyForNode(const UEdGraphNode* Node) const
 	{
-		return DebugObjectToPropertyMap.FindRef(MakeWeakObjectPtr(const_cast<UEdGraphNode*>(Node)));
+		return *DebugObjectToPropertyMap.FindRef(MakeWeakObjectPtr(const_cast<UEdGraphNode*>(Node)));
 	}
 
 	// Adds a debug record for a source node and destination in the bytecode of a specified function
@@ -488,11 +488,11 @@ public:
 
 	void GenerateReversePropertyMap(TMap<FProperty*, UObject*>& PropertySourceMap)
 	{
-		for (TMap<TWeakObjectPtr<UObject>, FProperty*>::TIterator MapIt(DebugObjectToPropertyMap); MapIt; ++MapIt)
+		for (TMap<TWeakObjectPtr<UObject>, TFieldPath<FProperty>>::TIterator MapIt(DebugObjectToPropertyMap); MapIt; ++MapIt)
 		{
 			if (UObject* SourceObj = MapIt.Key().Get())
 			{
-				PropertySourceMap.Add(MapIt.Value(), SourceObj);
+				PropertySourceMap.Add(*MapIt.Value(), SourceObj);
 			}
 		}
 	}
