Move per platform default skeletal mesh LOD streaming settings to project settings

diff --git a/Engine/Platforms/XXX/Source/Developer/XXXTargetPlatform/Private/XXXTargetPlatform.cpp b/Engine/Platforms/XXX/Source/Developer/XXXTargetPlatform/Private/XXXTargetPlatform.cpp
--- a/Engine/Platforms/XXX/Source/Developer/XXXTargetPlatform/Private/XXXTargetPlatform.cpp
+++ b/Engine/Platforms/XXX/Source/Developer/XXXTargetPlatform/Private/XXXTargetPlatform.cpp
@@ -35,7 +35,6 @@ FXXXTargetPlatform::FXXXTargetPlatform()
 	// load up texture settings from the config file
 	XXXLODSettings = nullptr; // These are registered by the device profile system.
 	StaticMeshLODSettings.Initialize(XXXEngineSettings);
-	SkeletalMeshDefaultLODStreamingSettings.Initialize(XXXEngineSettings);
 #endif
 }
 
@@ -151,12 +150,6 @@ const class FStaticMeshLODSettings& FXXXTargetPlatform::GetStaticMeshLODSettings
 }
 
 
-const FSkeletalMeshDefaultLODStreamingSettings& FXXXTargetPlatform::GetSkeletalMeshDefaultLODStreamingSettings() const
-{
-	return SkeletalMeshDefaultLODStreamingSettings;
-}
-
-
 void FXXXTargetPlatform::GetTextureFormats( const UTexture* InTexture, TArray< TArray<FName> >& OutFormats ) const
 {
 	GetDefaultTextureFormatNamePerLayer(OutFormats.AddDefaulted_GetRef(), this, InTexture, XXXEngineSettings, true);
diff --git a/Engine/Platforms/XXX/Source/Developer/XXXTargetPlatform/Private/XXXTargetPlatform.h b/Engine/Platforms/XXX/Source/Developer/XXXTargetPlatform/Private/XXXTargetPlatform.h
--- a/Engine/Platforms/XXX/Source/Developer/XXXTargetPlatform/Private/XXXTargetPlatform.h
+++ b/Engine/Platforms/XXX/Source/Developer/XXXTargetPlatform/Private/XXXTargetPlatform.h
@@ -16,7 +16,6 @@
 
 #if WITH_ENGINE
 #include "StaticMeshResources.h"
-#include "SkeletalMeshDefaultLODStreamingSettings.h"
 #endif // WITH_ENGINE
 
 /**
@@ -64,8 +63,6 @@ public:
 
 	virtual const class FStaticMeshLODSettings& GetStaticMeshLODSettings() const override;
 
-	virtual const FSkeletalMeshDefaultLODStreamingSettings& GetSkeletalMeshDefaultLODStreamingSettings() const override;
-
 	virtual void GetTextureFormats( const UTexture* InTexture, TArray< TArray<FName> >& OutFormats ) const override;
 
 	virtual void GetAllTextureFormats(TArray<FName>& OutFormats) const override;
@@ -112,8 +109,6 @@ private:
 
 	// Holds the static mesh LOD settings.
 	FStaticMeshLODSettings StaticMeshLODSettings;
-
-	FSkeletalMeshDefaultLODStreamingSettings SkeletalMeshDefaultLODStreamingSettings;
 #endif
 
 private:
diff --git a/Engine/Plugins/Runtime/AR/Microsoft/HoloLensAR/Source/HoloLensTargetPlatform/Private/HoloLensTargetPlatform.cpp b/Engine/Plugins/Runtime/AR/Microsoft/HoloLensAR/Source/HoloLensTargetPlatform/Private/HoloLensTargetPlatform.cpp
--- a/Engine/Plugins/Runtime/AR/Microsoft/HoloLensAR/Source/HoloLensTargetPlatform/Private/HoloLensTargetPlatform.cpp
+++ b/Engine/Plugins/Runtime/AR/Microsoft/HoloLensAR/Source/HoloLensTargetPlatform/Private/HoloLensTargetPlatform.cpp
@@ -24,7 +24,6 @@ FHoloLensTargetPlatform::FHoloLensTargetPlatform()
 	FConfigCacheIni::LoadLocalIniFile(EngineSettings, TEXT("Engine"), true, *PlatformName());
 	TextureLODSettings = nullptr; // These are registered by the device profile system.
 	StaticMeshLODSettings.Initialize(EngineSettings);
-	SkeletalMeshDefaultLODStreamingSettings.Initialize(EngineSettings);
 #endif
 
 	DeviceDetector = IHoloLensDeviceDetector::Create();
diff --git a/Engine/Plugins/Runtime/AR/Microsoft/HoloLensAR/Source/HoloLensTargetPlatform/Private/HoloLensTargetPlatform.h b/Engine/Plugins/Runtime/AR/Microsoft/HoloLensAR/Source/HoloLensTargetPlatform/Private/HoloLensTargetPlatform.h
--- a/Engine/Plugins/Runtime/AR/Microsoft/HoloLensAR/Source/HoloLensTargetPlatform/Private/HoloLensTargetPlatform.h
+++ b/Engine/Plugins/Runtime/AR/Microsoft/HoloLensAR/Source/HoloLensTargetPlatform/Private/HoloLensTargetPlatform.h
@@ -19,7 +19,6 @@
 
 #if WITH_ENGINE
 #include "StaticMeshResources.h"
-#include "SkeletalMeshDefaultLODStreamingSettings.h"
 #endif // WITH_ENGINE
 
 #define LOCTEXT_NAMESPACE "HoloLensTargetPlatform"
@@ -71,8 +70,6 @@ public:
 
 	virtual const class FStaticMeshLODSettings& GetStaticMeshLODSettings() const override { return StaticMeshLODSettings; }
 
-	virtual const FSkeletalMeshDefaultLODStreamingSettings& GetSkeletalMeshDefaultLODStreamingSettings() const override { return SkeletalMeshDefaultLODStreamingSettings; }
-
 	virtual void GetTextureFormats(const UTexture* InTexture, TArray< TArray<FName> >& OutFormats) const override;
 
 	virtual const UTextureLODSettings& GetTextureLODSettings() const override { return *TextureLODSettings; }
@@ -173,8 +170,6 @@ private:
 
 	// Holds the static mesh LOD settings.
 	FStaticMeshLODSettings StaticMeshLODSettings;
-
-	FSkeletalMeshDefaultLODStreamingSettings SkeletalMeshDefaultLODStreamingSettings;
 #endif //WITH_ENGINE
 
 private:
diff --git a/Engine/Source/Developer/AllDesktopTargetPlatform/Private/AllDesktopTargetPlatform.cpp b/Engine/Source/Developer/AllDesktopTargetPlatform/Private/AllDesktopTargetPlatform.cpp
--- a/Engine/Source/Developer/AllDesktopTargetPlatform/Private/AllDesktopTargetPlatform.cpp
+++ b/Engine/Source/Developer/AllDesktopTargetPlatform/Private/AllDesktopTargetPlatform.cpp
@@ -22,7 +22,6 @@ FAllDesktopTargetPlatform::FAllDesktopTargetPlatform()
 	// use non-platform specific settings
 	FConfigCacheIni::LoadLocalIniFile(EngineSettings, TEXT("Engine"), true, NULL);
 	StaticMeshLODSettings.Initialize(EngineSettings);
-	SkeletalMeshDefaultLODStreamingSettings.Initialize(EngineSettings);
 #endif // #if WITH_ENGINE
 
 }
diff --git a/Engine/Source/Developer/AllDesktopTargetPlatform/Private/AllDesktopTargetPlatform.h b/Engine/Source/Developer/AllDesktopTargetPlatform/Private/AllDesktopTargetPlatform.h
--- a/Engine/Source/Developer/AllDesktopTargetPlatform/Private/AllDesktopTargetPlatform.h
+++ b/Engine/Source/Developer/AllDesktopTargetPlatform/Private/AllDesktopTargetPlatform.h
@@ -15,7 +15,6 @@
 #if WITH_ENGINE
 #include "AudioCompressionSettings.h"
 #include "StaticMeshResources.h"
-#include "SkeletalMeshDefaultLODStreamingSettings.h"
 #endif // WITH_ENGINE
 
 class UTextureLODSettings;
@@ -134,11 +133,6 @@ public:
 		return StaticMeshLODSettings;
 	}
 
-	virtual const FSkeletalMeshDefaultLODStreamingSettings& GetSkeletalMeshDefaultLODStreamingSettings() const override
-	{
-		return SkeletalMeshDefaultLODStreamingSettings;
-	}
-
 	virtual void GetTextureFormats(const UTexture* Texture, TArray< TArray<FName> >& OutFormats) const override;
 
 	virtual void GetAllTextureFormats(TArray<FName>& OutFormats) const override;
@@ -183,8 +177,6 @@ private:
 
 	// Holds the static mesh LOD settings.
 	FStaticMeshLODSettings StaticMeshLODSettings;
-
-	FSkeletalMeshDefaultLODStreamingSettings SkeletalMeshDefaultLODStreamingSettings;
 #endif // WITH_ENGINE
 
 	// Holds an event delegate that is executed when a new target device has been discovered.
diff --git a/Engine/Source/Developer/Android/AndroidTargetPlatform/Private/AndroidTargetPlatform.cpp b/Engine/Source/Developer/Android/AndroidTargetPlatform/Private/AndroidTargetPlatform.cpp
--- a/Engine/Source/Developer/Android/AndroidTargetPlatform/Private/AndroidTargetPlatform.cpp
+++ b/Engine/Source/Developer/Android/AndroidTargetPlatform/Private/AndroidTargetPlatform.cpp
@@ -199,7 +199,6 @@ FAndroidTargetPlatform::FAndroidTargetPlatform(bool bInIsClient )
 		FConfigCacheIni::LoadLocalIniFile(EngineSettings, TEXT("Engine"), true, *IniPlatformName());
 			TextureLODSettings = nullptr; // These are registered by the device profile system.
 		StaticMeshLODSettings.Initialize(EngineSettings);
-		SkeletalMeshDefaultLODStreamingSettings.Initialize(EngineSettings);
 	#endif
 
 	TickDelegate = FTickerDelegate::CreateRaw(this, &FAndroidTargetPlatform::HandleTicker);
@@ -427,12 +426,6 @@ const FStaticMeshLODSettings& FAndroidTargetPlatform::GetStaticMeshLODSettings(
 }
 
 
-const FSkeletalMeshDefaultLODStreamingSettings& FAndroidTargetPlatform::GetSkeletalMeshDefaultLODStreamingSettings() const
-{
-	return SkeletalMeshDefaultLODStreamingSettings;
-}
-
-
 void FAndroidTargetPlatform::GetTextureFormats( const UTexture* InTexture, TArray< TArray<FName> >& OutFormats) const
 {
 #if WITH_EDITOR
diff --git a/Engine/Source/Developer/Android/AndroidTargetPlatform/Private/AndroidTargetPlatform.h b/Engine/Source/Developer/Android/AndroidTargetPlatform/Private/AndroidTargetPlatform.h
--- a/Engine/Source/Developer/Android/AndroidTargetPlatform/Private/AndroidTargetPlatform.h
+++ b/Engine/Source/Developer/Android/AndroidTargetPlatform/Private/AndroidTargetPlatform.h
@@ -24,7 +24,6 @@
 #if WITH_ENGINE
 #include "Internationalization/Text.h"
 #include "StaticMeshResources.h"
-#include "SkeletalMeshDefaultLODStreamingSettings.h"
 #endif // WITH_ENGINE
 
 #define LOCTEXT_NAMESPACE "FAndroidTargetPlatform" 
@@ -197,8 +196,6 @@ public:
 	
 	virtual const class FStaticMeshLODSettings& GetStaticMeshLODSettings() const override;
 
-	virtual const FSkeletalMeshDefaultLODStreamingSettings& GetSkeletalMeshDefaultLODStreamingSettings() const override;
-
 	virtual void GetTextureFormats( const UTexture* InTexture, TArray< TArray<FName> >& OutFormats) const override;
 
 	virtual void GetAllTextureFormats(TArray<FName>& OutFormats) const override;
@@ -311,8 +308,6 @@ protected:
 	// Holds the static mesh LOD settings.
 	FStaticMeshLODSettings StaticMeshLODSettings;
 
-	FSkeletalMeshDefaultLODStreamingSettings SkeletalMeshDefaultLODStreamingSettings;
-
 	ITargetDevicePtr DefaultDevice;
 #endif //WITH_ENGINE
 
diff --git a/Engine/Source/Developer/IOS/IOSTargetPlatform/Private/IOSTargetPlatform.cpp b/Engine/Source/Developer/IOS/IOSTargetPlatform/Private/IOSTargetPlatform.cpp
--- a/Engine/Source/Developer/IOS/IOSTargetPlatform/Private/IOSTargetPlatform.cpp
+++ b/Engine/Source/Developer/IOS/IOSTargetPlatform/Private/IOSTargetPlatform.cpp
@@ -37,7 +37,6 @@ FIOSTargetPlatform::FIOSTargetPlatform(bool bInIsTVOS, bool bInIsClientOnly)
 	FConfigCacheIni::LoadLocalIniFile(EngineSettings, TEXT("Engine"), true, *PlatformName());
 	TextureLODSettings = nullptr; // TextureLODSettings are registered by the device profile.
 	StaticMeshLODSettings.Initialize(EngineSettings);
-	SkeletalMeshDefaultLODStreamingSettings.Initialize(EngineSettings);
 #endif // #if WITH_ENGINE
 
 	// Initialize Ticker for device discovery
diff --git a/Engine/Source/Developer/IOS/IOSTargetPlatform/Private/IOSTargetPlatform.h b/Engine/Source/Developer/IOS/IOSTargetPlatform/Private/IOSTargetPlatform.h
--- a/Engine/Source/Developer/IOS/IOSTargetPlatform/Private/IOSTargetPlatform.h
+++ b/Engine/Source/Developer/IOS/IOSTargetPlatform/Private/IOSTargetPlatform.h
@@ -20,7 +20,6 @@
 #if WITH_ENGINE
 #include "AudioCompressionSettings.h"
 #include "StaticMeshResources.h"
-#include "SkeletalMeshDefaultLODStreamingSettings.h"
 #endif // WITH_ENGINE
 
 /**
@@ -116,11 +115,6 @@ public:
 		return StaticMeshLODSettings;
 	}
 
-	virtual const FSkeletalMeshDefaultLODStreamingSettings& GetSkeletalMeshDefaultLODStreamingSettings() const override
-	{
-		return SkeletalMeshDefaultLODStreamingSettings;
-	}
-
 	virtual void GetTextureFormats( const UTexture* Texture, TArray< TArray<FName> >& OutFormats) const override;
 
 	virtual void GetAllTextureFormats( TArray<FName>& OutFormats) const override;
@@ -218,8 +212,6 @@ private:
 
 	// Holds the static mesh LOD settings.
 	FStaticMeshLODSettings StaticMeshLODSettings;
-
-	FSkeletalMeshDefaultLODStreamingSettings SkeletalMeshDefaultLODStreamingSettings;
 #endif // WITH_ENGINE
 
     // holds usb device helper
diff --git a/Engine/Source/Developer/Linux/LinuxTargetPlatform/Private/LinuxTargetPlatform.h b/Engine/Source/Developer/Linux/LinuxTargetPlatform/Private/LinuxTargetPlatform.h
--- a/Engine/Source/Developer/Linux/LinuxTargetPlatform/Private/LinuxTargetPlatform.h
+++ b/Engine/Source/Developer/Linux/LinuxTargetPlatform/Private/LinuxTargetPlatform.h
@@ -18,7 +18,6 @@
 #include "AudioCompressionSettings.h"
 #include "Sound/SoundWave.h"
 #include "StaticMeshResources.h"
-#include "SkeletalMeshDefaultLODStreamingSettings.h"
 #endif // WITH_ENGINE
 #include "Interfaces/IProjectManager.h"
 #include "InstalledPlatformInfo.h"
@@ -60,7 +59,6 @@ public:
 		FConfigCacheIni::LoadLocalIniFile(EngineSettings, TEXT("Engine"), true, *this->PlatformName());
 		TextureLODSettings = nullptr;
 		StaticMeshLODSettings.Initialize(EngineSettings);
-		SkeletalMeshDefaultLODStreamingSettings.Initialize(EngineSettings);
 
 		InitDevicesFromConfig();
 
@@ -293,10 +291,6 @@ public:
 		return StaticMeshLODSettings;
 	}
 
-	virtual const FSkeletalMeshDefaultLODStreamingSettings& GetSkeletalMeshDefaultLODStreamingSettings() const override
-	{
-		return SkeletalMeshDefaultLODStreamingSettings;
-	}
 
 	virtual void GetTextureFormats( const UTexture* InTexture, TArray< TArray<FName> >& OutFormats) const override
 	{
@@ -523,8 +517,6 @@ protected:
 	// Holds static mesh LOD settings.
 	FStaticMeshLODSettings StaticMeshLODSettings;
 
-	FSkeletalMeshDefaultLODStreamingSettings SkeletalMeshDefaultLODStreamingSettings;
-
 	// True if the project requires encoded HDR reflection captures
 	bool bRequiresEncodedHDRReflectionCaptures;
 #endif // WITH_ENGINE
diff --git a/Engine/Source/Developer/Mac/MacTargetPlatform/Private/GenericMacTargetPlatform.h b/Engine/Source/Developer/Mac/MacTargetPlatform/Private/GenericMacTargetPlatform.h
--- a/Engine/Source/Developer/Mac/MacTargetPlatform/Private/GenericMacTargetPlatform.h
+++ b/Engine/Source/Developer/Mac/MacTargetPlatform/Private/GenericMacTargetPlatform.h
@@ -18,7 +18,6 @@
 #include "Sound/SoundWave.h"
 #include "TextureResource.h"
 #include "StaticMeshResources.h"
-#include "SkeletalMeshDefaultLODStreamingSettings.h"
 #endif // WITH_ENGINE
 
 #define LOCTEXT_NAMESPACE "TGenericMacTargetPlatform"
@@ -49,7 +48,6 @@ public:
 			FConfigCacheIni::LoadLocalIniFile(EngineSettings, TEXT("Engine"), true, *this->PlatformName());
 			TextureLODSettings = nullptr;
 			StaticMeshLODSettings.Initialize(EngineSettings);
-			SkeletalMeshDefaultLODStreamingSettings.Initialize(EngineSettings);
 		#endif
 	}
 
@@ -156,11 +154,6 @@ return TSuper::SupportsFeature(Feature);
 		return StaticMeshLODSettings;
 	}
 
-	virtual const FSkeletalMeshDefaultLODStreamingSettings& GetSkeletalMeshDefaultLODStreamingSettings() const override
-	{
-		return SkeletalMeshDefaultLODStreamingSettings;
-	}
-
 	virtual void GetTextureFormats( const UTexture* Texture, TArray< TArray<FName> >& OutFormats) const override
 	{
 		if (!IS_DEDICATED_SERVER)
@@ -417,7 +410,6 @@ private:
 	// Holds the static mesh LOD settings.
 	FStaticMeshLODSettings StaticMeshLODSettings;
 
-	FSkeletalMeshDefaultLODStreamingSettings SkeletalMeshDefaultLODStreamingSettings;
 #endif // WITH_ENGINE
 
 private:
diff --git a/Engine/Source/Developer/TargetPlatform/Public/Interfaces/ITargetPlatform.h b/Engine/Source/Developer/TargetPlatform/Public/Interfaces/ITargetPlatform.h
--- a/Engine/Source/Developer/TargetPlatform/Public/Interfaces/ITargetPlatform.h
+++ b/Engine/Source/Developer/TargetPlatform/Public/Interfaces/ITargetPlatform.h
@@ -517,13 +517,6 @@ public:
 	 */
 	virtual const class FStaticMeshLODSettings& GetStaticMeshLODSettings() const = 0;
 
-	/**
-	 * Gets the skeletal mesh default LOD streaming settings used by this platform.
-	 *
-	 * @return A skeletal mesh default LOD streaming settings structure.
-	 */
-	virtual const class FSkeletalMeshDefaultLODStreamingSettings& GetSkeletalMeshDefaultLODStreamingSettings() const = 0;
-
 	/** 
 	 * Gets the name of the mesh builder module that is responsible for 
 	 * building static and skeletal mesh derived data for this platform.
diff --git a/Engine/Source/Developer/Windows/WindowsTargetPlatform/Private/GenericWindowsTargetPlatform.h b/Engine/Source/Developer/Windows/WindowsTargetPlatform/Private/GenericWindowsTargetPlatform.h
--- a/Engine/Source/Developer/Windows/WindowsTargetPlatform/Private/GenericWindowsTargetPlatform.h
+++ b/Engine/Source/Developer/Windows/WindowsTargetPlatform/Private/GenericWindowsTargetPlatform.h
@@ -15,7 +15,6 @@
 	#include "StaticMeshResources.h"
 	#include "RHI.h"
 	#include "AudioCompressionSettings.h"
-	#include "SkeletalMeshDefaultLODStreamingSettings.h"
 #endif // WITH_ENGINE
 
 
@@ -46,7 +45,7 @@ public:
 		FConfigCacheIni::LoadLocalIniFile(EngineSettings, TEXT("Engine"), true, *this->PlatformName());
 		TextureLODSettings = nullptr; // These are registered by the device profile system.
 		StaticMeshLODSettings.Initialize(EngineSettings);
-		SkeletalMeshDefaultLODStreamingSettings.Initialize(EngineSettings);
+
 
 		// Get the Target RHIs for this platform, we do not always want all those that are supported.
 		TArray<FName> TargetedShaderFormats;
@@ -248,11 +247,6 @@ public:
 		return StaticMeshLODSettings;
 	}
 
-	virtual const FSkeletalMeshDefaultLODStreamingSettings& GetSkeletalMeshDefaultLODStreamingSettings() const override
-	{
-		return SkeletalMeshDefaultLODStreamingSettings;
-	}
-
 	virtual void GetTextureFormats( const UTexture* InTexture, TArray< TArray<FName> >& OutFormats) const override
 	{
 		if (!IS_DEDICATED_SERVER)
@@ -511,8 +505,6 @@ private:
 	// Holds static mesh LOD settings.
 	FStaticMeshLODSettings StaticMeshLODSettings;
 
-	FSkeletalMeshDefaultLODStreamingSettings SkeletalMeshDefaultLODStreamingSettings;
-
 	// True if the project supports non-DX11 texture formats.
 	bool bSupportDX11TextureFormats;
 
diff --git a/Engine/Source/Runtime/Engine/Classes/Engine/RendererSettings.h b/Engine/Source/Runtime/Engine/Classes/Engine/RendererSettings.h
--- a/Engine/Source/Runtime/Engine/Classes/Engine/RendererSettings.h
+++ b/Engine/Source/Runtime/Engine/Classes/Engine/RendererSettings.h
@@ -810,6 +810,16 @@ class ENGINE_API URendererSettings : public UDeveloperSettings
 		ConfigRestartRequired = true))
 		uint32 bLPV : 1;
 
+	UPROPERTY(config, EditAnywhere, Category="Experimental|LOD Streaming|Skeletal Mesh", meta=(
+		DisplayName="Stream LODs by default",
+		ToolTip="Whether to stream skeletal mesh LODs by default."))
+	FPerPlatformBool bStreamSkeletalMeshLODs;
+
+	UPROPERTY(config, EditAnywhere, Category="Experimental|LOD Streaming|Skeletal Mesh", meta=(
+		DisplayName="Discard optional LODs",
+		ToolTip="Whether to discard skeletal mesh LODs below minimum LOD levels at cook time."))
+	FPerPlatformBool bDiscardSkeletalMeshOptionalLODs;
+
 public:
 
 	//~ Begin UObject Interface
diff --git a/Engine/Source/Runtime/Engine/Private/SkeletalMesh.cpp b/Engine/Source/Runtime/Engine/Private/SkeletalMesh.cpp
--- a/Engine/Source/Runtime/Engine/Private/SkeletalMesh.cpp
+++ b/Engine/Source/Runtime/Engine/Private/SkeletalMesh.cpp
@@ -69,7 +69,7 @@
 
 #include "IMeshReductionManagerModule.h"
 #include "SkeletalMeshReductionSettings.h"
-#include "SkeletalMeshDefaultLODStreamingSettings.h"
+#include "Engine/RendererSettings.h"
 
 #endif // #if WITH_EDITOR
 
@@ -4101,15 +4101,15 @@ void USkeletalMesh::ResetLODInfo()
 bool USkeletalMesh::GetSupportsLODStreaming(const ITargetPlatform* TargetPlatform) const
 {
 	check(TargetPlatform);
+	const FName PlatformGroupName = TargetPlatform->GetPlatformInfo().PlatformGroupName;
+	const FName VanillaPlatformName = TargetPlatform->GetPlatformInfo().VanillaPlatformName;
 	if (bOverrideLODStreamingSettings)
 	{
-		return bSupportLODStreaming.GetValueForPlatformIdentifiers(
-			TargetPlatform->GetPlatformInfo().PlatformGroupName,
-			TargetPlatform->GetPlatformInfo().VanillaPlatformName);
+		return bSupportLODStreaming.GetValueForPlatformIdentifiers(PlatformGroupName, VanillaPlatformName);
 	}
 	else
 	{
-		return TargetPlatform->GetSkeletalMeshDefaultLODStreamingSettings().IsLODStreamingEnabled();
+		return GetDefault<URendererSettings>()->bStreamSkeletalMeshLODs.GetValueForPlatformIdentifiers(PlatformGroupName, VanillaPlatformName);
 	}
 }
 
@@ -4124,22 +4124,22 @@ int32 USkeletalMesh::GetMaxNumStreamedLODs(const ITargetPlatform* TargetPlatform
 	}
 	else
 	{
-		return TargetPlatform->GetSkeletalMeshDefaultLODStreamingSettings().GetMaxNumStreamedLODs();
+		return MAX_MESH_LOD_COUNT;
 	}
 }
 
 int32 USkeletalMesh::GetMaxNumOptionalLODs(const ITargetPlatform* TargetPlatform) const
 {
 	check(TargetPlatform);
+	const FName PlatformGroupName = TargetPlatform->GetPlatformInfo().PlatformGroupName;
+	const FName VanillaPlatformName = TargetPlatform->GetPlatformInfo().VanillaPlatformName;
 	if (bOverrideLODStreamingSettings)
 	{
-		return MaxNumOptionalLODs.GetValueForPlatformIdentifiers(
-			TargetPlatform->GetPlatformInfo().PlatformGroupName,
-			TargetPlatform->GetPlatformInfo().VanillaPlatformName);
+		return MaxNumOptionalLODs.GetValueForPlatformIdentifiers(PlatformGroupName, VanillaPlatformName) <= 0 ? 0 : MAX_MESH_LOD_COUNT;
 	}
 	else
 	{
-		return TargetPlatform->GetSkeletalMeshDefaultLODStreamingSettings().GetMaxNumOptionalLODs();
+		return GetDefault<URendererSettings>()->bDiscardSkeletalMeshOptionalLODs.GetValueForPlatformIdentifiers(PlatformGroupName, VanillaPlatformName) ? 0 : MAX_MESH_LOD_COUNT;
 	}
 }
 #endif
diff --git a/Engine/Source/Runtime/Engine/Private/SkeletalMeshLODSettings.cpp b/Engine/Source/Runtime/Engine/Private/SkeletalMeshLODSettings.cpp
--- a/Engine/Source/Runtime/Engine/Private/SkeletalMeshLODSettings.cpp
+++ b/Engine/Source/Runtime/Engine/Private/SkeletalMeshLODSettings.cpp
@@ -7,8 +7,6 @@
 #include "Animation/AnimSequence.h"
 #include "UObject/FortniteMainBranchObjectVersion.h"
 #include "UObject/EditorObjectVersion.h"
-#include "SkeletalMeshDefaultLODStreamingSettings.h"
-#include "Misc/ConfigCacheIni.h"
 
 DEFINE_LOG_CATEGORY_STATIC(LogSkeletalMeshLODSettings, Warning, All)
 
@@ -295,32 +293,3 @@ const float FSkeletalMeshLODGroupSettings::GetScreenSize() const
 {
 	return ScreenSize.Default;
 }
-
-/////////////////////////////////////////////////////////////
-// FSkeletalMeshDefaultLODStreamingSettings 
-/////////////////////////////////////////////////////////////
-FSkeletalMeshDefaultLODStreamingSettings::FSkeletalMeshDefaultLODStreamingSettings()
-	: bSupportLODStreaming(false)
-	, MaxNumStreamedLODs(0)
-	, MaxNumOptionalLODs(0)
-{}
-
-void FSkeletalMeshDefaultLODStreamingSettings::Initialize(const FConfigFile& IniFile)
-{
-	const TCHAR* IniSection = TEXT("SkeletalMeshDefaultLODStreamingSettings");
-	
-	if (!IniFile.GetBool(IniSection, TEXT("bSupportLODStreaming"), bSupportLODStreaming))
-	{
-		bSupportLODStreaming = false;
-	}
-
-	if (!IniFile.GetInt(IniSection, TEXT("MaxNumStreamedLODs"), MaxNumStreamedLODs))
-	{
-		MaxNumStreamedLODs = 0;
-	}
-
-	if (!IniFile.GetInt(IniSection, TEXT("MaxNumOptionalLODs"), MaxNumOptionalLODs))
-	{
-		MaxNumOptionalLODs = 0;
-	}
-}
diff --git a/Engine/Source/Runtime/Engine/Public/SkeletalMeshDefaultLODStreamingSettings.h b/Engine/Source/Runtime/Engine/Public/SkeletalMeshDefaultLODStreamingSettings.h
deleted file mode 100644
