Move the scope lock for data table changes into the header file so it can be used by classes derived from UDataTable.
Add a scope lock to UpdateCachedRowMap in composite data tables to prevent unnecessary work from being done.

diff --git a/Engine/Source/Runtime/Engine/Classes/Engine/DataTable.h b/Engine/Source/Runtime/Engine/Classes/Engine/DataTable.h
--- a/Engine/Source/Runtime/Engine/Classes/Engine/DataTable.h
+++ b/Engine/Source/Runtime/Engine/Classes/Engine/DataTable.h
@@ -379,8 +379,20 @@ protected:
 	 */
 	void OnPostDataImported(OUT TArray<FString>& OutCollectedImportProblems);
 
-
 	UScriptStruct& GetEmptyUsingStruct() const;
+
+	/** Used to trigger the data table changed delegate. This allows us to trigger the delegate only once from more complex changes */
+	struct FScopedDataTableChange
+	{
+		FScopedDataTableChange(UDataTable* InTable);
+		~FScopedDataTableChange();
+
+	private:
+		UDataTable* Table;
+
+		static TMap<UDataTable*, int32> ScopeCount;
+		static FCriticalSection CriticalSection;
+	};
 };
 
 
diff --git a/Engine/Source/Runtime/Engine/Private/CompositeDataTable.cpp b/Engine/Source/Runtime/Engine/Private/CompositeDataTable.cpp
--- a/Engine/Source/Runtime/Engine/Private/CompositeDataTable.cpp
+++ b/Engine/Source/Runtime/Engine/Private/CompositeDataTable.cpp
@@ -17,6 +17,8 @@
 
 #define LOCTEXT_NAMESPACE "CompositeDataTables"
 
+#define DATATABLE_CHANGE_SCOPE()	UDataTable::FScopedDataTableChange ActiveScope(this);
+
 UCompositeDataTable::UCompositeDataTable(const FObjectInitializer& ObjectInitializer)
 	: Super(ObjectInitializer)
 {
@@ -81,6 +83,8 @@ void UCompositeDataTable::UpdateCachedRowMap(bool bWarnOnInvalidChildren)
 		}
 	}
 
+	DATATABLE_CHANGE_SCOPE();
+
 #if WITH_EDITOR
 	FDataTableEditorUtils::BroadcastPreChange(this, FDataTableEditorUtils::EDataTableChangeInfo::RowList);
 #endif
diff --git a/Engine/Source/Runtime/Engine/Private/DataTable.cpp b/Engine/Source/Runtime/Engine/Private/DataTable.cpp
--- a/Engine/Source/Runtime/Engine/Private/DataTable.cpp
+++ b/Engine/Source/Runtime/Engine/Private/DataTable.cpp
@@ -31,40 +31,32 @@ namespace
 		}
 	}
 #endif // WITH_EDITORONLY_DATA
-	/** Used to trigger the trigger the data table changed delegate. This allows us to trigger the delegate once from more complex changes */
-	struct FScopedDataTableChange
-	{
-		FScopedDataTableChange(UDataTable* InTable)
-			: Table(InTable)
-		{
-			FScopeLock Lock(&CriticalSection);
-			int32& Count = ScopeCount.FindOrAdd(Table);
-			++Count;
-		}
-		~FScopedDataTableChange()
-		{
-			FScopeLock Lock(&CriticalSection);
-			int32& Count = ScopeCount.FindChecked(Table);
-			--Count;
-			if (Count == 0)
-			{
-				Table->OnDataTableChanged().Broadcast();
-				ScopeCount.Remove(Table);
-			}
-		}
+}
 
-	private:
-		UDataTable* Table;
+UDataTable::FScopedDataTableChange::FScopedDataTableChange(UDataTable* InTable)
+	: Table(InTable)
+{
+	FScopeLock Lock(&CriticalSection);
+	int32& Count = ScopeCount.FindOrAdd(Table);
+	++Count;
+}
 
-		static TMap<UDataTable*, int32> ScopeCount;
-		static FCriticalSection CriticalSection;
-	};
+UDataTable::FScopedDataTableChange::~FScopedDataTableChange()
+{
+	FScopeLock Lock(&CriticalSection);
+	int32& Count = ScopeCount.FindChecked(Table);
+	--Count;
+	if (Count == 0)
+	{
+		Table->OnDataTableChanged().Broadcast();
+		ScopeCount.Remove(Table);
+	}
+}
 
-	TMap< UDataTable*, int32> FScopedDataTableChange::ScopeCount;
-	FCriticalSection FScopedDataTableChange::CriticalSection;
+TMap<UDataTable*, int32> UDataTable::FScopedDataTableChange::ScopeCount;
+FCriticalSection UDataTable::FScopedDataTableChange::CriticalSection;
 
-#define DATATABLE_CHANGE_SCOPE()	FScopedDataTableChange ActiveScope(this);
-}
+#define DATATABLE_CHANGE_SCOPE()	UDataTable::FScopedDataTableChange ActiveScope(this);
 
 UDataTable::UDataTable(const FObjectInitializer& ObjectInitializer)
 	: Super(ObjectInitializer)
