Make sure we track property changes in Push Model due to mapping or unmapping network objects.
This is only necessary for Client Side Demo Recording.

diff --git a/Engine/Source/Runtime/Engine/Private/RepLayout.cpp b/Engine/Source/Runtime/Engine/Private/RepLayout.cpp
--- a/Engine/Source/Runtime/Engine/Private/RepLayout.cpp
+++ b/Engine/Source/Runtime/Engine/Private/RepLayout.cpp
@@ -788,7 +788,7 @@ struct FNetPrivatePushIdHelper
 		}
 	}
 
-	static void MarkPropertyDirty(UObject* InObject, const int32 RepIndex)
+	static void MarkPropertyDirty(const UObject* const InObject, const int32 RepIndex)
 	{
 		MARK_PROPERTY_DIRTY_UNSAFE(InObject, RepIndex);
 	}
@@ -3755,7 +3755,7 @@ void FRepLayout::GatherGuidReferences(
 	}
 }
 
-bool FRepLayout::MoveMappedObjectToUnmapped_r(FGuidReferencesMap* GuidReferencesMap, const FNetworkGUID& GUID) const
+bool FRepLayout::MoveMappedObjectToUnmapped_r(FGuidReferencesMap* GuidReferencesMap, const FNetworkGUID& GUID, const UObject* const OwningObject) const
 {
 	bool bFoundGUID = false;
 
@@ -3767,7 +3767,7 @@ bool FRepLayout::MoveMappedObjectToUnmapped_r(FGuidReferencesMap* GuidReferences
 		{
 			check(Cmds[GuidReferences.CmdIndex].Type == ERepLayoutCmdType::DynamicArray);
 
-			if (MoveMappedObjectToUnmapped_r(GuidReferences.Array, GUID))
+			if (MoveMappedObjectToUnmapped_r(GuidReferences.Array, GUID, OwningObject))
 			{
 				bFoundGUID = true;
 			}
@@ -3779,6 +3779,13 @@ bool FRepLayout::MoveMappedObjectToUnmapped_r(FGuidReferencesMap* GuidReferences
 			GuidReferences.MappedDynamicGUIDs.Remove(GUID);
 			GuidReferences.UnmappedGUIDs.Add(GUID);
 			bFoundGUID = true;
+
+#if WITH_PUSH_MODEL
+			if (OwningObject)
+			{
+				FNetPrivatePushIdHelper::MarkPropertyDirty(OwningObject, GuidReferences.ParentIndex);
+			}
+#endif
 		}
 	}
 
@@ -3794,7 +3801,7 @@ bool FRepLayout::MoveMappedObjectToUnmapped(
 
 	if (!IsEmpty())
 	{
-		bFound = MoveMappedObjectToUnmapped_r(&RepState->GuidReferencesMap, GUID);
+		bFound = MoveMappedObjectToUnmapped_r(&RepState->GuidReferencesMap, GUID, Params.Object);
 
 		// Custom Delta Properties
 		if (LifetimeCustomPropertyState && Params.Object)
@@ -3921,6 +3928,10 @@ void FRepLayout::UpdateUnmappedObjects_r(
 				// Remove from unmapped guids list
 				UnmappedIt.RemoveCurrent();
 				bMappedSomeGUIDs = true;
+
+#if WITH_PUSH_MODEL
+				FNetPrivatePushIdHelper::MarkPropertyDirty(OriginalObject, GuidReferences.ParentIndex);
+#endif
 			}
 		}
 
@@ -7333,7 +7344,7 @@ bool FRepLayout::MoveMappedObjectToUnmappedForFastArray(FFastArrayDeltaSerialize
 	bool bFound = false;
 	for (auto& GuidReferencesPair : ArraySerializer.GuidReferencesMap_StructDelta)
 	{
-		bFound |= MoveMappedObjectToUnmapped_r(&GuidReferencesPair.Value, MoveToUnmapped);
+		bFound |= MoveMappedObjectToUnmapped_r(&GuidReferencesPair.Value, MoveToUnmapped, Params.DeltaSerializeInfo.Object);
 	}
 	return bFound;
 }
diff --git a/Engine/Source/Runtime/Engine/Public/Net/RepLayout.h b/Engine/Source/Runtime/Engine/Public/Net/RepLayout.h
--- a/Engine/Source/Runtime/Engine/Public/Net/RepLayout.h
+++ b/Engine/Source/Runtime/Engine/Public/Net/RepLayout.h
@@ -1606,7 +1606,10 @@ private:
 		TSet<FNetworkGUID>& OutReferencedGuids,
 		int32& OutTrackedGuidMemoryBytes) const;
 
-	bool MoveMappedObjectToUnmapped_r(FGuidReferencesMap* GuidReferencesMap, const FNetworkGUID& GUID) const;
+	bool MoveMappedObjectToUnmapped_r(
+		FGuidReferencesMap* GuidReferencesMap,
+		const FNetworkGUID& GUID,
+		const UObject* const OwningObject) const;
 
 	void UpdateUnmappedObjects_r(
 		FReceivingRepState* RESTRICT RepState, 
