Refactor ULightMapVirtualTexture2D into a separate file

diff --git a/Engine/Source/Runtime/Engine/Classes/VT/LightmapVirtualTexture.h b/Engine/Source/Runtime/Engine/Classes/VT/LightmapVirtualTexture.h
new file mode 100644
--- /dev/null
+++ b/Engine/Source/Runtime/Engine/Classes/VT/LightmapVirtualTexture.h
@@ -0,0 +1,32 @@
+// Copyright Epic Games, Inc. All Rights Reserved.
+
+#pragma once
+
+#include "CoreMinimal.h"
+#include "Engine/Texture2D.h"
+#include "LightmapVirtualTexture.generated.h"
+
+enum class ELightMapVirtualTextureType
+{
+	HqLayer0,
+	HqLayer1,
+	ShadowMask,
+	SkyOcclusion,
+	AOMaterialMask,
+
+	Count,
+};
+
+UCLASS(ClassGroup = Rendering)
+class ENGINE_API ULightMapVirtualTexture2D : public UTexture2D
+{
+	GENERATED_UCLASS_BODY()
+
+	UPROPERTY(EditAnywhere, Category = VirtualTexture)
+	TArray<int8> TypeToLayer;
+
+	void SetLayerForType(ELightMapVirtualTextureType InType, uint8 InLayer);
+	uint32 GetLayerForType(ELightMapVirtualTextureType InType) const;
+
+	inline bool HasLayerForType(ELightMapVirtualTextureType InType) const { return GetLayerForType(InType) != ~0u; }
+};
diff --git a/Engine/Source/Runtime/Engine/Classes/VT/VirtualTexture.h b/Engine/Source/Runtime/Engine/Classes/VT/VirtualTexture.h
--- a/Engine/Source/Runtime/Engine/Classes/VT/VirtualTexture.h
+++ b/Engine/Source/Runtime/Engine/Classes/VT/VirtualTexture.h
@@ -3,14 +3,9 @@
 #pragma once
 
 #include "CoreMinimal.h"
-#include "UObject/ObjectMacros.h"
-#include "RenderResource.h"
-#include "RenderCommandFence.h"
-#include "Engine/Texture.h"
-#include "Engine/Texture2D.h"
-#include "VirtualTexturing.h"
 #include "VirtualTexture.generated.h"
 
+/** Deprecated class */
 UCLASS(ClassGroup = Rendering)
 class ENGINE_API UVirtualTexture : public UObject
 {
@@ -18,33 +13,9 @@ class ENGINE_API UVirtualTexture : public UObject
 	virtual void Serialize(FArchive& Ar) override;
 };
 
+/** Deprecated class */
 UCLASS(ClassGroup = Rendering)
 class ENGINE_API ULightMapVirtualTexture : public UVirtualTexture
 {
 	GENERATED_UCLASS_BODY()
 };
-
-enum class ELightMapVirtualTextureType
-{
-	HqLayer0,
-	HqLayer1,
-	ShadowMask,
-	SkyOcclusion,
-	AOMaterialMask,
-
-	Count,
-};
-
-UCLASS(ClassGroup = Rendering)
-class ENGINE_API ULightMapVirtualTexture2D : public UTexture2D
-{
-	GENERATED_UCLASS_BODY()
-
-	UPROPERTY(EditAnywhere, Category = VirtualTexture)
-	TArray<int8> TypeToLayer;
-
-	void SetLayerForType(ELightMapVirtualTextureType InType, uint8 InLayer);
-	uint32 GetLayerForType(ELightMapVirtualTextureType InType) const;
-
-	inline bool HasLayerForType(ELightMapVirtualTextureType InType) const { return GetLayerForType(InType) != ~0u; }
-};
diff --git a/Engine/Source/Runtime/Engine/Private/LightMap.cpp b/Engine/Source/Runtime/Engine/Private/LightMap.cpp
--- a/Engine/Source/Runtime/Engine/Private/LightMap.cpp
+++ b/Engine/Source/Runtime/Engine/Private/LightMap.cpp
@@ -19,6 +19,7 @@
 #include "UObject/Package.h"
 #include "GameFramework/WorldSettings.h"
 #include "Engine/MapBuildDataRegistry.h"
+#include "VT/LightmapVirtualTexture.h"
 #include "VT/VirtualTexture.h"
 #include "EngineModule.h"
 #include "Misc/PackageName.h"
diff --git a/Engine/Source/Runtime/Engine/Private/LightmapUniformShaderParameters.cpp b/Engine/Source/Runtime/Engine/Private/LightmapUniformShaderParameters.cpp
--- a/Engine/Source/Runtime/Engine/Private/LightmapUniformShaderParameters.cpp
+++ b/Engine/Source/Runtime/Engine/Private/LightmapUniformShaderParameters.cpp
@@ -3,7 +3,7 @@
 #include "LightmapUniformShaderParameters.h"
 #include "SceneManagement.h"
 #include "LightMap.h"
-#include "VT/VirtualTexture.h"
+#include "VT/LightmapVirtualTexture.h"
 #include "UnrealEngine.h"
 
 IMPLEMENT_GLOBAL_SHADER_PARAMETER_STRUCT(FPrecomputedLightingUniformParameters, "PrecomputedLightingBuffer");
diff --git a/Engine/Source/Runtime/Engine/Private/SceneManagement.cpp b/Engine/Source/Runtime/Engine/Private/SceneManagement.cpp
--- a/Engine/Source/Runtime/Engine/Private/SceneManagement.cpp
+++ b/Engine/Source/Runtime/Engine/Private/SceneManagement.cpp
@@ -14,7 +14,7 @@
 #include "Engine/Engine.h"
 #include "Engine/LightMapTexture2D.h"
 #include "Engine/ShadowMapTexture2D.h"
-#include "VT/VirtualTexture.h"
+#include "VT/LightmapVirtualTexture.h"
 #include "UnrealEngine.h"
 
 static TAutoConsoleVariable<float> CVarLODTemporalLag(
diff --git a/Engine/Source/Runtime/Engine/Private/VT/LightmapVirtualTexture.cpp b/Engine/Source/Runtime/Engine/Private/VT/LightmapVirtualTexture.cpp
new file mode 100644
--- /dev/null
+++ b/Engine/Source/Runtime/Engine/Private/VT/LightmapVirtualTexture.cpp
@@ -0,0 +1,28 @@
+// Copyright Epic Games, Inc. All Rights Reserved.
+
+#include "VT/LightmapVirtualTexture.h"
+#include "LightMap.h"
+
+ULightMapVirtualTexture2D::ULightMapVirtualTexture2D(const FObjectInitializer& ObjectInitializer)
+	: Super(ObjectInitializer)
+{
+	VirtualTextureStreaming = true;
+	SetLayerForType(ELightMapVirtualTextureType::HqLayer0, 0u);
+	SetLayerForType(ELightMapVirtualTextureType::HqLayer1, 1u);
+}
+
+void ULightMapVirtualTexture2D::SetLayerForType(ELightMapVirtualTextureType InType, uint8 InLayer)
+{
+	const int TypeIndex = (int)InType;
+	while (TypeIndex >= TypeToLayer.Num())
+	{
+		TypeToLayer.Add(-1);
+	}
+	TypeToLayer[TypeIndex] = InLayer;
+}
+
+uint32 ULightMapVirtualTexture2D::GetLayerForType(ELightMapVirtualTextureType InType) const
+{
+	const int TypeIndex = (int)InType;
+	return (TypeIndex >= TypeToLayer.Num()) ? ~0u : (uint32)TypeToLayer[TypeIndex];
+}
diff --git a/Engine/Source/Runtime/Engine/Private/VT/VirtualTexture.cpp b/Engine/Source/Runtime/Engine/Private/VT/VirtualTexture.cpp
--- a/Engine/Source/Runtime/Engine/Private/VT/VirtualTexture.cpp
+++ b/Engine/Source/Runtime/Engine/Private/VT/VirtualTexture.cpp
@@ -1,25 +1,26 @@
 // Copyright Epic Games, Inc. All Rights Reserved.
 
 #include "VT/VirtualTexture.h"
-#include "VT/VirtualTextureBuiltData.h"
-#include "Serialization/MemoryWriter.h"
-#include "LightMap.h"
+
 #include "EngineModule.h"
-#include "RendererInterface.h"
 #include "FileCache/FileCache.h"
+#include "RendererInterface.h"
+#include "Serialization/Archive.h"
+#include "VT/VirtualTextureBuiltData.h"
 #include "EngineModule.h"
 
 static FAutoConsoleCommand GVTFlushAndEvictFileCacheCommand(
 	TEXT("r.VT.FlushAndEvictFileCache"),
-	TEXT("Flush both the virtual texture physcial page cache and disk file cachje"),
+	TEXT("Flush both the virtual texture physcial page cache and disk file cache"),
 	FConsoleCommandDelegate::CreateStatic(
 		[]()
-{
-	IFileCacheHandle::EvictAll();
-	GetRendererModule().FlushVirtualTextureCache();
-})
+	{
+		IFileCacheHandle::EvictAll();
+		GetRendererModule().FlushVirtualTextureCache();
+	})
 );
 
+// Deprecated serialization code
 namespace
 {
 	// Version used to serialize dummy data for UVirtualTexture
@@ -73,7 +74,10 @@ namespace
 	};
 }
 
-UVirtualTexture::UVirtualTexture(const FObjectInitializer& ObjectInitializer) : Super(ObjectInitializer) {}
+UVirtualTexture::UVirtualTexture(const FObjectInitializer& ObjectInitializer) 
+	: Super(ObjectInitializer) 
+{
+}
 
 void UVirtualTexture::Serialize(FArchive& Ar)
 {
@@ -164,28 +168,7 @@ void UVirtualTexture::Serialize(FArchive& Ar)
 	}
 }
 
-ULightMapVirtualTexture::ULightMapVirtualTexture(const FObjectInitializer& ObjectInitializer) : Super(ObjectInitializer) {}
-
-ULightMapVirtualTexture2D::ULightMapVirtualTexture2D(const FObjectInitializer& ObjectInitializer)
-	: Super(ObjectInitializer)
-{
-	VirtualTextureStreaming = true;
-	SetLayerForType(ELightMapVirtualTextureType::HqLayer0, 0u);
-	SetLayerForType(ELightMapVirtualTextureType::HqLayer1, 1u);
-}
-
-void ULightMapVirtualTexture2D::SetLayerForType(ELightMapVirtualTextureType InType, uint8 InLayer)
-{
-	const int TypeIndex = (int)InType;
-	while (TypeIndex >= TypeToLayer.Num())
-	{
-		TypeToLayer.Add(-1);
-	}
-	TypeToLayer[TypeIndex] = InLayer;
-}
-
-uint32 ULightMapVirtualTexture2D::GetLayerForType(ELightMapVirtualTextureType InType) const
+ULightMapVirtualTexture::ULightMapVirtualTexture(const FObjectInitializer& ObjectInitializer) 
+	: Super(ObjectInitializer) 
 {
-	const int TypeIndex = (int)InType;
-	return (TypeIndex >= TypeToLayer.Num()) ? ~0u : (uint32)TypeToLayer[TypeIndex];
 }
diff --git a/Engine/Source/Runtime/Engine/Private/VT/VirtualTextureChunkManager.cpp b/Engine/Source/Runtime/Engine/Private/VT/VirtualTextureChunkManager.cpp
--- a/Engine/Source/Runtime/Engine/Private/VT/VirtualTextureChunkManager.cpp
+++ b/Engine/Source/Runtime/Engine/Private/VT/VirtualTextureChunkManager.cpp
@@ -5,7 +5,6 @@
 #include "RHI.h"
 #include "ScreenRendering.h"
 #include "UploadingVirtualTexture.h"
-#include "VT/VirtualTexture.h"
 #include "VirtualTextureBuiltData.h"
 #include "BlockCodingHelpers.h"
 #include "CrunchCompression.h"
diff --git a/Engine/Source/Runtime/Engine/Private/VT/VirtualTextureDataBuilder.cpp b/Engine/Source/Runtime/Engine/Private/VT/VirtualTextureDataBuilder.cpp
--- a/Engine/Source/Runtime/Engine/Private/VT/VirtualTextureDataBuilder.cpp
+++ b/Engine/Source/Runtime/Engine/Private/VT/VirtualTextureDataBuilder.cpp
@@ -4,7 +4,6 @@
 
 #include "VirtualTextureDataBuilder.h"
 #include "Modules/ModuleManager.h"
-#include "VT/VirtualTexture.h"
 #include "Modules/ModuleManager.h"
 #include "CrunchCompression.h"
 #include "Async/TaskGraphInterfaces.h"
diff --git a/Engine/Source/Runtime/Engine/Private/World.cpp b/Engine/Source/Runtime/Engine/Private/World.cpp
--- a/Engine/Source/Runtime/Engine/Private/World.cpp
+++ b/Engine/Source/Runtime/Engine/Private/World.cpp
@@ -81,7 +81,7 @@
 #include "Engine/LevelScriptBlueprint.h"
 #include "Engine/DemoNetDriver.h"
 #include "Modules/ModuleManager.h"
-#include "VT/VirtualTexture.h" 
+#include "VT/LightmapVirtualTexture.h" 
 #include "Interfaces/Interface_PostProcessVolume.h"
 
 #include "Materials/MaterialParameterCollectionInstance.h"
diff --git a/Engine/Source/Runtime/Renderer/Private/LightMapDensityRendering.h b/Engine/Source/Runtime/Renderer/Private/LightMapDensityRendering.h
--- a/Engine/Source/Runtime/Renderer/Private/LightMapDensityRendering.h
+++ b/Engine/Source/Runtime/Renderer/Private/LightMapDensityRendering.h
@@ -19,8 +19,6 @@
 #include "ShaderBaseClasses.h"
 #include "SceneRendering.h"
 #include "Engine/LightMapTexture2D.h"
-#include "Runtime/Engine/Classes/VT/VirtualTexture.h"
-#include "Runtime/Renderer/Private/VT/VirtualTextureSpace.h"
 
 BEGIN_GLOBAL_SHADER_PARAMETER_STRUCT(FLightmapDensityPassUniformParameters, )
 	SHADER_PARAMETER_STRUCT(FSceneTexturesUniformParameters, SceneTextures)
diff --git a/Engine/Source/Runtime/Renderer/Private/LightMapRendering.cpp b/Engine/Source/Runtime/Renderer/Private/LightMapRendering.cpp
--- a/Engine/Source/Runtime/Renderer/Private/LightMapRendering.cpp
+++ b/Engine/Source/Runtime/Renderer/Private/LightMapRendering.cpp
@@ -8,10 +8,6 @@ LightMapRendering.cpp: Light map rendering implementations.
 #include "ScenePrivate.h"
 #include "PrecomputedVolumetricLightmap.h"
 
-#include "Runtime/Engine/Classes/VT/VirtualTexture.h"
-#include "Runtime/Renderer/Private/VT/VirtualTextureSpace.h"
-#include "VT/VirtualTextureSpace.h"
-
 IMPLEMENT_TYPE_LAYOUT(FUniformLightMapPolicyShaderParametersType);
 
 IMPLEMENT_GLOBAL_SHADER_PARAMETER_STRUCT(FIndirectLightingCacheUniformParameters, "IndirectLightingCache");
