Monolithic build fixes
#rnx

diff --git a/Engine/Source/Runtime/Engine/Classes/VT/VirtualTextureBuilder.h b/Engine/Source/Runtime/Engine/Classes/VT/VirtualTextureBuilder.h
--- a/Engine/Source/Runtime/Engine/Classes/VT/VirtualTextureBuilder.h
+++ b/Engine/Source/Runtime/Engine/Classes/VT/VirtualTextureBuilder.h
@@ -9,7 +9,6 @@
 #if WITH_EDITOR
 
 /** Description object used to build the contents of a UVirtualTextureBuilder. */
-USTRUCT(BlueprintType)
 struct FVirtualTextureBuildDesc
 {
 	GENERATED_USTRUCT_BODY()
@@ -45,20 +44,19 @@ class ENGINE_API UVirtualTextureBuilder : public UObject
 {
 public:
 	GENERATED_UCLASS_BODY()
-		~UVirtualTextureBuilder();
+	~UVirtualTextureBuilder();
 
 	/** The UTexture object. */
 	UPROPERTY(VisibleAnywhere, BlueprintReadOnly, Category = Texture)
-		class UVirtualTexture2D* Texture;
+	class UVirtualTexture2D* Texture;
 
 	/** Some client defined hash of that defines how the Texture was built. */
 	UPROPERTY()
-		uint32 BuildHash;
+	uint32 BuildHash;
 
 #if WITH_EDITOR
-	UFUNCTION(BlueprintCallable, Category = Texture)
-		/** Creates a new UVirtualTexture2D and stores it in the contained Texture. */
-		void BuildTexture(FVirtualTextureBuildDesc const& BuildDesc);
+	/** Creates a new UVirtualTexture2D and stores it in the contained Texture. */
+	void BuildTexture(FVirtualTextureBuildDesc const& BuildDesc);
 #endif
 
 protected:
