Build fix
#rnx

diff --git a/Engine/Source/Runtime/Engine/Classes/VT/VirtualTextureBuilder.h b/Engine/Source/Runtime/Engine/Classes/VT/VirtualTextureBuilder.h
--- a/Engine/Source/Runtime/Engine/Classes/VT/VirtualTextureBuilder.h
+++ b/Engine/Source/Runtime/Engine/Classes/VT/VirtualTextureBuilder.h
@@ -11,9 +11,7 @@
 /** Description object used to build the contents of a UVirtualTextureBuilder. */
 struct FVirtualTextureBuildDesc
 {
-	GENERATED_USTRUCT_BODY()
-
-		uint32 BuildHash = 0;
+	uint32 BuildHash = 0;
 
 	int32 LayerCount = 0;
 	TArray<ETextureSourceFormat, TInlineAllocator<4>> LayerFormats;
