Move streaming low mips data outside of URuntimeVirtualTexture and into a UVirtualTextureBuilder.
This requires URuntimeVirtualTextureComponents to be updated.

diff --git a/Engine/Source/Editor/VirtualTexturingEditor/Private/RuntimeVirtualTextureBuild.cpp b/Engine/Source/Editor/VirtualTexturingEditor/Private/RuntimeVirtualTextureBuild.cpp
--- a/Engine/Source/Editor/VirtualTexturingEditor/Private/RuntimeVirtualTextureBuild.cpp
+++ b/Engine/Source/Editor/VirtualTexturingEditor/Private/RuntimeVirtualTextureBuild.cpp
@@ -121,13 +121,12 @@ namespace RuntimeVirtualTexture
 			return false;
 		}
 
-		URuntimeVirtualTexture* RuntimeVirtualTexture = InComponent->GetVirtualTexture();
-		if (RuntimeVirtualTexture == nullptr)
+		if (InComponent->GetVirtualTexture() == nullptr || InComponent->GetStreamingTexture() == nullptr)
 		{
 			return false;
 		}
 
-		if (RuntimeVirtualTexture->GetStreamLowMips() <= 0)
+		if (InComponent->GetStreamLowMips() <= 0)
 		{
 			return false;
 		}
@@ -156,7 +155,7 @@ namespace RuntimeVirtualTexture
 		const int32 TextureSizeX = VTDesc.WidthInBlocks * VTDesc.BlockWidthInTiles * TileSize;
 		const int32 TextureSizeY = VTDesc.HeightInBlocks * VTDesc.BlockHeightInTiles * TileSize;
 		const int32 MaxLevel = (int32)FMath::CeilLogTwo(FMath::Max(VTDesc.BlockWidthInTiles, VTDesc.BlockHeightInTiles));
-		const int32 RenderLevel = FMath::Max(MaxLevel - RuntimeVirtualTexture->GetStreamLowMips() + 1, 0);
+		const int32 RenderLevel = FMath::Max(MaxLevel - InComponent->GetStreamLowMips() + 1, 0);
 		const int32 ImageSizeX = FMath::Max(TileSize, TextureSizeX >> RenderLevel);
 		const int32 ImageSizeY = FMath::Max(TileSize, TextureSizeY >> RenderLevel);
 		const int32 NumTilesX = ImageSizeX / TileSize;
@@ -276,7 +275,8 @@ namespace RuntimeVirtualTexture
 			}
 		}
 
-		ReleaseResourceAndFlush(&RenderTileResources);
+		BeginReleaseResource(&RenderTileResources);
+		FlushRenderingCommands();
 
 		if (Task.ShouldCancel())
 		{
@@ -285,9 +285,8 @@ namespace RuntimeVirtualTexture
 
 		// Place final pixel data into the runtime virtual texture
 		Task.EnterProgressFrame(TaskWorkBuildBulkData);
-		RuntimeVirtualTexture->Modify();
-		RuntimeVirtualTexture->InitializeStreamingTexture(ImageSizeX, ImageSizeY, (uint8*)FinalPixels.GetData());
-		RuntimeVirtualTexture->PostEditChange();
+
+		InComponent->InitializeStreamingTexture(ImageSizeX, ImageSizeY, (uint8*)FinalPixels.GetData());
 
 		return true;
 	}
diff --git a/Engine/Source/Editor/VirtualTexturingEditor/Private/RuntimeVirtualTextureDetailsCustomization.cpp b/Engine/Source/Editor/VirtualTexturingEditor/Private/RuntimeVirtualTextureDetailsCustomization.cpp
--- a/Engine/Source/Editor/VirtualTexturingEditor/Private/RuntimeVirtualTextureDetailsCustomization.cpp
+++ b/Engine/Source/Editor/VirtualTexturingEditor/Private/RuntimeVirtualTextureDetailsCustomization.cpp
@@ -139,7 +139,6 @@ void FRuntimeVirtualTextureDetailsCustomization::CustomizeDetails(IDetailLayoutB
 	DetailBuilder.GetProperty(FName(TEXT("MaterialType")))->SetOnPropertyValueChanged(FSimpleDelegate::CreateSP(this, &FRuntimeVirtualTextureDetailsCustomization::RefreshDetails));
 	DetailBuilder.GetProperty(FName(TEXT("bCompressTextures")))->SetOnPropertyValueChanged(FSimpleDelegate::CreateSP(this, &FRuntimeVirtualTextureDetailsCustomization::RefreshDetails));
 	DetailBuilder.GetProperty(FName(TEXT("RemoveLowMips")))->SetOnPropertyValueChanged(FSimpleDelegate::CreateSP(this, &FRuntimeVirtualTextureDetailsCustomization::RefreshDetails));
-	DetailBuilder.GetProperty(FName(TEXT("StreamLowMips")))->SetOnPropertyValueChanged(FSimpleDelegate::CreateSP(this, &FRuntimeVirtualTextureDetailsCustomization::RefreshDetails));
 
 	// Initialize text blocks
 	RefreshDetails();
@@ -201,16 +200,13 @@ void FRuntimeVirtualTextureComponentDetailsCustomization::CustomizeDetails(IDeta
 		SNew(SHorizontalBox)
 		
 		+ SHorizontalBox::Slot()
-		.FillWidth(5.0f)
 		[
 			SourceActorValue->CreatePropertyValueWidget()
 		]
 
 		+ SHorizontalBox::Slot()
-		.FillWidth(1.0f)
 		[
 			SNew(SWrapBox)
-			.UseAllottedWidth(true)
 			
 			+ SWrapBox::Slot()
 			.Padding(FMargin(0.0f, 2.0f, 2.0f, 0.0f))
@@ -244,11 +240,11 @@ void FRuntimeVirtualTextureComponentDetailsCustomization::CustomizeDetails(IDeta
 	];
 
 	// Use existing property to add build button
-	TSharedPtr<IPropertyHandle> LowMipsValue = DetailBuilder.GetProperty("bUseStreamingLowMips");
+	TSharedPtr<IPropertyHandle> LowMipsValue = DetailBuilder.GetProperty("StreamLowMips");
 	DetailBuilder.HideProperty(LowMipsValue);
 
-	IDetailCategoryBuilder& VirtualTextureCategory = DetailBuilder.EditCategory("VirtualTexture", FText::GetEmpty());
-	VirtualTextureCategory.AddCustomRow(LowMipsValue->GetPropertyDisplayName(), true)
+	IDetailCategoryBuilder& VirtualTextureCategory = DetailBuilder.EditCategory("VirtualTextureBuild", FText::GetEmpty());
+	VirtualTextureCategory.AddCustomRow(LowMipsValue->GetPropertyDisplayName())
 	.NameContent()
 	[
 		LowMipsValue->CreatePropertyNameWidget()
@@ -259,6 +255,7 @@ void FRuntimeVirtualTextureComponentDetailsCustomization::CustomizeDetails(IDeta
 		SNew(SHorizontalBox)
 
 		+ SHorizontalBox::Slot()
+		.Padding(5.f, 0.f)
 		[
 			LowMipsValue->CreatePropertyValueWidget()
 		]
diff --git a/Engine/Source/Runtime/Engine/Classes/Components/RuntimeVirtualTextureComponent.h b/Engine/Source/Runtime/Engine/Classes/Components/RuntimeVirtualTextureComponent.h
--- a/Engine/Source/Runtime/Engine/Classes/Components/RuntimeVirtualTextureComponent.h
+++ b/Engine/Source/Runtime/Engine/Classes/Components/RuntimeVirtualTextureComponent.h
@@ -7,6 +7,7 @@
 #include "RuntimeVirtualTextureComponent.generated.h"
 
 class URuntimeVirtualTexture;
+class UVirtualTextureBuilder;
 
 /** Component used to place a URuntimeVirtualTexture in the world. */
 UCLASS(Blueprintable, ClassGroup = Rendering, collapsecategories, hidecategories = (Activation, Collision, Cooking, Mobility, LOD, Object, Physics, Rendering), editinlinenew)
@@ -19,30 +20,50 @@ protected:
 	UPROPERTY(EditAnywhere, BlueprintReadOnly, DuplicateTransient, Category = VirtualTexture)
 	URuntimeVirtualTexture* VirtualTexture = nullptr;
 
-	/** Use any streaming low mips stored in the virtual texture object. */
-	UPROPERTY(EditAnywhere, AdvancedDisplay, Category = VirtualTexture)
-	bool bUseStreamingLowMips = true;
-
-	/** Use any streaming low mips when rendering in editor. Set true to view and debug the baked streaming low mips. */
-	UPROPERTY(EditAnywhere, AdvancedDisplay, Category = VirtualTexture, meta = (EditCondition = "bUseStreamingLowMips"))
-	bool bUseStreamingLowMipsInEditor = false;
+	/** Texture object containing streamed low mips. */
+	UPROPERTY(EditAnywhere, BlueprintReadOnly, Category = VirtualTexture, meta = (DisplayName = "Streaming Mips"))
+	UVirtualTextureBuilder* StreamingTexture = nullptr;
 
 	/** Actor to copy the bounds from to set up the transform. */
 	UPROPERTY(EditAnywhere, DuplicateTransient, Category = TransformFromBounds, meta = (DisplayName = "Source Actor"))
 	AActor* BoundsSourceActor = nullptr;
 
+	/** Number of low mips to serialize and stream for the virtual texture. This can reduce rendering update cost. */
+	UPROPERTY(EditAnywhere, Category = VirtualTextureBuild, meta = (UIMin = "0", UIMax = "6", DisplayName = "Number of low mips to stream to the virtual texture"))
+	int32 StreamLowMips = 0;
+
+	/** Enable Crunch compression. ZLib compression is used when Crunch is disabled. */
+	UPROPERTY(EditAnywhere, AdvancedDisplay, Category = VirtualTextureBuild, meta = (DisplayName = "Enable Crunch compression"))
+	bool bEnableCompressCrunch = false;
+
+	/** Use any streaming low mips when rendering in editor. Set true to view and debug the baked streaming low mips. */
+	UPROPERTY(EditAnywhere, AdvancedDisplay, Category = VirtualTextureBuild)
+	bool bUseStreamingLowMipsInEditor = false;
+
 public:
 	/** Get the runtime virtual texture object on this component. */
 	URuntimeVirtualTexture* GetVirtualTexture() const { return VirtualTexture; }
 
+	/** Get the streaming virtual texture object on this component. */
+	class UVirtualTextureBuilder* GetStreamingTexture() const { return StreamingTexture; }
+
 	/** Get the runtime virtual texture UV to World transform on this component. */
 	UFUNCTION(BlueprintPure, Category = VirtualTexture)
 	FTransform GetVirtualTextureTransform() const;
 
-	/** Get if we want use any streaming low mips in the runtime virtual texture set on this component. */
+	/** Public getter for virtual texture streaming low mips */
+	int32 GetStreamLowMips() const { return FMath::Clamp(StreamLowMips, 0, 6); }
+
+	/** Get if we want to use any streaming low mips on this component. */
 	bool IsStreamingLowMips() const;
 
+	/** Public getter for crunch compression flag. */
+	bool IsCrunchCompressed() const { return bEnableCompressCrunch; }
+
 #if WITH_EDITOR
+	/** Initialize the low mip streaming texture with the passed in size and data. */
+	void InitializeStreamingTexture(uint32 InSizeX, uint32 InSizeY, uint8* InData);
+
 	/** Copy the rotation from BoundsSourceActor to this component. Called by our UI details customization. */
 	void SetRotation();
 
@@ -62,6 +83,11 @@ protected:
 	virtual FBoxSphereBounds CalcBounds(const FTransform& LocalToWorld) const override;
 	//~ End USceneComponent Interface
 
+	/** Calculate a hash used to determine if the StreamingTexture contents are valid for use. The hash doesn't include whether the contents are up to date. */
+	uint32 CalculateStreamingTextureSettingsHash() const;
+	/** Returns true if the StreamingTexure contents are valid for use. */
+	bool StreamingTextureIsValid() const;
+
 public:
 	/** Scene proxy object. Managed by the scene but stored here. */
 	class FRuntimeVirtualTextureSceneProxy* SceneProxy;
diff --git a/Engine/Source/Runtime/Engine/Classes/VT/RuntimeVirtualTexture.h b/Engine/Source/Runtime/Engine/Classes/VT/RuntimeVirtualTexture.h
--- a/Engine/Source/Runtime/Engine/Classes/VT/RuntimeVirtualTexture.h
+++ b/Engine/Source/Runtime/Engine/Classes/VT/RuntimeVirtualTexture.h
@@ -64,22 +64,14 @@ protected:
 	UPROPERTY(EditAnywhere, BluePrintGetter = GetTileBorderSize, Category = Size, meta = (UIMin = "0", UIMax = "4", DisplayName = "Border padding for each virtual texture tile"))
 	int32 TileBorderSize = 2; // 4
 
-	/** Number of low mips to serialize and stream for the virtual texture. This can reduce rendering update cost. */
-	UPROPERTY(EditAnywhere, Category = LowMips, meta = (UIMin = "0", UIMax = "6", DisplayName = "Number of low mips to stream to the virtual texture"))
-	int32 StreamLowMips = 0;
-
-	/** Texture object containing streamed low mips. */
-	UPROPERTY(VisibleAnywhere, BlueprintReadOnly, Category = LowMips, meta = (DisplayName = "Streaming low mip texture"))
-	class URuntimeVirtualTextureStreamingProxy* StreamingTexture;
-
-	/** Enable Crunch compression of streamed low mips. ZLib compression is used when Crunch is disabled. */
-	UPROPERTY(EditAnywhere, AdvancedDisplay, Category = LowMips, meta = (DisplayName = "Enable Crunch compression"))
-	bool bEnableCompressCrunch = false;
-
 	/** Number of low mips to cut from the virtual texture. This can reduce peak virtual texture update cost but will also increase the probability of mip shimmering. */
-	UPROPERTY(EditAnywhere, AdvancedDisplay, Category = LowMips, meta = (UIMin = "0", UIMax = "6", DisplayName = "Number of low mips to remove from the virtual texture"))
+	UPROPERTY(EditAnywhere, AdvancedDisplay, Category = Layout, meta = (UIMin = "0", UIMax = "6", DisplayName = "Number of low mips to remove from the virtual texture"))
 	int32 RemoveLowMips = 0;
 
+	/** Deprecated texture object containing streamed low mips. */
+	UPROPERTY()
+	class URuntimeVirtualTextureStreamingProxy* StreamingTexture_DEPRECATED = nullptr;
+
 public:
 	/** Public getter for enabled status */
 	bool GetEnabled() { return bEnable; }
@@ -101,17 +93,19 @@ public:
 	UFUNCTION(BlueprintGetter)
 	int32 GetTileBorderSize() const { return 2 * FMath::Clamp(TileBorderSize, 0, 4); }
 	
+	/** Get if this virtual texture uses compressed texture formats. */
+	bool GetCompressTextures() const { return bCompressTextures; }
 	/** Public getter for virtual texture removed low mips */
 	int32 GetRemoveLowMips() const { return FMath::Clamp(RemoveLowMips, 0, 5); }
-	/** Public getter for virtual texture streaming low mips */
-	int32 GetStreamLowMips() const { return FMath::Clamp(StreamLowMips, 0, 6); }
 	/** Public getter for virtual texture using single physical space flag. */
 	bool GetSinglePhysicalSpace() const { return bSinglePhysicalSpace; }
 
+#if WITH_EDITOR
 	/** Returns an approximate estimated value for the memory used by the page table texture. */
 	int32 GetEstimatedPageTableTextureMemoryKb() const;
 	/** Returns an approximate estimated value for the memory used by the physical texture. */
 	int32 GetEstimatedPhysicalTextureMemoryKb() const;
+#endif
 
 	/** Get virtual texture description based on the properties of this object and the passed in volume transform. */
 	void GetProducerDescription(FVTProducerDescription& OutDesc, FTransform const& VolumeToWorld) const;
@@ -143,18 +137,6 @@ public:
 	/** Getter for the shader uniform parameters. */
 	FVector4 GetUniformParameter(int32 Index) const;
 
-#if WITH_EDITOR
-	/** Get the streaming texture */
-	URuntimeVirtualTextureStreamingProxy* GetStreamingTexture() const { return StreamingTexture; }
-	/** Get a hash of the current state to use for streaming texture invalidation. */
-	uint32 GetStreamingTextureBuildHash() const;
-	/** Initialize the low mip streaming texture with the passed in size and data. */
-	void InitializeStreamingTexture(uint32 InSizeX, uint32 InSizeY, uint8* InData);
-#endif
-
-	/** Create the streaming texture producer to wrap an existing producer. */
-	IVirtualTexture* CreateStreamingTextureProducer(IVirtualTexture* InProducer, int32 InMaxLevel, int32& OutTransitionLevel) const;
-
 protected:
 	/** Initialize the render resources. This kicks off render thread work. */
 	void InitResource(IVirtualTexture* InProducer, FTransform const& VolumeToWorld);
@@ -163,7 +145,6 @@ protected:
 
 	//~ Begin UObject Interface.
 	virtual void GetAssetRegistryTags(TArray<FAssetRegistryTag>& OutTags) const override;
-	virtual void Serialize(FArchive& Ar) override;
 	virtual void PostLoad() override;
 #if WITH_EDITOR
 	virtual void PostEditChangeProperty(FPropertyChangedEvent& PropertyChangedEvent) override;
@@ -179,3 +160,16 @@ private:
 	/** Material uniform parameter used to pack world height. */
 	FVector4 WorldHeightUnpackParameter;
 };
+
+class UVirtualTexture2D;
+
+namespace RuntimeVirtualTexture
+{
+	/** Helper function to wrap a runtime virtual texture producer with a streaming producer. */
+	ENGINE_API IVirtualTexture* CreateStreamingTextureProducer(
+		IVirtualTexture* InProducer,
+		URuntimeVirtualTexture* InBaseVirtualTexure,
+		UVirtualTexture2D* InStreamingTexture,
+		int32 InMaxLevel, int32 InNumMips,
+		int32& OutTransitionLevel);
+}
diff --git a/Engine/Source/Runtime/Engine/Classes/VT/RuntimeVirtualTextureStreamingProxy.h b/Engine/Source/Runtime/Engine/Classes/VT/RuntimeVirtualTextureStreamingProxy.h
deleted file mode 100644
diff --git a/Engine/Source/Runtime/Engine/Classes/VT/VirtualTexture.h b/Engine/Source/Runtime/Engine/Classes/VT/VirtualTexture.h
--- a/Engine/Source/Runtime/Engine/Classes/VT/VirtualTexture.h
+++ b/Engine/Source/Runtime/Engine/Classes/VT/VirtualTexture.h
@@ -22,6 +22,14 @@ class ENGINE_API ULightMapVirtualTexture : public UVirtualTexture
 	GENERATED_UCLASS_BODY()
 };
 
+/** Deprecated class. */
+UCLASS(ClassGroup = Rendering)
+class URuntimeVirtualTextureStreamingProxy : public UTexture2D
+{
+	GENERATED_UCLASS_BODY()
+};
+
+
 /**
  * Virtual Texture with locally configurable build settings.
  * A raw UTexture2D can also represent a Virtual Texture but uses the one and only per-project build settings.
diff --git a/Engine/Source/Runtime/Engine/Private/Components/RuntimeVirtualTextureComponent.cpp b/Engine/Source/Runtime/Engine/Private/Components/RuntimeVirtualTextureComponent.cpp
--- a/Engine/Source/Runtime/Engine/Private/Components/RuntimeVirtualTextureComponent.cpp
+++ b/Engine/Source/Runtime/Engine/Private/Components/RuntimeVirtualTextureComponent.cpp
@@ -5,6 +5,8 @@
 #include "Components/PrimitiveComponent.h"
 #include "GameFramework/Actor.h"
 #include "SceneInterface.h"
+#include "VT/RuntimeVirtualTexture.h"
+#include "VT/VirtualTextureBuilder.h"
 
 URuntimeVirtualTextureComponent::URuntimeVirtualTextureComponent(const FObjectInitializer& ObjectInitializer)
 	: Super(ObjectInitializer)
@@ -61,17 +63,106 @@ FTransform URuntimeVirtualTextureComponent::GetVirtualTextureTransform() const
 	return FTransform(FVector(-0.5f, -0.5f, 0.f)) * GetComponentTransform();
 }
 
+uint32 URuntimeVirtualTextureComponent::CalculateStreamingTextureSettingsHash() const
+{
+	// Shouldn't need to call this when there is VirtualTexture == nullptr
+	check(VirtualTexture != nullptr);
+
+	// If a setting change can cause the streaming texture to no longer be valid then it should be included in this hash.
+	union FPackedSettings
+	{
+		uint32 PackedValue;
+		struct
+		{
+			uint32 MaterialType : 4;
+			uint32 TileSize : 12;
+			uint32 TileBorderSize : 4;
+			uint32 StreamLowMips : 4;
+			uint32 CompressTextures : 1;
+			uint32 SinglePhysicalSpace : 1;
+			uint32 EnableCompressCrunch : 1;
+		};
+	};
+
+	FPackedSettings Settings;
+	Settings.PackedValue = 0;
+	Settings.MaterialType = (uint32)VirtualTexture->GetMaterialType();
+	Settings.TileSize = (uint32)VirtualTexture->GetTileSize();
+	Settings.TileBorderSize = (uint32)VirtualTexture->GetTileBorderSize();
+	Settings.StreamLowMips = (uint32)StreamLowMips;
+	Settings.CompressTextures = (uint32)VirtualTexture->GetCompressTextures();
+	Settings.SinglePhysicalSpace = (uint32)VirtualTexture->GetSinglePhysicalSpace();
+	Settings.EnableCompressCrunch = (uint32)bEnableCompressCrunch;
+
+	return Settings.PackedValue;
+}
+
+bool URuntimeVirtualTextureComponent::StreamingTextureIsValid() const
+{
+	return VirtualTexture != nullptr && StreamingTexture != nullptr && StreamingTexture->Texture != nullptr && StreamingTexture->BuildHash == CalculateStreamingTextureSettingsHash();
+}
+
 bool URuntimeVirtualTextureComponent::IsStreamingLowMips() const
 {
 #if WITH_EDITOR
-	return bUseStreamingLowMips && bUseStreamingLowMipsInEditor;
-#else
-	return bUseStreamingLowMips;
+	if (!bUseStreamingLowMipsInEditor)
+	{
+		return false;
+	}
 #endif
+	return StreamLowMips > 0 && StreamingTextureIsValid();
 }
 
 #if WITH_EDITOR
 
+void URuntimeVirtualTextureComponent::InitializeStreamingTexture(uint32 InSizeX, uint32 InSizeY, uint8* InData)
+{
+	if (VirtualTexture != nullptr && StreamingTexture != nullptr)
+	{
+		VirtualTexture->Modify();
+		StreamingTexture->Modify();
+
+		// Release current runtime virtual texture producer.
+		// It may reference data inside the old StreamingTexture which could be garbage collected any time from now.
+		VirtualTexture->Release();
+
+		FVirtualTextureBuildDesc BuildDesc;
+		BuildDesc.bSinglePhysicalSpace = VirtualTexture->GetSinglePhysicalSpace();
+
+		BuildDesc.TileSize = VirtualTexture->GetTileSize();
+		BuildDesc.TileBorderSize = VirtualTexture->GetTileBorderSize();
+		BuildDesc.bCrunchCompressed = bEnableCompressCrunch;
+
+		BuildDesc.LayerCount = VirtualTexture->GetLayerCount();
+		check(BuildDesc.LayerCount <= RuntimeVirtualTexture::MaxTextureLayers);
+		BuildDesc.LayerFormats.AddDefaulted(BuildDesc.LayerCount);
+		BuildDesc.LayerFormatSettings.AddDefaulted(BuildDesc.LayerCount);
+
+		for (int32 Layer = 0; Layer < BuildDesc.LayerCount; Layer++)
+		{
+			const EPixelFormat LayerFormat = VirtualTexture->GetLayerFormat(Layer);
+			BuildDesc.LayerFormats[Layer] = LayerFormat == PF_G16 ? TSF_G16 : TSF_BGRA8;
+
+			BuildDesc.LayerFormatSettings[Layer].CompressionSettings = LayerFormat == PF_BC5 ? TC_Normalmap : TC_Default;
+			BuildDesc.LayerFormatSettings[Layer].CompressionNone = LayerFormat == PF_B8G8R8A8 || LayerFormat == PF_G16;
+			BuildDesc.LayerFormatSettings[Layer].CompressionNoAlpha = LayerFormat == PF_DXT1 || LayerFormat == PF_BC5;
+			BuildDesc.LayerFormatSettings[Layer].CompressionYCoCg = VirtualTexture->IsLayerYCoCg(Layer);
+			BuildDesc.LayerFormatSettings[Layer].SRGB = VirtualTexture->IsLayerSRGB(Layer);
+		}
+
+		BuildDesc.BuildHash = CalculateStreamingTextureSettingsHash();
+
+		BuildDesc.InSizeX = InSizeX;
+		BuildDesc.InSizeY = InSizeY;
+		BuildDesc.InData = InData;
+
+		StreamingTexture->BuildTexture(BuildDesc);
+
+		// Trigger refresh of the runtime virtual texture producer.
+		VirtualTexture->PostEditChange();
+	}
+}
+
 void URuntimeVirtualTextureComponent::SetRotation()
 {
 	if (BoundsSourceActor != nullptr)
diff --git a/Engine/Source/Runtime/Engine/Private/VT/RuntimeVirtualTexture.cpp b/Engine/Source/Runtime/Engine/Private/VT/RuntimeVirtualTexture.cpp
--- a/Engine/Source/Runtime/Engine/Private/VT/RuntimeVirtualTexture.cpp
+++ b/Engine/Source/Runtime/Engine/Private/VT/RuntimeVirtualTexture.cpp
@@ -4,9 +4,8 @@
 
 #include "EngineModule.h"
 #include "RendererInterface.h"
-#include "VT/VirtualTextureBuildSettings.h"
 #include "VT/RuntimeVirtualTextureNotify.h"
-#include "VT/RuntimeVirtualTextureStreamingProxy.h"
+#include "VT/VirtualTexture.h"
 #include "VT/VirtualTextureLevelRedirector.h"
 #include "VT/VirtualTextureScalability.h"
 #include "VT/UploadingVirtualTexture.h"
@@ -167,35 +166,6 @@ private:
 };
 
 
-URuntimeVirtualTextureStreamingProxy::URuntimeVirtualTextureStreamingProxy(const FObjectInitializer& ObjectInitializer)
-	: UTexture2D(ObjectInitializer)
-	, BuildHash(0)
-{
-}
-
-void URuntimeVirtualTextureStreamingProxy::GetVirtualTextureBuildSettings(FVirtualTextureBuildSettings& OutSettings) const
-{
-	OutSettings = Settings;
-}
-
-#if WITH_EDITOR
-
-void URuntimeVirtualTextureStreamingProxy::BeginCacheForCookedPlatformData(const ITargetPlatform* TargetPlatform)
-{
-	// Even though we skip the cook of this object for non VT platforms in URuntimeVirtualTexture::Serialize()
-	// we still load the object at cook time and kick off the DDC build. This will trigger an error in the texture DDC code.
-	// Either we need to make the DDC code more robust for non VT platforms or we can skip the process here...
-	if (!UseVirtualTexturing(GMaxRHIFeatureLevel, TargetPlatform))
-	{
-		return;
-	}
-
-	Super::BeginCacheForCookedPlatformData(TargetPlatform);
-}
-
-#endif
-
-
 URuntimeVirtualTexture::URuntimeVirtualTexture(const FObjectInitializer& ObjectInitializer)
 	: Super(ObjectInitializer)
 {
@@ -375,6 +345,8 @@ bool URuntimeVirtualTexture::IsLayerYCoCg(int32 LayerIndex) const
 	return false;
 }
 
+#if WITH_EDITOR
+
 int32 URuntimeVirtualTexture::GetEstimatedPageTableTextureMemoryKb() const
 {
 	//todo[vt]: Estimate memory usage
@@ -387,6 +359,8 @@ int32 URuntimeVirtualTexture::GetEstimatedPhysicalTextureMemoryKb() const
 	return 0;
 }
 
+#endif
+
 FVirtualTextureProducerHandle URuntimeVirtualTexture::GetProducerHandle() const
 {
 	return Resource->GetProducerHandle();
@@ -456,22 +430,6 @@ void URuntimeVirtualTexture::GetAssetRegistryTags(TArray<FAssetRegistryTag>& Out
 	OutTags.Add(FAssetRegistryTag("TileBorderSize", FString::FromInt(GetTileBorderSize()), FAssetRegistryTag::TT_Numerical));
 }
 
-void URuntimeVirtualTexture::Serialize(FArchive& Ar)
-{
-	if (Ar.IsCooking() && Ar.IsSaving() && !UseVirtualTexturing(GMaxRHIFeatureLevel, Ar.CookingTarget()))
-	{
-		// Clear StreamingTexture during cook for platforms that don't support virtual texturing
-		URuntimeVirtualTextureStreamingProxy* StreamingTextureBackup = StreamingTexture;
-		StreamingTexture = nullptr;
-		Super::Serialize(Ar);
-		StreamingTexture = StreamingTextureBackup;
-	}
-	else
-	{
-		Super::Serialize(Ar);
-	}
-}
-
 void URuntimeVirtualTexture::PostLoad()
 {
 	// Convert Size_DEPRECATED to TileCount
@@ -489,6 +447,9 @@ void URuntimeVirtualTexture::PostLoad()
 		MaterialType = ERuntimeVirtualTextureMaterialType::BaseColor_Normal_Specular;
 	}
 
+	// Remove StreamingTexture_DEPRECATED
+	StreamingTexture_DEPRECATED = nullptr;
+
 	Super::PostLoad();
 }
 
@@ -498,109 +459,44 @@ void URuntimeVirtualTexture::PostEditChangeProperty(FPropertyChangedEvent& Prope
 {
 	Super::PostEditChangeProperty(PropertyChangedEvent);
 
-	// Invalidate StreamingTexture if it is no longer compatible with this object
-	if (StreamingTexture != nullptr && StreamingTexture->BuildHash != GetStreamingTextureBuildHash())
-	{
-		StreamingTexture = nullptr;
-	}
-
 	RuntimeVirtualTexture::NotifyComponents(this);
 	RuntimeVirtualTexture::NotifyPrimitives(this);
 }
 
-uint32 URuntimeVirtualTexture::GetStreamingTextureBuildHash() const
-{
-	union FPackedSettings
-	{
-		uint32 PackedValue;
-		struct
-		{
-			uint32 MaterialType : 3;
-			uint32 CompressTextures : 1;
-			uint32 SinglePhysicalSpace : 1;
-			uint32 TileSize : 4;
-			uint32 TileBorderSize : 4;
-			uint32 StreamLowMips : 4;
-			uint32 EnableCompressCrunch : 1;
-		};
-	};
-
-	FPackedSettings Settings;
-	Settings.PackedValue = 0;
-	Settings.MaterialType = (uint32)MaterialType;
-	Settings.CompressTextures = (uint32)bCompressTextures;
-	Settings.SinglePhysicalSpace = (uint32)bSinglePhysicalSpace;
-	Settings.TileSize = (uint32)TileSize;
-	Settings.TileBorderSize = (uint32)TileBorderSize;
-	Settings.StreamLowMips = (uint32)GetStreamLowMips();
-	Settings.EnableCompressCrunch = (uint32)bEnableCompressCrunch;
-
-	return Settings.PackedValue;
-}
-
-void URuntimeVirtualTexture::InitializeStreamingTexture(uint32 InSizeX, uint32 InSizeY, uint8* InData)
-{
-	// Release current producer. 
-	// It may reference data inside StreamingTexture which could be garbage collected any time from now.
-	InitNullResource();
-
-	StreamingTexture = NewObject<URuntimeVirtualTextureStreamingProxy>(GetOutermost(), TEXT("StreamingTexture"));
-	StreamingTexture->VirtualTextureStreaming = true;
-	StreamingTexture->bSinglePhysicalSpace = bSinglePhysicalSpace;
-
-	StreamingTexture->Settings.Init();
-	StreamingTexture->Settings.TileSize = GetTileSize();
-	StreamingTexture->Settings.TileBorderSize = GetTileBorderSize();
-	StreamingTexture->Settings.bEnableCompressCrunch = bEnableCompressCrunch;
-
-	StreamingTexture->BuildHash = GetStreamingTextureBuildHash();
-
-	const int32 LayerCount = GetLayerCount();
-	check(LayerCount <= RuntimeVirtualTexture::MaxTextureLayers);
-	ETextureSourceFormat LayerFormats[RuntimeVirtualTexture::MaxTextureLayers];
-
-	for (int32 Layer = 0; Layer < LayerCount; Layer++)
-	{
-		EPixelFormat LayerFormat = GetLayerFormat(Layer);
-		LayerFormats[Layer] = LayerFormat == PF_G16 ? TSF_G16 : TSF_BGRA8;
-
-		FTextureFormatSettings FormatSettings;
-		FormatSettings.CompressionSettings = LayerFormat == PF_BC5 ? TC_Normalmap : TC_Default;
-		FormatSettings.CompressionNone = LayerFormat == PF_B8G8R8A8 || LayerFormat == PF_G16;
-		FormatSettings.CompressionNoAlpha = LayerFormat == PF_DXT1 || LayerFormat == PF_BC5;
-		FormatSettings.CompressionYCoCg = IsLayerYCoCg(Layer);
-		FormatSettings.SRGB = IsLayerSRGB(Layer);
-		
-		StreamingTexture->SetLayerFormatSettings(Layer, FormatSettings);
-	}
-
-	StreamingTexture->Source.InitLayered(InSizeX, InSizeY, 1, LayerCount, 1, LayerFormats, InData);
-
-	StreamingTexture->PostEditChange();
-}
-
 #endif
 
-IVirtualTexture* URuntimeVirtualTexture::CreateStreamingTextureProducer(IVirtualTexture* InProducer, int32 InMaxLevel, int32& OutTransitionLevel) const
+namespace RuntimeVirtualTexture
 {
-	if (StreamingTexture != nullptr)
+	IVirtualTexture* CreateStreamingTextureProducer(
+		IVirtualTexture* InProducer,
+		URuntimeVirtualTexture* InBaseVirtualTexure,
+		UVirtualTexture2D* InStreamingTexture,
+		int32 InMaxLevel, int32 InNumMips,
+		int32& OutTransitionLevel)
 	{
-		FTexturePlatformData** StreamingTextureData = StreamingTexture->GetRunningPlatformData();
-		if (StreamingTextureData != nullptr && *StreamingTextureData != nullptr)
+		if (InBaseVirtualTexure != nullptr && InStreamingTexture != nullptr)
 		{
-			FVirtualTextureBuiltData* VTData = (*StreamingTextureData)->VTData;
-			check(GetTileSize() == VTData->TileSize);
-			check(GetTileBorderSize() == VTData->TileBorderSize);
-
-			// Streaming data may have mips removed during cook
-			const int32 NumStreamMips = FMath::Min(GetStreamLowMips(), (*StreamingTextureData)->GetNumVTMips());
-
-			OutTransitionLevel = FMath::Max(InMaxLevel - NumStreamMips + 1, 0);
-			IVirtualTexture* StreamingProducer = new FUploadingVirtualTexture(VTData, 0);
-			return new FVirtualTextureLevelRedirector(InProducer, StreamingProducer, OutTransitionLevel);
+			FTexturePlatformData** StreamingTextureData = InStreamingTexture->GetRunningPlatformData();
+			if (StreamingTextureData != nullptr && *StreamingTextureData != nullptr)
+			{
+				FVirtualTextureBuiltData* VTData = (*StreamingTextureData)->VTData;
+
+				ensure(InBaseVirtualTexure->GetTileSize() == VTData->TileSize);
+				ensure(InBaseVirtualTexure->GetTileBorderSize() == VTData->TileBorderSize);
+				if (InBaseVirtualTexure->GetTileSize() == VTData->TileSize && InBaseVirtualTexure->GetTileBorderSize() == VTData->TileBorderSize)
+				{
+					// Streaming data may have mips removed during cook
+					const int32 NumStreamMips = FMath::Min(InNumMips, (*StreamingTextureData)->GetNumVTMips());
+
+					OutTransitionLevel = FMath::Max(InMaxLevel - NumStreamMips + 1, 0);
+					IVirtualTexture* StreamingProducer = new FUploadingVirtualTexture(VTData, 0);
+					return new FVirtualTextureLevelRedirector(InProducer, StreamingProducer, OutTransitionLevel);
+				}
+			}
 		}
+
+		// Can't create a streaming producer so return original producer.
+		OutTransitionLevel = InMaxLevel;
+		return InProducer;
 	}
-	// Can't create a streaming producer so return original producer.
-	OutTransitionLevel = InMaxLevel;
-	return InProducer;
 }
diff --git a/Engine/Source/Runtime/Engine/Private/VT/VirtualTexture.cpp b/Engine/Source/Runtime/Engine/Private/VT/VirtualTexture.cpp
--- a/Engine/Source/Runtime/Engine/Private/VT/VirtualTexture.cpp
+++ b/Engine/Source/Runtime/Engine/Private/VT/VirtualTexture.cpp
@@ -173,6 +173,11 @@ ULightMapVirtualTexture::ULightMapVirtualTexture(const FObjectInitializer& Objec
 {
 }
 
+URuntimeVirtualTextureStreamingProxy::URuntimeVirtualTextureStreamingProxy(const FObjectInitializer& ObjectInitializer)
+	: Super(ObjectInitializer)
+{
+}
+
 
 UVirtualTexture2D::UVirtualTexture2D(const FObjectInitializer& ObjectInitializer)
 	: UTexture2D(ObjectInitializer)
diff --git a/Engine/Source/Runtime/Renderer/Private/VT/RuntimeVirtualTextureSceneProxy.cpp b/Engine/Source/Runtime/Renderer/Private/VT/RuntimeVirtualTextureSceneProxy.cpp
--- a/Engine/Source/Runtime/Renderer/Private/VT/RuntimeVirtualTextureSceneProxy.cpp
+++ b/Engine/Source/Runtime/Renderer/Private/VT/RuntimeVirtualTextureSceneProxy.cpp
@@ -6,6 +6,8 @@
 #include "VirtualTextureSystem.h"
 #include "VT/RuntimeVirtualTexture.h"
 #include "VT/RuntimeVirtualTextureProducer.h"
+#include "VT/VirtualTexture.h"
+#include "VT/VirtualTextureBuilder.h"
 
 int32 FRuntimeVirtualTextureSceneProxy::ProducerIdGenerator = 1;
 
@@ -38,13 +40,15 @@ FRuntimeVirtualTextureSceneProxy::FRuntimeVirtualTextureSceneProxy(URuntimeVirtu
 		// The Producer object created here will be passed into the virtual texture system which will take ownership.
 		IVirtualTexture* Producer = new FRuntimeVirtualTextureProducer(Desc, ProducerId, MaterialType, bClearTextures, InComponent->GetScene(), Transform, Bounds);
 
-		if (InComponent->IsStreamingLowMips() && VirtualTexture->GetStreamLowMips() > 0)
+		int32 StreamLowMips = InComponent->IsStreamingLowMips() ? InComponent->GetStreamLowMips() : 0;
+		if (StreamLowMips > 0)
 		{
+			UVirtualTexture2D* StreamingTexture = InComponent->GetStreamingTexture()->Texture;
 			// Streaming mips start from the MaxLevel before taking into account the RemoveLowMips
 			const int32 MaxLevel = FMath::CeilLogTwo(FMath::Max(Desc.BlockWidthInTiles, Desc.BlockHeightInTiles));
 			// Wrap our producer to use a streaming producer for low mips
 			int32 StreamingTransitionLevel;
-			Producer = VirtualTexture->CreateStreamingTextureProducer(Producer, MaxLevel, StreamingTransitionLevel);
+			Producer = RuntimeVirtualTexture::CreateStreamingTextureProducer(Producer, VirtualTexture, StreamingTexture, MaxLevel, StreamLowMips, StreamingTransitionLevel);
 			// Any dirty flushes don't need to flush the streaming mips (they only change with a build step).
 			MaxDirtyLevel = FMath::Min(MaxDirtyLevel, StreamingTransitionLevel);
 		}
