Add cvar to disable shader validation errors

  - add support for disabling shader validation
    - UE flag already existed but didn't result in additional args being given to the shader compiler
      - D3D shader compiler now respects the flag
    - added new cvar r.Shaders.Validation to drive this globally
      - enabled by default

diff --git a/Engine/Source/Developer/Windows/ShaderFormatD3D/Private/D3DShaderCompiler.cpp b/Engine/Source/Developer/Windows/ShaderFormatD3D/Private/D3DShaderCompiler.cpp
--- a/Engine/Source/Developer/Windows/ShaderFormatD3D/Private/D3DShaderCompiler.cpp
+++ b/Engine/Source/Developer/Windows/ShaderFormatD3D/Private/D3DShaderCompiler.cpp
@@ -61,6 +61,7 @@ static uint32 TranslateCompilerFlagD3D11(ECompilerFlags CompilerFlag)
 	{
 	case CFLAG_PreferFlowControl: return D3D10_SHADER_PREFER_FLOW_CONTROL;
 	case CFLAG_AvoidFlowControl: return D3D10_SHADER_AVOID_FLOW_CONTROL;
+	case CFLAG_SkipValidation: return D3D10_SHADER_SKIP_VALIDATION;
 	default: return 0;
 	};
 }
diff --git a/Engine/Source/Runtime/Engine/Private/ShaderCompiler/ShaderCompiler.cpp b/Engine/Source/Runtime/Engine/Private/ShaderCompiler/ShaderCompiler.cpp
--- a/Engine/Source/Runtime/Engine/Private/ShaderCompiler/ShaderCompiler.cpp
+++ b/Engine/Source/Runtime/Engine/Private/ShaderCompiler/ShaderCompiler.cpp
@@ -293,7 +293,6 @@ static FAutoConsoleVariableRef CVarCreateShadersOnLoad(
 	TEXT("Whether to create shaders on load, which can reduce hitching, but use more memory.  Otherwise they will be created as needed.")
 );
 
-
 static TAutoConsoleVariable<FString> CVarShaderOverrideDebugDir(
 	TEXT("r.OverrideShaderDebugDir"),
 	"",
@@ -301,6 +300,11 @@ static TAutoConsoleVariable<FString> CVarShaderOverrideDebugDir(
 	TEXT("Empty: use default location Saved\\ShaderDebugInfo.\n"),
 	ECVF_ReadOnly);
 
+static TAutoConsoleVariable<int32> CVarShadersValidation(
+	TEXT("r.Shaders.Validation"),
+	1,
+	TEXT("Enabled shader compiler validation warnings and errors."),
+	ECVF_ReadOnly);
 
 extern bool CompileShaderPipeline(const IShaderFormat* Compiler, FName Format, FShaderPipelineCompileJob* PipelineJob, const FString& Dir);
 
@@ -3537,6 +3541,14 @@ void GlobalBeginCompileShader(
 		}
 	}
 
+	{
+		static const auto CVar = IConsoleManager::Get().FindConsoleVariable(TEXT("r.Shaders.Validation"));
+		if (CVar && CVar->GetInt() == 0)
+		{
+			Input.Environment.CompilerFlags.Add(CFLAG_SkipValidation);
+		}
+	}
+
 	if (IsD3DPlatform((EShaderPlatform)Target.Platform, false))
 	{
 		{
