Revise stripping of platform name from texture format name in virtual texturing system to handle platforms with shared prefixes

  -  change RemovePlatformPrefixFromName() in the VT system to include the underscore when matching system prefixes
    - prevents false positives when dealing with platform names which start with the name of another platform

diff --git a/Engine/Source/Runtime/Engine/Private/VT/VirtualTextureDataBuilder.cpp b/Engine/Source/Runtime/Engine/Private/VT/VirtualTextureDataBuilder.cpp
--- a/Engine/Source/Runtime/Engine/Private/VT/VirtualTextureDataBuilder.cpp
+++ b/Engine/Source/Runtime/Engine/Private/VT/VirtualTextureDataBuilder.cpp
@@ -789,9 +789,11 @@ static const FName RemovePlatformPrefixFromName(FName const& InName)
 
 	for (const auto& PlatformInfo : FDataDrivenPlatformInfoRegistry::GetAllPlatformInfos())
 	{
-		if (NameString.StartsWith(PlatformInfo.Key, ESearchCase::IgnoreCase))
+		FString PlatformTextureFormatPrefix = PlatformInfo.Key;
+		PlatformTextureFormatPrefix += TEXT('_');
+		if (NameString.StartsWith(PlatformTextureFormatPrefix, ESearchCase::IgnoreCase))
 		{
-			return *NameString.RightChop(PlatformInfo.Key.Len() + 1); // Platform name followed by `_`
+			return *NameString.RightChop(PlatformTextureFormatPrefix.Len());
 		}
 	}
 
