Support for dumping shader debug info only on compilation errors or warnings

- valid settings for r.DumpShaderDebugInfo are now 0-3
  - 0 - don't dump (unchanged)
  - 1 - always dump (unchanged)
  - 2 - dump on error (ie compilation failure)
  - 3 - dump on generated errors or warnings

- implementation details:
  - mode 2 and 3 are achieved by reissuing shader compilation jobs
    - jobs are issued as normal with no DumpInfoPath set so the platform specific shader compilers will not dump additional files
      - jobs that fail to compile (mode 2) or generate any errors/warnings (mode 3) are reissued with a valid DumpInfoPath which will cause the platform specific shader compilers to dump the additional files
        - the directory creation is deferred until the second pass on mode 2 or 3 as well so no empty directories are created

- limitations:
  - shader compilation jobs which crash the shader compile worker will not produce debug dump info on modes 2 or 3 as we will not get an opportunity to reissue the jobs

diff --git a/Engine/Config/ConsoleVariables.ini b/Engine/Config/ConsoleVariables.ini
--- a/Engine/Config/ConsoleVariables.ini
+++ b/Engine/Config/ConsoleVariables.ini
@@ -34,7 +34,7 @@
 
 ; Uncomment to get detailed logs on shader compiles and the opportunity to retry on errors
 ;r.ShaderDevelopmentMode=1
-; Uncomment to dump shaders in the Saved folder
+; Uncomment to dump shaders in the Saved folder (1 dump all, 2 dump on compilation failure only, 3 dump on compilation failure or warnings)
 ; Warning: leaving this on for a while will fill your hard drive with many small files and folders
 ;r.DumpShaderDebugInfo=1
 ; When this is enabled, SCW crashes will print out the list of jobs in the current worker
diff --git a/Engine/Plugins/Compositing/OpenColorIO/Source/OpenColorIO/Private/OpenColorIOShaderCompilationManager.cpp b/Engine/Plugins/Compositing/OpenColorIO/Source/OpenColorIO/Private/OpenColorIOShaderCompilationManager.cpp
--- a/Engine/Plugins/Compositing/OpenColorIO/Source/OpenColorIO/Private/OpenColorIOShaderCompilationManager.cpp
+++ b/Engine/Plugins/Compositing/OpenColorIO/Source/OpenColorIO/Private/OpenColorIOShaderCompilationManager.cpp
@@ -115,12 +115,14 @@ void FOpenColorIOShaderCompilationManager::RunCompileJobs()
 
 				UE_LOG(LogOpenColorIOShaderCompiler, Log, TEXT("Compile Job processing... %s"), *CurrentJob.Input.DebugGroupName);
 
-				FString AbsoluteDebugInfoDirectory = IFileManager::Get().ConvertToAbsolutePathForExternalAppForWrite(*(FPaths::ProjectSavedDir() / TEXT("ShaderDebugInfo")));
-				FPaths::NormalizeDirectoryName(AbsoluteDebugInfoDirectory);
-				CurrentJob.Input.DumpDebugInfoPath = AbsoluteDebugInfoDirectory / Format.ToString() / CurrentJob.Input.DebugGroupName;
-				if (!IFileManager::Get().DirectoryExists(*CurrentJob.Input.DumpDebugInfoPath))
+				CurrentJob.Input.DumpDebugInfoRootPath = GShaderCompilingManager->GetAbsoluteShaderDebugInfoDirectory() / Format.ToString();
+				FPaths::NormalizeDirectoryName(CurrentJob.Input.DumpDebugInfoRootPath);
+				const FShaderCompilingManager::EDumpShaderDebugInfo DumpShaderDebugInfo = GShaderCompilingManager->GetDumpShaderDebugInfo();
+				CurrentJob.Input.DebugExtension.Empty();
+				CurrentJob.Input.DumpDebugInfoRootPath.Empty();
+				if (DumpShaderDebugInfo == FShaderCompilingManager::EDumpShaderDebugInfo::Always)
 				{
-					verifyf(IFileManager::Get().MakeDirectory(*CurrentJob.Input.DumpDebugInfoPath, true), TEXT("Failed to create directory for shader debug info '%s'"), *CurrentJob.Input.DumpDebugInfoPath);
+					CurrentJob.Input.DumpDebugInfoRootPath = GShaderCompilingManager->CreateShaderDebugInfoPath(CurrentJob.Input);
 				}
 
 				if (IsValidRef(CurrentJob.Input.SharedEnvironment))
@@ -135,6 +137,13 @@ void FOpenColorIOShaderCompilationManager::RunCompileJobs()
 
 				CurrentJob.bSucceeded = CurrentJob.Output.bSucceeded;
 
+				// Recompile the shader to dump debug info if desired
+				if (GShaderCompilingManager->ShouldRecompileToDumpShaderDebugInfo(CurrentJob))
+				{
+					CurrentJob.Input.DumpDebugInfoPath = GShaderCompilingManager->CreateShaderDebugInfoPath(CurrentJob.Input);
+					Compiler->CompileShader(Format, CurrentJob.Input, CurrentJob.Output, FString(FPlatformProcess::ShaderDir()));
+				}
+
 				if (CurrentJob.Output.bSucceeded)
 				{
 					// Generate a hash of the output and cache it
diff --git a/Engine/Plugins/FX/Niagara/Source/NiagaraEditor/Private/NiagaraCompiler.cpp b/Engine/Plugins/FX/Niagara/Source/NiagaraEditor/Private/NiagaraCompiler.cpp
--- a/Engine/Plugins/FX/Niagara/Source/NiagaraEditor/Private/NiagaraCompiler.cpp
+++ b/Engine/Plugins/FX/Niagara/Source/NiagaraEditor/Private/NiagaraCompiler.cpp
@@ -957,24 +957,12 @@ int32 FHlslNiagaraCompiler::CompileScript(const FNiagaraCompileRequestData* InCo
 	Input.DumpDebugInfoRootPath = GShaderCompilingManager->GetAbsoluteShaderDebugInfoDirectory() / TEXT("VM");
 	FString UsageIdStr = !InOptions.TargetUsageId.IsValid() ? TEXT("") : (TEXT("_") + InOptions.TargetUsageId.ToString());
 	Input.DebugGroupName = InCompileRequest->SourceName / InCompileRequest->EmitterUniqueName / InCompileRequest->ENiagaraScriptUsageEnum->GetNameStringByValue((int64)InOptions.TargetUsage) + UsageIdStr;
-	Input.DumpDebugInfoPath = Input.DumpDebugInfoRootPath / Input.DebugGroupName;
+	Input.DebugExtension.Empty();
+	Input.DumpDebugInfoPath.Empty();
 
-	if (GShaderCompilingManager->GetDumpShaderDebugInfo())
+	if (GShaderCompilingManager->GetDumpShaderDebugInfo() == FShaderCompilingManager::EDumpShaderDebugInfo::Always)
 	{
-		// Sanitize the name to be used as a path
-		// List mostly comes from set of characters not allowed by windows in a path.  Just try to rename a file and type one of these for the list.
-		Input.DumpDebugInfoPath.ReplaceInline(TEXT("<"), TEXT("("));
-		Input.DumpDebugInfoPath.ReplaceInline(TEXT(">"), TEXT(")"));
-		Input.DumpDebugInfoPath.ReplaceInline(TEXT("::"), TEXT("=="));
-		Input.DumpDebugInfoPath.ReplaceInline(TEXT("|"), TEXT("_"));
-		Input.DumpDebugInfoPath.ReplaceInline(TEXT("*"), TEXT("-"));
-		Input.DumpDebugInfoPath.ReplaceInline(TEXT("?"), TEXT("!"));
-		Input.DumpDebugInfoPath.ReplaceInline(TEXT("\""), TEXT("\'"));
-
-		if (!IFileManager::Get().DirectoryExists(*Input.DumpDebugInfoPath))
-		{
-			verifyf(IFileManager::Get().MakeDirectory(*Input.DumpDebugInfoPath, true), TEXT("Failed to create directory for shader debug info '%s'"), *Input.DumpDebugInfoPath);
-		}
+		Input.DumpDebugInfoPath = GShaderCompilingManager->CreateShaderDebugInfoPath(Input);
 	}
 	CompileResults.DumpDebugInfoPath = Input.DumpDebugInfoPath;
 
@@ -994,7 +982,7 @@ int32 FHlslNiagaraCompiler::CompileScript(const FNiagaraCompileRequestData* InCo
 			CompileResults.bComputeSucceeded = true;
 		}
 		CompileResults.Data->LastHlslTranslationGPU = TranslatedHLSL;
-		DumpDebugInfo(CompileResults, true);
+		DumpDebugInfo(CompileResults, Input, true);
 		CompilationJob->CompileResults = CompileResults;
 		return JobID;
 	}
@@ -1047,28 +1035,53 @@ int32 FHlslNiagaraCompiler::CompileScript(const FNiagaraCompileRequestData* InCo
 
 	return JobID;
 }
-			//TODO: Map Lines of HLSL to their source Nodes and flag those nodes with errors associated with their lines.
-void FHlslNiagaraCompiler::DumpDebugInfo(FNiagaraCompileResults& CompileResult, bool bGPUScript)
+
+//TODO: Map Lines of HLSL to their source Nodes and flag those nodes with errors associated with their lines.
+void FHlslNiagaraCompiler::DumpDebugInfo(const FNiagaraCompileResults& CompileResult, const FShaderCompilerInput& Input, bool bGPUScript)
 {
-	if (GShaderCompilingManager->GetDumpShaderDebugInfo() && CompileResults.Data.IsValid())
+	if (CompileResults.Data.IsValid())
 	{
-		FString ExportText = CompileResults.Data->LastHlslTranslation;
-		FString ExportTextAsm = CompileResults.Data->LastAssemblyTranslation;
-		if (bGPUScript)
+		// Support dumping debug info only on failure or warnings
+		FString DumpDebugInfoPath = CompileResult.DumpDebugInfoPath;
+		if (DumpDebugInfoPath.IsEmpty())
 		{
-			ExportText = CompileResults.Data->LastHlslTranslationGPU;
-			ExportTextAsm = "";
+			const FShaderCompilingManager::EDumpShaderDebugInfo DumpShaderDebugInfo = GShaderCompilingManager->GetDumpShaderDebugInfo();
+			bool bDumpDebugInfo = false;
+			if (DumpShaderDebugInfo == FShaderCompilingManager::EDumpShaderDebugInfo::OnError)
+			{
+				bDumpDebugInfo = !CompileResult.bVMSucceeded;
+			}
+			else if (DumpShaderDebugInfo == FShaderCompilingManager::EDumpShaderDebugInfo::OnErrorOrWarning)
+			{
+				bDumpDebugInfo = !CompileResult.bVMSucceeded || (CompileResult.NumErrors + CompileResult.NumWarnings) > 0;
+			}
+
+			if (bDumpDebugInfo)
+			{
+				DumpDebugInfoPath = GShaderCompilingManager->CreateShaderDebugInfoPath(Input);
+			}
 		}
-		FString ExportTextParams;
-		for (const FNiagaraVariable& Var : CompileResults.Data->Parameters.Parameters)
+
+		if (!DumpDebugInfoPath.IsEmpty())
 		{
-			ExportTextParams += Var.ToString();
-			ExportTextParams += "\n";
-		}
+			FString ExportText = CompileResults.Data->LastHlslTranslation;
+			FString ExportTextAsm = CompileResults.Data->LastAssemblyTranslation;
+			if (bGPUScript)
+			{
+				ExportText = CompileResults.Data->LastHlslTranslationGPU;
+				ExportTextAsm = "";
+			}
+			FString ExportTextParams;
+			for (const FNiagaraVariable& Var : CompileResults.Data->Parameters.Parameters)
+			{
+				ExportTextParams += Var.ToString();
+				ExportTextParams += "\n";
+			}
 
-		FNiagaraEditorUtilities::WriteTextFileToDisk(CompileResult.DumpDebugInfoPath, TEXT("NiagaraEmitterInstance.ush"), ExportText, true);
-		FNiagaraEditorUtilities::WriteTextFileToDisk(CompileResult.DumpDebugInfoPath, TEXT("NiagaraEmitterInstance.asm"), ExportTextAsm, true);
-		FNiagaraEditorUtilities::WriteTextFileToDisk(CompileResult.DumpDebugInfoPath, TEXT("NiagaraEmitterInstance.params"), ExportTextParams, true);
+			FNiagaraEditorUtilities::WriteTextFileToDisk(DumpDebugInfoPath, TEXT("NiagaraEmitterInstance.ush"), ExportText, true);
+			FNiagaraEditorUtilities::WriteTextFileToDisk(DumpDebugInfoPath, TEXT("NiagaraEmitterInstance.asm"), ExportTextAsm, true);
+			FNiagaraEditorUtilities::WriteTextFileToDisk(DumpDebugInfoPath, TEXT("NiagaraEmitterInstance.params"), ExportTextParams, true);
+		}
 	}
 }
 
@@ -1232,7 +1245,7 @@ TOptional<FNiagaraCompileResults> FHlslNiagaraCompiler::GetCompileResult(int32 J
 			}
 		}
 	}
-	DumpDebugInfo(CompileResults, false);
+	DumpDebugInfo(CompileResults, CompilationJob->ShaderCompileJob->Input, false);
 
 	//Seems like Results is a bit of a cobbled together mess at this point.
 	//Ideally we can tidy this up in future.
diff --git a/Engine/Plugins/FX/Niagara/Source/NiagaraEditor/Private/NiagaraCompiler.h b/Engine/Plugins/FX/Niagara/Source/NiagaraEditor/Private/NiagaraCompiler.h
--- a/Engine/Plugins/FX/Niagara/Source/NiagaraEditor/Private/NiagaraCompiler.h
+++ b/Engine/Plugins/FX/Niagara/Source/NiagaraEditor/Private/NiagaraCompiler.h
@@ -51,5 +51,5 @@ public:
 private:
 	TUniquePtr<FNiagaraCompilerJob> CompilationJob;
 
-	void DumpDebugInfo(FNiagaraCompileResults& CompileResult, bool bGPUScript);
+	void DumpDebugInfo(const FNiagaraCompileResults& CompileResult, const FShaderCompilerInput& Input, bool bGPUScript);
 };
diff --git a/Engine/Plugins/FX/Niagara/Source/NiagaraShader/Private/NiagaraShaderCompilationManager.cpp b/Engine/Plugins/FX/Niagara/Source/NiagaraShader/Private/NiagaraShaderCompilationManager.cpp
--- a/Engine/Plugins/FX/Niagara/Source/NiagaraShader/Private/NiagaraShaderCompilationManager.cpp
+++ b/Engine/Plugins/FX/Niagara/Source/NiagaraShader/Private/NiagaraShaderCompilationManager.cpp
@@ -47,14 +47,13 @@ NIAGARASHADER_API void FNiagaraShaderCompilationManager::AddJobs(TArray<TSharedR
 
 		UE_LOG(LogNiagaraShaderCompiler, Verbose, TEXT("Adding niagara gpu shader compile job... %s"), *CurrentJob->Input.DebugGroupName);
 
-		static ITargetPlatformManagerModule& TPM = GetTargetPlatformManagerRef();
-		const FName Format = LegacyShaderPlatformToShaderFormat(EShaderPlatform(CurrentJob->Input.Target.Platform));
-		FString AbsoluteDebugInfoDirectory = IFileManager::Get().ConvertToAbsolutePathForExternalAppForWrite(*(FPaths::ProjectSavedDir() / TEXT("ShaderDebugInfo")));
-		FPaths::NormalizeDirectoryName(AbsoluteDebugInfoDirectory);
-		CurrentJob->Input.DumpDebugInfoPath = AbsoluteDebugInfoDirectory / Format.ToString() / CurrentJob->Input.DebugGroupName;
-		if (!IFileManager::Get().DirectoryExists(*CurrentJob->Input.DumpDebugInfoPath))
+		CurrentJob->Input.DumpDebugInfoRootPath = GShaderCompilingManager->GetAbsoluteShaderDebugInfoDirectory() / LegacyShaderPlatformToShaderFormat(EShaderPlatform(CurrentJob->Input.Target.Platform)).ToString();
+		FPaths::NormalizeDirectoryName(CurrentJob->Input.DumpDebugInfoRootPath);
+		CurrentJob->Input.DebugExtension.Empty();
+		CurrentJob->Input.DumpDebugInfoPath.Empty();
+		if (GShaderCompilingManager->GetDumpShaderDebugInfo() == FShaderCompilingManager::EDumpShaderDebugInfo::Always)
 		{
-			verifyf(IFileManager::Get().MakeDirectory(*CurrentJob->Input.DumpDebugInfoPath, true), TEXT("Failed to create directory for shader debug info '%s'"), *CurrentJob->Input.DumpDebugInfoPath);
+			CurrentJob->Input.DumpDebugInfoPath = GShaderCompilingManager->CreateShaderDebugInfoPath(CurrentJob->Input);
 		}
 	}
 	GShaderCompilingManager->AddJobs(InNewJobs, true, false, FString(), FString(), true);
diff --git a/Engine/Source/Runtime/Engine/Private/ShaderCompiler/ShaderCompiler.cpp b/Engine/Source/Runtime/Engine/Private/ShaderCompiler/ShaderCompiler.cpp
--- a/Engine/Source/Runtime/Engine/Private/ShaderCompiler/ShaderCompiler.cpp
+++ b/Engine/Source/Runtime/Engine/Private/ShaderCompiler/ShaderCompiler.cpp
@@ -123,6 +123,14 @@ static void ModalErrorOrLog(const FString& Text, int64 CurrentFilePos = 0, int64
 	}
 }
 
+template<class EnumType>
+constexpr auto& CastEnumToUnderlyingTypeReference(EnumType& Type)
+{
+	static_assert(TIsEnum<EnumType>::Value, "");
+	using UnderType = __underlying_type(EnumType);
+	return reinterpret_cast<UnderType&>(Type);
+}
+
 // Set to 1 to debug ShaderCompileWorker.exe. Set a breakpoint in LaunchWorker() to get the cmd-line.
 #define DEBUG_SHADERCOMPILEWORKER 0
 
@@ -131,11 +139,14 @@ static void ModalErrorOrLog(const FString& Text, int64 CurrentFilePos = 0, int64
 // For example if there are a lot of content shader compile errors you want to skip over without relaunching
 bool GRetryShaderCompilation = false;
 
-static int32 GDumpShaderDebugInfo = 0;
+static FShaderCompilingManager::EDumpShaderDebugInfo GDumpShaderDebugInfo = FShaderCompilingManager::EDumpShaderDebugInfo::Never;
 static FAutoConsoleVariableRef CVarDumpShaderDebugInfo(
 	TEXT("r.DumpShaderDebugInfo"),
-	GDumpShaderDebugInfo,
-	TEXT("When set to 1, will cause any material shaders that are then compiled to dump debug info to GameName/Saved/ShaderDebugInfo\n")
+	CastEnumToUnderlyingTypeReference(GDumpShaderDebugInfo),
+	TEXT("Dumps debug info for compiled shaders to GameName/Saved/ShaderDebugInfo\n")
+	TEXT("When set to 1, debug info is dumped for all compiled shader\n")
+	TEXT("When set to 2, it is restricted to shaders with compilation errors\n")
+	TEXT("When set to 3, it is restricted to shaders with compilation errors or warnings\n")
 	TEXT("The debug info is platform dependent, but usually includes a preprocessed version of the shader source.\n")
 	TEXT("Global shaders automatically dump debug info if r.ShaderDevelopmentMode is enabled, this cvar is not necessary.\n")
 	TEXT("On iOS, if the PowerVR graphics SDK is installed to the default path, the PowerVR shader compiler will be called and errors will be reported during the cook.")
@@ -145,7 +156,7 @@ static int32 GDumpShaderDebugInfoShort = 0;
 static FAutoConsoleVariableRef CVarDumpShaderDebugShortNames(
 	TEXT("r.DumpShaderDebugShortNames"),
 	GDumpShaderDebugInfoShort,
-	TEXT("Only valid when r.DumpShaderDebugInfo=1.\n")
+	TEXT("Only valid when r.DumpShaderDebugInfo > 0.\n")
 	TEXT("When set to 1, will shorten names factory and shader type folder names to avoid issues with long paths.")
 	);
 
@@ -153,7 +164,7 @@ static int32 GDumpShaderDebugInfoSCWCommandLine = 0;
 static FAutoConsoleVariableRef CVarDumpShaderDebugSCWCommandLine(
 	TEXT("r.DumpShaderDebugWorkerCommandLine"),
 	GDumpShaderDebugInfoSCWCommandLine,
-	TEXT("Only valid when r.DumpShaderDebugInfo=1.\n")
+	TEXT("Only valid when r.DumpShaderDebugInfo > 0.\n")
 	TEXT("When set to 1, it will generate a file that can be used with ShaderCompileWorker's -directcompile.")
 	);
 
@@ -644,7 +655,7 @@ static void ProcessErrors(const FShaderCompileJob& CurrentJob, TArray<FString>&
 	}
 }
 
-static void ReadSingleJob(FShaderCompileJob* CurrentJob, FArchive& OutputFile)
+static bool ReadSingleJob(FShaderCompileJob* CurrentJob, FArchive& OutputFile)
 {
 	check(!CurrentJob->bFinalized);
 	CurrentJob->bFinalized = true;
@@ -663,6 +674,16 @@ static void ReadSingleJob(FShaderCompileJob* CurrentJob, FArchive& OutputFile)
 		FString HashFileName = FPaths::Combine(CurrentJob->Input.DumpDebugInfoPath, TEXT("OutputHash.txt"));
 		FFileHelper::SaveStringToFile(CurrentJob->Output.OutputHash.ToString(), *HashFileName, FFileHelper::EEncodingOptions::ForceAnsi);
 	}
+
+	// Support dumping debug info for only failed compilations or those with warnings
+	if (GShaderCompilingManager->ShouldRecompileToDumpShaderDebugInfo(*CurrentJob))
+	{
+		// Build debug info path and create the directory if it doesn't already exist
+		CurrentJob->Input.DumpDebugInfoPath = GShaderCompilingManager->CreateShaderDebugInfoPath(CurrentJob->Input);
+		return true;
+	}
+
+	return false;
 };
 
 // Disable optimization for this crash handler to get full access to the entire stack frame when debugging a crash dump
@@ -849,6 +870,7 @@ void FShaderCompileUtilities::DoReadTaskResults(const TArray<TSharedRef<FShaderC
 	TArray<FShaderCompileJob*> QueuedSingleJobs;
 	TArray<FShaderPipelineCompileJob*> QueuedPipelineJobs;
 	SplitJobsByType(QueuedJobs, QueuedSingleJobs, QueuedPipelineJobs);
+	TArray<FShaderCompileJob*> ReissueSourceJobs;
 
 	// Read single jobs
 	{
@@ -871,7 +893,10 @@ void FShaderCompileUtilities::DoReadTaskResults(const TArray<TSharedRef<FShaderC
 		for (int32 JobIndex = 0; JobIndex < NumJobs; JobIndex++)
 		{
 			auto* CurrentJob = QueuedSingleJobs[JobIndex];
-			ReadSingleJob(CurrentJob, OutputFile);
+			if (ReadSingleJob(CurrentJob, OutputFile))
+			{
+				ReissueSourceJobs.Add(CurrentJob);
+			}
 		}
 	}
 
@@ -924,11 +949,32 @@ void FShaderCompileUtilities::DoReadTaskResults(const TArray<TSharedRef<FShaderC
 			for (int32 Index = 0; Index < NumStageJobs; Index++)
 			{
 				auto* SingleJob = CurrentJob->StageJobs[Index]->GetSingleShaderJob();
-				ReadSingleJob(SingleJob, OutputFile);
+				if (ReadSingleJob(SingleJob, OutputFile))
+				{
+					ReissueSourceJobs.Add(SingleJob);
+				}
 				CurrentJob->bFailedRemovingUnused = CurrentJob->bFailedRemovingUnused | SingleJob->Output.bFailedRemovingUnused;
 			}
 		}
 	}
+	
+	// Requeue any jobs we wish to run again
+	if (ReissueSourceJobs.Num())
+	{
+		TArray<TSharedRef<FShaderCommonCompileJob, ESPMode::ThreadSafe>> ReissueJobs;
+		ReissueJobs.Reserve(ReissueSourceJobs.Num());
+		const uint32 JobId = FShaderCommonCompileJob::GetNextJobId();
+		for (const FShaderCompileJob* ReissueSourceJob : ReissueSourceJobs)
+		{
+			FShaderCompileJob* ReissueJob = new FShaderCompileJob(*ReissueSourceJob);
+			ReissueJob->bFinalized = false;
+			ReissueJob->Id = JobId;
+			ReissueJob->Output = FShaderCompilerOutput();
+			ReissueJobs.Add(TSharedRef<FShaderCommonCompileJob, ESPMode::ThreadSafe>(ReissueJob));
+		}
+
+		GShaderCompilingManager->AddJobs(ReissueJobs, true, false, FString(""), FString(""), true);
+	}
 }
 
 static bool CheckSingleJob(FShaderCompileJob* SingleJob, const TArray<FMaterial*>& Materials, TArray<FString>& Errors)
@@ -2043,9 +2089,55 @@ FShaderCompilingManager::FShaderCompilingManager() :
 	Thread->StartThread();
 }
 
-bool FShaderCompilingManager::GetDumpShaderDebugInfo() const
+FShaderCompilingManager::EDumpShaderDebugInfo FShaderCompilingManager::GetDumpShaderDebugInfo() const
+{
+	if (GDumpShaderDebugInfo < EDumpShaderDebugInfo::Never || GDumpShaderDebugInfo > EDumpShaderDebugInfo::OnErrorOrWarning)
+	{
+		return EDumpShaderDebugInfo::Never;
+	}
+
+	return static_cast<FShaderCompilingManager::EDumpShaderDebugInfo>(GDumpShaderDebugInfo);
+}
+
+FString FShaderCompilingManager::CreateShaderDebugInfoPath(const FShaderCompilerInput& ShaderCompilerInput) const
 {
-	return GDumpShaderDebugInfo != 0;
+	FString DumpDebugInfoPath = ShaderCompilerInput.DumpDebugInfoRootPath / ShaderCompilerInput.DebugGroupName + ShaderCompilerInput.DebugExtension;
+
+	// Sanitize the name to be used as a path
+	// List mostly comes from set of characters not allowed by windows in a path.  Just try to rename a file and type one of these for the list.
+	DumpDebugInfoPath.ReplaceInline(TEXT("<"), TEXT("("));
+	DumpDebugInfoPath.ReplaceInline(TEXT(">"), TEXT(")"));
+	DumpDebugInfoPath.ReplaceInline(TEXT("::"), TEXT("=="));
+	DumpDebugInfoPath.ReplaceInline(TEXT("|"), TEXT("_"));
+	DumpDebugInfoPath.ReplaceInline(TEXT("*"), TEXT("-"));
+	DumpDebugInfoPath.ReplaceInline(TEXT("?"), TEXT("!"));
+	DumpDebugInfoPath.ReplaceInline(TEXT("\""), TEXT("\'"));
+
+	if (!IFileManager::Get().DirectoryExists(*DumpDebugInfoPath))
+	{
+		verifyf(IFileManager::Get().MakeDirectory(*DumpDebugInfoPath, true), TEXT("Failed to create directory for shader debug info '%s'"), *DumpDebugInfoPath);
+	}
+
+	return DumpDebugInfoPath;
+}
+
+bool FShaderCompilingManager::ShouldRecompileToDumpShaderDebugInfo(const FShaderCompileJob& Job) const
+{
+	if (Job.Input.DumpDebugInfoPath.IsEmpty())
+	{
+		const EDumpShaderDebugInfo DumpShaderDebugInfo = GetDumpShaderDebugInfo();
+
+		if (DumpShaderDebugInfo == EDumpShaderDebugInfo::OnError)
+		{
+			return !Job.bSucceeded;
+		}
+		else if (DumpShaderDebugInfo == EDumpShaderDebugInfo::OnErrorOrWarning)
+		{
+			return !Job.bSucceeded || Job.Output.Errors.Num() > 0;
+		}
+	}
+
+	return false;
 }
 
 void FShaderCompilingManager::AddJobs(TArray<TSharedRef<FShaderCommonCompileJob, ESPMode::ThreadSafe>>& NewJobs, bool bOptimizeForLowLatency, bool bRecreateComponentRenderStateOnCompletion, const FString MaterialBasePath, const FString PermutationString, bool bSkipResultProcessing)
@@ -3282,6 +3374,7 @@ void GlobalBeginCompileShader(
 	// asset material name or "Global"
 	Input.DebugGroupName = DebugGroupName;
 	Input.DebugDescription = DebugDescription;
+	Input.DebugExtension = DebugExtension;
 
 	if (ShaderType->GetRootParametersMetadata())
 	{
@@ -3394,24 +3487,10 @@ void GlobalBeginCompileShader(
 	}
 
 	// Setup the debug info path if requested, or if this is a global shader and shader development mode is enabled
-	if (GDumpShaderDebugInfo != 0)
+	Input.DumpDebugInfoPath.Empty();
+	if (GShaderCompilingManager->GetDumpShaderDebugInfo() == FShaderCompilingManager::EDumpShaderDebugInfo::Always)
 	{
-		Input.DumpDebugInfoPath = Input.DumpDebugInfoRootPath / Input.DebugGroupName + DebugExtension,
-		
-		// Sanitize the name to be used as a path
-		// List mostly comes from set of characters not allowed by windows in a path.  Just try to rename a file and type one of these for the list.
-		Input.DumpDebugInfoPath.ReplaceInline(TEXT("<"), TEXT("("));
-		Input.DumpDebugInfoPath.ReplaceInline(TEXT(">"), TEXT(")"));
-		Input.DumpDebugInfoPath.ReplaceInline(TEXT("::"), TEXT("=="));
-		Input.DumpDebugInfoPath.ReplaceInline(TEXT("|"), TEXT("_"));
-		Input.DumpDebugInfoPath.ReplaceInline(TEXT("*"), TEXT("-"));
-		Input.DumpDebugInfoPath.ReplaceInline(TEXT("?"), TEXT("!"));
-		Input.DumpDebugInfoPath.ReplaceInline(TEXT("\""), TEXT("\'"));
-
-		if (!IFileManager::Get().DirectoryExists(*Input.DumpDebugInfoPath))
-		{
-			verifyf(IFileManager::Get().MakeDirectory(*Input.DumpDebugInfoPath, true), TEXT("Failed to create directory for shader debug info '%s'"), *Input.DumpDebugInfoPath);
-		}
+		Input.DumpDebugInfoPath = GShaderCompilingManager->CreateShaderDebugInfoPath(Input);
 	}
 
 	// Add the appropriate definitions for the shader frequency.
diff --git a/Engine/Source/Runtime/Engine/Public/ShaderCompiler.h b/Engine/Source/Runtime/Engine/Public/ShaderCompiler.h
--- a/Engine/Source/Runtime/Engine/Public/ShaderCompiler.h
+++ b/Engine/Source/Runtime/Engine/Public/ShaderCompiler.h
@@ -578,7 +578,17 @@ public:
 		NumExternalJobs = NumJobs;
 	}
 
-	ENGINE_API bool GetDumpShaderDebugInfo() const;
+	enum class EDumpShaderDebugInfo : int32
+	{
+		Never				= 0,
+		Always				= 1,
+		OnError				= 2,
+		OnErrorOrWarning	= 3
+	};
+
+	ENGINE_API EDumpShaderDebugInfo GetDumpShaderDebugInfo() const;
+	ENGINE_API FString CreateShaderDebugInfoPath(const FShaderCompilerInput& ShaderCompilerInput) const;
+	ENGINE_API bool ShouldRecompileToDumpShaderDebugInfo(const FShaderCompileJob& Job) const;
 
 	const FString& GetAbsoluteShaderDebugInfoDirectory() const
 	{
diff --git a/Engine/Source/Runtime/RenderCore/Public/ShaderCore.h b/Engine/Source/Runtime/RenderCore/Public/ShaderCore.h
--- a/Engine/Source/Runtime/RenderCore/Public/ShaderCore.h
+++ b/Engine/Source/Runtime/RenderCore/Public/ShaderCore.h
@@ -17,7 +17,7 @@
 class Error;
 
 // this is for the protocol, not the data, bump if FShaderCompilerInput or ProcessInputFromArchive changes.
-const int32 ShaderCompileWorkerInputVersion = 12;
+const int32 ShaderCompileWorkerInputVersion = 13;
 // this is for the protocol, not the data, bump if FShaderCompilerOutput or WriteToOutputArchive changes.
 const int32 ShaderCompileWorkerOutputVersion = 5;
 // this is for the protocol, not the data, bump if FShaderCompilerOutput or WriteToOutputArchive changes.
@@ -596,6 +596,7 @@ struct FShaderCompilerInput
 	FString DumpDebugInfoPath;
 	// materialname or "Global" "for debugging and better error messages
 	FString DebugGroupName;
+	FString DebugExtension;
 
 	// Description of the configuration used when compiling. 
 	FString DebugDescription;
@@ -744,6 +745,7 @@ struct FShaderCompilerInput
 		Ar << Input.UsedOutputs;
 		Ar << Input.DumpDebugInfoRootPath;
 		Ar << Input.DumpDebugInfoPath;
+		Ar << Input.DebugExtension;
 		Ar << Input.DebugGroupName;
 		Ar << Input.DebugDescription;
 		Ar << Input.Environment;
