Ensure FPhysScene_ChaosPauseHandler gets cleaned up before exiting destructors get called as it can lead to a race where all the FMulticastDelegateBase have already gotten cleand up then causes a crash

diff --git a/Engine/Source/Runtime/Engine/Private/PhysicsEngine/Experimental/PhysScene_Chaos.cpp b/Engine/Source/Runtime/Engine/Private/PhysicsEngine/Experimental/PhysScene_Chaos.cpp
--- a/Engine/Source/Runtime/Engine/Private/PhysicsEngine/Experimental/PhysScene_Chaos.cpp
+++ b/Engine/Source/Runtime/Engine/Private/PhysicsEngine/Experimental/PhysScene_Chaos.cpp
@@ -881,6 +881,10 @@ void FPhysScene_Chaos::Shutdown()
 		ChaosModule->DestroySolver(GetSolver());
 	}
 
+#if WITH_EDITOR
+	PhysScene_ChaosPauseHandler.Reset();
+#endif
+
 	ChaosModule = nullptr;
 	SceneSolver = nullptr;
 
