Detected and reported if CrashReportClientEditor crashed itself and improved how Editor analytics summarize few stats:
  - Computed a more accurate 'idle' base on user inputs.
  - Experimenting a measurement of Editor 'idle' time based on Editor process CPU usage.
  - Recorded entering/exiting PIE right away rather than waiting the next 'heartbeat' up to 60 seconds.
  - In case the the session creation is delayed (because contention on the session lock), don't wait up to 60 seconds to retry. Retry immediatedly at the next tick.
  - Increased update rate of the session in the first minute to each second rather than each minute because lot of crashes occurs before the first minute.

diff --git a/Engine/Source/Developer/EditorAnalyticsSession/Private/EditorAnalyticsSession.cpp b/Engine/Source/Developer/EditorAnalyticsSession/Private/EditorAnalyticsSession.cpp
--- a/Engine/Source/Developer/EditorAnalyticsSession/Private/EditorAnalyticsSession.cpp
+++ b/Engine/Source/Developer/EditorAnalyticsSession/Private/EditorAnalyticsSession.cpp
@@ -41,6 +41,7 @@ namespace EditorAnalyticsDefs
 	static const FString PlatformProcessIDStoreKey(TEXT("PlatformProcessID"));
 	static const FString MonitorProcessIDStoreKey(TEXT("MonitorProcessID"));
 	static const FString ExitCodeStoreKey(TEXT("ExitCode"));
+	static const FString MonitorExceptCodeStoreKey(TEXT("MonitorExceptCode"));
 
 	// timestamps
 	static const FString StartupTimestampStoreKey(TEXT("StartupTimestamp"));
@@ -49,6 +50,8 @@ namespace EditorAnalyticsDefs
 	static const FString Idle1MinStoreKey(TEXT("Idle1Min"));
 	static const FString Idle5MinStoreKey(TEXT("Idle5Min"));
 	static const FString Idle30MinStoreKey(TEXT("Idle30Min"));
+	static const FString TotalUserInactivitySecondsStoreKey(TEXT("TotalUserInactivitySecs"));
+	static const FString TotalEditorInactivitySecondsStoreKey(TEXT("TotalEditorInactivitySecs"));
 	static const FString CurrentUserActivityStoreKey(TEXT("CurrentUserActivity"));
 	static const FString PluginsStoreKey(TEXT("Plugins"));
 	static const FString AverageFPSStoreKey(TEXT("AverageFPS"));
@@ -258,6 +261,14 @@ namespace EditorAnalyticsUtils
 			}
 		}
 
+		{
+			FString MonitorExceptCodeString;
+			if (FPlatformMisc::GetStoredValue(EditorAnalyticsDefs::StoreId, SectionName, EditorAnalyticsDefs::MonitorExceptCodeStoreKey, MonitorExceptCodeString))
+			{
+				Session.MonitorExceptCode.Emplace(FCString::Atoi(*MonitorExceptCodeString));
+			}
+		}
+
 		// scope is just to isolate the temporary value
 		{
 			FString StartupTimestampString;
@@ -274,6 +285,8 @@ namespace EditorAnalyticsUtils
 		GET_STORED_INT(Idle1Min);
 		GET_STORED_INT(Idle5Min);
 		GET_STORED_INT(Idle30Min);
+		GET_STORED_INT(TotalUserInactivitySeconds);
+		GET_STORED_INT(TotalEditorInactivitySeconds);
 
 		GET_STORED_STRING(CurrentUserActivity);
 
@@ -359,10 +372,6 @@ FEditorAnalyticsSession::FEditorAnalyticsSession()
 	MonitorProcessID = 0;
 	StartupTimestamp = FDateTime::MinValue();
 	Timestamp = FDateTime::MinValue();
-	IdleSeconds = 0;
-	Idle1Min = 0;
-	Idle5Min = 0;
-	Idle30Min = 0;
 	CurrentUserActivity = EditorAnalyticsDefs::DefaultUserActivity;
 	AverageFPS = 0;
 	GPUVendorID = 0;
@@ -478,6 +487,8 @@ bool FEditorAnalyticsSession::Save()
 			{EditorAnalyticsDefs::Idle1MinStoreKey,  FString::FromInt(Idle1Min)},
 			{EditorAnalyticsDefs::Idle5MinStoreKey,  FString::FromInt(Idle5Min)},
 			{EditorAnalyticsDefs::Idle30MinStoreKey, FString::FromInt(Idle30Min)},
+			{EditorAnalyticsDefs::TotalUserInactivitySecondsStoreKey, FString::FromInt(TotalUserInactivitySeconds)},
+			{EditorAnalyticsDefs::TotalEditorInactivitySecondsStoreKey, FString::FromInt(TotalEditorInactivitySeconds)},
 			{EditorAnalyticsDefs::CurrentUserActivityStoreKey, CurrentUserActivity},
 			{EditorAnalyticsDefs::AverageFPSStoreKey,     FString::SanitizeFloat(AverageFPS)},
 			{EditorAnalyticsDefs::IsDebuggerStoreKey,     EditorAnalyticsUtils::BoolToStoredString(bIsDebugger)},
@@ -494,6 +505,11 @@ bool FEditorAnalyticsSession::Save()
 			KeyValues.Emplace(EditorAnalyticsDefs::ExitCodeStoreKey, FString::FromInt(ExitCode.GetValue()));
 		}
 
+		if (MonitorExceptCode.IsSet())
+		{
+			KeyValues.Emplace(EditorAnalyticsDefs::MonitorExceptCodeStoreKey, FString::FromInt(MonitorExceptCode.GetValue()));
+		}
+
 		FPlatformMisc::SetStoredValues(EditorAnalyticsDefs::StoreId, StorageLocation, KeyValues);
 	}
 
@@ -534,6 +550,7 @@ bool FEditorAnalyticsSession::Delete() const
 	FPlatformMisc::DeleteStoredValue(EditorAnalyticsDefs::StoreId, SectionName, EditorAnalyticsDefs::PlatformProcessIDStoreKey);
 	FPlatformMisc::DeleteStoredValue(EditorAnalyticsDefs::StoreId, SectionName, EditorAnalyticsDefs::MonitorProcessIDStoreKey);
 	FPlatformMisc::DeleteStoredValue(EditorAnalyticsDefs::StoreId, SectionName, EditorAnalyticsDefs::ExitCodeStoreKey);
+	FPlatformMisc::DeleteStoredValue(EditorAnalyticsDefs::StoreId, SectionName, EditorAnalyticsDefs::MonitorExceptCodeStoreKey);
 
 	FPlatformMisc::DeleteStoredValue(EditorAnalyticsDefs::StoreId, SectionName, EditorAnalyticsDefs::StartupTimestampStoreKey);
 	FPlatformMisc::DeleteStoredValue(EditorAnalyticsDefs::StoreId, SectionName, EditorAnalyticsDefs::TimestampStoreKey);
@@ -541,6 +558,8 @@ bool FEditorAnalyticsSession::Delete() const
 	FPlatformMisc::DeleteStoredValue(EditorAnalyticsDefs::StoreId, SectionName, EditorAnalyticsDefs::Idle1MinStoreKey);
 	FPlatformMisc::DeleteStoredValue(EditorAnalyticsDefs::StoreId, SectionName, EditorAnalyticsDefs::Idle5MinStoreKey);
 	FPlatformMisc::DeleteStoredValue(EditorAnalyticsDefs::StoreId, SectionName, EditorAnalyticsDefs::Idle30MinStoreKey);
+	FPlatformMisc::DeleteStoredValue(EditorAnalyticsDefs::StoreId, SectionName, EditorAnalyticsDefs::TotalUserInactivitySecondsStoreKey);
+	FPlatformMisc::DeleteStoredValue(EditorAnalyticsDefs::StoreId, SectionName, EditorAnalyticsDefs::TotalEditorInactivitySecondsStoreKey);
 	FPlatformMisc::DeleteStoredValue(EditorAnalyticsDefs::StoreId, SectionName, EditorAnalyticsDefs::CurrentUserActivityStoreKey);
 	FPlatformMisc::DeleteStoredValue(EditorAnalyticsDefs::StoreId, SectionName, EditorAnalyticsDefs::PluginsStoreKey);
 	FPlatformMisc::DeleteStoredValue(EditorAnalyticsDefs::StoreId, SectionName, EditorAnalyticsDefs::AverageFPSStoreKey);
@@ -678,3 +697,17 @@ void FEditorAnalyticsSession::SaveExitCode(int32 InExitCode)
 	const FString StorageLocation = EditorAnalyticsUtils::GetSessionStorageLocation(SessionId);
 	FPlatformMisc::SetStoredValue(EditorAnalyticsDefs::StoreId, StorageLocation, EditorAnalyticsDefs::ExitCodeStoreKey, ExitCodeStr);
 }
+
+void FEditorAnalyticsSession::SaveMonitorExceptCode(int32 InExceptCode)
+{
+	if (!ensure(IsLocked()))
+	{
+		return;
+	}
+
+	MonitorExceptCode.Emplace(InExceptCode);
+	FString ExceptCodeStr = FString::Printf(TFormatSpecifier<decltype(InExceptCode)>::GetFormatSpecifier(), InExceptCode);
+
+	const FString StorageLocation = EditorAnalyticsUtils::GetSessionStorageLocation(SessionId);
+	FPlatformMisc::SetStoredValue(EditorAnalyticsDefs::StoreId, StorageLocation, EditorAnalyticsDefs::MonitorExceptCodeStoreKey, ExceptCodeStr);
+}
diff --git a/Engine/Source/Developer/EditorAnalyticsSession/Private/EditorSessionSummarySender.cpp b/Engine/Source/Developer/EditorAnalyticsSession/Private/EditorSessionSummarySender.cpp
--- a/Engine/Source/Developer/EditorAnalyticsSession/Private/EditorSessionSummarySender.cpp
+++ b/Engine/Source/Developer/EditorAnalyticsSession/Private/EditorSessionSummarySender.cpp
@@ -172,6 +172,8 @@ void FEditorSessionSummarySender::SendSessionSummaryEvent(const FEditorAnalytics
 	AnalyticsAttributes.Emplace(TEXT("1MinIdle"), Session.Idle1Min);
 	AnalyticsAttributes.Emplace(TEXT("5MinIdle"), Session.Idle5Min);
 	AnalyticsAttributes.Emplace(TEXT("30MinIdle"), Session.Idle30Min);
+	AnalyticsAttributes.Emplace(TEXT("TotalUserInactivitySecs"), Session.TotalUserInactivitySeconds);
+	AnalyticsAttributes.Emplace(TEXT("TotalEditorInactivitySecs"), Session.TotalEditorInactivitySeconds);
 	AnalyticsAttributes.Emplace(TEXT("CurrentUserActivity"), Session.CurrentUserActivity);
 	AnalyticsAttributes.Emplace(TEXT("AverageFPS"), Session.AverageFPS);
 	AnalyticsAttributes.Emplace(TEXT("Plugins"), PluginsString);
@@ -207,6 +209,12 @@ void FEditorSessionSummarySender::SendSessionSummaryEvent(const FEditorAnalytics
 		AnalyticsAttributes.Emplace(TEXT("ExitCode"), Session.ExitCode.GetValue());
 	}
 
+	// Add the monitor exception code in case the out-of-process monitor (CrashReportClientEditor) crashed itself, caught the exception and was able to store it in the session before dying.
+	if (Session.MonitorExceptCode.IsSet())
+	{
+		AnalyticsAttributes.Emplace(TEXT("MonitorExceptCode"), Session.MonitorExceptCode.GetValue());
+	}
+
 	// Was this summary produced by another process than itself or the out-of-process monitor for that run?
 	AnalyticsAttributes.Emplace(TEXT("DelayedSend"), Session.PlatformProcessID != CurrentSessionProcessId);
 
diff --git a/Engine/Source/Developer/EditorAnalyticsSession/Public/EditorAnalyticsSession.h b/Engine/Source/Developer/EditorAnalyticsSession/Public/EditorAnalyticsSession.h
--- a/Engine/Source/Developer/EditorAnalyticsSession/Public/EditorAnalyticsSession.h
+++ b/Engine/Source/Developer/EditorAnalyticsSession/Public/EditorAnalyticsSession.h
@@ -31,13 +31,16 @@ struct EDITORANALYTICSSESSION_API FEditorAnalyticsSession
 	uint32 PlatformProcessID;
 	uint32 MonitorProcessID; // Set to the CrashReportClientEditor PID when out-of-process reporting is used.
 	TOptional<int32> ExitCode; // Set by CrashReportClientEditor after the Editor process exit when out-of-process reporting is used and reading the exit code is supported.
+	TOptional<int32> MonitorExceptCode; // Set by CrashReportClientEditor if an exception is caught by monitoring the Editor. This is to detect if CRC crashes itself.
 
 	FDateTime StartupTimestamp;
 	FDateTime Timestamp;
-	volatile int32 IdleSeconds; // Can be updated from concurrent threads.
-	volatile int32 Idle1Min;
-	volatile int32 Idle5Min;
-	volatile int32 Idle30Min;
+	volatile int32 IdleSeconds = 0; // Can be updated from concurrent threads.
+	volatile int32 Idle1Min = 0;
+	volatile int32 Idle5Min = 0;
+	volatile int32 Idle30Min = 0;
+	volatile int32 TotalUserInactivitySeconds = 0; // If time elapsed between two user interaction is greater than a threshold, consider it inactivity and sum it up.
+	volatile int32 TotalEditorInactivitySeconds = 0; // Account for user input and Editor process CPU usage. Add up gaps where the CPU was not used intensively and the user did not interact.
 	FString CurrentUserActivity;
 	TArray<FString> Plugins;
 	float AverageFPS;
@@ -158,6 +161,14 @@ struct EDITORANALYTICSSESSION_API FEditorAnalyticsSession
 	 */
 	void SaveExitCode(int32 ExitCode);
 
+	/**
+	 * Set the exception code that caused the monitor application to crash. When the monitoring application raises
+	 * and catches an unexpected exception, it tries to save the exception code it in the session before dying.
+	 * @note This is to diagnose the cases where CrashReportClientEditor send the summary event delayed and without
+	 *       the Editor exit code.
+	 */
+	void SaveMonitorExceptCode(int32 ExceptCode);
+
 private:
 	static FSystemWideCriticalSection* StoredValuesLock;
 
diff --git a/Engine/Source/Programs/CrashReportClient/Private/Windows/CrashReportClientMainWindows.cpp b/Engine/Source/Programs/CrashReportClient/Private/Windows/CrashReportClientMainWindows.cpp
--- a/Engine/Source/Programs/CrashReportClient/Private/Windows/CrashReportClientMainWindows.cpp
+++ b/Engine/Source/Programs/CrashReportClient/Private/Windows/CrashReportClientMainWindows.cpp
@@ -3,6 +3,10 @@
 #include "CrashReportClientApp.h"
 #include "Windows/WindowsHWrapper.h"
 
+#if CRASH_REPORT_WITH_MTBF
+#include "EditorAnalyticsSession.h"
+#endif
+
 /**
  * WinMain, called when the application is started
  */
@@ -10,8 +14,34 @@ int WINAPI WinMain(_In_ HINSTANCE hInInstance, _In_opt_ HINSTANCE hPrevInstance,
 {
 	hInstance = hInInstance;
 
-	// Run the app
+#if CRASH_REPORT_WITH_MTBF && !PLATFORM_SEH_EXCEPTIONS_DISABLED
+	// Try to record if CrashReportClientEditor is crashing. Analytics shows that good number of Editor exit code are reported delayed, hinting
+	// that CRCEditor was not running anymore. Try figuring out if it crashed. Suspecting that the Editor crash reporter/handler code is crashing could also
+	// inadvertedly cause a crash in CRCEditor.
+	__try
+	{
+		RunCrashReportClient(GetCommandLineW());
+	}
+	__except(EXCEPTION_EXECUTE_HANDLER)
+	{
+		uint64 MonitorPid;
+		if (FParse::Value(GetCommandLineW(), TEXT("-MONITOR="), MonitorPid))
+		{
+			FTimespan Timeout = FTimespan::FromSeconds(5);
+			if (FEditorAnalyticsSession::Lock(Timeout))
+			{
+				FEditorAnalyticsSession MonitoredSession;
+				if (FEditorAnalyticsSession::FindSession(MonitorPid, MonitoredSession))
+				{
+					MonitoredSession.SaveMonitorExceptCode(GetExceptionCode());
+				}
+				FEditorAnalyticsSession::Unlock();
+			}
+		}
+	}
+#else
 	RunCrashReportClient(GetCommandLineW());
+#endif
 
 	return 0;
 }
diff --git a/Engine/Source/Runtime/Engine/Private/Analytics/EditorSessionSummaryWriter.cpp b/Engine/Source/Runtime/Engine/Private/Analytics/EditorSessionSummaryWriter.cpp
--- a/Engine/Source/Runtime/Engine/Private/Analytics/EditorSessionSummaryWriter.cpp
+++ b/Engine/Source/Runtime/Engine/Private/Analytics/EditorSessionSummaryWriter.cpp
@@ -13,6 +13,7 @@
 #include "Engine/Engine.h"
 #include "Framework/Application/SlateApplication.h"
 #include "HAL/PlatformProcess.h"
+#include "HAL/PlatformTime.h"
 #include "IAnalyticsProviderET.h"
 #include "Interfaces/IPluginManager.h"
 #include "Interfaces/IProjectManager.h"
@@ -26,6 +27,7 @@
 #include "EditorAnalyticsSession.h"
 #include "IVREditorModule.h"
 #include "Kismet2/DebuggerCommands.h"
+#include "Editor.h"
 
 #define LOCTEXT_NAMESPACE "SessionSummary"
 
@@ -35,6 +37,55 @@ namespace EditorSessionWriterDefs
 {
 	// Number of seconds to wait between each update of the mutable metrics.
 	static const float HeartbeatPeriodSeconds = 60;
+
+	// In the first minutes, update every seconds because lot of crashes occurs in the first minute.
+	static const float EarlyHeartbeatPeriodSeconds = 1;
+
+	// The number of seconds required between two user inputs to consider the user as 'inactive'.
+	constexpr double UserInactivitySecondsForIdleState = 120;
+
+	// The upper CPU usage % considered as Idle. If the CPU usage goes above this threshold, the Editor is considered 'active'.
+	constexpr float IdleCpuUsagePercent = 20;
+
+	// The number of seconds required between Editor activities to consider the Editor as 'inactive' (user input, cpu burst).
+	constexpr double EditorInactivitySecondsForIdleState = 120;
+}
+
+namespace EditorSessionWriterUtils
+{
+	bool UpdateIdleTime(const FDateTime& InActivityTimeUtc, TAtomic<uint64>& LastActivityUtcUnixTimestamp, volatile int32& TotalInactivitySecondsCounter, double InactivityThresholdSecondsForIdleState)
+	{
+		// Atomically load the time of the last activity. (user input, CPU burst/high usage, crash, terminate, shutdown)
+		FDateTime LastActivityTime = FDateTime::FromUnixTimestamp(LastActivityUtcUnixTimestamp.Load());
+		if (InActivityTimeUtc != LastActivityTime)
+		{
+			// The thread that exchanges the LastActivityUtcUnixTimestamp successfully is responsible to update the total counter. Other concurrent threads will fail
+			// to update and some fractions of seconds can be lost, but that's neglictable in the overall idle computation.
+			uint64 LastActivityExpectedTime = LastActivityTime.ToUnixTimestamp();
+			if (LastActivityUtcUnixTimestamp.CompareExchange(LastActivityExpectedTime, InActivityTimeUtc.ToUnixTimestamp()))
+			{
+				// Check if this period of inactivity was long enough to be considered as 'idle'.
+				FTimespan InactivityTimespan = InActivityTimeUtc - LastActivityTime;
+				if (InactivityTimespan > FTimespan::FromSeconds(InactivityThresholdSecondsForIdleState))
+				{
+					FPlatformAtomics::InterlockedAdd(&TotalInactivitySecondsCounter, FMath::FloorToInt(static_cast<float>(InactivityTimespan.GetTotalSeconds())));
+					return true; //  Total Idle time was updated.
+				}
+			}
+		}
+
+		return false; // Total Idle time wasn't updated.
+	}
+
+	bool IsMainLoopStaled(const FDateTime& CurrTime, const FDateTime& LastTickTime)
+	{
+		return (CurrTime - LastTickTime) > FTimespan::FromSeconds(2);
+	}
+
+	bool IsMainLoopStaled(const FDateTime& CurrTime, const TAtomic<uint64>& LastTickTimestamp)
+	{
+		return IsMainLoopStaled(CurrTime, FDateTime::FromUnixTimestamp(LastTickTimestamp));
+	}
 }
 
 FEditorSessionSummaryWriter::FEditorSessionSummaryWriter(uint32 InOutOfProcessMonitorProcessId)
@@ -79,22 +130,47 @@ void FEditorSessionSummaryWriter::Initialize()
 
 	if (CurrentSession)
 	{
+		// Reset all 'inactivity' timers to 'now'.
+		LastSlateInteractionTime = FSlateApplication::Get().GetCurrentTime(); // Use same clock than Slate does.
+		FDateTime CurrTimeUtc = FDateTime::UtcNow();
+		LastTickUtcTime.Store(CurrTimeUtc.ToUnixTimestamp());
+		LastEditorActivityUtcTime.Store(CurrTimeUtc.ToUnixTimestamp());
+		LastUserActivityUtcTime.Store(CurrTimeUtc.ToUnixTimestamp());
+
 		// Register for crash and app state callbacks
 		FCoreDelegates::OnHandleSystemError.AddRaw(this, &FEditorSessionSummaryWriter::OnCrashing); // WARNING: Don't assume this function is only called from game thread.
 		FCoreDelegates::ApplicationWillTerminateDelegate.AddRaw(this, &FEditorSessionSummaryWriter::OnTerminate); // WARNING: Don't assume this function is only called from game thread.
 		FCoreDelegates::IsVanillaProductChanged.AddRaw(this, &FEditorSessionSummaryWriter::OnVanillaStateChanged);
+		FEditorDelegates::PreBeginPIE.AddRaw(this, &FEditorSessionSummaryWriter::OnEnterPIE);
+		FEditorDelegates::EndPIE.AddRaw(this, &FEditorSessionSummaryWriter::OnExitPIE);
 		FUserActivityTracking::OnActivityChanged.AddRaw(this, &FEditorSessionSummaryWriter::OnUserActivity);
 		FSlateApplication::Get().GetOnModalLoopTickEvent().AddRaw(this, &FEditorSessionSummaryWriter::Tick);
+		FSlateApplication::Get().GetLastUserInteractionTimeUpdateEvent().AddRaw(this, &FEditorSessionSummaryWriter::OnSlateUserInteraction);
 	}
 }
 
-void FEditorSessionSummaryWriter::UpdateTimestamps()
+void FEditorSessionSummaryWriter::UpdateTimestamp(const FDateTime& InCurrTimeUtc)
+{
+	CurrentSession->Timestamp = InCurrTimeUtc;
+}
+
+void FEditorSessionSummaryWriter::UpdateEditorIdleTime(const FDateTime& InActivityTimeUtc, bool bSaveSession)
 {
-	CurrentSession->Timestamp = FDateTime::UtcNow();
-	UpdateIdleTimes();
+	if (EditorSessionWriterUtils::UpdateIdleTime(InActivityTimeUtc, LastEditorActivityUtcTime, CurrentSession->TotalEditorInactivitySeconds, EditorSessionWriterDefs::EditorInactivitySecondsForIdleState) && bSaveSession)
+	{
+		TrySaveCurrentSession();
+	}
 }
 
-void FEditorSessionSummaryWriter::UpdateIdleTimes()
+void FEditorSessionSummaryWriter::UpdateUserIdleTime(const FDateTime& InActivityTimeUtc, bool bSaveSession)
+{
+	if (EditorSessionWriterUtils::UpdateIdleTime(InActivityTimeUtc, LastUserActivityUtcTime, CurrentSession->TotalUserInactivitySeconds, EditorSessionWriterDefs::UserInactivitySecondsForIdleState) && bSaveSession)
+	{
+		TrySaveCurrentSession();
+	}
+}
+
+void FEditorSessionSummaryWriter::UpdateLegacyIdleTimes()
 {
 	int32 IdleSecondsTmp = FPlatformAtomics::AtomicRead(&CurrentSession->IdleSeconds); // Atomically load only once.
 
@@ -124,29 +200,59 @@ void FEditorSessionSummaryWriter::Tick(float DeltaTime)
 		return;
 	}
 
-	if (CurrentSession != nullptr)
+	// Try late initialization (in case the global lock was already taken during init and the session couldn't be created or the user just toggled 'send data' on).
+	if (CurrentSession == nullptr)
 	{
-		// Note: Update idle time in Tick() because Slate cannot be invoked from any thread and UpdateTimeStamps() can be called from a crashing thread.
-		//       Compute the idle time from Slate point-of view. Note that some tasks blocking the UI (such as importing large assets) may be considered idle time.
-		CurrentSession->IdleSeconds = FSlateApplication::Get().GetLastUserInteractionTime() != 0 ? // In case Slate did not register any interaction yet (ex the user just launches the Editor and goes away)
-			FMath::FloorToInt(static_cast<float>(FSlateApplication::Get().GetCurrentTime() - FSlateApplication::Get().GetLastUserInteractionTime())) :
-			FMath::FloorToInt(static_cast<float>((FDateTime::UtcNow() - CurrentSession->StartupTimestamp).GetTotalSeconds()));
+		Initialize();
 	}
+	else
+	{
+		// How much time the Editor ran without user input. Slate must be access from main thread.
+		FPlatformAtomics::InterlockedExchange(&CurrentSession->IdleSeconds, FMath::FloorToInt(static_cast<float>(FSlateApplication::Get().GetCurrentTime() - LastSlateInteractionTime)));
 
-	HeartbeatTimeElapsed += DeltaTime;
+		FDateTime CurrTickTime = FDateTime::UtcNow();
+		FDateTime LastTickTime = FDateTime::FromUnixTimestamp(LastTickUtcTime);
 
-	if (HeartbeatTimeElapsed > EditorSessionWriterDefs::HeartbeatPeriodSeconds)
-	{
-		HeartbeatTimeElapsed = 0.0f;
+		// Did not 'Tick' for a long time? The main thread was busy (or blocked). Consider the Editor as active for that period of time.
+		if (EditorSessionWriterUtils::IsMainLoopStaled(CurrTickTime, LastTickTime))
+		{
+			//UE_LOG(LogEditorSessionSummary, Verbose, TEXT("Main loop staled"));
+			LastEditorActivityUtcTime.Store(CurrTickTime.ToUnixTimestamp()); // Reset the Editor activity timer.
+		}
+		// The Editor process CPU usage is high enough? Consider this as an Editor activity.
+		else if (FPlatformTime::GetCPUTime().CPUTimePct > EditorSessionWriterDefs::IdleCpuUsagePercent)
+		{
+			//UE_LOG(LogEditorSessionSummary, Verbose, TEXT("Editor Usage: %s"), *LexToString(FPlatformTime::GetCPUTime().CPUTimePct));
+			UpdateEditorIdleTime(CurrTickTime, /*bSaveSession*/true); // Only save if the Editor did not register any activities in the last n seconds.
+		}
+		else if ((CurrTickTime - FDateTime::FromUnixTimestamp(LastEditorActivityUtcTime)).GetTotalSeconds() > EditorSessionWriterDefs::EditorInactivitySecondsForIdleState)
+		{
+			UpdateEditorIdleTime(CurrTickTime, /*bSaveSession*/true); // This period of time is considered idle, record it now.
+		}
 
-		// Try late initialization (in case the global lock was already taken during init and the session couldn't be created or the user just toggled 'send data' on).
-		if (CurrentSession == nullptr)
+		if ((CurrTickTime - FDateTime::FromUnixTimestamp(LastUserActivityUtcTime)).GetTotalSeconds() > EditorSessionWriterDefs::UserInactivitySecondsForIdleState)
 		{
-			Initialize();
+			UpdateUserIdleTime(CurrTickTime, /*bSaveSession*/true); // This period of time is considered idle, record it now.
 		}
 
-		if (CurrentSession != nullptr)
+		LastTickUtcTime.Store(CurrTickTime.ToUnixTimestamp());
+
+		// Update other session stats approximatively every minute.
+		HeartbeatTimeElapsed += DeltaTime;
+
+		// In the first minute, be more granular about updating session timestamp and save every second or so (Lot of abnormal terminations occur before the first minute).
+		FTimespan SessionAge = CurrTickTime - CurrentSession->StartupTimestamp;
+		if (SessionAge.GetTotalSeconds() < EditorSessionWriterDefs::HeartbeatPeriodSeconds && HeartbeatTimeElapsed >= EditorSessionWriterDefs::EarlyHeartbeatPeriodSeconds)
 		{
+			HeartbeatTimeElapsed = 0.0f;
+			UpdateTimestamp(CurrTickTime);
+			TrySaveCurrentSession();
+		}
+		// After the first minute, update the session every minute or so.
+		else if (HeartbeatTimeElapsed > EditorSessionWriterDefs::HeartbeatPeriodSeconds)
+		{
+			HeartbeatTimeElapsed = 0.0f;
+
 			// check if the debugger is present
 			bool bIsDebuggerPresent = FPlatformMisc::IsDebuggerPresent();
 			if (CurrentSession->bIsDebugger != bIsDebuggerPresent)
@@ -159,16 +265,13 @@ void FEditorSessionSummaryWriter::Tick(float DeltaTime)
 				}
 			}
 
-			UpdateTimestamps();
+			UpdateTimestamp(CurrTickTime);
+			UpdateLegacyIdleTimes();
 
-#if WITH_EDITOR
 			extern ENGINE_API float GAverageFPS;
-
 			CurrentSession->AverageFPS = GAverageFPS;
 			CurrentSession->bIsInVRMode = IVREditorModule::Get().IsVREditorModeActive();
-			CurrentSession->bIsInEnterprise = IProjectManager::Get().IsEnterpriseProject();
 			CurrentSession->bIsInPIE = FPlayWorldCommandCallbacks::IsInPIE();
-#endif
 
 			TrySaveCurrentSession();
 		}
@@ -191,14 +294,21 @@ void FEditorSessionSummaryWriter::Shutdown()
 	if (CurrentSession && !bShutdown)
 	{
 		// NOTE: Shutdown() may crash if a delegate is broadcasted from another thread at the same time (that's a bug in 4.25) the delegate are modified.
+		FEditorDelegates::PreBeginPIE.RemoveAll(this);
+		FEditorDelegates::EndPIE.RemoveAll(this);
 		FCoreDelegates::ApplicationWillTerminateDelegate.RemoveAll(this);
 		FCoreDelegates::IsVanillaProductChanged.RemoveAll(this);
 		FUserActivityTracking::OnActivityChanged.RemoveAll(this);
 		FSlateApplication::Get().GetOnModalLoopTickEvent().RemoveAll(this);
+		FSlateApplication::Get().GetLastUserInteractionTimeUpdateEvent().RemoveAll(this);
 		FCoreDelegates::OnHandleSystemError.RemoveAll(this);
 
 		CurrentSession->bWasShutdown = true;
-		UpdateTimestamps();
+		FDateTime CurrTimeUtc = FDateTime::UtcNow();
+		UpdateTimestamp(CurrTimeUtc);
+		UpdateEditorIdleTime(CurrTimeUtc, /*bSaveSession*/false); // Saved below, don't save twice.
+		UpdateUserIdleTime(CurrTimeUtc, /*bSaveSession*/false); // Saved below, don't save twice.
+		UpdateLegacyIdleTimes();
 		TrySaveCurrentSession();
 
 		CurrentSession.Reset();
@@ -267,6 +377,12 @@ TUniquePtr<FEditorAnalyticsSession> FEditorSessionSummaryWriter::CreateCurrentSe
 	Session->OSVersion = FPlatformMisc::GetOSVersion();
 	Session->bIs64BitOS = FPlatformMisc::Is64bitOperatingSystem();
 
+	extern ENGINE_API float GAverageFPS;
+	Session->AverageFPS = GAverageFPS;
+	Session->bIsInVRMode = IVREditorModule::Get().IsVREditorModeActive();
+	Session->bIsInEnterprise = IProjectManager::Get().IsEnterpriseProject();
+	Session->bIsInPIE = FPlayWorldCommandCallbacks::IsInPIE();
+
 	TArray<TSharedRef<IPlugin>> Plugins = IPluginManager::Get().GetEnabledPlugins();
 	for (const TSharedRef<IPlugin>& Plugin : Plugins)
 	{
@@ -284,12 +400,20 @@ void FEditorSessionSummaryWriter::OnCrashing()
 	// NOTE: This function is called from the crashing thread or the a crash processing thread and is concurrent with other functions such as Tick(), Initialize() or Shutdown() running on the game thread.
 	if (CurrentSession != nullptr)
 	{
-		UpdateIdleTimes();
-		CurrentSession->LogEvent(FEditorAnalyticsSession::EEventType::Crashed, FDateTime::UtcNow());
+		FDateTime CurrTime = FDateTime::UtcNow();
+
+		// If the main loop was not ticking, consider the Editor as 'active'. It was doing something or was waiting for something.
+		if (!EditorSessionWriterUtils::IsMainLoopStaled(CurrTime, LastTickUtcTime))
+		{
+			UpdateEditorIdleTime(CurrTime, /*bSaveSession*/false);
+		}
+		UpdateUserIdleTime(CurrTime, /*bSaveSession*/false);
+		UpdateLegacyIdleTimes();
+		CurrentSession->LogEvent(FEditorAnalyticsSession::EEventType::Crashed, CurrTime);
 
 		if (GIsGPUCrashed)
 		{
-			CurrentSession->LogEvent(FEditorAnalyticsSession::EEventType::GpuCrashed, FDateTime::UtcNow());
+			CurrentSession->LogEvent(FEditorAnalyticsSession::EEventType::GpuCrashed, CurrTime);
 		}
 
 		// NOTE: Don't try to save the session, we don't know if the lock used to save the key-store is corrupted (or held by the crashing thread) when OnCrashing() is called from the crash handler thread.
@@ -301,8 +425,16 @@ void FEditorSessionSummaryWriter::OnTerminate()
 	// NOTE: This function can be called from any thread (from the crashing thread too) and is likely concurrent with other functions such as Tick(), Initiallize() or Shutdown() running on the game thread.
 	if (CurrentSession != nullptr)
 	{
-		UpdateIdleTimes();
-		CurrentSession->LogEvent(FEditorAnalyticsSession::EEventType::Terminated, FDateTime::UtcNow());
+		FDateTime CurrTime = FDateTime::UtcNow();
+		
+		// If the main loop was not ticking, consider the Editor as 'active'. It was doing something or was waiting for something.
+		if (!EditorSessionWriterUtils::IsMainLoopStaled(CurrTime, LastTickUtcTime))
+		{
+			UpdateEditorIdleTime(CurrTime, /*bSaveSession*/false);
+		}
+		UpdateUserIdleTime(CurrTime, /*bSaveSession*/false);
+		UpdateLegacyIdleTimes();
+		CurrentSession->LogEvent(FEditorAnalyticsSession::EEventType::Terminated, CurrTime);
 
 		// NOTE: Don't try to save the session, we don't know if this is called from a crash handler (and if the crashing thread corrupted (or held) the lock to save the key-store.)
 		// NOTE: Don't explicitely Shutdown(), it is expected to be called on game thread to prevent unregistered delegate from a random thread. Just let the normal flow call Shutdown() or not. Destructor will do in last resort.
@@ -341,6 +473,34 @@ FString FEditorSessionSummaryWriter::GetUserActivityString()
 	return UserActivity.ActionName;
 }
 
+void FEditorSessionSummaryWriter::OnSlateUserInteraction(double CurrSlateInteractionTime)
+{
+	LastSlateInteractionTime = CurrSlateInteractionTime;
+	UpdateLegacyIdleTimes();
+
+	FDateTime CurrTimeUtc = FDateTime::UtcNow();
+	UpdateUserIdleTime(CurrTimeUtc, /*bSaveSession*/true);
+	UpdateEditorIdleTime(CurrTimeUtc, /*bSaveSession*/true);
+}
+
+void FEditorSessionSummaryWriter::OnEnterPIE(const bool /*bIsSimulating*/)
+{
+	if (CurrentSession != nullptr)
+	{
+		CurrentSession->bIsInPIE = true;
+		TrySaveCurrentSession();
+	}
+}
+
+void FEditorSessionSummaryWriter::OnExitPIE(const bool /*bIsSimulating*/)
+{
+	if (CurrentSession != nullptr)
+	{
+		CurrentSession->bIsInPIE = false;
+		TrySaveCurrentSession();
+	}
+}
+
 void FEditorSessionSummaryWriter::TrySaveCurrentSession()
 {
 	if (FEditorAnalyticsSession::TryLock()) // Inter-process lock to grant this process exclusive access to the key-store file/registry.
diff --git a/Engine/Source/Runtime/Engine/Private/Analytics/EditorSessionSummaryWriter.h b/Engine/Source/Runtime/Engine/Private/Analytics/EditorSessionSummaryWriter.h
--- a/Engine/Source/Runtime/Engine/Private/Analytics/EditorSessionSummaryWriter.h
+++ b/Engine/Source/Runtime/Engine/Private/Analytics/EditorSessionSummaryWriter.h
@@ -3,6 +3,7 @@
 #pragma once
 
 #include "CoreMinimal.h"
+#include "Templates/Atomic.h"
 
 #if WITH_EDITOR
 
@@ -26,11 +27,16 @@ private:
 	void OnTerminate();
 	void OnUserActivity(const FUserActivity& UserActivity);
 	void OnVanillaStateChanged(bool bIsVanilla);
+	void OnSlateUserInteraction(double CurrSlateInteractionTime);
+	void OnEnterPIE(const bool /*bIsSimulating*/);
+	void OnExitPIE(const bool /*bIsSimulating*/);
 
 	static TUniquePtr<FEditorAnalyticsSession> CreateCurrentSession(uint32 OutOfProcessMonitorProcessId);
 	static FString GetUserActivityString();
-	void UpdateTimestamps();
-	void UpdateIdleTimes();
+	void UpdateTimestamp(const FDateTime& InCurrTimeUtc);
+	void UpdateEditorIdleTime(const FDateTime& InActivityTimeUtc, bool bSaveSession);
+	void UpdateUserIdleTime(const FDateTime& InUserActivityTimeUtc, bool bSaveSession);
+	void UpdateLegacyIdleTimes();
 	void TrySaveCurrentSession();
 
 private:
@@ -38,8 +44,13 @@ private:
 	FString CurrentSessionSectionName;
 	FCriticalSection SaveSessionLock;
 	float HeartbeatTimeElapsed;
+
 	bool bShutdown;
 	const uint32 OutOfProcessMonitorProcessId; // Non-zero if out-of process monitoring is set. To ensure one CrashReportClient(CRC) doesn't report the session of another CRC instance (race condition).
+	double LastSlateInteractionTime;
+	TAtomic<uint64> LastTickUtcTime;           // Time since the last 'Tick()' to detect when the main thread isn't ticking.
+	TAtomic<uint64> LastEditorActivityUtcTime; // Since the last user action or CPU burst usage.
+	TAtomic<uint64> LastUserActivityUtcTime;   // Since the last user input.
 };
 
 #endif
