Subsystems will not attempt to instance non authrative classes.

diff --git a/Engine/Source/Runtime/Engine/Private/Subsystems/SubsystemCollection.cpp b/Engine/Source/Runtime/Engine/Private/Subsystems/SubsystemCollection.cpp
--- a/Engine/Source/Runtime/Engine/Private/Subsystems/SubsystemCollection.cpp
+++ b/Engine/Source/Runtime/Engine/Private/Subsystems/SubsystemCollection.cpp
@@ -194,6 +194,12 @@ USubsystem* FSubsystemCollectionBase::AddAndInitializeSubsystem(UClass* Subsyste
 			// Catch any attempt to add a subsystem of the wrong type
 			checkf(SubsystemClass->IsChildOf(BaseType), TEXT("ClassType (%s) must be a subclass of BaseType(%s)."), *SubsystemClass->GetName(), *BaseType->GetName());
 
+			// Do not create instances of classes aren't authoritative
+			if (SubsystemClass->GetAuthoritativeClass() != SubsystemClass)
+			{	
+				return false;
+			}
+
 			const USubsystem* CDO = SubsystemClass->GetDefaultObject<USubsystem>();
 			if (CDO->ShouldCreateSubsystem(Outer))
 			{
