Hide ray tracing debug show flags from editor menu, since corresponding functionality is exposed via view modes menu instead. Also skip rendering reflections in ray tracing debug or path tracing view modes.

diff --git a/Engine/Source/Runtime/Engine/Public/ShowFlagsValues.inl b/Engine/Source/Runtime/Engine/Public/ShowFlagsValues.inl
--- a/Engine/Source/Runtime/Engine/Public/ShowFlagsValues.inl
+++ b/Engine/Source/Runtime/Engine/Public/ShowFlagsValues.inl
@@ -344,8 +344,8 @@ SHOWFLAG_FIXED_IN_SHIPPING(0, VREditing, SFG_Hidden, NSLOCTEXT("UnrealEd", "VREd
 SHOWFLAG_FIXED_IN_SHIPPING(0, OcclusionMeshes, SFG_Visualize, NSLOCTEXT("UnrealEd", "VisualizeOcclusionQueries", "Visualize Occlusion Queries"))
 
 // RHI_RAYTRACING begin
-SHOWFLAG_FIXED_IN_SHIPPING(0, PathTracing, SFG_Developer, NSLOCTEXT("UnrealEd", "PathTracing", "Path tracing"))
-SHOWFLAG_FIXED_IN_SHIPPING(0, RayTracingDebug, SFG_Developer, NSLOCTEXT("UnrealEd", "RayTracingDebug", "Ray tracing debug"))
+SHOWFLAG_FIXED_IN_SHIPPING(0, PathTracing, SFG_Hidden, NSLOCTEXT("UnrealEd", "PathTracing", "Path tracing"))
+SHOWFLAG_FIXED_IN_SHIPPING(0, RayTracingDebug, SFG_Hidden, NSLOCTEXT("UnrealEd", "RayTracingDebug", "Ray tracing debug"))
 // RHI_RAYTRACING end
 
 /** Enable the SkyAtmosphere visualization to be drawn on screen */
diff --git a/Engine/Source/Runtime/Renderer/Private/IndirectLightRendering.cpp b/Engine/Source/Runtime/Renderer/Private/IndirectLightRendering.cpp
--- a/Engine/Source/Runtime/Renderer/Private/IndirectLightRendering.cpp
+++ b/Engine/Source/Runtime/Renderer/Private/IndirectLightRendering.cpp
@@ -584,7 +584,10 @@ void FDeferredShadingSceneRenderer::RenderDeferredReflectionsAndSkyLighting(FRHI
 {
 	check(RHICmdList.IsOutsideRenderPass());
 
-	if (ViewFamily.EngineShowFlags.VisualizeLightCulling || !ViewFamily.EngineShowFlags.Lighting)
+	if (ViewFamily.EngineShowFlags.VisualizeLightCulling 
+		|| ViewFamily.EngineShowFlags.RayTracingDebug
+		|| ViewFamily.EngineShowFlags.PathTracing
+		|| !ViewFamily.EngineShowFlags.Lighting)
 	{
 		return;
 	}
