fix odd tooltip

diff --git a/Engine/Source/Runtime/Engine/Public/Animation/AnimTypes.h b/Engine/Source/Runtime/Engine/Public/Animation/AnimTypes.h
--- a/Engine/Source/Runtime/Engine/Public/Animation/AnimTypes.h
+++ b/Engine/Source/Runtime/Engine/Public/Animation/AnimTypes.h
@@ -571,8 +571,7 @@ namespace ECurveBlendOption
 {
 	enum Type
 	{
-		/* Last pose that contains valid curve value override it. 
-		 * Unfortunately this has been default behavior for a long time, so we still keep it. */
+		/* Last pose that contains valid curve value override it. */
 		Override, // redirect from MaxWeight old legacy behavior
 		/** Only set the value if the previous pose doesn't have the curve value. */
 		DoNotOverride, 
