Change UpdateTriMeshVertices ensure to CHAOS_ENSURE.

diff --git a/Engine/Source/Runtime/Engine/Private/PhysicsEngine/BodyInstance.cpp b/Engine/Source/Runtime/Engine/Private/PhysicsEngine/BodyInstance.cpp
--- a/Engine/Source/Runtime/Engine/Private/PhysicsEngine/BodyInstance.cpp
+++ b/Engine/Source/Runtime/Engine/Private/PhysicsEngine/BodyInstance.cpp
@@ -470,7 +470,7 @@ int32 FBodyInstance::GetAllShapes_AssumesLocked(TArray<FPhysicsShapeHandle>& Out
 void FBodyInstance::UpdateTriMeshVertices(const TArray<FVector> & NewPositions)
 {
 #if WITH_CHAOS || WITH_IMMEDIATE_PHYSX
-	ensure(false);
+	CHAOS_ENSURE(false);
 #else
 #if PHYSICS_INTERFACE_PHYSX
 	if (BodySetup.IsValid())
