Fix RVT being marked modified when it shouldn't be on building streaming mips
Fix incorrect asset name being shown in progress when building streaming mips
Fix RVT volume bounds actor not being set for actors outside of current package

diff --git a/Engine/Source/Editor/VirtualTexturingEditor/Private/RuntimeVirtualTextureBuild.cpp b/Engine/Source/Editor/VirtualTexturingEditor/Private/RuntimeVirtualTextureBuild.cpp
--- a/Engine/Source/Editor/VirtualTexturingEditor/Private/RuntimeVirtualTextureBuild.cpp
+++ b/Engine/Source/Editor/VirtualTexturingEditor/Private/RuntimeVirtualTextureBuild.cpp
@@ -10,6 +10,7 @@
 #include "SceneInterface.h"
 #include "VT/RuntimeVirtualTexture.h"
 #include "VT/RuntimeVirtualTextureRender.h"
+#include "VT/VirtualTextureBuilder.h"
 
 namespace
 {
@@ -126,7 +127,7 @@ namespace RuntimeVirtualTexture
 			return false;
 		}
 
-		if (InComponent->GetStreamLowMips() <= 0)
+		if (InComponent->NumStreamingMips() <= 0)
 		{
 			return false;
 		}
@@ -155,7 +156,7 @@ namespace RuntimeVirtualTexture
 		const int32 TextureSizeX = VTDesc.WidthInBlocks * VTDesc.BlockWidthInTiles * TileSize;
 		const int32 TextureSizeY = VTDesc.HeightInBlocks * VTDesc.BlockHeightInTiles * TileSize;
 		const int32 MaxLevel = (int32)FMath::CeilLogTwo(FMath::Max(VTDesc.BlockWidthInTiles, VTDesc.BlockHeightInTiles));
-		const int32 RenderLevel = FMath::Max(MaxLevel - InComponent->GetStreamLowMips() + 1, 0);
+		const int32 RenderLevel = FMath::Max(MaxLevel - InComponent->NumStreamingMips() + 1, 0);
 		const int32 ImageSizeX = FMath::Max(TileSize, TextureSizeX >> RenderLevel);
 		const int32 ImageSizeY = FMath::Max(TileSize, TextureSizeY >> RenderLevel);
 		const int32 NumTilesX = ImageSizeX / TileSize;
@@ -172,7 +173,7 @@ namespace RuntimeVirtualTexture
 		// Spin up slow task UI
 		const float TaskWorkRender = NumTilesX * NumTilesY;
 		const float TaskWorkBuildBulkData = NumTilesX * NumTilesY / 2;
-		FScopedSlowTask Task(TaskWorkRender + TaskWorkBuildBulkData, FText::AsCultureInvariant(RuntimeVirtualTexture->GetName()));
+		FScopedSlowTask Task(TaskWorkRender + TaskWorkBuildBulkData, FText::AsCultureInvariant(InComponent->GetStreamingTexture()->GetName()));
 		Task.MakeDialog(true);
 
 		// Allocate render targets for rendering out the runtime virtual texture tiles
diff --git a/Engine/Source/Runtime/Engine/Classes/Components/RuntimeVirtualTextureComponent.h b/Engine/Source/Runtime/Engine/Classes/Components/RuntimeVirtualTextureComponent.h
--- a/Engine/Source/Runtime/Engine/Classes/Components/RuntimeVirtualTextureComponent.h
+++ b/Engine/Source/Runtime/Engine/Classes/Components/RuntimeVirtualTextureComponent.h
@@ -21,23 +21,23 @@ protected:
 	URuntimeVirtualTexture* VirtualTexture = nullptr;
 
 	/** Texture object containing streamed low mips. */
-	UPROPERTY(EditAnywhere, BlueprintReadOnly, Category = VirtualTexture, meta = (DisplayName = "Streaming Mips"))
+	UPROPERTY(EditAnywhere, BlueprintReadOnly, DuplicateTransient, Category = VirtualTexture)
 	UVirtualTextureBuilder* StreamingTexture = nullptr;
 
 	/** Actor to copy the bounds from to set up the transform. */
-	UPROPERTY(EditAnywhere, DuplicateTransient, Category = TransformFromBounds, meta = (DisplayName = "Source Actor"))
-	AActor* BoundsSourceActor = nullptr;
+	UPROPERTY(EditAnywhere, Category = TransformFromBounds, meta = (DisplayName = "Source Actor"))
+	TSoftObjectPtr<AActor> BoundsSourceActor;
 
 	/** Number of low mips to serialize and stream for the virtual texture. This can reduce rendering update cost. */
-	UPROPERTY(EditAnywhere, Category = VirtualTextureBuild, meta = (UIMin = "0", UIMax = "6", DisplayName = "Number of low mips to stream to the virtual texture"))
+	UPROPERTY(EditAnywhere, Category = VirtualTextureBuild, meta = (UIMin = "0", UIMax = "6", DisplayName = "Streaming Mips"))
 	int32 StreamLowMips = 0;
 
 	/** Enable Crunch compression. ZLib compression is used when Crunch is disabled. */
-	UPROPERTY(EditAnywhere, AdvancedDisplay, Category = VirtualTextureBuild, meta = (DisplayName = "Enable Crunch compression"))
+	UPROPERTY(EditAnywhere, AdvancedDisplay, Category = VirtualTextureBuild, meta = (DisplayName = "Enable Crunch"))
 	bool bEnableCompressCrunch = false;
 
 	/** Use any streaming low mips when rendering in editor. Set true to view and debug the baked streaming low mips. */
-	UPROPERTY(EditAnywhere, AdvancedDisplay, Category = VirtualTextureBuild)
+	UPROPERTY(EditAnywhere, AdvancedDisplay, Category = VirtualTextureBuild, meta = (DisplayName = "View Streaming Mips in Editor"))
 	bool bUseStreamingLowMipsInEditor = false;
 
 public:
@@ -52,7 +52,7 @@ public:
 	FTransform GetVirtualTextureTransform() const;
 
 	/** Public getter for virtual texture streaming low mips */
-	int32 GetStreamLowMips() const { return FMath::Clamp(StreamLowMips, 0, 6); }
+	int32 NumStreamingMips() const { return FMath::Clamp(StreamLowMips, 0, 6); }
 
 	/** Get if we want to use any streaming low mips on this component. */
 	bool IsStreamingLowMips() const;
@@ -86,7 +86,7 @@ protected:
 	/** Calculate a hash used to determine if the StreamingTexture contents are valid for use. The hash doesn't include whether the contents are up to date. */
 	uint32 CalculateStreamingTextureSettingsHash() const;
 	/** Returns true if the StreamingTexure contents are valid for use. */
-	bool StreamingTextureIsValid() const;
+	bool IsStreamingTextureValid() const;
 
 public:
 	/** Scene proxy object. Managed by the scene but stored here. */
diff --git a/Engine/Source/Runtime/Engine/Private/Components/RuntimeVirtualTextureComponent.cpp b/Engine/Source/Runtime/Engine/Private/Components/RuntimeVirtualTextureComponent.cpp
--- a/Engine/Source/Runtime/Engine/Private/Components/RuntimeVirtualTextureComponent.cpp
+++ b/Engine/Source/Runtime/Engine/Private/Components/RuntimeVirtualTextureComponent.cpp
@@ -97,7 +97,7 @@ uint32 URuntimeVirtualTextureComponent::CalculateStreamingTextureSettingsHash()
 	return Settings.PackedValue;
 }
 
-bool URuntimeVirtualTextureComponent::StreamingTextureIsValid() const
+bool URuntimeVirtualTextureComponent::IsStreamingTextureValid() const
 {
 	return VirtualTexture != nullptr && StreamingTexture != nullptr && StreamingTexture->Texture != nullptr && StreamingTexture->BuildHash == CalculateStreamingTextureSettingsHash();
 }
@@ -110,7 +110,7 @@ bool URuntimeVirtualTextureComponent::IsStreamingLowMips() const
 		return false;
 	}
 #endif
-	return StreamLowMips > 0 && StreamingTextureIsValid();
+	return StreamLowMips > 0 && IsStreamingTextureValid();
 }
 
 #if WITH_EDITOR
@@ -119,9 +119,6 @@ void URuntimeVirtualTextureComponent::InitializeStreamingTexture(uint32 InSizeX,
 {
 	if (VirtualTexture != nullptr && StreamingTexture != nullptr)
 	{
-		VirtualTexture->Modify();
-		StreamingTexture->Modify();
-
 		// Release current runtime virtual texture producer.
 		// It may reference data inside the old StreamingTexture which could be garbage collected any time from now.
 		VirtualTexture->Release();
@@ -156,6 +153,7 @@ void URuntimeVirtualTextureComponent::InitializeStreamingTexture(uint32 InSizeX,
 		BuildDesc.InSizeY = InSizeY;
 		BuildDesc.InData = InData;
 
+		StreamingTexture->Modify();
 		StreamingTexture->BuildTexture(BuildDesc);
 
 		// Trigger refresh of the runtime virtual texture producer.
@@ -165,7 +163,7 @@ void URuntimeVirtualTextureComponent::InitializeStreamingTexture(uint32 InSizeX,
 
 void URuntimeVirtualTextureComponent::SetRotation()
 {
-	if (BoundsSourceActor != nullptr)
+	if (BoundsSourceActor.IsValid())
 	{
 		// Copy the source actor rotation and notify the parent actor
 		SetWorldRotation(BoundsSourceActor->GetTransform().GetRotation());
@@ -175,7 +173,7 @@ void URuntimeVirtualTextureComponent::SetRotation()
 
 void URuntimeVirtualTextureComponent::SetTransformToBounds()
 {
-	if (BoundsSourceActor != nullptr)
+	if (BoundsSourceActor.IsValid())
 	{
 		// Calculate the bounds in our local rotation space translated to the BoundsSourceActor center
 		const FQuat TargetRotation = GetComponentToWorld().GetRotation();
diff --git a/Engine/Source/Runtime/Renderer/Private/VT/RuntimeVirtualTextureSceneProxy.cpp b/Engine/Source/Runtime/Renderer/Private/VT/RuntimeVirtualTextureSceneProxy.cpp
--- a/Engine/Source/Runtime/Renderer/Private/VT/RuntimeVirtualTextureSceneProxy.cpp
+++ b/Engine/Source/Runtime/Renderer/Private/VT/RuntimeVirtualTextureSceneProxy.cpp
@@ -40,7 +40,7 @@ FRuntimeVirtualTextureSceneProxy::FRuntimeVirtualTextureSceneProxy(URuntimeVirtu
 		// The Producer object created here will be passed into the virtual texture system which will take ownership.
 		IVirtualTexture* Producer = new FRuntimeVirtualTextureProducer(Desc, ProducerId, MaterialType, bClearTextures, InComponent->GetScene(), Transform, Bounds);
 
-		int32 StreamLowMips = InComponent->IsStreamingLowMips() ? InComponent->GetStreamLowMips() : 0;
+		int32 StreamLowMips = InComponent->IsStreamingLowMips() ? InComponent->NumStreamingMips() : 0;
 		if (StreamLowMips > 0)
 		{
 			UVirtualTexture2D* StreamingTexture = InComponent->GetStreamingTexture()->Texture;
