Fixed HQ particle lights under spawn fraction.

diff --git a/Engine/Source/Runtime/Engine/Private/Particles/ParticleModules.cpp b/Engine/Source/Runtime/Engine/Private/Particles/ParticleModules.cpp
--- a/Engine/Source/Runtime/Engine/Private/Particles/ParticleModules.cpp
+++ b/Engine/Source/Runtime/Engine/Private/Particles/ParticleModules.cpp
@@ -3055,7 +3055,7 @@ void UParticleModuleLight::SpawnEx(FParticleEmitterInstance* Owner, int32 Offset
 		LightData.bHighQuality = bHighQualityLights;
 		LightData.LightId = 0;
 
-		if (bHighQualityLights && ParticleLightQuality > 1)
+		if (bHighQualityLights && ParticleLightQuality > 1 && LightData.bValid)
 		{		
 			LightData.LightId = SpawnHQLight(LightData, Particle, Owner);
 		}
