Adding some very verbose logging to reloading materials, to try and pin down a crash

diff --git a/Engine/Source/Runtime/Engine/Private/Materials/MaterialShared.cpp b/Engine/Source/Runtime/Engine/Private/Materials/MaterialShared.cpp
--- a/Engine/Source/Runtime/Engine/Private/Materials/MaterialShared.cpp
+++ b/Engine/Source/Runtime/Engine/Private/Materials/MaterialShared.cpp
@@ -203,6 +203,8 @@ bool ReloadMaterialResource(
 	FString Filename;
 	GetReloadInfo(PackageName, &Filename);
 
+	UE_LOG(LogMaterial, VeryVerbose, TEXT("Attempting to load material resources for package %s (file name: %s)."), *PackageName, *Filename);
+
 	FMaterialResourceProxyReader Ar(*Filename, OffsetToFirstResource, FeatureLevel, QualityLevel);
 	FMaterialResource& Tmp = *InOutMaterialResource;
 	Tmp.SerializeInlineShaderMap(Ar);
