Reset property dirty states after Garbage Collection.
This helps prevent issues with Object Properties being nulled out, but not receiving dirty notifications.

diff --git a/Engine/Source/Runtime/Engine/Private/RepLayout.cpp b/Engine/Source/Runtime/Engine/Private/RepLayout.cpp
--- a/Engine/Source/Runtime/Engine/Private/RepLayout.cpp
+++ b/Engine/Source/Runtime/Engine/Private/RepLayout.cpp
@@ -1251,10 +1251,14 @@ namespace UE4_RepLayout_Private
 #if WITH_PUSH_MODEL
 	static bool IsPropertyDirty(
 		const int32 ParentIndex,
+		const bool bRecentlyCollectedGarbage,
 		const FComparePropertiesSharedParams& SharedParams,
 		FComparePropertiesStackParams& StackParams)
 	{
-		return !(*SharedParams.PushModelProperties)[ParentIndex] || SharedParams.PushModelState->IsPropertyDirty(ParentIndex);
+		return !(*SharedParams.PushModelProperties)[ParentIndex] ||
+			SharedParams.PushModelState->IsPropertyDirty(ParentIndex) ||
+			(SharedParams.PushModelState->DidRecentlyCollectGarbage() &&
+			EnumHasAnyFlags(SharedParams.Parents[ParentIndex].Flags, ERepParentFlags::HasObjectProperties | ERepParentFlags::IsNetSerialize));
 	}
 #endif // WITH_PUSH_MODEL	
 }	
@@ -1268,6 +1272,8 @@ static void CompareParentProperties(
 #if WITH_PUSH_MODEL
 	if (SharedParams.PushModelState != nullptr)
 	{
+		const bool bRecentlyCollectedGarbage = SharedParams.PushModelState->DidRecentlyCollectGarbage();
+
 		// Typically, on an initial compare all properties will be dirty anyway.
 		// However, if an Actor is awakened from Dormancy, it's SendingRepState will have been recreated, invalidating
 		// its saved Role and RemoteRole (regardless of whether or not they were changed).
@@ -1284,7 +1290,7 @@ static void CompareParentProperties(
 		}
 
 		// If we're forcibly comparing all properties, then don't bother checking dirty state.
-		if (SharedParams.bForceFail)
+		if (UNLIKELY(SharedParams.bForceFail))
 		{
 			for (int32 ParentIndex = 0; ParentIndex < SharedParams.Parents.Num(); ++ParentIndex)
 			{
@@ -1298,7 +1304,7 @@ static void CompareParentProperties(
 		{
 			for (int32 ParentIndex = 0; ParentIndex < SharedParams.Parents.Num(); ++ParentIndex)
 			{
-				const bool bIsPropertyDirty = UE4_RepLayout_Private::IsPropertyDirty(ParentIndex, SharedParams, StackParams);
+				const bool bIsPropertyDirty = UE4_RepLayout_Private::IsPropertyDirty(ParentIndex, bRecentlyCollectedGarbage, SharedParams, StackParams);
 				const bool bDidPropertyChange = UE4_RepLayout_Private::CompareParentPropertyHelper(ParentIndex, SharedParams, StackParams);
 
 				ensureAlwaysMsgf(!bDidPropertyChange || bIsPropertyDirty, TEXT("Push Model Property changed value, but was not marked dirty! Property=%s"), *SharedParams.Parents[ParentIndex].Property->GetPathName());
@@ -1307,7 +1313,7 @@ static void CompareParentProperties(
 #endif // WITH_PUSH_VALIDATION_SUPPORT
 
 		// If we have full push model support, then we only need to check properties that are actually dirty.
-		else if (EnumHasAnyFlags(SharedParams.Flags, ERepLayoutFlags::FullPushSupport))
+		else if (EnumHasAnyFlags(SharedParams.Flags, ERepLayoutFlags::FullPushSupport) && !bRecentlyCollectedGarbage)
 		{
 			for (TConstSetBitIterator<> It = SharedParams.PushModelState->GetDirtyProperties(); It; ++It)
 			{
@@ -1318,7 +1324,7 @@ static void CompareParentProperties(
 		{
 			for (int32 ParentIndex = 0; ParentIndex < SharedParams.Parents.Num(); ++ParentIndex)
 			{
-				if (UE4_RepLayout_Private::IsPropertyDirty(ParentIndex, SharedParams, StackParams))
+				if (UE4_RepLayout_Private::IsPropertyDirty(ParentIndex, bRecentlyCollectedGarbage, SharedParams, StackParams))
 				{
 					UE4_RepLayout_Private::CompareParentPropertyHelper(ParentIndex, SharedParams, StackParams);
 				}
@@ -4850,30 +4856,49 @@ static FName NAME_Vector_NetQuantize(TEXT("Vector_NetQuantize"));
 static FName NAME_UniqueNetIdRepl(TEXT("UniqueNetIdRepl"));
 static FName NAME_RepMovement(TEXT("RepMovement"));
 
+struct FInitFromPropertySharedParams
+{
+	TArray<FRepLayoutCmd>& Cmds;
+	const UNetConnection* ServerConnection;
+	const int32 ParentIndex;
+	bool bHasObjectProperties = false;
+	bool bHasNetSerializeProperties = false;
+};
+
+struct FInitFromPropertyStackParams
+{
+	FProperty* Property;
+	int32 Offset;
+	int32 RelativeHandle;
+	int32 ParentChecksum;
+	int32 StaticArrayIndex;
+};
+
+static uint32 GetRepLayoutCmdCompatibleChecksum(
+	FInitFromPropertySharedParams& SharedParams,
+	const FInitFromPropertyStackParams& StackParams)
+{
+	return GetRepLayoutCmdCompatibleChecksum(StackParams.Property, SharedParams.ServerConnection, StackParams.StaticArrayIndex, StackParams.ParentChecksum);
+}
+
 static uint32 AddPropertyCmd(
-	TArray<FRepLayoutCmd>& Cmds,
-	FProperty* Property,
-	int32 Offset,
-	int32 RelativeHandle,
-	int32 ParentIndex,
-	uint32 ParentChecksum,
-	int32 StaticArrayIndex,
-	const UNetConnection* ServerConnection)
+	FInitFromPropertySharedParams& SharedParams,
+	const FInitFromPropertyStackParams& StackParams)
 {
 	SCOPE_CYCLE_COUNTER(STAT_RepLayout_AddPropertyCmd);
 
-	FRepLayoutCmd & Cmd = Cmds.AddZeroed_GetRef();
+	FRepLayoutCmd & Cmd = SharedParams.Cmds.AddZeroed_GetRef();
 
-	Cmd.Property = Property;
+	Cmd.Property = StackParams.Property;
 	Cmd.Type = ERepLayoutCmdType::Property;		// Initially set to generic type
-	Cmd.Offset = Offset;
-	Cmd.ElementSize = Property->ElementSize;
-	Cmd.RelativeHandle = RelativeHandle;
-	Cmd.ParentIndex = ParentIndex;
-	Cmd.CompatibleChecksum = GetRepLayoutCmdCompatibleChecksum(Property, ServerConnection, StaticArrayIndex, ParentChecksum);
+	Cmd.Offset = StackParams.Offset;
+	Cmd.ElementSize = Cmd.Property->ElementSize;
+	Cmd.RelativeHandle = StackParams.RelativeHandle;
+	Cmd.ParentIndex = SharedParams.ParentIndex;
+	Cmd.CompatibleChecksum = GetRepLayoutCmdCompatibleChecksum(SharedParams, StackParams);
 
-	FProperty* UnderlyingProperty = Property;
-	if (FEnumProperty* EnumProperty = CastField<FEnumProperty>(Property))
+	FProperty* UnderlyingProperty = Cmd.Property;
+	if (FEnumProperty* EnumProperty = CastField<FEnumProperty>(UnderlyingProperty))
 	{
 		UnderlyingProperty = EnumProperty->GetUnderlyingProperty();
 	}
@@ -4945,6 +4970,7 @@ static uint32 AddPropertyCmd(
 	}
 	else if (UnderlyingProperty->IsA(FObjectPropertyBase::StaticClass()))
 	{
+		SharedParams.bHasObjectProperties = true;
 		if (UnderlyingProperty->IsA(FSoftObjectProperty::StaticClass()))
 		{
 			Cmd.Type = ERepLayoutCmdType::PropertySoftObject;
@@ -4991,24 +5017,18 @@ static uint32 AddPropertyCmd(
 }
 
 static FORCEINLINE uint32 AddArrayCmd(
-	TArray<FRepLayoutCmd>& Cmds,
-	FArrayProperty* Property,
-	int32 Offset,
-	int32 RelativeHandle,
-	int32 ParentIndex,
-	uint32 ParentChecksum,
-	int32 StaticArrayIndex,
-	const UNetConnection* ServerConnection)
+	FInitFromPropertySharedParams& SharedParams,
+	const FInitFromPropertyStackParams StackParams)
 {
-	FRepLayoutCmd& Cmd = Cmds.AddZeroed_GetRef();
+	FRepLayoutCmd& Cmd = SharedParams.Cmds.AddZeroed_GetRef();
 
 	Cmd.Type = ERepLayoutCmdType::DynamicArray;
-	Cmd.Property = Property;
-	Cmd.Offset = Offset;
-	Cmd.ElementSize = Property->Inner->ElementSize;
-	Cmd.RelativeHandle = RelativeHandle;
-	Cmd.ParentIndex = ParentIndex;
-	Cmd.CompatibleChecksum = GetRepLayoutCmdCompatibleChecksum(Property, ServerConnection, StaticArrayIndex, ParentChecksum);
+	Cmd.Property = StackParams.Property;
+	Cmd.Offset = StackParams.Offset;
+	Cmd.ElementSize = static_cast<FArrayProperty*>(StackParams.Property)->Inner->ElementSize;
+	Cmd.RelativeHandle = StackParams.RelativeHandle;
+	Cmd.ParentIndex = SharedParams.ParentIndex;
+	Cmd.CompatibleChecksum = GetRepLayoutCmdCompatibleChecksum(SharedParams, StackParams);
 
 	return Cmd.CompatibleChecksum;
 }
@@ -5039,37 +5059,33 @@ const FORCEINLINE int32 GetOffsetForProperty<ERepBuildType::Function>(FProperty&
 
 template<ERepBuildType BuildType>
 static int32 InitFromProperty_r(
-	TArray<FRepLayoutCmd>& Cmds,
-	FProperty* Property,
-	int32 Offset,
-	int32 RelativeHandle,
-	int32 ParentIndex,
-	uint32 ParentChecksum,
-	int32 StaticArrayIndex,
-	const UNetConnection* ServerConnection)
+	FInitFromPropertySharedParams& SharedParams,
+	FInitFromPropertyStackParams StackParams)
 {
-	FArrayProperty * ArrayProp = CastField<FArrayProperty>(Property);
-
-	if (ArrayProp != NULL)
+	if (FArrayProperty* ArrayProp = CastField<FArrayProperty>(StackParams.Property))
 	{
-		const int32 CmdStart = Cmds.Num();
+		const int32 CmdStart = SharedParams.Cmds.Num();
 
-		RelativeHandle++;
+		++StackParams.RelativeHandle;
+		StackParams.Offset += GetOffsetForProperty<BuildType>(*ArrayProp);
 
-		const uint32 ArrayChecksum = AddArrayCmd(Cmds, ArrayProp, Offset + GetOffsetForProperty<BuildType>(*ArrayProp), RelativeHandle, ParentIndex, ParentChecksum, StaticArrayIndex, ServerConnection);
+		const uint32 ArrayChecksum = AddArrayCmd(SharedParams, StackParams);
 
-		InitFromProperty_r<BuildType>(Cmds, ArrayProp->Inner, 0, 0, ParentIndex, ArrayChecksum, 0, ServerConnection);
+		FInitFromPropertyStackParams NewStackParams{
+			/*Property=*/ArrayProp->Inner,
+			/*Offset=*/0,
+			/*RelativeHandle=*/0,
+			/*ParentChecksum=*/ArrayChecksum,
+			/*StaticArrayIndex=*/0
+		};
 
-		AddReturnCmd(Cmds);
+		InitFromProperty_r<BuildType>(SharedParams, NewStackParams);
 
-		Cmds[CmdStart].EndCmd = Cmds.Num();		// Patch in the offset to jump over our array inner elements
+		AddReturnCmd(SharedParams.Cmds);
 
-		return RelativeHandle;
+		SharedParams.Cmds[CmdStart].EndCmd = SharedParams.Cmds.Num();		// Patch in the offset to jump over our array inner elements
 	}
-
-	FStructProperty* StructProp = CastField<FStructProperty>(Property);
-
-	if (StructProp != NULL)
+	else if (FStructProperty* StructProp = CastField<FStructProperty>(StackParams.Property))
 	{
 		UScriptStruct* Struct = StructProp->Struct;
 
@@ -5077,9 +5093,13 @@ static int32 InitFromProperty_r(
 		{
 			UE_CLOG(EnumHasAnyFlags(Struct->StructFlags, STRUCT_NetDeltaSerializeNative), LogRep, Warning, TEXT("RepLayout InitFromProperty_r: Struct marked both NetSerialize and NetDeltaSerialize: %s"), *StructProp->GetName());
 
-			RelativeHandle++;
-			AddPropertyCmd(Cmds, Property, Offset + GetOffsetForProperty<BuildType>(*Property), RelativeHandle, ParentIndex, ParentChecksum, StaticArrayIndex, ServerConnection);
-			return RelativeHandle;
+			SharedParams.bHasNetSerializeProperties = true;
+
+			++StackParams.RelativeHandle;
+			StackParams.Offset += GetOffsetForProperty<BuildType>(*StructProp);
+
+			AddPropertyCmd(SharedParams, StackParams);
+			return StackParams.RelativeHandle;
 		}
 
 		// Track properties so we can ensure they are sorted by offsets at the end
@@ -5116,25 +5136,36 @@ static int32 InitFromProperty_r(
 
 		Sort(NetProperties.GetData(), NetProperties.Num(), FCompareUFieldOffsets());
 
-		const uint32 StructChecksum = GetRepLayoutCmdCompatibleChecksum(Property, ServerConnection, StaticArrayIndex, ParentChecksum);
+		const uint32 StructChecksum = GetRepLayoutCmdCompatibleChecksum(SharedParams, StackParams);
 
 		for (int32 i = 0; i < NetProperties.Num(); i++)
 		{
 			for (int32 j = 0; j < NetProperties[i]->ArrayDim; j++)
 			{
 				const int32 ArrayElementOffset = j * NetProperties[i]->ElementSize;
-				RelativeHandle = InitFromProperty_r<BuildType>(Cmds, NetProperties[i], Offset + GetOffsetForProperty<BuildType>(*StructProp) + ArrayElementOffset, RelativeHandle, ParentIndex, StructChecksum, j, ServerConnection);
+
+				FInitFromPropertyStackParams NewStackParams{
+					/*Property=*/NetProperties[i],
+					/*Offset=*/StackParams.Offset + GetOffsetForProperty<BuildType>(*StructProp),
+					/*RelativeHandle=*/StackParams.RelativeHandle,
+					/*ParentChecksum=*/StructChecksum,
+					/*StaticArrayIndex=*/j
+				};
+
+				StackParams.RelativeHandle = InitFromProperty_r<BuildType>(SharedParams, NewStackParams);
 			}
 		}
-		return RelativeHandle;
 	}
+	else
+	{
+		// Add actual property
+		++StackParams.RelativeHandle;
+		StackParams.Offset += GetOffsetForProperty<BuildType>(*StackParams.Property);
 
-	// Add actual property
-	RelativeHandle++;
-
-	AddPropertyCmd(Cmds, Property, Offset + GetOffsetForProperty<BuildType>(*Property), RelativeHandle, ParentIndex, ParentChecksum, StaticArrayIndex, ServerConnection);
+		AddPropertyCmd(SharedParams, StackParams);
+	}
 
-	return RelativeHandle;
+	return StackParams.RelativeHandle;
 }
 
 static FORCEINLINE uint16 AddParentProperty(
@@ -5444,8 +5475,24 @@ void FRepLayout::InitFromClass(
 
 		const int32 ParentOffset = Property->ElementSize * ArrayIdx;
 
+		FInitFromPropertySharedParams SharedParams
+		{
+			/*Cmds=*/Cmds,
+			/*ServerConnection=*/ServerConnection,
+			/*ParentHandle=*/ParentHandle
+		};
+
+		FInitFromPropertyStackParams StackParams
+		{
+			/*Property=*/Property,
+			/*Offset=*/ParentOffset,
+			/*RelativeHandle=*/RelativeHandle,
+			/*ParentChecksum=*/0,
+			/*StaticArrayIndex=*/ArrayIdx
+		};
+
 		Parents[ParentHandle].CmdStart = Cmds.Num();
-		RelativeHandle = InitFromProperty_r<ERepBuildType::Class>(Cmds, Property, ParentOffset, RelativeHandle, ParentHandle, 0, ArrayIdx, ServerConnection);
+		RelativeHandle = InitFromProperty_r<ERepBuildType::Class>(SharedParams, StackParams);
 		Parents[ParentHandle].CmdEnd = Cmds.Num();
 		Parents[ParentHandle].Flags |= ERepParentFlags::IsConditional;
 		Parents[ParentHandle].Offset = GetOffsetForProperty<ERepBuildType::Class>(*Property) + ParentOffset;
@@ -5468,6 +5515,15 @@ void FRepLayout::InitFromClass(
 		{
 			HighestCustomDeltaRepIndex = ParentHandle;
 		}
+
+		if (SharedParams.bHasNetSerializeProperties)
+		{
+			Parents[ParentHandle].Flags |= ERepParentFlags::HasNetSerializeProperties;
+		}
+		if (SharedParams.bHasObjectProperties)
+		{
+			Parents[ParentHandle].Flags |= ERepParentFlags::HasObjectProperties;
+		}
 	}
 
 	// Make sure RemoteRole has a lower RepIndex than Role, otherwise assumptions RepLayout may break.
@@ -5533,6 +5589,11 @@ void FRepLayout::InitFromClass(
 				Parents[ParentIndex].Flags &= ~ERepParentFlags::IsConditional;
 			}
 
+			if (EnumHasAnyFlags(Parents[ParentIndex].Flags, ERepParentFlags::HasNetSerializeProperties | ERepParentFlags::HasObjectProperties))
+			{
+				Flags |= ERepLayoutFlags::HasObjectOrNetSerializeProperties;
+			}
+
 			++NumberOfLifetimeProperties;
 #if WITH_PUSH_MODEL
 			if (bIsPushModelEnabled && LifetimeProps[i].bIsPushBased)
@@ -5683,8 +5744,25 @@ void FRepLayout::InitFromFunction(
 		for (int32 ArrayIdx = 0; ArrayIdx < It->ArrayDim; ++ArrayIdx)
 		{
 			const int32 ParentHandle = AddParentProperty(Parents, *It, ArrayIdx);
+
+			FInitFromPropertySharedParams SharedParams
+			{
+				/*Cmds=*/Cmds,
+				/*ServerConnection=*/ServerConnection,
+				/*ParentHandle=*/ParentHandle
+			};
+
+			FInitFromPropertyStackParams StackParams
+			{
+				/*Property=*/*It,
+				/*Offset=*/It->ElementSize* ArrayIdx,
+				/*RelativeHandle=*/RelativeHandle,
+				/*ParentChecksum=*/0,
+				/*StaticArrayIndex=*/ArrayIdx
+			};
+
 			Parents[ParentHandle].CmdStart = Cmds.Num();
-			RelativeHandle = InitFromProperty_r<ERepBuildType::Function>(Cmds, *It, It->ElementSize * ArrayIdx, RelativeHandle, ParentHandle, 0, ArrayIdx, ServerConnection);
+			RelativeHandle = InitFromProperty_r<ERepBuildType::Function>(SharedParams, StackParams);
 			Parents[ParentHandle].CmdEnd = Cmds.Num();
 			Parents[ParentHandle].Offset = GetOffsetForProperty<ERepBuildType::Function>(**It);
 
@@ -5731,8 +5809,25 @@ void FRepLayout::InitFromStruct(
 		for (int32 ArrayIdx = 0; ArrayIdx < It->ArrayDim; ++ArrayIdx)
 		{
 			const int32 ParentHandle = AddParentProperty(Parents, *It, ArrayIdx);
+
+			FInitFromPropertySharedParams SharedParams
+			{
+				/*Cmds=*/Cmds,
+				/*ServerConnection=*/ServerConnection,
+				/*ParentHandle=*/ParentHandle
+			};
+
+			FInitFromPropertyStackParams StackParams
+			{
+				/*Property=*/*It,
+				/*Offset=*/It->ElementSize * ArrayIdx,
+				/*RelativeHandle=*/RelativeHandle,
+				/*ParentChecksum=*/0,
+				/*StaticArrayIndex=*/ArrayIdx
+			};
+
 			Parents[ParentHandle].CmdStart = Cmds.Num();
-			RelativeHandle = InitFromProperty_r<ERepBuildType::Struct>(Cmds, *It, It->ElementSize * ArrayIdx, RelativeHandle, ParentHandle, 0, ArrayIdx, ServerConnection);
+			RelativeHandle = InitFromProperty_r<ERepBuildType::Struct>(SharedParams, StackParams);
 			Parents[ParentHandle].CmdEnd = Cmds.Num();
 			Parents[ParentHandle].Offset = GetOffsetForProperty<ERepBuildType::Struct>(**It);
 
diff --git a/Engine/Source/Runtime/Engine/Public/Net/RepLayout.h b/Engine/Source/Runtime/Engine/Public/Net/RepLayout.h
--- a/Engine/Source/Runtime/Engine/Public/Net/RepLayout.h
+++ b/Engine/Source/Runtime/Engine/Public/Net/RepLayout.h
@@ -706,16 +706,18 @@ enum class ERepLayoutCmdType : uint8
 /** Various flags that describe how a Top Level Property should be handled. */
 enum class ERepParentFlags : uint32
 {
-	None				= 0,
-	IsLifetime			= (1 << 0),	//! This property is valid for the lifetime of the object (almost always set).
-	IsConditional		= (1 << 1),	//! This property has a secondary condition to check
-	IsConfig			= (1 << 2),	//! This property is defaulted from a config file
-	IsCustomDelta		= (1 << 3),	//! This property uses custom delta compression. Mutually exclusive with IsNetSerialize.
-	IsNetSerialize		= (1 << 4), //! This property uses a custom net serializer. Mutually exclusive with IsCustomDelta.
-	IsStructProperty	= (1 << 5),	//! This property is a FStructProperty.
-	IsZeroConstructible	= (1 << 6),	//! This property is ZeroConstructible.
-	IsFastArray			= (1 << 7),	//! This property is a FastArraySerializer. This can't be a ERepLayoutCmdType, because
-									//! these Custom Delta structs will have their inner properties tracked.
+	None						= 0,
+	IsLifetime					= (1 << 0),	//! This property is valid for the lifetime of the object (almost always set).
+	IsConditional				= (1 << 1),	//! This property has a secondary condition to check
+	IsConfig					= (1 << 2),	//! This property is defaulted from a config file
+	IsCustomDelta				= (1 << 3),	//! This property uses custom delta compression. Mutually exclusive with IsNetSerialize.
+	IsNetSerialize				= (1 << 4), //! This property uses a custom net serializer. Mutually exclusive with IsCustomDelta.
+	IsStructProperty			= (1 << 5),	//! This property is a FStructProperty.
+	IsZeroConstructible			= (1 << 6),	//! This property is ZeroConstructible.
+	IsFastArray					= (1 << 7),	//! This property is a FastArraySerializer. This can't be a ERepLayoutCmdType, because
+											//! these Custom Delta structs will have their inner properties tracked.
+	HasObjectProperties			= (1 << 8), //! This property is tracking UObjects (may be through nested properties).
+	HasNetSerializeProperties	= (1 << 9), //! This property contains Net Serialize properties (may be through nested properties).
 };
 
 ENUM_CLASS_FLAGS(ERepParentFlags)
@@ -1007,10 +1009,11 @@ ENUM_CLASS_FLAGS(ECreateRepLayoutFlags);
 
 enum class ERepLayoutFlags : uint8
 {
-	None				= 0,
-	IsActor 			= 1 << 0,	//! This RepLayout is for AActor or a subclass of AActor.
-	PartialPushSupport	= 1 << 1,	//! This RepLayout has some properties that use Push Model and some that don't.
-	FullPushSupport		= 1 << 2,	//! All properties in this RepLayout use Push Model.
+	None								= 0,
+	IsActor 							= (1 << 0),	//! This RepLayout is for AActor or a subclass of AActor.
+	PartialPushSupport					= (1 << 1),	//! This RepLayout has some properties that use Push Model and some that don't.
+	FullPushSupport						= (1 << 2),	//! All properties in this RepLayout use Push Model.
+	HasObjectOrNetSerializeProperties	= (1 << 3),	//! Will be set for any RepLayout that contains Object or Net Serialize property commands.
 };
 ENUM_CLASS_FLAGS(ERepLayoutFlags);
 
diff --git a/Engine/Source/Runtime/Net/Core/Private/Net/Core/PushModel/PushModel.cpp b/Engine/Source/Runtime/Net/Core/Private/Net/Core/PushModel/PushModel.cpp
--- a/Engine/Source/Runtime/Net/Core/Private/Net/Core/PushModel/PushModel.cpp
+++ b/Engine/Source/Runtime/Net/Core/Private/Net/Core/PushModel/PushModel.cpp
@@ -11,7 +11,9 @@
 #include "Types/PushModelPerObjectState.h"
 #include "Types/PushModelPerNetDriverState.h"
 #include "UObject/UObjectGlobals.h"
+#include "Stats/Stats2.h"
 
+DECLARE_CYCLE_STAT(TEXT("PushModel PostGarbageCollect"), STAT_PushModel_PostGarbageCollect, STATGROUP_Net);
 
 namespace UE4PushModelPrivate
 {
@@ -237,6 +239,8 @@ namespace UE4PushModelPrivate
 
 		void PostGarbageCollect()
 		{
+			SCOPE_CYCLE_COUNTER(STAT_PushModel_PostGarbageCollect);
+
 			// We can't compact PerObjectStates because we need ObjectIDs to be stable.
 			// But we can shrink it.
 
@@ -263,6 +267,10 @@ namespace UE4PushModelPrivate
 					ObjectIdToInternalId.Remove(It->GetObjectId());
 					It.RemoveCurrent();
 				}
+				else
+				{
+					It->SetRecentlyCollectedGarbage();
+				}
 			}
 
 			PerObjectStates.Shrink();
diff --git a/Engine/Source/Runtime/Net/Core/Private/Net/Core/PushModel/Types/PushModelPerObjectState.h b/Engine/Source/Runtime/Net/Core/Private/Net/Core/PushModel/Types/PushModelPerObjectState.h
--- a/Engine/Source/Runtime/Net/Core/Private/Net/Core/PushModel/Types/PushModelPerObjectState.h
+++ b/Engine/Source/Runtime/Net/Core/Private/Net/Core/PushModel/Types/PushModelPerObjectState.h
@@ -44,6 +44,19 @@ namespace UE4PushModelPrivate
 		FPushModelPerObjectState(const FPushModelPerObjectState& Other) = delete;
 		FPushModelPerObjectState& operator=(const FPushModelPerObjectState& Other) = delete;
 
+		/**
+		 * Currently, there are no notifications when a property gets nulled out from Garbage Collection.
+		 * This should'nt happen frequently in networked scenarios as it can cause other bugs with things like
+		 * Dormancy. However, it is easily reproducible on Clients recording Replays.
+		 */
+		void SetRecentlyCollectedGarbage()
+		{
+			for (FPushModelPerNetDriverState& NetDriverObject : PerNetDriverStates)
+			{
+				NetDriverObject.SetRecentlyCollectedGarbage();
+			}
+		}
+
 		void MarkPropertyDirty(const uint16 RepIndex)
 		{
 			DirtiedThisFrame[RepIndex] = true;
