Fix type mismatch from arguments in RepLayout

#rnx

diff --git a/Engine/Source/Runtime/Engine/Private/RepLayout.cpp b/Engine/Source/Runtime/Engine/Private/RepLayout.cpp
--- a/Engine/Source/Runtime/Engine/Private/RepLayout.cpp
+++ b/Engine/Source/Runtime/Engine/Private/RepLayout.cpp
@@ -4870,7 +4870,7 @@ struct FInitFromPropertyStackParams
 	FProperty* Property;
 	int32 Offset;
 	int32 RelativeHandle;
-	int32 ParentChecksum;
+	uint32 ParentChecksum;
 	int32 StaticArrayIndex;
 };
 
