Fix bad return type
#rnx

diff --git a/Engine/Source/Runtime/Engine/Private/Subsystems/SubsystemCollection.cpp b/Engine/Source/Runtime/Engine/Private/Subsystems/SubsystemCollection.cpp
--- a/Engine/Source/Runtime/Engine/Private/Subsystems/SubsystemCollection.cpp
+++ b/Engine/Source/Runtime/Engine/Private/Subsystems/SubsystemCollection.cpp
@@ -197,7 +197,7 @@ USubsystem* FSubsystemCollectionBase::AddAndInitializeSubsystem(UClass* Subsyste
 			// Do not create instances of classes aren't authoritative
 			if (SubsystemClass->GetAuthoritativeClass() != SubsystemClass)
 			{	
-				return false;
+				return nullptr;
 			}
 
 			const USubsystem* CDO = SubsystemClass->GetDefaultObject<USubsystem>();
