Fix race conditions in FVirtualTextureChunkStreamingManager create/destroy/update.
Problems are all down to temporary insanity on my part when I last touched this code.

diff --git a/Engine/Source/Runtime/Engine/Private/ContentStreaming.cpp b/Engine/Source/Runtime/Engine/Private/ContentStreaming.cpp
--- a/Engine/Source/Runtime/Engine/Private/ContentStreaming.cpp
+++ b/Engine/Source/Runtime/Engine/Private/ContentStreaming.cpp
@@ -22,6 +22,7 @@
 #include "Animation/AnimationStreaming.h"
 #include "AudioStreamingCache.h"
 #include "AudioCompressionSettingsUtils.h"
+#include "VT/VirtualTextureChunkManager.h"
 
 /*-----------------------------------------------------------------------------
 	Globals.
@@ -779,10 +780,17 @@ FStreamingManagerCollection::FStreamingManagerCollection()
 
 	AnimationStreamingManager = new FAnimationStreamingManager();
 	AddStreamingManager(AnimationStreamingManager);
+
+	VirtualTextureStreamingManager = new FVirtualTextureChunkStreamingManager();
+	AddStreamingManager(VirtualTextureStreamingManager);
 }
 
 FStreamingManagerCollection::~FStreamingManagerCollection()
 {
+	RemoveStreamingManager(VirtualTextureStreamingManager);
+	delete VirtualTextureStreamingManager;
+	VirtualTextureStreamingManager = nullptr;
+
 	RemoveStreamingManager(AnimationStreamingManager);
 	delete AnimationStreamingManager;
 	AnimationStreamingManager = nullptr;
@@ -1019,6 +1027,12 @@ IAnimationStreamingManager& FStreamingManagerCollection::GetAnimationStreamingMa
 	return *AnimationStreamingManager;
 }
 
+FVirtualTextureChunkStreamingManager& FStreamingManagerCollection::GetVirtualTextureStreamingManager() const
+{
+	check(VirtualTextureStreamingManager);
+	return *VirtualTextureStreamingManager;
+}
+
 /** Don't stream world resources for the next NumFrames. */
 void FStreamingManagerCollection::SetDisregardWorldResourcesForFrames(int32 NumFrames )
 {
diff --git a/Engine/Source/Runtime/Engine/Private/VT/UploadingVirtualTexture.cpp b/Engine/Source/Runtime/Engine/Private/VT/UploadingVirtualTexture.cpp
--- a/Engine/Source/Runtime/Engine/Private/VT/UploadingVirtualTexture.cpp
+++ b/Engine/Source/Runtime/Engine/Private/VT/UploadingVirtualTexture.cpp
@@ -64,7 +64,8 @@ FUploadingVirtualTexture::FUploadingVirtualTexture(FVirtualTextureBuiltData* InD
 	HandlePerChunk.AddDefaulted(InData->Chunks.Num());
 	CodecPerChunk.AddDefaulted(InData->Chunks.Num());
 	InvalidChunks.Init(false, InData->Chunks.Num());
-	StreamingManager = FVirtualTextureChunkStreamingManager::AddRef();
+
+	StreamingManager = &IStreamingManager::Get().GetVirtualTextureStreamingManager();
 
 	INC_MEMORY_STAT_BY(STAT_TotalHeaderSize, InData->GetMemoryFootprint());
 	INC_MEMORY_STAT_BY(STAT_TileHeaderSize, InData->GetTileMemoryFootprint());
@@ -72,7 +73,6 @@ FUploadingVirtualTexture::FUploadingVirtualTexture(FVirtualTextureBuiltData* InD
 	INC_DWORD_STAT_BY(STAT_NumTileHeaders, InData->GetNumTileHeaders());
 }
 
-
 FUploadingVirtualTexture::~FUploadingVirtualTexture()
 {
 	check(Data);
@@ -83,7 +83,6 @@ FUploadingVirtualTexture::~FUploadingVirtualTexture()
 
 	// Complete all open transcode requests before deleting IFileCacheHandle objects
 	StreamingManager->WaitTasksFinished();
-	FVirtualTextureChunkStreamingManager::DecRef();
 
 	for (TUniquePtr<FVirtualTextureCodec>& Codec : CodecPerChunk)
 	{
diff --git a/Engine/Source/Runtime/Engine/Private/VT/VirtualTextureChunkManager.cpp b/Engine/Source/Runtime/Engine/Private/VT/VirtualTextureChunkManager.cpp
--- a/Engine/Source/Runtime/Engine/Private/VT/VirtualTextureChunkManager.cpp
+++ b/Engine/Source/Runtime/Engine/Private/VT/VirtualTextureChunkManager.cpp
@@ -22,27 +22,8 @@ static FAutoConsoleVariableRef CVarNumTranscodeRequests(
 	TEXT("Number of transcode request that can be in flight. default 32\n"),
 	ECVF_Default);
 
-/** 
- * FVirtualTextureChunkStreamingManager is a treated as a singleton referenced by multiple FUploadingVirtualTexture objects.
- * Use RefCount to ensure it is unregistered and destroyed on shutdown.
- */
-struct FVirtualTextureChunkStreamingManagerSingleton
-{
-	FVirtualTextureChunkStreamingManager* Ptr = nullptr;
-	int32 RefCount = 0;
-
-	~FVirtualTextureChunkStreamingManagerSingleton()
-	{
-		check(Ptr == nullptr);
-		check(RefCount == 0);
-	}
-
-} GVirtualTextureChunkStreamingManager;
-
-
 FVirtualTextureChunkStreamingManager::FVirtualTextureChunkStreamingManager()
 {
-	IStreamingManager::Get().AddStreamingManager(this);
 #if WITH_EDITOR
 	GetVirtualTextureChunkDDCCache()->Initialize();
 #endif
@@ -55,43 +36,27 @@ FVirtualTextureChunkStreamingManager::~FVirtualTextureChunkStreamingManager()
 #if WITH_EDITOR
 	GetVirtualTextureChunkDDCCache()->ShutDown();
 #endif
-	IStreamingManager::Get().RemoveStreamingManager(this);
-}
-
-FVirtualTextureChunkStreamingManager* FVirtualTextureChunkStreamingManager::AddRef()
-{
-	if (GVirtualTextureChunkStreamingManager.RefCount++ == 0)
-	{
-		GVirtualTextureChunkStreamingManager.Ptr = new FVirtualTextureChunkStreamingManager();
-	}
-	return GVirtualTextureChunkStreamingManager.Ptr;
-}
-
-void FVirtualTextureChunkStreamingManager::DecRef()
-{
-	check(GVirtualTextureChunkStreamingManager.RefCount > 0);
-	if (--GVirtualTextureChunkStreamingManager.RefCount == 0)
-	{
-		delete GVirtualTextureChunkStreamingManager.Ptr;
-		GVirtualTextureChunkStreamingManager.Ptr = nullptr;
-	}
 }
 
 void FVirtualTextureChunkStreamingManager::UpdateResourceStreaming(float DeltaTime, bool bProcessEverything /*= false*/)
 {
-	FVirtualTextureChunkStreamingManager* StreamingManager = this;
-
-	ENQUEUE_RENDER_COMMAND(UpdateVirtualTextureStreaming)(
-		[StreamingManager](FRHICommandListImmediate& RHICmdList)
+	// It's OK to execute update if virtual texturing is disabled but we early out anyway to avoid an unnecessary render thread task.
+	if (UseVirtualTexturing(GMaxRHIFeatureLevel))
 	{
+		FVirtualTextureChunkStreamingManager* StreamingManager = this;
+
+		ENQUEUE_RENDER_COMMAND(UpdateVirtualTextureStreaming)(
+			[StreamingManager](FRHICommandListImmediate& RHICmdList)
+			{
 #if WITH_EDITOR
-		GetVirtualTextureChunkDDCCache()->UpdateRequests();
+				GetVirtualTextureChunkDDCCache()->UpdateRequests();
 #endif // WITH_EDITOR
-		StreamingManager->TranscodeCache.RetireOldTasks(StreamingManager->UploadCache);
-		StreamingManager->UploadCache.UpdateFreeList();
+				StreamingManager->TranscodeCache.RetireOldTasks(StreamingManager->UploadCache);
+				StreamingManager->UploadCache.UpdateFreeList();
 
-		FVirtualTextureCodec::RetireOldCodecs();
-	});
+				FVirtualTextureCodec::RetireOldCodecs();
+			});
+	}
 }
 
 int32 FVirtualTextureChunkStreamingManager::BlockTillAllRequestsFinished(float TimeLimit /*= 0.0f*/, bool bLogResults /*= false*/)
@@ -102,7 +67,6 @@ int32 FVirtualTextureChunkStreamingManager::BlockTillAllRequestsFinished(float T
 
 void FVirtualTextureChunkStreamingManager::CancelForcedResources()
 {
-
 }
 
 FVTRequestPageResult FVirtualTextureChunkStreamingManager::RequestTile(FUploadingVirtualTexture* VTexture, const FVirtualTextureProducerHandle& ProducerHandle, uint8 LayerMask, uint8 vLevel, uint32 vAddress, EVTRequestPagePriority Priority)
diff --git a/Engine/Source/Runtime/Engine/Private/VT/VirtualTextureChunkManager.h b/Engine/Source/Runtime/Engine/Private/VT/VirtualTextureChunkManager.h
--- a/Engine/Source/Runtime/Engine/Private/VT/VirtualTextureChunkManager.h
+++ b/Engine/Source/Runtime/Engine/Private/VT/VirtualTextureChunkManager.h
@@ -3,16 +3,10 @@
 #pragma once
 
 #include "ContentStreaming.h"
-#include "Containers/Queue.h"
-#include "Containers/List.h"
-#include "VirtualTextureBuiltData.h"
 #include "VirtualTextureUploadCache.h"
 #include "VirtualTextureTranscodeCache.h"
 
-class FChunkProvider;
-struct FVirtualTextureChunkStreamingManager;
 class FUploadingVirtualTexture;
-class IMemoryReadStream;
 
 DECLARE_STATS_GROUP(TEXT("Virtual Texturing Paging"), STATGROUP_VTP, STATCAT_Advanced);
 
@@ -33,14 +27,10 @@ DECLARE_DWORD_COUNTER_STAT(TEXT("Num uploads"), STAT_VTP_NumUploads, STATGROUP_V
 
 struct FVirtualTextureChunkStreamingManager final  : public IStreamingManager
 {
-	static FVirtualTextureChunkStreamingManager* AddRef();
-	static void DecRef();
-
-private:
+public:
 	FVirtualTextureChunkStreamingManager();
 	virtual ~FVirtualTextureChunkStreamingManager();
 
-public:
 	// IStreamingManager interface
 	virtual void UpdateResourceStreaming(float DeltaTime, bool bProcessEverything = false) override;
 	virtual int32 BlockTillAllRequestsFinished(float TimeLimit = 0.0f, bool bLogResults = false) override;
@@ -60,4 +50,3 @@ private:
 	FVirtualTextureUploadCache UploadCache;
 	FVirtualTextureTranscodeCache TranscodeCache;
 };
-
diff --git a/Engine/Source/Runtime/Engine/Public/ContentStreaming.h b/Engine/Source/Runtime/Engine/Public/ContentStreaming.h
--- a/Engine/Source/Runtime/Engine/Public/ContentStreaming.h
+++ b/Engine/Source/Runtime/Engine/Public/ContentStreaming.h
@@ -703,6 +703,11 @@ struct FStreamingManagerCollection : public IStreamingManager
 	 */
 	ENGINE_API IAnimationStreamingManager& GetAnimationStreamingManager() const;
 
+	/**
+	 * Gets a reference to the Virtual Texture Streaming Manager
+	*/
+	ENGINE_API struct FVirtualTextureChunkStreamingManager& GetVirtualTextureStreamingManager() const;
+
 	/**
 	 * Adds a streaming manager to the array of managers to route function calls to.
 	 *
@@ -801,6 +806,9 @@ protected:
 	/** The animation streaming manager, should always exist */
 	IAnimationStreamingManager* AnimationStreamingManager;
 
+	/** The virtual texture streaming manager, should always exist */
+	FVirtualTextureChunkStreamingManager* VirtualTextureStreamingManager;
+
 #if WITH_EDITOR
 	// Locks out any audio streaming manager call when we are re-initializing the audio streaming manager.
 	mutable FCriticalSection AudioStreamingManagerCriticalSection;
