Make sure flushing chaos flushes any dirty proxies

diff --git a/Engine/Source/Runtime/Engine/Private/PhysicsEngine/Experimental/PhysScene_Chaos.cpp b/Engine/Source/Runtime/Engine/Private/PhysicsEngine/Experimental/PhysScene_Chaos.cpp
--- a/Engine/Source/Runtime/Engine/Private/PhysicsEngine/Experimental/PhysScene_Chaos.cpp
+++ b/Engine/Source/Runtime/Engine/Private/PhysicsEngine/Experimental/PhysScene_Chaos.cpp
@@ -1344,6 +1344,9 @@ void FPhysScene_ChaosInterface::Flush_AssumesLocked()
 
 	if(Solver)
 	{
+		//Make sure any dirty proxy data is pushed
+		Solver->PushPhysicsState(Dispatcher);
+
 		TQueue<TFunction<void(Chaos::FPhysicsSolver*)>, EQueueMode::Mpsc>& Queue = Solver->GetCommandQueue();
 		TFunction<void(Chaos::FPhysicsSolver*)> Command;
 		while(Queue.Dequeue(Command))
