Fix for crash when initializating foliage structures with ray tracing enabled

diff --git a/Engine/Source/Runtime/Engine/Private/InstancedStaticMesh.cpp b/Engine/Source/Runtime/Engine/Private/InstancedStaticMesh.cpp
--- a/Engine/Source/Runtime/Engine/Private/InstancedStaticMesh.cpp
+++ b/Engine/Source/Runtime/Engine/Private/InstancedStaticMesh.cpp
@@ -1252,11 +1252,6 @@ void FInstancedStaticMeshSceneProxy::SetupRayTracingCullClusters()
 	{
 		return;
 	}
-	
-	if (!InstancedRenderData.PerInstanceRenderData.IsValid())
-	{
-		return;
-	}
 
 	//#dxr_todo: select the appropriate LOD depending on Context.View
 	int32 LOD = 0;
@@ -1264,7 +1259,7 @@ void FInstancedStaticMeshSceneProxy::SetupRayTracingCullClusters()
 	{
 		const float MaxClusterRadiusMultiplier = CVarRayTracingInstancesCullClusterMaxRadiusMultiplier.GetValueOnAnyThread();
 		const int32 Batches = GetNumMeshBatches();
-		const int32 InstanceCount = InstancedRenderData.PerInstanceRenderData->InstanceBuffer.GetNumInstances();
+		const int32 InstanceCount = InstancedRenderData.Component->PerInstanceSMData.Num();
 		int32 ClusterIndex = 0;
 		// We're in game thread and at this point this scene proxy hasn't been added to FScene, thus GetLocalToWorld() returns undefined transform
 		FMatrix ComponentLocalToWorld = InstancedRenderData.Component->GetComponentTransform().ToMatrixWithScale();
@@ -1285,16 +1280,14 @@ void FInstancedStaticMeshSceneProxy::SetupRayTracingCullClusters()
 		{
 			if (InstancedRenderData.Component->PerInstanceSMData.IsValidIndex(Instance))
 			{
-				FMatrix Transform;
-				InstancedRenderData.PerInstanceRenderData->InstanceBuffer.GetInstanceTransform(Instance, Transform);
-				Transform.M[3][3] = 1.0f;
-				FMatrix InstanceTransform = Transform * ComponentLocalToWorld;
+				const FInstancedStaticMeshInstanceData& InstanceData = InstancedRenderData.Component->PerInstanceSMData[Instance];
+				FMatrix InstanceTransform = InstanceData.Transform * ComponentLocalToWorld;
+				InstanceTransform.M[3][3] = 1.0f;
 
 				FVector VMin, VMax;
 				InstancedRenderData.Component->GetLocalBounds(VMin, VMax);
 				VMin = InstanceTransform.TransformPosition(VMin);
 				VMax = InstanceTransform.TransformPosition(VMax);
-
 				FVector VBBoxSize = VMax - VMin;
 
 				MaxInstanceRadius = FMath::Max(0.5f * VBBoxSize.Size(), MaxInstanceRadius);
@@ -1311,54 +1304,56 @@ void FInstancedStaticMeshSceneProxy::SetupRayTracingCullClusters()
 		// Build clusters
 		for (int32 Instance = 0; Instance < InstanceCount; ++Instance)
 		{
-			FMatrix Transform;
-			InstancedRenderData.PerInstanceRenderData->InstanceBuffer.GetInstanceTransform(Instance, Transform);
-			Transform.M[3][3] = 1.0f;
-			FMatrix InstanceTransform = Transform * ComponentLocalToWorld;
-			FVector InstanceLocation = InstanceTransform.TransformPosition({ 0.0f,0.0f,0.0f });
-			FVector VMin = InstanceLocation - FVector(MaxInstanceRadius, MaxInstanceRadius, MaxInstanceRadius);
-			FVector VMax = InstanceLocation + FVector(MaxInstanceRadius, MaxInstanceRadius, MaxInstanceRadius);
-			bool bClusterFound = false;
-
-			// Try to find suitable cluster
-			for (int32 CandidateCluster = 0; CandidateCluster <= ClusterIndex; ++CandidateCluster)
+			if (InstancedRenderData.Component->PerInstanceSMData.IsValidIndex(Instance))
 			{
-				// Build new candidate cluster bounds
-				FVector VCandidateMin = VMin.ComponentMin(RayTracingCullClusters[CandidateCluster].BoundsMin);
-				FVector VCandidateMax = VMax.ComponentMax(RayTracingCullClusters[CandidateCluster].BoundsMax);
+				const FInstancedStaticMeshInstanceData& InstanceData = InstancedRenderData.Component->PerInstanceSMData[Instance];
+				FMatrix InstanceTransform = InstanceData.Transform * ComponentLocalToWorld;
+				InstanceTransform.M[3][3] = 1.0f;
+				FVector InstanceLocation = InstanceTransform.TransformPosition({ 0.0f,0.0f,0.0f });
+				FVector VMin = InstanceLocation - FVector(MaxInstanceRadius, MaxInstanceRadius, MaxInstanceRadius);
+				FVector VMax = InstanceLocation + FVector(MaxInstanceRadius, MaxInstanceRadius, MaxInstanceRadius);
+				bool bClusterFound = false;
+
+				// Try to find suitable cluster
+				for (int32 CandidateCluster = 0; CandidateCluster <= ClusterIndex; ++CandidateCluster)
+				{
+					// Build new candidate cluster bounds
+					FVector VCandidateMin = VMin.ComponentMin(RayTracingCullClusters[CandidateCluster].BoundsMin);
+					FVector VCandidateMax = VMax.ComponentMax(RayTracingCullClusters[CandidateCluster].BoundsMax);
 
-				FVector VCandidateBBoxSize = VCandidateMax - VCandidateMin;
-				float MaxCandidateRadius = 0.5f * VCandidateBBoxSize.Size();
+					FVector VCandidateBBoxSize = VCandidateMax - VCandidateMin;
+					float MaxCandidateRadius = 0.5f * VCandidateBBoxSize.Size();
 
-				// If new candidate is still small enough, update current cluster
-				if (MaxCandidateRadius <= MaxClusterRadius)
+					// If new candidate is still small enough, update current cluster
+					if (MaxCandidateRadius <= MaxClusterRadius)
+					{
+						RayTracingCullClusters[CandidateCluster].BoundsMin = VCandidateMin;
+						RayTracingCullClusters[CandidateCluster].BoundsMax = VCandidateMax;
+						FCullNode& Node = RayTracingCullClusters[CandidateCluster].Nodes.Emplace_GetRef();
+						Node.Radius = InstanceRadii[Instance];
+						Node.Center = InstanceLocation;
+						Node.Instance = Instance;
+						bClusterFound = true;
+						break;
+					}
+				}
+
+				// if we couldn't add the instance to an existing cluster create a new one
+				if (!bClusterFound)
 				{
-					RayTracingCullClusters[CandidateCluster].BoundsMin = VCandidateMin;
-					RayTracingCullClusters[CandidateCluster].BoundsMax = VCandidateMax;
-					FCullNode& Node = RayTracingCullClusters[CandidateCluster].Nodes.Emplace_GetRef();
+					++ClusterIndex;
+
+					FRayTracingCullCluster& Cluster = RayTracingCullClusters.Emplace_GetRef();
+					Cluster.BoundsMin = VMin;
+					Cluster.BoundsMax = VMax;
+					FCullNode Node;
+
 					Node.Radius = InstanceRadii[Instance];
 					Node.Center = InstanceLocation;
 					Node.Instance = Instance;
-					bClusterFound = true;
-					break;
+					Cluster.Nodes.Add(Node);
 				}
 			}
-
-			// if we couldn't add the instance to an existing cluster create a new one
-			if (!bClusterFound)
-			{
-				++ClusterIndex;
-
-				FRayTracingCullCluster &Cluster = RayTracingCullClusters.Emplace_GetRef();
-				Cluster.BoundsMin = VMin;
-				Cluster.BoundsMax = VMax;
-				FCullNode Node;
-
-				Node.Radius = InstanceRadii[Instance];
-				Node.Center = InstanceLocation;
-				Node.Instance = Instance;
-				Cluster.Nodes.Add(Node);
-			}
 		}
 	}
 }
