Fix for incorrect clock being used by FAudioDevice::GetAudioTime, thus causing
1. Sequencer to actually use FAudioDevice::Update to time the game, rather than the actual audio clock
2. Our fix last night to not resync the audio, which just happened to be saving us from going out of sync with the correct audio clock because it was reseeking the file.

diff --git a/Engine/Source/Runtime/AudioMixer/Private/AudioMixerDevice.cpp b/Engine/Source/Runtime/AudioMixer/Private/AudioMixerDevice.cpp
--- a/Engine/Source/Runtime/AudioMixer/Private/AudioMixerDevice.cpp
+++ b/Engine/Source/Runtime/AudioMixer/Private/AudioMixerDevice.cpp
@@ -46,7 +46,6 @@ namespace Audio
 	FMixerDevice::FMixerDevice(IAudioMixerPlatformInterface* InAudioMixerPlatform)
 		: AudioMixerPlatform(InAudioMixerPlatform)
 		, AudioClockDelta(0.0)
-		, AudioClock(0.0)
 		, PreviousMasterVolume((float)INDEX_NONE)
 		, GameOrAudioThreadId(INDEX_NONE)
 		, AudioPlatformThreadId(INDEX_NONE)
diff --git a/Engine/Source/Runtime/AudioMixer/Public/AudioMixerDevice.h b/Engine/Source/Runtime/AudioMixer/Public/AudioMixerDevice.h
--- a/Engine/Source/Runtime/AudioMixer/Public/AudioMixerDevice.h
+++ b/Engine/Source/Runtime/AudioMixer/Public/AudioMixerDevice.h
@@ -316,9 +316,6 @@ namespace Audio
 		/** The time delta for each callback block. */
 		double AudioClockDelta;
 
-		/** The audio clock from device initialization, updated at block rate. */
-		double AudioClock;
-
 		/** What the previous master volume was. */
 		float PreviousMasterVolume;
 
diff --git a/Engine/Source/Runtime/Engine/Private/AudioDevice.cpp b/Engine/Source/Runtime/Engine/Private/AudioDevice.cpp
--- a/Engine/Source/Runtime/Engine/Private/AudioDevice.cpp
+++ b/Engine/Source/Runtime/Engine/Private/AudioDevice.cpp
@@ -4193,12 +4193,6 @@ void FAudioDevice::Update(bool bGameTicking)
 		UpdateDeviceDeltaTime();
 	}
 
-	{
-		SCOPED_NAMED_EVENT(FAudioDevice_UpdateAudioClock, FColor::Blue);
-		// Update the audio clock, this can be overridden per platform to get a sample-accurate clock
-		UpdateAudioClock();
-	}
-
 	{
 		SCOPED_NAMED_EVENT(FAudioDevice_UpdateVirtualLoops, FColor::Blue);
 		// Update which loops should re-trigger due to coming back into proximity
diff --git a/Engine/Source/Runtime/Engine/Public/AudioDevice.h b/Engine/Source/Runtime/Engine/Public/AudioDevice.h
--- a/Engine/Source/Runtime/Engine/Public/AudioDevice.h
+++ b/Engine/Source/Runtime/Engine/Public/AudioDevice.h
@@ -1636,12 +1636,6 @@ public:
 		LastUpdateTime = CurrTime;
 	}
 
-	/** Update the audio clock to be based off the update delta time */
-	virtual void UpdateAudioClock()
-	{
-		AudioClock += GetDeviceDeltaTime();
-	}
-
 private:
 	/** Processes the set of pending sounds that need to be stopped */
 	void ProcessingPendingActiveSoundStops(bool bForceDelete = false);
