More confident fix of the sporadic cook warnings we see on iOS and Mac. Max H was locally able to hit a repro and confirm that there is a race between the async DDC worker that reports this log and USoundWave::WillNeverCacheCookedPlatformDataAgain.
After trying various fixes, the quickest and easiest fix for this is to simply ensure that the DDC worker is finished when this is called.

diff --git a/Engine/Source/Runtime/Engine/Private/AudioDerivedData.cpp b/Engine/Source/Runtime/Engine/Private/AudioDerivedData.cpp
--- a/Engine/Source/Runtime/Engine/Private/AudioDerivedData.cpp
+++ b/Engine/Source/Runtime/Engine/Private/AudioDerivedData.cpp
@@ -1622,6 +1622,8 @@ void USoundWave::ClearCachedCookedPlatformData( const ITargetPlatform* TargetPla
 
 void USoundWave::WillNeverCacheCookedPlatformDataAgain()
 {
+	FinishCachePlatformData();
+
 	// this is called after we have finished caching the platform data but before we have saved the data
 	// so need to keep the cached platform data around
 	Super::WillNeverCacheCookedPlatformDataAgain();
