Fix collapsed hierarchy when using root motion on additive animations

diff --git a/Engine/Source/Runtime/Engine/Private/Animation/AnimCompressionDerivedData.cpp b/Engine/Source/Runtime/Engine/Private/Animation/AnimCompressionDerivedData.cpp
--- a/Engine/Source/Runtime/Engine/Private/Animation/AnimCompressionDerivedData.cpp
+++ b/Engine/Source/Runtime/Engine/Private/Animation/AnimCompressionDerivedData.cpp
@@ -35,7 +35,7 @@ const TCHAR* FDerivedDataAnimationCompression::GetVersionString() const
 	// This is a version string that mimics the old versioning scheme. If you
 	// want to bump this version, generate a new guid using VS->Tools->Create GUID and
 	// return it here. Ex.
-	return TEXT("0B129554D9A4427784EC46E2AFD314B7");
+	return TEXT("28FD8FC5802144D29FF3E40C3BBE5DB1");
 }
 
 bool FDerivedDataAnimationCompression::Build( TArray<uint8>& OutDataArray )
diff --git a/Engine/Source/Runtime/Engine/Private/Animation/AnimSequence.cpp b/Engine/Source/Runtime/Engine/Private/Animation/AnimSequence.cpp
--- a/Engine/Source/Runtime/Engine/Private/Animation/AnimSequence.cpp
+++ b/Engine/Source/Runtime/Engine/Private/Animation/AnimSequence.cpp
@@ -1624,7 +1624,8 @@ void UAnimSequence::GetBonePose(FCompactPose& OutPose, FBlendedCurve& OutCurve,
 		return;
 	}
 
-	FRootMotionReset RootMotionReset(bEnableRootMotion, RootMotionRootLock, bForceRootLock, ExtractRootTrackTransform(0.f, &RequiredBones), IsValidAdditive());
+	const bool bTreatAnimAsAdditive = (IsValidAdditive() && !bUseRawDataForPoseExtraction); // Raw data is never additive
+	FRootMotionReset RootMotionReset(bEnableRootMotion, RootMotionRootLock, bForceRootLock, ExtractRootTrackTransform(0.f, &RequiredBones), bTreatAnimAsAdditive);
 
 #if WITH_EDITOR
 	// this happens only with editor data
