Fixed crash when selecting an anim instance to debug while pose watches are active

Prevented access to the bone container of an anim instance that has been garbage collected or MarkPendingKill-ed

diff --git a/Engine/Source/Editor/Persona/Private/AnimationEditorViewportClient.cpp b/Engine/Source/Editor/Persona/Private/AnimationEditorViewportClient.cpp
--- a/Engine/Source/Editor/Persona/Private/AnimationEditorViewportClient.cpp
+++ b/Engine/Source/Editor/Persona/Private/AnimationEditorViewportClient.cpp
@@ -1247,7 +1247,10 @@ void FAnimationViewportClient::DrawWatchedPoses(UDebugSkelMeshComponent * MeshCo
 			{
 				for (const FAnimNodePoseWatch& AnimNodePoseWatch : AnimBlueprintGeneratedClass->GetAnimBlueprintDebugData().AnimNodePoseWatch)
 				{
-					DrawBonesFromCompactPose(*AnimNodePoseWatch.PoseInfo.Get(), MeshComponent, PDI, AnimNodePoseWatch.PoseDrawColour);
+					if(AnimNodePoseWatch.Object.Get() != nullptr)
+					{
+						DrawBonesFromCompactPose(*AnimNodePoseWatch.PoseInfo.Get(), MeshComponent, PDI, AnimNodePoseWatch.PoseDrawColour);
+					}
 				}
 			}
 		}
diff --git a/Engine/Source/Runtime/Engine/Classes/Engine/PoseWatch.h b/Engine/Source/Runtime/Engine/Classes/Engine/PoseWatch.h
--- a/Engine/Source/Runtime/Engine/Classes/Engine/PoseWatch.h
+++ b/Engine/Source/Runtime/Engine/Classes/Engine/PoseWatch.h
@@ -11,6 +11,8 @@ struct FCompactHeapPose;
 
 struct FAnimNodePoseWatch
 {
+	// Object (anim instance) that this pose came from
+	TWeakObjectPtr<const UObject>	Object;
 	TSharedPtr<FCompactHeapPose>	PoseInfo;
 	FColor							PoseDrawColour;
 	int32							NodeID;
diff --git a/Engine/Source/Runtime/Engine/Private/Animation/AnimBlueprintGeneratedClass.cpp b/Engine/Source/Runtime/Engine/Private/Animation/AnimBlueprintGeneratedClass.cpp
--- a/Engine/Source/Runtime/Engine/Private/Animation/AnimBlueprintGeneratedClass.cpp
+++ b/Engine/Source/Runtime/Engine/Private/Animation/AnimBlueprintGeneratedClass.cpp
@@ -189,6 +189,7 @@ void FAnimBlueprintDebugData::AddPoseWatch(int32 NodeID, FColor Color)
 	NewAnimNodePoseWatch.NodeID = NodeID;
 	NewAnimNodePoseWatch.PoseDrawColour = Color;
 	NewAnimNodePoseWatch.PoseInfo = MakeShareable(new FCompactHeapPose());
+	NewAnimNodePoseWatch.Object = nullptr;
 }
 
 void FAnimBlueprintDebugData::RemovePoseWatch(int32 NodeID)
diff --git a/Engine/Source/Runtime/Engine/Private/Animation/AnimInstanceProxy.cpp b/Engine/Source/Runtime/Engine/Private/Animation/AnimInstanceProxy.cpp
--- a/Engine/Source/Runtime/Engine/Private/Animation/AnimInstanceProxy.cpp
+++ b/Engine/Source/Runtime/Engine/Private/Animation/AnimInstanceProxy.cpp
@@ -2488,6 +2488,7 @@ void FAnimInstanceProxy::RegisterWatchedPose(const FCompactPose& Pose, int32 Lin
 			if (PoseWatch.NodeID == LinkID)
 			{
 				PoseWatch.PoseInfo->CopyBonesFrom(Pose);
+				PoseWatch.Object = GetAnimInstanceObject();
 				break;
 			}
 		}
@@ -2505,6 +2506,7 @@ void FAnimInstanceProxy::RegisterWatchedPose(const FCSPose<FCompactPose>& Pose,
 				FCompactPose TempPose;
 				FCSPose<FCompactPose>::ConvertComponentPosesToLocalPoses(Pose, TempPose);
 				PoseWatch.PoseInfo->CopyBonesFrom(TempPose);
+				PoseWatch.Object = GetAnimInstanceObject();
 				break;
 			}
 		}
