Fix missing offset for Static Arrays nested within structs or dynamic arrays.

#rnx

diff --git a/Engine/Source/Runtime/Engine/Private/RepLayout.cpp b/Engine/Source/Runtime/Engine/Private/RepLayout.cpp
--- a/Engine/Source/Runtime/Engine/Private/RepLayout.cpp
+++ b/Engine/Source/Runtime/Engine/Private/RepLayout.cpp
@@ -5146,7 +5146,7 @@ static int32 InitFromProperty_r(
 
 				FInitFromPropertyStackParams NewStackParams{
 					/*Property=*/NetProperties[i],
-					/*Offset=*/StackParams.Offset + GetOffsetForProperty<BuildType>(*StructProp),
+					/*Offset=*/StackParams.Offset + GetOffsetForProperty<BuildType>(*StructProp) + ArrayElementOffset,
 					/*RelativeHandle=*/StackParams.RelativeHandle,
 					/*ParentChecksum=*/StructChecksum,
 					/*StaticArrayIndex=*/j
