Fix for sound class distance attenuation

diff --git a/Engine/Source/Runtime/Engine/Classes/Sound/SoundClass.h b/Engine/Source/Runtime/Engine/Classes/Sound/SoundClass.h
--- a/Engine/Source/Runtime/Engine/Classes/Sound/SoundClass.h
+++ b/Engine/Source/Runtime/Engine/Classes/Sound/SoundClass.h
@@ -137,25 +137,20 @@ struct FSoundClassProperties
 	UPROPERTY(EditAnywhere, BlueprintReadOnly, Category = Submix, meta = (EditCondition = "!bApplyEffects"))
 	USoundSubmix* DefaultSubmix;
 
-	// Sets the attenuation scale of the sound class in the given amount of time
-	void SetAttenuationDistanceScale(float InAttenuationDistanceScale, float InTime);
-
-	// Sets the parent's attenuation scale
-	void SetParentAttenuationDistanceScale(float InAttenuationDistanceScale);
-
-	// Returns the current attenuation scale
-	float GetAttenuationDistanceScale() const;
-
-	// Updates any dynamic sound class properties
-	void UpdateSoundClassProperties(float DeltaTime);
-
 	FSoundClassProperties();
+};
+PRAGMA_ENABLE_DEPRECATION_WARNINGS
 
-private:
+/** Class for sound class properties which are intended to be dynamically changing without a sound mix. */
+struct FSoundClassDynamicProperties
+{
 	FDynamicParameter AttenuationScaleParam;
-	float ParentAttenuationScale;	
+
+	FSoundClassDynamicProperties()
+		: AttenuationScaleParam(1.0f)
+	{
+	}
 };
-PRAGMA_ENABLE_DEPRECATION_WARNINGS
 
 /**
  * Structure containing information on a SoundMix to activate passively.
diff --git a/Engine/Source/Runtime/Engine/Private/ActiveSound.cpp b/Engine/Source/Runtime/Engine/Private/ActiveSound.cpp
--- a/Engine/Source/Runtime/Engine/Private/ActiveSound.cpp
+++ b/Engine/Source/Runtime/Engine/Private/ActiveSound.cpp
@@ -1652,6 +1652,14 @@ void FActiveSound::UpdateAttenuation(float DeltaTime, FSoundParseParameters& Par
 	// The old audio engine will scale this together when the wave instance is queried for GetActualVolume.
 	if (Settings->bAttenuate)
 	{
+		// Apply the sound-class-based distance scale
+		if (ParseParams.SoundClass)
+		{
+			FSoundClassDynamicProperties* DynamicSoundClassProperties = AudioDevice->GetSoundClassDynamicProperties(ParseParams.SoundClass);
+
+			FocusDataToApply.DistanceScale *= FMath::Max(DynamicSoundClassProperties->AttenuationScaleParam.GetValue(), 0.0f);
+		}
+
 		if (Settings->AttenuationShape == EAttenuationShape::Sphere)
 		{
 			// Update attenuation data in-case it hasn't been updated
diff --git a/Engine/Source/Runtime/Engine/Private/AudioDevice.cpp b/Engine/Source/Runtime/Engine/Private/AudioDevice.cpp
--- a/Engine/Source/Runtime/Engine/Private/AudioDevice.cpp
+++ b/Engine/Source/Runtime/Engine/Private/AudioDevice.cpp
@@ -47,6 +47,7 @@
 #include "Editor/EditorEngine.h"
 #endif // WITH_EDITOR
 
+PRAGMA_DISABLE_OPTIMIZATION
 
 static int32 AudioChannelCountCVar = 0;
 FAutoConsoleVariableRef CVarSetAudioChannelCount(
@@ -2080,7 +2081,13 @@ void FAudioDevice::InitSoundClasses()
 	for (TObjectIterator<USoundClass> It; It; ++It)
 	{
 		USoundClass* SoundClass = *It;
-		FSoundClassProperties& Properties = SoundClasses.Add(SoundClass, SoundClass->Properties);
+		SoundClasses.Add(SoundClass, SoundClass->Properties);
+
+		// Set the dynamic properties
+		FSoundClassDynamicProperties DynamicProperty;
+		DynamicProperty.AttenuationScaleParam.Set(SoundClass->Properties.AttenuationDistanceScale, 0.0f);
+
+		DynamicSoundClassProperties.Add(SoundClass, DynamicProperty);
 	}
 
 	// Propagate the properties down the hierarchy
@@ -2172,7 +2179,6 @@ void FAudioDevice::RecurseIntoSoundClasses(USoundClass* CurrentClass, FSoundClas
 				Properties->Pitch *= ParentProperties.Pitch;
 				Properties->bIsUISound |= ParentProperties.bIsUISound;
 				Properties->bIsMusic |= ParentProperties.bIsMusic;
-				Properties->SetParentAttenuationDistanceScale(ParentProperties.GetAttenuationDistanceScale());
 
 				// Not all values propagate equally...
 				// VoiceCenterChannelVolume, RadioFilterVolume, RadioFilterVolumeThreshold, bApplyEffects, BleedStereo, bReverb, and bCenterChannelOnly do not propagate (sub-classes can be non-zero even if parent class is zero)
@@ -2226,9 +2232,12 @@ void FAudioDevice::ParseSoundClasses(float InDeltaTime)
 		USoundClass* SoundClass = It.Key();
 		if (SoundClass)
 		{
-			// Update any dynamic sound class properties
-			SoundClass->Properties.UpdateSoundClassProperties(InDeltaTime);
+			if (FSoundClassDynamicProperties* DynamicProperties = DynamicSoundClassProperties.Find(SoundClass))
+			{
+				DynamicProperties->AttenuationScaleParam.Update(InDeltaTime);
+			}
 
+			// Reset the property values
 			It.Value() = SoundClass->Properties;
 			if (SoundClass->ParentClass == NULL)
 			{
@@ -6016,6 +6025,10 @@ void FAudioDevice::RegisterSoundClass(USoundClass* InSoundClass)
 		if (!SoundClasses.Contains(InSoundClass))
 		{
 			SoundClasses.Add(InSoundClass, FSoundClassProperties());
+
+			FSoundClassDynamicProperties NewDynamicProperties;
+			NewDynamicProperties.AttenuationScaleParam.Set(InSoundClass->Properties.AttenuationDistanceScale, 0.0f);
+			DynamicSoundClassProperties.Add(InSoundClass, NewDynamicProperties);
 		}
 	}
 }
@@ -6038,6 +6051,7 @@ void FAudioDevice::UnregisterSoundClass(USoundClass* InSoundClass)
 		}
 
 		SoundClasses.Remove(InSoundClass);
+		DynamicSoundClassProperties.Remove(InSoundClass);
 	}
 }
 
@@ -6053,6 +6067,17 @@ FSoundClassProperties* FAudioDevice::GetSoundClassCurrentProperties(USoundClass*
 	return nullptr;
 }
 
+FSoundClassDynamicProperties* FAudioDevice::GetSoundClassDynamicProperties(USoundClass* InSoundClass)
+{
+	if (InSoundClass)
+	{
+		check(IsInAudioThread());
+
+		FSoundClassDynamicProperties* Properties = DynamicSoundClassProperties.Find(InSoundClass);
+		return Properties;
+	}
+	return nullptr;
+}
 
 #if !UE_BUILD_SHIPPING
 /**
@@ -6453,13 +6478,9 @@ void FAudioDevice::SetSoundClassDistanceScale(USoundClass* InSoundClass, float D
 		return;
 	}
 
-	if (FSoundClassProperties* Properties = SoundClasses.Find(InSoundClass))
+	if (FSoundClassDynamicProperties* DynamicProperties = DynamicSoundClassProperties.Find(InSoundClass))
 	{
-		Properties->SetAttenuationDistanceScale(DistanceScale, TimeSec);
-	}
-	else
-	{
-		UE_LOG(LogAudio, Warning, TEXT("Failed to find sound class %s to set attenuation scale."), *InSoundClass->GetName());
+		DynamicProperties->AttenuationScaleParam.Set(DistanceScale, TimeSec);
 	}
 }
 
@@ -6522,3 +6543,5 @@ int32 FAudioDevice::GetNumPrecacheFrames() const
 	// Otherwise, use the default value or value set in ini file
 	return NumPrecacheFrames;
 }
+
+PRAGMA_ENABLE_OPTIMIZATION
\ No newline at end of file
diff --git a/Engine/Source/Runtime/Engine/Private/SoundClass.cpp b/Engine/Source/Runtime/Engine/Private/SoundClass.cpp
--- a/Engine/Source/Runtime/Engine/Private/SoundClass.cpp
+++ b/Engine/Source/Runtime/Engine/Private/SoundClass.cpp
@@ -11,6 +11,8 @@
 #include "UObject/UObjectHash.h"
 #include "UObject/UObjectIterator.h"
 #include "Sound/SoundMix.h"
+#include "AudioDeviceManager.h"
+#include "AudioDevice.h"
 
 #if WITH_EDITOR
 #include "SoundClassGraph/SoundClassGraph.h"
@@ -38,29 +40,7 @@ FSoundClassProperties::FSoundClassProperties()
 	, OutputTarget(EAudioOutputTarget::Speaker)
 	, LoadingBehavior(ESoundWaveLoadingBehavior::Inherited)
 	, DefaultSubmix(nullptr)
-	, AttenuationScaleParam(AttenuationDistanceScale)
-	, ParentAttenuationScale(1.0f)
-{
-}
-
-void FSoundClassProperties::SetAttenuationDistanceScale(float InAttenuationScale, float InTime)
-{
-	AttenuationScaleParam.Set(AttenuationDistanceScale, InTime);
-}
-
-void FSoundClassProperties::SetParentAttenuationDistanceScale(float InAttenuationDistanceScale)
-{
-	ParentAttenuationScale = InAttenuationDistanceScale;
-}
-
-float FSoundClassProperties::GetAttenuationDistanceScale() const
 {
-	return AttenuationScaleParam.GetValue() * ParentAttenuationScale;
-}
-
-void FSoundClassProperties::UpdateSoundClassProperties(float DeltaTime)
-{
-	AttenuationScaleParam.Update(DeltaTime);
 }
 
 /*-----------------------------------------------------------------------------
@@ -142,6 +122,7 @@ void USoundClass::PostEditChangeProperty(struct FPropertyChangedEvent& PropertyC
 		static const FName NAME_ChildClasses = GET_MEMBER_NAME_CHECKED(USoundClass, ChildClasses);
 		static const FName NAME_ParentClass = GET_MEMBER_NAME_CHECKED(USoundClass, ParentClass);
 		static const FName NAME_Properties = GET_MEMBER_NAME_CHECKED(FSoundClassProperties, LoadingBehavior);
+		static const FName NAME_AttenuationDistanceScale = GET_MEMBER_NAME_CHECKED(FSoundClassProperties, AttenuationDistanceScale);
 
 		if (PropertyChangedEvent.GetPropertyName() == NAME_ChildClasses)
 		{
@@ -220,6 +201,16 @@ void USoundClass::PostEditChangeProperty(struct FPropertyChangedEvent& PropertyC
 				Properties.LoadingBehavior = BackupLoadingBehavior;
 			}
 		}
+		else if (PropertyChangedEvent.GetPropertyName() == NAME_AttenuationDistanceScale)
+		{
+			if (FAudioDeviceManager* AudioDeviceManager = FAudioDeviceManager::Get())
+			{
+				AudioDeviceManager->IterateOverAllDevices([this](Audio::FDeviceId DeviceId, FAudioDevice* InDevice)
+				{
+					InDevice->SetSoundClassDistanceScale(this, Properties.AttenuationDistanceScale, 0.2f);
+				});
+			}
+		}
 	}
 
 	if (PropertyChangedEvent.ChangeType != EPropertyChangeType::Interactive && PropertyChangedEvent.MemberProperty)
diff --git a/Engine/Source/Runtime/Engine/Private/SoundWave.cpp b/Engine/Source/Runtime/Engine/Private/SoundWave.cpp
--- a/Engine/Source/Runtime/Engine/Private/SoundWave.cpp
+++ b/Engine/Source/Runtime/Engine/Private/SoundWave.cpp
@@ -1806,8 +1806,6 @@ void USoundWave::Parse(FAudioDevice* AudioDevice, const UPTRINT NodeWaveInstance
 		WaveInstance->SetVolumeMultiplier(VolumeMultiplier* SoundClassProperties->Volume);
 		WaveInstance->SetPitch(WaveInstance->Pitch * SoundClassProperties->Pitch);
 
-		WaveInstance->AttenuationDistance *= FMath::Max(SoundClassProperties->GetAttenuationDistanceScale(), 0.0f);
-
 		WaveInstance->SoundClassFilterFrequency = SoundClassProperties->LowPassFilterFrequency;
 		WaveInstance->VoiceCenterChannelVolume = SoundClassProperties->VoiceCenterChannelVolume;
 		WaveInstance->RadioFilterVolume = SoundClassProperties->RadioFilterVolume * ParseParams.VolumeMultiplier;
diff --git a/Engine/Source/Runtime/Engine/Public/AudioDevice.h b/Engine/Source/Runtime/Engine/Public/AudioDevice.h
--- a/Engine/Source/Runtime/Engine/Public/AudioDevice.h
+++ b/Engine/Source/Runtime/Engine/Public/AudioDevice.h
@@ -859,6 +859,11 @@ public:
 	*/
 	FSoundClassProperties* GetSoundClassCurrentProperties(USoundClass* InSoundClass);
 
+	/** 
+	* Returns the parameters which are dynamic from the given sound class. 
+	*/
+	FSoundClassDynamicProperties* GetSoundClassDynamicProperties(USoundClass* InSoundClass);
+
 	/**
 	* Checks to see if a coordinate is within a distance of any listener
 	*/
@@ -1801,6 +1806,7 @@ private:
 
 	/** Current properties of all sound classes */
 	TMap<USoundClass*, FSoundClassProperties> SoundClasses;
+	TMap<USoundClass*, FSoundClassDynamicProperties> DynamicSoundClassProperties;
 
 	/** The Base SoundMix that's currently active */
 	USoundMix* BaseSoundMix;
