Removing pragma disable optimization flag

diff --git a/Engine/Source/Runtime/Engine/Private/AudioDevice.cpp b/Engine/Source/Runtime/Engine/Private/AudioDevice.cpp
--- a/Engine/Source/Runtime/Engine/Private/AudioDevice.cpp
+++ b/Engine/Source/Runtime/Engine/Private/AudioDevice.cpp
@@ -47,8 +47,6 @@
 #include "Editor/EditorEngine.h"
 #endif // WITH_EDITOR
 
-PRAGMA_DISABLE_OPTIMIZATION
-
 static int32 AudioChannelCountCVar = 0;
 FAutoConsoleVariableRef CVarSetAudioChannelCount(
 	TEXT("au.SetAudioChannelCount"),
@@ -6542,6 +6540,4 @@ int32 FAudioDevice::GetNumPrecacheFrames() const
 	}
 	// Otherwise, use the default value or value set in ini file
 	return NumPrecacheFrames;
-}
-
-PRAGMA_ENABLE_OPTIMIZATION
\ No newline at end of file
+}
\ No newline at end of file
