Fix crash in PhysicalAnimation EngineTest

diff --git a/Engine/Source/Runtime/Engine/Private/PhysicsEngine/ImmediatePhysics/ImmediatePhysicsChaos/ImmediatePhysicsActorHandle_Chaos.cpp b/Engine/Source/Runtime/Engine/Private/PhysicsEngine/ImmediatePhysics/ImmediatePhysicsChaos/ImmediatePhysicsActorHandle_Chaos.cpp
--- a/Engine/Source/Runtime/Engine/Private/PhysicsEngine/ImmediatePhysics/ImmediatePhysicsChaos/ImmediatePhysicsActorHandle_Chaos.cpp
+++ b/Engine/Source/Runtime/Engine/Private/PhysicsEngine/ImmediatePhysics/ImmediatePhysicsChaos/ImmediatePhysicsActorHandle_Chaos.cpp
@@ -190,8 +190,9 @@ namespace ImmediatePhysics_Chaos
 	{
 		using namespace Chaos;
 
-		if (BodyInstance == nullptr)
+		if ((BodyInstance == nullptr) || (BodyInstance->BodySetup == nullptr))
 		{
+			// @todo(ccaulfield): fix this path
 			return CreateDefaultGeometry(Scale, OutMass, OutInertia, OutCoMTransform, OutGeom, OutShapes);
 		}
 
