Changed the FParallelClothTask asset stat to show the skel mesh name instead of the skel mesh component name.

#rnx

diff --git a/Engine/Source/Runtime/Engine/Private/SkeletalMeshComponentPhysics.cpp b/Engine/Source/Runtime/Engine/Private/SkeletalMeshComponentPhysics.cpp
--- a/Engine/Source/Runtime/Engine/Private/SkeletalMeshComponentPhysics.cpp
+++ b/Engine/Source/Runtime/Engine/Private/SkeletalMeshComponentPhysics.cpp
@@ -3100,7 +3100,7 @@ public:
 
 	void DoTask(ENamedThreads::Type CurrentThread, const FGraphEventRef& MyCompletionGraphEvent)
 	{
-		FScopeCycleCounterUObject ContextScope(&SkeletalMeshComponent);
+		FScopeCycleCounterUObject ContextScope(SkeletalMeshComponent.SkeletalMesh);
 		SCOPE_CYCLE_COUNTER(STAT_ClothTotalTime);
 
 		if(SkeletalMeshComponent.ClothingSimulation)
