Removing "sound" has a hidden category for source buses

diff --git a/Engine/Source/Runtime/Engine/Classes/Sound/SoundSourceBus.h b/Engine/Source/Runtime/Engine/Classes/Sound/SoundSourceBus.h
--- a/Engine/Source/Runtime/Engine/Classes/Sound/SoundSourceBus.h
+++ b/Engine/Source/Runtime/Engine/Classes/Sound/SoundSourceBus.h
@@ -22,7 +22,7 @@ enum class ESourceBusChannels : uint8
 };
 
 // A source bus is a type of USoundBase and can be "played" like any sound.
-UCLASS(hidecategories= (Compression, SoundWave, Streaming, Subtitles, Sound, Analysis, Format, Loading, Info, ImportSettings), ClassGroup = Sound, meta = (BlueprintSpawnableComponent))
+UCLASS(hidecategories= (Compression, SoundWave, Streaming, Subtitles, Analysis, Format, Loading, Info, ImportSettings), ClassGroup = Sound, meta = (BlueprintSpawnableComponent))
 class ENGINE_API USoundSourceBus : public USoundWave
 {
 	GENERATED_UCLASS_BODY()
