Make sure MI cached data is up to date when recompiling due to changes to parent material, fix crash when textures are swapped in editor in certain cases

diff --git a/Engine/Source/Runtime/Engine/Private/Materials/MaterialShared.cpp b/Engine/Source/Runtime/Engine/Private/Materials/MaterialShared.cpp
--- a/Engine/Source/Runtime/Engine/Private/Materials/MaterialShared.cpp
+++ b/Engine/Source/Runtime/Engine/Private/Materials/MaterialShared.cpp
@@ -3151,6 +3151,7 @@ FMaterialUpdateContext::~FMaterialUpdateContext()
 			}
 		}
 
+		MI->UpdateCachedLayerParameters();
 		MI->RecacheUniformExpressions(true);
 		MI->InitStaticPermutation();//bHasStaticPermutation can change.
 		if (MI->bHasStaticPermutationResource)
