Fix compile error for non-editor

#ue4

diff --git a/Engine/Source/Runtime/Engine/Private/Materials/MaterialShared.cpp b/Engine/Source/Runtime/Engine/Private/Materials/MaterialShared.cpp
--- a/Engine/Source/Runtime/Engine/Private/Materials/MaterialShared.cpp
+++ b/Engine/Source/Runtime/Engine/Private/Materials/MaterialShared.cpp
@@ -3151,7 +3151,9 @@ FMaterialUpdateContext::~FMaterialUpdateContext()
 			}
 		}
 
+#if WITH_EDITOR
 		MI->UpdateCachedLayerParameters();
+#endif
 		MI->RecacheUniformExpressions(true);
 		MI->InitStaticPermutation();//bHasStaticPermutation can change.
 		if (MI->bHasStaticPermutationResource)
