Remove redundant call to destructor in ~FRayTracingMeshResourceCollector() which has caused memory corruption

diff --git a/Engine/Source/Runtime/Engine/Public/SceneManagement.h b/Engine/Source/Runtime/Engine/Public/SceneManagement.h
--- a/Engine/Source/Runtime/Engine/Public/SceneManagement.h
+++ b/Engine/Source/Runtime/Engine/Public/SceneManagement.h
@@ -2313,11 +2313,6 @@ public:
 		DynamicVertexBuffer = InDynamicVertexBuffer;
 		DynamicReadBuffer = InDynamicReadBuffer;
 	}
-
-	~FRayTracingMeshResourceCollector()
-	{
-		FMeshElementCollector::~FMeshElementCollector();
-	}
 };
 
 struct FRayTracingDynamicGeometryUpdateParams
