UE4.25 - Compile fix
#rnx

diff --git a/Engine/Source/Developer/Windows/ShaderFormatD3D/Private/D3DShaderCompiler.cpp b/Engine/Source/Developer/Windows/ShaderFormatD3D/Private/D3DShaderCompiler.cpp
--- a/Engine/Source/Developer/Windows/ShaderFormatD3D/Private/D3DShaderCompiler.cpp
+++ b/Engine/Source/Developer/Windows/ShaderFormatD3D/Private/D3DShaderCompiler.cpp
@@ -500,7 +500,7 @@ static FString DxcBlobEncodingToFString(TRefCountPtr<IDxcBlobEncoding> DxcBlob)
 
 
 static HRESULT DXCCompileWrapper(
-	TRefCountPtr<IDxcCompiler> Compiler,
+	TRefCountPtr<IDxcCompiler>& Compiler,
 	TRefCountPtr<IDxcBlobEncoding>& TextBlob,
 	LPCWSTR EntryPoint,
 	LPCWSTR TargetProfile,
diff --git a/Engine/Source/Runtime/Engine/Engine.Build.cs b/Engine/Source/Runtime/Engine/Engine.Build.cs
--- a/Engine/Source/Runtime/Engine/Engine.Build.cs
+++ b/Engine/Source/Runtime/Engine/Engine.Build.cs
@@ -37,6 +37,7 @@ public class Engine : ModuleRules
 				"AutomationWorker",
 				"MovieSceneCapture",
 				"DesktopPlatform",
+				"ShaderCompilerCommon"
 			}
 		);
 
diff --git a/Engine/Source/Runtime/Engine/Private/ShaderCompiler/ShaderCompiler.cpp b/Engine/Source/Runtime/Engine/Private/ShaderCompiler/ShaderCompiler.cpp
--- a/Engine/Source/Runtime/Engine/Private/ShaderCompiler/ShaderCompiler.cpp
+++ b/Engine/Source/Runtime/Engine/Private/ShaderCompiler/ShaderCompiler.cpp
@@ -48,6 +48,7 @@
 #include "ShaderParameterMetadata.h"
 #include "ProfilingDebugging/DiagnosticTable.h"
 #include "ProfilingDebugging/LoadTimeTracker.h"
+#include "ShaderCompilerCommon.h"
 
 #define LOCTEXT_NAMESPACE "ShaderCompiler"
 
@@ -343,7 +344,7 @@ namespace ShaderCompilerCookStats
 #endif
 
 // Make functions so the crash reporter can disambiguate the actual error because of the different callstacks
-namespace SCWErrorCode
+namespace ESCWErrorCode
 {
 	void HandleGeneralCrash(const TCHAR* ExceptionInfo, const TCHAR* Callstack)
 	{
@@ -741,7 +742,7 @@ static void HandleWorkerCrash(const TArray<TSharedRef<FShaderCommonCompileJob, E
 	switch (ErrorCode)
 	{
 	default:
-	case SCWErrorCode::GeneralCrash:
+	case ESCWErrorCode::GeneralCrash:
 	{
 		if (GDumpSCWJobInfoOnCrash != 0 || GIsBuildMachine)
 		{
@@ -780,37 +781,37 @@ static void HandleWorkerCrash(const TArray<TSharedRef<FShaderCommonCompileJob, E
 				}
 			}
 		}
-		SCWErrorCode::HandleGeneralCrash(ExceptionInfo.GetData(), Callstack.GetData());
+		ESCWErrorCode::HandleGeneralCrash(ExceptionInfo.GetData(), Callstack.GetData());
 	}
 	break;
-	case SCWErrorCode::BadShaderFormatVersion:
-		SCWErrorCode::HandleBadShaderFormatVersion(ExceptionInfo.GetData());
+	case ESCWErrorCode::BadShaderFormatVersion:
+		ESCWErrorCode::HandleBadShaderFormatVersion(ExceptionInfo.GetData());
 		break;
-	case SCWErrorCode::BadInputVersion:
-		SCWErrorCode::HandleBadInputVersion(ExceptionInfo.GetData());
+	case ESCWErrorCode::BadInputVersion:
+		ESCWErrorCode::HandleBadInputVersion(ExceptionInfo.GetData());
 		break;
-	case SCWErrorCode::BadSingleJobHeader:
-		SCWErrorCode::HandleBadSingleJobHeader(ExceptionInfo.GetData());
+	case ESCWErrorCode::BadSingleJobHeader:
+		ESCWErrorCode::HandleBadSingleJobHeader(ExceptionInfo.GetData());
 		break;
-	case SCWErrorCode::BadPipelineJobHeader:
-		SCWErrorCode::HandleBadPipelineJobHeader(ExceptionInfo.GetData());
+	case ESCWErrorCode::BadPipelineJobHeader:
+		ESCWErrorCode::HandleBadPipelineJobHeader(ExceptionInfo.GetData());
 		break;
-	case SCWErrorCode::CantDeleteInputFile:
-		SCWErrorCode::HandleCantDeleteInputFile(ExceptionInfo.GetData());
+	case ESCWErrorCode::CantDeleteInputFile:
+		ESCWErrorCode::HandleCantDeleteInputFile(ExceptionInfo.GetData());
 		break;
-	case SCWErrorCode::CantSaveOutputFile:
-		SCWErrorCode::HandleCantSaveOutputFile(ExceptionInfo.GetData());
+	case ESCWErrorCode::CantSaveOutputFile:
+		ESCWErrorCode::HandleCantSaveOutputFile(ExceptionInfo.GetData());
 		break;
-	case SCWErrorCode::NoTargetShaderFormatsFound:
-		SCWErrorCode::HandleNoTargetShaderFormatsFound(ExceptionInfo.GetData());
+	case ESCWErrorCode::NoTargetShaderFormatsFound:
+		ESCWErrorCode::HandleNoTargetShaderFormatsFound(ExceptionInfo.GetData());
 		break;
-	case SCWErrorCode::CantCompileForSpecificFormat:
-		SCWErrorCode::HandleCantCompileForSpecificFormat(ExceptionInfo.GetData());
+	case ESCWErrorCode::CantCompileForSpecificFormat:
+		ESCWErrorCode::HandleCantCompileForSpecificFormat(ExceptionInfo.GetData());
 		break;
-	case SCWErrorCode::CrashInsidePlatformCompiler:
-		SCWErrorCode::HandleCrashInsidePlatformCompiler(ExceptionInfo.GetData());
+	case ESCWErrorCode::CrashInsidePlatformCompiler:
+		ESCWErrorCode::HandleCrashInsidePlatformCompiler(ExceptionInfo.GetData());
 		break;
-	case SCWErrorCode::Success:
+	case ESCWErrorCode::Success:
 		// Can't get here...
 		break;
 	}
@@ -822,7 +823,7 @@ void FShaderCompileUtilities::DoReadTaskResults(const TArray<TSharedRef<FShaderC
 {
 	if (OutputFile.TotalSize() == 0)
 	{
-		SCWErrorCode::HandleOutputFileEmpty(*OutputFile.GetArchiveName());
+		ESCWErrorCode::HandleOutputFileEmpty(*OutputFile.GetArchiveName());
 	}
 
 	int32 OutputVersion = ShaderCompileWorkerOutputVersion;
@@ -840,7 +841,7 @@ void FShaderCompileUtilities::DoReadTaskResults(const TArray<TSharedRef<FShaderC
 	// Check for corrupted output file
 	if (FileSize > OutputFile.TotalSize())
 	{
-		SCWErrorCode::HandleOutputFileCorrupted(*OutputFile.GetArchiveName(), FileSize, OutputFile.TotalSize());
+		ESCWErrorCode::HandleOutputFileCorrupted(*OutputFile.GetArchiveName(), FileSize, OutputFile.TotalSize());
 	}
 
 	int32 ErrorCode = 0;
@@ -856,7 +857,7 @@ void FShaderCompileUtilities::DoReadTaskResults(const TArray<TSharedRef<FShaderC
 	OutputFile << ExceptionInfoLength;
 
 	// Worker crashed
-	if (ErrorCode != SCWErrorCode::Success)
+	if (ErrorCode != ESCWErrorCode::Success)
 	{
 		HandleWorkerCrash(QueuedJobs, OutputFile, OutputVersion, FileSize, ErrorCode, NumProcessedJobs, CallstackLength, ExceptionInfoLength);
 	}
