UE4.25 - Compile fix
#rnx

diff --git a/Engine/Source/Runtime/Engine/Private/ShaderCompiler/ShaderCompiler.cpp b/Engine/Source/Runtime/Engine/Private/ShaderCompiler/ShaderCompiler.cpp
--- a/Engine/Source/Runtime/Engine/Private/ShaderCompiler/ShaderCompiler.cpp
+++ b/Engine/Source/Runtime/Engine/Private/ShaderCompiler/ShaderCompiler.cpp
@@ -344,7 +344,7 @@ namespace ShaderCompilerCookStats
 #endif
 
 // Make functions so the crash reporter can disambiguate the actual error because of the different callstacks
-namespace ESCWErrorCode
+namespace SCWErrorCode
 {
 	void HandleGeneralCrash(const TCHAR* ExceptionInfo, const TCHAR* Callstack)
 	{
@@ -402,9 +402,9 @@ namespace ESCWErrorCode
 		ModalErrorOrLog(FString::Printf(TEXT("Output file corrupted (expected %I64d bytes, but only got %I64d): %s"), ExpectedSize, ActualSize, Filename));
 	}
 
-	void HandleCrashInsidePlatformCompiler(const TCHAR* Filename, int64 ExpectedSize, int64 ActualSize)
+	void HandleCrashInsidePlatformCompiler(const TCHAR* Data)
 	{
-		ModalErrorOrLog(FString::Printf(TEXT("Crash inside the platform compiler!")));
+		ModalErrorOrLog(FString::Printf(TEXT("Crash inside the platform compiler!\n%s"), Data));
 	}
 }
 
@@ -781,35 +781,35 @@ static void HandleWorkerCrash(const TArray<TSharedRef<FShaderCommonCompileJob, E
 				}
 			}
 		}
-		ESCWErrorCode::HandleGeneralCrash(ExceptionInfo.GetData(), Callstack.GetData());
+		SCWErrorCode::HandleGeneralCrash(ExceptionInfo.GetData(), Callstack.GetData());
 	}
 	break;
 	case ESCWErrorCode::BadShaderFormatVersion:
-		ESCWErrorCode::HandleBadShaderFormatVersion(ExceptionInfo.GetData());
+		SCWErrorCode::HandleBadShaderFormatVersion(ExceptionInfo.GetData());
 		break;
 	case ESCWErrorCode::BadInputVersion:
-		ESCWErrorCode::HandleBadInputVersion(ExceptionInfo.GetData());
+		SCWErrorCode::HandleBadInputVersion(ExceptionInfo.GetData());
 		break;
 	case ESCWErrorCode::BadSingleJobHeader:
-		ESCWErrorCode::HandleBadSingleJobHeader(ExceptionInfo.GetData());
+		SCWErrorCode::HandleBadSingleJobHeader(ExceptionInfo.GetData());
 		break;
 	case ESCWErrorCode::BadPipelineJobHeader:
-		ESCWErrorCode::HandleBadPipelineJobHeader(ExceptionInfo.GetData());
+		SCWErrorCode::HandleBadPipelineJobHeader(ExceptionInfo.GetData());
 		break;
 	case ESCWErrorCode::CantDeleteInputFile:
-		ESCWErrorCode::HandleCantDeleteInputFile(ExceptionInfo.GetData());
+		SCWErrorCode::HandleCantDeleteInputFile(ExceptionInfo.GetData());
 		break;
 	case ESCWErrorCode::CantSaveOutputFile:
-		ESCWErrorCode::HandleCantSaveOutputFile(ExceptionInfo.GetData());
+		SCWErrorCode::HandleCantSaveOutputFile(ExceptionInfo.GetData());
 		break;
 	case ESCWErrorCode::NoTargetShaderFormatsFound:
-		ESCWErrorCode::HandleNoTargetShaderFormatsFound(ExceptionInfo.GetData());
+		SCWErrorCode::HandleNoTargetShaderFormatsFound(ExceptionInfo.GetData());
 		break;
 	case ESCWErrorCode::CantCompileForSpecificFormat:
-		ESCWErrorCode::HandleCantCompileForSpecificFormat(ExceptionInfo.GetData());
+		SCWErrorCode::HandleCantCompileForSpecificFormat(ExceptionInfo.GetData());
 		break;
 	case ESCWErrorCode::CrashInsidePlatformCompiler:
-		ESCWErrorCode::HandleCrashInsidePlatformCompiler(ExceptionInfo.GetData());
+		SCWErrorCode::HandleCrashInsidePlatformCompiler(ExceptionInfo.GetData());
 		break;
 	case ESCWErrorCode::Success:
 		// Can't get here...
@@ -823,7 +823,7 @@ void FShaderCompileUtilities::DoReadTaskResults(const TArray<TSharedRef<FShaderC
 {
 	if (OutputFile.TotalSize() == 0)
 	{
-		ESCWErrorCode::HandleOutputFileEmpty(*OutputFile.GetArchiveName());
+		SCWErrorCode::HandleOutputFileEmpty(*OutputFile.GetArchiveName());
 	}
 
 	int32 OutputVersion = ShaderCompileWorkerOutputVersion;
@@ -841,7 +841,7 @@ void FShaderCompileUtilities::DoReadTaskResults(const TArray<TSharedRef<FShaderC
 	// Check for corrupted output file
 	if (FileSize > OutputFile.TotalSize())
 	{
-		ESCWErrorCode::HandleOutputFileCorrupted(*OutputFile.GetArchiveName(), FileSize, OutputFile.TotalSize());
+		SCWErrorCode::HandleOutputFileCorrupted(*OutputFile.GetArchiveName(), FileSize, OutputFile.TotalSize());
 	}
 
 	int32 ErrorCode = 0;
@@ -857,7 +857,7 @@ void FShaderCompileUtilities::DoReadTaskResults(const TArray<TSharedRef<FShaderC
 	OutputFile << ExceptionInfoLength;
 
 	// Worker crashed
-	if (ErrorCode != ESCWErrorCode::Success)
+	if (ErrorCode != (int32)ESCWErrorCode::Success)
 	{
 		HandleWorkerCrash(QueuedJobs, OutputFile, OutputVersion, FileSize, ErrorCode, NumProcessedJobs, CallstackLength, ExceptionInfoLength);
 	}
