UE4.25 - Fix Mac clang
#rnx

diff --git a/Engine/Source/Runtime/Engine/Private/ShaderCompiler/ShaderCompiler.cpp b/Engine/Source/Runtime/Engine/Private/ShaderCompiler/ShaderCompiler.cpp
--- a/Engine/Source/Runtime/Engine/Private/ShaderCompiler/ShaderCompiler.cpp
+++ b/Engine/Source/Runtime/Engine/Private/ShaderCompiler/ShaderCompiler.cpp
@@ -742,7 +742,7 @@ static void HandleWorkerCrash(const TArray<TSharedRef<FShaderCommonCompileJob, E
 	switch (ErrorCode)
 	{
 	default:
-	case ESCWErrorCode::GeneralCrash:
+	case (int32)ESCWErrorCode::GeneralCrash:
 	{
 		if (GDumpSCWJobInfoOnCrash != 0 || GIsBuildMachine)
 		{
@@ -784,34 +784,34 @@ static void HandleWorkerCrash(const TArray<TSharedRef<FShaderCommonCompileJob, E
 		SCWErrorCode::HandleGeneralCrash(ExceptionInfo.GetData(), Callstack.GetData());
 	}
 	break;
-	case ESCWErrorCode::BadShaderFormatVersion:
+	case (int32)ESCWErrorCode::BadShaderFormatVersion:
 		SCWErrorCode::HandleBadShaderFormatVersion(ExceptionInfo.GetData());
 		break;
-	case ESCWErrorCode::BadInputVersion:
+	case (int32)ESCWErrorCode::BadInputVersion:
 		SCWErrorCode::HandleBadInputVersion(ExceptionInfo.GetData());
 		break;
-	case ESCWErrorCode::BadSingleJobHeader:
+	case (int32)ESCWErrorCode::BadSingleJobHeader:
 		SCWErrorCode::HandleBadSingleJobHeader(ExceptionInfo.GetData());
 		break;
-	case ESCWErrorCode::BadPipelineJobHeader:
+	case (int32)ESCWErrorCode::BadPipelineJobHeader:
 		SCWErrorCode::HandleBadPipelineJobHeader(ExceptionInfo.GetData());
 		break;
-	case ESCWErrorCode::CantDeleteInputFile:
+	case (int32)ESCWErrorCode::CantDeleteInputFile:
 		SCWErrorCode::HandleCantDeleteInputFile(ExceptionInfo.GetData());
 		break;
-	case ESCWErrorCode::CantSaveOutputFile:
+	case (int32)ESCWErrorCode::CantSaveOutputFile:
 		SCWErrorCode::HandleCantSaveOutputFile(ExceptionInfo.GetData());
 		break;
-	case ESCWErrorCode::NoTargetShaderFormatsFound:
+	case (int32)ESCWErrorCode::NoTargetShaderFormatsFound:
 		SCWErrorCode::HandleNoTargetShaderFormatsFound(ExceptionInfo.GetData());
 		break;
-	case ESCWErrorCode::CantCompileForSpecificFormat:
+	case (int32)ESCWErrorCode::CantCompileForSpecificFormat:
 		SCWErrorCode::HandleCantCompileForSpecificFormat(ExceptionInfo.GetData());
 		break;
-	case ESCWErrorCode::CrashInsidePlatformCompiler:
+	case (int32)ESCWErrorCode::CrashInsidePlatformCompiler:
 		SCWErrorCode::HandleCrashInsidePlatformCompiler(ExceptionInfo.GetData());
 		break;
-	case ESCWErrorCode::Success:
+	case (int32)ESCWErrorCode::Success:
 		// Can't get here...
 		break;
 	}
