Fix issue where we would skip one chunk of ADPCM if a chunk failed to load in time. Also set the movie audio track to use the game clock for resync.

diff --git a/Engine/Source/Runtime/Engine/Private/ADPCMAudioInfo.cpp b/Engine/Source/Runtime/Engine/Private/ADPCMAudioInfo.cpp
--- a/Engine/Source/Runtime/Engine/Private/ADPCMAudioInfo.cpp
+++ b/Engine/Source/Runtime/Engine/Private/ADPCMAudioInfo.cpp
@@ -37,6 +37,13 @@ FAutoConsoleVariableRef CVarADPCMDisableSeekForwardOnChunkMisses(
 	TEXT("When there is a seek pending and this CVar is set to 0, we will scan forward in the file.\n"),
 	ECVF_Default);
 
+static float ChanceForIntentionalChunkMissCVar = 0.0f;
+FAutoConsoleVariableRef CVarChanceForIntentionalChunkMiss(
+	TEXT("au.adpcm.ChanceForIntentionalChunkMiss"),
+	ChanceForIntentionalChunkMissCVar,
+	TEXT("If this is set > 0 we will intentionally drop chunks. Used for debugging..\n"),
+	ECVF_Default);
+
 #define WAVE_FORMAT_LPCM  1
 #define WAVE_FORMAT_ADPCM 2
 
@@ -55,12 +62,14 @@ FADPCMAudioInfo::FADPCMAudioInfo(void)
 	, TotalDecodedSize(0)
 	, NumChannels(0)
 	, Format(0)
+	, PreviouslyRequestedChunkIndex(0)
 	, UncompressedBlockData(nullptr)
 	, SamplesPerBlock(0)
 	, FirstChunkSampleDataOffset(0)
 	, FirstChunkSampleDataIndex(0)
 	, bDecompressorReleased(false)
 	, bSeekPending(false)
+	, bSeekedFowardToNextChunk(false)
 	, TargetSeekTime(0.0f)
 	, LastSeekTime(0.0f)
 {
@@ -501,6 +510,8 @@ bool FADPCMAudioInfo::StreamCompressedInfoInternal(USoundWave* Wave, struct FSou
 		FirstChunkSampleDataIndex = 0;
 	}
 
+	PreviouslyRequestedChunkIndex = CurrentChunkIndex;
+
 	SrcBufferData = ChunkData;
 	CurrentChunkBufferOffset = 0;
 	CurCompressedChunkData = nullptr;
@@ -619,7 +630,25 @@ bool FADPCMAudioInfo::StreamCompressedData(uint8* Destination, bool bLooping, ui
 					// and invalidates CurCompressedChunkData.
 					if(CurCompressedChunkData != nullptr)
 					{
-						++CurrentChunkIndex;
+						// If we have already seeked forward to the next chunk because the chunk didn't load in time for the render callback, we don't increment here.
+						// Otherwise, we increment to the next chunk index and reset our chunk offset to the beginning of the new chunk.
+						if (!bSeekedFowardToNextChunk)
+						{
+							// Ensure that we're either seeking or moving sequentially through the file.
+							ensureAlways(bSeekPending || CurrentChunkIndex == PreviouslyRequestedChunkIndex);
+							++CurrentChunkIndex;
+
+							// Set the current buffer offset accounting for the header in the first chunk
+							if (!bSeekPending)
+							{
+								// Ensure that, if we hit this code, we are not at the beginning of the file.
+								ensureAlways(CurrentChunkIndex != FirstChunkSampleDataIndex);
+								CurrentChunkBufferOffset = 0;
+							}
+						}
+
+						// If this flag was raised, clear it for the next time we move on to another chunk.
+						bSeekedFowardToNextChunk = false;
 					}
 
 					if (CurrentChunkIndex >= StreamingSoundWave->GetNumChunks())
@@ -680,14 +709,27 @@ bool FADPCMAudioInfo::StreamCompressedData(uint8* Destination, bool bLooping, ui
 								}
 							}
 
+							// Seek forward in the file the amount we would be rendering otherwise.
 							CurrentChunkBufferOffset += NumCompressedBytesForCallback;
 							TotalSamplesStreamed += NumDecompressedSamplesForCallback;
 
 							const uint32 SizeOfCurrentChunk = StreamingSoundWave->GetSizeOfChunk(CurrentChunkIndex);
 							if (CurrentChunkBufferOffset >= SizeOfCurrentChunk)
 							{
+								// If this is the case, we've seeked forward past the end of the chunk.
+
+								// Ensure we never hit this section of code twice, resulting in a double increment of the chunk index.
+								ensureAlways(!bSeekedFowardToNextChunk);
+
+								// Ensure that we are either seeking of moving sequentially through the file.
+								ensureAlwaysMsgf(bSeekPending || CurrentChunkIndex == PreviouslyRequestedChunkIndex, TEXT("Failed to load a chunk of ADPCM audio for the entire duration that chunk of audio was supposed to be played."));
 								++CurrentChunkIndex;
-								CurrentChunkBufferOffset %= SizeOfCurrentChunk;
+								CurrentChunkBufferOffset -= SizeOfCurrentChunk;
+
+								// Ensure that we have not moved past the size of our new chunk with our current callback.
+								ensureAlways(CurrentChunkBufferOffset < StreamingSoundWave->GetSizeOfChunk(CurrentChunkIndex));
+
+								bSeekedFowardToNextChunk = true;
 
 								if (!ensureMsgf(CurrentChunkBufferOffset % CompressedBlockSize == 0, TEXT("Error: seeked partway into an ADPCM block. Please check the above code, as well as FAudioFormatADPCM::SplitDataForStreaming for accuracy.")))
 								{
@@ -718,10 +760,11 @@ bool FADPCMAudioInfo::StreamCompressedData(uint8* Destination, bool bLooping, ui
 						return false;
 					}
 
-					// Set the current buffer offset accounting for the header in the first chunk
-					if (!bSeekPending)
+					
+					if (CurrentChunkIndex == FirstChunkSampleDataIndex && !bSeekPending)
 					{
-						CurrentChunkBufferOffset = CurrentChunkIndex == FirstChunkSampleDataIndex ? FirstChunkSampleDataOffset : 0;
+						// If we're in the first chunk, set the current buffer offset accounting for the header.
+						CurrentChunkBufferOffset = FirstChunkSampleDataOffset;
 					}
 
 					bSeekPending = false;
@@ -916,6 +959,13 @@ bool FADPCMAudioInfo::ReleaseStreamChunk(bool bBlockUntilReleased)
 
 const uint8* FADPCMAudioInfo::GetLoadedChunk(USoundWave* InSoundWave, uint32 ChunkIndex, uint32& OutChunkSize)
 {
+	if (ChunkIndex != 0 && FMath::RandRange(0.0f, 1.0f) < ChanceForIntentionalChunkMissCVar)
+	{
+		UE_LOG(LogAudio, Display, TEXT("Intentionally dropping a chunk here."));
+		OutChunkSize = 0;
+		return nullptr;
+	}
+
 	if (!InSoundWave || ChunkIndex >= InSoundWave->GetNumChunks())
 	{
 		if(InSoundWave)
@@ -930,12 +980,19 @@ const uint8* FADPCMAudioInfo::GetLoadedChunk(USoundWave* InSoundWave, uint32 Chu
 	{
 		TArrayView<const uint8> ZerothChunk = InSoundWave->GetZerothChunk(true);
 		OutChunkSize = ZerothChunk.Num();
+		PreviouslyRequestedChunkIndex = ChunkIndex;
 		return ZerothChunk.GetData();
 	}
 	else
 	{
+		const bool bIsLastChunk = ChunkIndex == InSoundWave->GetNumChunks() - 1;
+		ensureAlwaysMsgf(bSeekPending || ChunkIndex == PreviouslyRequestedChunkIndex || ChunkIndex == PreviouslyRequestedChunkIndex + 1, TEXT("ADPCM playback error! We skipped from the end of chunk %d to the beginning of chunk %d."), PreviouslyRequestedChunkIndex, ChunkIndex);
+
 		CurCompressedChunkHandle = IStreamingManager::Get().GetAudioStreamingManager().GetLoadedChunk(InSoundWave, ChunkIndex, false, true);
 		OutChunkSize = CurCompressedChunkHandle.Num();
+
+		PreviouslyRequestedChunkIndex = ChunkIndex;
+		
 		return CurCompressedChunkHandle.GetData();
 	}
 }
diff --git a/Engine/Source/Runtime/Engine/Public/ADPCMAudioInfo.h b/Engine/Source/Runtime/Engine/Public/ADPCMAudioInfo.h
--- a/Engine/Source/Runtime/Engine/Public/ADPCMAudioInfo.h
+++ b/Engine/Source/Runtime/Engine/Public/ADPCMAudioInfo.h
@@ -154,6 +154,8 @@ private:
 	int32	NumChannels;
 	int32	Format;
 
+	uint32 PreviouslyRequestedChunkIndex;
+
 	uint8*			UncompressedBlockData;			// This holds the current block of compressed data for all channels
 	uint32			CurrentUncompressedBlockSampleIndex;	// This is the sample index within the current uncompressed block data
 	uint32			CurrentChunkIndex;				// This is the index that is currently being used, needed by streaming engine to make sure it stays loaded and the next chunk gets preloaded
@@ -172,6 +174,7 @@ private:
 	uint32			CurrentCompressedBlockIndex;		// For non disk streaming - the current compressed block in the compressed source data
 	uint32			TotalCompressedBlocksPerChannel;	// For non disk streaming - the total number of compressed blocks per channel
 	uint8			bSeekPending : 1;					// Whether or not seek has been requested and pending a read
+	uint8           bSeekedFowardToNextChunk : 1;       // If this is true, we have already seeked forward to the next chunk while waiting for the current chunk of audio to load.
 
 private:
 	void ProcessSeekRequest();
diff --git a/Engine/Source/Runtime/MovieSceneTracks/Private/Evaluation/MovieSceneAudioTemplate.cpp b/Engine/Source/Runtime/MovieSceneTracks/Private/Evaluation/MovieSceneAudioTemplate.cpp
--- a/Engine/Source/Runtime/MovieSceneTracks/Private/Evaluation/MovieSceneAudioTemplate.cpp
+++ b/Engine/Source/Runtime/MovieSceneTracks/Private/Evaluation/MovieSceneAudioTemplate.cpp
@@ -29,7 +29,7 @@ FAutoConsoleVariableRef CVarMaxSequenceAudioDesyncTolerance(
 	TEXT("Controls how many seconds an audio track can be out of sync in a Sequence before we attempt a time correction.\n"),
 	ECVF_Default);
 
-static int32 UseAudioClockForSequencerDesyncCVar = 1;
+static int32 UseAudioClockForSequencerDesyncCVar = 0;
 FAutoConsoleVariableRef CVaUseAudioClockForSequencerDesync(
 	TEXT("Sequencer.Audio.UseAudioClockForAudioDesync"),
 	UseAudioClockForSequencerDesyncCVar,
