Fixed potential crash gathering ISM instances

diff --git a/Engine/Source/Runtime/Engine/Private/InstancedStaticMesh.cpp b/Engine/Source/Runtime/Engine/Private/InstancedStaticMesh.cpp
--- a/Engine/Source/Runtime/Engine/Private/InstancedStaticMesh.cpp
+++ b/Engine/Source/Runtime/Engine/Private/InstancedStaticMesh.cpp
@@ -1155,7 +1155,8 @@ void FInstancedStaticMeshSceneProxy::GetDynamicRayTracingInstances(struct FRayTr
 	//preallocate the worst-case to prevent an explosion of reallocs
 	//#dxr_todo: possibly track used instances and reserve based on previous behavior
 	RayTracingInstanceTemplate.InstanceTransforms.Reserve(InstanceCount);
-
+	
+	FMatrix ToWorld = InstancedRenderData.Component->GetComponentTransform().ToMatrixWithScale();
 	if (CVarRayTracingRenderInstancesCulling.GetValueOnRenderThread() > 0 && RayTracingCullClusters.Num() > 0)
 	{
 		const float BVHCullRadius = CVarRayTracingInstancesCullClusterRadius.GetValueOnRenderThread();
@@ -1163,7 +1164,6 @@ void FInstancedStaticMeshSceneProxy::GetDynamicRayTracingInstances(struct FRayTr
 		const float BVHLowScaleRadius = CVarRayTracingInstancesLowScaleCullRadius.GetValueOnRenderThread();
 		const bool ApplyGeneralCulling = BVHCullRadius > 0.0f;
 		const bool ApplyLowScaleCulling = BVHLowScaleThreshold > 0.0f && BVHLowScaleRadius > 0.0f;
-		FMatrix ToWorld = GetLocalToWorld();
 
 		// Iterate over all culling clusters
 		for (int32 ClusterIdx = 0; ClusterIdx < RayTracingCullClusters.Num(); ++ClusterIdx)
@@ -1187,6 +1187,9 @@ void FInstancedStaticMeshSceneProxy::GetDynamicRayTracingInstances(struct FRayTr
 			{
 				const uint32 InstanceIdx = Node.Instance;
 
+				const FInstancedStaticMeshInstanceData& InstanceData = InstancedRenderData.Component->PerInstanceSMData[InstanceIdx];
+				FMatrix InstanceTransform = InstanceData.Transform * ToWorld;
+				InstanceTransform.M[3][3] = 1.0f;
 				FVector InstanceLocation = Node.Center;
 				FVector VToInstanceCenter = Context.ReferenceView->ViewLocation - InstanceLocation;
 				float DistanceToInstanceCenter = VToInstanceCenter.Size();
@@ -1204,11 +1207,6 @@ void FInstancedStaticMeshSceneProxy::GetDynamicRayTracingInstances(struct FRayTr
 						continue;
 				}
 
-				FMatrix Transform;
-				InstancedRenderData.PerInstanceRenderData->InstanceBuffer.GetInstanceTransform(InstanceIdx, Transform);
-				Transform.M[3][3] = 1.0f;
-				FMatrix InstanceTransform = Transform * GetLocalToWorld();
-
 				RayTracingInstanceTemplate.InstanceTransforms.Add(InstanceTransform);
 			}
 		}
@@ -1218,10 +1216,9 @@ void FInstancedStaticMeshSceneProxy::GetDynamicRayTracingInstances(struct FRayTr
 		// No culling
 		for (int32 InstanceIdx = 0; InstanceIdx < InstanceCount; ++InstanceIdx)
 		{
-			FMatrix Transform;
-			InstancedRenderData.PerInstanceRenderData->InstanceBuffer.GetInstanceTransform(InstanceIdx, Transform);
-			Transform.M[3][3] = 1.0f;
-			FMatrix InstanceTransform = Transform * GetLocalToWorld();
+			const FInstancedStaticMeshInstanceData& InstanceData = InstancedRenderData.Component->PerInstanceSMData[InstanceIdx];
+			FMatrix InstanceTransform = InstanceData.Transform * ToWorld;
+			InstanceTransform.M[3][3] = 1.0f;
 
 			RayTracingInstanceTemplate.InstanceTransforms.Add(InstanceTransform);
 		}
