Added stat drawcount to display per pass draw counts to the screen.

diff --git a/Engine/Config/BaseInput.ini b/Engine/Config/BaseInput.ini
--- a/Engine/Config/BaseInput.ini
+++ b/Engine/Config/BaseInput.ini
@@ -93,6 +93,7 @@ MaxScrollbackSize=1024
 +ManualAutoCompleteList=(Command="SetNoPEC",Desc="<ClassName> <PropertyName> <Value> Sets property to value on objectname without Pre/Post Edit Change notifications")
 +ManualAutoCompleteList=(Command="Stat FPS",Desc="Shows FPS counter")
 +ManualAutoCompleteList=(Command="Stat UNIT",Desc="Shows hardware unit framerate")
++ManualAutoCompleteList=(Command="Stat DrawCount",Desc="Shows draw counts broken down by category")
 +ManualAutoCompleteList=(Command="Stat UnitGraph",Desc="Draws simple unit time graph")
 +ManualAutoCompleteList=(Command="Stat NamedEvents",Desc="Stat NamedEvents (Enables named events for external profilers)")
 +ManualAutoCompleteList=(Command="Stat StartFile",Desc="Stat StartFile (starts a stats capture, creating a new file in the Profiling directory; stop with stat StopFile to close the file)")
diff --git a/Engine/Source/Runtime/Engine/Classes/Engine/Engine.h b/Engine/Source/Runtime/Engine/Classes/Engine/Engine.h
--- a/Engine/Source/Runtime/Engine/Classes/Engine/Engine.h
+++ b/Engine/Source/Runtime/Engine/Classes/Engine/Engine.h
@@ -3434,6 +3434,7 @@ private:
 	int32 RenderStatLevels(UWorld* World, FViewport* Viewport, FCanvas* Canvas, int32 X, int32 Y, const FVector* ViewLocation = nullptr, const FRotator* ViewRotation = nullptr);
 	int32 RenderStatLevelMap(UWorld* World, FViewport* Viewport, FCanvas* Canvas, int32 X, int32 Y, const FVector* ViewLocation = nullptr, const FRotator* ViewRotation = nullptr);
 	int32 RenderStatUnit(UWorld* World, FViewport* Viewport, FCanvas* Canvas, int32 X, int32 Y, const FVector* ViewLocation = nullptr, const FRotator* ViewRotation = nullptr);
+	int32 RenderStatDrawCount(UWorld* World, FViewport* Viewport, FCanvas* Canvas, int32 X, int32 Y, const FVector* ViewLocation = nullptr, const FRotator* ViewRotation = nullptr);
 #if !UE_BUILD_SHIPPING
 	int32 RenderStatSoundReverb(UWorld* World, FViewport* Viewport, FCanvas* Canvas, int32 X, int32 Y, const FVector* ViewLocation = nullptr, const FRotator* ViewRotation = nullptr);
  	int32 RenderStatSoundMixes(UWorld* World, FViewport* Viewport, FCanvas* Canvas, int32 X, int32 Y, const FVector* ViewLocation = nullptr, const FRotator* ViewRotation = nullptr);
diff --git a/Engine/Source/Runtime/Engine/Private/UnrealEngine.cpp b/Engine/Source/Runtime/Engine/Private/UnrealEngine.cpp
--- a/Engine/Source/Runtime/Engine/Private/UnrealEngine.cpp
+++ b/Engine/Source/Runtime/Engine/Private/UnrealEngine.cpp
@@ -1697,6 +1697,7 @@ void UEngine::Init(IEngineLoop* InEngineLoop)
 	EngineStats.Add(FEngineStatFuncs(TEXT("STAT_FPS"), TEXT("STATCAT_Engine"), FText::GetEmpty(), &UEngine::RenderStatFPS, &UEngine::ToggleStatFPS, bIsRHS));
 	EngineStats.Add(FEngineStatFuncs(TEXT("STAT_Summary"), TEXT("STATCAT_Engine"), FText::GetEmpty(), &UEngine::RenderStatSummary, NULL, bIsRHS));
 	EngineStats.Add(FEngineStatFuncs(TEXT("STAT_Unit"), TEXT("STATCAT_Engine"), FText::GetEmpty(), &UEngine::RenderStatUnit, &UEngine::ToggleStatUnit, bIsRHS));
+	EngineStats.Add(FEngineStatFuncs(TEXT("STAT_DrawCount"), TEXT("STATCAT_Engine"), FText::GetEmpty(), &UEngine::RenderStatDrawCount, NULL, bIsRHS));
 	/* @todo Slate Rendering
 	#if STATS
 	EngineStats.Add(FEngineStatFuncs(TEXT("STAT_SlateBatches"), TEXT("STATCAT_Engine"), FText::GetEmpty(), &UEngine::RenderStatSlateBatches, NULL, true));
@@ -15109,6 +15110,34 @@ int32 UEngine::RenderStatUnit(UWorld* World, FViewport* Viewport, FCanvas* Canva
 		checkf(Viewport->GetClient()->GetStatUnitData(), TEXT("StatUnitData must be allocated for this viewport if you wish to display stat."));
 		Y = Viewport->GetClient()->GetStatUnitData()->DrawStat(Viewport, Canvas, X, Y);
 	}
+
+	return Y;
+}
+
+// DRAWCOUNT
+int32 UEngine::RenderStatDrawCount(UWorld* World, FViewport* Viewport, FCanvas* Canvas, int32 X, int32 Y, const FVector* ViewLocation, const FRotator* ViewRotation)
+{
+#if CSV_PROFILER
+	// Display all the categories of draw counts. This may always report 0 in some modes if AreGPUStatsEnabled is not enabled.
+	// Most likely because we are not currently capturing a CSV.
+	for (int32 Index = 0; Index < FDrawCallCategoryName::NumCategory; ++Index)
+	{
+		FDrawCallCategoryName* CategoryName = FDrawCallCategoryName::Array[Index];
+		Canvas->DrawShadowedString(X - 50,
+			Y,
+			*FString::Printf(TEXT("%s: %i"), *CategoryName->Name.ToString(), FDrawCallCategoryName::DisplayCounts[Index]),
+			GetSmallFont(),
+			FColor::Green);
+		Y += 12;
+	}
+#else
+	Canvas->DrawShadowedString(X - 200,
+		Y,
+		*FString::Printf(TEXT("DrawCount only supported with CSV_PROFILER")),
+		GetSmallFont(),
+		FColor::Red);
+#endif
+
 	return Y;
 }
 
diff --git a/Engine/Source/Runtime/RHI/Private/RHI.cpp b/Engine/Source/Runtime/RHI/Private/RHI.cpp
--- a/Engine/Source/Runtime/RHI/Private/RHI.cpp
+++ b/Engine/Source/Runtime/RHI/Private/RHI.cpp
@@ -134,6 +134,7 @@ FRHIResource* FRHIResource::CurrentlyDeleting = nullptr;
 TArray<FRHIResource::ResourcesToDelete> FRHIResource::DeferredDeletionQueue;
 uint32 FRHIResource::CurrentFrame = 0;
 RHI_API FDrawCallCategoryName* FDrawCallCategoryName::Array[FDrawCallCategoryName::MAX_DRAWCALL_CATEGORY];
+RHI_API int32 FDrawCallCategoryName::DisplayCounts[FDrawCallCategoryName::MAX_DRAWCALL_CATEGORY];
 RHI_API int32 FDrawCallCategoryName::NumCategory = 0;
 
 FString FVertexElement::ToString() const
@@ -786,6 +787,7 @@ void RHIPrivateBeginFrame()
 	for (int32 Index=0; Index<FDrawCallCategoryName::NumCategory; ++Index)
 	{
 		FDrawCallCategoryName* CategoryName = FDrawCallCategoryName::Array[Index];
+		FDrawCallCategoryName::DisplayCounts[Index] = CategoryName->Counter;
 		GNumDrawCallsRHI += CategoryName->Counter;
 		FCsvProfiler::RecordCustomStat(CategoryName->Name, CSV_CATEGORY_INDEX(DrawCall), CategoryName->Counter, ECsvCustomStatOp::Set);
 		CategoryName->Counter = 0;
diff --git a/Engine/Source/Runtime/RHI/Public/RHI.h b/Engine/Source/Runtime/RHI/Public/RHI.h
--- a/Engine/Source/Runtime/RHI/Public/RHI.h
+++ b/Engine/Source/Runtime/RHI/Public/RHI.h
@@ -1670,6 +1670,7 @@ struct RHI_API FDrawCallCategoryName
 
 	static constexpr int32 MAX_DRAWCALL_CATEGORY = 256;
 	static FDrawCallCategoryName* Array[MAX_DRAWCALL_CATEGORY];
+	static int32 DisplayCounts[MAX_DRAWCALL_CATEGORY]; // A backup of the counts that can be used to display on screen to avoid flickering.
 	static int32 NumCategory;
 };
 
