Add more logging to texture, trying to track down CIS failure

diff --git a/Engine/Source/Runtime/Engine/Private/TextureDerivedData.cpp b/Engine/Source/Runtime/Engine/Private/TextureDerivedData.cpp
--- a/Engine/Source/Runtime/Engine/Private/TextureDerivedData.cpp
+++ b/Engine/Source/Runtime/Engine/Private/TextureDerivedData.cpp
@@ -1162,6 +1162,12 @@ bool FTexturePlatformData::TryLoadMips(int32 FirstMipToLoad, void** OutMipData,
 					Ar.Serialize(OutMipData[MipIndex - FirstMipToLoad], MipSize);
 				}
 			}
+			else
+			{
+				UE_LOG(LogTexture, Warning, TEXT("DDC.GetAsynchronousResults() failed for %s, MipIndex: %d"),
+					Texture ? *Texture->GetPathName() : TEXT("nullptr"),
+					MipIndex);
+			}
 			TempData.Reset();
 		}
 	}
@@ -1169,9 +1175,22 @@ bool FTexturePlatformData::TryLoadMips(int32 FirstMipToLoad, void** OutMipData,
 
 	if (NumMipsCached != (LoadableMips - FirstMipToLoad))
 	{
+		UE_LOG(LogTexture, Warning, TEXT("TryLoadMips failed for %s, NumMipsCached: %d, LoadableMips: %d, FirstMipToLoad: %d"),
+			Texture ? *Texture->GetPathName() : TEXT("nullptr"),
+			NumMipsCached,
+			LoadableMips,
+			FirstMipToLoad);
+
 		// Unable to cache all mips. Release memory for those that were cached.
 		for (int32 MipIndex = FirstMipToLoad; MipIndex < LoadableMips; ++MipIndex)
 		{
+			FTexture2DMipMap& Mip = Mips[MipIndex];
+			UE_LOG(LogTexture, Warning, TEXT("  Mip %d, BulkDataSize: %d, AsyncHandle: %d, DerivedDataKey: %s"),
+				MipIndex,
+				(int32)Mip.BulkData.GetBulkDataSize(),
+				AsyncHandles[MipIndex],
+				*Mip.DerivedDataKey);
+
 			if (OutMipData && OutMipData[MipIndex - FirstMipToLoad])
 			{
 				FMemory::Free(OutMipData[MipIndex - FirstMipToLoad]);
