Fixed bug where volume sample node could not be added to the material graph.

diff --git a/Engine/Source/Runtime/Engine/Private/Materials/MaterialExpressions.cpp b/Engine/Source/Runtime/Engine/Private/Materials/MaterialExpressions.cpp
--- a/Engine/Source/Runtime/Engine/Private/Materials/MaterialExpressions.cpp
+++ b/Engine/Source/Runtime/Engine/Private/Materials/MaterialExpressions.cpp
@@ -268,7 +268,7 @@ bool IsAllowedExpressionType(const UClass* const Class, const bool bMaterialFunc
 
 	// Exclude comments from the expression list, as well as the base parameter expression, as it should not be used directly
 	const bool bSharedAllowed = Class != UMaterialExpressionComment::StaticClass() 
-		&& Class != UMaterialExpressionParameter::StaticClass() && Class != UMaterialExpressionTextureSampleParameterVolume::StaticClass()
+		&& Class != UMaterialExpressionParameter::StaticClass()
 		&& (Class != UMaterialExpressionTextureSampleParameter2DArray::StaticClass() || AllowTextureArrayAssetCreationVar->GetValueOnGameThread() != 0);
 
 	if (bMaterialFunction)
