Removed some logging that is editor only

diff --git a/Engine/Source/Runtime/Engine/Private/TextureDerivedData.cpp b/Engine/Source/Runtime/Engine/Private/TextureDerivedData.cpp
--- a/Engine/Source/Runtime/Engine/Private/TextureDerivedData.cpp
+++ b/Engine/Source/Runtime/Engine/Private/TextureDerivedData.cpp
@@ -1185,11 +1185,9 @@ bool FTexturePlatformData::TryLoadMips(int32 FirstMipToLoad, void** OutMipData,
 		for (int32 MipIndex = FirstMipToLoad; MipIndex < LoadableMips; ++MipIndex)
 		{
 			FTexture2DMipMap& Mip = Mips[MipIndex];
-			UE_LOG(LogTexture, Warning, TEXT("  Mip %d, BulkDataSize: %d, AsyncHandle: %d, DerivedDataKey: %s"),
+			UE_LOG(LogTexture, Warning, TEXT("  Mip %d, BulkDataSize: %d"),
 				MipIndex,
-				(int32)Mip.BulkData.GetBulkDataSize(),
-				AsyncHandles[MipIndex],
-				*Mip.DerivedDataKey);
+				(int32)Mip.BulkData.GetBulkDataSize());
 
 			if (OutMipData && OutMipData[MipIndex - FirstMipToLoad])
 			{
