Add IsUsedWithVirtualHeightfieldMesh to material settings.
For use with the VirtualHeightfieldMesh plugin.

diff --git a/Engine/Source/Runtime/Engine/Classes/Materials/Material.h b/Engine/Source/Runtime/Engine/Classes/Materials/Material.h
--- a/Engine/Source/Runtime/Engine/Classes/Materials/Material.h
+++ b/Engine/Source/Runtime/Engine/Classes/Materials/Material.h
@@ -719,6 +719,13 @@ public:
 	UPROPERTY(EditAnywhere, BlueprintReadOnly, Category = Usage)
 	uint32 bUsedWithLidarPointCloud : 1;
 
+	/**
+	 * Indicates that the material and its instances can be used with Virtual Heightfield Mesh.
+	 * This will result in the shaders required to support Virtual Heightfield Mesh geometries being compiled which will increase shader compile time and memory usage.
+	 */
+	UPROPERTY(EditAnywhere, BlueprintReadOnly, Category = Usage)
+	uint32 bUsedWithVirtualHeightfieldMesh : 1;
+
 	/** 
 	 * Indicates that the material and its instances can be used with Slate UI and UMG
 	 * This will result in the shaders required to support UI materials being compiled which will increase shader compile time and memory usage.
diff --git a/Engine/Source/Runtime/Engine/Private/DebugViewModeMaterialProxy.cpp b/Engine/Source/Runtime/Engine/Private/DebugViewModeMaterialProxy.cpp
--- a/Engine/Source/Runtime/Engine/Private/DebugViewModeMaterialProxy.cpp
+++ b/Engine/Source/Runtime/Engine/Private/DebugViewModeMaterialProxy.cpp
@@ -62,6 +62,7 @@ FDebugViewModeMaterialProxy::FDebugViewModeMaterialProxy(
 				bIsUsedWithInstancedStaticMeshes = Resource->IsUsedWithInstancedStaticMeshes();
 				bIsUsedWithAPEXCloth = Resource->IsUsedWithAPEXCloth();
 				bIsUsedWithWater = Resource->IsUsedWithWater();
+				bIsUsedWithVirtualHeightfieldMesh = Resource->IsUsedWithVirtualHeightfieldMesh();
 			}
 		}
 
diff --git a/Engine/Source/Runtime/Engine/Private/DebugViewModeMaterialProxy.h b/Engine/Source/Runtime/Engine/Private/DebugViewModeMaterialProxy.h
--- a/Engine/Source/Runtime/Engine/Private/DebugViewModeMaterialProxy.h
+++ b/Engine/Source/Runtime/Engine/Private/DebugViewModeMaterialProxy.h
@@ -168,6 +168,7 @@ public:
 	virtual bool IsUsedWithInstancedStaticMeshes() const override { return bIsUsedWithInstancedStaticMeshes; }
 	virtual bool IsUsedWithAPEXCloth() const override { return bIsUsedWithAPEXCloth; }
 	virtual bool IsUsedWithWater() const override { return bIsUsedWithWater; }
+	virtual bool IsUsedWithVirtualHeightfieldMesh() const override { return bIsUsedWithVirtualHeightfieldMesh; }
 
 	virtual EMaterialShaderMapUsage::Type GetMaterialShaderMapUsage() const { return Usage; }
 
@@ -202,6 +203,7 @@ private:
 			uint32 bIsUsedWithInstancedStaticMeshes : 1;
 			uint32 bIsUsedWithAPEXCloth : 1;
 			uint32 bIsUsedWithWater : 1;
+			uint32 bIsUsedWithVirtualHeightfieldMesh : 1;
 		};
 	};
 
diff --git a/Engine/Source/Runtime/Engine/Private/Materials/MaterialShared.cpp b/Engine/Source/Runtime/Engine/Private/Materials/MaterialShared.cpp
--- a/Engine/Source/Runtime/Engine/Private/Materials/MaterialShared.cpp
+++ b/Engine/Source/Runtime/Engine/Private/Materials/MaterialShared.cpp
@@ -1076,6 +1076,11 @@ bool FMaterialResource::IsUsedWithLidarPointCloud() const
 	return Material->bUsedWithLidarPointCloud;
 }
 
+bool FMaterialResource::IsUsedWithVirtualHeightfieldMesh() const
+{
+	return Material->bUsedWithVirtualHeightfieldMesh;
+}
+
 bool FMaterialResource::IsUsedWithLandscape() const
 {
 	return false;
diff --git a/Engine/Source/Runtime/Engine/Public/MaterialShared.h b/Engine/Source/Runtime/Engine/Public/MaterialShared.h
--- a/Engine/Source/Runtime/Engine/Public/MaterialShared.h
+++ b/Engine/Source/Runtime/Engine/Public/MaterialShared.h
@@ -1552,6 +1552,7 @@ public:
 	virtual bool IsUsedWithWater() const { return false; }
 	virtual bool IsUsedWithHairStrands() const { return false; }
 	virtual bool IsUsedWithLidarPointCloud() const { return false; }
+	virtual bool IsUsedWithVirtualHeightfieldMesh() const { return false; }
 	ENGINE_API virtual enum EMaterialTessellationMode GetTessellationMode() const;
 	virtual bool IsCrackFreeDisplacementEnabled() const { return false; }
 	virtual bool IsAdaptiveTessellationEnabled() const { return false; }
@@ -2361,6 +2362,7 @@ public:
 	ENGINE_API virtual bool IsUsedWithWater() const override;
 	ENGINE_API virtual bool IsUsedWithHairStrands() const override;
 	ENGINE_API virtual bool IsUsedWithLidarPointCloud() const override;
+	ENGINE_API virtual bool IsUsedWithVirtualHeightfieldMesh() const override;
 	ENGINE_API virtual enum EMaterialTessellationMode GetTessellationMode() const override;
 	ENGINE_API virtual bool IsCrackFreeDisplacementEnabled() const override;
 	ENGINE_API virtual bool IsAdaptiveTessellationEnabled() const override;
@@ -2977,6 +2979,7 @@ struct FMaterialShaderParameters
 			uint64 bMaterialIsComplexityAccumulate : 1;
 			uint64 bMaterialIsLODColoration : 1;
 			uint64 bIsUsedWithLidarPointCloud : 1;
+			uint64 bIsUsedWithVirtualHeightfieldMesh : 1;
 			uint64 bIsStencilTestEnabled : 1;
 		};
 	};
@@ -3033,6 +3036,7 @@ struct FMaterialShaderParameters
 		bIsUsedWithInstancedStaticMeshes = InMaterial->IsUsedWithInstancedStaticMeshes();
 		bHasRuntimeVirtualTextureOutput = InMaterial->HasRuntimeVirtualTextureOutput();
 		bIsUsedWithLidarPointCloud = InMaterial->IsUsedWithLidarPointCloud();
+		bIsUsedWithVirtualHeightfieldMesh = InMaterial->IsUsedWithVirtualHeightfieldMesh();
 		bIsStencilTestEnabled = InMaterial->IsStencilTestEnabled();
 
 		// See FDebugViewModeMaterialProxy::GetFriendlyName()
