Fix a typo in a newly added API method (IsInSeperateFile -> IsInSeparateFile) before the 4.25 release to avoid future issues.

diff --git a/Engine/Source/Runtime/CoreUObject/Private/Serialization/BulkData2.cpp b/Engine/Source/Runtime/CoreUObject/Private/Serialization/BulkData2.cpp
--- a/Engine/Source/Runtime/CoreUObject/Private/Serialization/BulkData2.cpp
+++ b/Engine/Source/Runtime/CoreUObject/Private/Serialization/BulkData2.cpp
@@ -807,7 +807,7 @@ void FBulkDataBase::Serialize(FArchive& Ar, UObject* Owner, int32 /*Index*/, boo
 					}
 				}
 			}
-			else if (!Ar.IsAllowingLazyLoading() && !IsInSeperateFile())
+			else if (!Ar.IsAllowingLazyLoading() && !IsInSeparateFile())
 			{
 
 				// If the archive does not support lazy loading and the data is not in a different file then we have to load 
@@ -1082,7 +1082,7 @@ bool FBulkDataBase::IsInlined() const
 	return	(GetBulkDataFlags() & BULKDATA_PayloadAtEndOfFile) == 0;
 }
 
-bool FBulkDataBase::IsInSeperateFile() const
+bool FBulkDataBase::IsInSeparateFile() const
 {
 	return	(GetBulkDataFlags() & BULKDATA_PayloadInSeperateFile) != 0;
 }
@@ -1637,7 +1637,7 @@ FString FBulkDataBase::ConvertFilenameFromFlags(const FString& Filename) const
 		// optional data first.
 		return FPathViews::ChangeExtension(Filename, BulkDataExt::Optional);
 	}
-	else if (!IsInSeperateFile())
+	else if (!IsInSeparateFile())
 	{
 		return Filename;
 	}
diff --git a/Engine/Source/Runtime/CoreUObject/Public/Serialization/BulkData.h b/Engine/Source/Runtime/CoreUObject/Public/Serialization/BulkData.h
--- a/Engine/Source/Runtime/CoreUObject/Public/Serialization/BulkData.h
+++ b/Engine/Source/Runtime/CoreUObject/Public/Serialization/BulkData.h
@@ -427,14 +427,14 @@ public:
 		return   (GetBulkDataFlags() & BULKDATA_PayloadAtEndOfFile) == 0;
 	}
 
-	UE_DEPRECATED(4.25, "Use ::IsInSeperateFile() instead")
-	FORCEINLINE bool InSeperateFile() const { return IsInSeperateFile(); }
+	UE_DEPRECATED(4.25, "Use ::IsInSeparateFile() instead")
+	FORCEINLINE bool InSeperateFile() const { return IsInSeparateFile(); }
 
 	/**
 	* Returns whether this bulk data is currently stored in it's own file or not
 	* @return true if BULKDATA_PayloadInSeperateFile is not set
 	*/
-	bool IsInSeperateFile() const
+	bool IsInSeparateFile() const
 	{
 		return	(GetBulkDataFlags() & BULKDATA_PayloadInSeperateFile) != 0;
 	}
diff --git a/Engine/Source/Runtime/CoreUObject/Public/Serialization/BulkData2.h b/Engine/Source/Runtime/CoreUObject/Public/Serialization/BulkData2.h
--- a/Engine/Source/Runtime/CoreUObject/Public/Serialization/BulkData2.h
+++ b/Engine/Source/Runtime/CoreUObject/Public/Serialization/BulkData2.h
@@ -177,9 +177,9 @@ public:
 	bool IsDuplicateNonOptional() const;
 	bool IsOptional() const;
 	bool IsInlined() const;
-	UE_DEPRECATED(4.25, "Use ::IsInSeperateFile() instead")
-	FORCEINLINE bool InSeperateFile() const { return IsInSeperateFile(); }
-	bool IsInSeperateFile() const;
+	UE_DEPRECATED(4.25, "Use ::IsInSeparateFile() instead")
+	FORCEINLINE bool InSeperateFile() const { return IsInSeparateFile(); }
+	bool IsInSeparateFile() const;
 	bool IsSingleUse() const;
 	UE_DEPRECATED(4.26, "Use ::IsFileMemoryMapped() instead")
 	bool IsMemoryMapped() const { return IsFileMemoryMapped(); }
diff --git a/Engine/Source/Runtime/Engine/Private/Texture2D.cpp b/Engine/Source/Runtime/Engine/Private/Texture2D.cpp
--- a/Engine/Source/Runtime/Engine/Private/Texture2D.cpp
+++ b/Engine/Source/Runtime/Engine/Private/Texture2D.cpp
@@ -543,7 +543,7 @@ bool UTexture2D::GetMipDataFilename(const int32 MipIndex, FString& OutBulkDataFi
 #else
 			OutBulkDataFilename = PlatformData->CachedPackageFileName;
 
-			if (PlatformData->Mips[MipIndex].BulkData.IsInSeperateFile())
+			if (PlatformData->Mips[MipIndex].BulkData.IsInSeparateFile())
 			{	
 				const bool UseOptionalBulkDataFileName = PlatformData->Mips[MipIndex].BulkData.IsOptional();
 				OutBulkDataFilename = FPaths::ChangeExtension(OutBulkDataFilename, UseOptionalBulkDataFileName ? TEXT(".uptnl") : TEXT(".ubulk"));
diff --git a/Engine/Source/Runtime/Engine/Private/TextureDerivedData.cpp b/Engine/Source/Runtime/Engine/Private/TextureDerivedData.cpp
--- a/Engine/Source/Runtime/Engine/Private/TextureDerivedData.cpp
+++ b/Engine/Source/Runtime/Engine/Private/TextureDerivedData.cpp
@@ -1129,7 +1129,7 @@ bool FTexturePlatformData::TryLoadMips(int32 FirstMipToLoad, void** OutMipData,
 				// is it possible for streaming mips to be loaded in non streaming ways.
 				if (CVarSetTextureStreaming.GetValueOnAnyThread() != 0)
 				{
-					UE_CLOG(Mip.BulkData.IsInSeperateFile(), LogTexture, Error, TEXT("Loading non-streamed mips from an external bulk file.  This is not desireable.  File %s"), *(Mip.BulkData.GetFilename() ) );
+					UE_CLOG(Mip.BulkData.IsInSeparateFile(), LogTexture, Error, TEXT("Loading non-streamed mips from an external bulk file.  This is not desireable.  File %s"), *(Mip.BulkData.GetFilename() ) );
 				}
 #endif
 				Mip.BulkData.GetCopy(&OutMipData[MipIndex - FirstMipToLoad], true);
@@ -1210,7 +1210,7 @@ int32 FTexturePlatformData::GetNumNonStreamingMips() const
 
 		for (const FTexture2DMipMap& Mip : Mips)
 		{
-			if ( Mip.BulkData.IsInSeperateFile() || !Mip.BulkData.IsInlined() )
+			if ( Mip.BulkData.IsInSeparateFile() || !Mip.BulkData.IsInlined() )
 			{
 				--NumNonStreamingMips;
 			}
diff --git a/Engine/Source/Runtime/Engine/Public/TextureResource.h b/Engine/Source/Runtime/Engine/Public/TextureResource.h
--- a/Engine/Source/Runtime/Engine/Public/TextureResource.h
+++ b/Engine/Source/Runtime/Engine/Public/TextureResource.h
@@ -74,9 +74,9 @@ struct FTexture2DMipMap
 		bool IsAvailableForUse() const { return !(BulkDataFlags & BULKDATA_Unused); }
 		bool IsOptional() const { return (BulkDataFlags & BULKDATA_OptionalPayload) != 0; }
 		bool IsInlined() const { return (GetBulkDataFlags() & BULKDATA_PayloadInSeperateFile) == 0 && (GetBulkDataFlags() & BULKDATA_PayloadAtEndOfFile) == 0; }
-		UE_DEPRECATED(4.25, "Use ::IsInSeperateFile() instead")
-		FORCEINLINE bool InSeperateFile() const { return IsInSeperateFile(); }
-		bool IsInSeperateFile() const { return (GetBulkDataFlags() & BULKDATA_PayloadInSeperateFile) != 0; }
+		UE_DEPRECATED(4.25, "Use ::IsInSeparateFile() instead")
+		FORCEINLINE bool InSeperateFile() const { return IsInSeparateFile(); }
+		bool IsInSeparateFile() const { return (GetBulkDataFlags() & BULKDATA_PayloadInSeperateFile) != 0; }
 		bool IsBulkDataLoaded() const { return IsInlined(); }
 		bool IsAsyncLoadingComplete() const { return true; }
 		bool IsStoredCompressedOnDisk() const { return !!(BulkDataFlags & BULKDATA_SerializeCompressed); }
