Fixed CIS - bad forward declare causing nounity/nopch failures.

diff --git a/Engine/Source/Runtime/Engine/Classes/PhysicalMaterials/Experimental/ChaosPhysicalMaterial.h b/Engine/Source/Runtime/Engine/Classes/PhysicalMaterials/Experimental/ChaosPhysicalMaterial.h
--- a/Engine/Source/Runtime/Engine/Classes/PhysicalMaterials/Experimental/ChaosPhysicalMaterial.h
+++ b/Engine/Source/Runtime/Engine/Classes/PhysicalMaterials/Experimental/ChaosPhysicalMaterial.h
@@ -13,7 +13,10 @@
 #include "Physics/PhysicsInterfaceCore.h"
 #include "ChaosPhysicalMaterial.generated.h"
 
-class Chaos::FChaosPhysicsMaterial;
+namespace Chaos
+{
+	class FChaosPhysicsMaterial;
+}
 
 /**
  * Physical materials are used to define the response of a physical object when 
