Deprecate stat sounds/soundwaves/soundcues in favor of au.debug. variants.
- Allows user to enable and disable uniformly and avoid edgecases with stat system not reporting toggle and render call getting out-of-sync.
- Move debugger & internals to audio namespace to avoid collisions in unity build

diff --git a/Engine/Source/Editor/AudioEditor/Private/AssetTypeActions/AssetTypeActions_SoundBase.cpp b/Engine/Source/Editor/AudioEditor/Private/AssetTypeActions/AssetTypeActions_SoundBase.cpp
--- a/Engine/Source/Editor/AudioEditor/Private/AssetTypeActions/AssetTypeActions_SoundBase.cpp
+++ b/Engine/Source/Editor/AudioEditor/Private/AssetTypeActions/AssetTypeActions_SoundBase.cpp
@@ -296,7 +296,7 @@ void FAssetTypeActions_SoundBase::ExecuteMuteSound(TArray<TWeakObjectPtr<USoundB
 {	
 	if (FAudioDeviceManager* ADM = GEditor->GetAudioDeviceManager())
 	{
-		FAudioDebugger& Debugger = ADM->GetDebugger();
+		Audio::FAudioDebugger& Debugger = ADM->GetDebugger();
 
 		// In a selection that consists of some already muted, toggle everything in the same direction,
 		// to avoid AB problem.
@@ -320,7 +320,7 @@ void FAssetTypeActions_SoundBase::ExecuteSoloSound(TArray<TWeakObjectPtr<USoundB
 {
 	if (FAudioDeviceManager* ADM = GEditor->GetAudioDeviceManager())
 	{
-		FAudioDebugger& Debugger = ADM->GetDebugger();
+		Audio::FAudioDebugger& Debugger = ADM->GetDebugger();
 
 		// In a selection that consists of some already soloed, toggle everything in the same direction,
 		// to avoid AB problem.
@@ -345,7 +345,7 @@ bool FAssetTypeActions_SoundBase::IsActionCheckedMute(TArray<TWeakObjectPtr<USou
 	if (FAudioDeviceManager* ADM = GEditor->GetAudioDeviceManager())
 	{
 		// If *any* of the selection are muted, show the tick box as ticked.
-		FAudioDebugger& Debugger = ADM->GetDebugger();
+		Audio::FAudioDebugger& Debugger = ADM->GetDebugger();
 		for (TWeakObjectPtr<USoundBase> SoundBase : Objects)
 		{
 			if (USoundCue* Cue = Cast<USoundCue>(SoundBase.Get()))
@@ -372,7 +372,7 @@ bool FAssetTypeActions_SoundBase::IsActionCheckedSolo(TArray<TWeakObjectPtr<USou
 	// If *any* of the selection are solod, show the tick box as ticked.
 	if (FAudioDeviceManager* ADM = GEditor->GetAudioDeviceManager())
 	{
-		FAudioDebugger& Debugger = ADM->GetDebugger();
+		Audio::FAudioDebugger& Debugger = ADM->GetDebugger();
 		for (TWeakObjectPtr<USoundBase> SoundBase : Objects)
 		{
 			if (USoundCue* Cue = Cast<USoundCue>(SoundBase.Get()))
@@ -399,7 +399,7 @@ bool FAssetTypeActions_SoundBase::CanExecuteMuteCommand(TArray<TWeakObjectPtr<US
 	if (FAudioDeviceManager* ADM = GEditor->GetAudioDeviceManager())
 	{
 		// Allow muting if we're not Soloing.
-		FAudioDebugger& Debugger = ADM->GetDebugger();
+		Audio::FAudioDebugger& Debugger = ADM->GetDebugger();
 		for (TWeakObjectPtr<USoundBase> SoundBase : Objects)
 		{
 			if (USoundCue* Cue = Cast<USoundCue>(SoundBase.Get()))
@@ -429,7 +429,7 @@ bool FAssetTypeActions_SoundBase::CanExecuteSoloCommand(TArray<TWeakObjectPtr<US
 	if (FAudioDeviceManager* ADM = GEditor->GetAudioDeviceManager())
 	{
 		// Allow Soloing if we're not Muting.
-		FAudioDebugger& Debugger = ADM->GetDebugger();
+		Audio::FAudioDebugger& Debugger = ADM->GetDebugger();
 		for (TWeakObjectPtr<USoundBase> SoundBase : Objects)
 		{
 			if (USoundCue* Cue = Cast<USoundCue>(SoundBase.Get()))
diff --git a/Engine/Source/Editor/AudioEditor/Private/SoundClassEditor.h b/Engine/Source/Editor/AudioEditor/Private/SoundClassEditor.h
--- a/Engine/Source/Editor/AudioEditor/Private/SoundClassEditor.h
+++ b/Engine/Source/Editor/AudioEditor/Private/SoundClassEditor.h
@@ -15,6 +15,11 @@ class SDockableTab;
 class UEdGraph;
 class USoundClass;
 
+namespace Audio
+{
+	class FAudioDebugger;
+}
+
 //////////////////////////////////////////////////////////////////////////
 // FSoundClassEditor
 
@@ -119,7 +124,7 @@ private:
 	TSharedPtr<FUICommandList> GraphEditorCommands;
 
 	/** Cache the audio debugger instance */
-	class FAudioDebugger* Debugger;
+	Audio::FAudioDebugger* Debugger;
 
 	/**	The tab ids for all the tabs used */
 	static const FName GraphCanvasTabId;
diff --git a/Engine/Source/Editor/AudioEditor/Private/SoundCueEditor.h b/Engine/Source/Editor/AudioEditor/Private/SoundCueEditor.h
--- a/Engine/Source/Editor/AudioEditor/Private/SoundCueEditor.h
+++ b/Engine/Source/Editor/AudioEditor/Private/SoundCueEditor.h
@@ -18,6 +18,11 @@ class USoundCue;
 struct FPropertyChangedEvent;
 struct Rect;
 
+namespace Audio
+{
+	class FAudioDebugger;
+}
+
 /*-----------------------------------------------------------------------------
    FSoundCueEditor
 -----------------------------------------------------------------------------*/
@@ -211,7 +216,7 @@ private:
 	TSharedPtr<FUICommandList> GraphEditorCommands;
 
 	/** Cache of the Audio debugger instance */
-	class FAudioDebugger* Debugger;
+	Audio::FAudioDebugger* Debugger;
 
 	/**	The tab ids for all the tabs used */
 	static const FName GraphCanvasTabId;
diff --git a/Engine/Source/Editor/UnrealEd/Private/EditorServer.cpp b/Engine/Source/Editor/UnrealEd/Private/EditorServer.cpp
--- a/Engine/Source/Editor/UnrealEd/Private/EditorServer.cpp
+++ b/Engine/Source/Editor/UnrealEd/Private/EditorServer.cpp
@@ -2605,6 +2605,11 @@ bool UEditorEngine::Map_Load(const TCHAR* Str, FOutputDevice& Ar)
 					FReferenceChainSearch RefChainSearch(WorldPackage, EReferenceChainSearchMode::Shortest | EReferenceChainSearchMode::PrintResults);
 					UE_LOG(LogEditorServer, Fatal, TEXT("Failed to find the world in already loaded world package %s! Referenced by:") LINE_TERMINATOR TEXT("%s"), *WorldPackage->GetPathName(), *RefChainSearch.GetRootPath());
 				}
+
+				if (Context.AudioDeviceID == INDEX_NONE && GEngine->GetAudioDeviceManager())
+				{
+					Context.AudioDeviceID = GEngine->GetAudioDeviceManager()->GetMainAudioDeviceID();
+				}
 				Context.SetCurrentWorld(World);
 				GWorld = World;
 
diff --git a/Engine/Source/Editor/UnrealEd/Private/EditorViewportClient.cpp b/Engine/Source/Editor/UnrealEd/Private/EditorViewportClient.cpp
--- a/Engine/Source/Editor/UnrealEd/Private/EditorViewportClient.cpp
+++ b/Engine/Source/Editor/UnrealEd/Private/EditorViewportClient.cpp
@@ -5901,7 +5901,7 @@ void FEditorViewportClient::SetEnabledStats(const TArray<FString>& InEnabledStat
 	{
 		if (FAudioDeviceManager* DeviceManager = GEngine->GetAudioDeviceManager())
 		{
-			FAudioDebugger::ResolveDesiredStats(this);
+			Audio::FAudioDebugger::ResolveDesiredStats(this);
 		}
 	}
 #endif // ENABLE_AUDIO_DEBUG
diff --git a/Engine/Source/Runtime/AudioMixer/Private/AudioMixerSource.cpp b/Engine/Source/Runtime/AudioMixer/Private/AudioMixerSource.cpp
--- a/Engine/Source/Runtime/AudioMixer/Private/AudioMixerSource.cpp
+++ b/Engine/Source/Runtime/AudioMixer/Private/AudioMixerSource.cpp
@@ -423,7 +423,7 @@ namespace Audio
 		UpdateChannelMaps();
 
 #if ENABLE_AUDIO_DEBUG
-		FAudioDebugger::DrawDebugInfo(*this);
+		Audio::FAudioDebugger::DrawDebugInfo(*this);
 #endif // ENABLE_AUDIO_DEBUG
 	}
 
diff --git a/Engine/Source/Runtime/Engine/Private/ActiveSound.cpp b/Engine/Source/Runtime/Engine/Private/ActiveSound.cpp
--- a/Engine/Source/Runtime/Engine/Private/ActiveSound.cpp
+++ b/Engine/Source/Runtime/Engine/Private/ActiveSound.cpp
@@ -717,7 +717,7 @@ void FActiveSound::UpdateWaveInstances(TArray<FWaveInstance*> &InWaveInstances,
 	{
 		DebugColor = FColor::MakeRandomColor();
 	}
-	FAudioDebugger::DrawDebugInfo(*this, ThisSoundsWaveInstances, DeltaTime);
+	Audio::FAudioDebugger::DrawDebugInfo(*this, ThisSoundsWaveInstances, DeltaTime);
 #endif // ENABLE_AUDIO_DEBUG
 
 	InWaveInstances.Append(ThisSoundsWaveInstances);
diff --git a/Engine/Source/Runtime/Engine/Private/Audio.cpp b/Engine/Source/Runtime/Engine/Private/Audio.cpp
--- a/Engine/Source/Runtime/Engine/Private/Audio.cpp
+++ b/Engine/Source/Runtime/Engine/Private/Audio.cpp
@@ -473,7 +473,7 @@ float FSoundSource::GetDebugVolume(const float InVolume)
 	}
 
 	// Solos/Mutes (dev only).
-	FAudioDebugger& Debugger = GEngine->GetAudioDeviceManager()->GetDebugger();	
+	Audio::FAudioDebugger& Debugger = GEngine->GetAudioDeviceManager()->GetDebugger();	
 	FDebugInfo Info;
 				
 	// SoundWave Solo/Mutes.
diff --git a/Engine/Source/Runtime/Engine/Private/AudioDebug.cpp b/Engine/Source/Runtime/Engine/Private/AudioDebug.cpp
--- a/Engine/Source/Runtime/Engine/Private/AudioDebug.cpp
+++ b/Engine/Source/Runtime/Engine/Private/AudioDebug.cpp
@@ -26,6 +26,13 @@
 #include "Sound/SoundWave.h"
 #include "UnrealEngine.h"
 
+#if WITH_EDITOR
+#include "Editor.h"
+#include "Engine/GameViewportClient.h"
+#include "LevelEditorViewport.h"
+#include "UnrealClient.h"
+#endif // WITH_EDITOR
+
 
 #ifndef ENABLE_AUDIO_DEBUG
 #error "Please define ENABLE_AUDIO_DEBUG"
@@ -138,19 +145,23 @@ FAutoConsoleVariableRef CVarSoundCueDebugCharSpacing(
 	TEXT("Default: 7"),
 	ECVF_Default);
 
-// Test Data
-namespace ERequestedAudioStats
-{
-	static const uint8 SoundWaves = 0x1;
-	static const uint8 SoundCues = 0x2;
-	static const uint8 Sounds = 0x4;
-	static const uint8 SoundMixes = 0x8;
-	static const uint8 SoundModulation = 0x10;
-	static const uint8 DebugSounds = 0x20;
-	static const uint8 LongSoundNames = 0x40;
-};
-
-namespace
+static int32 SoundDebugDisplayCornerXCVar = 100;
+FAutoConsoleVariableRef CVarSoundDebugDisplayCornerX(
+	TEXT("au.Debug.Display.X"),
+	SoundDebugDisplayCornerXCVar,
+	TEXT("X position on screen of debug statistics. \n")
+	TEXT("Default: 100"),
+	ECVF_Default);
+
+static int32 SoundDebugDisplayCornerYCVar = 62;
+FAutoConsoleVariableRef CVarSoundDebugDisplayCornerY(
+	TEXT("au.Debug.Display.Y"),
+	SoundDebugDisplayCornerYCVar,
+	TEXT("X position on screen of debug statistics. \n")
+	TEXT("Default: 62"),
+	ECVF_Default);
+
+namespace Audio
 {
 	const FColor HeaderColor = FColor::Green;
 
@@ -197,7 +208,17 @@ namespace
 		return FColor::White;
 	}
 
-	const int32  TabWidth = 12;
+	static bool bAllowUsingDeprecatedDebugStats = true;
+
+	// Whether or not respective stat data is active for any audio device (set on game thread)
+	static bool bDebugWavesForAllViewsEnabled = false;
+	static bool bDebugCuesForAllViewsEnabled = false;
+	static bool bDebugSoundsForAllViewsEnabled = false;
+	static bool bDebugMixesForAllViewsEnabled = false;
+	static bool bDebugReverbForAllViewsEnabled = false;
+	static bool bDebugModulationForAllViewsEnabled = false;
+
+	const int32 DebuggerTabWidth = 12;
 
 	const float MinDisplayVolume = KINDA_SMALL_NUMBER; // -80 dB
 
@@ -215,6 +236,17 @@ namespace
 
 	struct FAudioStats
 	{
+		static const uint32 SoundWaves = 0x1;
+		static const uint32 SoundCues = 0x2;
+		static const uint32 Sounds = 0x4;
+		static const uint32 SoundMixes = 0x8;
+		static const uint32 SoundModulation = 0x10;
+		static const uint32 SoundReverb = 0x20;
+
+		// TODO: Move to console variables
+		static const uint32 DebugSounds = 0x40;
+		static const uint32 LongSoundNames = 0x80;
+
 		enum class EDisplaySort : uint8
 		{
 			Class,
@@ -239,7 +271,7 @@ namespace
 			float Volume;
 			int32 InstanceIndex;
 			FName WaveInstanceName;
-			uint8 bPlayWhenSilent:1;
+			uint8 bPlayWhenSilent : 1;
 		};
 
 		struct FStatSoundInfo
@@ -269,6 +301,7 @@ namespace
 
 		uint8 DisplayFlags;
 		EDisplaySort DisplaySort;
+		uint32 EnabledStats;
 		TArray<FTransform> ListenerTransforms;
 		TArray<FStatSoundInfo> StatSoundInfos;
 		TArray<FStatSoundMix> StatSoundMixes;
@@ -276,6 +309,7 @@ namespace
 		FAudioStats()
 			: DisplayFlags(0)
 			, DisplaySort(EDisplaySort::Name)
+			, EnabledStats(0)
 		{
 		}
 	};
@@ -314,7 +348,7 @@ namespace
 
 	struct FAudioStats_AudioThread
 	{
-		uint8 RequestedStats;
+		uint32 RequestedStats;
 
 		FAudioStats_AudioThread()
 			: RequestedStats(0)
@@ -322,8 +356,8 @@ namespace
 		}
 	};
 
-	TMap<uint32, FAudioStats> AudioDeviceStats;
-	TMap<uint32, FAudioStats_AudioThread> AudioDeviceStats_AudioThread;
+	TMap<FDeviceId, FAudioStats> AudioDeviceStats;
+	TMap<FDeviceId, FAudioStats_AudioThread> AudioDeviceStats_AudioThread;
 
 	void HandleDumpActiveSounds(UWorld* World)
 	{
@@ -332,7 +366,7 @@ namespace
 			GEngine->GetAudioDeviceManager()->GetDebugger().DumpActiveSounds();
 		}
 	}
-	
+
 	void HandleClearMutesAndSolos(UWorld* World)
 	{
 		if (GEngine && GEngine->GetAudioDeviceManager())
@@ -341,1557 +375,1904 @@ namespace
 		}
 	}
 
-	FAutoConsoleCommandWithWorld DumpActiveSounds(TEXT("au.DumpActiveSounds"), TEXT("Outputs data about all the currently active sounds."), FConsoleCommandWithWorldDelegate::CreateStatic(&HandleDumpActiveSounds), ECVF_Cheat);
-	FAutoConsoleCommandWithWorld ClearMutesAndSolos(TEXT("au.ClearMutesAndSolos"), TEXT("Clears any solo-ing/mute-ing sounds"), FConsoleCommandWithWorldDelegate::CreateStatic(&HandleClearMutesAndSolos), ECVF_Cheat);
-
-} // namespace <>
-
-
-/**
- * Audio Debugger Implementation
- */
-FAudioDebugger::FAudioDebugger()
-	: bVisualize3dDebug(0)
-{
-}
-
-bool FAudioDebugger::IsVisualizeDebug3dEnabled() const
-{
-	return bVisualize3dDebug;
-}
+	void DebugSoundObject(const TArray<FString>& Args, UWorld* InWorld, const uint32 InStatToEnable, bool& bAllEnabled)
+	{
+		bAllowUsingDeprecatedDebugStats = false;
 
-void FAudioDebugger::ToggleVisualizeDebug3dEnabled()
-{
-	bVisualize3dDebug = !bVisualize3dDebug;
-}
+		if (Args.Num() > 0)
+		{
+			bool bAllViews = false;
+			for (int32 i = 1; i < Args.Num(); ++i)
+			{
+				bAllViews |= Args[i] == TEXT("-AllViews");
+			}
 
-#if WITH_EDITOR
-void FAudioDebugger::OnBeginPIE()
-{
-	FAudioDeviceManager* DeviceManager = GEngine->GetAudioDeviceManager();
-	if (DeviceManager)
-	{
-		DeviceManager->GetDebugger().ClearMutesAndSolos();
+			bool bEnablementRequest = Args[0].ToBool();
+			bAllEnabled = bAllViews ? bEnablementRequest : false;
+			bEnablementRequest
+				? Audio::FAudioDebugger::SetStats(InStatToEnable, bAllViews ? nullptr : InWorld)
+				: Audio::FAudioDebugger::ClearStats(InStatToEnable, bAllViews ? nullptr : InWorld);
+		}
 	}
-}
 
-void FAudioDebugger::OnEndPIE()
-{
-	FAudioDeviceManager* DeviceManager = GEngine->GetAudioDeviceManager();
-	if (DeviceManager)
+	bool DebugShouldRenderStat(UWorld* World, FCanvas* Canvas, bool bEnablementBool, uint32 InAudioStat, FAudioDevice** OutAudioDevice)
 	{
-		DeviceManager->GetDebugger().ClearMutesAndSolos();
-	}
-}
-#endif // WITH_EDITOR
-
-void FAudioDebugger::QuerySoloMuteSoundClass(const FString& Name, bool& bOutIsSoloed, bool& bOutIsMuted, FString& OutReason) const
-{
-	GetDebugSoloMuteStateX(Name, DebugNames.SoloSoundClass, DebugNames.MuteSoundClass, bOutIsSoloed, bOutIsMuted, OutReason);
-}
-
-void FAudioDebugger::QuerySoloMuteSoundWave(const FString& Name, bool& bOutIsSoloed, bool& bOutIsMuted, FString& OutReason) const
-{
-	GetDebugSoloMuteStateX(Name, DebugNames.SoloSoundWave, DebugNames.MuteSoundWave, bOutIsSoloed, bOutIsMuted, OutReason);
-}
+		if (!Canvas || !World)
+		{
+			return false;
+		}
 
-void FAudioDebugger::QuerySoloMuteSoundCue(const FString& Name, bool& bOutIsSoloed, bool& bOutIsMuted, FString& OutReason) const
-{
-	GetDebugSoloMuteStateX(Name, DebugNames.SoloSoundCue, DebugNames.MuteSoundCue, bOutIsSoloed, bOutIsMuted, OutReason);
-}
+		*OutAudioDevice = World->GetAudioDeviceRaw();
+		if (!*OutAudioDevice)
+		{
+			return false;
+		}
 
-void FAudioDebugger::SetNameArray(FName InName, TArray<FName>& InNameArray, bool bOnOff)
-{
-	ExecuteCmdOnAudioThread([InName, &InNameArray, bOnOff]
-	{
-		if (bOnOff)
+		if (bEnablementBool || bAllowUsingDeprecatedDebugStats)
 		{
-			InNameArray.AddUnique(InName);
+			return true;
 		}
-		else
+
+		if (!bEnablementBool)
 		{
-			InNameArray.Remove(InName);
+			FAudioStats* AudioStats = AudioDeviceStats.Find((*OutAudioDevice)->DeviceID);
+			if (!AudioStats || !(AudioStats->EnabledStats & InAudioStat))
+			{
+				return false;
+			}
 		}
-	});
-}
 
-void FAudioDebugger::ToggleNameArray(FName InName, TArray<FName>& InNameArray, bool bExclusive )
+		return true;
+	}
+} // namespace Audio
+
+// Console Commands
+static FAutoConsoleCommandWithWorld GAudioDumpActiveSounds(TEXT("au.DumpActiveSounds"), TEXT("Outputs data about all the currently active sounds."), FConsoleCommandWithWorldDelegate::CreateStatic(&Audio::HandleDumpActiveSounds), ECVF_Cheat);
+static FAutoConsoleCommandWithWorld GAudioClearMutesAndSolos(TEXT("au.ClearMutesAndSolos"), TEXT("Clears any solo-ing/mute-ing sounds"), FConsoleCommandWithWorldDelegate::CreateStatic(&Audio::HandleClearMutesAndSolos), ECVF_Cheat);
+
+static FAutoConsoleCommandWithWorldAndArgs GAudioDebugSoundCues
+(
+	TEXT("au.Debug.SoundCues"),
+		TEXT("Post SoundCue information to viewport(s).\n")
+		TEXT("0: Disable, 1: Enable\n")
+	TEXT("(Optional) -AllViews: Enables/Disables for all viewports, not just those associated with the current world"),
+	FConsoleCommandWithWorldAndArgsDelegate::CreateStatic([](const TArray<FString>& Args, UWorld* InWorld)
+	{
+		Audio::DebugSoundObject(Args, InWorld, Audio::FAudioStats::SoundCues, Audio::bDebugCuesForAllViewsEnabled);
+	}), ECVF_Cheat
+);
+
+static FAutoConsoleCommandWithWorldAndArgs GAudioDebugSounds
+(
+	TEXT("au.Debug.Sounds"),
+		TEXT("Post Sound information to viewport(s).\n")
+		TEXT("0: Disable, 1: Enable\n")
+	TEXT("(Optional) -AllViews: Enables/Disables for all viewports, not just those associated with the current world"),
+	FConsoleCommandWithWorldAndArgsDelegate::CreateStatic([](const TArray<FString>& Args, UWorld* InWorld)
+	{
+		Audio::DebugSoundObject(Args, InWorld, Audio::FAudioStats::Sounds, Audio::bDebugSoundsForAllViewsEnabled);
+	}), ECVF_Cheat
+);
+
+static FAutoConsoleCommandWithWorldAndArgs GAudioDebugSoundWaves
+(
+	TEXT("au.Debug.SoundWaves"),
+	TEXT("Post SoundWave information to viewport(s).\n")
+		TEXT("0: Disable, 1: Enable\n")
+		TEXT("(Optional) -AllViews: Enables/Disables for all viewports, not just those associated with the current world"),
+	FConsoleCommandWithWorldAndArgsDelegate::CreateStatic([](const TArray<FString>& Args, UWorld* InWorld)
+	{
+		Audio::DebugSoundObject(Args, InWorld, Audio::FAudioStats::SoundWaves, Audio::bDebugWavesForAllViewsEnabled);
+	}), ECVF_Cheat
+);
+
+static FAutoConsoleCommandWithWorldAndArgs GAudioDebugSoundMixes
+(
+	TEXT("au.Debug.SoundMixes"),
+	TEXT("Post SoundMix information to viewport(s).\n")
+		TEXT("0: Disable, 1: Enable\n")
+		TEXT("(Optional) -AllViews: Enables/Disables for all viewports, not just those associated with the current world"),
+	FConsoleCommandWithWorldAndArgsDelegate::CreateStatic([](const TArray<FString>& Args, UWorld* InWorld)
+	{
+		Audio::DebugSoundObject(Args, InWorld, Audio::FAudioStats::SoundMixes, Audio::bDebugMixesForAllViewsEnabled);
+	}), ECVF_Cheat
+);
+
+static FAutoConsoleCommandWithWorldAndArgs GAudioDebugSoundReverb
+(
+	TEXT("au.Debug.SoundReverb"),
+	TEXT("Post SoundReverb information to viewport(s).\n")
+	TEXT("0: Disable, 1: Enable\n")
+	TEXT("(Optional) -AllViews: Enables/Disables for all viewports, not just those associated with the current world"),
+	FConsoleCommandWithWorldAndArgsDelegate::CreateStatic([](const TArray<FString>& Args, UWorld* InWorld)
+	{
+		Audio::DebugSoundObject(Args, InWorld, Audio::FAudioStats::SoundReverb, Audio::bDebugReverbForAllViewsEnabled);
+	}), ECVF_Cheat
+);
+
+static FAutoConsoleCommandWithWorldAndArgs GAudioDebugSoundModulation
+(
+	TEXT("au.Debug.SoundModulators"),
+	TEXT("Post SoundModulation information to viewport(s).\n")
+		TEXT("0: Disable, 1: Enable\n")
+		TEXT("(Optional) -AllViews: Enables/Disables for all viewports, not just those associated with the current world"),
+	FConsoleCommandWithWorldAndArgsDelegate::CreateStatic([](const TArray<FString>& Args, UWorld* InWorld)
+	{
+		Audio::DebugSoundObject(Args, InWorld, Audio::FAudioStats::SoundModulation, Audio::bDebugModulationForAllViewsEnabled);
+	}), ECVF_Cheat
+);
+
+
+namespace Audio
 {
-	ExecuteCmdOnAudioThread([InName, &InNameArray, bExclusive]
-	{ 
-		// On already?
-		int32 IndexOf = InNameArray.IndexOfByKey(InName);
-		if (IndexOf != INDEX_NONE)
-		{			
-			if (bExclusive)
-			{
-				// Turn off everything if we are exclusive.
-				InNameArray.Empty();
+	/** Audio Debugger Implementation */
+	FAudioDebugger::FAudioDebugger()
+		: bVisualize3dDebug(0)
+	{
+		WorldRegisteredWithDeviceHandle = FAudioDeviceManagerDelegates::OnWorldRegisteredToAudioDevice.AddLambda([this](const UWorld* InWorld, FDeviceId InDeviceId)
+		{
+			if (bAllowUsingDeprecatedDebugStats)
+			{
+				return;
+			}
+
+			uint32 StatsToSet = 0;
+			uint32 StatsToClear = 0;
+			if (bDebugSoundsForAllViewsEnabled)
+			{
+				StatsToSet |= FAudioStats::Sounds;
 			}
 			else
 			{
-				InNameArray.RemoveAtSwap(IndexOf);
+				StatsToClear |= FAudioStats::Sounds;
 			}
-		}
-		else // Add it.
-		{
-			// If we are exclusive, turn off everything else first.
-			if (bExclusive)
+
+			if (bDebugCuesForAllViewsEnabled)
+			{
+				StatsToSet |= FAudioStats::SoundCues;
+			}
+			else
 			{
-				InNameArray.Empty();
+				StatsToClear |= FAudioStats::SoundCues;
 			}
 
-			// Allow for NAME_None to act as a clear
-			if (InName != NAME_None)
+			if (bDebugWavesForAllViewsEnabled)
 			{
-				InNameArray.Add(InName);
-			}						
-		}
-	});
-}
+				StatsToSet |= FAudioStats::SoundWaves;
+			}
+			else
+			{
+				StatsToClear |= FAudioStats::SoundWaves;
+			}
 
-void FAudioDebugger::SetAudioMixerDebugSound(const TCHAR* SoundName)
-{
-	ExecuteCmdOnAudioThread([Name = FString(SoundName), this]{ DebugNames.DebugAudioMixerSoundName = Name; });
-}
+			if (bDebugModulationForAllViewsEnabled)
+			{
+				StatsToSet |= FAudioStats::SoundModulation;
+			}
+			else
+			{
+				StatsToClear |= FAudioStats::SoundModulation;
+			}
 
-void FAudioDebugger::SetAudioDebugSound(const TCHAR* SoundName)
-{
-	ExecuteCmdOnAudioThread([Name = FString(SoundName), this]{ 
-		DebugNames.DebugSoundName = Name;
-		DebugNames.bDebugSoundName = DebugNames.DebugSoundName != TEXT("");
-	});	
-}
+			if (bDebugMixesForAllViewsEnabled)
+			{
+				StatsToSet |= FAudioStats::SoundMixes;
+			}
+			else
+			{
+				StatsToClear |= FAudioStats::SoundMixes;
+			}
 
-const FString& FAudioDebugger::GetAudioMixerDebugSoundName() const
-{
-	check(IsInAudioThread());
-	return DebugNames.DebugAudioMixerSoundName;
-}
+			if (bDebugReverbForAllViewsEnabled)
+			{
+				StatsToSet |= FAudioStats::SoundReverb;
+			}
+			else
+			{
+				StatsToClear |= FAudioStats::SoundReverb;
+			}
 
-bool FAudioDebugger::GetAudioDebugSound(FString& OutDebugSound)
-{
-	check(IsInAudioThread());
-	if (DebugNames.bDebugSoundName)
-	{
-		OutDebugSound = DebugNames.DebugSoundName;
-		return true;
+			SetStats(InDeviceId, StatsToSet);
+			ClearStats(InDeviceId, StatsToClear);
+		});
 	}
-	return false;
-}
 
-void FAudioDebugger::ExecuteCmdOnAudioThread(TFunction<void()> Cmd)
-{
-	// If not on audio thread, queue it.
-	if (!IsInAudioThread())
+	bool FAudioDebugger::IsVisualizeDebug3dEnabled() const
 	{
-		FAudioThread::RunCommandOnAudioThread(Cmd);
-		return;
+		return bVisualize3dDebug;
 	}
 
-	// Otherwise, do it inline.
-	Cmd();
-}
-
-void FAudioDebugger::GetDebugSoloMuteStateX(
-	const FString& Name, 
-	const TArray<FName>& Solos,
-	const TArray<FName>& Mutes,	
-	bool& bOutIsSoloed, 
-	bool& bOutIsMuted, 
-	FString& OutReason) const
-{
-	check(IsInAudioThread());
-	
-	// Allow for partial matches of the name.
-	auto MatchesPartOfName = [&Name](FName i) -> bool { return Name.Contains(i.GetPlainNameString()); };
+	void FAudioDebugger::ToggleVisualizeDebug3dEnabled()
+	{
+		bVisualize3dDebug = !bVisualize3dDebug;
+	}
 
-	// Solo active?
-	if (Solos.Num() > 0)
+	#if WITH_EDITOR
+	void FAudioDebugger::OnBeginPIE()
 	{
-		if (Solos.ContainsByPredicate(MatchesPartOfName))
-		{
-			bOutIsSoloed = true;
-			OutReason = FString::Printf(TEXT("Sound is soloed explicitly."));
-		}
-		else
+		FAudioDeviceManager* DeviceManager = GEngine->GetAudioDeviceManager();
+		if (DeviceManager)
 		{
-			// Something else is soloed (record the first item in the solo list for debug reason).
-			bOutIsMuted = true;
-			OutReason = FString::Printf(TEXT("Sound is muted due to [%s] being soloed"), *Solos[0].ToString() );
+			DeviceManager->GetDebugger().ClearMutesAndSolos();
 		}
 	}
-	// Are we explicitly muted?
-	else if (Mutes.ContainsByPredicate(MatchesPartOfName))
+
+	void FAudioDebugger::OnEndPIE()
 	{
-		bOutIsMuted = true;
-		OutReason = FString::Printf(TEXT("Sound is explicitly muted"));
+		FAudioDeviceManager* DeviceManager = GEngine->GetAudioDeviceManager();
+		if (DeviceManager)
+		{
+			DeviceManager->GetDebugger().ClearMutesAndSolos();
+		}
 	}
-}
+	#endif // WITH_EDITOR
 
-void FAudioDebugger::DrawDebugInfo(const FSoundSource& SoundSource)
-{
-#if ENABLE_DRAW_DEBUG
-	const FWaveInstance* WaveInstance = SoundSource.GetWaveInstance();
-	if (!WaveInstance)
+	void FAudioDebugger::QuerySoloMuteSoundClass(const FString& Name, bool& bOutIsSoloed, bool& bOutIsMuted, FString& OutReason) const
 	{
-		return;
+		GetDebugSoloMuteStateX(Name, DebugNames.SoloSoundClass, DebugNames.MuteSoundClass, bOutIsSoloed, bOutIsMuted, OutReason);
 	}
 
-	const FActiveSound* ActiveSound = WaveInstance->ActiveSound;
-	if (!ActiveSound)
+	void FAudioDebugger::QuerySoloMuteSoundWave(const FString& Name, bool& bOutIsSoloed, bool& bOutIsMuted, FString& OutReason) const
 	{
-		return;
+		GetDebugSoloMuteStateX(Name, DebugNames.SoloSoundWave, DebugNames.MuteSoundWave, bOutIsSoloed, bOutIsMuted, OutReason);
 	}
 
-	if (!SpatialSourceVisualizeEnabledCVar)
+	void FAudioDebugger::QuerySoloMuteSoundCue(const FString& Name, bool& bOutIsSoloed, bool& bOutIsMuted, FString& OutReason) const
 	{
-		return;
+		GetDebugSoloMuteStateX(Name, DebugNames.SoloSoundCue, DebugNames.MuteSoundCue, bOutIsSoloed, bOutIsMuted, OutReason);
 	}
 
-	FAudioDeviceManager* DeviceManager = GEngine->GetAudioDeviceManager();
-	if (DeviceManager && DeviceManager->IsVisualizeDebug3dEnabled())
+	void FAudioDebugger::SetNameArray(FName InName, TArray<FName>& InNameArray, bool bOnOff)
 	{
-		DECLARE_CYCLE_STAT(TEXT("FAudioThreadTask.DrawSourceDebugInfo"), STAT_AudioDrawSourceDebugInfo, STATGROUP_TaskGraphTasks);
-
-		const FSoundBuffer* Buffer = SoundSource.GetBuffer();
-		const bool bSpatialized = Buffer && Buffer->NumChannels == 2 && WaveInstance->GetUseSpatialization();
-		if (bSpatialized)
+		ExecuteCmdOnAudioThread([InName, &InNameArray, bOnOff]
 		{
-			const FRotator Rotator = ActiveSound->Transform.GetRotation().Rotator();
+			if (bOnOff)
+			{
+				InNameArray.AddUnique(InName);
+			}
+			else
+			{
+				InNameArray.Remove(InName);
+			}
+		});
+	}
 
-			TWeakObjectPtr<UWorld> WorldPtr = WaveInstance->ActiveSound->GetWeakWorld();
-			FVector LeftChannelSourceLoc;
-			FVector RightChannelSourceLoc;
-			SoundSource.GetChannelLocations(LeftChannelSourceLoc, RightChannelSourceLoc);
-			FAudioThread::RunCommandOnGameThread([LeftChannelSourceLoc, RightChannelSourceLoc, Rotator, WorldPtr]()
+	void FAudioDebugger::ToggleNameArray(FName InName, TArray<FName>& InNameArray, bool bExclusive )
+	{
+		ExecuteCmdOnAudioThread([InName, &InNameArray, bExclusive]
+		{ 
+			// On already?
+			int32 IndexOf = InNameArray.IndexOfByKey(InName);
+			if (IndexOf != INDEX_NONE)
+			{			
+				if (bExclusive)
+				{
+					// Turn off everything if we are exclusive.
+					InNameArray.Empty();
+				}
+				else
+				{
+					InNameArray.RemoveAtSwap(IndexOf);
+				}
+			}
+			else // Add it.
 			{
-				if (WorldPtr.IsValid())
+				// If we are exclusive, turn off everything else first.
+				if (bExclusive)
 				{
-					UWorld* World = WorldPtr.Get();
-					DrawDebugCrosshairs(World, LeftChannelSourceLoc, Rotator, 20.0f, FColor::Red, false, -1.0f, SDPG_Foreground);
-					DrawDebugCrosshairs(World, RightChannelSourceLoc, Rotator, 20.0f, FColor::Green, false, -1.0f, SDPG_Foreground);
+					InNameArray.Empty();
 				}
-			}, GET_STATID(STAT_AudioDrawSourceDebugInfo));
-		}
+
+				// Allow for NAME_None to act as a clear
+				if (InName != NAME_None)
+				{
+					InNameArray.Add(InName);
+				}						
+			}
+		});
 	}
-#endif // ENABLE_DRAW_DEBUG
-}
 
-void FAudioDebugger::DrawDebugInfo(const FActiveSound& ActiveSound, const TArray<FWaveInstance*>& ThisSoundsWaveInstances, const float DeltaTime)
-{
-#if ENABLE_DRAW_DEBUG
-	if (!ActiveSoundVisualizeModeCVar)
+	void FAudioDebugger::SetAudioMixerDebugSound(const TCHAR* SoundName)
 	{
-		return;
+		ExecuteCmdOnAudioThread([Name = FString(SoundName), this]{ DebugNames.DebugAudioMixerSoundName = Name; });
 	}
 
-	// Only draw spatialized sounds
-	const USoundBase* Sound = ActiveSound.GetSound();
-	if (!Sound || !ActiveSound.bAllowSpatialization)
+	void FAudioDebugger::SetAudioDebugSound(const TCHAR* SoundName)
 	{
-		return;
+		ExecuteCmdOnAudioThread([Name = FString(SoundName), this]{ 
+			DebugNames.DebugSoundName = Name;
+			DebugNames.bDebugSoundName = DebugNames.DebugSoundName != TEXT("");
+		});	
 	}
 
-	const float PlaybackTime = ActiveSound.PlaybackTime;
-	const float PlaybackTimeNonVirtualized = ActiveSound.PlaybackTimeNonVirtualized;
-
-	// Sounds requiring culling can start and immediately stop repeatedly when subscribed
-	// concurrency is flooded, so don't show the initial frame.
-	if (FMath::IsNearlyZero(PlaybackTimeNonVirtualized))
+	const FString& FAudioDebugger::GetAudioMixerDebugSoundName() const
 	{
-		return;
+		check(IsInAudioThread());
+		return DebugNames.DebugAudioMixerSoundName;
 	}
 
-	if (ActiveSoundVisualizeTypeCVar > 0)
+	bool FAudioDebugger::GetAudioDebugSound(FString& OutDebugSound)
 	{
-		if (ActiveSoundVisualizeTypeCVar == 1 && ActiveSound.GetAudioComponentID() == 0)
+		check(IsInAudioThread());
+		if (DebugNames.bDebugSoundName)
 		{
-			return;
+			OutDebugSound = DebugNames.DebugSoundName;
+			return true;
 		}
+		return false;
+	}
 
-		if (ActiveSoundVisualizeTypeCVar == 2 && ActiveSound.GetAudioComponentID() > 0)
+	void FAudioDebugger::ExecuteCmdOnAudioThread(TFunction<void()> Cmd)
+	{
+		// If not on audio thread, queue it.
+		if (!IsInAudioThread())
 		{
+			FAudioThread::RunCommandOnAudioThread(Cmd);
 			return;
 		}
+
+		// Otherwise, do it inline.
+		Cmd();
 	}
 
-	FAudioDeviceManager* DeviceManager = GEngine->GetAudioDeviceManager();
-	if (DeviceManager && DeviceManager->IsVisualizeDebug3dEnabled())
+	void FAudioDebugger::GetDebugSoloMuteStateX(
+		const FString& Name, 
+		const TArray<FName>& Solos,
+		const TArray<FName>& Mutes,	
+		bool& bOutIsSoloed, 
+		bool& bOutIsMuted, 
+		FString& OutReason) const
 	{
-		DECLARE_CYCLE_STAT(TEXT("FAudioThreadTask.DrawActiveSoundDebugInfo"), STAT_AudioDrawActiveSoundDebugInfo, STATGROUP_TaskGraphTasks);
+		check(IsInAudioThread());
+	
+		// Allow for partial matches of the name.
+		auto MatchesPartOfName = [&Name](FName i) -> bool { return Name.Contains(i.GetPlainNameString()); };
 
-		const FString Name = Sound->GetName();
-		const FTransform CurTransform = ActiveSound.Transform;
-		FColor TextColor = FColor::White;
-		const float CurMaxDistance = ActiveSound.MaxDistance;
-		float DisplayValue = 0.0f;
-		if (ActiveSoundVisualizeModeCVar == 1 || ActiveSoundVisualizeModeCVar == 2)
+		// Solo active?
+		if (Solos.Num() > 0)
 		{
-			for (FWaveInstance* WaveInstance : ThisSoundsWaveInstances)
+			if (Solos.ContainsByPredicate(MatchesPartOfName))
 			{
-				DisplayValue = FMath::Max(DisplayValue, WaveInstance->GetVolumeWithDistanceAttenuation() * WaveInstance->GetDynamicVolume());
+				bOutIsSoloed = true;
+				OutReason = FString::Printf(TEXT("Sound is soloed explicitly."));
 			}
-		}
-		else if (ActiveSoundVisualizeModeCVar == 3)
-		{
-			if (ActiveSound.AudioDevice)
+			else
 			{
-				DisplayValue = ActiveSound.AudioDevice->GetDistanceToNearestListener(ActiveSound.Transform.GetLocation()) / CurMaxDistance;
+				// Something else is soloed (record the first item in the solo list for debug reason).
+				bOutIsMuted = true;
+				OutReason = FString::Printf(TEXT("Sound is muted due to [%s] being soloed"), *Solos[0].ToString() );
 			}
 		}
-		else if (ActiveSoundVisualizeModeCVar == 4)
+		// Are we explicitly muted?
+		else if (Mutes.ContainsByPredicate(MatchesPartOfName))
 		{
-			TextColor = ActiveSound.DebugColor;
+			bOutIsMuted = true;
+			OutReason = FString::Printf(TEXT("Sound is explicitly muted"));
 		}
+	}
 
-		TWeakObjectPtr<UWorld> WorldPtr = ActiveSound.GetWeakWorld();
-		FAudioThread::RunCommandOnGameThread([Name, TextColor, CurTransform, DisplayValue, WorldPtr, CurMaxDistance, PlaybackTime, PlaybackTimeNonVirtualized, DeltaTime]()
+	void FAudioDebugger::DrawDebugInfo(const FSoundSource& SoundSource)
+	{
+	#if ENABLE_DRAW_DEBUG
+		const FWaveInstance* WaveInstance = SoundSource.GetWaveInstance();
+		if (!WaveInstance)
 		{
-			if (WorldPtr.IsValid())
-			{
-				static const float ColorRedHue = 0.0f;
-				static const float ColorGreenHue = 85.0f;
+			return;
+		}
+
+		const FActiveSound* ActiveSound = WaveInstance->ActiveSound;
+		if (!ActiveSound)
+		{
+			return;
+		}
+
+		if (!SpatialSourceVisualizeEnabledCVar)
+		{
+			return;
+		}
+
+		FAudioDeviceManager* DeviceManager = GEngine->GetAudioDeviceManager();
+		if (DeviceManager && DeviceManager->IsVisualizeDebug3dEnabled())
+		{
+			DECLARE_CYCLE_STAT(TEXT("FAudioThreadTask.DrawSourceDebugInfo"), STAT_AudioDrawSourceDebugInfo, STATGROUP_TaskGraphTasks);
 
-				const FVector Location = CurTransform.GetLocation();
-				UWorld* DebugWorld = WorldPtr.Get();
-				DrawDebugSphere(DebugWorld, Location, 10.0f, 8, FColor::White, false, -1.0f, SDPG_Foreground);
-				FColor Color = TextColor;
+			const FSoundBuffer* Buffer = SoundSource.GetBuffer();
+			const bool bSpatialized = Buffer && Buffer->NumChannels == 2 && WaveInstance->GetUseSpatialization();
+			if (bSpatialized)
+			{
+				const FRotator Rotator = ActiveSound->Transform.GetRotation().Rotator();
 
-				FString Descriptor;
-				if (ActiveSoundVisualizeModeCVar == 1 || ActiveSoundVisualizeModeCVar == 2)
+				TWeakObjectPtr<UWorld> WorldPtr = WaveInstance->ActiveSound->GetWeakWorld();
+				FVector LeftChannelSourceLoc;
+				FVector RightChannelSourceLoc;
+				SoundSource.GetChannelLocations(LeftChannelSourceLoc, RightChannelSourceLoc);
+				FAudioThread::RunCommandOnGameThread([LeftChannelSourceLoc, RightChannelSourceLoc, Rotator, WorldPtr]()
 				{
-					const float DisplayDbVolume = Audio::ConvertToDecibels(DisplayValue);
-					if (ActiveSoundVisualizeModeCVar == 1)
+					if (WorldPtr.IsValid())
 					{
-						Descriptor = FString::Printf(TEXT(" (Vol: %.3f [Active: %.2fs, Playing: %.2fs])"), DisplayValue, PlaybackTime, PlaybackTimeNonVirtualized);
+						UWorld* World = WorldPtr.Get();
+						DrawDebugCrosshairs(World, LeftChannelSourceLoc, Rotator, 20.0f, FColor::Red, false, -1.0f, SDPG_Foreground);
+						DrawDebugCrosshairs(World, RightChannelSourceLoc, Rotator, 20.0f, FColor::Green, false, -1.0f, SDPG_Foreground);
 					}
-					else
-					{
-						Descriptor = FString::Printf(TEXT(" (Vol: %.3f dB [Active: %.2fs, Playing: %.2fs])"), DisplayDbVolume, PlaybackTime, PlaybackTimeNonVirtualized);
-					}
-					static const float DbColorMinVol = -30.0f;
-					const float DbVolume = FMath::Clamp(DisplayDbVolume, DbColorMinVol, 0.0f);
-					const float Hue = FMath::Lerp(ColorRedHue, ColorGreenHue, (-1.0f * DbVolume / DbColorMinVol) + 1.0f);
-					Color = FLinearColor::MakeFromHSV8(static_cast<uint8>(Hue), 255u, 255u).ToFColor(true);
-				}
-				else if (ActiveSoundVisualizeModeCVar == 3)
-				{
-					Descriptor = FString::Printf(TEXT(" (Dist: %.3f, Max: %.3f)"), DisplayValue * CurMaxDistance, CurMaxDistance);
-					const float Hue = FMath::Lerp(ColorGreenHue, ColorRedHue, DisplayValue);
-					Color = FLinearColor::MakeFromHSV8(static_cast<uint8>(FMath::Clamp(Hue, 0.0f, 255.f)), 255u, 255u).ToFColor(true);
-				}
-
-				const FString Description = FString::Printf(TEXT("%s%s"), *Name, *Descriptor);
-				DrawDebugString(DebugWorld, Location + FVector(0, 0, 32), *Description, nullptr, Color, DeltaTime, false);
+				}, GET_STATID(STAT_AudioDrawSourceDebugInfo));
 			}
-		}, GET_STATID(STAT_AudioDrawActiveSoundDebugInfo));
+		}
+	#endif // ENABLE_DRAW_DEBUG
 	}
-#endif // ENABLE_DRAW_DEBUG
-}
 
-void FAudioDebugger::DrawDebugInfo(UWorld& World, const TArray<FListener>& Listeners)
-{
-#if ENABLE_DRAW_DEBUG
-	if (!ActiveSoundVisualizeListenersCVar)
+	void FAudioDebugger::DrawDebugInfo(const FActiveSound& ActiveSound, const TArray<FWaveInstance*>& ThisSoundsWaveInstances, const float DeltaTime)
 	{
-		return;
-	}
+	#if ENABLE_DRAW_DEBUG
+		if (!ActiveSoundVisualizeModeCVar)
+		{
+			return;
+		}
+
+		// Only draw spatialized sounds
+		const USoundBase* Sound = ActiveSound.GetSound();
+		if (!Sound || !ActiveSound.bAllowSpatialization)
+		{
+			return;
+		}
 
-	check(IsInAudioThread());
+		const float PlaybackTime = ActiveSound.PlaybackTime;
+		const float PlaybackTimeNonVirtualized = ActiveSound.PlaybackTimeNonVirtualized;
 
-	FAudioDeviceManager* DeviceManager = GEngine->GetAudioDeviceManager();
-	if (DeviceManager && DeviceManager->IsVisualizeDebug3dEnabled())
-	{
-		DECLARE_CYCLE_STAT(TEXT("FAudioThreadTask.DrawListenerDebugInfo"), STAT_AudioDrawListenerDebugInfo, STATGROUP_TaskGraphTasks);
-		const TWeakObjectPtr<UWorld> WorldPtr = &World;
-		for (const FListener& Listener : Listeners)
+		// Sounds requiring culling can start and immediately stop repeatedly when subscribed
+		// concurrency is flooded, so don't show the initial frame.
+		if (FMath::IsNearlyZero(PlaybackTimeNonVirtualized))
+		{
+			return;
+		}
+
+		if (ActiveSoundVisualizeTypeCVar > 0)
+		{
+			if (ActiveSoundVisualizeTypeCVar == 1 && ActiveSound.GetAudioComponentID() == 0)
+			{
+				return;
+			}
+
+			if (ActiveSoundVisualizeTypeCVar == 2 && ActiveSound.GetAudioComponentID() > 0)
+			{
+				return;
+			}
+		}
+
+		FAudioDeviceManager* DeviceManager = GEngine->GetAudioDeviceManager();
+		if (DeviceManager && DeviceManager->IsVisualizeDebug3dEnabled())
 		{
-			const FVector ListenerPosition = Listener.GetPosition(true);
-			const FVector ListenerFront = Listener.GetFront();
-			const FVector ListenerUp = Listener.GetUp();
-			const FVector ListenerRight = Listener.GetRight();
+			DECLARE_CYCLE_STAT(TEXT("FAudioThreadTask.DrawActiveSoundDebugInfo"), STAT_AudioDrawActiveSoundDebugInfo, STATGROUP_TaskGraphTasks);
+
+			const FString Name = Sound->GetName();
+			const FTransform CurTransform = ActiveSound.Transform;
+			FColor TextColor = FColor::White;
+			const float CurMaxDistance = ActiveSound.MaxDistance;
+			float DisplayValue = 0.0f;
+			if (ActiveSoundVisualizeModeCVar == 1 || ActiveSoundVisualizeModeCVar == 2)
+			{
+				for (FWaveInstance* WaveInstance : ThisSoundsWaveInstances)
+				{
+					DisplayValue = FMath::Max(DisplayValue, WaveInstance->GetVolumeWithDistanceAttenuation() * WaveInstance->GetDynamicVolume());
+				}
+			}
+			else if (ActiveSoundVisualizeModeCVar == 3)
+			{
+				if (ActiveSound.AudioDevice)
+				{
+					DisplayValue = ActiveSound.AudioDevice->GetDistanceToNearestListener(ActiveSound.Transform.GetLocation()) / CurMaxDistance;
+				}
+			}
+			else if (ActiveSoundVisualizeModeCVar == 4)
+			{
+				TextColor = ActiveSound.DebugColor;
+			}
 
-			FAudioThread::RunCommandOnGameThread([WorldPtr, ListenerPosition, ListenerFront, ListenerUp, ListenerRight]()
+			TWeakObjectPtr<UWorld> WorldPtr = ActiveSound.GetWeakWorld();
+			FAudioThread::RunCommandOnGameThread([Name, TextColor, CurTransform, DisplayValue, WorldPtr, CurMaxDistance, PlaybackTime, PlaybackTimeNonVirtualized, DeltaTime]()
 			{
 				if (WorldPtr.IsValid())
 				{
-					static float ArrowLength	 = 30.0f;
-					static float ArrowHeadSize	 = 8.0f;
-					static float Lifetime		 = -1.0f;
-					static float Thickness		 = 0.9f;
-					static uint8 DepthPriority	 = SDPG_Foreground;
-					static bool bPersistentLines = false;
-
-					DrawDebugDirectionalArrow(WorldPtr.Get(), ListenerPosition, ListenerPosition + (ArrowLength * ListenerFront), ArrowHeadSize, FColor::Red, bPersistentLines, Lifetime, DepthPriority, Thickness);
-					DrawDebugDirectionalArrow(WorldPtr.Get(), ListenerPosition, ListenerPosition + (ArrowLength * ListenerUp), ArrowHeadSize, FColor::Blue, bPersistentLines, Lifetime, DepthPriority, Thickness);
-					DrawDebugDirectionalArrow(WorldPtr.Get(), ListenerPosition, ListenerPosition + (ArrowLength * ListenerRight), ArrowHeadSize, FColor::Green, bPersistentLines, Lifetime, DepthPriority, Thickness);
-					DrawDebugSphere(WorldPtr.Get(), ListenerPosition, 5.0f, 8, FColor::Magenta, bPersistentLines, Lifetime, DepthPriority);
+					static const float ColorRedHue = 0.0f;
+					static const float ColorGreenHue = 85.0f;
+
+					const FVector Location = CurTransform.GetLocation();
+					UWorld* DebugWorld = WorldPtr.Get();
+					DrawDebugSphere(DebugWorld, Location, 10.0f, 8, FColor::White, false, -1.0f, SDPG_Foreground);
+					FColor Color = TextColor;
+
+					FString Descriptor;
+					if (ActiveSoundVisualizeModeCVar == 1 || ActiveSoundVisualizeModeCVar == 2)
+					{
+						const float DisplayDbVolume = Audio::ConvertToDecibels(DisplayValue);
+						if (ActiveSoundVisualizeModeCVar == 1)
+						{
+							Descriptor = FString::Printf(TEXT(" (Vol: %.3f [Active: %.2fs, Playing: %.2fs])"), DisplayValue, PlaybackTime, PlaybackTimeNonVirtualized);
+						}
+						else
+						{
+							Descriptor = FString::Printf(TEXT(" (Vol: %.3f dB [Active: %.2fs, Playing: %.2fs])"), DisplayDbVolume, PlaybackTime, PlaybackTimeNonVirtualized);
+						}
+						static const float DbColorMinVol = -30.0f;
+						const float DbVolume = FMath::Clamp(DisplayDbVolume, DbColorMinVol, 0.0f);
+						const float Hue = FMath::Lerp(ColorRedHue, ColorGreenHue, (-1.0f * DbVolume / DbColorMinVol) + 1.0f);
+						Color = FLinearColor::MakeFromHSV8(static_cast<uint8>(Hue), 255u, 255u).ToFColor(true);
+					}
+					else if (ActiveSoundVisualizeModeCVar == 3)
+					{
+						Descriptor = FString::Printf(TEXT(" (Dist: %.3f, Max: %.3f)"), DisplayValue * CurMaxDistance, CurMaxDistance);
+						const float Hue = FMath::Lerp(ColorGreenHue, ColorRedHue, DisplayValue);
+						Color = FLinearColor::MakeFromHSV8(static_cast<uint8>(FMath::Clamp(Hue, 0.0f, 255.f)), 255u, 255u).ToFColor(true);
+					}
+
+					const FString Description = FString::Printf(TEXT("%s%s"), *Name, *Descriptor);
+					DrawDebugString(DebugWorld, Location + FVector(0, 0, 32), *Description, nullptr, Color, DeltaTime, false);
 				}
-			}, GET_STATID(STAT_AudioDrawListenerDebugInfo));
+			}, GET_STATID(STAT_AudioDrawActiveSoundDebugInfo));
 		}
+	#endif // ENABLE_DRAW_DEBUG
 	}
-#endif // ENABLE_DRAW_DEBUG
-}
 
-void FAudioDebugger::DrawDebugInfo(const FAudioVirtualLoop& VirtualLoop)
-{
-#if ENABLE_DRAW_DEBUG
-	if (!GEngine)
+	void FAudioDebugger::DrawDebugInfo(UWorld& World, const TArray<FListener>& Listeners)
 	{
-		return;
-	}
+	#if ENABLE_DRAW_DEBUG
+		if (!ActiveSoundVisualizeListenersCVar)
+		{
+			return;
+		}
 
-	if (!VirtualLoopsVisualizeEnabledCVar)
-	{
-		return;
-	}
+		check(IsInAudioThread());
 
-	FAudioDeviceManager* DeviceManager = GEngine->GetAudioDeviceManager();
-	if (DeviceManager && DeviceManager->IsVisualizeDebug3dEnabled())
-	{
-		DECLARE_CYCLE_STAT(TEXT("FAudioThreadTask.DrawVirtualLoopDebugInfo"), STAT_AudioDrawVirtualLoopDebugInfo, STATGROUP_TaskGraphTasks);
-
-		const FActiveSound& ActiveSound = VirtualLoop.GetActiveSound();
-		USoundBase* Sound = ActiveSound.GetSound();
-		check(Sound);
-
-		const FTransform Transform = ActiveSound.Transform;
-		const TWeakObjectPtr<UWorld> World = ActiveSound.GetWeakWorld();
-		const FString Name = Sound->GetName();
-		const float DrawInterval = VirtualLoop.GetUpdateInterval();
-		const float TimeVirtualized = VirtualLoop.GetTimeVirtualized();
-		FAudioThread::RunCommandOnGameThread([World, Transform, Name, DrawInterval, TimeVirtualized]()
-		{
-			if (World.IsValid())
-			{
-				const FString Description = FString::Printf(TEXT("%s [Virt: %.2fs]"), *Name, TimeVirtualized);
-				FVector Location = Transform.GetLocation();
-				FRotator Rotation = Transform.GetRotation().Rotator();
-				DrawDebugCrosshairs(World.Get(), Location, Rotation, 20.0f, FColor::Blue, false, DrawInterval, SDPG_Foreground);
-				DrawDebugString(World.Get(), Location + FVector(0, 0, 32), *Description, nullptr, FColor::Blue, DrawInterval, false);
+		FAudioDeviceManager* DeviceManager = GEngine->GetAudioDeviceManager();
+		if (DeviceManager && DeviceManager->IsVisualizeDebug3dEnabled())
+		{
+			DECLARE_CYCLE_STAT(TEXT("FAudioThreadTask.DrawListenerDebugInfo"), STAT_AudioDrawListenerDebugInfo, STATGROUP_TaskGraphTasks);
+			const TWeakObjectPtr<UWorld> WorldPtr = &World;
+			for (const FListener& Listener : Listeners)
+			{
+				const FVector ListenerPosition = Listener.GetPosition(true);
+				const FVector ListenerFront = Listener.GetFront();
+				const FVector ListenerUp = Listener.GetUp();
+				const FVector ListenerRight = Listener.GetRight();
+
+				FAudioThread::RunCommandOnGameThread([WorldPtr, ListenerPosition, ListenerFront, ListenerUp, ListenerRight]()
+				{
+					if (WorldPtr.IsValid())
+					{
+						static float ArrowLength	 = 30.0f;
+						static float ArrowHeadSize	 = 8.0f;
+						static float Lifetime		 = -1.0f;
+						static float Thickness		 = 0.9f;
+						static uint8 DepthPriority	 = SDPG_Foreground;
+						static bool bPersistentLines = false;
+
+						DrawDebugDirectionalArrow(WorldPtr.Get(), ListenerPosition, ListenerPosition + (ArrowLength * ListenerFront), ArrowHeadSize, FColor::Red, bPersistentLines, Lifetime, DepthPriority, Thickness);
+						DrawDebugDirectionalArrow(WorldPtr.Get(), ListenerPosition, ListenerPosition + (ArrowLength * ListenerUp), ArrowHeadSize, FColor::Blue, bPersistentLines, Lifetime, DepthPriority, Thickness);
+						DrawDebugDirectionalArrow(WorldPtr.Get(), ListenerPosition, ListenerPosition + (ArrowLength * ListenerRight), ArrowHeadSize, FColor::Green, bPersistentLines, Lifetime, DepthPriority, Thickness);
+						DrawDebugSphere(WorldPtr.Get(), ListenerPosition, 5.0f, 8, FColor::Magenta, bPersistentLines, Lifetime, DepthPriority);
+					}
+				}, GET_STATID(STAT_AudioDrawListenerDebugInfo));
 			}
-		}, GET_STATID(STAT_AudioDrawVirtualLoopDebugInfo));
+		}
+	#endif // ENABLE_DRAW_DEBUG
 	}
-#endif // ENABLE_DRAW_DEBUG
-}
 
-void FAudioDebugger::DumpActiveSounds() const
-{
-	if (!GEngine)
+	void FAudioDebugger::DrawDebugInfo(const FAudioVirtualLoop& VirtualLoop)
 	{
-		return;
-	}
+	#if ENABLE_DRAW_DEBUG
+		if (!GEngine)
+		{
+			return;
+		}
 
-	if (!IsInAudioThread())
-	{
-		DECLARE_CYCLE_STAT(TEXT("FAudioThreadTask.DumpActiveSounds"), STAT_AudioDumpActiveSounds, STATGROUP_TaskGraphTasks);
-		FAudioThread::RunCommandOnAudioThread([this]()
+		if (!VirtualLoopsVisualizeEnabledCVar)
+		{
+			return;
+		}
+
+		FAudioDeviceManager* DeviceManager = GEngine->GetAudioDeviceManager();
+		if (DeviceManager && DeviceManager->IsVisualizeDebug3dEnabled())
 		{
-			DumpActiveSounds();
-		}, GET_STATID(STAT_AudioDumpActiveSounds));
-		return;
+			DECLARE_CYCLE_STAT(TEXT("FAudioThreadTask.DrawVirtualLoopDebugInfo"), STAT_AudioDrawVirtualLoopDebugInfo, STATGROUP_TaskGraphTasks);
+
+			const FActiveSound& ActiveSound = VirtualLoop.GetActiveSound();
+			USoundBase* Sound = ActiveSound.GetSound();
+			check(Sound);
+
+			const FTransform Transform = ActiveSound.Transform;
+			const TWeakObjectPtr<UWorld> World = ActiveSound.GetWeakWorld();
+			const FString Name = Sound->GetName();
+			const float DrawInterval = VirtualLoop.GetUpdateInterval();
+			const float TimeVirtualized = VirtualLoop.GetTimeVirtualized();
+			FAudioThread::RunCommandOnGameThread([World, Transform, Name, DrawInterval, TimeVirtualized]()
+			{
+				if (World.IsValid())
+				{
+					const FString Description = FString::Printf(TEXT("%s [Virt: %.2fs]"), *Name, TimeVirtualized);
+					FVector Location = Transform.GetLocation();
+					FRotator Rotation = Transform.GetRotation().Rotator();
+					DrawDebugCrosshairs(World.Get(), Location, Rotation, 20.0f, FColor::Blue, false, DrawInterval, SDPG_Foreground);
+					DrawDebugString(World.Get(), Location + FVector(0, 0, 32), *Description, nullptr, FColor::Blue, DrawInterval, false);
+				}
+			}, GET_STATID(STAT_AudioDrawVirtualLoopDebugInfo));
+		}
+	#endif // ENABLE_DRAW_DEBUG
 	}
 
-	FAudioDeviceHandle AudioDevice = GEngine->GetAudioDeviceManager()->GetActiveAudioDevice();
-	if (!AudioDevice)
+	void FAudioDebugger::DrawDebugStats(UWorld& World, FViewport& Viewport)
 	{
-		return;
-	}
+		FCanvas* Canvas = Viewport.GetDebugCanvas();
+		if (!Canvas || bAllowUsingDeprecatedDebugStats)
+		{
+			return;
+		}
 
-	const TArray<FActiveSound*>& ActiveSounds = AudioDevice->GetActiveSounds();
-	UE_LOG(LogAudio, Display, TEXT("Active Sound Count: %d"), ActiveSounds.Num());
-	UE_LOG(LogAudio, Display, TEXT("------------------------"), ActiveSounds.Num());
+		int32 X = SoundDebugDisplayCornerXCVar;
+		int32 Y = SoundDebugDisplayCornerYCVar;
 
-	for (const FActiveSound* ActiveSound : ActiveSounds)
+		Y = RenderStatMixes(&World, &Viewport, Canvas, X, Y);
+		Y = RenderStatModulators(&World, &Viewport, Canvas, X, Y, nullptr, nullptr);
+		Y = RenderStatReverb(&World, &Viewport, Canvas, X, Y);
+		Y = RenderStatSounds(&World, &Viewport, Canvas, X, Y);
+		Y = RenderStatCues(&World, &Viewport, Canvas, X, Y);
+		Y = RenderStatWaves(&World, &Viewport, Canvas, X, Y);
+	}
+
+	void FAudioDebugger::DrawDebugStats(UWorld& World)
 	{
-		if (ActiveSound)
+		check(IsInGameThread());
+
+		if (UGameViewportClient* GameViewport = World.GetGameViewport())
 		{
-			UE_LOG(LogAudio, Display, TEXT("%s (%.3g) - %s"), *ActiveSound->GetSound()->GetName(), ActiveSound->GetSound()->GetDuration(), *ActiveSound->GetAudioComponentName());
+			if (GameViewport->Viewport)
+			{
+				DrawDebugStats(World, *GameViewport->Viewport);
+			}
+			return;
+		}
 
-			for (const TPair<UPTRINT, FWaveInstance*>& WaveInstancePair : ActiveSound->GetWaveInstances())
+#if WITH_EDITOR
+		if (GEditor)
+		{
+			// Level editor viewport may not be a game viewport, but should still display debug stats, so check if this case is covered and display accordingly
+			const TArray<FLevelEditorViewportClient*>& LevelViewportClients = GEditor->GetLevelViewportClients();
+			for (FLevelEditorViewportClient* ViewportClient : LevelViewportClients)
 			{
-				const FWaveInstance* WaveInstance = WaveInstancePair.Value;
-				UE_LOG(LogAudio, Display, TEXT("   %s (%.3g) (%d) - %.3g"),
-					*WaveInstance->GetName(), WaveInstance->WaveData->GetDuration(),
-					WaveInstance->WaveData->GetResourceSizeBytes(EResourceSizeMode::EstimatedTotal),
-					WaveInstance->GetVolumeWithDistanceAttenuation() * WaveInstance->GetDynamicVolume());
+				if (ViewportClient && ViewportClient->Viewport)
+				{
+					UWorld* EditorWorld = ViewportClient->GetWorld();
+					if (EditorWorld && EditorWorld == &World)
+					{
+						if (DrawDebugStatsEnabled())
+						{
+							ViewportClient->SetRealtimeOverride(true, NSLOCTEXT("AudioDebugger", "AudioDebugger_DrawDebugData", "Audio Debug Data"));
+						}
+						else if (!bAllowUsingDeprecatedDebugStats)
+						{
+							ViewportClient->RemoveRealtimeOverride();
+						}
+						DrawDebugStats(World, *ViewportClient->Viewport);
+					}
+				}
 			}
 		}
+#endif // WITH_EDITOR
 	}
-}
 
-void FAudioDebugger::ResolveDesiredStats(FViewportClient* ViewportClient)
-{
-	if (!ViewportClient)
+	bool FAudioDebugger::DrawDebugStatsEnabled()
 	{
-		return;
+		return bDebugWavesForAllViewsEnabled
+			|| bDebugCuesForAllViewsEnabled
+			|| bDebugSoundsForAllViewsEnabled
+			|| bDebugMixesForAllViewsEnabled
+			|| bDebugReverbForAllViewsEnabled
+			|| bDebugModulationForAllViewsEnabled;
 	}
 
-	FAudioDevice* AudioDevice = GetWorldAudio(ViewportClient->GetWorld());
-	if (!AudioDevice)
+	void FAudioDebugger::DumpActiveSounds() const
 	{
-		return;
-	}
+		if (!GEngine)
+		{
+			return;
+		}
 
-	uint8 SetStats = 0;
-	uint8 ClearStats = 0;
+		if (!IsInAudioThread())
+		{
+			DECLARE_CYCLE_STAT(TEXT("FAudioThreadTask.DumpActiveSounds"), STAT_AudioDumpActiveSounds, STATGROUP_TaskGraphTasks);
+			FAudioThread::RunCommandOnAudioThread([this]()
+			{
+				DumpActiveSounds();
+			}, GET_STATID(STAT_AudioDumpActiveSounds));
+			return;
+		}
 
-	if (ViewportClient->IsStatEnabled(TEXT("SoundCues")))
-	{
-		SetStats |= ERequestedAudioStats::SoundCues;
-	}
-	else
-	{
-		ClearStats |= ERequestedAudioStats::SoundCues;
-	}
+		FAudioDeviceHandle AudioDevice = GEngine->GetAudioDeviceManager()->GetActiveAudioDevice();
+		if (!AudioDevice)
+		{
+			return;
+		}
 
-	if (ViewportClient->IsStatEnabled(TEXT("SoundWaves")))
-	{
-		SetStats |= ERequestedAudioStats::SoundWaves;
-	}
-	else
-	{
-		ClearStats |= ERequestedAudioStats::SoundWaves;
-	}
+		const TArray<FActiveSound*>& ActiveSounds = AudioDevice->GetActiveSounds();
+		UE_LOG(LogAudio, Display, TEXT("Active Sound Count: %d"), ActiveSounds.Num());
+		UE_LOG(LogAudio, Display, TEXT("------------------------"), ActiveSounds.Num());
 
-	if (ViewportClient->IsStatEnabled(TEXT("SoundMixes")))
-	{
-		SetStats |= ERequestedAudioStats::SoundMixes;
-	}
-	else
-	{
-		ClearStats |= ERequestedAudioStats::SoundMixes;
+		for (const FActiveSound* ActiveSound : ActiveSounds)
+		{
+			if (ActiveSound)
+			{
+				UE_LOG(LogAudio, Display, TEXT("%s (%.3g) - %s"), *ActiveSound->GetSound()->GetName(), ActiveSound->GetSound()->GetDuration(), *ActiveSound->GetAudioComponentName());
+
+				for (const TPair<UPTRINT, FWaveInstance*>& WaveInstancePair : ActiveSound->GetWaveInstances())
+				{
+					const FWaveInstance* WaveInstance = WaveInstancePair.Value;
+					UE_LOG(LogAudio, Display, TEXT("   %s (%.3g) (%d) - %.3g"),
+						*WaveInstance->GetName(), WaveInstance->WaveData->GetDuration(),
+						WaveInstance->WaveData->GetResourceSizeBytes(EResourceSizeMode::EstimatedTotal),
+						WaveInstance->GetVolumeWithDistanceAttenuation() * WaveInstance->GetDynamicVolume());
+				}
+			}
+		}
 	}
 
-	if (ViewportClient->IsStatEnabled(TEXT("Sounds")))
+	void FAudioDebugger::ResolveDesiredStats(FViewportClient* ViewportClient)
 	{
-		FAudioStats& Stats = AudioDeviceStats.FindOrAdd(AudioDevice->DeviceID);
-		SetStats |= ERequestedAudioStats::Sounds;
+		if (!ViewportClient)
+		{
+			return;
+		}
+
+		FAudioDevice* AudioDevice = GetWorldAudio(ViewportClient->GetWorld());
+		if (!AudioDevice)
+		{
+			return;
+		}
+
+		uint32 SetStatFlags = 0;
+		uint32 ClearStatFlags = 0;
 
-		if (Stats.DisplayFlags & static_cast<uint8>(FAudioStats::EDisplayFlags::Debug))
+		if (ViewportClient->IsStatEnabled(TEXT("SoundCues")))
 		{
-			SetStats |= ERequestedAudioStats::DebugSounds;
+			SetStatFlags |= FAudioStats::SoundCues;
 		}
 		else
 		{
-			ClearStats |= ERequestedAudioStats::DebugSounds;
+			ClearStatFlags |= FAudioStats::SoundCues;
 		}
 
-		if (Stats.DisplayFlags & static_cast<uint8>(FAudioStats::EDisplayFlags::Long_Names))
+		if (ViewportClient->IsStatEnabled(TEXT("SoundWaves")))
 		{
-			SetStats |= ERequestedAudioStats::LongSoundNames;
+			SetStatFlags |= FAudioStats::SoundWaves;
 		}
 		else
 		{
-			ClearStats |= ERequestedAudioStats::LongSoundNames;
+			ClearStatFlags |= FAudioStats::SoundWaves;
 		}
-	}
-	else
-	{
-		ClearStats |= ERequestedAudioStats::Sounds;
-		ClearStats |= ERequestedAudioStats::DebugSounds;
-		ClearStats |= ERequestedAudioStats::LongSoundNames;
-	}
 
-	DECLARE_CYCLE_STAT(TEXT("FAudioThreadTask.ResolveDesiredStats"), STAT_AudioResolveDesiredStats, STATGROUP_TaskGraphTasks);
+		if (ViewportClient->IsStatEnabled(TEXT("SoundMixes")))
+		{
+			SetStatFlags |= FAudioStats::SoundMixes;
+		}
+		else
+		{
+			ClearStatFlags |= FAudioStats::SoundMixes;
+		}
 
-	const uint32 DeviceID = AudioDevice->DeviceID;
-	if (IsInAudioThread())
-	{
-		FAudioStats_AudioThread& Stats = AudioDeviceStats_AudioThread.FindOrAdd(DeviceID);
-		Stats.RequestedStats |= SetStats;
-		Stats.RequestedStats &= ~ClearStats;
-	}
-	else
-	{
-		FAudioThread::RunCommandOnAudioThread([SetStats, ClearStats, DeviceID]()
+		if (ViewportClient->IsStatEnabled(TEXT("Sounds")))
+		{
+			FAudioStats& Stats = AudioDeviceStats.FindOrAdd(AudioDevice->DeviceID);
+			SetStatFlags |= FAudioStats::Sounds;
+
+			if (Stats.DisplayFlags & static_cast<uint8>(FAudioStats::EDisplayFlags::Debug))
+			{
+				SetStatFlags |= FAudioStats::DebugSounds;
+			}
+			else
+			{
+				ClearStatFlags |= FAudioStats::DebugSounds;
+			}
+
+			if (Stats.DisplayFlags & static_cast<uint8>(FAudioStats::EDisplayFlags::Long_Names))
+			{
+				SetStatFlags |= FAudioStats::LongSoundNames;
+			}
+			else
+			{
+				ClearStatFlags |= FAudioStats::LongSoundNames;
+			}
+		}
+		else
+		{
+			ClearStatFlags |= FAudioStats::Sounds;
+			ClearStatFlags |= FAudioStats::DebugSounds;
+			ClearStatFlags |= FAudioStats::LongSoundNames;
+		}
+
+		DECLARE_CYCLE_STAT(TEXT("FAudioThreadTask.ResolveDesiredStats"), STAT_AudioResolveDesiredStats, STATGROUP_TaskGraphTasks);
+
+		const uint32 DeviceID = AudioDevice->DeviceID;
+		if (IsInAudioThread())
 		{
 			FAudioStats_AudioThread& Stats = AudioDeviceStats_AudioThread.FindOrAdd(DeviceID);
-			Stats.RequestedStats |= SetStats;
-			Stats.RequestedStats &= ~ClearStats;
-		}, GET_STATID(STAT_AudioResolveDesiredStats));
+			Stats.RequestedStats |= SetStatFlags;
+			Stats.RequestedStats &= ~ClearStatFlags;
+		}
+		else
+		{
+			FAudioThread::RunCommandOnAudioThread([SetStatFlags, ClearStatFlags, DeviceID]()
+			{
+				FAudioStats_AudioThread& Stats = AudioDeviceStats_AudioThread.FindOrAdd(DeviceID);
+				Stats.RequestedStats |= SetStatFlags;
+				Stats.RequestedStats &= ~ClearStatFlags;
+			}, GET_STATID(STAT_AudioResolveDesiredStats));
+		}
 	}
-}
 
-int32 FAudioDebugger::RenderStatCues(UWorld* World, FViewport* Viewport, FCanvas* Canvas, int32 X, int32 Y, const FVector* ViewLocation, const FRotator* ViewRotation)
-{
-	FAudioDevice* AudioDevice = GetWorldAudio(World);
-	if (!AudioDevice)
+	int32 FAudioDebugger::RenderStatCues(UWorld* World, FViewport* Viewport, FCanvas* Canvas, int32 X, int32 Y)
 	{
-		return Y;
-	}
+		FAudioDevice* AudioDevice = nullptr;
+		if (!DebugShouldRenderStat(World, Canvas, bDebugCuesForAllViewsEnabled, FAudioStats::SoundCues, &AudioDevice))
+		{
+			return Y;
+		}
 
-	UFont* StatsFont = GetStatsFont();
-	const int32 FontSpacing = 2;
-	const int32 FontHeight = StatsFont->GetMaxCharHeight() + FontSpacing;
-	Canvas->DrawShadowedString(X, Y, TEXT("Active Sound Cues:"), StatsFont, HeaderColor);
-	Y += FontHeight;
+		UFont* StatsFont = GetStatsFont();
+		const int32 FontSpacing = 2;
+		const int32 FontHeight = StatsFont->GetMaxCharHeight() + FontSpacing;
 
-	int32 ActiveSoundCount = 0;
-	FAudioStats& AudioStats = AudioDeviceStats.FindOrAdd(AudioDevice->DeviceID);
+		Canvas->DrawShadowedString(X, Y, TEXT("Active Sound Cues:"), StatsFont, HeaderColor);
+		Y += FontHeight;
 
-	if (!SoundCueDebugMinimalCVar)
-	{
-	for (const FAudioStats::FStatSoundInfo& StatSoundInfo : AudioStats.StatSoundInfos)
-	{
-		for (const FAudioStats::FStatWaveInstanceInfo& WaveInstanceInfo : StatSoundInfo.WaveInstanceInfos)
+		int32 ActiveSoundCount = 0;
+
+		FAudioStats& AudioStats = AudioDeviceStats.FindOrAdd(AudioDevice->DeviceID);
+		if (!SoundCueDebugMinimalCVar)
 		{
-			if (WaveInstanceInfo.Volume >= MinDisplayVolume)
+			for (const FAudioStats::FStatSoundInfo& StatSoundInfo : AudioStats.StatSoundInfos)
 			{
-				FColor Color = FColor::White;
-				FString MuteSoloReason;
-
-					if (FSoundSource::FDebugInfo* DebugInfo = WaveInstanceInfo.DebugInfo.Get())
+				for (const FAudioStats::FStatWaveInstanceInfo& WaveInstanceInfo : StatSoundInfo.WaveInstanceInfos)
+				{
+					if (WaveInstanceInfo.Volume >= MinDisplayVolume)
 					{
-						// Color code same as icons. Red (mute), Yellow (solo), White (normal).
-						FScopeLock Lock(&DebugInfo->CS);
-						Color = DebugInfo->bIsMuted ? FColor::Red : DebugInfo->bIsSoloed ? FColor::Yellow : FColor::White;
-						MuteSoloReason = !DebugInfo->MuteSoloReason.IsEmpty() ? FString::Printf(TEXT(" - %s"), *DebugInfo->MuteSoloReason) : TEXT("");
-					}
+						FColor Color = FColor::White;
+						FString MuteSoloReason;
 
-					const FString TheString = FString::Printf(TEXT("%4i. %6.2f %s %s %s"), ActiveSoundCount++, WaveInstanceInfo.Volume, *StatSoundInfo.SoundName, *StatSoundInfo.SoundClassName.ToString(), *MuteSoloReason);
-					Canvas->DrawShadowedString(X, Y, *TheString, StatsFont, Color);
-					Y += FontHeight;
-					break;
+						if (FSoundSource::FDebugInfo* DebugInfo = WaveInstanceInfo.DebugInfo.Get())
+						{
+							// Color code same as icons. Red (mute), Yellow (solo), White (normal).
+							FScopeLock Lock(&DebugInfo->CS);
+							Color = DebugInfo->bIsMuted ? FColor::Red : DebugInfo->bIsSoloed ? FColor::Yellow : FColor::White;
+							MuteSoloReason = !DebugInfo->MuteSoloReason.IsEmpty() ? FString::Printf(TEXT(" - %s"), *DebugInfo->MuteSoloReason) : TEXT("");
+						}
+
+						const FString TheString = FString::Printf(TEXT("%4i. %6.2f %s %s %s"), ActiveSoundCount++, WaveInstanceInfo.Volume, *StatSoundInfo.SoundName, *StatSoundInfo.SoundClassName.ToString(), *MuteSoloReason);
+						Canvas->DrawShadowedString(X, Y, *TheString, StatsFont, Color);
+						Y += FontHeight;
+						break;
+					}
 				}
 			}
 		}
-	}
-	else
-	{
-		TArray<TPair<FString, FString> > SplitNames;
-		FString SoundPath, SoundName;
-		int32 MaxNameLength = 0;
-		uint32 MaxClassLength = 0;
-
-		// Get spacing data
-		for (const FAudioStats::FStatSoundInfo& StatSoundInfo : AudioStats.StatSoundInfos)
+		else
 		{
-			for (const FAudioStats::FStatWaveInstanceInfo& WaveInstanceInfo : StatSoundInfo.WaveInstanceInfos)
+			TArray<TPair<FString, FString>> SplitNames;
+			FString SoundPath, SoundName;
+			int32 MaxNameLength = 0;
+			uint32 MaxClassLength = 0;
+
+			// Get spacing data
+			for (const FAudioStats::FStatSoundInfo& StatSoundInfo : AudioStats.StatSoundInfos)
 			{
-				if (WaveInstanceInfo.Volume >= MinDisplayVolume)
+				for (const FAudioStats::FStatWaveInstanceInfo& WaveInstanceInfo : StatSoundInfo.WaveInstanceInfos)
 				{
-					if (!StatSoundInfo.SoundPath.Split(TEXT("."), &SoundPath, &SoundName))
+					if (WaveInstanceInfo.Volume >= MinDisplayVolume)
 					{
-						if (!StatSoundInfo.SoundPath.Split(SUBOBJECT_DELIMITER, &SoundPath, &SoundName))
+						if (!StatSoundInfo.SoundPath.Split(TEXT("."), &SoundPath, &SoundName))
 						{
-							SoundPath = StatSoundInfo.SoundPath;
+							if (!StatSoundInfo.SoundPath.Split(SUBOBJECT_DELIMITER, &SoundPath, &SoundName))
+							{
+								SoundPath = StatSoundInfo.SoundPath;
+							}
 						}
-					}
-					SoundName = StatSoundInfo.SoundName;
-					SplitNames.Emplace(SoundName, SoundPath);
+						SoundName = StatSoundInfo.SoundName;
+						SplitNames.Emplace(SoundName, SoundPath);
 
-					if (SoundName.Len() > MaxNameLength)
-					{
-						MaxNameLength = SoundName.Len();
-					}
+						if (SoundName.Len() > MaxNameLength)
+						{
+							MaxNameLength = SoundName.Len();
+						}
 
-					if (StatSoundInfo.SoundClassName.GetStringLength() > MaxClassLength)
-					{
-						MaxClassLength = StatSoundInfo.SoundClassName.GetStringLength();
+						if (StatSoundInfo.SoundClassName.GetStringLength() > MaxClassLength)
+						{
+							MaxClassLength = StatSoundInfo.SoundClassName.GetStringLength();
+						}
 					}
 				}
 			}
-		}
 
-		const int32 TabSpacing = FMath::Clamp(SoundCueDebugTabSpacingCVar, 1, SoundCueDebugTabSpacingCVar);
-		const int32 CharSpacing = FMath::Clamp(SoundCueDebugCharSpacingCVar, 1, SoundCueDebugCharSpacingCVar);
-		const int32 NumberSpacing = 6 * CharSpacing;	// 6 character len for 2 decimal float + 2 spaces 'X.XX  '
+			const int32 TabSpacing = FMath::Clamp(SoundCueDebugTabSpacingCVar, 1, SoundCueDebugTabSpacingCVar);
+			const int32 CharSpacing = FMath::Clamp(SoundCueDebugCharSpacingCVar, 1, SoundCueDebugCharSpacingCVar);
+			const int32 NumberSpacing = 6 * CharSpacing;	// 6 character len for 2 decimal float + 2 spaces 'X.XX  '
 
-		// Tab out name and class length
-		MaxNameLength = (MaxNameLength / TabSpacing + 1) * TabSpacing;
-		MaxClassLength = (MaxClassLength / TabSpacing + 1) * TabSpacing;
+			// Tab out name and class length
+			MaxNameLength = (MaxNameLength / TabSpacing + 1) * TabSpacing;
+			MaxClassLength = (MaxClassLength / TabSpacing + 1) * TabSpacing;
 
-		for (const FAudioStats::FStatSoundInfo& StatSoundInfo : AudioStats.StatSoundInfos)
-		{
-			for (const FAudioStats::FStatWaveInstanceInfo& WaveInstanceInfo : StatSoundInfo.WaveInstanceInfos)
+			for (const FAudioStats::FStatSoundInfo& StatSoundInfo : AudioStats.StatSoundInfos)
 			{
-				if (WaveInstanceInfo.Volume >= MinDisplayVolume)
+				for (const FAudioStats::FStatWaveInstanceInfo& WaveInstanceInfo : StatSoundInfo.WaveInstanceInfos)
 				{
-					FColor Color = FColor::White;
-					FString MuteSoloReason;
-					bool bMutedOrSoloed = false;
-
-					if (FSoundSource::FDebugInfo* DebugInfo = WaveInstanceInfo.DebugInfo.Get())
+					if (WaveInstanceInfo.Volume >= MinDisplayVolume)
 					{
-						// Color code same as icons. Red (mute), Yellow (solo), White (normal).
-						FScopeLock Lock(&DebugInfo->CS);
-						Color = DebugInfo->bIsMuted ? FColor::Red : DebugInfo->bIsSoloed ? FColor::Yellow : FColor::White;
-						bMutedOrSoloed = DebugInfo->bIsMuted || DebugInfo->bIsSoloed;
-						MuteSoloReason = !DebugInfo->MuteSoloReason.IsEmpty() ? FString::Printf(TEXT(" - %s"), *DebugInfo->MuteSoloReason) : TEXT("");
-					}
+						FColor Color = FColor::White;
+						FString MuteSoloReason;
+						bool bMutedOrSoloed = false;
 
-					const int32 SoundNameIndex = ActiveSoundCount++;
-					const FString LeadingNumber = FString::Printf(TEXT("%4i. "), SoundNameIndex);
-					const FString Volume = FString::Printf(TEXT("%6.2f "), WaveInstanceInfo.Volume);
-					const FString PathAndMuting = FString::Printf(TEXT("Path: %s %s"), *SplitNames[SoundNameIndex].Value, *MuteSoloReason);
-
-					int32 CurrentX = X;
-					Canvas->DrawShadowedString(CurrentX, Y, *LeadingNumber, StatsFont, Color);
-					CurrentX += NumberSpacing;
-					Canvas->DrawShadowedString(CurrentX, Y, *Volume, StatsFont, bMutedOrSoloed ? Color : FColor::Orange);
-					CurrentX += NumberSpacing;
-					Canvas->DrawShadowedString(CurrentX, Y, *SplitNames[SoundNameIndex].Key, StatsFont, bMutedOrSoloed ? Color : FColor(0, 255, 255));
-					CurrentX += (MaxNameLength * CharSpacing);
-					Canvas->DrawShadowedString(CurrentX, Y, *StatSoundInfo.SoundClassName.ToString(), StatsFont, bMutedOrSoloed ? Color : FColor::Yellow);
-
-					if (SoundCueDebugShowPathCVar)
-					{
-						CurrentX += (MaxClassLength * CharSpacing);
-						Canvas->DrawShadowedString(CurrentX, Y, *PathAndMuting, StatsFont, Color);
-					}
+						if (FSoundSource::FDebugInfo* DebugInfo = WaveInstanceInfo.DebugInfo.Get())
+						{
+							// Color code same as icons. Red (mute), Yellow (solo), White (normal).
+							FScopeLock Lock(&DebugInfo->CS);
+							Color = DebugInfo->bIsMuted ? FColor::Red : DebugInfo->bIsSoloed ? FColor::Yellow : FColor::White;
+							bMutedOrSoloed = DebugInfo->bIsMuted || DebugInfo->bIsSoloed;
+							MuteSoloReason = !DebugInfo->MuteSoloReason.IsEmpty() ? FString::Printf(TEXT(" - %s"), *DebugInfo->MuteSoloReason) : TEXT("");
+						}
 
-					Y += FontHeight;
-					break;
+						const int32 SoundNameIndex = ActiveSoundCount++;
+						const FString LeadingNumber = FString::Printf(TEXT("%4i. "), SoundNameIndex);
+						const FString Volume = FString::Printf(TEXT("%6.2f "), WaveInstanceInfo.Volume);
+						const FString PathAndMuting = FString::Printf(TEXT("Path: %s %s"), *SplitNames[SoundNameIndex].Value, *MuteSoloReason);
+
+						int32 CurrentX = X;
+						Canvas->DrawShadowedString(CurrentX, Y, *LeadingNumber, StatsFont, Color);
+						CurrentX += NumberSpacing;
+						Canvas->DrawShadowedString(CurrentX, Y, *Volume, StatsFont, bMutedOrSoloed ? Color : FColor::Orange);
+						CurrentX += NumberSpacing;
+						Canvas->DrawShadowedString(CurrentX, Y, *SplitNames[SoundNameIndex].Key, StatsFont, bMutedOrSoloed ? Color : FColor(0, 255, 255));
+						CurrentX += (MaxNameLength * CharSpacing);
+						Canvas->DrawShadowedString(CurrentX, Y, *StatSoundInfo.SoundClassName.ToString(), StatsFont, bMutedOrSoloed ? Color : FColor::Yellow);
+
+						if (SoundCueDebugShowPathCVar)
+						{
+							CurrentX += (MaxClassLength * CharSpacing);
+							Canvas->DrawShadowedString(CurrentX, Y, *PathAndMuting, StatsFont, Color);
+						}
+
+						Y += FontHeight;
+						break;
+					}
 				}
 			}
 		}
-	}
-
-	Canvas->DrawShadowedString(X, Y, *FString::Printf(TEXT("Total: %i"), ActiveSoundCount), GetStatsFont(), GetBodyColor());
-	Y += FontHeight;
 
-	return Y;
-}
+		Canvas->DrawShadowedString(X, Y, *FString::Printf(TEXT("Total: %i"), ActiveSoundCount), GetStatsFont(), GetBodyColor());
+		Y += FontHeight;
 
-int32 FAudioDebugger::RenderStatMixes(UWorld* World, FViewport* Viewport, FCanvas* Canvas, int32 X, int32 Y, const FVector* ViewLocation, const FRotator* ViewRotation)
-{
-	FAudioDevice* AudioDevice = GetWorldAudio(World);
-	if (!AudioDevice)
-	{
 		return Y;
 	}
 
-	const int32 FontHeight = GetStatsFont()->GetMaxCharHeight() + 2;
-	Canvas->DrawShadowedString(X, Y, TEXT("Active Sound Mixes:"), GetStatsFont(), HeaderColor);
-	Y += FontHeight;
+	int32 FAudioDebugger::RenderStatMixes(UWorld* World, FViewport* Viewport, FCanvas* Canvas, int32 X, int32 Y)
+	{
+		FAudioDevice* AudioDevice = nullptr;
+		if (!DebugShouldRenderStat(World, Canvas, bDebugMixesForAllViewsEnabled, FAudioStats::SoundMixes, &AudioDevice))
+		{
+			return Y;
+		}
 
-	bool bDisplayedSoundMixes = false;
+		const int32 FontHeight = GetStatsFont()->GetMaxCharHeight() + 2;
+		FAudioStats& AudioStats = AudioDeviceStats.FindOrAdd(AudioDevice->DeviceID);
+		Canvas->DrawShadowedString(X, Y, TEXT("Active Sound Mixes:"), GetStatsFont(), HeaderColor);
+		Y += FontHeight;
 
-	FAudioStats& AudioStats = AudioDeviceStats.FindOrAdd(AudioDevice->DeviceID);
-	if (AudioStats.StatSoundMixes.Num() > 0)
-	{
-		bDisplayedSoundMixes = true;
+		bool bDisplayedSoundMixes = false;
 
-		for (const FAudioStats::FStatSoundMix& StatSoundMix : AudioStats.StatSoundMixes)
+		if (AudioStats.StatSoundMixes.Num() > 0)
 		{
-			const FString TheString = FString::Printf(TEXT("%s - Fade Proportion: %1.2f - Total Ref Count: %i"), *StatSoundMix.MixName, StatSoundMix.InterpValue, StatSoundMix.RefCount);
+			bDisplayedSoundMixes = true;
+
+			for (const FAudioStats::FStatSoundMix& StatSoundMix : AudioStats.StatSoundMixes)
+			{
+				const FString TheString = FString::Printf(TEXT("%s - Fade Proportion: %1.2f - Total Ref Count: %i"), *StatSoundMix.MixName, StatSoundMix.InterpValue, StatSoundMix.RefCount);
 
-			const FColor& TextColor = (StatSoundMix.bIsCurrentEQ ? FColor::Yellow : GetBodyColor());
+				const FColor& TextColor = (StatSoundMix.bIsCurrentEQ ? FColor::Yellow : GetBodyColor());
 
-			Canvas->DrawShadowedString(X + TabWidth, Y, *TheString, GetStatsFont(), TextColor);
+				Canvas->DrawShadowedString(X + DebuggerTabWidth, Y, *TheString, GetStatsFont(), TextColor);
+				Y += FontHeight;
+			}
+		}
+
+		if (!bDisplayedSoundMixes)
+		{
+			Canvas->DrawShadowedString(X + DebuggerTabWidth, Y, TEXT("None"), GetStatsFont(), GetBodyColor());
 			Y += FontHeight;
 		}
-	}
 
-	if (!bDisplayedSoundMixes)
-	{
-		Canvas->DrawShadowedString(X + TabWidth, Y, TEXT("None"), GetStatsFont(), GetBodyColor());
 		Y += FontHeight;
+		return Y;
 	}
 
-	return Y;
-}
-
-int32 FAudioDebugger::RenderStatModulators(UWorld* World, FViewport* Viewport, FCanvas* Canvas, int32 X, int32 Y, const FVector* ViewLocation, const FRotator* ViewRotation)
-{
-	FAudioDevice* AudioDevice = GetWorldAudio(World);
-	if (!AudioDevice)
+	int32 FAudioDebugger::RenderStatModulators(UWorld* World, FViewport* Viewport, FCanvas* Canvas, int32 X, int32 Y, const FVector* ViewLocation, const FRotator* ViewRotation)
 	{
-		return Y;
-	}
+		FAudioDevice* AudioDevice = nullptr;
+		if (!DebugShouldRenderStat(World, Canvas, bDebugModulationForAllViewsEnabled, FAudioStats::SoundModulation, &AudioDevice))
+		{
+			return Y;
+		}
 
-	const int32 FontHeight = GetStatsFont()->GetMaxCharHeight() + 2;
-	Canvas->DrawShadowedString(X, Y, TEXT("Active Sound Modulation:"), GetStatsFont(), HeaderColor);
-	Y += FontHeight;
+		const int32 FontHeight = GetStatsFont()->GetMaxCharHeight() + 2;
+		Canvas->DrawShadowedString(X, Y, TEXT("Active Sound Modulation:"), GetStatsFont(), HeaderColor);
 
-	bool bDisplayedSoundModulationInfo = false;
+		bool bDisplayedSoundModulationInfo = false;
+		if (IAudioModulation* Modulation = AudioDevice->ModulationInterface.Get())
+		{
+			const int32 YInit = Y;
+			Y = Modulation->OnRenderStat(Viewport, Canvas, X, Y, *GetStatsFont(), ViewLocation, ViewRotation);
+			bDisplayedSoundModulationInfo = Y != YInit;
+		}
 
-	if (IAudioModulation* Modulation = AudioDevice->ModulationInterface.Get())
-	{
-		const int32 YInit = Y;
-		Y = Modulation->OnRenderStat(Viewport, Canvas, X, Y, *GetStatsFont(), ViewLocation, ViewRotation);
-		bDisplayedSoundModulationInfo = Y != YInit;
-	}
+		if (!bDisplayedSoundModulationInfo)
+		{
+			Canvas->DrawShadowedString(X + DebuggerTabWidth, Y, TEXT("None"), GetStatsFont(), GetBodyColor());
+			Y += FontHeight;
+		}
 
-	if (!bDisplayedSoundModulationInfo)
-	{
-		Canvas->DrawShadowedString(X + TabWidth, Y, TEXT("None"), GetStatsFont(), GetBodyColor());
 		Y += FontHeight;
-	}
-
-	return Y;
-}
-
-int32 FAudioDebugger::RenderStatReverb(UWorld* World, FViewport* Viewport, FCanvas* Canvas, int32 X, int32 Y, const FVector* ViewLocation, const FRotator* ViewRotation)
-{
-	FAudioDevice* AudioDevice = GetWorldAudio(World);
-	if (!AudioDevice)
-	{
 		return Y;
 	}
 
-	check(IsInGameThread());
-
-	const int32 Height = static_cast<int32>(GetStatsFont()->GetMaxCharHeight() + 2);
-
-	FString TheString;
-	const FLinearColor LinearBodyColor = FLinearColor(GetBodyColor());
-	if (UReverbEffect* ReverbEffect = AudioDevice->GetCurrentReverbEffect())
+	int32 FAudioDebugger::RenderStatReverb(UWorld* World, FViewport* Viewport, FCanvas* Canvas, int32 X, int32 Y)
 	{
-		TheString = FString::Printf(TEXT("Active Reverb Effect: %s"), *ReverbEffect->GetName());
-		Canvas->DrawShadowedString(X, Y, *TheString, GetStatsFont(), LinearBodyColor);
-		Y += Height;
-
-		AAudioVolume* CurrentAudioVolume = nullptr;
-		const int32 ProxyCount = AudioDevice->ListenerProxies.Num();
-		for (int i = 0; i < ProxyCount; ++i)
-		{
-			const FTransform& Transform = AudioDevice->ListenerProxies[i].Transform;
-			AAudioVolume* PlayerAudioVolume = World->GetAudioSettings(Transform.GetLocation(), nullptr, nullptr);
-			if (PlayerAudioVolume && ((CurrentAudioVolume == nullptr) || (PlayerAudioVolume->GetPriority() > CurrentAudioVolume->GetPriority())))
-			{
-				CurrentAudioVolume = PlayerAudioVolume;
-			}
-		}
-		if (CurrentAudioVolume && CurrentAudioVolume->GetReverbSettings().ReverbEffect)
+		FAudioDevice* AudioDevice = nullptr;
+		if (!DebugShouldRenderStat(World, Canvas, bDebugReverbForAllViewsEnabled, FAudioStats::SoundReverb, &AudioDevice))
 		{
-			TheString = FString::Printf(TEXT("  Audio Volume Reverb Effect: %s (Priority: %g Volume Name: %s)"), *CurrentAudioVolume->GetReverbSettings().ReverbEffect->GetName(), CurrentAudioVolume->GetPriority(), *CurrentAudioVolume->GetName());
+			return Y;
 		}
-		else
-		{
-			TheString = TEXT("  Audio Volume Reverb Effect: None");
-		}
-		Canvas->DrawShadowedString(X, Y, *TheString, GetStatsFont(), LinearBodyColor);
-		Y += Height;
 
-		const TMap<FName, FActivatedReverb>& ActivatedReverbs = AudioDevice->GetActiveReverb();
-		if (ActivatedReverbs.Num() == 0)
+		const int32 Height = static_cast<int32>(GetStatsFont()->GetMaxCharHeight() + 2);
+
+		FString TheString;
+		const FLinearColor LinearBodyColor = FLinearColor(GetBodyColor());
+		if (UReverbEffect* ReverbEffect = AudioDevice->GetCurrentReverbEffect())
 		{
-			TheString = TEXT("  Activated Reverb: None");
-			Canvas->DrawShadowedString(X, Y, *TheString, GetStatsFont(), LinearBodyColor);
+			TheString = FString::Printf(TEXT("Active Reverb Effect: %s"), *ReverbEffect->GetName());
+			Canvas->DrawShadowedString(X, Y, *TheString, GetStatsFont(), HeaderColor);
 			Y += Height;
-		}
-		else if (ActivatedReverbs.Num() == 1)
-		{
-			auto It = ActivatedReverbs.CreateConstIterator();
-			TheString = FString::Printf(TEXT("  Activated Reverb Effect: %s (Priority: %g Tag: '%s')"), *It.Value().ReverbSettings.ReverbEffect->GetName(), It.Value().Priority, *It.Key().ToString());
+
+			AAudioVolume* CurrentAudioVolume = nullptr;
+			const int32 ProxyCount = AudioDevice->ListenerProxies.Num();
+			for (int i = 0; i < ProxyCount; ++i)
+			{
+				const FTransform& Transform = AudioDevice->ListenerProxies[i].Transform;
+				AAudioVolume* PlayerAudioVolume = World->GetAudioSettings(Transform.GetLocation(), nullptr, nullptr);
+				if (PlayerAudioVolume && ((CurrentAudioVolume == nullptr) || (PlayerAudioVolume->GetPriority() > CurrentAudioVolume->GetPriority())))
+				{
+					CurrentAudioVolume = PlayerAudioVolume;
+				}
+			}
+			if (CurrentAudioVolume && CurrentAudioVolume->GetReverbSettings().ReverbEffect)
+			{
+				TheString = FString::Printf(TEXT("  Audio Volume Reverb Effect: %s (Priority: %g Volume Name: %s)"), *CurrentAudioVolume->GetReverbSettings().ReverbEffect->GetName(), CurrentAudioVolume->GetPriority(), *CurrentAudioVolume->GetName());
+			}
+			else
+			{
+				TheString = TEXT("  Audio Volume Reverb Effect: None");
+			}
 			Canvas->DrawShadowedString(X, Y, *TheString, GetStatsFont(), LinearBodyColor);
 			Y += Height;
-		}
-		else
-		{
-			Canvas->DrawShadowedString(X, Y, TEXT("  Activated Reverb Effects:"), GetStatsFont(), LinearBodyColor);
-			Y += Height;
-			TMap<int32, FString> PrioritySortedActivatedReverbs;
-			for (auto It = ActivatedReverbs.CreateConstIterator(); It; ++It)
+
+			const TMap<FName, FActivatedReverb>& ActivatedReverbs = AudioDevice->GetActiveReverb();
+			if (ActivatedReverbs.Num() == 0)
 			{
-				TheString = FString::Printf(TEXT("    %s (Priority: %g Tag: '%s')"), *It.Value().ReverbSettings.ReverbEffect->GetName(), It.Value().Priority, *It.Key().ToString());
-				PrioritySortedActivatedReverbs.Add(It.Value().Priority, TheString);
+				TheString = TEXT("  Activated Reverb: None");
+				Canvas->DrawShadowedString(X, Y, *TheString, GetStatsFont(), LinearBodyColor);
+				Y += Height;
 			}
-			for (auto It = PrioritySortedActivatedReverbs.CreateConstIterator(); It; ++It)
+			else if (ActivatedReverbs.Num() == 1)
 			{
-				Canvas->DrawShadowedString(X, Y, *It.Value(), GetStatsFont(), LinearBodyColor);
+				auto It = ActivatedReverbs.CreateConstIterator();
+				TheString = FString::Printf(TEXT("  Activated Reverb Effect: %s (Priority: %g Tag: '%s')"), *It.Value().ReverbSettings.ReverbEffect->GetName(), It.Value().Priority, *It.Key().ToString());
+				Canvas->DrawShadowedString(X, Y, *TheString, GetStatsFont(), LinearBodyColor);
 				Y += Height;
 			}
+			else
+			{
+				Canvas->DrawShadowedString(X, Y, TEXT("  Activated Reverb Effects:"), GetStatsFont(), LinearBodyColor);
+				Y += Height;
+				TMap<int32, FString> PrioritySortedActivatedReverbs;
+				for (auto It = ActivatedReverbs.CreateConstIterator(); It; ++It)
+				{
+					TheString = FString::Printf(TEXT("    %s (Priority: %g Tag: '%s')"), *It.Value().ReverbSettings.ReverbEffect->GetName(), It.Value().Priority, *It.Key().ToString());
+					PrioritySortedActivatedReverbs.Add(It.Value().Priority, TheString);
+				}
+				for (auto It = PrioritySortedActivatedReverbs.CreateConstIterator(); It; ++It)
+				{
+					Canvas->DrawShadowedString(X, Y, *It.Value(), GetStatsFont(), LinearBodyColor);
+					Y += Height;
+				}
+			}
+		}
+		else
+		{
+			TheString = TEXT("Active Reverb Effect: None");
+			Canvas->DrawShadowedString(X, Y, *TheString, GetStatsFont(), HeaderColor);
+			Y += Height;
 		}
-	}
-	else
-	{
-		TheString = TEXT("Active Reverb Effect: None");
-		Canvas->DrawShadowedString(X, Y, *TheString, GetStatsFont(), LinearBodyColor);
-		Y += Height;
-	}
-
-	return Y;
-}
 
-int32 FAudioDebugger::RenderStatSounds(UWorld* World, FViewport* Viewport, FCanvas* Canvas, int32 X, int32 Y, const FVector* ViewLocation, const FRotator* ViewRotation)
-{
-	FAudioDevice* AudioDevice = GetWorldAudio(World);
-	if (!AudioDevice)
-	{
+		Y += Height;
 		return Y;
 	}
 
-	const int32 FontHeight = GetStatsFont()->GetMaxCharHeight() + 2;
-	Y += FontHeight;
-
-	FAudioStats& AudioStats = AudioDeviceStats.FindOrAdd(AudioDevice->DeviceID);
+	int32 FAudioDebugger::RenderStatSounds(UWorld* World, FViewport* Viewport, FCanvas* Canvas, int32 X, int32 Y)
+	{
+		FAudioDevice* AudioDevice = nullptr;
+		if (!DebugShouldRenderStat(World, Canvas, bDebugSoundsForAllViewsEnabled, FAudioStats::Sounds, &AudioDevice))
+		{
+			return Y;
+		}
 
-	const uint8 bDebug = AudioStats.DisplayFlags & static_cast<uint8>(FAudioStats::EDisplayFlags::Debug);
+		const int32 FontHeight = GetStatsFont()->GetMaxCharHeight() + 2;
 
-	UpdateDisplaySort(AudioStats);
+		FAudioStats& AudioStats = AudioDeviceStats.FindOrAdd(AudioDevice->DeviceID);
+		const uint8 bDebug = AudioStats.DisplayFlags & static_cast<uint8>(FAudioStats::EDisplayFlags::Debug);
+		UpdateDisplaySort(AudioStats);
 
-	// Sort the list.
-	FString SortingName;
-	switch (AudioStats.DisplaySort)
-	{
-		case FAudioStats::EDisplaySort::Class:
-	{
-			AudioStats.StatSoundInfos.Sort([](const FAudioStats::FStatSoundInfo& A, const FAudioStats::FStatSoundInfo& B) { return A.SoundClassName.LexicalLess(B.SoundClassName); });
-			SortingName = TEXT("Class");
-			break;
-	}
-		case FAudioStats::EDisplaySort::Distance:
-	{
-		AudioStats.StatSoundInfos.Sort([](const FAudioStats::FStatSoundInfo& A, const FAudioStats::FStatSoundInfo& B) { return A.Distance < B.Distance; });
-			SortingName = TEXT("Distance");
-			break;
-	}
-		case FAudioStats::EDisplaySort::PlaybackTime:
-	{
-			AudioStats.StatSoundInfos.Sort([](const FAudioStats::FStatSoundInfo& A, const FAudioStats::FStatSoundInfo& B) { return A.PlaybackTime > B.PlaybackTime; });
-			SortingName = TEXT("Time");
-			break;
+		// Sort the list.
+		FString SortingName;
+		switch (AudioStats.DisplaySort)
+		{
+			case FAudioStats::EDisplaySort::Class:
+		{
+				AudioStats.StatSoundInfos.Sort([](const FAudioStats::FStatSoundInfo& A, const FAudioStats::FStatSoundInfo& B) { return A.SoundClassName.LexicalLess(B.SoundClassName); });
+				SortingName = TEXT("Class");
+				break;
 		}
-		case FAudioStats::EDisplaySort::Priority:
+			case FAudioStats::EDisplaySort::Distance:
 		{
-			AudioStats.StatSoundInfos.Sort([](const FAudioStats::FStatSoundInfo& A, const FAudioStats::FStatSoundInfo& B) { return A.Priority > B.Priority; });
-			SortingName = TEXT("Priority");
-			break;
-	}
-		case FAudioStats::EDisplaySort::Waves:
-	{
-		AudioStats.StatSoundInfos.Sort([](const FAudioStats::FStatSoundInfo& A, const FAudioStats::FStatSoundInfo& B) { return A.WaveInstanceInfos.Num() > B.WaveInstanceInfos.Num(); });
-			SortingName = TEXT("Waves");
-			break;
+			AudioStats.StatSoundInfos.Sort([](const FAudioStats::FStatSoundInfo& A, const FAudioStats::FStatSoundInfo& B) { return A.Distance < B.Distance; });
+				SortingName = TEXT("Distance");
+				break;
 		}
-		case FAudioStats::EDisplaySort::Volume:
+			case FAudioStats::EDisplaySort::PlaybackTime:
 		{
-			AudioStats.StatSoundInfos.Sort([](const FAudioStats::FStatSoundInfo& A, const FAudioStats::FStatSoundInfo& B) { return A.Volume > B.Volume; });
-			SortingName = TEXT("Volume");
-			break;
+				AudioStats.StatSoundInfos.Sort([](const FAudioStats::FStatSoundInfo& A, const FAudioStats::FStatSoundInfo& B) { return A.PlaybackTime > B.PlaybackTime; });
+				SortingName = TEXT("Time");
+				break;
+			}
+			case FAudioStats::EDisplaySort::Priority:
+			{
+				AudioStats.StatSoundInfos.Sort([](const FAudioStats::FStatSoundInfo& A, const FAudioStats::FStatSoundInfo& B) { return A.Priority > B.Priority; });
+				SortingName = TEXT("Priority");
+				break;
 		}
-		case FAudioStats::EDisplaySort::Name:
-		default:
+			case FAudioStats::EDisplaySort::Waves:
 		{
-			AudioStats.StatSoundInfos.Sort([](const FAudioStats::FStatSoundInfo& A, const FAudioStats::FStatSoundInfo& B)
+			AudioStats.StatSoundInfos.Sort([](const FAudioStats::FStatSoundInfo& A, const FAudioStats::FStatSoundInfo& B) { return A.WaveInstanceInfos.Num() > B.WaveInstanceInfos.Num(); });
+				SortingName = TEXT("Waves");
+				break;
+			}
+			case FAudioStats::EDisplaySort::Volume:
+			{
+				AudioStats.StatSoundInfos.Sort([](const FAudioStats::FStatSoundInfo& A, const FAudioStats::FStatSoundInfo& B) { return A.Volume > B.Volume; });
+				SortingName = TEXT("Volume");
+				break;
+			}
+			case FAudioStats::EDisplaySort::Name:
+			default:
 			{
-				if (AudioDebugSoundShowPathCVar)
+				AudioStats.StatSoundInfos.Sort([](const FAudioStats::FStatSoundInfo& A, const FAudioStats::FStatSoundInfo& B)
 				{
-					return A.SoundPath < B.SoundPath;
-				}
-				return A.SoundName < B.SoundName;
-			});
-			SortingName = TEXT("Name");
-			break;
+					if (AudioDebugSoundShowPathCVar)
+					{
+						return A.SoundPath < B.SoundPath;
+					}
+					return A.SoundName < B.SoundName;
+				});
+				SortingName = TEXT("Name");
+				break;
+			}
 		}
-	}
-
-	Canvas->DrawShadowedString(X, Y, TEXT("Active Sounds:"), GetStatsFont(), HeaderColor);
-	Y += FontHeight;
 
-	const FString InfoText = FString::Printf(TEXT(" Sorting By: %s, Visualize Attenuation: %s"), *SortingName, bDebug ? TEXT("Enabled") : TEXT("Disabled"));
-	Canvas->DrawShadowedString(X, Y, *InfoText, GetStatsFont(), FColor(128, 255, 128));
-	Y += FontHeight;
+		Canvas->DrawShadowedString(X, Y, TEXT("Active Sounds:"), GetStatsFont(), HeaderColor);
+		Y += FontHeight;
 
-	static int32 FieldLength = 40;
-	static int32 NameHeaderPad = static_cast<int32>(FieldLength * 1.25f);
-	if (AudioStats.DisplaySort == FAudioStats::EDisplaySort::Name || AudioStats.DisplaySort == FAudioStats::EDisplaySort::Distance)
-	{
-		const FString FieldName = FString(TEXT("Distance")).RightPad(NameHeaderPad).Left(NameHeaderPad);
-		Canvas->DrawShadowedString(X, Y, *FString::Printf(TEXT("Index %s Name"), *FieldName), GetStatsFont(), GetBodyColor());
-	}
-	else
-	{
-		const FString FieldName = SortingName.RightPad(FieldLength).Left(FieldLength);
-		Canvas->DrawShadowedString(X, Y, *FString::Printf(TEXT("Index %s Name"), *FieldName), GetStatsFont(), GetBodyColor());
-	}
-	Y += FontHeight;
+		const FString InfoText = FString::Printf(TEXT(" Sorting By: %s, Visualize Attenuation: %s"), *SortingName, bDebug ? TEXT("Enabled") : TEXT("Disabled"));
+		Canvas->DrawShadowedString(X, Y, *InfoText, GetStatsFont(), FColor(128, 255, 128));
+		Y += FontHeight;
 
-	int32 TotalSoundWavesNum = 0;
-	const FColor BodyColor = GetStatSoundBodyColor();
-	for (int32 SoundIndex = 0; SoundIndex < AudioStats.StatSoundInfos.Num(); ++SoundIndex)
-	{
-		const FAudioStats::FStatSoundInfo& StatSoundInfo = AudioStats.StatSoundInfos[SoundIndex];
-		const int32 WaveInstancesNum = StatSoundInfo.WaveInstanceInfos.Num();
-		if (WaveInstancesNum == 0)
+		static int32 FieldLength = 40;
+		static int32 NameHeaderPad = static_cast<int32>(FieldLength * 1.25f);
+		if (AudioStats.DisplaySort == FAudioStats::EDisplaySort::Name || AudioStats.DisplaySort == FAudioStats::EDisplaySort::Distance)
 		{
-			continue;
+			const FString FieldName = FString(TEXT("Distance")).RightPad(NameHeaderPad).Left(NameHeaderPad);
+			Canvas->DrawShadowedString(X, Y, *FString::Printf(TEXT("Index %s Name"), *FieldName), GetStatsFont(), GetBodyColor());
 		}
-		TotalSoundWavesNum += WaveInstancesNum;
-
-		if (SoundIndex >= AudioDebugSoundMaxNumDisplayedCVar)
+		else
 		{
-			continue;
+			const FString FieldName = SortingName.RightPad(FieldLength).Left(FieldLength);
+			Canvas->DrawShadowedString(X, Y, *FString::Printf(TEXT("Index %s Name"), *FieldName), GetStatsFont(), GetBodyColor());
 		}
+		Y += FontHeight;
 
-		bool bDisplayWaves = false;
-		const FString DisplayName = AudioDebugSoundShowPathCVar ? StatSoundInfo.SoundPath : StatSoundInfo.SoundName;
-		FString DebugValue;
-		switch (AudioStats.DisplaySort)
+		int32 TotalSoundWavesNum = 0;
+		const FColor BodyColor = GetStatSoundBodyColor();
+		for (int32 SoundIndex = 0; SoundIndex < AudioStats.StatSoundInfos.Num(); ++SoundIndex)
 		{
-			case FAudioStats::EDisplaySort::Class:
+			const FAudioStats::FStatSoundInfo& StatSoundInfo = AudioStats.StatSoundInfos[SoundIndex];
+			const int32 WaveInstancesNum = StatSoundInfo.WaveInstanceInfos.Num();
+			if (WaveInstancesNum == 0)
 			{
-				DebugValue = *StatSoundInfo.SoundClassName.ToString();
-				break;
+				continue;
 			}
-			case FAudioStats::EDisplaySort::PlaybackTime:
+			TotalSoundWavesNum += WaveInstancesNum;
+
+			if (SoundIndex >= AudioDebugSoundMaxNumDisplayedCVar)
 			{
-				if (FMath::IsNearlyEqual(StatSoundInfo.PlaybackTime, StatSoundInfo.PlaybackTimeNonVirtualized))
+				continue;
+			}
+
+			bool bDisplayWaves = false;
+			const FString DisplayName = AudioDebugSoundShowPathCVar ? StatSoundInfo.SoundPath : StatSoundInfo.SoundName;
+			FString DebugValue;
+			switch (AudioStats.DisplaySort)
+			{
+				case FAudioStats::EDisplaySort::Class:
 				{
-					DebugValue = FString::Printf(TEXT("%08.2f"), StatSoundInfo.PlaybackTime);
+					DebugValue = *StatSoundInfo.SoundClassName.ToString();
+					break;
 				}
-				else
+				case FAudioStats::EDisplaySort::PlaybackTime:
 				{
-					DebugValue = FString::Printf(TEXT("%08.2f (%08.2f Non-Virt)"), StatSoundInfo.PlaybackTime, StatSoundInfo.PlaybackTimeNonVirtualized);
+					if (FMath::IsNearlyEqual(StatSoundInfo.PlaybackTime, StatSoundInfo.PlaybackTimeNonVirtualized))
+					{
+						DebugValue = FString::Printf(TEXT("%08.2f"), StatSoundInfo.PlaybackTime);
+					}
+					else
+					{
+						DebugValue = FString::Printf(TEXT("%08.2f (%08.2f Non-Virt)"), StatSoundInfo.PlaybackTime, StatSoundInfo.PlaybackTimeNonVirtualized);
+					}
+					break;
 				}
-				break;
-			}
-			case FAudioStats::EDisplaySort::Priority:
-			{
-				if (FMath::IsNearlyEqual(StatSoundInfo.Priority, TNumericLimits<float>::Max()))
+				case FAudioStats::EDisplaySort::Priority:
 				{
-					DebugValue += TEXT("Always  ");
+					if (FMath::IsNearlyEqual(StatSoundInfo.Priority, TNumericLimits<float>::Max()))
+					{
+						DebugValue += TEXT("Always  ");
+					}
+					else
+					{
+						DebugValue = FString::Printf(TEXT("%06.2f"), StatSoundInfo.Priority);
+					}
+					break;
 				}
-				else
+				case FAudioStats::EDisplaySort::Volume:
 				{
-					DebugValue = FString::Printf(TEXT("%06.2f"), StatSoundInfo.Priority);
+					DebugValue = FString::Printf(TEXT("%01.2f (%04.2f dB)"), StatSoundInfo.Volume, Audio::ConvertToDecibels(StatSoundInfo.Volume));
+					break;
 				}
-				break;
-			}
-			case FAudioStats::EDisplaySort::Volume:
+				case FAudioStats::EDisplaySort::Waves:
 			{
-				DebugValue = FString::Printf(TEXT("%01.2f (%04.2f dB)"), StatSoundInfo.Volume, Audio::ConvertToDecibels(StatSoundInfo.Volume));
-				break;
-			}
-			case FAudioStats::EDisplaySort::Waves:
-		{
-				DebugValue = FString::Printf(TEXT("%03u"), StatSoundInfo.WaveInstanceInfos.Num());
-				bDisplayWaves = true;
-				break;
-			}
-			case FAudioStats::EDisplaySort::Name:
-			case FAudioStats::EDisplaySort::Distance:
-			default:
+					DebugValue = FString::Printf(TEXT("%03u"), StatSoundInfo.WaveInstanceInfos.Num());
+					bDisplayWaves = true;
+					break;
+				}
+				case FAudioStats::EDisplaySort::Name:
+				case FAudioStats::EDisplaySort::Distance:
+				default:
+				{
+					DebugValue = FString::Printf(TEXT("%08.2f"), StatSoundInfo.Distance);
+					break;
+				}
+				}
+
+			const FString DebugStr = FString::Printf(TEXT("  %03i    %s %s"), SoundIndex, *DebugValue.RightPad(FieldLength - 4).Left(FieldLength - 4), *DisplayName);
+			Canvas->DrawShadowedString(X, Y, *DebugStr, GetStatsFont(), BodyColor);
+			Y += FontHeight;
+
+			if (bDisplayWaves)
 			{
-				DebugValue = FString::Printf(TEXT("%08.2f"), StatSoundInfo.Distance);
-				break;
-			}
+				for (int32 WaveIndex = 0; WaveIndex < WaveInstancesNum; WaveIndex++)
+				{
+					const FString WaveStr = *FString::Printf(TEXT("    %02i    %s"), WaveIndex, *StatSoundInfo.WaveInstanceInfos[WaveIndex].Description);
+					Canvas->DrawShadowedString(X, Y, *WaveStr, GetStatsFont(), FColor(205, 205, 205));
+					Y += FontHeight;
+				}
 			}
+		}
 
-		const FString DebugStr = FString::Printf(TEXT("  %03i    %s %s"), SoundIndex, *DebugValue.RightPad(FieldLength - 4).Left(FieldLength - 4), *DisplayName);
-		Canvas->DrawShadowedString(X, Y, *DebugStr, GetStatsFont(), BodyColor);
+		Canvas->DrawShadowedString(X, Y, *FString::Printf(TEXT("Audio Device ID: %u, Max Sounds Displayed: %i"), AudioDevice->DeviceID, AudioDebugSoundMaxNumDisplayedCVar), GetStatsFont(), HeaderColor);
 		Y += FontHeight;
 
-		if (bDisplayWaves)
+		Canvas->DrawShadowedString(X, Y, *FString::Printf(TEXT("Total Sounds: %i, Sound Waves: %i"), AudioStats.StatSoundInfos.Num(), TotalSoundWavesNum), GetStatsFont(), HeaderColor);
+		Y += FontHeight;
+
+		for (int32 i = 0; i < AudioStats.ListenerTransforms.Num(); ++i)
 		{
-			for (int32 WaveIndex = 0; WaveIndex < WaveInstancesNum; WaveIndex++)
-			{
-				const FString WaveStr = *FString::Printf(TEXT("    %02i    %s"), WaveIndex, *StatSoundInfo.WaveInstanceInfos[WaveIndex].Description);
-				Canvas->DrawShadowedString(X, Y, *WaveStr, GetStatsFont(), FColor(205, 205, 205));
-				Y += FontHeight;
-			}
+			FString LocStr = AudioStats.ListenerTransforms[i].GetLocation().ToString();
+			Canvas->DrawShadowedString(X, Y, *FString::Printf(TEXT("Listener '%d' Position: %s"), i, *LocStr), GetStatsFont(), HeaderColor);
+			Y += FontHeight;
 		}
-	}
-
-	Canvas->DrawShadowedString(X, Y, *FString::Printf(TEXT("Audio Device ID: %u, Max Sounds Displayed: %i"), AudioDevice->DeviceID, AudioDebugSoundMaxNumDisplayedCVar), GetStatsFont(), HeaderColor);
-	Y += FontHeight;
 
-	Canvas->DrawShadowedString(X, Y, *FString::Printf(TEXT("Total Sounds: %i, Sound Waves: %i"), AudioStats.StatSoundInfos.Num(), TotalSoundWavesNum), GetStatsFont(), HeaderColor);
-	Y += FontHeight;
-
-	for (int32 i = 0; i < AudioStats.ListenerTransforms.Num(); ++i)
-	{
-		FString LocStr = AudioStats.ListenerTransforms[i].GetLocation().ToString();
-		Canvas->DrawShadowedString(X, Y, *FString::Printf(TEXT("Listener '%d' Position: %s"), i, *LocStr), GetStatsFont(), HeaderColor);
 		Y += FontHeight;
-	}
-
-	if (!bDebug)
-	{
-		return Y;
-	}
 
-	// Draw sound cue's sphere only in debug.
-	for (const FAudioStats::FStatSoundInfo& StatSoundInfo : AudioStats.StatSoundInfos)
-	{
-		const FTransform& SoundTransform = StatSoundInfo.Transform;
-		const int32 WaveInstancesNum = StatSoundInfo.WaveInstanceInfos.Num();
+		if (!bDebug)
+		{
+			return Y;
+		}
 
-		if (StatSoundInfo.Distance > 100.0f && WaveInstancesNum > 0)
+		// Draw sound cue's sphere only in debug.
+		for (const FAudioStats::FStatSoundInfo& StatSoundInfo : AudioStats.StatSoundInfos)
 		{
-			float SphereRadius = 0.f;
-			float SphereInnerRadius = 0.f;
+			const FTransform& SoundTransform = StatSoundInfo.Transform;
+			const int32 WaveInstancesNum = StatSoundInfo.WaveInstanceInfos.Num();
 
-			if (StatSoundInfo.ShapeDetailsMap.Num() > 0)
+			if (StatSoundInfo.Distance > 100.0f && WaveInstancesNum > 0)
 			{
-				const FString DebugName = AudioDebugSoundShowPathCVar ? StatSoundInfo.SoundPath : StatSoundInfo.SoundName;
-				DrawDebugString(World, SoundTransform.GetTranslation(), *DebugName, nullptr, BodyColor, 0.01f);
+				float SphereRadius = 0.f;
+				float SphereInnerRadius = 0.f;
 
-				for (auto ShapeDetailsIt = StatSoundInfo.ShapeDetailsMap.CreateConstIterator(); ShapeDetailsIt; ++ShapeDetailsIt)
+				if (StatSoundInfo.ShapeDetailsMap.Num() > 0)
 				{
-					const FBaseAttenuationSettings::AttenuationShapeDetails& ShapeDetails = ShapeDetailsIt.Value();
-					switch (ShapeDetailsIt.Key())
-					{
-					case EAttenuationShape::Sphere:
-						if (ShapeDetails.Falloff > 0.f)
-						{
-							DrawDebugSphere(World, SoundTransform.GetTranslation(), ShapeDetails.Extents.X + ShapeDetails.Falloff, 10, FColor(155, 155, 255));
-							DrawDebugSphere(World, SoundTransform.GetTranslation(), ShapeDetails.Extents.X, 10, FColor(55, 55, 255));
-						}
-						else
-						{
-							DrawDebugSphere(World, SoundTransform.GetTranslation(), ShapeDetails.Extents.X, 10, FColor(155, 155, 255));
-						}
-						break;
-
-					case EAttenuationShape::Box:
-						if (ShapeDetails.Falloff > 0.f)
-						{
-							DrawDebugBox(World, SoundTransform.GetTranslation(), ShapeDetails.Extents + FVector(ShapeDetails.Falloff), SoundTransform.GetRotation(), FColor(155, 155, 255));
-							DrawDebugBox(World, SoundTransform.GetTranslation(), ShapeDetails.Extents, SoundTransform.GetRotation(), FColor(55, 55, 255));
-						}
-						else
-						{
-							DrawDebugBox(World, SoundTransform.GetTranslation(), ShapeDetails.Extents, SoundTransform.GetRotation(), FColor(155, 155, 255));
-						}
-						break;
-
-					case EAttenuationShape::Capsule:
-
-						if (ShapeDetails.Falloff > 0.f)
-						{
-							DrawDebugCapsule(World, SoundTransform.GetTranslation(), ShapeDetails.Extents.X + ShapeDetails.Falloff, ShapeDetails.Extents.Y + ShapeDetails.Falloff, SoundTransform.GetRotation(), FColor(155, 155, 255));
-							DrawDebugCapsule(World, SoundTransform.GetTranslation(), ShapeDetails.Extents.X, ShapeDetails.Extents.Y, SoundTransform.GetRotation(), FColor(55, 55, 255));
-						}
-						else
-						{
-							DrawDebugCapsule(World, SoundTransform.GetTranslation(), ShapeDetails.Extents.X, ShapeDetails.Extents.Y, SoundTransform.GetRotation(), FColor(155, 155, 255));
-						}
-						break;
+					const FString DebugName = AudioDebugSoundShowPathCVar ? StatSoundInfo.SoundPath : StatSoundInfo.SoundName;
+					DrawDebugString(World, SoundTransform.GetTranslation(), *DebugName, nullptr, BodyColor, 0.01f);
 
-					case EAttenuationShape::Cone:
+					for (auto ShapeDetailsIt = StatSoundInfo.ShapeDetailsMap.CreateConstIterator(); ShapeDetailsIt; ++ShapeDetailsIt)
 					{
-						const FVector Origin = SoundTransform.GetTranslation() - (SoundTransform.GetUnitAxis(EAxis::X) * ShapeDetails.ConeOffset);
-
-						if (ShapeDetails.Falloff > 0.f || ShapeDetails.Extents.Z > 0.f)
+						const FBaseAttenuationSettings::AttenuationShapeDetails& ShapeDetails = ShapeDetailsIt.Value();
+						switch (ShapeDetailsIt.Key())
 						{
-							const float OuterAngle = FMath::DegreesToRadians(ShapeDetails.Extents.Y + ShapeDetails.Extents.Z);
-							const float InnerAngle = FMath::DegreesToRadians(ShapeDetails.Extents.Y);
-							DrawDebugCone(World, Origin, SoundTransform.GetUnitAxis(EAxis::X), ShapeDetails.Extents.X + ShapeDetails.Falloff + ShapeDetails.ConeOffset, OuterAngle, OuterAngle, 10, FColor(155, 155, 255));
-							DrawDebugCone(World, Origin, SoundTransform.GetUnitAxis(EAxis::X), ShapeDetails.Extents.X + ShapeDetails.ConeOffset, InnerAngle, InnerAngle, 10, FColor(55, 55, 255));
-						}
-						else
+						case EAttenuationShape::Sphere:
+							if (ShapeDetails.Falloff > 0.f)
+							{
+								DrawDebugSphere(World, SoundTransform.GetTranslation(), ShapeDetails.Extents.X + ShapeDetails.Falloff, 10, FColor(155, 155, 255));
+								DrawDebugSphere(World, SoundTransform.GetTranslation(), ShapeDetails.Extents.X, 10, FColor(55, 55, 255));
+							}
+							else
+							{
+								DrawDebugSphere(World, SoundTransform.GetTranslation(), ShapeDetails.Extents.X, 10, FColor(155, 155, 255));
+							}
+							break;
+
+						case EAttenuationShape::Box:
+							if (ShapeDetails.Falloff > 0.f)
+							{
+								DrawDebugBox(World, SoundTransform.GetTranslation(), ShapeDetails.Extents + FVector(ShapeDetails.Falloff), SoundTransform.GetRotation(), FColor(155, 155, 255));
+								DrawDebugBox(World, SoundTransform.GetTranslation(), ShapeDetails.Extents, SoundTransform.GetRotation(), FColor(55, 55, 255));
+							}
+							else
+							{
+								DrawDebugBox(World, SoundTransform.GetTranslation(), ShapeDetails.Extents, SoundTransform.GetRotation(), FColor(155, 155, 255));
+							}
+							break;
+
+						case EAttenuationShape::Capsule:
+
+							if (ShapeDetails.Falloff > 0.f)
+							{
+								DrawDebugCapsule(World, SoundTransform.GetTranslation(), ShapeDetails.Extents.X + ShapeDetails.Falloff, ShapeDetails.Extents.Y + ShapeDetails.Falloff, SoundTransform.GetRotation(), FColor(155, 155, 255));
+								DrawDebugCapsule(World, SoundTransform.GetTranslation(), ShapeDetails.Extents.X, ShapeDetails.Extents.Y, SoundTransform.GetRotation(), FColor(55, 55, 255));
+							}
+							else
+							{
+								DrawDebugCapsule(World, SoundTransform.GetTranslation(), ShapeDetails.Extents.X, ShapeDetails.Extents.Y, SoundTransform.GetRotation(), FColor(155, 155, 255));
+							}
+							break;
+
+						case EAttenuationShape::Cone:
 						{
-							const float Angle = FMath::DegreesToRadians(ShapeDetails.Extents.Y);
-							DrawDebugCone(World, Origin, SoundTransform.GetUnitAxis(EAxis::X), ShapeDetails.Extents.X + ShapeDetails.ConeOffset, Angle, Angle, 10, FColor(155, 155, 255));
+							const FVector Origin = SoundTransform.GetTranslation() - (SoundTransform.GetUnitAxis(EAxis::X) * ShapeDetails.ConeOffset);
+
+							if (ShapeDetails.Falloff > 0.f || ShapeDetails.Extents.Z > 0.f)
+							{
+								const float OuterAngle = FMath::DegreesToRadians(ShapeDetails.Extents.Y + ShapeDetails.Extents.Z);
+								const float InnerAngle = FMath::DegreesToRadians(ShapeDetails.Extents.Y);
+								DrawDebugCone(World, Origin, SoundTransform.GetUnitAxis(EAxis::X), ShapeDetails.Extents.X + ShapeDetails.Falloff + ShapeDetails.ConeOffset, OuterAngle, OuterAngle, 10, FColor(155, 155, 255));
+								DrawDebugCone(World, Origin, SoundTransform.GetUnitAxis(EAxis::X), ShapeDetails.Extents.X + ShapeDetails.ConeOffset, InnerAngle, InnerAngle, 10, FColor(55, 55, 255));
+							}
+							else
+							{
+								const float Angle = FMath::DegreesToRadians(ShapeDetails.Extents.Y);
+								DrawDebugCone(World, Origin, SoundTransform.GetUnitAxis(EAxis::X), ShapeDetails.Extents.X + ShapeDetails.ConeOffset, Angle, Angle, 10, FColor(155, 155, 255));
+							}
+							break;
 						}
-						break;
-					}
 
-					default:
-						check(false);
+						default:
+							check(false);
+						}
 					}
 				}
 			}
 		}
-	}
 
-	return Y;
-}
-
-int32 FAudioDebugger::RenderStatWaves(UWorld* World, FViewport* Viewport, FCanvas* Canvas, int32 X, int32 Y, const FVector* ViewLocation, const FRotator* ViewRotation)
-{
-	FAudioDevice* AudioDevice = GetWorldAudio(World);
-	if (!AudioDevice)
-	{
 		return Y;
 	}
 
-	const int32 FontHeight = GetStatsFont()->GetMaxCharHeight() + 2;
+	int32 FAudioDebugger::RenderStatWaves(UWorld* World, FViewport* Viewport, FCanvas* Canvas, int32 X, int32 Y)
+	{
+		FAudioDevice* AudioDevice = nullptr;
+		if (!DebugShouldRenderStat(World, Canvas, bDebugWavesForAllViewsEnabled, FAudioStats::SoundWaves, &AudioDevice))
+		{
+			return Y;
+		}
+
+		const int32 FontHeight = GetStatsFont()->GetMaxCharHeight() + 2;
 
-	FAudioStats& AudioStats = AudioDeviceStats.FindOrAdd(AudioDevice->DeviceID);
-	Canvas->DrawShadowedString(X, Y, TEXT("Active Sound Waves:"), GetStatsFont(), FLinearColor::Green);
-	Y += TabWidth;
+		FAudioStats& AudioStats = AudioDeviceStats.FindOrAdd(AudioDevice->DeviceID);
+		Canvas->DrawShadowedString(X, Y, TEXT("Active Sound Waves:"), GetStatsFont(), FLinearColor::Green);
+		Y += DebuggerTabWidth;
 
-	using FWaveInstancePair = TPair<const FAudioStats::FStatWaveInstanceInfo*, const FAudioStats::FStatSoundInfo*>;
-	TArray<FWaveInstancePair> WaveInstances;
-	for (const FAudioStats::FStatSoundInfo& StatSoundInfo : AudioStats.StatSoundInfos)
-	{
-		for (const FAudioStats::FStatWaveInstanceInfo& WaveInstanceInfo : StatSoundInfo.WaveInstanceInfos)
+		using FWaveInstancePair = TPair<const FAudioStats::FStatWaveInstanceInfo*, const FAudioStats::FStatSoundInfo*>;
+		TArray<FWaveInstancePair> WaveInstances;
+		for (const FAudioStats::FStatSoundInfo& StatSoundInfo : AudioStats.StatSoundInfos)
 		{
-			if (WaveInstanceInfo.Volume >= MinDisplayVolume || WaveInstanceInfo.bPlayWhenSilent != 0)
+			for (const FAudioStats::FStatWaveInstanceInfo& WaveInstanceInfo : StatSoundInfo.WaveInstanceInfos)
 			{
-				WaveInstances.Emplace(&WaveInstanceInfo, &StatSoundInfo);
+				if (WaveInstanceInfo.Volume >= MinDisplayVolume || WaveInstanceInfo.bPlayWhenSilent != 0)
+				{
+					WaveInstances.Emplace(&WaveInstanceInfo, &StatSoundInfo);
+				}
 			}
 		}
-	}
-
-	WaveInstances.Sort([](const FWaveInstancePair& A, const FWaveInstancePair& B) { return A.Key->InstanceIndex < B.Key->InstanceIndex; });
 
-	const FColor BodyColor = GetBodyColor();
-	for (const FWaveInstancePair& WaveInstanceInfo : WaveInstances)
-	{
-		UAudioComponent* AudioComponent = UAudioComponent::GetAudioComponentFromID(WaveInstanceInfo.Value->AudioComponentID);
-		AActor* SoundOwner = AudioComponent ? AudioComponent->GetOwner() : nullptr;
-
-		FString TheString = *FString::Printf(TEXT("%4i.    %6.2f  %s   Owner: %s   SoundClass: %s"),
-			WaveInstanceInfo.Key->InstanceIndex,
-			WaveInstanceInfo.Key->Volume,
-			*WaveInstanceInfo.Key->WaveInstanceName.ToString(),
-			SoundOwner ? *SoundOwner->GetName() : TEXT("None"),
-			*WaveInstanceInfo.Value->SoundClassName.ToString());
-		Canvas->DrawShadowedString(X, Y, *TheString, GetStatsFont(), WaveInstanceInfo.Key->bPlayWhenSilent == 0 ? BodyColor : FColor::Yellow);
-		Y += FontHeight;
-	}
+		WaveInstances.Sort([](const FWaveInstancePair& A, const FWaveInstancePair& B) { return A.Key->InstanceIndex < B.Key->InstanceIndex; });
 
-	const int32 ActiveInstances = WaveInstances.Num();
+		const FColor BodyColor = GetBodyColor();
+		for (const FWaveInstancePair& WaveInstanceInfo : WaveInstances)
+		{
+			UAudioComponent* AudioComponent = UAudioComponent::GetAudioComponentFromID(WaveInstanceInfo.Value->AudioComponentID);
+			AActor* SoundOwner = AudioComponent ? AudioComponent->GetOwner() : nullptr;
 
-	const int32 Max = AudioDevice->GetMaxChannels() / 2;
-	float f = FMath::Clamp<float>((float)(ActiveInstances - Max) / (float)Max, 0.f, 1.f);
-	const int32 R = FMath::TruncToInt(f * 255);
+			FString TheString = *FString::Printf(TEXT("%4i.    %6.2f  %s   Owner: %s   SoundClass: %s"),
+				WaveInstanceInfo.Key->InstanceIndex,
+				WaveInstanceInfo.Key->Volume,
+				*WaveInstanceInfo.Key->WaveInstanceName.ToString(),
+				SoundOwner ? *SoundOwner->GetName() : TEXT("None"),
+				*WaveInstanceInfo.Value->SoundClassName.ToString());
+			Canvas->DrawShadowedString(X, Y, *TheString, GetStatsFont(), WaveInstanceInfo.Key->bPlayWhenSilent == 0 ? BodyColor : FColor::Yellow);
+			Y += FontHeight;
+		}
 
-	if (ActiveInstances > Max)
-	{
-		f = FMath::Clamp<float>((float)(Max - ActiveInstances) / (float)Max, 0.5f, 1.f);
-	}
-	else
-	{
-		f = 1.0f;
-	}
-	const int32 G = FMath::TruncToInt(f * 255);
-	const int32 B = 0;
+		const int32 ActiveInstances = WaveInstances.Num();
 
-	Canvas->DrawShadowedString(X, Y, *FString::Printf(TEXT(" Total: %i"), ActiveInstances), GetStatsFont(), FColor(R, G, B));
-	Y += FontHeight;
+		const int32 Max = AudioDevice->GetMaxChannels() / 2;
+		float f = FMath::Clamp<float>((float)(ActiveInstances - Max) / (float)Max, 0.f, 1.f);
+		const int32 R = FMath::TruncToInt(f * 255);
 
-	return Y;
-}
+		if (ActiveInstances > Max)
+		{
+			f = FMath::Clamp<float>((float)(Max - ActiveInstances) / (float)Max, 0.5f, 1.f);
+		}
+		else
+		{
+			f = 1.0f;
+		}
+		const int32 G = FMath::TruncToInt(f * 255);
+		const int32 B = 0;
 
-void FAudioDebugger::RemoveDevice(const FAudioDevice& AudioDevice)
-{
-	AudioDeviceStats.Remove(AudioDevice.DeviceID);
-	AudioDeviceStats_AudioThread.Remove(AudioDevice.DeviceID);
-}
+		Canvas->DrawShadowedString(X, Y, *FString::Printf(TEXT(" Total: %i"), ActiveInstances), GetStatsFont(), FColor(R, G, B));
+		Y += 2 * FontHeight;
 
-bool FAudioDebugger::ToggleStats(UWorld* World, const uint8 StatToToggle)
-{
-	if (!GEngine)
-	{
-		return false;
+		return Y;
 	}
 
-	FAudioDevice* AudioDevice = GetWorldAudio(World);
-	if (!AudioDevice)
+	void FAudioDebugger::RemoveDevice(const FAudioDevice& AudioDevice)
 	{
-		return false;
+		AudioDeviceStats.Remove(AudioDevice.DeviceID);
+		AudioDeviceStats_AudioThread.Remove(AudioDevice.DeviceID);
 	}
 
-	if (FAudioDeviceManager* DeviceManager = GEngine->GetAudioDeviceManager())
+	bool FAudioDebugger::ToggleStats(UWorld* World, const uint32 StatToToggle)
 	{
-		DeviceManager->GetDebugger().ToggleStats(AudioDevice->DeviceID, StatToToggle);
-	}
+		if (!GEngine)
+		{
+			return false;
+		}
+
+		FAudioDevice* AudioDevice = GetWorldAudio(World);
+		if (!AudioDevice)
+		{
+			return false;
+		}
 
-	return true;
-}
+		if (FAudioDeviceManager* DeviceManager = GEngine->GetAudioDeviceManager())
+		{
+			DeviceManager->GetDebugger().ToggleStats(AudioDevice->DeviceID, StatToToggle);
+		}
 
-void FAudioDebugger::ToggleStats(const uint32 AudioDeviceHandle, const uint8 StatsToToggle)
-{
-	if (!IsInAudioThread())
+		return true;
+	}
+
+	void FAudioDebugger::ClearStats(const uint32 StatsToClear, UWorld* InWorld)
 	{
-		DECLARE_CYCLE_STAT(TEXT("FAudioThreadTask.ToggleStats"), STAT_AudioToggleStats, STATGROUP_TaskGraphTasks);
+		if (!GEngine)
+		{
+			return;
+		}
 
-		FAudioThread::RunCommandOnAudioThread([this, AudioDeviceHandle, StatsToToggle]()
+		if (FAudioDeviceManager* DeviceManager = GEngine->GetAudioDeviceManager())
 		{
-			ToggleStats(AudioDeviceHandle, StatsToToggle);
-		}, GET_STATID(STAT_AudioToggleStats));
-		return;
+			Audio::FAudioDebugger& DebuggerInstance = DeviceManager->GetDebugger();
+			DeviceManager->IterateOverAllDevices(
+				[&DebuggerInstance, World = InWorld, ClearedStats = StatsToClear](FDeviceId DeviceId, const FAudioDevice* AudioDevice)
+				{
+					if (!World || World->GetAudioDevice().GetDeviceID() == DeviceId)
+					{
+						DebuggerInstance.ClearStats(DeviceId, ClearedStats);
+					}
+				}
+			);
+		}
 	}
 
-	FAudioStats_AudioThread& Stats = AudioDeviceStats_AudioThread.FindOrAdd(AudioDeviceHandle);
-	Stats.RequestedStats ^= StatsToToggle;
-}
+	void FAudioDebugger::ClearStats(FDeviceId DeviceId, const uint32 StatsToClear)
+	{
+		if (IsInGameThread())
+		{
+			FAudioStats& Stats = AudioDeviceStats.FindOrAdd(DeviceId);
+			Stats.EnabledStats &= ~StatsToClear;
+		}
 
-bool FAudioDebugger::ToggleStatCues(UWorld* World, FCommonViewportClient* ViewportClient, const TCHAR* Stream)
-{
-	return ToggleStats(World, ERequestedAudioStats::SoundCues);
-}
+		if (!IsInAudioThread())
+		{
+			DECLARE_CYCLE_STAT(TEXT("FAudioThreadTask.ClearStats"), STAT_AudioClearStats, STATGROUP_TaskGraphTasks);
 
-bool FAudioDebugger::ToggleStatMixes(UWorld* World, FCommonViewportClient* ViewportClient, const TCHAR* Stream)
-{
-	return ToggleStats(World, ERequestedAudioStats::SoundMixes);
-}
+			FAudioThread::RunCommandOnAudioThread([this, DeviceId, StatsToClear]()
+			{
+				ClearStats(DeviceId, StatsToClear);
+			}, GET_STATID(STAT_AudioClearStats));
+			return;
+		}
 
-bool FAudioDebugger::ToggleStatModulators(UWorld* World, FCommonViewportClient* ViewportClient, const TCHAR* Stream)
-{
-	if (!GEngine)
+		FAudioStats_AudioThread& Stats = AudioDeviceStats_AudioThread.FindOrAdd(DeviceId);
+		Stats.RequestedStats &= ~StatsToClear;
+	}
+
+	void FAudioDebugger::SetStats(const uint32 StatsToSet, UWorld* InWorld)
 	{
-		return false;
+		if (!GEngine)
+		{
+			return;
+		}
+
+		if (FAudioDeviceManager* DeviceManager = GEngine->GetAudioDeviceManager())
+		{
+			Audio::FAudioDebugger& DebuggerInstance = DeviceManager->GetDebugger();
+			DeviceManager->IterateOverAllDevices(
+				[&DebuggerInstance, World = InWorld, StatsSet = StatsToSet](FDeviceId DeviceId, const FAudioDevice* AudioDevice)
+				{
+					if (!World || World->GetAudioDevice().GetDeviceID() == DeviceId)
+					{
+						DebuggerInstance.SetStats(DeviceId, StatsSet);
+					}
+				}
+			);
+		}
 	}
 
-	FAudioDevice* AudioDevice = GetWorldAudio(World);
-	if (!AudioDevice)
+	void FAudioDebugger::SetStats(FDeviceId DeviceId, const uint32 StatsToSet)
 	{
-		return false;
+		if (IsInGameThread())
+		{
+			FAudioStats& Stats = AudioDeviceStats.FindOrAdd(DeviceId);
+			Stats.EnabledStats |= StatsToSet;
+		}
+
+		if (!IsInAudioThread())
+		{
+			DECLARE_CYCLE_STAT(TEXT("FAudioThreadTask.SetStats"), STAT_AudioSetStats, STATGROUP_TaskGraphTasks);
+
+			FAudioThread::RunCommandOnAudioThread([this, DeviceId, StatsToSet]()
+			{
+				SetStats(DeviceId, StatsToSet);
+			}, GET_STATID(STAT_AudioSetStats));
+			return;
+		}
+
+		FAudioStats_AudioThread& Stats = AudioDeviceStats_AudioThread.FindOrAdd(DeviceId);
+		Stats.RequestedStats |= StatsToSet;
 	}
 
-	if (AudioDevice->IsModulationPluginEnabled())
+	void FAudioDebugger::ToggleStats(FDeviceId DeviceId, const uint32 StatsToToggle)
 	{
-		if (IAudioModulation* Modulation = AudioDevice->ModulationInterface.Get())
+		if (!IsInAudioThread())
 		{
-			if (!Modulation->OnToggleStat(ViewportClient, Stream))
+			DECLARE_CYCLE_STAT(TEXT("FAudioThreadTask.ToggleStats"), STAT_AudioToggleStats, STATGROUP_TaskGraphTasks);
+
+			FAudioThread::RunCommandOnAudioThread([this, DeviceId, StatsToToggle]()
 			{
-				return false;
-			}
+				ToggleStats(DeviceId, StatsToToggle);
+			}, GET_STATID(STAT_AudioToggleStats));
+			return;
 		}
+
+		FAudioStats_AudioThread& Stats = AudioDeviceStats_AudioThread.FindOrAdd(DeviceId);
+		Stats.RequestedStats ^= StatsToToggle;
 	}
 
-	return true;
-}
+	bool FAudioDebugger::ToggleStatCues(UWorld* World, FCommonViewportClient* ViewportClient, const TCHAR* Stream)
+	{
+		return ToggleStats(World, FAudioStats::SoundCues);
+	}
 
-bool FAudioDebugger::PostStatModulatorHelp(UWorld* World, FCommonViewportClient* ViewportClient, const TCHAR* Stream)
-{
-	// Ignore if all Viewports are closed.
-	if (!ViewportClient)
+		bool FAudioDebugger::ToggleStatMixes(UWorld* World, FCommonViewportClient* ViewportClient, const TCHAR* Stream)
 	{
-		return false;
+		return ToggleStats(World, FAudioStats::SoundMixes);
 	}
 
-	if (FAudioDevice* AudioDevice = World->GetAudioDeviceRaw())
+	bool FAudioDebugger::ToggleStatModulators(UWorld* World, FCommonViewportClient* ViewportClient, const TCHAR* Stream)
 	{
+		if (!GEngine)
+		{
+			return false;
+		}
+
+		FAudioDevice* AudioDevice = GetWorldAudio(World);
+		if (!AudioDevice)
+		{
+			return false;
+		}
+
 		if (AudioDevice->IsModulationPluginEnabled())
 		{
 			if (IAudioModulation* Modulation = AudioDevice->ModulationInterface.Get())
 			{
-				if (!Modulation->OnPostHelp(ViewportClient, Stream))
+				if (!Modulation->OnToggleStat(ViewportClient, Stream))
 				{
 					return false;
 				}
 			}
 		}
-	}
-
-	return true;
-}
 
-bool FAudioDebugger::ToggleStatSounds(UWorld* World, FCommonViewportClient* ViewportClient, const TCHAR* Stream)
-{
-	// Ignore if all Viewports are closed.
-	if (!ViewportClient)
-	{
-		return false;
+		return true;
 	}
 
-	if (!ToggleStats(World, ERequestedAudioStats::Sounds))
+	bool FAudioDebugger::PostStatModulatorHelp(UWorld* World, FCommonViewportClient* ViewportClient, const TCHAR* Stream)
 	{
-		return false;
-	}
+		// Ignore if all Viewports are closed.
+		if (!ViewportClient)
+		{
+			return false;
+		}
 
-	const bool bHelp = Stream ? FCString::Stristr(Stream, TEXT("?")) != nullptr : false;
-	if (bHelp)
-	{
-		GLog->Logf(TEXT("Shows all active sounds. Displays value sorted by when sort is set"));
-		GLog->Logf(TEXT("  stat sounds sort=class|distance|name|priority|time|volume|waves"));
-		GLog->Logf(TEXT("      distance - sort list by distance to player"));
-		GLog->Logf(TEXT("      class - sort by sound class name"));
-		GLog->Logf(TEXT("      name - sort by cue pathname (default)"));
-		GLog->Logf(TEXT("      time - sort by time played back"));
-		GLog->Logf(TEXT("      waves - sort by waves' num"));
-		GLog->Logf(TEXT("      priority - sort by playback priority"));
-		GLog->Logf(TEXT("      volume - sort by volume"));
-		GLog->Logf(TEXT("  stat sounds -debug - enables debugging mode like showing sound radius sphere and names, but only for sounds with enabled 'Debug' property"));
-		GLog->Logf(TEXT(""));
-		GLog->Logf(TEXT("Ex. stat sounds sort=class -debug"));
-		GLog->Logf(TEXT(" This will show only sounds that have 'Debug' property set to true and will sort them by and display their sound class"));
+		if (FAudioDevice* AudioDevice = World->GetAudioDeviceRaw())
+		{
+			if (AudioDevice->IsModulationPluginEnabled())
+			{
+				if (IAudioModulation* Modulation = AudioDevice->ModulationInterface.Get())
+				{
+					if (!Modulation->OnPostHelp(ViewportClient, Stream))
+					{
+						return false;
+					}
+				}
+			}
+		}
+
+		return true;
 	}
 
-	uint8 ShowSounds = 0;
-	FAudioStats::EDisplaySort DisplaySort = FAudioStats::EDisplaySort::Name;
-	if (Stream)
+	bool FAudioDebugger::ToggleStatSounds(UWorld* World, FCommonViewportClient* ViewportClient, const TCHAR* Stream)
 	{
-		const bool bHide = FParse::Command(&Stream, TEXT("off"));
-		if (!bHide)
+		// Ignore if all Viewports are closed.
+		if (!ViewportClient)
 		{
-			const bool bDebug = FParse::Param(Stream, TEXT("debug"));
-			if (bDebug)
-			{
-				ShowSounds |= static_cast<uint8>(FAudioStats::EDisplayFlags::Debug);
-			}
+			return false;
+		}
 
-			const bool bLongNames = FParse::Param(Stream, TEXT("longnames"));
-			if (bLongNames)
-			{
-				ShowSounds |= static_cast<uint8>(FAudioStats::EDisplayFlags::Long_Names);
-			}
+		if (!ToggleStats(World, FAudioStats::Sounds))
+		{
+			return false;
+		}
 
-			FString SortStr;
-			if (FParse::Value(Stream, TEXT("sort="), SortStr))
-			{
-				AudioDebugSoundSortCVarCVar = SortStr;
-			}
+		const bool bHelp = Stream ? FCString::Stristr(Stream, TEXT("?")) != nullptr : false;
+		if (bHelp)
+		{
+			GLog->Logf(TEXT("Shows all active sounds. Displays value sorted by when sort is set"));
+			GLog->Logf(TEXT("  stat sounds sort=class|distance|name|priority|time|volume|waves"));
+			GLog->Logf(TEXT("      distance - sort list by distance to player"));
+			GLog->Logf(TEXT("      class - sort by sound class name"));
+			GLog->Logf(TEXT("      name - sort by cue pathname (default)"));
+			GLog->Logf(TEXT("      time - sort by time played back"));
+			GLog->Logf(TEXT("      waves - sort by waves' num"));
+			GLog->Logf(TEXT("      priority - sort by playback priority"));
+			GLog->Logf(TEXT("      volume - sort by volume"));
+			GLog->Logf(TEXT("  stat sounds -debug - enables debugging mode like showing sound radius sphere and names, but only for sounds with enabled 'Debug' property"));
+			GLog->Logf(TEXT(""));
+			GLog->Logf(TEXT("Ex. stat sounds sort=class -debug"));
+			GLog->Logf(TEXT(" This will show only sounds that have 'Debug' property set to true and will sort them by and display their sound class"));
 		}
-	}
 
-	FAudioDevice* AudioDevice = GetWorldAudio(World);
-	check(AudioDevice);
-	FAudioStats& Stats = AudioDeviceStats.FindOrAdd(AudioDevice->DeviceID);
-	Stats.DisplayFlags = ShowSounds;
-	Stats.DisplaySort = DisplaySort;
+		uint8 ShowSounds = 0;
+		FAudioStats::EDisplaySort DisplaySort = FAudioStats::EDisplaySort::Name;
+		if (Stream)
+		{
+			const bool bHide = FParse::Command(&Stream, TEXT("off"));
+			if (!bHide)
+			{
+				const bool bDebug = FParse::Param(Stream, TEXT("debug"));
+				if (bDebug)
+				{
+					ShowSounds |= static_cast<uint8>(FAudioStats::EDisplayFlags::Debug);
+				}
 
-	ResolveDesiredStats(ViewportClient);
+				const bool bLongNames = FParse::Param(Stream, TEXT("longnames"));
+				if (bLongNames)
+				{
+					ShowSounds |= static_cast<uint8>(FAudioStats::EDisplayFlags::Long_Names);
+				}
+
+				FString SortStr;
+				if (FParse::Value(Stream, TEXT("sort="), SortStr))
+				{
+					AudioDebugSoundSortCVarCVar = SortStr;
+				}
+			}
+		}
 
-	return true;
-}
+		FAudioDevice* AudioDevice = GetWorldAudio(World);
+		check(AudioDevice);
+		FAudioStats& Stats = AudioDeviceStats.FindOrAdd(AudioDevice->DeviceID);
+		Stats.DisplayFlags = ShowSounds;
+		Stats.DisplaySort = DisplaySort;
 
-bool FAudioDebugger::ToggleStatWaves(UWorld* World, FCommonViewportClient* ViewportClient, const TCHAR* Stream)
-{
-	return ToggleStats(World, ERequestedAudioStats::SoundWaves);
-}
+		ResolveDesiredStats(ViewportClient);
 
-void FAudioDebugger::SendUpdateResultsToGameThread(const FAudioDevice& AudioDevice, const int32 FirstActiveIndex)
-{
-	check(IsInAudioThread());
+		return true;
+	}
 
-	FAudioStats_AudioThread* Stats_AudioThread = AudioDeviceStats_AudioThread.Find(AudioDevice.DeviceID);
-	if (!Stats_AudioThread)
+	bool FAudioDebugger::ToggleStatWaves(UWorld* World, FCommonViewportClient* ViewportClient, const TCHAR* Stream)
 	{
-		return;
+		return ToggleStats(World, FAudioStats::SoundWaves);
 	}
 
-	TArray<FAudioStats::FStatSoundInfo> StatSoundInfos;
-	TArray<FAudioStats::FStatSoundMix> StatSoundMixes;
+	void FAudioDebugger::SendUpdateResultsToGameThread(const FAudioDevice& AudioDevice, const int32 FirstActiveIndex)
+	{
+		check(IsInAudioThread());
+
+		FAudioStats_AudioThread* Stats_AudioThread = AudioDeviceStats_AudioThread.Find(AudioDevice.DeviceID);
+		if (!Stats_AudioThread)
+		{
+			return;
+		}
+
+		TArray<FAudioStats::FStatSoundInfo> StatSoundInfos;
+		TArray<FAudioStats::FStatSoundMix> StatSoundMixes;
 
-	const uint8 RequestedStats = Stats_AudioThread->RequestedStats;
-	TMap<FActiveSound*, int32> ActiveSoundToInfoIndex;
+		const uint32 RequestedStats = Stats_AudioThread->RequestedStats;
+		TMap<FActiveSound*, int32> ActiveSoundToInfoIndex;
 	
-	const bool bDebug = (RequestedStats & ERequestedAudioStats::DebugSounds) != 0;
-
-	static const uint8 SoundMask =
-		ERequestedAudioStats::Sounds |
-		ERequestedAudioStats::SoundCues |
-		ERequestedAudioStats::SoundMixes |
-		ERequestedAudioStats::SoundWaves;
-	if (RequestedStats & SoundMask)
-	{
-		for (FActiveSound* ActiveSound : AudioDevice.GetActiveSounds())
+		const bool bDebug = (RequestedStats & FAudioStats::DebugSounds) != 0;
+
+		static const uint32 SoundMask =
+			FAudioStats::Sounds |
+			FAudioStats::SoundCues |
+			FAudioStats::SoundMixes |
+			FAudioStats::SoundWaves;
+		if (RequestedStats & SoundMask)
 		{
-			if (USoundBase* SoundBase = ActiveSound->GetSound())
-			{				
-				if (!bDebug || ActiveSound->GetSound()->bDebug)
-				{
-					ActiveSoundToInfoIndex.Add(ActiveSound, StatSoundInfos.AddDefaulted());
-					FAudioStats::FStatSoundInfo& StatSoundInfo = StatSoundInfos.Last();
-					StatSoundInfo.SoundName = SoundBase->GetName();
-					StatSoundInfo.SoundPath = SoundBase->GetPathName();
-					StatSoundInfo.Distance = AudioDevice.GetDistanceToNearestListener(ActiveSound->Transform.GetTranslation());
-					StatSoundInfo.PlaybackTime = ActiveSound->PlaybackTime;
-					StatSoundInfo.Priority = ActiveSound->GetHighestPriority();
-					StatSoundInfo.PlaybackTimeNonVirtualized = ActiveSound->PlaybackTimeNonVirtualized;
-					
-					StatSoundInfo.Volume = 0.0f;
-					for (const TPair<UPTRINT, FWaveInstance*>& Pair : ActiveSound->GetWaveInstances())
+			for (FActiveSound* ActiveSound : AudioDevice.GetActiveSounds())
+			{
+				if (USoundBase* SoundBase = ActiveSound->GetSound())
+				{				
+					if (!bDebug || ActiveSound->GetSound()->bDebug)
 					{
-						StatSoundInfo.Volume = FMath::Max(StatSoundInfo.Volume, Pair.Value->GetVolumeWithDistanceAttenuation() * Pair.Value->GetDynamicVolume());
-					}
+						ActiveSoundToInfoIndex.Add(ActiveSound, StatSoundInfos.AddDefaulted());
+						FAudioStats::FStatSoundInfo& StatSoundInfo = StatSoundInfos.Last();
+						StatSoundInfo.SoundName = SoundBase->GetName();
+						StatSoundInfo.SoundPath = SoundBase->GetPathName();
+						StatSoundInfo.Distance = AudioDevice.GetDistanceToNearestListener(ActiveSound->Transform.GetTranslation());
+						StatSoundInfo.PlaybackTime = ActiveSound->PlaybackTime;
+						StatSoundInfo.Priority = ActiveSound->GetHighestPriority();
+						StatSoundInfo.PlaybackTimeNonVirtualized = ActiveSound->PlaybackTimeNonVirtualized;
+					
+						StatSoundInfo.Volume = 0.0f;
+						for (const TPair<UPTRINT, FWaveInstance*>& Pair : ActiveSound->GetWaveInstances())
+						{
+							StatSoundInfo.Volume = FMath::Max(StatSoundInfo.Volume, Pair.Value->GetVolumeWithDistanceAttenuation() * Pair.Value->GetDynamicVolume());
+						}
 
-					if (USoundClass* SoundClass = ActiveSound->GetSoundClass())
-					{
-						StatSoundInfo.SoundClassName = SoundClass->GetFName();
+						if (USoundClass* SoundClass = ActiveSound->GetSoundClass())
+						{
+							StatSoundInfo.SoundClassName = SoundClass->GetFName();
 
-					}
-					else
-					{
-						StatSoundInfo.SoundClassName = NAME_None;
-					}
+						}
+						else
+						{
+							StatSoundInfo.SoundClassName = NAME_None;
+						}
 
-					StatSoundInfo.Transform = ActiveSound->Transform;
-					StatSoundInfo.AudioComponentID = ActiveSound->GetAudioComponentID();
+						StatSoundInfo.Transform = ActiveSound->Transform;
+						StatSoundInfo.AudioComponentID = ActiveSound->GetAudioComponentID();
 
-					if (bDebug && ActiveSound->GetSound()->bDebug)
-					{
-						ActiveSound->CollectAttenuationShapesForVisualization(StatSoundInfo.ShapeDetailsMap);
+						if (bDebug && ActiveSound->GetSound()->bDebug)
+						{
+							ActiveSound->CollectAttenuationShapesForVisualization(StatSoundInfo.ShapeDetailsMap);
+						}
 					}
 				}
 			}
-		}
 
-		// Iterate through all wave instances.
-		const TArray<FWaveInstance*>& WaveInstances = AudioDevice.GetActiveWaveInstances();
-		auto WaveInstanceSourceMap = AudioDevice.GetWaveInstanceSourceMap();
-		for (int32 InstanceIndex = FirstActiveIndex; InstanceIndex < WaveInstances.Num(); ++InstanceIndex)
-		{
-			const FWaveInstance* WaveInstance = WaveInstances[InstanceIndex];
-			int32* SoundInfoIndex = ActiveSoundToInfoIndex.Find(WaveInstance->ActiveSound);
-			if (SoundInfoIndex)
-			{
-				FAudioStats::FStatWaveInstanceInfo WaveInstanceInfo;
-				FSoundSource* Source = WaveInstanceSourceMap.FindRef(WaveInstance);
-				WaveInstanceInfo.Description = Source ? Source->Describe((RequestedStats & ERequestedAudioStats::LongSoundNames) != 0) : FString(TEXT("No source"));
-				WaveInstanceInfo.Volume = WaveInstance->GetVolumeWithDistanceAttenuation() * WaveInstance->GetDynamicVolume();
-				WaveInstanceInfo.InstanceIndex = InstanceIndex;
-				WaveInstanceInfo.WaveInstanceName = *WaveInstance->GetName();
-				WaveInstanceInfo.bPlayWhenSilent = WaveInstance->ActiveSound->IsPlayWhenSilent() ? 1 : 0;
-				WaveInstanceInfo.DebugInfo = Source ? Source->DebugInfo : WaveInstanceInfo.DebugInfo;
-				StatSoundInfos[*SoundInfoIndex].WaveInstanceInfos.Add(MoveTemp(WaveInstanceInfo));
+			// Iterate through all wave instances.
+			const TArray<FWaveInstance*>& WaveInstances = AudioDevice.GetActiveWaveInstances();
+			auto WaveInstanceSourceMap = AudioDevice.GetWaveInstanceSourceMap();
+			for (int32 InstanceIndex = FirstActiveIndex; InstanceIndex < WaveInstances.Num(); ++InstanceIndex)
+			{
+				const FWaveInstance* WaveInstance = WaveInstances[InstanceIndex];
+				int32* SoundInfoIndex = ActiveSoundToInfoIndex.Find(WaveInstance->ActiveSound);
+				if (SoundInfoIndex)
+				{
+					FAudioStats::FStatWaveInstanceInfo WaveInstanceInfo;
+					FSoundSource* Source = WaveInstanceSourceMap.FindRef(WaveInstance);
+					WaveInstanceInfo.Description = Source ? Source->Describe((RequestedStats & FAudioStats::LongSoundNames) != 0) : FString(TEXT("No source"));
+					WaveInstanceInfo.Volume = WaveInstance->GetVolumeWithDistanceAttenuation() * WaveInstance->GetDynamicVolume();
+					WaveInstanceInfo.InstanceIndex = InstanceIndex;
+					WaveInstanceInfo.WaveInstanceName = *WaveInstance->GetName();
+					WaveInstanceInfo.bPlayWhenSilent = WaveInstance->ActiveSound->IsPlayWhenSilent() ? 1 : 0;
+					WaveInstanceInfo.DebugInfo = Source ? Source->DebugInfo : WaveInstanceInfo.DebugInfo;
+					StatSoundInfos[*SoundInfoIndex].WaveInstanceInfos.Add(MoveTemp(WaveInstanceInfo));
+				}
 			}
 		}
-	}
 
-	if (RequestedStats & ERequestedAudioStats::SoundMixes)
-	{
-		if (const FAudioEffectsManager* Effects = AudioDevice.GetEffects())
+		if (RequestedStats & FAudioStats::SoundMixes)
 		{
-			const USoundMix* CurrentEQMix = Effects->GetCurrentEQMix();
-
-			for (const TPair<USoundMix*, FSoundMixState>& SoundMixPair : AudioDevice.GetSoundMixModifiers())
+			if (const FAudioEffectsManager* Effects = AudioDevice.GetEffects())
 			{
-				StatSoundMixes.AddDefaulted();
-				FAudioStats::FStatSoundMix& StatSoundMix = StatSoundMixes.Last();
-				StatSoundMix.MixName = SoundMixPair.Key->GetName();
-				StatSoundMix.InterpValue = SoundMixPair.Value.InterpValue;
-				StatSoundMix.RefCount = SoundMixPair.Value.ActiveRefCount + SoundMixPair.Value.PassiveRefCount;
-				StatSoundMix.bIsCurrentEQ = (SoundMixPair.Key == CurrentEQMix);
+				const USoundMix* CurrentEQMix = Effects->GetCurrentEQMix();
+
+				for (const TPair<USoundMix*, FSoundMixState>& SoundMixPair : AudioDevice.GetSoundMixModifiers())
+				{
+					StatSoundMixes.AddDefaulted();
+					FAudioStats::FStatSoundMix& StatSoundMix = StatSoundMixes.Last();
+					StatSoundMix.MixName = SoundMixPair.Key->GetName();
+					StatSoundMix.InterpValue = SoundMixPair.Value.InterpValue;
+					StatSoundMix.RefCount = SoundMixPair.Value.ActiveRefCount + SoundMixPair.Value.PassiveRefCount;
+					StatSoundMix.bIsCurrentEQ = (SoundMixPair.Key == CurrentEQMix);
+				}
 			}
 		}
-	}
 
-	DECLARE_CYCLE_STAT(TEXT("FGameThreadAudioTask.AudioSendResults"), STAT_AudioSendResults, STATGROUP_TaskGraphTasks);
+		DECLARE_CYCLE_STAT(TEXT("FGameThreadAudioTask.AudioSendResults"), STAT_AudioSendResults, STATGROUP_TaskGraphTasks);
 
-	const uint32 AudioDeviceID = AudioDevice.DeviceID;
+		const uint32 AudioDeviceID = AudioDevice.DeviceID;
 
-	TArray<FTransform> ListenerTransforms;
-	for (const FListener& Listener : AudioDevice.GetListeners())
-	{
-		ListenerTransforms.Add(Listener.Transform);
+		TArray<FTransform> ListenerTransforms;
+		for (const FListener& Listener : AudioDevice.GetListeners())
+		{
+			ListenerTransforms.Add(Listener.Transform);
+		}
+		FAudioThread::RunCommandOnGameThread([AudioDeviceID, ListenerTransforms, StatSoundInfos, StatSoundMixes]()
+		{
+			FAudioStats& Stats = AudioDeviceStats.FindOrAdd(AudioDeviceID);
+			Stats.ListenerTransforms = ListenerTransforms;
+			Stats.StatSoundInfos = StatSoundInfos;
+			Stats.StatSoundMixes = StatSoundMixes;
+		}, GET_STATID(STAT_AudioSendResults));
 	}
-	FAudioThread::RunCommandOnGameThread([AudioDeviceID, ListenerTransforms, StatSoundInfos, StatSoundMixes]()
-	{
-		FAudioStats& Stats = AudioDeviceStats.FindOrAdd(AudioDeviceID);
-		Stats.ListenerTransforms = ListenerTransforms;
-		Stats.StatSoundInfos = StatSoundInfos;
-		Stats.StatSoundMixes = StatSoundMixes;
-	}, GET_STATID(STAT_AudioSendResults));
-}
-
-void FAudioDebugger::UpdateAudibleInactiveSounds(const uint32 FirstActiveIndex, const TArray<FWaveInstance*>& WaveInstances)
-{
-#if STATS
-	uint32 AudibleInactiveSounds = 0;
-	// Count how many sounds are not being played but were audible
-	for (uint32 InstanceIndex = 0; InstanceIndex < FirstActiveIndex; ++InstanceIndex)
+
+	void FAudioDebugger::UpdateAudibleInactiveSounds(const uint32 FirstActiveIndex, const TArray<FWaveInstance*>& WaveInstances)
 	{
-		const FWaveInstance* WaveInstance = WaveInstances[InstanceIndex];
-		const float WaveInstanceVol = WaveInstance->GetVolumeWithDistanceAttenuation() * WaveInstance->GetDynamicVolume();
-		if (WaveInstanceVol > MinDisplayVolume)
+	#if STATS
+		uint32 AudibleInactiveSounds = 0;
+		// Count how many sounds are not being played but were audible
+		for (uint32 InstanceIndex = 0; InstanceIndex < FirstActiveIndex; ++InstanceIndex)
 		{
-			AudibleInactiveSounds++;
+			const FWaveInstance* WaveInstance = WaveInstances[InstanceIndex];
+			const float WaveInstanceVol = WaveInstance->GetVolumeWithDistanceAttenuation() * WaveInstance->GetDynamicVolume();
+			if (WaveInstanceVol > MinDisplayVolume)
+			{
+				AudibleInactiveSounds++;
+			}
 		}
+		SET_DWORD_STAT(STAT_AudibleWavesDroppedDueToPriority, AudibleInactiveSounds);
+	#endif
 	}
-	SET_DWORD_STAT(STAT_AudibleWavesDroppedDueToPriority, AudibleInactiveSounds);
-#endif
-}
-
-void FAudioDebugger::ClearMutesAndSolos()
-{
-	DebugNames.MuteSoundClass.Empty();
-	DebugNames.MuteSoundCue.Empty();
-	DebugNames.MuteSoundWave.Empty();
-	DebugNames.SoloSoundClass.Empty();
-	DebugNames.SoloSoundCue.Empty();
-	DebugNames.SoloSoundWave.Empty();
-}
 
+	void FAudioDebugger::ClearMutesAndSolos()
+	{
+		DebugNames.MuteSoundClass.Empty();
+		DebugNames.MuteSoundCue.Empty();
+		DebugNames.MuteSoundWave.Empty();
+		DebugNames.SoloSoundClass.Empty();
+		DebugNames.SoloSoundCue.Empty();
+		DebugNames.SoloSoundWave.Empty();
+	}
+} // namespace Audio
 #endif // ENABLE_AUDIO_DEBUG
\ No newline at end of file
diff --git a/Engine/Source/Runtime/Engine/Private/AudioDevice.cpp b/Engine/Source/Runtime/Engine/Private/AudioDevice.cpp
--- a/Engine/Source/Runtime/Engine/Private/AudioDevice.cpp
+++ b/Engine/Source/Runtime/Engine/Private/AudioDevice.cpp
@@ -722,7 +722,7 @@ void FAudioDevice::Teardown()
 	PluginListeners.Reset();
 
 #if ENABLE_AUDIO_DEBUG
-	FAudioDebugger::RemoveDevice(*this);
+	Audio::FAudioDebugger::RemoveDevice(*this);
 #endif // ENABLE_AUDIO_DEBUG
 }
 
@@ -1403,19 +1403,19 @@ void FAudioDevice::HandleAudioSoloCommon(const TCHAR* Cmd, FOutputDevice& Ar, FT
 
 bool FAudioDevice::HandleAudioSoloSoundClass(const TCHAR* Cmd, FOutputDevice& Ar)
 {
-	HandleAudioSoloCommon(Cmd, Ar, &FAudioDebugger::ToggleSoloSoundClass);
+	HandleAudioSoloCommon(Cmd, Ar, &Audio::FAudioDebugger::ToggleSoloSoundClass);
 	return true;
 }
 
 bool FAudioDevice::HandleAudioSoloSoundWave(const TCHAR* Cmd, FOutputDevice& Ar)
 {	
-	HandleAudioSoloCommon(Cmd, Ar, &FAudioDebugger::ToggleSoloSoundWave);
+	HandleAudioSoloCommon(Cmd, Ar, &Audio::FAudioDebugger::ToggleSoloSoundWave);
 	return true;
 }
 
 bool FAudioDevice::HandleAudioSoloSoundCue(const TCHAR* Cmd, FOutputDevice& Ar)
 {
-	HandleAudioSoloCommon(Cmd, Ar, &FAudioDebugger::ToggleSoloSoundCue);
+	HandleAudioSoloCommon(Cmd, Ar, &Audio::FAudioDebugger::ToggleSoloSoundCue);
 	return true;
 }
 
@@ -3931,7 +3931,7 @@ void FAudioDevice::StopSources(TArray<FWaveInstance*>& WaveInstances, int32 Firs
 			}
 
 #if ENABLE_AUDIO_DEBUG
-			FAudioDebugger::DrawDebugInfo(*Source);
+			Audio::FAudioDebugger::DrawDebugInfo(*Source);
 #endif // ENABLE_AUDIO_DEBUG
 		}
 	}
@@ -3946,7 +3946,7 @@ void FAudioDevice::StopSources(TArray<FWaveInstance*>& WaveInstances, int32 Firs
 	}
 
 #if ENABLE_AUDIO_DEBUG
-	FAudioDebugger::UpdateAudibleInactiveSounds(FirstActiveIndex, WaveInstances);
+	Audio::FAudioDebugger::UpdateAudibleInactiveSounds(FirstActiveIndex, WaveInstances);
 #endif // ENABLE_AUDIO_DEBUG
 }
 
@@ -4144,6 +4144,23 @@ void FAudioDevice::Update(bool bGameTicking)
 	{
 		// Make sure our referenced sound waves is up-to-date
 		UpdateReferencedSoundWaves();
+
+#if ENABLE_AUDIO_DEBUG
+		if (GEngine)
+		{
+			if (FAudioDeviceManager* DeviceManager = GEngine->GetAudioDeviceManager())
+			{
+				TArray<UWorld*> Worlds = DeviceManager->GetWorldsUsingAudioDevice(DeviceID);
+				for (UWorld* World : Worlds)
+				{
+					if (World)
+					{
+						Audio::FAudioDebugger::DrawDebugStats(*World);
+					}
+				}
+			}
+		}
+#endif // ENABLE_AUDIO_DEBUG
 	}
 
 	if (!IsInAudioThread())
@@ -4260,17 +4277,18 @@ void FAudioDevice::Update(bool bGameTicking)
 	}
 
 #if ENABLE_AUDIO_DEBUG
-	for (FActiveSound* ActiveSound : ActiveSounds)
+	if (GEngine)
 	{
-		if (!ActiveSound)
-		{
-			continue;
-		}
-		
-		if (UWorld* World = ActiveSound->GetWorld())
+		if (FAudioDeviceManager* DeviceManager = GEngine->GetAudioDeviceManager())
 		{
-			FAudioDebugger::DrawDebugInfo(*World, Listeners);
-			break;
+			TArray<UWorld*> Worlds = DeviceManager->GetWorldsUsingAudioDevice(DeviceID);
+			for (UWorld* World : Worlds)
+			{
+				if (World)
+				{
+					Audio::FAudioDebugger::DrawDebugInfo(*World, Listeners);
+				}
+			}
 		}
 	}
 #endif // ENABLE_AUDIO_DEBUG
@@ -4381,7 +4399,7 @@ void FAudioDevice::SendUpdateResultsToGameThread(const int32 FirstActiveIndex)
 	}, GET_STATID(STAT_AudioSendResults));
 
 #if ENABLE_AUDIO_DEBUG
-	FAudioDebugger::SendUpdateResultsToGameThread(*this, FirstActiveIndex);
+	Audio::FAudioDebugger::SendUpdateResultsToGameThread(*this, FirstActiveIndex);
 #endif // ENABLE_AUDIO_DEBUG
 }
 
@@ -6362,7 +6380,7 @@ void FAudioDevice::OnBeginPIE(const bool bIsSimulating)
 	}
 
 #if ENABLE_AUDIO_DEBUG
-	FAudioDebugger::OnBeginPIE();
+	Audio::FAudioDebugger::OnBeginPIE();
 #endif // ENABLE_AUDIO_DEBUG
 }
 
@@ -6375,7 +6393,7 @@ void FAudioDevice::OnEndPIE(const bool bIsSimulating)
 	}
 
 #if ENABLE_AUDIO_DEBUG
-	FAudioDebugger::OnEndPIE();
+	Audio::FAudioDebugger::OnEndPIE();
 #endif // ENABLE_AUDIO_DEBUG
 }
 #endif // WITH_EDITOR
diff --git a/Engine/Source/Runtime/Engine/Private/AudioDeviceManager.cpp b/Engine/Source/Runtime/Engine/Private/AudioDeviceManager.cpp
--- a/Engine/Source/Runtime/Engine/Private/AudioDeviceManager.cpp
+++ b/Engine/Source/Runtime/Engine/Private/AudioDeviceManager.cpp
@@ -134,7 +134,7 @@ FAudioDeviceManager::FAudioDeviceManager()
 {
 
 #if ENABLE_AUDIO_DEBUG
-	AudioDebugger = TUniquePtr<FAudioDebugger>(new FAudioDebugger());
+	AudioDebugger = MakeUnique<Audio::FAudioDebugger>();
 
 	// Check for a command line debug sound argument.
 	FString DebugSound;
@@ -381,13 +381,7 @@ FAudioDeviceHandle FAudioDeviceManager::RequestAudioDevice(const FAudioDevicePar
 		{
 			if (CanUseAudioDevice(InParams, Device.Value))
 			{
-				if (InParams.AssociatedWorld != nullptr)
-				{
-					
-					Device.Value.WorldsUsingThisDevice.AddUnique(InParams.AssociatedWorld);
-					FAudioDeviceManagerDelegates::OnWorldRegisteredToAudioDevice.Broadcast(InParams.AssociatedWorld, Device.Key);
-				}
-
+				RegisterWorld(InParams.AssociatedWorld, Device.Key);
 				return BuildNewHandle(Device.Value, Device.Key, InParams);
 			}
 		}
@@ -397,6 +391,40 @@ FAudioDeviceHandle FAudioDeviceManager::RequestAudioDevice(const FAudioDevicePar
 	}
 }
 
+void FAudioDeviceManager::RegisterWorld(UWorld* InWorld, Audio::FDeviceId DeviceId)
+{
+	if (!InWorld)
+	{
+		return;
+	}
+
+	if (FAudioDeviceContainer* DeviceContainer = Devices.Find(DeviceId))
+	{
+		if (!DeviceContainer->WorldsUsingThisDevice.Contains(InWorld))
+		{
+			DeviceContainer->WorldsUsingThisDevice.AddUnique(InWorld);
+			FAudioDeviceManagerDelegates::OnWorldRegisteredToAudioDevice.Broadcast(InWorld, DeviceId);
+		}
+	}
+}
+
+void FAudioDeviceManager::UnregisterWorld(UWorld* InWorld, Audio::FDeviceId DeviceId)
+{
+	if (!InWorld)
+	{
+		return;
+	}
+
+	if (FAudioDeviceContainer* DeviceContainer = Devices.Find(DeviceId))
+	{
+		if (DeviceContainer->WorldsUsingThisDevice.Contains(InWorld))
+		{
+			DeviceContainer->WorldsUsingThisDevice.Remove(InWorld);
+			FAudioDeviceManagerDelegates::OnWorldUnregisteredWithAudioDevice.Broadcast(InWorld, DeviceId);
+		}
+	}
+}
+
 bool FAudioDeviceManager::InitializeManager()
 {
 	if (LoadDefaultAudioDeviceModule())
@@ -421,6 +449,7 @@ bool FAudioDeviceManager::InitializeManager()
 		FAudioDeviceParams MainDeviceParams;
 		MainDeviceParams.Scope = EAudioDeviceScope::Shared;
 		MainDeviceParams.bIsNonRealtime = false;
+		MainDeviceParams.AssociatedWorld = GWorld;
 
 		MainAudioDeviceHandle = RequestAudioDevice(MainDeviceParams);
 		
@@ -530,6 +559,8 @@ FAudioDeviceHandle FAudioDeviceManager::CreateNewDevice(const FAudioDeviceParams
 	}
 	else
 	{
+		RegisterWorld(InParams.AssociatedWorld, DeviceID);
+
 		FAudioDeviceHandle Handle = BuildNewHandle(*ContainerPtr, DeviceID, InParams);
 		FAudioDeviceManagerDelegates::OnAudioDeviceCreated.Broadcast(DeviceID);
 		return Handle;
@@ -582,10 +613,8 @@ void FAudioDeviceManager::DecrementDevice(Audio::FDeviceId DeviceID, UWorld* InW
 		FAudioDeviceManagerDelegates::OnAudioDeviceDestroyed.Broadcast(DeviceID);
 		Devices.Remove(DeviceID);
 	}
-	else if (InWorld)
-	{
-		Container.WorldsUsingThisDevice.Remove(InWorld);
-	}
+
+	UnregisterWorld(InWorld, DeviceID);
 }
 
 bool FAudioDeviceManager::ShutdownAllAudioDevices()
@@ -744,7 +773,6 @@ void FAudioDeviceManager::UpdateActiveAudioDevices(bool bGameTicking)
 		}
 	}
 
-
 	IterateOverAllDevices(
 		[&bGameTicking](Audio::FDeviceId, FAudioDevice* InDevice)
 		{
@@ -1257,14 +1285,14 @@ void FAudioDeviceManager::SetDynamicSoundVolume(ESoundType SoundType, const FNam
 }
 
 #if ENABLE_AUDIO_DEBUG
-FAudioDebugger& FAudioDeviceManager::GetDebugger()
+Audio::FAudioDebugger& FAudioDeviceManager::GetDebugger()
 {
 	check(AudioDebugger.IsValid());
 
 	return *AudioDebugger;
 }
 
-const FAudioDebugger& FAudioDeviceManager::GetDebugger() const
+const Audio::FAudioDebugger& FAudioDeviceManager::GetDebugger() const
 {
 	check(AudioDebugger.IsValid());
 
@@ -1555,3 +1583,4 @@ FAudioDeviceManager::FAudioDeviceContainer::~FAudioDeviceContainer()
 FAudioDeviceManagerDelegates::FOnAudioDeviceCreated FAudioDeviceManagerDelegates::OnAudioDeviceCreated;
 FAudioDeviceManagerDelegates::FOnAudioDeviceDestroyed FAudioDeviceManagerDelegates::OnAudioDeviceDestroyed;
 FAudioDeviceManagerDelegates::FOnWorldRegisteredToAudioDevice FAudioDeviceManagerDelegates::OnWorldRegisteredToAudioDevice;
+FAudioDeviceManagerDelegates::FOnWorldUnregisteredWithAudioDevice FAudioDeviceManagerDelegates::OnWorldUnregisteredWithAudioDevice;
diff --git a/Engine/Source/Runtime/Engine/Private/AudioVirtualLoop.cpp b/Engine/Source/Runtime/Engine/Private/AudioVirtualLoop.cpp
--- a/Engine/Source/Runtime/Engine/Private/AudioVirtualLoop.cpp
+++ b/Engine/Source/Runtime/Engine/Private/AudioVirtualLoop.cpp
@@ -213,7 +213,7 @@ bool FAudioVirtualLoop::Update(float DeltaTime, bool bForceUpdate)
 	}
 
 #if ENABLE_AUDIO_DEBUG
-	FAudioDebugger::DrawDebugInfo(*this);
+	Audio::FAudioDebugger::DrawDebugInfo(*this);
 #endif // ENABLE_AUDIO_DEBUG
 
 	UpdateFocusData(UpdateDelta);
diff --git a/Engine/Source/Runtime/Engine/Private/GameViewportClient.cpp b/Engine/Source/Runtime/Engine/Private/GameViewportClient.cpp
--- a/Engine/Source/Runtime/Engine/Private/GameViewportClient.cpp
+++ b/Engine/Source/Runtime/Engine/Private/GameViewportClient.cpp
@@ -418,7 +418,7 @@ void UGameViewportClient::SetEnabledStats(const TArray<FString>& InEnabledStats)
 	{
 		if (FAudioDeviceManager* DeviceManager = GEngine->GetAudioDeviceManager())
 		{
-			FAudioDebugger::ResolveDesiredStats(this);
+			Audio::FAudioDebugger::ResolveDesiredStats(this);
 		}
 	}
 #endif // ENABLE_AUDIO_DEBUG
@@ -484,7 +484,7 @@ void UGameViewportClient::Init(struct FWorldContext& WorldContext, UGameInstance
 			if (AudioDevice.IsValid())
 			{
 #if ENABLE_AUDIO_DEBUG
-				FAudioDebugger::ResolveDesiredStats(this);
+				Audio::FAudioDebugger::ResolveDesiredStats(this);
 #endif // ENABLE_AUDIO_DEBUG
 
 				// Set the base mix of the new device based on the world settings of the world
diff --git a/Engine/Source/Runtime/Engine/Private/UnrealEngine.cpp b/Engine/Source/Runtime/Engine/Private/UnrealEngine.cpp
--- a/Engine/Source/Runtime/Engine/Private/UnrealEngine.cpp
+++ b/Engine/Source/Runtime/Engine/Private/UnrealEngine.cpp
@@ -15797,7 +15797,11 @@ static FAutoConsoleCommand PakFileTestCmd(
 int32 UEngine::RenderStatSoundReverb(UWorld* World, FViewport* Viewport, FCanvas* Canvas, int32 X, int32 Y, const FVector* ViewLocation, const FRotator* ViewRotation)
 {
 #if ENABLE_AUDIO_DEBUG
-	return FAudioDebugger::RenderStatReverb(World, Viewport, Canvas, X, Y, ViewLocation, ViewRotation);
+	const int32 FontHeight = GetStatsFont()->GetMaxCharHeight() + 2.0f;
+	Canvas->DrawShadowedString(X, Y, TEXT("'stat soundreverb' is deprecated. Please use au.Debug.SoundReverb"), GetStatsFont(), FLinearColor::Yellow);
+	Y += FontHeight;
+
+	return Audio::FAudioDebugger::RenderStatReverb(World, Viewport, Canvas, X, Y);
 #else // !ENABLE_AUDIO_DEBUG
 	return Y;
 #endif // !ENABLE_AUDIO_DEBUG
@@ -15807,7 +15811,11 @@ int32 UEngine::RenderStatSoundReverb(UWorld* World, FViewport* Viewport, FCanvas
 int32 UEngine::RenderStatSoundMixes(UWorld* World, FViewport* Viewport, FCanvas* Canvas, int32 X, int32 Y, const FVector* ViewLocation, const FRotator* ViewRotation)
 {
 #if ENABLE_AUDIO_DEBUG
-	return FAudioDebugger::RenderStatMixes(World, Viewport, Canvas, X, Y, ViewLocation, ViewRotation);
+	const int32 FontHeight = GetStatsFont()->GetMaxCharHeight() + 2.0f;
+	Canvas->DrawShadowedString(X, Y, TEXT("'stat soundmixes' is deprecated. Use au.Debug.SoundMixes"), GetStatsFont(), FLinearColor::Yellow);
+	Y += FontHeight;
+
+	return Audio::FAudioDebugger::RenderStatMixes(World, Viewport, Canvas, X, Y);
 #else // !ENABLE_AUDIO_DEBUG
 	return Y;
 #endif // !ENABLE_AUDIO_DEBUG
@@ -15817,7 +15825,11 @@ int32 UEngine::RenderStatSoundMixes(UWorld* World, FViewport* Viewport, FCanvas*
 int32 UEngine::RenderStatSoundModulators(UWorld* World, FViewport* Viewport, FCanvas* Canvas, int32 X, int32 Y, const FVector* ViewLocation, const FRotator* ViewRotation)
 {
 #if ENABLE_AUDIO_DEBUG
-	return FAudioDebugger::RenderStatModulators(World, Viewport, Canvas, X, Y, ViewLocation, ViewRotation);
+	const int32 FontHeight = GetStatsFont()->GetMaxCharHeight() + 2.0f;
+	Canvas->DrawShadowedString(X, Y, TEXT("'stat soundmodulators' is deprecated. Use au.Debug.Modulators"), GetStatsFont(), FLinearColor::Yellow);
+	Y += FontHeight;
+
+	return Audio::FAudioDebugger::RenderStatModulators(World, Viewport, Canvas, X, Y, ViewLocation, ViewRotation);
 #else // !ENABLE_AUDIO_DEBUG
 	return Y;
 #endif // !ENABLE_AUDIO_DEBUG
@@ -15827,7 +15839,7 @@ int32 UEngine::RenderStatSoundModulators(UWorld* World, FViewport* Viewport, FCa
 bool UEngine::ToggleStatSoundWaves(UWorld* World, FCommonViewportClient* ViewportClient, const TCHAR* Stream)
 {
 #if ENABLE_AUDIO_DEBUG
-	return FAudioDebugger::ToggleStatWaves(World, ViewportClient, Stream);
+	return Audio::FAudioDebugger::ToggleStatWaves(World, ViewportClient, Stream);
 #else // !ENABLE_AUDIO_DEBUG
 	return false;
 #endif // !ENABLE_AUDIO_DEBUG
@@ -15837,7 +15849,7 @@ bool UEngine::ToggleStatSoundWaves(UWorld* World, FCommonViewportClient* Viewpor
 bool UEngine::ToggleStatSoundCues(UWorld* World, FCommonViewportClient* ViewportClient, const TCHAR* Stream)
 {
 #if ENABLE_AUDIO_DEBUG
-	return FAudioDebugger::ToggleStatCues(World, ViewportClient, Stream);
+	return Audio::FAudioDebugger::ToggleStatCues(World, ViewportClient, Stream);
 #else // !ENABLE_AUDIO_DEBUG
 	return false;
 #endif // !ENABLE_AUDIO_DEBUG
@@ -15853,7 +15865,7 @@ bool UEngine::ToggleStatAudioStreaming(UWorld* World, FCommonViewportClient* Vie
 bool UEngine::ToggleStatSoundMixes(UWorld* World, FCommonViewportClient* ViewportClient, const TCHAR* Stream)
 {
 #if ENABLE_AUDIO_DEBUG
-	return FAudioDebugger::ToggleStatMixes(World, ViewportClient, Stream);
+	return Audio::FAudioDebugger::ToggleStatMixes(World, ViewportClient, Stream);
 #else // !ENABLE_AUDIO_DEBUG
 	return false;
 #endif // !ENABLE_AUDIO_DEBUG
@@ -15863,7 +15875,7 @@ bool UEngine::ToggleStatSoundMixes(UWorld* World, FCommonViewportClient* Viewpor
 bool UEngine::PostStatSoundModulatorHelp(UWorld* World, FCommonViewportClient* ViewportClient, const TCHAR* Stream)
 {
 #if ENABLE_AUDIO_DEBUG
-	return FAudioDebugger::PostStatModulatorHelp(World, ViewportClient, Stream);
+	return Audio::FAudioDebugger::PostStatModulatorHelp(World, ViewportClient, Stream);
 #else // !ENABLE_AUDIO_DEBUG
 	return false;
 #endif // !ENABLE_AUDIO_DEBUG
@@ -15873,7 +15885,7 @@ bool UEngine::PostStatSoundModulatorHelp(UWorld* World, FCommonViewportClient* V
 bool UEngine::ToggleStatSoundModulators(UWorld* World, FCommonViewportClient* ViewportClient, const TCHAR* Stream)
 {
 #if ENABLE_AUDIO_DEBUG
-	return FAudioDebugger::ToggleStatModulators(World, ViewportClient, Stream);
+	return Audio::FAudioDebugger::ToggleStatModulators(World, ViewportClient, Stream);
 #else // !ENABLE_AUDIO_DEBUG
 	return false;
 #endif // !ENABLE_AUDIO_DEBUG
@@ -15883,7 +15895,11 @@ bool UEngine::ToggleStatSoundModulators(UWorld* World, FCommonViewportClient* Vi
 int32 UEngine::RenderStatSoundWaves(UWorld* World, FViewport* Viewport, FCanvas* Canvas, int32 X, int32 Y, const FVector* ViewLocation, const FRotator* ViewRotation)
 {
 #if ENABLE_AUDIO_DEBUG
-	return FAudioDebugger::RenderStatWaves(World, Viewport, Canvas, X, Y, ViewLocation, ViewRotation);
+	const int32 FontHeight = GetStatsFont()->GetMaxCharHeight() + 2.0f;
+	Canvas->DrawShadowedString(X, Y, TEXT("'stat soundwaves' is deprecated. Use au.Debug.SoundWaves"), GetStatsFont(), FLinearColor::Yellow);
+	Y += FontHeight;
+
+	return Audio::FAudioDebugger::RenderStatWaves(World, Viewport, Canvas, X, Y);
 #else // !ENABLE_AUDIO_DEBUG
 	return Y;
 #endif // !ENABLE_AUDIO_DEBUG
@@ -15899,7 +15915,11 @@ int32 UEngine::RenderStatAudioStreaming(UWorld* World, FViewport* Viewport, FCan
 int32 UEngine::RenderStatSoundCues(UWorld* World, FViewport* Viewport, FCanvas* Canvas, int32 X, int32 Y, const FVector* ViewLocation, const FRotator* ViewRotation)
 {
 #if ENABLE_AUDIO_DEBUG
-	return FAudioDebugger::RenderStatCues(World, Viewport, Canvas, X, Y, ViewLocation, ViewRotation);
+	const int32 FontHeight = GetStatsFont()->GetMaxCharHeight() + 2.0f;
+	Canvas->DrawShadowedString(X, Y, TEXT("'stat soundcues' is deprecated. Use au.Debug.SoundCues"), GetStatsFont(), FLinearColor::Yellow);
+	Y += FontHeight;
+
+	return Audio::FAudioDebugger::RenderStatCues(World, Viewport, Canvas, X, Y);
 #else // !ENABLE_AUDIO_DEBUG
 	return Y;
 #endif // !ENABLE_AUDIO_DEBUG
@@ -15909,7 +15929,7 @@ int32 UEngine::RenderStatSoundCues(UWorld* World, FViewport* Viewport, FCanvas*
 bool UEngine::ToggleStatSounds(UWorld* World, FCommonViewportClient* ViewportClient, const TCHAR* Stream)
 {
 #if ENABLE_AUDIO_DEBUG
-	return FAudioDebugger::ToggleStatSounds(World, ViewportClient, Stream);
+	return Audio::FAudioDebugger::ToggleStatSounds(World, ViewportClient, Stream);
 #else // !ENABLE_AUDIO_DEBUG
 	return false;
 #endif // !ENABLE_AUDIO_DEBUG
@@ -15918,7 +15938,11 @@ bool UEngine::ToggleStatSounds(UWorld* World, FCommonViewportClient* ViewportCli
 int32 UEngine::RenderStatSounds(UWorld* World, FViewport* Viewport, FCanvas* Canvas, int32 X, int32 Y, const FVector* ViewLocation, const FRotator* ViewRotation)
 {
 #if ENABLE_AUDIO_DEBUG
-	return FAudioDebugger::RenderStatSounds(World, Viewport, Canvas, X, Y, ViewLocation, ViewRotation);
+	const int32 FontHeight = GetStatsFont()->GetMaxCharHeight() + 2.0f;
+	Canvas->DrawShadowedString(X, Y, TEXT("'stat sounds' is deprecated. Use au.Debug.Sounds"), GetStatsFont(), FLinearColor::Yellow);
+	Y += FontHeight;
+
+	return Audio::FAudioDebugger::RenderStatSounds(World, Viewport, Canvas, X, Y);
 #else // !ENABLE_AUDIO_DEBUG
 	return Y;
 #endif // !ENABLE_AUDIO_DEBUG
diff --git a/Engine/Source/Runtime/Engine/Private/World.cpp b/Engine/Source/Runtime/Engine/Private/World.cpp
--- a/Engine/Source/Runtime/Engine/Private/World.cpp
+++ b/Engine/Source/Runtime/Engine/Private/World.cpp
@@ -7568,7 +7568,23 @@ void UWorld::AddPostProcessingSettings(FVector ViewLocation, FSceneView* SceneVi
 
 void UWorld::SetAudioDevice(FAudioDeviceHandle& InHandle)
 {
+	if (InHandle.GetDeviceID() == AudioDeviceHandle.GetDeviceID())
+	{
+		return;
+	}
+
+	FAudioDeviceManager* DeviceManager = GEngine ? GEngine->GetAudioDeviceManager() : nullptr;
+	if (DeviceManager && InHandle.IsValid())
+	{
+		GEngine->GetAudioDeviceManager()->UnregisterWorld(this, InHandle.GetDeviceID());
+	}
+
 	AudioDeviceHandle = InHandle;
+
+	if (DeviceManager)
+	{
+		GEngine->GetAudioDeviceManager()->RegisterWorld(this, InHandle.GetDeviceID());
+	}
 }
 
 FAudioDeviceHandle UWorld::GetAudioDevice()
diff --git a/Engine/Source/Runtime/Engine/Public/Audio/AudioDebug.h b/Engine/Source/Runtime/Engine/Public/Audio/AudioDebug.h
--- a/Engine/Source/Runtime/Engine/Public/Audio/AudioDebug.h
+++ b/Engine/Source/Runtime/Engine/Public/Audio/AudioDebug.h
@@ -18,111 +18,126 @@ struct FWaveInstance;
 class FSoundSource;
 class UWorld;
 
-
-class ENGINE_API FAudioDebugger
+namespace Audio
 {
-public:
-	FAudioDebugger();
-
-	/** Struct which contains debug names for run-time debugging of sounds. */
-	struct FDebugNames
+	class ENGINE_API FAudioDebugger
 	{
-		TArray<FName> SoloSoundClass;
-		TArray<FName> SoloSoundWave;
-		TArray<FName> SoloSoundCue;
-		TArray<FName> MuteSoundClass;
-		TArray<FName> MuteSoundWave;
-		TArray<FName> MuteSoundCue;
-
-		FString DebugAudioMixerSoundName;
-		FString DebugSoundName;
-		bool bDebugSoundName;
-
-		FDebugNames()
-			: bDebugSoundName(false)
-		{}
-	};
-
-	static void DrawDebugInfo(const FSoundSource& SoundSource);
-	static void DrawDebugInfo(const FActiveSound& ActiveSound, const TArray<FWaveInstance*>& ThisSoundsWaveInstances, const float DeltaTime);
-	static void DrawDebugInfo(UWorld& World, const TArray<FListener>& Listeners);
-	static void DrawDebugInfo(const FAudioVirtualLoop& VirtualLoop);
-	static bool PostStatModulatorHelp(UWorld* World, FCommonViewportClient* ViewportClient, const TCHAR* Stream);
-	static int32 RenderStatCues(UWorld* World, FViewport* Viewport, FCanvas* Canvas, int32 X, int32 Y, const FVector* ViewLocation, const FRotator* ViewRotation);
-	static int32 RenderStatMixes(UWorld* World, FViewport* Viewport, FCanvas* Canvas, int32 X, int32 Y, const FVector* ViewLocation, const FRotator* ViewRotation);
-	static int32 RenderStatModulators(UWorld* World, FViewport* Viewport, FCanvas* Canvas, int32 X, int32 Y, const FVector* ViewLocation, const FRotator* ViewRotation);
-	static int32 RenderStatReverb(UWorld* World, FViewport* Viewport, FCanvas* Canvas, int32 X, int32 Y, const FVector* ViewLocation, const FRotator* ViewRotation);
-	static int32 RenderStatSounds(UWorld* World, FViewport* Viewport, FCanvas* Canvas, int32 X, int32 Y, const FVector* ViewLocation, const FRotator* ViewRotation);
-	static int32 RenderStatWaves(UWorld* World, FViewport* Viewport, FCanvas* Canvas, int32 X, int32 Y, const FVector* ViewLocation, const FRotator* ViewRotation);
-	static void RemoveDevice(const FAudioDevice& AudioDevice);
-	static void ResolveDesiredStats(FViewportClient* ViewportClient);
-	static void SendUpdateResultsToGameThread(const FAudioDevice& AudioDevice, const int32 FirstActiveIndex);
-	static bool ToggleStatCues(UWorld* World, FCommonViewportClient* ViewportClient, const TCHAR* Stream);
-	static bool ToggleStatMixes(UWorld* World, FCommonViewportClient* ViewportClient, const TCHAR* Stream);
-	static bool ToggleStatModulators(UWorld* World, FCommonViewportClient* ViewportClient, const TCHAR* Stream);
-	static bool ToggleStatSounds(UWorld* World, FCommonViewportClient* ViewportClient, const TCHAR* Stream);
-	static bool ToggleStatWaves(UWorld* World, FCommonViewportClient* ViewportClient, const TCHAR* Stream);
-	static void UpdateAudibleInactiveSounds(const uint32 FirstIndex, const TArray<FWaveInstance*>& WaveInstances);
-
-	void ClearMutesAndSolos();
-	void DumpActiveSounds() const;
-
-	bool IsVisualizeDebug3dEnabled() const;
-	void ToggleVisualizeDebug3dEnabled();
+	public:
+		FAudioDebugger();
+
+		/** Struct which contains debug names for run-time debugging of sounds. */
+		struct FDebugNames
+		{
+			TArray<FName> SoloSoundClass;
+			TArray<FName> SoloSoundWave;
+			TArray<FName> SoloSoundCue;
+			TArray<FName> MuteSoundClass;
+			TArray<FName> MuteSoundWave;
+			TArray<FName> MuteSoundCue;
+
+			FString DebugAudioMixerSoundName;
+			FString DebugSoundName;
+			bool bDebugSoundName;
+
+			FDebugNames()
+				: bDebugSoundName(false)
+			{}
+		};
+
+		static void DrawDebugInfo(const FSoundSource& SoundSource);
+		static void DrawDebugInfo(const FActiveSound& ActiveSound, const TArray<FWaveInstance*>& ThisSoundsWaveInstances, const float DeltaTime);
+		static void DrawDebugInfo(UWorld& World, const TArray<FListener>& Listeners);
+		static void DrawDebugInfo(const FAudioVirtualLoop& VirtualLoop);
+		static void DrawDebugStats(UWorld& World);
+		static bool DrawDebugStatsEnabled();
+		static bool PostStatModulatorHelp(UWorld* World, FCommonViewportClient* ViewportClient, const TCHAR* Stream);
+		static int32 RenderStatCues(UWorld* World, FViewport* Viewport, FCanvas* Canvas, int32 X, int32 Y);
+		static int32 RenderStatMixes(UWorld* World, FViewport* Viewport, FCanvas* Canvas, int32 X, int32 Y);
+		static int32 RenderStatModulators(UWorld* World, FViewport* Viewport, FCanvas* Canvas, int32 X, int32 Y, const FVector* ViewLocation, const FRotator* ViewRotation);
+		static int32 RenderStatReverb(UWorld* World, FViewport* Viewport, FCanvas* Canvas, int32 X, int32 Y);
+		static int32 RenderStatSounds(UWorld* World, FViewport* Viewport, FCanvas* Canvas, int32 X, int32 Y);
+		static int32 RenderStatWaves(UWorld* World, FViewport* Viewport, FCanvas* Canvas, int32 X, int32 Y);
+		static void RemoveDevice(const FAudioDevice& AudioDevice);
+		static void ResolveDesiredStats(FViewportClient* ViewportClient);
+		static void SendUpdateResultsToGameThread(const FAudioDevice& AudioDevice, const int32 FirstActiveIndex);
+		static bool ToggleStatCues(UWorld* World, FCommonViewportClient* ViewportClient, const TCHAR* Stream);
+		static bool ToggleStatMixes(UWorld* World, FCommonViewportClient* ViewportClient, const TCHAR* Stream);
+		static bool ToggleStatModulators(UWorld* World, FCommonViewportClient* ViewportClient, const TCHAR* Stream);
+		static bool ToggleStatSounds(UWorld* World, FCommonViewportClient* ViewportClient, const TCHAR* Stream);
+		static bool ToggleStatWaves(UWorld* World, FCommonViewportClient* ViewportClient, const TCHAR* Stream);
+		static void UpdateAudibleInactiveSounds(const uint32 FirstIndex, const TArray<FWaveInstance*>& WaveInstances);
+
+		static void ClearStats(const uint32 StatsToToggle, UWorld* InWorld);
+		static void SetStats(const uint32 StatsToToggle, UWorld* InWorld);
+
+		void ClearMutesAndSolos();
+		void DumpActiveSounds() const;
+
+		bool IsVisualizeDebug3dEnabled() const;
+		void ToggleVisualizeDebug3dEnabled();
 
 #if WITH_EDITOR
-	static void OnBeginPIE();
-	static void OnEndPIE();
+		static void OnBeginPIE();
+		static void OnEndPIE();
 #endif // WITH_EDITOR
 
-	// Evaluate Mute/Solos
-	void QuerySoloMuteSoundClass(const FString& Name, bool& bOutIsSoloed, bool& bOutIsMuted, FString& OutReason ) const;
-	void QuerySoloMuteSoundWave(const FString& Name, bool& bOutIsSoloed, bool& bOutIsMuted, FString& OutReason) const;
-	void QuerySoloMuteSoundCue(const FString& Name, bool& bOutIsSoloed, bool& bOutIsMuted, FString& OutReason ) const;
-   
-	// Is Mute/Solos. (only audio thread).
-	bool IsSoloSoundClass(FName InName) const { return DebugNames.SoloSoundClass.Contains(InName); }
-	bool IsSoloSoundWave(FName InName) const  { return DebugNames.SoloSoundWave.Contains(InName); }
-	bool IsSoloSoundCue(FName InName) const   { return DebugNames.SoloSoundCue.Contains(InName); }
-	bool IsMuteSoundClass(FName InName) const { return DebugNames.MuteSoundClass.Contains(InName); }
-	bool IsMuteSoundWave(FName InName) const  { return DebugNames.MuteSoundWave.Contains(InName); }
-	bool IsMuteSoundCue(FName InName) const   { return DebugNames.MuteSoundCue.Contains(InName); }
-		
-	// Mute/Solos toggles. (any thread). (If exclusive, toggle-on will clear everything first, and toggle-off will clear all).
-	void ToggleSoloSoundClass(FName InName, bool bExclusive = false) { ToggleNameArray(InName, DebugNames.SoloSoundClass, bExclusive); }
-	void ToggleSoloSoundWave(FName InName, bool bExclusive = false)  { ToggleNameArray(InName, DebugNames.SoloSoundWave, bExclusive); }
-	void ToggleSoloSoundCue(FName InName, bool bExclusive = false)   { ToggleNameArray(InName, DebugNames.SoloSoundCue, bExclusive); }
-	void ToggleMuteSoundClass(FName InName, bool bExclusive = false) { ToggleNameArray(InName, DebugNames.MuteSoundClass, bExclusive); }
-	void ToggleMuteSoundWave(FName InName, bool bExclusive = false)  { ToggleNameArray(InName, DebugNames.MuteSoundWave, bExclusive); }
-	void ToggleMuteSoundCue(FName InName, bool bExclusive = false)   { ToggleNameArray(InName, DebugNames.MuteSoundCue,	bExclusive); }
-
-	// Set Mute/Solo. (any thread).
-	void SetMuteSoundCue(FName InName, bool bInOnOff)  { SetNameArray(InName, DebugNames.MuteSoundCue, bInOnOff); }
-	void SetMuteSoundWave(FName InName, bool bInOnOff) { SetNameArray(InName, DebugNames.MuteSoundWave, bInOnOff); }
-	void SetSoloSoundCue(FName InName, bool bInOnOff)  { SetNameArray(InName, DebugNames.SoloSoundCue, bInOnOff); }	
-	void SetSoloSoundWave(FName InName, bool bInOnOff) { SetNameArray(InName, DebugNames.SoloSoundWave, bInOnOff); }
-		   	
-	void SetAudioMixerDebugSound(const TCHAR* SoundName);
-	void SetAudioDebugSound(const TCHAR* SoundName);
-
-	const FString& GetAudioMixerDebugSoundName() const;
-	bool GetAudioDebugSound(FString& OutDebugSound);	  	
-private:
-	void SetNameArray(FName InName, TArray<FName>& InNameArray, bool bOnOff);
-	void ToggleNameArray(FName InName, TArray<FName>& NameArray, bool bExclusive);
-	void ExecuteCmdOnAudioThread(TFunction<void()> Cmd);
-	
-	void GetDebugSoloMuteStateX(
-		const FString& Name, const TArray<FName>& Solos, const TArray<FName>& Mutes, 
-		bool& bOutIsSoloed, bool& bOutIsMuted, FString& OutReason) const;
-
-	static bool ToggleStats(UWorld* World, const uint8 StatToToggle);
-	void ToggleStats(const uint32 AudioDeviceHandle, const uint8 StatsToToggle);
-
-	/** Instance of the debug names struct. */
-	FDebugNames DebugNames;
-
-	/** Whether or not 3d debug visualization is enabled. */
-	uint8 bVisualize3dDebug : 1;
-};
+		// Evaluate Mute/Solos
+		void QuerySoloMuteSoundClass(const FString& Name, bool& bOutIsSoloed, bool& bOutIsMuted, FString& OutReason) const;
+		void QuerySoloMuteSoundWave(const FString& Name, bool& bOutIsSoloed, bool& bOutIsMuted, FString& OutReason) const;
+		void QuerySoloMuteSoundCue(const FString& Name, bool& bOutIsSoloed, bool& bOutIsMuted, FString& OutReason) const;
+
+		// Is Mute/Solos. (only audio thread).
+		bool IsSoloSoundClass(FName InName) const { return DebugNames.SoloSoundClass.Contains(InName); }
+		bool IsSoloSoundWave(FName InName) const { return DebugNames.SoloSoundWave.Contains(InName); }
+		bool IsSoloSoundCue(FName InName) const { return DebugNames.SoloSoundCue.Contains(InName); }
+		bool IsMuteSoundClass(FName InName) const { return DebugNames.MuteSoundClass.Contains(InName); }
+		bool IsMuteSoundWave(FName InName) const { return DebugNames.MuteSoundWave.Contains(InName); }
+		bool IsMuteSoundCue(FName InName) const { return DebugNames.MuteSoundCue.Contains(InName); }
+
+		// Mute/Solos toggles. (any thread). (If exclusive, toggle-on will clear everything first, and toggle-off will clear all).
+		void ToggleSoloSoundClass(FName InName, bool bExclusive = false) { ToggleNameArray(InName, DebugNames.SoloSoundClass, bExclusive); }
+		void ToggleSoloSoundWave(FName InName, bool bExclusive = false) { ToggleNameArray(InName, DebugNames.SoloSoundWave, bExclusive); }
+		void ToggleSoloSoundCue(FName InName, bool bExclusive = false) { ToggleNameArray(InName, DebugNames.SoloSoundCue, bExclusive); }
+		void ToggleMuteSoundClass(FName InName, bool bExclusive = false) { ToggleNameArray(InName, DebugNames.MuteSoundClass, bExclusive); }
+		void ToggleMuteSoundWave(FName InName, bool bExclusive = false) { ToggleNameArray(InName, DebugNames.MuteSoundWave, bExclusive); }
+		void ToggleMuteSoundCue(FName InName, bool bExclusive = false) { ToggleNameArray(InName, DebugNames.MuteSoundCue, bExclusive); }
+
+		// Set Mute/Solo. (any thread).
+		void SetMuteSoundCue(FName InName, bool bInOnOff) { SetNameArray(InName, DebugNames.MuteSoundCue, bInOnOff); }
+		void SetMuteSoundWave(FName InName, bool bInOnOff) { SetNameArray(InName, DebugNames.MuteSoundWave, bInOnOff); }
+		void SetSoloSoundCue(FName InName, bool bInOnOff) { SetNameArray(InName, DebugNames.SoloSoundCue, bInOnOff); }
+		void SetSoloSoundWave(FName InName, bool bInOnOff) { SetNameArray(InName, DebugNames.SoloSoundWave, bInOnOff); }
+
+		void SetAudioMixerDebugSound(const TCHAR* SoundName);
+		void SetAudioDebugSound(const TCHAR* SoundName);
+
+		const FString& GetAudioMixerDebugSoundName() const;
+		bool GetAudioDebugSound(FString& OutDebugSound);
+
+	private:
+		static void DrawDebugStats(UWorld& World, FViewport& Viewport);
+		void SetNameArray(FName InName, TArray<FName>& InNameArray, bool bOnOff);
+		void ToggleNameArray(FName InName, TArray<FName>& NameArray, bool bExclusive);
+		void ExecuteCmdOnAudioThread(TFunction<void()> Cmd);
+
+		void GetDebugSoloMuteStateX(
+			const FString& Name, const TArray<FName>& Solos, const TArray<FName>& Mutes,
+			bool& bOutIsSoloed, bool& bOutIsMuted, FString& OutReason) const;
+
+		void ClearStats(FDeviceId DeviceId, const uint32 StatsToClear);
+
+		void SetStats(FDeviceId DeviceId, const uint32 StatsToSet);
+
+		static bool ToggleStats(UWorld* World, const uint32 StatToToggle);
+		void ToggleStats(FDeviceId DeviceId, const uint32 StatsToToggle);
+
+		FDelegateHandle WorldRegisteredWithDeviceHandle;
+
+		/** Instance of the debug names struct. */
+		FDebugNames DebugNames;
+
+		/** Whether or not 3d debug visualization is enabled. */
+		uint8 bVisualize3dDebug : 1;
+	};
+} // namespace Audio
 #endif // ENABLE_AUDIO_DEBUG
\ No newline at end of file
diff --git a/Engine/Source/Runtime/Engine/Public/AudioDevice.h b/Engine/Source/Runtime/Engine/Public/AudioDevice.h
--- a/Engine/Source/Runtime/Engine/Public/AudioDevice.h
+++ b/Engine/Source/Runtime/Engine/Public/AudioDevice.h
@@ -27,7 +27,6 @@
  */
 
 class FArchive;
-class FAudioDebugger;
 class FAudioDevice;
 class FAudioEffectsManager;
 class FCanvas;
@@ -57,6 +56,10 @@ struct FAudioQualitySettings;
 struct FRotator;
 struct FWaveInstance;
 
+namespace Audio
+{
+	class FAudioDebugger;
+}
 
 /**
  * Debug state of the audio system
@@ -466,7 +469,7 @@ private:
 	bool HandleSetDynamicSoundCommand(const TCHAR* Cmd, FOutputDevice& Ar);
 
 	/** Handles all argument parsing for the solo commands in one place */
-	using FToggleSoloPtr = void (FAudioDebugger::*)(FName InName, bool bExclusive);
+	using FToggleSoloPtr = void (Audio::FAudioDebugger::*)(FName InName, bool bExclusive);
 	void HandleAudioSoloCommon(const TCHAR* Cmd, FOutputDevice& Ar, FToggleSoloPtr Funct);
 
 	/**
diff --git a/Engine/Source/Runtime/Engine/Public/AudioDeviceManager.h b/Engine/Source/Runtime/Engine/Public/AudioDeviceManager.h
--- a/Engine/Source/Runtime/Engine/Public/AudioDeviceManager.h
+++ b/Engine/Source/Runtime/Engine/Public/AudioDeviceManager.h
@@ -26,7 +26,6 @@ class USoundSubmixBase;
 class USoundWave;
 class UWorld;
 class FAudioDevice;
-class FAudioDebugger;
 struct FSourceEffectChainEntry;
 
 namespace Audio
@@ -35,6 +34,8 @@ namespace Audio
 	 * Typed identifier for Audio Device Id
 	 */
 	using FDeviceId = uint32;
+
+	class FAudioDebugger;
 }
 
 enum class ESoundType : uint8
@@ -194,11 +195,17 @@ public:
 	DECLARE_MULTICAST_DELEGATE_OneParam(FOnAudioDeviceDestroyed, Audio::FDeviceId /* AudioDeviceId*/);
 	static FOnAudioDeviceDestroyed OnAudioDeviceDestroyed;
 
-	// This delegate is called whenever a world is registered to an audio device. Please note that UWorlds are not guaranteed to
-	// be registered to the same audio device throughout their lifecycle,
-	// and there is no guarantee on the lifespan of both the UWorld and the Audio Device registered in this callback.
+	// Called whenever a world is registered to an audio device. UWorlds are not guaranteed to be registered to the same
+	// audio device throughout their lifecycle, and there is no guarantee on the lifespan of both the UWorld and the Audio
+	// Device registered in this callback.
 	DECLARE_MULTICAST_DELEGATE_TwoParams(FOnWorldRegisteredToAudioDevice, const UWorld* /*InWorld */, Audio::FDeviceId /* AudioDeviceId*/);
 	static FOnWorldRegisteredToAudioDevice OnWorldRegisteredToAudioDevice;
+
+	// Called whenever a world is unregistered from an audio device. UWorlds are not guaranteed to be registered to the same
+	// audio device throughout their lifecycle, and there is no guarantee on the lifespan of both the UWorld and the Audio
+	// Device registered in this callback.
+	DECLARE_MULTICAST_DELEGATE_TwoParams(FOnWorldUnregisteredWithAudioDevice, const UWorld* /*InWorld */, Audio::FDeviceId /* AudioDeviceId*/);
+	static FOnWorldUnregisteredWithAudioDevice OnWorldUnregisteredWithAudioDevice;
 };
 
 /**
@@ -297,6 +304,12 @@ public:
 	/** Unregisters the Sound Class for all active devices. */
 	void UnregisterSoundClass(USoundClass* SoundClass);
 
+	/** Registers the world with the provided device Id */
+	void RegisterWorld(UWorld* InWorld, Audio::FDeviceId DeviceId);
+
+	/** Unregisters the world from the provided device Id */
+	void UnregisterWorld(UWorld* InWorld, Audio::FDeviceId DeviceId);
+
 	/** Initializes the sound class for all active devices. */
 	void InitSoundClasses();
 
@@ -367,8 +380,8 @@ public:
 
 #if ENABLE_AUDIO_DEBUG
 	/** Get the audio debugger instance */
-	FAudioDebugger& GetDebugger();
-	const FAudioDebugger& GetDebugger() const;
+	Audio::FAudioDebugger& GetDebugger();
+	const Audio::FAudioDebugger& GetDebugger() const;
 #endif // ENABLE_AUDIO_DEBUG
 
 public:
@@ -395,7 +408,7 @@ private:
 
 #if ENABLE_AUDIO_DEBUG
 	/** Instance of audio debugger shared across audio devices */
-	TUniquePtr<FAudioDebugger> AudioDebugger;
+	TUniquePtr<Audio::FAudioDebugger> AudioDebugger;
 #endif // ENABLE_AUDIO_DEBUG
 
 	bool InitializeManager();
