Static analysis warning fix

diff --git a/Engine/Source/Runtime/Engine/Private/GameplayStatics.cpp b/Engine/Source/Runtime/Engine/Private/GameplayStatics.cpp
--- a/Engine/Source/Runtime/Engine/Private/GameplayStatics.cpp
+++ b/Engine/Source/Runtime/Engine/Private/GameplayStatics.cpp
@@ -999,9 +999,9 @@ UParticleSystemComponent* UGameplayStatics::InternalSpawnEmitterAtLocation(UWorl
 {
 	check(World && EmitterTemplate);
 
-	if (UParticleSystemComponent* PSC = CreateParticleSystem(EmitterTemplate, World, World->GetWorldSettings(), bAutoDestroy, PoolingMethod))
+	UParticleSystemComponent* PSC = CreateParticleSystem(EmitterTemplate, World, World->GetWorldSettings(), bAutoDestroy, PoolingMethod);
+	if (PSC)
 	{
-
 		PSC->SetUsingAbsoluteLocation(true);
 		PSC->SetUsingAbsoluteRotation(true);
 		PSC->SetUsingAbsoluteScale(true);
