Fix issue where forward rendering required lights to be shadow casting to be able to utilize Light Function Materials

diff --git a/Engine/Source/Runtime/Engine/Private/Components/LightComponent.cpp b/Engine/Source/Runtime/Engine/Private/Components/LightComponent.cpp
--- a/Engine/Source/Runtime/Engine/Private/Components/LightComponent.cpp
+++ b/Engine/Source/Runtime/Engine/Private/Components/LightComponent.cpp
@@ -1414,7 +1414,7 @@ void ULightComponent::ReassignStationaryLightChannels(UWorld* TargetWorld, bool
 			if (!LightingScenario || !LightLevel->bIsLightingScenario || LightLevel == LightingScenario)
 			{				
 				if (LightComponent->bAffectsWorld
-					&& LightComponent->CastShadows 
+					&& (LightComponent->CastShadows || LightComponent->LightFunctionMaterial)
 					&& LightComponent->CastStaticShadows)
 				{
 					LightToOverlapMap.Add(new FLightAndChannel(LightComponent), TArray<FLightAndChannel*>());
diff --git a/Engine/Source/Runtime/Renderer/Private/DynamicShadowMapChannelBindingHelper.cpp b/Engine/Source/Runtime/Renderer/Private/DynamicShadowMapChannelBindingHelper.cpp
--- a/Engine/Source/Runtime/Renderer/Private/DynamicShadowMapChannelBindingHelper.cpp
+++ b/Engine/Source/Runtime/Renderer/Private/DynamicShadowMapChannelBindingHelper.cpp
@@ -108,7 +108,7 @@ int32 FDynamicShadowMapChannelBindingHelper::GetPriority(const FLightSceneInfo*
 {
 	check(Light);
 	const FLightSceneProxy* Proxy = Light->Proxy;
-	if (Proxy && Proxy->CastsDynamicShadow())
+	if (Proxy && (Proxy->CastsDynamicShadow() || Proxy->GetLightFunctionMaterial()))
 	{
 		if (Proxy->HasStaticShadowing())
 		{
