Add support for passing strings in execcmds, enclosed in single quotes. Needed in order to support profiling of the significance manager.
Doesn't affect shipping builds.

diff --git a/Engine/Source/Runtime/Engine/Private/UnrealEngine.cpp b/Engine/Source/Runtime/Engine/Private/UnrealEngine.cpp
--- a/Engine/Source/Runtime/Engine/Private/UnrealEngine.cpp
+++ b/Engine/Source/Runtime/Engine/Private/UnrealEngine.cpp
@@ -1600,24 +1600,51 @@ void UEngine::Init(IEngineLoop* InEngineLoop)
 	FString ExecCmds;
 	if( FParse::Value(FCommandLine::Get(), TEXT("ExecCmds="), ExecCmds, false) )
 	{
+		// Read the command array, ignoring any commas in single-quotes. 
+		// Convert any single-quotes to double-quotes and skip leading whitespace
+		// This allows passing of strings, e:g -execcmds="exampleCvar '0,1,2,3'"
 		TArray<FString> CommandArray;
-		ExecCmds.ParseIntoArray( CommandArray, TEXT(","), true );
-
-		for( int32 Cx = 0; Cx < CommandArray.Num(); ++Cx )
+		FString CurrentCommand = "";
+		bool bInQuotes = false;
+		bool bSkippingWhitespace = true;
+		for (int i = 0; i < ExecCmds.Len(); i++)
 		{
-			const FString& Command = CommandArray[Cx];
-			// Skip leading whitespaces in the command.
-			int32 Index = 0;
-			while( FChar::IsWhitespace( Command[Index] ) )
+			TCHAR CurrentChar = ExecCmds[i];
+			if (CurrentChar == '\'')
 			{
-				Index++;
+				bInQuotes = !bInQuotes;
+				CurrentCommand += "\"";
 			}
-
-			if( Index < Command.Len()-1 )
+			else if (CurrentChar == ',' && !bInQuotes)
 			{
-				new(GEngine->DeferredCommands) FString(*Command+Index);
+				if (CurrentCommand.Len() > 0)
+				{
+					CommandArray.Add(CurrentCommand);
+					CurrentCommand = "";
+				}
+				bSkippingWhitespace = true;
+			}
+			else
+			{
+				if (bSkippingWhitespace)
+				{
+					bSkippingWhitespace = FChar::IsWhitespace(CurrentChar);
+				}
+				if (!bSkippingWhitespace)
+				{
+					CurrentCommand += CurrentChar;
+				}
 			}
 		}
+		if (CurrentCommand.Len() > 0)
+		{
+			CommandArray.Add(CurrentCommand);
+		}
+
+		for( int32 Cx = 0; Cx < CommandArray.Num(); ++Cx )
+		{
+			new(GEngine->DeferredCommands) FString(*CommandArray[Cx]);
+		}
 	}
 
 	// optionally set the vsync console variable
