Preventing AudioVolume's Proxy from being added by default.
# AudioProxies are an audio thread representation of AudioVolumes; but are only added when the audio volume is enabled.
# PostRegister was incorrectly adding a proxy for disabled volumes.

diff --git a/Engine/Source/Runtime/Engine/Private/AudioVolume.cpp b/Engine/Source/Runtime/Engine/Private/AudioVolume.cpp
--- a/Engine/Source/Runtime/Engine/Private/AudioVolume.cpp
+++ b/Engine/Source/Runtime/Engine/Private/AudioVolume.cpp
@@ -209,7 +209,10 @@ void AAudioVolume::PostRegisterAllComponents()
 	Super::PostRegisterAllComponents();
 
 	GetRootComponent()->TransformUpdated.AddUObject(this, &AAudioVolume::TransformUpdated);
-	AddProxy();
+	if (bEnabled)
+	{
+		AddProxy();
+	}
 
 	UWorld* World = GetWorld();
 	World->AudioVolumes.Add(this);
