Fix for "on exit" compression crash

diff --git a/Engine/Source/Runtime/Engine/Private/Animation/AnimCompressionDerivedData.cpp b/Engine/Source/Runtime/Engine/Private/Animation/AnimCompressionDerivedData.cpp
--- a/Engine/Source/Runtime/Engine/Private/Animation/AnimCompressionDerivedData.cpp
+++ b/Engine/Source/Runtime/Engine/Private/Animation/AnimCompressionDerivedData.cpp
@@ -46,6 +46,11 @@ bool FDerivedDataAnimationCompression::Build( TArray<uint8>& OutDataArray )
 	FCompressibleAnimData& DataToCompress = *DataToCompressPtr.Get();
 	FCompressedAnimSequence OutData;
 
+	if (DataToCompress.IsCancelled())
+	{
+		return false;
+	}
+
 	SCOPE_CYCLE_COUNTER(STAT_AnimCompressionDerivedData);
 	UE_LOG(LogAnimationCompression, Log, TEXT("Building Anim DDC data for %s"), *DataToCompress.FullName);
 
@@ -56,6 +61,10 @@ bool FDerivedDataAnimationCompression::Build( TArray<uint8>& OutDataArray )
 		DataToCompress.Update(OutData);
 
 		const bool bBoneCompressionOk = FAnimationUtils::CompressAnimBones(DataToCompress, CompressionResult);
+		if (DataToCompress.IsCancelled())
+		{
+			return false;
+		}
 		const bool bCurveCompressionOk = FAnimationUtils::CompressAnimCurves(DataToCompress, OutData);
 
 #if DO_CHECK
