Fix CIS for Linux - local variable name conflict with member.

diff --git a/Engine/Source/Runtime/Engine/Private/PhysicsEngine/BodyInstance.cpp b/Engine/Source/Runtime/Engine/Private/PhysicsEngine/BodyInstance.cpp
--- a/Engine/Source/Runtime/Engine/Private/PhysicsEngine/BodyInstance.cpp
+++ b/Engine/Source/Runtime/Engine/Private/PhysicsEngine/BodyInstance.cpp
@@ -970,8 +970,8 @@ void FBodyInstance::UpdatePhysicsFilterData()
 			if (SetupShapeIndex < BI->BodySetup->AggGeom.GetElementCount())
 			{
 				// Get the shape's CollisionEnabled masked with the body's CollisionEnabled and compute the shape's collisionflags.
-				const ECollisionEnabled::Type CollisionEnabled = CollisionEnabledIntersection(BI->GetCollisionEnabled(), BI->GetShapeCollisionEnabled(SetupShapeIndex));
-				BuildBodyCollisionFlags(PerShapeCollisionData.CollisionFlags, CollisionEnabled, bInstanceComplexAsSimple);
+				const ECollisionEnabled::Type ShapeCollisionEnabled = CollisionEnabledIntersection(BI->GetCollisionEnabled(), BI->GetShapeCollisionEnabled(SetupShapeIndex));
+				BuildBodyCollisionFlags(PerShapeCollisionData.CollisionFlags, ShapeCollisionEnabled, bInstanceComplexAsSimple);
 			}
 			else
 			{
