Fix ensure in AutoTest when an ADPCM fie loops: when Stream Caching is enabled. we might not be looping back to the zeroth chunk.

diff --git a/Engine/Source/Runtime/Engine/Private/ADPCMAudioInfo.cpp b/Engine/Source/Runtime/Engine/Private/ADPCMAudioInfo.cpp
--- a/Engine/Source/Runtime/Engine/Private/ADPCMAudioInfo.cpp
+++ b/Engine/Source/Runtime/Engine/Private/ADPCMAudioInfo.cpp
@@ -985,8 +985,8 @@ const uint8* FADPCMAudioInfo::GetLoadedChunk(USoundWave* InSoundWave, uint32 Chu
 	}
 	else
 	{
-		const bool bIsLastChunk = ChunkIndex == InSoundWave->GetNumChunks() - 1;
-		ensureAlwaysMsgf(bSeekPending || ChunkIndex == PreviouslyRequestedChunkIndex || ChunkIndex == PreviouslyRequestedChunkIndex + 1, TEXT("ADPCM playback error! We skipped from the end of chunk %d to the beginning of chunk %d."), PreviouslyRequestedChunkIndex, ChunkIndex);
+		const bool bIsSeekingOrLooping = bSeekPending || ChunkIndex == FirstChunkSampleDataIndex;
+		ensureAlwaysMsgf(bIsSeekingOrLooping || ChunkIndex == PreviouslyRequestedChunkIndex || ChunkIndex == PreviouslyRequestedChunkIndex + 1, TEXT("ADPCM playback error! We skipped from the end of chunk %d to the beginning of chunk %d."), PreviouslyRequestedChunkIndex, ChunkIndex);
 
 		CurCompressedChunkHandle = IStreamingManager::Get().GetAudioStreamingManager().GetLoadedChunk(InSoundWave, ChunkIndex, false, true);
 		OutChunkSize = CurCompressedChunkHandle.Num();
