Remove debug log

diff --git a/Engine/Source/Runtime/Engine/Private/PhysicsEngine/BodyInstance.cpp b/Engine/Source/Runtime/Engine/Private/PhysicsEngine/BodyInstance.cpp
--- a/Engine/Source/Runtime/Engine/Private/PhysicsEngine/BodyInstance.cpp
+++ b/Engine/Source/Runtime/Engine/Private/PhysicsEngine/BodyInstance.cpp
@@ -963,9 +963,6 @@ void FBodyInstance::UpdatePhysicsFilterData()
 				PerShapeCollisionData = BodyCollisionData;
 			}
 
-			const int32 ElementCount = BI->BodySetup->AggGeom.GetElementCount();
-			UE_LOG(LogTemp, Warning, TEXT("ShapeIndexBase: %d, ShapeIndex: %d, SetupShapeIndex: %d, ElementCount: %d"), ShapeIndexBase, ShapeIndex, SetupShapeIndex, ElementCount);
-
 			const bool bInstanceComplexAsSimple = BI->BodySetup.IsValid() ? (BI->BodySetup->GetCollisionTraceFlag() == CTF_UseComplexAsSimple) : false;
 			if (SetupShapeIndex < BI->BodySetup->AggGeom.GetElementCount())
 			{
