Fix CIS Linux... another shadowed member name

diff --git a/Engine/Source/Runtime/Engine/Private/PhysicsEngine/BodyInstance.cpp b/Engine/Source/Runtime/Engine/Private/PhysicsEngine/BodyInstance.cpp
--- a/Engine/Source/Runtime/Engine/Private/PhysicsEngine/BodyInstance.cpp
+++ b/Engine/Source/Runtime/Engine/Private/PhysicsEngine/BodyInstance.cpp
@@ -967,8 +967,8 @@ void FBodyInstance::UpdatePhysicsFilterData()
 			if (SetupShapeIndex < BI->BodySetup->AggGeom.GetElementCount())
 			{
 				// Get the shape's CollisionEnabled masked with the body's CollisionEnabled and compute the shape's collisionflags.
-				const ECollisionEnabled::Type ShapeCollisionEnabled = CollisionEnabledIntersection(BI->GetCollisionEnabled(), BI->GetShapeCollisionEnabled(SetupShapeIndex));
-				BuildBodyCollisionFlags(PerShapeCollisionData.CollisionFlags, ShapeCollisionEnabled, bInstanceComplexAsSimple);
+				const ECollisionEnabled::Type FilteredShapeCollision = CollisionEnabledIntersection(BI->GetCollisionEnabled(), BI->GetShapeCollisionEnabled(SetupShapeIndex));
+				BuildBodyCollisionFlags(PerShapeCollisionData.CollisionFlags, FilteredShapeCollision, bInstanceComplexAsSimple);
 			}
 			else
 			{
