Add contribute-to-mass checkbox to PHAT for individual primitives

diff --git a/Engine/Source/Editor/PhysicsAssetEditor/Private/PhysicsAssetEditor.cpp b/Engine/Source/Editor/PhysicsAssetEditor/Private/PhysicsAssetEditor.cpp
--- a/Engine/Source/Editor/PhysicsAssetEditor/Private/PhysicsAssetEditor.cpp
+++ b/Engine/Source/Editor/PhysicsAssetEditor/Private/PhysicsAssetEditor.cpp
@@ -699,6 +699,12 @@ void FPhysicsAssetEditor::BindCommands()
 		FCanExecuteAction::CreateSP(this, &FPhysicsAssetEditor::CanSetPrimitiveCollision, ECollisionEnabled::NoCollision),
 		FIsActionChecked::CreateSP(this, &FPhysicsAssetEditor::IsPrimitiveCollisionChecked, ECollisionEnabled::NoCollision));
 
+	ToolkitCommands->MapAction(
+		Commands.PrimitiveContributeToMass,
+		FExecuteAction::CreateSP(this, &FPhysicsAssetEditor::OnSetPrimitiveContributeToMass),
+		FCanExecuteAction::CreateSP(this, &FPhysicsAssetEditor::CanSetPrimitiveContributeToMass),
+		FIsActionChecked::CreateSP(this, &FPhysicsAssetEditor::GetPrimitiveContributeToMass));
+
 	ToolkitCommands->MapAction(
 		Commands.WeldToBody,
 		FExecuteAction::CreateSP(this, &FPhysicsAssetEditor::OnWeldToBody),
@@ -1190,6 +1196,9 @@ void FPhysicsAssetEditor::BuildMenuWidgetBody(FMenuBuilder& InMenuBuilder)
 				InSubMenuBuilder.AddMenuEntry(PhysicsAssetEditorCommands.PrimitiveNoCollision);
 				InSubMenuBuilder.EndSection();
 #endif
+				InSubMenuBuilder.BeginSection("MassHeader", LOCTEXT("MassHeader", "Mass"));
+				InSubMenuBuilder.AddMenuEntry(PhysicsAssetEditorCommands.PrimitiveContributeToMass);
+				InSubMenuBuilder.EndSection();
 			}
 		};
 
@@ -2241,6 +2250,21 @@ bool FPhysicsAssetEditor::IsPrimitiveCollisionChecked(ECollisionEnabled::Type Co
 	return SharedData->GetIsPrimitiveCollisionEnabled(CollisionEnabled);
 }
 
+void FPhysicsAssetEditor::OnSetPrimitiveContributeToMass()
+{
+	SharedData->SetPrimitiveContributeToMass(!SharedData->GetPrimitiveContributeToMass());
+}
+
+bool FPhysicsAssetEditor::CanSetPrimitiveContributeToMass() const
+{
+	return SharedData->CanSetPrimitiveContributeToMass();
+}
+
+bool FPhysicsAssetEditor::GetPrimitiveContributeToMass() const
+{
+	return SharedData->GetPrimitiveContributeToMass();
+}
+
 void FPhysicsAssetEditor::OnWeldToBody()
 {
 	SharedData->WeldSelectedBodies();
diff --git a/Engine/Source/Editor/PhysicsAssetEditor/Private/PhysicsAssetEditor.h b/Engine/Source/Editor/PhysicsAssetEditor/Private/PhysicsAssetEditor.h
--- a/Engine/Source/Editor/PhysicsAssetEditor/Private/PhysicsAssetEditor.h
+++ b/Engine/Source/Editor/PhysicsAssetEditor/Private/PhysicsAssetEditor.h
@@ -236,6 +236,9 @@ private:
 	void OnSetPrimitiveCollision(ECollisionEnabled::Type CollisionEnabled);
 	bool CanSetPrimitiveCollision(ECollisionEnabled::Type CollisionEnabled) const;
 	bool IsPrimitiveCollisionChecked(ECollisionEnabled::Type CollisionEnabled) const;
+	void OnSetPrimitiveContributeToMass();
+	bool CanSetPrimitiveContributeToMass() const;
+	bool GetPrimitiveContributeToMass() const;
 	void OnWeldToBody();
 	bool CanWeldToBody();
 	void OnAddSphere();
diff --git a/Engine/Source/Editor/PhysicsAssetEditor/Private/PhysicsAssetEditorActions.cpp b/Engine/Source/Editor/PhysicsAssetEditor/Private/PhysicsAssetEditorActions.cpp
--- a/Engine/Source/Editor/PhysicsAssetEditor/Private/PhysicsAssetEditorActions.cpp
+++ b/Engine/Source/Editor/PhysicsAssetEditor/Private/PhysicsAssetEditorActions.cpp
@@ -57,6 +57,7 @@ void FPhysicsAssetEditorCommands::RegisterCommands()
 	UI_COMMAND(PrimitiveQueryOnly, "Query Collision Only", "Enable query collision only on selected primitive(s)", EUserInterfaceActionType::RadioButton, FInputChord());
 	UI_COMMAND(PrimitivePhysicsOnly, "Physics Collision Only", "Enable physics collision only on selected primitive(s)", EUserInterfaceActionType::RadioButton, FInputChord());
 	UI_COMMAND(PrimitiveNoCollision, "No Collision", "Disable all collision on Selected primitive(s)", EUserInterfaceActionType::RadioButton, FInputChord());
+	UI_COMMAND(PrimitiveContributeToMass, "Primitive Contributes To Mass", "Toggle the contribution of selected primitive's volume to the overall mass of the body", EUserInterfaceActionType::Check, FInputChord());
 	UI_COMMAND(DuplicatePrimitive, "Duplicate", "Duplicate Selected Primitive(s)", EUserInterfaceActionType::Button, FInputChord());
 	UI_COMMAND(ResetConstraint, "Reset", "Reset Constraint", EUserInterfaceActionType::Button, FInputChord());
 	UI_COMMAND(SnapConstraint, "Snap", "Snap Constraint Translation To Bone", EUserInterfaceActionType::Button, FInputChord());
diff --git a/Engine/Source/Editor/PhysicsAssetEditor/Private/PhysicsAssetEditorActions.h b/Engine/Source/Editor/PhysicsAssetEditor/Private/PhysicsAssetEditorActions.h
--- a/Engine/Source/Editor/PhysicsAssetEditor/Private/PhysicsAssetEditorActions.h
+++ b/Engine/Source/Editor/PhysicsAssetEditor/Private/PhysicsAssetEditorActions.h
@@ -63,6 +63,7 @@ public:
 	TSharedPtr<FUICommandInfo> PrimitiveQueryOnly;
 	TSharedPtr<FUICommandInfo> PrimitivePhysicsOnly;
 	TSharedPtr<FUICommandInfo> PrimitiveNoCollision;
+	TSharedPtr<FUICommandInfo> PrimitiveContributeToMass;
 	TSharedPtr<FUICommandInfo> WeldToBody;
 	TSharedPtr<FUICommandInfo> AddSphere;
 	TSharedPtr<FUICommandInfo> AddSphyl;
diff --git a/Engine/Source/Editor/PhysicsAssetEditor/Private/PhysicsAssetEditorSharedData.cpp b/Engine/Source/Editor/PhysicsAssetEditor/Private/PhysicsAssetEditorSharedData.cpp
--- a/Engine/Source/Editor/PhysicsAssetEditor/Private/PhysicsAssetEditorSharedData.cpp
+++ b/Engine/Source/Editor/PhysicsAssetEditor/Private/PhysicsAssetEditorSharedData.cpp
@@ -879,7 +879,7 @@ void FPhysicsAssetEditorSharedData::SetPrimitiveCollision(ECollisionEnabled::Typ
 
 	for (FSelection SelectedBody : SelectedBodies)
 	{
-		PhysicsAsset->SetPrimitiveCollision(SelectedBody.Index, SelectedBody.PrimitiveIndex, CollisionEnabled);
+		PhysicsAsset->SetPrimitiveCollision(SelectedBody.Index, SelectedBody.PrimitiveType, SelectedBody.PrimitiveIndex, CollisionEnabled);
 	}
 
 	PreviewChangedEvent.Broadcast();
@@ -899,7 +899,33 @@ bool FPhysicsAssetEditorSharedData::GetIsPrimitiveCollisionEnabled(ECollisionEna
 {
 	for (const FSelection SelectedBody : SelectedBodies)
 	{
-		if (PhysicsAsset->GetPrimitiveCollision(SelectedBody.Index, SelectedBody.PrimitiveIndex) == CollisionEnabled)
+		if (PhysicsAsset->GetPrimitiveCollision(SelectedBody.Index, SelectedBody.PrimitiveType, SelectedBody.PrimitiveIndex) == CollisionEnabled)
+		{
+			return true;
+		}
+	}
+
+	return false;
+}
+
+void FPhysicsAssetEditorSharedData::SetPrimitiveContributeToMass(bool bContributeToMass)
+{
+	for (const FSelection SelectedBody : SelectedBodies)
+	{
+		PhysicsAsset->SetPrimitiveContributeToMass(SelectedBody.Index, SelectedBody.PrimitiveType, SelectedBody.PrimitiveIndex, bContributeToMass);
+	}
+}
+
+bool FPhysicsAssetEditorSharedData::CanSetPrimitiveContributeToMass() const
+{
+	return true;
+}
+
+bool FPhysicsAssetEditorSharedData::GetPrimitiveContributeToMass() const
+{
+	for (const FSelection SelectedBody : SelectedBodies)
+	{
+		if (PhysicsAsset->GetPrimitiveContributeToMass(SelectedBody.Index, SelectedBody.PrimitiveType, SelectedBody.PrimitiveIndex))
 		{
 			return true;
 		}
diff --git a/Engine/Source/Editor/PhysicsAssetEditor/Private/PhysicsAssetEditorSharedData.h b/Engine/Source/Editor/PhysicsAssetEditor/Private/PhysicsAssetEditorSharedData.h
--- a/Engine/Source/Editor/PhysicsAssetEditor/Private/PhysicsAssetEditorSharedData.h
+++ b/Engine/Source/Editor/PhysicsAssetEditor/Private/PhysicsAssetEditorSharedData.h
@@ -193,6 +193,9 @@ public:
 	void SetPrimitiveCollision(ECollisionEnabled::Type CollisionEnabled);
 	bool CanSetPrimitiveCollision(ECollisionEnabled::Type CollisionEnabled) const;
 	bool GetIsPrimitiveCollisionEnabled(ECollisionEnabled::Type CollisionEnabled) const;
+	void SetPrimitiveContributeToMass(bool bContributeToMass);
+	bool CanSetPrimitiveContributeToMass() const;
+	bool GetPrimitiveContributeToMass() const;
 
 	/** Prevents GC from collecting our objects */
 	void AddReferencedObjects(FReferenceCollector& Collector);
diff --git a/Engine/Source/Runtime/Engine/Classes/PhysicsEngine/AggregateGeom.h b/Engine/Source/Runtime/Engine/Classes/PhysicsEngine/AggregateGeom.h
--- a/Engine/Source/Runtime/Engine/Classes/PhysicsEngine/AggregateGeom.h
+++ b/Engine/Source/Runtime/Engine/Classes/PhysicsEngine/AggregateGeom.h
@@ -61,6 +61,26 @@ struct ENGINE_API FKAggregateGeom
 
 	int32 GetElementCount(EAggCollisionShape::Type Type) const;
 
+	FKShapeElem* GetElement(const EAggCollisionShape::Type Type, const int32 Index)
+	{
+		switch (Type)
+		{
+		case EAggCollisionShape::Sphere:
+			if (ensure(SphereElems.IsValidIndex(Index))) { return &SphereElems[Index]; }
+		case EAggCollisionShape::Box:
+			if (ensure(BoxElems.IsValidIndex(Index))) { return &BoxElems[Index]; }
+		case EAggCollisionShape::Sphyl:
+			if (ensure(SphylElems.IsValidIndex(Index))) { return &SphylElems[Index]; }
+		case EAggCollisionShape::Convex:
+			if (ensure(ConvexElems.IsValidIndex(Index))) { return &ConvexElems[Index]; }
+		case EAggCollisionShape::TaperedCapsule:
+			if (ensure(TaperedCapsuleElems.IsValidIndex(Index))) { return &TaperedCapsuleElems[Index]; }
+		default:
+			ensure(false);
+			return nullptr;
+		}
+	}
+
 	FKShapeElem* GetElement(const int32 InIndex)
 	{
 		int Index = InIndex;
diff --git a/Engine/Source/Runtime/Engine/Classes/PhysicsEngine/PhysicsAsset.h b/Engine/Source/Runtime/Engine/Classes/PhysicsEngine/PhysicsAsset.h
--- a/Engine/Source/Runtime/Engine/Classes/PhysicsEngine/PhysicsAsset.h
+++ b/Engine/Source/Runtime/Engine/Classes/PhysicsEngine/PhysicsAsset.h
@@ -229,10 +229,16 @@ public:
 	ENGINE_API bool IsCollisionEnabled(int32 BodyIndexA, int32 BodyIndexB) const;
 
 	// Get the per-primitive collision filtering mode for a body
-	ENGINE_API void SetPrimitiveCollision(int32 BodyIndex, int32 PrimitiveIndex, ECollisionEnabled::Type CollisionEnabled);
+	ENGINE_API void SetPrimitiveCollision(int32 BodyIndex, EAggCollisionShape::Type PrimitiveType, int32 PrimitiveIndex, ECollisionEnabled::Type CollisionEnabled);
 
 	// Get the per-primitive collision filtering mode for a body
-	ENGINE_API ECollisionEnabled::Type GetPrimitiveCollision(int32 BodyIndex, int32 PrimitiveIndex) const;
+	ENGINE_API ECollisionEnabled::Type GetPrimitiveCollision(int32 BodyIndex, EAggCollisionShape::Type PrimitiveType, int32 PrimitiveIndex) const;
+
+	// Set whether or not a primitive volume contributes to the mass of the object
+	ENGINE_API void SetPrimitiveContributeToMass(int32 BodyIndex, EAggCollisionShape::Type PrimitiveType, int32 PrimitiveIndex, bool bContributesToMass);
+
+	// Get whether or not a primitive volume contributes to the mass of the object
+	ENGINE_API bool GetPrimitiveContributeToMass(int32 BodyIndex, EAggCollisionShape::Type PrimitiveType, int32 PrimitiveIndex) const;
 
 	/** Update the BoundsBodies array and cache the indices of bodies marked with bConsiderForBounds to BoundsBodies array. */
 	ENGINE_API void UpdateBoundsBodiesArray();
diff --git a/Engine/Source/Runtime/Engine/Private/PhysicsEngine/PhysicsAsset.cpp b/Engine/Source/Runtime/Engine/Private/PhysicsEngine/PhysicsAsset.cpp
--- a/Engine/Source/Runtime/Engine/Private/PhysicsEngine/PhysicsAsset.cpp
+++ b/Engine/Source/Runtime/Engine/Private/PhysicsEngine/PhysicsAsset.cpp
@@ -253,22 +253,36 @@ bool UPhysicsAsset::IsCollisionEnabled(int32 BodyIndexA, int32 BodyIndexB) const
 	return true;
 }
 
-void UPhysicsAsset::SetPrimitiveCollision(int32 BodyIndex, int32 PrimitiveIndex, ECollisionEnabled::Type CollisionEnabled)
+void UPhysicsAsset::SetPrimitiveCollision(int32 BodyIndex, EAggCollisionShape::Type PrimitiveType, int32 PrimitiveIndex, ECollisionEnabled::Type CollisionEnabled)
 {
-#if WITH_CHAOS
 	check(SkeletalBodySetups.IsValidIndex(BodyIndex));
 	FKAggregateGeom* AggGeom = &SkeletalBodySetups[BodyIndex]->AggGeom;
 	ensure(PrimitiveIndex < AggGeom->GetElementCount());
-	AggGeom->GetElement(PrimitiveIndex)->SetCollisionEnabled(CollisionEnabled);
-#endif
+	AggGeom->GetElement(PrimitiveType, PrimitiveIndex)->SetCollisionEnabled(CollisionEnabled);
+}
+
+ECollisionEnabled::Type UPhysicsAsset::GetPrimitiveCollision(int32 BodyIndex, EAggCollisionShape::Type PrimitiveType, int32 PrimitiveIndex) const
+{
+	check(SkeletalBodySetups.IsValidIndex(BodyIndex));
+	FKAggregateGeom* AggGeom = &SkeletalBodySetups[BodyIndex]->AggGeom;
+	ensure(PrimitiveIndex < AggGeom->GetElementCount());
+	return AggGeom->GetElement(PrimitiveType, PrimitiveIndex)->GetCollisionEnabled();
+}
+
+void UPhysicsAsset::SetPrimitiveContributeToMass(int32 BodyIndex, EAggCollisionShape::Type PrimitiveType, int32 PrimitiveIndex, bool bContributeToMass)
+{
+	check(SkeletalBodySetups.IsValidIndex(BodyIndex));
+	FKAggregateGeom* AggGeom = &SkeletalBodySetups[BodyIndex]->AggGeom;
+	ensure(PrimitiveIndex < AggGeom->GetElementCount());
+	return AggGeom->GetElement(PrimitiveType, PrimitiveIndex)->SetContributeToMass(bContributeToMass);
 }
 
-ECollisionEnabled::Type UPhysicsAsset::GetPrimitiveCollision(int32 BodyIndex, int32 PrimitiveIndex) const
+bool UPhysicsAsset::GetPrimitiveContributeToMass(int32 BodyIndex, EAggCollisionShape::Type PrimitiveType, int32 PrimitiveIndex) const
 {
 	check(SkeletalBodySetups.IsValidIndex(BodyIndex));
 	FKAggregateGeom* AggGeom = &SkeletalBodySetups[BodyIndex]->AggGeom;
 	ensure(PrimitiveIndex < AggGeom->GetElementCount());
-	return AggGeom->GetElement(PrimitiveIndex)->GetCollisionEnabled();
+	return AggGeom->GetElement(PrimitiveType, PrimitiveIndex)->GetContributeToMass();
 }
 
 FBox UPhysicsAsset::CalcAABB(const USkinnedMeshComponent* MeshComp, const FTransform& LocalToWorld) const
