Fixed build.

diff --git a/Engine/Source/Runtime/Engine/Private/Materials/Material.cpp b/Engine/Source/Runtime/Engine/Private/Materials/Material.cpp
--- a/Engine/Source/Runtime/Engine/Private/Materials/Material.cpp
+++ b/Engine/Source/Runtime/Engine/Private/Materials/Material.cpp
@@ -1031,14 +1031,14 @@ void UMaterial::GetUsedTextures(TArray<UTexture*>& OutTextures, EMaterialQuality
 				if (MaterialInstance)
 				{
 					// Also look for any scalar parameters that are acting as lookups for an atlas texture, and store the atlas texture
-					const TArray<FMaterialScalarParameterInfo, FMemoryImageAllocator>* AtlasExpressions[1] =
+					const TArrayView<const FMaterialScalarParameterInfo> AtlasExpressions[1] =
 					{
-						&CurrentResource->GetUniformScalarParameterExpressions()
+						CurrentResource->GetUniformScalarParameterExpressions()
 					};
 					for (int32 TypeIndex = 0; TypeIndex < UE_ARRAY_COUNT(AtlasExpressions); TypeIndex++)
 					{
 						// Iterate over each of the material's texture expressions.
-						for (const FMaterialScalarParameterInfo& Parameter : *AtlasExpressions[TypeIndex])
+						for (const FMaterialScalarParameterInfo& Parameter : AtlasExpressions[TypeIndex])
 						{
 							bool bIsUsedAsAtlasPosition;
 							TSoftObjectPtr<class UCurveLinearColor> Curve;
@@ -1176,7 +1176,7 @@ void UMaterial::OverrideTexture(const UTexture* InTextureToOverride, UTexture* O
 		// Iterate over both the 2D textures and cube texture expressions.
 		for(int32 TypeIndex = 0;TypeIndex < NumMaterialTextureParameterTypes; TypeIndex++)
 		{
-			const TArray<FMaterialTextureParameterInfo, FMemoryImageAllocator>& Parameters = Resource->GetUniformTextureExpressions((EMaterialTextureParameterType)TypeIndex);
+			const TArrayView<const FMaterialTextureParameterInfo> Parameters = Resource->GetUniformTextureExpressions((EMaterialTextureParameterType)TypeIndex);
 			// Iterate over each of the material's texture expressions.
 			for (int32 ParameterIndex = 0; ParameterIndex < Parameters.Num(); ++ParameterIndex)
 			{
@@ -1208,7 +1208,7 @@ void UMaterial::OverrideVectorParameterDefault(const FHashedMaterialParameterInf
 	bool bShouldRecacheMaterialExpressions = false;
 
 	FMaterialResource* Resource = GetMaterialResource(InFeatureLevel);
-	const TArray<FMaterialVectorParameterInfo, FMemoryImageAllocator>& Parameters = Resource->GetUniformVectorParameterExpressions();
+	const TArrayView<const FMaterialVectorParameterInfo> Parameters = Resource->GetUniformVectorParameterExpressions();
 
 	// Iterate over each of the material's vector expressions.
 	for (int32 i = 0; i < Parameters.Num(); ++i)
@@ -1234,7 +1234,7 @@ void UMaterial::OverrideScalarParameterDefault(const FHashedMaterialParameterInf
 	bool bShouldRecacheMaterialExpressions = false;
 
 	FMaterialResource* Resource = GetMaterialResource(InFeatureLevel);
-	const TArray<FMaterialScalarParameterInfo, FMemoryImageAllocator>& Parameters = Resource->GetUniformScalarParameterExpressions();
+	const TArrayView<const FMaterialScalarParameterInfo> Parameters = Resource->GetUniformScalarParameterExpressions();
 
 	// Iterate over each of the material's vector expressions.
 	for (int32 i = 0; i < Parameters.Num(); ++i)
diff --git a/Engine/Source/Runtime/Engine/Private/Materials/MaterialInstance.cpp b/Engine/Source/Runtime/Engine/Private/Materials/MaterialInstance.cpp
--- a/Engine/Source/Runtime/Engine/Private/Materials/MaterialInstance.cpp
+++ b/Engine/Source/Runtime/Engine/Private/Materials/MaterialInstance.cpp
@@ -1476,7 +1476,7 @@ void UMaterialInstance::OverrideTexture(const UTexture* InTextureToOverride, UTe
 		
 	for (int32 TypeIndex = 0; TypeIndex < NumMaterialTextureParameterTypes; TypeIndex++)
 	{
-		const TArray<FMaterialTextureParameterInfo, FMemoryImageAllocator>& Parameters = SourceMaterialResource->GetUniformTextureExpressions((EMaterialTextureParameterType)TypeIndex);
+		const TArrayView<const FMaterialTextureParameterInfo> Parameters = SourceMaterialResource->GetUniformTextureExpressions((EMaterialTextureParameterType)TypeIndex);
 		// Iterate over each of the material's texture expressions.
 		for (int32 i = 0; i < Parameters.Num(); ++i)
 		{
@@ -1508,7 +1508,7 @@ void UMaterialInstance::OverrideVectorParameterDefault(const FHashedMaterialPara
 	if (bHasStaticPermutationResource)
 	{
 		FMaterialResource* SourceMaterialResource = GetMaterialResource(InFeatureLevel);
-		const TArray<FMaterialVectorParameterInfo, FMemoryImageAllocator>& Parameters = SourceMaterialResource->GetUniformVectorParameterExpressions();
+		const TArrayView<const FMaterialVectorParameterInfo> Parameters = SourceMaterialResource->GetUniformVectorParameterExpressions();
 
 		for (int32 i = 0; i < Parameters.Num(); ++i)
 		{
@@ -1535,7 +1535,7 @@ void UMaterialInstance::OverrideScalarParameterDefault(const FHashedMaterialPara
 	if (bHasStaticPermutationResource)
 	{
 		FMaterialResource* SourceMaterialResource = GetMaterialResource(InFeatureLevel);
-		const TArray<FMaterialScalarParameterInfo, FMemoryImageAllocator>& Parameters = SourceMaterialResource->GetUniformScalarParameterExpressions();
+		const TArrayView<const FMaterialScalarParameterInfo> Parameters = SourceMaterialResource->GetUniformScalarParameterExpressions();
 
 		for (int32 i = 0; i < Parameters.Num(); ++i)
 		{
