Rename NetworkDriver.cpp to NetDriver.cpp to be consistent with class name and coding standards.

diff --git a/Engine/Source/Runtime/Engine/Private/NetDriver.cpp b/Engine/Source/Runtime/Engine/Private/NetDriver.cpp
new file mode 100644
--- /dev/null
+++ b/Engine/Source/Runtime/Engine/Private/NetDriver.cpp
@@ -0,0 +1,5976 @@
+// Copyright Epic Games, Inc. All Rights Reserved.
+
+/*=============================================================================
+	NetworkDriver.cpp: Unreal network driver base class.
+=============================================================================*/
+
+#include "Engine/NetDriver.h"
+#include "CoreMinimal.h"
+#include "Misc/CoreMisc.h"
+#include "Misc/CommandLine.h"
+#include "Misc/NetworkGuid.h"
+#include "Stats/Stats.h"
+#include "Misc/App.h"
+#include "Misc/MemStack.h"
+#include "HAL/IConsoleManager.h"
+#include "HAL/LowLevelMemTracker.h"
+#include "UObject/ObjectMacros.h"
+#include "UObject/Object.h"
+#include "UObject/Class.h"
+#include "UObject/CoreNet.h"
+#include "UObject/UnrealType.h"
+#include "UObject/Package.h"
+#include "UObject/PropertyPortFlags.h"
+#include "EngineStats.h"
+#include "EngineGlobals.h"
+#include "Engine/EngineBaseTypes.h"
+#include "Engine/EngineTypes.h"
+#include "Components/ActorComponent.h"
+#include "Engine/Level.h"
+#include "GameFramework/Actor.h"
+#include "GameFramework/Pawn.h"
+#include "CollisionQueryParams.h"
+#include "GameFramework/GameModeBase.h"
+#include "Components/PrimitiveComponent.h"
+#include "Misc/ConfigCacheIni.h"
+#include "UObject/UObjectIterator.h"
+#include "GameFramework/PlayerController.h"
+#include "GameFramework/WorldSettings.h"
+#include "PacketHandler.h"
+#include "PacketHandlers/StatelessConnectHandlerComponent.h"
+#include "Net/Core/Analytics/NetAnalytics.h"
+#include "Engine/LocalPlayer.h"
+#include "Net/DataBunch.h"
+#include "Engine/NetConnection.h"
+#include "DrawDebugHelpers.h"
+#include "UnrealEngine.h"
+#include "EngineUtils.h"
+#include "Net/NetworkProfiler.h"
+#include "Engine/PackageMapClient.h"
+#include "Net/RepLayout.h"
+#include "Net/DataReplication.h"
+#include "Engine/ControlChannel.h"
+#include "Engine/ActorChannel.h"
+#include "Engine/VoiceChannel.h"
+#include "Engine/NetworkObjectList.h"
+#include "GameFramework/GameNetworkManager.h"
+#include "Net/OnlineEngineInterface.h"
+#include "NetworkingDistanceConstants.h"
+#include "Engine/ChildConnection.h"
+#include "Net/Core/Trace/NetTrace.h"
+#include "Misc/ScopeExit.h"
+#include "Net/DataChannel.h"
+#include "GameFramework/PlayerState.h"
+#include "Net/PerfCountersHelpers.h"
+#include "Stats/StatsMisc.h"
+#include "Engine/ReplicationDriver.h"
+#include "ProfilingDebugging/CsvProfiler.h"
+#include "Engine/LevelScriptActor.h"
+#include "Engine/NetworkSettings.h"
+#include "Net/NetworkGranularMemoryLogging.h"
+#include "SocketSubsystem.h"
+#include "AddressInfoTypes.h"
+#if USE_SERVER_PERF_COUNTERS
+#include "PerfCountersModule.h"
+#endif
+
+#if WITH_EDITOR
+#include "Editor.h"
+#endif
+
+// Default net driver stats
+DEFINE_STAT(STAT_Ping);
+DEFINE_STAT(STAT_Channels);
+DEFINE_STAT(STAT_MaxPacketOverhead);
+DEFINE_STAT(STAT_InRate);
+DEFINE_STAT(STAT_OutRate);
+DEFINE_STAT(STAT_InRateClientMax);
+DEFINE_STAT(STAT_InRateClientMin);
+DEFINE_STAT(STAT_InRateClientAvg);
+DEFINE_STAT(STAT_InPacketsClientMax);
+DEFINE_STAT(STAT_InPacketsClientMin);
+DEFINE_STAT(STAT_InPacketsClientAvg);
+DEFINE_STAT(STAT_OutRateClientMax);
+DEFINE_STAT(STAT_OutRateClientMin);
+DEFINE_STAT(STAT_OutRateClientAvg);
+DEFINE_STAT(STAT_OutPacketsClientMax);
+DEFINE_STAT(STAT_OutPacketsClientMin);
+DEFINE_STAT(STAT_OutPacketsClientAvg);
+DEFINE_STAT(STAT_NetNumClients);
+DEFINE_STAT(STAT_InPackets);
+DEFINE_STAT(STAT_OutPackets);
+DEFINE_STAT(STAT_InBunches);
+DEFINE_STAT(STAT_OutBunches);
+DEFINE_STAT(STAT_OutLoss);
+DEFINE_STAT(STAT_InLoss);
+DEFINE_STAT(STAT_NumConsideredActors);
+DEFINE_STAT(STAT_PrioritizedActors);
+DEFINE_STAT(STAT_NumReplicatedActors);
+DEFINE_STAT(STAT_NumReplicatedActorBytes);
+DEFINE_STAT(STAT_NumRelevantDeletedActors);
+DEFINE_STAT(STAT_NumActorChannels);
+DEFINE_STAT(STAT_NumActors);
+DEFINE_STAT(STAT_NumNetActors);
+DEFINE_STAT(STAT_NumDormantActors);
+DEFINE_STAT(STAT_NumInitiallyDormantActors);
+DEFINE_STAT(STAT_NumNetGUIDsAckd);
+DEFINE_STAT(STAT_NumNetGUIDsPending);
+DEFINE_STAT(STAT_NumNetGUIDsUnAckd);
+DEFINE_STAT(STAT_ObjPathBytes);
+DEFINE_STAT(STAT_NetGUIDInRate);
+DEFINE_STAT(STAT_NetGUIDOutRate);
+DEFINE_STAT(STAT_NetSaturated);
+DEFINE_STAT(STAT_ImportedNetGuids);
+DEFINE_STAT(STAT_PendingOuterNetGuids);
+DEFINE_STAT(STAT_UnmappedReplicators);
+
+// Voice specific stats
+DEFINE_STAT(STAT_VoiceBytesSent);
+DEFINE_STAT(STAT_VoiceBytesRecv);
+DEFINE_STAT(STAT_VoicePacketsSent);
+DEFINE_STAT(STAT_VoicePacketsRecv);
+DEFINE_STAT(STAT_PercentInVoice);
+DEFINE_STAT(STAT_PercentOutVoice);
+
+#if !UE_BUILD_SHIPPING
+// Packet stats
+DEFINE_STAT(STAT_MaxPacket);
+DEFINE_STAT(STAT_MaxPacketMinusReserved);
+DEFINE_STAT(STAT_PacketReservedTotal);
+DEFINE_STAT(STAT_PacketReservedNetConnection);
+DEFINE_STAT(STAT_PacketReservedPacketHandler);
+DEFINE_STAT(STAT_PacketReservedHandshake);
+#endif
+
+DECLARE_CYCLE_STAT(TEXT("NetDriver AddClientConnection"), Stat_NetDriverAddClientConnection, STATGROUP_Net);
+DECLARE_CYCLE_STAT(TEXT("NetDriver ProcessRemoteFunction"), STAT_NetProcessRemoteFunc, STATGROUP_Net);
+DECLARE_CYCLE_STAT(TEXT("Process Prioritized Actors Time"), STAT_NetProcessPrioritizedActorsTime, STATGROUP_Game);
+DECLARE_CYCLE_STAT(TEXT("NetDriver TickFlush"), STAT_NetTickFlush, STATGROUP_Game);
+DECLARE_CYCLE_STAT(TEXT("NetDriver TickFlush GatherStats"), STAT_NetTickFlushGatherStats, STATGROUP_Game);
+DECLARE_CYCLE_STAT(TEXT("NetDriver TickFlush GatherStatsPerfCounters"), STAT_NetTickFlushGatherStatsPerfCounters, STATGROUP_Game);
+
+int32 GNumSaturatedConnections; // Counter for how many connections are skipped/early out due to bandwidth saturation
+int32 GNumSharedSerializationHit;
+int32 GNumSharedSerializationMiss;
+
+extern int32 GNetRPCDebug;
+
+namespace NetEmulationHelper
+{
+#if DO_ENABLE_NET_TEST
+	/** Global that stores the network emulation values outside the NetDriver lifetime */
+	static TOptional<FPacketSimulationSettings> PersistentPacketSimulationSettings;
+
+	void CreatePersistentSimulationSettings()
+	{
+		if (!PersistentPacketSimulationSettings.IsSet())
+		{
+			PersistentPacketSimulationSettings.Emplace(FPacketSimulationSettings());
+		}
+	}
+
+	void ApplySimulationSettingsOnNetDrivers(UWorld* World, const FPacketSimulationSettings& Settings)
+	{
+		// Execute on all active NetDrivers
+		FWorldContext& Context = GEngine->GetWorldContextFromWorldChecked(World);
+		for (const FNamedNetDriver& ActiveNetDriver : Context.ActiveNetDrivers)
+		{
+			if (ActiveNetDriver.NetDriver)
+			{
+				ActiveNetDriver.NetDriver->SetPacketSimulationSettings(Settings);
+			}
+		}
+	}
+#endif //#if DO_ENABLE_NET_TEST
+}
+
+#if UE_BUILD_SHIPPING
+#define DEBUG_REMOTEFUNCTION(Format, ...)
+#else
+#define DEBUG_REMOTEFUNCTION(Format, ...) UE_LOG(LogNet, VeryVerbose, Format, __VA_ARGS__);
+#endif
+
+// CVars
+int32 GSetNetDormancyEnabled = 1;
+static FAutoConsoleVariableRef CVarSetNetDormancyEnabled(
+	TEXT("net.DormancyEnable"),
+	GSetNetDormancyEnabled,
+	TEXT("Enables Network Dormancy System for reducing CPU and bandwidth overhead of infrequently updated actors\n")
+	TEXT("1 Enables network dormancy. 0 disables network dormancy."),
+	ECVF_Default);
+
+int32 GNetDormancyValidate = 0;
+static FAutoConsoleVariableRef CVarNetDormancyValidate(
+	TEXT("net.DormancyValidate"),
+	GNetDormancyValidate,
+	TEXT("Validates that dormant actors do not change state while in a dormant state (on server only)")
+	TEXT("0: Dont validate. 1: Validate on wake up. 2: Validate on each net update"),
+	ECVF_Default);
+
+static TAutoConsoleVariable<int32> CVarNetDebugDraw(
+	TEXT("net.DebugDraw"),
+	0,
+	TEXT("Draws debug information for network dormancy and relevancy\n")
+	TEXT("1 Enables network debug drawing. 0 disables."),
+	ECVF_Default);
+
+static TAutoConsoleVariable<float> CVarNetDebugDrawCullDistance(
+	TEXT("net.DebugDrawCullDistance"),
+	0.f,
+	TEXT("Cull distance for net.DebugDraw. World Units")
+	TEXT("Max world units an actor can be away from the local view to draw its dormancy status. Zero disables culling"),
+	ECVF_Default);
+
+static TAutoConsoleVariable<int32> CVarUseAdaptiveNetUpdateFrequency(
+	TEXT( "net.UseAdaptiveNetUpdateFrequency" ), 
+	0, 
+	TEXT( "If 1, NetUpdateFrequency will be calculated based on how often actors actually send something when replicating" ) );
+
+TAutoConsoleVariable<int32> CVarNetAllowEncryption(
+	TEXT("net.AllowEncryption"),
+	1,
+	TEXT("If true, the engine will attempt to load an encryption PacketHandler component and fill in the EncryptionToken parameter of the NMT_Hello message based on the ?EncryptionToken= URL option and call callbacks if it's non-empty."));
+
+static TAutoConsoleVariable<int32> CVarActorChannelPool(
+	TEXT("net.ActorChannelPool"),
+	1,
+	TEXT("If nonzero, actor channels will be pooled to save memory and object creation cost."));
+
+static TAutoConsoleVariable<int32> CVarAllowReliableMulticastToNonRelevantChannels(
+	TEXT("net.AllowReliableMulticastToNonRelevantChannels"),
+	1,
+	TEXT("Allow Reliable Server Multicasts to be sent to non-Relevant Actors, as long as their is an existing ActorChannel."));
+
+static int32 GNetControlChannelDestructionInfo = 0;
+static FAutoConsoleVariableRef CVarNetControlChannelDestructionInfo(
+	TEXT("net.ControlChannelDestructionInfo"),
+	GNetControlChannelDestructionInfo,
+	TEXT("If enabled, send destruction info updates via the control channel instead of creating a new actor channel.")
+	TEXT("0: Old behavior, use an actor channel. 1: New behavior, use the control channel"),
+	ECVF_Default);
+
+
+/*-----------------------------------------------------------------------------
+	UNetDriver implementation.
+-----------------------------------------------------------------------------*/
+
+UNetDriver::UNetDriver(const FObjectInitializer& ObjectInitializer)
+:	UObject(ObjectInitializer)
+,	MaxInternetClientRate(10000)
+, 	MaxClientRate(15000)
+,   bNoTimeouts(false)
+,   ServerConnection(nullptr)
+,	ClientConnections()
+,	MappedClientConnections()
+,	RecentlyDisconnectedClients()
+,	RecentlyDisconnectedTrackingTime(0)
+,	ConnectionlessHandler()
+,	StatelessConnectComponent()
+,	AnalyticsProvider()
+,	AnalyticsAggregator()
+,   World(nullptr)
+,   Notify(nullptr)
+,	Time( 0.f )
+,	ElapsedTime( 0.0 )
+,	LastTickDispatchRealtime( 0.f )
+,   bIsPeer(false)
+,	bSkipLocalStats(false)
+,	InBytes(0)
+,	InTotalBytes(0)
+,	OutBytes(0)
+,	OutTotalBytes(0)
+,	NetGUIDOutBytes(0)
+,	NetGUIDInBytes(0)
+,	InPackets(0)
+,	InTotalPackets(0)
+,	OutPackets(0)
+,	OutTotalPackets(0)
+,	InBunches(0)
+,	OutBunches(0)
+,	InTotalBunches(0)
+,	OutTotalBunches(0)
+,	InPacketsLost(0)
+,	InTotalPacketsLost(0)
+,	OutPacketsLost(0)
+,	OutTotalPacketsLost(0)
+,	InOutOfOrderPackets(0)
+,	OutOutOfOrderPackets(0)
+,	StatUpdateTime(0.0)
+,	StatPeriod(1.f)
+,	TotalRPCsCalled(0)
+,	OutTotalAcks(0)
+,	bCollectNetStats(false)
+,	LastCleanupTime(0.0)
+,	NetTag(0)
+,	DebugRelevantActors(false)
+#if !UE_BUILD_SHIPPING
+,	SendRPCDel()
+#endif
+,	ProcessQueuedBunchesCurrentFrameMilliseconds(0.0f)
+,	DDoS()
+,	LocalAddr(nullptr)
+,	NetworkObjects(new FNetworkObjectList)
+,	LagState(ENetworkLagState::NotLagging)
+,	DuplicateLevelID(INDEX_NONE)
+,	PacketLossBurstEndTime(-1.0f)
+,	OutTotalNotifiedPackets(0u)
+{
+	UpdateDelayRandomStream.Initialize(FApp::bUseFixedSeed ? GetFName() : NAME_None);
+}
+
+void UNetDriver::InitPacketSimulationSettings()
+{
+#if DO_ENABLE_NET_TEST
+	if (bNeverApplyNetworkEmulationSettings)
+	{
+		return;
+	}
+
+	if (bForcedPacketSettings)
+	{
+		return;
+	}
+
+	if (NetEmulationHelper::PersistentPacketSimulationSettings.IsSet())
+	{
+		SetPacketSimulationSettings(NetEmulationHelper::PersistentPacketSimulationSettings.GetValue());
+		return;
+	}
+
+	// Read the settings from .ini and command line, with the command line taking precedence	
+	PacketSimulationSettings.ResetSettings();
+	PacketSimulationSettings.LoadConfig(*NetDriverName.ToString());
+	PacketSimulationSettings.ParseSettings(FCommandLine::Get(), *NetDriverName.ToString());
+#endif
+}
+
+#if DO_ENABLE_NET_TEST
+bool UNetDriver::IsSimulatingPacketLossBurst() const
+{
+	return PacketLossBurstEndTime > ElapsedTime;
+}
+#endif
+
+void UNetDriver::PostInitProperties()
+{
+	Super::PostInitProperties();
+
+	// By default we're the game net driver and any child ones must override this
+	NetDriverName = NAME_GameNetDriver;
+
+	if ( !HasAnyFlags(RF_ClassDefaultObject) )
+	{
+		// Init ConnectionIdHandler
+		{
+			constexpr uint32 IdCount = 1024;
+			ConnectionIdHandler.Init(IdCount);
+		}
+
+		InitPacketSimulationSettings();
+
+		
+		RoleProperty		= FindFieldChecked<FProperty>( AActor::StaticClass(), TEXT("Role"      ) );
+		RemoteRoleProperty	= FindFieldChecked<FProperty>( AActor::StaticClass(), TEXT("RemoteRole") );
+
+		GuidCache			= TSharedPtr< FNetGUIDCache >( new FNetGUIDCache( this ) );
+		NetCache			= TSharedPtr< FClassNetCacheMgr >( new FClassNetCacheMgr() );
+
+		ProfileStats		= FParse::Param(FCommandLine::Get(),TEXT("profilestats"));
+
+#if !UE_BUILD_SHIPPING
+		bNoTimeouts = bNoTimeouts || FParse::Param(FCommandLine::Get(), TEXT("NoTimeouts")) ? true : false;
+#endif // !UE_BUILD_SHIPPING
+
+#if WITH_EDITOR
+		// Do not time out in PIE since the server is local.
+		bNoTimeouts = bNoTimeouts || (GEditor && GEditor->PlayWorld);
+#endif // WITH_EDITOR
+	
+		OnLevelRemovedFromWorldHandle = FWorldDelegates::LevelRemovedFromWorld.AddUObject(this, &UNetDriver::OnLevelRemovedFromWorld);
+		OnLevelAddedToWorldHandle = FWorldDelegates::LevelAddedToWorld.AddUObject(this, &UNetDriver::OnLevelAddedToWorld);
+		PostGarbageCollectHandle = FCoreUObjectDelegates::GetPostGarbageCollect().AddUObject(this, &UNetDriver::PostGarbageCollect);
+
+		LoadChannelDefinitions();
+	}
+}
+
+void UNetDriver::PostReloadConfig(FProperty* PropertyToLoad)
+{
+	Super::PostReloadConfig(PropertyToLoad);
+	LoadChannelDefinitions();
+}
+
+void UNetDriver::LoadChannelDefinitions()
+{
+	if (!HasAnyFlags(RF_ClassDefaultObject))
+	{
+		if (ServerConnection != nullptr)
+		{
+			UE_LOG(LogNet, Warning, TEXT("Loading channel definitions while a server connection is active."));
+		}
+
+		if (ClientConnections.Num() != 0)
+		{
+			UE_LOG(LogNet, Warning, TEXT("Loading channel definitions while client connections are active."));
+		}
+
+		TSet<int32> StaticChannelIndices;
+
+		ChannelDefinitionMap.Reset();
+
+		for (FChannelDefinition& ChannelDef : ChannelDefinitions)
+		{
+			UE_CLOG(!ChannelDef.ChannelName.ToEName() || !ShouldReplicateAsInteger(*ChannelDef.ChannelName.ToEName()), 
+				LogNet, Warning, TEXT("Channel name will be serialized as a string: %s"), *ChannelDef.ChannelName.ToString());
+			UE_CLOG(ChannelDefinitionMap.Contains(ChannelDef.ChannelName), LogNet, Error, TEXT("Channel name is defined multiple times: %s"), *ChannelDef.ChannelName.ToString());
+			UE_CLOG(StaticChannelIndices.Contains(ChannelDef.StaticChannelIndex), LogNet, Error, TEXT("Channel static index is already in use: %s %i"), *ChannelDef.ChannelName.ToString(), ChannelDef.StaticChannelIndex);
+
+			ChannelDef.ChannelClass = StaticLoadClass(UChannel::StaticClass(), nullptr, *ChannelDef.ClassName.ToString(), nullptr, LOAD_Quiet);
+			if (ChannelDef.ChannelClass != nullptr)
+			{
+				ChannelDefinitionMap.Add(ChannelDef.ChannelName, ChannelDef);
+			}
+
+			if (ChannelDef.StaticChannelIndex != INDEX_NONE)
+			{
+				StaticChannelIndices.Add(ChannelDef.StaticChannelIndex);
+			}
+		}
+
+		ensureMsgf(IsKnownChannelName(NAME_Control), TEXT("Control channel type is not properly defined."));
+	}
+}
+
+void UNetDriver::AssertValid()
+{
+}
+
+/*static*/ bool UNetDriver::IsAdaptiveNetUpdateFrequencyEnabled()
+{
+	const bool bUseAdapativeNetFrequency = CVarUseAdaptiveNetUpdateFrequency.GetValueOnAnyThread() > 0;
+	return bUseAdapativeNetFrequency;
+}
+
+void UNetDriver::AddNetworkActor(AActor* Actor)
+{
+	LLM_SCOPE(ELLMTag::Networking);
+	ensureMsgf(Actor == nullptr || !(Actor->HasAnyFlags(RF_ClassDefaultObject | RF_ArchetypeObject)), TEXT("%s is a CDO or Archetype and should not be replicated."), *GetFullNameSafe(Actor));
+
+	if (!IsDormInitialStartupActor(Actor))
+	{
+		GetNetworkObjectList().FindOrAdd(Actor, this);
+		if (ReplicationDriver)
+		{
+			ReplicationDriver->AddNetworkActor(Actor);
+		}
+	}
+}
+
+void UNetDriver::SetRoleSwapOnReplicate(AActor* Actor, bool bSwapRoles)
+{
+	TSharedPtr<FNetworkObjectInfo> InfoPtr = GetNetworkObjectList().Find(Actor);
+	if (InfoPtr.IsValid())
+	{
+		InfoPtr->bSwapRolesOnReplicate = bSwapRoles;
+	}
+
+	if (ReplicationDriver)
+	{
+		ReplicationDriver->SetRoleSwapOnReplicate(Actor, bSwapRoles);
+	}
+}
+
+FNetworkObjectInfo* UNetDriver::FindOrAddNetworkObjectInfo(const AActor* InActor)
+{
+	ensureMsgf(InActor == nullptr || !(InActor->HasAnyFlags(RF_ClassDefaultObject | RF_ArchetypeObject)), TEXT("%s is a CDO or Archetype and should not be replicated."), *GetFullNameSafe(InActor));
+
+	bool bWasAdded = false;
+	if (TSharedPtr<FNetworkObjectInfo>* InfoPtr = GetNetworkObjectList().FindOrAdd(const_cast<AActor*>(InActor), this, &bWasAdded))
+	{
+		if (bWasAdded && ReplicationDriver)
+		{
+			ReplicationDriver->AddNetworkActor(const_cast<AActor*>(InActor));
+		}
+
+		return InfoPtr->Get();
+	}
+
+	return nullptr;
+}
+
+FNetworkObjectInfo* UNetDriver::FindNetworkObjectInfo(const AActor* InActor)
+{
+	return GetNetworkObjectList().Find(InActor).Get();
+}
+
+bool UNetDriver::IsNetworkActorUpdateFrequencyThrottled(const FNetworkObjectInfo& InNetworkActor) const
+{
+	bool bThrottled = false;
+	if (IsAdaptiveNetUpdateFrequencyEnabled())
+	{
+		// Must have been replicated once for this to happen (and for OptimalNetUpdateDelta to have been set)
+		const AActor* Actor = InNetworkActor.Actor;
+		if (Actor && InNetworkActor.LastNetReplicateTime != 0)
+		{
+			const float ExpectedNetDelay = (1.0f / Actor->NetUpdateFrequency);
+			if (InNetworkActor.OptimalNetUpdateDelta > ExpectedNetDelay)
+			{
+				bThrottled = true;
+			}
+		}
+	}
+
+	return bThrottled;
+}
+
+bool UNetDriver::IsNetworkActorUpdateFrequencyThrottled(const AActor* InActor) const
+{
+	bool bThrottled = false;
+	if (InActor && IsAdaptiveNetUpdateFrequencyEnabled())
+	{
+		if (const FNetworkObjectInfo* NetActor = FindNetworkObjectInfo(InActor))
+		{
+			bThrottled = IsNetworkActorUpdateFrequencyThrottled(*NetActor);
+		}
+	}
+
+	return bThrottled;
+}
+
+void UNetDriver::CancelAdaptiveReplication(FNetworkObjectInfo& InNetworkActor)
+{
+	if (IsAdaptiveNetUpdateFrequencyEnabled())
+	{
+		if (AActor* Actor = InNetworkActor.Actor)
+		{
+			if (UWorld* ActorWorld = Actor->GetWorld())
+			{
+				const float ExpectedNetDelay = (1.0f / Actor->NetUpdateFrequency);
+				Actor->SetNetUpdateTime( ActorWorld->GetTimeSeconds() + FMath::FRandRange( 0.5f, 1.0f ) * ExpectedNetDelay );
+				InNetworkActor.OptimalNetUpdateDelta = ExpectedNetDelay;
+				// TODO: we really need a way to cancel the throttling completely. OptimalNetUpdateDelta is going to be recalculated based on LastNetReplicateTime.
+			}
+		}
+	}
+}
+
+static TAutoConsoleVariable<int32> CVarOptimizedRemapping( TEXT( "net.OptimizedRemapping" ), 1, TEXT( "Uses optimized path to remap unmapped network guids" ) );
+static TAutoConsoleVariable<int32> CVarMaxClientGuidRemaps( TEXT( "net.MaxClientGuidRemaps" ), 100, TEXT( "Max client resolves of unmapped network guids per tick" ) );
+TAutoConsoleVariable<int32> CVarFilterGuidRemapping( TEXT( "net.FilterGuidRemapping" ), 1, TEXT( "Remove destroyed and parent guids from unmapped list" ) );
+
+/** Accounts for the network time we spent in the game driver. */
+double GTickFlushGameDriverTimeSeconds = 0.0;
+
+bool ShouldEnableScopeSecondsTimers()
+{
+#if STATS
+	return true;
+#elif CSV_PROFILER
+	return FCsvProfiler::Get()->IsCapturing();
+#else
+	return false;
+#endif
+}
+
+void UNetDriver::TickFlush(float DeltaSeconds)
+{
+	CSV_SCOPED_TIMING_STAT_EXCLUSIVE(NetworkOutgoing);
+	SCOPE_CYCLE_COUNTER(STAT_NetTickFlush);
+	bool bEnableTimer = (NetDriverName == NAME_GameNetDriver) && ShouldEnableScopeSecondsTimers();
+	if (bEnableTimer)
+	{
+		GTickFlushGameDriverTimeSeconds = 0.0;
+	}
+	FSimpleScopeSecondsCounter ScopedTimer(GTickFlushGameDriverTimeSeconds, bEnableTimer);
+
+	if ( IsServer() && ClientConnections.Num() > 0 && ClientConnections[0]->IsInternalAck() == false )
+	{
+		// Update all clients.
+#if WITH_SERVER_CODE
+		int32 Updated = ServerReplicateActors( DeltaSeconds );
+
+		static int32 LastUpdateCount = 0;
+		// Only log the zero replicated actors once after replicating an actor
+		if ((LastUpdateCount && !Updated) || Updated)
+		{
+			UE_LOG(LogNetTraffic, Verbose, TEXT("%s replicated %d actors"), *GetDescription(), Updated);
+		}
+		LastUpdateCount = Updated;
+#endif // WITH_SERVER_CODE
+	}
+
+	// Reset queued bunch amortization timer
+	ProcessQueuedBunchesCurrentFrameMilliseconds = 0.0f;
+
+	const double CurrentRealtimeSeconds = FPlatformTime::Seconds();
+
+	bool bCollectServerStats = false;
+#if USE_SERVER_PERF_COUNTERS || STATS
+	bCollectServerStats = true;
+#endif
+
+	if (bCollectNetStats || bCollectServerStats)
+	{
+		SCOPE_CYCLE_COUNTER(STAT_NetTickFlushGatherStats);
+		// Update network stats (only main game net driver for now) if stats or perf counters are used
+		if (NetDriverName == NAME_GameNetDriver &&
+			CurrentRealtimeSeconds - StatUpdateTime > StatPeriod)
+		{
+			int32 ClientInBytesMax = 0;
+			int32 ClientInBytesMin = 0;
+			int32 ClientInBytesAvg = 0;
+			int32 ClientInPacketsMax = 0;
+			int32 ClientInPacketsMin = 0;
+			int32 ClientInPacketsAvg = 0;
+			int32 ClientOutBytesMax = 0;
+			int32 ClientOutBytesMin = 0;
+			int32 ClientOutBytesAvg = 0;
+			int32 ClientOutPacketsMax = 0;
+			int32 ClientOutPacketsMin = 0;
+			int32 ClientOutPacketsAvg = 0;
+			int NumClients = 0;
+			int32 MaxPacketOverhead = 0;
+
+			// these need to be updated even if we are not collecting stats, since they get reported to analytics/QoS
+			for (UNetConnection * Client : ClientConnections)
+			{
+				if (Client)
+				{
+#define UpdatePerClientMinMaxAvg(VariableName) \
+				Client##VariableName##Max = FMath::Max(Client##VariableName##Max, Client->VariableName##PerSecond); \
+				if (Client##VariableName##Min == 0 || Client->VariableName##PerSecond < Client##VariableName##Min) \
+				{ \
+					Client##VariableName##Min = Client->VariableName##PerSecond; \
+				} \
+				Client##VariableName##Avg += Client->VariableName##PerSecond; \
+
+					UpdatePerClientMinMaxAvg(InBytes);
+					UpdatePerClientMinMaxAvg(OutBytes);
+					UpdatePerClientMinMaxAvg(InPackets);
+					UpdatePerClientMinMaxAvg(OutPackets);
+
+					MaxPacketOverhead = FMath::Max(Client->PacketOverhead, MaxPacketOverhead);
+
+#undef UpdatePerClientMinMaxAvg
+
+					++NumClients;
+				}
+			}
+
+			if (NumClients > 1)
+			{
+				ClientInBytesAvg /= NumClients;
+				ClientInPacketsAvg /= NumClients;
+				ClientOutBytesAvg /= NumClients;
+				ClientOutPacketsAvg /= NumClients;
+			}
+
+			int32 Ping = 0;
+			int32 NumOpenChannels = 0;
+			int32 NumActorChannels = 0;
+			int32 NumDormantActors = 0;
+			int32 NumActors = 0;
+			int32 AckCount = 0;
+			int32 UnAckCount = 0;
+			int32 PendingCount = 0;
+			int32 NetSaturated = 0;
+			int32 NumActiveNetActors = GetNetworkObjectList().GetActiveObjects().Num();
+
+			if (
+#if STATS
+				FThreadStats::IsCollectingData() ||
+#endif
+				bCollectNetStats)
+			{
+				const float RealTime = CurrentRealtimeSeconds - StatUpdateTime;
+
+				// Use the elapsed time to keep things scaled to one measured unit
+				InBytes = FMath::TruncToInt(InBytes / RealTime);
+				OutBytes = FMath::TruncToInt(OutBytes / RealTime);
+
+				NetGUIDOutBytes = FMath::TruncToInt(NetGUIDOutBytes / RealTime);
+				NetGUIDInBytes = FMath::TruncToInt(NetGUIDInBytes / RealTime);
+
+				// Save off for stats later
+
+				InBytesPerSecond = InBytes;
+				OutBytesPerSecond = OutBytes;
+
+				InPackets = FMath::TruncToInt(InPackets / RealTime);
+				OutPackets = FMath::TruncToInt(OutPackets / RealTime);
+				InBunches = FMath::TruncToInt(InBunches / RealTime);
+				OutBunches = FMath::TruncToInt(OutBunches / RealTime);
+				OutPacketsLost = FMath::TruncToInt(100.f * OutPacketsLost / FMath::Max((float)OutPackets, 1.f));
+				InPacketsLost = FMath::TruncToInt(100.f * InPacketsLost / FMath::Max((float)InPackets + InPacketsLost, 1.f));
+
+				if (ServerConnection != NULL && ServerConnection->PlayerController != NULL && ServerConnection->PlayerController->PlayerState != NULL)
+				{
+					Ping = FMath::TruncToInt(ServerConnection->PlayerController->PlayerState->ExactPing);
+				}
+
+				if (ServerConnection != NULL)
+				{
+					NumOpenChannels = ServerConnection->OpenChannels.Num();
+				}
+
+				for (int32 i = 0; i < ClientConnections.Num(); i++)
+				{
+					NumOpenChannels += ClientConnections[i]->OpenChannels.Num();
+				}
+
+				// Use the elapsed time to keep things scaled to one measured unit
+				VoicePacketsSent = FMath::TruncToInt(VoicePacketsSent / RealTime);
+				VoicePacketsRecv = FMath::TruncToInt(VoicePacketsRecv / RealTime);
+				VoiceBytesSent = FMath::TruncToInt(VoiceBytesSent / RealTime);
+				VoiceBytesRecv = FMath::TruncToInt(VoiceBytesRecv / RealTime);
+
+				// Determine voice percentages
+				VoiceInPercent = (InBytes > 0) ? FMath::TruncToInt(100.f * (float)VoiceBytesRecv / (float)InBytes) : 0;
+				VoiceOutPercent = (OutBytes > 0) ? FMath::TruncToInt(100.f * (float)VoiceBytesSent / (float)OutBytes) : 0;
+
+				if (World)
+				{
+					NumActors = World->GetActorCount();
+				}
+				
+				// Handle these stats to do this differently for client & server
+				if (ServerConnection != nullptr)
+				{
+					NumActorChannels = ServerConnection->ActorChannelsNum();
+					NumDormantActors = GetNetworkObjectList().GetNumDormantActorsForConnection(ServerConnection);
+#if STATS
+					ServerConnection->PackageMap->GetNetGUIDStats(AckCount, UnAckCount, PendingCount);
+#endif
+					NetSaturated = ServerConnection->IsNetReady(false) ? 0 : 1;
+				}
+				else
+				{
+					int32 SharedDormantActors = GetNetworkObjectList().GetDormantObjectsOnAllConnections().Num();
+					for (int32 i = 0; i < ClientConnections.Num(); i++)
+					{
+						NumActorChannels += ClientConnections[i]->ActorChannelsNum();
+						NumDormantActors += GetNetworkObjectList().GetNumDormantActorsForConnection(ClientConnections[i]);
+#if STATS
+						int32 ClientAckCount = 0;
+						int32 ClientUnAckCount = 0;
+						int32 ClientPendingCount = 0;
+						ClientConnections[i]->PackageMap->GetNetGUIDStats(AckCount, UnAckCount, PendingCount);
+						AckCount += ClientAckCount;
+						UnAckCount += ClientUnAckCount;
+						PendingCount += ClientPendingCount;
+#endif
+						NetSaturated |= ClientConnections[i]->IsNetReady(false) ? 0 : 1;
+					}
+					// Subtract out the duplicate counted dormant actors
+					NumDormantActors -= FMath::Max(0, ClientConnections.Num() - 1) * SharedDormantActors;
+				}
+			}
+
+#if STATS
+			if (!bSkipLocalStats)
+			{
+				// Copy the net status values over
+				SET_DWORD_STAT(STAT_Ping, Ping);
+				SET_DWORD_STAT(STAT_Channels, NumOpenChannels);
+				SET_DWORD_STAT(STAT_MaxPacketOverhead, MaxPacketOverhead);
+
+				SET_DWORD_STAT(STAT_OutLoss, OutPacketsLost);
+				SET_DWORD_STAT(STAT_InLoss, InPacketsLost);
+				SET_DWORD_STAT(STAT_InRate, InBytes);
+				SET_DWORD_STAT(STAT_OutRate, OutBytes);
+				SET_DWORD_STAT(STAT_InRateClientMax, ClientInBytesMax);
+				SET_DWORD_STAT(STAT_InRateClientMin, ClientInBytesMin);
+				SET_DWORD_STAT(STAT_InRateClientAvg, ClientInBytesAvg);
+				SET_DWORD_STAT(STAT_InPacketsClientMax, ClientInPacketsMax);
+				SET_DWORD_STAT(STAT_InPacketsClientMin, ClientInPacketsMin);
+				SET_DWORD_STAT(STAT_InPacketsClientAvg, ClientInPacketsAvg);
+				SET_DWORD_STAT(STAT_OutRateClientMax, ClientOutBytesMax);
+				SET_DWORD_STAT(STAT_OutRateClientMin, ClientOutBytesMin);
+				SET_DWORD_STAT(STAT_OutRateClientAvg, ClientOutBytesAvg);
+				SET_DWORD_STAT(STAT_OutPacketsClientMax, ClientOutPacketsMax);
+				SET_DWORD_STAT(STAT_OutPacketsClientMin, ClientOutPacketsMin);
+				SET_DWORD_STAT(STAT_OutPacketsClientAvg, ClientOutPacketsAvg);
+
+				SET_DWORD_STAT(STAT_NetNumClients, NumClients);
+				SET_DWORD_STAT(STAT_InPackets, InPackets);
+				SET_DWORD_STAT(STAT_OutPackets, OutPackets);
+				SET_DWORD_STAT(STAT_InBunches, InBunches);
+				SET_DWORD_STAT(STAT_OutBunches, OutBunches);
+
+				SET_DWORD_STAT(STAT_NetGUIDInRate, NetGUIDInBytes);
+				SET_DWORD_STAT(STAT_NetGUIDOutRate, NetGUIDOutBytes);
+
+				SET_DWORD_STAT(STAT_VoicePacketsSent, VoicePacketsSent);
+				SET_DWORD_STAT(STAT_VoicePacketsRecv, VoicePacketsRecv);
+				SET_DWORD_STAT(STAT_VoiceBytesSent, VoiceBytesSent);
+				SET_DWORD_STAT(STAT_VoiceBytesRecv, VoiceBytesRecv);
+
+				SET_DWORD_STAT(STAT_PercentInVoice, VoiceInPercent);
+				SET_DWORD_STAT(STAT_PercentOutVoice, VoiceOutPercent);
+
+				SET_DWORD_STAT(STAT_NumActorChannels, NumActorChannels);
+				SET_DWORD_STAT(STAT_NumDormantActors, NumDormantActors);
+				SET_DWORD_STAT(STAT_NumActors, NumActors);
+				SET_DWORD_STAT(STAT_NumNetActors, NumActiveNetActors);
+				SET_DWORD_STAT(STAT_NumNetGUIDsAckd, AckCount);
+				SET_DWORD_STAT(STAT_NumNetGUIDsPending, UnAckCount);
+				SET_DWORD_STAT(STAT_NumNetGUIDsUnAckd, PendingCount);
+				SET_DWORD_STAT(STAT_NetSaturated, NetSaturated);
+			}
+
+			// If we are to replicate server stats out to an observer, then set those values
+			if (AGameModeBase* GMBase = World->GetAuthGameMode())
+			{
+				if (GMBase->ServerStatReplicator != nullptr)
+				{
+					GMBase->ServerStatReplicator->Channels = NumOpenChannels;
+					GMBase->ServerStatReplicator->MaxPacketOverhead = MaxPacketOverhead;
+					GMBase->ServerStatReplicator->OutLoss = OutPacketsLost;
+					GMBase->ServerStatReplicator->InLoss = InPacketsLost;
+					GMBase->ServerStatReplicator->InRate = InBytes;
+					GMBase->ServerStatReplicator->OutRate = OutBytes;
+					GMBase->ServerStatReplicator->InRateClientMax = ClientInBytesMax;
+					GMBase->ServerStatReplicator->InRateClientMin = ClientInBytesMin;
+					GMBase->ServerStatReplicator->InRateClientAvg = ClientInBytesAvg;
+					GMBase->ServerStatReplicator->InPacketsClientMax = ClientInPacketsMax;
+					GMBase->ServerStatReplicator->InPacketsClientMin = ClientInPacketsMin;
+					GMBase->ServerStatReplicator->InPacketsClientAvg = ClientInPacketsAvg;
+					GMBase->ServerStatReplicator->OutRateClientMax = ClientOutBytesMax;
+					GMBase->ServerStatReplicator->OutRateClientMin = ClientOutBytesMin;
+					GMBase->ServerStatReplicator->OutRateClientAvg = ClientOutBytesAvg;
+					GMBase->ServerStatReplicator->OutPacketsClientMax = ClientOutPacketsMax;
+					GMBase->ServerStatReplicator->OutPacketsClientMin = ClientOutPacketsMin;
+					GMBase->ServerStatReplicator->OutPacketsClientAvg = ClientOutPacketsAvg;
+					GMBase->ServerStatReplicator->NetNumClients = NumClients;
+					GMBase->ServerStatReplicator->InPackets = InPackets;
+					GMBase->ServerStatReplicator->OutPackets = OutPackets;
+					GMBase->ServerStatReplicator->InBunches = InBunches;
+					GMBase->ServerStatReplicator->OutBunches = OutBunches;
+					GMBase->ServerStatReplicator->NetGUIDInRate = NetGUIDInBytes;
+					GMBase->ServerStatReplicator->NetGUIDOutRate = NetGUIDOutBytes;
+					GMBase->ServerStatReplicator->VoicePacketsSent = VoicePacketsSent;
+					GMBase->ServerStatReplicator->VoicePacketsRecv = VoicePacketsRecv;
+					GMBase->ServerStatReplicator->VoiceBytesSent = VoiceBytesSent;
+					GMBase->ServerStatReplicator->VoiceBytesRecv = VoiceBytesRecv;
+					GMBase->ServerStatReplicator->PercentInVoice = VoiceInPercent;
+					GMBase->ServerStatReplicator->PercentOutVoice = VoiceOutPercent;
+					GMBase->ServerStatReplicator->NumActorChannels = NumActorChannels;
+					GMBase->ServerStatReplicator->NumDormantActors = NumDormantActors;
+					GMBase->ServerStatReplicator->NumActors = NumActors;
+					GMBase->ServerStatReplicator->NumNetActors = NumActiveNetActors;
+					GMBase->ServerStatReplicator->NumNetGUIDsAckd = AckCount;
+					GMBase->ServerStatReplicator->NumNetGUIDsPending = UnAckCount;
+					GMBase->ServerStatReplicator->NumNetGUIDsUnAckd = PendingCount;
+					GMBase->ServerStatReplicator->NetSaturated = NetSaturated;
+				}
+			}
+#endif // STATS
+
+#if USE_SERVER_PERF_COUNTERS
+			IPerfCounters* PerfCounters = IPerfCountersModule::Get().GetPerformanceCounters();
+			if (PerfCounters)
+			{
+				SCOPE_CYCLE_COUNTER(STAT_NetTickFlushGatherStatsPerfCounters);
+
+				// Update total connections
+				PerfCounters->Set(TEXT("NumConnections"), ClientConnections.Num());
+
+				const int kNumBuckets = 8;	// evenly spaced with increment of 30 ms; last bucket collects all off-scale pings as well
+				if (ClientConnections.Num() > 0)
+				{
+					// Update per connection statistics
+					float MinPing = MAX_FLT;
+					float AvgPing = 0;
+					float MaxPing = -MAX_FLT;
+					float PingCount = 0;
+
+					int32 Buckets[kNumBuckets] = { 0 };
+
+					for (int32 i = 0; i < ClientConnections.Num(); i++)
+					{
+						UNetConnection* Connection = ClientConnections[i];
+
+						if (Connection != nullptr)
+						{
+							if (Connection->PlayerController != nullptr && Connection->PlayerController->PlayerState != nullptr)
+							{
+								// Ping value calculated per client
+								float ConnPing = Connection->PlayerController->PlayerState->ExactPing;
+
+								int Bucket = FMath::Max(0, FMath::Min(kNumBuckets - 1, (static_cast<int>(ConnPing) / 30)));
+								++Buckets[Bucket];
+
+								if (ConnPing < MinPing)
+								{
+									MinPing = ConnPing;
+								}
+
+								if (ConnPing > MaxPing)
+								{
+									MaxPing = ConnPing;
+								}
+
+								AvgPing += ConnPing;
+								PingCount++;
+							}
+						}
+					}
+
+					if (PingCount > 0)
+					{
+						AvgPing /= static_cast<float>(PingCount);
+					}
+
+					PerfCounters->Set(TEXT("AvgPing"), AvgPing, IPerfCounters::Flags::Transient);
+					PerfCounters->Set(TEXT("MaxPing"), MaxPing, IPerfCounters::Flags::Transient);
+					PerfCounters->Set(TEXT("MinPing"), MinPing, IPerfCounters::Flags::Transient);
+
+					// update buckets
+					for (int BucketIdx = 0; BucketIdx < UE_ARRAY_COUNT(Buckets); ++BucketIdx)
+					{
+						PerfCountersIncrement(FString::Printf(TEXT("PingBucketInt%d"), BucketIdx), Buckets[BucketIdx], 0, IPerfCounters::Flags::Transient);
+					}
+				}
+				else
+				{
+					PerfCounters->Set(TEXT("AvgPing"), 0.0f, IPerfCounters::Flags::Transient);
+					PerfCounters->Set(TEXT("MaxPing"), -FLT_MAX, IPerfCounters::Flags::Transient);
+					PerfCounters->Set(TEXT("MinPing"), FLT_MAX, IPerfCounters::Flags::Transient);
+
+					for (int BucketIdx = 0; BucketIdx < kNumBuckets; ++BucketIdx)
+					{
+						PerfCounters->Set(FString::Printf(TEXT("PingBucketInt%d"), BucketIdx), 0, IPerfCounters::Flags::Transient);
+					}
+				}
+
+				// set the per connection stats (these are calculated earlier).
+				// Note that NumClients may be != NumConnections. Also, if NumClients is 0, the rest of counters should be 0 as well
+				PerfCounters->Set(TEXT("NumClients"), NumClients);
+				PerfCounters->Set(TEXT("MaxPacketOverhead"), MaxPacketOverhead);
+				PerfCounters->Set(TEXT("InRateClientMax"), ClientInBytesMax);
+				PerfCounters->Set(TEXT("InRateClientMin"), ClientInBytesMin);
+				PerfCounters->Set(TEXT("InRateClientAvg"), ClientInBytesAvg);
+				PerfCounters->Set(TEXT("InPacketsClientMax"), ClientInPacketsMax);
+				PerfCounters->Set(TEXT("InPacketsClientMin"), ClientInPacketsMin);
+				PerfCounters->Set(TEXT("InPacketsClientAvg"), ClientInPacketsAvg);
+				PerfCounters->Set(TEXT("OutRateClientMax"), ClientOutBytesMax);
+				PerfCounters->Set(TEXT("OutRateClientMin"), ClientOutBytesMin);
+				PerfCounters->Set(TEXT("OutRateClientAvg"), ClientOutBytesAvg);
+				PerfCounters->Set(TEXT("OutPacketsClientMax"), ClientOutPacketsMax);
+				PerfCounters->Set(TEXT("OutPacketsClientMin"), ClientOutPacketsMin);
+				PerfCounters->Set(TEXT("OutPacketsClientAvg"), ClientOutPacketsAvg);
+
+				PerfCounters->Set(TEXT("InRate"), InBytes);
+				PerfCounters->Set(TEXT("OutRate"), OutBytes);
+				PerfCounters->Set(TEXT("InPacketsLost"), InPacketsLost);
+				PerfCounters->Set(TEXT("OutPacketsLost"), OutPacketsLost);
+				PerfCounters->Set(TEXT("InPackets"), InPackets);
+				PerfCounters->Set(TEXT("OutPackets"), OutPackets);
+				PerfCounters->Set(TEXT("InBunches"), InBunches);
+				PerfCounters->Set(TEXT("OutBunches"), OutBunches);
+			}
+#endif // USE_SERVER_PERF_COUNTERS
+
+			// Reset everything
+			InBytes = 0;
+			OutBytes = 0;
+			NetGUIDOutBytes = 0;
+			NetGUIDInBytes = 0;
+			InPackets = 0;
+			OutPackets = 0;
+			InBunches = 0;
+			OutBunches = 0;
+			OutPacketsLost = 0;
+			InPacketsLost = 0;
+			VoicePacketsSent = 0;
+			VoiceBytesSent = 0;
+			VoicePacketsRecv = 0;
+			VoiceBytesRecv = 0;
+			VoiceInPercent = 0;
+			VoiceOutPercent = 0;
+			StatUpdateTime = CurrentRealtimeSeconds;
+		}
+	} // bCollectNetStats ||(USE_SERVER_PERF_COUNTERS) || STATS
+
+	// Poll all sockets.
+	if( ServerConnection )
+	{
+		// Queue client voice packets in the server's voice channel
+		ProcessLocalClientPackets();
+		ServerConnection->Tick();
+	}
+	else
+	{
+		// Queue up any voice packets the server has locally
+		ProcessLocalServerPackets();
+	}
+	{
+		QUICK_SCOPE_CYCLE_COUNTER(STAT_NetDriver_TickClientConnections)
+
+	for( int32 i=0; i<ClientConnections.Num(); i++ )
+	{
+		ClientConnections[i]->Tick();
+	}
+	}
+
+	if (ConnectionlessHandler.IsValid())
+	{
+		ConnectionlessHandler->Tick(DeltaSeconds);
+
+		FlushHandler();
+	}
+
+	if (CVarNetDebugDraw.GetValueOnAnyThread() > 0)
+	{
+		DrawNetDriverDebug();
+	}
+
+	if ( CVarOptimizedRemapping.GetValueOnAnyThread() && GuidCache.IsValid() )
+	{
+		SCOPE_CYCLE_COUNTER( STAT_NetUpdateUnmappedObjectsTime );
+
+		// Go over recently imported network guids, and see if there are any replicators that need to map them
+		TSet< FNetworkGUID >& ImportedNetGuids = GuidCache->ImportedNetGuids;
+		TMap< FNetworkGUID, TSet< FNetworkGUID > >& PendingOuterNetGuids = GuidCache->PendingOuterNetGuids;
+
+		SET_DWORD_STAT( STAT_ImportedNetGuids, ImportedNetGuids.Num() );
+		SET_DWORD_STAT( STAT_PendingOuterNetGuids, PendingOuterNetGuids.Num() );
+		SET_DWORD_STAT( STAT_UnmappedReplicators, UnmappedReplicators.Num() );
+
+		TSet< FObjectReplicator* > ForceUpdateReplicators;
+
+		for (FObjectReplicator* Replicator : UnmappedReplicators)
+		{
+			if (Replicator->bForceUpdateUnmapped)
+			{
+				Replicator->bForceUpdateUnmapped = false;
+				ForceUpdateReplicators.Add(Replicator);
+			}
+		}
+
+		if ( ImportedNetGuids.Num() || ForceUpdateReplicators.Num() )
+		{
+			int32 NumRemapTests = 0;
+			int32 MaxRemaps = IsServer() ? 0 : CVarMaxClientGuidRemaps.GetValueOnAnyThread();
+			bool bFilterGuidRemapping = (CVarFilterGuidRemapping.GetValueOnAnyThread() > 0);
+
+			TArray< FNetworkGUID > UnmappedGuids;
+			TArray< FNetworkGUID > NewlyMappedGuids;
+
+			for ( auto It = ImportedNetGuids.CreateIterator(); It; ++It )
+			{
+				const FNetworkGUID NetworkGuid = *It;
+				bool bMappedOrBroken = false;
+
+				if ( GuidCache->GetObjectFromNetGUID( NetworkGuid, false ) != nullptr )
+				{
+					NewlyMappedGuids.Add( NetworkGuid );
+					bMappedOrBroken = true;
+				}
+
+				if ( GuidCache->IsGUIDBroken( NetworkGuid, false ) )
+				{
+					bMappedOrBroken = true;
+				}
+
+				It.RemoveCurrent();
+
+				if (!bMappedOrBroken)
+				{
+					FNetworkGUID OuterGUID = GuidCache->GetOuterNetGUID(NetworkGuid);
+
+					// we're missing the outer, stop checking until we map it
+					if (bFilterGuidRemapping && OuterGUID.IsValid() && !OuterGUID.IsDefault() && !GuidCache->IsGUIDLoaded(OuterGUID) && !GuidCache->IsGUIDPending(OuterGUID))
+					{
+						UE_LOG(LogNetPackageMap, Log, TEXT("Missing outer (%s) for unmapped guid (%s), marking pending"), *GuidCache->FullNetGUIDPath(OuterGUID), *GuidCache->FullNetGUIDPath(NetworkGuid));
+
+						TSet<FNetworkGUID>& PendingGuids = PendingOuterNetGuids.FindOrAdd(OuterGUID);
+						PendingGuids.Add(NetworkGuid);
+					}
+					else
+					{
+						if (ensure(NetworkGuid.IsValid()))
+						{
+						UnmappedGuids.Add(NetworkGuid);
+					}
+				}
+				}
+
+				++NumRemapTests;
+
+ 				if ((MaxRemaps > 0) && (NumRemapTests >= MaxRemaps))
+				{
+					break;
+				}
+			}
+
+			// attempt to resolve dependent guids next tick (outer is now mapped)
+			for ( const FNetworkGUID& NetGuid : NewlyMappedGuids )
+			{
+				TSet< FNetworkGUID >* DependentGuids = PendingOuterNetGuids.Find( NetGuid );
+				if ( DependentGuids )
+				{
+					UE_LOG(LogNetPackageMap, Log, TEXT("Newly mapped outer (%s) removing from pending"), *GuidCache->FullNetGUIDPath(NetGuid));
+
+					ImportedNetGuids.Append( *DependentGuids );
+					PendingOuterNetGuids.Remove( NetGuid );
+				}
+			}
+
+			// any tested guids that could not yet be mapped are added back to the list
+			if (UnmappedGuids.Num())
+			{
+				ImportedNetGuids.CompactStable();
+				ImportedNetGuids.Append(UnmappedGuids);
+			}
+
+			if ( NewlyMappedGuids.Num() || ForceUpdateReplicators.Num() )
+			{
+				TSet< FObjectReplicator* > AllReplicators = ForceUpdateReplicators;
+
+				for ( const FNetworkGUID& NetGuid : NewlyMappedGuids )
+				{
+					TSet< FObjectReplicator* >* Replicators = GuidToReplicatorMap.Find( NetGuid );
+
+					if ( Replicators )
+					{
+						AllReplicators.Append( *Replicators );
+					}
+				}
+
+				for ( FObjectReplicator* Replicator : AllReplicators )
+				{
+					if ( UnmappedReplicators.Contains( Replicator ) )
+					{
+						bool bHasMoreUnmapped = false;
+						Replicator->UpdateUnmappedObjects( bHasMoreUnmapped );
+
+						if ( !bHasMoreUnmapped )
+						{
+							UnmappedReplicators.Remove( Replicator );
+						}
+					}
+				}
+			}
+		}
+	}
+	else
+	{
+		SCOPE_CYCLE_COUNTER( STAT_NetUpdateUnmappedObjectsTime );
+
+		// Update properties that are unmapped, try to hook up the object pointers if they exist now
+		for ( auto It = UnmappedReplicators.CreateIterator(); It; ++It )
+		{
+			FObjectReplicator* Replicator = *It;
+
+			bool bHasMoreUnmapped = false;
+
+			Replicator->UpdateUnmappedObjects( bHasMoreUnmapped );
+
+			if ( !bHasMoreUnmapped )
+			{
+				// If there are no more unmapped objects, we can also stop checking
+				It.RemoveCurrent();
+			}
+		}
+	}
+
+	// Go over RepChangedPropertyTrackerMap periodically, and remove entries that no longer have valid objects
+	// Unfortunately if you mark an object as pending kill, it will no longer find itself in this map,
+	// so we do this as a fail safe to make sure we never leak memory from this map
+	const double CleanupTimeSeconds = 10.0;
+
+	if ( CurrentRealtimeSeconds - LastCleanupTime > CleanupTimeSeconds )
+	{
+		for ( auto It = RepChangedPropertyTrackerMap.CreateIterator(); It; ++It )
+		{
+			if ( !It.Value().IsObjectValid() )
+			{
+				It.RemoveCurrent();
+			}
+		}
+
+		for ( auto It = ReplicationChangeListMap.CreateIterator(); It; ++It )
+		{
+			if ( !It.Value().IsObjectValid() )
+			{
+				It.RemoveCurrent();
+			}
+		}
+
+		for ( UNetConnection* const ClientConnection : ClientConnections )
+		{
+			if ( ClientConnection )
+			{
+				ClientConnection->CleanupStaleDormantReplicators();
+			}
+		}
+
+		if ( ServerConnection )
+		{
+			ServerConnection->CleanupStaleDormantReplicators();
+		}
+
+		LastCleanupTime = CurrentRealtimeSeconds;
+	}
+
+	// Update the lag state
+	UpdateNetworkLagState();
+}
+
+void UNetDriver::UpdateNetworkLagState()
+{
+	ENetworkLagState::Type OldLagState = LagState;
+
+	// Percentage of the timeout time that a connection is considered "lagging"
+	const float TimeoutPercentThreshold = 0.75f;
+
+	if ( IsServer() )
+	{
+		// Server network lag detection
+
+		// See if all clients connected to us are lagging. If so there might be network connection problems.
+		// Only trigger this if there are a few connections since a single client could have just crashed or disconnected suddenly,
+		// and is less likely to happen with multiple clients simultaneously.
+		int32 NumValidConnections = 0;
+		int32 NumLaggingConnections = 0;
+		for (UNetConnection* Connection : ClientConnections)
+		{
+			if (Connection)
+			{
+				NumValidConnections++;
+
+				const double HalfTimeout = Connection->GetTimeoutValue() * TimeoutPercentThreshold;
+				const double DeltaTimeSinceLastMessage = ElapsedTime - Connection->LastReceiveTime;
+				if (DeltaTimeSinceLastMessage > HalfTimeout)
+				{
+					NumLaggingConnections++;
+				}
+			}
+		}
+
+		if (NumValidConnections >= 2 && NumValidConnections == NumLaggingConnections)
+		{
+			// All connections that we could measure are lagging and there are enough to know it is not likely the fault of the clients.
+			LagState = ENetworkLagState::Lagging;
+		}
+		else
+		{
+			// We have at least one non-lagging client or we don't have enough clients to know if the server is lagging.
+			LagState = ENetworkLagState::NotLagging;
+		}
+	}
+	else
+	{
+		// Client network lag detection.
+		
+		// Just check the server connection.
+		if (ensure(ServerConnection))
+		{
+			const double HalfTimeout = ServerConnection->GetTimeoutValue() * TimeoutPercentThreshold;
+			const double DeltaTimeSinceLastMessage = ElapsedTime - ServerConnection->LastReceiveTime;
+			if (DeltaTimeSinceLastMessage > HalfTimeout)
+			{
+				// We have exceeded half our timeout. We are lagging.
+				LagState = ENetworkLagState::Lagging;
+			}
+			else
+			{
+				// Not lagging yet. We have received a message recently.
+				LagState = ENetworkLagState::NotLagging;
+			}
+		}
+	}
+
+	if (OldLagState != LagState)
+	{
+		GEngine->BroadcastNetworkLagStateChanged(GetWorld(), this, LagState);
+	}
+}
+
+/**
+ * Determines which other connections should receive the voice packet and
+ * queues the packet for those connections. Used for sending both local/remote voice packets.
+ *
+ * @param VoicePacket the packet to be queued
+ * @param CameFromConn the connection this packet came from (NULL if local)
+ */
+void UNetDriver::ReplicateVoicePacket(TSharedPtr<FVoicePacket> VoicePacket, UNetConnection* CameFromConn)
+{
+	// Iterate the connections and see if they want the packet
+	for (int32 Index = 0; Index < ClientConnections.Num(); Index++)
+	{
+		UNetConnection* Conn = ClientConnections[Index];
+		// Skip the originating connection
+		if (CameFromConn != Conn)
+		{
+			// If server then determine if it should replicate the voice packet from another sender to this connection
+			const bool bReplicateAsServer = !bIsPeer && Conn->ShouldReplicateVoicePacketFrom(*VoicePacket->GetSender());
+			// If client peer then determine if it should send the voice packet to another client peer
+			//const bool bReplicateAsPeer = (bIsPeer && AllowPeerVoice) && Conn->ShouldReplicateVoicePacketToPeer(Conn->PlayerId);
+
+			if (bReplicateAsServer)// || bReplicateAsPeer)
+			{
+				UVoiceChannel* VoiceChannel = Conn->GetVoiceChannel();
+				if (VoiceChannel != NULL)
+				{
+					// Add the voice packet for network sending
+					VoiceChannel->AddVoicePacket(VoicePacket);
+				}
+			}
+		}
+	}
+}
+
+/**
+ * Process any local talker packets that need to be sent to clients
+ */
+void UNetDriver::ProcessLocalServerPackets()
+{
+	if (World)
+	{
+		int32 NumLocalTalkers = UOnlineEngineInterface::Get()->GetNumLocalTalkers(World);
+		// Process all of the local packets
+		for (int32 Index = 0; Index < NumLocalTalkers; Index++)
+		{
+			// Returns a ref counted copy of the local voice data or NULL if nothing to send
+			TSharedPtr<FVoicePacket> LocalPacket = UOnlineEngineInterface::Get()->GetLocalPacket(World, Index);
+			// Check for something to send for this local talker
+			if (LocalPacket.IsValid())
+			{
+				// See if anyone wants this packet
+				ReplicateVoicePacket(LocalPacket, NULL);
+
+				// once all local voice packets are processed then call ClearVoicePackets()
+			}
+		}
+	}
+}
+
+/**
+ * Process any local talker packets that need to be sent to the server
+ */
+void UNetDriver::ProcessLocalClientPackets()
+{
+	if (World)
+	{
+		int32 NumLocalTalkers = UOnlineEngineInterface::Get()->GetNumLocalTalkers(World);
+		if (NumLocalTalkers)
+		{
+			UVoiceChannel* VoiceChannel = ServerConnection->GetVoiceChannel();
+			if (VoiceChannel)
+			{
+				// Process all of the local packets
+				for (int32 Index = 0; Index < NumLocalTalkers; Index++)
+				{
+					// Returns a ref counted copy of the local voice data or NULL if nothing to send
+					TSharedPtr<FVoicePacket> LocalPacket = UOnlineEngineInterface::Get()->GetLocalPacket(World, Index);
+					// Check for something to send for this local talker
+					if (LocalPacket.IsValid())
+					{
+						// If there is a voice channel to the server, submit the packets
+						//if (ShouldSendVoicePacketsToServer())
+						{
+							// Add the voice packet for network sending
+							VoiceChannel->AddVoicePacket(LocalPacket);
+						}
+
+						// once all local voice packets are processed then call ClearLocalVoicePackets()
+					}
+				}
+			}
+		}
+	}
+}
+
+void UNetDriver::PostTickFlush()
+{
+	if (World)
+	{
+		UOnlineEngineInterface::Get()->ClearVoicePackets(World);
+	}
+}
+
+bool UNetDriver::InitConnectionClass(void)
+{
+	if (NetConnectionClass == NULL && NetConnectionClassName != TEXT(""))
+	{
+		NetConnectionClass = LoadClass<UNetConnection>(NULL,*NetConnectionClassName,NULL,LOAD_None,NULL);
+		if (NetConnectionClass == NULL)
+		{
+			UE_LOG(LogNet, Error,TEXT("Failed to load class '%s'"),*NetConnectionClassName);
+		}
+	}
+	return NetConnectionClass != NULL;
+}
+
+bool UNetDriver::InitReplicationDriverClass()
+{
+	if (ReplicationDriverClass == nullptr && !ReplicationDriverClassName.IsEmpty())
+	{
+		ReplicationDriverClass = LoadClass<UReplicationDriver>(nullptr, *ReplicationDriverClassName, nullptr, LOAD_None, nullptr);
+		if (ReplicationDriverClass == nullptr)
+		{
+			UE_LOG(LogNet, Error,TEXT("Failed to load class '%s'"), *ReplicationDriverClassName);
+		}
+	}
+
+	return ReplicationDriverClass != nullptr;
+}
+
+bool UNetDriver::InitBase(bool bInitAsClient, FNetworkNotify* InNotify, const FURL& URL, bool bReuseAddressAndPort, FString& Error)
+{
+	// Read any timeout overrides from the URL
+	if (const TCHAR* InitialConnectTimeoutOverride = URL.GetOption(TEXT("InitialConnectTimeout="), nullptr))
+	{
+		float ParsedValue;
+		LexFromString(ParsedValue, InitialConnectTimeoutOverride);
+		if (ParsedValue != 0.0f)
+		{
+			InitialConnectTimeout = ParsedValue;
+		}
+	}
+	if (const TCHAR* ConnectionTimeoutOverride = URL.GetOption(TEXT("ConnectionTimeout="), nullptr))
+	{
+		float ParsedValue;
+		LexFromString(ParsedValue, ConnectionTimeoutOverride);
+		if (ParsedValue != 0.0f)
+		{
+			ConnectionTimeout = ParsedValue;
+		}
+	}
+	if (URL.HasOption(TEXT("NoTimeouts")))
+	{
+		bNoTimeouts = true;
+	}
+
+	LastTickDispatchRealtime = FPlatformTime::Seconds();
+	bool bSuccess = InitConnectionClass();
+
+	if (!bInitAsClient)
+	{
+		ConnectionlessHandler.Reset(nullptr);
+		InitDestroyedStartupActors();
+		InitReplicationDriverClass();
+		SetReplicationDriver(UReplicationDriver::CreateReplicationDriver(this, URL, GetWorld()));
+
+		DDoS.Init(FMath::Clamp(NetServerMaxTickRate, 1, 1000));
+
+		DDoS.NotifySeverityEscalation.BindLambda(
+			[this](FString SeverityCategory)
+		{
+			GEngine->BroadcastNetworkDDosSEscalation(this->GetWorld(), this, SeverityCategory);
+		});
+	}
+
+#if DO_ENABLE_NET_TEST
+	bool bSettingFound(false);
+	FPacketSimulationSettings PacketSettings;
+
+	for (const FString& URLOption : URL.Op)
+	{
+		bSettingFound |= PacketSettings.ParseSettings(*URLOption);
+	}
+
+	if( bSettingFound )
+	{
+		SetPacketSimulationSettings(PacketSettings);
+	}
+#endif //#if DO_ENABLE_NET_TEST
+
+	Notify = InNotify;
+
+	// If we are not using Iris, we use the UniqueId to identify the NetDriver.
+	NetTraceId = GetUniqueID();
+
+	return bSuccess;
+}
+
+void UNetDriver::InitConnectionlessHandler()
+{
+	check(!ConnectionlessHandler.IsValid());
+
+#if !UE_BUILD_SHIPPING
+	if (!FParse::Param(FCommandLine::Get(), TEXT("NoPacketHandler")))
+#endif
+	{
+		ConnectionlessHandler = MakeUnique<PacketHandler>(&DDoS);
+
+		if (ConnectionlessHandler.IsValid())
+		{
+			PRAGMA_DISABLE_DEPRECATION_WARNINGS
+			ConnectionlessHandler->InitializeAddressSerializer([this](const FString& InAddress) {
+				return GetSocketSubsystem()->GetAddressFromString(InAddress);
+			});
+			PRAGMA_ENABLE_DEPRECATION_WARNINGS
+			ConnectionlessHandler->NotifyAnalyticsProvider(AnalyticsProvider, AnalyticsAggregator);
+			ConnectionlessHandler->Initialize(Handler::Mode::Server, MAX_PACKET_SIZE, true, nullptr, nullptr, NetDriverName);
+
+			// Add handling for the stateless connect handshake, for connectionless packets, as the outermost layer
+			TSharedPtr<HandlerComponent> NewComponent =
+				ConnectionlessHandler->AddHandler(TEXT("Engine.EngineHandlerComponentFactory(StatelessConnectHandlerComponent)"), true);
+
+			StatelessConnectComponent = StaticCastSharedPtr<StatelessConnectHandlerComponent>(NewComponent);
+
+			if (StatelessConnectComponent.IsValid())
+			{
+				StatelessConnectComponent.Pin()->SetDriver(this);
+			}
+
+			ConnectionlessHandler->InitializeComponents();
+		}
+	}
+}
+
+void UNetDriver::FlushHandler()
+{
+	BufferedPacket* QueuedPacket = ConnectionlessHandler->GetQueuedConnectionlessPacket();
+
+	while (QueuedPacket != nullptr)
+	{
+		LowLevelSend(QueuedPacket->Address, QueuedPacket->Data, QueuedPacket->CountBits, QueuedPacket->Traits);
+
+		delete QueuedPacket;
+
+		QueuedPacket = ConnectionlessHandler->GetQueuedConnectionlessPacket();
+	}
+}
+
+ENetMode UNetDriver::GetNetMode() const
+{
+	// Special case for PIE - forcing dedicated server behavior
+#if WITH_EDITOR
+	if (World && World->WorldType == EWorldType::PIE && IsServer())
+	{
+		if ( GEngine->GetWorldContextFromWorldChecked(World).RunAsDedicated )
+		{
+			return NM_DedicatedServer;
+		}
+	}
+#endif
+
+	// Normal
+	return (IsServer() ? (GIsClient ? NM_ListenServer : NM_DedicatedServer) : NM_Client);
+}
+
+void UNetDriver::RegisterTickEvents(class UWorld* InWorld)
+{
+	if (InWorld)
+	{
+		TickDispatchDelegateHandle  = InWorld->OnTickDispatch ().AddUObject(this, &UNetDriver::TickDispatch);
+		PostTickDispatchDelegateHandle	= InWorld->OnPostTickDispatch().AddUObject(this, &UNetDriver::PostTickDispatch);
+		TickFlushDelegateHandle     = InWorld->OnTickFlush    ().AddUObject(this, &UNetDriver::TickFlush);
+		PostTickFlushDelegateHandle		= InWorld->OnPostTickFlush	 ().AddUObject(this, &UNetDriver::PostTickFlush);
+	}
+}
+
+void UNetDriver::UnregisterTickEvents(class UWorld* InWorld)
+{
+	if (InWorld)
+	{
+		InWorld->OnTickDispatch ().Remove(TickDispatchDelegateHandle);
+		InWorld->OnPostTickDispatch().Remove(PostTickDispatchDelegateHandle);
+		InWorld->OnTickFlush    ().Remove(TickFlushDelegateHandle);
+		InWorld->OnPostTickFlush   ().Remove(PostTickFlushDelegateHandle);
+	}
+}
+
+static bool bCVarLogPendingGuidsOnShutdown = false;
+static FAutoConsoleVariableRef CVarLogPendingGuidsOnShutdown(
+	TEXT("Net.LogPendingGuidsOnShutdown"),
+	bCVarLogPendingGuidsOnShutdown,
+	TEXT("")
+);
+
+/** Shutdown all connections managed by this net driver */
+void UNetDriver::Shutdown()
+{
+	// Client closing connection to server
+	if (ServerConnection)
+	{
+		const UPackageMapClient* const ClientPackageMap = Cast<UPackageMapClient>(ServerConnection->PackageMap);
+		const bool bLogGuids = bCVarLogPendingGuidsOnShutdown && (ClientPackageMap != nullptr);
+		TSet<FNetworkGUID> GuidsToLog;
+
+		for (UChannel* Channel : ServerConnection->OpenChannels)
+		{
+			 if (UActorChannel * ActorChannel = Cast<UActorChannel>(Channel))
+			 {
+				 if (bLogGuids)
+				 {
+					 GuidsToLog.Append(ActorChannel->PendingGuidResolves);
+				 }
+				 ActorChannel->CleanupReplicators();
+			 }
+		}
+
+		if (GuidsToLog.Num() > 0)
+		{
+			TArray<FString> GuidStrings;
+			GuidStrings.Reserve(GuidsToLog.Num());
+			for (const FNetworkGUID GuidToLog : GuidsToLog)
+			{
+				FString FullNetGUIDPath = ClientPackageMap->GetFullNetGUIDPath(GuidToLog);
+				if (!FullNetGUIDPath.IsEmpty())
+				{
+					GuidStrings.Emplace(MoveTemp(FullNetGUIDPath));
+				}
+				else
+				{
+					GuidStrings.Emplace(GuidToLog.ToString());
+				}
+			}
+
+			UE_LOG(LogNet, Warning, TEXT("NetDriver::Shutdown: Pending Guids: \n%s"), *FString::Join(GuidStrings, TEXT("\n")));
+		}
+		
+
+		// Calls Channel[0]->Close to send a close bunch to server
+		ServerConnection->Close();
+		ServerConnection->FlushNet();
+	}
+
+	// Server closing connections with clients
+	if (ClientConnections.Num() > 0)
+	{
+		for (int32 ClientIndex = 0; ClientIndex < ClientConnections.Num(); ClientIndex++)
+		{
+			FString ErrorMsg = NSLOCTEXT("NetworkErrors", "HostClosedConnection", "Host closed the connection.").ToString();
+			FNetControlMessage<NMT_Failure>::Send(ClientConnections[ClientIndex], ErrorMsg);
+			ClientConnections[ClientIndex]->FlushNet(true);
+		}
+
+		for (int32 ClientIndex = ClientConnections.Num() - 1; ClientIndex >= 0; ClientIndex--)
+		{
+			if (ClientConnections[ClientIndex]->PlayerController)
+			{
+				APawn* Pawn = ClientConnections[ClientIndex]->PlayerController->GetPawn();
+				if( Pawn )
+				{
+					Pawn->Destroy( true );
+				}
+			}
+
+			// Calls Close() internally and removes from ClientConnections
+			ClientConnections[ClientIndex]->CleanUp();
+		}
+	}
+
+	// Empty our replication map here before we're destroyed, 
+	// even though we use AddReferencedObjects to keep the referenced properties
+	// in here from being collected, when we're all GC'd the order seems non-deterministic
+	RepLayoutMap.Empty();
+	ReplicationChangeListMap.Empty();
+
+	// Clean up the actor channel pool
+	ActorChannelPool.Empty();
+
+	ConnectionlessHandler.Reset(nullptr);
+
+	SetReplicationDriver(nullptr);
+
+	// End NetTrace session for this instance
+	UE_NET_TRACE_END_SESSION(GetNetTraceId());
+
+	if (AnalyticsAggregator.IsValid())
+	{
+		AnalyticsAggregator->SendAnalytics();
+		AnalyticsAggregator.Reset();
+	}
+}
+
+bool UNetDriver::IsServer() const
+{
+	// Client connections ALWAYS set the server connection object in InitConnect()
+	// @todo ONLINE improve this with a bool
+	return ServerConnection == NULL;
+}
+
+// ----------------------------------------------------------------------------------------
+//	RPC Tracking (receiving) via CSV profiler
+// ----------------------------------------------------------------------------------------
+CSV_DEFINE_CATEGORY(ReplicationRPCs, WITH_SERVER_CODE);
+
+#ifndef RPC_CSV_TRACKER	// Defines if RPC CSV tracking is compiled in. Default is only on server builds.
+#define RPC_CSV_TRACKER (CSV_PROFILER && WITH_SERVER_CODE)
+#endif
+
+/** Helper struct for tracking RPC (receive) timing */
+struct FRPCCSVTracker
+{
+	const double StatThreshold = 0.001; // 1MS minimum accumulated threshold for a stat to be dumped
+
+	void NotifyRPCProcessed(UFunction* Func, double Time)
+	{
+#if RPC_CSV_TRACKER
+		if (FItem* Item = FunctionMap.Find(Func))
+		{
+			Item->Time += Time;
+		}
+		else
+		{
+			FunctionMap.Emplace( Func, FItem(Func, Time) );
+		}
+#endif
+	}
+
+	void EndTickDispatch()
+	{
+#if RPC_CSV_TRACKER
+		FCsvProfiler* Profiler = FCsvProfiler::Get();
+		if (Profiler->IsCapturing())
+		{
+			// Record stat and reset accumulated time
+			for (auto& It : FunctionMap)
+			{
+				FItem& Item = It.Value;
+				if (Item.Time >= StatThreshold)
+				{
+					Profiler->RecordCustomStat(Item.Stat, CSV_CATEGORY_INDEX(ReplicationRPCs), static_cast<float>(Item.Time * 1000.0), ECsvCustomStatOp::Set);
+				}
+				Item.Time = 0.0;
+			}
+		}
+#endif
+	}
+
+	struct FItem
+	{
+		FItem(UFunction* Func, double InTime)
+		{
+#if RPC_CSV_TRACKER
+			Stat = FName(*Func->GetName());
+			Time = InTime;
+#endif
+		}
+		double Time;
+		FName Stat;
+	};
+
+	TMap<UFunction*, FItem>	FunctionMap;
+	
+} GRPCCSVTracker;
+
+// ----------------------------------------------------------------------------------------
+
+void UNetDriver::TickDispatch( float DeltaTime )
+{
+	SendCycles=0;
+
+	const double CurrentRealtime = FPlatformTime::Seconds();
+
+	const float DeltaRealtime = CurrentRealtime - LastTickDispatchRealtime;
+
+	LastTickDispatchRealtime = CurrentRealtime;
+
+	// Check to see if too much time is passing between ticks
+	// Setting this to somewhat large value for now, but small enough to catch blocking calls that are causing timeouts
+	const float TickLogThreshold = 5.0f;
+
+	bDidHitchLastFrame = (DeltaTime > TickLogThreshold || DeltaRealtime > TickLogThreshold);
+
+	if (bDidHitchLastFrame)
+	{
+		UE_LOG( LogNet, Log, TEXT( "UNetDriver::TickDispatch: Very long time between ticks. DeltaTime: %2.2f, Realtime: %2.2f. %s" ), DeltaTime, DeltaRealtime, *GetName() );
+	}
+
+	// Get new time.
+	PRAGMA_DISABLE_DEPRECATION_WARNINGS
+	Time += DeltaTime;
+	PRAGMA_ENABLE_DEPRECATION_WARNINGS
+	ElapsedTime += DeltaTime;
+
+	// Checks for standby cheats if enabled	
+	UpdateStandbyCheatStatus();
+
+	if (ServerConnection == nullptr)
+	{
+		// Delete any straggler connections
+		{
+			QUICK_SCOPE_CYCLE_COUNTER(UNetDriver_TickDispatch_CheckClientConnectionCleanup)
+
+			for (int32 i=ClientConnections.Num()-1; i>=0; i--)
+		{
+				if (ClientConnections[i]->State == USOCK_Closed)
+			{
+				ClientConnections[i]->CleanUp();
+			}
+		}
+	}
+
+		// Clean up recently disconnected client tracking
+		if (RecentlyDisconnectedClients.Num() > 0)
+		{
+			int32 NumToRemove = 0;
+
+			for (const FDisconnectedClient& CurElement : RecentlyDisconnectedClients)
+			{
+				if ((CurrentRealtime - CurElement.DisconnectTime) >= RecentlyDisconnectedTrackingTime)
+				{
+					verify(MappedClientConnections.Remove(CurElement.Address) == 1);
+
+					NumToRemove++;
+				}
+				else
+				{
+					break;
+				}
+			}
+
+			if (NumToRemove > 0)
+			{
+				RecentlyDisconnectedClients.RemoveAt(0, NumToRemove);
+			}
+		}
+	}
+
+#if RPC_CSV_TRACKER
+	GReceiveRPCTimingEnabled = (NetDriverName == NAME_GameNetDriver && ShouldEnableScopeSecondsTimers()) && (ServerConnection==nullptr);
+#endif
+}
+
+void UNetDriver::PostTickDispatch()
+{
+	// Flush out of order packet caches for connections that did not receive the missing packets during TickDispatch
+	if (ServerConnection != nullptr)
+	{
+		if (!ServerConnection->IsPendingKill())
+		{
+			ServerConnection->PostTickDispatch();
+		}
+	}
+
+	TArray<UNetConnection*> ClientConnCopy = ClientConnections;
+	for (UNetConnection* CurConn : ClientConnCopy)
+	{
+		if (!CurConn->IsPendingKill())
+		{
+			CurConn->PostTickDispatch();
+		}
+	}
+
+	if (ReplicationDriver)
+	{
+		ReplicationDriver->PostTickDispatch();
+	}
+
+	if (GReceiveRPCTimingEnabled)
+	{
+		GRPCCSVTracker.EndTickDispatch();
+		GReceiveRPCTimingEnabled = false;
+	}
+}
+
+void UNetDriver::NotifyRPCProcessed(UFunction* Function, UNetConnection* Connection, double ElapsedTimeSeconds)
+{
+	GRPCCSVTracker.NotifyRPCProcessed(Function, ElapsedTimeSeconds);
+}
+
+bool UNetDriver::IsLevelInitializedForActor(const AActor* InActor, const UNetConnection* InConnection) const
+{
+#if !(UE_BUILD_SHIPPING || UE_BUILD_TEST)
+	check(InActor);
+	check(InConnection);
+	check(World == InActor->GetWorld());
+#endif
+
+	// we can't create channels while the client is in the wrong world
+	const bool bCorrectWorld = GetWorldPackage() != nullptr && (InConnection->GetClientWorldPackageName() == GetWorldPackage()->GetFName()) && InConnection->ClientHasInitializedLevelFor(InActor);
+	// exception: Special case for PlayerControllers as they are required for the client to travel to the new world correctly			
+	const bool bIsConnectionPC = (InActor == InConnection->PlayerController);
+	return bCorrectWorld || bIsConnectionPC;
+}
+
+//
+// Internal RPC calling.
+//
+void UNetDriver::InternalProcessRemoteFunction
+(
+	AActor*			Actor,
+	UObject*		SubObject,
+	UNetConnection*	Connection,
+	UFunction*		Function,
+	void*			Parms,
+	FOutParmRec*	OutParms,
+	FFrame*			Stack,
+	bool			IsServer
+)
+{
+	// get the top most function
+	while (Function->GetSuperFunction())
+	{
+		Function = Function->GetSuperFunction();
+	}
+
+	// If saturated and function is unimportant, skip it. Note unreliable multicasts are queued at the actor channel level so they are not gated here.
+	if (!(Function->FunctionFlags & FUNC_NetReliable) && (!(Function->FunctionFlags & FUNC_NetMulticast)) && !Connection->IsNetReady(0))
+	{
+		DEBUG_REMOTEFUNCTION(TEXT("Network saturated, not calling %s::%s"), *GetNameSafe(Actor), *GetNameSafe(Function));
+		return;
+	}
+
+	// Route RPC calls to actual connection
+	if (Connection->GetUChildConnection())
+	{
+		Connection = ((UChildConnection*)Connection)->Parent;
+	}
+
+	// Prevent RPC calls to closed connections
+	if (Connection->State == USOCK_Closed)
+	{
+		DEBUG_REMOTEFUNCTION(TEXT("Attempting to call RPC on a closed connection. Not calling %s::%s"), *GetNameSafe(Actor), *GetNameSafe(Function));
+		return;
+	}
+
+	if (World == nullptr)
+	{
+		DEBUG_REMOTEFUNCTION(TEXT("Attempting to call RPC with a null World on the net driver. Not calling %s::%s"), *GetNameSafe(Actor), *GetNameSafe(Function));
+		return;
+	}
+
+	// If we have a subobject, thats who we are actually calling this on. If no subobject, we are calling on the actor.
+	UObject* TargetObj = SubObject ? SubObject : Actor;
+
+	// Make sure this function exists for both parties.
+	const FClassNetCache* ClassCache = NetCache->GetClassNetCache( TargetObj->GetClass() );
+	if (!ClassCache)
+	{
+		DEBUG_REMOTEFUNCTION(TEXT("ClassNetCache empty, not calling %s::%s"), *GetNameSafe(Actor), *GetNameSafe(Function));
+		return;
+	}
+		
+	const FFieldNetCache* FieldCache = ClassCache->GetFromField( Function );
+
+	if ( !FieldCache )
+	{
+		DEBUG_REMOTEFUNCTION(TEXT("FieldCache empty, not calling %s::%s"), *GetNameSafe(Actor), *GetNameSafe(Function));
+		return;
+	}
+		
+	// Get the actor channel.
+	UActorChannel* Ch = Connection->FindActorChannelRef(Actor);
+	if( !Ch )
+	{
+		if( IsServer )
+		{
+			if ( Actor->IsPendingKillPending() )
+			{
+				// Don't try opening a channel for me, I am in the process of being destroyed. Ignore my RPCs.
+				return;
+			}
+
+			if (IsLevelInitializedForActor(Actor, Connection))
+			{
+				Ch = (UActorChannel *)Connection->CreateChannelByName( NAME_Actor, EChannelCreateFlags::OpenedLocally );
+			}
+			else
+			{
+				UE_LOG(LogNet, Verbose, TEXT("Can't send function '%s' on actor '%s' because client hasn't loaded the level '%s' containing it"), *GetNameSafe(Function), *GetNameSafe(Actor), *GetNameSafe(Actor->GetLevel()));
+				return;
+			}
+		}
+		if (!Ch)
+		{
+			return;
+		}
+		if (IsServer)
+		{
+			Ch->SetChannelActor(Actor, ESetChannelActorFlags::None);
+		}	
+	}
+
+	ProcessRemoteFunctionForChannel(Ch, ClassCache, FieldCache, TargetObj, Connection, Function, Parms, OutParms, Stack, IsServer);
+}
+
+void UNetDriver::ProcessRemoteFunctionForChannel(UActorChannel* Ch, const FClassNetCache* ClassCache, const FFieldNetCache* FieldCache, UObject* TargetObj, UNetConnection* Connection, UFunction* Function, void* Parms, FOutParmRec* OutParms, FFrame* Stack, const bool IsServer, const ERemoteFunctionSendPolicy SendPolicy)
+{
+	// ---------------------------------------------------------------------------------------------------------------------------------------------------------------------------------------------------------------------------
+	// This function should be kept fast! Assume this is getting called multiple times at once. Don't look things up/recalc them if they do not change per connection/actor.
+	// ---------------------------------------------------------------------------------------------------------------------------------------------------------------------------------------------------------------------------
+
+	if (Ch->Closing)
+	{
+		return;
+	}
+
+	// Make sure initial channel-opening replication has taken place.
+	if( Ch->OpenPacketId.First==INDEX_NONE )
+	{
+		if (!IsServer)
+		{
+			DEBUG_REMOTEFUNCTION(TEXT("Initial channel replication has not occurred, not calling %s::%s"), *TargetObj->GetFullName(), *Function->GetName());
+			return;
+		}
+
+		// triggering replication of an Actor while already in the middle of replication can result in invalid data being sent and is therefore illegal
+		if (Ch->bIsReplicatingActor)
+		{
+			FString Error(FString::Printf(TEXT("Attempt to replicate function '%s' on Actor '%s' while it is in the middle of variable replication!"), *Function->GetName(), *TargetObj->GetFullName()));
+			UE_LOG(LogScript, Error, TEXT("%s"), *Error);
+			ensureMsgf(false, TEXT("%s"), *Error);
+			return;
+		}
+
+		// Bump the ReplicationFrame value to invalidate any properties marked as "unchanged" for this frame.
+		ReplicationFrame++;
+		
+		Ch->GetActor()->CallPreReplication(this);
+		Ch->ReplicateActor();
+	}
+
+	// Clients may be "closing" this connection but still processing bunches, we can't send anything if we have an invalid ChIndex.
+	if (Ch->ChIndex == -1)
+	{
+		ensure(!IsServer);
+		return;
+	}
+
+	// Form the RPC preamble.
+	FOutBunch Bunch( Ch, 0 );
+
+	// Reliability.
+	//warning: RPC's might overflow, preventing reliable functions from getting thorough.
+	if (Function->FunctionFlags & FUNC_NetReliable)
+	{
+		Bunch.bReliable = 1;
+	}
+
+	// verify we haven't overflowed unacked bunch buffer (Connection is not net ready)
+	//@warning: needs to be after parameter evaluation for script stack integrity
+	if ( Bunch.IsError() )
+	{
+		if ( !Bunch.bReliable )
+		{
+			// Not reliable, so not fatal. This can happen a lot in debug builds at startup if client is slow to get in game
+			UE_LOG( LogNet, Warning, TEXT( "Can't send function '%s' on '%s': Reliable buffer overflow. FieldCache->FieldNetIndex: %d Max %d. Ch MaxPacket: %d" ), *Function->GetName(), *TargetObj->GetFullName(), FieldCache->FieldNetIndex, ClassCache->GetMaxIndex(), Ch->Connection->MaxPacket );
+		}
+		else
+		{
+			// The connection has overflowed the reliable buffer. We cannot recover from this. Disconnect this user.
+			UE_LOG( LogNet, Warning, TEXT( "Closing connection. Can't send function '%s' on '%s': Reliable buffer overflow. FieldCache->FieldNetIndex: %d Max %d. Ch MaxPacket: %d." ), *Function->GetName(), *TargetObj->GetFullName(), FieldCache->FieldNetIndex, ClassCache->GetMaxIndex(), Ch->Connection->MaxPacket );
+
+			FString ErrorMsg = NSLOCTEXT( "NetworkErrors", "ClientReliableBufferOverflow", "Outgoing reliable buffer overflow" ).ToString();
+			FNetControlMessage<NMT_Failure>::Send( Connection, ErrorMsg );
+			Connection->FlushNet( true );
+			Connection->Close();
+#if USE_SERVER_PERF_COUNTERS
+			PerfCountersIncrement( TEXT( "ClosedConnectionsDueToReliableBufferOverflow" ) );
+#endif
+		}
+		return;
+	}
+
+#if !(UE_BUILD_SHIPPING || UE_BUILD_TEST)
+	extern TAutoConsoleVariable< int32 > CVarNetReliableDebug;
+
+	if ( CVarNetReliableDebug.GetValueOnAnyThread() > 0 )
+	{
+		Bunch.DebugString = FString::Printf( TEXT( "%.2f RPC: %s - %s" ), Connection->Driver->GetElapsedTime(), *TargetObj->GetFullName(), *Function->GetName() );
+	}
+#endif
+
+#if UE_NET_TRACE_ENABLED
+	SetTraceCollector(Bunch, UE_NET_TRACE_CREATE_COLLECTOR(ENetTraceVerbosity::Trace));
+#endif
+
+	TArray< FProperty * > LocalOutParms;
+
+	if( Stack == nullptr )
+	{
+		// Look for CPF_OutParm's, we'll need to copy these into the local parameter memory manually
+		// The receiving side will pull these back out when needed
+		for ( TFieldIterator<FProperty> It(Function); It && (It->PropertyFlags & (CPF_Parm|CPF_ReturnParm))==CPF_Parm; ++It )
+		{
+			if ( It->HasAnyPropertyFlags( CPF_OutParm ) )
+			{
+				if ( OutParms == NULL )
+				{
+					UE_LOG( LogNet, Warning, TEXT( "Missing OutParms. Property: %s, Function: %s, Actor: %s" ), *It->GetName(), *Function->GetName(), *TargetObj->GetFullName() );
+					continue;
+				}
+
+				FOutParmRec * Out = OutParms;
+
+				checkSlow( Out );
+
+				while ( Out->Property != *It )
+				{
+					Out = Out->NextOutParm;
+					checkSlow( Out );
+				}
+
+				void* Dest = It->ContainerPtrToValuePtr< void >( Parms );
+
+				const int32 CopySize = It->ElementSize * It->ArrayDim;
+
+				check( ( (uint8*)Dest - (uint8*)Parms ) + CopySize <= Function->ParmsSize );
+
+				It->CopyCompleteValue( Dest, Out->PropAddr );
+
+				LocalOutParms.Add( *It );
+			}
+		}
+	}
+
+	bool LogAsWarning = (GNetRPCDebug == 1);
+	if (LogAsWarning)
+	{
+		// Suppress spammy engine RPCs. This could be made a configable list in the future.
+		if (Function->GetName().Contains(TEXT("ServerUpdateCamera"))) LogAsWarning = false;
+		if (Function->GetName().Contains(TEXT("ClientAckGoodMove"))) LogAsWarning = false;
+		if (Function->GetName().Contains(TEXT("ServerMove"))) LogAsWarning = false;
+	}
+
+	FNetBitWriter TempWriter( Bunch.PackageMap, 0 );
+
+#if UE_NET_TRACE_ENABLED
+	// Create trace collector if tracing is enabled for the target bunch
+	SetTraceCollector(TempWriter, GetTraceCollector(Bunch) ? UE_NET_TRACE_CREATE_COLLECTOR(ENetTraceVerbosity::Trace) : nullptr);
+    ON_SCOPE_EXIT { UE_NET_TRACE_DESTROY_COLLECTOR(GetTraceCollector(TempWriter)); };
+#endif // UE_NET_TRACE_ENABLED
+
+	// Use the replication layout to send the rpc parameter values
+	TSharedPtr<FRepLayout> RepLayout = GetFunctionRepLayout( Function );
+	RepLayout->SendPropertiesForRPC( Function, Ch, TempWriter, Parms );
+
+	if (TempWriter.IsError())
+	{
+		if (LogAsWarning)
+		{
+			UE_LOG(LogNet, Warning, TEXT("Error: Can't send function '%s' on '%s': Failed to serialize properties"), *Function->GetName(), *TargetObj->GetFullName());
+		}
+		else
+		{
+			UE_LOG(LogNet, Log, TEXT("Error: Can't send function '%s' on '%s': Failed to serialize properties"), *Function->GetName(), *TargetObj->GetFullName());
+		}
+	}
+	else
+	{
+		// Make sure net field export group is registered
+		FNetFieldExportGroup* NetFieldExportGroup = Ch->GetOrCreateNetFieldExportGroupForClassNetCache( TargetObj );
+
+		int32 HeaderBits	= 0;
+		int32 ParameterBits	= 0;
+		
+		bool QueueBunch = false;
+		switch (SendPolicy)
+		{
+			case ERemoteFunctionSendPolicy::Default:
+				QueueBunch = ( !Bunch.bReliable && Function->FunctionFlags & FUNC_NetMulticast );
+				break;
+			case ERemoteFunctionSendPolicy::ForceQueue:
+				QueueBunch = true;
+				break;
+			case ERemoteFunctionSendPolicy::ForceSend:
+				QueueBunch = false;
+				break;
+		}
+
+		if ( QueueBunch )
+		{
+			Ch->WriteFieldHeaderAndPayload( Bunch, ClassCache, FieldCache, NetFieldExportGroup, TempWriter );
+			ParameterBits = Bunch.GetNumBits();
+		}
+		else
+		{
+			Ch->PrepareForRemoteFunction(TargetObj);
+
+			FNetBitWriter TempBlockWriter( Bunch.PackageMap, 0 );
+
+#if UE_NET_TRACE_ENABLED
+			UE_NET_TRACE_OBJECT_SCOPE(Ch->ActorNetGUID, Bunch, GetTraceCollector(Bunch), ENetTraceVerbosity::Trace);
+
+			// Ugliness to get data reported correctly, we basically fold the data from TempWriter into TempBlockWriter and then to Bunch
+			// Create trace collector if tracing is enabled for the target bunch			
+			SetTraceCollector(TempBlockWriter, GetTraceCollector(Bunch) ? UE_NET_TRACE_CREATE_COLLECTOR(ENetTraceVerbosity::Trace) : nullptr);
+#endif
+			
+			Ch->WriteFieldHeaderAndPayload( TempBlockWriter, ClassCache, FieldCache, NetFieldExportGroup, TempWriter );
+			ParameterBits = TempBlockWriter.GetNumBits();
+			HeaderBits = Ch->WriteContentBlockPayload( TargetObj, Bunch, false, TempBlockWriter );
+
+			UE_NET_TRACE_DESTROY_COLLECTOR(GetTraceCollector(TempBlockWriter));
+		}
+
+		// Destroy the memory used for the copied out parameters
+		for ( int32 i = 0; i < LocalOutParms.Num(); i++ )
+		{
+			check( LocalOutParms[i]->HasAnyPropertyFlags( CPF_OutParm ) );
+			LocalOutParms[i]->DestroyValue_InContainer( Parms );
+		}
+
+		// Send the bunch.
+		if( Bunch.IsError() )
+		{
+			UE_LOG(LogNet, Log, TEXT("Error: Can't send function '%s' on '%s': RPC bunch overflowed (too much data in parameters?)"), *Function->GetName(), *TargetObj->GetFullName());
+			ensureMsgf(false,TEXT("Error: Can't send function '%s' on '%s': RPC bunch overflowed (too much data in parameters?)"), *Function->GetName(), *TargetObj->GetFullName());
+		}
+		else if (Ch->Closing)
+		{
+			UE_LOG(LogNetTraffic, Log, TEXT("RPC bunch on closing channel") );
+		}
+		else
+		{
+			// Make sure we're tracking all the bits in the bunch
+			check(Bunch.GetNumBits() == HeaderBits + ParameterBits);
+
+			if (QueueBunch)
+			{
+				// Unreliable multicast functions are queued and sent out during property replication
+				if (LogAsWarning)
+				{
+					UE_LOG(LogNetTraffic, Warning,	TEXT("      Queing unreliable multicast RPC: %s::%s [%.1f bytes]"), *TargetObj->GetFullName(), *Function->GetName(), Bunch.GetNumBits() / 8.f );
+				}
+				else
+				{
+					UE_LOG(LogNetTraffic, Log,		TEXT("      Queing unreliable multicast RPC: %s::%s [%.1f bytes]"), *TargetObj->GetFullName(), *Function->GetName(), Bunch.GetNumBits() / 8.f );
+				}
+
+				NETWORK_PROFILER(GNetworkProfiler.TrackQueuedRPC(Connection, TargetObj, Ch->Actor, Function, HeaderBits, ParameterBits, 0));
+				Ch->QueueRemoteFunctionBunch(TargetObj, Function, Bunch);
+			}
+			else
+			{
+				if (LogAsWarning)
+				{
+					UE_LOG(LogNetTraffic, Warning,	TEXT("      Sent RPC: %s::%s [%.1f bytes]"), *TargetObj->GetFullName(), *Function->GetName(), Bunch.GetNumBits() / 8.f );
+				}
+				else
+				{
+					UE_LOG(LogNetTraffic, Log,		TEXT("      Sent RPC: %s::%s [%.1f bytes]"), *TargetObj->GetFullName(), *Function->GetName(), Bunch.GetNumBits() / 8.f );
+				}
+
+				NETWORK_PROFILER(GNetworkProfiler.TrackSendRPC(Ch->Actor, Function, HeaderBits, ParameterBits, 0, Connection));
+				Ch->SendBunch( &Bunch, 1 );
+			}
+		}
+	}
+
+	if ( Connection->IsInternalAck() )
+	{
+		Connection->FlushNet();
+	}
+}
+
+void UNetDriver::UpdateStandbyCheatStatus(void)
+{
+#if WITH_SERVER_CODE
+	// Only the server needs to check
+	if (ServerConnection == NULL && ClientConnections.Num())
+	{
+		QUICK_SCOPE_CYCLE_COUNTER(UNetDriver_TickDispatch_UpdateStandbyCheatStatus)
+
+		// Only check for cheats if enabled and one wasn't previously detected
+		if (bIsStandbyCheckingEnabled &&
+			bHasStandbyCheatTriggered == false &&
+			ClientConnections.Num() > 2)
+		{
+			int32 CountBadTx = 0;
+			int32 CountBadRx = 0;
+			int32 CountBadPing = 0;
+			
+			UWorld* FoundWorld = NULL;
+			// Look at each connection checking for a receive time and an ack time
+			for (int32 Index = 0; Index < ClientConnections.Num(); Index++)
+			{
+				UNetConnection* NetConn = ClientConnections[Index];
+				// Don't check connections that aren't fully formed (still loading & no controller)
+				// Controller won't be present until the join message is sent, which is after loading has completed
+				if (NetConn)
+				{
+					APlayerController* PlayerController = NetConn->PlayerController;
+					if(PlayerController)
+					{
+						UWorld* PlayerControllerWorld = PlayerController->GetWorld();
+						if( PlayerControllerWorld && 
+						   PlayerControllerWorld->GetTimeSeconds() - PlayerController->CreationTime > JoinInProgressStandbyWaitTime &&
+							// Ignore players with pending delete (kicked/timed out, but connection not closed)
+							PlayerController->IsPendingKillPending() == false)
+						{
+							if (!FoundWorld)
+							{
+								FoundWorld = PlayerControllerWorld;
+							}
+							else
+							{
+								check(FoundWorld == PlayerControllerWorld);
+							}
+							if (ElapsedTime - NetConn->LastReceiveTime > StandbyRxCheatTime)
+							{
+								CountBadRx++;
+							}
+							if (ElapsedTime - NetConn->GetLastRecvAckTime() > StandbyTxCheatTime)
+							{
+								CountBadTx++;
+							}
+							// Check for host tampering or crappy upstream bandwidth
+							if (PlayerController->PlayerState &&
+								PlayerController->PlayerState->ExactPing > BadPingThreshold)
+							{
+								CountBadPing++;
+							}
+						}
+					}
+				}
+			}
+			
+			if (FoundWorld)
+			{
+				AGameNetworkManager* const NetworkManager = FoundWorld->NetworkManager;
+				if (NetworkManager)
+				{
+					// See if we hit the percentage required for either TX or RX standby detection
+					if (float(CountBadRx) / float(ClientConnections.Num()) > PercentMissingForRxStandby)
+					{
+						bHasStandbyCheatTriggered = true;
+						NetworkManager->StandbyCheatDetected(STDBY_Rx);
+					}
+					else if (float(CountBadPing) / float(ClientConnections.Num()) > PercentForBadPing)
+					{
+						bHasStandbyCheatTriggered = true;
+						NetworkManager->StandbyCheatDetected(STDBY_BadPing);
+					}
+					// Check for the host not sending to the clients
+					else if (float(CountBadTx) / float(ClientConnections.Num()) > PercentMissingForTxStandby)
+					{
+						bHasStandbyCheatTriggered = true;
+						NetworkManager->StandbyCheatDetected(STDBY_Tx);
+					}
+				}
+			}
+		}
+	}
+#endif
+}
+
+
+void UNetDriver::SetAnalyticsProvider(TSharedPtr<IAnalyticsProvider> InProvider)
+{
+	AnalyticsProvider = InProvider;
+
+	if (AnalyticsProvider.IsValid())
+	{
+		if (AnalyticsAggregator.IsValid())
+		{
+			// If reinitializing an existing aggregator, hotfix its config
+			AnalyticsAggregator->InitConfig();
+
+			ensure(AnalyticsAggregator->GetAnalyticsProvider() == AnalyticsProvider);
+		}
+		else
+		{
+			AnalyticsAggregator = MakeShareable(new FNetAnalyticsAggregator(InProvider, NetDriverName));
+
+			AnalyticsAggregator->Init();
+		}
+	}
+	else
+	{
+		AnalyticsAggregator.Reset();
+	}
+
+	if (ConnectionlessHandler.IsValid())
+	{
+		ConnectionlessHandler->NotifyAnalyticsProvider(AnalyticsProvider, AnalyticsAggregator);
+	}
+
+	if (ServerConnection != nullptr)
+	{
+		ServerConnection->NotifyAnalyticsProvider();
+	}
+
+	for (UNetConnection* CurConn : ClientConnections)
+	{
+		if (CurConn != nullptr)
+		{
+			CurConn->NotifyAnalyticsProvider();
+		}
+	}
+}
+
+
+void UNetDriver::FRepChangedPropertyTrackerWrapper::CountBytes(FArchive& Ar) const
+{
+	if (FRepChangedPropertyTracker const * const LocalTracker = RepChangedPropertyTracker.Get())
+	{
+		LocalTracker->CountBytes(Ar);
+	}
+}
+
+void UNetDriver::FReplicationChangelistMgrWrapper::CountBytes(FArchive& Ar) const
+{
+	if (FReplicationChangelistMgr const * const ChangelistMgr = ReplicationChangelistMgr.Get())
+	{
+		Ar.CountBytes(sizeof(FReplicationChangelistMgr), sizeof(FReplicationChangelistMgr));
+		ChangelistMgr->CountBytes(Ar);
+	}
+}
+
+void UNetDriver::Serialize( FArchive& Ar )
+{
+	GRANULAR_NETWORK_MEMORY_TRACKING_INIT(Ar, "UNetDriver::Serialize");
+
+	GRANULAR_NETWORK_MEMORY_TRACKING_TRACK("UNetDriver::Super", Super::Serialize(Ar));
+
+	if (Ar.IsCountingMemory())
+	{
+		// TODO: We don't currently track:
+		//		StatelessConnectComponents
+		//		PacketHandlers
+		//		Network Address Bytes
+		//		AnalyticsData
+		//		NetworkNotify
+		//		Delegate Handles
+		//		DDoSDetection data
+		// These are probably insignificant, though.
+
+		GRANULAR_NETWORK_MEMORY_TRACKING_TRACK("MappedClientConnection", MappedClientConnections.CountBytes(Ar));
+		GRANULAR_NETWORK_MEMORY_TRACKING_TRACK("RecentlyDisconnectedClients", RecentlyDisconnectedClients.CountBytes(Ar));
+		GRANULAR_NETWORK_MEMORY_TRACKING_TRACK("GuidCache",
+			if (FNetGUIDCache const * const LocalGuidCache = GuidCache.Get())
+			{
+				LocalGuidCache->CountBytes(Ar);
+			}
+		);
+		
+		GRANULAR_NETWORK_MEMORY_TRACKING_TRACK("LocalNetCache",
+			if (FClassNetCacheMgr const * const LocalNetCache = NetCache.Get())
+			{
+				LocalNetCache->CountBytes(Ar);
+			}
+		);
+
+		GRANULAR_NETWORK_MEMORY_TRACKING_TRACK("LastPrioritizedActors", LastPrioritizedActors.CountBytes(Ar));
+		GRANULAR_NETWORK_MEMORY_TRACKING_TRACK("LastRelevantActors", LastRelevantActors.CountBytes(Ar));
+		GRANULAR_NETWORK_MEMORY_TRACKING_TRACK("LastSentActors", LastSentActors.CountBytes(Ar));
+		GRANULAR_NETWORK_MEMORY_TRACKING_TRACK("LastNonRelevantActors", LastNonRelevantActors.CountBytes(Ar));
+
+		GRANULAR_NETWORK_MEMORY_TRACKING_TRACK("DestroyedStartupOrDormantActors",
+			DestroyedStartupOrDormantActors.CountBytes(Ar);
+
+			for (const auto& DestroyedStartupOrDormantActorPair : DestroyedStartupOrDormantActors)
+			{
+				if (FActorDestructionInfo const * const DestructionInfo = DestroyedStartupOrDormantActorPair.Value.Get())
+				{
+					Ar.CountBytes(sizeof(FActorDestructionInfo), sizeof(FActorDestructionInfo));
+					DestructionInfo->PathName.CountBytes(Ar);
+				}
+			}
+		);
+
+		GRANULAR_NETWORK_MEMORY_TRACKING_TRACK("RenamedStartupActors", RenamedStartupActors.CountBytes(Ar));
+
+		GRANULAR_NETWORK_MEMORY_TRACKING_TRACK("RepChangedPropertyTrackerMap",
+			RepChangedPropertyTrackerMap.CountBytes(Ar);
+
+			for (const auto& RepChangedPropertyTrackerPair : RepChangedPropertyTrackerMap)
+			{
+				RepChangedPropertyTrackerPair.Value.CountBytes(Ar);
+			}
+		);
+
+		GRANULAR_NETWORK_MEMORY_TRACKING_TRACK("RepLayoutMap",
+			RepLayoutMap.CountBytes(Ar);
+
+			for (const auto& RepLayoutPair : RepLayoutMap)
+			{
+				if (FRepLayout const * const RepLayout = RepLayoutPair.Value.Get())
+				{
+					Ar.CountBytes(sizeof(FRepLayout), sizeof(FRepLayout));
+					RepLayout->CountBytes(Ar);
+				}
+			}
+		);
+
+		GRANULAR_NETWORK_MEMORY_TRACKING_TRACK("ReplicationChangeListMap", 
+			ReplicationChangeListMap.CountBytes(Ar);
+
+			for (const auto& ReplicationChangeListPair : ReplicationChangeListMap)
+			{
+				ReplicationChangeListPair.Value.CountBytes(Ar);
+			}
+		);
+
+		GRANULAR_NETWORK_MEMORY_TRACKING_TRACK("GuidToReplicatorMap", GuidToReplicatorMap.CountBytes(Ar));
+		GRANULAR_NETWORK_MEMORY_TRACKING_TRACK("UnmappedReplicators", UnmappedReplicators.CountBytes(Ar));
+		GRANULAR_NETWORK_MEMORY_TRACKING_TRACK("AllOwnedReplicators",
+		
+			// AllOwnedReplicators is the superset of all initialized FObjectReplicators,
+			// including DormantReplicators, Replicators on Channels, and UnmappedReplicators.
+			AllOwnedReplicators.CountBytes(Ar);
+			
+			for (FObjectReplicator const * const OwnedReplicator : AllOwnedReplicators)
+			{
+				if (OwnedReplicator)
+				{
+					Ar.CountBytes(sizeof(FObjectReplicator), sizeof(FObjectReplicator));
+					OwnedReplicator->CountBytes(Ar);
+				}
+			}
+		);
+
+		// Replicators are owned by UActorChannels, and so we don't track them here.
+
+		GRANULAR_NETWORK_MEMORY_TRACKING_TRACK("NetworkObjects",
+			if (FNetworkObjectList const * const NetObjList = NetworkObjects.Get())
+	{
+				Ar.CountBytes(sizeof(FNetworkObjectList), sizeof(FNetworkObjectList));
+				NetworkObjects->CountBytes(Ar);
+			}
+		);
+	}
+}
+
+void UNetDriver::FinishDestroy()
+{
+	if ( !HasAnyFlags(RF_ClassDefaultObject) )
+	{
+		// Make sure we tell listeners we are no longer lagging in case they set something up when lagging started.
+		if (GEngine && LagState != ENetworkLagState::NotLagging)
+		{
+			LagState = ENetworkLagState::NotLagging;
+			GEngine->BroadcastNetworkLagStateChanged(GetWorld(), this, LagState);
+		}
+
+		// Destroy server connection.
+		if( ServerConnection )
+		{
+			ServerConnection->CleanUp();
+		}
+		// Destroy client connections.
+		while( ClientConnections.Num() )
+		{
+			UNetConnection* ClientConnection = ClientConnections[0];
+			ClientConnection->CleanUp();
+		}
+		// Low level destroy.
+		LowLevelDestroy();
+
+		// Delete the guid cache
+		GuidCache.Reset();
+
+		FWorldDelegates::LevelRemovedFromWorld.Remove(OnLevelRemovedFromWorldHandle);
+		FWorldDelegates::LevelAddedToWorld.Remove(OnLevelAddedToWorldHandle);
+		FCoreUObjectDelegates::GetPostGarbageCollect().Remove(PostGarbageCollectHandle);
+	}
+	else
+	{
+		check(ServerConnection==NULL);
+		check(ClientConnections.Num()==0);
+		check(!GuidCache.IsValid());
+	}
+
+	// Make sure we've properly shut down all of the FObjectReplicator's
+	check( GuidToReplicatorMap.Num() == 0 );
+	check( TotalTrackedGuidMemoryBytes == 0 );
+	check( UnmappedReplicators.Num() == 0 );
+
+	Super::FinishDestroy();
+}
+
+void UNetDriver::LowLevelDestroy()
+{
+	// We are closing down all our sockets and low level communications.
+	// Sever the link with UWorld to ensure we don't tick again
+	SetWorld(NULL);
+}
+
+FString UNetDriver::LowLevelGetNetworkNumber()
+{
+	return LocalAddr.IsValid() ? LocalAddr->ToString(true) : FString(TEXT(""));
+}
+
+#if !UE_BUILD_SHIPPING
+bool UNetDriver::HandleSocketsCommand( const TCHAR* Cmd, FOutputDevice& Ar )
+{
+	// Print list of open connections.
+	Ar.Logf( TEXT("%s Connections:"), *GetDescription() );
+	if( ServerConnection )
+	{
+		Ar.Logf( TEXT("   Server %s"), *ServerConnection->LowLevelDescribe() );
+		for( int32 i=0; i<ServerConnection->OpenChannels.Num(); i++ )
+		{
+			UChannel* const Channel = ServerConnection->OpenChannels[i];
+			if (Channel)
+			{
+				Ar.Logf( TEXT("      Channel %i: %s"), Channel->ChIndex, *Channel->Describe() );
+			}
+		}
+	}
+#if WITH_SERVER_CODE
+	for( int32 i=0; i<ClientConnections.Num(); i++ )
+	{
+		UNetConnection* Connection = ClientConnections[i];
+		Ar.Logf( TEXT("   Client %s"), *Connection->LowLevelDescribe() );
+		for( int32 j=0; j<Connection->OpenChannels.Num(); j++ )
+		{
+			UChannel* const Channel = Connection->OpenChannels[j];
+			if (Channel)
+			{
+				Ar.Logf( TEXT("      Channel %i: %s"), Channel->ChIndex, *Channel->Describe() );
+			}
+		}
+	}
+#endif
+	return true;
+}
+
+bool UNetDriver::HandlePackageMapCommand( const TCHAR* Cmd, FOutputDevice& Ar )
+{
+	// Print packagemap for open connections
+	Ar.Logf(TEXT("Package Map:"));
+	if (ServerConnection != NULL)
+	{
+		Ar.Logf(TEXT("   Server %s"), *ServerConnection->LowLevelDescribe());
+		ServerConnection->PackageMap->LogDebugInfo(Ar);
+	}
+#if WITH_SERVER_CODE
+	for (int32 i = 0; i < ClientConnections.Num(); i++)
+	{
+		UNetConnection* Connection = ClientConnections[i];
+		Ar.Logf( TEXT("   Client %s"), *Connection->LowLevelDescribe() );
+		Connection->PackageMap->LogDebugInfo(Ar);
+	}
+#endif
+	return true;
+}
+
+bool UNetDriver::HandleNetFloodCommand( const TCHAR* Cmd, FOutputDevice& Ar )
+{
+	UNetConnection* TestConn = NULL;
+	if (ServerConnection != NULL)
+	{
+		TestConn = ServerConnection;
+	}
+#if WITH_SERVER_CODE
+	else if (ClientConnections.Num() > 0)
+	{
+		TestConn = ClientConnections[0];
+	}
+#endif
+	if (TestConn != NULL)
+	{
+		Ar.Logf(TEXT("Flooding connection 0 with control messages"));
+
+		for (int32 i = 0; i < 256 && TestConn->State == USOCK_Open; i++)
+		{
+			FNetControlMessage<NMT_Netspeed>::Send(TestConn, TestConn->CurrentNetSpeed);
+			TestConn->FlushNet();
+		}
+	}
+	return true;
+}
+
+bool UNetDriver::HandleNetDebugTextCommand( const TCHAR* Cmd, FOutputDevice& Ar )
+{
+	// Send a text string for testing connection
+	FString TestStr = FParse::Token(Cmd,false);
+	if (ServerConnection != NULL)
+	{
+		UE_LOG(LogNet, Log, TEXT("%s sending NMT_DebugText [%s] to [%s]"), 
+			*GetDescription(),*TestStr, *ServerConnection->LowLevelDescribe());
+
+		FNetControlMessage<NMT_DebugText>::Send(ServerConnection,TestStr);
+		ServerConnection->FlushNet(true);
+	}
+#if WITH_SERVER_CODE
+	else
+	{
+		for (int32 ClientIdx=0; ClientIdx < ClientConnections.Num(); ClientIdx++)
+		{
+			UNetConnection* Connection = ClientConnections[ClientIdx];
+			if (Connection)
+			{
+				UE_LOG(LogNet, Log, TEXT("%s sending NMT_DebugText [%s] to [%s]"), 
+					*GetDescription(),*TestStr, *Connection->LowLevelDescribe());
+
+				FNetControlMessage<NMT_DebugText>::Send(Connection,TestStr);
+				Connection->FlushNet(true);
+			}
+		}
+	}
+#endif // WITH_SERVER_CODE
+	return true;
+}
+
+bool UNetDriver::HandleNetDisconnectCommand( const TCHAR* Cmd, FOutputDevice& Ar )
+{
+	FString Msg = NSLOCTEXT("NetworkErrors", "NETDISCONNECTMSG", "NETDISCONNECT MSG").ToString();
+	if (ServerConnection != NULL)
+	{
+		UE_LOG(LogNet, Log, TEXT("%s disconnecting connection from host [%s]"), 
+			*GetDescription(),*ServerConnection->LowLevelDescribe());
+
+		FNetControlMessage<NMT_Failure>::Send(ServerConnection, Msg);
+	}
+#if WITH_SERVER_CODE
+	else
+	{
+		for (int32 ClientIdx=0; ClientIdx < ClientConnections.Num(); ClientIdx++)
+		{
+			UNetConnection* Connection = ClientConnections[ClientIdx];
+			if (Connection)
+			{
+				UE_LOG(LogNet, Log, TEXT("%s disconnecting from client [%s]"), 
+					*GetDescription(),*Connection->LowLevelDescribe());
+
+				FNetControlMessage<NMT_Failure>::Send(Connection, Msg);
+				Connection->FlushNet(true);
+			}
+		}
+	}
+#endif
+	return true;
+}
+
+bool UNetDriver::HandleNetDumpServerRPCCommand( const TCHAR* Cmd, FOutputDevice& Ar )
+{
+#if WITH_SERVER_CODE
+	for ( TObjectIterator<UClass> ClassIt; ClassIt; ++ClassIt )
+	{
+		bool bHasNetFields = false;
+
+		ClassIt->SetUpRuntimeReplicationData();
+
+		for ( int32 i = 0; i < ClassIt->NetFields.Num(); i++ )
+		{
+			UFunction * Function = Cast<UFunction>( ClassIt->NetFields[i] );
+
+			if ( Function != NULL && Function->FunctionFlags & FUNC_NetServer )
+			{
+				bHasNetFields = true;
+				break;
+			}
+		}
+
+		if ( !bHasNetFields )
+		{
+			continue;
+		}
+
+		Ar.Logf( TEXT( "Class: %s" ), *ClassIt->GetName() );
+
+		for ( int32 i = 0; i < ClassIt->NetFields.Num(); i++ )
+		{
+			UFunction * Function = Cast<UFunction>( ClassIt->NetFields[i] );
+
+			if ( Function != NULL && Function->FunctionFlags & FUNC_NetServer )
+			{
+				const FClassNetCache * ClassCache = NetCache->GetClassNetCache( *ClassIt );
+
+				const FFieldNetCache * FieldCache = ClassCache->GetFromField( Function );
+
+				TArray< FProperty * > Parms;
+
+				for ( TFieldIterator<FProperty> It( Function ); It && ( It->PropertyFlags & ( CPF_Parm | CPF_ReturnParm ) ) == CPF_Parm; ++It )
+				{
+					Parms.Add( *It );
+				}
+
+				if ( Parms.Num() == 0 )
+				{
+					Ar.Logf( TEXT( "    [0x%03x] %s();" ), FieldCache->FieldNetIndex, *Function->GetName() );
+					continue;
+				}
+
+				FString ParmString;
+
+				for ( int32 j = 0; j < Parms.Num(); j++ )
+				{
+					if ( CastField<FStructProperty>( Parms[j] ) )
+					{
+						ParmString += CastField<FStructProperty>( Parms[j] )->Struct->GetName();
+					}
+					else
+					{
+						ParmString += Parms[j]->GetClass()->GetName();
+					}
+
+					ParmString += TEXT( " " );
+
+					ParmString += Parms[j]->GetName();
+
+					if ( j < Parms.Num() - 1 )
+					{
+						ParmString += TEXT( ", " );
+					}
+				}						
+
+				Ar.Logf( TEXT( "    [0x%03x] %s( %s );" ), FieldCache->FieldNetIndex, *Function->GetName(), *ParmString );
+			}
+		}
+	}
+#endif
+	return true;
+}
+
+bool UNetDriver::HandleNetDumpDormancy(const TCHAR* Cmd, FOutputDevice& Ar)
+{
+	FProperty* Property = FindFProperty<FProperty>(AActor::StaticClass(), TEXT("NetDormancy"));
+	check(Property != nullptr);
+	Ar.Logf(TEXT(""));
+	TArray<AActor*> Actors;
+	if (FParse::Command(&Cmd, TEXT("ACTIVE")))
+	{
+		for (auto It = GetNetworkObjectList().GetActiveObjects().CreateConstIterator(); It; ++It)
+		{
+			FNetworkObjectInfo* ActorInfo = (*It).Get();
+			Actors.Add(ActorInfo->Actor);
+		}
+		Ar.Logf(TEXT("Logging network dormancy for %d active network objects"), Actors.Num());
+		Ar.Logf(TEXT(""));
+	}
+	else if (FParse::Command(&Cmd, TEXT("ALLDORMANT")))
+	{
+		for (auto It = GetNetworkObjectList().GetDormantObjectsOnAllConnections().CreateConstIterator(); It; ++It)
+		{
+			FNetworkObjectInfo* ActorInfo = (*It).Get();
+			Actors.Add(ActorInfo->Actor);
+		}
+		Ar.Logf(TEXT("Logging network dormancy for %d dormant network objects"), Actors.Num());
+		Ar.Logf(TEXT(""));
+	}
+	else
+	{
+		for (auto It = GetNetworkObjectList().GetAllObjects().CreateConstIterator(); It; ++It)
+		{
+			FNetworkObjectInfo* ActorInfo = (*It).Get();
+			Actors.Add(ActorInfo->Actor);
+		}
+		Ar.Logf(TEXT("Logging network dormancy for %d all network objects"), Actors.Num());
+		Ar.Logf(TEXT(""));
+	}
+	// Iterate through the objects reporting on their dormancy status
+	for (AActor* ThisActor : Actors)
+	{
+		uint8* BaseData = (uint8*)ThisActor;
+		FString ResultStr;
+		for (int32 i = 0; i < Property->ArrayDim; i++)
+		{
+			Property->ExportText_InContainer(i, ResultStr, BaseData, BaseData, ThisActor, PPF_IncludeTransient);
+		}
+
+		Ar.Logf(TEXT("%s.%s = %s"), *ThisActor->GetFullName(), *Property->GetName(), *ResultStr);
+	}
+	Ar.Logf(TEXT(""));
+	return true;
+}
+#endif // !UE_BUILD_SHIPPING
+
+void UNetDriver::HandlePacketLossBurstCommand( int32 DurationInMilliseconds )
+{
+#if DO_ENABLE_NET_TEST
+	UE_LOG(LogNet, Log, TEXT("%s simulating packet loss burst. Dropping incoming and outgoing packets for %d milliseconds."), *GetName(), DurationInMilliseconds);
+
+	PacketLossBurstEndTime = ElapsedTime + (DurationInMilliseconds / 1000.0);
+#endif
+}
+
+bool UNetDriver::Exec( UWorld* InWorld, const TCHAR* Cmd, FOutputDevice& Ar )
+{
+#if !UE_BUILD_SHIPPING
+	int32 PacketLossBurstMilliseconds = 0;
+
+	if( FParse::Command(&Cmd,TEXT("SOCKETS")) )
+	{
+		return HandleSocketsCommand( Cmd, Ar );
+	}
+	else if (FParse::Command(&Cmd, TEXT("PACKAGEMAP")))
+	{
+		return HandlePackageMapCommand( Cmd, Ar );
+	}
+	else if (FParse::Command(&Cmd, TEXT("NETFLOOD")))
+	{
+		return HandleNetFloodCommand( Cmd, Ar );
+	}
+#if DO_ENABLE_NET_TEST
+	// This will allow changing the Pkt* options at runtime
+	else if (bNeverApplyNetworkEmulationSettings == false && 
+			 PacketSimulationSettings.ParseSettings(Cmd, *NetDriverName.ToString()))
+	{
+		OnPacketSimulationSettingsChanged();
+		return true;
+	}
+#endif
+	else if (FParse::Command(&Cmd, TEXT("NETDEBUGTEXT")))
+	{
+		return HandleNetDebugTextCommand( Cmd, Ar );
+	}
+	else if (FParse::Command(&Cmd, TEXT("NETDISCONNECT")))
+	{
+		return HandleNetDisconnectCommand( Cmd, Ar );	
+	}
+	else if (FParse::Command(&Cmd, TEXT("DUMPSERVERRPC")))
+	{
+		return HandleNetDumpServerRPCCommand( Cmd, Ar );	
+	}
+	else if (FParse::Command(&Cmd, TEXT("DUMPDORMANCY")))
+	{
+		return HandleNetDumpDormancy(Cmd, Ar);
+	}
+	else if (FParse::Value(Cmd, TEXT("PKTLOSSBURST="), PacketLossBurstMilliseconds))
+	{
+		HandlePacketLossBurstCommand(PacketLossBurstMilliseconds);
+		return true;
+	}
+	else
+#endif // !UE_BUILD_SHIPPING
+	{
+		return false;
+	}
+}
+
+FActorDestructionInfo* UNetDriver::CreateDestructionInfo( UNetDriver* NetDriver, AActor* ThisActor, FActorDestructionInfo *DestructionInfo )
+{
+	if (DestructionInfo)
+	{
+		return DestructionInfo;
+	}
+
+	FNetworkGUID NetGUID = NetDriver->GuidCache->GetOrAssignNetGUID( ThisActor );
+	if (NetGUID.IsDefault())
+	{
+		UE_LOG(LogNet, Error, TEXT("CreateDestructionInfo got an invalid NetGUID for %s"), *ThisActor->GetName());
+		return nullptr;
+	}
+
+	TUniquePtr<FActorDestructionInfo>& NewInfoPtr = NetDriver->DestroyedStartupOrDormantActors.FindOrAdd( NetGUID );
+	if (NewInfoPtr.IsValid() == false)
+	{
+		NewInfoPtr = MakeUnique<FActorDestructionInfo>();
+	}
+
+	FActorDestructionInfo &NewInfo = *NewInfoPtr;
+	NewInfo.DestroyedPosition = ThisActor->GetActorLocation();
+	NewInfo.NetGUID = NetGUID;
+	NewInfo.Level = ThisActor->GetLevel();
+	NewInfo.ObjOuter = ThisActor->GetOuter();
+	NewInfo.PathName = ThisActor->GetName();
+
+	// Look for renamed actor now so we can clear it after the destroy is queued
+	FName RenamedPath = NetDriver->RenamedStartupActors.FindRef(ThisActor->GetFName());
+	if (RenamedPath != NAME_None)
+	{
+		NewInfo.PathName = RenamedPath.ToString();
+	}
+
+	if (NewInfo.Level.IsValid() && !NewInfo.Level->IsPersistentLevel() )
+	{
+		NewInfo.StreamingLevelName = NewInfo.Level->GetOutermost()->GetFName();
+	}
+	else
+	{
+		NewInfo.StreamingLevelName = NAME_None;
+	}
+
+	NewInfo.Reason = EChannelCloseReason::Destroyed;
+
+	return &NewInfo;
+}
+
+void UNetDriver::NotifyActorDestroyed( AActor* ThisActor, bool IsSeamlessTravel )
+{
+	// Remove the actor from the property tracker map
+	RepChangedPropertyTrackerMap.Remove(ThisActor);
+
+	const bool bIsServer = IsServer();
+	
+	if (bIsServer)
+	{
+		FActorDestructionInfo* DestructionInfo = nullptr;
+
+		const bool bIsActorStatic = !GuidCache->IsDynamicObject( ThisActor );
+		const bool bActorHasRole = ThisActor->GetRemoteRole() != ROLE_None;
+		const bool bShouldCreateDestructionInfo = bIsServer && bIsActorStatic && bActorHasRole && !IsSeamlessTravel;
+
+		if(bShouldCreateDestructionInfo)
+		{
+			UE_LOG(LogNet, VeryVerbose, TEXT("NotifyActorDestroyed %s - StartupActor"), *ThisActor->GetPathName() );
+			DestructionInfo = CreateDestructionInfo( this, ThisActor, DestructionInfo);
+		}
+
+		const FNetworkObjectInfo* NetworkObjectInfo = GetNetworkObjectList().Find( ThisActor ).Get();
+
+		for( int32 i=ClientConnections.Num()-1; i>=0; i-- )
+		{
+			UNetConnection* Connection = ClientConnections[i];
+			if( ThisActor->bNetTemporary )
+				Connection->SentTemporaries.Remove( ThisActor );
+			UActorChannel* Channel = Connection->FindActorChannelRef(ThisActor);
+			if( Channel )
+			{
+				check(Channel->OpenedLocally);
+				Channel->bClearRecentActorRefs = false;
+				Channel->Close(EChannelCloseReason::Destroyed);
+			}
+			else
+			{
+				const bool bDormantOrRecentlyDormant = NetworkObjectInfo && (NetworkObjectInfo->DormantConnections.Contains(Connection) || NetworkObjectInfo->RecentlyDormantConnections.Contains(Connection));
+
+				if (bShouldCreateDestructionInfo || bDormantOrRecentlyDormant)
+				{
+					// Make a new destruction info if necessary. It is necessary if the actor is dormant or recently dormant because
+					// even though the client knew about the actor at some point, it doesn't have a channel to handle destruction.
+					DestructionInfo = CreateDestructionInfo(this, ThisActor, DestructionInfo);
+					if (DestructionInfo)
+					{
+					Connection->AddDestructionInfo(DestructionInfo);
+				}
+			}
+			}
+
+			// Remove it from any dormancy lists
+			Connection->CleanupDormantReplicatorsForActor(ThisActor);
+		}
+	}
+
+	if (ServerConnection)
+	{
+		ServerConnection->CleanupDormantReplicatorsForActor(ThisActor);
+	}
+
+	// Remove this actor from the network object list
+	RemoveNetworkActor( ThisActor );
+}
+
+
+void UNetDriver::RemoveNetworkActor(AActor* Actor)
+{
+	GetNetworkObjectList().Remove(Actor);
+
+	// Remove from renamed list if destroyed
+	RenamedStartupActors.Remove(Actor->GetFName());
+
+	if (ReplicationDriver)
+	{
+		ReplicationDriver->RemoveNetworkActor(Actor);
+	}
+}
+
+void UNetDriver::NotifyActorRenamed(AActor* ThisActor, FName PreviousName)
+{
+	const bool bIsServer = IsServer();
+	const bool bIsActorStatic = !GuidCache->IsDynamicObject(ThisActor);
+	const bool bActorHasRole = ThisActor->GetRemoteRole() != ROLE_None;
+
+#if WITH_EDITOR
+	// When recompiling and reinstancing a Blueprint, we rename the old actor out of the way, which would cause this code to emit a warning during PIE
+	// Since that old actor is about to die (on both client and server) in the reinstancing case, it's safe to skip
+	if (GCompilingBlueprint)
+	{
+		return;
+	}
+#endif
+
+	if (bIsActorStatic && bActorHasRole)
+	{
+		if (bIsServer)
+		{
+			FName OriginalName = RenamedStartupActors.FindRef(PreviousName);
+			if (OriginalName != NAME_None)
+			{
+				PreviousName = OriginalName;
+			}
+
+			RenamedStartupActors.Add(ThisActor->GetFName(), PreviousName);
+
+			UE_LOG(LogNet, Log, TEXT("NotifyActorRenamed StartupActor: %s PreviousName: %s"), *ThisActor->GetName(), *PreviousName.ToString());
+		}
+		else 
+		{
+			UE_LOG(LogNet, Warning, TEXT("NotifyActorRenamed on client, StartupActor: %s"), *ThisActor->GetName());
+		}
+	}
+}
+
+// This method will be called when a Streaming Level is about to be Garbage Collected.
+void UNetDriver::NotifyStreamingLevelUnload(ULevel* Level)
+{
+	if (IsServer())
+	{
+		for (AActor* Actor : Level->Actors)
+		{
+			if (Actor && Actor->IsNetStartupActor())
+			{
+				NotifyActorLevelUnloaded(Actor);
+			}
+		}
+		
+		TArray<FNetworkGUID> RemovedGUIDs;
+		for (auto It = DestroyedStartupOrDormantActors.CreateIterator(); It; ++It)
+		{
+			FActorDestructionInfo* DestructInfo = It->Value.Get();
+			if (DestructInfo->Level == Level && DestructInfo->NetGUID.IsStatic())
+			{
+				for (UNetConnection* Connection : ClientConnections)
+				{
+					Connection->RemoveDestructionInfo(DestructInfo);
+				}
+
+				RemovedGUIDs.Add(It->Key);
+				It.RemoveCurrent();
+			}
+		}
+	}	
+
+	if (Level->LevelScriptActor)
+	{
+		RemoveClassRepLayoutReferences(Level->LevelScriptActor->GetClass());
+		ReplicationChangeListMap.Remove(Level->LevelScriptActor);
+	}
+
+	if (ServerConnection && ServerConnection->PackageMap)
+	{
+		UE_LOG(LogNet, Log, TEXT("NotifyStreamingLevelUnload: %s"), *Level->GetFullName() );
+
+		if (Level->LevelScriptActor)
+		{
+			UActorChannel * Channel = ServerConnection->FindActorChannelRef((AActor*)Level->LevelScriptActor);
+			if (Channel)
+			{
+				UE_LOG(LogNet, Log, TEXT("NotifyStreamingLevelUnload: BREAKING"));
+
+				Channel->Actor = NULL;
+				Channel->Broken = true;
+				Channel->CleanupReplicators();
+			}
+		}
+
+		ServerConnection->PackageMap->NotifyStreamingLevelUnload(Level);
+	}
+
+	for( int32 i=ClientConnections.Num()-1; i>=0; i-- )
+	{
+		UNetConnection* Connection = ClientConnections[i];
+		if (Connection && Connection->PackageMap)
+		{
+			Connection->PackageMap->NotifyStreamingLevelUnload(Level);
+		}
+	}
+}
+
+/** Called when an actor is being unloaded during a seamless travel or do due level streaming 
+ *  The main point is that it calls the normal NotifyActorDestroyed to destroy the channel on the server
+ *	but also removes the Actor reference, sets broken flag, and cleans up actor class references on clients.
+ */
+void UNetDriver::NotifyActorLevelUnloaded( AActor* TheActor )
+{
+	// server
+	NotifyActorDestroyed(TheActor, true);
+	// client
+	if (ServerConnection != NULL)
+	{
+		// we can't kill the channel until the server says so, so just clear the actor ref and break the channel
+		UActorChannel* Channel = ServerConnection->FindActorChannelRef(TheActor);
+		if (Channel != NULL)
+		{
+			ServerConnection->RemoveActorChannel(TheActor);
+			Channel->Actor = NULL;
+			Channel->Broken = true;
+			Channel->CleanupReplicators();
+		}
+	}
+}
+
+void UNetDriver::NotifyActorTearOff(AActor* Actor)
+{
+	if (ReplicationDriver)
+	{
+		ReplicationDriver->NotifyActorTearOff(Actor);
+	}
+}
+
+void UNetDriver::ForceNetUpdate(AActor* Actor)
+{
+	// Let Replication Driver handle it if its exists
+	if (ReplicationDriver)
+	{
+		ReplicationDriver->ForceNetUpdate(Actor);
+		return;
+	}
+	
+	// Legacy implementation
+	if ( FNetworkObjectInfo* NetActor = FindNetworkObjectInfo(Actor) )
+	{
+		NetActor->NextUpdateTime = World->TimeSeconds - 0.01f;
+	}
+}
+
+void UNetDriver::ForceAllActorsNetUpdateTime(float NetUpdateTimeOffset, TFunctionRef<bool(const AActor* const)> ValidActorTestFunc)
+{
+	for (auto It = GetNetworkObjectList().GetAllObjects().CreateConstIterator(); It; ++It)
+	{
+		FNetworkObjectInfo* NetActorInfo = (*It).Get();
+		if (NetActorInfo)
+		{
+			const AActor* const Actor = NetActorInfo->WeakActor.Get();
+			if (Actor)
+			{
+				if (ValidActorTestFunc(Actor))
+				{
+					// Only allow the next update to be sooner than the current one
+					const double NewUpdateTime = World->TimeSeconds + NetUpdateTimeOffset * FMath::FRand();
+					NetActorInfo->NextUpdateTime = FMath::Min(NetActorInfo->NextUpdateTime, NewUpdateTime);
+				}
+			}
+		}
+	}
+}
+
+/** UNetDriver::FlushActorDormancy(AActor* Actor)
+ *	 Flushes the actor from the NetDriver's dormant list and/or cancels pending dormancy on the actor channel.
+ *
+ *	 This does not change the Actor's actual NetDormant state. If a dormant actor is Flushed, it will net update at least one more
+ *	 time, and then go back to dormant.
+ */
+void UNetDriver::FlushActorDormancy(AActor* Actor, bool bWasDormInitial)
+{
+	// Note: Going into dormancy is completely handled in ServerReplicateActor. We want to avoid
+	// event-based handling of going into dormancy, because we have to deal with connections joining in progress.
+	// It is better to have ::ServerReplicateActor check the AActor and UChannel's states to determined if an actor
+	// needs to be moved into dormancy. The same amount of work will be done (1 time per connection when an actor goes dorm)
+	// and we avoid having to do special things when a new client joins.
+	//
+	// Going out of dormancy can be event based like this since it only affects clients already joined. Its more efficient in this
+	// way too, since we dont have to check every dormant actor in ::ServerReplicateActor to see if it needs to go out of dormancy
+
+	if (GSetNetDormancyEnabled == 0)
+		return;
+
+	check(Actor);
+	check(ServerConnection == NULL);
+
+	if (ReplicationDriver)
+	{
+		ReplicationDriver->FlushNetDormancy(Actor, bWasDormInitial);
+	}
+
+	FlushActorDormancyInternal(Actor);
+}
+
+void UNetDriver::NotifyActorDormancyChange(AActor* Actor, ENetDormancy OldDormancyState)
+{
+	if (ReplicationDriver)
+	{
+		ReplicationDriver->NotifyActorDormancyChange(Actor, OldDormancyState);
+	}
+
+	FlushActorDormancyInternal(Actor);
+}
+
+void UNetDriver::FlushActorDormancyInternal(AActor *Actor)
+{
+	QUICK_SCOPE_CYCLE_COUNTER(STAT_NetDriver_FlushActorDormancy);
+	// Go through each connection and remove the actor from the dormancy list
+	for (int32 i=0; i < ClientConnections.Num(); ++i)
+	{
+		UNetConnection *NetConnection = ClientConnections[i];
+		if(NetConnection != NULL)
+		{
+			NetConnection->FlushDormancy(Actor);
+		}
+	}
+}
+
+void UNetDriver::ForcePropertyCompare( AActor* Actor )
+{
+#if WITH_SERVER_CODE
+	check( Actor );
+
+	for ( int32 i=0; i < ClientConnections.Num(); ++i )
+	{
+		UNetConnection *NetConnection = ClientConnections[i];
+		if ( NetConnection != NULL )
+		{
+			NetConnection->ForcePropertyCompare( Actor );
+		}
+	}
+#endif // WITH_SERVER_CODE
+}
+
+void UNetDriver::ForceActorRelevantNextUpdate(AActor* Actor)
+{
+#if WITH_SERVER_CODE
+	check(Actor);
+	
+	GetNetworkObjectList().ForceActorRelevantNextUpdate(Actor, this);
+#endif // WITH_SERVER_CODE
+}
+
+UChildConnection* UNetDriver::CreateChild(UNetConnection* Parent)
+{
+	UE_LOG(LogNet, Log, TEXT("Creating child connection with %s parent"), *Parent->GetName());
+	auto Child = NewObject<UChildConnection>();
+	Child->Driver = this;
+	Child->URL = FURL();
+	Child->State = Parent->State;
+	Child->URL.Host = Parent->URL.Host;
+	Child->Parent = Parent;
+	Child->PackageMap = Parent->PackageMap;
+	Child->CurrentNetSpeed = Parent->CurrentNetSpeed;
+	Parent->Children.Add(Child);
+	return Child;
+}
+
+void UNetDriver::PostGarbageCollect()
+{
+	// We can't perform this logic in AddReferencedObjects because destroying GCObjects
+	// during Garbage Collection is illegal (@see UGCObjectReferencer::RemoveObject).
+	// FRepLayout's are GC objects, and either map could be holding onto the last reference
+	// to a given RepLayout.
+
+	for (auto It = RepLayoutMap.CreateIterator(); It; ++It)
+	{
+		if (!It.Key().IsValid())
+		{
+			It.RemoveCurrent();
+		}
+	}
+
+	for (auto It = ReplicationChangeListMap.CreateIterator(); It; ++It)
+	{
+		if (!It.Value().IsObjectValid())
+		{
+			It.RemoveCurrent();
+		}
+	}
+}
+
+void UNetDriver::AddReferencedObjects(UObject* InThis, FReferenceCollector& Collector)
+{
+	UNetDriver* This = CastChecked<UNetDriver>(InThis);
+	Super::AddReferencedObjects(This, Collector);
+
+	// TODO: It feels like we could get away without FRepLayout needing to track references.
+	//			E.G., if we detected that FObjectReplicator / FReplicationChangelistMgrs detected
+	//			their associated objects were destroyed, we could destroy the Shadow Buffers
+	//			which should be the last thing referencing the FRepLayout and the Properties.
+	for (auto It = This->RepLayoutMap.CreateIterator(); It; ++It)
+	{
+		It.Value()->AddReferencedObjects(Collector);
+	}
+	
+	for (FObjectReplicator* Replicator : This->AllOwnedReplicators)
+		{
+		Collector.AddReferencedObject(Replicator->ObjectPtr, This);
+		Collector.AddReferencedObject(Replicator->ObjectClass, This);
+			}
+
+	for (FConnectionMap::TIterator It(This->MappedClientConnections); It; ++It)
+		{
+		Collector.AddReferencedObject(It.Value(), This);
+		}
+
+	if (This->GuidCache.IsValid())
+	{
+		This->GuidCache->CollectReferences(Collector);
+	}
+}
+
+
+#if DO_ENABLE_NET_TEST
+
+void UNetDriver::SetPacketSimulationSettings(const FPacketSimulationSettings& NewSettings)
+{
+	if (bNeverApplyNetworkEmulationSettings)
+	{
+		return;
+	}
+
+	// Once set this prevents driver changes from reloading the packet emulation settings
+	bForcedPacketSettings = true;
+
+	PacketSimulationSettings = NewSettings;
+	OnPacketSimulationSettingsChanged();
+}
+
+void UNetDriver::OnPacketSimulationSettingsChanged()
+{
+	if (ServerConnection)
+	{
+		ServerConnection->UpdatePacketSimulationSettings();
+	}
+#if WITH_SERVER_CODE
+	for (UNetConnection* ClientConnection : ClientConnections)
+	{
+		if (ClientConnection)
+		{
+			ClientConnection->UpdatePacketSimulationSettings();
+		}
+	}
+#endif
+}
+
+class FPacketSimulationConsoleCommandVisitor 
+{
+public:
+	static void OnPacketSimulationConsoleCommand(const TCHAR *Name, IConsoleObject* CVar, TArray<IConsoleObject*>& Sink)
+	{
+		Sink.Add(CVar);
+	}
+};
+
+/** reads in settings from the .ini file 
+ * @note: overwrites all previous settings
+ */
+void FPacketSimulationSettings::LoadConfig(const TCHAR* OptionalQualifier)
+{
+	ConfigHelperInt(TEXT("PktLoss"), PktLoss, OptionalQualifier);
+	
+	ConfigHelperInt(TEXT("PktLossMinSize"), PktLossMinSize, OptionalQualifier);
+	ConfigHelperInt(TEXT("PktLossMaxSize"), PktLossMaxSize, OptionalQualifier);
+
+	bool InPktOrder = !!PktOrder;
+	ConfigHelperBool(TEXT("PktOrder"), InPktOrder, OptionalQualifier);
+	PktOrder = int32(InPktOrder);
+	
+	ConfigHelperInt(TEXT("PktLag"), PktLag, OptionalQualifier);
+	ConfigHelperInt(TEXT("PktLagVariance"), PktLagVariance, OptionalQualifier);
+
+	ConfigHelperInt(TEXT("PktLagMin"), PktLagMin, OptionalQualifier);
+	ConfigHelperInt(TEXT("PktLagMax"), PktLagMax, OptionalQualifier);
+	
+	ConfigHelperInt(TEXT("PktDup"), PktDup, OptionalQualifier);
+
+	ConfigHelperInt(TEXT("PktIncomingLagMin"), PktIncomingLagMin, OptionalQualifier);
+	ConfigHelperInt(TEXT("PktIncomingLagMax"), PktIncomingLagMax, OptionalQualifier);
+	ConfigHelperInt(TEXT("PktIncomingLoss"), PktIncomingLoss, OptionalQualifier);
+
+	ConfigHelperInt(TEXT("PktJitter"), PktJitter, OptionalQualifier);
+
+	ValidateSettings();
+}
+
+bool FPacketSimulationSettings::LoadEmulationProfile(const TCHAR* ProfileName)
+{
+	const FString SectionName = FString::Printf(TEXT("%s.%s"), TEXT("PacketSimulationProfile"), ProfileName);
+
+	TArray<FString> SectionConfigs;
+	bool bSectionExists = GConfig->GetSection(*SectionName, SectionConfigs, GEngineIni);
+	if (!bSectionExists)
+	{
+		UE_LOG(LogNet, Log, TEXT("EmulationProfile [%s] was not found in %s. Packet settings were not changed"), *SectionName, *GEngineIni);
+		return false;
+	}
+
+	ResetSettings();
+
+	UScriptStruct* ThisStruct = FPacketSimulationSettings::StaticStruct();
+
+	for (const FString& ConfigVar : SectionConfigs)
+	{
+		FString VarName;
+		FString VarValue;
+		if (ConfigVar.Split(TEXT("="), &VarName, &VarValue))
+		{
+			// If using the one line struct definition
+			if (VarName.Equals(TEXT("PacketSimulationSettings"), ESearchCase::IgnoreCase))
+			{
+				ThisStruct->ImportText(*VarValue, this, nullptr, 0, (FOutputDevice*)GWarn, TEXT("FPacketSimulationSettings"));
+			}
+			else if (FProperty* StructProperty = ThisStruct->FindPropertyByName(FName(*VarName, FNAME_Find)))
+			{
+				StructProperty->ImportText(*VarValue, StructProperty->ContainerPtrToValuePtr<void>(this), 0, nullptr);
+	}
+			else
+	{
+				UE_LOG(LogNet, Warning, TEXT("FPacketSimulationSettings::LoadEmulationProfile could not find property named %s"), *VarName);
+			}
+	}
+	}
+	
+	ValidateSettings();
+
+	return true;
+}
+
+void FPacketSimulationSettings::ResetSettings()
+{
+	*this = FPacketSimulationSettings();
+}
+
+void FPacketSimulationSettings::ValidateSettings()
+{
+	PktLoss = FMath::Clamp<int32>(PktLoss, 0, 100);
+
+	PktOrder = FMath::Clamp<int32>(PktOrder, 0, 1);
+
+	PktLagMin = FMath::Max(PktLagMin, 0);
+	PktLagMax = FMath::Max(PktLagMin, PktLagMax);
+
+	PktDup = FMath::Clamp<int32>(PktDup, 0, 100);
+
+	PktIncomingLagMin = FMath::Max(PktIncomingLagMin, 0);
+	PktIncomingLagMax = FMath::Max(PktIncomingLagMin, PktIncomingLagMax);
+	PktIncomingLoss = FMath::Clamp<int32>(PktIncomingLoss, 0, 100);
+}
+
+bool FPacketSimulationSettings::ConfigHelperInt(const TCHAR* Name, int32& Value, const TCHAR* OptionalQualifier)
+{
+	if (OptionalQualifier)
+	{
+		if (GConfig->GetInt(TEXT("PacketSimulationSettings"), *FString::Printf(TEXT("%s%s"), OptionalQualifier, Name), Value, GEngineIni))
+		{
+			return true;
+		}
+	}
+
+	if (GConfig->GetInt(TEXT("PacketSimulationSettings"), Name, Value, GEngineIni))
+	{
+		return true;
+	}
+
+	return false;
+}
+
+bool FPacketSimulationSettings::ConfigHelperBool(const TCHAR* Name, bool& Value, const TCHAR* OptionalQualifier)
+{
+	if (OptionalQualifier)
+	{
+		if (GConfig->GetBool(TEXT("PacketSimulationSettings"), *FString::Printf(TEXT("%s%s"), OptionalQualifier, Name), Value, GEngineIni))
+		{
+			return true;
+		}
+	}
+
+	if (GConfig->GetBool(TEXT("PacketSimulationSettings"), Name, Value, GEngineIni))
+	{
+		return true;
+	}
+
+	return false;
+}
+
+/**
+ * Reads the settings from a string: command line or an exec
+ *
+ * @param Stream the string to read the settings from
+ */
+bool FPacketSimulationSettings::ParseSettings(const TCHAR* Cmd, const TCHAR* OptionalQualifier)
+{
+	// note that each setting is tested.
+	// this is because the same function will be used to parse the command line as well
+	bool bParsed = false;
+
+	FString EmulationProfileName;
+	if (FParse::Value(Cmd, TEXT("PktEmulationProfile="), EmulationProfileName))
+	{
+		UE_LOG(LogNet, Log, TEXT("Applying EmulationProfile %s"), *EmulationProfileName);
+		bParsed = LoadEmulationProfile(*EmulationProfileName);
+	}
+	if( ParseHelper(Cmd, TEXT("PktLoss="), PktLoss, OptionalQualifier) )
+	{
+		bParsed = true;
+		UE_LOG(LogNet, Log, TEXT("PktLoss set to %d"), PktLoss);
+	}
+	if (ParseHelper(Cmd, TEXT("PktLossMinSize="), PktLossMinSize, OptionalQualifier))
+	{
+		bParsed = true;
+		UE_LOG(LogNet, Log, TEXT("PktLossMinSize set to %d"), PktLossMinSize);
+	}
+	if (ParseHelper(Cmd, TEXT("PktLossMaxSize="), PktLossMaxSize, OptionalQualifier))
+	{
+		bParsed = true;
+		UE_LOG(LogNet, Log, TEXT("PktLossMaxSize set to %d"), PktLossMaxSize);
+	}
+	if( ParseHelper(Cmd, TEXT("PktOrder="), PktOrder, OptionalQualifier) )
+	{
+		bParsed = true;
+		UE_LOG(LogNet, Log, TEXT("PktOrder set to %d"), PktOrder);
+	}
+	if( ParseHelper(Cmd, TEXT("PktLag="), PktLag, OptionalQualifier) )
+	{
+		bParsed = true;
+		UE_LOG(LogNet, Log, TEXT("PktLag set to %d"), PktLag);
+	}
+	if( ParseHelper(Cmd, TEXT("PktDup="), PktDup, OptionalQualifier) )
+	{
+		bParsed = true;
+		UE_LOG(LogNet, Log, TEXT("PktDup set to %d"), PktDup);
+	}	
+	if (ParseHelper(Cmd, TEXT("PktLagVariance="), PktLagVariance, OptionalQualifier))
+	{
+		bParsed = true;
+		UE_LOG(LogNet, Log, TEXT("PktLagVariance set to %d"), PktLagVariance);
+	}
+	if (ParseHelper(Cmd, TEXT("PktLagMin="), PktLagMin, OptionalQualifier))
+	{
+		bParsed = true;
+		UE_LOG(LogNet, Log, TEXT("PktLagMin set to %d"), PktLagMin);
+	}
+	if (ParseHelper(Cmd, TEXT("PktLagMax="), PktLagMax, OptionalQualifier))
+	{
+		bParsed = true;
+		UE_LOG(LogNet, Log, TEXT("PktLagMax set to %d"), PktLagMax);
+	}
+	if (ParseHelper(Cmd, TEXT("PktIncomingLagMin="), PktIncomingLagMin, OptionalQualifier))
+	{
+		bParsed = true;
+		UE_LOG(LogNet, Log, TEXT("PktIncomingLagMin set to %d"), PktIncomingLagMin);
+	}
+	if (ParseHelper(Cmd, TEXT("PktIncomingLagMax="), PktIncomingLagMax, OptionalQualifier))
+	{
+		bParsed = true;
+		UE_LOG(LogNet, Log, TEXT("PktIncomingLagMax set to %d"), PktIncomingLagMax);
+	}
+	if (ParseHelper(Cmd, TEXT("PktIncomingLoss="), PktIncomingLoss, OptionalQualifier))
+	{
+		bParsed = true;
+		UE_LOG(LogNet, Log, TEXT("PktIncomingLoss set to %d"), PktIncomingLoss);
+	}
+	if (ParseHelper(Cmd, TEXT("PktJitter="), PktJitter, OptionalQualifier))
+	{
+		bParsed = true;
+		UE_LOG(LogNet, Log, TEXT("PktJitter set to %d"), PktJitter);
+	}
+
+	ValidateSettings();
+	return bParsed;
+}
+
+bool FPacketSimulationSettings::ParseHelper(const TCHAR* Cmd, const TCHAR* Name, int32& Value, const TCHAR* OptionalQualifier)
+{
+	if (OptionalQualifier)
+	{
+		if (FParse::Value(Cmd, *FString::Printf(TEXT("%s%s"), OptionalQualifier, Name), Value))
+		{
+			return true;
+		}
+	}
+
+	if (FParse::Value(Cmd, Name, Value))
+	{
+		return true;
+	}
+	return false;
+}
+
+#endif //#if DO_ENABLE_NET_TEST
+
+FNetViewer::FNetViewer(UNetConnection* InConnection, float DeltaSeconds) :
+	Connection(InConnection),
+	InViewer(InConnection->PlayerController ? InConnection->PlayerController : InConnection->OwningActor),
+	ViewTarget(InConnection->ViewTarget),
+	ViewLocation(ForceInit),
+	ViewDir(ForceInit)
+{
+	check(InConnection->OwningActor);
+	check(!InConnection->PlayerController || (InConnection->PlayerController == InConnection->OwningActor));
+
+	APlayerController* ViewingController = InConnection->PlayerController;
+
+	// Get viewer coordinates.
+	ViewLocation = ViewTarget->GetActorLocation();
+	if (ViewingController)
+	{
+		FRotator ViewRotation = ViewingController->GetControlRotation();
+		ViewingController->GetPlayerViewPoint(ViewLocation, ViewRotation);
+		ViewDir = ViewRotation.Vector();
+	}
+}
+
+FActorPriority::FActorPriority(UNetConnection* InConnection, UActorChannel* InChannel, FNetworkObjectInfo* InActorInfo, const TArray<struct FNetViewer>& Viewers, bool bLowBandwidth)
+	: ActorInfo(InActorInfo), Channel(InChannel), DestructionInfo(NULL)
+{	
+	const float Time = Channel ? (InConnection->Driver->GetElapsedTime() - Channel->LastUpdateTime) : InConnection->Driver->SpawnPrioritySeconds;
+	// take the highest priority of the viewers on this connection
+	Priority = 0;
+	for (int32 i = 0; i < Viewers.Num(); i++)
+	{
+		Priority = FMath::Max<int32>(Priority, FMath::RoundToInt(65536.0f * ActorInfo->Actor->GetNetPriority(Viewers[i].ViewLocation, Viewers[i].ViewDir, Viewers[i].InViewer, Viewers[i].ViewTarget, InChannel, Time, bLowBandwidth)));
+	}
+}
+
+FActorPriority::FActorPriority(UNetConnection* InConnection, FActorDestructionInfo * Info, const TArray<FNetViewer>& Viewers )
+	: ActorInfo(NULL), Channel(NULL), DestructionInfo(Info)
+{
+	
+	Priority = 0;
+
+	for (int32 i = 0; i < Viewers.Num(); i++)
+	{
+		float Time  = InConnection->Driver->SpawnPrioritySeconds;
+
+		FVector Dir = DestructionInfo->DestroyedPosition - Viewers[i].ViewLocation;
+		float DistSq = Dir.SizeSquared();
+		
+		// adjust priority based on distance and whether actor is in front of viewer
+		if ( (Viewers[i].ViewDir | Dir) < 0.f )
+		{
+			if ( DistSq > NEARSIGHTTHRESHOLDSQUARED )
+				Time *= 0.2f;
+			else if ( DistSq > CLOSEPROXIMITYSQUARED )
+				Time *= 0.4f;
+		}
+		else if ( DistSq > MEDSIGHTTHRESHOLDSQUARED )
+			Time *= 0.4f;
+
+		Priority = FMath::Max<int32>(Priority, 65536.0f * Time);
+	}
+}
+
+namespace NetCmds
+{
+	static FAutoConsoleVariable MaxConnectionsToTickPerServerFrame( TEXT( "net.MaxConnectionsToTickPerServerFrame" ), 0, TEXT( "When non-zero, the maximum number of channels that will have changed replicated to them per server update" ) );
+}
+
+#if WITH_SERVER_CODE
+int32 UNetDriver::ServerReplicateActors_PrepConnections( const float DeltaSeconds )
+{
+	int32 NumClientsToTick = ClientConnections.Num();
+
+	// by default only throttle update for listen servers unless specified on the commandline
+	static bool bForceClientTickingThrottle = FParse::Param( FCommandLine::Get(), TEXT( "limitclientticks" ) );
+	if ( bForceClientTickingThrottle || GetNetMode() == NM_ListenServer )
+	{
+		// determine how many clients to tick this frame based on GEngine->NetTickRate (always tick at least one client), double for lan play
+		// FIXME: DeltaTimeOverflow is a static, and will conflict with other running net drivers, we investigate storing it on the driver itself!
+		static float DeltaTimeOverflow = 0.f;
+		// updates are doubled for lan play
+		static bool LanPlay = FParse::Param( FCommandLine::Get(), TEXT( "lanplay" ) );
+		//@todo - ideally we wouldn't want to tick more clients with a higher deltatime as that's not going to be good for performance and probably saturate bandwidth in hitchy situations, maybe 
+		// come up with a solution that is greedier with higher framerates, but still won't risk saturating server upstream bandwidth
+		float ClientUpdatesThisFrame = GEngine->NetClientTicksPerSecond * ( DeltaSeconds + DeltaTimeOverflow ) * ( LanPlay ? 2.f : 1.f );
+		NumClientsToTick = FMath::Min<int32>( NumClientsToTick, FMath::TruncToInt( ClientUpdatesThisFrame ) );
+		//UE_LOG(LogNet, Log, TEXT("%2.3f: Ticking %d clients this frame, %2.3f/%2.4f"),GetWorld()->GetTimeSeconds(),NumClientsToTick,DeltaSeconds,ClientUpdatesThisFrame);
+		if ( NumClientsToTick == 0 )
+		{
+			// if no clients are ticked this frame accumulate the time elapsed for the next frame
+			DeltaTimeOverflow += DeltaSeconds;
+			return 0;
+		}
+		DeltaTimeOverflow = 0.f;
+	}
+
+	if( NetCmds::MaxConnectionsToTickPerServerFrame->GetInt() > 0 )
+	{
+		NumClientsToTick = FMath::Min( ClientConnections.Num(), NetCmds::MaxConnectionsToTickPerServerFrame->GetInt() );
+	}
+
+	bool bFoundReadyConnection = false;
+
+	for ( int32 ConnIdx = 0; ConnIdx < ClientConnections.Num(); ConnIdx++ )
+	{
+		UNetConnection* Connection = ClientConnections[ConnIdx];
+		check( Connection );
+		check( Connection->State == USOCK_Pending || Connection->State == USOCK_Open || Connection->State == USOCK_Closed );
+		checkSlow( Connection->GetUChildConnection() == NULL );
+
+		// Handle not ready channels.
+		//@note: we cannot add Connection->IsNetReady(0) here to check for saturation, as if that's the case we still want to figure out the list of relevant actors
+		//			to reset their NetUpdateTime so that they will get sent as soon as the connection is no longer saturated
+		AActor* OwningActor = Connection->OwningActor;
+		if ( OwningActor != NULL && Connection->State == USOCK_Open && ( Connection->Driver->GetElapsedTime() - Connection->LastReceiveTime < 1.5 ) )
+		{
+			check( World == OwningActor->GetWorld() );
+
+			bFoundReadyConnection = true;
+
+			// the view target is what the player controller is looking at OR the owning actor itself when using beacons
+			AActor* DesiredViewTarget = OwningActor;
+			if (Connection->PlayerController)
+			{
+				if (AActor* ViewTarget = Connection->PlayerController->GetViewTarget())
+				{
+					if (ViewTarget->GetWorld())
+					{
+						// It is safe to use the player controller's view target.
+						DesiredViewTarget = ViewTarget;
+					}
+					else
+					{
+						// Log an error, since this means the view target for the player controller no longer has a valid world (this can happen
+						// if the player controller's view target was in a sublevel instance that has been unloaded).
+						UE_LOG(LogNet, Warning, TEXT("Player controller %s's view target (%s) no longer has a valid world! Was it unloaded as part a level instance?"),
+							*Connection->PlayerController->GetName(), *ViewTarget->GetName());
+					}
+				}
+			}
+			Connection->ViewTarget = DesiredViewTarget;
+
+			for ( int32 ChildIdx = 0; ChildIdx < Connection->Children.Num(); ChildIdx++ )
+			{
+				UNetConnection *Child = Connection->Children[ChildIdx];
+				APlayerController* ChildPlayerController = Child->PlayerController;
+				if ( ChildPlayerController != NULL )
+				{
+					Child->ViewTarget = ChildPlayerController->GetViewTarget();
+				}
+				else
+				{
+					Child->ViewTarget = NULL;
+				}
+			}
+		}
+		else
+		{
+			Connection->ViewTarget = NULL;
+			for ( int32 ChildIdx = 0; ChildIdx < Connection->Children.Num(); ChildIdx++ )
+			{
+				Connection->Children[ChildIdx]->ViewTarget = NULL;
+			}
+		}
+	}
+
+	return bFoundReadyConnection ? NumClientsToTick : 0;
+}
+
+void UNetDriver::ServerReplicateActors_BuildConsiderList( TArray<FNetworkObjectInfo*>& OutConsiderList, const float ServerTickTime )
+{
+	SCOPE_CYCLE_COUNTER( STAT_NetConsiderActorsTime );
+
+	UE_LOG( LogNetTraffic, Log, TEXT( "ServerReplicateActors_BuildConsiderList, Building ConsiderList %4.2f" ), World->GetTimeSeconds() );
+
+	int32 NumInitiallyDormant = 0;
+
+	const bool bUseAdapativeNetFrequency = IsAdaptiveNetUpdateFrequencyEnabled();
+
+	TArray<AActor*> ActorsToRemove;
+
+	for ( const TSharedPtr<FNetworkObjectInfo>& ObjectInfo : GetNetworkObjectList().GetActiveObjects() )
+	{
+		FNetworkObjectInfo* ActorInfo = ObjectInfo.Get();
+
+		if ( !ActorInfo->bPendingNetUpdate && World->TimeSeconds <= ActorInfo->NextUpdateTime )
+		{
+			continue;		// It's not time for this actor to perform an update, skip it
+		}
+
+		AActor* Actor = ActorInfo->Actor;
+
+		if ( Actor->IsPendingKillPending() )
+		{
+			// Actors aren't allowed to be placed in the NetworkObjectList if they are PendingKillPending.
+			// Actors should also be unconditionally removed from the NetworkObjectList when UWorld::DestroyActor is called.
+			// If this is happening, it means code is not destructing Actors properly, and that's not OK.
+			UE_LOG( LogNet, Warning, TEXT( "Actor %s was found in the NetworkObjectList, but is PendingKillPending" ), *Actor->GetName() );
+			ActorsToRemove.Add( Actor );
+			continue;
+		}
+
+		if ( Actor->GetRemoteRole() == ROLE_None )
+		{
+			ActorsToRemove.Add( Actor );
+			continue;
+		}
+
+		// This actor may belong to a different net driver, make sure this is the correct one
+		// (this can happen when using beacon net drivers for example)
+		if (Actor->GetNetDriverName() != NetDriverName)
+		{
+			UE_LOG(LogNetTraffic, Error, TEXT("Actor %s in wrong network actors list! (Has net driver '%s', expected '%s')"),
+					*Actor->GetName(), *Actor->GetNetDriverName().ToString(), *NetDriverName.ToString());
+
+			continue;
+		}
+
+		// Verify the actor is actually initialized (it might have been intentionally spawn deferred until a later frame)
+		if ( !Actor->IsActorInitialized() )
+		{
+			continue;
+		}
+
+		// Don't send actors that may still be streaming in or out
+		ULevel* Level = Actor->GetLevel();
+		if ( Level->HasVisibilityChangeRequestPending() || Level->bIsAssociatingLevel )
+		{
+			continue;
+		}
+
+		if ( IsDormInitialStartupActor(Actor) )
+		{
+			// This stat isn't that useful in its current form when using NetworkActors list
+			// We'll want to track initially dormant actors some other way to track them with stats
+			SCOPE_CYCLE_COUNTER( STAT_NetInitialDormantCheckTime );
+			NumInitiallyDormant++;
+			ActorsToRemove.Add( Actor );
+			//UE_LOG(LogNetTraffic, Log, TEXT("Skipping Actor %s - its initially dormant!"), *Actor->GetName() );
+			continue;
+		}
+
+		checkSlow( Actor->NeedsLoadForClient() ); // We have no business sending this unless the client can load
+		checkSlow( World == Actor->GetWorld() );
+
+		// Set defaults if this actor is replicating for first time
+		if ( ActorInfo->LastNetReplicateTime == 0 )
+		{
+			ActorInfo->LastNetReplicateTime = World->TimeSeconds;
+			ActorInfo->OptimalNetUpdateDelta = 1.0f / Actor->NetUpdateFrequency;
+		}
+
+		const float ScaleDownStartTime = 2.0f;
+		const float ScaleDownTimeRange = 5.0f;
+
+		const float LastReplicateDelta = World->TimeSeconds - ActorInfo->LastNetReplicateTime;
+
+		if ( LastReplicateDelta > ScaleDownStartTime )
+		{
+			if ( Actor->MinNetUpdateFrequency == 0.0f )
+			{
+				Actor->MinNetUpdateFrequency = 2.0f;
+			}
+
+			// Calculate min delta (max rate actor will update), and max delta (slowest rate actor will update)
+			const float MinOptimalDelta = 1.0f / Actor->NetUpdateFrequency;									  // Don't go faster than NetUpdateFrequency
+			const float MaxOptimalDelta = FMath::Max( 1.0f / Actor->MinNetUpdateFrequency, MinOptimalDelta ); // Don't go slower than MinNetUpdateFrequency (or NetUpdateFrequency if it's slower)
+
+			// Interpolate between MinOptimalDelta/MaxOptimalDelta based on how long it's been since this actor actually sent anything
+			const float Alpha = FMath::Clamp( ( LastReplicateDelta - ScaleDownStartTime ) / ScaleDownTimeRange, 0.0f, 1.0f );
+			ActorInfo->OptimalNetUpdateDelta = FMath::Lerp( MinOptimalDelta, MaxOptimalDelta, Alpha );
+		}
+
+		// Setup ActorInfo->NextUpdateTime, which will be the next time this actor will replicate properties to connections
+		// NOTE - We don't do this if bPendingNetUpdate is true, since this means we're forcing an update due to at least one connection
+		//	that wasn't to replicate previously (due to saturation, etc)
+		// NOTE - This also means all other connections will force an update (even if they just updated, we should look into this)
+		if ( !ActorInfo->bPendingNetUpdate )
+		{
+			UE_LOG( LogNetTraffic, Log, TEXT( "actor %s requesting new net update, time: %2.3f" ), *Actor->GetName(), World->TimeSeconds );
+
+			const float NextUpdateDelta = bUseAdapativeNetFrequency ? ActorInfo->OptimalNetUpdateDelta : 1.0f / Actor->NetUpdateFrequency;
+
+			// then set the next update time
+			ActorInfo->NextUpdateTime = World->TimeSeconds + UpdateDelayRandomStream.FRand() * ServerTickTime + NextUpdateDelta;
+
+			// and mark when the actor first requested an update
+			//@note: using Time because it's compared against UActorChannel.LastUpdateTime which also uses that value
+			PRAGMA_DISABLE_DEPRECATION_WARNINGS
+			ActorInfo->LastNetUpdateTime = ElapsedTime;
+			PRAGMA_ENABLE_DEPRECATION_WARNINGS
+			ActorInfo->LastNetUpdateTimestamp = ElapsedTime;
+		}
+
+		// and clear the pending update flag assuming all clients will be able to consider it
+		ActorInfo->bPendingNetUpdate = false;
+
+		// add it to the list to consider below
+		// For performance reasons, make sure we don't resize the array. It should already be appropriately sized above!
+		ensure( OutConsiderList.Num() < OutConsiderList.Max() );
+		OutConsiderList.Add( ActorInfo );
+
+		// Call PreReplication on all actors that will be considered
+		Actor->CallPreReplication( this );
+	}
+
+	for ( AActor* Actor : ActorsToRemove )
+	{
+		RemoveNetworkActor( Actor );
+	}
+
+	// Update stats
+	SET_DWORD_STAT( STAT_NumInitiallyDormantActors, NumInitiallyDormant );
+	SET_DWORD_STAT( STAT_NumConsideredActors, OutConsiderList.Num() );
+}
+
+// Returns true if this actor should replicate to *any* of the passed in connections
+static FORCEINLINE_DEBUGGABLE bool IsActorRelevantToConnection( const AActor* Actor, const TArray<FNetViewer>& ConnectionViewers )
+{
+	for ( int32 viewerIdx = 0; viewerIdx < ConnectionViewers.Num(); viewerIdx++ )
+	{
+		if ( Actor->IsNetRelevantFor( ConnectionViewers[viewerIdx].InViewer, ConnectionViewers[viewerIdx].ViewTarget, ConnectionViewers[viewerIdx].ViewLocation ) )
+		{
+			return true;
+		}
+	}
+
+	return false;
+}
+
+// Returns true if this actor is owned by, and should replicate to *any* of the passed in connections
+static FORCEINLINE_DEBUGGABLE UNetConnection* IsActorOwnedByAndRelevantToConnection( const AActor* Actor, const TArray<FNetViewer>& ConnectionViewers, bool& bOutHasNullViewTarget )
+{
+	const AActor* ActorOwner = Actor->GetNetOwner();
+
+	bOutHasNullViewTarget = false;
+
+	for ( int i = 0; i < ConnectionViewers.Num(); i++ )
+	{
+		UNetConnection* ViewerConnection = ConnectionViewers[i].Connection;
+
+		if ( ViewerConnection->ViewTarget == nullptr )
+		{
+			bOutHasNullViewTarget = true;
+		}
+
+		if ( ActorOwner == ViewerConnection->PlayerController ||
+			 ( ViewerConnection->PlayerController && ActorOwner == ViewerConnection->PlayerController->GetPawn() ) ||
+			 (ViewerConnection->ViewTarget && ViewerConnection->ViewTarget->IsRelevancyOwnerFor( Actor, ActorOwner, ViewerConnection->OwningActor ) ) )
+		{
+			return ViewerConnection;
+		}
+	}
+
+	return nullptr;
+}
+
+// Returns true if this actor is considered dormant (and all properties caught up) to the current connection
+static FORCEINLINE_DEBUGGABLE bool IsActorDormant( FNetworkObjectInfo* ActorInfo, const TWeakObjectPtr<UNetConnection>& Connection )
+{
+	// If actor is already dormant on this channel, then skip replication entirely
+	return ActorInfo->DormantConnections.Contains( Connection );
+}
+
+// Returns true if this actor wants to go dormant for a particular connection
+static FORCEINLINE_DEBUGGABLE bool ShouldActorGoDormant( AActor* Actor, const TArray<FNetViewer>& ConnectionViewers, UActorChannel* Channel, const float Time, const bool bLowNetBandwidth )
+{
+	if ( Actor->NetDormancy <= DORM_Awake || !Channel || Channel->bPendingDormancy || Channel->Dormant )
+	{
+		// Either shouldn't go dormant, or is already dormant
+		return false;
+	}
+
+	if ( Actor->NetDormancy == DORM_DormantPartial )
+	{
+		for ( int32 viewerIdx = 0; viewerIdx < ConnectionViewers.Num(); viewerIdx++ )
+		{
+			if ( !Actor->GetNetDormancy( ConnectionViewers[viewerIdx].ViewLocation, ConnectionViewers[viewerIdx].ViewDir, ConnectionViewers[viewerIdx].InViewer, ConnectionViewers[viewerIdx].ViewTarget, Channel, Time, bLowNetBandwidth ) )
+			{
+				return false;
+			}
+		}
+	}
+
+	return true;
+}
+
+int32 UNetDriver::ServerReplicateActors_PrioritizeActors( UNetConnection* Connection, const TArray<FNetViewer>& ConnectionViewers, const TArray<FNetworkObjectInfo*> ConsiderList, const bool bCPUSaturated, FActorPriority*& OutPriorityList, FActorPriority**& OutPriorityActors )
+{
+	SCOPE_CYCLE_COUNTER( STAT_NetPrioritizeActorsTime );
+
+	// Get list of visible/relevant actors.
+
+	NetTag++;
+
+	// Set up to skip all sent temporary actors
+	for ( int32 j = 0; j < Connection->SentTemporaries.Num(); j++ )
+	{
+		Connection->SentTemporaries[j]->NetTag = NetTag;
+	}
+
+	// Make list of all actors to consider.
+	check( World == Connection->OwningActor->GetWorld() );
+
+	int32 FinalSortedCount = 0;
+	int32 DeletedCount = 0;
+
+	// Make weak ptr once for IsActorDormant call
+	TWeakObjectPtr<UNetConnection> WeakConnection(Connection);
+
+	const int32 MaxSortedActors = ConsiderList.Num() + DestroyedStartupOrDormantActors.Num();
+	if ( MaxSortedActors > 0 )
+	{
+		OutPriorityList = new ( FMemStack::Get(), MaxSortedActors ) FActorPriority;
+		OutPriorityActors = new ( FMemStack::Get(), MaxSortedActors ) FActorPriority*;
+
+		check( World == Connection->ViewTarget->GetWorld() );
+
+		AGameNetworkManager* const NetworkManager = World->NetworkManager;
+		const bool bLowNetBandwidth = NetworkManager ? NetworkManager->IsInLowBandwidthMode() : false;
+
+		for ( FNetworkObjectInfo* ActorInfo : ConsiderList )
+		{
+			AActor* Actor = ActorInfo->Actor;
+
+			UActorChannel* Channel = Connection->FindActorChannelRef( ActorInfo->WeakActor );
+
+			// Skip actor if not relevant and theres no channel already.
+			// Historically Relevancy checks were deferred until after prioritization because they were expensive (line traces).
+			// Relevancy is now cheap and we are dealing with larger lists of considered actors, so we want to keep the list of
+			// prioritized actors low.
+			if (!Channel)
+			{
+				if (!IsLevelInitializedForActor(Actor, Connection))
+				{
+					// If the level this actor belongs to isn't loaded on client, don't bother sending
+					continue;
+				}
+
+				if (!IsActorRelevantToConnection(Actor, ConnectionViewers))
+				{
+					// If not relevant (and we don't have a channel), skip
+					continue;
+				}
+			}
+
+			UNetConnection* PriorityConnection = Connection;
+
+			if ( Actor->bOnlyRelevantToOwner )
+			{
+				// This actor should be owned by a particular connection, see if that connection is the one passed in
+				bool bHasNullViewTarget = false;
+
+				PriorityConnection = IsActorOwnedByAndRelevantToConnection( Actor, ConnectionViewers, bHasNullViewTarget );
+
+				if ( PriorityConnection == nullptr )
+				{
+					// Not owned by this connection, if we have a channel, close it, and continue
+					// NOTE - We won't close the channel if any connection has a NULL view target.
+					//	This is to give all connections a chance to own it
+					if ( !bHasNullViewTarget && Channel != NULL && ElapsedTime - Channel->RelevantTime >= RelevantTimeout )
+					{
+						Channel->Close(EChannelCloseReason::Relevancy);
+					}
+
+					// This connection doesn't own this actor
+					continue;
+				}
+			}
+			else if ( GSetNetDormancyEnabled != 0 )
+			{
+				// Skip Actor if dormant
+				if ( IsActorDormant( ActorInfo, WeakConnection ) )
+				{
+					continue;
+				}
+
+				// See of actor wants to try and go dormant
+				if ( ShouldActorGoDormant( Actor, ConnectionViewers, Channel, ElapsedTime, bLowNetBandwidth ) )
+				{
+					// Channel is marked to go dormant now once all properties have been replicated (but is not dormant yet)
+					Channel->StartBecomingDormant();
+				}
+			}
+
+			// Actor is relevant to this connection, add it to the list
+			// NOTE - We use NetTag to make sure SentTemporaries didn't already mark this actor to be skipped
+			if ( Actor->NetTag != NetTag )
+			{
+				UE_LOG( LogNetTraffic, Log, TEXT( "Consider %s alwaysrelevant %d frequency %f " ), *Actor->GetName(), Actor->bAlwaysRelevant, Actor->NetUpdateFrequency );
+
+				Actor->NetTag = NetTag;
+
+				OutPriorityList[FinalSortedCount] = FActorPriority( PriorityConnection, Channel, ActorInfo, ConnectionViewers, bLowNetBandwidth );
+				OutPriorityActors[FinalSortedCount] = OutPriorityList + FinalSortedCount;
+
+				FinalSortedCount++;
+
+				if ( DebugRelevantActors )
+				{
+					LastPrioritizedActors.Add( Actor );
+				}
+			}
+		}
+
+		// Add in deleted actors
+		for ( auto It = Connection->GetDestroyedStartupOrDormantActorGUIDs().CreateConstIterator(); It; ++It )
+		{
+			FActorDestructionInfo& DInfo = *DestroyedStartupOrDormantActors.FindChecked( *It );
+			OutPriorityList[FinalSortedCount] = FActorPriority( Connection, &DInfo, ConnectionViewers );
+			OutPriorityActors[FinalSortedCount] = OutPriorityList + FinalSortedCount;
+			FinalSortedCount++;
+			DeletedCount++;
+		}
+
+		// Sort by priority
+		Sort( OutPriorityActors, FinalSortedCount, FCompareFActorPriority() );
+	}
+
+	UE_LOG( LogNetTraffic, Log, TEXT( "ServerReplicateActors_PrioritizeActors: Potential %04i ConsiderList %03i FinalSortedCount %03i" ), MaxSortedActors, ConsiderList.Num(), FinalSortedCount );
+
+	// Setup stats
+	SET_DWORD_STAT( STAT_PrioritizedActors, FinalSortedCount );
+	SET_DWORD_STAT( STAT_NumRelevantDeletedActors, DeletedCount );
+
+	return FinalSortedCount;
+}
+
+int32 UNetDriver::ServerReplicateActors_ProcessPrioritizedActors( UNetConnection* Connection, const TArray<FNetViewer>& ConnectionViewers, FActorPriority** PriorityActors, const int32 FinalSortedCount, int32& OutUpdated )
+{
+	SCOPE_CYCLE_COUNTER(STAT_NetProcessPrioritizedActorsTime);
+
+	int32 ActorUpdatesThisConnection		= 0;
+	int32 ActorUpdatesThisConnectionSent	= 0;
+	int32 FinalRelevantCount				= 0;
+
+	if ( !Connection->IsNetReady( 0 ) )
+	{
+		GNumSaturatedConnections++;
+		// Connection saturated, don't process any actors
+		return 0;
+	}
+
+	for ( int32 j = 0; j < FinalSortedCount; j++ )
+	{
+		FNetworkObjectInfo*	ActorInfo = PriorityActors[j]->ActorInfo;
+
+		// Deletion entry
+		if ( ActorInfo == NULL && PriorityActors[j]->DestructionInfo )
+		{
+			// Make sure client has streaming level loaded
+			if ( PriorityActors[j]->DestructionInfo->StreamingLevelName != NAME_None && !Connection->ClientVisibleLevelNames.Contains( PriorityActors[j]->DestructionInfo->StreamingLevelName ) )
+			{
+				// This deletion entry is for an actor in a streaming level the connection doesn't have loaded, so skip it
+				continue;
+			}
+
+			FinalRelevantCount++;
+			UE_LOG( LogNetTraffic, Log, TEXT( "Server replicate actor creating destroy channel for NetGUID <%s,%s> Priority: %d" ), *PriorityActors[j]->DestructionInfo->NetGUID.ToString(), *PriorityActors[j]->DestructionInfo->PathName, PriorityActors[j]->Priority );
+
+			SendDestructionInfo(Connection, PriorityActors[j]->DestructionInfo);
+
+			Connection->RemoveDestructionInfo( PriorityActors[j]->DestructionInfo );		// Remove from connections to-be-destroyed list (close bunch of reliable, so it will make it there)
+			continue;
+		}
+
+#if !( UE_BUILD_SHIPPING || UE_BUILD_TEST )
+		static IConsoleVariable* DebugObjectCvar = IConsoleManager::Get().FindConsoleVariable( TEXT( "net.PackageMap.DebugObject" ) );
+		static IConsoleVariable* DebugAllObjectsCvar = IConsoleManager::Get().FindConsoleVariable( TEXT( "net.PackageMap.DebugAll" ) );
+		if ( ActorInfo &&
+			 ( ( DebugObjectCvar && !DebugObjectCvar->GetString().IsEmpty() && ActorInfo->Actor->GetName().Contains( DebugObjectCvar->GetString() ) ) ||
+			   ( DebugAllObjectsCvar && DebugAllObjectsCvar->GetInt() != 0 ) ) )
+		{
+			UE_LOG( LogNetPackageMap, Log, TEXT( "Evaluating actor for replication %s" ), *ActorInfo->Actor->GetName() );
+		}
+#endif
+
+		// Normal actor replication
+		UActorChannel* Channel = PriorityActors[j]->Channel;
+		UE_LOG( LogNetTraffic, Log, TEXT( " Maybe Replicate %s" ), *ActorInfo->Actor->GetName() );
+		if ( !Channel || Channel->Actor ) //make sure didn't just close this channel
+		{
+			AActor* Actor = ActorInfo->Actor;
+			bool bIsRelevant = false;
+
+			const bool bLevelInitializedForActor = IsLevelInitializedForActor( Actor, Connection );
+
+			// only check visibility on already visible actors every 1.0 + 0.5R seconds
+			// bTearOff actors should never be checked
+			if ( bLevelInitializedForActor )
+			{
+				if ( !Actor->GetTearOff() && ( !Channel || ElapsedTime - Channel->RelevantTime > 1.0 ) )
+				{
+					if ( IsActorRelevantToConnection( Actor, ConnectionViewers ) )
+					{
+						bIsRelevant = true;
+					}
+					else if ( DebugRelevantActors )
+					{
+						LastNonRelevantActors.Add( Actor );
+					}
+				}
+			}
+			else
+			{
+				// Actor is no longer relevant because the world it is/was in is not loaded by client
+				// exception: player controllers should never show up here
+				UE_LOG( LogNetTraffic, Log, TEXT( "- Level not initialized for actor %s" ), *Actor->GetName() );
+			}
+
+			// if the actor is now relevant or was recently relevant
+			const bool bIsRecentlyRelevant = bIsRelevant || ( Channel && ElapsedTime - Channel->RelevantTime < RelevantTimeout ) || (ActorInfo->ForceRelevantFrame >= Connection->LastProcessedFrame);
+
+			if ( bIsRecentlyRelevant )
+			{
+				FinalRelevantCount++;
+
+				TOptional<FScopedActorRoleSwap> SwapGuard;
+				if (ActorInfo->bSwapRolesOnReplicate)
+				{
+					SwapGuard = FScopedActorRoleSwap(Actor);
+				}
+
+				// Find or create the channel for this actor.
+				// we can't create the channel if the client is in a different world than we are
+				// or the package map doesn't support the actor's class/archetype (or the actor itself in the case of serializable actors)
+				// or it's an editor placed actor and the client hasn't initialized the level it's in
+				if ( Channel == NULL && GuidCache->SupportsObject( Actor->GetClass() ) && GuidCache->SupportsObject( Actor->IsNetStartupActor() ? Actor : Actor->GetArchetype() ) )
+				{
+					if ( bLevelInitializedForActor )
+					{
+						// Create a new channel for this actor.
+						Channel = (UActorChannel*)Connection->CreateChannelByName( NAME_Actor, EChannelCreateFlags::OpenedLocally );
+						if ( Channel )
+						{
+							Channel->SetChannelActor(Actor, ESetChannelActorFlags::None);
+						}
+					}
+					// if we couldn't replicate it for a reason that should be temporary, and this Actor is updated very infrequently, make sure we update it again soon
+					else if ( Actor->NetUpdateFrequency < 1.0f )
+					{
+						UE_LOG( LogNetTraffic, Log, TEXT( "Unable to replicate %s" ), *Actor->GetName() );
+						ActorInfo->NextUpdateTime = World->TimeSeconds + 0.2f * FMath::FRand();
+					}
+				}
+
+				if ( Channel )
+				{
+					// if it is relevant then mark the channel as relevant for a short amount of time
+					if ( bIsRelevant )
+					{
+						Channel->RelevantTime = ElapsedTime + 0.5 * UpdateDelayRandomStream.FRand();
+					}
+					// if the channel isn't saturated
+					if ( Channel->IsNetReady( 0 ) )
+					{
+						// replicate the actor
+						UE_LOG( LogNetTraffic, Log, TEXT( "- Replicate %s. %d" ), *Actor->GetName(), PriorityActors[j]->Priority );
+						if ( DebugRelevantActors )
+						{
+							LastRelevantActors.Add( Actor );
+						}
+
+						double ChannelLastNetUpdateTime = Channel->LastUpdateTime;
+
+						if ( Channel->ReplicateActor() )
+						{
+#if USE_SERVER_PERF_COUNTERS
+							if (const FNetworkObjectInfo* const ObjectInfo = Actor->FindNetworkObjectInfo())
+							{
+								// A channel time of 0.0 means this is the first time the actor is being replicated, so we don't need to record it
+								if (ChannelLastNetUpdateTime > 0.0)
+								{
+									Connection->GetActorsStarvedByClassTimeMap().FindOrAdd(Actor->GetClass()->GetName()).Add((World->RealTimeSeconds - ChannelLastNetUpdateTime) * 1000.0f);
+								}
+							}
+#endif
+
+							ActorUpdatesThisConnectionSent++;
+							if ( DebugRelevantActors )
+							{
+								LastSentActors.Add( Actor );
+							}
+
+							// Calculate min delta (max rate actor will upate), and max delta (slowest rate actor will update)
+							const float MinOptimalDelta				= 1.0f / Actor->NetUpdateFrequency;
+							const float MaxOptimalDelta				= FMath::Max( 1.0f / Actor->MinNetUpdateFrequency, MinOptimalDelta );
+							const float DeltaBetweenReplications	= ( World->TimeSeconds - ActorInfo->LastNetReplicateTime );
+
+							// Choose an optimal time, we choose 70% of the actual rate to allow frequency to go up if needed
+							ActorInfo->OptimalNetUpdateDelta = FMath::Clamp( DeltaBetweenReplications * 0.7f, MinOptimalDelta, MaxOptimalDelta );
+							ActorInfo->LastNetReplicateTime = World->TimeSeconds;
+						}
+						ActorUpdatesThisConnection++;
+						OutUpdated++;
+					}
+					else
+					{
+						UE_LOG( LogNetTraffic, Log, TEXT( "- Channel saturated, forcing pending update for %s" ), *Actor->GetName() );
+						// otherwise force this actor to be considered in the next tick again
+						Actor->ForceNetUpdate();
+					}
+					// second check for channel saturation
+					if ( !Connection->IsNetReady( 0 ) )
+					{
+						// We can bail out now since this connection is saturated, we'll return how far we got though
+						GNumSaturatedConnections++;
+						return j;
+					}
+				}
+			}
+
+			// If the actor wasn't recently relevant, or if it was torn off, close the actor channel if it exists for this connection
+			if ( ( !bIsRecentlyRelevant || Actor->GetTearOff() ) && Channel != NULL )
+			{
+				// Non startup (map) actors have their channels closed immediately, which destroys them.
+				// Startup actors get to keep their channels open.
+
+				// Fixme: this should be a setting
+				if ( !bLevelInitializedForActor || !Actor->IsNetStartupActor() )
+				{
+					UE_LOG( LogNetTraffic, Log, TEXT( "- Closing channel for no longer relevant actor %s" ), *Actor->GetName() );
+					Channel->Close(Actor->GetTearOff() ? EChannelCloseReason::TearOff : EChannelCloseReason::Relevancy);
+				}
+			}
+		}
+	}
+
+	return FinalSortedCount;
+}
+#endif
+
+int64 UNetDriver::SendDestructionInfo(UNetConnection* Connection, FActorDestructionInfo* DestructionInfo)
+{
+	QUICK_SCOPE_CYCLE_COUNTER(STAT_NetSendDestructionInfo);
+	CSV_SCOPED_TIMING_STAT_EXCLUSIVE(SendDestructionInfo);
+
+	int64 NumBits = 0;
+
+	checkf(Connection, TEXT("SendDestructionInfo called with invalid connection: %s"), *GetDescription());
+	checkf(DestructionInfo, TEXT("SendDestructionInfo called with invalid destruction info: %s"), *GetDescription());
+
+	if (GNetControlChannelDestructionInfo != 0)
+	{
+		if (UControlChannel* ControlChan = Cast<UControlChannel>(Connection->Channels[0]))
+		{
+			NumBits = ControlChan->SendDestructionInfo(DestructionInfo);
+		}
+	}
+	else
+	{
+		UActorChannel* Channel = (UActorChannel*)Connection->CreateChannelByName(NAME_Actor, EChannelCreateFlags::OpenedLocally);
+		if (Channel)
+		{
+			// Send a close bunch on the new channel
+			PRAGMA_DISABLE_DEPRECATION_WARNINGS
+				NumBits = Channel->SetChannelActorForDestroy(DestructionInfo);
+			PRAGMA_ENABLE_DEPRECATION_WARNINGS
+		}
+	}
+
+	return NumBits;
+}
+
+// -------------------------------------------------------------------------------------------------------------------------
+//	Replication profiling (server cpu) helpers.
+//		-CSV stat cateogry "Replication". Not enabled by default. (use -csvCategories=Replication to enable from cmd line)
+//		-Auto CSV capturing can be enabled with -ReplicationCSVCaptureFrames=X. This will start a CSV capture once there
+//			is one client connected and will auto terminate the process once X frames have elapsed.
+// -------------------------------------------------------------------------------------------------------------------------
+
+
+// Replication CSV category is enabled by default in server builds
+CSV_DEFINE_CATEGORY(Replication, WITH_SERVER_CODE);
+
+double GReplicationGatherPrioritizeTimeSeconds;
+double GServerReplicateActorTimeSeconds;
+
+int32 GNumClientConnections;
+int32 GNumClientUpdateLevelVisibility;
+
+DECLARE_DWORD_COUNTER_STAT(TEXT("Num Saturated Connections"), STAT_NumSaturatedConnections, STATGROUP_Net);
+
+DECLARE_DWORD_COUNTER_STAT(TEXT("SharedSerialization RPC Hit"), STAT_SharedSerializationRPCHit, STATGROUP_Net);
+DECLARE_DWORD_COUNTER_STAT(TEXT("SharedSerialization RPC Miss"), STAT_SharedSerializationRPCMiss, STATGROUP_Net);
+
+DECLARE_DWORD_COUNTER_STAT(TEXT("SharedSerialization Property Hit"), STAT_SharedSerializationPropertyHit, STATGROUP_Net);
+DECLARE_DWORD_COUNTER_STAT(TEXT("SharedSerialization Property Miss"), STAT_SharedSerializationPropertyMiss, STATGROUP_Net);
+
+struct FReplicationAutoCapture
+{
+	int32 CaptureFrames=-1;
+	int32 KillProcessFrames = -1;
+	bool StartedCapture = false;
+
+	void DoFrame()
+	{
+#if CSV_PROFILER
+		if (CaptureFrames == -1)
+		{
+			// First time see if we want to auto capture
+			CaptureFrames = 0;
+			FParse::Value(FCommandLine::Get(), TEXT("ReplicationCSVCaptureFrames="), CaptureFrames);
+		}
+		
+		if (CaptureFrames > 0)
+		{			
+			if (!StartedCapture)
+			{
+				StartedCapture = true;
+				FCsvProfiler::Get()->EnableCategoryByString(TEXT("Replication"));
+				FCsvProfiler::Get()->BeginCapture();
+			}
+				
+			if (--CaptureFrames <= 0)
+			{
+				FCsvProfiler::Get()->EndCapture();
+				KillProcessFrames = 60;
+			}
+
+		}
+		else if (KillProcessFrames > 0)
+		{
+			// Kill process when finished
+			--KillProcessFrames;
+			if (KillProcessFrames == 0)
+			{
+				GLog->PanicFlushThreadedLogs();
+				FPlatformMisc::RequestExit(1);
+			}
+		}
+#endif
+	}
+};
+FReplicationAutoCapture GReplicationAutoCapture;
+
+#if CSV_PROFILER
+struct FScopedNetDriverStats
+{
+	FScopedNetDriverStats(uint32& OutBytesRef, UNetDriver* InNetDriver) : OutBytes(OutBytesRef), NetDriver(InNetDriver)
+	{
+		GReplicationAutoCapture.DoFrame();
+
+		// Set GReplicateActorTimingEnabled for this frame. (This will determine if actor channels do cycle counting while replicating)
+		GReplicateActorTimingEnabled = ShouldEnableScopeSecondsTimers();
+		if (GReplicateActorTimingEnabled)
+	{
+		GReplicateActorTimeSeconds = 0;
+		GNumReplicateActorCalls = 0;
+		GNumSaturatedConnections = 0;
+
+		GReplicationGatherPrioritizeTimeSeconds = 0.f;
+		GServerReplicateActorTimeSeconds = 0.f;
+
+		// Whatever these values currently are were (mostly) set by RPCs (technically something else could have force ReplicateActor to be called but this is rare).
+		SET_DWORD_STAT(STAT_SharedSerializationRPCHit, GNumSharedSerializationHit);
+		SET_DWORD_STAT(STAT_SharedSerializationRPCMiss, GNumSharedSerializationMiss);
+
+		StartTime = FPlatformTime::Seconds();
+		StartOutBytes = GNetOutBytes;
+		}
+	}
+
+	~FScopedNetDriverStats()
+	{
+		if (GReplicateActorTimingEnabled)
+		{
+		const double TotalTime = FPlatformTime::Seconds() - StartTime;
+
+		GServerReplicateActorTimeSeconds = TotalTime;
+		GReplicationGatherPrioritizeTimeSeconds = TotalTime - GReplicateActorTimeSeconds;
+
+		uint32 FrameOutBytes = GNetOutBytes - StartOutBytes;
+
+		SET_FLOAT_STAT(STAT_NetServerGatherPrioritizeRepActorsTime, GReplicationGatherPrioritizeTimeSeconds * 1000.0);
+		SET_DWORD_STAT(STAT_NumReplicatedActors, GNumReplicateActorCalls);
+		SET_DWORD_STAT(STAT_NumSaturatedConnections, GNumSaturatedConnections);
+
+		CSV_CUSTOM_STAT(Replication, ServerReplicateActorTimeMS, (float)(GServerReplicateActorTimeSeconds * 1000.0), ECsvCustomStatOp::Set );
+		CSV_CUSTOM_STAT(Replication, GatherPrioritizeTimeMS, (float)(GReplicationGatherPrioritizeTimeSeconds * 1000.0), ECsvCustomStatOp::Set );
+		CSV_CUSTOM_STAT(Replication, ReplicateActorTimeMS, (float)(GReplicateActorTimeSeconds * 1000.0), ECsvCustomStatOp::Set );
+		CSV_CUSTOM_STAT(Replication, NumReplicateActorCallsPerConAvg, ((float)GNumReplicateActorCalls)/(float)GNumClientConnections, ECsvCustomStatOp::Set );
+		CSV_CUSTOM_STAT(Replication, Connections, (float)GNumClientConnections, ECsvCustomStatOp::Set );
+		CSV_CUSTOM_STAT(Replication, SatConnections, (float)GNumSaturatedConnections, ECsvCustomStatOp::Set );
+			CSV_CUSTOM_STAT(Replication, OutKBytes, ((float)FrameOutBytes) / 1024.f, ECsvCustomStatOp::Set );
+			CSV_CUSTOM_STAT(Replication, OutNetGUIDKBytesSec, ((float)NetDriver->NetGUIDOutBytes / 1024.f), ECsvCustomStatOp::Set );
+			CSV_CUSTOM_STAT(Replication, NumClientUpdateLevelVisibility, ((float)GNumClientUpdateLevelVisibility), ECsvCustomStatOp::Set );
+				
+
+		SET_DWORD_STAT(STAT_SharedSerializationPropertyHit, GNumSharedSerializationHit);
+		SET_DWORD_STAT(STAT_SharedSerializationPropertyMiss, GNumSharedSerializationMiss);
+
+		// Note: we want to reset this at the end of the frame since the RPC stats are incremented at the top (recv)
+		GNumSharedSerializationHit = 0;
+		GNumSharedSerializationMiss = 0;
+			GNumClientUpdateLevelVisibility = 0;
+		}
+	}
+	
+	uint32& OutBytes;
+	double StartTime;
+	uint32 StartOutBytes;
+	UNetDriver* NetDriver;
+};
+#endif
+
+// -------------------------------------------------------------------------------------------------------------------------
+//	ServerReplicateActors: this is main function to replicate actors to client connections. It can be "outsourced" to a Replication Driver.
+// -------------------------------------------------------------------------------------------------------------------------
+
+int32 UNetDriver::ServerReplicateActors(float DeltaSeconds)
+{
+	SCOPE_CYCLE_COUNTER(STAT_NetServerRepActorsTime);
+	CSV_SCOPED_TIMING_STAT_EXCLUSIVE(ServerReplicateActors);
+
+#if WITH_SERVER_CODE
+	if ( ClientConnections.Num() == 0 )
+	{
+		return 0;
+	}
+
+	SET_DWORD_STAT(STAT_NumReplicatedActors, 0);
+	SET_DWORD_STAT(STAT_NumReplicatedActorBytes, 0);
+
+#if CSV_PROFILER
+	FScopedNetDriverStats NetDriverStats(OutBytes, this);
+	GNumClientConnections = ClientConnections.Num();
+#endif
+	
+	if (ReplicationDriver)
+	{
+		return ReplicationDriver->ServerReplicateActors(DeltaSeconds);
+	}
+
+	check( World );
+
+	// Bump the ReplicationFrame value to invalidate any properties marked as "unchanged" for this frame.
+	ReplicationFrame++;
+
+	int32 Updated = 0;
+
+	const int32 NumClientsToTick = ServerReplicateActors_PrepConnections( DeltaSeconds );
+
+	if ( NumClientsToTick == 0 )
+	{
+		// No connections are ready this frame
+		return 0;
+	}
+
+	AWorldSettings* WorldSettings = World->GetWorldSettings();
+
+	bool bCPUSaturated		= false;
+	float ServerTickTime	= GEngine->GetMaxTickRate( DeltaSeconds );
+	if ( ServerTickTime == 0.f )
+	{
+		ServerTickTime = DeltaSeconds;
+	}
+	else
+	{
+		ServerTickTime	= 1.f/ServerTickTime;
+		bCPUSaturated	= DeltaSeconds > 1.2f * ServerTickTime;
+	}
+
+	TArray<FNetworkObjectInfo*> ConsiderList;
+	ConsiderList.Reserve( GetNetworkObjectList().GetActiveObjects().Num() );
+
+	// Build the consider list (actors that are ready to replicate)
+	ServerReplicateActors_BuildConsiderList( ConsiderList, ServerTickTime );
+
+	TSet<UNetConnection*> ConnectionsToClose;
+
+	FMemMark Mark( FMemStack::Get() );
+
+	for ( int32 i=0; i < ClientConnections.Num(); i++ )
+	{
+		UNetConnection* Connection = ClientConnections[i];
+		check(Connection);
+
+		// net.DormancyValidate can be set to 2 to validate all dormant actors against last known state before going dormant
+		if ( GNetDormancyValidate == 2 )
+		{
+			for ( auto It = Connection->DormantReplicatorMap.CreateIterator(); It; ++It )
+			{
+				FObjectReplicator& Replicator = It.Value().Get();
+
+				if ( Replicator.OwningChannel != nullptr )
+				{
+					Replicator.ValidateAgainstState( Replicator.OwningChannel->GetActor() );
+				}
+			}
+		}
+
+		// if this client shouldn't be ticked this frame
+		if (i >= NumClientsToTick)
+		{
+			//UE_LOG(LogNet, Log, TEXT("skipping update to %s"),*Connection->GetName());
+			// then mark each considered actor as bPendingNetUpdate so that they will be considered again the next frame when the connection is actually ticked
+			for (int32 ConsiderIdx = 0; ConsiderIdx < ConsiderList.Num(); ConsiderIdx++)
+			{
+				AActor *Actor = ConsiderList[ConsiderIdx]->Actor;
+				// if the actor hasn't already been flagged by another connection,
+				if (Actor != NULL && !ConsiderList[ConsiderIdx]->bPendingNetUpdate)
+				{
+					// find the channel
+					UActorChannel *Channel = Connection->FindActorChannelRef(ConsiderList[ConsiderIdx]->WeakActor);
+					// and if the channel last update time doesn't match the last net update time for the actor
+					if (Channel != NULL && Channel->LastUpdateTime < ConsiderList[ConsiderIdx]->LastNetUpdateTimestamp)
+					{
+						//UE_LOG(LogNet, Log, TEXT("flagging %s for a future update"),*Actor->GetName());
+						// flag it for a pending update
+						ConsiderList[ConsiderIdx]->bPendingNetUpdate = true;
+					}
+				}
+			}
+			// clear the time sensitive flag to avoid sending an extra packet to this connection
+			Connection->TimeSensitive = false;
+		}
+		else if (Connection->ViewTarget)
+		{		
+
+			const int32 LocalNumSaturated = GNumSaturatedConnections;
+
+			// Make a list of viewers this connection should consider (this connection and children of this connection)
+			TArray<FNetViewer>& ConnectionViewers = WorldSettings->ReplicationViewers;
+
+			ConnectionViewers.Reset();
+			new( ConnectionViewers )FNetViewer( Connection, DeltaSeconds );
+			for ( int32 ViewerIndex = 0; ViewerIndex < Connection->Children.Num(); ViewerIndex++ )
+			{
+				if ( Connection->Children[ViewerIndex]->ViewTarget != NULL )
+				{
+					new( ConnectionViewers )FNetViewer( Connection->Children[ViewerIndex], DeltaSeconds );
+				}
+			}
+
+			// send ClientAdjustment if necessary
+			// we do this here so that we send a maximum of one per packet to that client; there is no value in stacking additional corrections
+			if ( Connection->PlayerController )
+			{
+				Connection->PlayerController->SendClientAdjustment();
+			}
+
+			for ( int32 ChildIdx = 0; ChildIdx < Connection->Children.Num(); ChildIdx++ )
+			{
+				if ( Connection->Children[ChildIdx]->PlayerController != NULL )
+				{
+					Connection->Children[ChildIdx]->PlayerController->SendClientAdjustment();
+				}
+			}
+
+			FMemMark RelevantActorMark(FMemStack::Get());
+
+			FActorPriority* PriorityList	= NULL;
+			FActorPriority** PriorityActors = NULL;
+
+			// Get a sorted list of actors for this connection
+			const int32 FinalSortedCount = ServerReplicateActors_PrioritizeActors( Connection, ConnectionViewers, ConsiderList, bCPUSaturated, PriorityList, PriorityActors );
+
+			// Process the sorted list of actors for this connection
+			const int32 LastProcessedActor = ServerReplicateActors_ProcessPrioritizedActors( Connection, ConnectionViewers, PriorityActors, FinalSortedCount, Updated );
+
+			// relevant actors that could not be processed this frame are marked to be considered for next frame
+			for ( int32 k=LastProcessedActor; k<FinalSortedCount; k++ )
+			{
+				if (!PriorityActors[k]->ActorInfo)
+				{
+					// A deletion entry, skip it because we dont have anywhere to store a 'better give higher priority next time'
+					continue;
+				}
+
+				AActor* Actor = PriorityActors[k]->ActorInfo->Actor;
+
+				UActorChannel* Channel = PriorityActors[k]->Channel;
+				
+				UE_LOG(LogNetTraffic, Verbose, TEXT("Saturated. %s"), *Actor->GetName());
+				if (Channel != NULL && ElapsedTime - Channel->RelevantTime <= 1.0)
+				{
+					UE_LOG(LogNetTraffic, Log, TEXT(" Saturated. Mark %s NetUpdateTime to be checked for next tick"), *Actor->GetName());
+					PriorityActors[k]->ActorInfo->bPendingNetUpdate = true;
+				}
+				else if ( IsActorRelevantToConnection( Actor, ConnectionViewers ) )
+				{
+					// If this actor was relevant but didn't get processed, force another update for next frame
+					UE_LOG( LogNetTraffic, Log, TEXT( " Saturated. Mark %s NetUpdateTime to be checked for next tick" ), *Actor->GetName() );
+					PriorityActors[k]->ActorInfo->bPendingNetUpdate = true;
+					if ( Channel != NULL )
+					{
+						Channel->RelevantTime = ElapsedTime + 0.5 * UpdateDelayRandomStream.FRand();
+					}
+					}
+
+				// If the actor was forced to relevant and didn't get processed, try again on the next update;
+				if (PriorityActors[k]->ActorInfo->ForceRelevantFrame >= Connection->LastProcessedFrame)
+				{
+					PriorityActors[k]->ActorInfo->ForceRelevantFrame = ReplicationFrame+1;
+				}
+			}
+			RelevantActorMark.Pop();
+
+			ConnectionViewers.Reset();
+
+			Connection->LastProcessedFrame = ReplicationFrame;
+
+			const bool bWasSaturated = GNumSaturatedConnections > LocalNumSaturated;
+			Connection->TrackReplicationForAnalytics(bWasSaturated);
+		}
+
+		if (Connection->GetPendingCloseDueToReplicationFailure())
+		{
+			ConnectionsToClose.Add(Connection);
+		}
+	}
+
+	// shuffle the list of connections if not all connections were ticked
+	if (NumClientsToTick < ClientConnections.Num())
+	{
+		int32 NumConnectionsToMove = NumClientsToTick;
+		while (NumConnectionsToMove > 0)
+		{
+			// move all the ticked connections to the end of the list so that the other connections are considered first for the next frame
+			UNetConnection *Connection = ClientConnections[0];
+			ClientConnections.RemoveAt(0,1);
+			ClientConnections.Add(Connection);
+			NumConnectionsToMove--;
+		}
+	}
+	Mark.Pop();
+
+	if (DebugRelevantActors)
+	{
+		PrintDebugRelevantActors();
+		LastPrioritizedActors.Empty();
+		LastSentActors.Empty();
+		LastRelevantActors.Empty();
+		LastNonRelevantActors.Empty();
+
+		DebugRelevantActors  = false;
+	}
+
+	for (UNetConnection* ConnectionToClose : ConnectionsToClose)
+	{
+		ConnectionToClose->Close();
+	}
+
+	return Updated;
+#else
+	return 0;
+#endif // WITH_SERVER_CODE
+}
+
+UChannel* UNetDriver::GetOrCreateChannelByName(const FName& ChName)
+{
+	UChannel* RetVal = nullptr;
+	if (ChName == NAME_Actor && CVarActorChannelPool.GetValueOnAnyThread() != 0)
+	{
+		while (ActorChannelPool.Num() > 0 && RetVal == nullptr)
+		{
+			RetVal = ActorChannelPool.Pop();
+			if (RetVal && RetVal->GetClass() != ChannelDefinitionMap[ChName].ChannelClass)
+			{
+				// Channel type Changed since this channel was added to the pool. Throw it away.
+				RetVal->MarkPendingKill();
+				RetVal = nullptr;
+			}
+		}
+		if (RetVal)
+		{
+			check(RetVal->GetClass() == ChannelDefinitionMap[ChName].ChannelClass);
+			check(!RetVal->IsPendingKill());
+			RetVal->bPooled = false;
+		}
+	}
+
+	if (!RetVal)
+	{
+		RetVal = InternalCreateChannelByName(ChName);
+	}
+
+	return RetVal;
+}
+
+UChannel* UNetDriver::InternalCreateChannelByName(const FName& ChName)
+{
+	return NewObject<UChannel>(GetTransientPackage(), ChannelDefinitionMap[ChName].ChannelClass);
+}
+
+void UNetDriver::ReleaseToChannelPool(UChannel* Channel)
+{
+	check(Channel);
+	check(!Channel->IsPendingKill());
+	if (Channel->ChName == NAME_Actor && CVarActorChannelPool.GetValueOnAnyThread() != 0)
+	{
+		ActorChannelPool.Push(Channel);
+
+		Channel->AddedToChannelPool();
+	}
+}
+
+void UNetDriver::SetNetDriverName(FName NewNetDriverNamed)
+{
+	NetDriverName = NewNetDriverNamed;
+	InitPacketSimulationSettings();
+}
+
+void UNetDriver::PrintDebugRelevantActors()
+{
+	struct SLocal
+	{
+		static void AggregateAndPrint( TArray< TWeakObjectPtr<AActor> >	&List, FString txt )
+		{
+			TMap< TWeakObjectPtr<UClass>, int32>	ClassSummary;
+			TMap< TWeakObjectPtr<UClass>, int32>	SuperClassSummary;
+
+			for (auto It = List.CreateIterator(); It; ++It)
+			{
+				if (AActor* Actor = It->Get())
+				{
+
+					ClassSummary.FindOrAdd(Actor->GetClass())++;
+					if (Actor->GetClass()->GetSuperStruct())
+					{
+						SuperClassSummary.FindOrAdd( Actor->GetClass()->GetSuperClass() )++;
+					}
+				}
+			}
+
+			struct FCompareActorClassCount
+			{
+				FORCEINLINE bool operator()( int32 A, int32 B ) const
+				{
+					return A < B;
+				}
+			};
+
+
+			ClassSummary.ValueSort(FCompareActorClassCount());
+			SuperClassSummary.ValueSort(FCompareActorClassCount());
+
+			UE_LOG(LogNet, Warning, TEXT("------------------------------") );
+			UE_LOG(LogNet, Warning, TEXT(" %s Class Summary"), *txt );
+			UE_LOG(LogNet, Warning, TEXT("------------------------------") );
+
+			for (auto It = ClassSummary.CreateIterator(); It; ++It)
+			{
+				UE_LOG(LogNet, Warning, TEXT("%4d - %s (%s)"), It.Value(), *It.Key()->GetName(), It.Key()->GetSuperStruct() ? *It.Key()->GetSuperStruct()->GetName() : TEXT("NULL") );
+			}
+
+			UE_LOG(LogNet, Warning, TEXT("---------------------------------") );
+			UE_LOG(LogNet, Warning, TEXT(" %s Parent Class Summary "), *txt );
+			UE_LOG(LogNet, Warning, TEXT("------------------------------") );
+
+			for (auto It = SuperClassSummary.CreateIterator(); It; ++It)
+			{
+				UE_LOG(LogNet, Warning, TEXT("%4d - %s (%s)"), It.Value(), *It.Key()->GetName(), It.Key()->GetSuperStruct() ? *It.Key()->GetSuperStruct()->GetName() : TEXT("NULL") );
+			}
+
+			UE_LOG(LogNet, Warning, TEXT("---------------------------------") );
+			UE_LOG(LogNet, Warning, TEXT(" %s Total: %d"), *txt, List.Num() );
+			UE_LOG(LogNet, Warning, TEXT("---------------------------------") );
+		}
+	};
+
+	SLocal::AggregateAndPrint( LastPrioritizedActors, TEXT(" Prioritized Actor") );
+	SLocal::AggregateAndPrint( LastRelevantActors, TEXT(" Relevant Actor") );
+	SLocal::AggregateAndPrint( LastNonRelevantActors, TEXT(" NonRelevant Actor") );
+	SLocal::AggregateAndPrint( LastSentActors, TEXT(" Sent Actor") );
+
+	UE_LOG(LogNet, Warning, TEXT("---------------------------------") );
+	UE_LOG(LogNet, Warning, TEXT(" Num Connections: %d"), ClientConnections.Num() );
+
+	UE_LOG(LogNet, Warning, TEXT("---------------------------------") );
+
+}
+
+void UNetDriver::DrawNetDriverDebug()
+{
+#if ENABLE_DRAW_DEBUG
+	UNetConnection *Connection = (ServerConnection ? ServerConnection : (ClientConnections.Num() >= 1 ? ClientConnections[0] : NULL));
+	if (!Connection)
+	{
+		return;
+	}
+
+	UWorld* LocalWorld = GetWorld();
+	if (!LocalWorld)
+	{
+		return;
+	}
+
+	ULocalPlayer*	LocalPlayer = NULL;
+	for(FLocalPlayerIterator It(GEngine, LocalWorld);It;++It)
+	{
+		LocalPlayer = *It;
+		break;
+	}
+	if (!LocalPlayer)
+	{
+		return;
+	}
+
+	const float MaxDebugDrawDistanceSquare = FMath::Square(CVarNetDebugDrawCullDistance.GetValueOnAnyThread());
+
+	const FVector Extent(20.f);
+	// Used to draw additional boxes to show more state
+	const FBox BoxExpansion = FBox(&Extent, 1);
+
+	for (FActorIterator It(LocalWorld); It; ++It)
+	{
+		// Skip drawing non-replicated actors and controllers since their bounds are
+		// bogus and represented by the Character/Pawn
+		if (!It->GetIsReplicated() || It->IsA(AController::StaticClass()) || It->IsA(AInfo::StaticClass()))
+		{
+			continue;
+		}
+
+		FBox Box = It->GetComponentsBoundingBox();
+		// Skip actors without bounds in the world
+		if (Box.GetExtent().IsNearlyZero())
+		{
+			continue;
+		}
+
+		FColor ExtraStateDrawColor;
+
+		bool bWasCulled = false;
+		const bool bIsAlwaysRelevant = It->bAlwaysRelevant;
+		// Determine if we need to show that it was network culled or not
+		if (!bIsAlwaysRelevant && !It->bOnlyRelevantToOwner)
+		{
+			const float DistanceSquared = (It->GetActorLocation() - LocalPlayer->LastViewLocation).SizeSquared();
+			// Check for culling from this view based upon the CVar
+			if (MaxDebugDrawDistanceSquare > 0.f && MaxDebugDrawDistanceSquare < DistanceSquared)
+			{
+				continue;
+			}
+
+			if (DistanceSquared > It->NetCullDistanceSquared)
+			{
+				bWasCulled = true;
+				ExtraStateDrawColor = FColor::White;
+			}
+		}
+		else if (bIsAlwaysRelevant)
+		{
+			ExtraStateDrawColor = FColor::Red;
+		}
+		
+		const FNetworkObjectInfo* NetworkObjectInfo = Connection->Driver->FindNetworkObjectInfo( *It );
+
+		FColor DrawColor;
+		if ( NetworkObjectInfo && NetworkObjectInfo->DormantConnections.Contains( Connection ) )
+		{
+			DrawColor = FColor::Blue;
+		}
+		else if (Connection->FindActorChannelRef(*It) != NULL)
+		{
+			DrawColor = FColor::Green;
+		}
+		else
+		{
+			// A replicated actor that is not on a channel and not dormant
+			DrawColor = FColor::Orange;
+		}
+
+		DrawDebugBox( LocalWorld, Box.GetCenter(), Box.GetExtent(), FQuat::Identity, DrawColor, false );
+		const bool bDrawSecondBox = bWasCulled || bIsAlwaysRelevant;
+		// Draw a second box around the object if it was net culled
+		if (bDrawSecondBox)
+		{
+			Box += BoxExpansion;
+			DrawDebugBox(LocalWorld, Box.GetCenter(), Box.GetExtent(), FQuat::Identity, ExtraStateDrawColor, false);
+		}
+	}
+#endif
+}
+
+bool UNetDriver::NetObjectIsDynamic(const UObject *Object) const
+{
+	const UActorComponent *ActorComponent = Cast<const UActorComponent>(Object);
+	if (ActorComponent)
+	{
+		// Actor components are dynamic if their owning actor is.
+		return NetObjectIsDynamic(Object->GetOuter());
+	}
+
+	const AActor *Actor = Cast<const AActor>(Object);
+	if (!Actor || Actor->HasAnyFlags(RF_ClassDefaultObject | RF_ArchetypeObject) || Actor->IsNetStartupActor())
+	{
+		return false;
+	}
+
+	return true;
+}
+
+void UNetDriver::AddClientConnection(UNetConnection * NewConnection)
+{
+	SCOPE_CYCLE_COUNTER(Stat_NetDriverAddClientConnection);
+
+	UE_CLOG(!DDoS.CheckLogRestrictions(), LogNet, Log, TEXT("AddClientConnection: Added client connection: %s"),
+			*NewConnection->Describe());
+
+	ClientConnections.Add(NewConnection);
+
+	TSharedPtr<const FInternetAddr> ConnAddr = NewConnection->GetRemoteAddr();
+
+	if (ConnAddr.IsValid())
+	{
+		MappedClientConnections.Add(ConnAddr.ToSharedRef(), NewConnection);
+
+		// On the off-chance of the same IP:Port being reused, check RecentlyDisconnectedClients
+		int32 RecentDisconnectIdx = RecentlyDisconnectedClients.IndexOfByPredicate(
+			[&ConnAddr](const FDisconnectedClient& CurElement)
+			{
+				return *ConnAddr == *CurElement.Address;
+			});
+
+		if (RecentDisconnectIdx != INDEX_NONE)
+		{
+			RecentlyDisconnectedClients.RemoveAt(RecentDisconnectIdx);
+		}
+	}
+
+	if (ReplicationDriver)
+	{
+		ReplicationDriver->AddClientConnection(NewConnection);
+	}
+
+#if USE_SERVER_PERF_COUNTERS
+	PerfCountersIncrement(TEXT("AddedConnections"));
+#endif
+
+	CreateInitialServerChannels(NewConnection);
+
+	// When new connections join, we need to make sure to add all fully dormant actors back to the network list, so they can get processed for the new connection
+	// They'll eventually fall back off to this list when they are dormant on the new connection
+	GetNetworkObjectList().HandleConnectionAdded();
+
+	for (auto It = DestroyedStartupOrDormantActors.CreateIterator(); It; ++It)
+	{
+		if (It.Key().IsStatic())
+		{
+			UE_LOG(LogNet, VeryVerbose, TEXT("Adding actor NetGUID <%s> to new connection's destroy list"), *It.Key().ToString());
+			NewConnection->AddDestructionInfo(It.Value().Get());
+		}
+	}
+}
+
+void UNetDriver::CreateReplicatedStaticActorDestructionInfo(UNetDriver* NetDriver, ULevel* Level, const FReplicatedStaticActorDestructionInfo& Info)
+{
+	check(NetDriver && Level);
+
+	FNetworkGUID NetGUID = NetDriver->GuidCache->AssignNewNetGUIDFromPath_Server( Info.PathName.ToString(), Info.ObjOuter.Get(), Info.ObjClass );
+	if (NetGUID.IsDefault())
+	{
+		UE_LOG(LogNet, Error, TEXT("CreateReplicatedStaticActorDestructionInfo got an invalid NetGUID for %s"), *Info.PathName.ToString());
+		return;
+	}
+
+	UE_LOG(LogNet, VeryVerbose, TEXT("CreateReplicatedStaticActorDestructionInfo %s %s %s %s"), *NetDriver->GetName(), *Level->GetName(), *Info.PathName.ToString(), *NetGUID.ToString());
+	
+	TUniquePtr<FActorDestructionInfo>& NewInfoPtr = NetDriver->DestroyedStartupOrDormantActors.FindOrAdd( NetGUID );
+	if (NewInfoPtr.IsValid() == false)
+	{
+		NewInfoPtr = MakeUnique<FActorDestructionInfo>();
+	}
+	
+	FActorDestructionInfo &NewInfo = *NewInfoPtr;
+	NewInfo.DestroyedPosition = Info.DestroyedPosition;
+	NewInfo.NetGUID = NetGUID;
+	NewInfo.Level = Level;
+	NewInfo.ObjOuter = Info.ObjOuter.Get();
+	NewInfo.PathName = Info.PathName.ToString();
+
+	// Look for renamed actor now so we can clear it after the destroy is queued
+	FName RenamedPath = NetDriver->RenamedStartupActors.FindRef(Info.PathName);
+	if (RenamedPath != NAME_None)
+	{
+		NewInfo.PathName = RenamedPath.ToString();
+	}
+
+	if (NewInfo.Level.IsValid() && !NewInfo.Level->IsPersistentLevel() )
+	{
+		NewInfo.StreamingLevelName = NewInfo.Level->GetOutermost()->GetFName();
+	}
+	else
+	{
+		NewInfo.StreamingLevelName = NAME_None;
+	}
+
+	NewInfo.Reason = EChannelCloseReason::Destroyed;
+}
+
+void UNetDriver::InitDestroyedStartupActors()
+{
+	if (World)
+	{
+		// add startup actors destroyed before the creation of this net driver
+		for (auto LevelIt(World->GetLevelIterator()); LevelIt; ++LevelIt)
+		{
+			ULevel* Level = *LevelIt;
+			if (Level)
+			{
+				const TArray<FReplicatedStaticActorDestructionInfo>& DestroyedReplicatedStaticActors = Level->GetDestroyedReplicatedStaticActors();
+				for(const FReplicatedStaticActorDestructionInfo& Info : DestroyedReplicatedStaticActors)
+				{
+					CreateReplicatedStaticActorDestructionInfo(this, Level, Info);
+				}
+			}
+		}
+	}
+}
+
+void UNetDriver::NotifyActorFullyDormantForConnection(AActor* Actor, UNetConnection* Connection)
+{
+	GetNetworkObjectList().MarkDormant(Actor, Connection, IsServer() ? ClientConnections.Num() : 1, this);
+	if (ReplicationDriver)
+	{
+		ReplicationDriver->NotifyActorFullyDormantForConnection(Actor, Connection);
+	}
+}
+
+void UNetDriver::RemoveClientConnection(UNetConnection* ClientConnectionToRemove)
+{
+	verify(ClientConnections.Remove(ClientConnectionToRemove) == 1);
+
+	TSharedPtr<const FInternetAddr> AddrToRemove = ClientConnectionToRemove->GetRemoteAddr();
+
+	if (AddrToRemove.IsValid())
+	{
+		TSharedRef<const FInternetAddr> ConstAddrRef = AddrToRemove.ToSharedRef();
+
+		if (RecentlyDisconnectedTrackingTime > 0)
+		{
+			UNetConnection** FoundVal = MappedClientConnections.Find(ConstAddrRef);
+
+			// Mark recently disconnected clients as nullptr (don't wait for GC), and keep the MappedClientConections entry for a while.
+			// Required for identifying/ignoring packets from recently disconnected clients, with the same performance as for NetConnection's (important for DDoS detection)
+			if (ensure(FoundVal != nullptr))
+			{
+				RecentlyDisconnectedClients.Add(FDisconnectedClient(ConstAddrRef, FPlatformTime::Seconds()));
+
+				*FoundVal = nullptr;
+			}
+		}
+		else
+		{
+			verify(MappedClientConnections.Remove(ConstAddrRef) == 1);
+		}
+	}
+
+	if (ReplicationDriver)
+	{
+		ReplicationDriver->RemoveClientConnection(ClientConnectionToRemove);
+	}
+}
+
+void UNetDriver::SetWorld(class UWorld* InWorld)
+{
+	if (World)
+	{
+		// Remove old world association
+		UnregisterTickEvents(World);
+		World = NULL;
+		WorldPackage = NULL;
+		Notify = NULL;
+
+		GetNetworkObjectList().Reset();
+	}
+
+	if (InWorld)
+	{
+		// Setup new world association
+		World = InWorld;
+		WorldPackage = InWorld->GetOutermost();
+		Notify = InWorld;
+		RegisterTickEvents(InWorld);
+
+		GetNetworkObjectList().AddInitialObjects(InWorld, this);
+	}
+
+	if (ReplicationDriver)
+	{
+		ReplicationDriver->SetRepDriverWorld(InWorld);
+		ReplicationDriver->InitializeActorsInWorld(InWorld);
+	}
+}
+
+void UNetDriver::ResetGameWorldState()
+{
+	DestroyedStartupOrDormantActors.Empty();
+	RenamedStartupActors.Empty();
+
+	if ( NetCache.IsValid() )
+	{
+		NetCache->ClearClassNetCache();	// Clear the cache net: it will recreate itself after seamless travel
+	}
+
+	GetNetworkObjectList().ResetDormancyState();
+
+	if (ServerConnection)
+	{
+		ServerConnection->ResetGameWorldState();
+	}
+	for (auto It = ClientConnections.CreateIterator(); It; ++It)
+	{
+		(*It)->ResetGameWorldState();
+	}
+
+	if (ReplicationDriver)
+	{
+		ReplicationDriver->ResetGameWorldState();
+	}
+}
+
+void UNetDriver::CleanPackageMaps()
+{
+	if ( GuidCache.IsValid() )
+	{ 
+		GuidCache->CleanReferences();
+	}
+}
+
+void UNetDriver::RemoveClassRepLayoutReferences(UClass* Class)
+{
+	RepLayoutMap.Remove(Class);
+
+	for (auto Func : TFieldRange<UFunction>(Class, EFieldIteratorFlags::ExcludeSuper))
+	{
+		if (Func && Func->HasAnyFunctionFlags(EFunctionFlags::FUNC_Net))
+		{
+			RepLayoutMap.Remove(Func);
+		}
+	}
+}
+
+void UNetDriver::CleanupWorldForSeamlessTravel()
+{
+	if (World != nullptr)
+	{
+		for (auto LevelIt(World->GetLevelIterator()); LevelIt; ++LevelIt)
+		{
+			if (const ULevel* Level = *LevelIt)
+			{
+				if (Level->LevelScriptActor)
+				{
+					// workaround for this not being called on clients
+					if (ServerConnection != nullptr)
+					{
+						NotifyActorLevelUnloaded(Level->LevelScriptActor);
+					}
+
+					RemoveClassRepLayoutReferences(Level->LevelScriptActor->GetClass());
+
+					ReplicationChangeListMap.Remove(Level->LevelScriptActor);
+				}
+
+				// This is currently necessary because the actor iterator used in the seamless travel handler 
+				// skips over AWorldSettings actors for an unknown reason.
+				if (Level->GetWorldSettings())
+				{
+					NotifyActorLevelUnloaded(Level->GetWorldSettings());
+				}
+			}
+		}
+	}
+}
+
+void UNetDriver::PreSeamlessTravelGarbageCollect()
+{
+	ResetGameWorldState();
+}
+
+void UNetDriver::PostSeamlessTravelGarbageCollect()
+{
+	CleanPackageMaps();
+}
+
+void UNetDriver::SetReplicationDriver(UReplicationDriver* NewReplicationDriver)
+{
+	if (ReplicationDriver)
+	{
+		ReplicationDriver->TearDown();
+
+		// Though the NetDriver does not set the ConnectionReplicationConnectionDriver,
+		// We are going to clear it here so that references to the old replication driver
+		// do not leak. (E.g, a new replication driver may not set ConnectionReplicationDriver)
+		for (UNetConnection* Connection : ClientConnections)
+		{
+			if (Connection)
+			{
+				Connection->TearDownReplicationConnectionDriver();
+			}		
+		}
+		if (ServerConnection)
+		{
+			// Clients shouldn't currently ever have a repdriver but this could change in the future
+			ServerConnection->TearDownReplicationConnectionDriver();
+		}
+	}
+
+	ReplicationDriver = NewReplicationDriver;
+	if (ReplicationDriver)
+	{
+		ReplicationDriver->SetRepDriverWorld( GetWorld() );
+		ReplicationDriver->InitForNetDriver(this);
+		ReplicationDriver->InitializeActorsInWorld( GetWorld() );
+	}
+}
+
+UNetConnection* UNetDriver::GetConnectionById(uint32 ConnectionId) const
+{
+	if (ServerConnection != nullptr && ServerConnection->GetConnectionId() == ConnectionId)
+	{
+		return ServerConnection;
+	}
+
+	for (UNetConnection* Connection : ClientConnections)
+	{
+		if (Connection && Connection->GetConnectionId() == ConnectionId)
+		{
+			return Connection;
+		}
+	}
+
+	return nullptr;
+}
+
+static void	DumpRelevantActors( UWorld* InWorld )
+{
+	UNetDriver *NetDriver = InWorld->NetDriver;
+	if (!NetDriver)
+	{
+		return;
+	}
+
+	NetDriver->DebugRelevantActors = true;
+}
+
+TSharedPtr<FRepChangedPropertyTracker> UNetDriver::FindOrCreateRepChangedPropertyTracker(UObject* Obj)
+{
+	check(IsServer() || MaySendProperties());
+
+	FRepChangedPropertyTrackerWrapper * GlobalPropertyTrackerPtr = RepChangedPropertyTrackerMap.Find( Obj );
+
+	// Obj can be a new object with a pointer that matches an old, no longer valid, object
+	if ( GlobalPropertyTrackerPtr != nullptr && !GlobalPropertyTrackerPtr->IsObjectValid() )
+	{
+		RepChangedPropertyTrackerMap.Remove(Obj);
+		GlobalPropertyTrackerPtr = nullptr;
+	}
+
+	if ( !GlobalPropertyTrackerPtr ) 
+	{
+		const UWorld* const LocalWorld = GetWorld();
+		const bool bIsReplay = LocalWorld != nullptr && static_cast<void*>(LocalWorld->DemoNetDriver) == static_cast<void*>(this);
+		const bool bIsClientReplayRecording = LocalWorld != nullptr ? LocalWorld->IsRecordingClientReplay() : false;
+		FRepChangedPropertyTracker * Tracker = new FRepChangedPropertyTracker(bIsReplay, bIsClientReplayRecording);
+
+		GetObjectClassRepLayout( Obj->GetClass() )->InitChangedTracker( Tracker );
+
+		GlobalPropertyTrackerPtr = &RepChangedPropertyTrackerMap.Add( Obj, FRepChangedPropertyTrackerWrapper( Obj, TSharedPtr<FRepChangedPropertyTracker>( Tracker ) ) );
+	}
+
+	return GlobalPropertyTrackerPtr->RepChangedPropertyTracker;
+}
+
+TSharedPtr<FRepLayout> UNetDriver::GetObjectClassRepLayout( UClass * Class )
+{
+	TSharedPtr<FRepLayout>* RepLayoutPtr = RepLayoutMap.Find(Class);
+
+	if (!RepLayoutPtr) 
+	{
+		ECreateRepLayoutFlags Flags = MaySendProperties() ? ECreateRepLayoutFlags::MaySendProperties : ECreateRepLayoutFlags::None;
+		RepLayoutPtr = &RepLayoutMap.Add(Class, FRepLayout::CreateFromClass(Class, ServerConnection, Flags));
+	}
+
+	return *RepLayoutPtr;
+}
+
+TSharedPtr<FRepLayout> UNetDriver::GetFunctionRepLayout(UFunction * Function)
+{
+	TSharedPtr<FRepLayout>* RepLayoutPtr = RepLayoutMap.Find(Function);
+
+	if (!RepLayoutPtr) 
+	{
+		ECreateRepLayoutFlags Flags = MaySendProperties() ? ECreateRepLayoutFlags::MaySendProperties : ECreateRepLayoutFlags::None;
+		RepLayoutPtr = &RepLayoutMap.Add(Function, FRepLayout::CreateFromFunction(Function, ServerConnection, Flags));
+	}
+
+	return *RepLayoutPtr;
+}
+
+TSharedPtr<FRepLayout> UNetDriver::GetStructRepLayout( UStruct * Struct )
+{
+	TSharedPtr<FRepLayout>* RepLayoutPtr = RepLayoutMap.Find(Struct);
+
+	if (!RepLayoutPtr) 
+	{
+		ECreateRepLayoutFlags Flags = MaySendProperties() ? ECreateRepLayoutFlags::MaySendProperties : ECreateRepLayoutFlags::None;
+		RepLayoutPtr = &RepLayoutMap.Add(Struct, FRepLayout::CreateFromStruct(Struct, ServerConnection, Flags));
+	}
+
+	return *RepLayoutPtr;
+}
+
+TSharedPtr< FReplicationChangelistMgr > UNetDriver::GetReplicationChangeListMgr( UObject* Object )
+{
+	check(IsServer() || MaySendProperties());
+
+	FReplicationChangelistMgrWrapper* ReplicationChangeListMgrPtr = ReplicationChangeListMap.Find(Object);
+
+	// Object can be a new object with a pointer that matches an old, no longer valid, object
+	if (ReplicationChangeListMgrPtr != nullptr && !ReplicationChangeListMgrPtr->IsObjectValid())
+	{
+		ReplicationChangeListMap.Remove(Object);
+		ReplicationChangeListMgrPtr = nullptr;
+	}
+
+	if (!ReplicationChangeListMgrPtr)
+	{
+		const TSharedPtr<const FRepLayout> RepLayout = GetObjectClassRepLayout(Object->GetClass());
+		FReplicationChangelistMgrWrapper Wrapper(Object, RepLayout->CreateReplicationChangelistMgr(Object, GetCreateReplicationChangelistMgrFlags()));
+		ReplicationChangeListMgrPtr = &ReplicationChangeListMap.Add(Object, Wrapper);
+	}
+
+	return ReplicationChangeListMgrPtr->ReplicationChangelistMgr;
+}
+
+// This method will be called when Streaming Levels become Visible.
+void UNetDriver::OnLevelAddedToWorld(ULevel* Level, UWorld* InWorld)
+{
+	// Actors will be re-added to the network list when ULevel::SortActorList is called.
+}
+
+// This method will be called when Streaming Levels are hidden.
+void UNetDriver::OnLevelRemovedFromWorld(class ULevel* Level, class UWorld* InWorld)
+{
+	if (Level && InWorld == GetWorld())
+	{
+		if (!IsServer())
+		{
+			for (AActor* Actor : Level->Actors)
+			{
+				if (Actor)
+				{
+					// Always call this on clients.
+					// It won't actually destroy the Actor, but it will do some cleanup of the channel, etc.
+					NotifyActorLevelUnloaded(Actor);
+				}
+			}
+		}	
+		else
+	{
+			for (AActor* Actor : Level->Actors)
+		{
+			if (Actor)
+			{
+					// Keep Startup actors alive, because they haven't been destroyed yet.
+					// Technically, Dynamic actors may not have been destroyed either, but this
+					// resembles relevancy.
+					if (!Actor->IsNetStartupActor())
+					{
+				NotifyActorLevelUnloaded(Actor);
+					}
+					else
+					{
+						// We still want to remove Startup actors from the Network list so they aren't processed anymore.
+						GetNetworkObjectList().Remove(Actor);
+					}
+			}
+		}
+
+		TArray<FNetworkGUID> RemovedGUIDs;
+		for (auto It = DestroyedStartupOrDormantActors.CreateIterator(); It; ++It)
+		{
+				FActorDestructionInfo* DestructInfo = It->Value.Get();
+				
+				// Until the level is actually unloaded / reloaded, any Static Actors that have been
+				// destroyed will not be recreated, so we still need to track them.
+				if (DestructInfo->Level == Level && !DestructInfo->NetGUID.IsStatic())
+			{
+				for (UNetConnection* Connection : ClientConnections)
+				{
+					Connection->RemoveDestructionInfo(DestructInfo);
+				}
+
+				RemovedGUIDs.Add(It->Key);
+				It.RemoveCurrent();
+			}
+		}
+	}
+	}
+}
+
+// --------------------------------------------------------------------------------------------------------------------------------------------
+// --------------------------------------------------------------------------------------------------------------------------------------------
+// --------------------------------------------------------------------------------------------------------------------------------------------
+
+int32 CVar_RepDriver_Enable = 1;
+static FAutoConsoleVariableRef CVarRepManagerEnable(TEXT("Net.RepDriver.Enable"), CVar_RepDriver_Enable, TEXT("Enables Replication Driver. 0 will fallback to legacy NetDriver implementation."), ECVF_Default );
+
+UReplicationDriver::UReplicationDriver()
+{
+}
+
+UReplicationDriver* UReplicationDriver::CreateReplicationDriver(UNetDriver* NetDriver, const FURL& URL, UWorld* World)
+{
+	if (CreateReplicationDriverDelegate().IsBound())
+	{
+		return CreateReplicationDriverDelegate().Execute(NetDriver, URL, World);
+	}
+
+	// If we initialize the server map from the initial commandline, this will run before the "-execcmds" option has been handled. So explicitly check for force cmd line options
+	static bool FirstTime = true;
+	if (FirstTime)
+	{
+		FirstTime = false;
+		if (FParse::Param(FCommandLine::Get(), TEXT("RepDriverEnable")))
+		{
+			CVar_RepDriver_Enable = 1;
+		}
+		else if (FParse::Param(FCommandLine::Get(), TEXT("RepDriverDisable")))
+		{
+			CVar_RepDriver_Enable = 0;
+		}
+	}
+
+	if (CVar_RepDriver_Enable == 0)
+	{
+		return nullptr;
+	}
+
+	UClass* ReplicationDriverClass = NetDriver ? NetDriver->ReplicationDriverClass : nullptr;
+	if (ReplicationDriverClass == nullptr)
+	{
+		UE_LOG(LogNet, Log, TEXT("ReplicationDriverClass is null! Not using ReplicationDriver."));
+		return nullptr;
+	}
+
+	return NewObject<UReplicationDriver>(GetTransientPackage(), ReplicationDriverClass);
+}
+
+FCreateReplicationDriver& UReplicationDriver::CreateReplicationDriverDelegate()
+{
+	static FCreateReplicationDriver Delegate;
+	return Delegate;
+}
+
+// --------------------------------------------------------------------------------------------------------------------------------------------
+// --------------------------------------------------------------------------------------------------------------------------------------------
+// --------------------------------------------------------------------------------------------------------------------------------------------
+
+
+void UNetDriver::ProcessRemoteFunction(class AActor* Actor, UFunction* Function, void* Parameters, FOutParmRec* OutParms, FFrame* Stack, class UObject* SubObject )
+{
+#if !UE_BUILD_SHIPPING
+	SCOPE_CYCLE_COUNTER(STAT_NetProcessRemoteFunc);
+	SCOPE_CYCLE_UOBJECT(Function, Function);
+
+	{
+		UObject* TestObject = (SubObject == nullptr) ? Actor : SubObject;
+		ensureMsgf(TestObject->IsSupportedForNetworking() || TestObject->IsNameStableForNetworking(), TEXT("Attempted to call ProcessRemoteFunction with object that is not supported for networking. Object: %s Function: %s"), *TestObject->GetPathName(), *Function->GetName());
+	}
+
+	bool bBlockSendRPC = false;
+
+	SendRPCDel.ExecuteIfBound(Actor, Function, Parameters, OutParms, Stack, SubObject, bBlockSendRPC);
+
+	if (!bBlockSendRPC)
+#endif
+	{
+		const bool bIsServer = IsServer();
+		const bool bIsServerMulticast = bIsServer && (Function->FunctionFlags & FUNC_NetMulticast);
+
+		++TotalRPCsCalled;
+
+		// Forward to replication Driver if there is one
+		if (ReplicationDriver && ReplicationDriver->ProcessRemoteFunction(Actor, Function, Parameters, OutParms, Stack, SubObject))
+		{
+			return;
+		}
+		
+		// RepDriver didn't handle it, default implementation
+		UNetConnection* Connection = nullptr;
+		if (bIsServerMulticast)
+			{
+				TSharedPtr<FRepLayout> RepLayout = GetFunctionRepLayout( Function );
+
+				// Multicast functions go to every client
+				TArray<UNetConnection*> UniqueRealConnections;
+				for (int32 i=0; i<ClientConnections.Num(); ++i)
+				{
+					Connection = ClientConnections[i];
+					if (Connection && Connection->ViewTarget)
+					{
+						// Only send or queue multicasts if the actor is relevant to the connection
+						FNetViewer Viewer(Connection, 0.f);
+
+							if (Connection->GetUChildConnection() != nullptr)
+							{
+								Connection = ((UChildConnection*)Connection)->Parent;
+							}
+
+					// It's possible that an actor is not relevant to a specific connection, but the channel is still alive (due to hysteresis).
+					// However, it's also possible that the Actor could become relevant again before the channel ever closed, and in that case we
+					// don't want to lose Reliable RPCs.
+					if (Actor->IsNetRelevantFor(Viewer.InViewer, Viewer.ViewTarget, Viewer.ViewLocation) ||
+						((Function->FunctionFlags & FUNC_NetReliable) && !!CVarAllowReliableMulticastToNonRelevantChannels.GetValueOnGameThread() && Connection->FindActorChannelRef(Actor)))
+					{
+							// We don't want to call this unless necessary, and it will internally handle being called multiple times before a clear
+							// Builds any shared serialization state for this rpc
+							RepLayout->BuildSharedSerializationForRPC(Parameters);
+
+							InternalProcessRemoteFunction( Actor, SubObject, Connection, Function, Parameters, OutParms, Stack, bIsServer );
+						}
+					}
+				}			
+
+				// Finished sending this multicast rpc, clear any shared state
+				RepLayout->ClearSharedSerializationForRPC();
+
+				// Return here so we don't call InternalProcessRemoteFunction again at the bottom of this function
+				return;
+			}
+
+		// Send function data to remote.
+		Connection = Actor->GetNetConnection();
+		if (Connection)
+		{
+			InternalProcessRemoteFunction( Actor, SubObject, Connection, Function, Parameters, OutParms, Stack, bIsServer );
+		}
+		else
+		{
+			UE_LOG(LogNet, Warning, TEXT("UNetDriver::ProcessRemoteFunction: No owning connection for actor %s. Function %s will not be processed."), *Actor->GetName(), *Function->GetName());
+		}
+	}
+}
+
+void UNetDriver::CreateInitialClientChannels()
+{
+	if (ServerConnection != nullptr)
+	{
+		for (const FChannelDefinition& ChannelDef : ChannelDefinitions)
+		{
+			if (ChannelDef.bInitialClient && (ChannelDef.ChannelClass != nullptr))
+			{
+				ServerConnection->CreateChannelByName(ChannelDef.ChannelName, EChannelCreateFlags::OpenedLocally, ChannelDef.StaticChannelIndex);
+			}
+		}
+	}
+}
+
+void UNetDriver::CreateInitialServerChannels(UNetConnection* ClientConnection)
+{
+	if (ClientConnection != nullptr)
+	{
+		for (const FChannelDefinition& ChannelDef : ChannelDefinitions)
+		{
+			if (ChannelDef.bInitialServer && (ChannelDef.ChannelClass != nullptr))
+			{
+				ClientConnection->CreateChannelByName(ChannelDef.ChannelName, EChannelCreateFlags::OpenedLocally, ChannelDef.StaticChannelIndex);
+			}
+		}
+	}
+}
+
+bool UNetDriver::ShouldReplicateFunction(AActor* Actor, UFunction* Function) const
+{
+	return (Actor && Actor->GetNetDriverName() == NetDriverName);
+}
+
+bool UNetDriver::ShouldForwardFunction(AActor* Actor, UFunction* Function, void* Parms) const
+{
+	return !IsServer();
+}
+
+bool UNetDriver::ShouldReplicateActor(AActor* Actor) const
+{
+	return (Actor && Actor->GetNetDriverName() == NetDriverName);
+}
+
+bool UNetDriver::ShouldCallRemoteFunction(UObject* Object, UFunction* Function, const FReplicationFlags& RepFlags) const
+{
+	return ((!IsServer() || RepFlags.bNetOwner) && !RepFlags.bIgnoreRPCs);
+}
+
+bool UNetDriver::ShouldClientDestroyActor(AActor* Actor) const
+{
+	return (Actor && !Actor->IsA(ALevelScriptActor::StaticClass()));
+}
+
+void UNetDriver::NotifyActorChannelOpen(UActorChannel* Channel, AActor* Actor)
+{
+
+}
+
+void UNetDriver::NotifyActorChannelCleanedUp(UActorChannel* Channel, EChannelCloseReason CloseReason)
+{
+
+}
+
+void UNetDriver::NotifyActorTornOff(AActor* Actor)
+{
+
+}
+
+void UNetDriver::ConsumeAsyncLoadDelinquencyAnalytics(FNetAsyncLoadDelinquencyAnalytics& Out)
+{
+	if (FNetGUIDCache* LocalGuidCache = GuidCache.Get())
+	{
+		LocalGuidCache->ConsumeAsyncLoadDelinquencyAnalytics(Out);
+	}
+	else
+	{
+		Out.Reset();
+	}
+}
+
+const FNetAsyncLoadDelinquencyAnalytics& UNetDriver::GetAsyncLoadDelinquencyAnalytics() const
+{
+	static const FNetAsyncLoadDelinquencyAnalytics Empty(0);
+
+	if (FNetGUIDCache const * const LocalGuidCache = GuidCache.Get())
+	{
+		return LocalGuidCache->GetAsyncLoadDelinquencyAnalytics();
+	}
+
+	return Empty;
+}
+
+void UNetDriver::ResetAsyncLoadDelinquencyAnalytics()
+{
+	if (FNetGUIDCache* LocalGuidCache = GuidCache.Get())
+	{
+		LocalGuidCache->ResetAsyncLoadDelinquencyAnalytics();
+	}
+}
+
+bool UNetDriver::DidHitchLastFrame() const
+{
+	return bDidHitchLastFrame;
+}
+
+bool UNetDriver::IsDormInitialStartupActor(AActor* Actor)
+{
+	return Actor && Actor->IsNetStartupActor() && (Actor->NetDormancy == DORM_Initial);
+}
+
+ECreateReplicationChangelistMgrFlags UNetDriver::GetCreateReplicationChangelistMgrFlags() const
+{
+	return ECreateReplicationChangelistMgrFlags::None;
+}
+
+static TAutoConsoleVariable<int32> CVarRelinkMappedReferences(TEXT("net.RelinkMappedReferences"), 1, TEXT(""));
+
+void UNetDriver::MoveMappedObjectToUnmapped(const UObject* Object)
+{
+	if (!Object)
+	{
+		return;
+	}
+
+	if (!CVarRelinkMappedReferences.GetValueOnGameThread())
+	{
+		return;
+	}
+
+	if (IsServer())
+	{
+		return;
+	}
+
+	// Find all replicators that are referencing this object, and make sure to mark the references as unmapped
+	// This is so when/if this object is instantiated again (using same network guid), we can re-establish the old references
+	FNetworkGUID NetGuid = GuidCache->NetGUIDLookup.FindRef(const_cast<UObject*>(Object));
+
+	if (NetGuid.IsValid())
+	{
+		TSet<FObjectReplicator*>* Replicators = GuidToReplicatorMap.Find(NetGuid);
+
+		if (Replicators != nullptr)
+		{
+			for (FObjectReplicator* Replicator : *Replicators)
+			{
+				if (Replicator->MoveMappedObjectToUnmapped(NetGuid))
+				{
+					UnmappedReplicators.Add(Replicator);
+				}
+				else if (!UnmappedReplicators.Contains(Replicator))
+				{
+					UE_LOG(LogNet, Warning, TEXT("UActorChannel::MoveMappedObjectToUnmapped: MoveMappedObjectToUnmapped didn't find object: %s"), *GetPathNameSafe(Replicator->GetObject()));
+				}
+			}
+		}
+	}
+}
+
+FAutoConsoleCommandWithWorld	DumpRelevantActorsCommand(
+	TEXT("net.DumpRelevantActors"), 
+	TEXT( "Dumps information on relevant actors during next network update" ), 
+	FConsoleCommandWithWorldDelegate::CreateStatic(DumpRelevantActors)
+	);
+
+
+#if DO_ENABLE_NET_TEST
+
+FAutoConsoleCommandWithWorldArgsAndOutputDevice NetEmulationPktEmulationProfile(TEXT("NetEmulation.PktEmulationProfile"), TEXT("Apply a preconfigured emulation profile."),
+	FConsoleCommandWithWorldArgsAndOutputDeviceDelegate::CreateStatic([](const TArray<FString>& Args, UWorld* World, FOutputDevice& Output)
+{
+	bool bProfileApplied(false);
+	if (Args.Num() > 0)
+	{
+		FString CmdParams = FString::Printf(TEXT("PktEmulationProfile=%s"), *(Args[0]));
+
+		NetEmulationHelper::CreatePersistentSimulationSettings();
+
+		bProfileApplied = NetEmulationHelper::PersistentPacketSimulationSettings.GetValue().ParseSettings(*CmdParams, nullptr);
+		
+		if (bProfileApplied)
+		{
+			NetEmulationHelper::ApplySimulationSettingsOnNetDrivers(World, NetEmulationHelper::PersistentPacketSimulationSettings.GetValue());
+		}
+		else
+		{
+			Output.Log(FString::Printf(TEXT("EmulationProfile: %s was not found in Engine.ini"), *(Args[0])));
+		}
+	}
+	else
+	{
+		Output.Log(FString::Printf(TEXT("Missing emulation profile name")));
+	}
+
+	if (!bProfileApplied)
+	{
+		if (const UNetworkSettings* NetworkSettings = GetDefault<UNetworkSettings>())
+		{
+			Output.Log(TEXT("List of some supported emulation profiles:"));
+			for (const FNetworkEmulationProfileDescription& ProfileDesc : NetworkSettings->NetworkEmulationProfiles)
+			{
+				Output.Log(FString::Printf(TEXT("%s"), *ProfileDesc.ProfileName));
+			}
+		}
+	}
+}));
+
+FAutoConsoleCommandWithWorld NetEmulationOff(TEXT("NetEmulation.Off"), TEXT("Turn off network emulation"),
+	FConsoleCommandWithWorldDelegate::CreateStatic([](UWorld* World)
+{
+	NetEmulationHelper::CreatePersistentSimulationSettings();
+	NetEmulationHelper::PersistentPacketSimulationSettings.GetValue().ResetSettings();
+	NetEmulationHelper::ApplySimulationSettingsOnNetDrivers(World, NetEmulationHelper::PersistentPacketSimulationSettings.GetValue());
+}));
+
+#define BUILD_NETEMULATION_CONSOLE_COMMAND(CommandName, CommandHelp) FAutoConsoleCommandWithWorldAndArgs NetEmulation##CommandName(TEXT("NetEmulation."#CommandName), TEXT(CommandHelp), \
+	FConsoleCommandWithWorldAndArgsDelegate::CreateStatic([](const TArray<FString>& Args, UWorld* World) \
+	{ \
+		if (Args.Num() > 0) \
+		{ \
+			NetEmulationHelper::CreatePersistentSimulationSettings(); \
+			FString CmdParams = FString::Printf(TEXT(#CommandName"=%s"), *(Args[0])); \
+			NetEmulationHelper::PersistentPacketSimulationSettings.GetValue().ParseSettings(*CmdParams, nullptr); \
+			NetEmulationHelper::ApplySimulationSettingsOnNetDrivers(World, NetEmulationHelper::PersistentPacketSimulationSettings.GetValue()); \
+		} \
+	}));
+
+BUILD_NETEMULATION_CONSOLE_COMMAND(PktLoss, "Simulates network packet loss");
+BUILD_NETEMULATION_CONSOLE_COMMAND(PktOrder, "Simulates network packets received out of order");
+BUILD_NETEMULATION_CONSOLE_COMMAND(PktDup, "Simulates sending/receiving duplicate network packets");
+BUILD_NETEMULATION_CONSOLE_COMMAND(PktLag, "Simulates network packet lag");
+BUILD_NETEMULATION_CONSOLE_COMMAND(PktLagVariance, "Simulates variable network packet lag");
+BUILD_NETEMULATION_CONSOLE_COMMAND(PktLagMin, "Sets minimum outgoing packet latency");
+BUILD_NETEMULATION_CONSOLE_COMMAND(PktLagMax, "Sets maximum outgoing packet latency)");
+BUILD_NETEMULATION_CONSOLE_COMMAND(PktIncomingLagMin, "Sets minimum incoming packet latency");
+BUILD_NETEMULATION_CONSOLE_COMMAND(PktIncomingLagMax, "Sets maximum incoming packet latency");
+BUILD_NETEMULATION_CONSOLE_COMMAND(PktIncomingLoss, "Simulates incoming packet loss");
+BUILD_NETEMULATION_CONSOLE_COMMAND(PktJitter, "Simulates outgoing packet jitter");
+
+#endif //#if DO_ENABLE_NET_TEST
+
+/**
+ * Exec handler that routes online specific execs to the proper subsystem
+ *
+ * @param InWorld World context
+ * @param Cmd 	the exec command being executed
+ * @param Ar 	the archive to log results to
+ *
+ * @return true if the handler consumed the input, false to continue searching handlers
+ */
+static bool NetDriverExec(UWorld* InWorld, const TCHAR* Cmd, FOutputDevice& Ar)
+{
+	bool bHandled = false;
+
+	// Ignore any execs that don't start with NET
+	if (FParse::Command(&Cmd, TEXT("NET")))
+	{
+		UNetDriver* NamedDriver = NULL;
+		TCHAR TokenStr[128];
+
+		// Route the command to a specific beacon if a name is specified or all of them otherwise
+		if (FParse::Token(Cmd, TokenStr, UE_ARRAY_COUNT(TokenStr), true))
+		{
+			NamedDriver = GEngine->FindNamedNetDriver(InWorld, FName(TokenStr));
+			if (NamedDriver != NULL)
+			{
+				bHandled = NamedDriver->Exec(InWorld, Cmd, Ar);
+			}
+			else
+			{
+				FWorldContext &Context = GEngine->GetWorldContextFromWorldChecked(InWorld);
+
+				Cmd -= FCString::Strlen(TokenStr);
+				for (int32 NetDriverIdx=0; NetDriverIdx < Context.ActiveNetDrivers.Num(); NetDriverIdx++)
+				{
+					NamedDriver = Context.ActiveNetDrivers[NetDriverIdx].NetDriver;
+					if (NamedDriver)
+					{
+						bHandled |= NamedDriver->Exec(InWorld, Cmd, Ar);
+					}
+				}
+			}
+		}
+	}
+
+	return bHandled;
+}
+
+/** Our entry point for all net driver related exec routing */
+FStaticSelfRegisteringExec NetDriverExecRegistration(NetDriverExec);
diff --git a/Engine/Source/Runtime/Engine/Private/NetworkDriver.cpp b/Engine/Source/Runtime/Engine/Private/NetworkDriver.cpp
deleted file mode 100644
