Hopefully the last fix for sporadic cook warnings. Key issue was that we were detaching the archive in USoundWave::WillNeverCacheCookedPlatformDataAgain, which would cause our fallback behavior in CookSimpleWave that tries to load the raw PCM data to fail. Fix was to call UnloadBulkData, which deletes the buffer resident in memory but does not detach the archive, instead of RemoveBulkData, which deletes the buffer resident in memory and detaches the archive, thus causing us to never be able to load the raw pcm data for the duration of the executable.

diff --git a/Engine/Source/Runtime/Engine/Private/AudioDerivedData.cpp b/Engine/Source/Runtime/Engine/Private/AudioDerivedData.cpp
--- a/Engine/Source/Runtime/Engine/Private/AudioDerivedData.cpp
+++ b/Engine/Source/Runtime/Engine/Private/AudioDerivedData.cpp
@@ -1032,9 +1032,14 @@ static void CookSimpleWave(USoundWave* SoundWave, FName FormatName, const IAudio
 	bWasLocked = true;
 	int32 RawDataSize = SoundWave->RawData.GetBulkDataSize();
 
-	// parse the wave data
-	if (!WaveInfo.ReadWaveHeader(RawWaveData, RawDataSize, 0))
+
+	if (!RawWaveData || RawDataSize == 0)
+	{
+		UE_LOG(LogAudioDerivedData, Warning, TEXT("LPCM data failed to load for sound %s"), *SoundWave->GetFullName());
+	}
+	else if (!WaveInfo.ReadWaveHeader(RawWaveData, RawDataSize, 0))
 	{
+		// If we failed to parse the wave header, it's either because of an invalid bitdepth or channel configuration.
 		UE_LOG(LogAudioDerivedData, Warning, TEXT("Only mono or stereo 16 bit waves allowed: %s (%d bytes)"), *SoundWave->GetFullName(), RawDataSize);
 	}
 	else
@@ -1629,7 +1634,7 @@ void USoundWave::WillNeverCacheCookedPlatformDataAgain()
 	Super::WillNeverCacheCookedPlatformDataAgain();
 
 	// TODO: We can clear these arrays if we never need to cook again.
-	RawData.RemoveBulkData();
+	RawData.UnloadBulkData();
 	CompressedFormatData.FlushData();
 }
 #endif
