Unity build fix

jeremy.moore

diff --git a/Engine/Source/Runtime/Engine/Classes/Components/BoundsCopyComponent.h b/Engine/Source/Runtime/Engine/Classes/Components/BoundsCopyComponent.h
--- a/Engine/Source/Runtime/Engine/Classes/Components/BoundsCopyComponent.h
+++ b/Engine/Source/Runtime/Engine/Classes/Components/BoundsCopyComponent.h
@@ -4,6 +4,7 @@
 
 #include "CoreMinimal.h"
 #include "ActorComponent.h"
+#include "GameFramework/Actor.h"
 #include "BoundsCopyComponent.generated.h"
 
 /** Component used to copy the bounds of another Actor. */
