Add Device profile fragments, device profiles can now include 'fragments' each of which may define multiple CVars. This enables removal of duplicate CVar setup across multiple DPs.

diff --git a/Engine/Source/Runtime/Engine/Classes/DeviceProfiles/DeviceProfile.h b/Engine/Source/Runtime/Engine/Classes/DeviceProfiles/DeviceProfile.h
--- a/Engine/Source/Runtime/Engine/Classes/DeviceProfiles/DeviceProfile.h
+++ b/Engine/Source/Runtime/Engine/Classes/DeviceProfiles/DeviceProfile.h
@@ -39,6 +39,9 @@ class ENGINE_API UDeviceProfile : public UTextureLODSettings
 	/** This is not a property, it shouldn't be set by the editor */
 	FString ConfigPlatform;
 
+	/** A collection of UDeviceProfileFragment names, which can contain predefined sets of cvars */
+	TArray<FString> FragmentIncludes;
+
 public:
 
 	/** The collection of CVars which is set from this profile */
diff --git a/Engine/Source/Runtime/Engine/Classes/DeviceProfiles/DeviceProfileFragment.h b/Engine/Source/Runtime/Engine/Classes/DeviceProfiles/DeviceProfileFragment.h
new file mode 100644
--- /dev/null
+++ b/Engine/Source/Runtime/Engine/Classes/DeviceProfiles/DeviceProfileFragment.h
@@ -0,0 +1,19 @@
+// Copyright Epic Games, Inc. All Rights Reserved.
+
+/*=============================================================================
+	DeviceProfileFragment.h: Declares the UDeviceProfileFragment class.
+=============================================================================*/
+
+#pragma once
+
+#include "CoreMinimal.h"
+#include "UObject/ObjectMacros.h"
+#include "DeviceProfileFragment.generated.h"
+
+UCLASS(config=DeviceProfiles, perObjectConfig)
+class ENGINE_API UDeviceProfileFragment : public UObject
+{
+	GENERATED_UCLASS_BODY()
+public:
+	TArray<FString> CVars;
+};
diff --git a/Engine/Source/Runtime/Engine/Private/DeviceProfiles/DeviceProfile.cpp b/Engine/Source/Runtime/Engine/Private/DeviceProfiles/DeviceProfile.cpp
--- a/Engine/Source/Runtime/Engine/Private/DeviceProfiles/DeviceProfile.cpp
+++ b/Engine/Source/Runtime/Engine/Private/DeviceProfiles/DeviceProfile.cpp
@@ -8,6 +8,14 @@
 #include "UObject/UObjectIterator.h"
 #include "HAL/IConsoleManager.h"
 
+#include "DeviceProfiles/DeviceProfileFragment.h"
+
+UDeviceProfileFragment::UDeviceProfileFragment(const FObjectInitializer& ObjectInitializer)
+	: Super(ObjectInitializer)
+{
+}
+
+
 UDeviceProfile::UDeviceProfile(const FObjectInitializer& ObjectInitializer)
 	: Super(ObjectInitializer)
 {
diff --git a/Engine/Source/Runtime/Engine/Private/DeviceProfiles/DeviceProfileManager.cpp b/Engine/Source/Runtime/Engine/Private/DeviceProfiles/DeviceProfileManager.cpp
--- a/Engine/Source/Runtime/Engine/Private/DeviceProfiles/DeviceProfileManager.cpp
+++ b/Engine/Source/Runtime/Engine/Private/DeviceProfiles/DeviceProfileManager.cpp
@@ -18,6 +18,7 @@
 #include "PIEPreviewDeviceProfileSelectorModule.h"
 #endif
 #include "ProfilingDebugging/CsvProfiler.h"
+#include "DeviceProfiles/DeviceProfileFragment.h"
 
 static TAutoConsoleVariable<FString> CVarDeviceProfileOverride(
 	TEXT("dp.Override"),
@@ -81,6 +82,32 @@ UDeviceProfileManager& UDeviceProfileManager::Get(bool bFromPostCDOContruct)
 	return *DeviceProfileManagerSingleton;
 }
 
+static void GetFragmentCvars(const FString& CurrentSectionName, const FString& CVarArrayName, TArray<FString>& FragmentCVarsINOUT, const FString& DeviceProfileFileNameIn)
+{
+	FString FragmentIncludes = TEXT("FragmentIncludes");
+	TArray<FString> FragmentIncludeArray;
+	GConfig->GetArray(*CurrentSectionName, *FragmentIncludes, FragmentIncludeArray, DeviceProfileFileNameIn);
+
+	for(const FString& FragmentInclude : FragmentIncludeArray)
+	{
+		FString FragmentSectionName = FString::Printf(TEXT("%s %s"), *FragmentInclude, *UDeviceProfileFragment::StaticClass()->GetName());
+		if (GConfig->DoesSectionExist(*FragmentSectionName, DeviceProfileFileNameIn))
+		{
+			TArray<FString> FragmentCVars;
+			GConfig->GetArray(*FragmentSectionName, *CVarArrayName, FragmentCVars, DeviceProfileFileNameIn);
+			UE_CLOG(FragmentCVars.Num()>0, LogInit, Log, TEXT("Including %s from fragment: %s"), *CVarArrayName, *FragmentInclude);
+			FragmentCVarsINOUT += FragmentCVars;
+		}
+		else
+		{
+#if UE_BUILD_SHIPPING
+			UE_LOG(LogInit, Error, TEXT("Could not find device profile fragment %s."), *FragmentInclude);
+#else
+			UE_LOG(LogInit, Fatal, TEXT("Could not find device profile fragment %s."), *FragmentInclude);
+#endif
+		}
+	}
+}
 
 void UDeviceProfileManager::InitializeCVarsForActiveDeviceProfile(bool bPushSettings)
 {
@@ -196,9 +223,17 @@ void UDeviceProfileManager::InitializeCVarsForActiveDeviceProfile(bool bPushSett
 					ArrayName += BucketNames[(int32)FPlatformMemory::GetMemorySizeBucket()];
 				}
 
-				TArray< FString > CurrentProfilesCVars;
+				TArray< FString > CurrentProfilesCVars, FragmentCVars;
+				GetFragmentCvars(*CurrentSectionName, *ArrayName, FragmentCVars, DeviceProfileFileName);
 				GConfig->GetArray(*CurrentSectionName, *ArrayName, CurrentProfilesCVars, GDeviceProfilesIni);
 
+				if (FragmentCVars.Num())
+				{
+					// Prepend fragments to CurrentProfilesCVars, fragment cvars should be first so the DP's cvars take priority.
+					Swap(CurrentProfilesCVars, FragmentCVars);
+					CurrentProfilesCVars += FragmentCVars;
+				}
+
 				// Iterate over the profile and make sure we do not have duplicate CVars
 				{
 					TMap< FString, FString > ValidCVars;
