Fix merge issue causing build fail.

diff --git a/Engine/Source/Runtime/Engine/Private/DeviceProfiles/DeviceProfileManager.cpp b/Engine/Source/Runtime/Engine/Private/DeviceProfiles/DeviceProfileManager.cpp
--- a/Engine/Source/Runtime/Engine/Private/DeviceProfiles/DeviceProfileManager.cpp
+++ b/Engine/Source/Runtime/Engine/Private/DeviceProfiles/DeviceProfileManager.cpp
@@ -224,7 +224,7 @@ void UDeviceProfileManager::InitializeCVarsForActiveDeviceProfile(bool bPushSett
 				}
 
 				TArray< FString > CurrentProfilesCVars, FragmentCVars;
-				GetFragmentCvars(*CurrentSectionName, *ArrayName, FragmentCVars, DeviceProfileFileName);
+				GetFragmentCvars(*CurrentSectionName, *ArrayName, FragmentCVars, GDeviceProfilesIni);
 				GConfig->GetArray(*CurrentSectionName, *ArrayName, CurrentProfilesCVars, GDeviceProfilesIni);
 
 				if (FragmentCVars.Num())
