- In case of contention on the system session lock, fallback on the lockless mechanism (the one use for crash/terminate) to report the shutdown (we may not save everything, but we will save the most important information)

diff --git a/Engine/Source/Developer/EditorAnalyticsSession/Private/EditorAnalyticsSession.cpp b/Engine/Source/Developer/EditorAnalyticsSession/Private/EditorAnalyticsSession.cpp
--- a/Engine/Source/Developer/EditorAnalyticsSession/Private/EditorAnalyticsSession.cpp
+++ b/Engine/Source/Developer/EditorAnalyticsSession/Private/EditorAnalyticsSession.cpp
@@ -171,6 +171,7 @@ namespace EditorAnalyticsUtils
 						case FEditorAnalyticsSession::EEventType::Crashed:     Session.bCrashed = true;       break;
 						case FEditorAnalyticsSession::EEventType::GpuCrashed : Session.bGPUCrashed = true;    break;
 						case FEditorAnalyticsSession::EEventType::Terminated:  Session.bIsTerminating = true; break;
+						case FEditorAnalyticsSession::EEventType::Shutdown:    Session.bWasShutdown = true;   break;
 						default: break;
 					}
 
diff --git a/Engine/Source/Developer/EditorAnalyticsSession/Public/EditorAnalyticsSession.h b/Engine/Source/Developer/EditorAnalyticsSession/Public/EditorAnalyticsSession.h
--- a/Engine/Source/Developer/EditorAnalyticsSession/Public/EditorAnalyticsSession.h
+++ b/Engine/Source/Developer/EditorAnalyticsSession/Public/EditorAnalyticsSession.h
@@ -15,6 +15,7 @@ struct EDITORANALYTICSSESSION_API FEditorAnalyticsSession
 		Crashed = 0,
 		GpuCrashed,
 		Terminated,
+		Shutdown,
 	};
 
 	FString SessionId;
diff --git a/Engine/Source/Runtime/Engine/Private/Analytics/EditorSessionSummaryWriter.cpp b/Engine/Source/Runtime/Engine/Private/Analytics/EditorSessionSummaryWriter.cpp
--- a/Engine/Source/Runtime/Engine/Private/Analytics/EditorSessionSummaryWriter.cpp
+++ b/Engine/Source/Runtime/Engine/Private/Analytics/EditorSessionSummaryWriter.cpp
@@ -178,19 +178,19 @@ void FEditorSessionSummaryWriter::UpdateLegacyIdleTimes()
 	// 1 + 1 minutes
 	if (IdleSecondsTmp > (60 + 60))
 	{
-		FPlatformAtomics::InterlockedIncrement(&CurrentSession->Idle1Min); // User spent one more minute as idle > 1 min
+		FPlatformAtomics::InterlockedIncrement(&CurrentSession->Idle1Min); // User spent one more minute after being idle 1 minutes (first 1 minute in this idle time span is not counted as idle)
 	}
 
 	// 5 + 1 minutes
 	if (IdleSecondsTmp > (5 * 60 + 60))
 	{
-		FPlatformAtomics::InterlockedIncrement(&CurrentSession->Idle5Min); // User spent one more minute as idle > 5 min
+		FPlatformAtomics::InterlockedIncrement(&CurrentSession->Idle5Min); // User spent one more minute after being idle 5 minutes (first 5 minutes in this idle time span are not counted as idle)
 	}
 
 	// 30 + 1 minutes
 	if (IdleSecondsTmp > (30 * 60 + 60))
 	{
-		FPlatformAtomics::InterlockedIncrement(&CurrentSession->Idle30Min); // User spent one more minute as idle > 30 min
+		FPlatformAtomics::InterlockedIncrement(&CurrentSession->Idle30Min); // User spent one more minute after being idle 30 minutes (first 30 minutes in this idle time span are not counted as idle)
 	}
 }
 
@@ -320,7 +320,10 @@ void FEditorSessionSummaryWriter::Shutdown()
 		UpdateEditorIdleTime(CurrTimeUtc, /*bSaveSession*/false); // Saved below, don't save twice.
 		UpdateUserIdleTime(CurrTimeUtc, /*bSaveSession*/false); // Saved below, don't save twice.
 		UpdateLegacyIdleTimes();
-		TrySaveCurrentSession();
+		if (!TrySaveCurrentSession()) // The system lock was taken by another process.
+		{
+			CurrentSession->LogEvent(FEditorAnalyticsSession::EEventType::Shutdown, CurrTimeUtc); // Use the lockless mechanism. It doesn't save everything, but it carries the critical information.
+		}
 
 		CurrentSession.Reset();
 	}
@@ -518,14 +521,16 @@ void FEditorSessionSummaryWriter::OnExitPIE(const bool /*bIsSimulating*/)
 	}
 }
 
-void FEditorSessionSummaryWriter::TrySaveCurrentSession()
+bool FEditorSessionSummaryWriter::TrySaveCurrentSession()
 {
 	if (FEditorAnalyticsSession::TryLock()) // Inter-process lock to grant this process exclusive access to the key-store file/registry.
 	{
 		FScopeLock ScopedLock(&SaveSessionLock); // Intra-process lock to grant the calling thread exclusive access to the key-store file/registry.
 		CurrentSession->Save();
 		FEditorAnalyticsSession::Unlock();
+		return true;
 	}
+	return false;
 }
 
 #undef LOCTEXT_NAMESPACE
diff --git a/Engine/Source/Runtime/Engine/Private/Analytics/EditorSessionSummaryWriter.h b/Engine/Source/Runtime/Engine/Private/Analytics/EditorSessionSummaryWriter.h
--- a/Engine/Source/Runtime/Engine/Private/Analytics/EditorSessionSummaryWriter.h
+++ b/Engine/Source/Runtime/Engine/Private/Analytics/EditorSessionSummaryWriter.h
@@ -37,7 +37,7 @@ private:
 	void UpdateEditorIdleTime(const FDateTime& InActivityTimeUtc, bool bSaveSession);
 	void UpdateUserIdleTime(const FDateTime& InUserActivityTimeUtc, bool bSaveSession);
 	void UpdateLegacyIdleTimes();
-	void TrySaveCurrentSession();
+	bool TrySaveCurrentSession();
 
 private:
 	TUniquePtr<FEditorAnalyticsSession> CurrentSession;
