Convert UEdGraph::Serialize over to structured archive format

diff --git a/Engine/Source/Runtime/Engine/Classes/EdGraph/EdGraph.h b/Engine/Source/Runtime/Engine/Classes/EdGraph/EdGraph.h
--- a/Engine/Source/Runtime/Engine/Classes/EdGraph/EdGraph.h
+++ b/Engine/Source/Runtime/Engine/Classes/EdGraph/EdGraph.h
@@ -123,7 +123,7 @@ public:
 	//~ Begin UObject interface
 	virtual void BuildSubobjectMapping(UObject* OtherObject, TMap<UObject*, UObject*>& ObjectMapping) const override;
 #if WITH_EDITORONLY_DATA
-	virtual void Serialize( FArchive& Ar ) override;
+	virtual void Serialize(FStructuredArchiveRecord Record) override;
 	virtual void PostInitProperties() override;
 	virtual void PostLoad() override;
 	//~ End UObject Interface
diff --git a/Engine/Source/Runtime/Engine/Private/EdGraph/EdGraph.cpp b/Engine/Source/Runtime/Engine/Private/EdGraph/EdGraph.cpp
--- a/Engine/Source/Runtime/Engine/Private/EdGraph/EdGraph.cpp
+++ b/Engine/Source/Runtime/Engine/Private/EdGraph/EdGraph.cpp
@@ -147,16 +147,18 @@ void UEdGraph::PostInitProperties()
 	}
 }
 
-void UEdGraph::Serialize( FArchive& Ar )
+void UEdGraph::Serialize(FStructuredArchiveRecord Record)
 {
-	Super::Serialize(Ar);
+	Super::Serialize(Record);
+	const FArchiveState& ArchiveState = Record.GetArchiveState();
+
 	// Keep track of RF_Public
-	if( Ar.IsTransacting() )
+	if(ArchiveState.IsTransacting())
 	{
 		bool bIsPublic = HasAnyFlags(RF_Public);
-		if( Ar.IsLoading() )
+		if(ArchiveState.IsLoading())
 		{
-			Ar << bIsPublic;
+			Record << SA_VALUE(TEXT("IsPublic"), bIsPublic);
 			if (bIsPublic)
 			{
 				SetFlags( RF_Public );
@@ -166,9 +168,9 @@ void UEdGraph::Serialize( FArchive& Ar )
 				ClearFlags( RF_Public );
 			}
 		}
-		else if( Ar.IsSaving() )
+		else if( ArchiveState.IsSaving())
 		{
-			Ar << bIsPublic;
+			Record << SA_VALUE(TEXT("IsPublic"), bIsPublic);
 		}
 	}
 }
