Fixed a crash when destroying a FMaterialResourceProxyReader that was created via a filepath. Original work done by [REDACTED_PERSON].

diff --git a/Engine/Source/Runtime/Engine/Private/Materials/MaterialShared.cpp b/Engine/Source/Runtime/Engine/Private/Materials/MaterialShared.cpp
--- a/Engine/Source/Runtime/Engine/Private/Materials/MaterialShared.cpp
+++ b/Engine/Source/Runtime/Engine/Private/Materials/MaterialShared.cpp
@@ -3939,8 +3939,7 @@ FMaterialResourceProxyReader::FMaterialResourceProxyReader(
 	ERHIFeatureLevel::Type FeatureLevel,
 	EMaterialQualityLevel::Type QualityLevel) :
 	FArchiveProxy(Ar),
-	OffsetToEnd(-1),
-	bReleaseInnerArchive(false)
+	OffsetToEnd(-1)
 {
 	check(InnerArchive.IsLoading());
 	Initialize(FeatureLevel, QualityLevel, FeatureLevel != ERHIFeatureLevel::Num);
@@ -3951,9 +3950,9 @@ FMaterialResourceProxyReader::FMaterialResourceProxyReader(
 	uint32 NameMapOffset,
 	ERHIFeatureLevel::Type FeatureLevel,
 	EMaterialQualityLevel::Type QualityLevel) :
-	FArchiveProxy(*IFileManager::Get().CreateFileReader(Filename, FILEREAD_NoFail)),
-	OffsetToEnd(-1),
-	bReleaseInnerArchive(true)
+	TUniquePtr(IFileManager::Get().CreateFileReader(Filename, FILEREAD_NoFail)), // Create and store the FileArchive
+	FArchiveProxy(*Get()), // Now link it to the archive proxy
+	OffsetToEnd(-1)
 {
 	InnerArchive.Seek(NameMapOffset);
 	Initialize(FeatureLevel, QualityLevel);
@@ -3961,11 +3960,7 @@ FMaterialResourceProxyReader::FMaterialResourceProxyReader(
 
 FMaterialResourceProxyReader::~FMaterialResourceProxyReader()
 {
-	if (bReleaseInnerArchive)
-	{
-		delete &InnerArchive;
-	}
-	else if (OffsetToEnd != -1)
+	if (OffsetToEnd != -1)
 	{
 		InnerArchive.Seek(OffsetToEnd);
 	}
diff --git a/Engine/Source/Runtime/Engine/Public/MaterialShared.h b/Engine/Source/Runtime/Engine/Public/MaterialShared.h
--- a/Engine/Source/Runtime/Engine/Public/MaterialShared.h
+++ b/Engine/Source/Runtime/Engine/Public/MaterialShared.h
@@ -2846,7 +2846,7 @@ private:
 	const FMaterialResource& Resource;
 };
 
-class FMaterialResourceProxyReader final : public FArchiveProxy
+class FMaterialResourceProxyReader final : private TUniquePtr<FArchive>, public FArchiveProxy
 {
 public:
 	FMaterialResourceProxyReader(
@@ -2881,11 +2881,15 @@ public:
 
 	virtual FString GetArchiveName() const override { return TEXT("FMaterialResourceProxyReader"); }
 
+	operator FArchive& ()
+	{
+		return *static_cast<FArchiveProxy*>(this);
+	}
+
 private:
 	TArray<FName> Names;
 	int64 OffsetToFirstResource;
 	int64 OffsetToEnd;
-	bool bReleaseInnerArchive;
 
 	void Initialize(
 		ERHIFeatureLevel::Type FeatureLevel,
