Quick attempt at a CIS fix. If this doesn't work I'll back this whole change out (Paul Chipchase is out for the evening).

diff --git a/Engine/Source/Runtime/Engine/Public/MaterialShared.h b/Engine/Source/Runtime/Engine/Public/MaterialShared.h
--- a/Engine/Source/Runtime/Engine/Public/MaterialShared.h
+++ b/Engine/Source/Runtime/Engine/Public/MaterialShared.h
@@ -2881,11 +2881,6 @@ public:
 
 	virtual FString GetArchiveName() const override { return TEXT("FMaterialResourceProxyReader"); }
 
-	operator FArchive& ()
-	{
-		return *static_cast<FArchiveProxy*>(this);
-	}
-
 private:
 	TArray<FName> Names;
 	int64 OffsetToFirstResource;
