Delete audio chunk cache read requests when they are complete. Also mark the requests as AIOP_FLAG_DONTCACHE to avoid polluting the pak precache.

diff --git a/Engine/Source/Runtime/Engine/Private/AudioStreamingCache.cpp b/Engine/Source/Runtime/Engine/Private/AudioStreamingCache.cpp
--- a/Engine/Source/Runtime/Engine/Private/AudioStreamingCache.cpp
+++ b/Engine/Source/Runtime/Engine/Private/AudioStreamingCache.cpp
@@ -1176,6 +1176,46 @@ FAudioChunkCache::FCacheElement* FAudioChunkCache::EvictLeastRecentChunk()
 	return CacheElement;
 }
 
+static FAutoConsoleTaskPriority CPrio_ClearAudioChunkCacheReadRequest(
+	TEXT("TaskGraph.TaskPriorities.ClearAudioChunkCacheReadRequest"),
+	TEXT("Task and thread priority for an async task that clears FCacheElement::ReadRequest"),
+	ENamedThreads::BackgroundThreadPriority, // if we have background priority task threads, then use them...
+	ENamedThreads::NormalTaskPriority, // .. at normal task priority
+	ENamedThreads::NormalTaskPriority // if we don't have background threads, then use normal priority threads at normal task priority instead
+);
+
+class FClearAudioChunkCacheReadRequestTask
+{
+	IBulkDataIORequest* ReadRequest;
+
+public:
+	FORCEINLINE FClearAudioChunkCacheReadRequestTask(IBulkDataIORequest* InReadRequest)
+		: ReadRequest(InReadRequest)
+	{
+	}
+	static FORCEINLINE TStatId GetStatId()
+	{
+		RETURN_QUICK_DECLARE_CYCLE_STAT(FClearAudioChunkCacheReadRequestTask, STATGROUP_TaskGraphTasks);
+	}
+	static FORCEINLINE ENamedThreads::Type GetDesiredThread()
+	{
+		return CPrio_ClearAudioChunkCacheReadRequest.Get();
+	}
+	FORCEINLINE static ESubsequentsMode::Type GetSubsequentsMode()
+	{
+		return ESubsequentsMode::FireAndForget;
+	}
+	void DoTask(ENamedThreads::Type CurrentThread, const FGraphEventRef& MyCompletionGraphEvent)
+	{
+		if (ReadRequest)
+		{
+			ReadRequest->WaitCompletion();
+			delete ReadRequest;
+			ReadRequest = nullptr;
+		}
+	}
+};
+
 void FAudioChunkCache::KickOffAsyncLoad(FCacheElement* CacheElement, const FChunkKey& InKey, TFunction<void(EAudioChunkLoadResult)> OnLoadCompleted, ENamedThreads::Type CallbackThread, bool bNeededForPlayback)
 {
 	check(CacheElement);
@@ -1279,14 +1319,26 @@ void FAudioChunkCache::KickOffAsyncLoad(FCacheElement* CacheElement, const FChun
 			ExecuteOnLoadCompleteCallback(LoadResult, OnLoadCompleted, CallbackThread);
 
 			NumberOfLoadsInFlight.Decrement();
+
+			IBulkDataIORequest* LocalReadRequest = nullptr;
+			if (CacheElement->ReadRequest)
+			{
+				LocalReadRequest = (IBulkDataIORequest*)FPlatformAtomics::InterlockedExchangePtr((void* volatile*)&CacheElement->ReadRequest, nullptr);
+			}
+
+			if (LocalReadRequest)
+			{
+				TGraphTask<FClearAudioChunkCacheReadRequestTask>::CreateTask().ConstructAndDispatchWhenReady(LocalReadRequest);
+			}
 		};
 
 #if DEBUG_STREAM_CACHE
 		CacheElement->DebugInfo.TimeLoadStarted = FPlatformTime::Seconds();
 #endif
 
-		CacheElement->ReadRequest.Reset(Chunk.BulkData.CreateStreamingRequest(0, ChunkDataSize, AsyncIOPriority, &AsyncFileCallBack, CacheElement->ChunkData));
-		if (!CacheElement->ReadRequest.IsValid())
+		IBulkDataIORequest* LocalReadRequest = Chunk.BulkData.CreateStreamingRequest(0, ChunkDataSize, AsyncIOPriority | AIOP_FLAG_DONTCACHE, &AsyncFileCallBack, CacheElement->ChunkData);
+		CacheElement->ReadRequest = LocalReadRequest;
+		if (!LocalReadRequest)
 		{
 			UE_LOG(LogAudio, Error, TEXT("Chunk load in audio LRU cache failed."));
 			OnLoadCompleted(EAudioChunkLoadResult::ChunkOutOfBounds);
diff --git a/Engine/Source/Runtime/Engine/Public/AudioStreamingCache.h b/Engine/Source/Runtime/Engine/Public/AudioStreamingCache.h
--- a/Engine/Source/Runtime/Engine/Public/AudioStreamingCache.h
+++ b/Engine/Source/Runtime/Engine/Public/AudioStreamingCache.h
@@ -175,7 +175,7 @@ private:
 #endif
 
 		// Handle to our async read request operation.
-		TUniquePtr<IBulkDataIORequest> ReadRequest;
+		IBulkDataIORequest* ReadRequest;
 
 #if DEBUG_STREAM_CACHE
 		FCacheElementDebugInfo DebugInfo;
@@ -188,6 +188,7 @@ private:
 			, LessRecentElement(nullptr)
 			, CacheLookupID(InCacheIndex)
 			, bIsLoaded(false)
+			, ReadRequest(nullptr)
 		{
 		}
 
@@ -205,15 +206,21 @@ private:
 			}
 #endif
 
-			if (ReadRequest.IsValid())
+			IBulkDataIORequest* LocalReadRequest = nullptr;
+			if (ReadRequest)
+			{
+				LocalReadRequest = (IBulkDataIORequest*)FPlatformAtomics::InterlockedExchangePtr((void* volatile*)&ReadRequest, nullptr);
+			}
+
+			if (LocalReadRequest)
 			{
 				if (bCancel)
 				{
-					ReadRequest->Cancel();
+					LocalReadRequest->Cancel();
 				}
 				
-				ReadRequest->WaitCompletion();
-				ReadRequest.Reset();
+				LocalReadRequest->WaitCompletion();
+				delete LocalReadRequest;
 			}
 		}
 
