Sound Cues now update the loading behavior of their USoundWaves if appropriate

diff --git a/Engine/Source/Runtime/Engine/Classes/Sound/SoundCue.h b/Engine/Source/Runtime/Engine/Classes/Sound/SoundCue.h
--- a/Engine/Source/Runtime/Engine/Classes/Sound/SoundCue.h
+++ b/Engine/Source/Runtime/Engine/Classes/Sound/SoundCue.h
@@ -7,6 +7,7 @@
 #include "Sound/SoundAttenuation.h"
 #include "Sound/SoundBase.h"
 #include "Sound/SoundNode.h"
+#include "Sound/SoundWave.h"
 
 #if WITH_EDITOR
 #include "EdGraph/EdGraph.h"
@@ -264,6 +265,9 @@ public:
 	void RetainSoundCue();
 	void ReleaseRetainedAudio();
 
+	/** Call this when stream caching is enabled to update sound waves of loading behavior they are inheriting via SoundCue */
+	void CacheLoadingBehavior(ESoundWaveLoadingBehavior InBehavior);
+
 protected:
 	bool RecursiveFindPathToNode(USoundNode* CurrentNode, const UPTRINT CurrentHash, const UPTRINT NodeHashToFind, TArray<USoundNode*>& OutPath) const;
 
diff --git a/Engine/Source/Runtime/Engine/Classes/Sound/SoundNode.h b/Engine/Source/Runtime/Engine/Classes/Sound/SoundNode.h
--- a/Engine/Source/Runtime/Engine/Classes/Sound/SoundNode.h
+++ b/Engine/Source/Runtime/Engine/Classes/Sound/SoundNode.h
@@ -8,6 +8,7 @@
 #include "UObject/ObjectMacros.h"
 #include "UObject/Object.h"
 #include "UObject/Class.h"
+#include "Sound/SoundWave.h"
 #include "SoundNode.generated.h"
 
 class FAudioDevice;
@@ -212,6 +213,17 @@ public:
 	 */
 	virtual void RetainChildWavePlayers(bool bRecurse);
 
+	/**
+	 * When this is called and stream caching is enabled,
+	 * any wave player sound nodes childed off of this node
+	 * with loading behavior set to "Inherited"
+	 * will have their loading behavior updated and
+	 * their bLoadingBehaviorOverridden flag raised
+	 *
+	 * @param bRecurse when true, this will cause all children of child nodes to be overridden as well.
+	 */
+	virtual void OverrideLoadingBehaviorOnChildWaves(const bool bRecurse, const ESoundWaveLoadingBehavior InLoadingBehavior);
+
 	virtual void ReleaseRetainerOnChildWavePlayers(bool bRecurse);
 
 };
diff --git a/Engine/Source/Runtime/Engine/Classes/Sound/SoundWave.h b/Engine/Source/Runtime/Engine/Classes/Sound/SoundWave.h
--- a/Engine/Source/Runtime/Engine/Classes/Sound/SoundWave.h
+++ b/Engine/Source/Runtime/Engine/Classes/Sound/SoundWave.h
@@ -341,6 +341,9 @@ public:
 	UPROPERTY(EditAnywhere, Category = "Loading", meta = (DisplayName = "Loading Behavior Override"))
 	ESoundWaveLoadingBehavior LoadingBehavior;
 
+	/** Set to true if LoadingBehavior was inherited from a SoundCue. This is useful for debugging/logging */
+	uint8 bLoadingBehaviorOverridden:1;
+
 	/** Set to true for programmatically generated audio. */
 	uint8 bProcedural:1;
 	
@@ -493,6 +496,13 @@ public:
 
 	bool IsRetainingAudio();
 
+	/**
+	 * If Stream Caching is enabled, this can be used to override the default loading behavior of this USoundWave.
+	 * This can even be called on USoundWaves that still have the RF_NeedLoad flag, and won't be stomped by serialization.
+	 * NOTE: The new behavior will be ignored if it is less memory-aggressive than existing (even inherited) behavior
+	 */
+	void OverrideLoadingBehavior(ESoundWaveLoadingBehavior InLoadingBehavior);
+
 	/**
 	 * This is called by the audio stream cache once we put a chunk of compressed data in the chunk table. 
 	 * It allows us to look up the chunk directly, rather than searching the cache linearly.
diff --git a/Engine/Source/Runtime/Engine/Private/AudioStreamingCache.cpp b/Engine/Source/Runtime/Engine/Private/AudioStreamingCache.cpp
--- a/Engine/Source/Runtime/Engine/Private/AudioStreamingCache.cpp
+++ b/Engine/Source/Runtime/Engine/Private/AudioStreamingCache.cpp
@@ -337,7 +337,7 @@ FAudioChunkHandle FCachedAudioStreamingManager::GetLoadedChunk(const USoundWave*
 		// Finally, if there's a chunk after this in the sound, request that it is in the cache.
 		const int32 NextChunk = GetNextChunkIndex(SoundWave, ChunkIndex);
 
-		if (NextChunk != INDEX_NONE)
+		if (NextChunk != INDEX_NONE) 
 		{
 			const FAudioChunkCache::FChunkKey NextChunkKey = 
 			{ 
@@ -500,7 +500,16 @@ bool FCachedAudioStreamingManager::RequestChunk(USoundWave* SoundWave, uint32 Ch
 	FAudioChunkCache* Cache = GetCacheForWave(SoundWave);
 	if (Cache)
 	{
-		FAudioChunkCache::FChunkKey ChunkKey = { SoundWave, SoundWave->GetFName(), ChunkIndex };
+		FAudioChunkCache::FChunkKey ChunkKey = 
+		{
+			SoundWave
+		  , SoundWave->GetFName()
+		  , ChunkIndex
+#if WITH_EDITOR
+		  , (uint32)SoundWave->CurrentChunkRevision.GetValue()
+#endif
+		};
+
 		uint64 LookupIDForChunk = Cache->AddOrTouchChunk(ChunkKey, OnLoadCompleted, ThreadToCallOnLoadCompletedOn, bForImmediatePlayback);
 		SoundWave->SetCacheLookupIDForChunk(ChunkIndex, LookupIDForChunk);
 		return LookupIDForChunk != InvalidAudioStreamCacheLookupID;
@@ -564,16 +573,16 @@ uint64 FAudioChunkCache::AddOrTouchChunk(const FChunkKey& InKey, TFunction<void(
 		if (FoundElement->bIsLoaded)
 		{
 			ExecuteOnLoadCompleteCallback(EAudioChunkLoadResult::AlreadyLoaded, OnLoadCompleted, CallbackThread);
-#if DEBUG_STREAM_CACHE
-			checkf(FoundElement->DebugInfo.LoadingBehavior == InKey.SoundWave->GetLoadingBehavior(true), TEXT("Cached element had different LoadingBehavior {%i} than the new request {%i}")
-				, EnumToString(FoundElement->DebugInfo.LoadingBehavior), EnumToString(InKey.SoundWave->GetLoadingBehavior(true)))
-#endif
-
 		}
 
 #if DEBUG_STREAM_CACHE
 		FoundElement->DebugInfo.NumTimesTouched++;
-		FoundElement->DebugInfo.LoadingBehavior = InKey.SoundWave->GetLoadingBehavior(true);
+
+		// Recursing in no longer needed at this point since the inherited loading behavior has already been cached by the time this information is needed
+		const bool bRecurseSoundClasses = false;
+		FoundElement->DebugInfo.LoadingBehavior = InKey.SoundWave->GetLoadingBehavior(bRecurseSoundClasses);
+
+		FoundElement->DebugInfo.bLoadingBehaviorExternallyOverriden = InKey.SoundWave->bLoadingBehaviorOverridden;
 #endif
 
 		return FoundElement->CacheLookupID;
@@ -590,7 +599,12 @@ uint64 FAudioChunkCache::AddOrTouchChunk(const FChunkKey& InKey, TFunction<void(
 
 #if DEBUG_STREAM_CACHE
 		CacheElement->DebugInfo.bWasCacheMiss = bNeededForPlayback;
-		CacheElement->DebugInfo.LoadingBehavior = InKey.SoundWave->GetLoadingBehavior(true);
+
+		// Recursing in no longer needed at this point since the inherited loading behavior has already been cached by the time this information is needed
+		const bool bRecurseSoundClasses = false;
+		CacheElement->DebugInfo.LoadingBehavior = InKey.SoundWave->GetLoadingBehavior(bRecurseSoundClasses);
+
+		CacheElement->DebugInfo.bLoadingBehaviorExternallyOverriden = InKey.SoundWave->bLoadingBehaviorOverridden;
 #endif
 		const FStreamedAudioChunk& Chunk = InKey.SoundWave->RunningPlatformData->Chunks[InKey.ChunkIndex];
 		int32 ChunkDataSize = Chunk.AudioDataSize;
@@ -659,7 +673,7 @@ TArrayView<uint8> FAudioChunkCache::GetChunk(const FChunkKey& InKey, bool bBlock
 			int32 ChunkAudioDataSize = Chunk.AudioDataSize;
 #if DEBUG_STREAM_CACHE
 			FoundElement->DebugInfo.NumTotalChunks = InKey.SoundWave->GetNumChunks() - 1;
-			FoundElement->DebugInfo.TimeLoadStarted = FPlatformTime::Cycles64();
+			FoundElement->DebugInfo.TimeLoadStarted = FPlatformTime::Seconds();
 #endif
 			MemoryCounterBytes -= FoundElement->ChunkDataSize;
 
@@ -691,7 +705,19 @@ TArrayView<uint8> FAudioChunkCache::GetChunk(const FChunkKey& InKey, bool bBlock
 			FoundElement->bIsLoaded = true;
 #if DEBUG_STREAM_CACHE
 			FoundElement->DebugInfo.TimeToLoad = (FPlatformTime::Seconds() - FoundElement->DebugInfo.TimeLoadStarted) * 1000.0f;
+
 #endif
+			// If this value is ever negative, it means that we're decrementing more than we're incrementing:
+			if (ensureMsgf(FoundElement->NumConsumers.GetValue() >= 0, TEXT("NumConsumers was negative for FoundElement. Reseting to 1")))
+			{
+				FoundElement->NumConsumers.Increment();
+			}
+			else
+			{
+				FoundElement->NumConsumers.Set(1);
+			}
+
+			return TArrayView<uint8>(FoundElement->ChunkData, ChunkAudioDataSize);
 		}
 		else
 		{
@@ -1239,7 +1265,8 @@ void FAudioChunkCache::KickOffAsyncLoad(FCacheElement* CacheElement, const FChun
 
 #if DEBUG_STREAM_CACHE
 	CacheElement->DebugInfo.NumTotalChunks = InKey.SoundWave->GetNumChunks() - 1;
-	CacheElement->DebugInfo.LoadingBehavior = InKey.SoundWave->GetLoadingBehavior(true);
+	CacheElement->DebugInfo.LoadingBehavior = InKey.SoundWave->GetLoadingBehavior(false);
+	CacheElement->DebugInfo.bLoadingBehaviorExternallyOverriden = InKey.SoundWave->bLoadingBehaviorOverridden;
 #endif
 
 	// In editor, we retrieve from the DDC. In non-editor situations, we read the chunk async from the pak file.
@@ -1596,6 +1623,7 @@ TPair<int, int> FAudioChunkCache::DebugDisplay(UWorld* World, FViewport* Viewpor
 		double TimeToLoad = -1.0;
 		float AveragePlaceInCache = -1.0f;
 		ESoundWaveLoadingBehavior LoadingBehavior = ESoundWaveLoadingBehavior::Uninitialized;
+		bool bLoadingBehaviorExternallyOverriden = false;
 		bool bWasCacheMiss = false;
 		bool bIsStaleChunk = false;
 
@@ -1605,6 +1633,7 @@ TPair<int, int> FAudioChunkCache::DebugDisplay(UWorld* World, FViewport* Viewpor
 		TimeToLoad = CurrentElement->DebugInfo.TimeToLoad;
 		AveragePlaceInCache = CurrentElement->DebugInfo.AverageLocationInCacheWhenNeeded;
 		LoadingBehavior = CurrentElement->DebugInfo.LoadingBehavior;
+		bLoadingBehaviorExternallyOverriden = CurrentElement->DebugInfo.bLoadingBehaviorExternallyOverriden;
 		bWasCacheMiss = CurrentElement->DebugInfo.bWasCacheMiss;
 #endif
 
@@ -1615,7 +1644,7 @@ TPair<int, int> FAudioChunkCache::DebugDisplay(UWorld* World, FViewport* Viewpor
 
 		const bool bWasTrimmed = CurrentElement->ChunkDataSize == 0;
 
-		FString ElementInfo = *FString::Printf(TEXT("%4i. Size: %6.2f KB   Chunk: %d of %d   Request Count: %d    Average Index: %6.2f  Number of Handles Retaining Chunk: %d     Chunk Load Time(in ms): %6.4fms      Loading Behavior: %s      Name: %s Notes: %s %s"),
+		FString ElementInfo = *FString::Printf(TEXT("%4i. Size: %6.2f KB   Chunk: %d of %d   Request Count: %d    Average Index: %6.2f  Number of Handles Retaining Chunk: %d     Chunk Load Time(in ms): %6.4fms      Loading Behavior: %s%s      Name: %s Notes: %s %s"),
 			Index,
 			CurrentElement->ChunkDataSize / 1024.0f,
 			CurrentElement->Key.ChunkIndex,
@@ -1625,6 +1654,7 @@ TPair<int, int> FAudioChunkCache::DebugDisplay(UWorld* World, FViewport* Viewpor
 			CurrentElement->NumConsumers.GetValue(),
 			TimeToLoad,
 			EnumToString(LoadingBehavior),
+			bLoadingBehaviorExternallyOverriden ? TEXT("*") : TEXT(""),
 			bWasTrimmed ? TEXT("TRIMMED CHUNK") : *CurrentElement->Key.SoundWaveName.ToString(),
 			bWasCacheMiss ? TEXT("(Cache Miss!)") : TEXT(""),
 			bIsStaleChunk ? TEXT("(Stale Chunk)") : TEXT("")
@@ -1731,6 +1761,7 @@ FString FAudioChunkCache::DebugPrint()
 		double TimeToLoad = -1.0;
 		float AveragePlaceInCache = -1.0f;
 		ESoundWaveLoadingBehavior LoadingBehavior = ESoundWaveLoadingBehavior::Uninitialized;
+		bool bLoadingBehaviorExternallyOverriden = false;
 		bool bWasCacheMiss = false;
 		bool bIsStaleChunk = false;
 
@@ -1740,6 +1771,7 @@ FString FAudioChunkCache::DebugPrint()
 		TimeToLoad = CurrentElement->DebugInfo.TimeToLoad;
 		AveragePlaceInCache = CurrentElement->DebugInfo.AverageLocationInCacheWhenNeeded;
 		LoadingBehavior = CurrentElement->DebugInfo.LoadingBehavior;
+		bLoadingBehaviorExternallyOverriden = CurrentElement->DebugInfo.bLoadingBehaviorExternallyOverriden;
 		bWasCacheMiss = CurrentElement->DebugInfo.bWasCacheMiss;
 #endif
 
@@ -1750,7 +1782,7 @@ FString FAudioChunkCache::DebugPrint()
 
 		const bool bWasTrimmed = CurrentElement->ChunkDataSize == 0;
 
-		FString ElementInfo = *FString::Printf(TEXT("%4i.\t, %6.2f\t, %d of %d\t, %d\t, %6.2f\t, %d\t,  %6.4f\t, %s\t, %s, %s %s %s"),
+		FString ElementInfo = *FString::Printf(TEXT("%4i.\t, %6.2f\t, %d of %d\t, %d\t, %6.2f\t, %d\t,  %6.4f\t, %s\t, %s%s, %s %s %s"),
 			Index,
 			CurrentElement->ChunkDataSize / 1024.0f,
 			CurrentElement->Key.ChunkIndex,
@@ -1761,6 +1793,7 @@ FString FAudioChunkCache::DebugPrint()
 			TimeToLoad,
 			bWasTrimmed ? TEXT("TRIMMED CHUNK") : *CurrentElement->Key.SoundWaveName.ToString(),
 			EnumToString(LoadingBehavior),
+			bLoadingBehaviorExternallyOverriden ? TEXT("*") : TEXT(""),
 			bWasCacheMiss ? TEXT("(Cache Miss!)") : TEXT(""),
 			bIsStaleChunk ? TEXT("(Stale Chunk)") : TEXT(""),
 			CurrentElement->IsLoadInProgress() ? TEXT("(Loading In Progress)") : TEXT("")
diff --git a/Engine/Source/Runtime/Engine/Private/SoundCue.cpp b/Engine/Source/Runtime/Engine/Private/SoundCue.cpp
--- a/Engine/Source/Runtime/Engine/Private/SoundCue.cpp
+++ b/Engine/Source/Runtime/Engine/Private/SoundCue.cpp
@@ -13,7 +13,6 @@
 #include "Sound/SoundNode.h"
 #include "Sound/SoundNodeAssetReferencer.h"
 #include "Sound/SoundNodeMixer.h"
-#include "Sound/SoundWave.h"
 #include "Sound/SoundNodeAttenuation.h"
 #include "Sound/SoundNodeModulator.h"
 #include "Sound/SoundNodeQualityLevel.h"
@@ -117,6 +116,14 @@ void USoundCue::ReleaseRetainedAudio()
 	bIsRetainingAudio = false;
 }
 
+void USoundCue::CacheLoadingBehavior(ESoundWaveLoadingBehavior InBehavior)
+{
+	if (FirstNode)
+	{
+		FirstNode->OverrideLoadingBehaviorOnChildWaves(true, InBehavior);
+	}
+}
+
 void USoundCue::Serialize(FStructuredArchive::FRecord Record)
 {
 	FArchive& UnderlyingArchive = Record.GetUnderlyingArchive();
@@ -195,14 +202,25 @@ void USoundCue::PostLoad()
 		CurrentSoundClass = CurrentSoundClass->ParentClass;
 	}
 
+	// RETAIN behavior gets demoted to PRIME in editor because in editor so many sound waves are technically "loaded" at a time.
+	// If the Cache were not willing to evict them, that would cause issues.
 	if (!GIsEditor && SoundClassLoadingBehavior == ESoundWaveLoadingBehavior::RetainOnLoad)
 	{
 		RetainSoundCue();
 	}
-	else if (bPrimeOnLoad || SoundClassLoadingBehavior == ESoundWaveLoadingBehavior::PrimeOnLoad)
+	else if (bPrimeOnLoad
+		|| (SoundClassLoadingBehavior == ESoundWaveLoadingBehavior::PrimeOnLoad)
+		|| (GIsEditor && (SoundClassLoadingBehavior == ESoundWaveLoadingBehavior::RetainOnLoad))
+		)
 	{
+		// In editor, we call this for RetainOnLoad behavior as well
 		PrimeSoundCue();
 	}
+	else if (SoundClassLoadingBehavior == ESoundWaveLoadingBehavior::LoadOnDemand)
+	{
+		// update the soundWaves with the loading behavior they inherited
+		CacheLoadingBehavior(SoundClassLoadingBehavior);
+ 	}
 }
 
 bool USoundCue::CanBeClusterRoot() const
diff --git a/Engine/Source/Runtime/Engine/Private/SoundNode.cpp b/Engine/Source/Runtime/Engine/Private/SoundNode.cpp
--- a/Engine/Source/Runtime/Engine/Private/SoundNode.cpp
+++ b/Engine/Source/Runtime/Engine/Private/SoundNode.cpp
@@ -7,7 +7,6 @@
 #include "Misc/App.h"
 #include "Sound/SoundNodeWavePlayer.h"
 #include "ContentStreaming.h"
-#include "Sound/SoundWave.h"
 #include "AudioCompressionSettingsUtils.h"
 
 /*-----------------------------------------------------------------------------
@@ -93,38 +92,19 @@ UPTRINT USoundNode::GetNodeWaveInstanceHash(const UPTRINT ParentWaveInstanceHash
 
 void USoundNode::PrimeChildWavePlayers(bool bRecurse)
 {
-	if (FPlatformCompressionUtilities::IsCurrentPlatformUsingStreamCaching())
-	{
-		// Search child nodes for wave players, then prime their waves.
-		for (USoundNode* ChildNode : ChildNodes)
-		{
-			if (ChildNode)
-			{
-				ChildNode->ConditionalPostLoad();
-				if (bRecurse)
-				{
-					ChildNode->PrimeChildWavePlayers(true);
-				}
-
-				USoundNodeWavePlayer* WavePlayer = Cast<USoundNodeWavePlayer>(ChildNode);
-				if (WavePlayer != nullptr)
-				{
-					USoundWave* SoundWave = WavePlayer->GetSoundWave();
-					if (SoundWave && SoundWave->IsStreaming(nullptr) && SoundWave->GetNumChunks() > 1)
-					{
-						IStreamingManager::Get().GetAudioStreamingManager().RequestChunk(SoundWave, 1, [](EAudioChunkLoadResult) {});
-					}
-				}
-			}
-		}
-	}
+	OverrideLoadingBehaviorOnChildWaves(bRecurse, ESoundWaveLoadingBehavior::PrimeOnLoad);
 }
 
 void USoundNode::RetainChildWavePlayers(bool bRecurse)
 {
-	if (!GIsEditor && FPlatformCompressionUtilities::IsCurrentPlatformUsingStreamCaching())
+	OverrideLoadingBehaviorOnChildWaves(bRecurse, ESoundWaveLoadingBehavior::RetainOnLoad);
+}
+
+void USoundNode::OverrideLoadingBehaviorOnChildWaves(const bool bRecurse, const ESoundWaveLoadingBehavior InLoadingBehavior)
+{
+	if (FPlatformCompressionUtilities::IsCurrentPlatformUsingStreamCaching())
 	{
-		// Search child nodes for wave players, then prime their waves.
+		// Search child nodes for wave players, then override their waves' loading behavior.
 		for (USoundNode* ChildNode : ChildNodes)
 		{
 			if (ChildNode)
@@ -132,23 +112,26 @@ void USoundNode::RetainChildWavePlayers(bool bRecurse)
 				ChildNode->ConditionalPostLoad();
 				if (bRecurse)
 				{
-					ChildNode->RetainChildWavePlayers(true);
+					ChildNode->OverrideLoadingBehaviorOnChildWaves(true, InLoadingBehavior);
 				}
 
 				USoundNodeWavePlayer* WavePlayer = Cast<USoundNodeWavePlayer>(ChildNode);
 				if (WavePlayer != nullptr)
 				{
 					USoundWave* SoundWave = WavePlayer->GetSoundWave();
-					if (SoundWave && SoundWave->IsStreaming())
+					if (SoundWave)
 					{
-						SoundWave->ConditionalPostLoad();
-						SoundWave->RetainCompressedAudio();
+						SoundWave->OverrideLoadingBehavior(InLoadingBehavior);
 					}
 				}
 			}
 		}
 	}
-	bIsRetainingAudio = true;
+
+	if (!GIsEditor && InLoadingBehavior == ESoundWaveLoadingBehavior::RetainOnLoad)
+	{
+		bIsRetainingAudio = true;
+	}
 }
 
 void USoundNode::ReleaseRetainerOnChildWavePlayers(bool bRecurse)
diff --git a/Engine/Source/Runtime/Engine/Private/SoundNodeAssetReferencer.cpp b/Engine/Source/Runtime/Engine/Private/SoundNodeAssetReferencer.cpp
--- a/Engine/Source/Runtime/Engine/Private/SoundNodeAssetReferencer.cpp
+++ b/Engine/Source/Runtime/Engine/Private/SoundNodeAssetReferencer.cpp
@@ -8,8 +8,15 @@
 
 #define ASYNC_LOAD_RANDOMIZED_SOUNDS 1
 
+// If stream caching is enabled and this is set to 0,
+// we may result in dropped sounds due to USoundNodeWavePlayer::LoadAsset nulling out the SoundWave ptr.
+#define MAKE_SOUNDWAVES_HARD_REFERENCES 1
+
 bool USoundNodeAssetReferencer::ShouldHardReferenceAsset() const
 {
+#if MAKE_SOUNDWAVES_HARD_REFERENCES
+	return true;
+#else
 	bool bShouldHardReference = true;
 
 	if (USoundCue* Cue = Cast<USoundCue>(GetOuter()))
@@ -53,6 +60,7 @@ bool USoundNodeAssetReferencer::ShouldHardReferenceAsset() const
 	}
 
 	return bShouldHardReference;
+#endif // MAKE_SOUNDWAVES_HARD_REFERENCES
 }
 
 #if WITH_EDITOR
diff --git a/Engine/Source/Runtime/Engine/Private/SoundWave.cpp b/Engine/Source/Runtime/Engine/Private/SoundWave.cpp
--- a/Engine/Source/Runtime/Engine/Private/SoundWave.cpp
+++ b/Engine/Source/Runtime/Engine/Private/SoundWave.cpp
@@ -36,7 +36,7 @@
 #include "Async/Async.h"
 #include "Misc/CommandLine.h"
 
-static int32 SoundWaveDefaultLoadingBehaviorCVar = 1;
+static int32 SoundWaveDefaultLoadingBehaviorCVar = 0;
 FAutoConsoleVariableRef CVarSoundWaveDefaultLoadingBehavior(
 	TEXT("au.streamcache.SoundWaveDefaultLoadingBehavior"),
 	SoundWaveDefaultLoadingBehaviorCVar,
@@ -128,12 +128,13 @@ void FStreamedAudioChunk::Serialize(FArchive& Ar, UObject* Owner, int32 ChunkInd
 	DECLARE_SCOPE_CYCLE_COUNTER( TEXT("FStreamedAudioChunk::Serialize"), STAT_StreamedAudioChunk_Serialize, STATGROUP_LoadTime );
 	bool bShouldInlineAudioChunk = false;
 
-// 	const ITargetPlatform* CookingTarget = Ar.CookingTarget();
-// 	if (CookingTarget != nullptr)
-// 	{
-// 		const FPlatformAudioCookOverrides* Overrides = CookingTarget->GetAudioCompressionSettings();
-// 		bShouldInlineAudioChunk = Overrides->bInlineStreamedAudioChunks;
-// 	}
+	const ITargetPlatform* CookingTarget = Ar.CookingTarget();
+	if (CookingTarget != nullptr)
+	{
+		const FPlatformAudioCookOverrides* Overrides = FPlatformCompressionUtilities::GetCookOverrides(*CookingTarget->PlatformName());
+		check(Overrides);
+		bShouldInlineAudioChunk = Overrides->bInlineStreamedAudioChunks;
+	}
 
 	bool bCooked = Ar.IsCooking();
 	Ar << bCooked;
@@ -386,7 +387,7 @@ void USoundWave::Serialize( FArchive& Ar )
 				{
 					// for now we only support one format per wav
 					FName Format = CookingTarget->GetWaveFormat(this);
-					const FPlatformAudioCookOverrides* CompressionOverrides = FPlatformCompressionUtilities::GetCookOverrides(*CookingTarget->IniPlatformName());
+					const FPlatformAudioCookOverrides* CompressionOverrides = FPlatformCompressionUtilities::GetCookOverrides();
 
 					GetCompressedData(Format, CompressionOverrides); // Get the data from the DDC or build it
 					if (CompressionOverrides)
@@ -447,7 +448,7 @@ void USoundWave::Serialize( FArchive& Ar )
 #endif // #if WITH_EDITORONLY_DATA
 	}
 
-	if (!GIsEditor && !(IsTemplate() || IsRunningDedicatedServer()) &&  Ar.IsLoading())
+	if (!(IsTemplate() || IsRunningDedicatedServer()) &&  Ar.IsLoading())
 	{
 		// For non-editor builds, we can immediately cache the sample rate.
 		SampleRate = GetSampleRateForCurrentPlatform();
@@ -460,13 +461,16 @@ void USoundWave::Serialize( FArchive& Ar )
 			//EnsureZerothChunkIsLoaded();
 			const bool bHasFirstChunk = GetNumChunks() > 1;
 
-			if (LoadingBehavior == ESoundWaveLoadingBehavior::RetainOnLoad && bHasFirstChunk)
+			if (!GIsEditor && CurrentLoadingBehavior == ESoundWaveLoadingBehavior::RetainOnLoad && bHasFirstChunk)
 			{
 				RetainCompressedAudio(true);
 			}
-			else if (LoadingBehavior == ESoundWaveLoadingBehavior::PrimeOnLoad && bHasFirstChunk)
+			else if ((CurrentLoadingBehavior == ESoundWaveLoadingBehavior::PrimeOnLoad && bHasFirstChunk)
+				|| (GIsEditor && CurrentLoadingBehavior == ESoundWaveLoadingBehavior::RetainOnLoad))
 			{
 				// Prime first chunk of audio.
+				CachedSoundWaveLoadingBehavior = ESoundWaveLoadingBehavior::PrimeOnLoad;
+				bLoadingBehaviorOverridden = true;
 				IStreamingManager::Get().GetAudioStreamingManager().RequestChunk(this, 1, [](EAudioChunkLoadResult) {});
 			}
 		}
@@ -542,7 +546,7 @@ bool USoundWave::HasCompressedData(FName Format, ITargetPlatform* TargetPlatform
 	{
 		if (TargetPlatform)
 		{
-			CompressionOverrides = FPlatformCompressionUtilities::GetCookOverrides(*TargetPlatform->IniPlatformName());
+			CompressionOverrides = FPlatformCompressionUtilities::GetCookOverrides();
 		}
 	}
 	else
@@ -774,8 +778,8 @@ void USoundWave::InvalidateCompressedData(bool bFreeResources, bool bRebuildStre
 		CachePlatformData();
 		CurrentChunkRevision.Increment();
 	}
-	
-	
+
+
 	// If this sound wave is retained, release and re-retain the new chunk.
 	if (FirstChunk.IsValid())
 	{
@@ -817,27 +821,26 @@ void USoundWave::PostLoad()
 
 	CacheInheritedLoadingBehavior();
 
-	if (!GIsEditor)
-	{
-		// If our loading behavior is defined by a sound class, we need to update whether this sound wave actually needs to retain its audio data or not.
-		ESoundWaveLoadingBehavior ActualLoadingBehavior = GetLoadingBehavior();
+	
+	// If our loading behavior is defined by a sound class, we need to update whether this sound wave actually needs to retain its audio data or not.
+	ESoundWaveLoadingBehavior ActualLoadingBehavior = GetLoadingBehavior();
 
-		if (ShouldUseStreamCaching() && ActualLoadingBehavior != GetLoadingBehavior(false))
+	if (ShouldUseStreamCaching() && ActualLoadingBehavior != GetLoadingBehavior(false))
+	{
+		if (!GIsEditor && !DisableRetainingCVar && ActualLoadingBehavior == ESoundWaveLoadingBehavior::RetainOnLoad)
 		{
-			if (!DisableRetainingCVar && ActualLoadingBehavior == ESoundWaveLoadingBehavior::RetainOnLoad)
-			{
-				UE_LOG(LogAudio, Display, TEXT("Sound Wave '%s' will have to load its compressed audio data async."), *GetName());
-				RetainCompressedAudio(false);
-			}
-			else
-			{
-				// if a sound class defined our loading behavior as something other than Retain and our cvar default is to retain, we need to release our handle.
-				ReleaseCompressedAudio();
+			UE_LOG(LogAudio, Display, TEXT("Sound Wave '%s' will have to load its compressed audio data async."), *GetName());
+			RetainCompressedAudio(false);
+		}
+		else
+		{
+			// if a sound class defined our loading behavior as something other than Retain and our cvar default is to retain, we need to release our handle.
+			ReleaseCompressedAudio();
 
-				if (ActualLoadingBehavior == ESoundWaveLoadingBehavior::PrimeOnLoad && GetNumChunks() > 1)
-				{
-					IStreamingManager::Get().GetAudioStreamingManager().RequestChunk(this, 1, [](EAudioChunkLoadResult) {});
-				}
+			if ((ActualLoadingBehavior == ESoundWaveLoadingBehavior::PrimeOnLoad && GetNumChunks() > 1)
+				|| (GIsEditor && ActualLoadingBehavior == ESoundWaveLoadingBehavior::RetainOnLoad))
+			{
+				IStreamingManager::Get().GetAudioStreamingManager().RequestChunk(this, 1, [](EAudioChunkLoadResult) {});
 			}
 		}
 
@@ -862,14 +865,14 @@ void USoundWave::PostLoad()
 		{
 			if (!Platform->IsServerOnly())
 			{
-				BeginGetCompressedData(Platform->GetWaveFormat(this), FPlatformCompressionUtilities::GetCookOverrides(*Platform->IniPlatformName()));
+				BeginGetCompressedData(Platform->GetWaveFormat(this), FPlatformCompressionUtilities::GetCookOverrides());
 			}
 		}
 	}
 
 	// We don't precache default objects and we don't precache in the Editor as the latter will
 	// most likely cause us to run out of memory.
-	if (!GIsEditor && !IsTemplate( RF_ClassDefaultObject ) && GEngine)
+	if (!GIsEditor && !IsTemplate(RF_ClassDefaultObject) && GEngine)
 	{
 		FAudioDevice* AudioDevice = GEngine->GetMainAudioDeviceRaw();
 		if (AudioDevice)
@@ -878,7 +881,7 @@ void USoundWave::PostLoad()
 			AudioDevice->Precache(this);
 		}
 		// remove bulk data if no AudioDevice is used and no sounds were initialized
-		else if(IsRunningGame())
+		else if (IsRunningGame())
 		{
 			RawData.RemoveBulkData();
 		}
@@ -911,8 +914,8 @@ void USoundWave::PostLoad()
 	bNeedsThumbnailGeneration = true;
 #endif // #if WITH_EDITORONLY_DATA
 
-	INC_FLOAT_STAT_BY( STAT_AudioBufferTime, Duration );
-	INC_FLOAT_STAT_BY( STAT_AudioBufferTimeChannels, NumChannels * Duration );
+	INC_FLOAT_STAT_BY(STAT_AudioBufferTime, Duration);
+	INC_FLOAT_STAT_BY(STAT_AudioBufferTimeChannels, NumChannels * Duration);
 }
 
 void USoundWave::EnsureZerothChunkIsLoaded()
@@ -963,7 +966,7 @@ uint32 USoundWave::GetNumChunks() const
 	}
 	else
 	{
-		ensureAlwaysMsgf(false, TEXT("Call CachePlatformData(false) before calling this function in editor."));
+		ensureAlwaysMsgf(false, TEXT("Call CachePlatformData(false) before calling this function in editor. GetNumChunks() called on: %s"), *GetName());
 		return 0;
 	}
 }
@@ -972,7 +975,7 @@ uint32 USoundWave::GetSizeOfChunk(uint32 ChunkIndex)
 {
 	check(ChunkIndex < GetNumChunks());
 
-	if(RunningPlatformData)
+	if (RunningPlatformData)
 	{
 		return RunningPlatformData->Chunks[ChunkIndex].AudioDataSize;
 	}
@@ -1021,17 +1024,17 @@ void USoundWave::BeginDestroy()
 	ReleaseCompressedAudio();
 }
 
-void USoundWave::InitAudioResource( FByteBulkData& CompressedData )
+void USoundWave::InitAudioResource(FByteBulkData& CompressedData)
 {
-	if( !ResourceSize )
+	if (!ResourceSize)
 	{
 		// Grab the compressed vorbis data from the bulk data
 		ResourceSize = CompressedData.GetBulkDataSize();
-		if( ResourceSize > 0 )
+		if (ResourceSize > 0)
 		{
 #if WITH_EDITOR
 			check(!ResourceData);
-			CompressedData.GetCopy( ( void** )&ResourceData, true );
+			CompressedData.GetCopy((void**)&ResourceData, true);
 #else
 			check(!OwnedBulkDataPtr);
 			OwnedBulkDataPtr = CompressedData.StealFileMapping();
@@ -1056,7 +1059,7 @@ void USoundWave::InitAudioResource( FByteBulkData& CompressedData )
 
 bool USoundWave::InitAudioResource(FName Format)
 {
-	if( !ResourceSize && (!FPlatformProperties::SupportsAudioStreaming() || !IsStreaming(nullptr)) )
+	if (!ResourceSize && (!FPlatformProperties::SupportsAudioStreaming() || !IsStreaming(nullptr)))
 	{
 		FByteBulkData* Bulk = GetCompressedData(Format, GetPlatformCompressionOverridesForCurrentPlatform());
 		if (Bulk)
@@ -1079,7 +1082,7 @@ bool USoundWave::InitAudioResource(FName Format)
 void USoundWave::RemoveAudioResource()
 {
 #if WITH_EDITOR
-	if(ResourceData)
+	if (ResourceData)
 	{
 		FMemory::Free((void*)ResourceData);
 		ResourceSize = 0;
@@ -1110,9 +1113,9 @@ void USoundWave::InvalidateSoundWaveIfNeccessary()
 	{
 		InvalidateCompressedData(true);
 		bWasStreamCachingEnabledOnLastCook = bIsStreamCachingEnabled;
-		
+
 		// If stream caching is now turned on, recook the streaming audio if neccessary.
-		if(bIsStreamCachingEnabled && IsStreaming(nullptr))
+		if (bIsStreamCachingEnabled && IsStreaming(nullptr))
 		{
 			EnsureZerothChunkIsLoaded();
 		}
@@ -1121,7 +1124,7 @@ void USoundWave::InvalidateSoundWaveIfNeccessary()
 
 float USoundWave::GetSampleRateForTargetPlatform(const ITargetPlatform* TargetPlatform)
 {
-	const FPlatformAudioCookOverrides* Overrides = FPlatformCompressionUtilities::GetCookOverrides(*TargetPlatform->IniPlatformName());
+	const FPlatformAudioCookOverrides* Overrides = FPlatformCompressionUtilities::GetCookOverrides();
 	if (Overrides)
 	{
 		return GetSampleRateForCompressionOverrides(Overrides);
@@ -1213,37 +1216,37 @@ void USoundWave::LogBakedData()
 static bool AnyFFTAnalysisPropertiesChanged(const FName& PropertyName)
 {
 	// List of properties which cause analysis to get triggered
-	static FName OverrideSoundName						= GET_MEMBER_NAME_CHECKED(USoundWave, OverrideSoundToUseForAnalysis);
-	static FName EnableFFTAnalysisFName					= GET_MEMBER_NAME_CHECKED(USoundWave, bEnableBakedFFTAnalysis);
-	static FName FFTSizeFName							= GET_MEMBER_NAME_CHECKED(USoundWave, FFTSize);
-	static FName FFTAnalysisFrameSizeFName				= GET_MEMBER_NAME_CHECKED(USoundWave, FFTAnalysisFrameSize);
-	static FName FrequenciesToAnalyzeFName				= GET_MEMBER_NAME_CHECKED(USoundWave, FrequenciesToAnalyze);
-	static FName FFTAnalysisAttackTimeFName				= GET_MEMBER_NAME_CHECKED(USoundWave, FFTAnalysisAttackTime);
-	static FName FFTAnalysisReleaseTimeFName			= GET_MEMBER_NAME_CHECKED(USoundWave, FFTAnalysisReleaseTime);
+	static FName OverrideSoundName = GET_MEMBER_NAME_CHECKED(USoundWave, OverrideSoundToUseForAnalysis);
+	static FName EnableFFTAnalysisFName = GET_MEMBER_NAME_CHECKED(USoundWave, bEnableBakedFFTAnalysis);
+	static FName FFTSizeFName = GET_MEMBER_NAME_CHECKED(USoundWave, FFTSize);
+	static FName FFTAnalysisFrameSizeFName = GET_MEMBER_NAME_CHECKED(USoundWave, FFTAnalysisFrameSize);
+	static FName FrequenciesToAnalyzeFName = GET_MEMBER_NAME_CHECKED(USoundWave, FrequenciesToAnalyze);
+	static FName FFTAnalysisAttackTimeFName = GET_MEMBER_NAME_CHECKED(USoundWave, FFTAnalysisAttackTime);
+	static FName FFTAnalysisReleaseTimeFName = GET_MEMBER_NAME_CHECKED(USoundWave, FFTAnalysisReleaseTime);
 
 	return	PropertyName == OverrideSoundName ||
-			PropertyName == EnableFFTAnalysisFName ||
-			PropertyName == FFTSizeFName ||
-			PropertyName == FFTAnalysisFrameSizeFName ||
-			PropertyName == FrequenciesToAnalyzeFName ||
-			PropertyName == FFTAnalysisAttackTimeFName ||
-			PropertyName == FFTAnalysisReleaseTimeFName;
+		PropertyName == EnableFFTAnalysisFName ||
+		PropertyName == FFTSizeFName ||
+		PropertyName == FFTAnalysisFrameSizeFName ||
+		PropertyName == FrequenciesToAnalyzeFName ||
+		PropertyName == FFTAnalysisAttackTimeFName ||
+		PropertyName == FFTAnalysisReleaseTimeFName;
 }
 
 static bool AnyEnvelopeAnalysisPropertiesChanged(const FName& PropertyName)
 {
 	// List of properties which cause re-analysis to get triggered
-	static FName OverrideSoundName						= GET_MEMBER_NAME_CHECKED(USoundWave, OverrideSoundToUseForAnalysis);
+	static FName OverrideSoundName = GET_MEMBER_NAME_CHECKED(USoundWave, OverrideSoundToUseForAnalysis);
 	static FName EnableAmplitudeEnvelopeAnalysisFName = GET_MEMBER_NAME_CHECKED(USoundWave, bEnableAmplitudeEnvelopeAnalysis);
 	static FName EnvelopeFollowerFrameSizeFName = GET_MEMBER_NAME_CHECKED(USoundWave, EnvelopeFollowerFrameSize);
 	static FName EnvelopeFollowerAttackTimeFName = GET_MEMBER_NAME_CHECKED(USoundWave, EnvelopeFollowerAttackTime);
 	static FName EnvelopeFollowerReleaseTimeFName = GET_MEMBER_NAME_CHECKED(USoundWave, EnvelopeFollowerReleaseTime);
 
 	return	PropertyName == OverrideSoundName ||
-			PropertyName == EnableAmplitudeEnvelopeAnalysisFName ||
-			PropertyName == EnvelopeFollowerFrameSizeFName ||
-			PropertyName == EnvelopeFollowerAttackTimeFName ||
-			PropertyName == EnvelopeFollowerReleaseTimeFName;
+		PropertyName == EnableAmplitudeEnvelopeAnalysisFName ||
+		PropertyName == EnvelopeFollowerFrameSizeFName ||
+		PropertyName == EnvelopeFollowerAttackTimeFName ||
+		PropertyName == EnvelopeFollowerReleaseTimeFName;
 
 }
 
@@ -1556,12 +1559,12 @@ void USoundWave::FreeResources(bool bStopSoundsUsingThisResource)
 	check(IsInAudioThread());
 
 	// Housekeeping of stats
-	DEC_FLOAT_STAT_BY( STAT_AudioBufferTime, Duration );
-	DEC_FLOAT_STAT_BY( STAT_AudioBufferTimeChannels, NumChannels * Duration );
+	DEC_FLOAT_STAT_BY(STAT_AudioBufferTime, Duration);
+	DEC_FLOAT_STAT_BY(STAT_AudioBufferTimeChannels, NumChannels * Duration);
 
 	// GEngine is NULL during script compilation and GEngine->Client and its audio device might be
 	// destroyed first during the exit purge.
-	if( GEngine && !GExitPurge )
+	if (GEngine && !GExitPurge)
 	{
 		// Notify the audio device to free the bulk data associated with this wave.
 		FAudioDeviceManager* AudioDeviceManager = GEngine->GetAudioDeviceManager();
@@ -1635,7 +1638,7 @@ bool USoundWave::CleanupDecompressor(bool bForceWait)
 	return false;
 }
 
-FWaveInstance& USoundWave::HandleStart( FActiveSound& ActiveSound, const UPTRINT WaveInstanceHash ) const
+FWaveInstance& USoundWave::HandleStart(FActiveSound& ActiveSound, const UPTRINT WaveInstanceHash) const
 {
 	// Create a new wave instance and associate with the ActiveSound
 	FWaveInstance& WaveInstance = ActiveSound.AddWaveInstance(WaveInstanceHash);
@@ -1665,14 +1668,14 @@ bool USoundWave::IsReadyForFinishDestroy()
 {
 	{
 		FScopeLock Lock(&SourcesPlayingCs);
-		
+
 		for (ISoundWaveClient* SoundWaveClientPtr : SourcesPlaying)
 		{
 			// Should never have a null element here, inactive elements are removed from the container
 			check(SoundWaveClientPtr);
 			SoundWaveClientPtr->OnIsReadyForFinishDestroy(this);
 		}
-	} 
+	}
 
 	const bool bIsStreamingInProgress = IStreamingManager::Get().GetAudioStreamingManager().IsStreamingInProgress(this);
 
@@ -1680,16 +1683,16 @@ bool USoundWave::IsReadyForFinishDestroy()
 
 	// Wait till streaming and decompression finishes before deleting resource.
 	if (!bIsStreamingInProgress && ResourceState == ESoundWaveResourceState::NeedsFree)
-		{
-			DECLARE_CYCLE_STAT(TEXT("FAudioThreadTask.FreeResources"), STAT_AudioFreeResources, STATGROUP_AudioThreadCommands);
+	{
+		DECLARE_CYCLE_STAT(TEXT("FAudioThreadTask.FreeResources"), STAT_AudioFreeResources, STATGROUP_AudioThreadCommands);
 
-			USoundWave* SoundWave = this;
-			ResourceState = ESoundWaveResourceState::Freeing;
-			FAudioThread::RunCommandOnAudioThread([SoundWave]()
-			{
-				SoundWave->FreeResources();
-			}, GET_STATID(STAT_AudioFreeResources));
-		}
+		USoundWave* SoundWave = this;
+		ResourceState = ESoundWaveResourceState::Freeing;
+		FAudioThread::RunCommandOnAudioThread([SoundWave]()
+		{
+			SoundWave->FreeResources();
+		}, GET_STATID(STAT_AudioFreeResources));
+	}
 
 	return ResourceState == ESoundWaveResourceState::Freed;
 }
@@ -1736,7 +1739,7 @@ void USoundWave::Parse(FAudioDevice* AudioDevice, const UPTRINT NodeWaveInstance
 	const bool bIsNewWave = WaveInstance == nullptr;
 
 	// Create a new WaveInstance if this SoundWave doesn't already have one associated with it.
-	if(!WaveInstance)
+	if (!WaveInstance)
 	{
 		if (!ActiveSound.bRadioFilterSelected)
 		{
@@ -1803,7 +1806,7 @@ void USoundWave::Parse(FAudioDevice* AudioDevice, const UPTRINT NodeWaveInstance
 
 		// Use values from "parsed/ propagated" sound class properties
 		float VolumeMultiplier = WaveInstance->GetVolumeMultiplier();
-		WaveInstance->SetVolumeMultiplier(VolumeMultiplier* SoundClassProperties->Volume);
+		WaveInstance->SetVolumeMultiplier(VolumeMultiplier * SoundClassProperties->Volume);
 		WaveInstance->SetPitch(WaveInstance->Pitch * SoundClassProperties->Pitch);
 
 		WaveInstance->SoundClassFilterFrequency = SoundClassProperties->LowPassFilterFrequency;
@@ -1903,7 +1906,7 @@ void USoundWave::Parse(FAudioDevice* AudioDevice, const UPTRINT NodeWaveInstance
 
 		// Apply stereo normalization to wave instances if enabled
 		if (ParseParams.bApplyNormalizationToStereoSounds && NumChannels == 2)
-	{
+		{
 			float WaveInstanceVolume = WaveInstance->GetVolume();
 			WaveInstance->SetVolume(WaveInstanceVolume * 0.5f);
 		}
@@ -1916,8 +1919,8 @@ void USoundWave::Parse(FAudioDevice* AudioDevice, const UPTRINT NodeWaveInstance
 	WaveInstance->ReverbSendLevelRange = ParseParams.ReverbSendLevelRange;
 	WaveInstance->ReverbSendLevelDistanceRange = ParseParams.ReverbSendLevelDistanceRange;
 
-	// Copy the submix send settings
-	WaveInstance->SubmixSendSettings = ParseParams.SubmixSendSettings;
+ 	// Copy the submix send settings
+ 	WaveInstance->SubmixSendSettings = ParseParams.SubmixSendSettings;
 
 	// Get the envelope follower settings
 	WaveInstance->EnvelopeFollowerAttackTime = ParseParams.EnvelopeFollowerAttackTime;
@@ -1967,7 +1970,7 @@ void USoundWave::Parse(FAudioDevice* AudioDevice, const UPTRINT NodeWaveInstance
 
 #if !NO_LOGGING
 	// Sanity check
-	if(NumChannels > 2 && WaveInstance->GetUseSpatialization() && !WaveInstance->bReportedSpatializationWarning)
+	if (NumChannels > 2 && WaveInstance->GetUseSpatialization() && !WaveInstance->bReportedSpatializationWarning)
 	{
 		static TSet<USoundWave*> ReportedSounds;
 		if (!ReportedSounds.Contains(this))
@@ -1986,17 +1989,17 @@ void USoundWave::Parse(FAudioDevice* AudioDevice, const UPTRINT NodeWaveInstance
 					if (UAudioComponent* AudioComponent = UAudioComponent::GetAudioComponentFromID(AudioComponentID))
 					{
 						AActor* SoundOwner = AudioComponent->GetOwner();
-						UE_LOG(LogAudio, Warning, TEXT( "%s Actor: %s AudioComponent: %s" ), *SoundWarningInfo, (SoundOwner ? *SoundOwner->GetName() : TEXT("None")), *AudioComponent->GetName() );
+						UE_LOG(LogAudio, Warning, TEXT("%s Actor: %s AudioComponent: %s"), *SoundWarningInfo, (SoundOwner ? *SoundOwner->GetName() : TEXT("None")), *AudioComponent->GetName());
 					}
 					else
 					{
-						UE_LOG(LogAudio, Warning, TEXT("%s"), *SoundWarningInfo );
+						UE_LOG(LogAudio, Warning, TEXT("%s"), *SoundWarningInfo);
 					}
 				});
 			}
 			else
 			{
-				UE_LOG(LogAudio, Warning, TEXT("%s"), *SoundWarningInfo );
+				UE_LOG(LogAudio, Warning, TEXT("%s"), *SoundWarningInfo);
 			}
 
 			ReportedSounds.Add(this);
@@ -2023,14 +2026,13 @@ bool USoundWave::IsStreaming(const TCHAR* PlatformName/* = nullptr */) const
 		return false;
 	}
 
-	return IsStreaming(*FPlatformCompressionUtilities::GetCookOverrides(PlatformName));
+	return IsStreaming(*FPlatformCompressionUtilities::GetCookOverrides());
 }
 
 bool USoundWave::IsStreaming(const FPlatformAudioCookOverrides& Overrides) const
 {
 	// We stream if (A) bStreaming is set to true, (B) bForceInline is false and either bUseLoadOnDemand was set to true in
 	// our cook overrides, or the AutoStreamingThreshold was set and this sound is longer than the auto streaming threshold.
-	
 	if (bStreaming)
 	{
 		return true;
@@ -2045,7 +2047,7 @@ bool USoundWave::IsStreaming(const FPlatformAudioCookOverrides& Overrides) const
 	}
 
 	// For stream caching, the auto streaming threshold is used to force sounds to be inlined:
-	const bool bUsesStreamCache = (Overrides.bUseStreamCaching && Duration > Overrides.AutoStreamingThreshold);
+	const bool bUsesStreamCache = Overrides.bUseStreamCaching;
 	const bool bOverAutoStreamingThreshold = (Overrides.AutoStreamingThreshold > SMALL_NUMBER  && Duration > Overrides.AutoStreamingThreshold);
 
 	return bUsesStreamCache || bOverAutoStreamingThreshold;
@@ -2053,7 +2055,9 @@ bool USoundWave::IsStreaming(const FPlatformAudioCookOverrides& Overrides) const
 
 bool USoundWave::ShouldUseStreamCaching() const
 {
-	return  FPlatformCompressionUtilities::IsCurrentPlatformUsingStreamCaching() && IsStreaming(nullptr);
+	const bool bPlatformUsingStreamCaching = FPlatformCompressionUtilities::IsCurrentPlatformUsingStreamCaching();
+	const bool bIsStreaming = IsStreaming(nullptr);
+	return  bPlatformUsingStreamCaching && bIsStreaming;
 }
 
 TArrayView<const uint8> USoundWave::GetZerothChunk(bool bForImmediatePlayback)
@@ -2552,6 +2556,47 @@ bool USoundWave::IsRetainingAudio()
 	return FirstChunk.IsValid();
 }
 
+void USoundWave::OverrideLoadingBehavior(ESoundWaveLoadingBehavior InLoadingBehavior)
+{
+	const ESoundWaveLoadingBehavior OldBehavior = GetLoadingBehavior(false);
+	const bool bAlreadySetToRetained = (OldBehavior == ESoundWaveLoadingBehavior::RetainOnLoad);
+	const bool bAlreadyLoaded = !HasAnyFlags(RF_NeedLoad);
+
+	// already set to the most aggressive (non-inline) option
+	if (bAlreadySetToRetained)
+	{
+		return;
+	}
+
+	// we don't want to retain in editor, we prime instead
+	if (GIsEditor && InLoadingBehavior == ESoundWaveLoadingBehavior::RetainOnLoad)
+	{
+		InLoadingBehavior = ESoundWaveLoadingBehavior::PrimeOnLoad;
+	}
+
+
+	// record the new loading behavior
+	// (if this soundwave isn't loaded yet, 
+	// CachedSoundWaveLoadingBehavior will take precedence when it does load)
+	LoadingBehavior = InLoadingBehavior;
+	CachedSoundWaveLoadingBehavior = InLoadingBehavior;
+	bLoadingBehaviorOverridden = true;
+
+	// Manually perform prime/retain on already loaded sound waves
+	if (bAlreadyLoaded && IsStreaming())
+	{
+		if (InLoadingBehavior == ESoundWaveLoadingBehavior::RetainOnLoad)
+		{
+			ConditionalPostLoad();
+			RetainCompressedAudio();
+		}
+		else if (InLoadingBehavior == ESoundWaveLoadingBehavior::PrimeOnLoad && GetNumChunks() > 1)
+		{
+			IStreamingManager::Get().GetAudioStreamingManager().RequestChunk(this, 1, [](EAudioChunkLoadResult) {});
+		}
+	}
+}
+
 void USoundWave::SetCacheLookupIDForChunk(uint32 InChunkIndex, uint64 InCacheLookupID)
 {
 	check(FPlatformCompressionUtilities::IsCurrentPlatformUsingStreamCaching());
@@ -2576,10 +2621,21 @@ void USoundWave::CacheInheritedLoadingBehavior()
 	if (LoadingBehavior != ESoundWaveLoadingBehavior::Inherited)
 	{
 		// If this sound wave specifies it's own loading behavior, use that.
-		CachedSoundWaveLoadingBehavior = LoadingBehavior;
+		if (CachedSoundWaveLoadingBehavior == ESoundWaveLoadingBehavior::Uninitialized)
+		{
+			CachedSoundWaveLoadingBehavior = LoadingBehavior;
+		}
 	}
+ 	else if (bLoadingBehaviorOverridden)
+ 	{
+ 		ensureMsgf(CachedSoundWaveLoadingBehavior != ESoundWaveLoadingBehavior::Inherited, TEXT("SoundCue set loading behavior to Inherited on SoudWave: %s"), *GetFullName());
+ 		LoadingBehavior = CachedSoundWaveLoadingBehavior;
+ 	}
 	else
 	{
+		// if this is true then the behavior should not be Inherited here
+		check(!bLoadingBehaviorOverridden);
+
 		USoundClass* CurrentSoundClass = SoundClassObject;
 		ESoundWaveLoadingBehavior SoundClassLoadingBehavior = ESoundWaveLoadingBehavior::Inherited;
 
@@ -2593,8 +2649,13 @@ void USoundWave::CacheInheritedLoadingBehavior()
 		// If we could not find an override in the sound class hierarchy, use the loading behavior defined by our cvar.
 		if (SoundClassLoadingBehavior == ESoundWaveLoadingBehavior::Inherited)
 		{
+			// query the default loading behavior CVar
 			ensureAlwaysMsgf(SoundWaveDefaultLoadingBehaviorCVar >= 0 && SoundWaveDefaultLoadingBehaviorCVar < 4, TEXT("Invalid default loading behavior CVar value. Use value 0, 1, 2 or 3."));
-			SoundClassLoadingBehavior = (ESoundWaveLoadingBehavior)FMath::Clamp<int32>(SoundWaveDefaultLoadingBehaviorCVar, 0, (int32)ESoundWaveLoadingBehavior::LoadOnDemand);
+			ESoundWaveLoadingBehavior DefaultLoadingBehavior = (ESoundWaveLoadingBehavior)FMath::Clamp<int32>(SoundWaveDefaultLoadingBehaviorCVar, 0, (int32)ESoundWaveLoadingBehavior::LoadOnDemand);
+
+			// override this loading behavior w/ our default
+			SoundClassLoadingBehavior = DefaultLoadingBehavior;
+			bLoadingBehaviorOverridden = true;
 		}
 
 		CachedSoundWaveLoadingBehavior = SoundClassLoadingBehavior;
@@ -2603,15 +2664,22 @@ void USoundWave::CacheInheritedLoadingBehavior()
 
 ESoundWaveLoadingBehavior USoundWave::GetLoadingBehavior(bool bCheckSoundClasses /*= true*/) const
 {
-	checkf(!bCheckSoundClasses || CachedSoundWaveLoadingBehavior != ESoundWaveLoadingBehavior::Uninitialized, TEXT("this call is only valid if bCheckSoundClasses is false or CacheInheritedLoadingBehavior has already been called on the game thread."));
+	checkf(!bCheckSoundClasses || CachedSoundWaveLoadingBehavior != ESoundWaveLoadingBehavior::Uninitialized,
+		TEXT("Calling GetLoadingBehavior() is only valid if bCheckSoundClasses is false (which it %s) or CacheInheritedLoadingBehavior has already been called on the game thread. (SoundWave: %s)")
+		, bCheckSoundClasses? "is not":"is", *GetFullName());
 
 	if (!bCheckSoundClasses)
 	{
-		if (LoadingBehavior != ESoundWaveLoadingBehavior::Inherited)
+		if ((LoadingBehavior != ESoundWaveLoadingBehavior::Inherited && !bLoadingBehaviorOverridden))
 		{
 			// If this sound wave specifies it's own loading behavior, use that.
 			return LoadingBehavior;
 		}
+		else if (bLoadingBehaviorOverridden)
+		{
+			// If this sound wave has already had it's loading behavior cached from soundclasses or soundcues, use that.
+			return CachedSoundWaveLoadingBehavior;
+		}
 		else
 		{
 			// Otherwise, use the loading behavior defined by our cvar.
diff --git a/Engine/Source/Runtime/Engine/Public/AudioStreamingCache.h b/Engine/Source/Runtime/Engine/Public/AudioStreamingCache.h
--- a/Engine/Source/Runtime/Engine/Public/AudioStreamingCache.h
+++ b/Engine/Source/Runtime/Engine/Public/AudioStreamingCache.h
@@ -127,6 +127,7 @@ private:
 
 		// Note the loading behavior of the sound wave that inserted this element into the cache
 		ESoundWaveLoadingBehavior LoadingBehavior;
+		bool bLoadingBehaviorExternallyOverriden;
 
 		// if true, 
 		bool bWasCacheMiss;
