Fix for 5/5 crash on Android. Don't retain through SoundCues for now until Max's fixes come in, since the CPL here could be called on a USoundWave that isn't finished loading yet.

diff --git a/Engine/Source/Runtime/Engine/Private/SoundNode.cpp b/Engine/Source/Runtime/Engine/Private/SoundNode.cpp
--- a/Engine/Source/Runtime/Engine/Private/SoundNode.cpp
+++ b/Engine/Source/Runtime/Engine/Private/SoundNode.cpp
@@ -9,6 +9,13 @@
 #include "ContentStreaming.h"
 #include "AudioCompressionSettingsUtils.h"
 
+static int32 BypassRetainInSoundNodesCVar = 0;
+FAutoConsoleVariableRef CVarBypassRetainInSoundNodes(
+	TEXT("au.streamcache.priming.BypassRetainFromSoundCues"),
+	BypassRetainInSoundNodesCVar,
+	TEXT("When set to 1, we ignore the loading behavior of sound classes set on a Sound Cue directly.\n"),
+	ECVF_Default);
+
 /*-----------------------------------------------------------------------------
 	USoundNode implementation.
 -----------------------------------------------------------------------------*/
@@ -102,7 +109,7 @@ void USoundNode::RetainChildWavePlayers(bool bRecurse)
 
 void USoundNode::OverrideLoadingBehaviorOnChildWaves(const bool bRecurse, const ESoundWaveLoadingBehavior InLoadingBehavior)
 {
-	if (FPlatformCompressionUtilities::IsCurrentPlatformUsingStreamCaching())
+	if (!BypassRetainInSoundNodesCVar && FPlatformCompressionUtilities::IsCurrentPlatformUsingStreamCaching())
 	{
 		// Search child nodes for wave players, then override their waves' loading behavior.
 		for (USoundNode* ChildNode : ChildNodes)
