Fixing typo in include.

Somehow windows is ok with the slash being in the wrong direction!

diff --git a/Engine/Source/Runtime/Engine/Classes/Sound/SoundGenerator.h b/Engine/Source/Runtime/Engine/Classes/Sound/SoundGenerator.h
--- a/Engine/Source/Runtime/Engine/Classes/Sound/SoundGenerator.h
+++ b/Engine/Source/Runtime/Engine/Classes/Sound/SoundGenerator.h
@@ -3,8 +3,8 @@
 #pragma once
 
 #include "CoreMinimal.h"
-#include "Containers\Queue.h"
-#include "Templates\Function.h"
+#include "Containers/Queue.h"
+#include "Templates/Function.h"
 
 enum class EGeneratorState : uint8
 {
