Stop repeated DDC generations on load

##rb none

diff --git a/Engine/Source/Runtime/Engine/Private/Animation/AnimSequence.cpp b/Engine/Source/Runtime/Engine/Private/Animation/AnimSequence.cpp
--- a/Engine/Source/Runtime/Engine/Private/Animation/AnimSequence.cpp
+++ b/Engine/Source/Runtime/Engine/Private/Animation/AnimSequence.cpp
@@ -1236,7 +1236,7 @@ void UAnimSequence::PostEditChangeProperty(FPropertyChangedEvent& PropertyChange
 	// @Todo fix me: This is temporary fix to make sure they always have compressed data
 	if (RawAnimationData.Num() > 0 && bNeedPostProcess)
 	{
-		PostProcessSequence();
+		PostProcessSequence(false);
 	}
 
 	if (PropertyChangedEvent.Property != nullptr)
