Fix TSan (UE-91609).

diff --git a/Engine/Source/Runtime/Engine/Private/TickTaskManager.cpp b/Engine/Source/Runtime/Engine/Private/TickTaskManager.cpp
--- a/Engine/Source/Runtime/Engine/Private/TickTaskManager.cpp
+++ b/Engine/Source/Runtime/Engine/Private/TickTaskManager.cpp
@@ -1994,7 +1994,7 @@ void FTickFunction::QueueTickFunctionParallel(const struct FTickContext& TickCon
 				{
 					FTickFunction* Prereq = Prerequisites[PrereqIndex].Get();
 #if USING_THREAD_SANITISER
-					if (Prereq) { TSAN_AFTER(&Prereq->TickQueuedGFrameCounter); }
+					if (Prereq) { TSAN_AFTER(&InternalData->TickQueuedGFrameCounter); }
 #endif
 					if (!Prereq)
 					{
