Fix inability to distinguish Windows and WindowsNoEditor material SMs (UE-90852).

- ShaderMapId being identical makes the in-memory case alias material SMs when caching resources for -targetplatform=Windows+WindowsNoEditor
- Not a full fix, next up - Niagara and OpenColorIOColor

diff --git a/Engine/Source/Developer/MaterialBaking/Private/ExportMaterialProxy.h b/Engine/Source/Developer/MaterialBaking/Private/ExportMaterialProxy.h
--- a/Engine/Source/Developer/MaterialBaking/Private/ExportMaterialProxy.h
+++ b/Engine/Source/Developer/MaterialBaking/Private/ExportMaterialProxy.h
@@ -234,7 +234,7 @@ public:
 		const FMaterialResource* Resource = InMaterialInterface->GetMaterialResource(GMaxRHIFeatureLevel);
 
 		FMaterialShaderMapId ResourceId;
-		Resource->GetShaderMapId(GMaxRHIShaderPlatform, ResourceId);
+		Resource->GetShaderMapId(GMaxRHIShaderPlatform, nullptr, ResourceId);
 
 		// Our Id must be the same as BaseMaterialId for the shader compiler
 		// to be able to set back GameThreadShaderMap after async compilation.
diff --git a/Engine/Source/Developer/MaterialUtilities/Private/MaterialUtilities.cpp b/Engine/Source/Developer/MaterialUtilities/Private/MaterialUtilities.cpp
--- a/Engine/Source/Developer/MaterialUtilities/Private/MaterialUtilities.cpp
+++ b/Engine/Source/Developer/MaterialUtilities/Private/MaterialUtilities.cpp
@@ -416,7 +416,7 @@ public:
 		FMaterialResource* Resource = InMaterialInterface->GetMaterialResource(GMaxRHIFeatureLevel);
 
 		FMaterialShaderMapId ResourceId;
-		Resource->GetShaderMapId(GMaxRHIShaderPlatform, ResourceId);
+		Resource->GetShaderMapId(GMaxRHIShaderPlatform, nullptr, ResourceId);
 
 		{
 			TArray<FShaderType*> ShaderTypes;
diff --git a/Engine/Source/Editor/UnrealEd/Private/Lightmass/LightmassRender.cpp b/Engine/Source/Editor/UnrealEd/Private/Lightmass/LightmassRender.cpp
--- a/Engine/Source/Editor/UnrealEd/Private/Lightmass/LightmassRender.cpp
+++ b/Engine/Source/Editor/UnrealEd/Private/Lightmass/LightmassRender.cpp
@@ -245,7 +245,7 @@ public:
 			if (Resource)
 			{
 				FMaterialShaderMapId ResourceId;
-				Resource->GetShaderMapId(GMaxRHIShaderPlatform, ResourceId);
+				Resource->GetShaderMapId(GMaxRHIShaderPlatform, nullptr, ResourceId);
 
 				{
 					TArray<FShaderType*> ShaderTypes;
diff --git a/Engine/Source/Runtime/Core/Public/Serialization/MemoryLayout.h b/Engine/Source/Runtime/Core/Public/Serialization/MemoryLayout.h
--- a/Engine/Source/Runtime/Core/Public/Serialization/MemoryLayout.h
+++ b/Engine/Source/Runtime/Core/Public/Serialization/MemoryLayout.h
@@ -94,6 +94,38 @@ struct FPlatformTypeLayoutParameters
 	 * This is mostly needed when fixing bugs related to writing memory images with certain parameter combinations
 	 */
 	CORE_API void AppendKeyString(FString& KeyString) const;
+
+	/**
+	 * This is used for serializing into/from the DDC
+	 */
+	void Serialize(FArchive& Ar)
+	{
+		// if you change this code, please bump MATERIALSHADERMAP_DERIVEDDATA_VER (see FMaterialShaderMap::Serialize)
+		// since this is a part of ShaderMapId
+		Ar << MaxFieldAlignment;
+		Ar << b32Bit;
+		Ar << bForce64BitMemoryImagePointers;
+		Ar << bAlignBases;
+		Ar << bWithEditorOnly;
+		Ar << bWithRayTracing;
+		Ar << bIsCurrentPlatform;
+	}
+
+	bool operator==(const FPlatformTypeLayoutParameters& ReferenceSet) const
+	{
+		return MaxFieldAlignment == ReferenceSet.MaxFieldAlignment &&
+			b32Bit == ReferenceSet.b32Bit &&
+			bForce64BitMemoryImagePointers == ReferenceSet.bForce64BitMemoryImagePointers &&
+			bAlignBases == ReferenceSet.bAlignBases &&
+			bWithEditorOnly == ReferenceSet.bWithEditorOnly &&
+			bWithRayTracing == ReferenceSet.bWithRayTracing &&
+			bIsCurrentPlatform == ReferenceSet.bIsCurrentPlatform;
+	}
+
+	bool operator!=(const FPlatformTypeLayoutParameters& ReferenceSet) const
+	{
+		return !(*this == ReferenceSet);
+	}
 };
 
 struct FMemoryToStringContext
diff --git a/Engine/Source/Runtime/Engine/Private/DebugViewModeMaterialProxy.cpp b/Engine/Source/Runtime/Engine/Private/DebugViewModeMaterialProxy.cpp
--- a/Engine/Source/Runtime/Engine/Private/DebugViewModeMaterialProxy.cpp
+++ b/Engine/Source/Runtime/Engine/Private/DebugViewModeMaterialProxy.cpp
@@ -67,7 +67,7 @@ FDebugViewModeMaterialProxy::FDebugViewModeMaterialProxy(
 		}
 
 		FMaterialShaderMapId ResourceId;
-		Resource->GetShaderMapId(ShaderPlatform, ResourceId);
+		Resource->GetShaderMapId(ShaderPlatform, nullptr, ResourceId);
 
 		{
 			TArray<FShaderType*> ShaderTypes;
diff --git a/Engine/Source/Runtime/Engine/Private/Materials/Material.cpp b/Engine/Source/Runtime/Engine/Private/Materials/Material.cpp
--- a/Engine/Source/Runtime/Engine/Private/Materials/Material.cpp
+++ b/Engine/Source/Runtime/Engine/Private/Materials/Material.cpp
@@ -263,9 +263,9 @@ void FMaterialResource::EndAllowCachingStaticParameterValues()
 }
 #endif // WITH_EDITOR
 
-void FMaterialResource::GetShaderMapId(EShaderPlatform Platform, FMaterialShaderMapId& OutId) const
+void FMaterialResource::GetShaderMapId(EShaderPlatform Platform, const ITargetPlatform* TargetPlatform, FMaterialShaderMapId& OutId) const
 {
-	FMaterial::GetShaderMapId(Platform, OutId);
+	FMaterial::GetShaderMapId(Platform, TargetPlatform, OutId);
 #if WITH_EDITOR
 	Material->AppendReferencedFunctionIdsTo(OutId.ReferencedFunctions);
 	Material->AppendReferencedParameterCollectionIdsTo(OutId.ReferencedParameterCollections);
diff --git a/Engine/Source/Runtime/Engine/Private/Materials/MaterialShader.cpp b/Engine/Source/Runtime/Engine/Private/Materials/MaterialShader.cpp
--- a/Engine/Source/Runtime/Engine/Private/Materials/MaterialShader.cpp
+++ b/Engine/Source/Runtime/Engine/Private/Materials/MaterialShader.cpp
@@ -12,7 +12,6 @@
 #include "ShaderCompiler.h"
 #include "DerivedDataCacheInterface.h"
 #include "Interfaces/ITargetPlatformManagerModule.h"
-#include "Interfaces/ITargetPlatform.h"
 #include "ShaderDerivedDataVersion.h"
 #include "ProfilingDebugging/CookStats.h"
 #include "UObject/ReleaseObjectVersion.h"
@@ -131,18 +130,8 @@ static FString GetMaterialShaderMapKeyString(const FMaterialShaderMapId& ShaderM
 	FName Format = LegacyShaderPlatformToShaderFormat(Platform);
 	FString ShaderMapKeyString = Format.ToString() + TEXT("_") + FString(FString::FromInt(GetTargetPlatformManagerRef().ShaderFormatVersion(Format))) + TEXT("_");
 
-	FPlatformTypeLayoutParameters LayoutParameters;
-	if (TargetPlatform)
-	{
-		LayoutParameters.InitializeForPlatform(TargetPlatform->IniPlatformName(), TargetPlatform->HasEditorOnlyData());
-	}
-	else
-	{
-		LayoutParameters.InitializeForCurrent();
-	}
-
 	ShaderMapAppendKeyString(Platform, ShaderMapKeyString);
-	ShaderMapId.AppendKeyString(LayoutParameters, ShaderMapKeyString);
+	ShaderMapId.AppendKeyString(ShaderMapKeyString);
 	FMaterialAttributeDefinitionMap::AppendDDCKeyString(ShaderMapKeyString);
 	return FDerivedDataCacheInterface::BuildCacheKey(TEXT("MATSM"), MATERIALSHADERMAP_DERIVEDDATA_VER, *ShaderMapKeyString);
 }
@@ -400,6 +389,18 @@ void FMaterialShaderMapId::Serialize(FArchive& Ar, bool bLoadedByCookedMaterial)
 	checkf(CookedShaderMapIdHash != FSHAHash(), TEXT("Loaded an invalid cooked shadermap id hash"));
 #endif // WITH_EDITOR
 
+	// Based on the comment above, ShaderMapId would only be embedded into packages for archives < VER_UE4_PURGED_FMATERIAL_COMPILE_OUTPUTS (~2013 era).
+	// Since backward compatibility only works for the packages, we can assume that any archive with a version newer than that will have
+	// the LayoutParams serialized. The other option (old DDC) should be prevented by us having mutated MATERIALSHADERMAP_DERIVEDDATA_VER
+	if (Ar.UE4Ver() >= VER_UE4_PURGED_FMATERIAL_COMPILE_OUTPUTS)
+	{
+		LayoutParams.Serialize(Ar);
+	}
+	else
+	{
+		LayoutParams.InitializeForCurrent();
+	}
+
 	// Ensure loaded content is correct
 	check(!Ar.IsLoading() || IsContentValid());
 }
@@ -502,6 +503,11 @@ bool FMaterialShaderMapId::operator==(const FMaterialShaderMapId& ReferenceSet)
 		return false;
 	}
 
+	if (LayoutParams != ReferenceSet.LayoutParams)
+	{
+		return false;
+	}
+
 #if WITH_EDITOR
 	if (!IsCookedId())
 	{
@@ -664,7 +670,7 @@ void FMaterialShaderMapId::UpdateFromParameterSet(const FStaticParameterSet& Sta
 #endif // WITH_EDITOR
 
 #if WITH_EDITOR
-void FMaterialShaderMapId::AppendKeyString(const FPlatformTypeLayoutParameters& LayoutParams, FString& KeyString) const
+void FMaterialShaderMapId::AppendKeyString(FString& KeyString) const
 {
 	check(IsContentValid());
 	KeyString += BaseMaterialId.ToString();
diff --git a/Engine/Source/Runtime/Engine/Private/Materials/MaterialShared.cpp b/Engine/Source/Runtime/Engine/Private/Materials/MaterialShared.cpp
--- a/Engine/Source/Runtime/Engine/Private/Materials/MaterialShared.cpp
+++ b/Engine/Source/Runtime/Engine/Private/Materials/MaterialShared.cpp
@@ -42,6 +42,7 @@
 #include "ProfilingDebugging/LoadTimeTracker.h"
 #include "UObject/CoreRedirects.h"
 #include "RayTracingDefinitions.h"
+#include "Interfaces/ITargetPlatform.h"
 
 #define LOCTEXT_NAMESPACE "MaterialShared"
 
@@ -513,7 +514,7 @@ int32 FMaterialAttributesInput::CompileWithDefault(class FMaterialCompiler* Comp
 }
 #endif  // WITH_EDITOR
 
-void FMaterial::GetShaderMapId(EShaderPlatform Platform, FMaterialShaderMapId& OutId) const
+void FMaterial::GetShaderMapId(EShaderPlatform Platform, const ITargetPlatform* TargetPlatform, FMaterialShaderMapId& OutId) const
 { 
 	if (bLoadedCookedShaderMapId)
 	{
@@ -545,11 +546,26 @@ void FMaterial::GetShaderMapId(EShaderPlatform Platform, FMaterialShaderMapId& O
 		OutId.FeatureLevel = GetFeatureLevel();
 		OutId.SetShaderDependencies(ShaderTypes, ShaderPipelineTypes, VFTypes, Platform);
 		GetReferencedTexturesHash(Platform, OutId.TextureReferencesHash);
+
+		if (TargetPlatform)
+		{
+			OutId.LayoutParams.InitializeForPlatform(TargetPlatform->IniPlatformName(), TargetPlatform->HasEditorOnlyData());
+		}
+		else
+		{
+			OutId.LayoutParams.InitializeForCurrent();
+		}
 #else
 		OutId.QualityLevel = GetQualityLevelForShaderMapId();
 		OutId.FeatureLevel = GetFeatureLevel();
 
-		UE_LOG(LogMaterial, Log, TEXT("Tried to access an uncooked shader map ID in a cooked application"));
+		if (TargetPlatform != nullptr)
+		{
+			UE_LOG(LogMaterial, Error, TEXT("FMaterial::GetShaderMapId: TargetPlatform is not null, but a cooked executable cannot target platforms other than its own."));
+		}
+		OutId.LayoutParams.InitializeForCurrent();
+
+		UE_LOG(LogMaterial, Error, TEXT("Tried to access an uncooked shader map ID in a cooked application"));
 #endif
 	}
 }
@@ -1717,7 +1733,7 @@ bool FMaterial::CacheShaders(EShaderPlatform Platform, const ITargetPlatform* Ta
 {
 	FAllowCachingStaticParameterValues AllowCachingStaticParameterValues(*this);
 	FMaterialShaderMapId NoStaticParametersId;
-	GetShaderMapId(Platform, NoStaticParametersId);
+	GetShaderMapId(Platform, TargetPlatform, NoStaticParametersId);
 	return CacheShaders(NoStaticParametersId, Platform, TargetPlatform);
 }
 
@@ -2912,7 +2928,7 @@ bool FMaterial::GetMaterialExpressionSource( FString& OutSource )
 
 	FMaterialCompilationOutput TempOutput;
 	FMaterialShaderMapId ShaderMapID;
-	GetShaderMapId(GMaxRHIShaderPlatform, ShaderMapID);
+	GetShaderMapId(GMaxRHIShaderPlatform, nullptr, ShaderMapID);
 	FStaticParameterSet StaticParamSet;
 	GetStaticParameterSet(GMaxRHIShaderPlatform, StaticParamSet);
 
@@ -3300,7 +3316,7 @@ void UMaterialInterface::AnalyzeMaterialProperty(EMaterialProperty InProperty, i
 	FMaterialCompilationOutput TempOutput;
 	FMaterialResource* MaterialResource = GetMaterialResource(GMaxRHIFeatureLevel);
 	FMaterialShaderMapId ShaderMapID;
-	MaterialResource->GetShaderMapId(GMaxRHIShaderPlatform, ShaderMapID);
+	MaterialResource->GetShaderMapId(GMaxRHIShaderPlatform, nullptr, ShaderMapID);
 	FStaticParameterSet StaticParamSet;
 	MaterialResource->GetStaticParameterSet(GMaxRHIShaderPlatform, StaticParamSet);
 	FMaterialAnalyzer MaterialTranslator(MaterialResource, TempOutput, StaticParamSet, GMaxRHIShaderPlatform, MaterialResource->GetQualityLevel(), GMaxRHIFeatureLevel);
@@ -3338,7 +3354,7 @@ bool UMaterialInterface::IsTextureReferencedByProperty(EMaterialProperty InPrope
 	FMaterialCompilationOutput TempOutput;
 	FMaterialResource* MaterialResource = GetMaterialResource(GMaxRHIFeatureLevel);
 	FMaterialShaderMapId ShaderMapID;
-	MaterialResource->GetShaderMapId(GMaxRHIShaderPlatform, ShaderMapID);
+	MaterialResource->GetShaderMapId(GMaxRHIShaderPlatform, nullptr, ShaderMapID);
 	FStaticParameterSet StaticParamSet;
 	MaterialResource->GetStaticParameterSet(GMaxRHIShaderPlatform, StaticParamSet);
 	FHLSLMaterialTranslator MaterialTranslator(MaterialResource, TempOutput, StaticParamSet, GMaxRHIShaderPlatform, MaterialResource->GetQualityLevel(), GMaxRHIFeatureLevel);
diff --git a/Engine/Source/Runtime/Engine/Private/ShaderDerivedDataVersion.h b/Engine/Source/Runtime/Engine/Private/ShaderDerivedDataVersion.h
--- a/Engine/Source/Runtime/Engine/Private/ShaderDerivedDataVersion.h
+++ b/Engine/Source/Runtime/Engine/Private/ShaderDerivedDataVersion.h
@@ -10,5 +10,5 @@
 
 // In case of merge conflicts with DDC versions, you MUST generate a new GUID and set this new
 // guid as version
-#define GLOBALSHADERMAP_DERIVEDDATA_VER			TEXT("F3CB60CC102B46E79BFDA8B2E2E05686")
-#define MATERIALSHADERMAP_DERIVEDDATA_VER		TEXT("89EAC9ABC0FD493FA7A9F6E2CDC7142C")
+#define GLOBALSHADERMAP_DERIVEDDATA_VER			TEXT("CE4A26A5071E4C9596A658A9AC69B29B")
+#define MATERIALSHADERMAP_DERIVEDDATA_VER		TEXT("EEC7Ef19B1AF435984C9604C80FF9350")
diff --git a/Engine/Source/Runtime/Engine/Public/MaterialShared.h b/Engine/Source/Runtime/Engine/Public/MaterialShared.h
--- a/Engine/Source/Runtime/Engine/Public/MaterialShared.h
+++ b/Engine/Source/Runtime/Engine/Public/MaterialShared.h
@@ -789,7 +789,12 @@ public:
 	/** A hash of the base property overrides for this material instance. */
 	FSHAHash BasePropertyOverridesHash;
 #endif // WITH_EDITOR
-	
+
+	/*
+	 * Type layout parameters of the memory image
+	 */
+	FPlatformTypeLayoutParameters LayoutParams;
+
 	FMaterialShaderMapId()
 		: QualityLevel(EMaterialQualityLevel::High)
 		, FeatureLevel(ERHIFeatureLevel::SM5)
@@ -875,7 +880,7 @@ public:
 	void UpdateFromParameterSet(const FStaticParameterSet& StaticParameters);
 
 	/** Appends string representations of this Id to a key string. */
-	void AppendKeyString(const FPlatformTypeLayoutParameters& LayoutParameters, FString& KeyString) const;
+	void AppendKeyString(FString& KeyString) const;
 
 	const TArray<FStaticSwitchParameter> &GetStaticSwitchParameters() const 					{ return StaticSwitchParameters; }
 	const TArray<FStaticComponentMaskParameter> &GetStaticComponentMaskParameters() const 		{ return StaticComponentMaskParameters; }
@@ -1028,7 +1033,7 @@ public:
 
 	/**
 	 * Finds the shader map for a material.
-	 * @param StaticParameterSet - The static parameter set identifying the shader map
+	 * @param ShaderMapId - The static parameter set and other properties identifying the shader map
 	 * @param Platform - The platform to lookup for
 	 * @return NULL if no cached shader map was found.
 	 */
@@ -1271,7 +1276,7 @@ private:
 	float CompileTime;
 #endif
 
-	/** The static parameter set that this shader map was compiled with */
+	/** The static parameter set that this shader map was compiled with and other parameters unique to this shadermap */
 	FMaterialShaderMapId ShaderMapId;
 
 	/** The platform being compiled, or nullptr for current platform */
@@ -1507,7 +1512,7 @@ public:
 	void DiscardShaderMap();
 
 	// Material properties.
-	ENGINE_API virtual void GetShaderMapId(EShaderPlatform Platform, FMaterialShaderMapId& OutId) const;
+	ENGINE_API virtual void GetShaderMapId(EShaderPlatform Platform, const ITargetPlatform* TargetPlatform, FMaterialShaderMapId& OutId) const;
 #if WITH_EDITORONLY_DATA
 	ENGINE_API virtual void GetStaticParameterSet(EShaderPlatform Platform, FStaticParameterSet& OutSet) const;
 #endif // WITH_EDITORONLY_DATA
@@ -2314,7 +2319,7 @@ public:
 	ENGINE_API virtual FString GetMaterialUsageDescription() const override;
 
 	// FMaterial interface.
-	ENGINE_API virtual void GetShaderMapId(EShaderPlatform Platform, FMaterialShaderMapId& OutId) const override;
+	ENGINE_API virtual void GetShaderMapId(EShaderPlatform Platform, const ITargetPlatform* TargetPlatform, FMaterialShaderMapId& OutId) const override;
 #if WITH_EDITORONLY_DATA
 	ENGINE_API virtual void GetStaticParameterSet(EShaderPlatform Platform, FStaticParameterSet& OutSet) const override;
 #endif
diff --git a/Engine/Source/Runtime/Landscape/Private/LandscapeRender.cpp b/Engine/Source/Runtime/Landscape/Private/LandscapeRender.cpp
--- a/Engine/Source/Runtime/Landscape/Private/LandscapeRender.cpp
+++ b/Engine/Source/Runtime/Landscape/Private/LandscapeRender.cpp
@@ -3635,9 +3635,9 @@ public:
 	{
 	}
 
-	void GetShaderMapId(EShaderPlatform Platform, FMaterialShaderMapId& OutId) const override
+	void GetShaderMapId(EShaderPlatform Platform, const ITargetPlatform* TargetPlatform, FMaterialShaderMapId& OutId) const override
 	{
-		FMaterialResource::GetShaderMapId(Platform, OutId);
+		FMaterialResource::GetShaderMapId(Platform, TargetPlatform, OutId);
 
 #if WITH_EDITOR
 		if (bIsLayerThumbnail || bDisableTessellation)
