Create net.DormancyHysteresis CVar that allows a timeout before actually putting Actors into Dormancy.
This can help in cases where certain actors are being flushed very frequently.

diff --git a/Engine/Source/Runtime/Engine/Classes/Engine/Channel.h b/Engine/Source/Runtime/Engine/Classes/Engine/Channel.h
--- a/Engine/Source/Runtime/Engine/Classes/Engine/Channel.h
+++ b/Engine/Source/Runtime/Engine/Classes/Engine/Channel.h
@@ -67,6 +67,7 @@ class ENGINE_API UChannel
 	uint32				Broken:1;			// Has encountered errors and is ignoring subsequent packets.
 	uint32				bTornOff:1;			// Actor associated with this channel was torn off
 	uint32				bPendingDormancy:1;	// Channel wants to go dormant (it will check during tick if it can go dormant)
+	uint32				bIsInDormancyHysteresis:1; // Channel wants to go dormant, and is otherwise ready to become dormant, but is waiting for a timeout before doing so.
 	uint32				bPausedUntilReliableACK:1; // Unreliable property replication is paused until all reliables are ack'd.
 	uint32				SentClosingBunch:1;	// Set when sending closing bunch to avoid recursion in send-failure-close case.
 	uint32				bPooled:1;			// Set when placed in the actor channel pool
diff --git a/Engine/Source/Runtime/Engine/Classes/Engine/NetDriver.h b/Engine/Source/Runtime/Engine/Classes/Engine/NetDriver.h
--- a/Engine/Source/Runtime/Engine/Classes/Engine/NetDriver.h
+++ b/Engine/Source/Runtime/Engine/Classes/Engine/NetDriver.h
@@ -1329,6 +1329,9 @@ public:
 	/** Flushes actor from NetDriver's dormancy list, but does not change any state on the Actor itself */
 	ENGINE_API void FlushActorDormancy(AActor *Actor, bool bWasDormInitial=false);
 
+	//~ This probably doesn't need to be exported, since it's only called by AActor::SetNetDormancy.
+
+	/** Notifies the NetDriver that the desired Dormancy state for this Actor has changed. */
 	ENGINE_API void NotifyActorDormancyChange(AActor* Actor, ENetDormancy OldDormancyState);
 
 	/** Forces properties on this actor to do a compare for one frame (rather than share shadow state) */
@@ -1494,6 +1497,7 @@ public:
 	/** Adds (fully initialized, ready to go) client connection to the ClientConnections list + any other game related setup */
 	ENGINE_API void	AddClientConnection(UNetConnection * NewConnection);
 
+	//~ This method should only be called by internal networking systems.
 	ENGINE_API void NotifyActorFullyDormantForConnection(AActor* Actor, UNetConnection* Connection);
 
 	/** Returns true if this actor is considered to be in a loaded level */
diff --git a/Engine/Source/Runtime/Engine/Private/DataChannel.cpp b/Engine/Source/Runtime/Engine/Private/DataChannel.cpp
--- a/Engine/Source/Runtime/Engine/Private/DataChannel.cpp
+++ b/Engine/Source/Runtime/Engine/Private/DataChannel.cpp
@@ -86,6 +86,13 @@ static FAutoConsoleVariableRef CVarNetMaxConstructedPartialBunchSizeBytes(
 	TEXT("The maximum size allowed for Partial Bunches.")
 );
 
+static float DormancyHysteresis = 0;
+static FAutoConsoleVariableRef CVarDormancyHysteresis(
+	TEXT("net.DormancyHysteresis"),
+	DormancyHysteresis,
+	TEXT("When > 0, represents the time we'll wait before letting a channel become fully dormant (in seconds). This can prevent churn when objects are going in and out of dormant more frequently than normal.")
+);
+
 template<typename T>
 static const bool IsBunchTooLarge(UNetConnection* Connection, T* Bunch)
 {
@@ -1354,6 +1361,7 @@ void UChannel::AddedToChannelPool()
 	Broken = false;
 	bTornOff = false;
 	bPendingDormancy = false;
+	bIsInDormancyHysteresis = false;
 	bPausedUntilReliableACK = false;
 	SentClosingBunch = false;
 	bOpenedForCheckpoint = false;
@@ -2945,6 +2953,11 @@ int64 UActorChannel::ReplicateActor()
 		GNumReplicateActorCalls++;
 	}
 
+	if (bIsInDormancyHysteresis)
+	{
+		return 0;
+	}
+
 	// triggering replication of an Actor while already in the middle of replication can result in invalid data being sent and is therefore illegal
 	if (bIsReplicatingActor)
 	{
@@ -2953,15 +2966,17 @@ int64 UActorChannel::ReplicateActor()
 		ensureMsgf(false, TEXT("%s"), *Error);
 		return 0;
 	}
-	else if (bActorIsPendingKill)
+
+	if (bActorIsPendingKill)
 	{
 		// Don't need to do anything, because it should have already been logged.
 		return 0;
 	}
+
 	// If our Actor is PendingKill, that's bad. It means that somehow it wasn't properly removed
 	// from the NetDriver or ReplicationDriver.
 	// TODO: Maybe notify the NetDriver / RepDriver about this, and have the channel close?
-	else if (Actor->IsPendingKillOrUnreachable())
+	if (Actor->IsPendingKillOrUnreachable())
 	{
 		bActorIsPendingKill = true;
 		ActorReplicator.Reset();
@@ -2970,7 +2985,8 @@ int64 UActorChannel::ReplicateActor()
 		ensureMsgf(false, TEXT("%s"), *Error);
 		return 0;
 	}
-	else if (bPausedUntilReliableACK)
+
+	if (bPausedUntilReliableACK)
 	{
 		if (NumOutRec > 0)
 		{
@@ -3347,6 +3363,7 @@ void UActorChannel::BecomeDormant()
 {
 	UE_LOG(LogNetDormancy, Verbose, TEXT("BecomeDormant: %s"), *Describe() );
 	bPendingDormancy = false;
+	bIsInDormancyHysteresis = false;
 	Dormant = true;
 	// Replays do not close dormant channels currently, so just mark the channel dormant directly.
 	if (Connection && Connection->IsInternalAck())
@@ -3361,13 +3378,29 @@ void UActorChannel::BecomeDormant()
 
 bool UActorChannel::ReadyForDormancy(bool suppressLogs)
 {
-	for (auto MapIt = ReplicationMap.CreateIterator(); MapIt; ++MapIt)
+	// We need to keep replicating the Actor and its subobjects until none of them have
+	// changes, and would otherwise go Dormant normally.
+	if (!bIsInDormancyHysteresis)
 	{
-		if (!MapIt.Value()->ReadyForDormancy(suppressLogs))
+		for (auto MapIt = ReplicationMap.CreateIterator(); MapIt; ++MapIt)
+		{
+			if (!MapIt.Value()->ReadyForDormancy(suppressLogs))
+			{
+				return false;
+			}
+		}
+	}
+
+	if (DormancyHysteresis > 0 && Connection && Connection->Driver)
+	{
+		bIsInDormancyHysteresis = true;
+		const double TimePassed = Connection->Driver->Time - LastUpdateTime;
+		if (TimePassed < DormancyHysteresis)
 		{
 			return false;
 		}
 	}
+
 	return true;
 }
 
@@ -3384,7 +3417,8 @@ void UActorChannel::StartBecomingDormant()
 	{
 		MapIt.Value()->StartBecomingDormant();
 	}
-	bPendingDormancy = 1;
+	bPendingDormancy = true;
+	bIsInDormancyHysteresis = false;
 	Connection->StartTickingChannel(this);
 }
 
diff --git a/Engine/Source/Runtime/Engine/Private/NetConnection.cpp b/Engine/Source/Runtime/Engine/Private/NetConnection.cpp
--- a/Engine/Source/Runtime/Engine/Private/NetConnection.cpp
+++ b/Engine/Source/Runtime/Engine/Private/NetConnection.cpp
@@ -3376,6 +3376,8 @@ void UNetConnection::Tick()
 		// We should never need more ticking channels than open channels
 		checkf(ChannelsToTick.Num() <= OpenChannels.Num(), TEXT("More ticking channels (%d) than open channels (%d) for net connection!"), ChannelsToTick.Num(), OpenChannels.Num())
 
+		QUICK_SCOPE_CYCLE_COUNTER(STAT_NetConnection_TickChannels);
+
 		// Tick the channels.
 		if (CVarTickAllOpenChannels.GetValueOnAnyThread() == 0)
 		{
@@ -3780,6 +3782,7 @@ void UNetConnection::FlushDormancy(class AActor* Actor)
 
 		Ch->Dormant = false;
 		Ch->bPendingDormancy = false;
+		Ch->bIsInDormancyHysteresis = false;
 	}
 
 }
@@ -3797,6 +3800,8 @@ void UNetConnection::ForcePropertyCompare( AActor* Actor )
 /** Wrapper for validating an objects dormancy state, and to prepare the object for replication again */
 void UNetConnection::FlushDormancyForObject( UObject* Object )
 {
+	QUICK_SCOPE_CYCLE_COUNTER(STAT_NetConnection_FlushDormancyForObject)
+
 	const bool ValidateProperties = (GNetDormancyValidate == 1);
 
 	TSharedRef< FObjectReplicator > * Replicator = DormantReplicatorMap.Find( Object );
diff --git a/Engine/Source/Runtime/Engine/Private/NetDriver.cpp b/Engine/Source/Runtime/Engine/Private/NetDriver.cpp
--- a/Engine/Source/Runtime/Engine/Private/NetDriver.cpp
+++ b/Engine/Source/Runtime/Engine/Private/NetDriver.cpp
@@ -2866,6 +2866,58 @@ void UNetDriver::HandlePacketLossBurstCommand( int32 DurationInMilliseconds )
 #endif
 }
 
+#if !UE_BUILD_SHIPPING
+void HandleNetFlushAllDormancy(UNetDriver* InNetDriver, UWorld* InWorld)
+{
+	QUICK_SCOPE_CYCLE_COUNTER(STAT_NetDriver_HandleNetFlushAllDormancy);
+
+	// TODO: Maybe prevent this from happening on Beacons or Demo Net Drivers?
+	if (InNetDriver && InWorld && InNetDriver->ServerConnection == nullptr)
+	{
+		const FNetworkObjectList& NetworkObjectList = InNetDriver->GetNetworkObjectList();
+		const int32 NumberOfActorsDormantOnAllConnections = NetworkObjectList.GetDormantObjectsOnAllConnections().Num();
+
+		TSet<AActor*> ActorsToReset;
+		ActorsToReset.Reserve(NumberOfActorsDormantOnAllConnections);
+
+		// Just iterate over all network objects here to make sure we get Actors that are Pending Dormancy.
+		for (const TSharedPtr<FNetworkObjectInfo>& NetworkObjectInfo : NetworkObjectList.GetAllObjects())
+		{
+			if (NetworkObjectInfo.IsValid())
+			{
+				if (AActor* Actor = NetworkObjectInfo->WeakActor.Get())
+				{
+					if (Actor->NetDormancy.GetValue() > (int32)DORM_Awake)
+					{
+						ActorsToReset.Add(Actor);
+					}
+				}
+			}
+		}
+
+		const double StartTime = FPlatformTime::Seconds();
+
+		for (AActor* Actor : ActorsToReset)
+		{
+			Actor->FlushNetDormancy();
+		}
+
+		const double TotalTime = FPlatformTime::Seconds() - StartTime;
+
+
+		float DormancyHysteresis = -1.f;
+		if (IConsoleVariable* CVarDormancyHysteresis = IConsoleManager::Get().FindConsoleVariable(TEXT("net.DormancyHysteresis")))
+		{
+			DormancyHysteresis = CVarDormancyHysteresis->GetFloat();
+		}
+
+		UE_LOG(LogNet, Warning, TEXT("HandleNetFlushAllDormancy: NumberOfActors: %d, NumberOfActorsDormantOnAllConnections: %d, TimeForFlush: %lf, DormancyHysteresis: %f"),
+			ActorsToReset.Num(), NumberOfActorsDormantOnAllConnections, TotalTime, DormancyHysteresis);
+	}
+}
+#endif
+
+
 bool UNetDriver::Exec( UWorld* InWorld, const TCHAR* Cmd, FOutputDevice& Ar )
 {
 #if !UE_BUILD_SHIPPING
@@ -2913,6 +2965,11 @@ bool UNetDriver::Exec( UWorld* InWorld, const TCHAR* Cmd, FOutputDevice& Ar )
 		HandlePacketLossBurstCommand(PacketLossBurstMilliseconds);
 		return true;
 	}
+	else if (FParse::Command(&Cmd, TEXT("FLUSHALLNETDORMANCY")))
+	{
+		HandleNetFlushAllDormancy(this, InWorld);
+		return true;
+	}
 	else
 #endif // !UE_BUILD_SHIPPING
 	{
@@ -4492,20 +4549,24 @@ struct FScopedNetDriverStats
 		// Set GReplicateActorTimingEnabled for this frame. (This will determine if actor channels do cycle counting while replicating)
 		GReplicateActorTimingEnabled = ShouldEnableScopeSecondsTimers();
 		if (GReplicateActorTimingEnabled)
-	{
-		GReplicateActorTimeSeconds = 0;
-		GNumReplicateActorCalls = 0;
-		GNumSaturatedConnections = 0;
+		{
+			GReplicateActorTimeSeconds = 0;
+			GNumReplicateActorCalls = 0;
+			GNumSaturatedConnections = 0;
 
-		GReplicationGatherPrioritizeTimeSeconds = 0.f;
-		GServerReplicateActorTimeSeconds = 0.f;
+			GReplicationGatherPrioritizeTimeSeconds = 0.f;
+			GServerReplicateActorTimeSeconds = 0.f;
 
-		// Whatever these values currently are were (mostly) set by RPCs (technically something else could have force ReplicateActor to be called but this is rare).
-		SET_DWORD_STAT(STAT_SharedSerializationRPCHit, GNumSharedSerializationHit);
-		SET_DWORD_STAT(STAT_SharedSerializationRPCMiss, GNumSharedSerializationMiss);
+			// Whatever these values currently are were (mostly) set by RPCs (technically something else could have force ReplicateActor to be called but this is rare).
+			SET_DWORD_STAT(STAT_SharedSerializationRPCHit, GNumSharedSerializationHit);
+			SET_DWORD_STAT(STAT_SharedSerializationRPCMiss, GNumSharedSerializationMiss);
 
-		StartTime = FPlatformTime::Seconds();
-		StartOutBytes = GNetOutBytes;
+			const FNetworkObjectList& NetworkObjectList = NetDriver->GetNetworkObjectList();
+			CSV_CUSTOM_STAT(Replication, NumberOfActiveActors, NetworkObjectList.GetActiveObjects().Num(), ECsvCustomStatOp::Set);
+			CSV_CUSTOM_STAT(Replication, NumberOfFullyDormantActors, NetworkObjectList.GetDormantObjectsOnAllConnections().Num(), ECsvCustomStatOp::Set);
+
+			StartTime = FPlatformTime::Seconds();
+			StartOutBytes = GNetOutBytes;
 		}
 	}
 
@@ -4513,34 +4574,34 @@ struct FScopedNetDriverStats
 	{
 		if (GReplicateActorTimingEnabled)
 		{
-		const double TotalTime = FPlatformTime::Seconds() - StartTime;
+			const double TotalTime = FPlatformTime::Seconds() - StartTime;
 
-		GServerReplicateActorTimeSeconds = TotalTime;
-		GReplicationGatherPrioritizeTimeSeconds = TotalTime - GReplicateActorTimeSeconds;
+			GServerReplicateActorTimeSeconds = TotalTime;
+			GReplicationGatherPrioritizeTimeSeconds = TotalTime - GReplicateActorTimeSeconds;
 
-		uint32 FrameOutBytes = GNetOutBytes - StartOutBytes;
+			uint32 FrameOutBytes = GNetOutBytes - StartOutBytes;
 
-		SET_FLOAT_STAT(STAT_NetServerGatherPrioritizeRepActorsTime, GReplicationGatherPrioritizeTimeSeconds * 1000.0);
-		SET_DWORD_STAT(STAT_NumReplicatedActors, GNumReplicateActorCalls);
-		SET_DWORD_STAT(STAT_NumSaturatedConnections, GNumSaturatedConnections);
+			SET_FLOAT_STAT(STAT_NetServerGatherPrioritizeRepActorsTime, GReplicationGatherPrioritizeTimeSeconds * 1000.0);
+			SET_DWORD_STAT(STAT_NumReplicatedActors, GNumReplicateActorCalls);
+			SET_DWORD_STAT(STAT_NumSaturatedConnections, GNumSaturatedConnections);
 
-		CSV_CUSTOM_STAT(Replication, ServerReplicateActorTimeMS, (float)(GServerReplicateActorTimeSeconds * 1000.0), ECsvCustomStatOp::Set );
-		CSV_CUSTOM_STAT(Replication, GatherPrioritizeTimeMS, (float)(GReplicationGatherPrioritizeTimeSeconds * 1000.0), ECsvCustomStatOp::Set );
-		CSV_CUSTOM_STAT(Replication, ReplicateActorTimeMS, (float)(GReplicateActorTimeSeconds * 1000.0), ECsvCustomStatOp::Set );
-		CSV_CUSTOM_STAT(Replication, NumReplicateActorCallsPerConAvg, ((float)GNumReplicateActorCalls)/(float)GNumClientConnections, ECsvCustomStatOp::Set );
-		CSV_CUSTOM_STAT(Replication, Connections, (float)GNumClientConnections, ECsvCustomStatOp::Set );
-		CSV_CUSTOM_STAT(Replication, SatConnections, (float)GNumSaturatedConnections, ECsvCustomStatOp::Set );
+			CSV_CUSTOM_STAT(Replication, ServerReplicateActorTimeMS, (float)(GServerReplicateActorTimeSeconds * 1000.0), ECsvCustomStatOp::Set );
+			CSV_CUSTOM_STAT(Replication, GatherPrioritizeTimeMS, (float)(GReplicationGatherPrioritizeTimeSeconds * 1000.0), ECsvCustomStatOp::Set );
+			CSV_CUSTOM_STAT(Replication, ReplicateActorTimeMS, (float)(GReplicateActorTimeSeconds * 1000.0), ECsvCustomStatOp::Set );
+			CSV_CUSTOM_STAT(Replication, NumReplicateActorCallsPerConAvg, ((float)GNumReplicateActorCalls)/(float)GNumClientConnections, ECsvCustomStatOp::Set );
+			CSV_CUSTOM_STAT(Replication, Connections, (float)GNumClientConnections, ECsvCustomStatOp::Set );
+			CSV_CUSTOM_STAT(Replication, SatConnections, (float)GNumSaturatedConnections, ECsvCustomStatOp::Set );
 			CSV_CUSTOM_STAT(Replication, OutKBytes, ((float)FrameOutBytes) / 1024.f, ECsvCustomStatOp::Set );
 			CSV_CUSTOM_STAT(Replication, OutNetGUIDKBytesSec, ((float)NetDriver->NetGUIDOutBytes / 1024.f), ECsvCustomStatOp::Set );
 			CSV_CUSTOM_STAT(Replication, NumClientUpdateLevelVisibility, ((float)GNumClientUpdateLevelVisibility), ECsvCustomStatOp::Set );
 				
 
-		SET_DWORD_STAT(STAT_SharedSerializationPropertyHit, GNumSharedSerializationHit);
-		SET_DWORD_STAT(STAT_SharedSerializationPropertyMiss, GNumSharedSerializationMiss);
+			SET_DWORD_STAT(STAT_SharedSerializationPropertyHit, GNumSharedSerializationHit);
+			SET_DWORD_STAT(STAT_SharedSerializationPropertyMiss, GNumSharedSerializationMiss);
 
-		// Note: we want to reset this at the end of the frame since the RPC stats are incremented at the top (recv)
-		GNumSharedSerializationHit = 0;
-		GNumSharedSerializationMiss = 0;
+			// Note: we want to reset this at the end of the frame since the RPC stats are incremented at the top (recv)
+			GNumSharedSerializationHit = 0;
+			GNumSharedSerializationMiss = 0;
 			GNumClientUpdateLevelVisibility = 0;
 		}
 	}
