More TSan fixage

diff --git a/Engine/Source/Runtime/Engine/Private/TickTaskManager.cpp b/Engine/Source/Runtime/Engine/Private/TickTaskManager.cpp
--- a/Engine/Source/Runtime/Engine/Private/TickTaskManager.cpp
+++ b/Engine/Source/Runtime/Engine/Private/TickTaskManager.cpp
@@ -2060,7 +2060,7 @@ void FTickFunction::QueueTickFunctionParallel(const struct FTickContext& TickCon
 		}
 		InternalData->bWasInterval = false;
 		
-		TSAN_BEFORE(&TickQueuedGFrameCounter);
+		TSAN_BEFORE(&InternalData->TickQueuedGFrameCounter);
 		FPlatformMisc::MemoryBarrier();
 		
 		// MSVC enforces acq/rel semantics on volatile values, but clang cannot (supports more backend architectures).
