Move distance field downsample task setup out of parallel for. It can only be called on render thread.

diff --git a/Engine/Source/Runtime/Engine/Private/DistanceFieldAtlas.cpp b/Engine/Source/Runtime/Engine/Private/DistanceFieldAtlas.cpp
--- a/Engine/Source/Runtime/Engine/Private/DistanceFieldAtlas.cpp
+++ b/Engine/Source/Runtime/Engine/Private/DistanceFieldAtlas.cpp
@@ -611,22 +611,29 @@ void FDistanceFieldVolumeTextureAtlas::UpdateAllocations(FRHICommandListImmediat
 				// FUpdateTexture3DData default constructor is private. it might not be used if copy is done on GPU
 				// Not sure we want to make it public. let's be conservative, and allocate on stack an array of its size
 				uint8 TextureUpdateDataStackMem[sizeof(FUpdateTexture3DData)];
-				FUpdateTexture3DData* TextureUpdateDataPtr = reinterpret_cast<FUpdateTexture3DData*>(&TextureUpdateDataStackMem);
+				FUpdateTexture3DData* AtlasUpdateDataPtr = reinterpret_cast<FUpdateTexture3DData*>(&TextureUpdateDataStackMem);
 
 				if (!bRuntimeDownsampling)
 				{
-					*TextureUpdateDataPtr = RHIBeginUpdateTexture3D(VolumeTextureRHI, 0, UpdateRegion);
+					*AtlasUpdateDataPtr = RHIBeginUpdateTexture3D(VolumeTextureRHI, 0, UpdateRegion);
 				}
-
-				if (bRuntimeDownsampling)
+				else
 				{
 					UpdateDataArray.Empty(LocalPendingAllocations->Num());
 					UpdateDataArray.AddUninitialized(LocalPendingAllocations->Num());
 					DownsamplingTasks.AddDefaulted(LocalPendingAllocations->Num());
+
+					for (int32 Idx = 0; Idx < LocalPendingAllocations->Num(); ++Idx)
+					{
+						FDistanceFieldDownsamplingDataTask& DownsamplingTask = DownsamplingTasks[Idx];
+						FUpdateTexture3DData& UpdateData = UpdateDataArray[Idx];
+						const FDistanceFieldVolumeTexture* Texture = (*LocalPendingAllocations)[Idx];
+						FDistanceFieldDownsampling::FillDownsamplingTask(Texture->VolumeData.Size, Texture->SizeInAtlas, Texture->GetAllocationMin(), Format, DownsamplingTask, UpdateData);
+					}
 				}
 
 				// @todo arne @todo beni: can you verify i properly merged the optimizations here and in Main?
-				ParallelFor(LocalPendingAllocations->Num(), [FormatSize, this, bDataIsCompressed, bRuntimeDownsampling, &LocalPendingAllocations, &TextureUpdateDataPtr, &DownsamplingTasks, &UpdateDataArray](int32 AllocationIndex)
+				ParallelFor(LocalPendingAllocations->Num(), [FormatSize, this, bDataIsCompressed, bRuntimeDownsampling, LocalPendingAllocations, AtlasUpdateDataPtr, &UpdateDataArray](int32 AllocationIndex)
 				{
 					TArray<uint8> UncompressedData;
 					FDistanceFieldVolumeTexture* Texture = (*LocalPendingAllocations)[AllocationIndex];
@@ -650,17 +657,18 @@ void FDistanceFieldVolumeTextureAtlas::UpdateAllocations(FRHICommandListImmediat
 						SourceDataPtr = &Texture->VolumeData.CompressedDistanceFieldVolume;
 					}
 
-					FIntVector DstOffset = FIntVector::ZeroValue;
+					FIntVector DstOffset;
+					FUpdateTexture3DData* TextureUpdateDataPtr;
 					
 					if (bRuntimeDownsampling)
 					{
-						FDistanceFieldDownsamplingDataTask& DownsamplingTask = DownsamplingTasks[AllocationIndex];
+						DstOffset = FIntVector::ZeroValue;
 						TextureUpdateDataPtr = &UpdateDataArray[AllocationIndex];
-						FDistanceFieldDownsampling::FillDownsamplingTask(Size, Texture->SizeInAtlas, Texture->GetAllocationMin(), Format, DownsamplingTask, *TextureUpdateDataPtr);
 					}
 					else
 					{
 						DstOffset = Texture->GetAllocationMin();
+						TextureUpdateDataPtr = AtlasUpdateDataPtr;
 					}
 					
 					CopyToUpdateTextureData(Size, FormatSize, *SourceDataPtr, *TextureUpdateDataPtr, DstOffset);
@@ -668,7 +676,7 @@ void FDistanceFieldVolumeTextureAtlas::UpdateAllocations(FRHICommandListImmediat
 
 				if (!bRuntimeDownsampling)
 				{
-					RHIEndUpdateTexture3D(*TextureUpdateDataPtr);
+					RHIEndUpdateTexture3D(*AtlasUpdateDataPtr);
 				}
 			}
 		}
