Fix deprecation warning.

#rnx
None

diff --git a/Engine/Source/Runtime/Engine/Private/DataChannel.cpp b/Engine/Source/Runtime/Engine/Private/DataChannel.cpp
--- a/Engine/Source/Runtime/Engine/Private/DataChannel.cpp
+++ b/Engine/Source/Runtime/Engine/Private/DataChannel.cpp
@@ -3394,7 +3394,7 @@ bool UActorChannel::ReadyForDormancy(bool suppressLogs)
 	if (DormancyHysteresis > 0 && Connection && Connection->Driver)
 	{
 		bIsInDormancyHysteresis = true;
-		const double TimePassed = Connection->Driver->Time - LastUpdateTime;
+		const double TimePassed = Connection->Driver->GetElapsedTime() - LastUpdateTime;
 		if (TimePassed < DormancyHysteresis)
 		{
 			return false;
