fix for cook failiure after stream caching updates due to ensure. Make sure we don't try to prime or retain chunks if we're not caching the streamed audio data (for example, when running the cook commandlet)

diff --git a/Engine/Source/Runtime/Engine/Private/SoundWave.cpp b/Engine/Source/Runtime/Engine/Private/SoundWave.cpp
--- a/Engine/Source/Runtime/Engine/Private/SoundWave.cpp
+++ b/Engine/Source/Runtime/Engine/Private/SoundWave.cpp
@@ -429,6 +429,8 @@ void USoundWave::Serialize( FArchive& Ar )
 
 	Ar << CompressedDataGuid;
 
+	bool bBuiltStreamedAudio = false;
+
 	if (bShouldStreamSound)
 	{
 		if (bCooked)
@@ -444,6 +446,7 @@ void USoundWave::Serialize( FArchive& Ar )
 		if (Ar.IsLoading() && !Ar.IsTransacting() && !bCooked && !GetOutermost()->HasAnyPackageFlags(PKG_ReloadingForCooker))
 		{
 			CachePlatformData(false);
+			bBuiltStreamedAudio = true;
 		}
 #endif // #if WITH_EDITORONLY_DATA
 	}
@@ -453,8 +456,10 @@ void USoundWave::Serialize( FArchive& Ar )
 		// For non-editor builds, we can immediately cache the sample rate.
 		SampleRate = GetSampleRateForCurrentPlatform();
 
+		const bool bShouldLoadChunks = bCooked || bBuiltStreamedAudio;
+
 		// If stream caching is enabled, here we determine if we should retain or prime this wave on load.
-		if (bShouldStreamSound && FPlatformCompressionUtilities::IsCurrentPlatformUsingStreamCaching())
+		if (bShouldStreamSound && bShouldLoadChunks && FPlatformCompressionUtilities::IsCurrentPlatformUsingStreamCaching())
 		{
 			ESoundWaveLoadingBehavior CurrentLoadingBehavior = GetLoadingBehavior(false);
 
