checks() designed to catch material parameter caching bugs switched to ensures, added more info on failure

diff --git a/Engine/Source/Runtime/Engine/Private/Materials/Material.cpp b/Engine/Source/Runtime/Engine/Private/Materials/Material.cpp
--- a/Engine/Source/Runtime/Engine/Private/Materials/Material.cpp
+++ b/Engine/Source/Runtime/Engine/Private/Materials/Material.cpp
@@ -2090,8 +2090,8 @@ bool UMaterial::GetScalarParameterValue(const FHashedMaterialParameterInfo& Para
 		float OldValue = 0.0f;
 		const bool bOldResult = GetScalarParameterValue_Legacy(ParameterInfo, OldValue, bOveriddenOnly);
 
-		check(bOldResult == bResult);
-		check(!bResult || OutValue == OldValue);
+		ensureMsgf(bOldResult == bResult, TEXT("UMaterial::GetScalarParameterValue() mismatch, bOldResult: %d, bResult: %d"), bOldResult, bResult);
+		ensureMsgf(!bResult || OutValue == OldValue, TEXT("UMaterial::GetScalarParameterValue() mismatch, OutValue: %g, OldValue: %g"), OutValue, OldValue);
 	}
 #endif
 	return bResult;
@@ -2234,8 +2234,8 @@ bool UMaterial::GetVectorParameterValue(const FHashedMaterialParameterInfo& Para
 		FLinearColor OldValue(ForceInitToZero);
 		const bool bOldResult = GetVectorParameterValue_Legacy(ParameterInfo, OldValue, bOveriddenOnly);
 
-		check(bOldResult == bResult);
-		check(!bResult || OutValue == OldValue);
+		ensureMsgf(bOldResult == bResult, TEXT("UMaterial::GetVectorParameterValue() mismatch, bOldResult: %d, bResult: %d"), bOldResult, bResult);
+		ensureMsgf(!bResult || OutValue == OldValue, TEXT("UMaterial::GetVectorParameterValue() mismatch, OutValue: %s, OldValue: %s"), *OutValue.ToString(), *OldValue.ToString());
 	}
 #endif
 	return bResult;
@@ -2362,8 +2362,10 @@ bool UMaterial::GetTextureParameterValue(const FHashedMaterialParameterInfo& Par
 		UTexture* OldValue = nullptr;
 		const bool bOldResult = GetTextureParameterValue_Legacy(ParameterInfo, OldValue, bOveriddenOnly);
 
-		check(bOldResult == bResult);
-		check(!bOldResult || OutValue == OldValue); // if result is false, texture value is undefined
+		ensureMsgf(bOldResult == bResult, TEXT("UMaterial::GetTextureParameterValue() mismatch, bOldResult: %d, bResult: %d"), bOldResult, bResult);
+		ensureMsgf(!bResult || OutValue == OldValue, TEXT("UMaterial::GetTextureParameterValue() mismatch, OutValue: %s, OldValue: %s"),
+			OutValue ? *OutValue->GetName() : TEXT("nullptr"),
+			OldValue ? *OldValue->GetName() : TEXT("nullptr"));
 	}
 #endif
 	return bResult;
@@ -2462,8 +2464,10 @@ bool UMaterial::GetRuntimeVirtualTextureParameterValue(const FHashedMaterialPara
 		URuntimeVirtualTexture* OldValue = nullptr;
 		const bool bOldResult = GetRuntimeVirtualTextureParameterValue_Legacy(ParameterInfo, OldValue, bOveriddenOnly);
 
-		check(bOldResult == bResult);
-		check(!bResult || OutValue == OldValue);
+		ensureMsgf(bOldResult == bResult, TEXT("UMaterial::GetRuntimeVirtualTextureParameterValue() mismatch, bOldResult: %d, bResult: %d"), bOldResult, bResult);
+		ensureMsgf(!bResult || OutValue == OldValue, TEXT("UMaterial::GetRuntimeVirtualTextureParameterValue() mismatch, OutValue: %s, OldValue: %s"),
+			OutValue ? *OutValue->GetName() : TEXT("nullptr"),
+			OldValue ? *OldValue->GetName() : TEXT("nullptr"));
 	}
 #endif
 	return bResult;
@@ -2576,8 +2580,10 @@ bool UMaterial::GetFontParameterValue(const FHashedMaterialParameterInfo& Parame
 		int32 OldPage = INDEX_NONE;
 		const bool bOldResult = GetFontParameterValue_Legacy(ParameterInfo, OldValue, OldPage, bOveriddenOnly);
 
-		check(bOldResult == bResult);
-		check(!bResult || (OutFontValue == OldValue && OldPage == OutFontPage));
+		ensureMsgf(bOldResult == bResult, TEXT("UMaterial::GetFontParameterValue() mismatch, bOldResult: %d, bResult: %d"), bOldResult, bResult);
+		ensureMsgf(!bResult || (OutFontValue == OldValue && OldPage == OutFontPage), TEXT("UMaterial::GetFontParameterValue() mismatch, OutValue: %s, OutPage: %d OldValue: %s, OldPage: %d"),
+			OutFontValue ? *OutFontValue->GetName() : TEXT("nullptr"), OutFontPage,
+			OldValue ? *OldValue->GetName() : TEXT("nullptr"), OldPage);
 	}
 #endif
 	return bResult;
