CIS Fix: Missing include for non unity builds

diff --git a/Engine/Source/Runtime/Engine/Private/Materials/Material.cpp b/Engine/Source/Runtime/Engine/Private/Materials/Material.cpp
--- a/Engine/Source/Runtime/Engine/Private/Materials/Material.cpp
+++ b/Engine/Source/Runtime/Engine/Private/Materials/Material.cpp
@@ -64,6 +64,7 @@
 #include "UObject/ReleaseObjectVersion.h"
 #include "RenderCore/Public/RenderUtils.h"
 #include "Materials/MaterialStaticParameterValueResolver.h"
+#include "Engine/Font.h"
 
 #if WITH_EDITOR
 #include "Logging/TokenizedMessage.h"
