Fix CIS Issue

#rnx

diff --git a/Engine/Source/Runtime/Engine/Private/NetDriver.cpp b/Engine/Source/Runtime/Engine/Private/NetDriver.cpp
--- a/Engine/Source/Runtime/Engine/Private/NetDriver.cpp
+++ b/Engine/Source/Runtime/Engine/Private/NetDriver.cpp
@@ -2905,14 +2905,14 @@ void HandleNetFlushAllDormancy(UNetDriver* InNetDriver, UWorld* InWorld)
 		const double TotalTime = FPlatformTime::Seconds() - StartTime;
 
 
-		float DormancyHysteresis = -1.f;
-		if (IConsoleVariable* CVarDormancyHysteresis = IConsoleManager::Get().FindConsoleVariable(TEXT("net.DormancyHysteresis")))
+		float CurrentDormancyHysteresis = -1.f;
+		if (IConsoleVariable* CurrentDormancyHysteresisCVar = IConsoleManager::Get().FindConsoleVariable(TEXT("net.DormancyHysteresis")))
 		{
-			DormancyHysteresis = CVarDormancyHysteresis->GetFloat();
+			CurrentDormancyHysteresis = CurrentDormancyHysteresisCVar->GetFloat();
 		}
 
 		UE_LOG(LogNet, Warning, TEXT("HandleNetFlushAllDormancy: NumberOfActors: %d, NumberOfActorsDormantOnAllConnections: %d, TimeForFlush: %lf, DormancyHysteresis: %f"),
-			ActorsToReset.Num(), NumberOfActorsDormantOnAllConnections, TotalTime, DormancyHysteresis);
+			ActorsToReset.Num(), NumberOfActorsDormantOnAllConnections, TotalTime, CurrentDormancyHysteresis);
 	}
 }
 #endif
