Fix blend space marker sync when markers land on animation boundaries

diff --git a/Engine/Source/Runtime/Engine/Private/Animation/AnimSequence.cpp b/Engine/Source/Runtime/Engine/Private/Animation/AnimSequence.cpp
--- a/Engine/Source/Runtime/Engine/Private/Animation/AnimSequence.cpp
+++ b/Engine/Source/Runtime/Engine/Private/Animation/AnimSequence.cpp
@@ -4833,8 +4833,8 @@ void UAnimSequence::AdvanceMarkerPhaseAsLeader(bool bLooping, float MoveDelta, c
 					MarkerTimeOffset = SequenceLength;
 				}
 			}
-			const float NextMarkerTime = NextSyncMarker.Time + MarkerTimeOffset;
-			const float TimeToMarker = NextMarkerTime - CurrentTime;
+
+			const float TimeToMarker = NextMarker.TimeToMarker;
 
 			if (CurrentMoveDelta > TimeToMarker)
 			{
@@ -4862,6 +4862,7 @@ void UAnimSequence::AdvanceMarkerPhaseAsLeader(bool bLooping, float MoveDelta, c
 						MarkerTimeOffset += SequenceLength;
 					}
 				} while (!ValidMarkerNames.Contains(AuthoredSyncMarkers[NextMarker.MarkerIndex].MarkerName));
+				NextMarker.TimeToMarker = MarkerTimeOffset + AuthoredSyncMarkers[NextMarker.MarkerIndex].Time - CurrentTime;
 			}
 			else
 			{
@@ -4898,8 +4899,8 @@ void UAnimSequence::AdvanceMarkerPhaseAsLeader(bool bLooping, float MoveDelta, c
 					MarkerTimeOffset = -SequenceLength;
 				}
 			}
-			const float PrevMarkerTime = PrevSyncMarker.Time + MarkerTimeOffset;
-			const float TimeToMarker = PrevMarkerTime - CurrentTime;
+
+			const float TimeToMarker = NextMarker.TimeToMarker;
 
 			if (CurrentMoveDelta < TimeToMarker)
 			{
@@ -4927,6 +4928,7 @@ void UAnimSequence::AdvanceMarkerPhaseAsLeader(bool bLooping, float MoveDelta, c
 						MarkerTimeOffset -= SequenceLength;
 					}
 				} while (!ValidMarkerNames.Contains(AuthoredSyncMarkers[PrevMarker.MarkerIndex].MarkerName));
+				PrevMarker.TimeToMarker = MarkerTimeOffset + AuthoredSyncMarkers[PrevMarker.MarkerIndex].Time - CurrentTime;
 			}
 			else
 			{
@@ -5225,8 +5227,15 @@ FMarkerSyncAnimPosition UAnimSequence::GetMarkerSyncPositionfromMarkerIndicies(i
 	}
 
 	// Account for looping
-	PrevTime = (PrevTime > CurrentTime) ? PrevTime - SequenceLength : PrevTime;
-	NextTime = (NextTime < CurrentTime) ? NextTime + SequenceLength : NextTime;
+	if(PrevTime > NextTime)
+	{
+		PrevTime = (PrevTime > CurrentTime) ? PrevTime - SequenceLength : PrevTime;
+		NextTime = (NextTime < CurrentTime) ? NextTime + SequenceLength : NextTime;
+	}
+	else if (PrevTime > CurrentTime)
+	{
+		CurrentTime += SequenceLength;
+	}
 
 	if (PrevTime == NextTime)
 	{
