Fixed potential assert when rendering ISM with ray tracing enabled if the ISM proxy had no instances

diff --git a/Engine/Source/Runtime/Engine/Private/InstancedStaticMesh.cpp b/Engine/Source/Runtime/Engine/Private/InstancedStaticMesh.cpp
--- a/Engine/Source/Runtime/Engine/Private/InstancedStaticMesh.cpp
+++ b/Engine/Source/Runtime/Engine/Private/InstancedStaticMesh.cpp
@@ -1134,11 +1134,6 @@ void FInstancedStaticMeshSceneProxy::GetDynamicRayTracingInstances(struct FRayTr
 	{
 		return;
 	}
-	
-	if (!InstancedRenderData.PerInstanceRenderData.IsValid())
-	{
-		return;
-	}
 
 	uint32 LOD = GetCurrentFirstLODIdx_RenderThread();
 	const int32 InstanceCount = InstancedRenderData.PerInstanceRenderData->InstanceBuffer.GetNumInstances();
@@ -1155,8 +1150,7 @@ void FInstancedStaticMeshSceneProxy::GetDynamicRayTracingInstances(struct FRayTr
 	//preallocate the worst-case to prevent an explosion of reallocs
 	//#dxr_todo: possibly track used instances and reserve based on previous behavior
 	RayTracingInstanceTemplate.InstanceTransforms.Reserve(InstanceCount);
-	
-	FMatrix ToWorld = InstancedRenderData.Component->GetComponentTransform().ToMatrixWithScale();
+
 	if (CVarRayTracingRenderInstancesCulling.GetValueOnRenderThread() > 0 && RayTracingCullClusters.Num() > 0)
 	{
 		const float BVHCullRadius = CVarRayTracingInstancesCullClusterRadius.GetValueOnRenderThread();
@@ -1164,6 +1158,7 @@ void FInstancedStaticMeshSceneProxy::GetDynamicRayTracingInstances(struct FRayTr
 		const float BVHLowScaleRadius = CVarRayTracingInstancesLowScaleCullRadius.GetValueOnRenderThread();
 		const bool ApplyGeneralCulling = BVHCullRadius > 0.0f;
 		const bool ApplyLowScaleCulling = BVHLowScaleThreshold > 0.0f && BVHLowScaleRadius > 0.0f;
+		FMatrix ToWorld = GetLocalToWorld();
 
 		// Iterate over all culling clusters
 		for (int32 ClusterIdx = 0; ClusterIdx < RayTracingCullClusters.Num(); ++ClusterIdx)
@@ -1187,9 +1182,6 @@ void FInstancedStaticMeshSceneProxy::GetDynamicRayTracingInstances(struct FRayTr
 			{
 				const uint32 InstanceIdx = Node.Instance;
 
-				const FInstancedStaticMeshInstanceData& InstanceData = InstancedRenderData.Component->PerInstanceSMData[InstanceIdx];
-				FMatrix InstanceTransform = InstanceData.Transform * ToWorld;
-				InstanceTransform.M[3][3] = 1.0f;
 				FVector InstanceLocation = Node.Center;
 				FVector VToInstanceCenter = Context.ReferenceView->ViewLocation - InstanceLocation;
 				float DistanceToInstanceCenter = VToInstanceCenter.Size();
@@ -1207,6 +1199,11 @@ void FInstancedStaticMeshSceneProxy::GetDynamicRayTracingInstances(struct FRayTr
 						continue;
 				}
 
+				FMatrix Transform;
+				InstancedRenderData.PerInstanceRenderData->InstanceBuffer.GetInstanceTransform(InstanceIdx, Transform);
+				Transform.M[3][3] = 1.0f;
+				FMatrix InstanceTransform = Transform * GetLocalToWorld();
+
 				RayTracingInstanceTemplate.InstanceTransforms.Add(InstanceTransform);
 			}
 		}
@@ -1216,9 +1213,10 @@ void FInstancedStaticMeshSceneProxy::GetDynamicRayTracingInstances(struct FRayTr
 		// No culling
 		for (int32 InstanceIdx = 0; InstanceIdx < InstanceCount; ++InstanceIdx)
 		{
-			const FInstancedStaticMeshInstanceData& InstanceData = InstancedRenderData.Component->PerInstanceSMData[InstanceIdx];
-			FMatrix InstanceTransform = InstanceData.Transform * ToWorld;
-			InstanceTransform.M[3][3] = 1.0f;
+			FMatrix Transform;
+			InstancedRenderData.PerInstanceRenderData->InstanceBuffer.GetInstanceTransform(InstanceIdx, Transform);
+			Transform.M[3][3] = 1.0f;
+			FMatrix InstanceTransform = Transform * GetLocalToWorld();
 
 			RayTracingInstanceTemplate.InstanceTransforms.Add(InstanceTransform);
 		}
@@ -1250,13 +1248,18 @@ void FInstancedStaticMeshSceneProxy::SetupRayTracingCullClusters()
 		return;
 	}
 
+	const int32 InstanceCount = InstancedRenderData.Component->PerInstanceSMData.Num();
+	if (InstanceCount <= 0)
+	{
+		return;
+	}
+
 	//#dxr_todo: select the appropriate LOD depending on Context.View
 	int32 LOD = 0;
 	if (RenderData->LODResources.Num() > LOD && RenderData->LODResources[LOD].RayTracingGeometry.IsInitialized())
 	{
 		const float MaxClusterRadiusMultiplier = CVarRayTracingInstancesCullClusterMaxRadiusMultiplier.GetValueOnAnyThread();
 		const int32 Batches = GetNumMeshBatches();
-		const int32 InstanceCount = InstancedRenderData.Component->PerInstanceSMData.Num();
 		int32 ClusterIndex = 0;
 		// We're in game thread and at this point this scene proxy hasn't been added to FScene, thus GetLocalToWorld() returns undefined transform
 		FMatrix ComponentLocalToWorld = InstancedRenderData.Component->GetComponentTransform().ToMatrixWithScale();
