Fix view target blending not working on the dedicated server when bUseClientSideCameraUpdates is on

diff --git a/Engine/Source/Runtime/Engine/Classes/Camera/PlayerCameraManager.h b/Engine/Source/Runtime/Engine/Classes/Camera/PlayerCameraManager.h
--- a/Engine/Source/Runtime/Engine/Classes/Camera/PlayerCameraManager.h
+++ b/Engine/Source/Runtime/Engine/Classes/Camera/PlayerCameraManager.h
@@ -943,6 +943,12 @@ protected:
 	 */
 	virtual void UpdateViewTargetInternal(FTViewTarget& OutVT, float DeltaTime);
 
+	// ClientSide camera updates prevents DoUpdateCamera from swapping PendingViewTarget in when the blend is complete, just use a timer to swap
+	UFUNCTION()
+	void SwapPendingViewTargetWhenUsingClientSideCameraUpdates();
+
+	FTimerHandle SwapPendingViewTargetWhenUsingClientSideCameraUpdatesTimerHandle;
+
 private:
 	// Buried to prevent use; use GetCameraRotation instead
 	FRotator GetActorRotation() const { return Super::GetActorRotation(); }
diff --git a/Engine/Source/Runtime/Engine/Private/PlayerCameraManager.cpp b/Engine/Source/Runtime/Engine/Private/PlayerCameraManager.cpp
--- a/Engine/Source/Runtime/Engine/Private/PlayerCameraManager.cpp
+++ b/Engine/Source/Runtime/Engine/Private/PlayerCameraManager.cpp
@@ -94,6 +94,16 @@ APlayerController* APlayerCameraManager::GetOwningPlayerController() const
 	return PCOwner;
 }
 
+void APlayerCameraManager::SwapPendingViewTargetWhenUsingClientSideCameraUpdates()
+{
+	if (PendingViewTarget.Target)
+	{
+		AssignViewTarget(PendingViewTarget.Target, ViewTarget);
+		ViewTarget.CheckViewTarget(PCOwner);
+		// remove old pending ViewTarget so we don't still try to switch to it
+		PendingViewTarget.Target = NULL;
+	}
+}
 
 void APlayerCameraManager::SetViewTarget(class AActor* NewTarget, struct FViewTargetTransitionParams TransitionParams)
 {
@@ -116,6 +126,11 @@ void APlayerCameraManager::SetViewTarget(class AActor* NewTarget, struct FViewTa
 		return;
 	}
 
+	if (UWorld* World = GetWorld())
+	{
+		World->GetTimerManager().ClearTimer(SwapPendingViewTargetWhenUsingClientSideCameraUpdatesTimerHandle);
+	}
+
 	// if viewtarget different then new one or we're transitioning from the same target with locked outgoing, then assign it
 	if((NewTarget != ViewTarget.Target) || (PendingViewTarget.Target && BlendParams.bLockOutgoing))
 	{
@@ -135,6 +150,14 @@ void APlayerCameraManager::SetViewTarget(class AActor* NewTarget, struct FViewTa
 			
 			AssignViewTarget(NewTarget, PendingViewTarget, TransitionParams);
 			PendingViewTarget.CheckViewTarget(PCOwner);
+
+			if (bUseClientSideCameraUpdates && GetNetMode() != NM_Client)
+			{
+				if (UWorld* World = GetWorld())
+				{
+					World->GetTimerManager().SetTimer(SwapPendingViewTargetWhenUsingClientSideCameraUpdatesTimerHandle, this, &ThisClass::SwapPendingViewTargetWhenUsingClientSideCameraUpdates, TransitionParams.BlendTime, false);
+				}
+			}
 		}
 		else
 		{
