#UE4 IWYU fix

diff --git a/Engine/Source/Runtime/Engine/Private/PlayerCameraManager.cpp b/Engine/Source/Runtime/Engine/Private/PlayerCameraManager.cpp
--- a/Engine/Source/Runtime/Engine/Private/PlayerCameraManager.cpp
+++ b/Engine/Source/Runtime/Engine/Private/PlayerCameraManager.cpp
@@ -21,6 +21,7 @@
 #include "GameFramework/PlayerState.h"
 #include "IXRTrackingSystem.h" // for IsHeadTrackingAllowed()
 #include "GameFramework/GameNetworkManager.h"
+#include "TimerManager.h"
 
 DEFINE_LOG_CATEGORY_STATIC(LogPlayerCameraManager, Log, All);
 
