Fix for ensure hitting in USoundCue::PostLoad when using the cook commandlet.

diff --git a/Engine/Source/Runtime/Engine/Private/SoundWave.cpp b/Engine/Source/Runtime/Engine/Private/SoundWave.cpp
--- a/Engine/Source/Runtime/Engine/Private/SoundWave.cpp
+++ b/Engine/Source/Runtime/Engine/Private/SoundWave.cpp
@@ -969,6 +969,11 @@ uint32 USoundWave::GetNumChunks() const
 	{
 		return 0;
 	}
+	else if (GetOutermost()->HasAnyPackageFlags(PKG_ReloadingForCooker))
+	{
+		UE_LOG(LogAudio, Warning, TEXT("USoundWave::GetNumChunks called in the cook commandlet."))
+		return 0;
+	}
 	else
 	{
 		ensureAlwaysMsgf(false, TEXT("Call CachePlatformData(false) before calling this function in editor. GetNumChunks() called on: %s"), *GetName());
@@ -2587,8 +2592,11 @@ void USoundWave::OverrideLoadingBehavior(ESoundWaveLoadingBehavior InLoadingBeha
 	CachedSoundWaveLoadingBehavior = InLoadingBehavior;
 	bLoadingBehaviorOverridden = true;
 
+	// If we're reloading for the cook commandlet, we don't have streamed audio chunks to load.
+	const bool bReloadingForCooker = GetOutermost()->HasAnyPackageFlags(PKG_ReloadingForCooker);
+
 	// Manually perform prime/retain on already loaded sound waves
-	if (bAlreadyLoaded && IsStreaming())
+	if (!bReloadingForCooker && bAlreadyLoaded && IsStreaming())
 	{
 		if (InLoadingBehavior == ESoundWaveLoadingBehavior::RetainOnLoad)
 		{
