Added a total count to stat drawcount and lowered the update rate to make the numbers more stable.
#RNX

diff --git a/Engine/Source/Runtime/Engine/Private/UnrealEngine.cpp b/Engine/Source/Runtime/Engine/Private/UnrealEngine.cpp
--- a/Engine/Source/Runtime/Engine/Private/UnrealEngine.cpp
+++ b/Engine/Source/Runtime/Engine/Private/UnrealEngine.cpp
@@ -15159,10 +15159,12 @@ int32 UEngine::RenderStatUnit(UWorld* World, FViewport* Viewport, FCanvas* Canva
 int32 UEngine::RenderStatDrawCount(UWorld* World, FViewport* Viewport, FCanvas* Canvas, int32 X, int32 Y, const FVector* ViewLocation, const FRotator* ViewRotation)
 {
 #if CSV_PROFILER
+	int32 TotalCount = 0;
 	// Display all the categories of draw counts. This may always report 0 in some modes if AreGPUStatsEnabled is not enabled.
 	// Most likely because we are not currently capturing a CSV.
 	for (int32 Index = 0; Index < FDrawCallCategoryName::NumCategory; ++Index)
 	{
+		TotalCount += FDrawCallCategoryName::DisplayCounts[Index];
 		FDrawCallCategoryName* CategoryName = FDrawCallCategoryName::Array[Index];
 		Canvas->DrawShadowedString(X - 50,
 			Y,
@@ -15171,6 +15173,12 @@ int32 UEngine::RenderStatDrawCount(UWorld* World, FViewport* Viewport, FCanvas*
 			FColor::Green);
 		Y += 12;
 	}
+	Canvas->DrawShadowedString(X - 50,
+		Y,
+		*FString::Printf(TEXT("Total: %i"), TotalCount),
+		GetSmallFont(),
+		FColor::Green);
+	Y += 12;
 #else
 	Canvas->DrawShadowedString(X - 200,
 		Y,
diff --git a/Engine/Source/Runtime/RHI/Private/RHI.cpp b/Engine/Source/Runtime/RHI/Private/RHI.cpp
--- a/Engine/Source/Runtime/RHI/Private/RHI.cpp
+++ b/Engine/Source/Runtime/RHI/Private/RHI.cpp
@@ -789,10 +789,24 @@ void RHIPrivateBeginFrame()
 	GNumDrawCallsRHI = GCurrentNumDrawCallsRHI;
 	
 #if CSV_PROFILER
+	// Only copy the display counters every so many frames to keep things more stable.
+	const int32 FramesUntilDisplayCopy = 30;
+	static int32 FrameCount = 0;
+	bool bCopyDisplayFrames = false;
+	++FrameCount;
+	if (FrameCount >= FramesUntilDisplayCopy)
+	{
+		bCopyDisplayFrames = true;
+		FrameCount = 0;
+	}
+
 	for (int32 Index=0; Index<FDrawCallCategoryName::NumCategory; ++Index)
 	{
 		FDrawCallCategoryName* CategoryName = FDrawCallCategoryName::Array[Index];
-		FDrawCallCategoryName::DisplayCounts[Index] = CategoryName->Counter;
+		if (bCopyDisplayFrames)
+		{
+			FDrawCallCategoryName::DisplayCounts[Index] = CategoryName->Counter;
+		}
 		GNumDrawCallsRHI += CategoryName->Counter;
 		FCsvProfiler::RecordCustomStat(CategoryName->Name, CSV_CATEGORY_INDEX(DrawCall), CategoryName->Counter, ECsvCustomStatOp::Set);
 		CategoryName->Counter = 0;
