Fix deprecation comment
#rnx

diff --git a/Engine/Source/Runtime/Engine/Private/UnrealEngine.cpp b/Engine/Source/Runtime/Engine/Private/UnrealEngine.cpp
--- a/Engine/Source/Runtime/Engine/Private/UnrealEngine.cpp
+++ b/Engine/Source/Runtime/Engine/Private/UnrealEngine.cpp
@@ -15875,7 +15875,7 @@ int32 UEngine::RenderStatSoundModulators(UWorld* World, FViewport* Viewport, FCa
 {
 #if ENABLE_AUDIO_DEBUG
 	const int32 FontHeight = GetStatsFont()->GetMaxCharHeight() + 2.0f;
-	Canvas->DrawShadowedString(X, Y, TEXT("'stat soundmodulators' is deprecated. Use au.Debug.Modulators"), GetStatsFont(), FLinearColor::Yellow);
+	Canvas->DrawShadowedString(X, Y, TEXT("'stat soundmodulators' is deprecated. Use au.Debug.SoundModulators"), GetStatsFont(), FLinearColor::Yellow);
 	Y += FontHeight;
 
 	return Audio::FAudioDebugger::RenderStatModulators(World, Viewport, Canvas, X, Y, ViewLocation, ViewRotation);
