Added cache size to error log for Audio Streaming Cache getting blown.
#rnx

diff --git a/Engine/Source/Runtime/Engine/Private/AudioStreamingCache.cpp b/Engine/Source/Runtime/Engine/Private/AudioStreamingCache.cpp
--- a/Engine/Source/Runtime/Engine/Private/AudioStreamingCache.cpp
+++ b/Engine/Source/Runtime/Engine/Private/AudioStreamingCache.cpp
@@ -1189,7 +1189,7 @@ FAudioChunkCache::FCacheElement* FAudioChunkCache::EvictLeastRecentChunk()
 		// If we ever hit this, it means that we couldn't find any cache elements that aren't in use.
 		if (!CacheElement || CacheElement == ElementToStopAt)
 		{
-			ensureMsgf(false, TEXT("Cache blown! Please increase the cache size or load less audio."));
+			ensureMsgf(false, TEXT("Cache blown! Please increase the cache size (currently %i bytes) or load less audio."), ReportCacheSize());
 			return nullptr;
 		}
 	}
