Multiple stability fixes for VLM mobile path

diff --git a/Engine/Source/Runtime/Engine/Private/PrecomputedVolumetricLightmap.cpp b/Engine/Source/Runtime/Engine/Private/PrecomputedVolumetricLightmap.cpp
--- a/Engine/Source/Runtime/Engine/Private/PrecomputedVolumetricLightmap.cpp
+++ b/Engine/Source/Runtime/Engine/Private/PrecomputedVolumetricLightmap.cpp
@@ -304,6 +304,7 @@ FPrecomputedVolumetricLightmapData::FPrecomputedVolumetricLightmapData()
 	, BrickSize(0)
 	, BrickDataDimensions(EForceInit::ForceInitToZero)
 	, BrickDataBaseOffsetInAtlas(0)
+	, IndexInCPUSubLevelBrickDataList(INDEX_NONE)
 {}
 
 FPrecomputedVolumetricLightmapData::~FPrecomputedVolumetricLightmapData()
@@ -497,12 +498,12 @@ ENGINE_API void FPrecomputedVolumetricLightmapData::AddToSceneData(FPrecomputedV
 		return;
 	}
 
-	SceneDataAdded.Add(SceneData);
-
 	FRHICommandListImmediate& RHICmdList = FRHICommandListExecutor::GetImmediateCommandList();
 
 	if (IndirectionTextureDimensions.GetMax() > 0)
 	{
+		SceneDataAdded.Add(SceneData);
+
 		// Copy parameters from the persistent level VLM
 		SceneData->Bounds = Bounds;
 		SceneData->BrickSize = BrickSize;
@@ -513,6 +514,20 @@ ENGINE_API void FPrecomputedVolumetricLightmapData::AddToSceneData(FPrecomputedV
 
 		if (GetFeatureLevel() >= ERHIFeatureLevel::SM5)
 		{
+			// Tracking down UE-88573
+			ensureMsgf(IsInitialized(), TEXT("FPrecomputedVolumetricLightmapData is being added to a scene data without initialization"));
+
+			if (!IsInitialized())
+			{
+				InitResource();
+			}
+
+			if (!IndirectionTexture.Texture)
+			{
+				ensureMsgf(IsInitialized(), TEXT("FPrecomputedVolumetricLightmapData IndirectionTexture is still invalid after manual initialization, returning"));
+				return;
+			}
+
 			// GPU Path
 
 			const int32 PaddedBrickSize = BrickSize + 1;
@@ -558,7 +573,7 @@ ENGINE_API void FPrecomputedVolumetricLightmapData::AddToSceneData(FPrecomputedV
 			SceneData->CPUSubLevelIndirectionTable.Empty();
 			SceneData->CPUSubLevelIndirectionTable.AddZeroed(IndirectionTextureDimensions.X * IndirectionTextureDimensions.Y * IndirectionTextureDimensions.Z);
 			SceneData->CPUSubLevelBrickDataList.Empty();
-			SceneData->CPUSubLevelBrickDataList.Add(this);
+			IndexInCPUSubLevelBrickDataList = SceneData->CPUSubLevelBrickDataList.Add(this);
 		}
 	}
 	else
@@ -568,6 +583,8 @@ ENGINE_API void FPrecomputedVolumetricLightmapData::AddToSceneData(FPrecomputedV
 			// GPU Path
 			if (SceneData->IndirectionTexture.Texture.IsValid())
 			{
+				SceneDataAdded.Add(SceneData);
+
 				InitRHIForSubLevelResources();
 
 				FGlobalShaderMap* GlobalShaderMap = GetGlobalShaderMap(GetFeatureLevel());
@@ -595,9 +612,11 @@ ENGINE_API void FPrecomputedVolumetricLightmapData::AddToSceneData(FPrecomputedV
 			// CPU Path
 			if (SceneData->IndirectionTexture.Data.Num() > 0)
 			{
-				int32 Index = SceneData->CPUSubLevelBrickDataList.Add(this);
-				check(Index < UINT8_MAX);
-				uint8 Value = (uint8)Index;
+				SceneDataAdded.Add(SceneData);
+
+				IndexInCPUSubLevelBrickDataList = SceneData->CPUSubLevelBrickDataList.Add(this);
+				check(IndexInCPUSubLevelBrickDataList < UINT8_MAX);
+				uint8 Value = (uint8)IndexInCPUSubLevelBrickDataList;
 
 				for (int32 BrickIndex = 0; BrickIndex < SubLevelBrickPositions.Num(); BrickIndex++)
 				{
@@ -684,28 +703,32 @@ ENGINE_API void FPrecomputedVolumetricLightmapData::RemoveFromSceneData(FPrecomp
 			// CPU Path
 			if (SceneData->IndirectionTexture.Data.Num() > 0)
 			{
-				SceneData->CPUSubLevelBrickDataList.Remove(this);
-
-				for (int32 BrickIndex = 0; BrickIndex < SubLevelBrickPositions.Num(); BrickIndex++)
+				ensure(SceneData->CPUSubLevelBrickDataList[IndexInCPUSubLevelBrickDataList] == this);
+				if (SceneData->CPUSubLevelBrickDataList[IndexInCPUSubLevelBrickDataList] == this)
 				{
-					const FColor OriginalValue = IndirectionTextureOriginalValues[BrickIndex];
-
-					const FIntVector IndirectionDestDataCoordinate = SubLevelBrickPositions[BrickIndex];
-					const int32 IndirectionDestDataIndex =
-						((IndirectionDestDataCoordinate.Z * SceneData->IndirectionTextureDimensions.Y) + IndirectionDestDataCoordinate.Y) *
-						SceneData->IndirectionTextureDimensions.X + IndirectionDestDataCoordinate.X;
-
-					{
-						const int32 IndirectionTextureDataStride = GPixelFormats[SceneData->IndirectionTexture.Format].BlockBytes;
-						uint8* IndirectionVoxelPtr = (uint8*)&SceneData->IndirectionTexture.Data[IndirectionDestDataIndex * IndirectionTextureDataStride];
-						*(IndirectionVoxelPtr + 0) = OriginalValue.R;
-						*(IndirectionVoxelPtr + 1) = OriginalValue.G;
-						*(IndirectionVoxelPtr + 2) = OriginalValue.B;
-						*(IndirectionVoxelPtr + 3) = 1;
-					}
+					SceneData->CPUSubLevelBrickDataList.RemoveAt(IndexInCPUSubLevelBrickDataList);
 
+					for (int32 BrickIndex = 0; BrickIndex < SubLevelBrickPositions.Num(); BrickIndex++)
 					{
-						SceneData->CPUSubLevelIndirectionTable[IndirectionDestDataIndex] = 0;
+						const FColor OriginalValue = IndirectionTextureOriginalValues[BrickIndex];
+
+						const FIntVector IndirectionDestDataCoordinate = SubLevelBrickPositions[BrickIndex];
+						const int32 IndirectionDestDataIndex =
+							((IndirectionDestDataCoordinate.Z * SceneData->IndirectionTextureDimensions.Y) + IndirectionDestDataCoordinate.Y) *
+							SceneData->IndirectionTextureDimensions.X + IndirectionDestDataCoordinate.X;
+
+						{
+							const int32 IndirectionTextureDataStride = GPixelFormats[SceneData->IndirectionTexture.Format].BlockBytes;
+							uint8* IndirectionVoxelPtr = (uint8*)&SceneData->IndirectionTexture.Data[IndirectionDestDataIndex * IndirectionTextureDataStride];
+							*(IndirectionVoxelPtr + 0) = OriginalValue.R;
+							*(IndirectionVoxelPtr + 1) = OriginalValue.G;
+							*(IndirectionVoxelPtr + 2) = OriginalValue.B;
+							*(IndirectionVoxelPtr + 3) = 1;
+						}
+
+						{
+							SceneData->CPUSubLevelIndirectionTable[IndirectionDestDataIndex] = 0;
+						}
 					}
 				}
 			}
diff --git a/Engine/Source/Runtime/Engine/Public/PrecomputedVolumetricLightmap.h b/Engine/Source/Runtime/Engine/Public/PrecomputedVolumetricLightmap.h
--- a/Engine/Source/Runtime/Engine/Public/PrecomputedVolumetricLightmap.h
+++ b/Engine/Source/Runtime/Engine/Public/PrecomputedVolumetricLightmap.h
@@ -206,7 +206,8 @@ public:
 
 	// CPU indirection table for mobile path
 	TArray<uint8> CPUSubLevelIndirectionTable;
-	TArray<FPrecomputedVolumetricLightmapData*> CPUSubLevelBrickDataList;
+	TSparseArray<FPrecomputedVolumetricLightmapData*> CPUSubLevelBrickDataList;
+	int32 IndexInCPUSubLevelBrickDataList;
 
 private:
 
diff --git a/Engine/Source/Runtime/Renderer/Private/RendererScene.cpp b/Engine/Source/Runtime/Renderer/Private/RendererScene.cpp
--- a/Engine/Source/Runtime/Renderer/Private/RendererScene.cpp
+++ b/Engine/Source/Runtime/Renderer/Private/RendererScene.cpp
@@ -1132,6 +1132,7 @@ FScene::FScene(UWorld* InWorld, bool bInRequiresHitProxies, bool bInIsEditorScen
 ,	SimpleDirectionalLight(NULL)
 ,	ReflectionSceneData(InFeatureLevel)
 ,	IndirectLightingCache(InFeatureLevel)
+,	VolumetricLightmapSceneData(this)
 ,	DistanceFieldSceneData(GShaderPlatformForFeatureLevel[InFeatureLevel])
 ,	PreshadowCacheLayout(0, 0, 0, 0, false)
 ,	AtmosphericFog(NULL)
@@ -2377,6 +2378,28 @@ void FVolumetricLightmapSceneData::RemoveLevelVolume(const FPrecomputedVolumetri
 	}
 }
 
+const FPrecomputedVolumetricLightmap* FVolumetricLightmapSceneData::GetLevelVolumetricLightmap() const
+{
+	return &GlobalVolumetricLightmap;
+}
+
+bool FVolumetricLightmapSceneData::HasData() const
+{
+	if (LevelVolumetricLightmaps.Num() > 0)
+	{
+		if (Scene->GetFeatureLevel() >= ERHIFeatureLevel::SM5)
+		{
+			return GlobalVolumetricLightmapData.IndirectionTexture.Texture.IsValid();
+		}
+		else
+		{
+			return GlobalVolumetricLightmapData.IndirectionTexture.Data.Num() > 0;
+		}
+	}
+
+	return false;
+}
+
 bool FScene::HasPrecomputedVolumetricLightmap_RenderThread() const
 {
 	return VolumetricLightmapSceneData.HasData();
diff --git a/Engine/Source/Runtime/Renderer/Private/ScenePrivate.h b/Engine/Source/Runtime/Renderer/Private/ScenePrivate.h
--- a/Engine/Source/Runtime/Renderer/Private/ScenePrivate.h
+++ b/Engine/Source/Runtime/Renderer/Private/ScenePrivate.h
@@ -1631,19 +1631,22 @@ class FVolumetricLightmapSceneData
 {
 public:
 
-	FVolumetricLightmapSceneData() { GlobalVolumetricLightmap.Data = &GlobalVolumetricLightmapData; }
-	bool HasData() const { return LevelVolumetricLightmaps.Num() > 0; }
+	FVolumetricLightmapSceneData(FScene* InScene)
+		: Scene(InScene)
+	{
+		GlobalVolumetricLightmap.Data = &GlobalVolumetricLightmapData;
+	}
+
+	bool HasData() const;
 	void AddLevelVolume(const class FPrecomputedVolumetricLightmap* InVolume, EShadingPath ShadingPath, bool bIsPersistentLevel);
 	void RemoveLevelVolume(const class FPrecomputedVolumetricLightmap* InVolume);
-	const FPrecomputedVolumetricLightmap* GetLevelVolumetricLightmap() const 
-	{ 
-		return &GlobalVolumetricLightmap;
-	}
+	const FPrecomputedVolumetricLightmap* GetLevelVolumetricLightmap() const;
 
 	TMap<FVector, FVolumetricLightmapInterpolation> CPUInterpolationCache;
 
 	FPrecomputedVolumetricLightmapData GlobalVolumetricLightmapData;
 private:
+	FScene* Scene;
 	FPrecomputedVolumetricLightmap GlobalVolumetricLightmap;
 	const FPrecomputedVolumetricLightmap* PersistentLevelVolumetricLightmap = nullptr;
 	TArray<const FPrecomputedVolumetricLightmap*> LevelVolumetricLightmaps;
