Fix VLM crash when changing preview feature level

diff --git a/Engine/Source/Runtime/Engine/Private/PrecomputedVolumetricLightmap.cpp b/Engine/Source/Runtime/Engine/Private/PrecomputedVolumetricLightmap.cpp
--- a/Engine/Source/Runtime/Engine/Private/PrecomputedVolumetricLightmap.cpp
+++ b/Engine/Source/Runtime/Engine/Private/PrecomputedVolumetricLightmap.cpp
@@ -422,32 +422,35 @@ ENGINE_API void FPrecomputedVolumetricLightmapData::HandleDataMovementInAtlas(in
 
 		for (FPrecomputedVolumetricLightmapData* SceneData : SceneDataAdded)
 		{
-			FGlobalShaderMap* GlobalShaderMap = GetGlobalShaderMap(GetFeatureLevel());
+			if (SceneData->IndirectionTexture.Texture.IsValid())
+			{
+				FGlobalShaderMap* GlobalShaderMap = GetGlobalShaderMap(GetFeatureLevel());
 
-			TShaderMapRef<FMoveWholeIndirectionTextureCS> ComputeShader(GlobalShaderMap);
+				TShaderMapRef<FMoveWholeIndirectionTextureCS> ComputeShader(GlobalShaderMap);
 
-			FVolumetricLightmapDataLayer NewIndirectionTexture = SceneData->IndirectionTexture;
-			NewIndirectionTexture.CreateTargetTexture(IndirectionTextureDimensions);
-			NewIndirectionTexture.CreateUAV();
+				FVolumetricLightmapDataLayer NewIndirectionTexture = SceneData->IndirectionTexture;
+				NewIndirectionTexture.CreateTargetTexture(IndirectionTextureDimensions);
+				NewIndirectionTexture.CreateUAV();
 
-			FMoveWholeIndirectionTextureCS::FParameters Parameters;
-			Parameters.NumBricks = NumBricks;
-			Parameters.StartPosInOldVolume = OldOffset;
-			Parameters.StartPosInNewVolume = BrickDataBaseOffsetInAtlas;
-			Parameters.OldIndirectionTexture = SceneData->IndirectionTexture.Texture;
-			Parameters.IndirectionTexture = NewIndirectionTexture.UAV;
+				FMoveWholeIndirectionTextureCS::FParameters Parameters;
+				Parameters.NumBricks = NumBricks;
+				Parameters.StartPosInOldVolume = OldOffset;
+				Parameters.StartPosInNewVolume = BrickDataBaseOffsetInAtlas;
+				Parameters.OldIndirectionTexture = SceneData->IndirectionTexture.Texture;
+				Parameters.IndirectionTexture = NewIndirectionTexture.UAV;
 
-			FComputeShaderUtils::Dispatch(RHICmdList, ComputeShader, Parameters,
-				FIntVector(
-					FMath::DivideAndRoundUp(IndirectionTextureDimensions.X, 4),
-					FMath::DivideAndRoundUp(IndirectionTextureDimensions.Y, 4),
-					FMath::DivideAndRoundUp(IndirectionTextureDimensions.Z, 4))
-			);
+				FComputeShaderUtils::Dispatch(RHICmdList, ComputeShader, Parameters,
+					FIntVector(
+						FMath::DivideAndRoundUp(IndirectionTextureDimensions.X, 4),
+						FMath::DivideAndRoundUp(IndirectionTextureDimensions.Y, 4),
+						FMath::DivideAndRoundUp(IndirectionTextureDimensions.Z, 4))
+				);
 
-			SceneData->IndirectionTexture = NewIndirectionTexture;
+				SceneData->IndirectionTexture = NewIndirectionTexture;
 
- 			FRHIUnorderedAccessView* UAV = NewIndirectionTexture.UAV;
- 			RHICmdList.TransitionResources(EResourceTransitionAccess::EReadable, EResourceTransitionPipeline::EComputeToGfx, &UAV, 1);
+				FRHIUnorderedAccessView* UAV = NewIndirectionTexture.UAV;
+				RHICmdList.TransitionResources(EResourceTransitionAccess::EReadable, EResourceTransitionPipeline::EComputeToGfx, &UAV, 1);
+			}
 		}
 	}
 	else
