Remove unintentional debug message

diff --git a/Engine/Source/Runtime/Engine/Private/World.cpp b/Engine/Source/Runtime/Engine/Private/World.cpp
--- a/Engine/Source/Runtime/Engine/Private/World.cpp
+++ b/Engine/Source/Runtime/Engine/Private/World.cpp
@@ -4225,7 +4225,6 @@ void UWorld::CleanupWorldInternal(bool bSessionEnded, bool bCleanupResources, UW
 		return;
 	}
 	bool bWorldChanged = NewWorld != this && NewWorld != nullptr;
-	FPlatformMisc::LowLevelOutputDebugStringf(TEXT("XYXYXY XYXYXY Clearnupworld %p -> %p\n"), this, NewWorld);
 	CleanupWorldTag = CleanupWorldGlobalTag;
 
 	UE_LOG(LogWorld, Log, TEXT("UWorld::CleanupWorld for %s, bSessionEnded=%s, bCleanupResources=%s"), *GetName(), bSessionEnded ? TEXT("true") : TEXT("false"), bCleanupResources ? TEXT("true") : TEXT("false"));
