Fixed potential race condition when trying to render data that was not ready yet.

diff --git a/Engine/Source/Runtime/Engine/Private/StaticMeshRender.cpp b/Engine/Source/Runtime/Engine/Private/StaticMeshRender.cpp
--- a/Engine/Source/Runtime/Engine/Private/StaticMeshRender.cpp
+++ b/Engine/Source/Runtime/Engine/Private/StaticMeshRender.cpp
@@ -197,6 +197,7 @@ FStaticMeshSceneProxy::FStaticMeshSceneProxy(UStaticMeshComponent* InComponent,
 
 #if RHI_RAYTRACING
 	bDynamicRayTracingGeometry = InComponent->bEvaluateWorldPositionOffset && MaterialRelevance.bUsesWorldPositionOffset;
+	
 	if (IsRayTracingEnabled())
 	{
 		RayTracingGeometries.AddDefaulted(RenderData->LODResources.Num());
@@ -205,21 +206,10 @@ FStaticMeshSceneProxy::FStaticMeshSceneProxy(UStaticMeshComponent* InComponent,
 			RayTracingGeometries[LODIndex] = &RenderData->LODResources[LODIndex].RayTracingGeometry;
 		}
 
-		if(bDynamicRayTracingGeometry)
 		{
-			DynamicRayTracingGeometries.AddDefaulted(RenderData->LODResources.Num());
-			for (int32 LODIndex = 0; LODIndex < RenderData->LODResources.Num(); LODIndex++)
-			{
-				auto& Initializer = DynamicRayTracingGeometries[LODIndex].Initializer;
-				Initializer = RenderData->LODResources[LODIndex].RayTracingGeometry.Initializer;
-				for (FRayTracingGeometrySegment& Segment : Initializer.Segments)
-				{
-					Segment.VertexBuffer = nullptr;
-				}
-				Initializer.bAllowUpdate = true;
-				Initializer.bFastBuild = true;
-			}
+			RayTracingGeometries[LODIndex] = &RenderData->LODResources[LODIndex].RayTracingGeometry;
 		}
+
 	}
 #endif
 
@@ -626,12 +616,27 @@ void FStaticMeshSceneProxy::CreateRenderThreadResources()
 #if RHI_RAYTRACING
 	if(IsRayTracingEnabled())
 	{
+		if (bDynamicRayTracingGeometry)
+		{
+			DynamicRayTracingGeometries.AddDefaulted(RenderData->LODResources.Num());
+			for (int32 LODIndex = 0; LODIndex < RenderData->LODResources.Num(); LODIndex++)
+			{
+				auto& Initializer = DynamicRayTracingGeometries[LODIndex].Initializer;
+				Initializer = RenderData->LODResources[LODIndex].RayTracingGeometry.Initializer;
+				for (FRayTracingGeometrySegment& Segment : Initializer.Segments)
+				{
+					Segment.VertexBuffer = nullptr;
+				}
+				Initializer.bAllowUpdate = true;
+				Initializer.bFastBuild = true;
+			}
+		}
+
 		DynamicRayTracingGeometryVertexBuffers.AddDefaulted(DynamicRayTracingGeometries.Num());
 		for(int32 i = 0; i < DynamicRayTracingGeometries.Num(); i++)
 		{
 			auto& Geometry = DynamicRayTracingGeometries[i];
-			DynamicRayTracingGeometryVertexBuffers[i]
-				.Initialize(4, 256, PF_R32_FLOAT, BUF_UnorderedAccess | BUF_ShaderResource, TEXT("FStaticMeshSceneProxy::RayTracingDynamicVertexBuffer"));
+			DynamicRayTracingGeometryVertexBuffers[i].Initialize(4, 256, PF_R32_FLOAT, BUF_UnorderedAccess | BUF_ShaderResource, TEXT("FStaticMeshSceneProxy::RayTracingDynamicVertexBuffer"));
 			Geometry.InitResource();
 		}
 	}
