Fixed bad merge

diff --git a/Engine/Source/Runtime/Engine/Private/StaticMeshRender.cpp b/Engine/Source/Runtime/Engine/Private/StaticMeshRender.cpp
--- a/Engine/Source/Runtime/Engine/Private/StaticMeshRender.cpp
+++ b/Engine/Source/Runtime/Engine/Private/StaticMeshRender.cpp
@@ -205,11 +205,6 @@ FStaticMeshSceneProxy::FStaticMeshSceneProxy(UStaticMeshComponent* InComponent,
 		{
 			RayTracingGeometries[LODIndex] = &RenderData->LODResources[LODIndex].RayTracingGeometry;
 		}
-
-		{
-			RayTracingGeometries[LODIndex] = &RenderData->LODResources[LODIndex].RayTracingGeometry;
-		}
-
 	}
 #endif
 
