Fix: Avoid Subsystem Deinitialize when doing a Map Build from the Editor (check for bCleanupResources). This would leave the World with no subsystems.

diff --git a/Engine/Source/Runtime/Engine/Private/World.cpp b/Engine/Source/Runtime/Engine/Private/World.cpp
--- a/Engine/Source/Runtime/Engine/Private/World.cpp
+++ b/Engine/Source/Runtime/Engine/Private/World.cpp
@@ -4357,7 +4357,10 @@ void UWorld::CleanupWorldInternal(bool bSessionEnded, bool bCleanupResources, UW
 
 	FWorldDelegates::OnPostWorldCleanup.Broadcast(this, bSessionEnded, bCleanupResources);
 
-	SubsystemCollection.Deinitialize();
+	if (bCleanupResources)
+	{
+		SubsystemCollection.Deinitialize();
+	}
 
 	if(FXSystem && bWorldChanged)
 	{
