re-enable static mesh lod streaming. It was accidentally disabled from last merge.

diff --git a/Engine/Source/Runtime/Engine/Private/StaticMesh.cpp b/Engine/Source/Runtime/Engine/Private/StaticMesh.cpp
--- a/Engine/Source/Runtime/Engine/Private/StaticMesh.cpp
+++ b/Engine/Source/Runtime/Engine/Private/StaticMesh.cpp
@@ -2698,13 +2698,6 @@ void UStaticMesh::InitResources()
 	{
 		LinkStreaming();
 	}
-	
-	// Make sure the first LOD is actually the lowest LOD.
-	if (RenderData)
-	{
-		RenderData->CurrentFirstLODIdx = RenderData->GetFirstValidLODIdx(0);
-		SetCachedNumResidentLODs(RenderData->LODResources.Num() - RenderData->CurrentFirstLODIdx);
-	}
 
 #if	STATS
 	UStaticMesh* This = this;
