#UE4 If both packages and paths are specified in UAssetManager::ScanPathsForPrimaryAssets it will now discover assets found in the paths and the specific assets. Previously it would only find assets that are in the specified paths AND one of the specific assets, which cannot be the expected behavior because RefreshPrimaryAssetDirectory unions all scan paths/packagenames into one scan per type which may come from multiple sources.

diff --git a/Engine/Source/Runtime/Engine/Private/AssetManager.cpp b/Engine/Source/Runtime/Engine/Private/AssetManager.cpp
--- a/Engine/Source/Runtime/Engine/Private/AssetManager.cpp
+++ b/Engine/Source/Runtime/Engine/Private/AssetManager.cpp
@@ -490,21 +490,36 @@ int32 UAssetManager::ScanPathsForPrimaryAssets(FPrimaryAssetType PrimaryAssetTyp
 		}
 	}
 
+	const bool bBothDirectoriesAndPackageNames = (Directories.Num() > 0 && PackageNames.Num() > 0);
 	for (const FString& Directory : Directories)
 	{
 		ARFilter.PackagePaths.Add(FName(*Directory));
 	}
 
-	for (const FString& PackageName : PackageNames)
+	if (!bBothDirectoriesAndPackageNames)
 	{
-		ARFilter.PackageNames.Add(FName(*PackageName));
+		// To get both the directories and package names we have to do two queries, since putting both in the same query only returns assets of those package names AND are in those directories.
+		for (const FString& PackageName : PackageNames)
+		{
+			ARFilter.PackageNames.Add(FName(*PackageName));
+		}
 	}
 
 	ARFilter.bRecursivePaths = true;
 	ARFilter.bIncludeOnlyOnDiskAssets = !GIsEditor; // In editor check in memory, otherwise don't
 
 	TArray<FAssetData> AssetDataList;
+	if (bBothDirectoriesAndPackageNames)
+	{
+		// To get both the directories and package names we have to do two queries, since putting both in the same query only returns assets of those package names AND are in those directories.
+		AssetRegistry.GetAssets(ARFilter, AssetDataList);
 
+		for (const FString& PackageName : PackageNames)
+		{
+			ARFilter.PackageNames.Add(FName(*PackageName));
+		}
+		ARFilter.PackagePaths.Empty();
+	}
 	AssetRegistry.GetAssets(ARFilter, AssetDataList);
 
 	int32 NumAdded = 0;
