Fixed potential race condition accessing data when ray tracing skeletal meshes that is being destroyed from another thread

diff --git a/Engine/Source/Runtime/Engine/Private/SkeletalMesh.cpp b/Engine/Source/Runtime/Engine/Private/SkeletalMesh.cpp
--- a/Engine/Source/Runtime/Engine/Private/SkeletalMesh.cpp
+++ b/Engine/Source/Runtime/Engine/Private/SkeletalMesh.cpp
@@ -5211,7 +5211,8 @@ void FSkeletalMeshSceneProxy::GetDynamicElementsSection(const TArray<const FScen
 #if RHI_RAYTRACING
 void FSkeletalMeshSceneProxy::GetDynamicRayTracingInstances(FRayTracingMaterialGatheringContext & Context, TArray<struct FRayTracingInstance>& OutRayTracingInstances)
 {
-	if (MeshObject->GetRayTracingGeometry())
+	// GetRayTracingGeometry()->IsInitialized() is checked as a workaround for UE-92634. FSkeletalMeshSceneProxy's resources may have already been released, but proxy has not removed yet)
+	if (MeshObject->GetRayTracingGeometry() && MeshObject->GetRayTracingGeometry()->IsInitialized() )
 	{
 		// #dxr: the only case where RayTracingGeometryRHI is invalid is the very first frame - if that's not the case we have a bug somewhere else
 		if (MeshObject->GetRayTracingGeometry()->RayTracingGeometryRHI.IsValid())
