Unregister device/world pair only if no more references found
#rnx

diff --git a/Engine/Source/Runtime/Engine/Private/AudioDeviceManager.cpp b/Engine/Source/Runtime/Engine/Private/AudioDeviceManager.cpp
--- a/Engine/Source/Runtime/Engine/Private/AudioDeviceManager.cpp
+++ b/Engine/Source/Runtime/Engine/Private/AudioDeviceManager.cpp
@@ -612,9 +612,9 @@ void FAudioDeviceManager::DecrementDevice(Audio::FDeviceId DeviceID, UWorld* InW
 
 		FAudioDeviceManagerDelegates::OnAudioDeviceDestroyed.Broadcast(DeviceID);
 		Devices.Remove(DeviceID);
-	}
 
-	UnregisterWorld(InWorld, DeviceID);
+		UnregisterWorld(InWorld, DeviceID);
+	}
 }
 
 bool FAudioDeviceManager::ShutdownAllAudioDevices()
