Remove check if device is valid before attempting to parse sound node graph and early out instead
#rnx

diff --git a/Engine/Source/Runtime/Engine/Private/ActiveSound.cpp b/Engine/Source/Runtime/Engine/Private/ActiveSound.cpp
--- a/Engine/Source/Runtime/Engine/Private/ActiveSound.cpp
+++ b/Engine/Source/Runtime/Engine/Private/ActiveSound.cpp
@@ -495,8 +495,6 @@ bool FActiveSound::GetConcurrencyFadeDuration(float& OutFadeDuration) const
 
 void FActiveSound::UpdateWaveInstances(TArray<FWaveInstance*> &InWaveInstances, const float DeltaTime)
 {
-	check(AudioDevice);
-
 	// Reset whether or not the active sound is playing audio.
 	bIsPlayingAudio = false;
 
@@ -504,7 +502,7 @@ void FActiveSound::UpdateWaveInstances(TArray<FWaveInstance*> &InWaveInstances,
 	MinCurrentPitch = 1.0f;
 
 	// Early outs.
-	if (Sound == nullptr || !Sound->IsPlayable())
+	if (!AudioDevice || Sound == nullptr || !Sound->IsPlayable())
 	{
 		return;
 	}
