Clear blend weights for any nodes running in layers in SetState(). Fixes issues with auto transition rules reading stale blend weight data.
#rnx

diff --git a/Engine/Source/Runtime/Engine/Private/Animation/AnimNode_StateMachine.cpp b/Engine/Source/Runtime/Engine/Private/Animation/AnimNode_StateMachine.cpp
--- a/Engine/Source/Runtime/Engine/Private/Animation/AnimNode_StateMachine.cpp
+++ b/Engine/Source/Runtime/Engine/Private/Animation/AnimNode_StateMachine.cpp
@@ -995,6 +995,25 @@ void FAnimNode_StateMachine::SetState(const FAnimationBaseContext& Context, int3
 			}
 		}
 
+		// Clear any currently cached blend weights for asset player nodes in layers.
+		for (const int32& LayerIdx : GetStateInfo(CurrentState).LayerNodeIndices)
+		{
+			// Try and retrieve the actual node object
+			if (FAnimNode_LinkedAnimLayer* Layer = Context.AnimInstanceProxy->GetNodeFromIndex<FAnimNode_LinkedAnimLayer>(LayerIdx))
+			{
+				// Retrieve the AnimInstance running for this layer
+				if (UAnimInstance* CurrentTarget = Layer->GetTargetInstance<UAnimInstance>())
+				{
+					// Retrieve all asset player nodes from the corresponding Anim blueprint class and clear their cached blend weight
+					TArray<FAnimNode_AssetPlayerBase*> PlayerNodesInLayer = CurrentTarget->GetInstanceAssetPlayers(Layer->Layer);
+					for (FAnimNode_AssetPlayerBase* Player : PlayerNodesInLayer)
+					{
+						Player->ClearCachedBlendWeight();
+					}
+				}
+			}
+		}
+
 		if ((!bAlreadyActive || bForceReset) && !IsAConduitState(NewStateIndex))
 		{
 			// Initialize the new state since it's not part of an active transition (and thus not still initialized)
