Subsurface profile fix in USubsurfaceProfile::PostLoad.

diff --git a/Engine/Source/Runtime/Engine/Private/Rendering/SubsurfaceProfile.cpp b/Engine/Source/Runtime/Engine/Private/Rendering/SubsurfaceProfile.cpp
--- a/Engine/Source/Runtime/Engine/Private/Rendering/SubsurfaceProfile.cpp
+++ b/Engine/Source/Runtime/Engine/Private/Rendering/SubsurfaceProfile.cpp
@@ -524,22 +524,23 @@ void USubsurfaceProfile::PostLoad()
 
 	const auto CVar = IConsoleManager::Get().
 		FindTConsoleVariableDataInt(TEXT("r.SSS.Burley.AlwaysUpdateParametersFromSeparable"));
-	check(CVar);
-	
-	const bool bUpdateBurleyParametersFromSeparable = CVar->GetValueOnAnyThread() == 1;
-
-	if (bUpdateBurleyParametersFromSeparable)
+	if (CVar)
 	{
-		MapFallOffColor2SurfaceAlbedoAndDiffuseMeanFreePath(Settings.FalloffColor.R, Settings.SurfaceAlbedo.R, Settings.MeanFreePathColor.R);
-		MapFallOffColor2SurfaceAlbedoAndDiffuseMeanFreePath(Settings.FalloffColor.G, Settings.SurfaceAlbedo.G, Settings.MeanFreePathColor.G);
-		MapFallOffColor2SurfaceAlbedoAndDiffuseMeanFreePath(Settings.FalloffColor.B, Settings.SurfaceAlbedo.B, Settings.MeanFreePathColor.B);
+		const bool bUpdateBurleyParametersFromSeparable = CVar->GetValueOnAnyThread() == 1;
 
-		//Normalize mean free path color and set the corresponding dfmp
-		float MaxMeanFreePathColor = FMath::Max3(Settings.MeanFreePathColor.R, Settings.MeanFreePathColor.G, Settings.MeanFreePathColor.B);
-		if (MaxMeanFreePathColor > 1)
+		if (bUpdateBurleyParametersFromSeparable)
 		{
-			Settings.MeanFreePathColor /= MaxMeanFreePathColor;
-			Settings.MeanFreePathDistance = FMath::Clamp(Settings.ScatterRadius*MaxMeanFreePathColor,0.1f,50.0f);	// 50.0f is the ClampMax of MeanFreePathDistance.
+			MapFallOffColor2SurfaceAlbedoAndDiffuseMeanFreePath(Settings.FalloffColor.R, Settings.SurfaceAlbedo.R, Settings.MeanFreePathColor.R);
+			MapFallOffColor2SurfaceAlbedoAndDiffuseMeanFreePath(Settings.FalloffColor.G, Settings.SurfaceAlbedo.G, Settings.MeanFreePathColor.G);
+			MapFallOffColor2SurfaceAlbedoAndDiffuseMeanFreePath(Settings.FalloffColor.B, Settings.SurfaceAlbedo.B, Settings.MeanFreePathColor.B);
+
+			//Normalize mean free path color and set the corresponding dfmp
+			float MaxMeanFreePathColor = FMath::Max3(Settings.MeanFreePathColor.R, Settings.MeanFreePathColor.G, Settings.MeanFreePathColor.B);
+			if (MaxMeanFreePathColor > 1)
+			{
+				Settings.MeanFreePathColor /= MaxMeanFreePathColor;
+				Settings.MeanFreePathDistance = FMath::Clamp(Settings.ScatterRadius*MaxMeanFreePathColor, 0.1f, 50.0f);	// 50.0f is the ClampMax of MeanFreePathDistance.
+			}
 		}
 	}
 }
\ No newline at end of file
