Adding a counter for the number of times the audio streaming cache has overflown. This will be reported in audiomemeport and stat audiostreaming debug utils.
As a side effect, Blowing the cache will no longer spam the log or GameThread task spawning to release retainers on all USoundWaves.

#rnx

diff --git a/Engine/Source/Runtime/Engine/Private/AudioStreamingCache.cpp b/Engine/Source/Runtime/Engine/Private/AudioStreamingCache.cpp
--- a/Engine/Source/Runtime/Engine/Private/AudioStreamingCache.cpp
+++ b/Engine/Source/Runtime/Engine/Private/AudioStreamingCache.cpp
@@ -352,30 +352,33 @@ FAudioChunkHandle FCachedAudioStreamingManager::GetLoadedChunk(const USoundWave*
 			uint64 LookupIDForNextChunk = Cache->AddOrTouchChunk(NextChunkKey, [](EAudioChunkLoadResult) {}, ENamedThreads::AnyThread, false);
 			if (LookupIDForNextChunk == InvalidAudioStreamCacheLookupID)
 			{
-				UE_LOG(LogAudio, Warning, TEXT("Cache overflow!!! couldn't load chunk %d for sound %s!"), ChunkIndex, *SoundWave->GetName());
+				// this bool is true while we are waiting on the game thread to reset chunk handles owned by USoundWaves
+				static FThreadSafeBool bCacheCurrentlyBlown = false;
 
-
-				AsyncTask(ENamedThreads::GameThread, []()
+				if (!bCacheCurrentlyBlown)
 				{
-					if (GEngine && GEngine->GetMainAudioDevice())
-					{
-						GEngine->GetMainAudioDevice()->Exec(nullptr, TEXT("audiomemreport"));
-					}
-					
-					int32 NumChunksReleased = 0;
-
-					for (TObjectIterator<USoundWave> It; It; ++It)
+					bCacheCurrentlyBlown = true;
+					UE_LOG(LogAudio, Warning, TEXT("Cache overflow!!! couldn't load chunk %d for sound %s!"), ChunkIndex, *SoundWave->GetName());
+					AsyncTask(ENamedThreads::GameThread, []()
 					{
-						USoundWave* Wave = *It;
-						if (Wave && Wave->IsRetainingAudio())
+						if (GEngine && GEngine->GetMainAudioDevice())
 						{
-							Wave->ReleaseCompressedAudio();
-							NumChunksReleased++;
+							GEngine->GetMainAudioDevice()->Exec(nullptr, TEXT("audiomemreport"));
 						}
-					}
-
-					UE_LOG(LogAudio, Warning, TEXT("Removed %d retained sounds from the stream cache."), NumChunksReleased);
-				});
+						int32 NumChunksReleased = 0;
+						for (TObjectIterator<USoundWave> It; It; ++It)
+						{
+							USoundWave* Wave = *It;
+							if (Wave && Wave->IsRetainingAudio())
+							{
+								Wave->ReleaseCompressedAudio();
+								NumChunksReleased++;
+							}
+						}
+						UE_LOG(LogAudio, Warning, TEXT("Removed %d retained sounds from the stream cache."), NumChunksReleased);
+						bCacheCurrentlyBlown = false;
+					});
+				}
 			}
 			else
 			{
@@ -537,6 +540,7 @@ FAudioChunkCache::FAudioChunkCache(uint32 InMaxChunkSize, uint32 NumChunks, uint
 	{
 		CachePool.Emplace(MaxChunkSize, Index);
 	}
+	CacheOverflowCount.Set(0);
 }
 
 FAudioChunkCache::~FAudioChunkCache()
@@ -1573,6 +1577,9 @@ TPair<int, int> FAudioChunkCache::DebugDisplay(UWorld* World, FViewport* Viewpor
 
 	FString NumElementsDetail = *FString::Printf(TEXT("Number of chunks loaded: %d of %d"), ChunksInUse, CachePool.Num());
 
+	const int32 NumCacheOverflows = CacheOverflowCount.GetValue();
+	FString CacheOverflowsDetail = *FString::Printf(TEXT("The cache has blown %d times)"), NumCacheOverflows);
+
 	// Offset our number of elements loaded horizontally to the right next to the cache title:
 	int32 CacheTitleOffsetX = 0;
 	int32 CacheTitleOffsetY = 0;
@@ -1581,6 +1588,9 @@ TPair<int, int> FAudioChunkCache::DebugDisplay(UWorld* World, FViewport* Viewpor
 	Canvas->DrawShadowedString(X + CacheTitleOffsetX, Y - 12, *NumElementsDetail, UEngine::GetSmallFont(), FLinearColor::Green);
 	Y += 10;
 
+	Canvas->DrawShadowedString(X + CacheTitleOffsetX, Y - 12, *CacheOverflowsDetail, UEngine::GetSmallFont(), NumCacheOverflows != 0? FLinearColor::Red : FLinearColor::Green);
+	Y += 10;
+
 	// First pass: We run through and get a snap shot of the amount of memory currently in use.
 	FCacheElement* CurrentElement = MostRecentElement;
 	uint32 NumBytesCounter = 0;
@@ -1713,8 +1723,10 @@ FString FAudioChunkCache::DebugPrint()
 	FString OutputString;
 
 	FString NumElementsDetail = *FString::Printf(TEXT("Number of chunks loaded: %d of %d"), ChunksInUse, CachePool.Num());
+	FString NumCacheOverflows = *FString::Printf(TEXT("The cache has blown %d times"), CacheOverflowCount.GetValue());
 
 	OutputString += NumElementsDetail + TEXT("\n");
+	OutputString += NumCacheOverflows + TEXT("\n");
 
 	// First pass: We run through and get a snap shot of the amount of memory currently in use.
 	FCacheElement* CurrentElement = MostRecentElement;
diff --git a/Engine/Source/Runtime/Engine/Public/AudioStreamingCache.h b/Engine/Source/Runtime/Engine/Public/AudioStreamingCache.h
--- a/Engine/Source/Runtime/Engine/Public/AudioStreamingCache.h
+++ b/Engine/Source/Runtime/Engine/Public/AudioStreamingCache.h
@@ -106,6 +106,16 @@ public:
 	// Generate a formatted text file for this cache.
 	FString DebugPrint();
 
+	void IncrementCacheOverflowCounter()
+	{
+		CacheOverflowCount.Increment();
+	}
+
+	int32 GetNumberOfCacheOverflows() const
+	{
+		return CacheOverflowCount.GetValue();
+	}
+
 private:
 
 #if DEBUG_STREAM_CACHE
@@ -156,6 +166,11 @@ private:
 	};
 #endif
 
+
+	// counter for the number of times this cache has overflown
+	FThreadSafeCounter CacheOverflowCount;
+
+
 	// Struct containing a single element in our LRU Cache.  
 	struct FCacheElement
 	{
