Fixed possible race condition where Texture2DResource::InitRHI() can access invalid data because the Platform mips are being regenerated.
Added extra FlushRenderingCommands() when enabling or disabling dynamicaly the render asset streaming manager to prevent draw commands from accessing released resources.

diff --git a/Engine/Source/Runtime/Engine/Private/ContentStreaming.cpp b/Engine/Source/Runtime/Engine/Private/ContentStreaming.cpp
--- a/Engine/Source/Runtime/Engine/Private/ContentStreaming.cpp
+++ b/Engine/Source/Runtime/Engine/Private/ContentStreaming.cpp
@@ -1281,6 +1281,8 @@ void FStreamingManagerCollection::AddOrRemoveTextureStreamingManagerIfNeeded(boo
 		//Add the texture streaming manager if it's needed.
 		if( !TextureStreamingManager )
 		{
+			FlushRenderingCommands();
+
 			GConfig->GetFloat( TEXT("TextureStreaming"), TEXT("LoadMapTimeLimit"), LoadMapTimeLimit, GEngineIni );
 			// Create the streaming manager and add the default streamers.
 			TextureStreamingManager = new FRenderAssetStreamingManager();
@@ -1323,6 +1325,7 @@ void FStreamingManagerCollection::AddOrRemoveTextureStreamingManagerIfNeeded(boo
 		//Remove the texture streaming manager if needed.
 		if( TextureStreamingManager )
 		{
+			FlushRenderingCommands();
 			TextureStreamingManager->BlockTillAllRequestsFinished();
 
 			RemoveStreamingManager(TextureStreamingManager);
diff --git a/Engine/Source/Runtime/Engine/Private/Texture2D.cpp b/Engine/Source/Runtime/Engine/Private/Texture2D.cpp
--- a/Engine/Source/Runtime/Engine/Private/Texture2D.cpp
+++ b/Engine/Source/Runtime/Engine/Private/Texture2D.cpp
@@ -816,6 +816,12 @@ void UTexture2D::UpdateResource()
 	}
 
 #if WITH_EDITOR
+	// Wait for FTexture2DResource::InitRHI() complete before modifying the PlatformData.
+	FTexture2DResource* Texture2DResource = (FTexture2DResource*)Resource;
+	if (Texture2DResource && !Texture2DResource->bReadyForStreaming)
+	{
+		FlushRenderingCommands();
+	}
 	// Recache platform data if the source has changed.
 	CachePlatformData();
 	// clear all the cooked cached platform data if the source could have changed... 
