#UE4 Adjust verbosity of a spammy log in FStreamLevelAction::UpdateOperation

diff --git a/Engine/Source/Runtime/Engine/Private/LevelStreaming.cpp b/Engine/Source/Runtime/Engine/Private/LevelStreaming.cpp
--- a/Engine/Source/Runtime/Engine/Private/LevelStreaming.cpp
+++ b/Engine/Source/Runtime/Engine/Private/LevelStreaming.cpp
@@ -93,7 +93,7 @@ void FStreamLevelAction::UpdateOperation(FLatentResponse& Response)
 {
 	ULevelStreaming* LevelStreamingObject = Level.Get(); // to avoid confusion.
 	const bool bIsLevelValid = LevelStreamingObject != nullptr;
-	UE_LOG(LogLevelStreaming, Display, TEXT("FStreamLevelAction::UpdateOperation() LevelName %s, bIsLevelValid %d"), *LevelName.ToString(), (int32)bIsLevelValid);
+	UE_LOG(LogLevelStreaming, Verbose, TEXT("FStreamLevelAction::UpdateOperation() LevelName %s, bIsLevelValid %d"), *LevelName.ToString(), (int32)bIsLevelValid);
 	if (bIsLevelValid)
 	{
 		bool bIsOperationFinished = UpdateLevel(LevelStreamingObject);
