- Modify particle perf to capture across GC's for full replay captures
- Capture a small history of peak effect costs to help track down effects that 'spike'

#rnx

diff --git a/Engine/Source/Runtime/Engine/Classes/Particles/ParticlePerfStats.h b/Engine/Source/Runtime/Engine/Classes/Particles/ParticlePerfStats.h
--- a/Engine/Source/Runtime/Engine/Classes/Particles/ParticlePerfStats.h
+++ b/Engine/Source/Runtime/Engine/Classes/Particles/ParticlePerfStats.h
@@ -45,39 +45,49 @@ struct ENGINE_API FParticlePerfStats
 		uint64			GetDynamicMeshElementsCycles;	// Total cycles spent setting up dynamic mesh elements
 	};
 
+	static constexpr int32 kNumMaxSamples = 10;
+
+	FParticlePerfStats()
+	{
+		Reset();
+	}
+
 	uint32			AccumulatedNumFrames = 0;
 	FGroupedStats	AccumulatedStats;
 
+	TArray<uint64, TInlineAllocator<kNumMaxSamples>>	MaxPerFrameTotalGameThreadCycles;
+	TArray<uint64, TInlineAllocator<kNumMaxSamples>>	MaxPerFrameTotalRenderThreadCycles;
+
 	FGroupedStats	CurrentFrameStats;
 
+	static FParticlePerfStats* GetPerfStats(class UFXSystemAsset* Asset);
+
 	void Reset();
 	void Tick();
 
 	static void OnStartup();
 	static void OnShutdown();
-
-	static FParticlePerfStats Dummy;
 };
 
 //-TODO: Need to remove the task graph WakeUp cost otherwise it skews results!
 #define PARTICLE_PERF_STAT_INSTANCE_COUNT(STATS, COUNT) \
-	(STATS) != nullptr ? (STATS)->ParticlePerfStats.CurrentFrameStats.NumInstances += COUNT : 0
+	FParticlePerfStats::GetPerfStats(STATS)->CurrentFrameStats.NumInstances += COUNT
 
 #define PARTICLE_PERF_STAT_CYCLES(STATS, NAME) \
 	struct FScopedParticleStat_##NAME \
 	{ \
-		FORCEINLINE FScopedParticleStat_ ## NAME(FParticlePerfStats& InStats) \
-			: Stats(InStats) \
+		FORCEINLINE FScopedParticleStat_ ## NAME(class UFXSystemAsset* InSystemAsset) \
+			: SystemAsset(InSystemAsset) \
 			, StartCycles(FPlatformTime::Cycles64()) \
 		{ \
 		} \
 		FORCEINLINE ~FScopedParticleStat_##NAME() \
 		{ \
-			Stats.CurrentFrameStats.NAME##Cycles += FPlatformTime::Cycles64() - StartCycles; \
+			FParticlePerfStats::GetPerfStats(SystemAsset)->CurrentFrameStats.NAME##Cycles += FPlatformTime::Cycles64() - StartCycles; \
 		} \
-		FParticlePerfStats& Stats; \
+		class UFXSystemAsset* SystemAsset; \
 		uint64 StartCycles; \
-	} NAME##Stat((STATS) ? (STATS)->ParticlePerfStats : FParticlePerfStats::Dummy);
+	} NAME##Stat(STATS);
 
 #else //WITH_PARTICLE_PERF_STATS
 
diff --git a/Engine/Source/Runtime/Engine/Classes/Particles/ParticleSystem.h b/Engine/Source/Runtime/Engine/Classes/Particles/ParticleSystem.h
--- a/Engine/Source/Runtime/Engine/Classes/Particles/ParticleSystem.h
+++ b/Engine/Source/Runtime/Engine/Classes/Particles/ParticleSystem.h
@@ -121,7 +121,7 @@ public:
 	//TODO: Allow pool size overriding per world and possibly implement some preallocation too.
 
 #if WITH_PARTICLE_PERF_STATS
-	FParticlePerfStats ParticlePerfStats;
+	FParticlePerfStats* ParticlePerfStats;
 #endif
 };
 
diff --git a/Engine/Source/Runtime/Engine/Private/Particles/ParticlePerfStats.cpp b/Engine/Source/Runtime/Engine/Private/Particles/ParticlePerfStats.cpp
--- a/Engine/Source/Runtime/Engine/Private/Particles/ParticlePerfStats.cpp
+++ b/Engine/Source/Runtime/Engine/Private/Particles/ParticlePerfStats.cpp
@@ -11,26 +11,40 @@ namespace ParticlePerfStatsLocal
 	static int32 GbParticlePerfStatsEnabled = 0;
 	static int32 GbParticlePerfStatsEnabledLatch = 0;
 	static int32 GParticlePerfStatsFramesRemaining = 0;
+	static FCriticalSection SystemToPerfStatsGuard;
+	static TMap<FName, TUniquePtr<FParticlePerfStats>> SystemToPerfStats;
 
 	static void DumpParticlePerfStats(FOutputDevice& Ar)
 	{
 		FlushRenderingCommands();
+		FScopeLock ScopeLock(&ParticlePerfStatsLocal::SystemToPerfStatsGuard);
 
 		FString tempString;
 
 		Ar.Logf(TEXT("**** Particle Performance Stats"));
-		Ar.Logf(TEXT("Name,Average PerFrame GameThread,Average PerInstance GameThread,Average PerFrame RenderThread,Average PerInstance RenderThread,NumFrames,Total Instances,Total Tick GameThread,Total Tick Concurrent,Total Finalize,Total End Of Frame,Total Render Update,Total Get Dynamic Mesh Elements"));
-		for (TObjectIterator<UFXSystemAsset> SystemIt; SystemIt; ++SystemIt)
+		Ar.Logf(TEXT(",Name,Average PerFrame GameThread,Average PerInstance GameThread,Average PerFrame RenderThread,Average PerInstance RenderThread,NumFrames,Total Instances,Total Tick GameThread,Total Tick Concurrent,Total Finalize,Total End Of Frame,Total Render Update,Total Get Dynamic Mesh Elements,Max PerFrame GameThread,Max Range PerFrame GameThread,Max PerFrame RenderThread,Max Range PerFrame RenderThread"));
+
+		for (auto it = SystemToPerfStats.CreateIterator(); it; ++it)
 		{
-			const FParticlePerfStats& PerfStats = SystemIt->ParticlePerfStats;
-			if (PerfStats.AccumulatedStats.NumInstances == 0)
+			const FParticlePerfStats& PerfStats = *it.Value();
+
+			if (PerfStats.AccumulatedNumFrames == 0)
+			{
 				continue;
+			}
+
+			FString SystemName = it.Key().ToString();
 
 			const uint64 TotalGameThread = PerfStats.AccumulatedStats.TickGameThreadCycles + PerfStats.AccumulatedStats.TickConcurrentCycles + PerfStats.AccumulatedStats.FinalizeCycles + PerfStats.AccumulatedStats.EndOfFrameCycles;
 			const uint64 TotalRenderThread = PerfStats.AccumulatedStats.RenderUpdateCycles + PerfStats.AccumulatedStats.GetDynamicMeshElementsCycles;
 
+			const uint64 MaxPerFrameTotalGameThreadFirst = PerfStats.MaxPerFrameTotalGameThreadCycles.Num() > 0 ? PerfStats.MaxPerFrameTotalGameThreadCycles[0] : 0;
+			const uint64 MaxPerFrameTotalRenderThreadFirst = PerfStats.MaxPerFrameTotalRenderThreadCycles.Num() > 0 ? PerfStats.MaxPerFrameTotalRenderThreadCycles[0] : 0;
+			const uint64 MaxPerFrameTotalGameThreadLast = PerfStats.MaxPerFrameTotalGameThreadCycles.Num() > 0 ? PerfStats.MaxPerFrameTotalGameThreadCycles.Last() : 0;
+			const uint64 MaxPerFrameTotalRenderThreadLast = PerfStats.MaxPerFrameTotalRenderThreadCycles.Num() > 0 ? PerfStats.MaxPerFrameTotalRenderThreadCycles.Last() : 0;
+
 			tempString.Reset();
-			tempString.Append(*SystemIt->GetFullName());
+			tempString.Appendf(TEXT(",%s"), *SystemName);
 			tempString.Appendf(TEXT(",%u"), uint32(FPlatformTime::ToMilliseconds64(TotalGameThread) * 1000.0 / double(PerfStats.AccumulatedNumFrames)));
 			tempString.Appendf(TEXT(",%u"), uint32(FPlatformTime::ToMilliseconds64(TotalGameThread) * 1000.0 / double(PerfStats.AccumulatedStats.NumInstances)));
 			tempString.Appendf(TEXT(",%u"), uint32(FPlatformTime::ToMilliseconds64(TotalRenderThread) * 1000.0 / double(PerfStats.AccumulatedNumFrames)));
@@ -43,6 +57,21 @@ namespace ParticlePerfStatsLocal
 			tempString.Appendf(TEXT(",%u"), uint32(FPlatformTime::ToMilliseconds64(PerfStats.AccumulatedStats.EndOfFrameCycles) * 1000.0));
 			tempString.Appendf(TEXT(",%u"), uint32(FPlatformTime::ToMilliseconds64(PerfStats.AccumulatedStats.RenderUpdateCycles) * 1000.0));
 			tempString.Appendf(TEXT(",%u"), uint32(FPlatformTime::ToMilliseconds64(PerfStats.AccumulatedStats.GetDynamicMeshElementsCycles) * 1000.0));
+
+			tempString.Appendf(TEXT(",%u,[ "), uint32(FPlatformTime::ToMilliseconds64(PerfStats.MaxPerFrameTotalGameThreadCycles.Num() > 0 ? PerfStats.MaxPerFrameTotalGameThreadCycles[0] : 0) * 1000.0));
+			for (uint64 v : PerfStats.MaxPerFrameTotalGameThreadCycles)
+			{
+				tempString.Appendf(TEXT("%u "), uint32(FPlatformTime::ToMilliseconds64(v) * 1000.0));
+			}
+			tempString.Append(TEXT("]"));
+
+			tempString.Appendf(TEXT(",%u,[ "), uint32(FPlatformTime::ToMilliseconds64(PerfStats.MaxPerFrameTotalRenderThreadCycles.Num() > 0 ? PerfStats.MaxPerFrameTotalRenderThreadCycles[0] : 0) * 1000.0));
+			for (uint64 v : PerfStats.MaxPerFrameTotalRenderThreadCycles)
+			{
+				tempString.Appendf(TEXT("%u "), uint32(FPlatformTime::ToMilliseconds64(v) * 1000.0));
+			}
+			tempString.Append(TEXT("]"));
+
 			Ar.Log(*tempString);
 		}
 	}
@@ -51,19 +80,22 @@ namespace ParticlePerfStatsLocal
 	{
 		FlushRenderingCommands();
 
+		FScopeLock ScopeLock(&ParticlePerfStatsLocal::SystemToPerfStatsGuard);
 		for (TObjectIterator<UFXSystemAsset> SystemIt; SystemIt; ++SystemIt)
 		{
-			SystemIt->ParticlePerfStats.Reset();
+			SystemIt->ParticlePerfStats = nullptr;
 		}
+		SystemToPerfStats.Empty();
 	}
 
 	static void TickParticlePerfStats()
 	{
 		if (GbParticlePerfStatsEnabledLatch != 0)
 		{
-			for (TObjectIterator<UFXSystemAsset> SystemIt; SystemIt; ++SystemIt)
+			FScopeLock ScopeLock(&ParticlePerfStatsLocal::SystemToPerfStatsGuard);
+			for (auto it=SystemToPerfStats.CreateIterator(); it; ++it)
 			{
-				SystemIt->ParticlePerfStats.Tick();
+				it.Value()->Tick();
 			}
 
 			bool bDisableGathering = GbParticlePerfStatsEnabled == 0;
@@ -131,12 +163,34 @@ namespace ParticlePerfStatsLocal
 	);
 }
 
-FParticlePerfStats	FParticlePerfStats::Dummy;
+FParticlePerfStats* FParticlePerfStats::GetPerfStats(class UFXSystemAsset* Asset)
+{
+	static FParticlePerfStats Dummy;
+
+	if (Asset == nullptr)
+	{
+		return &Dummy;
+	}
+
+	if (Asset->ParticlePerfStats == nullptr)
+	{
+		FScopeLock ScopeLock(&ParticlePerfStatsLocal::SystemToPerfStatsGuard);
+		TUniquePtr<FParticlePerfStats>& PerfStats = ParticlePerfStatsLocal::SystemToPerfStats.FindOrAdd(Asset->GetFName());
+		if (PerfStats == nullptr)
+		{
+			PerfStats.Reset(new FParticlePerfStats());
+		}
+		Asset->ParticlePerfStats = PerfStats.Get();
+	}
+	return Asset->ParticlePerfStats;
+}
 
 void FParticlePerfStats::Reset()
 {
 	AccumulatedNumFrames = 0;
 	AccumulatedStats.Reset();
+	MaxPerFrameTotalGameThreadCycles.Empty();
+	MaxPerFrameTotalRenderThreadCycles.Empty();
 	CurrentFrameStats.Reset();
 }
 
@@ -154,6 +208,35 @@ void FParticlePerfStats::Tick()
 	AccumulatedStats.FinalizeCycles += CurrentFrameStats.FinalizeCycles;
 	AccumulatedStats.EndOfFrameCycles += CurrentFrameStats.EndOfFrameCycles;
 
+	const uint64 ThisFrameMaxTotalGameThread = CurrentFrameStats.TickGameThreadCycles + CurrentFrameStats.TickConcurrentCycles + CurrentFrameStats.FinalizeCycles + CurrentFrameStats.EndOfFrameCycles;
+	const uint64 ThisFrameMaxTotalRenderThread = CurrentFrameStats.RenderUpdateCycles + CurrentFrameStats.GetDynamicMeshElementsCycles;
+
+	// Update Max Samples
+	{
+		int32 InsertIndex;
+		InsertIndex = MaxPerFrameTotalGameThreadCycles.IndexOfByPredicate([&](uint32 v) {return ThisFrameMaxTotalGameThread > v; });
+		if (InsertIndex != INDEX_NONE)
+		{
+			MaxPerFrameTotalGameThreadCycles.Pop(false);
+			MaxPerFrameTotalGameThreadCycles.Insert(ThisFrameMaxTotalGameThread, InsertIndex);
+		}
+		else if (MaxPerFrameTotalGameThreadCycles.Num() < kNumMaxSamples)
+		{
+			MaxPerFrameTotalGameThreadCycles.Add(ThisFrameMaxTotalGameThread);
+		}
+
+		InsertIndex = MaxPerFrameTotalRenderThreadCycles.IndexOfByPredicate([&](uint32 v) {return ThisFrameMaxTotalRenderThread > v; });
+		if (InsertIndex != INDEX_NONE)
+		{
+			MaxPerFrameTotalRenderThreadCycles.Pop(false);
+			MaxPerFrameTotalRenderThreadCycles.Insert(ThisFrameMaxTotalRenderThread, InsertIndex);
+		}
+		else if (MaxPerFrameTotalRenderThreadCycles.Num() < kNumMaxSamples)
+		{
+			MaxPerFrameTotalRenderThreadCycles.Add(ThisFrameMaxTotalRenderThread);
+		}
+	}
+
 	CurrentFrameStats.NumInstances = 0;
 	CurrentFrameStats.TickGameThreadCycles = 0;
 	CurrentFrameStats.TickConcurrentCycles = 0;
