Change cache blown from being an ensure to a warning log.
Increase the cache size on PS4 again while we mess with chunk size.

diff --git a/Engine/Source/Runtime/Engine/Private/AudioStreamingCache.cpp b/Engine/Source/Runtime/Engine/Private/AudioStreamingCache.cpp
--- a/Engine/Source/Runtime/Engine/Private/AudioStreamingCache.cpp
+++ b/Engine/Source/Runtime/Engine/Private/AudioStreamingCache.cpp
@@ -1193,7 +1193,7 @@ FAudioChunkCache::FCacheElement* FAudioChunkCache::EvictLeastRecentChunk()
 		// If we ever hit this, it means that we couldn't find any cache elements that aren't in use.
 		if (!CacheElement || CacheElement == ElementToStopAt)
 		{
-			ensureMsgf(false, TEXT("Cache blown! Please increase the cache size (currently %i bytes) or load less audio."), ReportCacheSize());
+			UE_LOG(LogAudio, Warning, TEXT("Cache blown! Please increase the cache size (currently %lu bytes) or load less audio."), ReportCacheSize());
 			return nullptr;
 		}
 	}
