Chaos - RBAN fix material hookup, friction, etc

diff --git a/Engine/Source/Runtime/Engine/Private/PhysicsEngine/ImmediatePhysics/ImmediatePhysicsChaos/ImmediatePhysicsSimulation_Chaos.cpp b/Engine/Source/Runtime/Engine/Private/PhysicsEngine/ImmediatePhysics/ImmediatePhysicsChaos/ImmediatePhysicsSimulation_Chaos.cpp
--- a/Engine/Source/Runtime/Engine/Private/PhysicsEngine/ImmediatePhysics/ImmediatePhysicsChaos/ImmediatePhysicsSimulation_Chaos.cpp
+++ b/Engine/Source/Runtime/Engine/Private/PhysicsEngine/ImmediatePhysics/ImmediatePhysicsChaos/ImmediatePhysicsSimulation_Chaos.cpp
@@ -196,6 +196,7 @@ namespace ImmediatePhysics_Chaos
 			, MaxNumRollingAverageStepTimes(ChaosImmediate_Evolution_DeltaTimeCount)
 			, bActorsDirty(false)
 		{
+			Particles.GetParticleHandles().AddArray(&CollidedParticles);
 			Particles.GetParticleHandles().AddArray(&ParticleMaterials);
 			Particles.GetParticleHandles().AddArray(&PerParticleMaterials);
 			Particles.GetParticleHandles().AddArray(&ParticlePrevXs);
@@ -401,11 +402,11 @@ namespace ImmediatePhysics_Chaos
 			}
 		}
 
-		Implementation->ParticleMaterials.Add(MakeSerializable(Material));
-		Implementation->PerParticleMaterials.Add(MoveTemp(Material));
-		Implementation->CollidedParticles.Add(false);
-		Implementation->ParticlePrevXs.Add(ActorHandle->GetParticle()->X());
-		Implementation->ParticlePrevRs.Add(ActorHandle->GetParticle()->R());
+		ActorHandle->GetParticle()->AuxilaryValue(Implementation->ParticleMaterials) = MakeSerializable(Material);
+		ActorHandle->GetParticle()->AuxilaryValue(Implementation->PerParticleMaterials) = MoveTemp(Material);
+		ActorHandle->GetParticle()->AuxilaryValue(Implementation->CollidedParticles) = false;
+		ActorHandle->GetParticle()->AuxilaryValue(Implementation->ParticlePrevXs) = ActorHandle->GetParticle()->X();
+		ActorHandle->GetParticle()->AuxilaryValue(Implementation->ParticlePrevRs) = ActorHandle->GetParticle()->R();
 
 		Implementation->bActorsDirty = true;
 
@@ -420,12 +421,6 @@ namespace ImmediatePhysics_Chaos
 		int32 Index = Implementation->ActorHandles.Remove(ActorHandle);
 		delete ActorHandle;
 
-		Implementation->ParticleMaterials.RemoveAt(Index, 1);
-		Implementation->PerParticleMaterials.RemoveAt(Index, 1);
-		Implementation->CollidedParticles.RemoveAt(Index, 1);
-		Implementation->ParticlePrevXs.RemoveAt(Index, 1);
-		Implementation->ParticlePrevRs.RemoveAt(Index, 1);
-
 		Implementation->bActorsDirty = true;
 	}
 
diff --git a/Engine/Source/Runtime/Experimental/Chaos/Private/Chaos/PBDCollisionConstraints.cpp b/Engine/Source/Runtime/Experimental/Chaos/Private/Chaos/PBDCollisionConstraints.cpp
--- a/Engine/Source/Runtime/Experimental/Chaos/Private/Chaos/PBDCollisionConstraints.cpp
+++ b/Engine/Source/Runtime/Experimental/Chaos/Private/Chaos/PBDCollisionConstraints.cpp
@@ -252,8 +252,6 @@ namespace Chaos
 			Manifolds.Add(Handle->GetKey(), Handle);
 #endif
 		}
-
-		UpdateConstraintMaterialProperties(Constraints.SinglePointConstraints[Idx]);
 	}
 
 	void FPBDCollisionConstraints::AddConstraint(const FRigidBodySweptPointContactConstraint& InConstraint)
@@ -277,8 +275,6 @@ namespace Chaos
 #endif
 			}
 		}
-
-		UpdateConstraintMaterialProperties(Constraints.SinglePointSweptConstraints[Idx]);
 	}
 
 
@@ -301,8 +297,28 @@ namespace Chaos
 			Manifolds.Add(Handle->GetKey(), Handle);
 #endif
 		}
+	}
+
+	void FPBDCollisionConstraints::PrepareIteration(float dt)
+	{
+		// NOTE: We could set material properties as we add constraints, but the ParticlePairBroadphase
+		// skips the call to AddConstraint and writes directly to the constraint array, so we
+		// need to do it after all constraints are added.
 
-		UpdateConstraintMaterialProperties(Constraints.MultiPointConstraints[Idx]);
+		for (FRigidBodyPointContactConstraint& Contact : Constraints.SinglePointConstraints)
+		{
+			UpdateConstraintMaterialProperties(Contact);
+		}
+
+		for (FRigidBodyMultiPointContactConstraint& Contact : Constraints.MultiPointConstraints)
+		{
+			UpdateConstraintMaterialProperties(Contact);
+		}
+
+		for (FRigidBodySweptPointContactConstraint& Contact : Constraints.SinglePointSweptConstraints)
+		{
+			UpdateConstraintMaterialProperties(Contact);
+		}
 	}
 
 	void FPBDCollisionConstraints::UpdatePositionBasedState(const FReal Dt)
diff --git a/Engine/Source/Runtime/Experimental/Chaos/Public/Chaos/PBDCollisionConstraints.h b/Engine/Source/Runtime/Experimental/Chaos/Public/Chaos/PBDCollisionConstraints.h
--- a/Engine/Source/Runtime/Experimental/Chaos/Public/Chaos/PBDCollisionConstraints.h
+++ b/Engine/Source/Runtime/Experimental/Chaos/Public/Chaos/PBDCollisionConstraints.h
@@ -141,7 +141,7 @@ public:
 
 	void UnprepareTick() {}
 
-	void PrepareIteration(FReal Dt) {}
+	void PrepareIteration(FReal Dt);
 
 	void UnprepareIteration(FReal Dt) {}
 
