When loading a soundcue, null out sound waves on wave players not currently being used, and only reload them if we change quality levels mid-game.
Gives us similar behavior to what we were getting when using weak references for USoundWaves, except that we are guaranteed to get access to the USoundWave on USoundCue::PostLoad, allowing us to properly setup and evaluate loading behaviors for the cache.

This also fixes an issue where unselected quality levels on a quality level could still hold onto USoundWaves with loaded audio if we set the default loading behavior to retain for all soundwaves via CVar.

diff --git a/Engine/Source/Runtime/Engine/Classes/Sound/SoundCue.h b/Engine/Source/Runtime/Engine/Classes/Sound/SoundCue.h
--- a/Engine/Source/Runtime/Engine/Classes/Sound/SoundCue.h
+++ b/Engine/Source/Runtime/Engine/Classes/Sound/SoundCue.h
@@ -280,8 +280,6 @@ private:
 	FDelegateHandle OnPostEngineInitHandle;
 	static int32 CachedQualityLevel;
 
-	void ApplySoundClassLoadingBehaviorToSoundWaves();
-
 public:
 
 	/**
diff --git a/Engine/Source/Runtime/Engine/Classes/Sound/SoundNode.h b/Engine/Source/Runtime/Engine/Classes/Sound/SoundNode.h
--- a/Engine/Source/Runtime/Engine/Classes/Sound/SoundNode.h
+++ b/Engine/Source/Runtime/Engine/Classes/Sound/SoundNode.h
@@ -226,5 +226,11 @@ public:
 
 	virtual void ReleaseRetainerOnChildWavePlayers(bool bRecurse);
 
+	/**
+	 * When called, this will find any child wave players connected to this node
+	 * and null out their associated USoundWave, allowing the USoundWave to be garbage collected.
+	 * This should only be called by USoundNodeQualityLevel::PostLoad when au.CullSoundWaveHardReferences is 1.
+	 */
+	virtual void RemoveSoundWaveOnChildWavePlayers();
 };
 
diff --git a/Engine/Source/Runtime/Engine/Classes/Sound/SoundNodeWavePlayer.h b/Engine/Source/Runtime/Engine/Classes/Sound/SoundNodeWavePlayer.h
--- a/Engine/Source/Runtime/Engine/Classes/Sound/SoundNodeWavePlayer.h
+++ b/Engine/Source/Runtime/Engine/Classes/Sound/SoundNodeWavePlayer.h
@@ -33,7 +33,14 @@ private:
 
 	void OnSoundWaveLoaded(const FName& PackageName, UPackage * Package, EAsyncLoadingResult::Type Result, bool bAddToRoot);
 
-	uint32 bAsyncLoading:1;
+	uint8 bAsyncLoading:1;
+
+	// Set to true when we enqueue a task to the game thread to load
+	// the asset associated with this thread.
+	// This only occurs if ClearAssetReferences() was called and we still tried to play this wave player,
+	// likely because the quality level was changed.
+	FThreadSafeBool bAsyncLoadRequestPending;
+
 public:
 
 	UPROPERTY(EditAnywhere, Category=WavePlayer)
@@ -65,5 +72,7 @@ public:
 	virtual void ClearAssetReferences() override;
 	//~ End USoundNode Interface
 
+	// If this returns true, this wave player currently has an async load for the USoundWave in flight.
+	bool IsCurrentlyAsyncLoadingAsset() const { return bAsyncLoading; }
 };
 
diff --git a/Engine/Source/Runtime/Engine/Private/SoundNode.cpp b/Engine/Source/Runtime/Engine/Private/SoundNode.cpp
--- a/Engine/Source/Runtime/Engine/Private/SoundNode.cpp
+++ b/Engine/Source/Runtime/Engine/Private/SoundNode.cpp
@@ -145,7 +145,7 @@ void USoundNode::ReleaseRetainerOnChildWavePlayers(bool bRecurse)
 {
 	if (bIsRetainingAudio && FPlatformCompressionUtilities::IsCurrentPlatformUsingStreamCaching())
 	{
-		// Search child nodes for wave players, then prime their waves.
+		// Search child nodes for wave players, then release their retainers.
 		for (USoundNode* ChildNode : ChildNodes)
 		{
 			if (ChildNode)
@@ -172,6 +172,29 @@ void USoundNode::ReleaseRetainerOnChildWavePlayers(bool bRecurse)
 	bIsRetainingAudio = false;
 }
 
+void USoundNode::RemoveSoundWaveOnChildWavePlayers()
+{
+	// Search child nodes for wave players, then null out their sound wave.
+	for (USoundNode* ChildNode : ChildNodes)
+	{
+		if (ChildNode)
+		{
+			ChildNode->ConditionalPostLoad();
+			ChildNode->RemoveSoundWaveOnChildWavePlayers();
+
+			USoundNodeWavePlayer* WavePlayer = Cast<USoundNodeWavePlayer>(ChildNode);
+			if (WavePlayer != nullptr)
+			{
+				USoundWave* SoundWave = WavePlayer->GetSoundWave();
+				if (SoundWave && !WavePlayer->IsCurrentlyAsyncLoadingAsset())
+				{
+					WavePlayer->SetSoundWave(nullptr);
+				}
+			}
+		}
+	}
+}
+
 void USoundNode::BeginDestroy()
 {
 	Super::BeginDestroy();
diff --git a/Engine/Source/Runtime/Engine/Private/SoundNodeQualityLevel.cpp b/Engine/Source/Runtime/Engine/Private/SoundNodeQualityLevel.cpp
--- a/Engine/Source/Runtime/Engine/Private/SoundNodeQualityLevel.cpp
+++ b/Engine/Source/Runtime/Engine/Private/SoundNodeQualityLevel.cpp
@@ -16,12 +16,37 @@
 #include "Settings/LevelEditorPlaySettings.h"
 #endif
 
+static int32 CullSoundWaveHardReferencesCvar = 1;
+FAutoConsoleVariableRef CVarCullSoundWaveHardReferences(
+	TEXT("au.CullSoundWaveHardReferences"),
+	CullSoundWaveHardReferencesCvar,
+	TEXT("When set to 1, this deliberately removes USoundWave hard references from currently unselected quality nodes.\n")
+	TEXT("0: Cull soundwaves, 1: do not cull sound waves."),
+	ECVF_Default);
+
 #if WITH_EDITOR
 void USoundNodeQualityLevel::PostLoad()
 {
 	Super::PostLoad();
 
 	ReconcileNode(false);
+
+	uint32 CachedQualityLevel = USoundCue::GetCachedQualityLevel();
+	if (CullSoundWaveHardReferencesCvar && ChildNodes.IsValidIndex(CachedQualityLevel))
+	{
+		// go through any waveplayers on an unselcted quality level and null out their soundwave references.
+		for (int32 Index = 0; Index < ChildNodes.Num(); Index++)
+		{
+			if (Index == CachedQualityLevel)
+			{
+				continue;
+			}
+			else
+			{
+				ChildNodes[Index]->RemoveSoundWaveOnChildWavePlayers();
+			}
+		}
+	}
 }
 
 void USoundNodeQualityLevel::ReconcileNode(bool bReconstructNode)
diff --git a/Engine/Source/Runtime/Engine/Private/SoundNodeWavePlayer.cpp b/Engine/Source/Runtime/Engine/Private/SoundNodeWavePlayer.cpp
--- a/Engine/Source/Runtime/Engine/Private/SoundNodeWavePlayer.cpp
+++ b/Engine/Source/Runtime/Engine/Private/SoundNodeWavePlayer.cpp
@@ -5,6 +5,7 @@
 #include "ActiveSound.h"
 #include "Sound/SoundWave.h"
 #include "UObject/FrameworkObjectVersion.h"
+#include "Async/Async.h"
 
 #define LOCTEXT_NAMESPACE "SoundNodeWavePlayer"
 
@@ -182,6 +183,24 @@ void USoundNodeWavePlayer::ParseNodes( FAudioDevice* AudioDevice, const UPTRINT
 
 		SoundWave->bLooping = bWaveIsLooping;
 	}
+	else if (!SoundWaveAssetPtr.IsNull() && !bAsyncLoadRequestPending)
+	{
+		UE_LOG(LogAudio, Warning, TEXT("Asynchronous load of %s required in USoundNodeWavePlayer::ParseNodes when we attempted to play, likely because the quality node was changed."), *GetFullNameSafe(this));
+		
+		// raise this flag in case this node is parsed again before the game thread task queue is executed.
+		bAsyncLoadRequestPending = true;
+
+		TWeakObjectPtr<USoundNodeWavePlayer> WeakThis = MakeWeakObjectPtr(this);
+
+		// Dispatch a call to the game thread to load this asset.
+		AsyncTask(ENamedThreads::GameThread, [WeakThis]() {
+			if (WeakThis.IsValid())
+			{
+				WeakThis->LoadAsset();
+				WeakThis->bAsyncLoadRequestPending = false;
+			}
+		});	
+	}
 }
 
 float USoundNodeWavePlayer::GetDuration()
