Change Fast Array Limits to configurable CVar. Add a log on server when we know we have exceeded these limits.

diff --git a/Engine/Source/Runtime/Engine/Classes/Engine/NetSerialization.h b/Engine/Source/Runtime/Engine/Classes/Engine/NetSerialization.h
--- a/Engine/Source/Runtime/Engine/Classes/Engine/NetSerialization.h
+++ b/Engine/Source/Runtime/Engine/Classes/Engine/NetSerialization.h
@@ -17,6 +17,7 @@
 #include "EngineLogs.h"
 #include "Containers/ArrayView.h"
 #include "Net/GuidReferences.h"
+#include "HAL/IConsoleManager.h"
 #include "NetSerialization.generated.h"
 
 class Error;
@@ -547,10 +548,22 @@ struct FFastArraySerializer
 		return DeltaFlags;
 	}
 
+	static const int32 GetMaxNumberOfAllowedChangesPerUpdate()
+	{
+		return MaxNumberOfAllowedChangesPerUpdate;
+	}
+
+	static const int32 GetMaxNumberOfAllowedDeletionsPerUpdate()
+	{
+		return MaxNumberOfAllowedDeletionsPerUpdate;
+	}
+
 private:
 
-	static constexpr int32 MAX_NUM_CHANGED = 2048;
-	static constexpr int32 MAX_NUM_DELETED = 2048;
+	ENGINE_API static int32 MaxNumberOfAllowedChangesPerUpdate;
+	ENGINE_API static int32 MaxNumberOfAllowedDeletionsPerUpdate;
+	ENGINE_API static FAutoConsoleVariableRef CVarMaxNumberOfAllowedChangesPerUpdate;
+	ENGINE_API static FAutoConsoleVariableRef CVarMaxNumberOfAllowedDeletionsPerUpdate;
 
 	/** Struct containing common header data that is written / read when serializing Fast Arrays. */
 	struct FFastArraySerializerHeader
@@ -882,6 +895,14 @@ void FFastArraySerializer::TFastArraySerializeHelper<Type, SerializerType>::Writ
 	UE_LOG(LogNetFastTArray, Log, TEXT("   Writing Bunch. NumChange: %d. NumDel: %d [%d/%d]"),
 		Header.NumChanged, Header.DeletedIndices.Num(), Header.ArrayReplicationKey, Header.BaseReplicationKey);
 
+	const int32 MaxNumDeleted = FFastArraySerializer::GetMaxNumberOfAllowedDeletionsPerUpdate();
+	const int32 MaxNumChanged = FFastArraySerializer::GetMaxNumberOfAllowedChangesPerUpdate();
+	
+	// TODO: We should consider propagating this error in the same way we handle
+	// array overflows in RepLayout SendProperties / CompareProperties.
+	UE_CLOG(NumDeletes > MaxNumDeleted, LogNetFastTArray, Warning, TEXT("NumDeletes > GetMaxNumberOfAllowedDeletionsPerUpdate: %d > %d. (Write)"), NumDeletes, MaxNumDeleted);
+	UE_CLOG(Header.NumChanged > MaxNumChanged, LogNetFastTArray, Warning, TEXT("NumChanged > GetMaxNumberOfAllowedChangesPerUpdate: %d > %d. (Write)"), Header.NumChanged, MaxNumChanged);
+
 	// Serialize deleted items, just by their ID
 	for (auto It = Header.DeletedIndices.CreateIterator(); It; ++It)
 	{
@@ -908,18 +929,20 @@ bool FFastArraySerializer::TFastArraySerializeHelper<Type, SerializerType>::Read
 
 	UE_LOG(LogNetFastTArray, Log, TEXT("Received [%d/%d]."), Header.ArrayReplicationKey, Header.BaseReplicationKey);
 
-	if (NumDeletes > MAX_NUM_DELETED)
+	const int32 MaxNumDeleted = FFastArraySerializer::GetMaxNumberOfAllowedDeletionsPerUpdate();
+	if (NumDeletes > MaxNumDeleted)
 	{
-		UE_LOG(LogNetFastTArray, Warning, TEXT("NumDeletes > MAX_NUM_DELETED: %d."), NumDeletes);
+		UE_LOG(LogNetFastTArray, Warning, TEXT("NumDeletes > GetMaxNumberOfAllowedDeletionsPerUpdate: %d > %d. (Read)"), NumDeletes, MaxNumDeleted);
 		Reader.SetError();
 		return false;
 	}
 
 	Reader << Header.NumChanged;
 
-	if (Header.NumChanged > MAX_NUM_CHANGED)
+	const int32 MaxNumChanged = FFastArraySerializer::GetMaxNumberOfAllowedChangesPerUpdate();
+	if (Header.NumChanged > MaxNumChanged)
 	{
-		UE_LOG(LogNetFastTArray, Warning, TEXT("NumChanged > MAX_NUM_CHANGED: %d."), Header.NumChanged);
+		UE_LOG(LogNetFastTArray, Warning, TEXT("NumChanged > GetMaxNumberOfAllowedChangesPerUpdate: %d > %d. (Read)"), Header.NumChanged, MaxNumChanged);
 		Reader.SetError();
 		return false;
 	}
diff --git a/Engine/Source/Runtime/Engine/Private/NetSerialization.cpp b/Engine/Source/Runtime/Engine/Private/NetSerialization.cpp
--- a/Engine/Source/Runtime/Engine/Private/NetSerialization.cpp
+++ b/Engine/Source/Runtime/Engine/Private/NetSerialization.cpp
@@ -1,7 +1,22 @@
 // Copyright Epic Games, Inc. All Rights Reserved.
 
 #include "Engine/NetSerialization.h"
+#include "HAL/IConsoleManager.h"
 
 DEFINE_STAT(STAT_NetSerializeFastArray);
 DEFINE_STAT(STAT_NetSerializeFastArray_BuildMap);
-DEFINE_STAT(STAT_NetSerializeFastArray_DeltaStruct);
\ No newline at end of file
+DEFINE_STAT(STAT_NetSerializeFastArray_DeltaStruct);
+
+int32 FFastArraySerializer::MaxNumberOfAllowedChangesPerUpdate = 2048;
+FAutoConsoleVariableRef FFastArraySerializer::CVarMaxNumberOfAllowedChangesPerUpdate(
+	TEXT("net.MaxNumberOfAllowedTArrayChangesPerUpdate"),
+	FFastArraySerializer::MaxNumberOfAllowedChangesPerUpdate,
+	TEXT("")
+);
+
+int32 FFastArraySerializer::MaxNumberOfAllowedDeletionsPerUpdate = 2048;
+FAutoConsoleVariableRef FFastArraySerializer::CVarMaxNumberOfAllowedDeletionsPerUpdate(
+	TEXT("net.MaxNumberOfAllowedTArrayDeletionsPerUpdate"),
+	FFastArraySerializer::MaxNumberOfAllowedDeletionsPerUpdate,
+	TEXT("")
+);
\ No newline at end of file
