fix missing nullptr check.

diff --git a/Engine/Source/Runtime/Engine/Private/SoundNodeQualityLevel.cpp b/Engine/Source/Runtime/Engine/Private/SoundNodeQualityLevel.cpp
--- a/Engine/Source/Runtime/Engine/Private/SoundNodeQualityLevel.cpp
+++ b/Engine/Source/Runtime/Engine/Private/SoundNodeQualityLevel.cpp
@@ -41,7 +41,7 @@ void USoundNodeQualityLevel::PostLoad()
 			{
 				continue;
 			}
-			else
+			else if(ChildNodes[Index] != nullptr)
 			{
 				ChildNodes[Index]->RemoveSoundWaveOnChildWavePlayers();
 			}
