Redo undone streamable handle optimization with a fix to bulk load handles being inadvertently cancelled

Pending bulk load requests would get cancelled when individual load request changed bundle state of a single asset.

Also fixed two invalid ensures.

diff --git a/Engine/Source/Runtime/Engine/Private/StreamableManager.cpp b/Engine/Source/Runtime/Engine/Private/StreamableManager.cpp
--- a/Engine/Source/Runtime/Engine/Private/StreamableManager.cpp
+++ b/Engine/Source/Runtime/Engine/Private/StreamableManager.cpp
@@ -1467,7 +1467,7 @@ void FStreamableManager::RemoveReferencedAsset(const FSoftObjectPath& Target, TS
 	// This should always be in the active handles list
 	if (ensureMsgf(Existing, TEXT("Failed to find existing streamable for %s"), *Target.ToString()))
 	{
-		ensureMsgf(Existing->ActiveHandles.Remove(Handle) > 0, TEXT("Failed to remove active handle for %s"), *Target.ToString());
+		ensureMsgf(Existing->ActiveHandles.RemoveSingleSwap(Handle) > 0, TEXT("Failed to remove active handle for %s"), *Target.ToString());
 
 		// Try removing from loading list if it's still there, this won't call the callback as it's being called from cancel
 		// This may remove more than one copy if streamables were merged
