Enable post-process calibration material in cooked builds. Changed "debug" -> "calibration" for clarity.

diff --git a/Engine/Source/Editor/UnrealEd/Private/EditorViewportClient.cpp b/Engine/Source/Editor/UnrealEd/Private/EditorViewportClient.cpp
--- a/Engine/Source/Editor/UnrealEd/Private/EditorViewportClient.cpp
+++ b/Engine/Source/Editor/UnrealEd/Private/EditorViewportClient.cpp
@@ -2566,15 +2566,15 @@ FText FEditorViewportClient::GetCurrentBufferVisualizationModeDisplayName() cons
 		? FBufferVisualizationData::GetMaterialDefaultDisplayName() : GetBufferVisualizationData().GetMaterialDisplayName(CurrentBufferVisualizationMode));
 }
 
-bool FEditorViewportClient::IsVisualizeDebugMaterialEnabled() const
+bool FEditorViewportClient::IsVisualizeCalibrationMaterialEnabled() const
 {
 	// Get the list of requested buffers from the console
 	const URendererSettings* Settings = GetDefault<URendererSettings>();
 	check(Settings);
 
-	return ((EngineShowFlags.VisualizeDebugCustomPostProcessMaterial && Settings->VisualizeDebugCustomMaterialPath.IsValid()) ||
-		(EngineShowFlags.VisualizeDebugColor  && Settings->VisualizeDebugColorMaterialPath.IsValid()) ||
-		(EngineShowFlags.VisualizeDebugGrayscale && Settings->VisualizeDebugGrayscaleMaterialPath.IsValid()));
+	return ((EngineShowFlags.VisualizeCalibrationCustom && Settings->VisualizeCalibrationCustomMaterialPath.IsValid()) ||
+		(EngineShowFlags.VisualizeCalibrationColor  && Settings->VisualizeCalibrationColorMaterialPath.IsValid()) ||
+		(EngineShowFlags.VisualizeCalibrationGrayscale && Settings->VisualizeCalibrationGrayscaleMaterialPath.IsValid()));
 }
 
 void FEditorViewportClient::ChangeRayTracingDebugVisualizationMode(FName InName)
@@ -2611,7 +2611,7 @@ bool FEditorViewportClient::SupportsPreviewResolutionFraction() const
 	}
 
 	// Don't do preview screen percentage for buffer visualization.
-	if (EngineShowFlags.VisualizeBuffer || IsVisualizeDebugMaterialEnabled())
+	if (EngineShowFlags.VisualizeBuffer || IsVisualizeCalibrationMaterialEnabled())
 	{
 		return false;
 	}
@@ -3877,7 +3877,7 @@ void FEditorViewportClient::Draw(FViewport* InViewport, FCanvas* Canvas)
 	}
 
 	// Axes indicators
-	if (bDrawAxes && !ViewFamily.EngineShowFlags.Game && !GLevelEditorModeTools().IsViewportUIHidden() && !IsVisualizeDebugMaterialEnabled())
+	if (bDrawAxes && !ViewFamily.EngineShowFlags.Game && !GLevelEditorModeTools().IsViewportUIHidden() && !IsVisualizeCalibrationMaterialEnabled())
 	{
 		switch (GetViewportType())
 		{
diff --git a/Engine/Source/Editor/UnrealEd/Public/EditorViewportClient.h b/Engine/Source/Editor/UnrealEd/Public/EditorViewportClient.h
--- a/Engine/Source/Editor/UnrealEd/Public/EditorViewportClient.h
+++ b/Engine/Source/Editor/UnrealEd/Public/EditorViewportClient.h
@@ -1192,7 +1192,7 @@ public:
 	/**
 	* Returns whether visualize debug material is enabled.
 	*/
-	bool IsVisualizeDebugMaterialEnabled() const;
+	bool IsVisualizeCalibrationMaterialEnabled() const;
 
 	/**
 	 * Changes the ray tracing debug visualization mode for this viewport
diff --git a/Engine/Source/Runtime/Engine/Classes/Engine/RendererSettings.h b/Engine/Source/Runtime/Engine/Classes/Engine/RendererSettings.h
--- a/Engine/Source/Runtime/Engine/Classes/Engine/RendererSettings.h
+++ b/Engine/Source/Runtime/Engine/Classes/Engine/RendererSettings.h
@@ -822,25 +822,25 @@ class ENGINE_API URendererSettings : public UDeveloperSettings
 	FPerPlatformBool bDiscardSkeletalMeshOptionalLODs;
 
 	/**
-	" Visualize debug material settings for post process debug materials, used for setting full-screen images used for monitor calibration."
+	" Visualize calibration material settings for post process calibration materials, used for setting full-screen images used for monitor calibration."
 	*/
-	UPROPERTY(config, EditAnywhere, Category = PostProcessDebugMaterials, meta = (AllowedClasses = "Material",
-		DisplayName = "Visualize Debug Color Material Path",
-		ToolTip = "When the VisualizeDebugColor show flag is enabled, this path will be used as the post-process material to render.",
+	UPROPERTY(config, EditAnywhere, Category = PostProcessCalibrationMaterials, meta = (AllowedClasses = "Material",
+		DisplayName = "Visualize Calibration Color Material Path",
+		ToolTip = "When the VisualizeCalibrationColor show flag is enabled, this path will be used as the post-process material to render.",
 		ConfigRestartRequired = true))
-	FSoftObjectPath VisualizeDebugColorMaterialPath;
+	FSoftObjectPath VisualizeCalibrationColorMaterialPath;
 
-	UPROPERTY(config, EditAnywhere, Category = PostProcessDebugMaterials, meta = (AllowedClasses = "Material",
-		DisplayName = "Visualize Debug Grayscale Material Path",
-		ToolTip = "When the VisualizeDebugGrayscale show flag is enabled, this path will be used as the post-process material to render.",
+	UPROPERTY(config, EditAnywhere, Category = PostProcessCalibrationMaterials, meta = (AllowedClasses = "Material",
+		DisplayName = "Visualize Calibration Grayscale Material Path",
+		ToolTip = "When the VisualizeCalibrationGrayscale show flag is enabled, this path will be used as the post-process material to render.",
 		ConfigRestartRequired = true))
-	FSoftObjectPath VisualizeDebugGrayscaleMaterialPath;
+	FSoftObjectPath VisualizeCalibrationGrayscaleMaterialPath;
 
-	UPROPERTY(config, EditAnywhere, Category = PostProcessDebugMaterials, meta = (AllowedClasses = "Material",
-		DisplayName = "Visualize Debug Custom Material Path",
-		ToolTip = "When the VisualizeDebugCustom show flag is enabled, this path will be used as the post-process material to render.",
+	UPROPERTY(config, EditAnywhere, Category = PostProcessCalibrationMaterials, meta = (AllowedClasses = "Material",
+		DisplayName = "Visualize Calibration Custom Material Path",
+		ToolTip = "When the VisualizeCalibrationCustom show flag is enabled, this path will be used as the post-process material to render.",
 		ConfigRestartRequired = true))
-	FSoftObjectPath VisualizeDebugCustomMaterialPath;
+	FSoftObjectPath VisualizeCalibrationCustomMaterialPath;
 
 public:
 
diff --git a/Engine/Source/Runtime/Engine/Private/SceneView.cpp b/Engine/Source/Runtime/Engine/Private/SceneView.cpp
--- a/Engine/Source/Runtime/Engine/Private/SceneView.cpp
+++ b/Engine/Source/Runtime/Engine/Private/SceneView.cpp
@@ -2011,11 +2011,15 @@ void FSceneView::EndFinalPostprocessSettings(const FSceneViewInitOptions& ViewIn
 	if (AllowDebugViewmodes())
 	{
 		ConfigureBufferVisualizationSettings();
-#if !(UE_BUILD_SHIPPING || UE_BUILD_TEST)
-		ConfigureDebugVisualizationSettings();
-#endif
 	}
 
+#if !(UE_BUILD_SHIPPING)
+	if (Family->EngineShowFlags.IsVisualizeCalibrationEnabled())
+	{
+		ConfigureVisualizeCalibrationSettings();
+	}
+#endif
+
 #if WITH_EDITOR
 	FHighResScreenshotConfig& Config = GetHighResScreenshotConfig();
 
@@ -2127,70 +2131,46 @@ void FSceneView::ConfigureBufferVisualizationSettings()
 	}
 }
 
-#if !(UE_BUILD_SHIPPING || UE_BUILD_TEST)
+#if !(UE_BUILD_SHIPPING)
 
-void FSceneView::ConfigureDebugVisualizationSettings()
+void FSceneView::ConfigureVisualizeCalibrationSettings()
 {
-	// Get the list of requested buffers from the console
 	const URendererSettings* Settings = GetDefault<URendererSettings>();
 	check(Settings);
-
-	if (Family->EngineShowFlags.VisualizeDebugCustomPostProcessMaterial)
+	
+	auto ConfigureCalibrationSettings = [](const FSoftObjectPath& InPath, UMaterialInterface*& OutMaterialInterface, FName& OutMaterialName)
 	{
-		if (Settings->VisualizeDebugCustomMaterialPath.IsValid())
+		if (InPath.IsValid())
 		{
-			if (UMaterial* Material = Cast<UMaterial>(Settings->VisualizeDebugCustomMaterialPath.TryLoad()))
+			if (UMaterial* Material = Cast<UMaterial>(InPath.TryLoad()))
 			{
-				FinalPostProcessSettings.DebugCustomVisualizationMaterial = Material;
-				CurrentVisualizeDebugCustomMaterialName = *Material->GetPathName();
+				OutMaterialInterface = Material;
+				OutMaterialName = *Material->GetPathName();
 			}
 			else
 			{
-				UE_LOG(LogBufferVisualization, Warning, TEXT("Error loading material '%s'"), *Settings->VisualizeDebugCustomMaterialPath.ToString());
-				FinalPostProcessSettings.DebugCustomVisualizationMaterial = nullptr;
-				CurrentVisualizeDebugCustomMaterialName = NAME_None;
+				UE_LOG(LogBufferVisualization, Warning, TEXT("Error loading material '%s'"), *InPath.ToString());
+				OutMaterialInterface = nullptr;
+				OutMaterialName = NAME_None;
 			}
 		}
-	}
+	};
 
-	if (Family->EngineShowFlags.VisualizeDebugColor)
+	if (Family->EngineShowFlags.VisualizeCalibrationColor)
 	{
-		if (Settings->VisualizeDebugColorMaterialPath.IsValid())
-		{
-			if (UMaterial* Material = Cast<UMaterial>(Settings->VisualizeDebugColorMaterialPath.TryLoad()))
-			{
-				FinalPostProcessSettings.DebugColorVisualizationMaterial = Material;
-				CurrentVisualizeDebugColorMaterialName = *Material->GetPathName();
-			}
-			else
-			{
-				UE_LOG(LogBufferVisualization, Warning, TEXT("Error loading material '%s'"), *Settings->VisualizeDebugColorMaterialPath.ToString());
-				FinalPostProcessSettings.DebugColorVisualizationMaterial = nullptr;
-				CurrentVisualizeDebugColorMaterialName = NAME_None;
-			}
-		}
+		ConfigureCalibrationSettings(Settings->VisualizeCalibrationColorMaterialPath, FinalPostProcessSettings.VisualizeCalibrationColorMaterial, CurrentVisualizeCalibrationColorMaterialName);
 	}
-
-	if (Family->EngineShowFlags.VisualizeDebugGrayscale)
+	else if (Family->EngineShowFlags.VisualizeCalibrationGrayscale)
 	{
-		if (Settings->VisualizeDebugGrayscaleMaterialPath.IsValid())
-		{
-			if (UMaterial* Material = Cast<UMaterial>(Settings->VisualizeDebugGrayscaleMaterialPath.TryLoad()))
-			{
-				FinalPostProcessSettings.DebugGrayscaleVisualizationMaterial = Material;
-				CurrentVisualizeDebugGrayscaleMaterialName = *Material->GetPathName();
-			}
-			else
-			{
-				UE_LOG(LogBufferVisualization, Warning, TEXT("Error loading material '%s'"), *Settings->VisualizeDebugGrayscaleMaterialPath.ToString());
-				FinalPostProcessSettings.DebugGrayscaleVisualizationMaterial = nullptr;
-				CurrentVisualizeDebugGrayscaleMaterialName = NAME_None;
-			}
-		}
+		ConfigureCalibrationSettings(Settings->VisualizeCalibrationGrayscaleMaterialPath, FinalPostProcessSettings.VisualizeCalibrationGrayscaleMaterial, CurrentVisualizeCalibrationGrayscaleMaterialName);
+	}
+	else if (Family->EngineShowFlags.VisualizeCalibrationCustom)
+	{
+		ConfigureCalibrationSettings(Settings->VisualizeCalibrationCustomMaterialPath, FinalPostProcessSettings.VisualizeCalibrationCustomMaterial, CurrentVisualizeCalibrationCustomMaterialName);
 	}
 }
 
-#endif
+#endif // #if !(UE_BUILD_SHIPPING)
 
 EShaderPlatform FSceneView::GetShaderPlatform() const
 {
diff --git a/Engine/Source/Runtime/Engine/Private/ShowFlags.cpp b/Engine/Source/Runtime/Engine/Private/ShowFlags.cpp
--- a/Engine/Source/Runtime/Engine/Private/ShowFlags.cpp
+++ b/Engine/Source/Runtime/Engine/Private/ShowFlags.cpp
@@ -333,11 +333,6 @@ void ApplyViewMode(EViewModeIndex ViewModeIndex, bool bPerspective, FEngineShowF
 	EngineShowFlags.SetPathTracing(ViewModeIndex == VMI_PathTracing);
 }
 
-static bool IsVisualizeDebugMaterialEnabled(FEngineShowFlags& EngineShowFlags)
-{
-	return (EngineShowFlags.VisualizeDebugColor || EngineShowFlags.VisualizeDebugGrayscale || EngineShowFlags.VisualizeDebugCustomPostProcessMaterial);
-}
-
 void EngineShowFlagOverride(EShowFlagInitMode ShowFlagInitMode, EViewModeIndex ViewModeIndex, FEngineShowFlags& EngineShowFlags, bool bCanDisableTonemapper)
 {
 	if(ShowFlagInitMode == ESFIM_Game)
@@ -514,8 +509,8 @@ void EngineShowFlagOverride(EShowFlagInitMode ShowFlagInitMode, EViewModeIndex V
 		}
 	}
 
-	// disable AA in full screen GBuffer visualization or debug material visualization
-	if (bCanDisableTonemapper && (EngineShowFlags.VisualizeBuffer || IsVisualizeDebugMaterialEnabled(EngineShowFlags)))
+	// disable AA in full screen GBuffer visualization or calibration material visualization
+	if (bCanDisableTonemapper && EngineShowFlags.VisualizeBuffer)
 	{
 		EngineShowFlags.SetTonemapper(false);
 	}
diff --git a/Engine/Source/Runtime/Engine/Public/FinalPostProcessSettings.h b/Engine/Source/Runtime/Engine/Public/FinalPostProcessSettings.h
--- a/Engine/Source/Runtime/Engine/Public/FinalPostProcessSettings.h
+++ b/Engine/Source/Runtime/Engine/Public/FinalPostProcessSettings.h
@@ -17,9 +17,9 @@ class FFinalPostProcessSettings : public FPostProcessSettings
 {
 public:
 	FFinalPostProcessSettings()
-	: DebugColorVisualizationMaterial(NULL)
-	, DebugGrayscaleVisualizationMaterial(NULL)
-	, DebugCustomVisualizationMaterial(NULL)
+	: VisualizeCalibrationColorMaterial(NULL)
+	, VisualizeCalibrationGrayscaleMaterial(NULL)
+	, VisualizeCalibrationCustomMaterial(NULL)
 	, HighResScreenshotMaterial(NULL)
 	, HighResScreenshotMaskMaterial(NULL)
 	, HighResScreenshotCaptureRegionMaterial(NULL)
@@ -171,14 +171,14 @@ public:
 	// List of materials to use in the buffer visualization overview
 	TArray<UMaterialInterface*> BufferVisualizationOverviewMaterials;
 
-	// Material to use for debug color visualization
-	UMaterialInterface* DebugColorVisualizationMaterial;
+	// Material to use for calibration color visualization
+	UMaterialInterface* VisualizeCalibrationColorMaterial;
 
 	// Material to use for debug grayscale visualization
-	UMaterialInterface* DebugGrayscaleVisualizationMaterial;
+	UMaterialInterface* VisualizeCalibrationGrayscaleMaterial;
 
 	// Material to use for debug custom post-processing visualization
-	UMaterialInterface* DebugCustomVisualizationMaterial;
+	UMaterialInterface* VisualizeCalibrationCustomMaterial;
 
 	// Material to use for rendering high res screenshot with mask. Post process expects this material to be set all the time.
 	UMaterialInterface* HighResScreenshotMaterial;
diff --git a/Engine/Source/Runtime/Engine/Public/SceneView.h b/Engine/Source/Runtime/Engine/Public/SceneView.h
--- a/Engine/Source/Runtime/Engine/Public/SceneView.h
+++ b/Engine/Source/Runtime/Engine/Public/SceneView.h
@@ -931,14 +931,14 @@ public:
 	/** Current buffer visualization mode */
 	FName CurrentBufferVisualizationMode;
 
-	/** Current visualize debug color material name */
-	FName CurrentVisualizeDebugColorMaterialName;
+	/** Current visualize calibration color material name */
+	FName CurrentVisualizeCalibrationColorMaterialName;
 
-	/** Current visualize debug grayscale material name */
-	FName CurrentVisualizeDebugGrayscaleMaterialName;
+	/** Current visualize calibration grayscale material name */
+	FName CurrentVisualizeCalibrationGrayscaleMaterialName;
 
-	/** Current visualize debug custom material name */
-	FName CurrentVisualizeDebugCustomMaterialName;
+	/** Current visualize calibration custom material name */
+	FName CurrentVisualizeCalibrationCustomMaterialName;
 
 #if WITH_EDITOR
 	/* Whether to use the pixel inspector */
@@ -1271,9 +1271,9 @@ public:
 	/** Configure post process settings for the buffer visualization system */
 	void ConfigureBufferVisualizationSettings();
 
-#if !(UE_BUILD_SHIPPING || UE_BUILD_TEST)
-	/** Configure post process settings for debug image */
-	void ConfigureDebugVisualizationSettings();
+#if !(UE_BUILD_SHIPPING)
+	/** Configure post process settings for calibration material */
+	void ConfigureVisualizeCalibrationSettings();
 #endif
 
 	/** Get the feature level for this view (cached from the scene so this is not different per view) **/
diff --git a/Engine/Source/Runtime/Engine/Public/ShowFlags.h b/Engine/Source/Runtime/Engine/Public/ShowFlags.h
--- a/Engine/Source/Runtime/Engine/Public/ShowFlags.h
+++ b/Engine/Source/Runtime/Engine/Public/ShowFlags.h
@@ -223,6 +223,10 @@ struct FEngineShowFlags
 		SetDistanceFieldGI(true);
 	}
 
+	bool IsVisualizeCalibrationEnabled() const
+	{
+		return (VisualizeCalibrationColor || VisualizeCalibrationGrayscale || VisualizeCalibrationCustom);
+	}
 
 	// ---------------------------------------------------------
 	// The following methods are there for serialization, localization and in general to iterate and manipulate flags.
@@ -400,9 +404,9 @@ private:
 		SetPathTracing(false);
 		SetRayTracingDebug(false);
 		SetVisualizeSkyAtmosphere(false);
-		SetVisualizeDebugColor(false);
-		SetVisualizeDebugGrayscale(false);
-		SetVisualizeDebugCustomPostProcessMaterial(false);
+		SetVisualizeCalibrationColor(false);
+		SetVisualizeCalibrationGrayscale(false);
+		SetVisualizeCalibrationCustom(false);
 	}
 
 
diff --git a/Engine/Source/Runtime/Engine/Public/ShowFlagsValues.inl b/Engine/Source/Runtime/Engine/Public/ShowFlagsValues.inl
--- a/Engine/Source/Runtime/Engine/Public/ShowFlagsValues.inl
+++ b/Engine/Source/Runtime/Engine/Public/ShowFlagsValues.inl
@@ -350,12 +350,12 @@ SHOWFLAG_FIXED_IN_SHIPPING(0, RayTracingDebug, SFG_Hidden, NSLOCTEXT("UnrealEd",
 
 /** Enable the SkyAtmosphere visualization to be drawn on screen */
 SHOWFLAG_FIXED_IN_SHIPPING(0, VisualizeSkyAtmosphere, SFG_Visualize, NSLOCTEXT("UnrealEd", "VisualizeSkyAtmosphereSF", "Sky Atmosphere"))
-/** Shows a debug color post process material defined in renderer config  */
-SHOWFLAG_FIXED_IN_SHIPPING(0, VisualizeDebugColor, SFG_Developer, NSLOCTEXT("UnrealEd", "VisualizeDebugColorSF", "Visualize Debug Color"))
-/** Shows a debug grayscale post process material defined in renderer config  */
-SHOWFLAG_FIXED_IN_SHIPPING(0, VisualizeDebugGrayscale, SFG_Developer, NSLOCTEXT("UnrealEd", "VisualizeDebugGrayscaleSF", "Visualize Debug Grayscale"))
-/** Shows a debug custom post process material defined in renderer config  */
-SHOWFLAG_FIXED_IN_SHIPPING(0, VisualizeDebugCustomPostProcessMaterial, SFG_Developer, NSLOCTEXT("UnrealEd", "VisualizeDebugCustomPostProcessMaterialSF", "Visualize Debug Custom"))
+/** Shows a full-screen calibration color post process material defined in renderer config  */
+SHOWFLAG_FIXED_IN_SHIPPING(0, VisualizeCalibrationColor, SFG_Developer, NSLOCTEXT("UnrealEd", "VisualizeCalibrationColorSF", "Visualize Calibration Color"))
+/** Shows a full-screen calibration grayscale post process material defined in renderer config  */
+SHOWFLAG_FIXED_IN_SHIPPING(0, VisualizeCalibrationGrayscale, SFG_Developer, NSLOCTEXT("UnrealEd", "VisualizeCalibrationGrayscaleSF", "Visualize Calibration Grayscale"))
+/** Shows a full-screen calibration custom post process material defined in renderer config  */
+SHOWFLAG_FIXED_IN_SHIPPING(0, VisualizeCalibrationCustom, SFG_Developer, NSLOCTEXT("UnrealEd", "VisualizeCalibrationCustomSF", "Visualize Calibration Custom"))
 
 #undef SHOWFLAG_ALWAYS_ACCESSIBLE
 #undef SHOWFLAG_FIXED_IN_SHIPPING
diff --git a/Engine/Source/Runtime/Renderer/Private/DeferredShadingRenderer.cpp b/Engine/Source/Runtime/Renderer/Private/DeferredShadingRenderer.cpp
--- a/Engine/Source/Runtime/Renderer/Private/DeferredShadingRenderer.cpp
+++ b/Engine/Source/Runtime/Renderer/Private/DeferredShadingRenderer.cpp
@@ -13,7 +13,7 @@
 #include "ScreenRendering.h"
 #include "PostProcess/SceneFilterRendering.h"
 #include "PostProcess/PostProcessSubsurface.h"
-#include "PostProcess/PostProcessVisualizeDebugMaterial.h"
+#include "PostProcess/PostProcessVisualizeCalibrationMaterial.h"
 #include "CompositionLighting/CompositionLighting.h"
 #include "FXSystem.h"
 #include "OneColorShader.h"
@@ -2915,16 +2915,16 @@ void FDeferredShadingSceneRenderer::Render(FRHICommandListImmediate& RHICmdList)
 			{
 				FViewInfo& View = Views[ViewIndex];
 
-#if !(UE_BUILD_SHIPPING || UE_BUILD_TEST)
+#if !(UE_BUILD_SHIPPING)
 
-				if (IsPostProcessVisualizeDebugMaterialEnabled(View))
+				if (IsPostProcessVisualizeCalibrationMaterialEnabled(View))
 				{
-					const UMaterialInterface* DebugMaterialInterface = GetPostProcessVisualizeDebugMaterialInterface(View);
+					const UMaterialInterface* DebugMaterialInterface = GetPostProcessVisualizeCalibrationMaterialInterface(View);
 					check(DebugMaterialInterface);
 
 					SCOPED_GPU_MASK(RHICmdList, View.GPUMask);
 					RDG_EVENT_SCOPE_CONDITIONAL(GraphBuilder, Views.Num() > 1, "View%d", ViewIndex);
-					AddVisualizeDebugMaterialPostProcessingPasses(GraphBuilder, View, PostProcessingInputs, DebugMaterialInterface);
+					AddVisualizeCalibrationMaterialPostProcessingPasses(GraphBuilder, View, PostProcessingInputs, DebugMaterialInterface);
 
 				}
 				else
@@ -2934,7 +2934,7 @@ void FDeferredShadingSceneRenderer::Render(FRHICommandListImmediate& RHICmdList)
 					RDG_EVENT_SCOPE_CONDITIONAL(GraphBuilder, Views.Num() > 1, "View%d", ViewIndex);
 					AddPostProcessingPasses(GraphBuilder, View, PostProcessingInputs);
 
-#if !(UE_BUILD_SHIPPING || UE_BUILD_TEST)
+#if !(UE_BUILD_SHIPPING)
 				}
 #endif
 			}
diff --git a/Engine/Source/Runtime/Renderer/Private/PostProcess/PostProcessVisualizeCalibrationMaterial.cpp b/Engine/Source/Runtime/Renderer/Private/PostProcess/PostProcessVisualizeCalibrationMaterial.cpp
new file mode 100644
--- /dev/null
+++ b/Engine/Source/Runtime/Renderer/Private/PostProcess/PostProcessVisualizeCalibrationMaterial.cpp
@@ -0,0 +1,43 @@
+// Copyright Epic Games, Inc. All Rights Reserved.
+
+#include "PostProcess/PostProcessVisualizeCalibrationMaterial.h"
+#include "EngineGlobals.h"
+
+static bool IsPostProcessVisualizeCalibrationColorMaterialEnabled(const FViewInfo& View) 
+{
+	return View.Family->EngineShowFlags.VisualizeCalibrationColor && View.CurrentVisualizeCalibrationColorMaterialName != NAME_None;
+}
+
+static bool IsPostProcessVisualizeCalibrationGrayscaleMaterialEnabled(const FViewInfo& View) 
+{
+	return View.Family->EngineShowFlags.VisualizeCalibrationGrayscale && View.CurrentVisualizeCalibrationGrayscaleMaterialName != NAME_None;
+}
+
+static bool IsPostProcessVisualizeCalibrationCustomMaterialEnabled(const FViewInfo& View) 
+{
+	return View.Family->EngineShowFlags.VisualizeCalibrationCustom && View.CurrentVisualizeCalibrationCustomMaterialName != NAME_None;
+}
+
+bool IsPostProcessVisualizeCalibrationMaterialEnabled(const FViewInfo& View) 
+{
+	return (IsPostProcessVisualizeCalibrationColorMaterialEnabled(View) || IsPostProcessVisualizeCalibrationGrayscaleMaterialEnabled(View) || IsPostProcessVisualizeCalibrationCustomMaterialEnabled(View));
+}
+
+// Returns whether the Calibration custom material pass needs to render on screen.
+const UMaterialInterface* GetPostProcessVisualizeCalibrationMaterialInterface(const FViewInfo& View) 
+{
+	if (IsPostProcessVisualizeCalibrationColorMaterialEnabled(View))
+	{
+		return View.FinalPostProcessSettings.VisualizeCalibrationColorMaterial;
+	}
+	else if (IsPostProcessVisualizeCalibrationGrayscaleMaterialEnabled(View))
+	{
+		return View.FinalPostProcessSettings.VisualizeCalibrationGrayscaleMaterial;
+	}
+	else if (IsPostProcessVisualizeCalibrationCustomMaterialEnabled(View))
+	{
+		return View.FinalPostProcessSettings.VisualizeCalibrationCustomMaterial;
+	}
+
+	return NULL;
+}
diff --git a/Engine/Source/Runtime/Renderer/Private/PostProcess/PostProcessVisualizeCalibrationMaterial.h b/Engine/Source/Runtime/Renderer/Private/PostProcess/PostProcessVisualizeCalibrationMaterial.h
new file mode 100644
--- /dev/null
+++ b/Engine/Source/Runtime/Renderer/Private/PostProcess/PostProcessVisualizeCalibrationMaterial.h
@@ -0,0 +1,11 @@
+// Copyright Epic Games, Inc. All Rights Reserved.
+
+#pragma once
+
+#include "ScreenPass.h"
+
+// Returns whether any calibration material pass is enabled.
+bool IsPostProcessVisualizeCalibrationMaterialEnabled(const FViewInfo& View);
+
+// Returns whether any calibration material pass is enabled.
+const UMaterialInterface* GetPostProcessVisualizeCalibrationMaterialInterface(const FViewInfo& View);
diff --git a/Engine/Source/Runtime/Renderer/Private/PostProcess/PostProcessVisualizeDebugMaterial.cpp b/Engine/Source/Runtime/Renderer/Private/PostProcess/PostProcessVisualizeDebugMaterial.cpp
deleted file mode 100644
diff --git a/Engine/Source/Runtime/Renderer/Private/PostProcess/PostProcessVisualizeDebugMaterial.h b/Engine/Source/Runtime/Renderer/Private/PostProcess/PostProcessVisualizeDebugMaterial.h
deleted file mode 100644
diff --git a/Engine/Source/Runtime/Renderer/Private/PostProcess/PostProcessing.cpp b/Engine/Source/Runtime/Renderer/Private/PostProcess/PostProcessing.cpp
--- a/Engine/Source/Runtime/Renderer/Private/PostProcess/PostProcessing.cpp
+++ b/Engine/Source/Runtime/Renderer/Private/PostProcess/PostProcessing.cpp
@@ -32,7 +32,7 @@
 #include "PostProcess/PostProcessVisualizeComplexity.h"
 #include "PostProcess/PostProcessCompositeEditorPrimitives.h"
 #include "PostProcess/PostProcessTestImage.h"
-#include "PostProcess/PostProcessVisualizeDebugMaterial.h"
+#include "PostProcess/PostProcessVisualizeCalibrationMaterial.h"
 #include "PostProcess/PostProcessFFTBloom.h"
 #include "PostProcess/PostProcessStreamingAccuracyLegend.h"
 #include "PostProcess/PostProcessSubsurface.h"
@@ -1077,9 +1077,9 @@ void AddDebugViewPostProcessingPasses(FRDGBuilder& GraphBuilder, const FViewInfo
 	}
 }
 
-#if !(UE_BUILD_SHIPPING || UE_BUILD_TEST)
+#if !(UE_BUILD_SHIPPING)
 
-void AddVisualizeDebugMaterialPostProcessingPasses(FRDGBuilder& GraphBuilder, const FViewInfo& View, const FPostProcessingInputs& Inputs, const UMaterialInterface* InMaterialInterface)
+void AddVisualizeCalibrationMaterialPostProcessingPasses(FRDGBuilder& GraphBuilder, const FViewInfo& View, const FPostProcessingInputs& Inputs, const UMaterialInterface* InMaterialInterface)
 {
 	CSV_SCOPED_TIMING_STAT_EXCLUSIVE(RenderPostProcessing);
 	QUICK_SCOPE_CYCLE_COUNTER(STAT_PostProcessing_Process);
diff --git a/Engine/Source/Runtime/Renderer/Private/PostProcess/PostProcessing.h b/Engine/Source/Runtime/Renderer/Private/PostProcess/PostProcessing.h
--- a/Engine/Source/Runtime/Renderer/Private/PostProcess/PostProcessing.h
+++ b/Engine/Source/Runtime/Renderer/Private/PostProcess/PostProcessing.h
@@ -57,9 +57,9 @@ void AddPostProcessingPasses(FRDGBuilder& GraphBuilder, const FViewInfo& View, c
 
 void AddDebugViewPostProcessingPasses(FRDGBuilder& GraphBuilder, const FViewInfo& View, const FPostProcessingInputs& Inputs);
 
-#if !(UE_BUILD_SHIPPING || UE_BUILD_TEST)
+#if !(UE_BUILD_SHIPPING)
 
-void AddVisualizeDebugMaterialPostProcessingPasses(FRDGBuilder& GraphBuilder, const FViewInfo& View, const FPostProcessingInputs& Inputs, const UMaterialInterface* InMaterialInterface);
+void AddVisualizeCalibrationMaterialPostProcessingPasses(FRDGBuilder& GraphBuilder, const FViewInfo& View, const FPostProcessingInputs& Inputs, const UMaterialInterface* InMaterialInterface);
 
 #endif
 
