Add texture LOD group setting to runtime virtual texture.
Add VT tile size bias to texture level of detail settings.
This allows us to set up tile size bias per platform to use smaller tile sizes on platforms with less processing power.

diff --git a/Engine/Source/Runtime/Engine/Classes/Components/RuntimeVirtualTextureComponent.h b/Engine/Source/Runtime/Engine/Classes/Components/RuntimeVirtualTextureComponent.h
--- a/Engine/Source/Runtime/Engine/Classes/Components/RuntimeVirtualTextureComponent.h
+++ b/Engine/Source/Runtime/Engine/Classes/Components/RuntimeVirtualTextureComponent.h
@@ -74,7 +74,7 @@ protected:
 	//~ End USceneComponent Interface
 
 	/** Calculate a hash used to determine if the StreamingTexture contents are valid for use. The hash doesn't include whether the contents are up to date. */
-	uint32 CalculateStreamingTextureSettingsHash() const;
+	uint64 CalculateStreamingTextureSettingsHash() const;
 	/** Returns true if the StreamingTexure contents are valid for use. */
 	bool IsStreamingTextureValid() const;
 
diff --git a/Engine/Source/Runtime/Engine/Classes/Engine/TextureLODSettings.h b/Engine/Source/Runtime/Engine/Classes/Engine/TextureLODSettings.h
--- a/Engine/Source/Runtime/Engine/Classes/Engine/TextureLODSettings.h
+++ b/Engine/Source/Runtime/Engine/Classes/Engine/TextureLODSettings.h
@@ -38,6 +38,7 @@ struct ENGINE_API FTextureLODGroup
 		, DuplicateNonOptionalMips(false)
 		, Downscale(1.0)
 		, DownscaleOptions(ETextureDownscaleOptions::SimpleAverage)
+		, VirtualTextureTileSizeBias(0)
 	{
 		SetupGroup();
 	}
@@ -112,7 +113,10 @@ struct ENGINE_API FTextureLODGroup
 
 	UPROPERTY()
 	ETextureDownscaleOptions DownscaleOptions;
-	
+
+	UPROPERTY()
+	int32 VirtualTextureTileSizeBias;
+
 	void SetupGroup();
 };
 
diff --git a/Engine/Source/Runtime/Engine/Classes/VT/RuntimeVirtualTexture.h b/Engine/Source/Runtime/Engine/Classes/VT/RuntimeVirtualTexture.h
--- a/Engine/Source/Runtime/Engine/Classes/VT/RuntimeVirtualTexture.h
+++ b/Engine/Source/Runtime/Engine/Classes/VT/RuntimeVirtualTexture.h
@@ -40,9 +40,9 @@ protected:
 	UPROPERTY(EditAnywhere, AdvancedDisplay, Category = Layout, meta = (DisplayName = "Enable private page table"))
 	bool bPrivateSpace = true;
 
-	/** Enable device scalability settings to modify the TileCount of the virtual texture. */
-	UPROPERTY(EditAnywhere, AdvancedDisplay, Category = Layout)
-	bool bEnableScalability = false;
+	/** Number of low mips to cut from the virtual texture. This can reduce peak virtual texture update cost but will also increase the probability of mip shimmering. */
+	UPROPERTY(EditAnywhere, AdvancedDisplay, Category = Layout, meta = (UIMin = "0", UIMax = "6", DisplayName = "Number of low mips to remove from the virtual texture"))
+	int32 RemoveLowMips = 0;
 
 	/** Size of virtual texture along the largest axis. (Actual values increase in powers of 2) */
 	UPROPERTY()
@@ -64,9 +64,9 @@ protected:
 	UPROPERTY(EditAnywhere, BluePrintGetter = GetTileBorderSize, Category = Size, meta = (UIMin = "0", UIMax = "4", DisplayName = "Border padding for each virtual texture tile"))
 	int32 TileBorderSize = 2; // 4
 
-	/** Number of low mips to cut from the virtual texture. This can reduce peak virtual texture update cost but will also increase the probability of mip shimmering. */
-	UPROPERTY(EditAnywhere, AdvancedDisplay, Category = Layout, meta = (UIMin = "0", UIMax = "6", DisplayName = "Number of low mips to remove from the virtual texture"))
-	int32 RemoveLowMips = 0;
+	/** Texture group this texture belongs to */
+	UPROPERTY(EditAnywhere, BlueprintReadWrite, Category = LevelOfDetail, meta = (DisplayName = "Texture Group"), AssetRegistrySearchable)
+	TEnumAsByte<enum TextureGroup> LODGroup;
 
 	/** Deprecated texture object containing streamed low mips. */
 	UPROPERTY()
@@ -88,11 +88,15 @@ public:
 	static int32 GetTileCount(int32 InTileCount) { return 1 << FMath::Clamp(InTileCount, 0, 12); }
 	/** Public getter for virtual texture tile size */
 	UFUNCTION(BlueprintGetter)
-	int32 GetTileSize() const { return 1 << FMath::Clamp(TileSize + 6, 6, 10); }
+	int32 GetTileSize() const { return GetTileSize(TileSize); }
+	static int32 GetTileSize(int32 InTileSize) { return 1 << FMath::Clamp(InTileSize + 6, 6, 10); }
 	/** Public getter for virtual texture tile border size */
 	UFUNCTION(BlueprintGetter)
 	int32 GetTileBorderSize() const { return 2 * FMath::Clamp(TileBorderSize, 0, 4); }
 	
+	/** Public getter for texture LOD Group */
+	TEnumAsByte<enum TextureGroup> GetLODGroup() const { return LODGroup; }
+
 	/** Get if this virtual texture uses compressed texture formats. */
 	bool GetCompressTextures() const { return bCompressTextures; }
 	/** Public getter for virtual texture removed low mips */
@@ -168,8 +172,8 @@ namespace RuntimeVirtualTexture
 	/** Helper function to wrap a runtime virtual texture producer with a streaming producer. */
 	ENGINE_API IVirtualTexture* CreateStreamingTextureProducer(
 		IVirtualTexture* InProducer,
-		URuntimeVirtualTexture* InBaseVirtualTexure,
+		FVTProducerDescription const& InProducerDesc,
 		UVirtualTexture2D* InStreamingTexture,
-		int32 InMaxLevel, int32 InNumMips,
+		int32 InMaxLevel,
 		int32& OutTransitionLevel);
 }
diff --git a/Engine/Source/Runtime/Engine/Classes/VT/VirtualTextureBuilder.h b/Engine/Source/Runtime/Engine/Classes/VT/VirtualTextureBuilder.h
--- a/Engine/Source/Runtime/Engine/Classes/VT/VirtualTextureBuilder.h
+++ b/Engine/Source/Runtime/Engine/Classes/VT/VirtualTextureBuilder.h
@@ -11,7 +11,7 @@
 /** Description object used to build the contents of a UVirtualTextureBuilder. */
 struct FVirtualTextureBuildDesc
 {
-	uint32 BuildHash = 0;
+	uint64 BuildHash = 0;
 
 	int32 LayerCount = 0;
 	TArray<ETextureSourceFormat, TInlineAllocator<4>> LayerFormats;
@@ -20,6 +20,8 @@ struct FVirtualTextureBuildDesc
 	int32 TileSize = 0;
 	int32 TileBorderSize = 0;
 
+	TEnumAsByte<enum TextureGroup> LODGroup;
+
 	bool bCrunchCompressed = false;
 	bool bSinglePhysicalSpace = false;
 
@@ -50,7 +52,7 @@ public:
 
 	/** Some client defined hash of that defines how the Texture was built. */
 	UPROPERTY()
-	uint32 BuildHash;
+	uint64 BuildHash;
 
 #if WITH_EDITOR
 	/** Creates a new UVirtualTexture2D and stores it in the contained Texture. */
diff --git a/Engine/Source/Runtime/Engine/Private/Components/RuntimeVirtualTextureComponent.cpp b/Engine/Source/Runtime/Engine/Private/Components/RuntimeVirtualTextureComponent.cpp
--- a/Engine/Source/Runtime/Engine/Private/Components/RuntimeVirtualTextureComponent.cpp
+++ b/Engine/Source/Runtime/Engine/Private/Components/RuntimeVirtualTextureComponent.cpp
@@ -61,7 +61,7 @@ FTransform URuntimeVirtualTextureComponent::GetVirtualTextureTransform() const
 	return FTransform(FVector(-0.5f, -0.5f, 0.f)) * GetComponentTransform();
 }
 
-uint32 URuntimeVirtualTextureComponent::CalculateStreamingTextureSettingsHash() const
+uint64 URuntimeVirtualTextureComponent::CalculateStreamingTextureSettingsHash() const
 {
 	// Shouldn't need to call this when there is VirtualTexture == nullptr
 	check(VirtualTexture != nullptr);
@@ -69,13 +69,14 @@ uint32 URuntimeVirtualTextureComponent::CalculateStreamingTextureSettingsHash()
 	// If a setting change can cause the streaming texture to no longer be valid then it should be included in this hash.
 	union FPackedSettings
 	{
-		uint32 PackedValue;
+		uint64 PackedValue;
 		struct
 		{
 			uint32 MaterialType : 4;
 			uint32 TileSize : 12;
 			uint32 TileBorderSize : 4;
 			uint32 StreamLowMips : 4;
+			uint32 LODGroup : 8;
 			uint32 CompressTextures : 1;
 			uint32 SinglePhysicalSpace : 1;
 			uint32 EnableCompressCrunch : 1;
@@ -88,6 +89,7 @@ uint32 URuntimeVirtualTextureComponent::CalculateStreamingTextureSettingsHash()
 	Settings.TileSize = (uint32)VirtualTexture->GetTileSize();
 	Settings.TileBorderSize = (uint32)VirtualTexture->GetTileBorderSize();
 	Settings.StreamLowMips = (uint32)StreamLowMips;
+	Settings.LODGroup = (uint32)VirtualTexture->GetLODGroup();
 	Settings.CompressTextures = (uint32)VirtualTexture->GetCompressTextures();
 	Settings.SinglePhysicalSpace = (uint32)VirtualTexture->GetSinglePhysicalSpace();
 	Settings.EnableCompressCrunch = (uint32)bEnableCompressCrunch;
@@ -126,6 +128,7 @@ void URuntimeVirtualTextureComponent::InitializeStreamingTexture(uint32 InSizeX,
 
 		BuildDesc.TileSize = VirtualTexture->GetTileSize();
 		BuildDesc.TileBorderSize = VirtualTexture->GetTileBorderSize();
+		BuildDesc.LODGroup = VirtualTexture->GetLODGroup();
 		BuildDesc.bCrunchCompressed = bEnableCompressCrunch;
 
 		BuildDesc.LayerCount = VirtualTexture->GetLayerCount();
diff --git a/Engine/Source/Runtime/Engine/Private/TextureDerivedData.cpp b/Engine/Source/Runtime/Engine/Private/TextureDerivedData.cpp
--- a/Engine/Source/Runtime/Engine/Private/TextureDerivedData.cpp
+++ b/Engine/Source/Runtime/Engine/Private/TextureDerivedData.cpp
@@ -459,7 +459,12 @@ static void GetTextureBuildSettings(
 		OutBuildSettings.bVirtualTextureEnableCompressCrunch = VirtualTextureBuildSettings.bEnableCompressCrunch;
 		OutBuildSettings.VirtualTextureTileSize = FMath::RoundUpToPowerOfTwo(VirtualTextureBuildSettings.TileSize);
 
-		// don't all max resolution to be less than VT tile size
+		// Apply any LOD group tile size bias here
+		const int32 TileSizeBias = TextureLODSettings.GetTextureLODGroup(Texture.LODGroup).VirtualTextureTileSizeBias;
+		OutBuildSettings.VirtualTextureTileSize >>= (TileSizeBias < 0) ? -TileSizeBias : 0;
+		OutBuildSettings.VirtualTextureTileSize <<= (TileSizeBias > 0) ? TileSizeBias : 0;
+
+		// Don't allow max resolution to be less than VT tile size
 		OutBuildSettings.MaxTextureResolution = FMath::Max<uint32>(OutBuildSettings.MaxTextureResolution, OutBuildSettings.VirtualTextureTileSize);
 
 		// 0 is a valid value for border size
diff --git a/Engine/Source/Runtime/Engine/Private/VT/RuntimeVirtualTexture.cpp b/Engine/Source/Runtime/Engine/Private/VT/RuntimeVirtualTexture.cpp
--- a/Engine/Source/Runtime/Engine/Private/VT/RuntimeVirtualTexture.cpp
+++ b/Engine/Source/Runtime/Engine/Private/VT/RuntimeVirtualTexture.cpp
@@ -2,7 +2,11 @@
 
 #include "VT/RuntimeVirtualTexture.h"
 
+#include "DeviceProfiles/DeviceProfile.h"
+#include "DeviceProfiles/DeviceProfileManager.h"
 #include "EngineModule.h"
+#include "Engine/TextureLODSettings.h"
+#include "Interfaces/ITargetPlatform.h"
 #include "RendererInterface.h"
 #include "VT/RuntimeVirtualTextureNotify.h"
 #include "VT/VirtualTexture.h"
@@ -184,17 +188,22 @@ void URuntimeVirtualTexture::GetProducerDescription(FVTProducerDescription& OutD
 {
 	OutDesc.Name = GetFName();
 	OutDesc.Dimensions = 2;
-	OutDesc.TileSize = GetTileSize();
-	OutDesc.TileBorderSize = GetTileBorderSize();
 	OutDesc.DepthInTiles = 1;
 	OutDesc.WidthInBlocks = 1;
 	OutDesc.HeightInBlocks = 1;
 
-	// Apply TileCount modifier here to allow size scalability option
+	// Apply LODGroup TileSize bias here.
+	const int32 TileSizeBias = UDeviceProfileManager::Get().GetActiveProfile()->GetTextureLODSettings()->GetTextureLODGroup(LODGroup).VirtualTextureTileSizeBias;
+	OutDesc.TileSize = GetTileSize(TileSize + TileSizeBias);
+
+	OutDesc.TileBorderSize = GetTileBorderSize();
+
+	// Apply runtime scalability option for TileCount bias here.
+	// We could also add a LODGroup driven bias in the future if it is useful.
 	const int32 TileCountBias = VirtualTextureScalability::GetRuntimeVirtualTextureSizeBias();
 	const int32 MaxSizeInTiles = GetTileCount(TileCount + TileCountBias);
 
-	// Set width and height to best match the runtime virtual texture volume's aspect ratio
+	// Set width and height to best match the runtime virtual texture volume's aspect ratio.
 	const FVector VolumeSize = VolumeToWorld.GetScale3D();
 	const float VolumeSizeX = FMath::Max(FMath::Abs(VolumeSize.X), 0.0001f);
 	const float VolumeSizeY = FMath::Max(FMath::Abs(VolumeSize.Y), 0.0001f);
@@ -469,28 +478,35 @@ namespace RuntimeVirtualTexture
 {
 	IVirtualTexture* CreateStreamingTextureProducer(
 		IVirtualTexture* InProducer,
-		URuntimeVirtualTexture* InBaseVirtualTexure,
+		FVTProducerDescription const& InProducerDesc,
 		UVirtualTexture2D* InStreamingTexture,
-		int32 InMaxLevel, int32 InNumMips,
+		int32 InMaxLevel,
 		int32& OutTransitionLevel)
 	{
-		if (InBaseVirtualTexure != nullptr && InStreamingTexture != nullptr)
+		if (InProducer != nullptr && InStreamingTexture != nullptr)
 		{
 			FTexturePlatformData** StreamingTextureData = InStreamingTexture->GetRunningPlatformData();
 			if (StreamingTextureData != nullptr && *StreamingTextureData != nullptr)
 			{
 				FVirtualTextureBuiltData* VTData = (*StreamingTextureData)->VTData;
 
-				ensure(InBaseVirtualTexure->GetTileSize() == VTData->TileSize);
-				ensure(InBaseVirtualTexure->GetTileBorderSize() == VTData->TileBorderSize);
-				if (InBaseVirtualTexure->GetTileSize() == VTData->TileSize && InBaseVirtualTexure->GetTileBorderSize() == VTData->TileBorderSize)
+				ensure(InProducerDesc.TileSize == VTData->TileSize);
+				ensure(InProducerDesc.TileBorderSize == VTData->TileBorderSize);
+				if (InProducerDesc.TileSize == VTData->TileSize && InProducerDesc.TileBorderSize == VTData->TileBorderSize)
 				{
-					// Streaming data may have mips removed during cook
-					const int32 NumStreamMips = FMath::Min(InNumMips, (*StreamingTextureData)->GetNumVTMips());
-
-					OutTransitionLevel = FMath::Max(InMaxLevel - NumStreamMips + 1, 0);
-					IVirtualTexture* StreamingProducer = new FUploadingVirtualTexture(VTData, 0);
-					return new FVirtualTextureLevelRedirector(InProducer, StreamingProducer, OutTransitionLevel);
+					// Note that streaming data may have mips added/removed during cook.
+					const uint32 Size = FMath::Max(VTData->Width, VTData->Height);
+					const uint32 NumTiles = FMath::DivideAndRoundUp(Size, VTData->TileSize);
+					const uint32 NumMips = FMath::CeilLogTwo(NumTiles) + 1;
+
+					// If the streaming texture is bigger then the runtime virtual texture then offset the first mip.
+					const int32 TransitionLevel = InMaxLevel - (int32)NumMips + 1;
+					const int32 FirstStreamingMip = TransitionLevel < 0 ? -TransitionLevel : 0;
+					const int32 AdjustedTransitionLevel = TransitionLevel + FirstStreamingMip;
+					OutTransitionLevel = TransitionLevel;
+
+					IVirtualTexture* StreamingProducer = new FUploadingVirtualTexture(VTData, FirstStreamingMip);
+					return new FVirtualTextureLevelRedirector(InProducer, StreamingProducer, AdjustedTransitionLevel);
 				}
 			}
 		}
diff --git a/Engine/Source/Runtime/Engine/Private/VT/VirtualTextureBuildSettings.cpp b/Engine/Source/Runtime/Engine/Private/VT/VirtualTextureBuildSettings.cpp
--- a/Engine/Source/Runtime/Engine/Private/VT/VirtualTextureBuildSettings.cpp
+++ b/Engine/Source/Runtime/Engine/Private/VT/VirtualTextureBuildSettings.cpp
@@ -31,7 +31,7 @@ static TAutoConsoleVariable<int32> CVarVTTileBorderSize(
 void FVirtualTextureBuildSettings::Init()
 {
 	TileSize = CVarVTTileSize.GetValueOnAnyThread();
-	TileBorderSize = CVarVTTileBorderSize.GetValueOnAnyThread();;
+	TileBorderSize = CVarVTTileBorderSize.GetValueOnAnyThread();
 	bEnableCompressCrunch = CVarVTCompressCrunch.GetValueOnAnyThread() != 0;
 	bEnableCompressZlib = CVarVTCompressZlib.GetValueOnAnyThread() != 0;
 }
diff --git a/Engine/Source/Runtime/Engine/Private/VT/VirtualTextureBuilder.cpp b/Engine/Source/Runtime/Engine/Private/VT/VirtualTextureBuilder.cpp
--- a/Engine/Source/Runtime/Engine/Private/VT/VirtualTextureBuilder.cpp
+++ b/Engine/Source/Runtime/Engine/Private/VT/VirtualTextureBuilder.cpp
@@ -38,6 +38,7 @@ void UVirtualTextureBuilder::BuildTexture(FVirtualTextureBuildDesc const& BuildD
 
 	Texture = NewObject<UVirtualTexture2D>(this, TEXT("Texture"));
 	Texture->VirtualTextureStreaming = true;
+	Texture->LODGroup = BuildDesc.LODGroup;
 
 	Texture->Settings.Init();
 	Texture->Settings.TileSize = BuildDesc.TileSize;
diff --git a/Engine/Source/Runtime/Renderer/Private/VT/RuntimeVirtualTextureSceneProxy.cpp b/Engine/Source/Runtime/Renderer/Private/VT/RuntimeVirtualTextureSceneProxy.cpp
--- a/Engine/Source/Runtime/Renderer/Private/VT/RuntimeVirtualTextureSceneProxy.cpp
+++ b/Engine/Source/Runtime/Renderer/Private/VT/RuntimeVirtualTextureSceneProxy.cpp
@@ -40,17 +40,16 @@ FRuntimeVirtualTextureSceneProxy::FRuntimeVirtualTextureSceneProxy(URuntimeVirtu
 		// The Producer object created here will be passed into the virtual texture system which will take ownership.
 		IVirtualTexture* Producer = new FRuntimeVirtualTextureProducer(Desc, ProducerId, MaterialType, bClearTextures, InComponent->GetScene(), Transform, Bounds);
 
-		int32 StreamLowMips = InComponent->IsStreamingLowMips() ? InComponent->NumStreamingMips() : 0;
-		if (StreamLowMips > 0)
+		if (InComponent->IsStreamingLowMips())
 		{
 			UVirtualTexture2D* StreamingTexture = InComponent->GetStreamingTexture()->Texture;
 			// Streaming mips start from the MaxLevel before taking into account the RemoveLowMips
 			const int32 MaxLevel = FMath::CeilLogTwo(FMath::Max(Desc.BlockWidthInTiles, Desc.BlockHeightInTiles));
 			// Wrap our producer to use a streaming producer for low mips
 			int32 StreamingTransitionLevel;
-			Producer = RuntimeVirtualTexture::CreateStreamingTextureProducer(Producer, VirtualTexture, StreamingTexture, MaxLevel, StreamLowMips, StreamingTransitionLevel);
+			Producer = RuntimeVirtualTexture::CreateStreamingTextureProducer(Producer, Desc, StreamingTexture, MaxLevel, StreamingTransitionLevel);
 			// Any dirty flushes don't need to flush the streaming mips (they only change with a build step).
-			MaxDirtyLevel = FMath::Min(MaxDirtyLevel, StreamingTransitionLevel);
+			MaxDirtyLevel = FMath::Clamp(StreamingTransitionLevel, 0, MaxDirtyLevel);
 		}
 
 		// The Initialize() call will allocate the virtual texture by spawning work on the render thread.
