Add MapCheck test for incompatible RVT streaming texture

diff --git a/Engine/Source/Runtime/Engine/Classes/Components/RuntimeVirtualTextureComponent.h b/Engine/Source/Runtime/Engine/Classes/Components/RuntimeVirtualTextureComponent.h
--- a/Engine/Source/Runtime/Engine/Classes/Components/RuntimeVirtualTextureComponent.h
+++ b/Engine/Source/Runtime/Engine/Classes/Components/RuntimeVirtualTextureComponent.h
@@ -66,6 +66,7 @@ protected:
 	virtual void CreateRenderState_Concurrent(FRegisterComponentContext* Context) override;
 	virtual void SendRenderTransform_Concurrent() override;
 	virtual void DestroyRenderState_Concurrent() override;
+	virtual void CheckForErrors() override;
 	//~ End UActorComponent Interface
 
 	//~ Begin USceneComponent Interface
diff --git a/Engine/Source/Runtime/Engine/Private/Components/RuntimeVirtualTextureComponent.cpp b/Engine/Source/Runtime/Engine/Private/Components/RuntimeVirtualTextureComponent.cpp
--- a/Engine/Source/Runtime/Engine/Private/Components/RuntimeVirtualTextureComponent.cpp
+++ b/Engine/Source/Runtime/Engine/Private/Components/RuntimeVirtualTextureComponent.cpp
@@ -4,10 +4,14 @@
 
 #include "Components/PrimitiveComponent.h"
 #include "GameFramework/Actor.h"
+#include "Misc/UObjectToken.h"
+#include "Misc/MapErrors.h"
 #include "SceneInterface.h"
 #include "VT/RuntimeVirtualTexture.h"
 #include "VT/VirtualTextureBuilder.h"
 
+#define LOCTEXT_NAMESPACE "URuntimeVirtualTextureComponent"
+
 URuntimeVirtualTextureComponent::URuntimeVirtualTextureComponent(const FObjectInitializer& ObjectInitializer)
 	: Super(ObjectInitializer)
 	, SceneProxy(nullptr)
@@ -162,4 +166,20 @@ void URuntimeVirtualTextureComponent::InitializeStreamingTexture(uint32 InSizeX,
 	}
 }
 
+void URuntimeVirtualTextureComponent::CheckForErrors()
+{
+	Super::CheckForErrors();
+
+	// Check if streaming texture has been built with the latest settings. If not then it won't be used which would cause a performance regression.
+	if (VirtualTexture != nullptr && StreamingTexture != nullptr && StreamingTexture->Texture != nullptr && StreamingTexture->BuildHash != CalculateStreamingTextureSettingsHash())
+	{
+		FMessageLog("MapCheck").PerformanceWarning()
+			->AddToken(FUObjectToken::Create(this))
+			->AddToken(FTextToken::Create(LOCTEXT("RuntimeVirtualTextureComponent_StreamingTextureNeedsUpdate", "The settings have changed since the streaming texture was last rebuilt. Streaming mips are disabled.")))
+			->AddToken(FMapErrorToken::Create(FName(TEXT("RuntimeVirtualTextureComponent_StreamingTextureNeedsUpdate"))));
+	}
+}
+
 #endif
+
+#undef LOCTEXT_NAMESPACE
