Non-Editor build fix
#rnx

diff --git a/Engine/Source/Runtime/Engine/Classes/Components/RuntimeVirtualTextureComponent.h b/Engine/Source/Runtime/Engine/Classes/Components/RuntimeVirtualTextureComponent.h
--- a/Engine/Source/Runtime/Engine/Classes/Components/RuntimeVirtualTextureComponent.h
+++ b/Engine/Source/Runtime/Engine/Classes/Components/RuntimeVirtualTextureComponent.h
@@ -66,7 +66,9 @@ protected:
 	virtual void CreateRenderState_Concurrent(FRegisterComponentContext* Context) override;
 	virtual void SendRenderTransform_Concurrent() override;
 	virtual void DestroyRenderState_Concurrent() override;
+#if WITH_EDITOR
 	virtual void CheckForErrors() override;
+#endif
 	//~ End UActorComponent Interface
 
 	//~ Begin USceneComponent Interface
