Build fix
#rnx

diff --git a/Engine/Source/Runtime/Engine/Private/Components/RuntimeVirtualTextureComponent.cpp b/Engine/Source/Runtime/Engine/Private/Components/RuntimeVirtualTextureComponent.cpp
--- a/Engine/Source/Runtime/Engine/Private/Components/RuntimeVirtualTextureComponent.cpp
+++ b/Engine/Source/Runtime/Engine/Private/Components/RuntimeVirtualTextureComponent.cpp
@@ -3,7 +3,7 @@
 #include "Components/RuntimeVirtualTextureComponent.h"
 
 #include "Components/PrimitiveComponent.h"
-#include "GameFramework/Actor.h"
+#include "Logging/MessageLog.h"
 #include "Misc/UObjectToken.h"
 #include "Misc/MapErrors.h"
 #include "SceneInterface.h"
