Move part of the TickComponent that performs the movement to a virtual method

diff --git a/Engine/Source/Runtime/Engine/Classes/GameFramework/CharacterMovementComponent.h b/Engine/Source/Runtime/Engine/Classes/GameFramework/CharacterMovementComponent.h
--- a/Engine/Source/Runtime/Engine/Classes/GameFramework/CharacterMovementComponent.h
+++ b/Engine/Source/Runtime/Engine/Classes/GameFramework/CharacterMovementComponent.h
@@ -1590,6 +1590,11 @@ protected:
 	/** Use new physics after landing. Defaults to swimming if in water, walking otherwise. */
 	virtual void SetPostLandedPhysics(const FHitResult& Hit);
 
+	/** Updates acceleration and perform movement, called from the TickComponent on the authoritative side for controlled characters, 
+	 *	or on the client for characters without a controller when either playing root motion or bRunPhysicsWithNoController is true.
+	 */
+	virtual void ControlledCharacterMove(const FVector& InputVector, float DeltaSeconds);
+
 	/** Switch collision settings for NavWalking mode (ignore world collisions) */
 	virtual void SetNavWalkingPhysics(bool bEnable);
 
diff --git a/Engine/Source/Runtime/Engine/Private/Components/CharacterMovementComponent.cpp b/Engine/Source/Runtime/Engine/Private/Components/CharacterMovementComponent.cpp
--- a/Engine/Source/Runtime/Engine/Private/Components/CharacterMovementComponent.cpp
+++ b/Engine/Source/Runtime/Engine/Private/Components/CharacterMovementComponent.cpp
@@ -1317,28 +1317,9 @@ void UCharacterMovementComponent::TickComponent(float DeltaTime, enum ELevelTick
 		}
 
 		// Allow root motion to move characters that have no controller.
-		if( CharacterOwner->IsLocallyControlled() || (!CharacterOwner->Controller && bRunPhysicsWithNoController) || (!CharacterOwner->Controller && CharacterOwner->IsPlayingRootMotion()) )
+		if (CharacterOwner->IsLocallyControlled() || (!CharacterOwner->Controller && bRunPhysicsWithNoController) || (!CharacterOwner->Controller && CharacterOwner->IsPlayingRootMotion()))
 		{
-			{
-				SCOPE_CYCLE_COUNTER(STAT_CharUpdateAcceleration);
-
-				// We need to check the jump state before adjusting input acceleration, to minimize latency
-				// and to make sure acceleration respects our potentially new falling state.
-				CharacterOwner->CheckJumpInput(DeltaTime);
-
-				// apply input to acceleration
-				Acceleration = ScaleInputAcceleration(ConstrainInputAcceleration(InputVector));
-				AnalogInputModifier = ComputeAnalogInputModifier();
-			}
-
-			if (CharacterOwner->GetLocalRole() == ROLE_Authority)
-			{
-				PerformMovement(DeltaTime);
-			}
-			else if (bIsClient)
-			{
-				ReplicateMoveToServer(DeltaTime, Acceleration);
-			}
+			ControlledCharacterMove(InputVector, DeltaTime);
 		}
 		else if (CharacterOwner->GetRemoteRole() == ROLE_AutonomousProxy)
 		{
@@ -5530,6 +5511,30 @@ void UCharacterMovementComponent::SetPostLandedPhysics(const FHitResult& Hit)
 	}
 }
 
+void UCharacterMovementComponent::ControlledCharacterMove(const FVector& InputVector, float DeltaSeconds)
+{
+	{
+		SCOPE_CYCLE_COUNTER(STAT_CharUpdateAcceleration);
+
+		// We need to check the jump state before adjusting input acceleration, to minimize latency
+		// and to make sure acceleration respects our potentially new falling state.
+		CharacterOwner->CheckJumpInput(DeltaSeconds);
+
+		// apply input to acceleration
+		Acceleration = ScaleInputAcceleration(ConstrainInputAcceleration(InputVector));
+		AnalogInputModifier = ComputeAnalogInputModifier();
+	}
+
+	if (CharacterOwner->GetLocalRole() == ROLE_Authority)
+	{
+		PerformMovement(DeltaSeconds);
+	}
+	else if (CharacterOwner->GetLocalRole() == ROLE_AutonomousProxy && IsNetMode(NM_Client))
+	{
+		ReplicateMoveToServer(DeltaSeconds, Acceleration);
+	}
+}
+
 void UCharacterMovementComponent::SetNavWalkingPhysics(bool bEnable)
 {
 	if (UpdatedPrimitive)
