- Remove GEngine references where not necessary from AudioDeviceManager
- Add logging around device manager lifetime
- Fix crash on shutdown when attempting to debug audio device handles w/INSTRUMENT_AUDIODEVICE_HANDLES
- Clean-up device manager shutdown logic & crit section for when reseting device map in destructor (probably not necessary but safer)
- Notify UWorld & GameViewportClient when shutting down and forcefully destroying associated audio device prior to destruction
- Fix unregistration not being called if world is swapping Audio Device (don't think this happens anywhere, but to be safe)

diff --git a/Engine/Source/Runtime/Engine/Classes/Engine/GameViewportClient.h b/Engine/Source/Runtime/Engine/Classes/Engine/GameViewportClient.h
--- a/Engine/Source/Runtime/Engine/Classes/Engine/GameViewportClient.h
+++ b/Engine/Source/Runtime/Engine/Classes/Engine/GameViewportClient.h
@@ -109,6 +109,9 @@ protected:
 	/** Strong handle to the audio device used by this viewport. */
 	FAudioDeviceHandle AudioDevice;
 
+	/** Handle to delegate in case audio device is destroyed. */
+	FDelegateHandle AudioDeviceDestroyedHandle;
+
 public:
 
 	/** see enum EViewModeIndex */
diff --git a/Engine/Source/Runtime/Engine/Classes/Engine/World.h b/Engine/Source/Runtime/Engine/Classes/Engine/World.h
--- a/Engine/Source/Runtime/Engine/Classes/Engine/World.h
+++ b/Engine/Source/Runtime/Engine/Classes/Engine/World.h
@@ -3388,6 +3388,9 @@ public:
 private:
 	bool bIsRefreshingStreamingLevels;
 
+	/** Handle to delegate in case audio device is destroyed. */
+	FDelegateHandle AudioDeviceDestroyedHandle;
+
 public:
 
 	bool IsRefreshingStreamingLevels() const { return bIsRefreshingStreamingLevels; }
diff --git a/Engine/Source/Runtime/Engine/Private/AudioDeviceManager.cpp b/Engine/Source/Runtime/Engine/Private/AudioDeviceManager.cpp
--- a/Engine/Source/Runtime/Engine/Private/AudioDeviceManager.cpp
+++ b/Engine/Source/Runtime/Engine/Private/AudioDeviceManager.cpp
@@ -76,6 +76,57 @@ static FAutoConsoleCommand GReportAudioDevicesCommand(
 	})
 );
 
+namespace AudioDeviceManagerUtils
+{
+	FString PrintDeviceInfo(const Audio::FDeviceId InDeviceId, const EAudioDeviceScope InScope, const bool bInIsNonRealtime, const int32* NumHandles = nullptr, const TMap<uint32, FString>* InStackWalk = nullptr)
+	{
+		FString ScopeStr;
+		switch (InScope)
+		{
+			case EAudioDeviceScope::Shared:
+				ScopeStr = TEXT("Shared");
+				break;
+
+			case EAudioDeviceScope::Unique:
+				ScopeStr = TEXT("Unique");
+				break;
+
+			case EAudioDeviceScope::Default:
+			default:
+				ScopeStr = TEXT("Default");
+				break;
+		}
+
+		FString DeviceInfo = FString::Printf(
+			TEXT("                Id: %d, Scope: %s, Realtime: %s"),
+			InDeviceId,
+			*ScopeStr,
+			bInIsNonRealtime ? TEXT("False") : TEXT("True"));
+
+		if (!NumHandles)
+		{
+			return DeviceInfo;
+		}
+
+		DeviceInfo += FString::Printf(TEXT(", Num Handles: %d"), *NumHandles);
+
+#if INSTRUMENT_AUDIODEVICE_HANDLES
+		if (InStackWalk)
+		{
+			DeviceInfo += TEXT("\n            Active Handles:\n\n");
+			for (const TPair<uint32, FString>& StackWalkPairs : *InStackWalk)
+			{
+				DeviceInfo += StackWalkPairs.Value;
+				DeviceInfo += TEXT("\n\n");
+			}
+		}
+#endif
+
+		return DeviceInfo;
+	}
+}
+
+
 FAudioDeviceManager* FAudioDeviceManager::Singleton = nullptr;
 
 // Some stress tests:
@@ -148,20 +199,34 @@ FAudioDeviceManager::FAudioDeviceManager()
 
 FAudioDeviceManager::~FAudioDeviceManager()
 {
-	FAudioCommandFence Fence;
-	Fence.BeginFence();
-	Fence.Wait();
-	FAudioThread::StopAudioThread();
+	UE_LOG(LogAudio, Display, TEXT("Beginning Audio Device Manager Shutdown (Module: %s)..."), *AudioMixerModuleName);
 
-	MainAudioDeviceHandle.Reset();
+	TArray<Audio::FDeviceId> DeviceIds;
+	{
+		FScopeLock ScopeLock(&DeviceMapCriticalSection);
+		Devices.GetKeys(DeviceIds);
+	}
 
-	// Notify anyone listening to the device manager that we are about to destroy the audio device.
-	for (auto& Device : Devices)
+	if (DeviceIds.Num() > 0)
 	{
-		FAudioDeviceManagerDelegates::OnAudioDeviceDestroyed.Broadcast(Device.Key);
+		UE_LOG(LogAudio, Display, TEXT("Destroying %d Remaining Audio Device(s)..."), DeviceIds.Num());
+
+		// Notify anyone listening to the device manager that we are about to destroy the audio device.
+		for (Audio::FDeviceId DeviceId : DeviceIds)
+		{
+			FAudioDeviceManagerDelegates::OnAudioDeviceDestroyed.Broadcast(DeviceId);
+		}
 	}
 
-	Devices.Reset();
+	FAudioThread::StopAudioThread();
+
+	FCoreDelegates::ApplicationWillEnterBackgroundDelegate.RemoveAll(this);
+
+	{
+		FScopeLock ScopeLock(&DeviceMapCriticalSection);
+		MainAudioDeviceHandle.Reset();
+		Devices.Reset();
+	}
 
 	// Release any loaded buffers - this calls stop on any sources that need it
 	for (int32 Index = Buffers.Num() - 1; Index >= 0; Index--)
@@ -347,8 +412,6 @@ IAudioDeviceModule* FAudioDeviceManager::GetAudioDeviceModule()
 	return AudioDeviceModule;
 }
 
-
-
 FAudioDeviceParams FAudioDeviceManager::GetDefaultParamsForNewWorld()
 {
 	bool bCreateNewAudioDeviceForPlayInEditor = false;
@@ -402,6 +465,7 @@ void FAudioDeviceManager::RegisterWorld(UWorld* InWorld, Audio::FDeviceId Device
 	{
 		if (!DeviceContainer->WorldsUsingThisDevice.Contains(InWorld))
 		{
+			UE_LOG(LogAudio, Display, TEXT("Audio Device (ID: %d) registered with world '%s'."), DeviceId, *InWorld->GetName());
 			DeviceContainer->WorldsUsingThisDevice.AddUnique(InWorld);
 			FAudioDeviceManagerDelegates::OnWorldRegisteredToAudioDevice.Broadcast(InWorld, DeviceId);
 		}
@@ -419,6 +483,7 @@ void FAudioDeviceManager::UnregisterWorld(UWorld* InWorld, Audio::FDeviceId Devi
 	{
 		if (DeviceContainer->WorldsUsingThisDevice.Contains(InWorld))
 		{
+			UE_LOG(LogAudio, Display, TEXT("Audio Device unregistered from world '%s'."), *InWorld->GetName());
 			DeviceContainer->WorldsUsingThisDevice.Remove(InWorld);
 			FAudioDeviceManagerDelegates::OnWorldUnregisteredWithAudioDevice.Broadcast(InWorld, DeviceId);
 		}
@@ -455,7 +520,7 @@ bool FAudioDeviceManager::InitializeManager()
 		
 		if (!MainAudioDeviceHandle)
 		{
-			UE_LOG(LogAudio, Display, TEXT("Audio device could not be initialized. Please check the value for AudioDeviceModuleName and AudioMixerModuleName in [Platform]Engine.ini."));
+			UE_LOG(LogAudio, Display, TEXT("Main audio device could not be initialized. Please check the value for AudioDeviceModuleName and AudioMixerModuleName in [Platform]Engine.ini."));
 			return false;
 		}
 
@@ -545,13 +610,17 @@ bool FAudioDeviceManager::LoadDefaultAudioDeviceModule()
 
 FAudioDeviceHandle FAudioDeviceManager::CreateNewDevice(const FAudioDeviceParams& InParams)
 {
-	Audio::FDeviceId DeviceID = GetNewDeviceID();
+	const Audio::FDeviceId DeviceID = GetNewDeviceID();
+
+	FString DeviceInfo = AudioDeviceManagerUtils::PrintDeviceInfo(DeviceID, InParams.Scope, InParams.bIsNonRealtime);
+	UE_LOG(LogAudio, Display, TEXT("Creating Audio Device: %s"), *DeviceInfo);
 	Devices.Emplace(DeviceID, FAudioDeviceContainer(InParams, DeviceID, this));
+
 	FAudioDeviceContainer* ContainerPtr = Devices.Find(DeviceID);
 	check(ContainerPtr);
 	if (!ContainerPtr->Device)
 	{
-		UE_LOG(LogAudio, Display, TEXT("Audio device could not be initialized. Please check the value for AudioDeviceModuleName and AudioMixerModuleName in [Platform]Engine.ini."));
+		UE_LOG(LogAudio, Display, TEXT("Destroying Audio Device %d: could not be initialized. Check AudioDeviceModuleName and AudioMixerModuleName in [Platform]Engine.ini."), DeviceID);
 
 		// Initializing the audio device failed. Remove the device container and return an empty handle.
 		Devices.Remove(DeviceID);
@@ -569,14 +638,10 @@ FAudioDeviceHandle FAudioDeviceManager::CreateNewDevice(const FAudioDeviceParams
 
 bool FAudioDeviceManager::IsValidAudioDevice(Audio::FDeviceId Handle) const
 {
-	return Devices.Contains(Handle);
-}
+	FCriticalSection* ConstCastCritSection = const_cast<FCriticalSection*>(&DeviceMapCriticalSection);
+	FScopeLock ScopeLock(ConstCastCritSection);
 
-bool FAudioDeviceManager::ShutdownAudioDevice(Audio::FDeviceId Handle)
-{
-	// Make sure we have a non-null device ptr in the index slot, then delete it
-	Devices.Remove(Handle);
-	return true;
+	return Devices.Contains(Handle);
 }
 
 void FAudioDeviceManager::IncrementDevice(Audio::FDeviceId DeviceID)
@@ -611,22 +676,14 @@ void FAudioDeviceManager::DecrementDevice(Audio::FDeviceId DeviceID, UWorld* InW
 		}
 
 		FAudioDeviceManagerDelegates::OnAudioDeviceDestroyed.Broadcast(DeviceID);
+
+		UE_LOG(LogAudio, Display, TEXT("Destroying Audio Device %d: All handles released."), DeviceID);
 		Devices.Remove(DeviceID);
 
 		UnregisterWorld(InWorld, DeviceID);
 	}
 }
 
-bool FAudioDeviceManager::ShutdownAllAudioDevices()
-{
-	FCoreDelegates::ApplicationWillEnterBackgroundDelegate.RemoveAll(this);
-
-	MainAudioDeviceHandle.Reset();
-
-	Devices.Reset();
-	return true;
-}
-
 FAudioDeviceHandle FAudioDeviceManager::BuildNewHandle(FAudioDeviceContainer&Container, Audio::FDeviceId DeviceID, const FAudioDeviceParams &InParams)
 {
 	FAudioDeviceManager::Get()->IncrementDevice(DeviceID);
@@ -696,11 +753,16 @@ bool FAudioDeviceManager::Initialize()
 
 		if (bUseSound)
 		{
+			UE_LOG(LogAudio, Display, TEXT("Initializing Audio Device Manager..."));
 			Singleton = new FAudioDeviceManager();
 
-			if (!Singleton->InitializeManager())
+			if (Singleton->InitializeManager())
+			{
+				UE_LOG(LogAudio, Display, TEXT("Audio Device Manager Initialized"));
+			}
+			else
 			{
-				UE_LOG(LogAudio, Warning, TEXT("Initializing the Audio Device Manager Failed!"));
+				UE_LOG(LogAudio, Warning, TEXT("Audio Device Manager Initialization Failed!"));
 				delete Singleton;
 				Singleton = nullptr;
 			}
@@ -721,6 +783,8 @@ void FAudioDeviceManager::Shutdown()
 	{
 		delete Singleton;
 		Singleton = nullptr;
+
+		UE_LOG(LogAudio, Display, TEXT("Audio Device Manager Shutdown"));
 	}
 }
 
@@ -730,7 +794,8 @@ FAudioDeviceHandle FAudioDeviceManager::GetActiveAudioDevice()
 	{
 		return GetAudioDevice(ActiveAudioDeviceID);
 	}
-	return GEngine->GetMainAudioDevice();
+
+	return MainAudioDeviceHandle;
 }
 
 void FAudioDeviceManager::UpdateActiveAudioDevices(bool bGameTicking)
@@ -1043,32 +1108,20 @@ void FAudioDeviceManager::LogListOfAudioDevices()
 {
 	FString ListOfDevices;
 
-	for (auto& DeviceContainer : Devices)
+	for (const TPair<Audio::FDeviceId, FAudioDeviceContainer>& DevicePair : Devices)
 	{
-		FString DeviceInfo = FString::Printf(TEXT(R"(
-					Device %d:
-					Scope: %s 
-					Realtime: %s
-					Number Of Owners: %d 
-		)"),
-			DeviceContainer.Key,
-			DeviceContainer.Value.Scope == EAudioDeviceScope::Unique ? TEXT("Unique") : TEXT("Shared"),
-			DeviceContainer.Value.bIsNonRealtime ? TEXT("No") : TEXT("Yes"),
-			DeviceContainer.Value.NumberOfHandlesToThisDevice);
-
+		ListOfDevices += AudioDeviceManagerUtils::PrintDeviceInfo(
+			DevicePair.Key,
+			DevicePair.Value.Scope,
+			DevicePair.Value.bIsNonRealtime,
+			&DevicePair.Value.NumberOfHandlesToThisDevice
 #if INSTRUMENT_AUDIODEVICE_HANDLES
-		for (auto& StackWalkString : DeviceContainer.Value.HandleCreationStackWalks)
-		{
-			DeviceInfo += TEXT("Handle Created here still alive:\n");
-			DeviceInfo += StackWalkString.Value;
-			DeviceInfo += TEXT("\n\n");
-		}
+			, &DevicePair.Value.HandleCreationStackWalks
 #endif
-
-		ListOfDevices += DeviceInfo;
+		);
 	}
 
-	UE_LOG(LogAudio, Display, TEXT("List of devices: \n%s"), *ListOfDevices);
+	UE_LOG(LogAudio, Display, TEXT("Active Audio Devices:\n%s"), *ListOfDevices);
 }
 
 uint32 FAudioDeviceManager::GetNewDeviceID()
@@ -1307,14 +1360,14 @@ void FAudioDeviceManager::AppWillEnterBackground()
 	// Flush all commands to the audio thread and the audio render thread:
 	if (GCVarFlushAudioRenderCommandsOnSuspend)
 	{
-		if (GEngine && GEngine->GetMainAudioDevice())
+		if (MainAudioDeviceHandle.IsValid())
 		{
-			FAudioDeviceHandle AudioDevice = GEngine->GetMainAudioDevice();
-
-			FAudioThread::RunCommandOnAudioThread([AudioDevice]()
+			FAudioThread::RunCommandOnAudioThread([this]()
 			{
-				FAudioDevice* AudioDevicePtr = const_cast<FAudioDevice*>(AudioDevice.GetAudioDevice());
-				AudioDevicePtr->FlushAudioRenderingCommands(true);
+				if (MainAudioDeviceHandle.IsValid())
+				{
+					MainAudioDeviceHandle->FlushAudioRenderingCommands(true);
+				}
 			}, TStatId());
 		}
 
@@ -1385,7 +1438,7 @@ FAudioDeviceHandle::~FAudioDeviceHandle()
 	if (IsValid())
 	{
 		FAudioDeviceManager* AudioDeviceManager = FAudioDeviceManager::Get();
-		if (AudioDeviceManager)
+		if (ensure(AudioDeviceManager))
 		{
 			AudioDeviceManager->DecrementDevice(DeviceId, World);
 
@@ -1450,7 +1503,7 @@ FAudioDeviceHandle& FAudioDeviceHandle::operator=(FAudioDeviceHandle&& Other)
 	{
 #if INSTRUMENT_AUDIODEVICE_HANDLES
 		check(StackWalkID != INDEX_NONE);
-		GEngine->GetAudioDeviceManager()->RemoveStackWalkForContainer(DeviceId, StackWalkID);
+		FAudioDeviceManager::Get()->RemoveStackWalkForContainer(DeviceId, StackWalkID);
 #endif
 
 		FAudioDeviceManager::Get()->DecrementDevice(DeviceId, World);
diff --git a/Engine/Source/Runtime/Engine/Private/GameViewportClient.cpp b/Engine/Source/Runtime/Engine/Private/GameViewportClient.cpp
--- a/Engine/Source/Runtime/Engine/Private/GameViewportClient.cpp
+++ b/Engine/Source/Runtime/Engine/Private/GameViewportClient.cpp
@@ -287,6 +287,13 @@ UGameViewportClient::UGameViewportClient(const FObjectInitializer& ObjectInitial
 		}
 #endif
 
+		AudioDeviceDestroyedHandle = FAudioDeviceManagerDelegates::OnAudioDeviceDestroyed.AddLambda([this](const Audio::FDeviceId InDeviceId)
+		{
+			if (InDeviceId == AudioDevice.GetDeviceID())
+			{
+				AudioDevice.Reset();
+			}
+		});
 	}
 }
 
@@ -344,6 +351,7 @@ void UGameViewportClient::PostInitProperties()
 
 void UGameViewportClient::BeginDestroy()
 {
+	AudioDeviceDestroyedHandle.Reset();
 	AudioDevice.Reset();
 
 	RemoveAllViewportWidgets();
diff --git a/Engine/Source/Runtime/Engine/Private/World.cpp b/Engine/Source/Runtime/Engine/Private/World.cpp
--- a/Engine/Source/Runtime/Engine/Private/World.cpp
+++ b/Engine/Source/Runtime/Engine/Private/World.cpp
@@ -63,6 +63,7 @@
 #include "TickTaskManagerInterface.h"
 #include "FXSystem.h"
 #include "AudioDevice.h"
+#include "AudioDeviceManager.h"
 #include "VisualLogger/VisualLogger.h"
 #include "LevelUtils.h"
 #include "Physics/PhysicsInterfaceCore.h"
@@ -366,6 +367,15 @@ UWorld::UWorld( const FObjectInitializer& ObjectInitializer )
 	FWorldDelegates::OnPostWorldCreation.Broadcast(this);
 
 	PerfTrackers = new FWorldInGamePerformanceTrackers();
+
+	AudioDeviceDestroyedHandle = FAudioDeviceManagerDelegates::OnAudioDeviceDestroyed.AddLambda([this](const Audio::FDeviceId InDeviceId)
+		{
+			if (InDeviceId == AudioDeviceHandle.GetDeviceID())
+			{
+				FAudioDeviceHandle EmptyHandle;
+				SetAudioDevice(EmptyHandle);
+			}
+		});
 }
 
 UWorld::~UWorld()
@@ -848,6 +858,7 @@ void UWorld::BeginDestroy()
 		Scene->UpdateParameterCollections(TArray<FMaterialParameterCollectionInstanceResource*>());
 	}
 
+	AudioDeviceDestroyedHandle.Reset();
 	AudioDeviceHandle.Reset();
 }
 
@@ -7576,16 +7587,16 @@ void UWorld::SetAudioDevice(FAudioDeviceHandle& InHandle)
 	}
 
 	FAudioDeviceManager* DeviceManager = GEngine ? GEngine->GetAudioDeviceManager() : nullptr;
-	if (DeviceManager && InHandle.IsValid())
+	if (DeviceManager && AudioDeviceHandle.IsValid())
 	{
-		GEngine->GetAudioDeviceManager()->UnregisterWorld(this, InHandle.GetDeviceID());
+		GEngine->GetAudioDeviceManager()->UnregisterWorld(this, AudioDeviceHandle.GetDeviceID());
 	}
 
 	AudioDeviceHandle = InHandle;
 
-	if (DeviceManager)
+	if (DeviceManager && AudioDeviceHandle.IsValid())
 	{
-		GEngine->GetAudioDeviceManager()->RegisterWorld(this, InHandle.GetDeviceID());
+		GEngine->GetAudioDeviceManager()->RegisterWorld(this, AudioDeviceHandle.GetDeviceID());
 	}
 }
 
diff --git a/Engine/Source/Runtime/Engine/Public/AudioDeviceManager.h b/Engine/Source/Runtime/Engine/Public/AudioDeviceManager.h
--- a/Engine/Source/Runtime/Engine/Public/AudioDeviceManager.h
+++ b/Engine/Source/Runtime/Engine/Public/AudioDeviceManager.h
@@ -13,7 +13,7 @@ class FAudioDebugger;
 // Set this to one if you'd like to check who owns
 // handles to an audio device.
 #ifndef INSTRUMENT_AUDIODEVICE_HANDLES
-#define INSTRUMENT_AUDIODEVICE_HANDLES 0
+#define INSTRUMENT_AUDIODEVICE_HANDLES 1
 #endif
 
 class FReferenceCollector;
@@ -424,23 +424,12 @@ private:
 
 	FAudioDeviceHandle CreateNewDevice(const FAudioDeviceParams& InParams);
 
-	/**
-	* Shutsdown the audio device associated with the handle. The handle
-	* will become invalid after the audio device is shut down.
-	*/
-	bool ShutdownAudioDevice(Audio::FDeviceId DeviceID);
-
 	// Called exclusively by the FAudioDeviceHandle copy constructor and assignment operators:
 	void IncrementDevice(Audio::FDeviceId DeviceID);
 
 	// Called exclusively by the FAudioDeviceHandle dtor.
 	void DecrementDevice(Audio::FDeviceId DeviceID, UWorld* InWorld);
 
-	/**
-	* Shuts down all active audio devices
-	*/
-	bool ShutdownAllAudioDevices();
-
 	/** Application enters background handler */
 	void AppWillEnterBackground();
 
