IWYU Fix for monolythic build

diff --git a/Engine/Source/Runtime/Engine/Classes/Engine/World.h b/Engine/Source/Runtime/Engine/Classes/Engine/World.h
--- a/Engine/Source/Runtime/Engine/Classes/Engine/World.h
+++ b/Engine/Source/Runtime/Engine/Classes/Engine/World.h
@@ -9,6 +9,7 @@
 #include "UObject/Object.h"
 #include "Misc/Guid.h"
 #include "UObject/Class.h"
+#include "Delegates/IDelegateInstance.h"
 #include "Engine/EngineTypes.h"
 #include "Engine/EngineBaseTypes.h"
 #include "CollisionQueryParams.h"
