Missing include

diff --git a/Engine/Source/Runtime/Engine/Classes/Engine/Engine.h b/Engine/Source/Runtime/Engine/Classes/Engine/Engine.h
--- a/Engine/Source/Runtime/Engine/Classes/Engine/Engine.h
+++ b/Engine/Source/Runtime/Engine/Classes/Engine/Engine.h
@@ -16,6 +16,7 @@
 #include "Misc/FrameRate.h"
 #include "Subsystems/SubsystemCollection.h"
 #include "Subsystems/EngineSubsystem.h"
+#include "RHI.h"
 #include "AudioDeviceManager.h"
 #include "Engine.generated.h"
 
