Fixed crash in ray tracing when mesh stream is enabled due invalid LOD usage.

diff --git a/Engine/Source/Runtime/Engine/Private/StaticMeshRender.cpp b/Engine/Source/Runtime/Engine/Private/StaticMeshRender.cpp
--- a/Engine/Source/Runtime/Engine/Private/StaticMeshRender.cpp
+++ b/Engine/Source/Runtime/Engine/Private/StaticMeshRender.cpp
@@ -1643,7 +1643,7 @@ void FStaticMeshSceneProxy::GetDynamicRayTracingInstances(FRayTracingMaterialGat
 	}
 
 	uint8 PrimitiveDPG = GetStaticDepthPriorityGroup();
-	const uint32 LODIndex = GetLOD(Context.ReferenceView);
+	const uint32 LODIndex = FMath::Max(GetLOD(Context.ReferenceView), (int32)GetCurrentFirstLODIdx_RenderThread());
 	const FStaticMeshLODResources& LODModel = RenderData->LODResources[LODIndex];
 
 	FRayTracingGeometry& Geometry = DynamicRayTracingGeometries[LODIndex];
