Cleanup BlueprintFunctionLibrary.h to not have unused protected function

#rnx

diff --git a/Engine/Source/Runtime/Engine/Classes/Kismet/BlueprintFunctionLibrary.h b/Engine/Source/Runtime/Engine/Classes/Kismet/BlueprintFunctionLibrary.h
--- a/Engine/Source/Runtime/Engine/Classes/Kismet/BlueprintFunctionLibrary.h
+++ b/Engine/Source/Runtime/Engine/Classes/Kismet/BlueprintFunctionLibrary.h
@@ -19,9 +19,4 @@ class UBlueprintFunctionLibrary : public UObject
 	// UObject interface
 	ENGINE_API virtual int32 GetFunctionCallspace(UFunction* Function, FFrame* Stack) override;
 	// End of UObject interface
-
-protected:
-
-	//bool	CheckAuthority(UObject* 
-
 };
