Fix monolythic build issue with cyclic dependencies between World & AudioDeviceManager

diff --git a/Engine/Source/Runtime/Engine/Classes/Engine/World.h b/Engine/Source/Runtime/Engine/Classes/Engine/World.h
--- a/Engine/Source/Runtime/Engine/Classes/Engine/World.h
+++ b/Engine/Source/Runtime/Engine/Classes/Engine/World.h
@@ -212,6 +212,23 @@ private:
 	UWorld* World;
 };
 
+// List of delegates for the world being registered to an audio device.
+class ENGINE_API FAudioDeviceWorldDelegates
+{
+public:
+	// Called whenever a world is registered to an audio device. UWorlds are not guaranteed to be registered to the same
+	// audio device throughout their lifecycle, and there is no guarantee on the lifespan of both the UWorld and the Audio
+	// Device registered in this callback.
+	DECLARE_MULTICAST_DELEGATE_TwoParams(FOnWorldRegisteredToAudioDevice, const UWorld* /*InWorld */, Audio::FDeviceId /* AudioDeviceId*/);
+	static FOnWorldRegisteredToAudioDevice OnWorldRegisteredToAudioDevice;
+
+	// Called whenever a world is unregistered from an audio device. UWorlds are not guaranteed to be registered to the same
+	// audio device throughout their lifecycle, and there is no guarantee on the lifespan of both the UWorld and the Audio
+	// Device registered in this callback.
+	DECLARE_MULTICAST_DELEGATE_TwoParams(FOnWorldUnregisteredWithAudioDevice, const UWorld* /*InWorld */, Audio::FDeviceId /* AudioDeviceId*/);
+	static FOnWorldUnregisteredWithAudioDevice OnWorldUnregisteredWithAudioDevice;
+};
+
 /** class that encapsulates seamless world traveling */
 class FSeamlessTravelHandler
 {
diff --git a/Engine/Source/Runtime/Engine/Private/AudioDebug.cpp b/Engine/Source/Runtime/Engine/Private/AudioDebug.cpp
--- a/Engine/Source/Runtime/Engine/Private/AudioDebug.cpp
+++ b/Engine/Source/Runtime/Engine/Private/AudioDebug.cpp
@@ -11,6 +11,7 @@
 #include "DrawDebugHelpers.h"
 #include "DSP/Dsp.h"
 #include "Engine/Font.h"
+#include "Engine/World.h"
 #include "GameFramework/GameUserSettings.h"
 #include "HAL/IConsoleManager.h"
 #include "Misc/CommandLine.h"
@@ -509,7 +510,7 @@ namespace Audio
 	FAudioDebugger::FAudioDebugger()
 		: bVisualize3dDebug(0)
 	{
-		WorldRegisteredWithDeviceHandle = FAudioDeviceManagerDelegates::OnWorldRegisteredToAudioDevice.AddLambda([this](const UWorld* InWorld, FDeviceId InDeviceId)
+		WorldRegisteredWithDeviceHandle = FAudioDeviceWorldDelegates::OnWorldRegisteredToAudioDevice.AddLambda([this](const UWorld* InWorld, FDeviceId InDeviceId)
 		{
 			if (bAllowUsingDeprecatedDebugStats)
 			{
diff --git a/Engine/Source/Runtime/Engine/Private/AudioDeviceManager.cpp b/Engine/Source/Runtime/Engine/Private/AudioDeviceManager.cpp
--- a/Engine/Source/Runtime/Engine/Private/AudioDeviceManager.cpp
+++ b/Engine/Source/Runtime/Engine/Private/AudioDeviceManager.cpp
@@ -6,6 +6,7 @@
 #include "AudioDefines.h"
 #include "AudioDevice.h"
 #include "AudioMixerDevice.h"
+#include "Engine/World.h"
 #include "Sound/AudioSettings.h"
 #include "Sound/SoundWave.h"
 #include "GameFramework/GameUserSettings.h"
@@ -467,7 +468,7 @@ void FAudioDeviceManager::RegisterWorld(UWorld* InWorld, Audio::FDeviceId Device
 		{
 			UE_LOG(LogAudio, Display, TEXT("Audio Device (ID: %d) registered with world '%s'."), DeviceId, *InWorld->GetName());
 			DeviceContainer->WorldsUsingThisDevice.AddUnique(InWorld);
-			FAudioDeviceManagerDelegates::OnWorldRegisteredToAudioDevice.Broadcast(InWorld, DeviceId);
+			FAudioDeviceWorldDelegates::OnWorldRegisteredToAudioDevice.Broadcast(InWorld, DeviceId);
 		}
 	}
 }
@@ -485,7 +486,7 @@ void FAudioDeviceManager::UnregisterWorld(UWorld* InWorld, Audio::FDeviceId Devi
 		{
 			UE_LOG(LogAudio, Display, TEXT("Audio Device unregistered from world '%s'."), *InWorld->GetName());
 			DeviceContainer->WorldsUsingThisDevice.Remove(InWorld);
-			FAudioDeviceManagerDelegates::OnWorldUnregisteredWithAudioDevice.Broadcast(InWorld, DeviceId);
+			FAudioDeviceWorldDelegates::OnWorldUnregisteredWithAudioDevice.Broadcast(InWorld, DeviceId);
 		}
 	}
 }
@@ -1635,5 +1636,3 @@ FAudioDeviceManager::FAudioDeviceContainer::~FAudioDeviceContainer()
 
 FAudioDeviceManagerDelegates::FOnAudioDeviceCreated FAudioDeviceManagerDelegates::OnAudioDeviceCreated;
 FAudioDeviceManagerDelegates::FOnAudioDeviceDestroyed FAudioDeviceManagerDelegates::OnAudioDeviceDestroyed;
-FAudioDeviceManagerDelegates::FOnWorldRegisteredToAudioDevice FAudioDeviceManagerDelegates::OnWorldRegisteredToAudioDevice;
-FAudioDeviceManagerDelegates::FOnWorldUnregisteredWithAudioDevice FAudioDeviceManagerDelegates::OnWorldUnregisteredWithAudioDevice;
diff --git a/Engine/Source/Runtime/Engine/Private/World.cpp b/Engine/Source/Runtime/Engine/Private/World.cpp
--- a/Engine/Source/Runtime/Engine/Private/World.cpp
+++ b/Engine/Source/Runtime/Engine/Private/World.cpp
@@ -303,7 +303,9 @@ FScopedLevelCollectionContextSwitch::~FScopedLevelCollectionContextSwitch()
 	}
 }
 
+FAudioDeviceWorldDelegates::FOnWorldRegisteredToAudioDevice FAudioDeviceWorldDelegates::OnWorldRegisteredToAudioDevice;
 
+FAudioDeviceWorldDelegates::FOnWorldUnregisteredWithAudioDevice FAudioDeviceWorldDelegates::OnWorldUnregisteredWithAudioDevice;
 
 /*-----------------------------------------------------------------------------
 	UWorld implementation.
diff --git a/Engine/Source/Runtime/Engine/Public/AudioDeviceManager.h b/Engine/Source/Runtime/Engine/Public/AudioDeviceManager.h
--- a/Engine/Source/Runtime/Engine/Public/AudioDeviceManager.h
+++ b/Engine/Source/Runtime/Engine/Public/AudioDeviceManager.h
@@ -194,18 +194,6 @@ public:
 	// This delegate is called whenever an audio device is destroyed.
 	DECLARE_MULTICAST_DELEGATE_OneParam(FOnAudioDeviceDestroyed, Audio::FDeviceId /* AudioDeviceId*/);
 	static FOnAudioDeviceDestroyed OnAudioDeviceDestroyed;
-
-	// Called whenever a world is registered to an audio device. UWorlds are not guaranteed to be registered to the same
-	// audio device throughout their lifecycle, and there is no guarantee on the lifespan of both the UWorld and the Audio
-	// Device registered in this callback.
-	DECLARE_MULTICAST_DELEGATE_TwoParams(FOnWorldRegisteredToAudioDevice, const UWorld* /*InWorld */, Audio::FDeviceId /* AudioDeviceId*/);
-	static FOnWorldRegisteredToAudioDevice OnWorldRegisteredToAudioDevice;
-
-	// Called whenever a world is unregistered from an audio device. UWorlds are not guaranteed to be registered to the same
-	// audio device throughout their lifecycle, and there is no guarantee on the lifespan of both the UWorld and the Audio
-	// Device registered in this callback.
-	DECLARE_MULTICAST_DELEGATE_TwoParams(FOnWorldUnregisteredWithAudioDevice, const UWorld* /*InWorld */, Audio::FDeviceId /* AudioDeviceId*/);
-	static FOnWorldUnregisteredWithAudioDevice OnWorldUnregisteredWithAudioDevice;
 };
 
 /**
