Back out changelists: 13151948, 13148935, 13148649

diff --git a/Engine/Source/Runtime/Engine/Classes/Engine/GameViewportClient.h b/Engine/Source/Runtime/Engine/Classes/Engine/GameViewportClient.h
--- a/Engine/Source/Runtime/Engine/Classes/Engine/GameViewportClient.h
+++ b/Engine/Source/Runtime/Engine/Classes/Engine/GameViewportClient.h
@@ -109,9 +109,6 @@ protected:
 	/** Strong handle to the audio device used by this viewport. */
 	FAudioDeviceHandle AudioDevice;
 
-	/** Handle to delegate in case audio device is destroyed. */
-	FDelegateHandle AudioDeviceDestroyedHandle;
-
 public:
 
 	/** see enum EViewModeIndex */
diff --git a/Engine/Source/Runtime/Engine/Classes/Engine/World.h b/Engine/Source/Runtime/Engine/Classes/Engine/World.h
--- a/Engine/Source/Runtime/Engine/Classes/Engine/World.h
+++ b/Engine/Source/Runtime/Engine/Classes/Engine/World.h
@@ -9,7 +9,6 @@
 #include "UObject/Object.h"
 #include "Misc/Guid.h"
 #include "UObject/Class.h"
-#include "Delegates/IDelegateInstance.h"
 #include "Engine/EngineTypes.h"
 #include "Engine/EngineBaseTypes.h"
 #include "CollisionQueryParams.h"
@@ -212,23 +211,6 @@ private:
 	UWorld* World;
 };
 
-// List of delegates for the world being registered to an audio device.
-class ENGINE_API FAudioDeviceWorldDelegates
-{
-public:
-	// Called whenever a world is registered to an audio device. UWorlds are not guaranteed to be registered to the same
-	// audio device throughout their lifecycle, and there is no guarantee on the lifespan of both the UWorld and the Audio
-	// Device registered in this callback.
-	DECLARE_MULTICAST_DELEGATE_TwoParams(FOnWorldRegisteredToAudioDevice, const UWorld* /*InWorld */, Audio::FDeviceId /* AudioDeviceId*/);
-	static FOnWorldRegisteredToAudioDevice OnWorldRegisteredToAudioDevice;
-
-	// Called whenever a world is unregistered from an audio device. UWorlds are not guaranteed to be registered to the same
-	// audio device throughout their lifecycle, and there is no guarantee on the lifespan of both the UWorld and the Audio
-	// Device registered in this callback.
-	DECLARE_MULTICAST_DELEGATE_TwoParams(FOnWorldUnregisteredWithAudioDevice, const UWorld* /*InWorld */, Audio::FDeviceId /* AudioDeviceId*/);
-	static FOnWorldUnregisteredWithAudioDevice OnWorldUnregisteredWithAudioDevice;
-};
-
 /** class that encapsulates seamless world traveling */
 class FSeamlessTravelHandler
 {
@@ -3406,9 +3388,6 @@ public:
 private:
 	bool bIsRefreshingStreamingLevels;
 
-	/** Handle to delegate in case audio device is destroyed. */
-	FDelegateHandle AudioDeviceDestroyedHandle;
-
 public:
 
 	bool IsRefreshingStreamingLevels() const { return bIsRefreshingStreamingLevels; }
diff --git a/Engine/Source/Runtime/Engine/Private/AudioDebug.cpp b/Engine/Source/Runtime/Engine/Private/AudioDebug.cpp
--- a/Engine/Source/Runtime/Engine/Private/AudioDebug.cpp
+++ b/Engine/Source/Runtime/Engine/Private/AudioDebug.cpp
@@ -11,7 +11,6 @@
 #include "DrawDebugHelpers.h"
 #include "DSP/Dsp.h"
 #include "Engine/Font.h"
-#include "Engine/World.h"
 #include "GameFramework/GameUserSettings.h"
 #include "HAL/IConsoleManager.h"
 #include "Misc/CommandLine.h"
@@ -510,7 +509,7 @@ namespace Audio
 	FAudioDebugger::FAudioDebugger()
 		: bVisualize3dDebug(0)
 	{
-		WorldRegisteredWithDeviceHandle = FAudioDeviceWorldDelegates::OnWorldRegisteredToAudioDevice.AddLambda([this](const UWorld* InWorld, FDeviceId InDeviceId)
+		WorldRegisteredWithDeviceHandle = FAudioDeviceManagerDelegates::OnWorldRegisteredToAudioDevice.AddLambda([this](const UWorld* InWorld, FDeviceId InDeviceId)
 		{
 			if (bAllowUsingDeprecatedDebugStats)
 			{
diff --git a/Engine/Source/Runtime/Engine/Private/AudioDeviceManager.cpp b/Engine/Source/Runtime/Engine/Private/AudioDeviceManager.cpp
--- a/Engine/Source/Runtime/Engine/Private/AudioDeviceManager.cpp
+++ b/Engine/Source/Runtime/Engine/Private/AudioDeviceManager.cpp
@@ -6,7 +6,6 @@
 #include "AudioDefines.h"
 #include "AudioDevice.h"
 #include "AudioMixerDevice.h"
-#include "Engine/World.h"
 #include "Sound/AudioSettings.h"
 #include "Sound/SoundWave.h"
 #include "GameFramework/GameUserSettings.h"
@@ -77,57 +76,6 @@ static FAutoConsoleCommand GReportAudioDevicesCommand(
 	})
 );
 
-namespace AudioDeviceManagerUtils
-{
-	FString PrintDeviceInfo(const Audio::FDeviceId InDeviceId, const EAudioDeviceScope InScope, const bool bInIsNonRealtime, const int32* NumHandles = nullptr, const TMap<uint32, FString>* InStackWalk = nullptr)
-	{
-		FString ScopeStr;
-		switch (InScope)
-		{
-			case EAudioDeviceScope::Shared:
-				ScopeStr = TEXT("Shared");
-				break;
-
-			case EAudioDeviceScope::Unique:
-				ScopeStr = TEXT("Unique");
-				break;
-
-			case EAudioDeviceScope::Default:
-			default:
-				ScopeStr = TEXT("Default");
-				break;
-		}
-
-		FString DeviceInfo = FString::Printf(
-			TEXT("                Id: %d, Scope: %s, Realtime: %s"),
-			InDeviceId,
-			*ScopeStr,
-			bInIsNonRealtime ? TEXT("False") : TEXT("True"));
-
-		if (!NumHandles)
-		{
-			return DeviceInfo;
-		}
-
-		DeviceInfo += FString::Printf(TEXT(", Num Handles: %d"), *NumHandles);
-
-#if INSTRUMENT_AUDIODEVICE_HANDLES
-		if (InStackWalk)
-		{
-			DeviceInfo += TEXT("\n            Active Handles:\n\n");
-			for (const TPair<uint32, FString>& StackWalkPairs : *InStackWalk)
-			{
-				DeviceInfo += StackWalkPairs.Value;
-				DeviceInfo += TEXT("\n\n");
-			}
-		}
-#endif
-
-		return DeviceInfo;
-	}
-}
-
-
 FAudioDeviceManager* FAudioDeviceManager::Singleton = nullptr;
 
 // Some stress tests:
@@ -200,35 +148,21 @@ FAudioDeviceManager::FAudioDeviceManager()
 
 FAudioDeviceManager::~FAudioDeviceManager()
 {
-	UE_LOG(LogAudio, Display, TEXT("Beginning Audio Device Manager Shutdown (Module: %s)..."), *AudioMixerModuleName);
-
-	TArray<Audio::FDeviceId> DeviceIds;
-	{
-		FScopeLock ScopeLock(&DeviceMapCriticalSection);
-		Devices.GetKeys(DeviceIds);
-	}
-
-	if (DeviceIds.Num() > 0)
-	{
-		UE_LOG(LogAudio, Display, TEXT("Destroying %d Remaining Audio Device(s)..."), DeviceIds.Num());
-
-		// Notify anyone listening to the device manager that we are about to destroy the audio device.
-		for (Audio::FDeviceId DeviceId : DeviceIds)
-		{
-			FAudioDeviceManagerDelegates::OnAudioDeviceDestroyed.Broadcast(DeviceId);
-		}
-	}
-
+	FAudioCommandFence Fence;
+	Fence.BeginFence();
+	Fence.Wait();
 	FAudioThread::StopAudioThread();
 
-	FCoreDelegates::ApplicationWillEnterBackgroundDelegate.RemoveAll(this);
+	MainAudioDeviceHandle.Reset();
 
+	// Notify anyone listening to the device manager that we are about to destroy the audio device.
+	for (auto& Device : Devices)
 	{
-		FScopeLock ScopeLock(&DeviceMapCriticalSection);
-		MainAudioDeviceHandle.Reset();
-		Devices.Reset();
+		FAudioDeviceManagerDelegates::OnAudioDeviceDestroyed.Broadcast(Device.Key);
 	}
 
+	Devices.Reset();
+
 	// Release any loaded buffers - this calls stop on any sources that need it
 	for (int32 Index = Buffers.Num() - 1; Index >= 0; Index--)
 	{
@@ -413,6 +347,8 @@ IAudioDeviceModule* FAudioDeviceManager::GetAudioDeviceModule()
 	return AudioDeviceModule;
 }
 
+
+
 FAudioDeviceParams FAudioDeviceManager::GetDefaultParamsForNewWorld()
 {
 	bool bCreateNewAudioDeviceForPlayInEditor = false;
@@ -466,9 +402,8 @@ void FAudioDeviceManager::RegisterWorld(UWorld* InWorld, Audio::FDeviceId Device
 	{
 		if (!DeviceContainer->WorldsUsingThisDevice.Contains(InWorld))
 		{
-			UE_LOG(LogAudio, Display, TEXT("Audio Device (ID: %d) registered with world '%s'."), DeviceId, *InWorld->GetName());
 			DeviceContainer->WorldsUsingThisDevice.AddUnique(InWorld);
-			FAudioDeviceWorldDelegates::OnWorldRegisteredToAudioDevice.Broadcast(InWorld, DeviceId);
+			FAudioDeviceManagerDelegates::OnWorldRegisteredToAudioDevice.Broadcast(InWorld, DeviceId);
 		}
 	}
 }
@@ -484,9 +419,8 @@ void FAudioDeviceManager::UnregisterWorld(UWorld* InWorld, Audio::FDeviceId Devi
 	{
 		if (DeviceContainer->WorldsUsingThisDevice.Contains(InWorld))
 		{
-			UE_LOG(LogAudio, Display, TEXT("Audio Device unregistered from world '%s'."), *InWorld->GetName());
 			DeviceContainer->WorldsUsingThisDevice.Remove(InWorld);
-			FAudioDeviceWorldDelegates::OnWorldUnregisteredWithAudioDevice.Broadcast(InWorld, DeviceId);
+			FAudioDeviceManagerDelegates::OnWorldUnregisteredWithAudioDevice.Broadcast(InWorld, DeviceId);
 		}
 	}
 }
@@ -521,7 +455,7 @@ bool FAudioDeviceManager::InitializeManager()
 		
 		if (!MainAudioDeviceHandle)
 		{
-			UE_LOG(LogAudio, Display, TEXT("Main audio device could not be initialized. Please check the value for AudioDeviceModuleName and AudioMixerModuleName in [Platform]Engine.ini."));
+			UE_LOG(LogAudio, Display, TEXT("Audio device could not be initialized. Please check the value for AudioDeviceModuleName and AudioMixerModuleName in [Platform]Engine.ini."));
 			return false;
 		}
 
@@ -611,17 +545,13 @@ bool FAudioDeviceManager::LoadDefaultAudioDeviceModule()
 
 FAudioDeviceHandle FAudioDeviceManager::CreateNewDevice(const FAudioDeviceParams& InParams)
 {
-	const Audio::FDeviceId DeviceID = GetNewDeviceID();
-
-	FString DeviceInfo = AudioDeviceManagerUtils::PrintDeviceInfo(DeviceID, InParams.Scope, InParams.bIsNonRealtime);
-	UE_LOG(LogAudio, Display, TEXT("Creating Audio Device: %s"), *DeviceInfo);
+	Audio::FDeviceId DeviceID = GetNewDeviceID();
 	Devices.Emplace(DeviceID, FAudioDeviceContainer(InParams, DeviceID, this));
-
 	FAudioDeviceContainer* ContainerPtr = Devices.Find(DeviceID);
 	check(ContainerPtr);
 	if (!ContainerPtr->Device)
 	{
-		UE_LOG(LogAudio, Display, TEXT("Destroying Audio Device %d: could not be initialized. Check AudioDeviceModuleName and AudioMixerModuleName in [Platform]Engine.ini."), DeviceID);
+		UE_LOG(LogAudio, Display, TEXT("Audio device could not be initialized. Please check the value for AudioDeviceModuleName and AudioMixerModuleName in [Platform]Engine.ini."));
 
 		// Initializing the audio device failed. Remove the device container and return an empty handle.
 		Devices.Remove(DeviceID);
@@ -639,12 +569,16 @@ FAudioDeviceHandle FAudioDeviceManager::CreateNewDevice(const FAudioDeviceParams
 
 bool FAudioDeviceManager::IsValidAudioDevice(Audio::FDeviceId Handle) const
 {
-	FCriticalSection* ConstCastCritSection = const_cast<FCriticalSection*>(&DeviceMapCriticalSection);
-	FScopeLock ScopeLock(ConstCastCritSection);
-
 	return Devices.Contains(Handle);
 }
 
+bool FAudioDeviceManager::ShutdownAudioDevice(Audio::FDeviceId Handle)
+{
+	// Make sure we have a non-null device ptr in the index slot, then delete it
+	Devices.Remove(Handle);
+	return true;
+}
+
 void FAudioDeviceManager::IncrementDevice(Audio::FDeviceId DeviceID)
 {
 	FScopeLock ScopeLock(&DeviceMapCriticalSection);
@@ -677,14 +611,22 @@ void FAudioDeviceManager::DecrementDevice(Audio::FDeviceId DeviceID, UWorld* InW
 		}
 
 		FAudioDeviceManagerDelegates::OnAudioDeviceDestroyed.Broadcast(DeviceID);
-
-		UE_LOG(LogAudio, Display, TEXT("Destroying Audio Device %d: All handles released."), DeviceID);
 		Devices.Remove(DeviceID);
 
 		UnregisterWorld(InWorld, DeviceID);
 	}
 }
 
+bool FAudioDeviceManager::ShutdownAllAudioDevices()
+{
+	FCoreDelegates::ApplicationWillEnterBackgroundDelegate.RemoveAll(this);
+
+	MainAudioDeviceHandle.Reset();
+
+	Devices.Reset();
+	return true;
+}
+
 FAudioDeviceHandle FAudioDeviceManager::BuildNewHandle(FAudioDeviceContainer&Container, Audio::FDeviceId DeviceID, const FAudioDeviceParams &InParams)
 {
 	FAudioDeviceManager::Get()->IncrementDevice(DeviceID);
@@ -754,16 +696,11 @@ bool FAudioDeviceManager::Initialize()
 
 		if (bUseSound)
 		{
-			UE_LOG(LogAudio, Display, TEXT("Initializing Audio Device Manager..."));
 			Singleton = new FAudioDeviceManager();
 
-			if (Singleton->InitializeManager())
-			{
-				UE_LOG(LogAudio, Display, TEXT("Audio Device Manager Initialized"));
-			}
-			else
+			if (!Singleton->InitializeManager())
 			{
-				UE_LOG(LogAudio, Warning, TEXT("Audio Device Manager Initialization Failed!"));
+				UE_LOG(LogAudio, Warning, TEXT("Initializing the Audio Device Manager Failed!"));
 				delete Singleton;
 				Singleton = nullptr;
 			}
@@ -784,8 +721,6 @@ void FAudioDeviceManager::Shutdown()
 	{
 		delete Singleton;
 		Singleton = nullptr;
-
-		UE_LOG(LogAudio, Display, TEXT("Audio Device Manager Shutdown"));
 	}
 }
 
@@ -795,8 +730,7 @@ FAudioDeviceHandle FAudioDeviceManager::GetActiveAudioDevice()
 	{
 		return GetAudioDevice(ActiveAudioDeviceID);
 	}
-
-	return MainAudioDeviceHandle;
+	return GEngine->GetMainAudioDevice();
 }
 
 void FAudioDeviceManager::UpdateActiveAudioDevices(bool bGameTicking)
@@ -1109,20 +1043,32 @@ void FAudioDeviceManager::LogListOfAudioDevices()
 {
 	FString ListOfDevices;
 
-	for (const TPair<Audio::FDeviceId, FAudioDeviceContainer>& DevicePair : Devices)
+	for (auto& DeviceContainer : Devices)
 	{
-		ListOfDevices += AudioDeviceManagerUtils::PrintDeviceInfo(
-			DevicePair.Key,
-			DevicePair.Value.Scope,
-			DevicePair.Value.bIsNonRealtime,
-			&DevicePair.Value.NumberOfHandlesToThisDevice
+		FString DeviceInfo = FString::Printf(TEXT(R"(
+					Device %d:
+					Scope: %s 
+					Realtime: %s
+					Number Of Owners: %d 
+		)"),
+			DeviceContainer.Key,
+			DeviceContainer.Value.Scope == EAudioDeviceScope::Unique ? TEXT("Unique") : TEXT("Shared"),
+			DeviceContainer.Value.bIsNonRealtime ? TEXT("No") : TEXT("Yes"),
+			DeviceContainer.Value.NumberOfHandlesToThisDevice);
+
 #if INSTRUMENT_AUDIODEVICE_HANDLES
-			, &DevicePair.Value.HandleCreationStackWalks
+		for (auto& StackWalkString : DeviceContainer.Value.HandleCreationStackWalks)
+		{
+			DeviceInfo += TEXT("Handle Created here still alive:\n");
+			DeviceInfo += StackWalkString.Value;
+			DeviceInfo += TEXT("\n\n");
+		}
 #endif
-		);
+
+		ListOfDevices += DeviceInfo;
 	}
 
-	UE_LOG(LogAudio, Display, TEXT("Active Audio Devices:\n%s"), *ListOfDevices);
+	UE_LOG(LogAudio, Display, TEXT("List of devices: \n%s"), *ListOfDevices);
 }
 
 uint32 FAudioDeviceManager::GetNewDeviceID()
@@ -1361,14 +1307,14 @@ void FAudioDeviceManager::AppWillEnterBackground()
 	// Flush all commands to the audio thread and the audio render thread:
 	if (GCVarFlushAudioRenderCommandsOnSuspend)
 	{
-		if (MainAudioDeviceHandle.IsValid())
+		if (GEngine && GEngine->GetMainAudioDevice())
 		{
-			FAudioThread::RunCommandOnAudioThread([this]()
+			FAudioDeviceHandle AudioDevice = GEngine->GetMainAudioDevice();
+
+			FAudioThread::RunCommandOnAudioThread([AudioDevice]()
 			{
-				if (MainAudioDeviceHandle.IsValid())
-				{
-					MainAudioDeviceHandle->FlushAudioRenderingCommands(true);
-				}
+				FAudioDevice* AudioDevicePtr = const_cast<FAudioDevice*>(AudioDevice.GetAudioDevice());
+				AudioDevicePtr->FlushAudioRenderingCommands(true);
 			}, TStatId());
 		}
 
@@ -1439,7 +1385,7 @@ FAudioDeviceHandle::~FAudioDeviceHandle()
 	if (IsValid())
 	{
 		FAudioDeviceManager* AudioDeviceManager = FAudioDeviceManager::Get();
-		if (ensure(AudioDeviceManager))
+		if (AudioDeviceManager)
 		{
 			AudioDeviceManager->DecrementDevice(DeviceId, World);
 
@@ -1504,7 +1450,7 @@ FAudioDeviceHandle& FAudioDeviceHandle::operator=(FAudioDeviceHandle&& Other)
 	{
 #if INSTRUMENT_AUDIODEVICE_HANDLES
 		check(StackWalkID != INDEX_NONE);
-		FAudioDeviceManager::Get()->RemoveStackWalkForContainer(DeviceId, StackWalkID);
+		GEngine->GetAudioDeviceManager()->RemoveStackWalkForContainer(DeviceId, StackWalkID);
 #endif
 
 		FAudioDeviceManager::Get()->DecrementDevice(DeviceId, World);
@@ -1636,3 +1582,5 @@ FAudioDeviceManager::FAudioDeviceContainer::~FAudioDeviceContainer()
 
 FAudioDeviceManagerDelegates::FOnAudioDeviceCreated FAudioDeviceManagerDelegates::OnAudioDeviceCreated;
 FAudioDeviceManagerDelegates::FOnAudioDeviceDestroyed FAudioDeviceManagerDelegates::OnAudioDeviceDestroyed;
+FAudioDeviceManagerDelegates::FOnWorldRegisteredToAudioDevice FAudioDeviceManagerDelegates::OnWorldRegisteredToAudioDevice;
+FAudioDeviceManagerDelegates::FOnWorldUnregisteredWithAudioDevice FAudioDeviceManagerDelegates::OnWorldUnregisteredWithAudioDevice;
diff --git a/Engine/Source/Runtime/Engine/Private/GameViewportClient.cpp b/Engine/Source/Runtime/Engine/Private/GameViewportClient.cpp
--- a/Engine/Source/Runtime/Engine/Private/GameViewportClient.cpp
+++ b/Engine/Source/Runtime/Engine/Private/GameViewportClient.cpp
@@ -287,13 +287,6 @@ UGameViewportClient::UGameViewportClient(const FObjectInitializer& ObjectInitial
 		}
 #endif
 
-		AudioDeviceDestroyedHandle = FAudioDeviceManagerDelegates::OnAudioDeviceDestroyed.AddLambda([this](const Audio::FDeviceId InDeviceId)
-		{
-			if (InDeviceId == AudioDevice.GetDeviceID())
-			{
-				AudioDevice.Reset();
-			}
-		});
 	}
 }
 
@@ -351,7 +344,6 @@ void UGameViewportClient::PostInitProperties()
 
 void UGameViewportClient::BeginDestroy()
 {
-	AudioDeviceDestroyedHandle.Reset();
 	AudioDevice.Reset();
 
 	RemoveAllViewportWidgets();
diff --git a/Engine/Source/Runtime/Engine/Private/World.cpp b/Engine/Source/Runtime/Engine/Private/World.cpp
--- a/Engine/Source/Runtime/Engine/Private/World.cpp
+++ b/Engine/Source/Runtime/Engine/Private/World.cpp
@@ -63,7 +63,6 @@
 #include "TickTaskManagerInterface.h"
 #include "FXSystem.h"
 #include "AudioDevice.h"
-#include "AudioDeviceManager.h"
 #include "VisualLogger/VisualLogger.h"
 #include "LevelUtils.h"
 #include "Physics/PhysicsInterfaceCore.h"
@@ -303,9 +302,7 @@ FScopedLevelCollectionContextSwitch::~FScopedLevelCollectionContextSwitch()
 	}
 }
 
-FAudioDeviceWorldDelegates::FOnWorldRegisteredToAudioDevice FAudioDeviceWorldDelegates::OnWorldRegisteredToAudioDevice;
 
-FAudioDeviceWorldDelegates::FOnWorldUnregisteredWithAudioDevice FAudioDeviceWorldDelegates::OnWorldUnregisteredWithAudioDevice;
 
 /*-----------------------------------------------------------------------------
 	UWorld implementation.
@@ -369,15 +366,6 @@ UWorld::UWorld( const FObjectInitializer& ObjectInitializer )
 	FWorldDelegates::OnPostWorldCreation.Broadcast(this);
 
 	PerfTrackers = new FWorldInGamePerformanceTrackers();
-
-	AudioDeviceDestroyedHandle = FAudioDeviceManagerDelegates::OnAudioDeviceDestroyed.AddLambda([this](const Audio::FDeviceId InDeviceId)
-		{
-			if (InDeviceId == AudioDeviceHandle.GetDeviceID())
-			{
-				FAudioDeviceHandle EmptyHandle;
-				SetAudioDevice(EmptyHandle);
-			}
-		});
 }
 
 UWorld::~UWorld()
@@ -860,7 +848,6 @@ void UWorld::BeginDestroy()
 		Scene->UpdateParameterCollections(TArray<FMaterialParameterCollectionInstanceResource*>());
 	}
 
-	AudioDeviceDestroyedHandle.Reset();
 	AudioDeviceHandle.Reset();
 }
 
@@ -7589,16 +7576,16 @@ void UWorld::SetAudioDevice(FAudioDeviceHandle& InHandle)
 	}
 
 	FAudioDeviceManager* DeviceManager = GEngine ? GEngine->GetAudioDeviceManager() : nullptr;
-	if (DeviceManager && AudioDeviceHandle.IsValid())
+	if (DeviceManager && InHandle.IsValid())
 	{
-		GEngine->GetAudioDeviceManager()->UnregisterWorld(this, AudioDeviceHandle.GetDeviceID());
+		GEngine->GetAudioDeviceManager()->UnregisterWorld(this, InHandle.GetDeviceID());
 	}
 
 	AudioDeviceHandle = InHandle;
 
-	if (DeviceManager && AudioDeviceHandle.IsValid())
+	if (DeviceManager)
 	{
-		GEngine->GetAudioDeviceManager()->RegisterWorld(this, AudioDeviceHandle.GetDeviceID());
+		GEngine->GetAudioDeviceManager()->RegisterWorld(this, InHandle.GetDeviceID());
 	}
 }
 
diff --git a/Engine/Source/Runtime/Engine/Public/AudioDeviceManager.h b/Engine/Source/Runtime/Engine/Public/AudioDeviceManager.h
--- a/Engine/Source/Runtime/Engine/Public/AudioDeviceManager.h
+++ b/Engine/Source/Runtime/Engine/Public/AudioDeviceManager.h
@@ -13,7 +13,7 @@ class FAudioDebugger;
 // Set this to one if you'd like to check who owns
 // handles to an audio device.
 #ifndef INSTRUMENT_AUDIODEVICE_HANDLES
-#define INSTRUMENT_AUDIODEVICE_HANDLES 1
+#define INSTRUMENT_AUDIODEVICE_HANDLES 0
 #endif
 
 class FReferenceCollector;
@@ -194,6 +194,18 @@ public:
 	// This delegate is called whenever an audio device is destroyed.
 	DECLARE_MULTICAST_DELEGATE_OneParam(FOnAudioDeviceDestroyed, Audio::FDeviceId /* AudioDeviceId*/);
 	static FOnAudioDeviceDestroyed OnAudioDeviceDestroyed;
+
+	// Called whenever a world is registered to an audio device. UWorlds are not guaranteed to be registered to the same
+	// audio device throughout their lifecycle, and there is no guarantee on the lifespan of both the UWorld and the Audio
+	// Device registered in this callback.
+	DECLARE_MULTICAST_DELEGATE_TwoParams(FOnWorldRegisteredToAudioDevice, const UWorld* /*InWorld */, Audio::FDeviceId /* AudioDeviceId*/);
+	static FOnWorldRegisteredToAudioDevice OnWorldRegisteredToAudioDevice;
+
+	// Called whenever a world is unregistered from an audio device. UWorlds are not guaranteed to be registered to the same
+	// audio device throughout their lifecycle, and there is no guarantee on the lifespan of both the UWorld and the Audio
+	// Device registered in this callback.
+	DECLARE_MULTICAST_DELEGATE_TwoParams(FOnWorldUnregisteredWithAudioDevice, const UWorld* /*InWorld */, Audio::FDeviceId /* AudioDeviceId*/);
+	static FOnWorldUnregisteredWithAudioDevice OnWorldUnregisteredWithAudioDevice;
 };
 
 /**
@@ -412,12 +424,23 @@ private:
 
 	FAudioDeviceHandle CreateNewDevice(const FAudioDeviceParams& InParams);
 
+	/**
+	* Shutsdown the audio device associated with the handle. The handle
+	* will become invalid after the audio device is shut down.
+	*/
+	bool ShutdownAudioDevice(Audio::FDeviceId DeviceID);
+
 	// Called exclusively by the FAudioDeviceHandle copy constructor and assignment operators:
 	void IncrementDevice(Audio::FDeviceId DeviceID);
 
 	// Called exclusively by the FAudioDeviceHandle dtor.
 	void DecrementDevice(Audio::FDeviceId DeviceID, UWorld* InWorld);
 
+	/**
+	* Shuts down all active audio devices
+	*/
+	bool ShutdownAllAudioDevices();
+
 	/** Application enters background handler */
 	void AppWillEnterBackground();
 
