#UE4 CIS IWYU fix

diff --git a/Engine/Source/Runtime/Engine/Public/PerPlatformProperties.h b/Engine/Source/Runtime/Engine/Public/PerPlatformProperties.h
--- a/Engine/Source/Runtime/Engine/Public/PerPlatformProperties.h
+++ b/Engine/Source/Runtime/Engine/Public/PerPlatformProperties.h
@@ -11,6 +11,11 @@ PerPlatformProperties.h: Property types that can be overridden on a per-platform
 #include "Containers/Map.h"
 #include "Algo/Find.h"
 #include "Serialization/MemoryLayout.h"
+
+#if WITH_EDITORONLY_DATA && WITH_EDITOR
+#include "Engine/Engine.h"
+#endif
+
 #include "PerPlatformProperties.generated.h"
 
 /** TPerPlatformProperty - template parent class for per-platform properties 
