Fix Build.

#rnx

diff --git a/Engine/Source/Runtime/Engine/Classes/Engine/SkeletalMesh.h b/Engine/Source/Runtime/Engine/Classes/Engine/SkeletalMesh.h
--- a/Engine/Source/Runtime/Engine/Classes/Engine/SkeletalMesh.h
+++ b/Engine/Source/Runtime/Engine/Classes/Engine/SkeletalMesh.h
@@ -20,6 +20,7 @@
 #include "BoneContainer.h"
 #include "Interfaces/Interface_CollisionDataProvider.h"
 #include "EngineTypes.h"
+#include "Engine/Engine.h"
 #include "SkeletalMeshSampling.h"
 #include "PerPlatformProperties.h"
 #include "SkeletalMeshLODSettings.h"
