Fixed an issue where the ChildActorComponent was not properly previewing in the editor.

diff --git a/Engine/Source/Runtime/Engine/Private/Components/ChildActorComponent.cpp b/Engine/Source/Runtime/Engine/Private/Components/ChildActorComponent.cpp
--- a/Engine/Source/Runtime/Engine/Private/Components/ChildActorComponent.cpp
+++ b/Engine/Source/Runtime/Engine/Private/Components/ChildActorComponent.cpp
@@ -393,7 +393,14 @@ void UChildActorComponent::ApplyComponentInstanceData(FChildActorComponentInstan
 {
 	check(ChildActorInstanceData);
 
-	if (ChildActorClass == ChildActorInstanceData->ChildActorClass)
+	bool bCreatedChildActor = false;
+	if (!ChildActor ||
+		ChildActor->GetClass() != ChildActorClass)
+	{
+		bCreatedChildActor = true;
+		CreateChildActor();
+	}
+	else if (ChildActorClass == ChildActorInstanceData->ChildActorClass)
 	{
 		ChildActorName = ChildActorInstanceData->ChildActorName;
 	}
@@ -401,7 +408,7 @@ void UChildActorComponent::ApplyComponentInstanceData(FChildActorComponentInstan
 	if (ChildActor)
 	{
 		// Only rename if it is safe to
-		if(ChildActorName != NAME_None)
+		if(!bCreatedChildActor && ChildActorName != NAME_None)
 		{
 			const FString ChildActorNameString = ChildActorName.ToString();
 			if (ChildActor->Rename(*ChildActorNameString, nullptr, REN_Test))
