Added CVars for enabling/disabling all ray tracing geometry types

diff --git a/Engine/Plugins/Experimental/GeometryCache/Source/GeometryCache/Private/GeometryCacheSceneProxy.cpp b/Engine/Plugins/Experimental/GeometryCache/Source/GeometryCache/Private/GeometryCacheSceneProxy.cpp
--- a/Engine/Plugins/Experimental/GeometryCache/Source/GeometryCache/Private/GeometryCacheSceneProxy.cpp
+++ b/Engine/Plugins/Experimental/GeometryCache/Source/GeometryCache/Private/GeometryCacheSceneProxy.cpp
@@ -35,6 +35,11 @@ static TAutoConsoleVariable<int32> CVarInterpolateFrames(
 	TEXT("Interpolate between geometry cache frames (if topology allows this)."),
 	ECVF_Scalability | ECVF_RenderThreadSafe);
 
+static TAutoConsoleVariable<int32> CVarRayTracingGeometryCache(
+	TEXT("r.RayTracing.Geometry.GeometryCache"),
+	1,
+	TEXT("Include geometry cache primitives in ray tracing effects (default = 1 (geometry cache enabled in ray tracing))"));
+
 /**
 * All vertex information except the position.
 */
@@ -489,6 +494,11 @@ void FGeometryCacheSceneProxy::GetDynamicMeshElements(const TArray<const FSceneV
 #if RHI_RAYTRACING
 void FGeometryCacheSceneProxy::GetDynamicRayTracingInstances(FRayTracingMaterialGatheringContext& Context, TArray<FRayTracingInstance>& OutRayTracingInstances)
 {
+	if (!CVarRayTracingGeometryCache.GetValueOnRenderThread())
+	{
+		return;
+	}
+
 	for (FGeomCacheTrackProxy* TrackProxy : Tracks)
 	{
 		const FVisibilitySample& VisibilitySample = TrackProxy->GetVisibilitySample(Time, bLooping);
diff --git a/Engine/Plugins/FX/Niagara/Source/Niagara/Private/NiagaraRendererRibbons.cpp b/Engine/Plugins/FX/Niagara/Source/Niagara/Private/NiagaraRendererRibbons.cpp
--- a/Engine/Plugins/FX/Niagara/Source/Niagara/Private/NiagaraRendererRibbons.cpp
+++ b/Engine/Plugins/FX/Niagara/Source/Niagara/Private/NiagaraRendererRibbons.cpp
@@ -68,6 +68,11 @@ static FAutoConsoleVariableRef CVarEnableNiagaraRibbonRendering(
 	ECVF_Default
 );
 
+static TAutoConsoleVariable<int32> CVarRayTracingNiagaraRibbons(
+	TEXT("r.RayTracing.Geometry.NiagaraRibbons"),
+	1,
+	TEXT("Include Niagara ribbons in ray tracing effects (default = 1 (Niagara ribbons enabled in ray tracing))"));
+
 // max absolute error 9.0x10^-3
 // Eberly's polynomial degree 1 - respect bounds
 // input [-1, 1] and output [0, PI]
@@ -1103,6 +1108,11 @@ void FNiagaraRendererRibbons::CreatePerViewResources(
 #if RHI_RAYTRACING
 void FNiagaraRendererRibbons::GetDynamicRayTracingInstances(FRayTracingMaterialGatheringContext& Context, TArray<FRayTracingInstance>& OutRayTracingInstances, const FNiagaraSceneProxy* SceneProxy)
 {
+	if (!CVarRayTracingNiagaraRibbons.GetValueOnRenderThread())
+	{
+		return;
+	}
+
 	SCOPE_CYCLE_COUNTER(STAT_NiagaraRender);
 	SCOPE_CYCLE_COUNTER(STAT_NiagaraRenderRibbons);
 	check(SceneProxy);
@@ -1174,4 +1184,4 @@ void FNiagaraRendererRibbons::GetDynamicRayTracingInstances(FRayTracingMaterialG
 
 	OutRayTracingInstances.Add(RayTracingInstance);
 }
-#endif
\ No newline at end of file
+#endif
diff --git a/Engine/Plugins/FX/Niagara/Source/Niagara/Private/NiagaraRendererSprites.cpp b/Engine/Plugins/FX/Niagara/Source/Niagara/Private/NiagaraRendererSprites.cpp
--- a/Engine/Plugins/FX/Niagara/Source/Niagara/Private/NiagaraRendererSprites.cpp
+++ b/Engine/Plugins/FX/Niagara/Source/Niagara/Private/NiagaraRendererSprites.cpp
@@ -39,6 +39,11 @@ static FAutoConsoleVariableRef CVarbEnableMinimalGPUBuffers(
 	ECVF_Default
 );
 
+static TAutoConsoleVariable<int32> CVarRayTracingNiagaraSprites(
+	TEXT("r.RayTracing.Geometry.NiagaraSprites"),
+	1,
+	TEXT("Include Niagara sprites in ray tracing effects (default = 1 (Niagara sprites enabled in ray tracing))"));
+
 
 /** Dynamic data for sprite renderers. */
 struct FNiagaraDynamicDataSprites : public FNiagaraDynamicDataBase
@@ -604,6 +609,11 @@ void FNiagaraRendererSprites::GetDynamicMeshElements(const TArray<const FSceneVi
 #if RHI_RAYTRACING
 void FNiagaraRendererSprites::GetDynamicRayTracingInstances(FRayTracingMaterialGatheringContext& Context, TArray<FRayTracingInstance>& OutRayTracingInstances, const FNiagaraSceneProxy* SceneProxy)
 {
+	if (!CVarRayTracingNiagaraSprites.GetValueOnRenderThread())
+	{
+		return;
+	}
+
 	SCOPE_CYCLE_COUNTER(STAT_NiagaraRender);
 	SCOPE_CYCLE_COUNTER(STAT_NiagaraRenderSprites);
 	check(SceneProxy);
diff --git a/Engine/Plugins/Runtime/ProceduralMeshComponent/Source/ProceduralMeshComponent/Private/ProceduralMeshComponent.cpp b/Engine/Plugins/Runtime/ProceduralMeshComponent/Source/ProceduralMeshComponent/Private/ProceduralMeshComponent.cpp
--- a/Engine/Plugins/Runtime/ProceduralMeshComponent/Source/ProceduralMeshComponent/Private/ProceduralMeshComponent.cpp
+++ b/Engine/Plugins/Runtime/ProceduralMeshComponent/Source/ProceduralMeshComponent/Private/ProceduralMeshComponent.cpp
@@ -30,6 +30,11 @@ DECLARE_CYCLE_STAT(TEXT("Update Collision"), STAT_ProcMesh_UpdateCollision, STAT
 
 DEFINE_LOG_CATEGORY_STATIC(LogProceduralComponent, Log, All);
 
+static TAutoConsoleVariable<int32> CVarRayTracingProceduralMesh(
+	TEXT("r.RayTracing.Geometry.ProceduralMeshes"),
+	1,
+	TEXT("Include procedural meshes in ray tracing effects (default = 1 (procedural meshes enabled in ray tracing))"));
+
 /** Resource array to pass  */
 class FProcMeshVertexResourceArray : public FResourceArrayInterface
 {
@@ -427,6 +432,11 @@ public:
 
 	virtual void GetDynamicRayTracingInstances(FRayTracingMaterialGatheringContext& Context, TArray<FRayTracingInstance>& OutRayTracingInstances) override final
 	{
+		if (!CVarRayTracingProceduralMesh.GetValueOnRenderThread())
+		{
+			return;
+		}
+
 		for (const FProcMeshProxySection* Section : Sections)
 		{
 			if (Section != nullptr && Section->bSectionVisible)
@@ -1071,4 +1081,4 @@ UMaterialInterface* UProceduralMeshComponent::GetMaterialFromCollisionFaceIndex(
 	}
 
 	return Result;
-}
\ No newline at end of file
+}
diff --git a/Engine/Source/Runtime/Engine/Private/InstancedStaticMesh.cpp b/Engine/Source/Runtime/Engine/Private/InstancedStaticMesh.cpp
--- a/Engine/Source/Runtime/Engine/Private/InstancedStaticMesh.cpp
+++ b/Engine/Source/Runtime/Engine/Private/InstancedStaticMesh.cpp
@@ -62,32 +62,32 @@ TAutoConsoleVariable<int32> CVarMinLOD(
 	ECVF_Scalability | ECVF_Default);
 
 static TAutoConsoleVariable<int32> CVarRayTracingRenderInstances(
-	TEXT("r.RayTracing.InstancedStaticMeshes"),
+	TEXT("r.RayTracing.Geometry.InstancedStaticMeshes"),
 	1,
 	TEXT("Include static mesh instances in ray tracing effects (default = 1 (Instances enabled in ray tracing))"));
 
 static TAutoConsoleVariable<int32> CVarRayTracingRenderInstancesCulling(
-	TEXT("r.RayTracing.InstancedStaticMeshes.Culling"),
+	TEXT("r.RayTracing.Geometry.InstancedStaticMeshes.Culling"),
 	1,
 	TEXT("Enable culling for instances in ray tracing (default = 1 (Culling enabled))"));
 
 static TAutoConsoleVariable<float> CVarRayTracingInstancesCullClusterMaxRadiusMultiplier(
-	TEXT("r.RayTracing.InstancedStaticMeshes.CullClusterMaxRadiusMultiplier"),
+	TEXT("r.RayTracing.Geometry.InstancedStaticMeshes.CullClusterMaxRadiusMultiplier"),
 	20.0f, 
 	TEXT("Multiplier for the maximum instance size (default = 20)"));
 
 static TAutoConsoleVariable<float> CVarRayTracingInstancesCullClusterRadius(
-	TEXT("r.RayTracing.InstancedStaticMeshes.CullClusterRadius"),
+	TEXT("r.RayTracing.Geometry.InstancedStaticMeshes.CullClusterRadius"),
 	10000.0f, // 100 m
 	TEXT("Ignore instances outside of this radius in ray tracing effects (default = 10000 (100m))"));
 
 static TAutoConsoleVariable<float> CVarRayTracingInstancesLowScaleThreshold(
-	TEXT("r.RayTracing.InstancedStaticMeshes.LowScaleRadiusThreshold"),
+	TEXT("r.RayTracing.Geometry.InstancedStaticMeshes.LowScaleRadiusThreshold"),
 	50.0f, // Instances with a radius smaller than this threshold get culled after CVarRayTracingInstancesLowScaleCullRadius
 	TEXT("Threshold that classifies instances as small (default = 50cm))"));
 
 static TAutoConsoleVariable<float> CVarRayTracingInstancesLowScaleCullRadius(
-	TEXT("r.RayTracing.InstancedStaticMeshes.LowScaleCullRadius"),
+	TEXT("r.RayTracing.Geometry.InstancedStaticMeshes.LowScaleCullRadius"),
 	1000.0f, 
 	TEXT("Cull radius for small instances (default = 1000 (10m))"));
 
diff --git a/Engine/Source/Runtime/Engine/Private/SkeletalMesh.cpp b/Engine/Source/Runtime/Engine/Private/SkeletalMesh.cpp
--- a/Engine/Source/Runtime/Engine/Private/SkeletalMesh.cpp
+++ b/Engine/Source/Runtime/Engine/Private/SkeletalMesh.cpp
@@ -112,6 +112,11 @@ TAutoConsoleVariable<int32> CVarDebugDrawBoneAxes(TEXT("a.DebugDrawBoneAxes"), 0
 const FGuid FSkeletalMeshCustomVersion::GUID(0xD78A4A00, 0xE8584697, 0xBAA819B5, 0x487D46B4);
 FCustomVersionRegistration GRegisterSkeletalMeshCustomVersion(FSkeletalMeshCustomVersion::GUID, FSkeletalMeshCustomVersion::LatestVersion, TEXT("SkeletalMeshVer"));
 
+static TAutoConsoleVariable<int32> CVarRayTracingSkeletalMeshes(
+	TEXT("r.RayTracing.Geometry.SkeletalMeshes"),
+	1,
+	TEXT("Include skeletal meshes in ray tracing effects (default = 1 (skeletal meshes enabled in ray tracing))"));
+
 #if WITH_APEX_CLOTHING
 /*-----------------------------------------------------------------------------
 	utility functions for apex clothing 
@@ -5213,6 +5218,11 @@ void FSkeletalMeshSceneProxy::GetDynamicElementsSection(const TArray<const FScen
 #if RHI_RAYTRACING
 void FSkeletalMeshSceneProxy::GetDynamicRayTracingInstances(FRayTracingMaterialGatheringContext & Context, TArray<struct FRayTracingInstance>& OutRayTracingInstances)
 {
+	if (!CVarRayTracingSkeletalMeshes.GetValueOnRenderThread())
+	{
+		return;
+	}
+
 	// GetRayTracingGeometry()->IsInitialized() is checked as a workaround for UE-92634. FSkeletalMeshSceneProxy's resources may have already been released, but proxy has not removed yet)
 	if (MeshObject->GetRayTracingGeometry() && MeshObject->GetRayTracingGeometry()->IsInitialized() )
 	{
diff --git a/Engine/Source/Runtime/Engine/Private/StaticMeshRender.cpp b/Engine/Source/Runtime/Engine/Private/StaticMeshRender.cpp
--- a/Engine/Source/Runtime/Engine/Private/StaticMeshRender.cpp
+++ b/Engine/Source/Runtime/Engine/Private/StaticMeshRender.cpp
@@ -108,6 +108,16 @@ static FAutoConsoleCommand GToggleForceDefaultMaterialCmd(
 	FConsoleCommandDelegate::CreateStatic(ToggleForceDefaultMaterial)
 	);
 
+static TAutoConsoleVariable<int32> CVarRayTracingStaticMeshes(
+	TEXT("r.RayTracing.Geometry.StaticMeshes"),
+	1,
+	TEXT("Include static meshes in ray tracing effects (default = 1 (static meshes enabled in ray tracing))"));
+
+static TAutoConsoleVariable<int32> CVarRayTracingWPOStaticMeshes(
+	TEXT("r.RayTracing.Geometry.WPOStaticMeshes"),
+	1,
+	TEXT("Include static meshes with world position offset in ray tracing effects (default = 1 (static meshes with world position offset enabled in ray tracing))"));
+
 /** Initialization constructor. */
 FStaticMeshSceneProxy::FStaticMeshSceneProxy(UStaticMeshComponent* InComponent, bool bForceLODsShareStaticLighting)
 	: FPrimitiveSceneProxy(InComponent, InComponent->GetStaticMesh()->GetFName())
@@ -1637,7 +1647,7 @@ void FStaticMeshSceneProxy::GetDynamicMeshElements(const TArray<const FSceneView
 #if RHI_RAYTRACING
 void FStaticMeshSceneProxy::GetDynamicRayTracingInstances(FRayTracingMaterialGatheringContext& Context, TArray<FRayTracingInstance>& OutRayTracingInstances )
 {
-	if (DynamicRayTracingGeometries.Num() <= 0)
+	if (DynamicRayTracingGeometries.Num() <= 0 || !CVarRayTracingStaticMeshes.GetValueOnRenderThread() || !CVarRayTracingWPOStaticMeshes.GetValueOnRenderThread())
 	{
 		return;
 	}
diff --git a/Engine/Source/Runtime/Landscape/Private/LandscapeRender.cpp b/Engine/Source/Runtime/Landscape/Private/LandscapeRender.cpp
--- a/Engine/Source/Runtime/Landscape/Private/LandscapeRender.cpp
+++ b/Engine/Source/Runtime/Landscape/Private/LandscapeRender.cpp
@@ -130,20 +130,20 @@ FAutoConsoleVariableRef CVarLandscapeDebugViewMode(
 
 #if RHI_RAYTRACING
 static TAutoConsoleVariable<int32> CVarRayTracingLandscape(
-	TEXT("r.RayTracing.Landscape"),
+	TEXT("r.RayTracing.Geometry.Landscape"),
 	1,
 	TEXT("Include landscapes in ray tracing effects (default = 1 (landscape enabled in ray tracing))"));
 
 int32 GLandscapeRayTracingGeometryLODsThatUpdateEveryFrame = 0;
 static FAutoConsoleVariableRef CVarLandscapeRayTracingGeometryLODsThatUpdateEveryFrame(
-	TEXT("r.RayTracing.Landscape.LODsUpdateEveryFrame"),
+	TEXT("r.RayTracing.Geometry.Landscape.LODsUpdateEveryFrame"),
 	GLandscapeRayTracingGeometryLODsThatUpdateEveryFrame,
 	TEXT("If on, LODs that are lower than the specified level will be updated every frame, which can be used to workaround some artifacts caused by texture streaming if you're using WorldPositionOffset on the landscape")
 );
 
 int32 GLandscapeRayTracingGeometryDetectTextureStreaming = 1;
 static FAutoConsoleVariableRef CVarLandscapeRayTracingGeometryDetectTextureStreaming(
-	TEXT("r.RayTracing.Landscape.DetectTextureStreaming"),
+	TEXT("r.RayTracing.Geometry.Landscape.DetectTextureStreaming"),
 	GLandscapeRayTracingGeometryDetectTextureStreaming,
 	TEXT("If on, update ray tracing geometry when texture streaming state changes. Useful when WorldPositionOffset is used in the landscape material")
 );
diff --git a/Engine/Source/Runtime/Renderer/Private/DeferredShadingRenderer.cpp b/Engine/Source/Runtime/Renderer/Private/DeferredShadingRenderer.cpp
--- a/Engine/Source/Runtime/Renderer/Private/DeferredShadingRenderer.cpp
+++ b/Engine/Source/Runtime/Renderer/Private/DeferredShadingRenderer.cpp
@@ -810,7 +810,9 @@ bool FDeferredShadingSceneRenderer::GatherRayTracingWorldInstances(FRHICommandLi
 
 				//#dxr_todo UE-68621  The Raytracing code path does not support ShowFlags since data moved to the SceneInfo. 
 				//Touching the SceneProxy to determine this would simply cost too much
-				if (SceneInfo->bIsRayTracingStaticRelevant && View.Family->EngineShowFlags.StaticMeshes)
+				static const auto RayTracingStaticMeshesCVar = IConsoleManager::Get().FindTConsoleVariableDataInt(TEXT("r.RayTracing.StaticMeshes"));
+
+				if (SceneInfo->bIsRayTracingStaticRelevant && View.Family->EngineShowFlags.StaticMeshes && RayTracingStaticMeshesCVar && RayTracingStaticMeshesCVar->GetValueOnRenderThread() > 0)
 				{
 					Item.ViewIndex = ViewIndex;
 					RelevantPrimitives.Add(Item);
