Fix sounds occasionally not playing in PlaySound notifies

#rnx

diff --git a/Engine/Source/Runtime/Engine/Private/Animation/AnimNotify_PlaySound.cpp b/Engine/Source/Runtime/Engine/Private/Animation/AnimNotify_PlaySound.cpp
--- a/Engine/Source/Runtime/Engine/Private/Animation/AnimNotify_PlaySound.cpp
+++ b/Engine/Source/Runtime/Engine/Private/Animation/AnimNotify_PlaySound.cpp
@@ -28,7 +28,7 @@ UAnimNotify_PlaySound::UAnimNotify_PlaySound()
 void UAnimNotify_PlaySound::Notify(class USkeletalMeshComponent* MeshComp, class UAnimSequenceBase* Animation)
 {
 	// Don't call super to avoid call back in to blueprints
-	if (Sound)
+	if (Sound && MeshComp)
 	{
 		if (Sound->IsLooping())
 		{
@@ -36,13 +36,20 @@ void UAnimNotify_PlaySound::Notify(class USkeletalMeshComponent* MeshComp, class
 			return;
 		}
 
-		if (bFollow)
+		if (APawn* Pawn = Cast<APawn>(MeshComp->GetOwner()))
 		{
-			UGameplayStatics::SpawnSoundAttached(Sound, MeshComp, AttachName, FVector(ForceInit), EAttachLocation::SnapToTarget, false, VolumeMultiplier, PitchMultiplier);
+			if (bFollow)
+			{
+				UGameplayStatics::SpawnSoundAttached(Sound, MeshComp, AttachName, FVector(ForceInit), EAttachLocation::SnapToTarget, false, VolumeMultiplier, PitchMultiplier);
+			}
+			else
+			{
+				UGameplayStatics::PlaySoundAtLocation(MeshComp->GetWorld(), Sound, MeshComp->GetComponentLocation(), VolumeMultiplier, PitchMultiplier);
+			}
 		}
 		else
 		{
-			UGameplayStatics::PlaySoundAtLocation(MeshComp->GetWorld(), Sound, MeshComp->GetComponentLocation(), VolumeMultiplier, PitchMultiplier);
+			UGameplayStatics::PlaySound2D(MeshComp->GetWorld(), Sound, VolumeMultiplier, PitchMultiplier);
 		}
 	}
 }
