Fixed shader compilation error when sampling texture2d array

diff --git a/Engine/Source/Runtime/Engine/Private/Materials/MaterialUniformExpressions.cpp b/Engine/Source/Runtime/Engine/Private/Materials/MaterialUniformExpressions.cpp
--- a/Engine/Source/Runtime/Engine/Private/Materials/MaterialUniformExpressions.cpp
+++ b/Engine/Source/Runtime/Engine/Private/Materials/MaterialUniformExpressions.cpp
@@ -440,7 +440,7 @@ FShaderParametersMetadata* FUniformExpressionSet::CreateBufferStruct()
 		check((NextMemberOffset % SHADER_PARAMETER_POINTER_ALIGNMENT) == 0);
 		new(Members) FShaderParametersMetadata::FMember(*Texture2DArrayNames[i], TEXT("Texture2DArray"), NextMemberOffset, UBMT_TEXTURE, EShaderPrecisionModifier::Float, 1, 1, 0, NULL);
 		NextMemberOffset += SHADER_PARAMETER_POINTER_ALIGNMENT;
-		new(Members) FShaderParametersMetadata::FMember(*Texture2DSamplerNames[i], TEXT("SamplerState"), NextMemberOffset, UBMT_SAMPLER, EShaderPrecisionModifier::Float, 1, 1, 0, NULL);
+		new(Members) FShaderParametersMetadata::FMember(*Texture2DArraySamplerNames[i], TEXT("SamplerState"), NextMemberOffset, UBMT_SAMPLER, EShaderPrecisionModifier::Float, 1, 1, 0, NULL);
 		NextMemberOffset += SHADER_PARAMETER_POINTER_ALIGNMENT;
 	}
 
