- Fix crash on shutdown when attempting to debug audio device handles w/INSTRUMENT_AUDIODEVICE_HANDLES
- Clean-up device manager shutdown logic & crit section for when reseting device map in destructor (probably not necessary but safer)
- Notify UWorld & GameViewportClient when shutting down and forcefully destroying associated audio device prior to destruction

Minor:
- Remove GEngine references where not necessary from AudioDeviceManager
- Add logging around device manager lifetime
- Fix unregistration not being called if world is swapping Audio Device (don't think this happens anywhere, but to be safe)

diff --git a/Engine/Source/Runtime/Engine/Classes/Engine/GameViewportClient.h b/Engine/Source/Runtime/Engine/Classes/Engine/GameViewportClient.h
--- a/Engine/Source/Runtime/Engine/Classes/Engine/GameViewportClient.h
+++ b/Engine/Source/Runtime/Engine/Classes/Engine/GameViewportClient.h
@@ -3,6 +3,7 @@
 #pragma once
 
 #include "CoreMinimal.h"
+#include "Delegates/IDelegateInstance.h"
 #include "UObject/ObjectMacros.h"
 #include "InputCoreTypes.h"
 #include "Engine/EngineBaseTypes.h"
@@ -109,6 +110,9 @@ protected:
 	/** Strong handle to the audio device used by this viewport. */
 	FAudioDeviceHandle AudioDevice;
 
+	/** Handle to delegate in case audio device is destroyed. */
+	FDelegateHandle AudioDeviceDestroyedHandle;
+
 public:
 
 	/** see enum EViewModeIndex */
diff --git a/Engine/Source/Runtime/Engine/Classes/Engine/World.h b/Engine/Source/Runtime/Engine/Classes/Engine/World.h
--- a/Engine/Source/Runtime/Engine/Classes/Engine/World.h
+++ b/Engine/Source/Runtime/Engine/Classes/Engine/World.h
@@ -9,6 +9,7 @@
 #include "UObject/Object.h"
 #include "Misc/Guid.h"
 #include "UObject/Class.h"
+#include "Delegates/IDelegateInstance.h"
 #include "Engine/EngineTypes.h"
 #include "Engine/EngineBaseTypes.h"
 #include "CollisionQueryParams.h"
@@ -211,6 +212,23 @@ private:
 	UWorld* World;
 };
 
+// List of delegates for the world being registered to an audio device.
+class ENGINE_API FAudioDeviceWorldDelegates
+{
+public:
+	// Called whenever a world is registered to an audio device. UWorlds are not guaranteed to be registered to the same
+	// audio device throughout their lifecycle, and there is no guarantee on the lifespan of both the UWorld and the Audio
+	// Device registered in this callback.
+	DECLARE_MULTICAST_DELEGATE_TwoParams(FOnWorldRegisteredToAudioDevice, const UWorld* /*InWorld */, Audio::FDeviceId /* AudioDeviceId*/);
+	static FOnWorldRegisteredToAudioDevice OnWorldRegisteredToAudioDevice;
+
+	// Called whenever a world is unregistered from an audio device. UWorlds are not guaranteed to be registered to the same
+	// audio device throughout their lifecycle, and there is no guarantee on the lifespan of both the UWorld and the Audio
+	// Device registered in this callback.
+	DECLARE_MULTICAST_DELEGATE_TwoParams(FOnWorldUnregisteredWithAudioDevice, const UWorld* /*InWorld */, Audio::FDeviceId /* AudioDeviceId*/);
+	static FOnWorldUnregisteredWithAudioDevice OnWorldUnregisteredWithAudioDevice;
+};
+
 /** class that encapsulates seamless world traveling */
 class FSeamlessTravelHandler
 {
@@ -1200,6 +1218,9 @@ public:
 
 private:
 
+	/** Handle to delegate in case audio device is destroyed. */
+	FDelegateHandle AudioDeviceDestroyedHandle;
+
 #if WITH_EDITORONLY_DATA
 	/** Pointer to the current level being edited. Level has to be in the Levels array and == PersistentLevel in the game.		*/
 	UPROPERTY(Transient)
@@ -3314,7 +3335,7 @@ public:
 	 */
 	class AAudioVolume* GetAudioSettings( const FVector& ViewLocation, struct FReverbSettings* OutReverbSettings, struct FInteriorSettings* OutInteriorSettings );
 
-	void SetAudioDevice(FAudioDeviceHandle& InHandle);
+	void SetAudioDevice(const FAudioDeviceHandle& InHandle);
 
 	/**
 	 * Get the audio device used by this world.
diff --git a/Engine/Source/Runtime/Engine/Private/AudioDebug.cpp b/Engine/Source/Runtime/Engine/Private/AudioDebug.cpp
--- a/Engine/Source/Runtime/Engine/Private/AudioDebug.cpp
+++ b/Engine/Source/Runtime/Engine/Private/AudioDebug.cpp
@@ -11,6 +11,7 @@
 #include "DrawDebugHelpers.h"
 #include "DSP/Dsp.h"
 #include "Engine/Font.h"
+#include "Engine/World.h"
 #include "GameFramework/GameUserSettings.h"
 #include "HAL/IConsoleManager.h"
 #include "Misc/CommandLine.h"
@@ -509,7 +510,7 @@ namespace Audio
 	FAudioDebugger::FAudioDebugger()
 		: bVisualize3dDebug(0)
 	{
-		WorldRegisteredWithDeviceHandle = FAudioDeviceManagerDelegates::OnWorldRegisteredToAudioDevice.AddLambda([this](const UWorld* InWorld, FDeviceId InDeviceId)
+		WorldRegisteredWithDeviceHandle = FAudioDeviceWorldDelegates::OnWorldRegisteredToAudioDevice.AddLambda([this](const UWorld* InWorld, FDeviceId InDeviceId)
 		{
 			if (bAllowUsingDeprecatedDebugStats)
 			{
diff --git a/Engine/Source/Runtime/Engine/Private/AudioDeviceManager.cpp b/Engine/Source/Runtime/Engine/Private/AudioDeviceManager.cpp
--- a/Engine/Source/Runtime/Engine/Private/AudioDeviceManager.cpp
+++ b/Engine/Source/Runtime/Engine/Private/AudioDeviceManager.cpp
@@ -6,6 +6,7 @@
 #include "AudioDefines.h"
 #include "AudioDevice.h"
 #include "AudioMixerDevice.h"
+#include "Engine/World.h"
 #include "Sound/AudioSettings.h"
 #include "Sound/SoundWave.h"
 #include "GameFramework/GameUserSettings.h"
@@ -76,6 +77,57 @@ static FAutoConsoleCommand GReportAudioDevicesCommand(
 	})
 );
 
+namespace AudioDeviceManagerUtils
+{
+	FString PrintDeviceInfo(const Audio::FDeviceId InDeviceId, const EAudioDeviceScope InScope, const bool bInIsNonRealtime, const int32* NumHandles = nullptr, const TMap<uint32, FString>* InStackWalk = nullptr)
+	{
+		FString ScopeStr;
+		switch (InScope)
+		{
+			case EAudioDeviceScope::Shared:
+				ScopeStr = TEXT("Shared");
+				break;
+
+			case EAudioDeviceScope::Unique:
+				ScopeStr = TEXT("Unique");
+				break;
+
+			case EAudioDeviceScope::Default:
+			default:
+				ScopeStr = TEXT("Default");
+				break;
+		}
+
+		FString DeviceInfo = FString::Printf(
+			TEXT("                Id: %d, Scope: %s, Realtime: %s"),
+			InDeviceId,
+			*ScopeStr,
+			bInIsNonRealtime ? TEXT("False") : TEXT("True"));
+
+		if (!NumHandles)
+		{
+			return DeviceInfo;
+		}
+
+		DeviceInfo += FString::Printf(TEXT(", Num Handles: %d"), *NumHandles);
+
+#if INSTRUMENT_AUDIODEVICE_HANDLES
+		if (InStackWalk)
+		{
+			DeviceInfo += TEXT("\n            Active Handles:\n\n");
+			for (const TPair<uint32, FString>& StackWalkPairs : *InStackWalk)
+			{
+				DeviceInfo += StackWalkPairs.Value;
+				DeviceInfo += TEXT("\n\n");
+			}
+		}
+#endif
+
+		return DeviceInfo;
+	}
+}
+
+
 FAudioDeviceManager* FAudioDeviceManager::Singleton = nullptr;
 
 // Some stress tests:
@@ -148,20 +200,34 @@ FAudioDeviceManager::FAudioDeviceManager()
 
 FAudioDeviceManager::~FAudioDeviceManager()
 {
-	FAudioCommandFence Fence;
-	Fence.BeginFence();
-	Fence.Wait();
-	FAudioThread::StopAudioThread();
+	UE_LOG(LogAudio, Display, TEXT("Beginning Audio Device Manager Shutdown (Module: %s)..."), *AudioMixerModuleName);
 
-	MainAudioDeviceHandle.Reset();
+	TArray<Audio::FDeviceId> DeviceIds;
+	{
+		FScopeLock ScopeLock(&DeviceMapCriticalSection);
+		Devices.GetKeys(DeviceIds);
+	}
+
+	if (DeviceIds.Num() > 0)
+	{
+		UE_LOG(LogAudio, Display, TEXT("Destroying %d Remaining Audio Device(s)..."), DeviceIds.Num());
+
+		// Notify anyone listening to the device manager that we are about to destroy the audio device.
+		for (Audio::FDeviceId DeviceId : DeviceIds)
+		{
+			FAudioDeviceManagerDelegates::OnAudioDeviceDestroyed.Broadcast(DeviceId);
+		}
+	}
+
+	FAudioThread::StopAudioThread();
 
-	// Notify anyone listening to the device manager that we are about to destroy the audio device.
-	for (auto& Device : Devices)
 	{
-		FAudioDeviceManagerDelegates::OnAudioDeviceDestroyed.Broadcast(Device.Key);
+		FScopeLock ScopeLock(&DeviceMapCriticalSection);
+		MainAudioDeviceHandle.Reset();
+		Devices.Reset();
 	}
 
-	Devices.Reset();
+	FCoreDelegates::ApplicationWillEnterBackgroundDelegate.RemoveAll(this);
 
 	// Release any loaded buffers - this calls stop on any sources that need it
 	for (int32 Index = Buffers.Num() - 1; Index >= 0; Index--)
@@ -347,8 +413,6 @@ IAudioDeviceModule* FAudioDeviceManager::GetAudioDeviceModule()
 	return AudioDeviceModule;
 }
 
-
-
 FAudioDeviceParams FAudioDeviceManager::GetDefaultParamsForNewWorld()
 {
 	bool bCreateNewAudioDeviceForPlayInEditor = false;
@@ -402,8 +466,9 @@ void FAudioDeviceManager::RegisterWorld(UWorld* InWorld, Audio::FDeviceId Device
 	{
 		if (!DeviceContainer->WorldsUsingThisDevice.Contains(InWorld))
 		{
+			UE_LOG(LogAudio, Display, TEXT("Audio Device (ID: %d) registered with world '%s'."), DeviceId, *InWorld->GetName());
 			DeviceContainer->WorldsUsingThisDevice.AddUnique(InWorld);
-			FAudioDeviceManagerDelegates::OnWorldRegisteredToAudioDevice.Broadcast(InWorld, DeviceId);
+			FAudioDeviceWorldDelegates::OnWorldRegisteredToAudioDevice.Broadcast(InWorld, DeviceId);
 		}
 	}
 }
@@ -419,8 +484,9 @@ void FAudioDeviceManager::UnregisterWorld(UWorld* InWorld, Audio::FDeviceId Devi
 	{
 		if (DeviceContainer->WorldsUsingThisDevice.Contains(InWorld))
 		{
+			UE_LOG(LogAudio, Display, TEXT("Audio Device unregistered from world '%s'."), *InWorld->GetName());
 			DeviceContainer->WorldsUsingThisDevice.Remove(InWorld);
-			FAudioDeviceManagerDelegates::OnWorldUnregisteredWithAudioDevice.Broadcast(InWorld, DeviceId);
+			FAudioDeviceWorldDelegates::OnWorldUnregisteredWithAudioDevice.Broadcast(InWorld, DeviceId);
 		}
 	}
 }
@@ -455,7 +521,7 @@ bool FAudioDeviceManager::InitializeManager()
 		
 		if (!MainAudioDeviceHandle)
 		{
-			UE_LOG(LogAudio, Display, TEXT("Audio device could not be initialized. Please check the value for AudioDeviceModuleName and AudioMixerModuleName in [Platform]Engine.ini."));
+			UE_LOG(LogAudio, Display, TEXT("Main audio device could not be initialized. Please check the value for AudioDeviceModuleName and AudioMixerModuleName in [Platform]Engine.ini."));
 			return false;
 		}
 
@@ -545,13 +611,17 @@ bool FAudioDeviceManager::LoadDefaultAudioDeviceModule()
 
 FAudioDeviceHandle FAudioDeviceManager::CreateNewDevice(const FAudioDeviceParams& InParams)
 {
-	Audio::FDeviceId DeviceID = GetNewDeviceID();
+	const Audio::FDeviceId DeviceID = GetNewDeviceID();
+
+	FString DeviceInfo = AudioDeviceManagerUtils::PrintDeviceInfo(DeviceID, InParams.Scope, InParams.bIsNonRealtime);
+	UE_LOG(LogAudio, Display, TEXT("Creating Audio Device: %s"), *DeviceInfo);
 	Devices.Emplace(DeviceID, FAudioDeviceContainer(InParams, DeviceID, this));
+
 	FAudioDeviceContainer* ContainerPtr = Devices.Find(DeviceID);
 	check(ContainerPtr);
 	if (!ContainerPtr->Device)
 	{
-		UE_LOG(LogAudio, Display, TEXT("Audio device could not be initialized. Please check the value for AudioDeviceModuleName and AudioMixerModuleName in [Platform]Engine.ini."));
+		UE_LOG(LogAudio, Display, TEXT("Destroying Audio Device %d: could not be initialized. Check AudioDeviceModuleName and AudioMixerModuleName in [Platform]Engine.ini."), DeviceID);
 
 		// Initializing the audio device failed. Remove the device container and return an empty handle.
 		Devices.Remove(DeviceID);
@@ -569,14 +639,10 @@ FAudioDeviceHandle FAudioDeviceManager::CreateNewDevice(const FAudioDeviceParams
 
 bool FAudioDeviceManager::IsValidAudioDevice(Audio::FDeviceId Handle) const
 {
-	return Devices.Contains(Handle);
-}
+	FCriticalSection* ConstCastCritSection = const_cast<FCriticalSection*>(&DeviceMapCriticalSection);
+	FScopeLock ScopeLock(ConstCastCritSection);
 
-bool FAudioDeviceManager::ShutdownAudioDevice(Audio::FDeviceId Handle)
-{
-	// Make sure we have a non-null device ptr in the index slot, then delete it
-	Devices.Remove(Handle);
-	return true;
+	return Devices.Contains(Handle);
 }
 
 void FAudioDeviceManager::IncrementDevice(Audio::FDeviceId DeviceID)
@@ -611,20 +677,12 @@ void FAudioDeviceManager::DecrementDevice(Audio::FDeviceId DeviceID, UWorld* InW
 		}
 
 		FAudioDeviceManagerDelegates::OnAudioDeviceDestroyed.Broadcast(DeviceID);
-		Devices.Remove(DeviceID);
 
 		UnregisterWorld(InWorld, DeviceID);
-	}
-}
-
-bool FAudioDeviceManager::ShutdownAllAudioDevices()
-{
-	FCoreDelegates::ApplicationWillEnterBackgroundDelegate.RemoveAll(this);
-
-	MainAudioDeviceHandle.Reset();
 
-	Devices.Reset();
-	return true;
+		UE_LOG(LogAudio, Display, TEXT("Destroying Audio Device %d: All handles released."), DeviceID);
+		Devices.Remove(DeviceID);
+	}
 }
 
 FAudioDeviceHandle FAudioDeviceManager::BuildNewHandle(FAudioDeviceContainer&Container, Audio::FDeviceId DeviceID, const FAudioDeviceParams &InParams)
@@ -648,14 +706,14 @@ uint32 FAudioDeviceManager::CreateUniqueStackWalkID()
 }
 #endif
 
-FAudioDeviceHandle FAudioDeviceManager::GetAudioDevice(Audio::FDeviceId Handle)
+FAudioDeviceHandle FAudioDeviceManager::GetAudioDevice(Audio::FDeviceId InDeviceID)
 {
 	FScopeLock ScopeLock(&DeviceMapCriticalSection);
-	FAudioDeviceContainer* Container = Devices.Find(Handle);
+	FAudioDeviceContainer* Container = Devices.Find(InDeviceID);
 	if (Container)
 	{
 		FAudioDeviceParams Params = FAudioDeviceParams();
-		return BuildNewHandle(*Container, Handle, Params);
+		return BuildNewHandle(*Container, InDeviceID, Params);
 	}
 	else
 	{
@@ -663,20 +721,39 @@ FAudioDeviceHandle FAudioDeviceManager::GetAudioDevice(Audio::FDeviceId Handle)
 	}
 }
 
-FAudioDevice* FAudioDeviceManager::GetAudioDeviceRaw(Audio::FDeviceId Handle)
+FAudioDevice* FAudioDeviceManager::GetAudioDeviceRaw(Audio::FDeviceId InDeviceID)
 {
 	FScopeLock ScopeLock(&DeviceMapCriticalSection);
-	if (!IsValidAudioDevice(Handle))
+	if (!IsValidAudioDevice(InDeviceID))
 	{
 		return nullptr;
 	}
 
-	FAudioDevice* AudioDevice = Devices[Handle].Device;
+	FAudioDevice* AudioDevice = Devices[InDeviceID].Device;
 	check(AudioDevice != nullptr);
 
 	return AudioDevice;
 }
 
+void FAudioDeviceManager::SetAudioDevice(UWorld& InWorld, Audio::FDeviceId InDeviceID)
+{
+	FScopeLock ScopeLock(&DeviceMapCriticalSection);
+	{
+		if (FAudioDeviceContainer* Container = Devices.Find(InDeviceID))
+		{
+			FAudioDeviceParams Params = FAudioDeviceParams();
+			Params.AssociatedWorld = &InWorld;
+			FAudioDeviceHandle Handle = BuildNewHandle(*Container, InDeviceID, Params);
+			InWorld.SetAudioDevice(Handle);
+		}
+		else
+		{
+			static const FAudioDeviceHandle InvalidHandle;
+			InWorld.SetAudioDevice(InvalidHandle);
+		}
+	}
+}
+
 bool FAudioDeviceManager::Initialize()
 {
 	if (!Singleton)
@@ -696,11 +773,16 @@ bool FAudioDeviceManager::Initialize()
 
 		if (bUseSound)
 		{
+			UE_LOG(LogAudio, Display, TEXT("Initializing Audio Device Manager..."));
 			Singleton = new FAudioDeviceManager();
 
-			if (!Singleton->InitializeManager())
+			if (Singleton->InitializeManager())
 			{
-				UE_LOG(LogAudio, Warning, TEXT("Initializing the Audio Device Manager Failed!"));
+				UE_LOG(LogAudio, Display, TEXT("Audio Device Manager Initialized"));
+			}
+			else
+			{
+				UE_LOG(LogAudio, Warning, TEXT("Audio Device Manager Initialization Failed!"));
 				delete Singleton;
 				Singleton = nullptr;
 			}
@@ -721,6 +803,8 @@ void FAudioDeviceManager::Shutdown()
 	{
 		delete Singleton;
 		Singleton = nullptr;
+
+		UE_LOG(LogAudio, Display, TEXT("Audio Device Manager Shutdown"));
 	}
 }
 
@@ -730,7 +814,8 @@ FAudioDeviceHandle FAudioDeviceManager::GetActiveAudioDevice()
 	{
 		return GetAudioDevice(ActiveAudioDeviceID);
 	}
-	return GEngine->GetMainAudioDevice();
+
+	return MainAudioDeviceHandle;
 }
 
 void FAudioDeviceManager::UpdateActiveAudioDevices(bool bGameTicking)
@@ -1043,32 +1128,20 @@ void FAudioDeviceManager::LogListOfAudioDevices()
 {
 	FString ListOfDevices;
 
-	for (auto& DeviceContainer : Devices)
+	for (const TPair<Audio::FDeviceId, FAudioDeviceContainer>& DevicePair : Devices)
 	{
-		FString DeviceInfo = FString::Printf(TEXT(R"(
-					Device %d:
-					Scope: %s 
-					Realtime: %s
-					Number Of Owners: %d 
-		)"),
-			DeviceContainer.Key,
-			DeviceContainer.Value.Scope == EAudioDeviceScope::Unique ? TEXT("Unique") : TEXT("Shared"),
-			DeviceContainer.Value.bIsNonRealtime ? TEXT("No") : TEXT("Yes"),
-			DeviceContainer.Value.NumberOfHandlesToThisDevice);
-
+		ListOfDevices += AudioDeviceManagerUtils::PrintDeviceInfo(
+			DevicePair.Key,
+			DevicePair.Value.Scope,
+			DevicePair.Value.bIsNonRealtime,
+			&DevicePair.Value.NumberOfHandlesToThisDevice
 #if INSTRUMENT_AUDIODEVICE_HANDLES
-		for (auto& StackWalkString : DeviceContainer.Value.HandleCreationStackWalks)
-		{
-			DeviceInfo += TEXT("Handle Created here still alive:\n");
-			DeviceInfo += StackWalkString.Value;
-			DeviceInfo += TEXT("\n\n");
-		}
+			, &DevicePair.Value.HandleCreationStackWalks
 #endif
-
-		ListOfDevices += DeviceInfo;
+		);
 	}
 
-	UE_LOG(LogAudio, Display, TEXT("List of devices: \n%s"), *ListOfDevices);
+	UE_LOG(LogAudio, Display, TEXT("Active Audio Devices:\n%s"), *ListOfDevices);
 }
 
 uint32 FAudioDeviceManager::GetNewDeviceID()
@@ -1307,14 +1380,14 @@ void FAudioDeviceManager::AppWillEnterBackground()
 	// Flush all commands to the audio thread and the audio render thread:
 	if (GCVarFlushAudioRenderCommandsOnSuspend)
 	{
-		if (GEngine && GEngine->GetMainAudioDevice())
+		if (MainAudioDeviceHandle.IsValid())
 		{
-			FAudioDeviceHandle AudioDevice = GEngine->GetMainAudioDevice();
-
-			FAudioThread::RunCommandOnAudioThread([AudioDevice]()
+			FAudioThread::RunCommandOnAudioThread([this]()
 			{
-				FAudioDevice* AudioDevicePtr = const_cast<FAudioDevice*>(AudioDevice.GetAudioDevice());
-				AudioDevicePtr->FlushAudioRenderingCommands(true);
+				if (MainAudioDeviceHandle.IsValid())
+				{
+					MainAudioDeviceHandle->FlushAudioRenderingCommands(true);
+				}
 			}, TStatId());
 		}
 
@@ -1385,7 +1458,7 @@ FAudioDeviceHandle::~FAudioDeviceHandle()
 	if (IsValid())
 	{
 		FAudioDeviceManager* AudioDeviceManager = FAudioDeviceManager::Get();
-		if (AudioDeviceManager)
+		if (ensure(AudioDeviceManager))
 		{
 			AudioDeviceManager->DecrementDevice(DeviceId, World);
 
@@ -1402,6 +1475,11 @@ FAudioDevice* FAudioDeviceHandle::GetAudioDevice() const
 	return Device;
 }
 
+UWorld* FAudioDeviceHandle::GetWorld() const
+{
+	return World;
+}
+
 Audio::FDeviceId FAudioDeviceHandle::GetDeviceID() const
 {
 	return DeviceId;
@@ -1419,26 +1497,37 @@ void FAudioDeviceHandle::Reset()
 
 FAudioDeviceHandle& FAudioDeviceHandle::operator=(const FAudioDeviceHandle& Other)
 {
-	if (IsValid())
-	{
-		check(FAudioDeviceManager::Get());
-		FAudioDeviceManager::Get()->DecrementDevice(DeviceId, World);
-	}
+	FAudioDeviceManager* AudioDeviceManager = FAudioDeviceManager::Get();
+
+	const bool bWasValid = IsValid();
+	const Audio::FDeviceId OldDeviceId = DeviceId;
+	UWorld* OldWorld = World;
+
+#if INSTRUMENT_AUDIODEVICE_HANDLES
+	const uint32 OldStackWalkID = StackWalkID;
+#endif
 
 	Device = Other.Device;
 	DeviceId = Other.DeviceId;
+	World = Other.World;
 
-	if (IsValid())
+	if (AudioDeviceManager && IsValid())
 	{
-		FAudioDeviceManager* AudioDeviceManager = FAudioDeviceManager::Get();
-		if (AudioDeviceManager)
-		{
-			AudioDeviceManager->IncrementDevice(DeviceId);
+		AudioDeviceManager->IncrementDevice(DeviceId);
+
+#if INSTRUMENT_AUDIODEVICE_HANDLES
+		AddStackDumpToAudioDeviceContainer();
+#endif
+	}
 
+	if (AudioDeviceManager && bWasValid)
+	{
 #if INSTRUMENT_AUDIODEVICE_HANDLES
-			AddStackDumpToAudioDeviceContainer();
+		check(OldStackWalkID != INDEX_NONE);
+		AudioDeviceManager->RemoveStackWalkForContainer(OldDeviceId, OldStackWalkID);
 #endif
-		}
+
+		AudioDeviceManager->DecrementDevice(OldDeviceId, OldWorld);
 	}
 
 	return *this;
@@ -1446,28 +1535,50 @@ FAudioDeviceHandle& FAudioDeviceHandle::operator=(const FAudioDeviceHandle& Othe
 
 FAudioDeviceHandle& FAudioDeviceHandle::operator=(FAudioDeviceHandle&& Other)
 {
-	if (FAudioDeviceManager::Get() && IsValid())
-	{
 #if INSTRUMENT_AUDIODEVICE_HANDLES
-		check(StackWalkID != INDEX_NONE);
-		GEngine->GetAudioDeviceManager()->RemoveStackWalkForContainer(DeviceId, StackWalkID);
+	const uint32 OldStackWalkID = StackWalkID;
 #endif
 
-		FAudioDeviceManager::Get()->DecrementDevice(DeviceId, World);
-	}
+	const bool bWasValid = IsValid();
+	const Audio::FDeviceId OldDeviceId = DeviceId;
+	UWorld* OldWorld = World;
 
 	Device = Other.Device;
 	DeviceId = Other.DeviceId;
+	World = Other.World;
 
-	Other.Device = nullptr;
-	Other.DeviceId = INDEX_NONE;
-
+	FAudioDeviceManager* AudioDeviceManager = FAudioDeviceManager::Get();
+	if (AudioDeviceManager && IsValid())
+	{
 #if INSTRUMENT_AUDIODEVICE_HANDLES
-	if (IsValid())
+		AddStackDumpToAudioDeviceContainer();
+#endif
+	}
+
+	if (AudioDeviceManager && bWasValid)
 	{
+#if INSTRUMENT_AUDIODEVICE_HANDLES
+		check(OldStackWalkID != INDEX_NONE);
+		AudioDeviceManager->RemoveStackWalkForContainer(OldDeviceId, OldStackWalkID);
+#endif
 
-		AddStackDumpToAudioDeviceContainer();
+		AudioDeviceManager->DecrementDevice(OldDeviceId, OldWorld);
 	}
+
+	if (AudioDeviceManager && Other.IsValid())
+	{
+#if INSTRUMENT_AUDIODEVICE_HANDLES
+		check(Other.StackWalkID != INDEX_NONE);
+		AudioDeviceManager->RemoveStackWalkForContainer(Other.DeviceId, Other.StackWalkID);
+#endif
+	}
+
+	Other.Device = nullptr;
+	Other.DeviceId = INDEX_NONE;
+	Other.World = nullptr;
+
+#if INSTRUMENT_AUDIODEVICE_HANDLES
+	Other.StackWalkID = INDEX_NONE;
 #endif
 
 	return *this;
@@ -1582,5 +1693,3 @@ FAudioDeviceManager::FAudioDeviceContainer::~FAudioDeviceContainer()
 
 FAudioDeviceManagerDelegates::FOnAudioDeviceCreated FAudioDeviceManagerDelegates::OnAudioDeviceCreated;
 FAudioDeviceManagerDelegates::FOnAudioDeviceDestroyed FAudioDeviceManagerDelegates::OnAudioDeviceDestroyed;
-FAudioDeviceManagerDelegates::FOnWorldRegisteredToAudioDevice FAudioDeviceManagerDelegates::OnWorldRegisteredToAudioDevice;
-FAudioDeviceManagerDelegates::FOnWorldUnregisteredWithAudioDevice FAudioDeviceManagerDelegates::OnWorldUnregisteredWithAudioDevice;
diff --git a/Engine/Source/Runtime/Engine/Private/GameViewportClient.cpp b/Engine/Source/Runtime/Engine/Private/GameViewportClient.cpp
--- a/Engine/Source/Runtime/Engine/Private/GameViewportClient.cpp
+++ b/Engine/Source/Runtime/Engine/Private/GameViewportClient.cpp
@@ -287,6 +287,13 @@ UGameViewportClient::UGameViewportClient(const FObjectInitializer& ObjectInitial
 		}
 #endif
 
+		AudioDeviceDestroyedHandle = FAudioDeviceManagerDelegates::OnAudioDeviceDestroyed.AddLambda([this](const Audio::FDeviceId InDeviceId)
+		{
+			if (InDeviceId == AudioDevice.GetDeviceID())
+			{
+				AudioDevice.Reset();
+			}
+		});
 	}
 }
 
@@ -344,6 +351,7 @@ void UGameViewportClient::PostInitProperties()
 
 void UGameViewportClient::BeginDestroy()
 {
+	AudioDeviceDestroyedHandle.Reset();
 	AudioDevice.Reset();
 
 	RemoveAllViewportWidgets();
diff --git a/Engine/Source/Runtime/Engine/Private/UnrealEngine.cpp b/Engine/Source/Runtime/Engine/Private/UnrealEngine.cpp
--- a/Engine/Source/Runtime/Engine/Private/UnrealEngine.cpp
+++ b/Engine/Source/Runtime/Engine/Private/UnrealEngine.cpp
@@ -1109,11 +1109,12 @@ void FWorldContext::SetCurrentWorld(UWorld* World)
 {
 	if (World != nullptr && AudioDeviceID != INDEX_NONE)
 	{
-		// Set the world's audio device handle so that audio components playing in the 
-		// world will use the correct audio device instance.
-		FAudioDeviceHandle AudioDevice = FAudioDeviceManager::Get()->GetAudioDevice(AudioDeviceID);
-
-		World->SetAudioDevice(AudioDevice);
+		if (FAudioDeviceManager* DeviceManager = FAudioDeviceManager::Get())
+		{
+			// Set the world's audio device handle so that audio components playing in the 
+			// world will use the correct audio device instance.
+			FAudioDeviceManager::Get()->SetAudioDevice(*World, AudioDeviceID);
+		}
 	}
 
 	for (int32 idx = 0; idx < ExternalReferences.Num(); ++idx)
diff --git a/Engine/Source/Runtime/Engine/Private/World.cpp b/Engine/Source/Runtime/Engine/Private/World.cpp
--- a/Engine/Source/Runtime/Engine/Private/World.cpp
+++ b/Engine/Source/Runtime/Engine/Private/World.cpp
@@ -63,6 +63,7 @@
 #include "TickTaskManagerInterface.h"
 #include "FXSystem.h"
 #include "AudioDevice.h"
+#include "AudioDeviceManager.h"
 #include "VisualLogger/VisualLogger.h"
 #include "LevelUtils.h"
 #include "Physics/PhysicsInterfaceCore.h"
@@ -302,7 +303,9 @@ FScopedLevelCollectionContextSwitch::~FScopedLevelCollectionContextSwitch()
 	}
 }
 
+FAudioDeviceWorldDelegates::FOnWorldRegisteredToAudioDevice FAudioDeviceWorldDelegates::OnWorldRegisteredToAudioDevice;
 
+FAudioDeviceWorldDelegates::FOnWorldUnregisteredWithAudioDevice FAudioDeviceWorldDelegates::OnWorldUnregisteredWithAudioDevice;
 
 /*-----------------------------------------------------------------------------
 	UWorld implementation.
@@ -366,6 +369,15 @@ UWorld::UWorld( const FObjectInitializer& ObjectInitializer )
 	FWorldDelegates::OnPostWorldCreation.Broadcast(this);
 
 	PerfTrackers = new FWorldInGamePerformanceTrackers();
+
+	AudioDeviceDestroyedHandle = FAudioDeviceManagerDelegates::OnAudioDeviceDestroyed.AddLambda([this](const Audio::FDeviceId InDeviceId)
+	{
+		if (InDeviceId == AudioDeviceHandle.GetDeviceID())
+		{
+			FAudioDeviceHandle EmptyHandle;
+			SetAudioDevice(EmptyHandle);
+		}
+	});
 }
 
 UWorld::~UWorld()
@@ -848,7 +860,9 @@ void UWorld::BeginDestroy()
 		Scene->UpdateParameterCollections(TArray<FMaterialParameterCollectionInstanceResource*>());
 	}
 
-	AudioDeviceHandle.Reset();
+	AudioDeviceDestroyedHandle.Reset();
+	FAudioDeviceHandle EmptyHandle;
+	SetAudioDevice(EmptyHandle);
 }
 
 void UWorld::ReleasePhysicsScene()
@@ -7568,7 +7582,7 @@ void UWorld::AddPostProcessingSettings(FVector ViewLocation, FSceneView* SceneVi
 	}
 }
 
-void UWorld::SetAudioDevice(FAudioDeviceHandle& InHandle)
+void UWorld::SetAudioDevice(const FAudioDeviceHandle& InHandle)
 {
 	if (InHandle.GetDeviceID() == AudioDeviceHandle.GetDeviceID())
 	{
@@ -7576,16 +7590,22 @@ void UWorld::SetAudioDevice(FAudioDeviceHandle& InHandle)
 	}
 
 	FAudioDeviceManager* DeviceManager = GEngine ? GEngine->GetAudioDeviceManager() : nullptr;
+
+	// Register new world with incoming device first to avoid premature reporting due to no handles being valid...
 	if (DeviceManager && InHandle.IsValid())
 	{
-		GEngine->GetAudioDeviceManager()->UnregisterWorld(this, InHandle.GetDeviceID());
+		check(InHandle.GetWorld() == this);
+		DeviceManager->RegisterWorld(this, InHandle.GetDeviceID());
 	}
 
+	const Audio::FDeviceId OldDeviceId = AudioDeviceHandle.GetDeviceID();
+	const bool bUnregister = AudioDeviceHandle.IsValid();
+
 	AudioDeviceHandle = InHandle;
 
-	if (DeviceManager)
+	if (DeviceManager && bUnregister)
 	{
-		GEngine->GetAudioDeviceManager()->RegisterWorld(this, InHandle.GetDeviceID());
+		DeviceManager->UnregisterWorld(this, OldDeviceId);
 	}
 }
 
diff --git a/Engine/Source/Runtime/Engine/Public/AudioDeviceManager.h b/Engine/Source/Runtime/Engine/Public/AudioDeviceManager.h
--- a/Engine/Source/Runtime/Engine/Public/AudioDeviceManager.h
+++ b/Engine/Source/Runtime/Engine/Public/AudioDeviceManager.h
@@ -96,6 +96,8 @@ public:
 	// Returns the device ID for the audio device referenced by this handle.
 	Audio::FDeviceId GetDeviceID() const;
 
+	UWorld* GetWorld() const;
+
 	// Checks whether this points to a valid compressed chunk.
 	bool IsValid() const;
 
@@ -194,18 +196,6 @@ public:
 	// This delegate is called whenever an audio device is destroyed.
 	DECLARE_MULTICAST_DELEGATE_OneParam(FOnAudioDeviceDestroyed, Audio::FDeviceId /* AudioDeviceId*/);
 	static FOnAudioDeviceDestroyed OnAudioDeviceDestroyed;
-
-	// Called whenever a world is registered to an audio device. UWorlds are not guaranteed to be registered to the same
-	// audio device throughout their lifecycle, and there is no guarantee on the lifespan of both the UWorld and the Audio
-	// Device registered in this callback.
-	DECLARE_MULTICAST_DELEGATE_TwoParams(FOnWorldRegisteredToAudioDevice, const UWorld* /*InWorld */, Audio::FDeviceId /* AudioDeviceId*/);
-	static FOnWorldRegisteredToAudioDevice OnWorldRegisteredToAudioDevice;
-
-	// Called whenever a world is unregistered from an audio device. UWorlds are not guaranteed to be registered to the same
-	// audio device throughout their lifecycle, and there is no guarantee on the lifespan of both the UWorld and the Audio
-	// Device registered in this callback.
-	DECLARE_MULTICAST_DELEGATE_TwoParams(FOnWorldUnregisteredWithAudioDevice, const UWorld* /*InWorld */, Audio::FDeviceId /* AudioDeviceId*/);
-	static FOnWorldUnregisteredWithAudioDevice OnWorldUnregisteredWithAudioDevice;
 };
 
 /**
@@ -256,15 +246,20 @@ public:
 
 	/**
 	 * Returns a strong handle to the audio device associated with the given device ID.
-	 * if the device ID is invalid we return an invalid, zeroed handle.
+	 * if the device ID is invalid returns an invalid, zeroed handle. 
 	 */
-	FAudioDeviceHandle GetAudioDevice(Audio::FDeviceId DeviceID);
+	FAudioDeviceHandle GetAudioDevice(Audio::FDeviceId InDeviceID);
 
 	/**
-	* Returns a raw ptr to the audio device associated with the handle. If the
-	* handle is invalid then a NULL device ptr will be returned.
-	*/
-	FAudioDevice* GetAudioDeviceRaw(Audio::FDeviceId DeviceID);
+	 * Returns a raw ptr to the audio device associated with the handle. If the
+	 * handle is invalid then a NULL device ptr will be returned.
+	 */
+	FAudioDevice* GetAudioDeviceRaw(Audio::FDeviceId InDeviceID);
+
+	/**
+	  Sets the device associated with the given world.
+	  */
+	void SetAudioDevice(UWorld& InWorld, Audio::FDeviceId InDeviceID);
 
 	/**
 	* Initialize the audio device manager.
@@ -424,23 +419,12 @@ private:
 
 	FAudioDeviceHandle CreateNewDevice(const FAudioDeviceParams& InParams);
 
-	/**
-	* Shutsdown the audio device associated with the handle. The handle
-	* will become invalid after the audio device is shut down.
-	*/
-	bool ShutdownAudioDevice(Audio::FDeviceId DeviceID);
-
 	// Called exclusively by the FAudioDeviceHandle copy constructor and assignment operators:
 	void IncrementDevice(Audio::FDeviceId DeviceID);
 
 	// Called exclusively by the FAudioDeviceHandle dtor.
 	void DecrementDevice(Audio::FDeviceId DeviceID, UWorld* InWorld);
 
-	/**
-	* Shuts down all active audio devices
-	*/
-	bool ShutdownAllAudioDevices();
-
 	/** Application enters background handler */
 	void AppWillEnterBackground();
 
