Add csv scope around log flush calls on dedicated server

diff --git a/Engine/Source/Runtime/Engine/Private/GameEngine.cpp b/Engine/Source/Runtime/Engine/Private/GameEngine.cpp
--- a/Engine/Source/Runtime/Engine/Private/GameEngine.cpp
+++ b/Engine/Source/Runtime/Engine/Private/GameEngine.cpp
@@ -1609,6 +1609,7 @@ void UGameEngine::Tick( float DeltaSeconds, bool bIdleMode )
 		double CurrentTime = FPlatformTime::Seconds();
 		if (CurrentTime - LastTimeLogsFlushed > static_cast<double>(ServerFlushLogInterval))
 		{
+			CSV_SCOPED_TIMING_STAT_EXCLUSIVE(LogFlush);
 			GLog->Flush();
 
 			LastTimeLogsFlushed = FPlatformTime::Seconds();
