Renamed Brush Wireframe to Wireframe in menus to avoid confusion. We dont differentiate between wireframe and brush wireframe anymore

diff --git a/Engine/Source/Runtime/Engine/Classes/Engine/EngineBaseTypes.h b/Engine/Source/Runtime/Engine/Classes/Engine/EngineBaseTypes.h
--- a/Engine/Source/Runtime/Engine/Classes/Engine/EngineBaseTypes.h
+++ b/Engine/Source/Runtime/Engine/Classes/Engine/EngineBaseTypes.h
@@ -914,9 +914,9 @@ UENUM()
 enum EViewModeIndex
 {
 	/** Wireframe w/ brushes. */
-	VMI_BrushWireframe = 0 UMETA(DisplayName = "Brush Wireframe"),
+	VMI_BrushWireframe = 0 UMETA(DisplayName = "Wireframe"),
 	/** Wireframe w/ BSP. */
-	VMI_Wireframe = 1 UMETA(DisplayName = "Wireframe"),
+	VMI_Wireframe = 1 UMETA(DisplayName = "CSG Wireframe"),
 	/** Unlit. */
 	VMI_Unlit = 2 UMETA(DisplayName = "Unlit"),
 	/** Lit. */
diff --git a/Engine/Source/Runtime/Engine/Private/ViewModeNames.cpp b/Engine/Source/Runtime/Engine/Private/ViewModeNames.cpp
--- a/Engine/Source/Runtime/Engine/Private/ViewModeNames.cpp
+++ b/Engine/Source/Runtime/Engine/Private/ViewModeNames.cpp
@@ -18,12 +18,12 @@ TArray<FText> FillViewModeDisplayNames()
 		// Wireframe w/ brushes
 		if (ViewModeIndex == VMI_BrushWireframe)
 		{
-			ViewModeDisplayNames.Emplace(LOCTEXT("UViewModeUtils_VMI_BrushWireframe", "Brush Wireframe"));
+			ViewModeDisplayNames.Emplace(LOCTEXT("UViewModeUtils_VMI_BrushWireframe", "Wireframe"));
 		}
 		// Wireframe w/ BSP
 		else if (ViewModeIndex == VMI_Wireframe)
 		{
-			ViewModeDisplayNames.Emplace(LOCTEXT("UViewModeUtils_VMI_Wireframe", "Wireframe"));
+			ViewModeDisplayNames.Emplace(LOCTEXT("UViewModeUtils_VMI_Wireframe", "CSG Wireframe"));
 		}
 		// Unlit
 		else if (ViewModeIndex == VMI_Unlit)
