Add kismet math functions for double support

#rnx

diff --git a/Engine/Source/Runtime/Engine/Classes/Kismet/KismetMathLibrary.h b/Engine/Source/Runtime/Engine/Classes/Kismet/KismetMathLibrary.h
--- a/Engine/Source/Runtime/Engine/Classes/Kismet/KismetMathLibrary.h
+++ b/Engine/Source/Runtime/Engine/Classes/Kismet/KismetMathLibrary.h
@@ -525,6 +525,22 @@ class ENGINE_API UKismetMathLibrary : public UBlueprintFunctionLibrary
 	UFUNCTION(BlueprintPure, Category="Math|Float")
 	static float Fraction(float A);
 
+	/** Addition (A + B) */
+	UFUNCTION(BlueprintPure, meta = (DisplayName = "double + double", CompactNodeTitle = "+", Keywords = "+ add plus", CommutativeAssociativeBinaryOperator = "true"), Category = "Math|Double")
+	static double Add_DoubleDouble(double A, double B = 1.0);
+
+	/** Subtraction (A - B) */
+	UFUNCTION(BlueprintPure, meta = (DisplayName = "double - double", CompactNodeTitle = "-", Keywords = "- subtract minus"), Category = "Math|Double")
+	static double Subtract_DoubleDouble(double A, double B = 1.0);
+
+	/** Multiplication (A * B) */
+	UFUNCTION(BlueprintPure, meta = (DisplayName = "double * double", CompactNodeTitle = "*", Keywords = "* multiply", CommutativeAssociativeBinaryOperator = "true"), Category = "Math|Double")
+	static double Multiply_DoubleDouble(double A, double B);
+
+	/** Division (A / B) */
+	UFUNCTION(BlueprintPure, meta = (DisplayName = "double / double", CompactNodeTitle = "/", Keywords = "/ divide division"), Category = "Math|Double")
+	static double Divide_DoubleDouble(double A, double B = 1.0);
+
 	/** Addition (A + B) */
 	UFUNCTION(BlueprintPure, meta=(DisplayName = "float + float", CompactNodeTitle = "+", Keywords = "+ add plus", CommutativeAssociativeBinaryOperator = "true"), Category="Math|Float")
 	static float Add_FloatFloat(float A, float B = 1.f);
@@ -4040,6 +4056,7 @@ private:
 	static void ReportError_Divide_ByteByte();
 	static void ReportError_Percent_ByteByte();
 	static void ReportError_Divide_IntInt();
+	static void ReportError_Divide_DoubleDouble();
 	static void ReportError_Divide_Int64Int64();
 	static void ReportError_Percent_IntInt();
 	static void ReportError_Sqrt();
diff --git a/Engine/Source/Runtime/Engine/Classes/Kismet/KismetMathLibrary.inl b/Engine/Source/Runtime/Engine/Classes/Kismet/KismetMathLibrary.inl
--- a/Engine/Source/Runtime/Engine/Classes/Kismet/KismetMathLibrary.inl
+++ b/Engine/Source/Runtime/Engine/Classes/Kismet/KismetMathLibrary.inl
@@ -463,6 +463,34 @@ float UKismetMathLibrary::Fraction(float A)
 	return FMath::Fractional(A);
 }
 
+KISMET_MATH_FORCEINLINE
+double UKismetMathLibrary::Add_DoubleDouble(double A, double B)
+{
+	return A + B;
+}
+
+KISMET_MATH_FORCEINLINE
+double UKismetMathLibrary::Subtract_DoubleDouble(double A, double B)
+{
+	return A - B;
+}
+
+KISMET_MATH_FORCEINLINE
+double UKismetMathLibrary::Multiply_DoubleDouble(double A, double B)
+{
+	return A * B;
+}
+
+KISMET_MATH_FORCEINLINE
+double UKismetMathLibrary::Divide_DoubleDouble(double A, double B)
+{
+	if(B == 0.0)
+	{
+		ReportError_Divide_DoubleDouble();
+	}
+	return A / B;
+}
+
 KISMET_MATH_FORCEINLINE
 float UKismetMathLibrary::Add_FloatFloat(float A, float B)
 {
diff --git a/Engine/Source/Runtime/Engine/Private/KismetMathLibrary.cpp b/Engine/Source/Runtime/Engine/Private/KismetMathLibrary.cpp
--- a/Engine/Source/Runtime/Engine/Private/KismetMathLibrary.cpp
+++ b/Engine/Source/Runtime/Engine/Private/KismetMathLibrary.cpp
@@ -83,6 +83,11 @@ void UKismetMathLibrary::ReportError_Divide_IntInt()
 	FFrame::KismetExecutionMessage(TEXT("Divide by zero: Divide_IntInt"), ELogVerbosity::Warning, DivideByZeroWarning);
 }
 
+void UKismetMathLibrary::ReportError_Divide_DoubleDouble()
+{
+	FFrame::KismetExecutionMessage(TEXT("Divide by zero: ReportError_Divide_DoubleDouble"), ELogVerbosity::Warning, DivideByZeroWarning);
+}
+
 void UKismetMathLibrary::ReportError_Divide_Int64Int64()
 {
 	FFrame::KismetExecutionMessage(TEXT("Divide by zero: Divide_Int64Int64"), ELogVerbosity::Warning, DivideByZeroWarning);
