Made the Set and Get Actor label functions more easily searchable in blueprints

diff --git a/Engine/Source/Runtime/Engine/Classes/GameFramework/Actor.h b/Engine/Source/Runtime/Engine/Classes/GameFramework/Actor.h
--- a/Engine/Source/Runtime/Engine/Classes/GameFramework/Actor.h
+++ b/Engine/Source/Runtime/Engine/Classes/GameFramework/Actor.h
@@ -1934,7 +1934,7 @@ public:
 	virtual void CheckForDeprecated();
 
 	/** Returns this actor's current label.  Actor labels are only available in development builds. */
-	UFUNCTION(BlueprintCallable, Category = "Editor Scripting | Actor Editing")
+	UFUNCTION(BlueprintCallable, Category = "Editor Scripting | Actor Editing", meta = (KeyWords = "Display Name"))
 	const FString& GetActorLabel() const;
 
 	/**
@@ -1942,7 +1942,7 @@ public:
 	 * @param	NewActorLabel	The new label string to assign to the actor.  If empty, the actor will have a default label.
 	 * @param	bMarkDirty		If true the actor's package will be marked dirty for saving.  Otherwise it will not be.  You should pass false for this parameter if dirtying is not allowed (like during loads)
 	 */
-	UFUNCTION(BlueprintCallable, Category = "Editor Scripting | Actor Editing")
+	UFUNCTION(BlueprintCallable, Category = "Editor Scripting | Actor Editing", meta = (KeyWords = "Display Name"))
 	void SetActorLabel( const FString& NewActorLabel, bool bMarkDirty = true );
 
 	/** Advanced - clear the actor label. */
