Prevent multiple notificationitem from appearing at the same time by doing a single update per game-thread tick so slate has the time to cleanup the older one before we spawn a new one

diff --git a/Engine/Source/Runtime/Engine/Private/TextureCompiler.cpp b/Engine/Source/Runtime/Engine/Private/TextureCompiler.cpp
--- a/Engine/Source/Runtime/Engine/Private/TextureCompiler.cpp
+++ b/Engine/Source/Runtime/Engine/Private/TextureCompiler.cpp
@@ -245,21 +245,21 @@ void FTextureCompilingManager::UpdateCompilationNotification()
 	check(IsInGameThread());
 	static TWeakPtr<SNotificationItem> TextureCompilationPtr;
 
-	// Finished all requests! Notify the UI.
 	TSharedPtr<SNotificationItem> NotificationItem = TextureCompilationPtr.Pin();
 
 	const int32 NumRemainingCompilations = GetNumRemainingTextures();
-	if (NumRemainingCompilations == 0 && NotificationItem.IsValid())
+	if (NumRemainingCompilations == 0)
 	{
-		NotificationItem->SetText(NSLOCTEXT("TextureBuild", "TextureBuildFinished", "Finished building Textures!"));
-		NotificationItem->SetCompletionState(SNotificationItem::CS_Success);
-		NotificationItem->ExpireAndFadeout();
+		if (NotificationItem.IsValid())
+		{
+			NotificationItem->SetText(NSLOCTEXT("TextureBuild", "TextureBuildFinished", "Finished building Textures!"));
+			NotificationItem->SetCompletionState(SNotificationItem::CS_Success);
+			NotificationItem->ExpireAndFadeout();
 
-		TextureCompilationPtr.Reset();
-		return;
+			TextureCompilationPtr.Reset();
+		}
 	}
-
-	if (NumRemainingCompilations > 0)
+	else
 	{
 		if (!NotificationItem.IsValid())
 		{
@@ -270,8 +270,8 @@ void FTextureCompilingManager::UpdateCompilationNotification()
 			Info.FadeOutDuration = 0.0f;
 			Info.ExpireDuration = 0.0f;
 
-			TextureCompilationPtr = FSlateNotificationManager::Get().AddNotification(Info);
-			NotificationItem = TextureCompilationPtr.Pin();
+			NotificationItem = FSlateNotificationManager::Get().AddNotification(Info);
+			TextureCompilationPtr = NotificationItem;
 		}
 
 		FFormatNamedArguments Args;
@@ -362,8 +362,6 @@ void FTextureCompilingManager::AddTextures(const TArray<UTexture*>& InTextures)
 	// so this will allow us to refresh compiled textures
 	// of the highest priority to improve the UI experience.
 	ProcessTextures(1 /* Maximum Priority */);
-
-	UpdateCompilationNotification();
 }
 
 void FTextureCompilingManager::FinishCompilation(const TArray<UTexture*>& InTextures)
@@ -436,8 +434,6 @@ void FTextureCompilingManager::FinishCompilation(const TArray<UTexture*>& InText
 			{
 				Bucket.Remove(Texture);
 			}
-
-			UpdateCompilationNotification();
 		}
 	}
 }
@@ -512,7 +508,6 @@ void FTextureCompilingManager::ProcessTextures(int32 MaximumPriority)
 					}
 
 					RegisteredTextureBuckets[PriorityIndex] = MoveTemp(TexturesToPostpone);
-					UpdateCompilationNotification();
 				}
 			}
 		}
@@ -591,6 +586,8 @@ void FTextureCompilingManager::ProcessTextures(int32 MaximumPriority)
 void FTextureCompilingManager::Tick(float DeltaTime)
 {
 	ProcessTextures();
+
+	UpdateCompilationNotification();
 }
 
 #endif // #if WITH_EDITOR
