Edigrate FlushRenderingCommands fix from Jeremy.Moore to avoid crash when editing terrain

Move FlushRenderingCommands() into UTexture::CachePlatformData().
Whenever we recache the texture platform data we need to ensure that previous Resource processing is completed since the texture resource currently accesses the platform data.
This was handled for some cases within UTexture2D::UpdateResource() but there are other cases we need to catch.

diff --git a/Engine/Source/Runtime/Engine/Private/Texture2D.cpp b/Engine/Source/Runtime/Engine/Private/Texture2D.cpp
--- a/Engine/Source/Runtime/Engine/Private/Texture2D.cpp
+++ b/Engine/Source/Runtime/Engine/Private/Texture2D.cpp
@@ -1080,12 +1080,6 @@ void UTexture2D::UpdateResource()
 	}
 
 #if WITH_EDITOR
-	// Wait for FTexture2DResource::InitRHI() complete before modifying the PlatformData.
-	FTexture2DResource* Texture2DResource = (FTexture2DResource*)GetResource();
-	if (Texture2DResource && !Texture2DResource->bReadyForStreaming)
-	{
-		FlushRenderingCommands();
-	}
 	// Recache platform data if the source has changed.
 	if (FTextureCompilingManager::Get().IsAsyncCompilationAllowed(this))
 	{
diff --git a/Engine/Source/Runtime/Engine/Private/TextureDerivedData.cpp b/Engine/Source/Runtime/Engine/Private/TextureDerivedData.cpp
--- a/Engine/Source/Runtime/Engine/Private/TextureDerivedData.cpp
+++ b/Engine/Source/Runtime/Engine/Private/TextureDerivedData.cpp
@@ -1640,6 +1640,10 @@ void UTexture::CachePlatformData(bool bAsyncCache, bool bAllowAsyncBuild, bool b
 				if (PlatformDataLink)
 				{
 					ReleaseResource();
+
+					// Need to wait for any previous InitRHI() to complete before modifying PlatformData
+					// We could remove this flush if InitRHI() was modified to not access PlatformData directly
+					FlushRenderingCommands();
 				}
 				else
 				{
