Debug draw utilities in KismetTraceUtils.h are now tagged with ENGINE_API

diff --git a/Engine/Source/Runtime/Engine/Private/KismetTraceUtils.h b/Engine/Source/Runtime/Engine/Private/KismetTraceUtils.h
--- a/Engine/Source/Runtime/Engine/Private/KismetTraceUtils.h
+++ b/Engine/Source/Runtime/Engine/Private/KismetTraceUtils.h
@@ -15,20 +15,20 @@ FCollisionObjectQueryParams ConfigureCollisionObjectParams(const TArray<TEnumAsB
 
 #if ENABLE_DRAW_DEBUG
 
-void DrawDebugSweptSphere(const UWorld* InWorld, FVector const& Start, FVector const& End, float Radius, FColor const& Color, bool bPersistentLines = false, float LifeTime = -1.f, uint8 DepthPriority = 0);
-void DrawDebugSweptBox(const UWorld* InWorld, FVector const& Start, FVector const& End, FRotator const & Orientation, FVector const & HalfSize, FColor const& Color, bool bPersistentLines = false, float LifeTime = -1.f, uint8 DepthPriority = 0);
+ENGINE_API void DrawDebugSweptSphere(const UWorld* InWorld, FVector const& Start, FVector const& End, float Radius, FColor const& Color, bool bPersistentLines = false, float LifeTime = -1.f, uint8 DepthPriority = 0);
+ENGINE_API void DrawDebugSweptBox(const UWorld* InWorld, FVector const& Start, FVector const& End, FRotator const & Orientation, FVector const & HalfSize, FColor const& Color, bool bPersistentLines = false, float LifeTime = -1.f, uint8 DepthPriority = 0);
 
-void DrawDebugLineTraceSingle(const UWorld* World, const FVector& Start, const FVector& End, EDrawDebugTrace::Type DrawDebugType, bool bHit, const FHitResult& OutHit, FLinearColor TraceColor, FLinearColor TraceHitColor, float DrawTime);
-void DrawDebugLineTraceMulti(const UWorld* World, const FVector& Start, const FVector& End, EDrawDebugTrace::Type DrawDebugType, bool bHit, const TArray<FHitResult>& OutHits, FLinearColor TraceColor, FLinearColor TraceHitColor, float DrawTime);
+ENGINE_API void DrawDebugLineTraceSingle(const UWorld* World, const FVector& Start, const FVector& End, EDrawDebugTrace::Type DrawDebugType, bool bHit, const FHitResult& OutHit, FLinearColor TraceColor, FLinearColor TraceHitColor, float DrawTime);
+ENGINE_API void DrawDebugLineTraceMulti(const UWorld* World, const FVector& Start, const FVector& End, EDrawDebugTrace::Type DrawDebugType, bool bHit, const TArray<FHitResult>& OutHits, FLinearColor TraceColor, FLinearColor TraceHitColor, float DrawTime);
 
-void DrawDebugBoxTraceSingle(const UWorld* World, const FVector& Start, const FVector& End, const FVector HalfSize, const FRotator Orientation, EDrawDebugTrace::Type DrawDebugType, bool bHit, const FHitResult& OutHit, FLinearColor TraceColor, FLinearColor TraceHitColor, float DrawTime);
-void DrawDebugBoxTraceMulti(const UWorld* World, const FVector& Start, const FVector& End, const FVector HalfSize, const FRotator Orientation, EDrawDebugTrace::Type DrawDebugType, bool bHit, const TArray<FHitResult>& OutHits, FLinearColor TraceColor, FLinearColor TraceHitColor, float DrawTime);
+ENGINE_API void DrawDebugBoxTraceSingle(const UWorld* World, const FVector& Start, const FVector& End, const FVector HalfSize, const FRotator Orientation, EDrawDebugTrace::Type DrawDebugType, bool bHit, const FHitResult& OutHit, FLinearColor TraceColor, FLinearColor TraceHitColor, float DrawTime);
+ENGINE_API void DrawDebugBoxTraceMulti(const UWorld* World, const FVector& Start, const FVector& End, const FVector HalfSize, const FRotator Orientation, EDrawDebugTrace::Type DrawDebugType, bool bHit, const TArray<FHitResult>& OutHits, FLinearColor TraceColor, FLinearColor TraceHitColor, float DrawTime);
 
-void DrawDebugSphereTraceSingle(const UWorld* World, const FVector& Start, const FVector& End, float Radius, EDrawDebugTrace::Type DrawDebugType, bool bHit, const FHitResult& OutHit, FLinearColor TraceColor, FLinearColor TraceHitColor, float DrawTime);
-void DrawDebugSphereTraceMulti(const UWorld* World, const FVector& Start, const FVector& End, float Radius, EDrawDebugTrace::Type DrawDebugType, bool bHit, const TArray<FHitResult>& OutHits, FLinearColor TraceColor, FLinearColor TraceHitColor, float DrawTime);
+ENGINE_API void DrawDebugSphereTraceSingle(const UWorld* World, const FVector& Start, const FVector& End, float Radius, EDrawDebugTrace::Type DrawDebugType, bool bHit, const FHitResult& OutHit, FLinearColor TraceColor, FLinearColor TraceHitColor, float DrawTime);
+ENGINE_API void DrawDebugSphereTraceMulti(const UWorld* World, const FVector& Start, const FVector& End, float Radius, EDrawDebugTrace::Type DrawDebugType, bool bHit, const TArray<FHitResult>& OutHits, FLinearColor TraceColor, FLinearColor TraceHitColor, float DrawTime);
 
-void DrawDebugCapsuleTraceSingle(const UWorld* World, const FVector& Start, const FVector& End, float Radius, float HalfHeight, EDrawDebugTrace::Type DrawDebugType, bool bHit, const FHitResult& OutHit, FLinearColor TraceColor, FLinearColor TraceHitColor, float DrawTime);
-void DrawDebugCapsuleTraceMulti(const UWorld* World, const FVector& Start, const FVector& End, float Radius, float HalfHeight, EDrawDebugTrace::Type DrawDebugType, bool bHit, const TArray<FHitResult>& OutHits, FLinearColor TraceColor, FLinearColor TraceHitColor, float DrawTime);
+ENGINE_API void DrawDebugCapsuleTraceSingle(const UWorld* World, const FVector& Start, const FVector& End, float Radius, float HalfHeight, EDrawDebugTrace::Type DrawDebugType, bool bHit, const FHitResult& OutHit, FLinearColor TraceColor, FLinearColor TraceHitColor, float DrawTime);
+ENGINE_API void DrawDebugCapsuleTraceMulti(const UWorld* World, const FVector& Start, const FVector& End, float Radius, float HalfHeight, EDrawDebugTrace::Type DrawDebugType, bool bHit, const TArray<FHitResult>& OutHits, FLinearColor TraceColor, FLinearColor TraceHitColor, float DrawTime);
 
 #endif // ENABLE_DRAW_DEBUG
 
