#PUMA DistanceCurves

diff --git a/Engine/Source/Runtime/Engine/Private/Animation/AnimCurveTypes.cpp b/Engine/Source/Runtime/Engine/Private/Animation/AnimCurveTypes.cpp
--- a/Engine/Source/Runtime/Engine/Private/Animation/AnimCurveTypes.cpp
+++ b/Engine/Source/Runtime/Engine/Private/Animation/AnimCurveTypes.cpp
@@ -308,6 +308,116 @@ void FTransformCurve::Resize(float NewLength, bool bInsert/* whether insert or r
 	ScaleCurve.Resize(NewLength, bInsert, OldStartTime, OldEndTime);
 }
 
+////////////////////////////////////////////////////
+//  FCachedFloatCurve
+
+bool FCachedFloatCurve::IsValid(UAnimSequenceBase* InAnimSequence) const
+{
+	return ((CurveName != NAME_None) && (GetFloatCurve(InAnimSequence) != nullptr));
+}
+
+float FCachedFloatCurve::GetValueAtPosition(UAnimSequenceBase* InAnimSequence, const float& InPosition) const
+{
+	if (const FFloatCurve* DistanceCurve = GetFloatCurve(InAnimSequence))
+	{
+		return DistanceCurve->Evaluate(InPosition);
+	}
+
+	return 0.f;
+}
+
+USkeleton::AnimCurveUID FCachedFloatCurve::GetAnimCurveUID(UAnimSequenceBase* InAnimSequence) const
+{
+	if (CachedUID == SmartName::MaxUID && InAnimSequence)
+	{
+		if (const USkeleton* Skeleton = InAnimSequence->GetSkeleton())
+		{
+			const FSmartNameMapping* CurveNameMapping = Skeleton->GetSmartNameContainer(USkeleton::AnimCurveMappingName);
+			if (CurveNameMapping)
+			{
+				CachedUID = CurveNameMapping->FindUID(CurveName);
+			}
+		}
+	}
+
+	return CachedUID;
+}
+
+const FFloatCurve* FCachedFloatCurve::GetFloatCurve(UAnimSequenceBase* InAnimSequence) const
+{
+	if (InAnimSequence)
+	{
+		USkeleton::AnimCurveUID DistanceCurveUID = GetAnimCurveUID(InAnimSequence);
+		if (DistanceCurveUID != SmartName::MaxUID)
+		{
+			return (const FFloatCurve*)(InAnimSequence->GetCurveData().GetCurveData(DistanceCurveUID));
+		}
+	}
+
+	return nullptr;
+}
+
+////////////////////////////////////////////////////
+//  FDistanceCurve
+
+float FDistanceCurve::GetDistanceRange(UAnimSequenceBase* InAnimSequence) const
+{
+	if (const FFloatCurve* DistanceCurve = GetFloatCurve(InAnimSequence))
+	{
+		if (DistanceCurve->FloatCurve.GetNumKeys() >= 2)
+		{
+			return (DistanceCurve->FloatCurve.GetLastKey().Value - DistanceCurve->FloatCurve.GetFirstKey().Value);
+		}
+	}
+	return 0.f;
+}
+
+float FDistanceCurve::GetAnimPositionFromDistance(UAnimSequenceBase* InAnimSequence, const float& InDistance)
+{
+	if (const FFloatCurve* DistanceCurve = GetFloatCurve(InAnimSequence))
+	{
+		const TArray<FRichCurveKey>& Keys = DistanceCurve->FloatCurve.GetConstRefOfKeys();
+
+		const int32 NumKeys = Keys.Num();
+		if (NumKeys < 2)
+		{
+			return 0.f;
+		}
+
+		// Some assumptions: 
+		// - keys have unique values, so for a given value, it maps to a single position on the timeline of the animation.
+		// - key values are sorted in increasing order.
+
+		int32 first = 1;
+		int32 last = NumKeys - 1;
+		int32 count = last - first;
+
+		while (count > 0)
+		{
+			int32 step = count / 2;
+			int32 middle = first + step;
+
+			if (InDistance > Keys[middle].Value)
+			{
+				first = middle + 1;
+				count -= step + 1;
+			}
+			else
+			{
+				count = step;
+			}
+		}
+
+		const FRichCurveKey& KeyA = Keys[first - 1];
+		const FRichCurveKey& KeyB = Keys[first];
+		const float Diff = KeyB.Value - KeyA.Value;
+		const float Alpha = !FMath::IsNearlyZero(Diff) ? ((InDistance - KeyA.Value) / Diff) : 0.f;
+		return FMath::Lerp(KeyA.Time, KeyB.Time, Alpha);
+	}
+
+	return 0.f;
+}
+
 /////////////////////////////////////////////////////
 // FRawCurveTracks
 
diff --git a/Engine/Source/Runtime/Engine/Public/Animation/AnimCurveTypes.h b/Engine/Source/Runtime/Engine/Public/Animation/AnimCurveTypes.h
--- a/Engine/Source/Runtime/Engine/Public/Animation/AnimCurveTypes.h
+++ b/Engine/Source/Runtime/Engine/Public/Animation/AnimCurveTypes.h
@@ -245,6 +245,49 @@ struct FTransformCurve: public FAnimCurveBase
 	void Resize(float NewLength, bool bInsert/* whether insert or remove*/, float OldStartTime, float OldEndTime);
 };
 
+USTRUCT(BlueprintType)
+struct ENGINE_API FCachedFloatCurve
+{
+	GENERATED_USTRUCT_BODY()
+
+public:
+	UPROPERTY(EditAnywhere, Category = "Curve Settings")
+		FName CurveName;
+
+private:
+	mutable USkeleton::AnimCurveUID CachedUID;
+
+public:
+	FCachedFloatCurve()
+		: CachedUID(SmartName::MaxUID)
+	{}
+
+	bool IsValid(UAnimSequenceBase* InAnimSequence) const;
+	float GetValueAtPosition(UAnimSequenceBase* InAnimSequence, const float& InPosition) const;
+	const FFloatCurve* GetFloatCurve(UAnimSequenceBase* InAnimSequence) const;
+
+protected:
+	USkeleton::AnimCurveUID GetAnimCurveUID(UAnimSequenceBase* InAnimSequence) const;
+};
+
+USTRUCT(BlueprintType)
+struct ENGINE_API FDistanceCurve : public FCachedFloatCurve
+{
+	GENERATED_USTRUCT_BODY()
+
+public:
+	FDistanceCurve()
+		: Super()
+	{}
+
+	FDistanceCurve(const Super& Other)
+		: Super(Other)
+	{}
+
+	float GetDistanceRange(UAnimSequenceBase* InAnimSequence) const;
+	float GetAnimPositionFromDistance(UAnimSequenceBase* InAnimSequence, const float& InDistance);
+};
+
 /**
 * This is array of curves that run when collecting curves natively 
 */
