Fix crash in VT because textures placeholders were never converted back to actual VT.
   - IsCurrentlyVirtualTextured was returning false during the CreateResource() preventing VT creation.
   - This also fixes a potentially wrong assumption about the type of the Resource stored in UTexture2D.

diff --git a/Engine/Source/Runtime/Engine/Classes/Engine/Texture2D.h b/Engine/Source/Runtime/Engine/Classes/Engine/Texture2D.h
--- a/Engine/Source/Runtime/Engine/Classes/Engine/Texture2D.h
+++ b/Engine/Source/Runtime/Engine/Classes/Engine/Texture2D.h
@@ -79,6 +79,11 @@ private:
 	UPROPERTY()
 	FIntPoint ImportedSize;
 
+#if WITH_EDITOR
+	/** Whether the current texture is a default placeholder because compilation is still ongoing. */
+	bool bIsDefaultTexture = false;
+#endif
+
 	/** The derived data for this texture on this platform. */
 	FTexturePlatformData* PrivatePlatformData;
 public:
diff --git a/Engine/Source/Runtime/Engine/Classes/Engine/TextureCube.h b/Engine/Source/Runtime/Engine/Classes/Engine/TextureCube.h
--- a/Engine/Source/Runtime/Engine/Classes/Engine/TextureCube.h
+++ b/Engine/Source/Runtime/Engine/Classes/Engine/TextureCube.h
@@ -18,6 +18,12 @@ class UTextureCube : public UTexture
 
 	/** The derived data for this texture on this platform. */
 	FTexturePlatformData* PrivatePlatformData;
+
+#if WITH_EDITOR
+	/** Whether the current texture is a default placeholder because compilation is still ongoing. */
+	bool bIsDefaultTexture = false;
+#endif
+
 public:
 	// The deprecation will be enabled further along the dev cycle and fixed accordingly.
 	// This is to avoid merge conflicts with other branches that fixing this deprecation might cause.
diff --git a/Engine/Source/Runtime/Engine/Private/Texture2D.cpp b/Engine/Source/Runtime/Engine/Private/Texture2D.cpp
--- a/Engine/Source/Runtime/Engine/Private/Texture2D.cpp
+++ b/Engine/Source/Runtime/Engine/Private/Texture2D.cpp
@@ -645,8 +645,7 @@ int32 UTexture2D::GetNumRequestedMips() const
 
 bool UTexture2D::IsDefaultTexture() const
 {
-	const FTexture2DResource* TextureResource = (const FTexture2DResource*)GetResource();
-	return (PrivatePlatformData && !PrivatePlatformData->IsAsyncWorkComplete()) || (TextureResource && TextureResource->ProxiedResource != nullptr);
+	return bIsDefaultTexture || (PrivatePlatformData && !PrivatePlatformData->IsAsyncWorkComplete());
 }
 
 void UTexture2D::PostEditUndo()
@@ -1259,8 +1258,10 @@ FTextureResource* UTexture2D::CreateResource()
 
 		UnlinkStreaming();
 		bIsStreamable = false;
+		bIsDefaultTexture = true;
 		return new FTexture2DResource(this, (const FTexture2DResource*)GetDefaultTexture2D(this)->GetResource());
 	}
+	bIsDefaultTexture = false;
 #endif
 
 	if (IsCurrentlyVirtualTextured())
diff --git a/Engine/Source/Runtime/Engine/Private/TextureCube.cpp b/Engine/Source/Runtime/Engine/Private/TextureCube.cpp
--- a/Engine/Source/Runtime/Engine/Private/TextureCube.cpp
+++ b/Engine/Source/Runtime/Engine/Private/TextureCube.cpp
@@ -586,8 +586,10 @@ FTextureResource* UTextureCube::CreateResource()
 	{
 		FTextureCompilingManager::Get().AddTextures({ this });
 
+		bIsDefaultTexture = true;
 		return new FTextureCubeResource(this, (const FTextureCubeResource*)GetDefaultTextureCube(this)->GetResource());
 	}
+	bIsDefaultTexture = false;
 #endif
 
 	FTextureResource* NewResource = NULL;
@@ -617,8 +619,7 @@ void UTextureCube::GetResourceSizeEx(FResourceSizeEx& CumulativeResourceSize)
 #if WITH_EDITOR
 bool UTextureCube::IsDefaultTexture() const
 {
-	const FTextureCubeResource* TextureResource = (const FTextureCubeResource*)GetResource();
-	return (PrivatePlatformData && !PrivatePlatformData->IsAsyncWorkComplete()) || (TextureResource && TextureResource->GetProxiedResource() != nullptr);
+	return bIsDefaultTexture || (PrivatePlatformData && !PrivatePlatformData->IsAsyncWorkComplete());
 }
 
 uint32 UTextureCube::GetMaximumDimension() const
