Fix race condition in previous fix because 2 resources can exist at the same time and each must have it's own boolean value

diff --git a/Engine/Source/Runtime/Engine/Classes/Engine/Texture2D.h b/Engine/Source/Runtime/Engine/Classes/Engine/Texture2D.h
--- a/Engine/Source/Runtime/Engine/Classes/Engine/Texture2D.h
+++ b/Engine/Source/Runtime/Engine/Classes/Engine/Texture2D.h
@@ -79,11 +79,6 @@ private:
 	UPROPERTY()
 	FIntPoint ImportedSize;
 
-#if WITH_EDITOR
-	/** Whether the current texture is a default placeholder because compilation is still ongoing. */
-	bool bIsDefaultTexture = false;
-#endif
-
 	/** The derived data for this texture on this platform. */
 	FTexturePlatformData* PrivatePlatformData;
 public:
diff --git a/Engine/Source/Runtime/Engine/Classes/Engine/TextureCube.h b/Engine/Source/Runtime/Engine/Classes/Engine/TextureCube.h
--- a/Engine/Source/Runtime/Engine/Classes/Engine/TextureCube.h
+++ b/Engine/Source/Runtime/Engine/Classes/Engine/TextureCube.h
@@ -18,12 +18,6 @@ class UTextureCube : public UTexture
 
 	/** The derived data for this texture on this platform. */
 	FTexturePlatformData* PrivatePlatformData;
-
-#if WITH_EDITOR
-	/** Whether the current texture is a default placeholder because compilation is still ongoing. */
-	bool bIsDefaultTexture = false;
-#endif
-
 public:
 	// The deprecation will be enabled further along the dev cycle and fixed accordingly.
 	// This is to avoid merge conflicts with other branches that fixing this deprecation might cause.
diff --git a/Engine/Source/Runtime/Engine/Private/Texture2D.cpp b/Engine/Source/Runtime/Engine/Private/Texture2D.cpp
--- a/Engine/Source/Runtime/Engine/Private/Texture2D.cpp
+++ b/Engine/Source/Runtime/Engine/Private/Texture2D.cpp
@@ -645,7 +645,7 @@ int32 UTexture2D::GetNumRequestedMips() const
 
 bool UTexture2D::IsDefaultTexture() const
 {
-	return bIsDefaultTexture || (PrivatePlatformData && !PrivatePlatformData->IsAsyncWorkComplete());
+	return (PrivatePlatformData && !PrivatePlatformData->IsAsyncWorkComplete()) || (GetResource() && GetResource()->IsProxy());
 }
 
 void UTexture2D::PostEditUndo()
@@ -1258,13 +1258,11 @@ FTextureResource* UTexture2D::CreateResource()
 
 		UnlinkStreaming();
 		bIsStreamable = false;
-		bIsDefaultTexture = true;
 		return new FTexture2DResource(this, (const FTexture2DResource*)GetDefaultTexture2D(this)->GetResource());
 	}
-	bIsDefaultTexture = false;
 #endif
-
-	if (IsCurrentlyVirtualTextured())
+	// We can't use IsCurrentlyVirtualTextured here because the call might get redirected to the placeholder texture
+	if (VirtualTextureStreaming && PrivatePlatformData && PrivatePlatformData->VTData)
 	{
 		UnlinkStreaming();
 		bIsStreamable = false;
diff --git a/Engine/Source/Runtime/Engine/Private/TextureCube.cpp b/Engine/Source/Runtime/Engine/Private/TextureCube.cpp
--- a/Engine/Source/Runtime/Engine/Private/TextureCube.cpp
+++ b/Engine/Source/Runtime/Engine/Private/TextureCube.cpp
@@ -503,6 +503,8 @@ public:
 		return TextureCubeRHI;
 	}
 
+	virtual bool IsProxy() const override { return ProxiedResource != nullptr; }
+
 	const FTextureCubeResource* GetProxiedResource() const { return ProxiedResource; }
 private:
 	/** A reference to the texture's RHI resource as a cube-map texture. */
@@ -586,10 +588,8 @@ FTextureResource* UTextureCube::CreateResource()
 	{
 		FTextureCompilingManager::Get().AddTextures({ this });
 
-		bIsDefaultTexture = true;
 		return new FTextureCubeResource(this, (const FTextureCubeResource*)GetDefaultTextureCube(this)->GetResource());
 	}
-	bIsDefaultTexture = false;
 #endif
 
 	FTextureResource* NewResource = NULL;
@@ -619,7 +619,7 @@ void UTextureCube::GetResourceSizeEx(FResourceSizeEx& CumulativeResourceSize)
 #if WITH_EDITOR
 bool UTextureCube::IsDefaultTexture() const
 {
-	return bIsDefaultTexture || (PrivatePlatformData && !PrivatePlatformData->IsAsyncWorkComplete());
+	return (PrivatePlatformData && !PrivatePlatformData->IsAsyncWorkComplete()) || (GetResource() && GetResource()->IsProxy());
 }
 
 uint32 UTextureCube::GetMaximumDimension() const
diff --git a/Engine/Source/Runtime/Engine/Public/TextureResource.h b/Engine/Source/Runtime/Engine/Public/TextureResource.h
--- a/Engine/Source/Runtime/Engine/Public/TextureResource.h
+++ b/Engine/Source/Runtime/Engine/Public/TextureResource.h
@@ -82,6 +82,11 @@ public:
 	// used when updating mip map bias offset
 	virtual void RefreshSamplerStates() {}
 
+	/**
+	* Returns true if the resource is proxying another one.
+	*/
+	virtual bool IsProxy() const { return false; }
+
 #if STATS
 	/* The Stat_ FName corresponding to each TEXTUREGROUP */
 	static FName TextureGroupStatFNames[TEXTUREGROUP_MAX];
@@ -137,6 +142,7 @@ public:
 	}
 
 	virtual FString GetFriendlyName() const override;
+	virtual bool IsProxy() const override { return ProxiedResource != nullptr; }
 
 	//Returns the current first mip (always valid)
 	int32 GetCurrentFirstMip() const
