Fix infinite wait on cleared cooked platform data in case an UpdateResource clears them

diff --git a/Engine/Source/Runtime/Engine/Private/TextureDerivedData.cpp b/Engine/Source/Runtime/Engine/Private/TextureDerivedData.cpp
--- a/Engine/Source/Runtime/Engine/Private/TextureDerivedData.cpp
+++ b/Engine/Source/Runtime/Engine/Private/TextureDerivedData.cpp
@@ -1838,7 +1838,11 @@ bool UTexture::IsCachedCookedPlatformDataLoaded( const ITargetPlatform* TargetPl
 
 		// begin cache hasn't been called
 		if ( !PlatformData )
+		{
+			// In case an UpdateResource happens, cooked platform data might be cleared and we might need to reschedule
+			BeginCacheForCookedPlatformData(TargetPlatform);
 			return false;
+		}
 
 		if ( (PlatformData->AsyncTask != NULL) && ( PlatformData->AsyncTask->IsWorkDone() == true ) )
 		{
