Integrate 13765000 - Remove PhysScene_Chaos::Shutdown registering with AppPreExit.
Fixes HLOD builds crashing.

diff --git a/Engine/Source/Runtime/Engine/Private/PhysicsEngine/Experimental/PhysScene_Chaos.cpp b/Engine/Source/Runtime/Engine/Private/PhysicsEngine/Experimental/PhysScene_Chaos.cpp
--- a/Engine/Source/Runtime/Engine/Private/PhysicsEngine/Experimental/PhysScene_Chaos.cpp
+++ b/Engine/Source/Runtime/Engine/Private/PhysicsEngine/Experimental/PhysScene_Chaos.cpp
@@ -422,10 +422,6 @@ FPhysScene_Chaos::FPhysScene_Chaos(AActor* InSolverActor
 #endif
 	}
 
-	// #BGallagher Temporary while we're using the global scene singleton. Shouldn't be required
-	// once we have a better lifecycle for the scenes.
-	FCoreDelegates::OnPreExit.AddRaw(this, &FPhysScene_Chaos::Shutdown);
-
 	PhysicsProxyToComponentMap.Reset();
 	ComponentToPhysicsProxyMap.Reset();
 
