SS - Added missing Tags and fixed some olt standing issues with ParallelGameThread beeing recognized as ParallelRenderingThread (epecially in cases with Tessellation)

diff --git a/Engine/Source/Runtime/Engine/Private/LevelTick.cpp b/Engine/Source/Runtime/Engine/Private/LevelTick.cpp
--- a/Engine/Source/Runtime/Engine/Private/LevelTick.cpp
+++ b/Engine/Source/Runtime/Engine/Private/LevelTick.cpp
@@ -1024,6 +1024,7 @@ void UWorld::SendAllEndOfFrameUpdates()
 	auto ParallelWork = 
 		[](int32 Index) 
 		{
+			FOptionalTaskTagScope Scope(ETaskTag::EParallelGameThread);
 			UActorComponent* NextComponent = LocalComponentsThatNeedEndOfFrameUpdate[Index];
 			if (NextComponent)
 			{
diff --git a/Engine/Source/Runtime/Engine/Private/Particles/ParticleSystemManager.cpp b/Engine/Source/Runtime/Engine/Private/Particles/ParticleSystemManager.cpp
--- a/Engine/Source/Runtime/Engine/Private/Particles/ParticleSystemManager.cpp
+++ b/Engine/Source/Runtime/Engine/Private/Particles/ParticleSystemManager.cpp
@@ -109,6 +109,7 @@ public:
 
 	void DoTask(ENamedThreads::Type CurrentThread, const FGraphEventRef& MyCompletionGraphEvent)
 	{
+		FTaskTagScope Scope(ETaskTag::EParallelRenderingThread);
 		SCOPE_CYCLE_COUNTER(STAT_PSCMan_AsyncBatch);
 
 // 		FString Ticked;
diff --git a/Engine/Source/Runtime/Engine/Private/TessellationRendering.cpp b/Engine/Source/Runtime/Engine/Private/TessellationRendering.cpp
--- a/Engine/Source/Runtime/Engine/Private/TessellationRendering.cpp
+++ b/Engine/Source/Runtime/Engine/Private/TessellationRendering.cpp
@@ -9,7 +9,7 @@
   * Game thread version that looks at the material settings. Will not change answer during a shader compile */
 bool MaterialSettingsRequireAdjacencyInformation_GameThread(UMaterialInterface* Material, const FVertexFactoryType* VertexFactoryType, const FStaticFeatureLevel InFeatureLevel)
 {
-	check(IsInGameThread());
+	check(IsInParallelGameThread() || IsInGameThread());
 
 	//if we pass null here we have to guarantee that the VF supports tessellation (e.g by using type LocalVF)
 	bool VertexFactorySupportsTessellation = !VertexFactoryType || (VertexFactoryType && VertexFactoryType->SupportsTessellationShaders());
@@ -65,14 +65,14 @@ bool RequiresAdjacencyInformation(UMaterialInterface* Material, const FVertexFac
 		return false;
 	}
 
-	if (IsInRenderingThread() || IsInParallelRenderingThread())
-	{
-		return MaterialRenderingRequiresAdjacencyInformation_RenderingThread(Material, VertexFactoryType, InFeatureLevel);
-	}
-	else if (IsInGameThread())
+	if (IsInParallelGameThread() || IsInGameThread())
 	{
 		return MaterialSettingsRequireAdjacencyInformation_GameThread(Material, VertexFactoryType, InFeatureLevel);
 	}
+	else if (IsInRenderingThread() || IsInParallelRenderingThread())
+	{
+		return MaterialRenderingRequiresAdjacencyInformation_RenderingThread(Material, VertexFactoryType, InFeatureLevel);
+	} 
 	else
 	{
 		//if we pass null here we have to guarantee that the VF supports tessellation (e.g by using type LocalVF)
