SS - Fix RequiresAdjacencyInformation logig to work with ParallelGameThread.

diff --git a/Engine/Source/Runtime/Engine/Private/TessellationRendering.cpp b/Engine/Source/Runtime/Engine/Private/TessellationRendering.cpp
--- a/Engine/Source/Runtime/Engine/Private/TessellationRendering.cpp
+++ b/Engine/Source/Runtime/Engine/Private/TessellationRendering.cpp
@@ -51,6 +51,26 @@ bool MaterialRenderingRequiresAdjacencyInformation_RenderingThread(UMaterialInte
 	return false;
 }
 
+
+/** Returns true if the Material and Vertex Factory combination require adjacency information.
+  * Parallel thread version that looks at the current shader that will be used. **Will change answer during a shader compile** */
+bool MaterialRenderingRequiresAdjacencyInformation_Parallel(UMaterialInterface* Material, const FVertexFactoryType* VertexFactoryType, const FStaticFeatureLevel InFeatureLevel)
+{
+	//if we pass null here we have to guarantee that the VF supports tessellation (e.g by using type LocalVF)
+	bool VertexFactorySupportsTessellation = !VertexFactoryType || (VertexFactoryType && VertexFactoryType->SupportsTessellationShaders());
+
+	// Concurrent?
+	if (VertexFactorySupportsTessellation && Material)
+	{
+		const UMaterial* BaseMaterial = Material->GetMaterial_Concurrent();
+		check(BaseMaterial);
+		EMaterialTessellationMode TessellationMode = (EMaterialTessellationMode)BaseMaterial->D3D11TessellationMode;
+		bool bEnableCrackFreeDisplacement = BaseMaterial->bEnableCrackFreeDisplacement;
+		return TessellationMode == MTM_PNTriangles || (TessellationMode == MTM_FlatTessellation && bEnableCrackFreeDisplacement);
+	}
+	return false;
+}
+
 /** Returns true if the Material and Vertex Factory combination require adjacency information.
   * Returns different information depending on whether it is called on the rendering thread or game thread -
   * On the game thread, it looks at the material settings. Will not change answer during a shader compile
@@ -65,28 +85,21 @@ bool RequiresAdjacencyInformation(UMaterialInterface* Material, const FVertexFac
 		return false;
 	}
 
-	if (IsInParallelGameThread() || IsInGameThread())
+	if (IsInParallelGameThread())
 	{
-		return MaterialSettingsRequireAdjacencyInformation_GameThread(Material, VertexFactoryType, InFeatureLevel);
+		return MaterialRenderingRequiresAdjacencyInformation_Parallel(Material, VertexFactoryType, InFeatureLevel);
 	}
 	else if (IsInRenderingThread() || IsInParallelRenderingThread())
 	{
 		return MaterialRenderingRequiresAdjacencyInformation_RenderingThread(Material, VertexFactoryType, InFeatureLevel);
-	} 
+	}
+	else if (IsInGameThread())
+	{
+		return MaterialSettingsRequireAdjacencyInformation_GameThread(Material, VertexFactoryType, InFeatureLevel);
+	}
 	else
 	{
-		//if we pass null here we have to guarantee that the VF supports tessellation (e.g by using type LocalVF)
-		bool VertexFactorySupportsTessellation = !VertexFactoryType || (VertexFactoryType && VertexFactoryType->SupportsTessellationShaders());
-
-		// Concurrent?
-		if (VertexFactorySupportsTessellation && Material)
-		{
-			const UMaterial* BaseMaterial = Material->GetMaterial_Concurrent();
-			check(BaseMaterial);
-			EMaterialTessellationMode TessellationMode = (EMaterialTessellationMode)BaseMaterial->D3D11TessellationMode;
-			bool bEnableCrackFreeDisplacement = BaseMaterial->bEnableCrackFreeDisplacement;
-			return TessellationMode == MTM_PNTriangles || (TessellationMode == MTM_FlatTessellation && bEnableCrackFreeDisplacement);
-		}
+		return MaterialRenderingRequiresAdjacencyInformation_Parallel(Material, VertexFactoryType, InFeatureLevel);
 	}
 	return false;
 }
