Marked FSkeletalMeshLODRenderData::BuildFromLODModel() as ENGINE_API.

diff --git a/Engine/Source/Runtime/Engine/Public/Rendering/SkeletalMeshLODRenderData.h b/Engine/Source/Runtime/Engine/Public/Rendering/SkeletalMeshLODRenderData.h
--- a/Engine/Source/Runtime/Engine/Public/Rendering/SkeletalMeshLODRenderData.h
+++ b/Engine/Source/Runtime/Engine/Public/Rendering/SkeletalMeshLODRenderData.h
@@ -200,7 +200,7 @@ public:
 	 * Initialize render data (e.g. vertex buffers) from model info
 	 * @param BuildFlags See ESkeletalMeshVertexFlags.
 	 */
-	void BuildFromLODModel(const FSkeletalMeshLODModel* LODModel, uint32 BuildFlags);
+	void ENGINE_API BuildFromLODModel(const FSkeletalMeshLODModel* LODModel, uint32 BuildFlags);
 #endif // WITH_EDITOR
 
 	uint32 GetNumVertices() const
