Fix potential crash due to race condition when accessing Mips while the platform data is still computing asynchronously

diff --git a/Engine/Source/Runtime/Engine/Private/TextureDerivedData.cpp b/Engine/Source/Runtime/Engine/Private/TextureDerivedData.cpp
--- a/Engine/Source/Runtime/Engine/Private/TextureDerivedData.cpp
+++ b/Engine/Source/Runtime/Engine/Private/TextureDerivedData.cpp
@@ -1100,6 +1100,14 @@ FTexturePlatformData::~FTexturePlatformData()
 
 bool FTexturePlatformData::IsReadyForAsyncPostLoad() const
 {
+#if WITH_EDITOR
+	// Can't touch the Mips until async work is finished
+	if (!IsAsyncWorkComplete())
+	{
+		return false;
+	}
+#endif
+
 	for (int32 MipIndex = 0; MipIndex < Mips.Num(); ++MipIndex)
 	{
 		const FTexture2DMipMap& Mip = Mips[MipIndex];
