Fix material not being updated propertly after async texture compilation
Reduce game-thread stutter by improving algorithmic complexity of UpdateStaticMeshesForMaterials

diff --git a/Engine/Source/Runtime/Engine/Classes/Materials/Material.h b/Engine/Source/Runtime/Engine/Classes/Materials/Material.h
--- a/Engine/Source/Runtime/Engine/Classes/Materials/Material.h
+++ b/Engine/Source/Runtime/Engine/Classes/Materials/Material.h
@@ -995,6 +995,7 @@ public:
 #endif
 	ENGINE_API virtual bool GetFontParameterValue(const FHashedMaterialParameterInfo& ParameterInfo,class UFont*& OutFontValue,int32& OutFontPage, bool bOveriddenOnly = false) const override;
 	ENGINE_API virtual bool GetRefractionSettings(float& OutBiasValue) const override;
+	ENGINE_API virtual void GetDependencies(TSet<UMaterialInterface*>& Dependencies) override;
 	ENGINE_API virtual FMaterialRenderProxy* GetRenderProxy() const override;
 	ENGINE_API virtual UPhysicalMaterial* GetPhysicalMaterial() const override;
 	ENGINE_API virtual UPhysicalMaterialMask* GetPhysicalMaterialMask() const override;
diff --git a/Engine/Source/Runtime/Engine/Classes/Materials/MaterialInstance.h b/Engine/Source/Runtime/Engine/Classes/Materials/MaterialInstance.h
--- a/Engine/Source/Runtime/Engine/Classes/Materials/MaterialInstance.h
+++ b/Engine/Source/Runtime/Engine/Classes/Materials/MaterialInstance.h
@@ -516,6 +516,7 @@ public:
 	virtual ENGINE_API bool GetTerrainLayerWeightParameterValue(const FHashedMaterialParameterInfo& ParameterInfo, int32& OutWeightmapIndex, FGuid &OutExpressionGuid) const override;
 			ENGINE_API bool UpdateMaterialLayersParameterValue(const FHashedMaterialParameterInfo& ParameterInfo, const FMaterialLayersFunctions& LayersValue, const bool bOverridden, const FGuid& GUID);
 	virtual ENGINE_API bool IsDependent(UMaterialInterface* TestDependency) override;
+	virtual ENGINE_API void GetDependencies(TSet<UMaterialInterface*>& Dependencies) override;
 	virtual ENGINE_API FMaterialRenderProxy* GetRenderProxy() const override;
 	virtual ENGINE_API UPhysicalMaterial* GetPhysicalMaterial() const override;
 	virtual ENGINE_API UPhysicalMaterialMask* GetPhysicalMaterialMask() const override;
diff --git a/Engine/Source/Runtime/Engine/Classes/Materials/MaterialInterface.h b/Engine/Source/Runtime/Engine/Classes/Materials/MaterialInterface.h
--- a/Engine/Source/Runtime/Engine/Classes/Materials/MaterialInterface.h
+++ b/Engine/Source/Runtime/Engine/Classes/Materials/MaterialInterface.h
@@ -333,6 +333,12 @@ public:
 	*/
 	virtual bool IsDependent(UMaterialInterface* TestDependency) { return TestDependency == this; }
 
+	/**
+	* Get this material dependencies.
+	* @param	Dependencies - List of materials this interface depends on.
+	*/
+	virtual void GetDependencies(TSet<UMaterialInterface*>& Dependencies) PURE_VIRTUAL(UMaterialInterface::GetDependencies, return;);
+
 	/**
 	* Return a pointer to the FMaterialRenderProxy used for rendering.
 	* @param	Selected	specify true to return an alternate material used for rendering this material when part of a selection
diff --git a/Engine/Source/Runtime/Engine/Private/Materials/Material.cpp b/Engine/Source/Runtime/Engine/Private/Materials/Material.cpp
--- a/Engine/Source/Runtime/Engine/Private/Materials/Material.cpp
+++ b/Engine/Source/Runtime/Engine/Private/Materials/Material.cpp
@@ -2825,6 +2825,11 @@ bool UMaterial::GetRefractionSettings(float& OutBiasValue) const
 	return true;
 }
 
+void UMaterial::GetDependencies(TSet<UMaterialInterface*>& Dependencies) 
+{
+	Dependencies.Add(this);
+}
+
 FMaterialRenderProxy* UMaterial::GetRenderProxy() const
 {
 	return DefaultMaterialInstance;
diff --git a/Engine/Source/Runtime/Engine/Private/Materials/MaterialInstance.cpp b/Engine/Source/Runtime/Engine/Private/Materials/MaterialInstance.cpp
--- a/Engine/Source/Runtime/Engine/Private/Materials/MaterialInstance.cpp
+++ b/Engine/Source/Runtime/Engine/Private/Materials/MaterialInstance.cpp
@@ -1629,6 +1629,22 @@ bool UMaterialInstance::CheckMaterialUsage_Concurrent(const EMaterialUsage Usage
 	}
 }
 
+void UMaterialInstance::GetDependencies(TSet<UMaterialInterface*>& Dependencies)
+{
+	if (GetReentrantFlag())
+	{
+		return;
+	}
+
+	Dependencies.Add(this);
+	
+	if (Parent)
+	{
+		FMICReentranceGuard	Guard(this);
+		Parent->GetDependencies(Dependencies);
+	}
+}
+
 bool UMaterialInstance::IsDependent(UMaterialInterface* TestDependency)
 {
 	if(TestDependency == this)
diff --git a/Engine/Source/Runtime/Engine/Private/TextureCompiler.cpp b/Engine/Source/Runtime/Engine/Private/TextureCompiler.cpp
--- a/Engine/Source/Runtime/Engine/Private/TextureCompiler.cpp
+++ b/Engine/Source/Runtime/Engine/Private/TextureCompiler.cpp
@@ -59,29 +59,24 @@ namespace TextureCompilingManagerImpl
 		return StaticEnum<TextureGroup>()->GetMetaData(TEXT("DisplayName"), Texture->LODGroup);
 	}
 
-	static TMultiMap<UTexture*, UMaterial*> GetTexturesAffectingMaterials()
+	static TMultiMap<UTexture*, UMaterialInterface*> GetTexturesAffectingMaterials()
 	{
 		TRACE_CPUPROFILER_EVENT_SCOPE(GetTexturesAffectingMaterials);
 
-		// Update any material that uses this texture and must force a recompile of cache resource
-		TMultiMap<UTexture*, UMaterial*> TexturesRequiringMaterialUpdate;
-		for (TObjectIterator<UMaterial> It; It; ++It)
+		// Update any material that uses this texture
+		TMultiMap<UTexture*, UMaterialInterface*> TexturesRequiringMaterialUpdate;
+
+		TArray<UTexture*> UsedTextures;
+		for (TObjectIterator<UMaterialInterface> It; It; ++It)
 		{
-			UMaterial* Material = *It;
-			for (const UMaterialExpression* MaterialExpression : Material->Expressions)
+			UsedTextures.Reset();
+
+			UMaterialInterface* MaterialInterface = *It;
+			MaterialInterface->GetUsedTextures(UsedTextures, EMaterialQualityLevel::Num, true, GMaxRHIFeatureLevel, true);
+
+			for (UTexture* Texture : UsedTextures)
 			{
-				if (MaterialExpression == nullptr)
-				{
-					continue;
-				}
-	
-				TArray<UTexture*> ForceRecompileTextures;
-				MaterialExpression->GetTexturesForceMaterialRecompile(ForceRecompileTextures);
-	
-				for (UTexture* Texture : ForceRecompileTextures)
-				{
-					TexturesRequiringMaterialUpdate.Emplace(Texture, Material);
-				}
+				TexturesRequiringMaterialUpdate.Emplace(Texture, MaterialInterface);
 			}
 		}
 
@@ -366,6 +361,8 @@ void FTextureCompilingManager::AddTextures(const TArray<UTexture*>& InTextures)
 
 void FTextureCompilingManager::FinishCompilation(const TArray<UTexture*>& InTextures)
 {
+	TRACE_CPUPROFILER_EVENT_SCOPE(FTextureCompilingManager::FinishCompilation);
+
 	using namespace TextureCompilingManagerImpl;
 	check(IsInGameThread());
 
@@ -387,7 +384,7 @@ void FTextureCompilingManager::FinishCompilation(const TArray<UTexture*>& InText
 	if (PendingTextures.Num())
 	{
 		FScopedSlowTask SlowTask((float)PendingTextures.Num(), LOCTEXT("FinishTextureCompilation", "Waiting on texture compilation"), true);
-		SlowTask.MakeDialog(true);
+		SlowTask.MakeDialogDelayed(1.0f);
 
 		struct FTextureTask : public IQueuedWork
 		{
@@ -516,9 +513,9 @@ void FTextureCompilingManager::ProcessTextures(int32 MaximumPriority)
 		{
 			TRACE_CPUPROFILER_EVENT_SCOPE(ProcessTexturesRequiringMaterialUpdate);
 
-			TMultiMap<UTexture*, UMaterial*> TexturesAffectingMaterials = GetTexturesAffectingMaterials();
+			TMultiMap<UTexture*, UMaterialInterface*> TexturesAffectingMaterials = GetTexturesAffectingMaterials();
 
-			TArray<UMaterial*> MaterialsToUpdate;
+			TArray<UMaterialInterface*> MaterialsToUpdate;
 			for (UTexture* Texture : ProcessedTextures)
 			{
 				TexturesAffectingMaterials.MultiFind(Texture, MaterialsToUpdate);
@@ -527,11 +524,11 @@ void FTextureCompilingManager::ProcessTextures(int32 MaximumPriority)
 			if (MaterialsToUpdate.Num())
 			{
 				TRACE_CPUPROFILER_EVENT_SCOPE(UpdateMaterials);
-				FMaterialUpdateContext UpdateContext;
+				FMaterialUpdateContext UpdateContext(FMaterialUpdateContext::EOptions::SyncWithRenderingThread);
 
-				for (UMaterial* MaterialToUpdate : MaterialsToUpdate)
+				for (UMaterialInterface* MaterialToUpdate : MaterialsToUpdate)
 				{
-					UpdateContext.AddMaterial(MaterialToUpdate);
+					UpdateContext.AddMaterialInterface(MaterialToUpdate);
 				}
 			}
 		}
diff --git a/Engine/Source/Runtime/Renderer/Private/RendererScene.cpp b/Engine/Source/Runtime/Renderer/Private/RendererScene.cpp
--- a/Engine/Source/Runtime/Renderer/Private/RendererScene.cpp
+++ b/Engine/Source/Runtime/Renderer/Private/RendererScene.cpp
@@ -4534,6 +4534,7 @@ void FRendererModule::UpdateStaticDrawLists()
 void UpdateStaticMeshesForMaterials(const TArray<const FMaterial*>& MaterialResourcesToUpdate)
 {
 	TArray<UMaterialInterface*> UsedMaterials;
+	TSet<UMaterialInterface*> UsedMaterialsDependencies;
 	TMap<FScene*, TArray<FPrimitiveSceneInfo*>> UsedPrimitives;
 	for (TObjectIterator<UPrimitiveComponent> PrimitiveIt; PrimitiveIt; ++PrimitiveIt)
 	{
@@ -4541,42 +4542,40 @@ void UpdateStaticMeshesForMaterials(const TArray<const FMaterial*>& MaterialReso
 
 		if (PrimitiveComponent->IsRenderStateCreated() && PrimitiveComponent->SceneProxy)
 		{
+			UsedMaterialsDependencies.Reset();
 			UsedMaterials.Reset();
-			bool bPrimitiveIsDependentOnMaterial = false;
 
 			// Note: relying on GetUsedMaterials to be accurate, or else we won't propagate to the right primitives and the renderer will crash later
 			// FPrimitiveSceneProxy::VerifyUsedMaterial is used to make sure that all materials used for rendering are reported in GetUsedMaterials
 			PrimitiveComponent->GetUsedMaterials(UsedMaterials);
 
-			if (UsedMaterials.Num() > 0)
+			for (UMaterialInterface* UsedMaterial : UsedMaterials)
 			{
-				for (TArray<const FMaterial*>::TConstIterator MaterialIt(MaterialResourcesToUpdate); MaterialIt; ++MaterialIt)
+				if (UsedMaterial)
 				{
-					UMaterialInterface* UpdatedMaterialInterface = (*MaterialIt)->GetMaterialInterface();
+					UsedMaterial->GetDependencies(UsedMaterialsDependencies);
+				}
+			}
+
+			if (UsedMaterialsDependencies.Num() > 0)
+			{
+				for (const FMaterial* MaterialResourceToUpdate : MaterialResourcesToUpdate)
+				{
+					UMaterialInterface* UpdatedMaterialInterface = MaterialResourceToUpdate->GetMaterialInterface();
 
 					if (UpdatedMaterialInterface)
 					{
-						for (int32 MaterialIndex = 0; MaterialIndex < UsedMaterials.Num(); MaterialIndex++)
+						if (UsedMaterialsDependencies.Contains(UpdatedMaterialInterface))
 						{
-							UMaterialInterface* TestMaterial = UsedMaterials[MaterialIndex];
-
-							if (TestMaterial && (TestMaterial == UpdatedMaterialInterface || TestMaterial->IsDependent(UpdatedMaterialInterface)))
-							{
-								bPrimitiveIsDependentOnMaterial = true;
-								break;
-							}
+							FPrimitiveSceneProxy* SceneProxy = PrimitiveComponent->SceneProxy;
+							FPrimitiveSceneInfo* SceneInfo = SceneProxy->GetPrimitiveSceneInfo();
+							FScene* Scene = SceneInfo->Scene;
+							TArray<FPrimitiveSceneInfo*>& SceneInfos = UsedPrimitives.FindOrAdd(Scene);
+							SceneInfos.Add(SceneInfo);
+							break;
 						}
 					}
 				}
-
-				if (bPrimitiveIsDependentOnMaterial)
-				{
-					FPrimitiveSceneProxy* SceneProxy = PrimitiveComponent->SceneProxy;
-					FPrimitiveSceneInfo* SceneInfo = SceneProxy->GetPrimitiveSceneInfo();
-					FScene* Scene = SceneInfo->Scene;
-					TArray<FPrimitiveSceneInfo*>& SceneInfos = UsedPrimitives.FindOrAdd(Scene);
-					SceneInfos.Add(SceneInfo);
-				}
 			}
 		}
 	}
