Optimize time spent in game-thread updating materials after async texture compilation

diff --git a/Engine/Source/Runtime/Engine/Private/TextureCompiler.cpp b/Engine/Source/Runtime/Engine/Private/TextureCompiler.cpp
--- a/Engine/Source/Runtime/Engine/Private/TextureCompiler.cpp
+++ b/Engine/Source/Runtime/Engine/Private/TextureCompiler.cpp
@@ -59,12 +59,12 @@ namespace TextureCompilingManagerImpl
 		return StaticEnum<TextureGroup>()->GetMetaData(TEXT("DisplayName"), Texture->LODGroup);
 	}
 
-	static TMultiMap<UTexture*, UMaterialInterface*> GetTexturesAffectingMaterials()
+	static TMultiMap<UObject*, UMaterialInterface*> GetTexturesAffectingMaterials()
 	{
 		TRACE_CPUPROFILER_EVENT_SCOPE(GetTexturesAffectingMaterials);
 
 		// Update any material that uses this texture
-		TMultiMap<UTexture*, UMaterialInterface*> TexturesRequiringMaterialUpdate;
+		TMultiMap<UObject*, UMaterialInterface*> TexturesRequiringMaterialUpdate;
 
 		TArray<UTexture*> UsedTextures;
 		for (TObjectIterator<UMaterialInterface> It; It; ++It)
@@ -72,9 +72,7 @@ namespace TextureCompilingManagerImpl
 			UsedTextures.Reset();
 
 			UMaterialInterface* MaterialInterface = *It;
-			MaterialInterface->GetUsedTextures(UsedTextures, EMaterialQualityLevel::Num, true, GMaxRHIFeatureLevel, true);
-
-			for (UTexture* Texture : UsedTextures)
+			for (UObject* Texture : MaterialInterface->GetReferencedTextures())
 			{
 				TexturesRequiringMaterialUpdate.Emplace(Texture, MaterialInterface);
 			}
@@ -513,7 +511,7 @@ void FTextureCompilingManager::ProcessTextures(int32 MaximumPriority)
 		{
 			TRACE_CPUPROFILER_EVENT_SCOPE(ProcessTexturesRequiringMaterialUpdate);
 
-			TMultiMap<UTexture*, UMaterialInterface*> TexturesAffectingMaterials = GetTexturesAffectingMaterials();
+			TMultiMap<UObject*, UMaterialInterface*> TexturesAffectingMaterials = GetTexturesAffectingMaterials();
 
 			TArray<UMaterialInterface*> MaterialsToUpdate;
 			for (UTexture* Texture : ProcessedTextures)
