SS - Fix Is InGT perf tracker

diff --git a/Engine/Source/Runtime/Engine/Public/InGamePerformanceTracker.h b/Engine/Source/Runtime/Engine/Public/InGamePerformanceTracker.h
--- a/Engine/Source/Runtime/Engine/Public/InGamePerformanceTracker.h
+++ b/Engine/Source/Runtime/Engine/Public/InGamePerformanceTracker.h
@@ -130,7 +130,12 @@ public:
 	{
 		if (Tracker && FInGamePerformanceTracker::CachedEnabled)
 		{
-			if (IsInGameThread())
+			if (IsInParallelGameThread())
+			{
+				//On other threads we don't allow (or check for) re-entrance and store our own cycle data.
+				BeginCycles = FPlatformTime::Cycles();
+			}
+			else if (IsInGameThread())
 			{
 				Tracker->EnterTimedSection();//On the GT we use the tracker directly so we can allow re-entrance.
 			}
@@ -145,7 +150,14 @@ public:
 	{
 		if (Tracker && FInGamePerformanceTracker::CachedEnabled)
 		{
-			if (IsInGameThread())
+			if (IsInParallelGameThread())
+			{
+				//On other threads we don't allow (or check for) re-entrance and store our own cycle data.
+				uint32 Cycles = FPlatformTime::Cycles() - BeginCycles;
+				Tracker->AddCycles(Cycles);
+				BeginCycles = 0;
+			}
+			else if (IsInGameThread())
 			{
 				Tracker->ExitTimedSection();//On the GT we use the tracker directly so we can allow re-entrance.
 			}
