SS - Change the order of the test for parallel rendering in Tessellation rendering because the ParallelRenderingThread has been erroneously identified as the RenderingThread during startup in the original implementation.

diff --git a/Engine/Source/Runtime/Engine/Private/TessellationRendering.cpp b/Engine/Source/Runtime/Engine/Private/TessellationRendering.cpp
--- a/Engine/Source/Runtime/Engine/Private/TessellationRendering.cpp
+++ b/Engine/Source/Runtime/Engine/Private/TessellationRendering.cpp
@@ -29,7 +29,7 @@ bool MaterialSettingsRequireAdjacencyInformation_GameThread(UMaterialInterface*
   * Rendering thread version that looks at the current shader that will be used. **Will change answer during a shader compile** */
 bool MaterialRenderingRequiresAdjacencyInformation_RenderingThread(UMaterialInterface* Material, const FVertexFactoryType* VertexFactoryType, const FStaticFeatureLevel InFeatureLevel)
 {
-	check(IsInRenderingThread() || IsInParallelRenderingThread());
+	check(IsInParallelRenderingThread() || IsInRenderingThread());
 
 	//if we pass null here we have to guarantee that the VF supports tessellation (e.g by using type LocalVF)
 	bool VertexFactorySupportsTessellation = !VertexFactoryType || (VertexFactoryType && VertexFactoryType->SupportsTessellationShaders());
@@ -89,7 +89,7 @@ bool RequiresAdjacencyInformation(UMaterialInterface* Material, const FVertexFac
 	{
 		return MaterialRenderingRequiresAdjacencyInformation_Parallel(Material, VertexFactoryType, InFeatureLevel);
 	}
-	else if (IsInRenderingThread() || IsInParallelRenderingThread())
+	else if (IsInParallelRenderingThread() || IsInRenderingThread())
 	{
 		return MaterialRenderingRequiresAdjacencyInformation_RenderingThread(Material, VertexFactoryType, InFeatureLevel);
 	}
