Avoid flushrenderingcommands for even better game-thread performance when updating materials after async texture compilation

diff --git a/Engine/Source/Runtime/Engine/Private/TextureCompiler.cpp b/Engine/Source/Runtime/Engine/Private/TextureCompiler.cpp
--- a/Engine/Source/Runtime/Engine/Private/TextureCompiler.cpp
+++ b/Engine/Source/Runtime/Engine/Private/TextureCompiler.cpp
@@ -59,7 +59,7 @@ namespace TextureCompilingManagerImpl
 		return StaticEnum<TextureGroup>()->GetMetaData(TEXT("DisplayName"), Texture->LODGroup);
 	}
 
-	static TMultiMap<UObject*, UMaterialInterface*> GetTexturesAffectingMaterials()
+	static TMultiMap<UObject*, UMaterialInterface*> GetTexturesAffectingMaterialInterfaces()
 	{
 		TRACE_CPUPROFILER_EVENT_SCOPE(GetTexturesAffectingMaterials);
 
@@ -509,9 +509,9 @@ void FTextureCompilingManager::ProcessTextures(int32 MaximumPriority)
 
 		if (ProcessedTextures.Num())
 		{
-			TRACE_CPUPROFILER_EVENT_SCOPE(ProcessTexturesRequiringMaterialUpdate);
+			TRACE_CPUPROFILER_EVENT_SCOPE(RecacheUniformExpressions);
 
-			TMultiMap<UObject*, UMaterialInterface*> TexturesAffectingMaterials = GetTexturesAffectingMaterials();
+			TMultiMap<UObject*, UMaterialInterface*> TexturesAffectingMaterials = GetTexturesAffectingMaterialInterfaces();
 
 			TArray<UMaterialInterface*> MaterialsToUpdate;
 			for (UTexture* Texture : ProcessedTextures)
@@ -522,11 +522,10 @@ void FTextureCompilingManager::ProcessTextures(int32 MaximumPriority)
 			if (MaterialsToUpdate.Num())
 			{
 				TRACE_CPUPROFILER_EVENT_SCOPE(UpdateMaterials);
-				FMaterialUpdateContext UpdateContext(FMaterialUpdateContext::EOptions::SyncWithRenderingThread);
-
+				
 				for (UMaterialInterface* MaterialToUpdate : MaterialsToUpdate)
 				{
-					UpdateContext.AddMaterialInterface(MaterialToUpdate);
+					MaterialToUpdate->RecacheUniformExpressions(false);
 				}
 			}
 		}
