Fix warning about async texture placeholder being of the wrong type

diff --git a/Engine/Source/Runtime/Engine/Private/Materials/MaterialShared.cpp b/Engine/Source/Runtime/Engine/Private/Materials/MaterialShared.cpp
--- a/Engine/Source/Runtime/Engine/Private/Materials/MaterialShared.cpp
+++ b/Engine/Source/Runtime/Engine/Private/Materials/MaterialShared.cpp
@@ -2386,8 +2386,14 @@ IAllocatedVirtualTexture* FMaterialRenderProxy::AllocateVTStack(const FMaterialR
 		}
 		if (!Texture->IsCurrentlyVirtualTextured())
 		{
-			UE_LOG(LogMaterial, Warning, TEXT("Material '%s' expects texture '%s' to be Virtual"),
-				*GetFriendlyName(), *Texture->GetName());
+			// The placeholder used during async texture compilation is expected to be of the wrong type since
+			// no VT infos are available until later in the compilation process. This will be resolved
+			// once the final texture resource is available.
+			if (!Texture->IsDefaultTexture())
+			{
+				UE_LOG(LogMaterial, Warning, TEXT("Material '%s' expects texture '%s' to be Virtual"),
+					*GetFriendlyName(), *Texture->GetName());
+			}
 			continue;
 		}
 
