Fix game-thread stutter when compressing texture from worker threads

diff --git a/Engine/Source/Developer/TextureCompressor/Private/TextureCompressorModule.cpp b/Engine/Source/Developer/TextureCompressor/Private/TextureCompressorModule.cpp
--- a/Engine/Source/Developer/TextureCompressor/Private/TextureCompressorModule.cpp
+++ b/Engine/Source/Developer/TextureCompressor/Private/TextureCompressorModule.cpp
@@ -1676,8 +1676,9 @@ void ITextureCompressorModule::AdjustImageColors(FImage& Image, const FTextureBu
 			}
 		};
 
-		EParallelForFlags PriorityFlags = IsInGameThread() ? EParallelForFlags::None : EParallelForFlags::BackgroundPriority;
-		ParallelFor(NumJobs, AdjustImageColorsFunc, bForceSingleThread ? EParallelForFlags::ForceSingleThread : PriorityFlags);
+		const EParallelForFlags PriorityFlag = IsInGameThread() ? EParallelForFlags::None : EParallelForFlags::BackgroundPriority;
+		const EParallelForFlags SingleThreadedFlag = bForceSingleThread ? EParallelForFlags::ForceSingleThread : EParallelForFlags::None;
+		ParallelFor(NumJobs, AdjustImageColorsFunc, SingleThreadedFlag | PriorityFlag);
 	}
 }
 
diff --git a/Engine/Source/Runtime/Engine/Private/VT/VirtualTextureDataBuilder.cpp b/Engine/Source/Runtime/Engine/Private/VT/VirtualTextureDataBuilder.cpp
--- a/Engine/Source/Runtime/Engine/Private/VT/VirtualTextureDataBuilder.cpp
+++ b/Engine/Source/Runtime/Engine/Private/VT/VirtualTextureDataBuilder.cpp
@@ -613,6 +613,9 @@ void FVirtualTextureDataBuilder::BuildTiles(const TArray<FVTSourceTileEntry>& Ti
  
 		GeneratedData.TilePayload.AddDefaulted(TileList.Num());
 
+		const EParallelForFlags ThreadingFlag = bAllowAsync ? EParallelForFlags::None : EParallelForFlags::ForceSingleThread;
+		const EParallelForFlags PriorityFlag = IsInGameThread() ? EParallelForFlags::None : EParallelForFlags::BackgroundPriority;
+
 		ParallelFor(TileList.Num(), [&](int32 TileIndex)
 		{
 			const FVTSourceTileEntry& Tile = TileList[TileIndex];
@@ -677,7 +680,7 @@ void FVirtualTextureDataBuilder::BuildTiles(const TArray<FVTSourceTileEntry>& Ti
 			{
 				GeneratedData.TilePayload[TileIndex] = MoveTemp(CompressedMip[0].RawData);
 			}
-		}, !bAllowAsync); // ParallelFor
+		}, ThreadingFlag | PriorityFlag); // ParallelFor
 
 		if (BuildSettingsLayer0.bVirtualTextureEnableCompressZlib)
 		{
