Add support for adding tasks from multiple threads like distance fields (CL 8571312)
Fix threadpool lifetime management like distance field (CL 10733543)

diff --git a/Engine/Source/Developer/MeshUtilities/Private/MeshCardRepresentationUtilities.cpp b/Engine/Source/Developer/MeshUtilities/Private/MeshCardRepresentationUtilities.cpp
--- a/Engine/Source/Developer/MeshUtilities/Private/MeshCardRepresentationUtilities.cpp
+++ b/Engine/Source/Developer/MeshUtilities/Private/MeshCardRepresentationUtilities.cpp
@@ -522,6 +522,8 @@ bool FMeshUtilities::GenerateCardRepresentationData(
 	FCardRepresentationData& OutData)
 {
 #if USE_EMBREE
+	TRACE_CPUPROFILER_EVENT_SCOPE(FMeshUtilities::GenerateCardRepresentationData);
+
 	RTCDevice EmbreeDevice = rtcNewDevice(nullptr);
 
 	RTCError ReturnErrorNewDevice = rtcDeviceGetError(EmbreeDevice);
diff --git a/Engine/Source/Runtime/Engine/Private/MeshCardRepresentation.cpp b/Engine/Source/Runtime/Engine/Private/MeshCardRepresentation.cpp
--- a/Engine/Source/Runtime/Engine/Private/MeshCardRepresentation.cpp
+++ b/Engine/Source/Runtime/Engine/Private/MeshCardRepresentation.cpp
@@ -153,8 +153,15 @@ public:
 	{
 		check(!bIsRunning);
 
+		// Calling Reset will call Kill which in turn will call Stop and set bForceFinish to true.
+		Thread.Reset();
+
+		// Now we can set bForceFinish to false without being overwritten by the old thread shutting down.
 		bForceFinish = false;
-		Thread.Reset(FRunnableThread::Create(this, *FString::Printf(TEXT("BuildCardRepresentationThread%u"), NextThreadIndex), 0, TPri_Normal, FPlatformAffinity::GetPoolThreadMask()));
+		Thread.Reset(FRunnableThread::Create(this, *FString::Printf(TEXT("BuildCardRepresentationThread%u"), NextThreadIndex), 0, TPri_SlightlyBelowNormal, FPlatformAffinity::GetPoolThreadMask()));
+
+		// Set this now before exiting so that IsRunning() returns true without having to wait on the thread to be completely started.
+		bIsRunning = true;
 		NextThreadIndex++;
 	}
 
@@ -190,25 +197,44 @@ FQueuedThreadPool* CreateCardWorkerThreadPool()
 uint32 FBuildCardRepresentationThreadRunnable::Run()
 {
 	bool bHasWork = true;
+	
+	// Do not exit right away if no work to do as it often leads to stop and go problems
+	// when tasks are being queued at a slower rate than the processor capability to process them.
+	const uint64 ExitAfterIdleCycle = static_cast<uint64>(10.0 / FPlatformTime::GetSecondsPerCycle64()); // 10s
 
-	while (!bForceFinish && bHasWork)
+	uint64 LastWorkCycle = FPlatformTime::Cycles64();
+	while (!bForceFinish && (bHasWork || (FPlatformTime::Cycles64() - LastWorkCycle) < ExitAfterIdleCycle))
 	{
 		// LIFO build order, since meshes actually visible in a map are typically loaded last
 		FAsyncCardRepresentationTask* Task = AsyncQueue.TaskQueue.Pop();
 
+		FQueuedThreadPool* ThreadPool = nullptr;
+
+#if WITH_EDITOR
+		ThreadPool = GLargeThreadPool;
+#endif
+
 		if (Task)
 		{
-			if (!WorkerThreadPool)
+			if (!ThreadPool)
 			{
-				WorkerThreadPool.Reset(CreateCardWorkerThreadPool());
+				if (!WorkerThreadPool)
+				{
+					WorkerThreadPool.Reset(CreateCardWorkerThreadPool());
+				}
+
+				ThreadPool = WorkerThreadPool.Get();
 			}
 
-			AsyncQueue.Build(Task, *WorkerThreadPool);
+			AsyncQueue.Build(Task, *ThreadPool);
+			LastWorkCycle = FPlatformTime::Cycles64();
+
 			bHasWork = true;
 		}
 		else
 		{
 			bHasWork = false;
+			FPlatformProcess::Sleep(.01f);
 		}
 	}
 
@@ -245,12 +271,18 @@ void FCardRepresentationAsyncQueue::AddTask(FAsyncCardRepresentationTask* Task)
 		MeshUtilities = &FModuleManager::Get().LoadModuleChecked<IMeshUtilities>(TEXT("MeshUtilities"));
 	}
 	
-	ReferencedTasks.Add(Task);
+	{
+		// Array protection when called from multiple threads
+		FScopeLock Lock(&CriticalSection);
+		ReferencedTasks.Add(Task);
+	}
 
 	if (GUseAsyncCardRepresentationBuildQueue)
 	{
 		TaskQueue.Push(Task);
 
+		// Logic protection when called from multiple threads
+		FScopeLock Lock(&CriticalSection);
 		if (!ThreadRunnable->IsRunning())
 		{
 			ThreadRunnable->Launch();
@@ -283,10 +315,13 @@ void FCardRepresentationAsyncQueue::BlockUntilBuildComplete(UStaticMesh* StaticM
 
 		bReferenced = false;
 
-		for (int TaskIndex = 0; TaskIndex < ReferencedTasks.Num(); TaskIndex++)
 		{
-			bReferenced = bReferenced || ReferencedTasks[TaskIndex]->StaticMesh == StaticMesh;
-			bReferenced = bReferenced || ReferencedTasks[TaskIndex]->GenerateSource == StaticMesh;
+			FScopeLock Lock(&CriticalSection);
+			for (int TaskIndex = 0; TaskIndex < ReferencedTasks.Num(); TaskIndex++)
+			{
+				bReferenced = bReferenced || ReferencedTasks[TaskIndex]->StaticMesh == StaticMesh;
+				bReferenced = bReferenced || ReferencedTasks[TaskIndex]->GenerateSource == StaticMesh;
+			}
 		}
 
 		if (bReferenced)
@@ -317,6 +352,7 @@ void FCardRepresentationAsyncQueue::BlockUntilBuildComplete(UStaticMesh* StaticM
 
 void FCardRepresentationAsyncQueue::BlockUntilAllBuildsComplete()
 {
+	TRACE_CPUPROFILER_EVENT_SCOPE(FCardRepresentationAsyncQueue::BlockUntilAllBuildsComplete)
 	do 
 	{
 		ProcessAsyncTasks();
@@ -328,7 +364,7 @@ void FCardRepresentationAsyncQueue::BlockUntilAllBuildsComplete()
 void FCardRepresentationAsyncQueue::Build(FAsyncCardRepresentationTask* Task, FQueuedThreadPool& ThreadPool)
 {
 #if WITH_EDITOR
-
+	
 	// Editor 'force delete' can null any UObject pointers which are seen by reference collecting (eg UProperty or serialized)
 	if (Task->StaticMesh && Task->GenerateSource)
 	{
@@ -350,6 +386,7 @@ void FCardRepresentationAsyncQueue::Build(FAsyncCardRepresentationTask* Task, FQ
 
 void FCardRepresentationAsyncQueue::AddReferencedObjects(FReferenceCollector& Collector)
 {	
+	FScopeLock Lock(&CriticalSection);
 	for (int TaskIndex = 0; TaskIndex < ReferencedTasks.Num(); TaskIndex++)
 	{
 		// Make sure none of the UObjects referenced by the async tasks are GC'ed during the task
@@ -361,6 +398,8 @@ void FCardRepresentationAsyncQueue::AddReferencedObjects(FReferenceCollector& Co
 void FCardRepresentationAsyncQueue::ProcessAsyncTasks()
 {
 #if WITH_EDITOR
+	TRACE_CPUPROFILER_EVENT_SCOPE(FCardRepresentationAsyncQueue::ProcessAsyncTasks)
+
 	TArray<FAsyncCardRepresentationTask*> LocalCompletedTasks;
 	CompletedTasks.PopAll(LocalCompletedTasks);
 
@@ -370,7 +409,10 @@ void FCardRepresentationAsyncQueue::ProcessAsyncTasks()
 		COOK_STAT(auto Timer = CardRepresentationCookStats::UsageStats.TimeSyncWork());
 		FAsyncCardRepresentationTask* Task = LocalCompletedTasks[TaskIndex];
 
-		ReferencedTasks.Remove(Task);
+		{
+			FScopeLock Lock(&CriticalSection);
+			ReferencedTasks.Remove(Task);
+		}
 
 		// Editor 'force delete' can null any UObject pointers which are seen by reference collecting (eg UProperty or serialized)
 		if (Task->StaticMesh && Task->bSuccess)
@@ -401,7 +443,13 @@ void FCardRepresentationAsyncQueue::ProcessAsyncTasks()
 		delete Task;
 	}
 
-	if (ReferencedTasks.Num() > 0 && !ThreadRunnable->IsRunning())
+	bool bRemainingTasks = false;
+	{
+		FScopeLock Lock(&CriticalSection);
+		bRemainingTasks = ReferencedTasks.Num() > 0;
+	}
+
+	if (bRemainingTasks && !ThreadRunnable->IsRunning())
 	{
 		ThreadRunnable->Launch();
 	}
diff --git a/Engine/Source/Runtime/Engine/Public/MeshCardRepresentation.h b/Engine/Source/Runtime/Engine/Public/MeshCardRepresentation.h
--- a/Engine/Source/Runtime/Engine/Public/MeshCardRepresentation.h
+++ b/Engine/Source/Runtime/Engine/Public/MeshCardRepresentation.h
@@ -195,6 +195,7 @@ public:
 
 	int32 GetNumOutstandingTasks() const
 	{
+		FScopeLock Lock(&CriticalSection);
 		return ReferencedTasks.Num();
 	}
 
@@ -217,6 +218,8 @@ private:
 
 	class IMeshUtilities* MeshUtilities;
 
+	mutable FCriticalSection CriticalSection;
+
 	friend class FBuildCardRepresentationThreadRunnable;
 };
 
