integrating 13707877 to fix Chaos "players falling through the floor"

diff --git a/Engine/Source/Runtime/Engine/Private/PhysicsEngine/Experimental/PhysScene_Chaos.cpp b/Engine/Source/Runtime/Engine/Private/PhysicsEngine/Experimental/PhysScene_Chaos.cpp
--- a/Engine/Source/Runtime/Engine/Private/PhysicsEngine/Experimental/PhysScene_Chaos.cpp
+++ b/Engine/Source/Runtime/Engine/Private/PhysicsEngine/Experimental/PhysScene_Chaos.cpp
@@ -1709,7 +1709,8 @@ void FPhysScene_ChaosInterface::ProcessDeferredCreatePhysicsState()
 	// TODO explore parallelization of other physics initialization, not trivial and likely to break stuff.
 	for (UPrimitiveComponent* PrimitiveComponent : DeferredCreatePhysicsStateComponents)
 	{
-		if (PrimitiveComponent->GetOwner() && PrimitiveComponent->GetOwner()->IsPendingKill() == false && PrimitiveComponent->ShouldCreatePhysicsState() && PrimitiveComponent->IsPhysicsStateCreated() == false)
+		const bool bPendingKill = PrimitiveComponent->GetOwner() && PrimitiveComponent->GetOwner()->IsPendingKill();
+		if (!bPendingKill && PrimitiveComponent->ShouldCreatePhysicsState() && PrimitiveComponent->IsPhysicsStateCreated() == false)
 		{
 			PrimitiveComponent->OnCreatePhysicsState();
 			PrimitiveComponent->GlobalCreatePhysicsDelegate.Broadcast(PrimitiveComponent);
