Fix potential reentrancy in createresource caused by the async texture compiler which might cause leaks and crash in ExitRHI when running the editor in debug

diff --git a/Engine/Source/Runtime/Engine/Private/TextureCompiler.cpp b/Engine/Source/Runtime/Engine/Private/TextureCompiler.cpp
--- a/Engine/Source/Runtime/Engine/Private/TextureCompiler.cpp
+++ b/Engine/Source/Runtime/Engine/Private/TextureCompiler.cpp
@@ -331,6 +331,17 @@ void FTextureCompilingManager::AddTextures(const TArray<UTexture*>& InTextures)
 	TRACE_CPUPROFILER_EVENT_SCOPE(FTextureCompilingManager::AddTextures)
 	check(IsInGameThread());
 
+	// We might not get ticked very often during load time
+	// so this will allow us to refresh compiled textures
+	// of the highest priority to improve the UI experience.
+	ProcessTextures(1 /* Maximum Priority */);
+
+	// Register new textures after ProcessTextures to avoid
+	// potential reentrant calls to CreateResource on the
+	// textures being added. This would cause multiple
+	// TextureResource to be created and assigned to the same Owner
+	// which would obviously be bad and causing leaks including
+	// in the RHI.
 	for (UTexture* Texture : InTextures)
 	{
 		int32 TexturePriority = 2;
@@ -350,11 +361,6 @@ void FTextureCompilingManager::AddTextures(const TArray<UTexture*>& InTextures)
 		}
 		RegisteredTextureBuckets[TexturePriority].Emplace(Texture);
 	}
-
-	// We might not get ticked very often during load time
-	// so this will allow us to refresh compiled textures
-	// of the highest priority to improve the UI experience.
-	ProcessTextures(1 /* Maximum Priority */);
 }
 
 void FTextureCompilingManager::FinishCompilation(const TArray<UTexture*>& InTextures)
