Schedule async textures on GThreadPool instead of GLargeThreadPool to reduce memory usage and improve editor responsiveness
Add texture compilation shutdown with cancellation when exiting the editor with textures still being processed
Improve game-thread responsiveness during forced async texture flush

diff --git a/Engine/Source/Runtime/Engine/Classes/Engine/Texture.h b/Engine/Source/Runtime/Engine/Classes/Engine/Texture.h
--- a/Engine/Source/Runtime/Engine/Classes/Engine/Texture.h
+++ b/Engine/Source/Runtime/Engine/Classes/Engine/Texture.h
@@ -1160,6 +1160,8 @@ public:
 protected:
 
 #if WITH_EDITOR
+	// The Texture compiler might use TryCancelCachePlatformData on shutdown
+	friend class FTextureCompilingManager;
 
 	/** Try to cancel any async tasks on PlatformData. 
 	 *  Returns true if there is no more async tasks pending, false otherwise.
@@ -1169,7 +1171,7 @@ protected:
 	/** Notify any loaded material instances that the texture has changed. */
 	ENGINE_API void NotifyMaterials();
 
-#endif //WITH_EDOTIR
+#endif //WITH_EDITOR
 };
 
 /** 
diff --git a/Engine/Source/Runtime/Engine/Private/TextureCompiler.cpp b/Engine/Source/Runtime/Engine/Private/TextureCompiler.cpp
--- a/Engine/Source/Runtime/Engine/Private/TextureCompiler.cpp
+++ b/Engine/Source/Runtime/Engine/Private/TextureCompiler.cpp
@@ -168,6 +168,35 @@ EQueuedWorkPriority FTextureCompilingManager::GetBasePriority(UTexture* InTextur
 	return TextureCompilingManagerImpl::GetBasePriority(InTexture);
 }
 
+void FTextureCompilingManager::Shutdown()
+{
+	bHasShutdown = true;
+	if (GetNumRemainingTextures())
+	{
+		TArray<UTexture*> PendingTextures;
+		PendingTextures.Reserve(GetNumRemainingTextures());
+
+		for (TSet<TWeakObjectPtr<UTexture>>& Bucket : RegisteredTextureBuckets)
+		{
+			for (TWeakObjectPtr<UTexture>& WeakTexture : Bucket)
+			{
+				if (WeakTexture.IsValid())
+				{
+					UTexture* Texture = WeakTexture.Get();
+					
+					if (!Texture->TryCancelCachePlatformData())
+					{
+						PendingTextures.Add(Texture);
+					}
+				}
+			}
+		}
+
+		// Wait on textures already in progress we couldn't cancel
+		FinishCompilation(PendingTextures);
+	}
+}
+
 FQueuedThreadPool* FTextureCompilingManager::GetThreadPool() const
 {
 	static FQueuedThreadPoolWrapper* GTextureThreadPool = nullptr;
@@ -175,10 +204,13 @@ FQueuedThreadPool* FTextureCompilingManager::GetThreadPool() const
 	{
 		TextureCompilingManagerImpl::EnsureInitializedCVars();
 
-		// Wrapping GLargeThreadPool to give TextureThreadPool it's own set of priorities and allow Pausable functionality
-		// All texture priorities will resolve to a Low priority once being scheduled in the LargeThreadPool.
+		// Wrapping GThreadPool to give TextureThreadPool it's own set of priorities and allow Pausable functionality
+		// We're using GThreadPool instead of GLargeThreadPool because asset compilation is hard on total memory and memory bandwidth and can run slower when going wider than actual cores.
+		// All texture priorities will resolve to either Low or Lowest priority once being scheduled in the ThreadPool.
+		// Any asset supporting being built async should be scheduled lower than Normal on the LargeThreadPool to let non-async stuff go first
+		const auto TexturePriorityMapper = [](EQueuedWorkPriority TexturePriority) { return FMath::Max(TexturePriority, EQueuedWorkPriority::Low); };
 		const int32 MaxConcurrency = CVarAsyncTextureCompilationMaxConcurrency.GetValueOnAnyThread();
-		GTextureThreadPool = new FQueuedThreadPoolWrapper(GLargeThreadPool, MaxConcurrency, [](EQueuedWorkPriority) { return EQueuedWorkPriority::Low; });
+		GTextureThreadPool = new FQueuedThreadPoolWrapper(GThreadPool, MaxConcurrency, TexturePriorityMapper);
 
 		CVarAsyncTextureCompilation->SetOnChangedCallback(
 			FConsoleVariableDelegate::CreateLambda(
@@ -228,6 +260,11 @@ FQueuedThreadPool* FTextureCompilingManager::GetThreadPool() const
 
 bool FTextureCompilingManager::IsAsyncTextureCompilationEnabled() const
 {
+	if (bHasShutdown)
+	{
+		return false;
+	}
+
 	TextureCompilingManagerImpl::EnsureInitializedCVars();
 
 	return CVarAsyncTextureCompilation.GetValueOnAnyThread() != 0;
@@ -408,7 +445,7 @@ void FTextureCompilingManager::FinishCompilation(const TArray<UTexture*>& InText
 		for (UTexture* Texture : PendingTextures)
 		{
 			PendingTasks[PendingTaskIndex].Texture.Reset(Texture);
-			GLargeThreadPool->AddQueuedWork(&PendingTasks[PendingTaskIndex], EQueuedWorkPriority::High);
+			GThreadPool->AddQueuedWork(&PendingTasks[PendingTaskIndex], EQueuedWorkPriority::High);
 			PendingTaskIndex++;
 		}
 
@@ -425,6 +462,12 @@ void FTextureCompilingManager::FinishCompilation(const TArray<UTexture*>& InText
 			// Be nice with the game thread and tick the progress at 60 fps even when no progress is being made...
 			while (!PendingTask.Event->Wait(16))
 			{
+				// If the task pool is saturated, we want to be able to contribute so we can make some progress quickly...
+				if (GThreadPool->RetractQueuedWork(&PendingTask))
+				{
+					PendingTask.DoThreadedWork();
+				}
+
 				UpdateProgress(0.0f, TextureIndex, InTextures.Num(), TextureName);
 			}
 			UpdateProgress(1.f, TextureIndex++, InTextures.Num(), TextureName);
diff --git a/Engine/Source/Runtime/Engine/Public/TextureCompiler.h b/Engine/Source/Runtime/Engine/Public/TextureCompiler.h
--- a/Engine/Source/Runtime/Engine/Public/TextureCompiler.h
+++ b/Engine/Source/Runtime/Engine/Public/TextureCompiler.h
@@ -53,6 +53,11 @@ public:
 	 */
 	ENGINE_API FQueuedThreadPool* GetThreadPool() const;
 
+	/**
+	 * Cancel any pending work and blocks until it is safe to shut down.
+	 */
+	ENGINE_API void Shutdown();
+
 protected:
 	virtual void Tick(float DeltaTime) override;
 	virtual bool IsTickableWhenPaused() const override { return true; }
@@ -60,6 +65,7 @@ protected:
 	virtual TStatId GetStatId() const override { RETURN_QUICK_DECLARE_CYCLE_STAT(FTextureCompilingManager, STATGROUP_Tickables); }
 
 private:
+	bool bHasShutdown = false;
 	TArray<TSet<TWeakObjectPtr<UTexture>>> RegisteredTextureBuckets;
 
 	void ProcessTextures(int32 MaximumPriority = -1);
diff --git a/Engine/Source/Runtime/Launch/Private/LaunchEngineLoop.cpp b/Engine/Source/Runtime/Launch/Private/LaunchEngineLoop.cpp
--- a/Engine/Source/Runtime/Launch/Private/LaunchEngineLoop.cpp
+++ b/Engine/Source/Runtime/Launch/Private/LaunchEngineLoop.cpp
@@ -89,6 +89,7 @@
 	#include "Interfaces/IEditorStyleModule.h"
 	#include "PIEPreviewDeviceProfileSelectorModule.h"
 	#include "Misc/QueuedThreadPoolWrapper.h"
+	#include "TextureCompiler.h"
 
 	#if PLATFORM_WINDOWS
 		#include "Windows/AllowWindowsPlatformTypes.h"
@@ -4147,6 +4148,10 @@ void FEngineLoop::Exit()
 		}
 	}
 
+#if WITH_EDITOR
+	FTextureCompilingManager::Get().Shutdown();
+#endif
+
 	// Block till all outstanding resource streaming requests are fulfilled.
 	if (!IStreamingManager::HasShutdown())
 	{
