Moved USelection to UnrealEd

#rnx

diff --git a/Engine/Source/Developer/GameplayDebugger/Private/GameplayDebuggerLocalController.cpp b/Engine/Source/Developer/GameplayDebugger/Private/GameplayDebuggerLocalController.cpp
--- a/Engine/Source/Developer/GameplayDebugger/Private/GameplayDebuggerLocalController.cpp
+++ b/Engine/Source/Developer/GameplayDebugger/Private/GameplayDebuggerLocalController.cpp
@@ -700,6 +700,7 @@ FString UGameplayDebuggerLocalController::GetKeyDescriptionLong(const FKey& KeyB
 	return (KeyDisplay == KeyName) ? FString::Printf(TEXT("[%s]"), *KeyDisplay) : FString::Printf(TEXT("%s [%s key])"), *KeyDisplay, *KeyName);
 }
 
+#if WITH_EDITOR
 void UGameplayDebuggerLocalController::OnSelectionChanged(UObject* Object)
 {
 	USelection* Selection = Cast<USelection>(Object);
@@ -733,6 +734,7 @@ void UGameplayDebuggerLocalController::OnSelectedObject(UObject* Object)
 		CachedReplicator->CollectCategoryData(/*bForce=*/true);
 	}
 }
+#endif
 
 void UGameplayDebuggerLocalController::OnCategoriesChanged()
 {
diff --git a/Engine/Source/Developer/GameplayDebugger/Public/GameplayDebuggerLocalController.h b/Engine/Source/Developer/GameplayDebugger/Public/GameplayDebuggerLocalController.h
--- a/Engine/Source/Developer/GameplayDebugger/Public/GameplayDebuggerLocalController.h
+++ b/Engine/Source/Developer/GameplayDebugger/Public/GameplayDebuggerLocalController.h
@@ -115,9 +115,11 @@ protected:
 	/** draw header for category */
 	void DrawCategoryHeader(int32 CategoryId, TSharedRef<FGameplayDebuggerCategory> Category, FGameplayDebuggerCanvasContext& CanvasContext);
 
+#if WITH_EDITOR
 	/** event for simulate in editor mode */
 	void OnSelectionChanged(UObject* Object);
 	void OnSelectedObject(UObject* Object);
+#endif
 
 	FString GetKeyDescriptionShort(const FKey& KeyBind) const;
 	FString GetKeyDescriptionLong(const FKey& KeyBind) const;
diff --git a/Engine/Source/Editor/UnrealEd/Private/Selection.cpp b/Engine/Source/Editor/UnrealEd/Private/Selection.cpp
new file mode 100644
--- /dev/null
+++ b/Engine/Source/Editor/UnrealEd/Private/Selection.cpp
@@ -0,0 +1,244 @@
+// Copyright Epic Games, Inc. All Rights Reserved.
+
+#include "Selection.h"
+#include "UObject/Package.h"
+#include "UObject/UObjectAnnotation.h"
+
+DEFINE_LOG_CATEGORY_STATIC(LogSelection, Log, All);
+
+USelection::FOnSelectionChanged	USelection::SelectionChangedEvent;
+USelection::FOnSelectionChanged	USelection::SelectObjectEvent;
+FSimpleMulticastDelegate		USelection::SelectNoneEvent;
+
+void USelection::Initialize(FUObjectAnnotationSparseBool* InSelectionAnnotation)
+{
+	if (InSelectionAnnotation)
+	{
+		SelectionAnnotation = InSelectionAnnotation;
+		bOwnsSelectionAnnotation = false;
+	}
+	else
+	{
+		SelectionAnnotation = new FUObjectAnnotationSparseBool;
+		bOwnsSelectionAnnotation = true;
+	}
+}
+
+void USelection::Select(UObject* InObject)
+{
+	check( InObject );
+
+	const bool bSelectionChanged = !SelectionAnnotation->Get(InObject);
+	SelectionAnnotation->Set(InObject);
+
+	if(bSelectionChanged)
+	{
+		// Add to selected lists.
+		SelectedObjects.Add(InObject);
+		MarkBatchDirty();
+	}
+
+	FSelectedClassInfo* SelectedClassInfo = SelectedClasses.Find(InObject->GetClass());
+	if( SelectedClassInfo )
+	{
+		++SelectedClassInfo->SelectionCount;
+	}
+	else
+	{
+		// 1 Object with a new class type has been selected
+		FSelectedClassInfo NewInfo(InObject->GetClass(), 1);
+		SelectedClasses.Add(NewInfo);
+	}
+
+	if( !IsBatchSelecting() )
+	{
+		// Call this after the item has been added from the selection set.
+		USelection::SelectObjectEvent.Broadcast( InObject );
+	}
+}
+
+
+void USelection::Deselect(UObject* InObject)
+{
+	check( InObject );
+
+	const bool bSelectionChanged = InObject->IsSelected();
+	SelectionAnnotation->Clear(InObject);
+
+	// Remove from selected list.
+	SelectedObjects.RemoveSingle( InObject );
+
+	if( !IsBatchSelecting() )
+	{
+		// Call this after the item has been removed from the selection set.
+		USelection::SelectObjectEvent.Broadcast( InObject );
+	}
+
+	FSetElementId Id = SelectedClasses.FindId(InObject->GetClass());
+	if(Id.IsValidId())
+	{
+		FSelectedClassInfo& ClassInfo = SelectedClasses[Id];
+		// One less object of this class is selected;
+		--ClassInfo.SelectionCount;
+		// If no more objects of the selected class exists, remove it
+		if(ClassInfo.SelectionCount == 0)
+		{
+			SelectedClasses.Remove(Id);
+		}
+	}
+
+
+	if ( bSelectionChanged )
+	{
+		MarkBatchDirty();
+	}
+}
+
+
+void USelection::Select(UObject* InObject, bool bSelect)
+{
+	if( bSelect )
+	{
+		Select( InObject );
+	}
+	else
+	{
+		Deselect( InObject );
+	}
+}
+
+
+void USelection::ToggleSelect( UObject* InObject )
+{
+	check( InObject );
+
+	Select( InObject, !InObject->IsSelected() );
+}
+
+void USelection::DeselectAll( UClass* InClass )
+{
+	// Fast path for deselecting all UObjects with any flags
+	if ( InClass == UObject::StaticClass() )
+	{
+		InClass = nullptr;
+	}
+
+	bool bSelectionChanged = false;
+
+	TSet<FSelectedClassInfo> RemovedClasses;
+	// Remove from the end to minimize memmoves.
+	for ( int32 i = SelectedObjects.Num()-1 ; i >= 0 ; --i )
+	{
+		UObject* Object = GetSelectedObject(i);
+
+		if ( !Object )
+		{		
+			// Remove NULLs from SelectedObjects array.
+			SelectedObjects.RemoveAt( i );
+		}
+		else if( !InClass || Object->IsA( InClass ) )
+		{
+			// if the object is of type InClass then all objects of that same type will be removed
+			RemovedClasses.Add(FSelectedClassInfo(Object->GetClass()));
+
+			SelectionAnnotation->Clear(Object);
+			SelectedObjects.RemoveAt( i );
+
+			// Call this after the item has been removed from the selection set.
+			USelection::SelectObjectEvent.Broadcast( Object );
+
+			bSelectionChanged = true;
+		}
+	}
+
+	if( InClass == nullptr )
+	{
+		SelectedClasses.Empty();
+	}
+	else
+	{
+		// Remove the passed in class and all child classes that were removed
+		// from the list of currently selected classes
+		RemovedClasses.Add(InClass);
+		SelectedClasses = SelectedClasses.Difference(RemovedClasses);
+	}
+
+	if ( bSelectionChanged )
+	{
+		MarkBatchDirty();
+		if ( !IsBatchSelecting() )
+		{
+			USelection::SelectionChangedEvent.Broadcast(this);
+		}
+	}
+}
+
+void USelection::MarkBatchDirty()
+{
+	if ( IsBatchSelecting() )
+	{
+		bIsBatchDirty = true;
+	}
+}
+
+
+
+bool USelection::IsSelected(const UObject* InObject) const
+{
+	if ( InObject )
+	{
+		return SelectedObjects.Contains(InObject);
+	}
+
+	return false;
+}
+
+
+void USelection::Serialize(FArchive& Ar)
+{
+	Super::Serialize( Ar );
+	Ar << SelectedObjects;
+
+	if(Ar.IsLoading())
+	{
+		// The set of selected objects may have changed, so make sure our annotations exactly match the list, otherwise
+		// UObject::IsSelected() could return a result that was different from the list of objects returned by GetSelectedObjects()
+		// This needs to happen in serialize because other code may check the selection state in PostEditUndo and the order of PostEditUndo is indeterminate.
+		SelectionAnnotation->ClearAll();
+
+		for(TWeakObjectPtr<UObject>& ObjectPtr : SelectedObjects)
+		{
+			if (UObject* Object = ObjectPtr.Get(true))
+			{
+				SelectionAnnotation->Set(Object);
+			}
+		}
+	}
+}
+
+bool USelection::Modify(bool bAlwaysMarkDirty/* =true */)
+{
+	// If the selection currently contains any PIE objects we should not be including it in the transaction buffer
+	for (TWeakObjectPtr<UObject>& ObjectPtr : SelectedObjects)
+	{
+		UObject* Object = ObjectPtr.Get();
+		if (Object && Object->GetOutermost()->HasAnyPackageFlags(PKG_PlayInEditor | PKG_ContainsScript | PKG_CompiledIn))
+		{
+			return false;
+		}
+	}
+
+	return Super::Modify(bAlwaysMarkDirty);
+}
+
+void USelection::BeginDestroy()
+{
+	Super::BeginDestroy();
+
+	if (bOwnsSelectionAnnotation)
+	{
+		delete SelectionAnnotation;
+		SelectionAnnotation = nullptr;
+	}
+}
+
diff --git a/Engine/Source/Editor/UnrealEd/Public/Selection.h b/Engine/Source/Editor/UnrealEd/Public/Selection.h
new file mode 100644
--- /dev/null
+++ b/Engine/Source/Editor/UnrealEd/Public/Selection.h
@@ -0,0 +1,523 @@
+// Copyright Epic Games, Inc. All Rights Reserved.
+
+#pragma once
+
+#include "CoreMinimal.h"
+#include "UObject/ObjectMacros.h"
+#include "UObject/UObjectGlobals.h"
+#include "UObject/Object.h"
+#include "Components/ActorComponent.h"
+#include "Selection.generated.h"
+
+/**
+ * Manages selections of objects.  Used in the editor for selecting
+ * objects in the various browser windows.
+ */
+UCLASS(transient)
+class UNREALED_API USelection : public UObject
+{
+	GENERATED_BODY()
+
+private:
+	/** Contains info about each class and how many objects of that class are selected */
+	struct FSelectedClassInfo
+	{
+		/** The selected class */
+		const UClass* Class;
+		/** How many objects of that class are selected */
+		int32 SelectionCount;
+
+		FSelectedClassInfo(const UClass* InClass)
+			: Class(InClass)
+			, SelectionCount(0)
+		{}
+
+		FSelectedClassInfo(const UClass* InClass, int32 InSelectionCount)
+			: Class(InClass)
+			, SelectionCount(InSelectionCount)
+		{}
+		bool operator==(const FSelectedClassInfo& Info) const
+		{
+			return Class == Info.Class;
+		}
+
+		friend uint32 GetTypeHash(const FSelectedClassInfo& Info)
+		{
+			return GetTypeHash(Info.Class);
+		}
+	};
+
+	typedef TArray<TWeakObjectPtr<UObject> >	ObjectArray;
+	typedef TSet<FSelectedClassInfo>	ClassArray;
+
+	template<typename SelectionFilter>
+	friend class TSelectionIterator;
+
+public:
+	/** Params: UObject* NewSelection */
+	DECLARE_MULTICAST_DELEGATE_OneParam(FOnSelectionChanged, UObject*);
+
+	/** Called when selection in editor has changed */
+	static FOnSelectionChanged SelectionChangedEvent;
+	/** Called when an object has been selected (generally an actor) */
+	static FOnSelectionChanged SelectObjectEvent;
+	/** Called to deselect everything */
+	static FSimpleMulticastDelegate SelectNoneEvent;
+
+	/** Initializes the selection set with an annotation used to quickly look up selection state */
+	void Initialize(FUObjectAnnotationSparseBool* InSelectionAnnotation);
+
+	typedef ClassArray::TIterator TClassIterator;
+	typedef ClassArray::TConstIterator TClassConstIterator;
+
+	TClassIterator			ClassItor()				{ return TClassIterator( SelectedClasses ); }
+	TClassConstIterator		ClassConstItor() const	{ return TClassConstIterator( SelectedClasses ); }
+
+	/**
+	 * Returns the number of objects in the selection set.  This function is used by clients in
+	 * conjunction with op::() to iterate over selected objects.  Note that some of these objects
+	 * may be NULL, and so clients should use CountSelections() to get the true number of
+	 * non-NULL selected objects.
+	 * 
+	 * @return		Number of objects in the selection set.
+	 */
+	int32 Num() const
+	{
+		return SelectedObjects.Num();
+	}
+
+	/**
+	 * @return	The Index'th selected objects.  May be NULL.
+	 */
+	UObject* GetSelectedObject(const int32 InIndex)
+	{
+		return (SelectedObjects.IsValidIndex(InIndex) && SelectedObjects[InIndex].IsValid() ? SelectedObjects[InIndex].Get() : nullptr);
+	}
+
+	/**
+	 * @return	The Index'th selected objects.  May be NULL.
+	 */
+	const UObject* GetSelectedObject(const int32 InIndex) const
+	{
+		return (SelectedObjects.IsValidIndex(InIndex) && SelectedObjects[InIndex].IsValid() ? SelectedObjects[InIndex].Get() : nullptr);
+	}
+
+	/**
+	 * Call before beginning selection operations
+	 */
+	void BeginBatchSelectOperation()
+	{
+		SelectionMutex++;
+	}
+
+	/**
+	 * Should be called when selection operations are complete.  If all selection operations are complete, notifies all listeners
+	 * that the selection has been changed.
+	 */
+	void EndBatchSelectOperation(bool bNotify = true)
+	{
+		if ( --SelectionMutex == 0 )
+		{
+			const bool bSelectionChanged = bIsBatchDirty;
+			bIsBatchDirty = false;
+
+			if ( bSelectionChanged && bNotify )
+			{
+				// new version - includes which selection set was modified
+				USelection::SelectionChangedEvent.Broadcast(this);
+			}
+		}
+	}
+
+	/**
+	 * @return	Returns whether or not the selection object is currently in the middle of a batch select block.
+	 */
+	bool IsBatchSelecting() const
+	{
+		return SelectionMutex != 0;
+	}
+
+	/**
+	 * Selects the specified object.
+	 *
+	 * @param	InObject	The object to select/deselect.  Must be non-NULL.
+	 */
+	void Select(UObject* InObject);
+
+	/**
+	 * Deselects the specified object.
+	 *
+	 * @param	InObject	The object to deselect.  Must be non-NULL.
+	 */
+	void Deselect(UObject* InObject);
+
+	/**
+	 * Selects or deselects the specified object, depending on the value of the bSelect flag.
+	 *
+	 * @param	InObject	The object to select/deselect.  Must be non-NULL.
+	 * @param	bSelect		true selects the object, false deselects.
+	 */
+	void Select(UObject* InObject, bool bSelect);
+
+	/**
+	 * Toggles the selection state of the specified object.
+	 *
+	 * @param	InObject	The object to select/deselect.  Must be non-NULL.
+	 */
+	void ToggleSelect(UObject* InObject);
+
+	/**
+	 * Deselects all objects of the specified class, if no class is specified it deselects all objects.
+	 *
+	 * @param	InClass		The type of object to deselect.  Can be NULL.
+	 */
+	void DeselectAll( UClass* InClass = NULL );
+
+	/**
+	 * If batch selection is active, sets flag indicating something actually changed.
+	 */
+	void MarkBatchDirty();
+
+	/**
+	 * Returns the first selected object of the specified class.
+	 *
+	 * @param	InClass				The class of object to return.  Must be non-NULL.
+	 * @param	RequiredInterface	[opt] Interface this class must implement to be returned.  May be NULL.
+	 * @param	bArchetypesOnly		[opt] true to only return archetype objects, false otherwise
+	 * @return						The first selected object of the specified class.
+	 */
+	UObject* GetTop(UClass* InClass, UClass* RequiredInterface=nullptr, bool bArchetypesOnly=false)
+	{
+		check( InClass );
+		for( int32 i=0; i<SelectedObjects.Num(); ++i )
+		{
+			UObject* SelectedObject = GetSelectedObject(i);
+			if (SelectedObject)
+			{
+				// maybe filter out non-archetypes
+				if ( bArchetypesOnly && !SelectedObject->HasAnyFlags(RF_ArchetypeObject) )
+				{
+					continue;
+				}
+
+				if ( InClass->HasAnyClassFlags(CLASS_Interface) )
+				{
+					//~ Begin InClass is an Interface, and we want the top object that implements it
+					if ( SelectedObject->GetClass()->ImplementsInterface(InClass) )
+					{
+						return SelectedObject;
+					}
+				}
+				else if ( SelectedObject->IsA(InClass) )
+				{
+					//~ Begin InClass is a class, so we want the top object of that class that implements the required Interface, if specified
+					if ( !RequiredInterface || SelectedObject->GetClass()->ImplementsInterface(RequiredInterface) )
+					{
+						return SelectedObject;
+					}
+				}
+			}
+		}
+		return nullptr;
+	}
+
+	/**
+	* Returns the last selected object of the specified class.
+	*
+	* @param	InClass		The class of object to return.  Must be non-NULL.
+	* @return				The last selected object of the specified class.
+	*/
+	UObject* GetBottom(UClass* InClass)
+	{
+		check( InClass );
+		for( int32 i = SelectedObjects.Num()-1 ; i > -1 ; --i )
+		{
+			UObject* SelectedObject = GetSelectedObject(i);
+			if( SelectedObject && SelectedObject->IsA(InClass) )
+			{
+				return SelectedObject;
+			}
+		}
+		return nullptr;
+	}
+
+	/**
+	 * Returns the first selected object.
+	 *
+	 * @return				The first selected object.
+	 */
+	template< class T > T* GetTop()
+	{
+		UObject* Selected = GetTop(T::StaticClass());
+		return Selected ? CastChecked<T>(Selected) : nullptr;
+	}
+
+	/**
+	* Returns the last selected object.
+	*
+	* @return				The last selected object.
+	*/
+	template< class T > T* GetBottom()
+	{
+		UObject* Selected = GetBottom(T::StaticClass());
+		return Selected ? CastChecked<T>(Selected) : nullptr;
+	}
+
+	/**
+	 * Returns true if the specified object is non-NULL and selected.
+	 *
+	 * @param	InObject	The object to query.  Can be NULL.
+	 * @return				true if the object is selected, or false if InObject is unselected or NULL.
+	 */
+	bool IsSelected(const UObject* InObject) const;
+
+	/**
+	 * Returns the number of selected objects of the specified type.
+	 *
+	 * @param	bIgnorePendingKill	specify true to count only those objects which are not pending kill (marked for garbage collection)
+	 * @return						The number of objects of the specified type.
+	 */
+	template< class T >
+	int32 CountSelections( bool bIgnorePendingKill=false )
+	{
+		int32 Count = 0;
+		for( int32 i=0; i<SelectedObjects.Num(); ++i )
+		{
+			UObject* SelectedObject = GetSelectedObject(i);
+			if( SelectedObject && SelectedObject->IsA(T::StaticClass()) && !(bIgnorePendingKill && SelectedObject->IsPendingKill()) )
+			{
+				++Count;
+			}
+		}
+		return Count;
+	}
+
+	/**
+	 * Untemplated version of CountSelections.
+	 */
+	int32 CountSelections(UClass *ClassToCount, bool bIgnorePendingKill=false)
+	{
+		int32 Count = 0;
+		for( int32 i=0; i<SelectedObjects.Num(); ++i )
+		{
+			UObject* SelectedObject = GetSelectedObject(i);
+			if( SelectedObject && SelectedObject->IsA(ClassToCount) && !(bIgnorePendingKill && SelectedObject->IsPendingKill()) )
+			{
+				++Count;
+			}
+		}
+		return Count;
+	}
+
+	bool IsClassSelected(UClass* Class) const
+	{
+		const FSelectedClassInfo* Info = SelectedClasses.Find(Class);
+		return Info && Info->SelectionCount > 0;
+	}
+
+	//~ Begin UObject Interface
+	virtual void Serialize(FArchive& Ar) override;
+	virtual bool Modify( bool bAlwaysMarkDirty=true) override;
+	virtual void BeginDestroy() override;
+
+	//~ End UObject Interface
+
+
+	/**
+	 * Fills in the specified array with all selected objects of the desired type.
+	 * 
+	 * @param	OutSelectedObjects		[out] Array to fill with selected objects of type T
+	 * @return							The number of selected objects of the specified type.
+	 */
+	template< class T > 
+	int32 GetSelectedObjects(TArray<T*> &OutSelectedObjects)
+	{
+		OutSelectedObjects.Empty();
+		for ( int32 Idx = 0 ; Idx < SelectedObjects.Num() ; ++Idx )
+		{
+			UObject* SelectedObject = GetSelectedObject(Idx);
+			if ( SelectedObject && SelectedObject->IsA(T::StaticClass()) )
+			{
+				OutSelectedObjects.Add( (T*)SelectedObject );
+			}
+		}
+		return OutSelectedObjects.Num();
+	}
+	int32 GetSelectedObjects( ObjectArray& out_SelectedObjects )
+	{
+		out_SelectedObjects = SelectedObjects;
+		return SelectedObjects.Num();
+	}
+
+	int32 GetSelectedObjects(UClass *FilterClass, TArray<UObject*> &OutSelectedObjects)
+	{
+		OutSelectedObjects.Empty();
+		for ( int32 Idx = 0 ; Idx < SelectedObjects.Num() ; ++Idx )
+		{
+			UObject* SelectedObject = GetSelectedObject(Idx);
+			if ( SelectedObject && SelectedObject->IsA(FilterClass) )
+			{
+				OutSelectedObjects.Add( SelectedObject );
+			}
+		}
+		return OutSelectedObjects.Num();
+	}
+
+protected:
+	/** List of selected objects, ordered as they were selected. */
+	ObjectArray	SelectedObjects;
+
+	/** Tracks the most recently selected actor classes.  Used for UnrealEd menus. */
+	ClassArray	SelectedClasses;
+
+	/** Tracks the number of active selection operations.  Allows batched selection operations to only send one notification at the end of the batch */
+	int32			SelectionMutex = 0;
+
+	/** Tracks whether the selection set changed during a batch selection operation */
+	bool		bIsBatchDirty = false;
+	
+	/** Selection annotation for fast lookup */
+	FUObjectAnnotationSparseBool* SelectionAnnotation = nullptr;
+	bool bOwnsSelectionAnnotation = false;
+
+private:
+	// Hide IsSelected(), as calling IsSelected() on a selection set almost always indicates
+	// an error where the caller should use IsSelected(UObject* InObject).
+	bool IsSelected() const
+	{
+		return UObject::IsSelected();
+	}
+};
+
+
+/** A filter for generic selection sets.  Simply allows objects which are non-null */
+class FGenericSelectionFilter
+{
+public:
+	bool IsObjectValid( const UObject* InObject ) const
+	{
+		return InObject != nullptr;
+	}
+};
+
+/**
+ * Manages selections of objects.  Used in the editor for selecting
+ * objects in the various browser windows.
+ */
+template<typename SelectionFilter>
+class TSelectionIterator
+{
+public:
+	TSelectionIterator(USelection& InSelection)
+		: Selection( InSelection )
+		, Filter( SelectionFilter() )
+	{
+		Reset();
+	}
+
+	/** Advances iterator to the next valid element in the container. */
+	void operator++()
+	{
+		while ( true )
+		{
+			++Index;
+
+			// Halt if the end of the selection set has been reached.
+			if ( !IsIndexValid() )
+			{
+				return;
+			}
+
+			// Halt if at a valid object.
+			if ( IsObjectValid() )
+			{
+				return;
+			}
+		}
+	}
+
+	/** Element access. */
+	UObject* operator*() const
+	{
+		return GetCurrentObject();
+	}
+
+	/** Element access. */
+	UObject* operator->() const
+	{
+		return GetCurrentObject();
+	}
+
+	/** Returns true if the iterator has not yet reached the end of the selection set. */
+	explicit operator bool() const
+	{
+		return IsIndexValid();
+	}
+
+	/** Resets the iterator to the beginning of the selection set. */
+	void Reset()
+	{
+		Index = -1;
+		++( *this );
+	}
+
+	/** Returns an index to the current element. */
+	int32 GetIndex() const
+	{
+		return Index;
+	}
+
+private:
+	UObject* GetCurrentObject() const
+	{
+		return Selection.GetSelectedObject(Index);
+	}
+
+	bool IsObjectValid() const
+	{
+		return Filter.IsObjectValid( GetCurrentObject() );
+	}
+
+	bool IsIndexValid() const
+	{
+		return Selection.SelectedObjects.IsValidIndex( Index );
+	}
+
+	USelection&	Selection;
+	SelectionFilter Filter;
+	int32			Index;
+};
+
+
+class FSelectionIterator : public TSelectionIterator<FGenericSelectionFilter>
+{
+public:
+	FSelectionIterator(USelection& InSelection)
+		: TSelectionIterator<FGenericSelectionFilter>( InSelection )
+	{}
+};
+
+/** A filter for only iterating through editable components */
+class FSelectedEditableComponentFilter
+{
+public:
+	bool IsObjectValid(const UObject* Object) const
+	{
+		if (const UActorComponent* Comp = Cast<UActorComponent>( Object ))
+		{
+			return Comp->IsEditableWhenInherited();
+		}
+		return false;
+	}
+};
+
+/**
+ * An iterator used to iterate through selected components that are editable (i.e. not created in a blueprint)
+ */
+class FSelectedEditableComponentIterator : public TSelectionIterator<FSelectedEditableComponentFilter>
+{
+public:
+	FSelectedEditableComponentIterator(USelection& InSelection)
+		: TSelectionIterator<FSelectedEditableComponentFilter>(InSelection)
+	{}
+};
diff --git a/Engine/Source/Runtime/Engine/Classes/Engine/Selection.h b/Engine/Source/Runtime/Engine/Classes/Engine/Selection.h
--- a/Engine/Source/Runtime/Engine/Classes/Engine/Selection.h
+++ b/Engine/Source/Runtime/Engine/Classes/Engine/Selection.h
@@ -1,532 +1,8 @@
 // Copyright Epic Games, Inc. All Rights Reserved.
 
-/**
- * Manages selections of objects.  Used in the editor for selecting
- * objects in the various browser windows.
- */
-
 #pragma once
 
-#include "CoreMinimal.h"
-#include "UObject/ObjectMacros.h"
-#include "UObject/UObjectGlobals.h"
-#include "UObject/Object.h"
-#include "Components/ActorComponent.h"
-
-#include "Selection.generated.h"
-
-UCLASS(customConstructor, transient)
-class ENGINE_API USelection : public UObject
-{
-	GENERATED_UCLASS_BODY()
-private:
-	/** Contains info about each class and how many objects of that class are selected */
-	struct FSelectedClassInfo
-	{
-		/** The selected class */
-		const UClass* Class;
-		/** How many objects of that class are selected */
-		int32 SelectionCount;
-
-		FSelectedClassInfo(const UClass* InClass)
-			: Class(InClass)
-			, SelectionCount(0)
-		{}
-
-		FSelectedClassInfo(const UClass* InClass, int32 InSelectionCount)
-			: Class(InClass)
-			, SelectionCount(InSelectionCount)
-		{}
-		bool operator==(const FSelectedClassInfo& Info) const
-		{
-			return Class == Info.Class;
-		}
-
-		friend uint32 GetTypeHash(const FSelectedClassInfo& Info)
-		{
-			return GetTypeHash(Info.Class);
-		}
-	};
-
-	typedef TArray<TWeakObjectPtr<UObject> >	ObjectArray;
-	typedef TSet<FSelectedClassInfo>	ClassArray;
-
-	template<typename SelectionFilter>
-	friend class TSelectionIterator;
-
-public:
-	/** Params: UObject* NewSelection */
-	DECLARE_MULTICAST_DELEGATE_OneParam(FOnSelectionChanged, UObject*);
-
-	/** Called when selection in editor has changed */
-	static FOnSelectionChanged SelectionChangedEvent;
-	/** Called when an object has been selected (generally an actor) */
-	static FOnSelectionChanged SelectObjectEvent;
-	/** Called to deselect everything */
-	static FSimpleMulticastDelegate SelectNoneEvent;
-
-	USelection(const FObjectInitializer& ObjectInitializer = FObjectInitializer::Get());
-
-	/** Initializes the selection set with an annotation used to quickly look up selection state */
-	void Initialize(FUObjectAnnotationSparseBool* InSelectionAnnotation);
-
-	typedef ClassArray::TIterator TClassIterator;
-	typedef ClassArray::TConstIterator TClassConstIterator;
-
-	TClassIterator			ClassItor()				{ return TClassIterator( SelectedClasses ); }
-	TClassConstIterator		ClassConstItor() const	{ return TClassConstIterator( SelectedClasses ); }
-
-	/**
-	 * Returns the number of objects in the selection set.  This function is used by clients in
-	 * conjunction with op::() to iterate over selected objects.  Note that some of these objects
-	 * may be NULL, and so clients should use CountSelections() to get the true number of
-	 * non-NULL selected objects.
-	 * 
-	 * @return		Number of objects in the selection set.
-	 */
-	int32 Num() const
-	{
-		return SelectedObjects.Num();
-	}
-
-	/**
-	 * @return	The Index'th selected objects.  May be NULL.
-	 */
-	UObject* GetSelectedObject(const int32 InIndex)
-	{
-		return (SelectedObjects.IsValidIndex(InIndex) && SelectedObjects[InIndex].IsValid() ? SelectedObjects[InIndex].Get() : nullptr);
-	}
-
-	/**
-	 * @return	The Index'th selected objects.  May be NULL.
-	 */
-	const UObject* GetSelectedObject(const int32 InIndex) const
-	{
-		return (SelectedObjects.IsValidIndex(InIndex) && SelectedObjects[InIndex].IsValid() ? SelectedObjects[InIndex].Get() : nullptr);
-	}
-
-	/**
-	 * Call before beginning selection operations
-	 */
-	void BeginBatchSelectOperation()
-	{
-		SelectionMutex++;
-	}
-
-	/**
-	 * Should be called when selection operations are complete.  If all selection operations are complete, notifies all listeners
-	 * that the selection has been changed.
-	 */
-	void EndBatchSelectOperation(bool bNotify = true)
-	{
-		if ( --SelectionMutex == 0 )
-		{
-			const bool bSelectionChanged = bIsBatchDirty;
-			bIsBatchDirty = false;
-
-			if ( bSelectionChanged && bNotify )
-			{
-				// new version - includes which selection set was modified
-				USelection::SelectionChangedEvent.Broadcast(this);
-			}
-		}
-	}
-
-	/**
-	 * @return	Returns whether or not the selection object is currently in the middle of a batch select block.
-	 */
-	bool IsBatchSelecting() const
-	{
-		return SelectionMutex != 0;
-	}
-
-	/**
-	 * Selects the specified object.
-	 *
-	 * @param	InObject	The object to select/deselect.  Must be non-NULL.
-	 */
-	void Select(UObject* InObject);
-
-	/**
-	 * Deselects the specified object.
-	 *
-	 * @param	InObject	The object to deselect.  Must be non-NULL.
-	 */
-	void Deselect(UObject* InObject);
-
-	/**
-	 * Selects or deselects the specified object, depending on the value of the bSelect flag.
-	 *
-	 * @param	InObject	The object to select/deselect.  Must be non-NULL.
-	 * @param	bSelect		true selects the object, false deselects.
-	 */
-	void Select(UObject* InObject, bool bSelect);
-
-	/**
-	 * Toggles the selection state of the specified object.
-	 *
-	 * @param	InObject	The object to select/deselect.  Must be non-NULL.
-	 */
-	void ToggleSelect(UObject* InObject);
-
-	/**
-	 * Deselects all objects of the specified class, if no class is specified it deselects all objects.
-	 *
-	 * @param	InClass		The type of object to deselect.  Can be NULL.
-	 */
-	void DeselectAll( UClass* InClass = NULL );
-
-	/**
-	 * If batch selection is active, sets flag indicating something actually changed.
-	 */
-	void MarkBatchDirty();
-
-	/**
-	 * Returns the first selected object of the specified class.
-	 *
-	 * @param	InClass				The class of object to return.  Must be non-NULL.
-	 * @param	RequiredInterface	[opt] Interface this class must implement to be returned.  May be NULL.
-	 * @param	bArchetypesOnly		[opt] true to only return archetype objects, false otherwise
-	 * @return						The first selected object of the specified class.
-	 */
-	UObject* GetTop(UClass* InClass, UClass* RequiredInterface=nullptr, bool bArchetypesOnly=false)
-	{
-		check( InClass );
-		for( int32 i=0; i<SelectedObjects.Num(); ++i )
-		{
-			UObject* SelectedObject = GetSelectedObject(i);
-			if (SelectedObject)
-			{
-				// maybe filter out non-archetypes
-				if ( bArchetypesOnly && !SelectedObject->HasAnyFlags(RF_ArchetypeObject) )
-				{
-					continue;
-				}
-
-				if ( InClass->HasAnyClassFlags(CLASS_Interface) )
-				{
-					//~ Begin InClass is an Interface, and we want the top object that implements it
-					if ( SelectedObject->GetClass()->ImplementsInterface(InClass) )
-					{
-						return SelectedObject;
-					}
-				}
-				else if ( SelectedObject->IsA(InClass) )
-				{
-					//~ Begin InClass is a class, so we want the top object of that class that implements the required Interface, if specified
-					if ( !RequiredInterface || SelectedObject->GetClass()->ImplementsInterface(RequiredInterface) )
-					{
-						return SelectedObject;
-					}
-				}
-			}
-		}
-		return nullptr;
-	}
-
-	/**
-	* Returns the last selected object of the specified class.
-	*
-	* @param	InClass		The class of object to return.  Must be non-NULL.
-	* @return				The last selected object of the specified class.
-	*/
-	UObject* GetBottom(UClass* InClass)
-	{
-		check( InClass );
-		for( int32 i = SelectedObjects.Num()-1 ; i > -1 ; --i )
-		{
-			UObject* SelectedObject = GetSelectedObject(i);
-			if( SelectedObject && SelectedObject->IsA(InClass) )
-			{
-				return SelectedObject;
-			}
-		}
-		return nullptr;
-	}
-
-	/**
-	 * Returns the first selected object.
-	 *
-	 * @return				The first selected object.
-	 */
-	template< class T > T* GetTop()
-	{
-		UObject* Selected = GetTop(T::StaticClass());
-		return Selected ? CastChecked<T>(Selected) : nullptr;
-	}
-
-	/**
-	* Returns the last selected object.
-	*
-	* @return				The last selected object.
-	*/
-	template< class T > T* GetBottom()
-	{
-		UObject* Selected = GetBottom(T::StaticClass());
-		return Selected ? CastChecked<T>(Selected) : nullptr;
-	}
-
-	/**
-	 * Returns true if the specified object is non-NULL and selected.
-	 *
-	 * @param	InObject	The object to query.  Can be NULL.
-	 * @return				true if the object is selected, or false if InObject is unselected or NULL.
-	 */
-	bool IsSelected(const UObject* InObject) const;
-
-	/**
-	 * Returns the number of selected objects of the specified type.
-	 *
-	 * @param	bIgnorePendingKill	specify true to count only those objects which are not pending kill (marked for garbage collection)
-	 * @return						The number of objects of the specified type.
-	 */
-	template< class T >
-	int32 CountSelections( bool bIgnorePendingKill=false )
-	{
-		int32 Count = 0;
-		for( int32 i=0; i<SelectedObjects.Num(); ++i )
-		{
-			UObject* SelectedObject = GetSelectedObject(i);
-			if( SelectedObject && SelectedObject->IsA(T::StaticClass()) && !(bIgnorePendingKill && SelectedObject->IsPendingKill()) )
-			{
-				++Count;
-			}
-		}
-		return Count;
-	}
-
-	/**
-	 * Untemplated version of CountSelections.
-	 */
-	int32 CountSelections(UClass *ClassToCount, bool bIgnorePendingKill=false)
-	{
-		int32 Count = 0;
-		for( int32 i=0; i<SelectedObjects.Num(); ++i )
-		{
-			UObject* SelectedObject = GetSelectedObject(i);
-			if( SelectedObject && SelectedObject->IsA(ClassToCount) && !(bIgnorePendingKill && SelectedObject->IsPendingKill()) )
-			{
-				++Count;
-			}
-		}
-		return Count;
-	}
-
-	bool IsClassSelected(UClass* Class) const
-	{
-		const FSelectedClassInfo* Info = SelectedClasses.Find(Class);
-		return Info && Info->SelectionCount > 0;
-	}
-
-	//~ Begin UObject Interface
-	virtual void Serialize(FArchive& Ar) override;
 #if WITH_EDITOR
-	virtual bool Modify( bool bAlwaysMarkDirty=true) override;
+// TODO: Deprecate this adapter include in favor of using the UnrealEd version directly
+#include "UnrealEd/Public/Selection.h"
 #endif
-	virtual void BeginDestroy() override;
-
-	//~ End UObject Interface
-
-
-	/**
-	 * Fills in the specified array with all selected objects of the desired type.
-	 * 
-	 * @param	OutSelectedObjects		[out] Array to fill with selected objects of type T
-	 * @return							The number of selected objects of the specified type.
-	 */
-	template< class T > 
-	int32 GetSelectedObjects(TArray<T*> &OutSelectedObjects)
-	{
-		OutSelectedObjects.Empty();
-		for ( int32 Idx = 0 ; Idx < SelectedObjects.Num() ; ++Idx )
-		{
-			UObject* SelectedObject = GetSelectedObject(Idx);
-			if ( SelectedObject && SelectedObject->IsA(T::StaticClass()) )
-			{
-				OutSelectedObjects.Add( (T*)SelectedObject );
-			}
-		}
-		return OutSelectedObjects.Num();
-	}
-	int32 GetSelectedObjects( ObjectArray& out_SelectedObjects )
-	{
-		out_SelectedObjects = SelectedObjects;
-		return SelectedObjects.Num();
-	}
-
-	int32 GetSelectedObjects(UClass *FilterClass, TArray<UObject*> &OutSelectedObjects)
-	{
-		OutSelectedObjects.Empty();
-		for ( int32 Idx = 0 ; Idx < SelectedObjects.Num() ; ++Idx )
-		{
-			UObject* SelectedObject = GetSelectedObject(Idx);
-			if ( SelectedObject && SelectedObject->IsA(FilterClass) )
-			{
-				OutSelectedObjects.Add( SelectedObject );
-			}
-		}
-		return OutSelectedObjects.Num();
-	}
-
-protected:
-	/** List of selected objects, ordered as they were selected. */
-	ObjectArray	SelectedObjects;
-
-	/** Tracks the most recently selected actor classes.  Used for UnrealEd menus. */
-	ClassArray	SelectedClasses;
-
-	/** Tracks the number of active selection operations.  Allows batched selection operations to only send one notification at the end of the batch */
-	int32			SelectionMutex;
-
-	/** Tracks whether the selection set changed during a batch selection operation */
-	bool		bIsBatchDirty;
-	
-	/** Selection annotation for fast lookup */
-	FUObjectAnnotationSparseBool* SelectionAnnotation;
-
-	bool bOwnsSelectionAnnotation;
-private:
-	// Hide IsSelected(), as calling IsSelected() on a selection set almost always indicates
-	// an error where the caller should use IsSelected(UObject* InObject).
-	bool IsSelected() const
-	{
-		return UObject::IsSelected();
-	}
-};
-
-
-/** A filter for generic selection sets.  Simply allows objects which are non-null */
-class FGenericSelectionFilter
-{
-public:
-	bool IsObjectValid( const UObject* InObject ) const
-	{
-		return InObject != nullptr;
-	}
-};
-
-/**
- * Manages selections of objects.  Used in the editor for selecting
- * objects in the various browser windows.
- */
-template<typename SelectionFilter>
-class TSelectionIterator
-{
-public:
-	TSelectionIterator(USelection& InSelection)
-		: Selection( InSelection )
-		, Filter( SelectionFilter() )
-	{
-		Reset();
-	}
-
-	/** Advances iterator to the next valid element in the container. */
-	void operator++()
-	{
-		while ( true )
-		{
-			++Index;
-
-			// Halt if the end of the selection set has been reached.
-			if ( !IsIndexValid() )
-			{
-				return;
-			}
-
-			// Halt if at a valid object.
-			if ( IsObjectValid() )
-			{
-				return;
-			}
-		}
-	}
-
-	/** Element access. */
-	UObject* operator*() const
-	{
-		return GetCurrentObject();
-	}
-
-	/** Element access. */
-	UObject* operator->() const
-	{
-		return GetCurrentObject();
-	}
-
-	/** Returns true if the iterator has not yet reached the end of the selection set. */
-	explicit operator bool() const
-	{
-		return IsIndexValid();
-	}
-
-	/** Resets the iterator to the beginning of the selection set. */
-	void Reset()
-	{
-		Index = -1;
-		++( *this );
-	}
-
-	/** Returns an index to the current element. */
-	int32 GetIndex() const
-	{
-		return Index;
-	}
-
-private:
-	UObject* GetCurrentObject() const
-	{
-		return Selection.GetSelectedObject(Index);
-	}
-
-	bool IsObjectValid() const
-	{
-		return Filter.IsObjectValid( GetCurrentObject() );
-	}
-
-	bool IsIndexValid() const
-	{
-		return Selection.SelectedObjects.IsValidIndex( Index );
-	}
-
-	USelection&	Selection;
-	SelectionFilter Filter;
-	int32			Index;
-};
-
-
-class FSelectionIterator : public TSelectionIterator<FGenericSelectionFilter>
-{
-public:
-	FSelectionIterator(USelection& InSelection)
-		: TSelectionIterator<FGenericSelectionFilter>( InSelection )
-	{}
-};
-
-/** A filter for only iterating through editable components */
-class FSelectedEditableComponentFilter
-{
-public:
-	bool IsObjectValid(const UObject* Object) const
-	{
-#if WITH_EDITOR
-		const UActorComponent* Comp = Cast<UActorComponent>( Object );
-		if (Comp)
-		{
-			return Comp->IsEditableWhenInherited();
-		}
-#endif
-		
-		return false;
-	}
-};
-
-/**
- * An iterator used to iterate through selected components that are editable (i.e. not created in a blueprint)
- */
-class FSelectedEditableComponentIterator : public TSelectionIterator<FSelectedEditableComponentFilter>
-{
-public:
-	FSelectedEditableComponentIterator(USelection& InSelection)
-		: TSelectionIterator<FSelectedEditableComponentFilter>(InSelection)
-	{}
-};
diff --git a/Engine/Source/Runtime/Engine/Private/Selection.cpp b/Engine/Source/Runtime/Engine/Private/Selection.cpp
deleted file mode 100644
