Removed unnecessary GetOrigin/GetBoxExtent accessors from FBoxSphereBounds

diff --git a/Engine/Source/Runtime/Core/Public/Math/BoxSphereBounds.h b/Engine/Source/Runtime/Core/Public/Math/BoxSphereBounds.h
--- a/Engine/Source/Runtime/Core/Public/Math/BoxSphereBounds.h
+++ b/Engine/Source/Runtime/Core/Public/Math/BoxSphereBounds.h
@@ -212,26 +212,6 @@ public:
 		return FSphere(Origin,SphereRadius);
 	}
 
-	/**
-	 * Gets the box/sphere origin.
-	 *
-	 * @return The box/sphere origin.
-	 */
-	FORCEINLINE FVector GetOrigin() const
-	{
-		return Origin;
-	}
-
-	/**
-	 * Gets the box extent.
-	 *
-	 * @return The box extent.
-	 */
-	FORCEINLINE FVector GetBoxExtent() const
-	{
-		return BoxExtent;
-	}
-
 	/**
 	 * Increase the size of the box and sphere by a given size.
 	 *
diff --git a/Engine/Source/Runtime/Engine/Private/InstanceUniformShaderParameters.cpp b/Engine/Source/Runtime/Engine/Private/InstanceUniformShaderParameters.cpp
--- a/Engine/Source/Runtime/Engine/Private/InstanceUniformShaderParameters.cpp
+++ b/Engine/Source/Runtime/Engine/Private/InstanceUniformShaderParameters.cpp
@@ -10,8 +10,8 @@ FInstanceSceneShaderData::FInstanceSceneShaderData(const FPrimitiveInstance& Ins
 		Instance.LocalToWorld,
 		Instance.PrevLocalToWorld,
 		Instance.WorldToLocal,
-		Instance.LocalBounds.GetOrigin(),
-		Instance.LocalBounds.GetBoxExtent(),
+		Instance.LocalBounds.Origin,
+		Instance.LocalBounds.BoxExtent,
 		Instance.NonUniformScale,
 		Instance.InvNonUniformScaleAndDeterminantSign,
 		Instance.NaniteInfo,
