Initial classes for generalized data driven compute kernels. Also basic flow for shader compiling w/ stubs.

diff --git a/Engine/Source/Runtime/Engine/Private/ComputeFramework/ComputeKernel.cpp b/Engine/Source/Runtime/Engine/Private/ComputeFramework/ComputeKernel.cpp
new file mode 100644
--- /dev/null
+++ b/Engine/Source/Runtime/Engine/Private/ComputeFramework/ComputeKernel.cpp
@@ -0,0 +1,118 @@
+// Copyright Epic Games, Inc. All Rights Reserved.
+
+#include "ComputeFramework/ComputeKernel.h"
+#include "ComputeFramework/ComputeKernelPermutationSet.h"
+#include "ComputeFramework/ComputeKernelSource.h"
+#include "ComputeFramework/ComputeKernelResource.h"
+
+DEFINE_LOG_CATEGORY(ComputeKernel);
+
+#if WITH_EDITOR
+void UComputeKernel::PostEditChangeChainProperty(FPropertyChangedChainEvent& PropertyChangedEvent)
+{
+	Super::PostEditChangeChainProperty(PropertyChangedEvent);
+
+	FProperty* ModifiedProperty = PropertyChangedEvent.PropertyChain.GetActiveMemberNode()->GetValue();
+	if (!ModifiedProperty)
+	{
+		return;
+	}
+
+	FName ModifiedPropName = ModifiedProperty->GetFName();
+
+	if (ModifiedPropName == GET_MEMBER_NAME_CHECKED(UComputeKernel, KernelSource))
+	{
+		if (KernelSource)
+		{
+			PermutationSet = KernelSource->PermutationSet;
+			DefinitionsSet = KernelSource->DefinitionsSet;
+		}
+		else
+		{
+			PermutationSet = FComputeKernelPermutationSet();
+			DefinitionsSet = FComputeKernelDefinitionsSet();
+		}
+		
+		CacheResourceShadersForRendering();
+	}
+	else if (ModifiedPropName == GET_MEMBER_NAME_CHECKED(UComputeKernel, PermutationSet) ||
+		ModifiedPropName == GET_MEMBER_NAME_CHECKED(UComputeKernel, DefinitionsSet))
+	{
+		CacheResourceShadersForRendering();
+	}
+}
+
+void UComputeKernel::CacheResourceShadersForRendering(
+	uint32 CompilationFlags /*= uint32(EComputeKernelCompilationFlags::None)*/
+	)
+{
+	if (!KernelSource)
+	{
+		KernelResource->Invalidate();
+		return;
+	}
+
+	if (!KernelResource)
+	{
+		KernelResource = MakeUnique<FComputeKernelResource>();
+	}
+
+	// #TODO_ZABIR: Initialize KernelResouce with needed data
+
+	ERHIFeatureLevel::Type CacheFeatureLevel = GMaxRHIFeatureLevel;
+	const EShaderPlatform ShaderPlatform = GShaderPlatformForFeatureLevel[CacheFeatureLevel];
+
+	UComputeKernel::CacheShadersForResource(ShaderPlatform, CompilationFlags, KernelResource.Get());
+}
+
+void UComputeKernel::CacheShadersForResource(
+	EShaderPlatform ShaderPlatform,
+	uint32 CompilationFlags,
+	FComputeKernelResource* KernelResource
+	)
+{
+	const bool bIsDefault = (KernelResource->GetKernelFlags() & uint32(EComputeKernelFlags::IsDefaultKernel)) != 0;
+	if (!GIsEditor || GIsAutomationTesting || bIsDefault)
+	{
+		CompilationFlags |= uint32(EComputeKernelCompilationFlags::Synchronous);
+	}
+
+	KernelResource->CacheShaders(ShaderPlatform, CompilationFlags);
+
+	const FComputeKernelCompilationResults& Results = KernelResource->GetCompilationResults();
+
+	if (!Results.bIsSuccess)
+	{
+		if (bIsDefault)
+		{
+			UE_LOG(
+				ComputeKernel, 
+				Fatal, 
+				TEXT("Failed to compile default FComputeKernelResource [%s] for platform [%s]!"),
+				*KernelResource->GetFriendlyName(),
+				*LegacyShaderPlatformToShaderFormat(ShaderPlatform).ToString()
+				);
+		}
+
+		UE_LOG(
+			ComputeKernel, 
+			Warning, 
+			TEXT("Failed to compile FComputeKernelResource [%s] for platform [%s]."),
+			*KernelResource->GetFriendlyName(),
+			*LegacyShaderPlatformToShaderFormat(ShaderPlatform).ToString()
+			);
+
+		uint32 WarningCount = Results.CompileWarnings.Num();
+		for (uint32 i = 0; i < WarningCount; ++i)
+		{
+			UE_LOG(ComputeKernel, Warning, TEXT("    [Warning] - %s"), *Results.CompileWarnings[i]);
+		}
+
+		uint32 ErrorCount = Results.CompileWarnings.Num();
+		for (uint32 i = 0; i < ErrorCount; ++i)
+		{
+			UE_LOG(ComputeKernel, Warning, TEXT("      [Error] - %s"), *Results.CompileErrors[i]);
+		}
+	}
+}
+#endif
diff --git a/Engine/Source/Runtime/Engine/Private/ComputeFramework/ComputeKernelFromText.cpp b/Engine/Source/Runtime/Engine/Private/ComputeFramework/ComputeKernelFromText.cpp
new file mode 100644
--- /dev/null
+++ b/Engine/Source/Runtime/Engine/Private/ComputeFramework/ComputeKernelFromText.cpp
@@ -0,0 +1,129 @@
+// Copyright Epic Games, Inc. All Rights Reserved.
+
+#include "ComputeFramework/ComputeKernelFromText.h"
+#include "Misc/FileHelper.h"
+#include "Misc/ScopeExit.h"
+
+#include "Internationalization/Regex.h"
+
+DECLARE_LOG_CATEGORY_EXTERN(ComputeKernel, Log, All);
+
+UComputeKernelFromText::UComputeKernelFromText()
+{
+	UniqueId = FGuid::NewGuid();
+}
+
+#if WITH_EDITOR
+
+void UComputeKernelFromText::PreEditChange(FEditPropertyChain& PropertyAboutToChange)
+{
+	Super::PreEditChange(PropertyAboutToChange);
+	PrevSourceFile = SourceFile;
+}
+
+void UComputeKernelFromText::PostEditChangeChainProperty(FPropertyChangedChainEvent& PropertyChangedEvent)
+{
+	Super::PostEditChangeChainProperty(PropertyChangedEvent);
+
+	FProperty* ModifiedProperty = PropertyChangedEvent.PropertyChain.GetActiveMemberNode()->GetValue();
+	if (!ModifiedProperty)
+	{
+		return;
+	}
+
+	FName ModifiedPropName = ModifiedProperty->GetFName();
+
+	if (ModifiedPropName == GET_MEMBER_NAME_CHECKED(UComputeKernelFromText, SourceFile))
+	{
+		ReparseKernelSourceText();
+	}
+}
+
+void UComputeKernelFromText::ReparseKernelSourceText()
+{
+	ON_SCOPE_EXIT
+	{
+		//FriendlyName = FPaths::GetCleanFilename(SourceFile.FilePath);
+		//Name = (FriendlyName.IsEmpty()) ? TEXT("Unset"_SV) : FriendlyName;
+	};
+
+	if (SourceFile.FilePath.IsEmpty())
+	{
+		EntryPointName = FString();
+		PermutationSet = FComputeKernelPermutationSet();
+		DefinitionsSet = FComputeKernelDefinitionsSet();
+
+		return;
+	}
+
+	FString FullKernelPath = FPaths::ConvertRelativePathToFull(SourceFile.FilePath);
+
+	IPlatformFile& PlatformFileSystem = IPlatformFile::GetPlatformPhysical();
+	if (!PlatformFileSystem.FileExists(*FullKernelPath))
+	{
+		UE_LOG(ComputeKernel, Error, TEXT("Unable to find kernel file \"%s\""), *FullKernelPath);
+
+		SourceFile = PrevSourceFile;
+		return;
+	}
+
+	FString KernelSourceText;
+	if (!FFileHelper::LoadFileToString(KernelSourceText, &PlatformFileSystem, *FullKernelPath))
+	{
+		UE_LOG(ComputeKernel, Error, TEXT("Unable to read kernel file \"%s\""), *FullKernelPath);
+
+		SourceFile = PrevSourceFile;
+		return;
+	}
+
+	const FRegexPattern KernelEntryPointPattern(TEXT(R"(KERNEL_ENTRY_POINT\(\s*([a-zA-Z_][\w\d]*)\s*\))"));
+	{
+		FRegexMatcher KernelEntryPointMatcher(KernelEntryPointPattern, KernelSourceText);
+		if (KernelEntryPointMatcher.FindNext())
+		{
+			EntryPointName = KernelEntryPointMatcher.GetCaptureGroup(1);
+		}
+	}
+
+	const FRegexPattern KernelPermutationBoolPattern(TEXT(R"(KERNEL_PERMUTATION_BOOL\(\s*([a-zA-Z_][\w\d]*)\s*\))"));
+	{
+		FComputeKernelPermutationSet NewPermutationSet;
+
+		FRegexMatcher BoolPermuationMatcher(KernelPermutationBoolPattern, KernelSourceText);
+		while (BoolPermuationMatcher.FindNext())
+		{
+			NewPermutationSet.BooleanOptions.Emplace(BoolPermuationMatcher.GetCaptureGroup(1));
+		}
+
+		PermutationSet = NewPermutationSet;
+	}
+
+	const FRegexPattern KernelDefinePattern(TEXT(R"(KERNEL_DEFINE\(\s*([a-zA-Z_][\w\d]*)\s*\))"));
+	{
+		FComputeKernelDefinitionsSet NewDefinitionsSet;
+
+		FRegexMatcher DefineMatcher(KernelDefinePattern, KernelSourceText);
+		while (DefineMatcher.FindNext())
+		{
+			NewDefinitionsSet.Defines.Emplace(DefineMatcher.GetCaptureGroup(1));
+		}
+
+		DefinitionsSet = NewDefinitionsSet;
+	}
+
+	/*
+	uint64 NewHash = 0;
+
+	NewHash = Hash(NewHash, UniqueId);
+	NewHash = Hash(NewHash, *SourceFile.FilePath);
+	NewHash = Hash(NewHash, *KernelSource);
+
+	if (SourceHash != NewHash)
+	{
+		SourceHash = NewHash;
+		// Trigger recompile
+	}
+	*/
+}
+
+#endif
diff --git a/Engine/Source/Runtime/Engine/Private/ComputeFramework/ComputeKernelPermutationSet.cpp b/Engine/Source/Runtime/Engine/Private/ComputeFramework/ComputeKernelPermutationSet.cpp
new file mode 100644
--- /dev/null
+++ b/Engine/Source/Runtime/Engine/Private/ComputeFramework/ComputeKernelPermutationSet.cpp
@@ -0,0 +1,3 @@
+// Copyright Epic Games, Inc. All Rights Reserved.
+
+#include "ComputeFramework/ComputeKernelPermutationSet.h"
diff --git a/Engine/Source/Runtime/Engine/Private/ComputeFramework/ComputeKernelResource.cpp b/Engine/Source/Runtime/Engine/Private/ComputeFramework/ComputeKernelResource.cpp
new file mode 100644
--- /dev/null
+++ b/Engine/Source/Runtime/Engine/Private/ComputeFramework/ComputeKernelResource.cpp
@@ -0,0 +1,132 @@
+// Copyright Epic Games, Inc. All Rights Reserved.
+
+#include "ComputeFramework/ComputeKernelResource.h"
+#include "ComputeFramework/ComputeKernelShaderMap.h"
+#include "ComputeFramework/ComputeKernel.h"
+#include "ShaderCompiler.h"
+
+#if WITH_EDITOR
+void FComputeKernelResource::CreateShaderMapId( 
+	FComputeKernelShaderMapId* OutShaderMapId 
+	) const
+{
+}
+
+void FComputeKernelResource::BeginCompileShaderMap(
+	EShaderPlatform ShaderPlatform,
+	uint32 CompilationFlags,
+	const FComputeKernelShaderMapId& ShaderMapId,
+	TRefCountPtr<FComputeKernelShaderMap>& OutShaderMap
+	)
+{
+	check(IsInGameThread());
+
+#if WITH_EDITORONLY_DATA
+	TRefCountPtr<FComputeKernelShaderMap> NewShaderMap = new FComputeKernelShaderMap();
+
+	const bool bSynchronousCompile = 
+		(CompilationFlags & uint32(EComputeKernelCompilationFlags::Synchronous)) || 
+		(CompilationFlags & uint32(EComputeKernelFlags::IsDefaultKernel)) || 
+		!GShaderCompilingManager->AllowAsynchronousShaderCompiling();
+
+	NewShaderMap->Compile(ShaderPlatform, this, ShaderMapId, bSynchronousCompile);
+
+	if (bSynchronousCompile && NewShaderMap->CompiledSuccessfully())
+	{
+		OutShaderMap = NewShaderMap;
+	}
+	else
+	{
+#if DEBUG_INFINITESHADERCOMPILE
+		UE_LOG(ComputeKernel, Display, TEXT("Kicking off shader compilation for FComputeKernelResource [%s], ShaderMap_GT 0x%08X%08X"), TEXT("GetFriendlyName()"), (int)((int64)(NewShaderMap.GetReference()) >> 32), (int)((int64)(NewShaderMap.GetReference())));
+#endif
+		InFlightCompilationIds.AddUnique(NewShaderMap->GetCompilingId());
+
+		// Async compile, engine will need to utilize fallback while compiling.
+		OutShaderMap = nullptr;
+	}
+#else
+	UE_LOG(ComputeKernel, Fatal, TEXT("Shader compilation outside the editor is not supported."));
+#endif
+}
+
+void FComputeKernelResource::CacheShaders(
+	EShaderPlatform ShaderPlatform,
+	uint32 CompilationFlags
+	)
+{
+	check(IsInGameThread());
+
+	FComputeKernelShaderMapId ShaderMapId;
+	CreateShaderMapId(&ShaderMapId);
+
+	ShaderMap_GT = nullptr;
+	{
+		ShaderMap_GT = FComputeKernelShaderMap::Find(ShaderPlatform, ShaderMapId);
+	}
+
+	const bool bForceRecompile = (CompilationFlags & uint32(EComputeKernelCompilationFlags::Force)) != 0;
+
+	// Attempt to load from the derived data cache uncooked
+	if (!bForceRecompile && !ShaderMap_GT && !FPlatformProperties::RequiresCookedData())
+	{
+		FComputeKernelShaderMap::LoadFromDerivedDataCache(ShaderPlatform, ShaderMapId, this, ShaderMap_GT);
+		if (ShaderMap_GT && ShaderMap_GT->IsValid())
+		{
+			UE_LOG(ComputeKernel, Verbose, TEXT("Loaded FShaderMap [%s] for FComputeKernelResource [%s] from DDC"), *ShaderMapId.GetFriendlyName(), TEXT("GetFriendlyName()"));
+		}
+		else
+		{
+			UE_LOG(ComputeKernel, Display, TEXT("Loading FShaderMap for FComputeKernelResource [%s] from DDC failed. Proceeding to compilation"), TEXT("GetFriendlyName()"));
+		}
+	}
+
+	bool bAssumeShaderMapIsComplete = false;
+#if UE_BUILD_SHIPPING || UE_BUILD_TEST
+	bAssumeShaderMapIsComplete = FPlatformProperties::RequiresCookedData();
+#endif
+
+	if (ShaderMap_GT && ShaderMap_GT->TryToAddToExistingCompilationTask(this))
+	{
+#if DEBUG_INFINITESHADERCOMPILE
+		UE_LOG(ComputeKernel, Display, TEXT("Found in flight compilation task for FComputeKernel %s, linking to other FShaderMap 0x%08X%08X"), TEXT("GetFriendlyName()"), (int)((int64)(GameThreadShaderMap.GetReference()) >> 32), (int)((int64)(GameThreadShaderMap.GetReference())));
+#endif
+
+		InFlightCompilationIds.AddUnique(ShaderMap_GT->GetCompilingId());
+		UE_LOG(ComputeKernel, Log, TEXT("FComputeKernelResource [%s] found exisiting in flight compilation id [%d]"), TEXT("GetFriendlyName()"), ShaderMap_GT->GetCompilingId());
+
+		// Reset the shader map so fall back rendering continues untill compilation is complete.
+		ShaderMap_GT = nullptr;
+	}
+	else if (bForceRecompile || !ShaderMap_GT || !(bAssumeShaderMapIsComplete || ShaderMap_GT->IsComplete()))
+	{
+		if (FPlatformProperties::RequiresCookedData())
+		{
+			UE_LOG(ComputeKernel, Fatal, TEXT("FComputeKernelResource [%s] cannot compile on the fly on cooked target!"), TEXT("GetFriendlyName()"));
+			ShaderMap_GT = nullptr;
+		}
+		else
+		{
+			if (ShaderMap_GT)
+			{
+				UE_LOG(ComputeKernel, Log, TEXT("FComputeKernelResource's [%s] FShaderMap [%s] is compiling. Waiting for compilation."), TEXT("GetFriendlyName()"), *ShaderMapId.GetFriendlyName());
+			}
+			else
+			{
+				UE_LOG(ComputeKernel, Log, TEXT("FComputeKernelResource's [%s] FShaderMap is missing. Proceeding to compilation."), TEXT("GetFriendlyName()"));
+			}
+
+			// Compilation is async unless EComputeKernelCompilationFlags::Synchronous flag is used.
+			BeginCompileShaderMap(ShaderPlatform, CompilationFlags, ShaderMapId, ShaderMap_GT);
+		}
+	}
+
+	FComputeKernelShaderMap* LoadedShaderMap = ShaderMap_RT;
+	ENQUEUE_RENDER_COMMAND(FSetRenderThreadShaderMap)(
+		[this, LoadedShaderMap](FRHICommandListImmediate& RHICmdList)
+	{
+		ShaderMap_RT = LoadedShaderMap;
+	});
+}
+
+#endif
diff --git a/Engine/Source/Runtime/Engine/Private/ComputeFramework/ComputeKernelShaderMap.cpp b/Engine/Source/Runtime/Engine/Private/ComputeFramework/ComputeKernelShaderMap.cpp
new file mode 100644
--- /dev/null
+++ b/Engine/Source/Runtime/Engine/Private/ComputeFramework/ComputeKernelShaderMap.cpp
@@ -0,0 +1,121 @@
+// Copyright Epic Games, Inc. All Rights Reserved.
+
+#include "ComputeFramework/ComputeKernelShaderMap.h"
+#include "ShaderCompiler.h"
+#include "ComputeFramework/ComputeKernelShaderType.h"
+#include "ComputeFramework/ComputeKernelResource.h"
+
+DECLARE_LOG_CATEGORY_EXTERN(ComputeKernel, Log, All);
+
+IMPLEMENT_TYPE_LAYOUT(FComputeKernelCompilationOutput);
+IMPLEMENT_TYPE_LAYOUT(FComputeKernelShaderMapId);
+IMPLEMENT_TYPE_LAYOUT(FComputeKernelShaderMapContent);
+
+TMap<TRefCountPtr<FComputeKernelShaderMap>, TArray<const FComputeKernelResource*>> FComputeKernelShaderMap::ComputeKernelShaderMapsBeingCompiled;
+
+static bool ShouldCacheComputeKernelShader(
+	EShaderPlatform ShaderPlatform, 
+	const FComputeKernelShaderType* ShaderType, 
+	const FComputeKernelResource* KernelShader
+	)
+{
+	return ShaderType->ShouldCache(ShaderPlatform, KernelShader) && KernelShader->ShouldCache(ShaderPlatform, ShaderType);
+}
+
+void FComputeKernelShaderMap::LoadFromDerivedDataCache(
+	EShaderPlatform ShaderPlatform, 
+	const FComputeKernelShaderMapId& ShaderMapId, 
+	FComputeKernelResource* KernelShader, 
+	TRefCountPtr<FComputeKernelShaderMap>& InOutGameThreadShaderMap
+	)
+{
+
+}
+
+bool FComputeKernelShaderMap::TryToAddToExistingCompilationTask(
+	FComputeKernelResource* KernelShader
+	)
+{
+	return false;
+}
+
+void FComputeKernelShaderMap::Compile(
+	EShaderPlatform ShaderPlatform,
+	FComputeKernelResource* KernelResource,
+	const FComputeKernelShaderMapId& ShaderMapId,
+	bool bSynchronousCompile
+	)
+{
+	check(IsInGameThread());
+	check(GetRefCount() > 0);
+	
+	if (FPlatformProperties::RequiresCookedData())
+	{
+		UE_LOG(ComputeKernel, Fatal, TEXT("Trying to compile FComputeKernelResource [%s] at run-time is not supported on %s!"), *KernelResource->GetFriendlyName(), FPlatformProperties::PlatformName());
+		return;
+	}
+
+	TArray<const FComputeKernelResource*>* CorrespondingKernels = ComputeKernelShaderMapsBeingCompiled.Find(this);
+
+	if (CorrespondingKernels)
+	{
+		// if this shader map is already doing an async, cannot do sync compile at this time.
+		check(!bSynchronousCompile);
+
+		CorrespondingKernels->AddUnique(KernelResource);
+		return;
+	}
+
+	ComputeKernelShaderMapsBeingCompiled.Add(this, { KernelResource });
+
+#if DEBUG_INFINITESHADERCOMPILE
+	UE_LOG(ComputeKernel, Display, TEXT("Added FComputeKernelShaderMap 0x%08X%08X with FComputeKernelResource 0x%08X%08X to ComputeKernelShaderMapsBeingCompiled"), (int)((int64)(this) >> 32), (int)((int64)(this)), (int)((int64)(KernelShader) >> 32), (int)((int64)(KernelShader)));
+#endif  
+
+	FComputeKernelShaderMapContent* NewContent = new FComputeKernelShaderMapContent(ShaderPlatform);
+	NewContent->FriendlyName = KernelResource->GetFriendlyName();
+	NewContent->ShaderMapId = ShaderMapId;
+	AssignContent(NewContent);
+	
+	TArray<TSharedRef<FShaderCommonCompileJob, ESPMode::ThreadSafe>> NewJobs;
+	TMap<TShaderTypePermutation<const FShaderType>, TSharedRef<FShaderCommonCompileJob, ESPMode::ThreadSafe>> SharedShaderJobs;
+
+	const TArray<FShaderType*>& ComputeKernelShaderTypeCollection = FShaderType::GetSortedTypes(FShaderType::EShaderTypeForDynamicCast::ComputeKernel);
+	for (auto* ShaderType : ComputeKernelShaderTypeCollection)
+	{
+		FComputeKernelShaderType* ComputeKernelShaderType = ShaderType->GetComputeKernelShaderType();
+		if (!ComputeKernelShaderType || !ShouldCacheComputeKernelShader(ShaderPlatform, ComputeKernelShaderType, KernelResource) || NewContent->HasShader(ShaderType, KernelResource->GetPermutationId()))
+		{
+			continue;
+		}
+
+		TSharedRef<FShaderCommonCompileJob, ESPMode::ThreadSafe> Job = ComputeKernelShaderType->BeginCompileShader(
+			ShaderPlatform,
+			KernelResource,
+			NewJobs
+			);
+
+		TShaderTypePermutation<const FShaderType> ShaderTypePermutation(ComputeKernelShaderType, KernelResource->GetPermutationId());
+		check(!SharedShaderJobs.Find(ShaderTypePermutation));
+
+		SharedShaderJobs.Add(ShaderTypePermutation, Job);
+
+		CompilationRequestId = Job->Id;
+	}
+
+	/**
+	Register(ShaderPlatform);
+
+	GComputeKernelShaderCompilationManager.AddJobs(NewJobs);
+
+	if (bSynchronousCompile)
+	{
+		TArray<int32> CurrentShaderMapId;
+		CurrentShaderMapId.Add(CompilationRequestId);
+		GComputeKernelShaderCompilationManager->FinishCompilation(
+			NewContent->FriendlyName,
+			CurrentShaderMapId
+			);
+	}
+	/**/
+}
diff --git a/Engine/Source/Runtime/Engine/Private/ComputeFramework/ComputeKernelShaderType.cpp b/Engine/Source/Runtime/Engine/Private/ComputeFramework/ComputeKernelShaderType.cpp
new file mode 100644
--- /dev/null
+++ b/Engine/Source/Runtime/Engine/Private/ComputeFramework/ComputeKernelShaderType.cpp
@@ -0,0 +1,45 @@
+// Copyright Epic Games, Inc. All Rights Reserved.
+
+#include "ComputeFramework/ComputeKernelShaderType.h"
+#include "ShaderCompiler.h"
+#include "RHI.h"
+#include "ComputeFramework/ComputeKernelResource.h"
+
+TSharedRef<FShaderCommonCompileJob, ESPMode::ThreadSafe> FComputeKernelShaderType::BeginCompileShader(
+	EShaderPlatform ShaderPlatform,
+	FComputeKernelResource* KernelShader,
+	TArray<TSharedRef<FShaderCommonCompileJob, ESPMode::ThreadSafe>>& InOutNewJobs
+	)
+{
+	FShaderCompileJob* NewJob = new FShaderCompileJob(FShaderCommonCompileJob::GetNextJobId(), nullptr, this, KernelShader->GetPermutationId());
+	{
+		NewJob->Input.Target = FShaderTarget(SF_Compute, ShaderPlatform);
+		NewJob->Input.ShaderFormat = LegacyShaderPlatformToShaderFormat(ShaderPlatform);
+
+		NewJob->Input.VirtualSourceFilePath = KernelShader->GetSourceFileName();
+		NewJob->Input.EntryPointName = KernelShader->GetEntryPointName();
+
+		NewJob->Input.SharedEnvironment = KernelShader->CreateShaderCompilationEnvironment(ShaderPlatform);
+
+		//NewJob->Input.Environment.SetDefine(TEXT("GPU_SIMULATION"), 1);
+		//NewJob->Input.Environment.IncludeVirtualPathToContentsMap.Add(TEXT("/Engine/Generated/NiagaraEmitterInstance.ush"), Script->HlslOutput);
+
+		SetupCompileEnvironment(ShaderPlatform, KernelShader, NewJob->Input.Environment);
+	}
+
+	TSharedRef<FShaderCommonCompileJob, ESPMode::ThreadSafe> SharedJob(NewJob);
+	
+	::GlobalBeginCompileShader(
+		KernelShader->GetFriendlyName(),
+		nullptr,
+		this,
+		nullptr,
+		KernelShader->GetSourceFileName(),
+		KernelShader->GetEntryPointName(),
+		FShaderTarget(SF_Compute, ShaderPlatform),
+		SharedJob,
+		InOutNewJobs
+		);
+
+	return SharedJob;
+}
\ No newline at end of file
diff --git a/Engine/Source/Runtime/Engine/Private/ComputeFramework/ComputeKernelSource.cpp b/Engine/Source/Runtime/Engine/Private/ComputeFramework/ComputeKernelSource.cpp
new file mode 100644
--- /dev/null
+++ b/Engine/Source/Runtime/Engine/Private/ComputeFramework/ComputeKernelSource.cpp
@@ -0,0 +1,3 @@
+// Copyright Epic Games, Inc. All Rights Reserved.
+
+#include "ComputeFramework/ComputeKernelSource.h"
diff --git a/Engine/Source/Runtime/Engine/Public/ComputeFramework/ComputeKernel.h b/Engine/Source/Runtime/Engine/Public/ComputeFramework/ComputeKernel.h
new file mode 100644
--- /dev/null
+++ b/Engine/Source/Runtime/Engine/Public/ComputeFramework/ComputeKernel.h
@@ -0,0 +1,146 @@
+// Copyright Epic Games, Inc. All Rights Reserved.
+
+#pragma once
+
+#include "CoreMinimal.h"
+#include "ComputeFramework/ShaderParamTypeDefinition.h"
+#include "ComputeFramework/ComputeKernelPermutationSet.h"
+#include "ComputeFramework\ComputeKernelResource.h"
+
+#include "ComputeKernel.generated.h"
+
+class UComputeKernelSource;
+
+DECLARE_LOG_CATEGORY_EXTERN(ComputeKernel, Log, All);
+
+/* Describes the size and shape (threads) of a kernel invocation. */
+USTRUCT(BlueprintType)
+struct FKernelInvocationDimension
+{
+	GENERATED_BODY()
+
+	UPROPERTY(EditAnywhere, Category = "Kernel")
+	uint16 X = 0;
+
+	UPROPERTY(EditAnywhere, Category = "Kernel")
+	uint16 Y = 0;
+
+	UPROPERTY(EditAnywhere, Category = "Kernel")
+	uint16 Z = 0;
+};
+
+/* Flags that convey kernel behavior to aid compilation/optimizations. */
+UENUM(meta = (Bitflags, UseEnumValuesAsMaskValuesInEditor = "true"))
+enum class EComputeKernelFlags : uint32
+{
+	/*
+	 * Default implies that this kernel must be compiled before the system is functional.
+	 * It also implies that this will be compiled synchronously. Other than a pass
+	 * through kernel, default shouldn't be used.
+	 */
+	IsDefaultKernel = 1 << 0, // KERNEL_FLAG(IS_DEFAUL_KERNEL)
+
+	/*
+	 * Promise from the author that all memory writes will be unique per shader
+	 * dispatch thread. i.e. ThreadX will be the only thread to write to MemoryY,
+	 * thus no synchronization is necessary by the compute graph.
+	 */
+	IsolatedMemoryWrites = 1 << 1, // KERNEL_FLAG(ISOLATED_MEMORY_WRITES)
+};
+ENUM_CLASS_FLAGS(EComputeKernelFlags);
+
+enum class EComputeKernelCompilationFlags
+{
+	None		= 0,
+	
+	/* Force recompilation even if kernel is not dirty and/or DDC data is available. */
+	Force		= 1 << 0,
+
+	/* Compile the shader while blocking the main thread. */
+	Synchronous = 1 << 1,
+};
+
+/* Base class representing a kernel that will be run as a shader on the GPU. */
+UCLASS(hidecategories = (Object))
+class ENGINE_API UComputeKernel : public UObject
+{
+	GENERATED_BODY()
+
+public:
+	/* 
+	 * A kernel's source may be authored by different mechanisms; e.g. HLSL text, VPL graph, ML Meta Lang, etc
+	 * This abstracts the source and compilation process.
+	 */
+	UPROPERTY(EditAnywhere, AssetRegistrySearchable, meta = (ShowOnlyInnerProperties), Category = "Kernel")
+	UComputeKernelSource* KernelSource = nullptr;
+
+	/* Specifying certain memory access flags allows for optimizations such as kernel fusing. */
+	UPROPERTY(VisibleAnywhere, BlueprintReadOnly, meta = (Bitmask, BitmaskEnum = "EComputeKernelFlags"), Category = "Kernel")
+	int32 KernelFlags = 0;
+
+	UPROPERTY(EditDefaultsOnly, meta = (ShowOnlyInnerProperties), Category = "Kernel")
+	FComputeKernelPermutationSet PermutationSet;
+
+	UPROPERTY(EditDefaultsOnly, meta = (ShowOnlyInnerProperties), Category = "Kernel")
+	FComputeKernelDefinitionsSet DefinitionsSet;
+
+	/*
+	 * The minimum number of invocations (threads) launched by the kernel definition.
+	 * Generally aim for multiple of 64 total. i.e. X * Y * Z = 64
+	 */
+	UPROPERTY(VisibleAnywhere, Category = "Kernel")
+	FKernelInvocationDimension GroupSizeDim;
+
+	/* Named input types to the kernel. These are discovered via reflection from the shader. */
+	UPROPERTY(VisibleAnywhere, EditFixedSize, Category = "Kernel")
+	TArray<FShaderParamTypeDefinition> Inputs;
+
+	/* Named output types to the kernel. These are discovered via reflection from the shader. */
+	UPROPERTY(VisibleAnywhere, EditFixedSize, Category = "Kernel")
+	TArray<FShaderParamTypeDefinition> Outputs;
+
+#if WITH_EDITOR
+	DECLARE_EVENT_OneParam(UComputeKernel, FShaderResetEvent, const UComputeKernel*);
+
+	/*
+	 * Delegate is invoked when the shader or shader input/output list is changed.
+	 * All priorly retrieved bind points are invalid after this signal is raised.
+	 */
+	FShaderResetEvent ShaderResetSignal;
+#endif
+
+	/* Retrieve the IO bind points for the kernel - only way to get a valid bind point. */
+	//FShaderBindPoint GetInputBindPoint(const FString& Name) const;
+	//FShaderBindPoint GetOutputBindPoint(const FString& Name) const;
+
+	/*
+	 * Implemented by derived classes to specify how many invocations of the kernel are necessary 
+	 * given the specific data payloads the kernel will operate on. 
+	 */
+	//virtual FKernelInvocationDimension CalculateInvocationDimension(const TArray<FBindPointDataDesc>& BindingCtx) const = 0;
+
+	//=========================================================================
+	// UObject interface
+#if WITH_EDITOR
+	virtual void PostEditChangeChainProperty(FPropertyChangedChainEvent& PropertyChangedEvent) override;
+#endif
+	//=========================================================================
+
+private:
+#if WITH_EDITOR
+	/* Entry point to initiate any [re]compilation necessary due to changes. */
+	void CacheResourceShadersForRendering(
+		uint32 CompilationFlags = uint32(EComputeKernelCompilationFlags::None)
+		);
+
+	/* Given a FKernelResouce, its associated shader is compiled - used for rendering and cooking. */
+	static void CacheShadersForResource(
+		EShaderPlatform ShaderPlatform,
+		uint32 CompilationFlags,
+		FComputeKernelResource* KernelResource
+		);
+#endif
+
+	/* The shader resouce encapsulating the kernel. Akin to FMaterialResource */
+	TUniquePtr<FComputeKernelResource> KernelResource = nullptr;
+};
diff --git a/Engine/Source/Runtime/Engine/Public/ComputeFramework/ComputeKernelFromText.h b/Engine/Source/Runtime/Engine/Public/ComputeFramework/ComputeKernelFromText.h
new file mode 100644
--- /dev/null
+++ b/Engine/Source/Runtime/Engine/Public/ComputeFramework/ComputeKernelFromText.h
@@ -0,0 +1,44 @@
+// Copyright Epic Games, Inc. All Rights Reserved.
+
+#pragma once
+
+#include "ComputeKernelSource.h"
+
+#include "ComputeKernelFromText.generated.h"
+
+/* 
+ * Responsible for loading HLSL text and parsing the options available.
+ */
+UCLASS(BlueprintType)
+class ENGINE_API UComputeKernelFromText : public UComputeKernelSource
+{
+	GENERATED_BODY()
+
+public:
+	UComputeKernelFromText();
+
+	/* Filepath to the source file containing the kernel entry points. */
+	UPROPERTY(EditDefaultsOnly, AssetRegistrySearchable, meta = (ContentDir, RelativeToGameContentDir, FilePathFilter = "Unreal Shader File (*.usf)|*.usf"), Category = "Kernel")
+	FFilePath SourceFile;
+
+	UPROPERTY(VisibleAnywhere, AssetRegistrySearchable, Category = "Kernel")
+	FString EntryPointName;
+
+	UPROPERTY()
+	FGuid UniqueId;
+
+	UPROPERTY()
+	uint64 SourceHash = 0;
+
+#if WITH_EDITOR
+	FFilePath PrevSourceFile;
+
+	virtual void PreEditChange(FEditPropertyChain& PropertyAboutToChange) override;
+	virtual void PostEditChangeChainProperty(FPropertyChangedChainEvent& PropertyChangedEvent) override;
+#endif
+
+private:
+#if WITH_EDITOR
+	void ReparseKernelSourceText();
+#endif
+};
diff --git a/Engine/Source/Runtime/Engine/Public/ComputeFramework/ComputeKernelPermutationSet.h b/Engine/Source/Runtime/Engine/Public/ComputeFramework/ComputeKernelPermutationSet.h
new file mode 100644
--- /dev/null
+++ b/Engine/Source/Runtime/Engine/Public/ComputeFramework/ComputeKernelPermutationSet.h
@@ -0,0 +1,117 @@
+// Copyright Epic Games, Inc. All Rights Reserved.
+
+#pragma once
+
+#include "ComputeKernelPermutationSet.generated.h"
+
+UENUM()
+enum class EComputeKernelPermutationType : uint8
+{
+	Bool,
+	Range,
+	Set,
+	Enum,
+
+	Count,
+};
+
+// #TODO_ZABIR: refactor to a compact bitset once we do the UI work to present it reasonably.
+USTRUCT()
+struct FComputeKernelPermutationBool
+{
+	GENERATED_BODY()
+
+	UPROPERTY(VisibleAnywhere, Category = "Permutation Options")
+	FString Name;
+
+	UPROPERTY(EditDefaultsOnly, Category = "Permutation Options")
+	bool Value;
+
+	FComputeKernelPermutationBool()
+		: Value(false)
+	{
+	}
+
+	explicit FComputeKernelPermutationBool(FString InName, bool InValue = false)
+		: Name(MoveTemp(InName))
+		, Value(InValue)
+	{
+	}
+};
+
+USTRUCT()
+struct FComputeKernelPermutationSet
+{
+	GENERATED_BODY()
+
+	UPROPERTY(EditDefaultsOnly, EditFixedSize, meta = (EditFixedOrder), Category = "Permutation Options")
+	TArray<FComputeKernelPermutationBool> BooleanOptions;
+
+	uint32 GetPermutationCount() const
+	{
+		uint32 PermutationCount = (1 << BooleanOptions.Num());
+
+		check(PermutationCount < INT32_MAX);
+		return PermutationCount;
+	}
+
+	uint32 GetPermutationId() const
+	{
+		checkSlow(GetPermutationCount());
+
+		uint32 PermutationId = 0;
+
+		for (auto& BoolOpt : BooleanOptions)
+		{
+			Encode(&PermutationId, 2, BoolOpt.Value ? 1 : 0);
+		}
+
+		return PermutationId;
+	}
+
+private:
+	static void Encode(uint32* Encoded, uint32 ValueRange, uint32 Value)
+	{
+		*Encoded *= ValueRange;
+		*Encoded += Value;
+	}
+
+	static uint32 Decode(uint32* Encoded, uint32 ValueRange)
+	{
+		uint32 Value = *Encoded % ValueRange;
+		*Encoded /= ValueRange;
+
+		return Value;
+	}
+};
+
+USTRUCT()
+struct FComputeKernelDefinitions
+{
+	GENERATED_BODY()
+
+public:
+	UPROPERTY(VisibleAnywhere, Category = "Kernel")
+	FString Symbol;
+
+	UPROPERTY(EditDefaultsOnly, Category = "Kernel")
+	FString Define;
+
+	FComputeKernelDefinitions() = default;
+
+	explicit FComputeKernelDefinitions(FString InSymbol, FString InDefine = FString())
+		: Symbol(MoveTemp(InSymbol))
+		, Define(MoveTemp(InDefine))
+	{
+	}
+};
+
+USTRUCT()
+struct FComputeKernelDefinitionsSet
+{
+	GENERATED_BODY()
+
+public:
+	UPROPERTY(EditDefaultsOnly, EditFixedSize, meta = (EditFixedOrder), Category = "Permutation Options")
+	TArray<FComputeKernelDefinitions> Defines;
+};
diff --git a/Engine/Source/Runtime/Engine/Public/ComputeFramework/ComputeKernelResource.h b/Engine/Source/Runtime/Engine/Public/ComputeFramework/ComputeKernelResource.h
new file mode 100644
--- /dev/null
+++ b/Engine/Source/Runtime/Engine/Public/ComputeFramework/ComputeKernelResource.h
@@ -0,0 +1,104 @@
+// Copyright Epic Games, Inc. All Rights Reserved.
+
+#pragma once
+
+#include "ComputeFramework//ComputeKernelShaderMap.h"
+
+struct FComputeKernelCompilationResults
+{
+	bool bIsSuccess : 1;
+
+	TArray<FString> CompileWarnings;
+	TArray<FString> CompileErrors;
+
+	FComputeKernelCompilationResults()
+		: bIsSuccess(false)
+	{}
+};
+
+class FComputeKernelResource
+{
+public:
+	uint32 GetKernelFlags() const { return 0; }
+
+	const FComputeKernelCompilationResults& GetCompilationResults() const { return CompilationResults; }
+
+#if WITH_EDITOR
+	void CacheShaders(
+		EShaderPlatform ShaderPlatform,
+		uint32 CompilationFlags
+		);
+#endif
+
+private:
+	FComputeKernelCompilationResults CompilationResults;
+
+	/*
+	 * Game thread view of the shader map. The shader map uses deferred deletion so that the rendering thread 
+	 * has a chance to process, then release when the shader map is no longer used by rendering thread.
+	 * Code that sets this is responsible for updating ShaderMap_RT in a thread safe way.
+	 * During an async compile, this will be nullptr.
+	 */
+	TRefCountPtr<FComputeKernelShaderMap> ShaderMap_GT;
+
+	/**
+	 * Render thread view of the shader map. Updates should originate by modifying ShaderMap_GT and then 
+	 * propagate to the render thread.
+	 */
+	FComputeKernelShaderMap* ShaderMap_RT;
+
+	/*
+	 * Contains the compiling IDs of this shader map when it is being compiled asynchronously.
+	 * This can be used to access the shader map during async compiling.
+	 */
+	TArray<uint32> InFlightCompilationIds;
+
+	void CreateShaderMapId(
+		FComputeKernelShaderMapId* OutShaderMapId
+		) const;
+
+#if WITH_EDITOR
+	void BeginCompileShaderMap(
+		EShaderPlatform ShaderPlatform,
+		uint32 CompilationFlags,
+		const FComputeKernelShaderMapId& ShaderMapId,
+		TRefCountPtr<FComputeKernelShaderMap>& OutShaderMap
+		);
+#endif
+
+
+
+
+
+
+
+
+
+
+
+
+
+
+
+	
+public:
+	const FString& GetFriendlyName() const { return FriendlyName; }
+	const TCHAR* GetSourceFileName() const { return nullptr; }
+	const TCHAR* GetEntryPointName() const { return nullptr; }
+	uint32 GetPermutationId() const { return 0; }
+
+	void Invalidate() {}
+
+	TRefCountPtr<FShaderCompilerEnvironment> CreateShaderCompilationEnvironment(EShaderPlatform ShaderPlatform) const 
+	{ 
+		return new FShaderCompilerEnvironment(); 
+	}
+
+	bool ShouldCache(EShaderPlatform ShaderPlatform, const FShaderType* ShaderType) const 
+	{ 
+		return ShaderType->GetComputeKernelShaderType() != nullptr; 
+	}
+
+private:
+	FString FriendlyName;
+};
diff --git a/Engine/Source/Runtime/Engine/Public/ComputeFramework/ComputeKernelShaderMap.h b/Engine/Source/Runtime/Engine/Public/ComputeFramework/ComputeKernelShaderMap.h
new file mode 100644
--- /dev/null
+++ b/Engine/Source/Runtime/Engine/Public/ComputeFramework/ComputeKernelShaderMap.h
@@ -0,0 +1,114 @@
+// Copyright Epic Games, Inc. All Rights Reserved.
+
+#pragma once
+
+class FComputeKernelResource;
+
+/* Stores all the output/stats from the compilation process, eg. error messages, error codes. */
+class FComputeKernelCompilationOutput
+{
+	DECLARE_TYPE_LAYOUT(FComputeKernelCompilationOutput, NonVirtual);
+
+public:
+};
+
+/** Contains all the information needed to uniquely identify a FComputeKernelShaderMap. */
+class FComputeKernelShaderMapId
+{
+	DECLARE_TYPE_LAYOUT(FComputeKernelShaderMapId, NonVirtual);
+
+public:
+	LAYOUT_FIELD(ERHIFeatureLevel::Type, FeatureLevel);
+	LAYOUT_FIELD(TMemoryImageArray<FMemoryImageString>, AdditionalDefines);
+	LAYOUT_FIELD_EDITORONLY(FMemoryImageString, FriendlyName);
+
+	FComputeKernelShaderMapId()
+		: FeatureLevel(GMaxRHIFeatureLevel)
+	{
+	}
+
+	SIZE_T GetSizeBytes() const
+	{
+		return sizeof(*this);
+	}
+
+	const FMemoryImageString& GetFriendlyName() const { return FriendlyName; }
+};
+
+class FComputeKernelShaderMapContent : public FShaderMapContent
+{
+	DECLARE_TYPE_LAYOUT(FComputeKernelShaderMapContent, NonVirtual);
+	friend class FComputeKernelShaderMap;
+
+public:
+	using Super = FShaderMapContent;
+
+	explicit FComputeKernelShaderMapContent(EShaderPlatform InPlatform = EShaderPlatform::SP_NumPlatforms) 
+		: FShaderMapContent(InPlatform) 
+	{
+	}
+
+private:
+	LAYOUT_FIELD(FComputeKernelCompilationOutput, CompilationOutput);
+	LAYOUT_FIELD(FSHAHash, ShaderContentHash);
+	LAYOUT_FIELD(FComputeKernelShaderMapId, ShaderMapId);
+
+	LAYOUT_FIELD_EDITORONLY(FMemoryImageString, FriendlyName);
+	LAYOUT_FIELD_EDITORONLY(FMemoryImageString, DebugDescription);
+};
+
+class FComputeKernelShaderMap : public TShaderMap<FComputeKernelShaderMapContent, FShaderMapPointerTable>, public FDeferredCleanupInterface
+{
+public:
+	static FComputeKernelShaderMap* Find(
+		EShaderPlatform ShaderPlatform, 
+		const FComputeKernelShaderMapId& ShaderMapId
+		) 
+	{ 
+		return nullptr; 
+	}
+
+	static void LoadFromDerivedDataCache(
+		EShaderPlatform ShaderPlatform, 
+		const FComputeKernelShaderMapId& ShaderMapId, 
+		FComputeKernelResource* KernelShader, 
+		TRefCountPtr<FComputeKernelShaderMap>& InOutGameThreadShaderMap
+		);
+
+	/*
+	 * Checks to see if the shader map is already being compiled, and if so
+	 * add this to the list to be applied once the compile finishes.
+	 * Returns true if the shader map was being compiled and was added.
+	 */
+	bool TryToAddToExistingCompilationTask(FComputeKernelResource* KernelShader);
+
+
+
+
+	bool IsValid() { return false; }
+
+	const FMemoryImageString& GetFriendlyName() const { return GetContent()->FriendlyName; }
+
+	void Compile(
+		EShaderPlatform ShaderPlatform,
+		FComputeKernelResource* KernelShader,
+		const FComputeKernelShaderMapId& ShaderMapId,
+		bool bSynchronousCompile
+		);
+
+	uint32 GetCompilingId() const { return CompilationRequestId; }
+	bool CompiledSuccessfully() const { return bCompiledSuccessfully; }
+	bool IsComplete() const { return false; }
+
+private:
+	/* Tracks resources and their shader map compilations that are in flight. */
+	static TMap<TRefCountPtr<FComputeKernelShaderMap>, TArray<const FComputeKernelResource*>> ComputeKernelShaderMapsBeingCompiled;
+
+	/* 
+	 * Uniquely identifies this shader map during compilation. Needed for deferred compilation
+	 * where many shader compilations are running concurrently owned by different shader maps. 
+	 */
+	uint32 CompilationRequestId;
+
+	bool bCompiledSuccessfully : 1;
+};
diff --git a/Engine/Source/Runtime/Engine/Public/ComputeFramework/ComputeKernelShaderType.h b/Engine/Source/Runtime/Engine/Public/ComputeFramework/ComputeKernelShaderType.h
new file mode 100644
--- /dev/null
+++ b/Engine/Source/Runtime/Engine/Public/ComputeFramework/ComputeKernelShaderType.h
@@ -0,0 +1,38 @@
+// Copyright Epic Games, Inc. All Rights Reserved.
+
+#pragma once
+
+struct FShaderCompilerEnvironment;
+class FComputeKernelResource;
+
+struct FComputeKernelShaderPermutationParameters : public FShaderPermutationParameters
+{
+	const FComputeKernelResource* KernelShader;
+
+	FComputeKernelShaderPermutationParameters(EShaderPlatform InPlatform, const FComputeKernelResource* InKernelShader)
+		: FShaderPermutationParameters(InPlatform)
+		, KernelShader(InKernelShader)
+	{}
+};
+
+class FComputeKernelShaderType : public FShaderType
+{
+public:
+	TSharedRef<FShaderCommonCompileJob, ESPMode::ThreadSafe> BeginCompileShader(
+		EShaderPlatform ShaderPlatform,
+		FComputeKernelResource* KernelShader,
+		TArray<TSharedRef<FShaderCommonCompileJob, ESPMode::ThreadSafe>>& InOutNewJobs
+		);
+
+
+	bool ShouldCache(EShaderPlatform ShaderPlatform, const FComputeKernelResource* KernelShader) const
+	{
+		return ShouldCompilePermutation(FComputeKernelShaderPermutationParameters(ShaderPlatform, KernelShader));
+	}
+
+protected:
+	void SetupCompileEnvironment(EShaderPlatform ShaderPlatform, const FComputeKernelResource* KernelShader, FShaderCompilerEnvironment& Environment) const
+	{
+		ModifyCompilationEnvironment(FComputeKernelShaderPermutationParameters(ShaderPlatform, KernelShader), Environment);
+	}
+};
diff --git a/Engine/Source/Runtime/Engine/Public/ComputeFramework/ComputeKernelSource.h b/Engine/Source/Runtime/Engine/Public/ComputeFramework/ComputeKernelSource.h
new file mode 100644
--- /dev/null
+++ b/Engine/Source/Runtime/Engine/Public/ComputeFramework/ComputeKernelSource.h
@@ -0,0 +1,21 @@
+// Copyright Epic Games, Inc. All Rights Reserved.
+
+#pragma once
+
+#include "ComputeFramework/ComputeKernelPermutationSet.h"
+#include "ComputeKernelSource.generated.h"
+
+class FComputeKernelResource;
+
+UCLASS(Abstract, EditInlineNew)
+class UComputeKernelSource : public UObject
+{
+	GENERATED_BODY()
+
+public:
+	UPROPERTY(VisibleAnywhere, meta = (ShowOnlyInnerProperties), Category = "Kernel")
+	FComputeKernelPermutationSet PermutationSet;
+
+	UPROPERTY(VisibleAnywhere, meta = (ShowOnlyInnerProperties), Category = "Kernel")
+	FComputeKernelDefinitionsSet DefinitionsSet;
+};
diff --git a/Engine/Source/Runtime/Engine/Public/ComputeFramework/ShaderParamTypeDefinition.h b/Engine/Source/Runtime/Engine/Public/ComputeFramework/ShaderParamTypeDefinition.h
new file mode 100644
--- /dev/null
+++ b/Engine/Source/Runtime/Engine/Public/ComputeFramework/ShaderParamTypeDefinition.h
@@ -0,0 +1,138 @@
+// Copyright Epic Games, Inc. All Rights Reserved.
+
+#pragma once
+
+#include "ShaderParamTypeDefinition.generated.h"
+
+/* The base types of data that shaders can consume/expose */
+UENUM()
+enum class EShaderFundamentalType : uint8
+{
+	Bool,
+	Int,
+	Uint,
+	Float,
+	Double,
+	Struct,
+};
+
+/*
+ * Shader types can then be in the form of a scalar, vector, matrix.
+ * e.g Scalar: float a; 	Vector: float3 n; 		Matrix: float3x4 WVP;
+ * Note: float b[5]; is still considered scalar. It is an array of scalars.
+ */
+UENUM()
+enum class EShaderFundamentalDimensionType : uint8
+{
+	Scalar,
+	Vector,
+	Matrix,
+};
+
+/* Describes how the shader parameters are bound. */
+UENUM()
+enum class EShaderParamBindingType : uint8
+{
+	ConstantParameter,
+	ReadOnlyResource,  // SRV, treated as Inputs
+	ReadWriteResource, // UAV, treated as Outputs
+};
+
+UENUM()
+enum class EShaderResourceType : uint8
+{
+	Texture1D,
+	Texture2D,
+	Texture3D,
+	TextureCube,
+	Buffer,
+	StructuredBuffer,
+	ByteAddressBuffer,
+};
+
+/* Fully describes the name and type of a parameter a shader exposes. */
+USTRUCT()
+struct FShaderParamTypeDefinition
+{
+	GENERATED_BODY()
+
+public:
+	UPROPERTY(EditAnywhere, Category = "Kernel")
+	FString	Name;
+
+// #TODO_ZABIR: Shader reflection needs to extract much richer type info. Currently skip param type validation
+#if 0
+	UPROPERTY()
+	uint16 ArrayElementCount; // 0 indicates not an array. >= 1 indicates an array
+
+	UPROPERTY()
+	EShaderFundamentalType FundamentalType;
+
+	UPROPERTY()
+	EShaderFundamentalDimensionType	DimType;
+
+	UPROPERTY()
+	EShaderParamBindingType BindingType;
+
+	UPROPERTY()
+	EShaderResourceType	ResourceType;
+
+	union
+	{
+		uint8 VectorElemCount : 2;
+
+		struct
+		{
+			uint8 MatrixRowCount : 2;
+			uint8 MatrixColumnCount : 2;
+		};
+	};
+
+	bool IsAnyBufferType() const
+	{
+		return
+			ResourceType == EShaderResourceType::Buffer ||
+			ResourceType == EShaderResourceType::ByteAddressBuffer ||
+			ResourceType == EShaderResourceType::StructuredBuffer;
+	}
+
+	bool IsAnyTextureType() const
+	{
+		return
+			ResourceType == EShaderResourceType::Texture1D ||
+			ResourceType == EShaderResourceType::Texture2D ||
+			ResourceType == EShaderResourceType::Texture3D ||
+			ResourceType == EShaderResourceType::TextureCube;
+	}
+
+	/* Determines if the type definition is valid according to HLSL rules. */
+	bool IsValid() const
+	{
+		if (FundamentalType == EShaderFundamentalType::Struct && DimType != EShaderFundamentalDimensionType::Scalar)
+		{
+			// cannot have anything but scalar struct types
+			return false;
+		}
+
+		if (IsAnyTextureType() && FundamentalType == EShaderFundamentalType::Struct)
+		{
+			// cannot have textures of structs
+			return false;
+		}
+
+		if (IsAnyTextureType() && DimType == EShaderFundamentalDimensionType::Matrix)
+		{
+			// cannot have textures of matrices
+			return false;
+		}
+
+		if ((IsAnyBufferType() || IsAnyTextureType()) && BindingType == EShaderParamBindingType::ConstantParameter)
+		{
+			// cannot have buffers and textures bound as const params
+			return false;
+		}
+
+		return true;
+	}
+#endif
+};
diff --git a/Engine/Source/Runtime/RenderCore/Public/Shader.h b/Engine/Source/Runtime/RenderCore/Public/Shader.h
--- a/Engine/Source/Runtime/RenderCore/Public/Shader.h
+++ b/Engine/Source/Runtime/RenderCore/Public/Shader.h
@@ -32,6 +32,7 @@ class FMaterialShaderType;
 class FNiagaraShaderType;
 class FOpenColorIOShaderType;
 class FMeshMaterialShaderType;
+class FComputeKernelShaderType;
 class FShader;
 class FShaderMapBase;
 class FShaderPipelineType;
@@ -978,6 +979,7 @@ public:
 		MeshMaterial,
 		Niagara,
 		OCIO,
+		ComputeKernel,
 		NumShaderTypes,
 	};
 
@@ -1091,6 +1093,14 @@ public:
 	{
 		return (ShaderTypeForDynamicCast == EShaderTypeForDynamicCast::OCIO) ? reinterpret_cast<FOpenColorIOShaderType*>(this) : nullptr;
 	}
+	FORCEINLINE const FComputeKernelShaderType* GetComputeKernelShaderType() const
+	{
+		return (ShaderTypeForDynamicCast == EShaderTypeForDynamicCast::ComputeKernel) ? reinterpret_cast<const FComputeKernelShaderType*>(this) : nullptr;
+	}
+	FORCEINLINE FComputeKernelShaderType* GetComputeKernelShaderType()
+	{
+		return (ShaderTypeForDynamicCast == EShaderTypeForDynamicCast::ComputeKernel) ? reinterpret_cast<FComputeKernelShaderType*>(this) : nullptr;
+	}
 
 	inline EShaderTypeForDynamicCast GetTypeForDynamicCast() const
 	{
