Compile fix for unity build

diff --git a/Engine/Source/Runtime/Engine/Private/ComputeFramework/ComputeKernelFromText.cpp b/Engine/Source/Runtime/Engine/Private/ComputeFramework/ComputeKernelFromText.cpp
--- a/Engine/Source/Runtime/Engine/Private/ComputeFramework/ComputeKernelFromText.cpp
+++ b/Engine/Source/Runtime/Engine/Private/ComputeFramework/ComputeKernelFromText.cpp
@@ -1,13 +1,12 @@
 // Copyright Epic Games, Inc. All Rights Reserved.
 
 #include "ComputeFramework/ComputeKernelFromText.h"
+#include "ComputeFramework/ComputeKernel.h"
 #include "Misc/FileHelper.h"
 #include "Misc/ScopeExit.h"
 
 #include "Internationalization/Regex.h"
 
-DECLARE_LOG_CATEGORY_EXTERN(ComputeKernel, Log, All);
-
 UComputeKernelFromText::UComputeKernelFromText()
 {
 	UniqueId = FGuid::NewGuid();
diff --git a/Engine/Source/Runtime/Engine/Private/ComputeFramework/ComputeKernelShaderMap.cpp b/Engine/Source/Runtime/Engine/Private/ComputeFramework/ComputeKernelShaderMap.cpp
--- a/Engine/Source/Runtime/Engine/Private/ComputeFramework/ComputeKernelShaderMap.cpp
+++ b/Engine/Source/Runtime/Engine/Private/ComputeFramework/ComputeKernelShaderMap.cpp
@@ -4,8 +4,7 @@
 #include "ShaderCompiler.h"
 #include "ComputeFramework/ComputeKernelShaderType.h"
 #include "ComputeFramework/ComputeKernelResource.h"
-
-DECLARE_LOG_CATEGORY_EXTERN(ComputeKernel, Log, All);
+#include "ComputeFramework/ComputeKernel.h"
 
 IMPLEMENT_TYPE_LAYOUT(FComputeKernelCompilationOutput);
 IMPLEMENT_TYPE_LAYOUT(FComputeKernelShaderMapId);
