Compile fixes for Standalone builds
-Added missing WITH_EDITORONLY_DATA sections
-removed /* within comment that clang doesn't like

diff --git a/Engine/Source/Runtime/Engine/Private/ComputeFramework/ComputeKernelShaderMap.cpp b/Engine/Source/Runtime/Engine/Private/ComputeFramework/ComputeKernelShaderMap.cpp
--- a/Engine/Source/Runtime/Engine/Private/ComputeFramework/ComputeKernelShaderMap.cpp
+++ b/Engine/Source/Runtime/Engine/Private/ComputeFramework/ComputeKernelShaderMap.cpp
@@ -72,7 +72,9 @@ void FComputeKernelShaderMap::Compile(
 #endif  
 
 	FComputeKernelShaderMapContent* NewContent = new FComputeKernelShaderMapContent(ShaderPlatform);
+#if WITH_EDITORONLY_DATA
 	NewContent->FriendlyName = KernelResource->GetFriendlyName();
+#endif
 	NewContent->ShaderMapId = ShaderMapId;
 	AssignContent(NewContent);
 	
@@ -116,5 +118,5 @@ void FComputeKernelShaderMap::Compile(
 			CurrentShaderMapId
 			);
 	}
-	/**/
+	**/
 }
diff --git a/Engine/Source/Runtime/Engine/Public/ComputeFramework/ComputeKernelShaderMap.h b/Engine/Source/Runtime/Engine/Public/ComputeFramework/ComputeKernelShaderMap.h
--- a/Engine/Source/Runtime/Engine/Public/ComputeFramework/ComputeKernelShaderMap.h
+++ b/Engine/Source/Runtime/Engine/Public/ComputeFramework/ComputeKernelShaderMap.h
@@ -32,7 +32,9 @@ public:
 		return sizeof(*this);
 	}
 
+#if WITH_EDITORONLY_DATA
 	const FMemoryImageString& GetFriendlyName() const { return FriendlyName; }
+#endif
 };
 
 class FComputeKernelShaderMapContent : public FShaderMapContent
@@ -87,7 +89,9 @@ public:
 
 	bool IsValid() { return false; }
 
+#if WITH_EDITORONLY_DATA
 	const FMemoryImageString& GetFriendlyName() const { return GetContent()->FriendlyName; }
+#endif
 
 	void Compile(
 		EShaderPlatform ShaderPlatform,
