WorldPartition HLODs
- Added support for mesh simplification
- Moved HLOD mesh building code out of the HLOD Layer class
- Renamed MeshProxy to MeshSimplify
- Re-add CacheHLODParents() call that was lost during a swarm review refactor
- SWorldPartitionEditor::Construct() - Use LoadModuleChecked() instead of GetModuleChecked() to obtain the WorldBrowser module

diff --git a/Engine/Source/Editor/WorldPartitionEditor/Private/WorldPartition/SWorldPartitionEditor.cpp b/Engine/Source/Editor/WorldPartitionEditor/Private/WorldPartition/SWorldPartitionEditor.cpp
--- a/Engine/Source/Editor/WorldPartitionEditor/Private/WorldPartition/SWorldPartitionEditor.cpp
+++ b/Engine/Source/Editor/WorldPartitionEditor/Private/WorldPartition/SWorldPartitionEditor.cpp
@@ -20,7 +20,7 @@ void SWorldPartitionEditor::Construct(const FArguments& InArgs)
 
 	OnBrowseWorld(InArgs._InWorld);
 
-	FWorldBrowserModule& WorldBrowserModule = FModuleManager::GetModuleChecked<FWorldBrowserModule>("WorldBrowser");
+	FWorldBrowserModule& WorldBrowserModule = FModuleManager::LoadModuleChecked<FWorldBrowserModule>("WorldBrowser");
 	WorldBrowserModule.OnBrowseWorld.AddSP(this, &SWorldPartitionEditor::OnBrowseWorld);
 }
 
diff --git a/Engine/Source/Runtime/Engine/Private/WorldPartition/HLOD/HLODBuilder.cpp b/Engine/Source/Runtime/Engine/Private/WorldPartition/HLOD/HLODBuilder.cpp
new file mode 100644
--- /dev/null
+++ b/Engine/Source/Runtime/Engine/Private/WorldPartition/HLOD/HLODBuilder.cpp
@@ -0,0 +1,261 @@
+#include "WorldPartition/HLOD/HLODBuilder.h"
+
+#include "WorldPartition/WorldPartition.h"
+#include "WorldPartition/HLOD/HLODActor.h"
+#include "WorldPartition/HLOD/HLODLayer.h"
+
+#include "Algo/ForEach.h"
+#include "Algo/Transform.h"
+#include "Templates/UniquePtr.h"
+
+#include "IMeshMergeUtilities.h"
+#include "MeshMergeModule.h"
+
+/**
+ * Base class for all HLODBuilders
+ */
+class FHLODBuilder
+{
+public:
+	virtual ~FHLODBuilder() {}
+
+	virtual void Build(const TArray<AActor*>& InSubActors) const = 0;
+
+	typedef TFunction<UPrimitiveComponent*(AWorldPartitionHLOD*)> FCreateComponentFunction;
+		
+	AWorldPartitionHLOD* SpawnHLODActor(const TCHAR* InName, const TArray<UPrimitiveComponent*>& InSubComponents, FCreateComponentFunction InCreateComponentFunc) const
+	{
+		FActorSpawnParameters SpawnParams;
+		SpawnParams.bDeferConstruction = true;
+		SpawnParams.bCreateActorPackage = true;
+		AWorldPartitionHLOD* HLODActor = World->SpawnActor<AWorldPartitionHLOD>(SpawnParams);
+
+		UPrimitiveComponent* ParentComponent = InCreateComponentFunc(HLODActor);
+
+		if (ParentComponent)
+		{
+			HLODActor->SetHLODLayer(HLODLayer, iLevel);
+			HLODActor->SetParentPrimitive(ParentComponent);
+			HLODActor->SetChildrenPrimitives(InSubComponents);
+			HLODActor->SetActorLabel(FString::Printf(TEXT("%s_%s_%s"), *HLODLayer->GetName(), *CellName.ToString(), InName));
+			HLODActor->RuntimeGrid = GetHLODLevelSettings().TargetGrid;
+
+			WorldPartition->UpdateActorDesc(HLODActor);
+		}
+		else
+		{
+			World->DestroyActor(HLODActor);
+			HLODActor = nullptr;
+		}
+
+		return HLODActor;
+	}
+
+	TArray<UPrimitiveComponent*> GatherPrimitiveComponents(int32 iHLODLevel, const TArray<AActor*> InActors) const
+	{
+		TArray<UPrimitiveComponent*> PrimitiveComponents;
+		for (AActor* SubActor : InActors)
+		{
+			for (UActorComponent* SubComponent : SubActor->GetComponents())
+			{
+				if (UPrimitiveComponent* PrimitiveComponent = Cast<UPrimitiveComponent>(SubComponent))
+				{
+					if (UHLODLayer::ShouldIncludeInHLOD(PrimitiveComponent, iHLODLevel))
+					{
+						PrimitiveComponents.Add(PrimitiveComponent);
+					}
+				}
+			}
+		}
+		return PrimitiveComponents;
+	}
+
+	const FHLODLevelSettings& GetHLODLevelSettings() const
+	{
+		return HLODLayer->GetLevels()[iLevel];
+	}
+
+public:
+	UWorld*				World;
+	UWorldPartition*	WorldPartition;
+	const UHLODLayer*	HLODLayer;
+	int32				iLevel;
+	FName				CellName;
+};
+
+
+/**
+ * Build a AWorldPartitionHLOD whose components are ISMC
+ */
+class FHLODBuilder_Instancing : public FHLODBuilder
+{
+	virtual void Build(const TArray<AActor*>& InSubActors) const override
+	{
+		TRACE_CPUPROFILER_EVENT_SCOPE(FHLODBuilder_Instancing::BuildHLOD);
+
+		const FHLODLevelSettings& LevelSettings = GetHLODLevelSettings();
+
+		// Gather all meshes to instantiate along with their transforms
+		TArray<UPrimitiveComponent*> PrimitiveComponents = GatherPrimitiveComponents(iLevel, InSubActors);
+		TMap<UStaticMesh*, TArray<UPrimitiveComponent*>> Instances;
+		for (UPrimitiveComponent* Primitive : PrimitiveComponents)
+		{
+			if (UStaticMeshComponent* SMC = Cast<UStaticMeshComponent>(Primitive))
+			{
+				Instances.FindOrAdd(SMC->GetStaticMesh()).Add(SMC);
+			}
+		}
+
+		// Now, create an ISMC for each SM asset we found
+		for (const auto& Entry : Instances)
+		{
+			UStaticMesh* StaticMesh = Entry.Key;
+			const TArray<UPrimitiveComponent*>& SubComponents = Entry.Value;
+
+			FCreateComponentFunction CreateComponentLambda = [&LevelSettings, StaticMesh, SubComponents](AWorldPartitionHLOD* HLODActor)
+			{
+				UInstancedStaticMeshComponent* Component = NewObject<UInstancedStaticMeshComponent>(HLODActor);
+				Component->SetStaticMesh(StaticMesh);
+				Component->SetForcedLodModel(StaticMesh->GetNumLODs());
+
+				// Add all instances
+				for (UPrimitiveComponent* SMC : SubComponents)
+				{
+					Component->AddInstanceWorldSpace(SMC->GetComponentTransform());
+				}
+
+				return Component;
+			};
+
+			SpawnHLODActor(*StaticMesh->GetName(), SubComponents, CreateComponentLambda);
+		}
+	}
+};
+
+
+/**
+ * Build a merged mesh using geometry from the provided actors
+ */
+class FHLODBuilder_MeshMerge : public FHLODBuilder
+{
+	virtual void Build(const TArray<AActor*>& InSubActors) const override
+	{
+		TRACE_CPUPROFILER_EVENT_SCOPE(UHLODLayer::BuildHLOD_MeshMerge);
+
+		const FHLODLevelSettings& LevelSettings = GetHLODLevelSettings();
+
+		TArray<UPrimitiveComponent*> SubComponents = GatherPrimitiveComponents(iLevel, InSubActors);
+
+		FCreateComponentFunction CreateComponentLambda = [&LevelSettings, SubComponents, this](AWorldPartitionHLOD* HLODActor)
+		{
+			TArray<UObject*> Assets;
+			FVector MergedActorLocation;
+
+			const IMeshMergeUtilities& MeshMergeUtilities = FModuleManager::Get().LoadModuleChecked<IMeshMergeModule>("MeshMergeUtilities").GetUtilities();
+			MeshMergeUtilities.MergeComponentsToStaticMesh(SubComponents, HLODActor->GetWorld(), LevelSettings.MergeSetting, LevelSettings.FlattenMaterial.LoadSynchronous(), HLODActor->GetPackage(), CellName.ToString(), Assets, MergedActorLocation, 0.25f, false);
+
+			// All merged mesh assets are stored in the HLOD Actor package
+			Algo::ForEach(Assets, [](UObject* Asset) { Asset->ClearFlags(RF_Public | RF_Standalone); });
+
+			UStaticMesh* StaticMesh = nullptr;
+			UStaticMeshComponent* Component = nullptr;
+			if (Assets.FindItemByClass<UStaticMesh>(&StaticMesh))
+			{
+				Component = NewObject<UStaticMeshComponent>(HLODActor);
+				Component->SetStaticMesh(StaticMesh);
+				Component->SetWorldLocation(MergedActorLocation);
+			}
+
+			return Component;
+		};
+		
+		SpawnHLODActor(TEXT("MergedMesh"), SubComponents, CreateComponentLambda);
+	}
+};
+
+/**
+ * Build a simplified mesh using geometry from the provided actors
+ */
+class FHLODBuilder_MeshSimplify : public FHLODBuilder
+{
+	virtual void Build(const TArray<AActor*>& InSubActors) const override
+	{
+		TRACE_CPUPROFILER_EVENT_SCOPE(UHLODLayer::BuildHLOD_MeshProxy);
+
+		const FHLODLevelSettings& LevelSettings = GetHLODLevelSettings();
+
+		TArray<UPrimitiveComponent*> SubComponents = GatherPrimitiveComponents(iLevel, InSubActors);
+
+		FCreateComponentFunction CreateComponentLambda = [&LevelSettings, SubComponents, this] (AWorldPartitionHLOD* HLODActor)
+		{
+			TArray<UObject*> Assets;
+			FCreateProxyDelegate ProxyDelegate;
+			ProxyDelegate.BindLambda([&Assets](const FGuid Guid, TArray<UObject*>& InAssetsCreated) { Assets = InAssetsCreated; });
+
+			TArray<UStaticMeshComponent*> StaticMeshComponents;
+			Algo::TransformIf(SubComponents, StaticMeshComponents, [](UPrimitiveComponent* InPrimitiveComponent) { return InPrimitiveComponent->IsA<UStaticMeshComponent>(); }, [](UPrimitiveComponent* InPrimitiveComponent) { return Cast<UStaticMeshComponent>(InPrimitiveComponent); });
+
+			const IMeshMergeUtilities& MeshMergeUtilities = FModuleManager::Get().LoadModuleChecked<IMeshMergeModule>("MeshMergeUtilities").GetUtilities();
+			MeshMergeUtilities.CreateProxyMesh(StaticMeshComponents, LevelSettings.ProxySetting, LevelSettings.FlattenMaterial.LoadSynchronous(), HLODActor->GetPackage(), CellName.ToString(), FGuid::NewGuid(), ProxyDelegate, true);
+
+			// All merged mesh assets are stored in the HLOD Actor package
+			Algo::ForEach(Assets, [](UObject* Asset) { Asset->ClearFlags(RF_Public | RF_Standalone); });
+
+			UStaticMesh* StaticMesh = nullptr;
+			UStaticMeshComponent* Component = nullptr;
+			if (Assets.FindItemByClass<UStaticMesh>(&StaticMesh))
+			{
+				Component = NewObject<UStaticMeshComponent>(HLODActor);
+				Component->SetStaticMesh(StaticMesh);
+			}
+
+			return Component;
+		};
+		
+		SpawnHLODActor(TEXT("SimplifiedMesh"), SubComponents, CreateComponentLambda);
+	}
+};
+
+
+/**
+ *
+ */
+void FHLODBuilderUtilities::BuildHLODs(UWorldPartition* InWorldPartition, FName InCellName, const UHLODLayer* InHLODLayer, const TArray<AActor*>& InSubActors)
+{
+	const TArray<FHLODLevelSettings>& LevelsSettings = InHLODLayer->GetLevels();
+
+	for (int32 iLevel = 0; iLevel < LevelsSettings.Num(); ++iLevel)
+	{
+		TUniquePtr<FHLODBuilder> HLODBuilder = nullptr;
+
+		EHLODLevelType HLODLevelType = LevelsSettings[iLevel].LevelType;
+		switch (HLODLevelType)
+		{
+		case EHLODLevelType::Instancing:
+			HLODBuilder = TUniquePtr<FHLODBuilder>(new FHLODBuilder_Instancing());
+			break;
+
+		case EHLODLevelType::MeshMerge:
+			HLODBuilder = TUniquePtr<FHLODBuilder>(new FHLODBuilder_MeshMerge());
+			break;
+
+		case EHLODLevelType::MeshSimplify:
+			HLODBuilder = TUniquePtr<FHLODBuilder>(new FHLODBuilder_MeshSimplify());
+			break;
+
+		default:
+			checkf(false, TEXT("Unsupported type"));
+		}
+
+		if (HLODBuilder)
+		{
+			HLODBuilder->World = InWorldPartition->GetWorld();
+			HLODBuilder->WorldPartition = InWorldPartition;
+			HLODBuilder->HLODLayer = InHLODLayer;
+			HLODBuilder->iLevel = iLevel;
+			HLODBuilder->CellName = InCellName;
+
+			HLODBuilder->Build(InSubActors);
+		}
+	}
+}
diff --git a/Engine/Source/Runtime/Engine/Private/WorldPartition/HLOD/HLODBuilder.h b/Engine/Source/Runtime/Engine/Private/WorldPartition/HLOD/HLODBuilder.h
new file mode 100644
--- /dev/null
+++ b/Engine/Source/Runtime/Engine/Private/WorldPartition/HLOD/HLODBuilder.h
@@ -0,0 +1,28 @@
+// Copyright Epic Games, Inc. All Rights Reserved.
+
+#pragma once
+
+#include "CoreMinimal.h"
+#include "Containers/Array.h"
+
+class AActor;
+class UWorldPartition;
+class UHLODLayer;
+
+/**
+ * Tools for building HLODs in WorldPartition
+ */
+class FHLODBuilderUtilities
+{
+public:
+	/**
+	 * Build HLODs for a given cell of a WorldPartition level.
+	 * May spawn multiple AWorldPartitionHLOD actors depending on the HLODLayer settings.
+	 *
+	 * @param 	InWorldPartition	The WorldPartition for which we are building HLODs
+	 * @param 	InCellName			The name of the cell
+	 * @param 	InHLODLayer			The HLODLayer which will provide the HLOD building parameters.
+	 * @param 	InSubActors			The actors from which we'll gather geometry to generate an HLOD mesh
+	 */
+	static void BuildHLODs(UWorldPartition* InWorldPartition, FName InCellName, const UHLODLayer* InHLODLayer, const TArray<AActor*>& InSubActors);
+};
diff --git a/Engine/Source/Runtime/Engine/Private/WorldPartition/HLOD/HLODLayer.cpp b/Engine/Source/Runtime/Engine/Private/WorldPartition/HLOD/HLODLayer.cpp
--- a/Engine/Source/Runtime/Engine/Private/WorldPartition/HLOD/HLODLayer.cpp
+++ b/Engine/Source/Runtime/Engine/Private/WorldPartition/HLOD/HLODLayer.cpp
@@ -10,14 +10,7 @@
 #include "WorldPartition/WorldPartition.h"
 #include "WorldPartition/WorldPartitionActorDesc.h"
 #include "WorldPartition/HLOD/HLODActor.h"
-
-#include "Algo/ForEach.h"
-
-#include "IMeshMergeUtilities.h"
-#include "IMeshDescriptionModule.h"
-#include "MeshMergeModule.h"
-#include "Materials/Material.h"
-#include "Engine/StaticMesh.h"
+#include "WorldPartition/HLOD/HLODBuilder.h"
 #endif
 
 UHLODLayer::UHLODLayer(const FObjectInitializer& ObjectInitializer)
@@ -61,7 +54,7 @@ void UHLODLayer::GenerateHLODForCell(UWorldPartition* InWorldPartition, FName In
 			continue;
 		}
 
-		HLODLayer->BuildHLOD(InWorldPartition, InCellName, HLODLayerActors.Value);
+		FHLODBuilderUtilities::BuildHLODs(InWorldPartition, InCellName, HLODLayer, HLODLayerActors.Value);
 	}
 }
 
@@ -140,38 +133,6 @@ bool UHLODLayer::ShouldIncludeInHLOD(UPrimitiveComponent* InComponent, int32 InL
 	return true;
 }
 
-
-
-int32 UHLODLayer::BuildHLOD(UWorldPartition* InWorldPartition, FName InCellName, const TArray<AActor*>& InSubActors)
-{
-	TRACE_CPUPROFILER_EVENT_SCOPE(UHLODLayer::BuildHLOD);
-
-	int32 NumActorSpawned = 0;
-
-	for (int32 iLevel = 0; iLevel < Levels.Num(); ++iLevel)
-	{
-		switch (Levels[iLevel].LevelType)
-		{
-		case EHLODLevelType::Instancing:
-			NumActorSpawned += BuildHLOD_Instancing(iLevel, InWorldPartition, InCellName, InSubActors);
-			break;
-
-		case EHLODLevelType::MeshMerge:
-			NumActorSpawned += BuildHLOD_MeshMerge(iLevel, InWorldPartition, InCellName, InSubActors);
-			break;
-
-		case EHLODLevelType::MeshProxy:
-			NumActorSpawned += BuildHLOD_MeshProxy(iLevel, InWorldPartition, InCellName, InSubActors);
-			break;
-
-		default:
-			checkf(false, TEXT("Unsupported type"));
-		}
-	}
-
-	return NumActorSpawned;
-}
-
 UHLODLayer* UHLODLayer::GetHLODLayer(const AActor* InActor)
 {
 	if (UHLODLayer* HLODLayer = InActor->GetHLODLayer())
@@ -185,135 +146,4 @@ UHLODLayer* UHLODLayer::GetHLODLayer(const AActor* InActor)
 	return nullptr;
 }
 
-TArray<UPrimitiveComponent*> GatherPrimitiveComponents(int32 iHLODLevel, const TArray<AActor*> InActors)
-{
-	TArray<UPrimitiveComponent*> PrimitiveComponents;
-	for (AActor* SubActor : InActors)
-	{
-		for (UActorComponent* SubComponent : SubActor->GetComponents())
-		{
-			if (UPrimitiveComponent* PrimitiveComponent = Cast<UPrimitiveComponent>(SubComponent))
-			{
-				if (UHLODLayer::ShouldIncludeInHLOD(PrimitiveComponent, iHLODLevel))
-				{
-					PrimitiveComponents.Add(PrimitiveComponent);
-				}
-			}
-		}
-	}
-	return PrimitiveComponents;
-}
-
-int32 UHLODLayer::BuildHLOD_Instancing(int32 iLevel, UWorldPartition* InWorldPartition, FName InCellName, const TArray<AActor*>& InSubActors) const
-{
-	TRACE_CPUPROFILER_EVENT_SCOPE(UHLODLayer::BuildHLOD_Instancing);
-
-	const FHLODLevelSettings& LevelSettings = Levels[iLevel];
-
-	TArray<UPrimitiveComponent*> PrimitiveComponents = GatherPrimitiveComponents(iLevel, InSubActors);
-
-	// Gather all meshes to instantiate along with their transforms
-	TMap<UStaticMesh*, TArray<UPrimitiveComponent*>> Instances;
-	for (UPrimitiveComponent* Primitive : PrimitiveComponents)
-	{
-		if (UStaticMeshComponent* SMC = Cast<UStaticMeshComponent>(Primitive))
-		{
-			Instances.FindOrAdd(SMC->GetStaticMesh()).Add(SMC);
-		}
-	}
-
-	int32 NumActors = 0;
-
-	// Now, create an ISMC for each SM asset we found
-	for (const auto& Entry : Instances)
-	{
-		FActorSpawnParameters SpawnParams;
-		SpawnParams.bDeferConstruction = true;
-		SpawnParams.bCreateActorPackage = true;
-		AWorldPartitionHLOD* HLODActor = InWorldPartition->GetWorld()->SpawnActor<AWorldPartitionHLOD>(SpawnParams);
-
-		UInstancedStaticMeshComponent* Component = NewObject<UInstancedStaticMeshComponent>(HLODActor);
-		Component->SetStaticMesh(Entry.Key);
-		Component->SetForcedLodModel(Entry.Key->GetNumLODs());
-
-		// Add all instances
-		TSet<AActor*> Actors;
-
-		for (UPrimitiveComponent* SMC : Entry.Value)
-		{
-			Actors.Add(SMC->GetOwner());
-			Component->AddInstanceWorldSpace(SMC->GetComponentTransform());
-		}
-
-		HLODActor->SetHLODLayer(this, iLevel);
-		HLODActor->SetParentPrimitive(Component);
-		HLODActor->SetChildrenPrimitives(Entry.Value);
-		HLODActor->SetActorLabel(FString::Printf(TEXT("%s_%s_%s"), *GetName(), *InCellName.ToString(), *Entry.Key->GetName()));
-		HLODActor->RuntimeGrid = Levels[iLevel].TargetGrid;
-
-		InWorldPartition->UpdateActorDesc(HLODActor);
-
-		NumActors++;
-	}
-
-	return NumActors;
-}
-
-int32 UHLODLayer::BuildHLOD_MeshMerge(int32 iLevel, UWorldPartition* InWorldPartition, FName InCellName, const TArray<AActor*>& InSubActors) const
-{
-	TRACE_CPUPROFILER_EVENT_SCOPE(UHLODLayer::BuildHLOD_MeshMerge);
-
-	FActorSpawnParameters SpawnParams;
-	SpawnParams.bDeferConstruction = true;
-	SpawnParams.bCreateActorPackage = true;
-	AWorldPartitionHLOD* HLODActor = InWorldPartition->GetWorld()->SpawnActor<AWorldPartitionHLOD>(SpawnParams);
-
-	const FHLODLevelSettings& LevelSettings = Levels[iLevel];
-
-	const IMeshDescriptionModule& MeshDescriptionModule = IMeshDescriptionModule::Get();
-	const IMeshMergeUtilities& MergeUtilities = FModuleManager::Get().LoadModuleChecked<IMeshMergeModule>("MeshMergeUtilities").GetUtilities();
-
-	TArray<UPrimitiveComponent*> PrimitiveComponents = GatherPrimitiveComponents(iLevel, InSubActors);
-
-	TArray<UObject*> Assets;
-	FVector MergedActorLocation;
-	MergeUtilities.MergeComponentsToStaticMesh(PrimitiveComponents, HLODActor->GetWorld(), LevelSettings.MergeSetting, LevelSettings.FlattenMaterial.LoadSynchronous(), HLODActor->GetPackage(), TEXT(""), Assets, MergedActorLocation, 0.25f, false);
-
-	// All merged mesh assets are stored in the HLOD Actor package
-	Algo::ForEach(Assets, [](UObject* Asset) { Asset->ClearFlags(RF_Public | RF_Standalone); });
-
-	UStaticMesh* StaticMesh;
-
-	int32 NumComponents = 0;
-	if (Assets.FindItemByClass<UStaticMesh>(&StaticMesh))
-	{
-		UStaticMeshComponent* Component = NewObject<UStaticMeshComponent>(HLODActor);
-		Component->SetStaticMesh(StaticMesh);
-		Component->SetWorldLocation(MergedActorLocation);
-		
-		HLODActor->SetHLODLayer(this, iLevel);
-		HLODActor->SetParentPrimitive(Component);
-		HLODActor->SetChildrenPrimitives(PrimitiveComponents);
-		HLODActor->SetActorLabel(FString::Printf(TEXT("%s_%s_MergedMesh"), *GetName(), *InCellName.ToString()));
-		HLODActor->RuntimeGrid = Levels[iLevel].TargetGrid;
-
-		InWorldPartition->UpdateActorDesc(HLODActor);
-		NumComponents++;
-	}
-	else
-	{
-		InWorldPartition->GetWorld()->DestroyActor(HLODActor);
-	}
-
-	return NumComponents;
-}
-
-int32 UHLODLayer::BuildHLOD_MeshProxy(int32 iLevel, UWorldPartition* InWorldPartition, FName InCellName, const TArray<AActor*>& InSubActors) const
-{
-	TRACE_CPUPROFILER_EVENT_SCOPE(UHLODLayer::BuildHLOD_MeshProxy);
-
-	const FHLODLevelSettings& LevelSettings = Levels[iLevel];
-
-	return 0;
-}
 #endif // WITH_EDITOR
\ No newline at end of file
diff --git a/Engine/Source/Runtime/Engine/Private/WorldPartition/WorldPartitionRuntimeSpatialHash.cpp b/Engine/Source/Runtime/Engine/Private/WorldPartition/WorldPartitionRuntimeSpatialHash.cpp
--- a/Engine/Source/Runtime/Engine/Private/WorldPartition/WorldPartitionRuntimeSpatialHash.cpp
+++ b/Engine/Source/Runtime/Engine/Private/WorldPartition/WorldPartitionRuntimeSpatialHash.cpp
@@ -470,6 +470,9 @@ bool UWorldPartitionRuntimeSpatialHash::GenerateStreaming(EWorldPartitionStreami
 
 	check(!StreamingGrids.Num());
 
+	// Build a map of Actor GUID -> HLODActor GUID once instead of having to recompute for every streaming grid we create
+	CacheHLODParents();
+
 	TMap<FName, int32> GridsMapping;
 	GridsMapping.Add(NAME_None, 0);
 	for (int32 i = 0; i < AllGrids.Num(); i++)
diff --git a/Engine/Source/Runtime/Engine/Public/WorldPartition/HLOD/HLODLayer.h b/Engine/Source/Runtime/Engine/Public/WorldPartition/HLOD/HLODLayer.h
--- a/Engine/Source/Runtime/Engine/Public/WorldPartition/HLOD/HLODLayer.h
+++ b/Engine/Source/Runtime/Engine/Public/WorldPartition/HLOD/HLODLayer.h
@@ -19,9 +19,9 @@ class AWorldPartitionHLOD;
 UENUM()
 enum class EHLODLevelType : uint8
 {
-	Instancing = 0	UMETA(DisplayName = "Use instancing"),
-	MeshMerge = 1	UMETA(DisplayName = "Generate merged mesh"),
-	MeshProxy = 2	UMETA(DisplayName = "Generate proxy mesh")
+	Instancing = 0		UMETA(DisplayName = "Use instancing"),
+	MeshMerge = 1		UMETA(DisplayName = "Generate merged mesh"),
+	MeshSimplify = 2	UMETA(DisplayName = "Generate simplified mesh")
 };
 
 USTRUCT()
@@ -36,10 +36,10 @@ struct FHLODLevelSettings
 	UPROPERTY(EditAnywhere, Category=HLOD, meta = (editcondition = "LevelType == EHLODLevelType::MeshMerge", ShowOnlyInnerProperties))
 	FMeshMergingSettings MergeSetting;
 
-	UPROPERTY(EditAnywhere, Category=HLOD, meta = (editcondition = "LevelType == EHLODLevelType::MeshProxy", ShowOnlyInnerProperties))
+	UPROPERTY(EditAnywhere, Category=HLOD, meta = (editcondition = "LevelType == EHLODLevelType::MeshSimplify", ShowOnlyInnerProperties))
 	FMeshProxySettings ProxySetting;
 
-	UPROPERTY(EditAnywhere, Category=HLOD, meta = (editcondition = "LevelType == EHLODLevelType::MeshMerge || LevelType == EHLODLevelType::MeshProxy", ShowOnlyInnerProperties))
+	UPROPERTY(EditAnywhere, Category=HLOD, meta = (editcondition = "LevelType == EHLODLevelType::MeshMerge || LevelType == EHLODLevelType::MeshSimplify", ShowOnlyInnerProperties))
 	TSoftObjectPtr<UMaterial> FlattenMaterial;
 
 	UPROPERTY(EditAnywhere, Category=HLOD)
@@ -64,13 +64,6 @@ public:
 	static UHLODLayer* GetHLODLayer(const AActor* InActor);
 	static bool ShouldIncludeInHLOD(AActor* InActor);
 	static bool ShouldIncludeInHLOD(UPrimitiveComponent* InComponent, int32 InLevelIndex);
-
-private:
-	int32 BuildHLOD(UWorldPartition* InWorldPartition, FName InCellName, const TArray<AActor*>& InSubActors);
-
-	int32 BuildHLOD_Instancing(int32 iLevel, UWorldPartition* InWorldPartition, FName InCellName, const TArray<AActor*>& InSubActors) const;
-	int32 BuildHLOD_MeshMerge(int32 iLevel, UWorldPartition* InWorldPartition, FName InCellName, const TArray<AActor*>& InSubActors) const;
-	int32 BuildHLOD_MeshProxy(int32 iLevel, UWorldPartition* InWorldPartition, FName InCellName, const TArray<AActor*>& InSubActors) const;
 #endif
 
 
