Reverted accidental change to FoundationSubsystem

diff --git a/Engine/Source/Runtime/Engine/Private/Foundation/FoundationSubsystem.cpp b/Engine/Source/Runtime/Engine/Private/Foundation/FoundationSubsystem.cpp
--- a/Engine/Source/Runtime/Engine/Private/Foundation/FoundationSubsystem.cpp
+++ b/Engine/Source/Runtime/Engine/Private/Foundation/FoundationSubsystem.cpp
@@ -684,7 +684,7 @@ UFoundationSubsystem::FLevelsToRemoveScope::~FLevelsToRemoveScope()
 		double StartTime = FPlatformTime::Seconds();
 		const bool bClearSelection = false;
 		// No need to clear the whole editor selection since actor of this level will be removed from the selection by: UEditorEngine::OnLevelRemovedFromWorld
-		EditorLevelUtils::RemoveLevelsFromWorld(Levels, bClearSelection, true);
+		EditorLevelUtils::RemoveLevelsFromWorld(Levels, bClearSelection, bResetTrans);
 		double ElapsedTime = FPlatformTime::Seconds() - StartTime;
 		UE_LOG(LogFoundation, Log, TEXT("Unloaded %s levels in %s seconds"), *FText::AsNumber(Levels.Num()).ToString(), *FText::AsNumber(ElapsedTime).ToString());
 	}
