Fix build error

diff --git a/Engine/Source/Runtime/Engine/Private/WorldPartition/HLOD/HLODBuilder.cpp b/Engine/Source/Runtime/Engine/Private/WorldPartition/HLOD/HLODBuilder.cpp
--- a/Engine/Source/Runtime/Engine/Private/WorldPartition/HLOD/HLODBuilder.cpp
+++ b/Engine/Source/Runtime/Engine/Private/WorldPartition/HLOD/HLODBuilder.cpp
@@ -1,5 +1,7 @@
 #include "WorldPartition/HLOD/HLODBuilder.h"
 
+#if WITH_EDITOR
+
 #include "WorldPartition/WorldPartition.h"
 #include "WorldPartition/HLOD/HLODActor.h"
 #include "WorldPartition/HLOD/HLODLayer.h"
@@ -259,3 +261,5 @@ void FHLODBuilderUtilities::BuildHLODs(UWorldPartition* InWorldPartition, FName
 		}
 	}
 }
+
+#endif
diff --git a/Engine/Source/Runtime/Engine/Private/WorldPartition/HLOD/HLODBuilder.h b/Engine/Source/Runtime/Engine/Private/WorldPartition/HLOD/HLODBuilder.h
--- a/Engine/Source/Runtime/Engine/Private/WorldPartition/HLOD/HLODBuilder.h
+++ b/Engine/Source/Runtime/Engine/Private/WorldPartition/HLOD/HLODBuilder.h
@@ -2,6 +2,8 @@
 
 #pragma once
 
+#if WITH_EDITOR
+
 #include "CoreMinimal.h"
 #include "Containers/Array.h"
 
@@ -26,3 +28,5 @@ public:
 	 */
 	static void BuildHLODs(UWorldPartition* InWorldPartition, FName InCellName, const UHLODLayer* InHLODLayer, const TArray<AActor*>& InSubActors);
 };
+
+#endif
\ No newline at end of file
