Fix lexical error on Mac.

#rnx

diff --git a/Engine/Source/Runtime/Engine/Public/ComputeFramework/ComputeKernel.h b/Engine/Source/Runtime/Engine/Public/ComputeFramework/ComputeKernel.h
--- a/Engine/Source/Runtime/Engine/Public/ComputeFramework/ComputeKernel.h
+++ b/Engine/Source/Runtime/Engine/Public/ComputeFramework/ComputeKernel.h
@@ -5,7 +5,7 @@
 #include "CoreMinimal.h"
 #include "ComputeFramework/ShaderParamTypeDefinition.h"
 #include "ComputeFramework/ComputeKernelPermutationSet.h"
-#include "ComputeFramework\ComputeKernelResource.h"
+#include "ComputeFramework/ComputeKernelResource.h"
 
 #include "ComputeKernel.generated.h"
 
