Fixed incorrect call to RHIEndMultiUpdateTexture3D with no updates

diff --git a/Engine/Source/Runtime/Engine/Private/DistanceFieldAtlas.cpp b/Engine/Source/Runtime/Engine/Private/DistanceFieldAtlas.cpp
--- a/Engine/Source/Runtime/Engine/Private/DistanceFieldAtlas.cpp
+++ b/Engine/Source/Runtime/Engine/Private/DistanceFieldAtlas.cpp
@@ -782,7 +782,7 @@ void FDistanceFieldVolumeTextureAtlas::UpdateAllocations(FRHICommandListImmediat
 
 				}, !GDistanceFieldParallelAtlasUpdate, false);
 
-			if (!bRuntimeDownsampling)
+			if (!bRuntimeDownsampling && UpdateDataArray.Num() > 0)
 			{
 				// For some RHIs, this has the advantage of reducing transition barriers
 				RHIEndMultiUpdateTexture3D(UpdateDataArray);
