Added a setting to WorldSettings alowing user to configure which AISystemBase subclass is to be used with a given level.

diff --git a/Engine/Source/Runtime/Engine/Classes/GameFramework/WorldSettings.h b/Engine/Source/Runtime/Engine/Classes/GameFramework/WorldSettings.h
--- a/Engine/Source/Runtime/Engine/Classes/GameFramework/WorldSettings.h
+++ b/Engine/Source/Runtime/Engine/Classes/GameFramework/WorldSettings.h
@@ -19,6 +19,7 @@ class UAssetUserData;
 class UNetConnection;
 class UNavigationSystemConfig;
 class UWorldPartitionStreamingPolicy;
+class UAISystemBase;
 
 UENUM()
 enum EVisibilityAggressiveness
@@ -454,11 +455,14 @@ protected:
 	UPROPERTY(BlueprintReadOnly, config, Category = World, meta=(DisplayName = "DEPRECATED_bEnableNavigationSystem"))
 	uint8 bEnableNavigationSystem:1;
 
-public:
 	/** if set to false AI system will not get created. Use it to disable all AI-related activity on a map */
 	UPROPERTY(EditAnywhere, BlueprintReadOnly, config, Category=AI, AdvancedDisplay)
 	uint8 bEnableAISystem:1;
 
+public:
+	UPROPERTY(EditAnywhere, BlueprintReadOnly, Category=AI, meta=(MetaClass="AISystemBase", editcondition="bEnableAISystem"), AdvancedDisplay)
+	TSoftClassPtr<UAISystemBase> AISystemClass;
+
 	/** 
 	 * Enables tools for composing a tiled world. 
 	 * Level has to be saved and all sub-levels removed before enabling this option.
@@ -792,6 +796,9 @@ public:
 	/** @return whether given world is configured to host any NavigationSystem */
 	bool IsNavigationSystemEnabled() const;
 
+	/** @return whether given world is configured to host an AISystem */
+	bool IsAISystemEnabled() const { return bEnableAISystem; }
+
 	/**
 	 * Called from GameStateBase, calls BeginPlay on all actors
 	 */
diff --git a/Engine/Source/Runtime/Engine/Private/World.cpp b/Engine/Source/Runtime/Engine/Private/World.cpp
--- a/Engine/Source/Runtime/Engine/Private/World.cpp
+++ b/Engine/Source/Runtime/Engine/Private/World.cpp
@@ -1273,7 +1273,7 @@ UAISystemBase* UWorld::CreateAISystem()
 	{
 		const FName AIModuleName = UAISystemBase::GetAISystemModuleName();
 		const AWorldSettings* WorldSettings = PersistentLevel->GetWorldSettings(false);
-		if (AIModuleName.IsNone() == false && WorldSettings && WorldSettings->bEnableAISystem)
+		if (AIModuleName.IsNone() == false && WorldSettings && WorldSettings->IsAISystemEnabled())
 		{
 			IAISystemModule* AISystemModule = FModuleManager::LoadModulePtr<IAISystemModule>(AIModuleName);
 			if (AISystemModule)
@@ -1475,7 +1475,7 @@ void UWorld::InitWorld(const InitializationValues IVS)
 			{
 				FNavigationSystem::AddNavigationSystemToWorld(*this, FNavigationSystemRunMode::InvalidMode, WorldSettings->GetNavigationSystemConfig(), /*bInitializeForWorld=*/false);
 			}
-			if (IVS.bCreateAISystem && WorldSettings->bEnableAISystem)
+			if (IVS.bCreateAISystem && WorldSettings->IsAISystemEnabled())
 			{
 				CreateAISystem();
 			}
diff --git a/Engine/Source/Runtime/Engine/Private/WorldSettings.cpp b/Engine/Source/Runtime/Engine/Private/WorldSettings.cpp
--- a/Engine/Source/Runtime/Engine/Private/WorldSettings.cpp
+++ b/Engine/Source/Runtime/Engine/Private/WorldSettings.cpp
@@ -77,6 +77,7 @@ PRAGMA_DISABLE_DEPRECATION_WARNINGS
 PRAGMA_ENABLE_DEPRECATION_WARNINGS
 	NavigationSystemConfig = nullptr;
 	bEnableAISystem = true;
+	AISystemClass = UAISystemBase::GetAISystemClassName();
 	bEnableWorldComposition = false;
 	bEnableWorldOriginRebasing = false;
 #if WITH_EDITORONLY_DATA
@@ -644,6 +645,8 @@ bool AWorldSettings::CanEditChange(const FProperty* InProperty) const
 
 void AWorldSettings::PostEditChangeProperty(FPropertyChangedEvent& PropertyChangedEvent)
 {
+	UWorld* World = GetWorld();
+
 	FProperty* PropertyThatChanged = PropertyChangedEvent.Property;
 	if (PropertyThatChanged)
 	{
@@ -667,18 +670,28 @@ void AWorldSettings::PostEditChangeProperty(FPropertyChangedEvent& PropertyChang
 	LightmassSettings.EnvironmentIntensity = FMath::Max(LightmassSettings.EnvironmentIntensity, 0.0f);
 
 #if WITH_EDITOR
-	if (PropertyChangedEvent.GetPropertyName() == GET_MEMBER_NAME_CHECKED(FLightmassWorldInfoSettings, VolumetricLightmapDetailCellSize))
+	const FName PropName = PropertyChangedEvent.GetPropertyName();
+	if (PropName == GET_MEMBER_NAME_CHECKED(FLightmassWorldInfoSettings, VolumetricLightmapDetailCellSize))
 	{
 		FStaticLightingSystemInterface::OnLightmassImportanceVolumeModified.Broadcast();
 	}
+	else if (PropName == GET_MEMBER_NAME_CHECKED(AWorldSettings, bEnableAISystem)
+		|| PropName == GET_MEMBER_NAME_CHECKED(AWorldSettings, AISystemClass))
+	{
+		if (World)
+		{
+			World->CreateAISystem();
+		}
+	}
+
 #endif
 
 	// Ensure texture size is power of two between 512 and 4096.
 	PackedLightAndShadowMapTextureSize = FMath::Clamp<uint32>( FMath::RoundUpToPowerOfTwo( PackedLightAndShadowMapTextureSize ), 512, 4096 );
 
-	if (PropertyThatChanged != nullptr && GetWorld() != nullptr && GetWorld()->Scene)
+	if (PropertyThatChanged != nullptr && World != nullptr && World->Scene)
 	{
-		GetWorld()->Scene->UpdateSceneSettings(this);
+		World->Scene->UpdateSceneSettings(this);
 	}
 
 	for (UAssetUserData* Datum : AssetUserData)
@@ -958,7 +971,7 @@ void AWorldSettings::UpdateBookmarkClass()
 
 FSoftClassPath AWorldSettings::GetAISystemClassName() const
 {
-	return bEnableAISystem ? UAISystemBase::GetAISystemClassName() : FSoftClassPath();
+	return bEnableAISystem ? AISystemClass.Get() : FSoftClassPath();
 }
 
 void AWorldSettings::RewindForReplay()
