Added missing includes to HLODBuilder. And the copyrights line.

diff --git a/Engine/Source/Runtime/Engine/Private/WorldPartition/HLOD/HLODBuilder.cpp b/Engine/Source/Runtime/Engine/Private/WorldPartition/HLOD/HLODBuilder.cpp
--- a/Engine/Source/Runtime/Engine/Private/WorldPartition/HLOD/HLODBuilder.cpp
+++ b/Engine/Source/Runtime/Engine/Private/WorldPartition/HLOD/HLODBuilder.cpp
@@ -1,3 +1,5 @@
+// Copyright Epic Games, Inc. All Rights Reserved.
+
 #include "WorldPartition/HLOD/HLODBuilder.h"
 
 #if WITH_EDITOR
@@ -12,6 +14,9 @@
 
 #include "IMeshMergeUtilities.h"
 #include "MeshMergeModule.h"
+#include "Modules/ModuleManager.h"
+#include "Engine/StaticMesh.h"
+#include "Materials/Material.h"
 
 /**
  * Base class for all HLODBuilders
