Fix missing material updates in async texture compilation due to a change in behavior from CL 13480995

diff --git a/Engine/Source/Runtime/Engine/Private/TextureCompiler.cpp b/Engine/Source/Runtime/Engine/Private/TextureCompiler.cpp
--- a/Engine/Source/Runtime/Engine/Private/TextureCompiler.cpp
+++ b/Engine/Source/Runtime/Engine/Private/TextureCompiler.cpp
@@ -15,6 +15,7 @@
 #include "UObject/StrongObjectPtr.h"
 #include "UObject/UObjectIterator.h"
 #include "Materials/MaterialInterface.h"
+#include "Materials/MaterialInstance.h"
 #include "Materials/Material.h"
 #include "TextureDerivedDataTask.h"
 #include "Misc/IQueuedWork.h"
@@ -66,16 +67,27 @@ namespace TextureCompilingManagerImpl
 		// Update any material that uses this texture
 		TMultiMap<UObject*, UMaterialInterface*> TexturesRequiringMaterialUpdate;
 
-		TArray<UTexture*> UsedTextures;
 		for (TObjectIterator<UMaterialInterface> It; It; ++It)
 		{
-			UsedTextures.Reset();
-
 			UMaterialInterface* MaterialInterface = *It;
 			for (UObject* Texture : MaterialInterface->GetReferencedTextures())
 			{
 				TexturesRequiringMaterialUpdate.Emplace(Texture, MaterialInterface);
 			}
+			
+			// Fix in CL 13480995 broke GetReferencedTextures() ability to return all referenced textures
+			// so we manually gather them instead...
+			UMaterialInstance* MaterialInstance = Cast<UMaterialInstance>(MaterialInterface);
+			if (MaterialInstance)
+			{
+				for (const FTextureParameterValue& TextureParam : MaterialInstance->TextureParameterValues)
+				{
+					if (TextureParam.ParameterValue)
+					{
+						TexturesRequiringMaterialUpdate.Emplace(TextureParam.ParameterValue, MaterialInterface);
+					}
+				}
+			}
 		}
 
 		return TexturesRequiringMaterialUpdate;
