Added a way for billboard components to set the opacity clip mask value directly.

Also improve billboard rendering quality a bit by fixing incorrect alpha values.

diff --git a/Engine/Source/Runtime/Engine/Classes/Components/BillboardComponent.h b/Engine/Source/Runtime/Engine/Classes/Components/BillboardComponent.h
--- a/Engine/Source/Runtime/Engine/Classes/Components/BillboardComponent.h
+++ b/Engine/Source/Runtime/Engine/Classes/Components/BillboardComponent.h
@@ -41,6 +41,10 @@ class ENGINE_API UBillboardComponent : public UPrimitiveComponent
 	UPROPERTY(EditAnywhere, BlueprintReadWrite, Category=Sprite)
 	float VL;
 
+	/** The billboard is not rendered where texture opacity < OpacityMaskRefVal */
+	UPROPERTY(EditAnywhere, BlueprintReadWrite, AdvancedDisplay, Category=Sprite)
+	float OpacityMaskRefVal;
+
 #if WITH_EDITORONLY_DATA
 	/** Sprite category that the component belongs to. Value serves as a key into the localization file. */
 	UPROPERTY()
@@ -66,6 +70,9 @@ class ENGINE_API UBillboardComponent : public UPrimitiveComponent
 	UFUNCTION(BlueprintCallable, Category="Rendering|Components|Sprite")
 	virtual void SetSpriteAndUV(class UTexture2D* NewSprite, int32 NewU, int32 NewUL, int32 NewV, int32 NewVL);
 
+	/** Changed the opacity masked used by this component */
+	UFUNCTION(BlueprintCallable, Category = "Rendering|Components|Sprite")
+	void SetOpacityMaskRefVal(float RefVal);
 
 	//~ Begin UPrimitiveComponent Interface
 	virtual FPrimitiveSceneProxy* CreateSceneProxy() override;
diff --git a/Engine/Source/Runtime/Engine/Private/BatchedElements.cpp b/Engine/Source/Runtime/Engine/Private/BatchedElements.cpp
--- a/Engine/Source/Runtime/Engine/Private/BatchedElements.cpp
+++ b/Engine/Source/Runtime/Engine/Private/BatchedElements.cpp
@@ -312,7 +312,8 @@ void FBatchedElements::AddSprite(
 	float UL,
 	float V,
 	float VL,
-	uint8 BlendMode
+	uint8 BlendMode,
+	float OpacityMaskRefVal
 	)
 {
 	check(Texture);
@@ -328,6 +329,7 @@ void FBatchedElements::AddSprite(
 	Sprite->UL = UL == 0.f ? Texture->GetSizeX() : UL;
 	Sprite->V = V;
 	Sprite->VL = VL == 0.f ? Texture->GetSizeY() : VL;
+	Sprite->OpacityMaskRefVal = OpacityMaskRefVal;
 	Sprite->BlendMode = BlendMode;
 }
 
@@ -430,11 +432,10 @@ static void SetHitTestingBlendState(FRHICommandList& RHICmdList, FGraphicsPipeli
 	}
 }
 
-/** Global alpha ref test value for rendering masked batched elements */
-float GBatchedElementAlphaRefVal = 128.f;
 /** Global smoothing width for rendering batched elements with distance field blend modes */
 float GBatchedElementSmoothWidth = 4;
 
+
 FAutoConsoleVariableRef CVarWellCanvasDistanceFieldSmoothWidth(TEXT("Canvas.DistanceFieldSmoothness"), GBatchedElementSmoothWidth, TEXT("Global sharpness of distance field fonts/shapes rendered by canvas."), ECVF_Default);
 
 template<class TSimpleElementPixelShader>
@@ -458,7 +459,8 @@ void FBatchedElements::PrepareShaders(
 	bool bHitTesting,
 	float Gamma,
 	const FDepthFieldGlowInfo* GlowInfo,
-	const FSceneView* View
+	const FSceneView* View,
+	float OpacityMaskRefVal
 	) const
 {
 	// used to mask individual channels and desaturate
@@ -568,7 +570,7 @@ void FBatchedElements::PrepareShaders(
 					SetGraphicsPipelineState(RHICmdList, GraphicsPSOInit);
 
 					MaskedPixelShader->SetEditorCompositingParameters(RHICmdList, View);
-					MaskedPixelShader->SetParameters(RHICmdList, Texture, Gamma, GBatchedElementAlphaRefVal / 255.0f, BlendMode);
+					MaskedPixelShader->SetParameters(RHICmdList, Texture, Gamma, OpacityMaskRefVal, BlendMode);
 				}
 				else
 				{
@@ -578,7 +580,7 @@ void FBatchedElements::PrepareShaders(
 					SetGraphicsPipelineState(RHICmdList, GraphicsPSOInit);
 
 					MaskedPixelShader->SetEditorCompositingParameters(RHICmdList, View);
-					MaskedPixelShader->SetParameters(RHICmdList, Texture, Gamma, GBatchedElementAlphaRefVal / 255.0f, BlendMode);
+					MaskedPixelShader->SetParameters(RHICmdList, Texture, Gamma, OpacityMaskRefVal, BlendMode);
 				}
 			}
 			// render distance field elements
@@ -588,7 +590,7 @@ void FBatchedElements::PrepareShaders(
 				BlendMode == SE_BLEND_TranslucentDistanceField	||
 				BlendMode == SE_BLEND_TranslucentDistanceFieldShadowed)
 			{
-				float AlphaRefVal = GBatchedElementAlphaRefVal;
+				float AlphaRefVal = OpacityMaskRefVal;
 				if (BlendMode == SE_BLEND_TranslucentDistanceField ||
 					BlendMode == SE_BLEND_TranslucentDistanceFieldShadowed)
 				{
@@ -621,7 +623,7 @@ void FBatchedElements::PrepareShaders(
 					RHICmdList, 
 					Texture,
 					Gamma,
-					AlphaRefVal / 255.0f,
+					AlphaRefVal,
 					GBatchedElementSmoothWidth,
 					EnableShadow,
 					ShadowDirection,
@@ -1102,6 +1104,7 @@ bool FBatchedElements::Draw(FRHICommandList& RHICmdList, const FMeshPassProcesso
 				const FTexture* CurrentTexture = Sprites[0].Texture;
 				ESimpleElementBlendMode CurrentBlendMode = (ESimpleElementBlendMode)Sprites[0].BlendMode;
 				int32 BatchStartIndex = 0;
+				float CurrentOpacityMask = Sprites[0].OpacityMaskRefVal;
 
 				// Start loop at 1, since we've already started the first batch with the first sprite in the list
 				for (int32 SpriteIndex = 1; SpriteIndex < ValidSpriteCount + 1; SpriteIndex++)
@@ -1109,12 +1112,13 @@ bool FBatchedElements::Draw(FRHICommandList& RHICmdList, const FMeshPassProcesso
 					// Need to flush the current batch once we hit the end of the list, or if state of this sprite doesn't match current batch
 					if (SpriteIndex == ValidSpriteCount ||
 						CurrentTexture != Sprites[SpriteIndex].Texture ||
-						CurrentBlendMode != Sprites[SpriteIndex].BlendMode)
+						CurrentBlendMode != Sprites[SpriteIndex].BlendMode ||
+						CurrentOpacityMask != Sprites[SpriteIndex].OpacityMaskRefVal)
 					{
 						const int32 SpriteNum = SpriteIndex - BatchStartIndex;
 						const int32 BaseVertex = BatchStartIndex * 6;
 						const int32 PrimCount = SpriteNum * 2;
-						PrepareShaders(RHICmdList, GraphicsPSOInit, FeatureLevel, CurrentBlendMode, Transform, bNeedToSwitchVerticalAxis, BatchedElementParameters, CurrentTexture, bHitTesting, Gamma, NULL, &View);
+						PrepareShaders(RHICmdList, GraphicsPSOInit, FeatureLevel, CurrentBlendMode, Transform, bNeedToSwitchVerticalAxis, BatchedElementParameters, CurrentTexture, bHitTesting, Gamma, NULL, &View, CurrentOpacityMask);
 						RHICmdList.SetStencilRef(StencilRef);
 						RHICmdList.SetStreamSource(0, VertexBufferRHI, 0);
 						RHICmdList.DrawPrimitive(BaseVertex, PrimCount, 1);
@@ -1125,6 +1129,7 @@ bool FBatchedElements::Draw(FRHICommandList& RHICmdList, const FMeshPassProcesso
 							BatchStartIndex = SpriteIndex;
 							CurrentTexture = Sprites[SpriteIndex].Texture;
 							CurrentBlendMode = (ESimpleElementBlendMode)Sprites[SpriteIndex].BlendMode;
+							CurrentOpacityMask = Sprites[SpriteIndex].OpacityMaskRefVal;
 						}
 					}
 				}
diff --git a/Engine/Source/Runtime/Engine/Private/Components/BillboardComponent.cpp b/Engine/Source/Runtime/Engine/Private/Components/BillboardComponent.cpp
--- a/Engine/Source/Runtime/Engine/Private/Components/BillboardComponent.cpp
+++ b/Engine/Source/Runtime/Engine/Private/Components/BillboardComponent.cpp
@@ -61,6 +61,8 @@ public:
 		// Calculate the scale factor for the sprite.
 		float Scale = InComponent->GetComponentTransform().GetMaximumAxisScale();
 
+		OpacityMaskRefVal = InComponent->OpacityMaskRefVal;
+
 		if(InComponent->Sprite)
 		{
 			Texture = InComponent->Sprite;
@@ -188,9 +190,12 @@ public:
 					FLinearColor LevelColorToUse = IsSelected() ? ColorToUse : (FLinearColor)GetLevelColor();
 					FLinearColor PropertyColorToUse = GetPropertyColor();
 
+					ColorToUse.A = 1.0f;
+
 					const FLinearColor& SpriteColor = View->Family->EngineShowFlags.LevelColoration ? LevelColorToUse :
 						( (View->Family->EngineShowFlags.PropertyColoration) ? PropertyColorToUse : ColorToUse );
 
+
 					Collector.GetPDI(ViewIndex)->DrawSprite(
 						Origin,
 						ViewedSizeX,
@@ -198,7 +203,9 @@ public:
 						TextureResource,
 						SpriteColor,
 						GetDepthPriorityGroup(View),
-						U,UL,V,VL
+						U,UL,V,VL,
+						SE_BLEND_Masked,
+						OpacityMaskRefVal
 						);
 				}
 			}
@@ -251,6 +258,7 @@ private:
 	float UL;
 	const float V;
 	float VL;
+	float OpacityMaskRefVal;
 	FLinearColor Color;
 	const uint32 bIsScreenSizeScaled : 1;
 	uint32 bIsActorLocked : 1;
@@ -292,6 +300,7 @@ UBillboardComponent::UBillboardComponent(const FObjectInitializer& ObjectInitial
 	V = 0;
 	UL = 0;
 	VL = 0;
+	OpacityMaskRefVal = .5f;
 	bHiddenInGame = true;
 	SetGenerateOverlapEvents(false);
 	bUseEditorCompositing = true;
@@ -393,6 +402,12 @@ void UBillboardComponent::SetSpriteAndUV(UTexture2D* NewSprite, int32 NewU, int3
 	SetSprite(NewSprite);
 }
 
+void UBillboardComponent::SetOpacityMaskRefVal(float RefVal)
+{
+	OpacityMaskRefVal = RefVal;
+	MarkRenderStateDirty();
+}
+
 #if WITH_EDITORONLY_DATA
 void UBillboardComponent::SetEditorScale(float InEditorScale)
 {
diff --git a/Engine/Source/Runtime/Engine/Private/SceneManagement.cpp b/Engine/Source/Runtime/Engine/Private/SceneManagement.cpp
--- a/Engine/Source/Runtime/Engine/Private/SceneManagement.cpp
+++ b/Engine/Source/Runtime/Engine/Private/SceneManagement.cpp
@@ -101,7 +101,8 @@ void FSimpleElementCollector::DrawSprite(
 	float UL,
 	float V,
 	float VL,
-	uint8 BlendMode
+	uint8 BlendMode,
+	float OpacityMaskRefVal
 	)
 {
 	FBatchedElements& Elements = DepthPriorityGroup == SDPG_World ? BatchedElements : TopBatchedElements;
@@ -117,7 +118,8 @@ void FSimpleElementCollector::DrawSprite(
 		UL,
 		V,
 		VL,
-		BlendMode
+		BlendMode,
+		OpacityMaskRefVal
 		);
 }
 
diff --git a/Engine/Source/Runtime/Engine/Public/BatchedElements.h b/Engine/Source/Runtime/Engine/Public/BatchedElements.h
--- a/Engine/Source/Runtime/Engine/Public/BatchedElements.h
+++ b/Engine/Source/Runtime/Engine/Public/BatchedElements.h
@@ -190,7 +190,8 @@ public:
 		float UL,
 		float V,
 		float VL,
-		uint8 BlendMode = SE_BLEND_Masked
+		uint8 BlendMode = SE_BLEND_Masked,
+		float OpacityMaskRefVal = .5f
 		);
 
 	/** 
@@ -304,6 +305,7 @@ private:
 		float UL;
 		float V;
 		float VL;
+		float OpacityMaskRefVal;
 		uint8 BlendMode;
 	};
 	/** This array is sorted during draw-calls */
@@ -370,8 +372,9 @@ private:
 		const FTexture* Texture,
 		bool bHitTesting,
 		float Gamma,
-		const FDepthFieldGlowInfo* GlowInfo = NULL,
-		const FSceneView* View = NULL
+		const FDepthFieldGlowInfo* GlowInfo = nullptr,
+		const FSceneView* View = nullptr,
+		float OpacityMaskRefVal = .5f
 		) const;
 
 	/** if false then prevent the use of HDR encoded shaders. */
diff --git a/Engine/Source/Runtime/Engine/Public/SceneManagement.h b/Engine/Source/Runtime/Engine/Public/SceneManagement.h
--- a/Engine/Source/Runtime/Engine/Public/SceneManagement.h
+++ b/Engine/Source/Runtime/Engine/Public/SceneManagement.h
@@ -1866,7 +1866,8 @@ public:
 		float UL,
 		float V,
 		float VL,
-		uint8 BlendMode = 1 /*SE_BLEND_Masked*/
+		uint8 BlendMode = 1, /*SE_BLEND_Masked*/
+		float OpacityMaskRefVal = .5f
 		) = 0;
 
 	virtual void DrawLine(
@@ -1972,8 +1973,9 @@ public:
 		float UL,
 		float V,
 		float VL,
-		uint8 BlendMode = SE_BLEND_Masked
-		) override;
+		uint8 BlendMode = SE_BLEND_Masked,
+		float OpacityMaskRevVal = .5f
+	) override;
 
 	virtual void DrawLine(
 		const FVector& Start,
diff --git a/Engine/Source/Runtime/Renderer/Private/DynamicPrimitiveDrawing.h b/Engine/Source/Runtime/Renderer/Private/DynamicPrimitiveDrawing.h
--- a/Engine/Source/Runtime/Renderer/Private/DynamicPrimitiveDrawing.h
+++ b/Engine/Source/Runtime/Renderer/Private/DynamicPrimitiveDrawing.h
@@ -33,8 +33,10 @@ public:
 		float UL,
 		float V,
 		float VL,
-		uint8 BlendMode = SE_BLEND_Masked
+		uint8 BlendMode = SE_BLEND_Masked,
+		float OpacityMaskRefVal = .5f
 		) override;
+
 	virtual void DrawLine(
 		const FVector& Start,
 		const FVector& End,
diff --git a/Engine/Source/Runtime/Renderer/Private/DynamicPrimitiveDrawing.inl b/Engine/Source/Runtime/Renderer/Private/DynamicPrimitiveDrawing.inl
--- a/Engine/Source/Runtime/Renderer/Private/DynamicPrimitiveDrawing.inl
+++ b/Engine/Source/Runtime/Renderer/Private/DynamicPrimitiveDrawing.inl
@@ -86,7 +86,8 @@ inline void FViewElementPDI::DrawSprite(
 	float UL,
 	float V,
 	float VL,
-	uint8 BlendMode
+	uint8 BlendMode,
+	float OpacityMaskRefVal
 	)
 {
 	FBatchedElements &Elements = GetElements(DepthPriorityGroup);
@@ -98,8 +99,9 @@ inline void FViewElementPDI::DrawSprite(
 		Sprite,
 		Color,
 		CurrentHitProxy ? CurrentHitProxy->Id : FHitProxyId(),
-		U,UL,V,VL,
-		BlendMode
+		U, UL, V, VL,
+		BlendMode,
+		OpacityMaskRefVal
 	);
 }
 
