WorldPartition : When the world is partitioned, we don't want to duplicate for PIE any external actors as they will be streamed by the runtime grids.

diff --git a/Engine/Source/Runtime/Engine/Private/Level.cpp b/Engine/Source/Runtime/Engine/Private/Level.cpp
--- a/Engine/Source/Runtime/Engine/Private/Level.cpp
+++ b/Engine/Source/Runtime/Engine/Private/Level.cpp
@@ -432,10 +432,15 @@ void ULevel::Serialize( FArchive& Ar )
 			}
 
 #if WITH_EDITOR
-			// Otherwise, don't filter out external actors if duplicating the world to get the actors properly duplicated.
-			if (IsUsingExternalActors() && !(Ar.GetPortFlags() & PPF_Duplicate))
+			if (IsUsingExternalActors() && Actor->IsPackageExternal())
 			{
-				if (Actor->IsPackageExternal())
+				// Don't filter out external actors if duplicating the world to get the actors properly duplicated.
+				if (!(Ar.GetPortFlags() & PPF_Duplicate))
+				{
+					return false;
+				}
+				// When the world is partitioned and duplicating for PIE, we don't want to duplicate any external actors as they will be streamed by the runtime grids.
+				else if (bIsPartitioned && (Ar.GetPortFlags() & PPF_DuplicateForPIE))
 				{
 					return false;
 				}
