Fixed UWorldPartitionRuntimeSpatialHash::GetStreamingCells to return always loaded cells even when no streaming sources are provided.

diff --git a/Engine/Source/Runtime/Engine/Private/WorldPartition/WorldPartitionRuntimeSpatialHash.cpp b/Engine/Source/Runtime/Engine/Private/WorldPartition/WorldPartitionRuntimeSpatialHash.cpp
--- a/Engine/Source/Runtime/Engine/Private/WorldPartition/WorldPartitionRuntimeSpatialHash.cpp
+++ b/Engine/Source/Runtime/Engine/Private/WorldPartition/WorldPartitionRuntimeSpatialHash.cpp
@@ -813,11 +813,26 @@ int32 UWorldPartitionRuntimeSpatialHash::GetAllStreamingCells(TSet<const UWorldP
 	return Cells.Num();
 }
 
-int32 UWorldPartitionRuntimeSpatialHash::GetStreamingCells(const FVector& Position, TSet<const UWorldPartitionRuntimeCell*>& Cells) const
+int32 UWorldPartitionRuntimeSpatialHash::GetStreamingCells(const TArray<FWorldPartitionStreamingSource>& Sources, TSet<const UWorldPartitionRuntimeCell*>& Cells) const
 {
-	for (const FSpatialHashStreamingGrid& StreamingGrid: StreamingGrids)
+	if (Sources.Num() == 0)
 	{
-		GetStreamingCells(Position, StreamingGrid, Cells);
+		// Get always loaded cells
+		for (const FSpatialHashStreamingGrid& StreamingGrid : StreamingGrids)
+		{
+			GetAlwaysLoadedStreamingCells(StreamingGrid, Cells);
+		}
+	}
+	else
+	{
+		// Get cells based on streaming sources
+		for (const FSpatialHashStreamingGrid& StreamingGrid : StreamingGrids)
+		{
+			for (const FWorldPartitionStreamingSource& Source : Sources)
+			{
+				GetStreamingCells(Source.Location, StreamingGrid, Cells);
+			}
+		}
 	}
 
 	return Cells.Num();
@@ -837,6 +852,19 @@ void UWorldPartitionRuntimeSpatialHash::GetStreamingCells(const FVector& Positio
 	});
 }
 
+void UWorldPartitionRuntimeSpatialHash::GetAlwaysLoadedStreamingCells(const FSpatialHashStreamingGrid& StreamingGrid, TSet<const UWorldPartitionRuntimeCell*>& Cells) const
+{
+	if (StreamingGrid.GridLevels.Num() > 0)
+	{
+		const int32 TopLevel = StreamingGrid.GridLevels.Num() - 1;
+		check(StreamingGrid.GridLevels[TopLevel].GridCells.Num() == 1);
+		if (UWorldPartitionRuntimeSpatialHashCell* Cell = StreamingGrid.GridLevels[TopLevel].GridCells[0])
+		{
+			Cells.Add(Cell);
+		}
+	}
+}
+
 void UWorldPartitionRuntimeSpatialHash::SortStreamingCellsByDistance(const TSet<const UWorldPartitionRuntimeCell*>& InCells, const TArray<FWorldPartitionStreamingSource>& InSources, TArray<const UWorldPartitionRuntimeCell*>& OutSortedCells)
 {
 	struct FCellShortestDist
diff --git a/Engine/Source/Runtime/Engine/Private/WorldPartition/WorldPartitionStreamingPolicy.cpp b/Engine/Source/Runtime/Engine/Private/WorldPartition/WorldPartitionStreamingPolicy.cpp
--- a/Engine/Source/Runtime/Engine/Private/WorldPartition/WorldPartitionStreamingPolicy.cpp
+++ b/Engine/Source/Runtime/Engine/Private/WorldPartition/WorldPartitionStreamingPolicy.cpp
@@ -99,12 +99,13 @@ void UWorldPartitionStreamingPolicy::UpdateStreamingState()
 			return;
 		}
 
-		// Get cells to load/unload
-		for (const FWorldPartitionStreamingSource& Source : StreamingSources)
+		// When uninitializing, UpdateStreamingState is called, but we don't want any cells to be loaded
+		if (WorldPartition->IsInitialized())
 		{
-			WorldPartition->RuntimeHash->GetStreamingCells(Source.Location, StreamingCells);
+			WorldPartition->RuntimeHash->GetStreamingCells(StreamingSources, StreamingCells);
 		}
 
+		// Determine cells to load/unload
 		TSet<const UWorldPartitionRuntimeCell*> ToLoadCells = StreamingCells.Difference(LoadedCells);
 		TSet<const UWorldPartitionRuntimeCell*> ToUnloadCells = LoadedCells.Difference(StreamingCells);
 
diff --git a/Engine/Source/Runtime/Engine/Public/WorldPartition/WorldPartitionRuntimeHash.h b/Engine/Source/Runtime/Engine/Public/WorldPartition/WorldPartitionRuntimeHash.h
--- a/Engine/Source/Runtime/Engine/Public/WorldPartition/WorldPartitionRuntimeHash.h
+++ b/Engine/Source/Runtime/Engine/Public/WorldPartition/WorldPartitionRuntimeHash.h
@@ -24,7 +24,7 @@ class ENGINE_API UWorldPartitionRuntimeHash : public UObject
 
 	// Streaming interface
 	virtual int32 GetAllStreamingCells(TSet<const UWorldPartitionRuntimeCell*>& Cells) const { return 0; }
-	virtual int32 GetStreamingCells(const FVector& Position, TSet<const UWorldPartitionRuntimeCell*>& Cells) const { return 0; };
+	virtual int32 GetStreamingCells(const TArray<FWorldPartitionStreamingSource>& Sources, TSet<const UWorldPartitionRuntimeCell*>& Cells) const { return 0; };
 	virtual void SortStreamingCellsByDistance(const TSet<const UWorldPartitionRuntimeCell*>& InCells, const TArray<FWorldPartitionStreamingSource>& InSources, TArray<const UWorldPartitionRuntimeCell*>& OutSortedCells) {}
 
 	/* Returns desired footprint that ShowDebugInfo should take relative to given Canvas size (the value can exceed the given size).
diff --git a/Engine/Source/Runtime/Engine/Public/WorldPartition/WorldPartitionRuntimeSpatialHash.h b/Engine/Source/Runtime/Engine/Public/WorldPartition/WorldPartitionRuntimeSpatialHash.h
--- a/Engine/Source/Runtime/Engine/Public/WorldPartition/WorldPartitionRuntimeSpatialHash.h
+++ b/Engine/Source/Runtime/Engine/Public/WorldPartition/WorldPartitionRuntimeSpatialHash.h
@@ -147,7 +147,7 @@ public:
 
 	// streaming interface
 	virtual int32 GetAllStreamingCells(TSet<const UWorldPartitionRuntimeCell*>& Cells) const override;
-	virtual int32 GetStreamingCells(const FVector& Position, TSet<const UWorldPartitionRuntimeCell*>& Cells) const override;
+	virtual int32 GetStreamingCells(const TArray<FWorldPartitionStreamingSource>& Sources, TSet<const UWorldPartitionRuntimeCell*>& Cells) const override;
 	virtual void SortStreamingCellsByDistance(const TSet<const UWorldPartitionRuntimeCell*>& InCells, const TArray<FWorldPartitionStreamingSource>& InSources, TArray<const UWorldPartitionRuntimeCell*>& OutSortedCells) override;
 
 private:
@@ -181,6 +181,7 @@ private:
 	virtual void ShowDebugInfo(class UCanvas* Canvas, const TArray<FWorldPartitionStreamingSource>& Sources, const FVector2D& PartitionCanvasOffset, const FVector2D& PartitionCanvasSize) const override;
 
 private:
+	void GetAlwaysLoadedStreamingCells(const FSpatialHashStreamingGrid& StreamingGrid, TSet<const UWorldPartitionRuntimeCell*>& Cells) const;
 	void GetStreamingCells(const FVector& Position, const FSpatialHashStreamingGrid& StreamingGrid, TSet<const UWorldPartitionRuntimeCell*>& Cells) const;
 #if WITH_EDITOR
 	bool CreateStreamingGrid(const FSpatialHashRuntimeGrid& RuntimeGrid, const FSquare2DGridHelper& PartionedActors, EWorldPartitionStreamingMode Mode, UWorldPartitionStreamingPolicy* StreamingPolicy);
