Added support for creating Actor and Component elements

This implements the initial set of functionality needed to support the element list implementation of USelection.

#rnx

diff --git a/Engine/Source/Runtime/Engine/Classes/Components/ActorComponent.h b/Engine/Source/Runtime/Engine/Classes/Components/ActorComponent.h
--- a/Engine/Source/Runtime/Engine/Classes/Components/ActorComponent.h
+++ b/Engine/Source/Runtime/Engine/Classes/Components/ActorComponent.h
@@ -15,6 +15,7 @@
 #include "Interfaces/Interface_AssetUserData.h"
 #include "UObject/UObjectAnnotation.h"
 #include "UObject/StructOnScope.h"
+#include "TypedElementHandle.h"
 #include "ComponentInstanceDataCache.h"
 #include "ActorComponent.generated.h"
 
@@ -535,6 +536,13 @@ private:
 	/** Private version without inlining that does *not* check Dedicated server build flags (which should already have been done). */
 	ENetMode InternalGetNetMode() const;
 
+public:
+	TTypedElementOwner<struct FComponentElementData> CreateComponentElement() const;
+	void DestroyComponentElement(TTypedElementOwner<struct FComponentElementData>& InOutComponentElement) const;
+#if WITH_EDITOR
+	FTypedElementHandle AcquireEditorElementHandle(const bool bAllowCreate = true) const;
+#endif
+
 protected:
 	/** Return true if this component is in a state where it can be activated normally. */
 	virtual bool ShouldActivate() const;
diff --git a/Engine/Source/Runtime/Engine/Classes/GameFramework/Actor.h b/Engine/Source/Runtime/Engine/Classes/GameFramework/Actor.h
--- a/Engine/Source/Runtime/Engine/Classes/GameFramework/Actor.h
+++ b/Engine/Source/Runtime/Engine/Classes/GameFramework/Actor.h
@@ -15,6 +15,7 @@
 #include "ComponentInstanceDataCache.h"
 #include "Components/ChildActorComponent.h"
 #include "RenderCommandFence.h"
+#include "TypedElementHandle.h"
 #include "Misc/ITransaction.h"
 #include "Engine/Level.h"
 
@@ -3300,6 +3301,12 @@ public:
 	/** Returns the instance components array */
 	const TArray<UActorComponent*>& GetInstanceComponents() const;
 
+public:
+	TTypedElementOwner<struct FActorElementData> CreateActorElement() const;
+	void DestroyActorElement(TTypedElementOwner<struct FActorElementData>& InOutActorElement) const;
+#if WITH_EDITOR
+	FTypedElementHandle AcquireEditorElementHandle(const bool bAllowCreate = true) const;
+#endif
 
 	//~=============================================================================
 	// Navigation/AI related functions
diff --git a/Engine/Source/Runtime/Engine/Engine.Build.cs b/Engine/Source/Runtime/Engine/Engine.Build.cs
--- a/Engine/Source/Runtime/Engine/Engine.Build.cs
+++ b/Engine/Source/Runtime/Engine/Engine.Build.cs
@@ -97,6 +97,12 @@ public class Engine : ModuleRules
             }
 		);
 
+		PublicIncludePathModuleNames.AddRange(
+			new string[] {
+				"TypedElementFramework",
+			}
+		);
+
 		PrivateDependencyModuleNames.AddRange(
 			new string[] {
 				"AnimationCore",
@@ -105,6 +111,8 @@ public class Engine : ModuleRules
 				"Landscape",
 				"UMG",
 				"Projects",
+				"TypedElementFramework",
+				"TypedElementInterfaces",
 				"MaterialShaderQualitySettings",
 				"CinematicCamera",
 				"Analytics",
diff --git a/Engine/Source/Runtime/Engine/Private/Actor.cpp b/Engine/Source/Runtime/Engine/Private/Actor.cpp
--- a/Engine/Source/Runtime/Engine/Private/Actor.cpp
+++ b/Engine/Source/Runtime/Engine/Private/Actor.cpp
@@ -4,6 +4,9 @@
 #include "EngineDefines.h"
 #include "EngineStats.h"
 #include "EngineGlobals.h"
+#include "TypedElementRegistry.h"
+#include "TypedElementOwnerStore.h"
+#include "Elements/Actor/ActorElementData.h"
 #include "GameFramework/DamageType.h"
 #include "TimerManager.h"
 #include "GameFramework/Pawn.h"
@@ -64,6 +67,7 @@ FMakeNoiseDelegate AActor::MakeNoiseDelegate = FMakeNoiseDelegate::CreateStatic(
 
 #if WITH_EDITOR
 FUObjectAnnotationSparseBool GSelectedActorAnnotation;
+TTypedElementOwnerStore<FActorElementData, const AActor*> GActorElementOwnerStore;
 #endif
 
 #if !UE_BUILD_SHIPPING
@@ -220,6 +224,47 @@ void AActor::InitializeDefaults()
 	GridPlacement = EActorGridPlacement::Bounds;
 }
 
+TTypedElementOwner<FActorElementData> AActor::CreateActorElement() const
+{
+	UTypedElementRegistry* Registry = UTypedElementRegistry::GetInstance();
+
+	TTypedElementOwner<FActorElementData> ActorElement;
+	if (ensureMsgf(Registry, TEXT("Typed element was requested for '%s' before the registry was available! This usually means that NewObject was used instead of CreateDefaultSubobject during CDO construction."), *GetPathName()))
+	{
+		ActorElement = Registry->CreateElement<FActorElementData>(NAME_Actor);
+		ActorElement.GetDataChecked().Actor = const_cast<AActor*>(this);
+	}
+
+	return ActorElement;
+}
+
+void AActor::DestroyActorElement(TTypedElementOwner<FActorElementData>& InOutActorElement) const
+{
+	if (InOutActorElement)
+	{
+		checkf(InOutActorElement.GetDataChecked().Actor == this, TEXT("Actor element was not for this actor instance! %s"), *GetPathName());
+		UTypedElementRegistry::GetInstance()->DestroyElement(InOutActorElement);
+	}
+}
+
+#if WITH_EDITOR
+FTypedElementHandle AActor::AcquireEditorElementHandle(const bool bAllowCreate) const
+{
+	if (TTypedElementOwnerScopedAccess<FActorElementData> EditorElementScopedAccess = GActorElementOwnerStore.FindElementOwner(this))
+	{
+		return EditorElementScopedAccess->AcquireHandle();
+	}
+
+	if (bAllowCreate)
+	{
+		TTypedElementOwnerScopedAccess<FActorElementData> EditorElementScopedAccess = GActorElementOwnerStore.RegisterElementOwner(this, CreateActorElement());
+		return EditorElementScopedAccess->AcquireHandle();
+	}
+
+	return FTypedElementHandle();
+}
+#endif
+
 void FActorTickFunction::ExecuteTick(float DeltaTime, enum ELevelTick TickType, ENamedThreads::Type CurrentThread, const FGraphEventRef& MyCompletionGraphEvent)
 {
 	if (Target && !Target->IsPendingKillOrUnreachable())
@@ -372,6 +417,14 @@ void AActor::ResetOwnedComponents()
 void AActor::PostInitProperties()
 {
 	Super::PostInitProperties();
+
+#if WITH_EDITOR
+	if (GIsEditor && !HasAnyFlags(RF_ClassDefaultObject | RF_ArchetypeObject))
+	{
+		GActorElementOwnerStore.RegisterElementOwner(this, CreateActorElement());
+	}
+#endif	// WITH_EDITOR
+
 	RemoteRole = (bReplicates ? ROLE_SimulatedProxy : ROLE_None);
 
 	// Make sure the OwnedComponents list correct.  
@@ -777,6 +830,13 @@ void AActor::BeginDestroy()
 	}
 #endif // (CSV_PROFILER && !UE_BUILD_SHIPPING)
 
+#if WITH_EDITOR
+	if (TTypedElementOwner<FActorElementData> EditorElement = GActorElementOwnerStore.UnregisterElementOwner(this))
+	{
+		DestroyActorElement(EditorElement);
+	}
+#endif	// WITH_EDITOR
+
 	Super::BeginDestroy();
 }
 
diff --git a/Engine/Source/Runtime/Engine/Private/Components/ActorComponent.cpp b/Engine/Source/Runtime/Engine/Private/Components/ActorComponent.cpp
--- a/Engine/Source/Runtime/Engine/Private/Components/ActorComponent.cpp
+++ b/Engine/Source/Runtime/Engine/Private/Components/ActorComponent.cpp
@@ -13,6 +13,9 @@
 #include "Components/PrimitiveComponent.h"
 #include "AI/NavigationSystemBase.h"
 #include "Engine/BlueprintGeneratedClass.h"
+#include "TypedElementRegistry.h"
+#include "TypedElementOwnerStore.h"
+#include "Elements/Component/ComponentElementData.h"
 #include "ContentStreaming.h"
 #include "ComponentReregisterContext.h"
 #include "Engine/AssetUserData.h"
@@ -74,6 +77,7 @@ FAutoConsoleVariableRef GTickComponentLatentActionsWithTheComponentCVar(
 
 #if WITH_EDITOR
 FUObjectAnnotationSparseBool GSelectedComponentAnnotation;
+TTypedElementOwnerStore<FComponentElementData, const UActorComponent*> GComponentElementOwnerStore;
 #endif
 
 /** Static var indicating activity of reregister context */
@@ -289,6 +293,13 @@ void UActorComponent::PostInitProperties()
 {
 	Super::PostInitProperties();
 
+#if WITH_EDITOR
+	if (GIsEditor && !HasAnyFlags(RF_ClassDefaultObject | RF_ArchetypeObject))
+	{
+		GComponentElementOwnerStore.RegisterElementOwner(this, CreateComponentElement());
+	}
+#endif	// WITH_EDITOR
+
 	// Instance components will be added during the owner's initialization
 	if (OwnerPrivate && CreationMethod != EComponentCreationMethod::Instance)
 	{
@@ -333,6 +344,47 @@ void UActorComponent::PostInitProperties()
 	}
 }
 
+TTypedElementOwner<FComponentElementData> UActorComponent::CreateComponentElement() const
+{
+	UTypedElementRegistry* Registry = UTypedElementRegistry::GetInstance();
+
+	TTypedElementOwner<FComponentElementData> ComponentElement;
+	if (ensureMsgf(Registry, TEXT("Typed element was requested for '%s' before the registry was available! This usually means that NewObject was used instead of CreateDefaultSubobject during CDO construction."), *GetPathName()))
+	{
+		ComponentElement = Registry->CreateElement<FComponentElementData>(NAME_Components);
+		ComponentElement.GetDataChecked().Component = const_cast<UActorComponent*>(this);
+	}
+
+	return ComponentElement;
+}
+
+void UActorComponent::DestroyComponentElement(TTypedElementOwner<FComponentElementData>& InOutComponentElement) const
+{
+	if (InOutComponentElement)
+	{
+		checkf(InOutComponentElement.GetDataChecked().Component == this, TEXT("Component element was not for this component instance! %s"), *GetPathName());
+		UTypedElementRegistry::GetInstance()->DestroyElement(InOutComponentElement);
+	}
+}
+
+#if WITH_EDITOR
+FTypedElementHandle UActorComponent::AcquireEditorElementHandle(const bool bAllowCreate) const
+{
+	if (TTypedElementOwnerScopedAccess<FComponentElementData> EditorElementScopedAccess = GComponentElementOwnerStore.FindElementOwner(this))
+	{
+		return EditorElementScopedAccess->AcquireHandle();
+	}
+
+	if (bAllowCreate)
+	{
+		TTypedElementOwnerScopedAccess<FComponentElementData> EditorElementScopedAccess = GComponentElementOwnerStore.RegisterElementOwner(this, CreateComponentElement());
+		return EditorElementScopedAccess->AcquireHandle();
+	}
+
+	return FTypedElementHandle();
+}
+#endif
+
 void UActorComponent::PostLoad()
 {
 	Super::PostLoad();
@@ -667,6 +719,13 @@ void UActorComponent::BeginDestroy()
 		MyOwner->RemoveOwnedComponent(this);
 	}
 
+#if WITH_EDITOR
+	if (TTypedElementOwner<FComponentElementData> EditorElement = GComponentElementOwnerStore.UnregisterElementOwner(this))
+	{
+		DestroyComponentElement(EditorElement);
+	}
+#endif	// WITH_EDITOR
+
 	Super::BeginDestroy();
 }
 
diff --git a/Engine/Source/Runtime/Engine/Private/Elements/Actor/ActorElementData.cpp b/Engine/Source/Runtime/Engine/Private/Elements/Actor/ActorElementData.cpp
new file mode 100644
--- /dev/null
+++ b/Engine/Source/Runtime/Engine/Private/Elements/Actor/ActorElementData.cpp
@@ -0,0 +1,5 @@
+// Copyright Epic Games, Inc. All Rights Reserved.
+
+#include "ActorElementData.h"
+
+UE_DEFINE_TYPED_ELEMENT_DATA_RTTI(FActorElementData);
diff --git a/Engine/Source/Runtime/Engine/Private/Elements/Actor/ActorElementData.h b/Engine/Source/Runtime/Engine/Private/Elements/Actor/ActorElementData.h
new file mode 100644
--- /dev/null
+++ b/Engine/Source/Runtime/Engine/Private/Elements/Actor/ActorElementData.h
@@ -0,0 +1,17 @@
+// Copyright Epic Games, Inc. All Rights Reserved.
+
+#pragma once
+
+#include "TypedElementData.h"
+
+class AActor;
+
+/**
+ * Element data that represents an Actor.
+ */
+struct FActorElementData
+{
+	UE_DECLARE_TYPED_ELEMENT_DATA_RTTI(FActorElementData);
+
+	AActor* Actor = nullptr;
+};
diff --git a/Engine/Source/Runtime/Engine/Private/Elements/Actor/ActorElementSelectionInterface.cpp b/Engine/Source/Runtime/Engine/Private/Elements/Actor/ActorElementSelectionInterface.cpp
new file mode 100644
--- /dev/null
+++ b/Engine/Source/Runtime/Engine/Private/Elements/Actor/ActorElementSelectionInterface.cpp
@@ -0,0 +1,17 @@
+// Copyright Epic Games, Inc. All Rights Reserved.
+
+#include "ActorElementSelectionInterface.h"
+#include "ActorElementData.h"
+#include "GameFramework/Actor.h"
+
+bool UActorElementSelectionInterface::IsValidSelection(const FTypedElementHandle& InElementHandle)
+{
+	// TODO: Validate that this actor is in a valid state to be selected?
+	return true;
+}
+
+UObject* UActorElementSelectionInterface::Legacy_GetSelectionObject(const FTypedElementHandle& InElementHandle)
+{
+	const FActorElementData* ActorData = InElementHandle.GetDataScript<FActorElementData>();
+	return ActorData ? ActorData->Actor : nullptr;
+}
diff --git a/Engine/Source/Runtime/Engine/Private/Elements/Actor/ActorElementSelectionInterface.h b/Engine/Source/Runtime/Engine/Private/Elements/Actor/ActorElementSelectionInterface.h
new file mode 100644
--- /dev/null
+++ b/Engine/Source/Runtime/Engine/Private/Elements/Actor/ActorElementSelectionInterface.h
@@ -0,0 +1,16 @@
+// Copyright Epic Games, Inc. All Rights Reserved.
+
+#pragma once
+
+#include "TypedElementSelectionInterface.h"
+#include "ActorElementSelectionInterface.generated.h"
+
+UCLASS()
+class UActorElementSelectionInterface : public UTypedElementSelectionInterface
+{
+	GENERATED_BODY()
+
+public:
+	virtual bool IsValidSelection(const FTypedElementHandle& InElementHandle) override;
+	virtual UObject* Legacy_GetSelectionObject(const FTypedElementHandle& InElementHandle) override;
+};
diff --git a/Engine/Source/Runtime/Engine/Private/Elements/Component/ComponentElementData.cpp b/Engine/Source/Runtime/Engine/Private/Elements/Component/ComponentElementData.cpp
new file mode 100644
--- /dev/null
+++ b/Engine/Source/Runtime/Engine/Private/Elements/Component/ComponentElementData.cpp
@@ -0,0 +1,5 @@
+// Copyright Epic Games, Inc. All Rights Reserved.
+
+#include "ComponentElementData.h"
+
+UE_DEFINE_TYPED_ELEMENT_DATA_RTTI(FComponentElementData);
diff --git a/Engine/Source/Runtime/Engine/Private/Elements/Component/ComponentElementData.h b/Engine/Source/Runtime/Engine/Private/Elements/Component/ComponentElementData.h
new file mode 100644
--- /dev/null
+++ b/Engine/Source/Runtime/Engine/Private/Elements/Component/ComponentElementData.h
@@ -0,0 +1,17 @@
+// Copyright Epic Games, Inc. All Rights Reserved.
+
+#pragma once
+
+#include "TypedElementData.h"
+
+class UActorComponent;
+
+/**
+ * Element data that represents an Actor Component.
+ */
+struct FComponentElementData
+{
+	UE_DECLARE_TYPED_ELEMENT_DATA_RTTI(FComponentElementData);
+
+	UActorComponent* Component = nullptr;
+};
diff --git a/Engine/Source/Runtime/Engine/Private/Elements/Component/ComponentElementSelectionInterface.cpp b/Engine/Source/Runtime/Engine/Private/Elements/Component/ComponentElementSelectionInterface.cpp
new file mode 100644
--- /dev/null
+++ b/Engine/Source/Runtime/Engine/Private/Elements/Component/ComponentElementSelectionInterface.cpp
@@ -0,0 +1,17 @@
+// Copyright Epic Games, Inc. All Rights Reserved.
+
+#include "ComponentElementSelectionInterface.h"
+#include "ComponentElementData.h"
+#include "Components/ActorComponent.h"
+
+bool UComponentElementSelectionInterface::IsValidSelection(const FTypedElementHandle& InElementHandle)
+{
+	// TODO: Validate that this component is in a valid state to be selected?
+	return true;
+}
+
+UObject* UComponentElementSelectionInterface::Legacy_GetSelectionObject(const FTypedElementHandle& InElementHandle)
+{
+	const FComponentElementData* ComponentData = InElementHandle.GetDataScript<FComponentElementData>();
+	return ComponentData ? ComponentData->Component : nullptr;
+}
diff --git a/Engine/Source/Runtime/Engine/Private/Elements/Component/ComponentElementSelectionInterface.h b/Engine/Source/Runtime/Engine/Private/Elements/Component/ComponentElementSelectionInterface.h
new file mode 100644
--- /dev/null
+++ b/Engine/Source/Runtime/Engine/Private/Elements/Component/ComponentElementSelectionInterface.h
@@ -0,0 +1,16 @@
+// Copyright Epic Games, Inc. All Rights Reserved.
+
+#pragma once
+
+#include "TypedElementSelectionInterface.h"
+#include "ComponentElementSelectionInterface.generated.h"
+
+UCLASS()
+class UComponentElementSelectionInterface : public UTypedElementSelectionInterface
+{
+	GENERATED_BODY()
+
+public:
+	virtual bool IsValidSelection(const FTypedElementHandle& InElementHandle) override;
+	virtual UObject* Legacy_GetSelectionObject(const FTypedElementHandle& InElementHandle) override;
+};
diff --git a/Engine/Source/Runtime/Engine/Private/Elements/EngineElements.cpp b/Engine/Source/Runtime/Engine/Private/Elements/EngineElements.cpp
new file mode 100644
--- /dev/null
+++ b/Engine/Source/Runtime/Engine/Private/Elements/EngineElements.cpp
@@ -0,0 +1,38 @@
+// Copyright Epic Games, Inc. All Rights Reserved.
+
+#include "EngineElements.h"
+#include "Modules/ModuleManager.h"
+
+#include "TypedElementRegistry.h"
+
+#include "Actor/ActorElementData.h"
+#include "Actor/ActorElementSelectionInterface.h"
+
+#include "Component/ComponentElementData.h"
+#include "Component/ComponentElementSelectionInterface.h"
+
+void RegisterActorElements()
+{
+	UTypedElementRegistry* Registry = UTypedElementRegistry::GetInstance();
+
+	Registry->RegisterElementType<FActorElementData>(NAME_Actor);
+	Registry->RegisterElementInterface<UTypedElementSelectionInterface>(NAME_Actor, NewObject<UActorElementSelectionInterface>());
+}
+
+void RegisterComponentElements()
+{
+	UTypedElementRegistry* Registry = UTypedElementRegistry::GetInstance();
+
+	Registry->RegisterElementType<FComponentElementData>(NAME_Components);
+	Registry->RegisterElementInterface<UTypedElementSelectionInterface>(NAME_Components, NewObject<UComponentElementSelectionInterface>());
+}
+
+void RegisterEngineElements()
+{
+	// Ensure the framework and base interfaces are also loaded
+	FModuleManager::Get().LoadModuleChecked("TypedElementFramework");
+	FModuleManager::Get().LoadModuleChecked("TypedElementInterfaces");
+
+	RegisterActorElements();
+	RegisterComponentElements();
+}
diff --git a/Engine/Source/Runtime/Engine/Private/Elements/EngineElements.h b/Engine/Source/Runtime/Engine/Private/Elements/EngineElements.h
new file mode 100644
--- /dev/null
+++ b/Engine/Source/Runtime/Engine/Private/Elements/EngineElements.h
@@ -0,0 +1,5 @@
+// Copyright Epic Games, Inc. All Rights Reserved.
+
+#pragma once
+
+void RegisterEngineElements();
diff --git a/Engine/Source/Runtime/Engine/Private/UnrealEngine.cpp b/Engine/Source/Runtime/Engine/Private/UnrealEngine.cpp
--- a/Engine/Source/Runtime/Engine/Private/UnrealEngine.cpp
+++ b/Engine/Source/Runtime/Engine/Private/UnrealEngine.cpp
@@ -104,6 +104,7 @@ UnrealEngine.cpp: Implements the UEngine class and helpers.
 #include "IImageWrapper.h"
 #include "IImageWrapperModule.h"
 #include "EngineAnalytics.h"
+#include "Elements/EngineElements.h"
 #include "TickTaskManagerInterface.h"
 #include "Net/NetworkProfiler.h"
 #include "ProfilingDebugging/MallocProfiler.h"
@@ -1483,6 +1484,9 @@ void UEngine::Init(IEngineLoop* InEngineLoop)
 
 	EngineLoop = InEngineLoop;
 
+	// Elements.
+	RegisterEngineElements();
+
 	// Subsystems.
 	FURL::StaticInit();
 	FLinkerLoad::StaticInit(UTexture2D::StaticClass());
