Exposed AISystemBase.AISystemClassName and AISystemBase.AISystemModuleName config properties to "AI Settings" (visible in Project Settings)

diff --git a/Engine/Source/Runtime/Engine/Classes/AI/AISystemBase.h b/Engine/Source/Runtime/Engine/Classes/AI/AISystemBase.h
--- a/Engine/Source/Runtime/Engine/Classes/AI/AISystemBase.h
+++ b/Engine/Source/Runtime/Engine/Classes/AI/AISystemBase.h
@@ -52,11 +52,11 @@ class ENGINE_API UAISystemBase : public UObject
 
 private:
 	/** List of specific AI system class to create, can be game-specific */
-	UPROPERTY(globalconfig, noclear, meta = (MetaClass = "AISystem", DisplayName = "AISystem Class"))
+	UPROPERTY(globalconfig, EditAnywhere, Category = "AISystem", noclear, meta = (MetaClass = "AISystem", DisplayName = "AISystem Class"))
 	FSoftClassPath AISystemClassName;
 
 	/** Name of a module used to spawn the AI system. If not empty, this module has to implement IAISystemModule */
-	UPROPERTY(globalconfig, noclear, meta = (MetaClass = "AISystem", DisplayName = "AISystem Module"))
+	UPROPERTY(globalconfig, EditAnywhere, Category = "AISystem", noclear, meta = (MetaClass = "AISystem", DisplayName = "AISystem Module"))
 	FName AISystemModuleName;
 
 	FDelegateHandle OnMatchStateSetHandle;
