Fix compile warnings.

diff --git a/Engine/Source/Runtime/Engine/Public/ComputeFramework/ComputeKernelResource.h b/Engine/Source/Runtime/Engine/Public/ComputeFramework/ComputeKernelResource.h
--- a/Engine/Source/Runtime/Engine/Public/ComputeFramework/ComputeKernelResource.h
+++ b/Engine/Source/Runtime/Engine/Public/ComputeFramework/ComputeKernelResource.h
@@ -2,7 +2,7 @@
 
 #pragma once
 
-#include "ComputeFramework//ComputeKernelShaderMap.h"
+#include "ComputeFramework/ComputeKernelShaderMap.h"
 
 struct FComputeKernelCompilationResults
 {
