Unshelved from pending changelist '4784118':

PR #5490: Add `IsEmpty` Blueprint array/set/map helpers (Contributed by [REDACTED_PERSON])

diff --git a/Engine/Source/Runtime/Engine/Classes/Kismet/BlueprintMapLibrary.h b/Engine/Source/Runtime/Engine/Classes/Kismet/BlueprintMapLibrary.h
--- a/Engine/Source/Runtime/Engine/Classes/Kismet/BlueprintMapLibrary.h
+++ b/Engine/Source/Runtime/Engine/Classes/Kismet/BlueprintMapLibrary.h
@@ -85,6 +85,24 @@ public:
 	UFUNCTION(BlueprintPure, CustomThunk, meta=(DisplayName = "Length", CompactNodeTitle = "LENGTH", MapParam = "TargetMap", Keywords = "num size count", BlueprintThreadSafe), Category="Utilities|Map")
 	static int32 Map_Length(const TMap<int32, int32>& TargetMap);
 	
+	/**
+	 * Check if the map does not have any entires
+	 *
+	 * @param	TargetMap		The map to check
+	 * @return	A boolean indicating if the map has any entires
+	 */
+	UFUNCTION(BlueprintPure, CustomThunk, meta = (DisplayName = "Is Empty", CompactNodeTitle = "IS EMPTY", MapParam = "TargetMap"), Category = "Utilities|Map")
+	static bool Map_IsEmpty(const TMap<int32, int32>& TargetMap);
+
+	/**
+	 * Check if the map has any entries
+	 *
+	 * @param	TargetMap		The map to check
+	 * @return	A boolean indicating if the map has any entires
+	 */
+	UFUNCTION(BlueprintPure, CustomThunk, meta = (DisplayName = "Is Not Empty", CompactNodeTitle = "IS NOT EMPTY", MapParam = "TargetMap"), Category = "Utilities|Map")
+	static bool Map_IsNotEmpty(const TMap<int32, int32>& TargetMap);
+
 	/** 
 	 * Clears a map of all entries, resetting it to empty
 	 *
@@ -311,6 +329,42 @@ public:
 		P_NATIVE_END
 	}
 
+	DECLARE_FUNCTION(execMap_IsEmpty)
+	{
+		Stack.MostRecentProperty = nullptr;
+		Stack.StepCompiledIn<FMapProperty>(nullptr);
+		void* MapAddr = Stack.MostRecentPropertyAddress;
+		FMapProperty* MapProperty = CastField<FMapProperty>(Stack.MostRecentProperty);
+		if (!MapProperty)
+		{
+			Stack.bArrayContextFailed = true;
+			return;
+		}
+
+		P_FINISH;
+		P_NATIVE_BEGIN;
+		*(bool*)RESULT_PARAM = GenericMap_IsEmpty(MapAddr, MapProperty);
+		P_NATIVE_END
+	}
+
+	DECLARE_FUNCTION(execMap_IsNotEmpty)
+	{
+		Stack.MostRecentProperty = nullptr;
+		Stack.StepCompiledIn<FMapProperty>(nullptr);
+		void* MapAddr = Stack.MostRecentPropertyAddress;
+		FMapProperty* MapProperty = CastField<FMapProperty>(Stack.MostRecentProperty);
+		if (!MapProperty)
+		{
+			Stack.bArrayContextFailed = true;
+			return;
+		}
+
+		P_FINISH;
+		P_NATIVE_BEGIN;
+		*(bool*)RESULT_PARAM = GenericMap_IsNotEmpty(MapAddr, MapProperty);
+		P_NATIVE_END
+	}
+
 	DECLARE_FUNCTION(execMap_Clear)
 	{
 		Stack.MostRecentProperty = nullptr;
@@ -350,6 +404,8 @@ public:
 	static void GenericMap_Keys(const void* MapAddr, const FMapProperty* MapProperty, const void* ArrayAddr, const FArrayProperty* ArrayProperty);
 	static void GenericMap_Values(const void* MapAddr, const FMapProperty* MapProperty, const void* ArrayAddr, const FArrayProperty* ArrayProperty);
 	static int32 GenericMap_Length(const void* TargetMap, const FMapProperty* MapProperty);
+	static bool GenericMap_IsEmpty(const void* TargetMap, const FMapProperty* MapProperty);
+	static bool GenericMap_IsNotEmpty(const void* TargetMap, const FMapProperty* MapProperty);
 	static void GenericMap_Clear(const void* TargetMap, const FMapProperty* MapProperty);
 	static void GenericMap_SetMapPropertyByName(UObject* OwnerObject, FName MapPropertyName, const void* SrcMapAddr);
 };
diff --git a/Engine/Source/Runtime/Engine/Classes/Kismet/BlueprintSetLibrary.h b/Engine/Source/Runtime/Engine/Classes/Kismet/BlueprintSetLibrary.h
--- a/Engine/Source/Runtime/Engine/Classes/Kismet/BlueprintSetLibrary.h
+++ b/Engine/Source/Runtime/Engine/Classes/Kismet/BlueprintSetLibrary.h
@@ -45,6 +45,24 @@ public:
 	UFUNCTION(BlueprintCallable, CustomThunk, meta=(DisplayName = "Remove", CompactNodeTitle = "REMOVE", SetParam = "TargetSet|Item", AutoCreateRefTerm = "Item"), Category="Utilities|Set")
 	static bool Set_Remove(const TSet<int32>& TargetSet, const int32& Item);
 	
+	/**
+	 * Check if the set is empty
+	 *
+	 * @param	TargetSet		The set to check
+	 * @return	A boolean indicating if the array is empty
+	 */
+	UFUNCTION(BlueprintPure, CustomThunk, meta = (DisplayName = "Is Empty", CompactNodeTitle = "IS EMPTY", SetParam = "TargetSet"), Category = "Utilities|Set")
+	static bool Set_IsEmpty(const TSet<int32>& TargetSet);
+
+	/**
+	 * Check if the set has any elements
+	 *
+	 * @param	TargetSet		The set to check
+	 * @return	A boolean indicating if the array has any elements
+	 */
+	UFUNCTION(BlueprintPure, CustomThunk, meta = (DisplayName = "Is Not Empty", CompactNodeTitle = "IS NOT EMPTY", SetParam = "TargetSet"), Category = "Utilities|Set")
+	static bool Set_IsNotEmpty(const TSet<int32>& TargetSet);
+
 	/**
 	 * Removes all elements in an Array from a set.
 	 *
@@ -355,6 +373,44 @@ public:
 		ElementProp->DestroyValue(StorageSpace);
 	}
 
+	DECLARE_FUNCTION(execSet_IsEmpty)
+	{
+		Stack.MostRecentProperty = nullptr;
+		Stack.StepCompiledIn<FSetProperty>(nullptr);
+		void* SetAddr = Stack.MostRecentPropertyAddress;
+		FSetProperty* SetProperty = CastField<FSetProperty>(Stack.MostRecentProperty);
+		if (!SetProperty)
+		{
+			Stack.bArrayContextFailed = true;
+			return;
+		}
+
+		P_FINISH;
+
+		P_NATIVE_BEGIN;
+		*(bool*)RESULT_PARAM = GenericSet_IsEmpty(SetAddr, SetProperty);
+		P_NATIVE_END;
+	}
+
+	DECLARE_FUNCTION(execSet_IsNotEmpty)
+	{
+		Stack.MostRecentProperty = nullptr;
+		Stack.StepCompiledIn<FSetProperty>(nullptr);
+		void* SetAddr = Stack.MostRecentPropertyAddress;
+		FSetProperty* SetProperty = CastField<FSetProperty>(Stack.MostRecentProperty);
+		if (!SetProperty)
+		{
+			Stack.bArrayContextFailed = true;
+			return;
+		}
+
+		P_FINISH;
+
+		P_NATIVE_BEGIN;
+		*(bool*)RESULT_PARAM = GenericSet_IsNotEmpty(SetAddr, SetProperty);
+		P_NATIVE_END;
+	}
+
 	DECLARE_FUNCTION(execSet_Intersection)
 	{
 		Stack.MostRecentProperty = nullptr;
@@ -495,6 +551,8 @@ public:
 	static void GenericSet_Clear(const void* TargetSet, const FSetProperty* SetProperty);
 	static int32 GenericSet_Length(const void* TargetSet, const FSetProperty* SetProperty);
 	static bool GenericSet_Contains(const void* TargetSet, const FSetProperty* SetProperty, const void* ItemToFind);
+	static bool GenericSet_IsEmpty(const void* TargetSet, const FSetProperty* SetProperty);
+	static bool GenericSet_IsNotEmpty(const void* TargetSet, const FSetProperty* SetProperty);
 	static void GenericSet_Intersect(const void* SetA, const FSetProperty* SetPropertyA, const void* SetB, const FSetProperty* SetPropertyB, const void* SetResult, const FSetProperty* SetPropertyResult);
 	static void GenericSet_Union(const void* SetA, const FSetProperty* SetPropertyA, const void* SetB, const FSetProperty* SetPropertyB, const void* SetResult, const FSetProperty* SetPropertyResult);
 	static void GenericSet_Difference(const void* SetA, const FSetProperty* SetPropertyA, const void* SetB, const FSetProperty* SetPropertyB, const void* SetResult, const FSetProperty* SetPropertyResult);
diff --git a/Engine/Source/Runtime/Engine/Classes/Kismet/KismetArrayLibrary.h b/Engine/Source/Runtime/Engine/Classes/Kismet/KismetArrayLibrary.h
--- a/Engine/Source/Runtime/Engine/Classes/Kismet/KismetArrayLibrary.h
+++ b/Engine/Source/Runtime/Engine/Classes/Kismet/KismetArrayLibrary.h
@@ -122,6 +122,24 @@ class ENGINE_API UKismetArrayLibrary : public UBlueprintFunctionLibrary
 	UFUNCTION(BlueprintPure, CustomThunk, meta=(DisplayName = "Length", CompactNodeTitle = "LENGTH", ArrayParm = "TargetArray", Keywords = "num size count", BlueprintThreadSafe), Category="Utilities|Array")
 	static int32 Array_Length(const TArray<int32>& TargetArray);
 
+	/*
+	 *Check if the array is empty
+	 *
+	 *@param	TargetArray		The array to check
+	 *@return	A boolean indicating if the array is empty
+	*/
+	UFUNCTION(BlueprintPure, CustomThunk, meta = (DisplayName = "Is Empty", CompactNodeTitle = "IS EMPTY", ArrayParm = "TargetArray", BlueprintThreadSafe), Category = "Utilities|Array")
+	static bool Array_IsEmpty(const TArray<int32>& TargetArray);
+
+	/*
+	 *Check if the array has any elements
+	 *
+	 *@param	TargetArray		The array to check
+	 *@return	A boolean indicating if the array has any elements
+	*/
+	UFUNCTION(BlueprintPure, CustomThunk, meta = (DisplayName = "Is Not Empty", CompactNodeTitle = "IS NOT EMPTY", ArrayParm = "TargetArray", BlueprintThreadSafe), Category = "Utilities|Array")
+	static bool Array_IsNotEmpty(const TArray<int32>& TargetArray);
+
 
 	/* 
 	 *Get the last valid index into an array
@@ -222,6 +240,8 @@ class ENGINE_API UKismetArrayLibrary : public UBlueprintFunctionLibrary
 	static void GenericArray_Clear(void* TargetArray, const FArrayProperty* ArrayProp);
 	static void GenericArray_Resize(void* TargetArray, const FArrayProperty* ArrayProp, int32 Size);
 	static int32 GenericArray_Length(const void* TargetArray, const FArrayProperty* ArrayProp);
+	static bool GenericArray_IsEmpty(const void* TargetArray, const FArrayProperty* ArrayProp);
+	static bool GenericArray_IsNotEmpty(const void* TargetArray, const FArrayProperty* ArrayProp);
 	static int32 GenericArray_LastIndex(const void* TargetArray, const FArrayProperty* ArrayProp);
 	static void GenericArray_Get(void* TargetArray, const FArrayProperty* ArrayProp, int32 Index, void* Item);
 	static void GenericArray_Set(void* TargetArray, const FArrayProperty* ArrayProp, int32 Index, const void* NewItem, bool bSizeToFit);
@@ -513,6 +533,40 @@ public:
 		*(int32*)RESULT_PARAM = GenericArray_Length(ArrayAddr, ArrayProperty);
 		P_NATIVE_END;
 	}
+	
+	DECLARE_FUNCTION(execArray_IsEmpty)
+	{
+		Stack.MostRecentProperty = nullptr;
+		Stack.StepCompiledIn<FArrayProperty>(nullptr);
+		void* ArrayAddr = Stack.MostRecentPropertyAddress;
+		FArrayProperty* ArrayProperty = CastField<FArrayProperty>(Stack.MostRecentProperty);
+		if (!ArrayProperty)
+		{
+			Stack.bArrayContextFailed = true;
+			return;
+		}
+		P_FINISH;
+		P_NATIVE_BEGIN;
+		*(bool*)RESULT_PARAM = GenericArray_IsEmpty(ArrayAddr, ArrayProperty);
+		P_NATIVE_END;
+	}
+	
+	DECLARE_FUNCTION(execArray_IsNotEmpty)
+	{
+		Stack.MostRecentProperty = nullptr;
+		Stack.StepCompiledIn<FArrayProperty>(nullptr);
+		void* ArrayAddr = Stack.MostRecentPropertyAddress;
+		FArrayProperty* ArrayProperty = CastField<FArrayProperty>(Stack.MostRecentProperty);
+		if (!ArrayProperty)
+		{
+			Stack.bArrayContextFailed = true;
+			return;
+		}
+		P_FINISH;
+		P_NATIVE_BEGIN;
+		*(bool*)RESULT_PARAM = GenericArray_IsNotEmpty(ArrayAddr, ArrayProperty);
+		P_NATIVE_END;
+	}
 
 	DECLARE_FUNCTION(execArray_LastIndex)
 	{
diff --git a/Engine/Source/Runtime/Engine/Private/BlueprintMapLibrary.cpp b/Engine/Source/Runtime/Engine/Private/BlueprintMapLibrary.cpp
--- a/Engine/Source/Runtime/Engine/Private/BlueprintMapLibrary.cpp
+++ b/Engine/Source/Runtime/Engine/Private/BlueprintMapLibrary.cpp
@@ -99,6 +99,26 @@ int32 UBlueprintMapLibrary::GenericMap_Length(const void* TargetMap, const FMapP
 	return 0;
 }
 
+bool UBlueprintMapLibrary::GenericMap_IsEmpty(const void* TargetMap, const FMapProperty* MapProperty)
+{
+	if (TargetMap)
+	{
+		FScriptMapHelper MapHelper(MapProperty, TargetMap);
+		return MapHelper.Num() == 0;
+	}
+	return true;
+}
+
+bool UBlueprintMapLibrary::GenericMap_IsNotEmpty(const void* TargetMap, const FMapProperty* MapProperty)
+{
+	if (TargetMap)
+	{
+		FScriptMapHelper MapHelper(MapProperty, TargetMap);
+		return MapHelper.Num() > 0;
+	}
+	return false;
+}
+
 void UBlueprintMapLibrary::GenericMap_Clear(const void* TargetMap, const FMapProperty* MapProperty)
 {
 	if(TargetMap)
diff --git a/Engine/Source/Runtime/Engine/Private/BlueprintSetLibrary.cpp b/Engine/Source/Runtime/Engine/Private/BlueprintSetLibrary.cpp
--- a/Engine/Source/Runtime/Engine/Private/BlueprintSetLibrary.cpp
+++ b/Engine/Source/Runtime/Engine/Private/BlueprintSetLibrary.cpp
@@ -101,6 +101,29 @@ bool UBlueprintSetLibrary::GenericSet_Contains(const void* TargetSet, const FSet
 	return false;
 }
 
+bool UBlueprintSetLibrary::GenericSet_IsEmpty(const void* TargetSet, const FSetProperty* SetProperty)
+{
+	if (TargetSet)
+	{
+		FScriptSetHelper SetHelper(SetProperty, TargetSet);
+
+		return SetHelper.Num() == 0;
+	}
+
+	return true;
+}
+
+bool UBlueprintSetLibrary::GenericSet_IsNotEmpty(const void* TargetSet, const FSetProperty* SetProperty)
+{
+	if (TargetSet)
+	{
+		FScriptSetHelper SetHelper(SetProperty, TargetSet);
+
+		return SetHelper.Num() > 0;
+	}
+	return false;
+}
+
 void UBlueprintSetLibrary::GenericSet_Intersect(const void* SetA, const FSetProperty* SetPropertyA, const void* SetB, const FSetProperty* SetPropertyB, const void* SetResult, const FSetProperty* SetPropertyResult)
 {
 	if (SetA && SetB && SetResult)
diff --git a/Engine/Source/Runtime/Engine/Private/KismetArrayLibrary.cpp b/Engine/Source/Runtime/Engine/Private/KismetArrayLibrary.cpp
--- a/Engine/Source/Runtime/Engine/Private/KismetArrayLibrary.cpp
+++ b/Engine/Source/Runtime/Engine/Private/KismetArrayLibrary.cpp
@@ -267,6 +267,30 @@ int32 UKismetArrayLibrary::GenericArray_Length(const void* TargetArray, const FA
 	return 0;
 }
 
+bool UKismetArrayLibrary::GenericArray_IsEmpty(const void* TargetArray, const FArrayProperty* ArrayProp)
+{
+	if( TargetArray )
+	{
+		FScriptArrayHelper ArrayHelper(ArrayProp, TargetArray);
+
+		return ArrayHelper.Num() == 0;
+	}
+	
+	return true;
+}
+
+bool UKismetArrayLibrary::GenericArray_IsNotEmpty(const void* TargetArray, const FArrayProperty* ArrayProp)
+{
+	if( TargetArray )
+	{
+		FScriptArrayHelper ArrayHelper(ArrayProp, TargetArray);
+
+		return ArrayHelper.Num() > 0;
+	}
+
+	return false;
+}
+
 int32 UKismetArrayLibrary::GenericArray_LastIndex(const void* TargetArray, const FArrayProperty* ArrayProp)
 {
 	if( TargetArray )
@@ -477,6 +501,20 @@ int32 UKismetArrayLibrary::Array_Length(const TArray<int32>& TargetArray)
 	return 0;
 }
 
+bool UKismetArrayLibrary::Array_IsEmpty(const TArray<int32>& TargetArray)
+{
+	// We should never hit these!  They're stubs to avoid NoExport on the class.  Call the Generic* equivalent instead
+	check(0);
+	return true;
+}
+
+bool UKismetArrayLibrary::Array_IsNotEmpty(const TArray<int32>& TargetArray)
+{
+	// We should never hit these!  They're stubs to avoid NoExport on the class.  Call the Generic* equivalent instead
+	check(0);
+	return false;
+}
+
 int32 UKismetArrayLibrary::Array_LastIndex(const TArray<int32>& TargetArray)
 {
 	// We should never hit these!  They're stubs to avoid NoExport on the class.  Call the Generic* equivalent instead
diff --git a/Engine/Source/Runtime/Engine/Public/GeneratedCodeHelpers.h b/Engine/Source/Runtime/Engine/Public/GeneratedCodeHelpers.h
--- a/Engine/Source/Runtime/Engine/Public/GeneratedCodeHelpers.h
+++ b/Engine/Source/Runtime/Engine/Public/GeneratedCodeHelpers.h
@@ -360,6 +360,18 @@ public:
 		return TargetSet.Remove(Item) > 0;
 	}
 
+	template<typename T>
+	static bool Set_IsEmpty(const TSet<T>& TargetSet)
+	{
+		return TargetSet.Num() == 0;
+	}
+
+	template<typename T>
+	static bool Set_IsNotEmpty(const TSet<T>& TargetSet)
+	{
+		return TargetSet.Num() > 0;
+	}
+
 	template<typename T, typename U>
 	static void Set_RemoveItems(TSet<T>& TargetSet, const TArray<U>& Items)
 	{
@@ -466,6 +478,18 @@ public:
 		return TargetMap.Num();
 	}
 	
+	template<typename T, typename U>
+	static bool Map_IsEmpty(const TMap<T, U>& TargetMap)
+	{
+		return TargetMap.Num() == 0;
+	}
+
+	template<typename T, typename U>
+	static bool Map_IsNotEmpty(const TMap<T, U>& TargetMap)
+	{
+		return TargetMap.Num() > 0;
+	}
+
 	template<typename T, typename U>
 	static void Map_Clear(TMap<T, U>& TargetMap)
 	{
