Blueprints: Add int to Double conversion function

diff --git a/Engine/Source/Runtime/Engine/Classes/Kismet/KismetMathLibrary.h b/Engine/Source/Runtime/Engine/Classes/Kismet/KismetMathLibrary.h
--- a/Engine/Source/Runtime/Engine/Classes/Kismet/KismetMathLibrary.h
+++ b/Engine/Source/Runtime/Engine/Classes/Kismet/KismetMathLibrary.h
@@ -3403,6 +3403,10 @@ class ENGINE_API UKismetMathLibrary : public UBlueprintFunctionLibrary
 	UFUNCTION(BlueprintPure, meta=(DisplayName = "ToFloat (integer)", CompactNodeTitle = "->", Keywords="cast convert", BlueprintAutocast), Category="Math|Conversions")
 	static float Conv_IntToFloat(int32 InInt);
 
+	/** Converts an integer to a double */
+	UFUNCTION(BlueprintPure, meta = (DisplayName = "ToDouble (integer)", CompactNodeTitle = "->", Keywords = "cast convert", BlueprintAutocast), Category = "Math|Conversions")
+	static double Conv_IntToDouble(int32 InInt);
+
 	/** Converts an integer to a 64 bit integer */
 	UFUNCTION(BlueprintPure, meta = (DisplayName = "ToInt64 (integer)", CompactNodeTitle = "->", Keywords = "cast convert", BlueprintAutocast), Category = "Math|Conversions")
 	static int64 Conv_IntToInt64(int32 InInt);
diff --git a/Engine/Source/Runtime/Engine/Classes/Kismet/KismetMathLibrary.inl b/Engine/Source/Runtime/Engine/Classes/Kismet/KismetMathLibrary.inl
--- a/Engine/Source/Runtime/Engine/Classes/Kismet/KismetMathLibrary.inl
+++ b/Engine/Source/Runtime/Engine/Classes/Kismet/KismetMathLibrary.inl
@@ -2607,6 +2607,12 @@ float UKismetMathLibrary::Conv_IntToFloat(int32 InInt)
 	return (float)InInt;
 }
 
+KISMET_MATH_FORCEINLINE
+double UKismetMathLibrary::Conv_IntToDouble(int32 InInt)
+{
+	return (double)InInt;
+}
+
 KISMET_MATH_FORCEINLINE
 int64 UKismetMathLibrary::Conv_IntToInt64(int32 InInt)
 {
