ComputeFramework: Fix non-unity/non-PCH builds.

#rnx

diff --git a/Engine/Source/Runtime/Engine/Public/ComputeFramework/ComputeKernelFromText.h b/Engine/Source/Runtime/Engine/Public/ComputeFramework/ComputeKernelFromText.h
--- a/Engine/Source/Runtime/Engine/Public/ComputeFramework/ComputeKernelFromText.h
+++ b/Engine/Source/Runtime/Engine/Public/ComputeFramework/ComputeKernelFromText.h
@@ -3,6 +3,7 @@
 #pragma once
 
 #include "ComputeKernelSource.h"
+#include "Engine/EngineTypes.h"
 
 #include "ComputeKernelFromText.generated.h"
 
diff --git a/Engine/Source/Runtime/Engine/Public/ComputeFramework/ComputeKernelShaderMap.h b/Engine/Source/Runtime/Engine/Public/ComputeFramework/ComputeKernelShaderMap.h
--- a/Engine/Source/Runtime/Engine/Public/ComputeFramework/ComputeKernelShaderMap.h
+++ b/Engine/Source/Runtime/Engine/Public/ComputeFramework/ComputeKernelShaderMap.h
@@ -2,6 +2,12 @@
 
 #pragma once
 
+#include "RHI.h"
+#include "RHIDefinitions.h"
+#include "Serialization/MemoryImage.h"
+#include "Shader.h"
+
+
 class FComputeKernelResource;
 
 /* Stores all the output/stats from the compilation process, eg. error messages, error codes. */
diff --git a/Engine/Source/Runtime/Engine/Public/ComputeFramework/ComputeKernelShaderType.h b/Engine/Source/Runtime/Engine/Public/ComputeFramework/ComputeKernelShaderType.h
--- a/Engine/Source/Runtime/Engine/Public/ComputeFramework/ComputeKernelShaderType.h
+++ b/Engine/Source/Runtime/Engine/Public/ComputeFramework/ComputeKernelShaderType.h
@@ -2,6 +2,9 @@
 
 #pragma once
 
+#include "Shader.h"
+#include "ShaderCompiler.h"
+
 struct FShaderCompilerEnvironment;
 class FComputeKernelResource;
 
