Added missing GIsEditor

#rnx

diff --git a/Engine/Source/Runtime/Engine/Private/Actor.cpp b/Engine/Source/Runtime/Engine/Private/Actor.cpp
--- a/Engine/Source/Runtime/Engine/Private/Actor.cpp
+++ b/Engine/Source/Runtime/Engine/Private/Actor.cpp
@@ -250,15 +250,18 @@ void AActor::DestroyActorElement(TTypedElementOwner<FActorElementData>& InOutAct
 #if WITH_EDITOR
 FTypedElementHandle AActor::AcquireEditorElementHandle(const bool bAllowCreate) const
 {
-	if (TTypedElementOwnerScopedAccess<FActorElementData> EditorElementScopedAccess = GActorElementOwnerStore.FindElementOwner(this))
+	if (GIsEditor)
 	{
-		return EditorElementScopedAccess->AcquireHandle();
-	}
+		if (TTypedElementOwnerScopedAccess<FActorElementData> EditorElementScopedAccess = GActorElementOwnerStore.FindElementOwner(this))
+		{
+			return EditorElementScopedAccess->AcquireHandle();
+		}
 
-	if (bAllowCreate)
-	{
-		TTypedElementOwnerScopedAccess<FActorElementData> EditorElementScopedAccess = GActorElementOwnerStore.RegisterElementOwner(this, CreateActorElement());
-		return EditorElementScopedAccess->AcquireHandle();
+		if (bAllowCreate)
+		{
+			TTypedElementOwnerScopedAccess<FActorElementData> EditorElementScopedAccess = GActorElementOwnerStore.RegisterElementOwner(this, CreateActorElement());
+			return EditorElementScopedAccess->AcquireHandle();
+		}
 	}
 
 	return FTypedElementHandle();
diff --git a/Engine/Source/Runtime/Engine/Private/Components/ActorComponent.cpp b/Engine/Source/Runtime/Engine/Private/Components/ActorComponent.cpp
--- a/Engine/Source/Runtime/Engine/Private/Components/ActorComponent.cpp
+++ b/Engine/Source/Runtime/Engine/Private/Components/ActorComponent.cpp
@@ -370,15 +370,18 @@ void UActorComponent::DestroyComponentElement(TTypedElementOwner<FComponentEleme
 #if WITH_EDITOR
 FTypedElementHandle UActorComponent::AcquireEditorElementHandle(const bool bAllowCreate) const
 {
-	if (TTypedElementOwnerScopedAccess<FComponentElementData> EditorElementScopedAccess = GComponentElementOwnerStore.FindElementOwner(this))
+	if (GIsEditor)
 	{
-		return EditorElementScopedAccess->AcquireHandle();
-	}
+		if (TTypedElementOwnerScopedAccess<FComponentElementData> EditorElementScopedAccess = GComponentElementOwnerStore.FindElementOwner(this))
+		{
+			return EditorElementScopedAccess->AcquireHandle();
+		}
 
-	if (bAllowCreate)
-	{
-		TTypedElementOwnerScopedAccess<FComponentElementData> EditorElementScopedAccess = GComponentElementOwnerStore.RegisterElementOwner(this, CreateComponentElement());
-		return EditorElementScopedAccess->AcquireHandle();
+		if (bAllowCreate)
+		{
+			TTypedElementOwnerScopedAccess<FComponentElementData> EditorElementScopedAccess = GComponentElementOwnerStore.RegisterElementOwner(this, CreateComponentElement());
+			return EditorElementScopedAccess->AcquireHandle();
+		}
 	}
 
 	return FTypedElementHandle();
