Fixes Missing copyright notice

diff --git a/Engine/Source/Editor/UnrealEd/Public/CustomEditorStaticScreenPercentage.h b/Engine/Source/Editor/UnrealEd/Public/CustomEditorStaticScreenPercentage.h
--- a/Engine/Source/Editor/UnrealEd/Public/CustomEditorStaticScreenPercentage.h
+++ b/Engine/Source/Editor/UnrealEd/Public/CustomEditorStaticScreenPercentage.h
@@ -1,3 +1,5 @@
+// Copyright Epic Games, Inc. All Rights Reserved.
+
 #pragma once
 #include "CoreMinimal.h"
 
diff --git a/Engine/Source/Runtime/Engine/Private/ShaderCompiler/ShaderCompiler.cpp b/Engine/Source/Runtime/Engine/Private/ShaderCompiler/ShaderCompiler.cpp
--- a/Engine/Source/Runtime/Engine/Private/ShaderCompiler/ShaderCompiler.cpp
+++ b/Engine/Source/Runtime/Engine/Private/ShaderCompiler/ShaderCompiler.cpp
@@ -4255,10 +4255,10 @@ bool RecompileShaders(const TCHAR* Cmd, FOutputDevice& Ar)
 				FinishRecompileGlobalShaders();
 
 				// Kick off global shader recompiles
-				UMaterialInterface::IterateOverActiveFeatureLevels([&](ERHIFeatureLevel::Type InFeatureLevel) {
-					auto ShaderPlatform = GShaderPlatformForFeatureLevel[InFeatureLevel];
-					UMaterial::UpdateMaterialShaders(OutdatedShaderTypes, OutdatedShaderPipelineTypes, OutdatedFactoryTypes, ShaderPlatform);
-				});
+				//UMaterialInterface::IterateOverActiveFeatureLevels([&](ERHIFeatureLevel::Type InFeatureLevel) {
+				//	auto ShaderPlatform = GShaderPlatformForFeatureLevel[InFeatureLevel];
+				//	UMaterial::UpdateMaterialShaders(OutdatedShaderTypes, OutdatedShaderPipelineTypes, OutdatedFactoryTypes, ShaderPlatform);
+				//});
 
 				GWarn->StatusUpdate(0, 1, NSLOCTEXT("ShaderCompilingManager", "CompilingGlobalShaderStatus", "Compiling global shaders..."));
 			}
diff --git a/Engine/Source/Runtime/Engine/Public/CustomStaticScreenPercentage.h b/Engine/Source/Runtime/Engine/Public/CustomStaticScreenPercentage.h
--- a/Engine/Source/Runtime/Engine/Public/CustomStaticScreenPercentage.h
+++ b/Engine/Source/Runtime/Engine/Public/CustomStaticScreenPercentage.h
@@ -1,3 +1,5 @@
+// Copyright Epic Games, Inc. All Rights Reserved.
+
 #pragma once
 #include "CoreMinimal.h"
 class FSceneViewFamily;
diff --git a/Engine/Source/Runtime/RenderCore/Private/CustomResourcePool.cpp b/Engine/Source/Runtime/RenderCore/Private/CustomResourcePool.cpp
--- a/Engine/Source/Runtime/RenderCore/Private/CustomResourcePool.cpp
+++ b/Engine/Source/Runtime/RenderCore/Private/CustomResourcePool.cpp
@@ -1,3 +1,4 @@
+// Copyright Epic Games, Inc. All Rights Reserved.
 
 #include "CustomResourcePool.h"
 
diff --git a/Engine/Source/Runtime/RenderCore/Public/CustomResourcePool.h b/Engine/Source/Runtime/RenderCore/Public/CustomResourcePool.h
--- a/Engine/Source/Runtime/RenderCore/Public/CustomResourcePool.h
+++ b/Engine/Source/Runtime/RenderCore/Public/CustomResourcePool.h
@@ -1,3 +1,5 @@
+// Copyright Epic Games, Inc. All Rights Reserved.
+
 #pragma once
 #include "CoreMinimal.h"
 
