Undo single //UE5/Main/Engine/Source/Runtime/Engine/Private/ShaderCompiler/ShaderCompiler.cpp revision 7

diff --git a/Engine/Source/Runtime/Engine/Private/ShaderCompiler/ShaderCompiler.cpp b/Engine/Source/Runtime/Engine/Private/ShaderCompiler/ShaderCompiler.cpp
--- a/Engine/Source/Runtime/Engine/Private/ShaderCompiler/ShaderCompiler.cpp
+++ b/Engine/Source/Runtime/Engine/Private/ShaderCompiler/ShaderCompiler.cpp
@@ -4255,10 +4255,10 @@ bool RecompileShaders(const TCHAR* Cmd, FOutputDevice& Ar)
 				FinishRecompileGlobalShaders();
 
 				// Kick off global shader recompiles
-				//UMaterialInterface::IterateOverActiveFeatureLevels([&](ERHIFeatureLevel::Type InFeatureLevel) {
-				//	auto ShaderPlatform = GShaderPlatformForFeatureLevel[InFeatureLevel];
-				//	UMaterial::UpdateMaterialShaders(OutdatedShaderTypes, OutdatedShaderPipelineTypes, OutdatedFactoryTypes, ShaderPlatform);
-				//});
+				UMaterialInterface::IterateOverActiveFeatureLevels([&](ERHIFeatureLevel::Type InFeatureLevel) {
+					auto ShaderPlatform = GShaderPlatformForFeatureLevel[InFeatureLevel];
+					UMaterial::UpdateMaterialShaders(OutdatedShaderTypes, OutdatedShaderPipelineTypes, OutdatedFactoryTypes, ShaderPlatform);
+				});
 
 				GWarn->StatusUpdate(0, 1, NSLOCTEXT("ShaderCompilingManager", "CompilingGlobalShaderStatus", "Compiling global shaders..."));
 			}
