Fixes missing copyright notice

diff --git a/Engine/Source/Runtime/Engine/Private/CustomStaticScreenPercentage.cpp b/Engine/Source/Runtime/Engine/Private/CustomStaticScreenPercentage.cpp
--- a/Engine/Source/Runtime/Engine/Private/CustomStaticScreenPercentage.cpp
+++ b/Engine/Source/Runtime/Engine/Private/CustomStaticScreenPercentage.cpp
@@ -1,3 +1,4 @@
+// Copyright Epic Games, Inc. All Rights Reserved.
 
 #include "CustomStaticScreenPercentage.h"
 
