Rename Foundation to LevelInstance

diff --git a/Engine/Source/Editor/EditorStyle/Private/StarshipStyle.cpp b/Engine/Source/Editor/EditorStyle/Private/StarshipStyle.cpp
--- a/Engine/Source/Editor/EditorStyle/Private/StarshipStyle.cpp
+++ b/Engine/Source/Editor/EditorStyle/Private/StarshipStyle.cpp
@@ -5907,7 +5907,7 @@ void FStarshipEditorStyle::FStyle::SetupClassIconsAndThumbnails()
 			TEXT("CameraRig_Crane"),
 			TEXT("CameraRig_Rail"),
 			TEXT("FoliageType_Actor"),
-			TEXT("FoundationActor")
+			TEXT("LevelInstance")
 		};
 
 		for (int32 TypeIndex = 0; TypeIndex < UE_ARRAY_COUNT(AssetTypes); ++TypeIndex)
diff --git a/Engine/Source/Editor/FoundationEditor/FoundationEditor.Build.cs b/Engine/Source/Editor/FoundationEditor/FoundationEditor.Build.cs
deleted file mode 100644
diff --git a/Engine/Source/Editor/FoundationEditor/Private/FoundationActorDetails.cpp b/Engine/Source/Editor/FoundationEditor/Private/FoundationActorDetails.cpp
deleted file mode 100644
diff --git a/Engine/Source/Editor/FoundationEditor/Private/FoundationActorDetails.h b/Engine/Source/Editor/FoundationEditor/Private/FoundationActorDetails.h
deleted file mode 100644
diff --git a/Engine/Source/Editor/FoundationEditor/Private/FoundationEditorModule.cpp b/Engine/Source/Editor/FoundationEditor/Private/FoundationEditorModule.cpp
deleted file mode 100644
diff --git a/Engine/Source/Editor/FoundationEditor/Private/FoundationEditorModule.h b/Engine/Source/Editor/FoundationEditor/Private/FoundationEditorModule.h
deleted file mode 100644
diff --git a/Engine/Source/Editor/FoundationEditor/Private/FoundationEditorSettings.cpp b/Engine/Source/Editor/FoundationEditor/Private/FoundationEditorSettings.cpp
deleted file mode 100644
diff --git a/Engine/Source/Editor/FoundationEditor/Public/FoundationEditorSettings.h b/Engine/Source/Editor/FoundationEditor/Public/FoundationEditorSettings.h
deleted file mode 100644
diff --git a/Engine/Source/Editor/LevelInstanceEditor/LevelInstanceEditor.Build.cs b/Engine/Source/Editor/LevelInstanceEditor/LevelInstanceEditor.Build.cs
new file mode 100644
--- /dev/null
+++ b/Engine/Source/Editor/LevelInstanceEditor/LevelInstanceEditor.Build.cs
@@ -0,0 +1,43 @@
+// Copyright Epic Games, Inc. All Rights Reserved.
+
+using UnrealBuildTool;
+
+public class LevelInstanceEditor : ModuleRules
+{
+    public LevelInstanceEditor(ReadOnlyTargetRules Target) : base(Target)
+    {
+        PrivateIncludePaths.Add("Editor/LevelInstanceEditor/Private");	// For PCH includes (because they don't work with relative paths, yet)
+
+        PrivateIncludePathModuleNames.AddRange(
+            new string[] {
+            }
+        );
+     
+        PrivateDependencyModuleNames.AddRange(
+            new string[] {
+				"Core",
+				"CoreUObject",
+				"EditorStyle",
+				"Engine",
+				"InputCore",
+				"Slate",
+				"SlateCore",
+				"EditorFramework",
+				"UnrealEd",
+				"LevelEditor",
+				"ToolMenus",
+				"PropertyEditor",
+				"NewLevelDialog",
+				"MainFrame",
+				"ContentBrowser",
+				"AssetTools",
+				"ClassViewer"
+			}
+		);
+
+        DynamicallyLoadedModuleNames.AddRange(
+            new string[] {
+		    }
+		);
+    }
+}
diff --git a/Engine/Source/Editor/LevelInstanceEditor/Private/LevelInstanceActorDetails.cpp b/Engine/Source/Editor/LevelInstanceEditor/Private/LevelInstanceActorDetails.cpp
new file mode 100644
--- /dev/null
+++ b/Engine/Source/Editor/LevelInstanceEditor/Private/LevelInstanceActorDetails.cpp
@@ -0,0 +1,151 @@
+// Copyright Epic Games, Inc. All Rights Reserved.
+
+#include "LevelInstanceActorDetails.h"
+#include "DetailLayoutBuilder.h"
+#include "DetailCategoryBuilder.h"
+#include "DetailWidgetRow.h"
+#include "Widgets/Input/SButton.h"
+#include "Widgets/Input/SMultiLineEditableTextBox.h"
+#include "ScopedTransaction.h"
+
+#include "Engine/World.h"
+#include "LevelInstance/LevelInstanceActor.h"
+
+#define LOCTEXT_NAMESPACE "FLevelInstanceActorDetails"
+
+namespace LevelInstanceActorDetailsCallbacks
+{
+	static bool IsEditCommitButtonEnabled(TWeakObjectPtr<ALevelInstance> LevelInstanceActorPtr)
+	{
+		if (ALevelInstance* LevelInstanceActor = LevelInstanceActorPtr.Get())
+		{
+			return LevelInstanceActor->CanEdit() || LevelInstanceActor->CanCommit();
+		}
+
+		return false;
+	}
+
+	static FText GetEditCommitButtonText(TWeakObjectPtr<ALevelInstance> LevelInstanceActorPtr)
+	{
+		if (ALevelInstance* LevelInstanceActor = LevelInstanceActorPtr.Get())
+		{
+			if (LevelInstanceActor->CanCommit())
+			{
+				return LOCTEXT("CommitChanges", "Commit Changes");
+			}
+		}
+		
+		return LOCTEXT("Edit", "Edit");
+	}
+
+	static FText GetEditCommitReasonText(TWeakObjectPtr<ALevelInstance> LevelInstanceActorPtr)
+	{
+		FText Reason;
+		if (ALevelInstance* LevelInstanceActor = LevelInstanceActorPtr.Get())
+		{
+			if (!LevelInstanceActor->IsEditing())
+			{
+				LevelInstanceActor->CanEdit(&Reason);
+				return Reason;
+			}
+
+			LevelInstanceActor->CanCommit(&Reason);
+		}
+		return Reason;
+	}
+
+	static EVisibility GetEditCommitReasonVisibility(TWeakObjectPtr<ALevelInstance> LevelInstanceActorPtr)
+	{
+		if (ALevelInstance* LevelInstanceActor = LevelInstanceActorPtr.Get())
+		{
+			return IsEditCommitButtonEnabled(LevelInstanceActor) ? EVisibility::Collapsed : EVisibility::Visible;
+		}
+
+		return EVisibility::Collapsed;
+	}
+
+	static FReply OnEditCommitButtonClicked(TWeakObjectPtr<ALevelInstance> LevelInstanceActorPtr)
+	{
+		if (ALevelInstance* LevelInstanceActor = LevelInstanceActorPtr.Get())
+		{
+			if (LevelInstanceActor->CanCommit())
+			{
+				LevelInstanceActor->Commit();
+			}
+			else if (LevelInstanceActor->CanEdit())
+			{
+				LevelInstanceActor->Edit();
+			}
+		}
+		return FReply::Handled();
+	}
+}
+
+FLevelInstanceActorDetails::FLevelInstanceActorDetails()
+{
+}
+
+TSharedRef<IDetailCustomization> FLevelInstanceActorDetails::MakeInstance()
+{
+	return MakeShareable(new FLevelInstanceActorDetails);
+}
+
+void FLevelInstanceActorDetails::CustomizeDetails(IDetailLayoutBuilder& DetailBuilder)
+{
+	TArray<TWeakObjectPtr<UObject>> EditingObjects;
+	DetailBuilder.GetObjectsBeingCustomized(EditingObjects);
+
+	if (EditingObjects.Num() > 1)
+	{
+		return;
+	}
+
+	TWeakObjectPtr<ALevelInstance> EditingObject = Cast<ALevelInstance>(EditingObjects[0].Get());
+	UWorld* World = EditingObject->GetWorld();
+
+	if (!World)
+	{
+		return;
+	}
+
+	IDetailCategoryBuilder& LevelInstanceEditingCategory = DetailBuilder.EditCategory("LevelInstance Editing", FText::GetEmpty(), ECategoryPriority::Transform);
+
+	LevelInstanceEditingCategory.AddCustomRow(FText::GetEmpty())
+	[
+		SNew(SVerticalBox)
+		+ SVerticalBox::Slot()
+		.AutoHeight()
+		[
+			SNew(SHorizontalBox)
+			+ SHorizontalBox::Slot()
+			.FillWidth(1)
+			[
+				SNew(SMultiLineEditableTextBox)
+				.Visibility_Static(&LevelInstanceActorDetailsCallbacks::GetEditCommitReasonVisibility, EditingObject)
+				.Font(DetailBuilder.GetDetailFontBold())
+				.BackgroundColor(TAttribute<FSlateColor>::Create(TAttribute<FSlateColor>::FGetter::CreateLambda([]() { return FEditorStyle::GetColor("ErrorReporting.WarningBackgroundColor"); })))
+				.Text_Static(&LevelInstanceActorDetailsCallbacks::GetEditCommitReasonText, EditingObject)
+				.AutoWrapText(true)
+				.IsReadOnly(true)
+			]
+		]
+		+ SVerticalBox::Slot()
+		.AutoHeight()
+		[
+			SNew(SHorizontalBox)
+			+ SHorizontalBox::Slot()
+			.FillWidth(1)
+			[
+				SNew(SButton)
+				.IsEnabled_Static(&LevelInstanceActorDetailsCallbacks::IsEditCommitButtonEnabled, EditingObject)
+				.Text_Static(&LevelInstanceActorDetailsCallbacks::GetEditCommitButtonText, EditingObject)
+				.HAlign(HAlign_Center)
+				.OnClicked_Static(&LevelInstanceActorDetailsCallbacks::OnEditCommitButtonClicked, EditingObject)
+			]
+		]
+	];
+}
+
+
+
+#undef LOCTEXT_NAMESPACE
\ No newline at end of file
diff --git a/Engine/Source/Editor/LevelInstanceEditor/Private/LevelInstanceActorDetails.h b/Engine/Source/Editor/LevelInstanceEditor/Private/LevelInstanceActorDetails.h
new file mode 100644
--- /dev/null
+++ b/Engine/Source/Editor/LevelInstanceEditor/Private/LevelInstanceActorDetails.h
@@ -0,0 +1,22 @@
+// Copyright Epic Games, Inc. All Rights Reserved.
+
+#pragma once
+
+#include "CoreMinimal.h"
+#include "Types/SlateEnums.h"
+#include "IDetailCustomization.h"
+
+class FLevelInstanceActorDetails : public IDetailCustomization
+{
+public:
+	/** Makes a new instance of this detail layout class for a specific detail view requesting it */
+	static TSharedRef<IDetailCustomization> MakeInstance();
+
+	/** IDetailCustomization interface */
+	virtual void CustomizeDetails(IDetailLayoutBuilder& DetailBuilder) override;
+	/** End IDetailCustomization interface */
+
+private:
+	/** Use MakeInstance to create an instance of this class */
+	FLevelInstanceActorDetails();
+};
\ No newline at end of file
diff --git a/Engine/Source/Editor/LevelInstanceEditor/Private/LevelInstanceEditorModule.cpp b/Engine/Source/Editor/LevelInstanceEditor/Private/LevelInstanceEditorModule.cpp
new file mode 100644
--- /dev/null
+++ b/Engine/Source/Editor/LevelInstanceEditor/Private/LevelInstanceEditorModule.cpp
@@ -0,0 +1,671 @@
+// Copyright Epic Games, Inc. All Rights Reserved.
+#include "LevelInstanceEditorModule.h"
+#include "LevelInstanceActorDetails.h"
+#include "LevelInstance/LevelInstanceSubsystem.h"
+#include "LevelInstance/LevelInstanceActor.h"
+#include "LevelInstanceEditorSettings.h"
+#include "ToolMenus.h"
+#include "Editor.h"
+#include "LevelEditorMenuContext.h"
+#include "ContentBrowserMenuContexts.h"
+#include "ContentBrowserModule.h"
+#include "IContentBrowserSingleton.h"
+#include "LevelEditor.h"
+#include "Engine/Selection.h"
+#include "PropertyEditorModule.h"
+#include "EditorLevelUtils.h"
+#include "Modules/ModuleManager.h"
+#include "Misc/MessageDialog.h"
+#include "NewLevelDialogModule.h"
+#include "Interfaces/IMainFrameModule.h"
+#include "Editor/EditorEngine.h"
+#include "AssetToolsModule.h"
+#include "IAssetTools.h"
+#include "Factories/BlueprintFactory.h"
+#include "ClassViewerModule.h"
+#include "ClassViewerFilter.h"
+#include "Kismet2/BlueprintEditorUtils.h"
+#include "Misc/ScopeExit.h"
+#include "Widgets/Input/SNumericEntryBox.h"
+
+IMPLEMENT_MODULE( FLevelInstanceEditorModule, LevelInstanceEditor );
+
+#define LOCTEXT_NAMESPACE "LevelInstanceEditor"
+
+namespace LevelInstanceMenuUtils
+{
+	FToolMenuSection& CreateLevelInstanceSection(UToolMenu* Menu)
+	{
+		FToolMenuSection& Section = Menu->FindOrAddSection(TEXT("LevelInstance"));
+		Section.Label = LOCTEXT("LevelInstance", "Level Instance");
+		return Section;
+	}
+
+	void CreateEditSubMenu(UToolMenu* Menu, TArray<ALevelInstance*> LevelInstanceHierarchy, AActor* ContextActor)
+	{
+		FToolMenuSection& Section = Menu->AddSection(NAME_None, LOCTEXT("LevelInstanceContextEditSection", "Context"));
+		for (ALevelInstance* LevelInstanceActor : LevelInstanceHierarchy)
+		{
+			FToolUIAction LevelInstanceEditAction;
+			FText EntryDesc = LOCTEXT("LevelInstanceEditSubMenuEntry","");
+			const bool bCanEdit = LevelInstanceActor->CanEdit(&EntryDesc);
+
+			LevelInstanceEditAction.ExecuteAction.BindLambda([LevelInstanceActor, ContextActor](const FToolMenuContext&)
+			{
+				LevelInstanceActor->Edit(ContextActor);
+			});
+			LevelInstanceEditAction.CanExecuteAction.BindLambda([bCanEdit](const FToolMenuContext&)
+			{ 
+				return bCanEdit;
+			});
+
+			const FText EntryLabel = FText::Format(LOCTEXT("LevelInstanceName", "{0}:{1}"), FText::FromString(LevelInstanceActor->GetName()), FText::FromString(LevelInstanceActor->GetWorldAssetPackage()));
+			Section.AddMenuEntry(NAME_None, EntryLabel, EntryDesc, FSlateIcon(), LevelInstanceEditAction);
+		}
+	}
+
+	void CreateCommitSubMenu(UToolMenu* Menu, TArray<ALevelInstance*> LevelInstanceEdits, ALevelInstance* ContextLevelInstance)
+	{
+		FText OtherSectionLabel = LOCTEXT("LevelInstanceOtherCommitSection", "Other(s)");
+		FToolMenuSection* Section = &Menu->AddSection("LevelInstanceContextCommitSection", ContextLevelInstance != nullptr? LOCTEXT("LevelInstanceContextCommitSection", "Context") : OtherSectionLabel);
+		for (ALevelInstance* LevelInstanceActor : LevelInstanceEdits)
+		{
+			FText EntryDesc = LOCTEXT("LevelInstanceEditSubMenuEntry", "");
+			const bool bCanCommit = LevelInstanceActor->CanCommit(&EntryDesc);
+
+			FToolUIAction LevelInstanceEditAction;
+			LevelInstanceEditAction.ExecuteAction.BindLambda([LevelInstanceActor](const FToolMenuContext&)
+			{
+				LevelInstanceActor->Commit();
+			});
+			LevelInstanceEditAction.CanExecuteAction.BindLambda([bCanCommit](const FToolMenuContext&)
+			{
+				return bCanCommit;
+			});
+
+			const FText EntryLabel = FText::Format(LOCTEXT("LevelInstanceName", "{0}:{1}"), FText::FromString(LevelInstanceActor->GetName()), FText::FromString(LevelInstanceActor->GetWorldAssetPackage()));
+			Section->AddMenuEntry(NAME_None, EntryLabel, EntryDesc, FSlateIcon(), LevelInstanceEditAction);
+
+			if (ContextLevelInstance == LevelInstanceActor && LevelInstanceEdits.Num() > 1)
+			{
+				Section = &Menu->AddSection(FName("LevelInstanceOtherCommitSection"), OtherSectionLabel);
+			}
+		}
+	}
+
+	void CreateSetCurrentSubMenu(UToolMenu* Menu, TArray<ALevelInstance*> LevelInstanceEdits, ALevelInstance* ContextLevelInstance)
+	{
+		FText OtherSectionLabel = LOCTEXT("LevelInstanceOtherSetCurrentSection", "Other(s)");
+		FToolMenuSection* Section = &Menu->AddSection(FName("LevelInstanceContextSetCurrentSection"), ContextLevelInstance != nullptr? LOCTEXT("LevelInstanceContextSetCurrentSection", "Context") : OtherSectionLabel);
+		for (ALevelInstance* LevelInstanceActor : LevelInstanceEdits)
+		{
+			FToolUIAction LevelInstanceSetCurrentAction;
+			LevelInstanceSetCurrentAction.ExecuteAction.BindLambda([LevelInstanceActor](const FToolMenuContext&)
+			{
+				LevelInstanceActor->SetCurrent();
+			});
+		
+			const FText EntryLabel = FText::Format(LOCTEXT("LevelInstanceName", "{0}:{1}"), FText::FromString(LevelInstanceActor->GetName()), FText::FromString(LevelInstanceActor->GetWorldAssetPackage()));
+			Section->AddMenuEntry(NAME_None, EntryLabel, TAttribute<FText>(), FSlateIcon(), LevelInstanceSetCurrentAction);
+
+			if (LevelInstanceActor == ContextLevelInstance && LevelInstanceEdits.Num() > 1)
+			{
+				Section = &Menu->AddSection(FName("LevelInstanceOtherSetCurrentSection"), OtherSectionLabel);
+			}
+		}
+	}
+
+	void MoveSelectionToLevelInstance(ALevelInstance* DestinationLevelInstance)
+	{
+		if (ULevelInstanceSubsystem* LevelInstanceSubsystem = DestinationLevelInstance->GetLevelInstanceSubsystem())
+		{
+			TArray<AActor*> ActorsToMove;
+			ActorsToMove.Reserve(GEditor->GetSelectedActorCount());
+			for (FSelectionIterator It(GEditor->GetSelectedActorIterator()); It; ++It)
+			{
+				if (AActor* Actor = Cast<AActor>(*It))
+				{
+					ActorsToMove.Add(Actor);
+				}
+			}
+			
+			LevelInstanceSubsystem->MoveActorsTo(DestinationLevelInstance, ActorsToMove);
+		}
+	}
+
+	void CreateMoveSelectionToSubMenu(UToolMenu* Menu, TArray<ALevelInstance*> LevelInstanceEdits)
+	{
+		FToolMenuSection* Section = &Menu->AddSection(NAME_None);
+		for (ALevelInstance* LevelInstanceActor : LevelInstanceEdits)
+		{
+			FToolUIAction LevelInstanceMoveSelectionAction;
+			LevelInstanceMoveSelectionAction.ExecuteAction.BindLambda([LevelInstanceActor](const FToolMenuContext&)
+			{
+				MoveSelectionToLevelInstance(LevelInstanceActor);
+			});
+
+			const FText EntryLabel = FText::Format(LOCTEXT("LevelInstanceName", "{0}:{1}"), FText::FromString(LevelInstanceActor->GetName()), FText::FromString(LevelInstanceActor->GetWorldAssetPackage()));
+			Section->AddMenuEntry(NAME_None, EntryLabel, TAttribute<FText>(), FSlateIcon(), LevelInstanceMoveSelectionAction);
+		}
+	}
+
+	void CreateEditMenu(UToolMenu* Menu, AActor* ContextActor)
+	{
+		if (ULevelInstanceSubsystem* LevelInstanceSubsystem = ContextActor->GetWorld()->GetSubsystem<ULevelInstanceSubsystem>())
+		{
+			TArray<ALevelInstance*> LevelInstanceHierarchy;
+			LevelInstanceSubsystem->ForEachLevelInstanceAncestorsAndSelf(ContextActor, [&LevelInstanceHierarchy](ALevelInstance* AncestorLevelInstance)
+			{
+				LevelInstanceHierarchy.Add(AncestorLevelInstance);
+				return true;
+			});
+
+			if (LevelInstanceHierarchy.Num() > 0)
+			{
+				FToolMenuSection& Section = CreateLevelInstanceSection(Menu);
+				Section.AddSubMenu(
+					"EditLevelInstances",
+					LOCTEXT("EditLevelInstances", "Edit"),
+					TAttribute<FText>(),
+					FNewToolMenuDelegate::CreateStatic(&CreateEditSubMenu, MoveTemp(LevelInstanceHierarchy), ContextActor)
+				);
+			}
+		}
+	}
+	
+	void CreateCommitMenu(UToolMenu* Menu, AActor* ContextActor)
+	{
+		ALevelInstance* ContextLevelInstance = nullptr;
+		TArray<ALevelInstance*> LevelInstanceEdits;
+		if (ContextActor)
+		{
+			if (ULevelInstanceSubsystem* LevelInstanceSubsystem = ContextActor->GetWorld()->GetSubsystem<ULevelInstanceSubsystem>())
+			{
+				LevelInstanceSubsystem->ForEachLevelInstanceAncestorsAndSelf(ContextActor, [&ContextLevelInstance,&LevelInstanceEdits](ALevelInstance* LevelInstanceActor)
+				{
+					if (LevelInstanceActor->IsEditing())
+					{
+						ContextLevelInstance = LevelInstanceActor;
+						LevelInstanceEdits.Add(ContextLevelInstance);
+						return false;
+					}
+					return true;
+				});
+			}
+		}
+
+		if (ULevelInstanceSubsystem* LevelInstanceSubsystem = GEditor->GetEditorWorldContext().World()->GetSubsystem<ULevelInstanceSubsystem>())
+		{
+			LevelInstanceSubsystem->ForEachLevelInstanceEdit([&LevelInstanceEdits, ContextLevelInstance](ALevelInstance* LevelInstanceActor)
+			{
+				if (ContextLevelInstance != LevelInstanceActor)
+				{
+					LevelInstanceEdits.Add(LevelInstanceActor);
+				}
+				return true;
+			});
+		}
+
+		if (LevelInstanceEdits.Num() > 0)
+		{
+			FToolMenuSection& Section = CreateLevelInstanceSection(Menu);
+			Section.AddSubMenu(
+				"CommitLevelInstances",
+				LOCTEXT("CommitLevelInstances", "Commit"),
+				TAttribute<FText>(),
+				FNewToolMenuDelegate::CreateStatic(&CreateCommitSubMenu, MoveTemp(LevelInstanceEdits), ContextLevelInstance)
+			);
+		}
+	}
+
+	void CreateSetCurrentMenu(UToolMenu* Menu, AActor* ContextActor)
+	{
+		ALevelInstance* ContextLevelInstance = nullptr;
+		TArray<ALevelInstance*> LevelInstanceEdits;
+		if (ContextActor)
+		{
+			if (ULevelInstanceSubsystem* LevelInstanceSubsystem = ContextActor->GetWorld()->GetSubsystem<ULevelInstanceSubsystem>())
+			{
+				LevelInstanceSubsystem->ForEachLevelInstanceAncestorsAndSelf(ContextActor, [&ContextLevelInstance,&LevelInstanceEdits](ALevelInstance* LevelInstanceActor)
+				{
+					if (LevelInstanceActor->IsEditing())
+					{
+						if (!LevelInstanceActor->IsCurrent())
+						{
+							ContextLevelInstance = LevelInstanceActor;
+							LevelInstanceEdits.Add(ContextLevelInstance);
+						}
+						return false;
+					}
+					return true;
+				});
+			}
+		}
+
+		if (ULevelInstanceSubsystem* LevelInstanceSubsystem = GEditor->GetEditorWorldContext().World()->GetSubsystem<ULevelInstanceSubsystem>())
+		{
+			LevelInstanceSubsystem->ForEachLevelInstanceEdit([&LevelInstanceEdits, ContextLevelInstance](ALevelInstance* LevelInstanceActor)
+			{
+				if (ContextLevelInstance != LevelInstanceActor && !LevelInstanceActor->IsCurrent())
+				{
+					LevelInstanceEdits.Add(LevelInstanceActor);
+				}
+				return true;
+			});
+		}
+
+		if (LevelInstanceEdits.Num() > 0)
+		{
+			FToolMenuSection& Section = CreateLevelInstanceSection(Menu);
+			Section.AddSubMenu(
+				"SetCurrentLevelInstances",
+				LOCTEXT("SetCurrentLevelInstances", "Set Current"),
+				TAttribute<FText>(),
+				FNewToolMenuDelegate::CreateStatic(&CreateSetCurrentSubMenu, MoveTemp(LevelInstanceEdits), ContextLevelInstance)
+			);
+		}
+	}
+
+	void CreateMoveSelectionToMenu(UToolMenu* Menu)
+	{
+		if (GEditor->GetSelectedActorCount() > 0)
+		{
+
+			TArray<ALevelInstance*> LevelInstanceEdits;
+			if (ULevelInstanceSubsystem* LevelInstanceSubsystem = GEditor->GetEditorWorldContext().World()->GetSubsystem<ULevelInstanceSubsystem>())
+			{
+				LevelInstanceSubsystem->ForEachLevelInstanceEdit([&LevelInstanceEdits](ALevelInstance* LevelInstanceActor)
+				{
+					LevelInstanceEdits.Add(LevelInstanceActor);
+					return true;
+				});
+			}
+
+			if (LevelInstanceEdits.Num() > 0)
+			{
+				FToolMenuSection& Section = CreateLevelInstanceSection(Menu);
+				Section.AddSubMenu(
+					"MoveSelectionToLevelInstances",
+					LOCTEXT("MoveSelectionToLevelInstances", "Move Selection to"),
+					TAttribute<FText>(),
+					FNewToolMenuDelegate::CreateStatic(&CreateMoveSelectionToSubMenu, MoveTemp(LevelInstanceEdits))
+				);
+			}
+		}
+	}
+
+	void CreateLevelInstanceFromSelection(ULevelInstanceSubsystem* LevelInstanceSubsystem)
+	{
+		TArray<AActor*> ActorsToMove;
+		ActorsToMove.Reserve(GEditor->GetSelectedActorCount());
+		for (FSelectionIterator It(GEditor->GetSelectedActorIterator()); It; ++It)
+		{
+			if (AActor* Actor = Cast<AActor>(*It))
+			{
+				ActorsToMove.Add(Actor);
+			}
+		}
+
+		IMainFrameModule& MainFrameModule = FModuleManager::GetModuleChecked<IMainFrameModule>("MainFrame");
+		FNewLevelDialogModule& NewLevelDialogModule = FModuleManager::LoadModuleChecked<FNewLevelDialogModule>("NewLevelDialog");
+		FString TemplateMapPackage;
+		if (!GetMutableDefault<ULevelInstanceEditorSettings>()->TemplateMapInfos.Num() || NewLevelDialogModule.CreateAndShowTemplateDialog(MainFrameModule.GetParentWindow(), LOCTEXT("LevelInstanceTemplateDialog", "Choose Level Instance Template..."), GetMutableDefault<ULevelInstanceEditorSettings>()->TemplateMapInfos, TemplateMapPackage))
+		{
+			UPackage* TemplatePackage = !TemplateMapPackage.IsEmpty() ? LoadPackage(nullptr, *TemplateMapPackage, LOAD_None) : nullptr;
+			UWorld* TemplateWorld = TemplatePackage ? UWorld::FindWorldInPackage(TemplatePackage) : nullptr;
+
+			if (!LevelInstanceSubsystem->CreateLevelInstanceFrom(ActorsToMove, TemplateWorld))
+			{
+				FText Title = LOCTEXT("CreateFromSelectionFailTitle", "Create from selection failed");
+				FMessageDialog::Open(EAppMsgType::Ok, LOCTEXT("CreateFromSelectionFailMsg", "Failed to create LevelInstance from selection. Check log for details."), &Title);
+			}
+		}
+	}
+
+	void CreateCreateMenu(UToolMenu* Menu)
+	{
+		if (ULevelInstanceSubsystem* LevelInstanceSubsystem = GEditor->GetEditorWorldContext().World()->GetSubsystem<ULevelInstanceSubsystem>())
+		{
+			FToolMenuSection& Section = CreateLevelInstanceSection(Menu);
+			FToolUIAction UIAction;
+			UIAction.ExecuteAction.BindLambda([LevelInstanceSubsystem](const FToolMenuContext& MenuContext)
+			{
+				CreateLevelInstanceFromSelection(LevelInstanceSubsystem);
+			});
+			UIAction.CanExecuteAction.BindLambda([](const FToolMenuContext& MenuContext)
+			{
+				return GEditor->GetSelectedActorCount() > 0;
+			});
+
+			Section.AddMenuEntry(
+				"CreateLevelInstanceFromSelection",
+				LOCTEXT("CreateLevelInstanceFromSelection", "Create from selection"),
+				TAttribute<FText>(),
+				TAttribute<FSlateIcon>(),
+				UIAction);
+		}
+	}
+
+	void CreateSaveAsMenu(UToolMenu* Menu, AActor* ContextActor)
+	{
+		ALevelInstance* ContextLevelInstance = nullptr;
+		if (ContextActor)
+		{
+			if (ULevelInstanceSubsystem* LevelInstanceSubsystem = ContextActor->GetWorld()->GetSubsystem<ULevelInstanceSubsystem>())
+			{
+				LevelInstanceSubsystem->ForEachLevelInstanceAncestorsAndSelf(ContextActor, [&ContextLevelInstance](ALevelInstance* LevelInstanceActor)
+					{
+						if (LevelInstanceActor->IsEditing())
+						{
+							ContextLevelInstance = LevelInstanceActor;
+							return false;
+						}
+						return true;
+					});
+			}
+		}
+
+		if (ContextLevelInstance)
+		{
+			FToolMenuSection& Section = CreateLevelInstanceSection(Menu);
+			FText EntryDesc = LOCTEXT("LevelInstanceEditSubMenuEntry", "");
+			const bool bCanCommit = ContextLevelInstance->CanCommit(&EntryDesc);
+
+			FToolUIAction SaveAction;
+			SaveAction.ExecuteAction.BindLambda([ContextLevelInstance](const FToolMenuContext& MenuContext)
+				{
+					ContextLevelInstance->SaveAs();
+				});
+			SaveAction.CanExecuteAction.BindLambda([bCanCommit](const FToolMenuContext& MenuContext)
+				{
+					return bCanCommit;
+				});
+
+			Section.AddMenuEntry(
+				"SaveLevelInstanceAs",
+				LOCTEXT("SaveLevelInstanceAs", "Save Level as..."),
+				TAttribute<FText>(),
+				TAttribute<FSlateIcon>(),
+				SaveAction);
+		}
+	}
+
+	void CreateBreakSubMenu(UToolMenu* Menu, ALevelInstance* ContextLevelInstance)
+	{
+		static int32 BreakLevels = 1;
+
+		check(ContextLevelInstance);
+
+		if (ULevelInstanceSubsystem* LevelInstanceSubsystem = ContextLevelInstance->GetWorld()->GetSubsystem<ULevelInstanceSubsystem>())
+		{
+			FToolMenuSection& Section = Menu->AddSection(NAME_None, LOCTEXT("LevelInstanceBreakSection", "Break Level Instance"));
+			TSharedRef<SWidget> MenuWidget =
+				SNew(SVerticalBox)
+
+				+SVerticalBox::Slot()
+				[
+					SNew(SHorizontalBox)
+					+ SHorizontalBox::Slot()
+					[
+						SNew(SNumericEntryBox<int32>)
+						.MinValue(1)
+						.Value_Lambda([]() { return BreakLevels; })
+						.OnValueChanged_Lambda([](int32 InValue) { BreakLevels = InValue; })
+						.LabelPadding(0)
+						.Label()
+						[
+							SNumericEntryBox<int32>::BuildLabel(LOCTEXT("BreakLevelsLabel", "Levels"), FLinearColor::White, SNumericEntryBox<int32>::BlueLabelBackgroundColor)
+						]
+					]
+				]
+
+				+SVerticalBox::Slot()
+				.VAlign(VAlign_Center)
+				.HAlign(HAlign_Center)
+				.Padding(0, 5, 0, 0)
+				[
+					SNew(SButton)
+					.HAlign(HAlign_Center)
+					.ContentPadding(FEditorStyle::GetMargin("StandardDialog.ContentPadding"))
+					.OnClicked_Lambda([ContextLevelInstance, LevelInstanceSubsystem]() {LevelInstanceSubsystem->BreakLevelInstance(ContextLevelInstance, BreakLevels); return FReply::Handled(); })
+					.Text(LOCTEXT("BreakLevelInstances_BreakLevelInstanceButton", "Break Level Instance"))
+				];
+
+			Section.AddEntry(FToolMenuEntry::InitWidget("SetBreakLevels", MenuWidget, FText::GetEmpty(), false));
+		}
+	}
+
+	void CreateBreakMenu(UToolMenu* Menu, AActor* ContextActor)
+	{
+		check(ContextActor);
+
+		if (ULevelInstanceSubsystem* LevelInstanceSubsystem = ContextActor->GetWorld()->GetSubsystem<ULevelInstanceSubsystem>())
+		{
+			ALevelInstance* ContextLevelInstance = nullptr;
+
+			// Find the top level LevelInstance
+			LevelInstanceSubsystem->ForEachLevelInstanceAncestorsAndSelf(ContextActor, [LevelInstanceSubsystem, ContextActor, &ContextLevelInstance](ALevelInstance* Ancestor)
+				{
+					if (Ancestor->GetLevel() == ContextActor->GetWorld()->GetCurrentLevel())
+					{
+						ContextLevelInstance = Ancestor;
+						return false;
+					}
+					return true;
+				});
+
+			if (ContextLevelInstance && !ContextLevelInstance->IsEditing())
+			{
+				FToolMenuSection& Section = CreateLevelInstanceSection(Menu);
+
+				Section.AddSubMenu(
+					"BreakLevelInstances",
+					LOCTEXT("BreakLevelInstances", "Break..."),
+					TAttribute<FText>(),
+					FNewToolMenuDelegate::CreateStatic(&CreateBreakSubMenu, ContextLevelInstance)
+				);
+			}
+		}
+	}
+
+	class FLevelInstanceClassFilter : public IClassViewerFilter
+	{
+	public:
+		
+		virtual bool IsClassAllowed(const FClassViewerInitializationOptions& InInitOptions, const UClass* InClass, TSharedRef< FClassViewerFilterFuncs > InFilterFuncs) override
+		{
+			return InClass && InClass->IsChildOf(ALevelInstance::StaticClass()) && !InClass->HasAnyClassFlags(CLASS_Deprecated);
+		}
+
+		virtual bool IsUnloadedClassAllowed(const FClassViewerInitializationOptions& InInitOptions, const TSharedRef< const IUnloadedBlueprintData > InUnloadedClassData, TSharedRef< FClassViewerFilterFuncs > InFilterFuncs) override
+		{
+			return InUnloadedClassData->IsChildOf(ALevelInstance::StaticClass()) && !InUnloadedClassData->HasAnyClassFlags(CLASS_Deprecated);
+		}
+	};
+
+	void CreateBlueprintFromWorld(UWorld* WorldAsset)
+	{
+		TSoftObjectPtr<UWorld> LevelInstancePtr(WorldAsset);
+
+		int32 LastSlashIndex = 0;
+		FString LongPackageName = LevelInstancePtr.GetLongPackageName();
+		LongPackageName.FindLastChar('/', LastSlashIndex);
+		
+		FString PackagePath = LongPackageName.Mid(0, LastSlashIndex == INDEX_NONE ? MAX_int32 : LastSlashIndex);
+		FString AssetName = LevelInstancePtr.GetAssetName() + "_LevelInstance";
+		IAssetTools& AssetTools = FAssetToolsModule::GetModule().Get();
+
+		UBlueprintFactory* BlueprintFactory = NewObject<UBlueprintFactory>();
+		BlueprintFactory->AddToRoot();
+		BlueprintFactory->OnConfigurePropertiesDelegate.BindLambda([](FClassViewerInitializationOptions* Options)
+		{
+			Options->bShowDefaultClasses = false;
+			Options->bIsBlueprintBaseOnly = false;
+			Options->InitiallySelectedClass = ALevelInstance::StaticClass();
+			Options->bIsActorsOnly = true;
+			Options->ClassFilter = MakeShareable(new FLevelInstanceClassFilter);
+		});
+		ON_SCOPE_EXIT
+		{
+			BlueprintFactory->OnConfigurePropertiesDelegate.Unbind();
+			BlueprintFactory->RemoveFromRoot();
+		};
+
+		if (UBlueprint* NewBlueprint = Cast<UBlueprint>(AssetTools.CreateAssetWithDialog(AssetName, PackagePath, UBlueprint::StaticClass(), BlueprintFactory, FName("Create LevelInstance Blueprint"))))
+		{
+			ALevelInstance* CDO = CastChecked<ALevelInstance>(NewBlueprint->GeneratedClass->GetDefaultObject());
+			CDO->SetWorldAsset(LevelInstancePtr);
+			FBlueprintEditorUtils::MarkBlueprintAsModified(NewBlueprint);
+			
+			FContentBrowserModule& ContentBrowserModule = FModuleManager::LoadModuleChecked<FContentBrowserModule>("ContentBrowser");
+			TArray<UObject*> Assets;
+			Assets.Add(NewBlueprint);
+			ContentBrowserModule.Get().SyncBrowserToAssets(Assets);
+		}		
+	}
+
+	void CreateBlueprintFromMenu(UToolMenu* Menu, UWorld* WorldAsset)
+	{
+		FToolMenuSection& Section = CreateLevelInstanceSection(Menu);
+		FToolUIAction UIAction;
+		UIAction.ExecuteAction.BindLambda([WorldAsset](const FToolMenuContext& MenuContext)
+		{
+			CreateBlueprintFromWorld(WorldAsset);
+		});
+
+		Section.AddMenuEntry(
+			"CreateLevelInstanceBlueprint",
+			LOCTEXT("CreateLevelInstanceBlueprint", "New Blueprint..."),
+			TAttribute<FText>(),
+			TAttribute<FSlateIcon>(),
+			UIAction);
+	}
+};
+
+void FLevelInstanceEditorModule::StartupModule()
+{
+	ExtendContextMenu();
+
+	FPropertyEditorModule& PropertyModule = FModuleManager::LoadModuleChecked<FPropertyEditorModule>("PropertyEditor");
+	PropertyModule.RegisterCustomClassLayout("LevelInstanceActor", FOnGetDetailCustomizationInstance::CreateStatic(&FLevelInstanceActorDetails::MakeInstance));
+	PropertyModule.NotifyCustomizationModuleChanged();
+
+	// GEditor needs to be set before this module is loaded
+	check(GEditor);
+	GEditor->OnLevelActorDeleted().AddRaw(this, &FLevelInstanceEditorModule::OnLevelActorDeleted);
+	
+	EditorLevelUtils::CanMoveActorToLevelDelegate.AddRaw(this, &FLevelInstanceEditorModule::CanMoveActorToLevel);
+
+	FLevelEditorModule& LevelEditorModule = FModuleManager::LoadModuleChecked<FLevelEditorModule>("LevelEditor");
+	LevelEditorModule.OnMapChanged().AddRaw(this, &FLevelInstanceEditorModule::OnMapChanged);
+}
+
+void FLevelInstanceEditorModule::ShutdownModule()
+{
+	if (GEditor)
+	{
+		GEditor->OnLevelActorDeleted().RemoveAll(this);
+	}
+
+	EditorLevelUtils::CanMoveActorToLevelDelegate.RemoveAll(this);
+
+	if (FLevelEditorModule* LevelEditorModule = FModuleManager::GetModulePtr<FLevelEditorModule>("LevelEditor"))
+	{
+		LevelEditorModule->OnMapChanged().RemoveAll(this);
+	}
+}
+
+void FLevelInstanceEditorModule::OnLevelActorDeleted(AActor* Actor)
+{
+	if (ULevelInstanceSubsystem* LevelInstanceSubsystem = Actor->GetWorld()->GetSubsystem<ULevelInstanceSubsystem>())
+	{
+		LevelInstanceSubsystem->OnActorDeleted(Actor);
+	}
+}
+
+void FLevelInstanceEditorModule::OnMapChanged(UWorld* World, EMapChangeType MapChangeType)
+{
+	// On Map Changed, Users will be asked to save unchanged maps. Once we hit the teardown we need to force
+	// LevelInstance Edits to be cancelled. If they are still dirty it means the user decided not to save changes.
+	if (World && MapChangeType == EMapChangeType::TearDownWorld)
+	{
+		if (ULevelInstanceSubsystem* LevelInstanceSubsystem = World->GetSubsystem<ULevelInstanceSubsystem>())
+		{
+			LevelInstanceSubsystem->DiscardEdits();
+		}
+	}
+}
+
+void FLevelInstanceEditorModule::CanMoveActorToLevel(const AActor* ActorToMove, const ULevel* DestLevel, bool& bOutCanMove)
+{
+	if (UWorld* World = ActorToMove->GetWorld())
+	{
+		if (ULevelInstanceSubsystem* LevelInstanceSubsystem = World->GetSubsystem<ULevelInstanceSubsystem>())
+		{
+			if (!LevelInstanceSubsystem->CanMoveActorToLevel(ActorToMove))
+			{
+				bOutCanMove = false;
+				return;
+			}
+		}
+	}
+}
+
+void FLevelInstanceEditorModule::ExtendContextMenu()
+{
+	if (UToolMenu* Menu = UToolMenus::Get()->ExtendMenu("LevelEditor.ActorContextMenu"))
+	{
+		FToolMenuSection& Section = Menu->AddDynamicSection("ActorLevelInstance", FNewToolMenuDelegate::CreateLambda([this](UToolMenu* ToolMenu)
+		{
+			if (ToolMenu)
+			{
+				if (ULevelEditorContextMenuContext* LevelEditorMenuContext = ToolMenu->Context.FindContext<ULevelEditorContextMenuContext>())
+				{
+					AActor* ContextActor = LevelEditorMenuContext->HitProxyActor;
+					if (!ContextActor && GEditor->GetSelectedActorCount() == 1)
+					{
+						ContextActor = Cast<AActor>(GEditor->GetSelectedActors()->GetSelectedObject(0));
+					}
+
+					if (ContextActor)
+					{
+						LevelInstanceMenuUtils::CreateEditMenu(ToolMenu, ContextActor);
+						LevelInstanceMenuUtils::CreateCommitMenu(ToolMenu, ContextActor);
+						LevelInstanceMenuUtils::CreateSaveAsMenu(ToolMenu, ContextActor);
+						LevelInstanceMenuUtils::CreateBreakMenu(ToolMenu, ContextActor);
+					}
+
+					LevelInstanceMenuUtils::CreateSetCurrentMenu(ToolMenu, ContextActor);
+					LevelInstanceMenuUtils::CreateMoveSelectionToMenu(ToolMenu);
+					LevelInstanceMenuUtils::CreateCreateMenu(ToolMenu);
+				}
+			}
+		}), FToolMenuInsert(NAME_None, EToolMenuInsertType::First));
+	}
+
+	if (UToolMenu* WorldAssetMenu = UToolMenus::Get()->ExtendMenu("ContentBrowser.AssetContextMenu.World"))
+	{
+		FToolMenuSection& Section = WorldAssetMenu->AddDynamicSection("ActorLevelInstance", FNewToolMenuDelegate::CreateLambda([this](UToolMenu* ToolMenu)
+		{
+			if (ToolMenu)
+			{
+				if (UContentBrowserAssetContextMenuContext* AssetMenuContext = ToolMenu->Context.FindContext<UContentBrowserAssetContextMenuContext>())
+				{
+					if (AssetMenuContext->SelectedObjects.Num() != 1)
+					{
+						return;
+					}
+					// World is already loaded by the AssetContextMenu code
+					if (UWorld* WorldAsset = Cast<UWorld>(AssetMenuContext->SelectedObjects[0].Get()))
+					{
+						LevelInstanceMenuUtils::CreateBlueprintFromMenu(ToolMenu, WorldAsset);
+					}
+				}
+			}
+		}), FToolMenuInsert(NAME_None, EToolMenuInsertType::Default));
+	}
+	
+}
+
+#undef LOCTEXT_NAMESPACE
+
diff --git a/Engine/Source/Editor/LevelInstanceEditor/Private/LevelInstanceEditorModule.h b/Engine/Source/Editor/LevelInstanceEditor/Private/LevelInstanceEditorModule.h
new file mode 100644
--- /dev/null
+++ b/Engine/Source/Editor/LevelInstanceEditor/Private/LevelInstanceEditorModule.h
@@ -0,0 +1,35 @@
+// Copyright Epic Games, Inc. All Rights Reserved.
+#pragma once
+
+#include "CoreMinimal.h"
+#include "LevelInstance/ILevelInstanceEditorModule.h"
+
+class AActor;
+class ULevel;
+enum class EMapChangeType : uint8;
+
+/**
+ * The module holding all of the UI related pieces for LevelInstance management
+ */
+class FLevelInstanceEditorModule : public ILevelInstanceEditorModule
+{
+public:
+	virtual ~FLevelInstanceEditorModule(){}
+
+	/**
+	 * Called right after the module DLL has been loaded and the module object has been created
+	 */
+	virtual void StartupModule();
+
+	/**
+	 * Called before the module is unloaded, right before the module object is destroyed.
+	 */
+	virtual void ShutdownModule();
+
+private:
+	void OnLevelActorDeleted(AActor* Actor);
+	void OnMapChanged(UWorld* World, EMapChangeType MapChangeType);
+	void CanMoveActorToLevel(const AActor* ActorToMove, const ULevel* DestLevel, bool& bOutCanMove);
+
+	void ExtendContextMenu();
+};
diff --git a/Engine/Source/Editor/LevelInstanceEditor/Private/LevelInstanceEditorSettings.cpp b/Engine/Source/Editor/LevelInstanceEditor/Private/LevelInstanceEditorSettings.cpp
new file mode 100644
--- /dev/null
+++ b/Engine/Source/Editor/LevelInstanceEditor/Private/LevelInstanceEditorSettings.cpp
@@ -0,0 +1,7 @@
+// Copyright Epic Games, Inc. All Rights Reserved.
+
+#include "LevelInstanceEditorSettings.h"
+
+ULevelInstanceEditorSettings::ULevelInstanceEditorSettings()
+{
+}
diff --git a/Engine/Source/Editor/LevelInstanceEditor/Public/LevelInstanceEditorSettings.h b/Engine/Source/Editor/LevelInstanceEditor/Public/LevelInstanceEditorSettings.h
new file mode 100644
--- /dev/null
+++ b/Engine/Source/Editor/LevelInstanceEditor/Public/LevelInstanceEditorSettings.h
@@ -0,0 +1,22 @@
+// Copyright Epic Games, Inc. All Rights Reserved.
+
+#pragma once
+
+#include "CoreMinimal.h"
+#include "UObject/ObjectMacros.h"
+#include "UObject/Object.h"
+#include "Editor/UnrealEdEngine.h"
+#include "LevelInstanceEditorSettings.generated.h"
+
+UCLASS(config = Editor)
+class ULevelInstanceEditorSettings : public UObject
+{
+	GENERATED_BODY()
+
+public:
+	ULevelInstanceEditorSettings();
+
+	/** List of info for all known LevelInstance template maps */
+	UPROPERTY(config)
+	TArray<FTemplateMapInfo> TemplateMapInfos;
+};
\ No newline at end of file
diff --git a/Engine/Source/Editor/SceneOutliner/Private/ActorBrowsingMode.cpp b/Engine/Source/Editor/SceneOutliner/Private/ActorBrowsingMode.cpp
--- a/Engine/Source/Editor/SceneOutliner/Private/ActorBrowsingMode.cpp
+++ b/Engine/Source/Editor/SceneOutliner/Private/ActorBrowsingMode.cpp
@@ -26,9 +26,9 @@
 #include "ComponentTreeItem.h"
 #include "ActorBrowsingModeSettings.h"
 #include "ScopedTransaction.h"
-#include "Foundation/FoundationActor.h"
-#include "Foundation/FoundationSubsystem.h"
-#include "Foundation/FoundationEditorInstanceActor.h"
+#include "LevelInstance/LevelInstanceActor.h"
+#include "LevelInstance/LevelInstanceSubsystem.h"
+#include "LevelInstance/LevelInstanceEditorInstanceActor.h"
 
 DEFINE_LOG_CATEGORY_STATIC(LogActorBrowser, Log, All);
 
@@ -105,32 +105,32 @@ FActorBrowsingMode::FActorBrowsingMode(SSceneOutliner* InSceneOutliner, TWeakObj
 
 	FilterInfoMap.Add(TEXT("HideComponentsFilter"), HideComponentsInfo);
 
-	FSceneOutlinerFilterInfo HideFoundationsInfo(LOCTEXT("ToggleHideFoundations", "Hide Foundations"), LOCTEXT("ToggleHideFoundationsToolTip", "When enabled, hides all foundation content."), SharedSettings->bHideFoundationHierarchy, FCreateSceneOutlinerFilter::CreateStatic(&FActorBrowsingMode::CreateHideFoundationsFilter));
-	HideFoundationsInfo.OnToggle().AddLambda([this](bool bIsActive)
+	FSceneOutlinerFilterInfo HideLevelInstancesInfo(LOCTEXT("ToggleHideLevelInstances", "Hide Level Instances"), LOCTEXT("ToggleHideLevelInstancesToolTip", "When enabled, hides all level instance content."), SharedSettings->bHideLevelInstanceHierarchy, FCreateSceneOutlinerFilter::CreateStatic(&FActorBrowsingMode::CreateHideLevelInstancesFilter));
+	HideLevelInstancesInfo.OnToggle().AddLambda([this](bool bIsActive)
 		{
 			UActorBrowsingModeSettings* Settings = GetMutableDefault<UActorBrowsingModeSettings>();
-			Settings->bHideFoundationHierarchy = bHideFoundationHierarchy = bIsActive;
+			Settings->bHideLevelInstanceHierarchy = bHideLevelInstanceHierarchy = bIsActive;
 			Settings->PostEditChange();
 
 			if (auto ActorHierarchy = StaticCast<FActorHierarchy*>(Hierarchy.Get()))
 			{
-				ActorHierarchy->SetShowingFoundations(!bIsActive);
+				ActorHierarchy->SetShowingLevelInstances(!bIsActive);
 			}
 		});
-	FilterInfoMap.Add(TEXT("HideFoundationsFilter"), HideFoundationsInfo);
+	FilterInfoMap.Add(TEXT("HideLevelInstancesFilter"), HideLevelInstancesInfo);
 
-	// Add a filter which sets the interactive mode of foundation items and their children
+	// Add a filter which sets the interactive mode of LevelInstance items and their children
 	SceneOutliner->AddFilter(MakeShared<FActorFilter>(FActorTreeItem::FFilterPredicate::CreateStatic([](const AActor* Actor) {return true; }), FSceneOutlinerFilter::EDefaultBehaviour::Pass, FActorTreeItem::FFilterPredicate::CreateLambda([this](const AActor* Actor)
 		{
-			if (!bHideFoundationHierarchy)
+			if (!bHideLevelInstanceHierarchy)
 			{
-				if (const UFoundationSubsystem* FoundationSubsystem = RepresentingWorld->GetSubsystem<UFoundationSubsystem>())
+				if (const ULevelInstanceSubsystem* LevelInstanceSubsystem = RepresentingWorld->GetSubsystem<ULevelInstanceSubsystem>())
 				{
 					// if actor has a valid parent and the parent is not being edited,
 					// then the actor should not be selectable.
-					if (const AFoundationActor* ParentFoundation = FoundationSubsystem->GetParentFoundation(Actor))
+					if (const ALevelInstance* ParentLevelInstance = LevelInstanceSubsystem->GetParentLevelInstance(Actor))
 					{
-						if (!FoundationSubsystem->IsEditingFoundation(ParentFoundation))
+						if (!LevelInstanceSubsystem->IsEditingLevelInstance(ParentLevelInstance))
 						{
 							return false;
 						}
@@ -270,23 +270,23 @@ TSharedRef<FSceneOutlinerFilter> FActorBrowsingMode::CreateHideComponentsFilter(
 		FSceneOutlinerFilter::EDefaultBehaviour::Pass));
 }
 
-TSharedRef<FSceneOutlinerFilter> FActorBrowsingMode::CreateHideFoundationsFilter()
+TSharedRef<FSceneOutlinerFilter> FActorBrowsingMode::CreateHideLevelInstancesFilter()
 {
 	return MakeShareable(new FActorFilter(FActorTreeItem::FFilterPredicate::CreateStatic([](const AActor* Actor)
 		{
-			// Check if actor belongs to a foundation
-			if (const UFoundationSubsystem* FoundationSubsystem = Actor->GetWorld()->GetSubsystem<UFoundationSubsystem>())
+			// Check if actor belongs to a LevelInstance
+			if (const ULevelInstanceSubsystem* LevelInstanceSubsystem = Actor->GetWorld()->GetSubsystem<ULevelInstanceSubsystem>())
 			{
-				if (const AFoundationActor* ParentFoundation = FoundationSubsystem->GetParentFoundation(Actor))
+				if (const ALevelInstance* ParentLevelInstance = LevelInstanceSubsystem->GetParentLevelInstance(Actor))
 				{
-					if (!FoundationSubsystem->IsEditingFoundation(ParentFoundation))
+					if (!LevelInstanceSubsystem->IsEditingLevelInstance(ParentLevelInstance))
 					{
 						return false;
 					}
 				}
 			}
-			// Or if the actor itself is a foundation editor instance
-			return Cast<AFoundationEditorInstanceActor>(Actor) == nullptr;
+			// Or if the actor itself is a LevelInstance editor instance
+			return Cast<ALevelInstanceEditorInstanceActor>(Actor) == nullptr;
 		}), FSceneOutlinerFilter::EDefaultBehaviour::Pass));
 }
 
@@ -554,16 +554,16 @@ void FActorBrowsingMode::OnItemDoubleClick(FSceneOutlinerTreeItemPtr Item)
 		AActor* Actor = ActorItem->Actor.Get();
 		check(Actor);
 
-		AFoundationActor* FoundationActor = Cast<AFoundationActor>(Actor);
-		if (FoundationActor && FSlateApplication::Get().GetModifierKeys().IsAltDown())
+		ALevelInstance* LevelInstanceActor = Cast<ALevelInstance>(Actor);
+		if (LevelInstanceActor && FSlateApplication::Get().GetModifierKeys().IsAltDown())
 		{
-			if (FoundationActor->CanEdit())
+			if (LevelInstanceActor->CanEdit())
 			{
-				FoundationActor->Edit();
+				LevelInstanceActor->Edit();
 			}
-			else if (FoundationActor->CanCommit())
+			else if (LevelInstanceActor->CanCommit())
 			{
-				FoundationActor->Commit();
+				LevelInstanceActor->Commit();
 			}
 		}
 		else if (Item->CanInteract())
@@ -785,15 +785,15 @@ FSceneOutlinerDragValidationInfo FActorBrowsingMode::ValidateDrop(const ISceneOu
 			return FSceneOutlinerDragValidationInfo(ESceneOutlinerDropCompatibility::IncompatibleGeneric, FText());
 		}
 
-		const AFoundationActor* FoundationTarget = Cast<AFoundationActor>(ActorTarget);
-		const UFoundationSubsystem* FoundationSubsystem = RepresentingWorld->GetSubsystem<UFoundationSubsystem>();
+		const ALevelInstance* LevelInstanceTarget = Cast<ALevelInstance>(ActorTarget);
+		const ULevelInstanceSubsystem* LevelInstanceSubsystem = RepresentingWorld->GetSubsystem<ULevelInstanceSubsystem>();
 
-		if (FoundationTarget)
+		if (LevelInstanceTarget)
 		{
-			check(FoundationSubsystem);
-			if (!FoundationTarget->IsEditing())
+			check(LevelInstanceSubsystem);
+			if (!LevelInstanceTarget->IsEditing())
 			{
-				return FSceneOutlinerDragValidationInfo(ESceneOutlinerDropCompatibility::IncompatibleGeneric, LOCTEXT("Error_AttachToClosedFoundation", "Cannot attach to foundation which is not being edited"));
+				return FSceneOutlinerDragValidationInfo(ESceneOutlinerDropCompatibility::IncompatibleGeneric, LOCTEXT("Error_AttachToClosedLevelInstance", "Cannot attach to LevelInstance which is not being edited"));
 			}
 		}
 
@@ -808,20 +808,20 @@ FSceneOutlinerDragValidationInfo FActorBrowsingMode::ValidateDrop(const ISceneOu
 			{
 				if (bCanAttach)
 				{
-					if (FoundationSubsystem)
+					if (LevelInstanceSubsystem)
 					{
-						// Either all actors must be in a foundation or none of them
-						if (const AFoundationActor* ParentFoundation = FoundationSubsystem->GetParentFoundation(DragActor))
+						// Either all actors must be in a LevelInstance or none of them
+						if (const ALevelInstance* ParentLevelInstance = LevelInstanceSubsystem->GetParentLevelInstance(DragActor))
 						{
-							if (!ParentFoundation->IsEditing())
+							if (!ParentLevelInstance->IsEditing())
 							{
-								return FSceneOutlinerDragValidationInfo(ESceneOutlinerDropCompatibility::IncompatibleGeneric, LOCTEXT("Error_RemoveEditingFoundation", "Cannot detach from a foundation which is not being edited"));
+								return FSceneOutlinerDragValidationInfo(ESceneOutlinerDropCompatibility::IncompatibleGeneric, LOCTEXT("Error_RemoveEditingLevelInstance", "Cannot detach from a LevelInstance which is not being edited"));
 							}
 						}
 
-						if (!FoundationSubsystem->CanMoveActorToLevel(DragActor))
+						if (!LevelInstanceSubsystem->CanMoveActorToLevel(DragActor))
 						{
-							AttachErrorMsg = LOCTEXT("Error_MoveActorToFoundation", "Cannot move foundation while it or its children are being edited");
+							AttachErrorMsg = LOCTEXT("Error_MoveActorToLevelInstance", "Cannot move LevelInstance while it or its children are being edited");
 							bCanAttach = bDraggedOntoAttachmentParent = false;
 							break;
 						}
@@ -833,7 +833,7 @@ FSceneOutlinerDragValidationInfo FActorBrowsingMode::ValidateDrop(const ISceneOu
 						bCanAttach = bDraggedOntoAttachmentParent = false;
 						break;
 					}
-					if (!FoundationTarget && !GEditor->CanParentActors(ActorTarget, DragActor, &AttachErrorMsg))
+					if (!LevelInstanceTarget && !GEditor->CanParentActors(ActorTarget, DragActor, &AttachErrorMsg))
 					{
 						bCanAttach = false;
 					}
@@ -941,29 +941,29 @@ FSceneOutlinerDragValidationInfo FActorBrowsingMode::ValidateDrop(const ISceneOu
 
 		if (Payload.Has<FActorTreeItem>())
 		{
-			const UFoundationSubsystem* FoundationSubsystem = RepresentingWorld->GetSubsystem<UFoundationSubsystem>();
+			const ULevelInstanceSubsystem* LevelInstanceSubsystem = RepresentingWorld->GetSubsystem<ULevelInstanceSubsystem>();
 			// Iterate over all the actors that have been dragged
 			for (const TWeakObjectPtr<AActor>& WeakActor : Payload.GetData<TWeakObjectPtr<AActor>>(SceneOutliner::FWeakActorSelector()))
 			{
 				const AActor* Actor = WeakActor.Get();
 
-				bool bActorContainedInFoundation = false;
-				if (FoundationSubsystem)
+				bool bActorContainedInLevelInstance = false;
+				if (LevelInstanceSubsystem)
 				{
-					if (const AFoundationActor* ParentFoundation = FoundationSubsystem->GetParentFoundation(Actor))
+					if (const ALevelInstance* ParentLevelInstance = LevelInstanceSubsystem->GetParentLevelInstance(Actor))
 					{
-						if (!ParentFoundation->IsEditing())
+						if (!ParentLevelInstance->IsEditing())
 						{
-							return FSceneOutlinerDragValidationInfo(ESceneOutlinerDropCompatibility::IncompatibleGeneric, LOCTEXT("Error_RemoveEditingFoundation", "Cannot detach from a foundation which is not being edited"));
+							return FSceneOutlinerDragValidationInfo(ESceneOutlinerDropCompatibility::IncompatibleGeneric, LOCTEXT("Error_RemoveEditingLevelInstance", "Cannot detach from a LevelInstance which is not being edited"));
 						}
-						bActorContainedInFoundation = true;
+						bActorContainedInLevelInstance = true;
 					}
 
-					if (const AFoundationActor* FoundationActor = Cast<AFoundationActor>(Actor))
+					if (const ALevelInstance* LevelInstanceActor = Cast<ALevelInstance>(Actor))
 					{
-						if (!FoundationSubsystem->CanMoveActorToLevel(FoundationActor))
+						if (!LevelInstanceSubsystem->CanMoveActorToLevel(LevelInstanceActor))
 						{
-							return FSceneOutlinerDragValidationInfo(ESceneOutlinerDropCompatibility::IncompatibleGeneric, LOCTEXT("Error_RemoveEditingFoundation", "Cannot detach a foundation which is currently being edited"));
+							return FSceneOutlinerDragValidationInfo(ESceneOutlinerDropCompatibility::IncompatibleGeneric, LOCTEXT("Error_RemoveEditingLevelInstance", "Cannot detach a LevelInstance which is currently being edited"));
 						}
 					}
 				}
@@ -972,7 +972,7 @@ FSceneOutlinerDragValidationInfo FActorBrowsingMode::ValidateDrop(const ISceneOu
 				{
 					return FSceneOutlinerDragValidationInfo(ESceneOutlinerDropCompatibility::IncompatibleGeneric, FText::Format(LOCTEXT("Error_AttachChildActor", "Cannot move {0} as it is a child actor."), FText::FromString(Actor->GetActorLabel())));
 				}
-				else if (Actor->GetFolderPath() == DestinationPath && !Actor->GetAttachParentActor() && !bActorContainedInFoundation)
+				else if (Actor->GetFolderPath() == DestinationPath && !Actor->GetAttachParentActor() && !bActorContainedInLevelInstance)
 				{
 					FFormatNamedArguments Args;
 					Args.Add(TEXT("SourceName"), FText::FromString(Actor->GetActorLabel()));
@@ -1052,22 +1052,22 @@ void FActorBrowsingMode::OnDrop(ISceneOutlinerTreeItem& DropTarget, const FScene
 		{
 			// Show socket chooser if we have sockets to select
 
-			if (AFoundationActor* TargetFoundation = Cast<AFoundationActor>(DropActor))
+			if (ALevelInstance* TargetLevelInstance = Cast<ALevelInstance>(DropActor))
 			{
-				// Actors inside foundations cannot have folder paths
+				// Actors inside LevelInstances cannot have folder paths
 				TArray<AActor*> DraggedActors = Payload.GetData<AActor*>(SceneOutliner::FActorSelector());
 				for (auto& Actor : DraggedActors)
 				{
 					Actor->SetFolderPath_Recursively(FName());
 				}
 
-				UFoundationSubsystem* FoundationSubsystem = RepresentingWorld->GetSubsystem<UFoundationSubsystem>();
-				check(FoundationSubsystem);
+				ULevelInstanceSubsystem* LevelInstanceSubsystem = RepresentingWorld->GetSubsystem<ULevelInstanceSubsystem>();
+				check(LevelInstanceSubsystem);
 
-				check(TargetFoundation->IsEditing());
-				const FScopedTransaction Transaction(LOCTEXT("UndoAction_MoveActorsToFoundation", "Move actors to foundation"));
+				check(TargetLevelInstance->IsEditing());
+				const FScopedTransaction Transaction(LOCTEXT("UndoAction_MoveActorsToLevelInstance", "Move actors to LevelInstance"));
 
-				FoundationSubsystem->MoveActorsTo(TargetFoundation, DraggedActors);
+				LevelInstanceSubsystem->MoveActorsTo(TargetLevelInstance, DraggedActors);
 			}
 			else
 			{
@@ -1174,25 +1174,25 @@ void FActorBrowsingMode::OnDrop(ISceneOutlinerTreeItem& DropTarget, const FScene
 				}
 			}
 
-			const UFoundationSubsystem* FoundationSubsystem = RepresentingWorld->GetSubsystem<UFoundationSubsystem>();
-			check(FoundationSubsystem);
+			const ULevelInstanceSubsystem* LevelInstanceSubsystem = RepresentingWorld->GetSubsystem<ULevelInstanceSubsystem>();
+			check(LevelInstanceSubsystem);
 			// Since we are moving to a folder (or root), we must be moving into the persistent level.
 			ULevel* DestinationLevel = RepresentingWorld->PersistentLevel;
 			check(DestinationLevel);
 
 			TArray<AActor*> ActorsToMove;
-			Payload.ForEachItem<FActorTreeItem>([FoundationSubsystem, &ActorsToMove](const FActorTreeItem& ActorItem)
+			Payload.ForEachItem<FActorTreeItem>([LevelInstanceSubsystem, &ActorsToMove](const FActorTreeItem& ActorItem)
 				{
 					AActor* Actor = ActorItem.Actor.Get();
-					if (const AFoundationActor* ParentFoundation = FoundationSubsystem->GetParentFoundation(Actor))
+					if (const ALevelInstance* ParentLevelInstance = LevelInstanceSubsystem->GetParentLevelInstance(Actor))
 					{
-						check(ParentFoundation->IsEditing());
+						check(ParentLevelInstance->IsEditing());
 						ActorsToMove.Add(Actor);
 					}
 				});
 
 			TArray<AActor*> DraggedActors = Payload.GetData<AActor*>(SceneOutliner::FActorSelector());
-			FoundationSubsystem->MoveActorsToLevel(ActorsToMove, DestinationLevel);
+			LevelInstanceSubsystem->MoveActorsToLevel(ActorsToMove, DestinationLevel);
 		}
 	}
 }
diff --git a/Engine/Source/Editor/SceneOutliner/Private/ActorBrowsingModeSettings.h b/Engine/Source/Editor/SceneOutliner/Private/ActorBrowsingModeSettings.h
--- a/Engine/Source/Editor/SceneOutliner/Private/ActorBrowsingModeSettings.h
+++ b/Engine/Source/Editor/SceneOutliner/Private/ActorBrowsingModeSettings.h
@@ -33,7 +33,7 @@ class UActorBrowsingModeSettings
 	UPROPERTY(config)
 	uint32 bHideActorComponents : 1;
 
-	/** True when the Scene Outliner is not displaying foundations */
+	/** True when the Scene Outliner is not displaying LevelInstances */
 	UPROPERTY(config)
-	uint32 bHideFoundationHierarchy : 1;
+	uint32 bHideLevelInstanceHierarchy : 1;
 };
diff --git a/Engine/Source/Editor/SceneOutliner/Private/ActorHierarchy.cpp b/Engine/Source/Editor/SceneOutliner/Private/ActorHierarchy.cpp
--- a/Engine/Source/Editor/SceneOutliner/Private/ActorHierarchy.cpp
+++ b/Engine/Source/Editor/SceneOutliner/Private/ActorHierarchy.cpp
@@ -13,8 +13,8 @@
 #include "LevelUtils.h"
 #include "GameFramework/WorldSettings.h"
 #include "EditorActorFolders.h"
-#include "Foundation/FoundationActor.h"
-#include "Foundation/FoundationSubsystem.h"
+#include "LevelInstance/LevelInstanceActor.h"
+#include "LevelInstance/LevelInstanceSubsystem.h"
 
 TUniquePtr<FActorHierarchy> FActorHierarchy::Create(ISceneOutlinerMode* Mode, const TWeakObjectPtr<UWorld>& World)
 {
@@ -71,9 +71,9 @@ FActorHierarchy::~FActorHierarchy()
 
 void FActorHierarchy::FindChildren(const ISceneOutlinerTreeItem& Item, const TMap<FSceneOutlinerTreeItemID, FSceneOutlinerTreeItemPtr>& Items, TArray<FSceneOutlinerTreeItemPtr>& OutChildren) const
 {
-	UFoundationSubsystem* FoundationSubsystem = RepresentingWorld->GetSubsystem<UFoundationSubsystem>();
+	ULevelInstanceSubsystem* LevelInstanceSubsystem = RepresentingWorld->GetSubsystem<ULevelInstanceSubsystem>();
 
-	check(!bShowingFoundations || FoundationSubsystem);
+	check(!bShowingLevelInstances || LevelInstanceSubsystem);
 
 	if (const FWorldTreeItem* WorldTreeItem = Item.CastTo<FWorldTreeItem>())
 	{
@@ -110,9 +110,9 @@ void FActorHierarchy::FindChildren(const ISceneOutlinerTreeItem& Item, const TMa
 	}
 	else if (const FActorTreeItem* ActorTreeItem = Item.CastTo<FActorTreeItem>())
 	{
-		if (AFoundationActor* FoundationActor = Cast<AFoundationActor>(ActorTreeItem->Actor))
+		if (ALevelInstance* LevelInstanceActor = Cast<ALevelInstance>(ActorTreeItem->Actor))
 		{
-			FoundationSubsystem->ForEachActorInFoundation(FoundationActor, [&OutChildren, &Items](AActor* Actor)
+			LevelInstanceSubsystem->ForEachActorInLevelInstance(LevelInstanceActor, [&OutChildren, &Items](AActor* Actor)
 				{
 					if (const FSceneOutlinerTreeItemPtr* ContainedActorItem = Items.Find(Actor))
 					{
@@ -193,18 +193,18 @@ FSceneOutlinerTreeItemPtr FActorHierarchy::FindParent(const ISceneOutlinerTreeIt
 				}
 			}
 
-			if (const UFoundationSubsystem* FoundationSubsystem = RepresentingWorld->GetSubsystem<UFoundationSubsystem>())
+			if (const ULevelInstanceSubsystem* LevelInstanceSubsystem = RepresentingWorld->GetSubsystem<ULevelInstanceSubsystem>())
 			{
-				if (const AFoundationActor* OwningFoundation = FoundationSubsystem->GetParentFoundation(Actor))
+				if (const ALevelInstance* OwningLevelInstance = LevelInstanceSubsystem->GetParentLevelInstance(Actor))
 				{
-					const AFoundationActor* FoundationActor = Cast<AFoundationActor>(Actor);
-					const bool bIsAnEditingFoundation = FoundationActor ? FoundationActor->IsEditing() : false;
-					// Parent this to a foundation if the parent foundation is being edited or if this is a sub foundation which is being edited
-					if (bShowingFoundations || (OwningFoundation->IsEditing() || bIsAnEditingFoundation))
+					const ALevelInstance* LevelInstanceActor = Cast<ALevelInstance>(Actor);
+					const bool bIsAnEditingLevelInstance = LevelInstanceActor ? LevelInstanceActor->IsEditing() : false;
+					// Parent this to a LevelInstance if the parent LevelInstance is being edited or if this is a sub LevelInstance which is being edited
+					if (bShowingLevelInstances || (OwningLevelInstance->IsEditing() || bIsAnEditingLevelInstance))
 					{
-						if (const FSceneOutlinerTreeItemPtr* OwningFoundationItem = Items.Find(OwningFoundation))
+						if (const FSceneOutlinerTreeItemPtr* OwningLevelInstanceItem = Items.Find(OwningLevelInstance))
 						{
-							return *OwningFoundationItem;
+							return *OwningLevelInstanceItem;
 						}
 						else
 						{
@@ -281,17 +281,17 @@ void FActorHierarchy::CreateWorldChildren(UWorld* World, TArray<FSceneOutlinerTr
 {
 	check(World);
 
-	const UFoundationSubsystem* FoundationSubsystem = World->GetSubsystem<UFoundationSubsystem>();
+	const ULevelInstanceSubsystem* LevelInstanceSubsystem = World->GetSubsystem<ULevelInstanceSubsystem>();
 	// Create all actor items
 	for (FActorIterator ActorIt(World); ActorIt; ++ActorIt)
 	{
 		AActor* Actor = *ActorIt;
-		// If we are not showing foundations, foundation sub actor items should not be created unless they belong to a foundation which is being edited
-		if (FoundationSubsystem)
+		// If we are not showing LevelInstances, LevelInstance sub actor items should not be created unless they belong to a LevelInstance which is being edited
+		if (LevelInstanceSubsystem)
 		{
-			if (const AFoundationActor* ParentFoundation = FoundationSubsystem->GetParentFoundation(Actor))
+			if (const ALevelInstance* ParentLevelInstance = LevelInstanceSubsystem->GetParentLevelInstance(Actor))
 		{
-				if (!bShowingFoundations && !ParentFoundation->IsEditing())
+				if (!bShowingLevelInstances && !ParentLevelInstance->IsEditing())
 				{
 					continue;
 				}
@@ -349,16 +349,16 @@ void FActorHierarchy::CreateChildren(const FSceneOutlinerTreeItemPtr& Item, TArr
 
 		TArray<AActor*> ChildActors;
 
-		if (const AFoundationActor* FoundationParentActor = Cast<AFoundationActor>(ParentActor))
+		if (const ALevelInstance* LevelInstanceParentActor = Cast<ALevelInstance>(ParentActor))
 		{
-			const UFoundationSubsystem* FoundationSubsystem = RepresentingWorld->GetSubsystem<UFoundationSubsystem>();
-			check(FoundationSubsystem);
+			const ULevelInstanceSubsystem* LevelInstanceSubsystem = RepresentingWorld->GetSubsystem<ULevelInstanceSubsystem>();
+			check(LevelInstanceSubsystem);
 
-			FoundationSubsystem->ForEachActorInFoundation(FoundationParentActor, [this, FoundationParentActor, FoundationSubsystem, &ChildActors](AActor* SubActor)
+			LevelInstanceSubsystem->ForEachActorInLevelInstance(LevelInstanceParentActor, [this, LevelInstanceParentActor, LevelInstanceSubsystem, &ChildActors](AActor* SubActor)
 				{
-					const AFoundationActor* FoundationActor = Cast<AFoundationActor>(SubActor);
-					const bool bIsAnEditingFoundation = FoundationActor ? FoundationSubsystem->IsEditingFoundation(FoundationActor) : false;
-					if (bShowingFoundations || (FoundationSubsystem->IsEditingFoundation(FoundationParentActor) || bIsAnEditingFoundation))
+					const ALevelInstance* LevelInstanceActor = Cast<ALevelInstance>(SubActor);
+					const bool bIsAnEditingLevelInstance = LevelInstanceActor ? LevelInstanceSubsystem->IsEditingLevelInstance(LevelInstanceActor) : false;
+					if (bShowingLevelInstances || (LevelInstanceSubsystem->IsEditingLevelInstance(LevelInstanceParentActor) || bIsAnEditingLevelInstance))
 					{
 						ChildActors.Add(SubActor);
 					}
@@ -428,16 +428,16 @@ FSceneOutlinerTreeItemPtr FActorHierarchy::CreateParentItem(const FSceneOutliner
 				return Mode->CreateItemFor<FActorFolderTreeItem>(FActorFolderTreeItem(ActorTreeItem->Actor->GetFolderPath(), ActorTreeItem->Actor->GetWorld()), true);
 			}
 
-			// If item belongs to a foundation
-			if (const UFoundationSubsystem* FoundationSubsystem = RepresentingWorld->GetSubsystem<UFoundationSubsystem>())
+			// If item belongs to a LevelInstance
+			if (const ULevelInstanceSubsystem* LevelInstanceSubsystem = RepresentingWorld->GetSubsystem<ULevelInstanceSubsystem>())
 			{
-				if (AFoundationActor* ParentFoundation = FoundationSubsystem->GetParentFoundation(Actor))
+				if (ALevelInstance* ParentLevelInstance = LevelInstanceSubsystem->GetParentLevelInstance(Actor))
 				{
-					const AFoundationActor* FoundationActor = Cast<AFoundationActor>(Actor);
-					const bool bIsAnEditingFoundation = FoundationActor ? FoundationActor->IsEditing() : false;
-					if (bShowingFoundations || (ParentFoundation->IsEditing() || bIsAnEditingFoundation))
+					const ALevelInstance* LevelInstanceActor = Cast<ALevelInstance>(Actor);
+					const bool bIsAnEditingLevelInstance = LevelInstanceActor ? LevelInstanceActor->IsEditing() : false;
+					if (bShowingLevelInstances || (ParentLevelInstance->IsEditing() || bIsAnEditingLevelInstance))
 					{
-						return Mode->CreateItemFor<FActorTreeItem>(ParentFoundation, true);
+						return Mode->CreateItemFor<FActorTreeItem>(ParentLevelInstance, true);
 					}
 				}
 			}
diff --git a/Engine/Source/Editor/SceneOutliner/Private/ActorMode.cpp b/Engine/Source/Editor/SceneOutliner/Private/ActorMode.cpp
--- a/Engine/Source/Editor/SceneOutliner/Private/ActorMode.cpp
+++ b/Engine/Source/Editor/SceneOutliner/Private/ActorMode.cpp
@@ -14,7 +14,7 @@
 #include "FolderTreeItem.h"
 #include "ComponentTreeItem.h"
 #include "WorldTreeItem.h"
-#include "Foundation/FoundationEditorInstanceActor.h"
+#include "LevelInstance/LevelInstanceEditorInstanceActor.h"
 
 #define LOCTEXT_NAMESPACE "SceneOutliner_ActorMode"
 
@@ -70,7 +70,7 @@ namespace SceneOutliner
 FActorMode::FActorMode(const FActorModeParams& Params)
 	: ISceneOutlinerMode(Params.SceneOutliner)
 	, bHideComponents(Params.bHideComponents)
-	, bHideFoundationHierarchy(Params.bHideFoundationHierarchy)
+	, bHideLevelInstanceHierarchy(Params.bHideLevelInstanceHierarchy)
 	, SpecifiedWorldToDisplay(Params.SpecifiedWorldToDisplay)
 {
 	SceneOutliner->AddFilter(MakeShared<FActorFilter>(FActorTreeItem::FFilterPredicate::CreateLambda([this](const AActor* Actor)
@@ -91,7 +91,7 @@ TUniquePtr<ISceneOutlinerHierarchy> FActorMode::CreateHierarchy()
 {
 	TUniquePtr<FActorHierarchy> ActorHierarchy = FActorHierarchy::Create(this, RepresentingWorld);
 	ActorHierarchy->SetShowingComponents(!bHideComponents);
-	ActorHierarchy->SetShowingFoundations(!bHideFoundationHierarchy);
+	ActorHierarchy->SetShowingLevelInstances(!bHideLevelInstanceHierarchy);
 
 	return ActorHierarchy;
 }
diff --git a/Engine/Source/Editor/SceneOutliner/Private/ActorTreeItem.cpp b/Engine/Source/Editor/SceneOutliner/Private/ActorTreeItem.cpp
--- a/Engine/Source/Editor/SceneOutliner/Private/ActorTreeItem.cpp
+++ b/Engine/Source/Editor/SceneOutliner/Private/ActorTreeItem.cpp
@@ -17,7 +17,7 @@
 #include "ISceneOutlinerMode.h"
 #include "Logging/MessageLog.h"
 #include "SSocketChooser.h"
-#include "Foundation/FoundationActor.h"
+#include "LevelInstance/LevelInstanceActor.h"
 
 #define LOCTEXT_NAMESPACE "SceneOutliner_ActorTreeItem"
 
@@ -112,14 +112,14 @@ private:
 	FText GetDisplayText() const
 	{
 		const AActor* Actor = ActorPtr.Get();
-		if (const AFoundationActor* FoundationActor = Cast<AFoundationActor>(Actor))
+		if (const ALevelInstance* LevelInstanceActor = Cast<ALevelInstance>(Actor))
 		{
-			if (FoundationActor->IsDirty())
+			if (LevelInstanceActor->IsDirty())
 			{
 				FFormatNamedArguments Args;
-				Args.Add(TEXT("ActorLabel"), FText::FromString(FoundationActor->GetActorLabel()));
-				Args.Add(TEXT("EditTag"), LOCTEXT("EditingFoundationLabel", "*"));
-				return FText::Format(LOCTEXT("FoundationDisplay", "{ActorLabel}{EditTag}"), Args);
+				Args.Add(TEXT("ActorLabel"), FText::FromString(LevelInstanceActor->GetActorLabel()));
+				Args.Add(TEXT("EditTag"), LOCTEXT("EditingLevelInstanceLabel", "*"));
+				return FText::Format(LOCTEXT("LevelInstanceDisplay", "{ActorLabel}{EditTag}"), Args);
 			}
 		}
 
@@ -247,10 +247,10 @@ private:
 	{
 		AActor* Actor = ActorPtr.Get();
 
-		// Color foundations differently if they are being edited
-		if (const AFoundationActor* FoundationActor = Cast<AFoundationActor>(Actor))
+		// Color LevelInstances differently if they are being edited
+		if (const ALevelInstance* LevelInstanceActor = Cast<ALevelInstance>(Actor))
 		{
-			if (FoundationActor->IsEditing() && !FoundationActor->IsSelected())
+			if (LevelInstanceActor->IsEditing() && !LevelInstanceActor->IsSelected())
 			{
 				return FAppStyle::Get().GetSlateColor("Colors.AccentGreen");
 			}
diff --git a/Engine/Source/Editor/SceneOutliner/Private/SceneOutlinerModule.cpp b/Engine/Source/Editor/SceneOutliner/Private/SceneOutlinerModule.cpp
--- a/Engine/Source/Editor/SceneOutliner/Private/SceneOutlinerModule.cpp
+++ b/Engine/Source/Editor/SceneOutliner/Private/SceneOutlinerModule.cpp
@@ -62,7 +62,7 @@ TSharedRef<ISceneOutliner> FSceneOutlinerModule::CreateActorPicker(const FSceneO
 			Params.SceneOutliner = Outliner;
 			Params.SpecifiedWorldToDisplay = SpecifiedWorld;
 			Params.bHideComponents = true;
-			Params.bHideFoundationHierarchy = true;
+			Params.bHideLevelInstanceHierarchy = true;
 			return new FActorPickingMode(Params, OnItemPicked);
 		});
 
@@ -97,7 +97,7 @@ TSharedRef<ISceneOutliner> FSceneOutlinerModule::CreateComponentPicker(const FSc
 			Params.SceneOutliner = Outliner;
 			Params.SpecifiedWorldToDisplay = SpecifiedWorld;
 			Params.bHideComponents = false;
-			Params.bHideFoundationHierarchy = true;
+			Params.bHideLevelInstanceHierarchy = true;
 			return new FActorPickingMode(Params, OnItemPicked);
 		});
 
diff --git a/Engine/Source/Editor/SceneOutliner/Public/ActorBrowsingMode.h b/Engine/Source/Editor/SceneOutliner/Public/ActorBrowsingMode.h
--- a/Engine/Source/Editor/SceneOutliner/Public/ActorBrowsingMode.h
+++ b/Engine/Source/Editor/SceneOutliner/Public/ActorBrowsingMode.h
@@ -98,7 +98,7 @@ private:
 	static TSharedRef<FSceneOutlinerFilter> CreateHideTemporaryActorsFilter();
 	static TSharedRef<FSceneOutlinerFilter> CreateIsInCurrentLevelFilter();
 	static TSharedRef<FSceneOutlinerFilter> CreateHideComponentsFilter();
-	static TSharedRef<FSceneOutlinerFilter> CreateHideFoundationsFilter();
+	static TSharedRef<FSceneOutlinerFilter> CreateHideLevelInstancesFilter();
 private:
 	/** Number of actors which have passed through the filters */
 	uint32 FilteredActorCount = 0;
diff --git a/Engine/Source/Editor/SceneOutliner/Public/ActorHierarchy.h b/Engine/Source/Editor/SceneOutliner/Public/ActorHierarchy.h
--- a/Engine/Source/Editor/SceneOutliner/Public/ActorHierarchy.h
+++ b/Engine/Source/Editor/SceneOutliner/Public/ActorHierarchy.h
@@ -25,7 +25,7 @@ public:
 	virtual FSceneOutlinerTreeItemPtr CreateParentItem(const FSceneOutlinerTreeItemPtr& Item) const override;
 
 	void SetShowingComponents(bool bInShowingComponents) { bShowingComponents = bInShowingComponents; }
-	void SetShowingFoundations(bool bInShowingFoundations) { bShowingFoundations = bInShowingFoundations; }
+	void SetShowingLevelInstances(bool bInShowingLevelInstances) { bShowingLevelInstances = bInShowingLevelInstances; }
 private:
 	/** Adds all the direct and indirect children of a world to OutItems */
 	void CreateWorldChildren(UWorld* World, TArray<FSceneOutlinerTreeItemPtr>& OutItems) const;
@@ -63,7 +63,7 @@ private:
 	void FullRefreshEvent();
 
 	bool bShowingComponents = false;
-	bool bShowingFoundations = false;
+	bool bShowingLevelInstances = false;
 
 	TWeakObjectPtr<UWorld> RepresentingWorld;
 private:
diff --git a/Engine/Source/Editor/SceneOutliner/Public/ActorMode.h b/Engine/Source/Editor/SceneOutliner/Public/ActorMode.h
--- a/Engine/Source/Editor/SceneOutliner/Public/ActorMode.h
+++ b/Engine/Source/Editor/SceneOutliner/Public/ActorMode.h
@@ -23,17 +23,17 @@ struct SCENEOUTLINER_API FActorModeParams
 {
 	FActorModeParams() {}
 
-	FActorModeParams(SSceneOutliner* InSceneOutliner, const TWeakObjectPtr<UWorld>& InSpecifiedWorldToDisplay = nullptr, bool bInHideComponents = true, bool bInHideFoundationHierarchy = true)
+	FActorModeParams(SSceneOutliner* InSceneOutliner, const TWeakObjectPtr<UWorld>& InSpecifiedWorldToDisplay = nullptr, bool bInHideComponents = true, bool bInHideLevelInstanceHierarchy = true)
 		: SpecifiedWorldToDisplay(InSpecifiedWorldToDisplay)
 		, SceneOutliner(InSceneOutliner)
 		, bHideComponents(bInHideComponents)
-		, bHideFoundationHierarchy(bInHideFoundationHierarchy)
+		, bHideLevelInstanceHierarchy(bInHideLevelInstanceHierarchy)
 	{}
 
 	TWeakObjectPtr<UWorld> SpecifiedWorldToDisplay = nullptr;
 	SSceneOutliner* SceneOutliner;
 	bool bHideComponents;
-	bool bHideFoundationHierarchy;
+	bool bHideLevelInstanceHierarchy;
 };
 
 class SCENEOUTLINER_API FActorMode : public ISceneOutlinerMode
@@ -72,8 +72,8 @@ protected:
 protected:
 	// Should the hide components filter be enabled
 	bool bHideComponents;
-	// Are foundations being hidden
-	bool bHideFoundationHierarchy;
+	// Are LevelInstances being hidden
+	bool bHideLevelInstanceHierarchy;
 	/** The world which we are currently representing */
 	TWeakObjectPtr<UWorld> RepresentingWorld;
 	/** The world which the user manually selected */
diff --git a/Engine/Source/Editor/UnrealEd/Classes/Commandlets/WorldPartitionPreCookCommandlet.h b/Engine/Source/Editor/UnrealEd/Classes/Commandlets/WorldPartitionPreCookCommandlet.h
--- a/Engine/Source/Editor/UnrealEd/Classes/Commandlets/WorldPartitionPreCookCommandlet.h
+++ b/Engine/Source/Editor/UnrealEd/Classes/Commandlets/WorldPartitionPreCookCommandlet.h
@@ -18,7 +18,7 @@ class UWorldPartitionPreCookCommandlet : public UCommandlet
 	//~ End UCommandlet Interface
 
 private:
-	void OnFoundationActorPostLoad(class AFoundationActor* FoundationActor);
+	void OnLevelInstanceActorPostLoad(class ALevelInstance* LevelInstanceActor);
 	bool PreCookLevelAndSave(ULevel* InLevel);
 	bool SaveLevel(ULevel* InLevel);
 	void RemoveSubLevel(ULevel* InLevel);
diff --git a/Engine/Source/Editor/UnrealEd/Private/Commandlets/WorldPartitionPreCookCommandlet.cpp b/Engine/Source/Editor/UnrealEd/Private/Commandlets/WorldPartitionPreCookCommandlet.cpp
--- a/Engine/Source/Editor/UnrealEd/Private/Commandlets/WorldPartitionPreCookCommandlet.cpp
+++ b/Engine/Source/Editor/UnrealEd/Private/Commandlets/WorldPartitionPreCookCommandlet.cpp
@@ -4,7 +4,7 @@
 #include "WorldPartition/WorldPartitionRuntimeHash.h"
 #include "WorldPartition/WorldPartition.h"
 #include "WorldPartition/WorldPartitionLevelHelper.h"
-#include "Foundation/FoundationActor.h"
+#include "LevelInstance/LevelInstanceActor.h"
 #include "Engine/LevelStreamingDynamic.h"
 #include "LevelUtils.h"
 #include "Editor.h"
@@ -18,26 +18,26 @@ UWorldPartitionPreCookCommandlet::UWorldPartitionPreCookCommandlet(const FObject
 {
 }
 
-void UWorldPartitionPreCookCommandlet::OnFoundationActorPostLoad(AFoundationActor* FoundationActor)
+void UWorldPartitionPreCookCommandlet::OnLevelInstanceActorPostLoad(ALevelInstance* LevelInstanceActor)
 {
-	if (FoundationActor->IsFoundationPathValid())
+	if (LevelInstanceActor->IsLevelInstancePathValid())
 	{
-		if (ULevel::GetIsLevelPartitionedFromPackage(FName(*FoundationActor->GetFoundationPackage())))
+		if (ULevel::GetIsLevelPartitionedFromPackage(FName(*LevelInstanceActor->GetWorldAssetPackage())))
 		{
-			const TSoftObjectPtr<UWorld>& Foundation = FoundationActor->GetFoundation();
-			PartitionedWorldsToGenerate.Add(Foundation.GetLongPackageName());
+			const TSoftObjectPtr<UWorld>& WorldAsset = LevelInstanceActor->GetWorldAsset();
+			PartitionedWorldsToGenerate.Add(WorldAsset.GetLongPackageName());
 
-			FString FoundationPath = Foundation.ToString();
-			FString PackageName = FPackageName::ObjectPathToPackageName(FoundationPath);
-			check(PackageName != FoundationPath);
-			FString ObjectName = FPackageName::ObjectPathToObjectName(Foundation.ToString());
-			check(ObjectName != FoundationPath);
+			FString WorldAssetPath = WorldAsset.ToString();
+			FString PackageName = FPackageName::ObjectPathToPackageName(WorldAssetPath);
+			check(PackageName != WorldAssetPath);
+			FString ObjectName = FPackageName::ObjectPathToObjectName(WorldAsset.ToString());
+			check(ObjectName != WorldAssetPath);
 
 			FString PackagePath = FPackageName::GetLongPackagePath(*PackageName);
 			FName PackageShortName = FPackageName::GetShortFName(*PackageName);
 
-			FString NewFoundation = FString::Printf(TEXT("%s/%s/%s/%s_Main.%s_Main"), *PackagePath, *PackageShortName.ToString(), FWorldPartitionLevelHelper::GetSavedLevelOutputSubFolderName(), *PackageShortName.ToString(), *ObjectName);
-			FoundationActor->SetFoundation(TSoftObjectPtr<UWorld>(FSoftObjectPath(NewFoundation)));
+			FString NewLevelInstance = FString::Printf(TEXT("%s/%s/%s/%s_Main.%s_Main"), *PackagePath, *PackageShortName.ToString(), FWorldPartitionLevelHelper::GetSavedLevelOutputSubFolderName(), *PackageShortName.ToString(), *ObjectName);
+			LevelInstanceActor->SetWorldAsset(TSoftObjectPtr<UWorld>(FSoftObjectPath(NewLevelInstance)));
 		}
 	}
 }
@@ -184,7 +184,7 @@ int32 UWorldPartitionPreCookCommandlet::Main(const FString& Params)
 	FAssetRegistryModule& AssetRegistryModule = FModuleManager::LoadModuleChecked<FAssetRegistryModule>(TEXT("AssetRegistry"));
 	AssetRegistryModule.Get().SearchAllAssets(true);
 
-	AFoundationActor::OnFoundationActorPostLoad.AddUObject(this, &UWorldPartitionPreCookCommandlet::OnFoundationActorPostLoad);
+	ALevelInstance::OnLevelInstanceActorPostLoad.AddUObject(this, &UWorldPartitionPreCookCommandlet::OnLevelInstanceActorPostLoad);
 
 	ULevel* MainLevel = LoadLevel(Tokens[0]);
 	if (!PreCookLevelAndSave(MainLevel))
diff --git a/Engine/Source/Runtime/Engine/Classes/GameFramework/Actor.h b/Engine/Source/Runtime/Engine/Classes/GameFramework/Actor.h
--- a/Engine/Source/Runtime/Engine/Classes/GameFramework/Actor.h
+++ b/Engine/Source/Runtime/Engine/Classes/GameFramework/Actor.h
@@ -3630,7 +3630,7 @@ private:
 
 	friend UWorld;
 	friend class FFoliageHelper;
-	friend class UFoundationSubsystem;
+	friend class ULevelInstanceSubsystem;
 };
 
 struct FSetActorGuid
@@ -3640,7 +3640,6 @@ private:
 	{
 		InActor->ActorGuid = InActorGuid;
 	}
-	friend class ULevelStreamingFoundationInstance;
 	friend UWorld;
 };
 #endif
diff --git a/Engine/Source/Runtime/Engine/Engine.Build.cs b/Engine/Source/Runtime/Engine/Engine.Build.cs
--- a/Engine/Source/Runtime/Engine/Engine.Build.cs
+++ b/Engine/Source/Runtime/Engine/Engine.Build.cs
@@ -365,7 +365,7 @@ public class Engine : ModuleRules
             PrivateIncludePathModuleNames.Add("NaniteBuilder");
             DynamicallyLoadedModuleNames.Add("NaniteBuilder");
 
-			DynamicallyLoadedModuleNames.Add("FoundationEditor");
+			DynamicallyLoadedModuleNames.Add("LevelInstanceEditor");
         }
 
 		SetupModulePhysicsSupport(Target);
@@ -395,7 +395,7 @@ public class Engine : ModuleRules
 
 		if (Target.Platform == UnrealTargetPlatform.Mac)
 		{
-			PublicFrameworks.AddRange(new string[] { "AVFoundation", "CoreVideo", "CoreMedia" });
+			PublicFrameworks.AddRange(new string[] { "AVLevelInstance", "CoreVideo", "CoreMedia" });
 		}
 
 		if (Target.IsInPlatformGroup(UnrealPlatformGroup.Android))
diff --git a/Engine/Source/Runtime/Engine/Private/Foundation/FoundationActor.cpp b/Engine/Source/Runtime/Engine/Private/Foundation/FoundationActor.cpp
deleted file mode 100644
diff --git a/Engine/Source/Runtime/Engine/Private/Foundation/FoundationComponent.cpp b/Engine/Source/Runtime/Engine/Private/Foundation/FoundationComponent.cpp
deleted file mode 100644
diff --git a/Engine/Source/Runtime/Engine/Private/Foundation/FoundationComponent.h b/Engine/Source/Runtime/Engine/Private/Foundation/FoundationComponent.h
deleted file mode 100644
diff --git a/Engine/Source/Runtime/Engine/Private/Foundation/FoundationEditorInstanceActor.cpp b/Engine/Source/Runtime/Engine/Private/Foundation/FoundationEditorInstanceActor.cpp
deleted file mode 100644
diff --git a/Engine/Source/Runtime/Engine/Private/Foundation/FoundationEditorLevelStreaming.cpp b/Engine/Source/Runtime/Engine/Private/Foundation/FoundationEditorLevelStreaming.cpp
deleted file mode 100644
diff --git a/Engine/Source/Runtime/Engine/Private/Foundation/FoundationInstanceLevelStreaming.cpp b/Engine/Source/Runtime/Engine/Private/Foundation/FoundationInstanceLevelStreaming.cpp
deleted file mode 100644
diff --git a/Engine/Source/Runtime/Engine/Private/Foundation/FoundationPrivate.h b/Engine/Source/Runtime/Engine/Private/Foundation/FoundationPrivate.h
deleted file mode 100644
diff --git a/Engine/Source/Runtime/Engine/Private/Foundation/FoundationSubsystem.cpp b/Engine/Source/Runtime/Engine/Private/Foundation/FoundationSubsystem.cpp
deleted file mode 100644
diff --git a/Engine/Source/Runtime/Engine/Private/LevelInstance/LevelInstanceActor.cpp b/Engine/Source/Runtime/Engine/Private/LevelInstance/LevelInstanceActor.cpp
new file mode 100644
--- /dev/null
+++ b/Engine/Source/Runtime/Engine/Private/LevelInstance/LevelInstanceActor.cpp
@@ -0,0 +1,710 @@
+// Copyright Epic Games, Inc. All Rights Reserved.
+
+#include "LevelInstance/LevelInstanceActor.h"
+#include "LevelInstance/LevelInstanceSubsystem.h"
+#include "LevelInstance/LevelInstanceComponent.h"
+#include "LevelInstance/LevelInstanceEditorInstanceActor.h"
+#include "WorldPartition/WorldPartitionSubsystem.h"
+#include "WorldPartition/WorldPartition.h"
+#include "Engine/World.h"
+#include "Engine/Engine.h"
+#include "LevelInstancePrivate.h"
+
+#if WITH_EDITOR
+#include "ActorRegistry.h"
+#include "AssetData.h"
+#include "Logging/MessageLog.h"
+#include "Logging/TokenizedMessage.h"
+#include "Misc/MapErrors.h"
+#include "Misc/UObjectToken.h"
+#endif
+
+#define LOCTEXT_NAMESPACE "LevelInstanceActor"
+
+ALevelInstance::ALevelInstance(const FObjectInitializer& ObjectInitializer)
+	: Super(ObjectInitializer)
+#if WITH_EDITOR
+	, CachedLevelInstanceID(InvalidLevelInstanceID)
+	, bGuardLoadUnload(false)
+#endif
+	, LevelInstanceID(InvalidLevelInstanceID)
+{
+	RootComponent = CreateDefaultSubobject<ULevelInstanceComponent>(TEXT("Root"));
+	RootComponent->Mobility = EComponentMobility::Static;
+}
+
+ULevelInstanceSubsystem* ALevelInstance::GetLevelInstanceSubsystem() const
+{
+	if (UWorld* CurrentWorld = GetWorld())
+	{
+		return CurrentWorld->GetSubsystem<ULevelInstanceSubsystem>();
+	}
+
+	return nullptr;
+}
+
+void ALevelInstance::Serialize(FArchive& Ar)
+{
+	Super::Serialize(Ar);
+#if WITH_EDITOR
+	if (Ar.IsSaving() && Ar.IsCooking() && !IsTemplate())
+	{
+		FGuid Guid = GetLevelInstanceActorGuid();
+		Ar << Guid;
+	}
+#else
+	if (Ar.IsLoading())
+	{
+		if (IsTemplate())
+		{
+			check(!LevelInstanceActorGuid.IsValid());
+		}
+		else if (Ar.GetPortFlags() & PPF_Duplicate)
+		{
+			LevelInstanceActorGuid = FGuid::NewGuid();
+		}
+		else if (Ar.IsPersistent())
+		{
+			Ar << LevelInstanceActorGuid;
+		}
+	}
+#endif
+}
+
+void ALevelInstance::PostRegisterAllComponents()
+{
+	Super::PostRegisterAllComponents();
+
+	if (ULevelInstanceSubsystem* LevelInstanceSubsystem = GetLevelInstanceSubsystem())
+	{
+		LevelInstanceID = LevelInstanceSubsystem->RegisterLevelInstance(this);
+
+		if (!IsRunningCommandlet())
+		{
+			LoadLevelInstance();
+		}
+	}
+}
+
+void ALevelInstance::PostUnregisterAllComponents()
+{
+	Super::PostUnregisterAllComponents();
+
+	if (ULevelInstanceSubsystem* LevelInstanceSubsystem = GetLevelInstanceSubsystem())
+	{
+		// If LevelInstance has already been unregistered it will have an Invalid LevelInstanceID. Avoid processing it.
+		if (!HasValidLevelInstanceID())
+		{
+			return;
+		}
+
+		LevelInstanceSubsystem->UnregisterLevelInstance(this);
+
+		if (!IsRunningCommandlet())
+		{
+			UnloadLevelInstance();
+		}
+
+		// To avoid processing PostUnregisterAllComponents multiple times (BP Recompile is one use case)
+		LevelInstanceID = InvalidLevelInstanceID;
+	}
+
+}
+
+void ALevelInstance::LoadLevelInstance()
+{
+#if WITH_EDITOR
+	if (!bGuardLoadUnload)
+#endif
+	{
+		if (ULevelInstanceSubsystem* LevelInstanceSubsystem = GetLevelInstanceSubsystem())
+		{
+			bool bForce = false;
+#if WITH_EDITOR
+			// When reinstancing. Avoid reloading the level but if the package changed, force the load
+			bForce = GIsReinstancing && IsLoaded() && LevelInstanceSubsystem->GetLevelInstanceLevel(this)->GetPackage()->FileName != FName(*GetWorldAssetPackage());
+#endif
+			LevelInstanceSubsystem->RequestLoadLevelInstance(this, bForce);
+		}
+	}
+}
+
+void ALevelInstance::UnloadLevelInstance()
+{
+#if WITH_EDITOR
+	if (!bGuardLoadUnload)
+#endif
+	{
+		if (ULevelInstanceSubsystem* LevelInstanceSubsystem = GetLevelInstanceSubsystem())
+		{
+#if WITH_EDITOR
+			check(!HasDirtyChildren());
+#endif
+			LevelInstanceSubsystem->RequestUnloadLevelInstance(this);
+		}
+	}
+}
+
+bool ALevelInstance::IsLevelInstancePathValid() const
+{
+	return WorldAsset.GetUniqueID().IsValid();
+}
+
+bool ALevelInstance::HasValidLevelInstanceID() const
+{
+	return LevelInstanceID != InvalidLevelInstanceID;
+}
+
+const FLevelInstanceID& ALevelInstance::GetLevelInstanceID() const
+{
+	check(HasValidLevelInstanceID());
+	return LevelInstanceID;
+}
+
+const FGuid& ALevelInstance::GetLevelInstanceActorGuid() const
+{
+#if WITH_EDITOR
+	const FGuid& Guid = GetActorGuid();
+#else
+	const FGuid& Guid = LevelInstanceActorGuid;
+#endif
+	check(IsTemplate() || Guid.IsValid());
+	return Guid;
+}
+
+#if WITH_EDITOR
+
+AActor* ALevelInstance::FindEditorInstanceActor() const
+{
+	AActor* FoundActor = nullptr;
+	if (ULevelInstanceSubsystem* LevelInstanceSubsystem = GetLevelInstanceSubsystem())
+	{
+		if (LevelInstanceSubsystem->IsLoaded(this))
+		{
+			LevelInstanceSubsystem->ForEachActorInLevelInstance(this, [&FoundActor, this](AActor* LevelActor)
+			{
+				if (ALevelInstanceEditorInstanceActor* LevelInstanceEditorInstanceActor = Cast<ALevelInstanceEditorInstanceActor>(LevelActor))
+				{
+					check(LevelInstanceEditorInstanceActor->GetLevelInstanceID() == GetLevelInstanceID());
+					FoundActor = LevelInstanceEditorInstanceActor;
+					return false;
+				}
+				return true;
+			});
+		}
+	}
+
+	return FoundActor;
+}
+
+ALevelInstance::FOnLevelInstanceActorPostLoad ALevelInstance::OnLevelInstanceActorPostLoad;
+
+void ALevelInstance::PostLoad()
+{
+	Super::PostLoad();
+
+	OnLevelInstanceActorPostLoad.Broadcast(this);
+}
+
+void ALevelInstance::PreEditUndo()
+{
+	CachedLevelInstanceID = LevelInstanceID;
+	CachedWorldAsset = WorldAsset;
+	bCachedIsTemporarilyHiddenInEditor = IsTemporarilyHiddenInEditor(false);
+
+	{
+		TGuardValue<bool> LoadUnloadGuard(bGuardLoadUnload, true);
+		Super::PreEditUndo();
+	}
+}
+
+void ALevelInstance::PostEditUndo()
+{
+	{
+		TGuardValue<bool> LoadUnloadGuard(bGuardLoadUnload, true);
+		Super::PostEditUndo();
+	}
+
+	if (CachedWorldAsset != WorldAsset)
+	{
+		UpdateLevelInstance();
+	}
+
+	if (bCachedIsTemporarilyHiddenInEditor != IsTemporarilyHiddenInEditor(false))
+	{
+		if (ULevelInstanceSubsystem* LevelInstanceSubsystem = GetLevelInstanceSubsystem())
+		{
+			LevelInstanceSubsystem->SetIsTemporarilyHiddenInEditor(this, !bCachedIsTemporarilyHiddenInEditor);
+		}
+	}
+
+	// Here we want to load or unload based on our current state
+	if (HasValidLevelInstanceID() && !IsLoaded())
+	{
+		LoadLevelInstance();
+	}
+	else if (IsPendingKill())
+	{
+		// Temp restore the ID so that we can unload
+		TGuardValue<FLevelInstanceID> LevelInstanceIDGuard(LevelInstanceID, CachedLevelInstanceID);
+		if (IsLoaded())
+		{
+			UnloadLevelInstance();
+		}
+	}
+
+	CachedLevelInstanceID = InvalidLevelInstanceID;
+	CachedWorldAsset.Reset();
+
+	if (ULevelInstanceComponent* LevelInstanceComponent = Cast<ULevelInstanceComponent>(GetRootComponent()))
+	{
+		// Order of operations when undoing may lead to the RootComponent being undone before our actor so we need to make sure we update here and in the component when undoing
+		LevelInstanceComponent->UpdateEditorInstanceActor();
+	}
+}
+
+FString ALevelInstance::GetWorldAssetPackage() const
+{
+	return WorldAsset.GetUniqueID().GetLongPackageName();
+}
+
+void ALevelInstance::GetAssetRegistryTags(TArray<FAssetRegistryTag>& OutTags) const
+{
+	Super::GetAssetRegistryTags(OutTags);
+
+	if (IsPackageExternal())
+	{
+		static const FName NAME_LevelInstancePackage(TEXT("LevelInstancePackage"));
+		FActorRegistry::SaveActorMetaData(NAME_LevelInstancePackage, GetWorldAssetPackage(), OutTags);
+
+		static const FName NAME_LevelInstanceTransform(TEXT("LevelInstanceTransform"));
+		FTransform LevelInstanceTransform(GetActorRotation(), GetActorLocation());
+		FActorRegistry::SaveActorMetaData(NAME_LevelInstanceTransform, LevelInstanceTransform, OutTags);
+	}
+}
+
+void ALevelInstance::PreEditChange(FProperty* PropertyThatWillChange)
+{
+	{
+		TGuardValue<bool> LoadUnloadGuard(bGuardLoadUnload, true);
+		Super::PreEditChange(PropertyThatWillChange);
+	}
+
+	if (PropertyThatWillChange->GetFName() == GET_MEMBER_NAME_CHECKED(ALevelInstance, WorldAsset))
+	{
+		CachedWorldAsset = WorldAsset;
+	}
+}
+
+void ALevelInstance::CheckForErrors()
+{
+	Super::CheckForErrors();
+
+	TArray<TPair<FText, TSoftObjectPtr<UWorld>>> LoopInfo;
+	const ALevelInstance* LoopStart = nullptr;
+	if (!CheckForLoop(GetWorldAsset(), &LoopInfo, &LoopStart))
+	{
+		TSharedRef<FTokenizedMessage> Error = FMessageLog("MapCheck").Error()->AddToken(FTextToken::Create(LOCTEXT("LevelInstanceActor_Loop_CheckForErrors", "LevelInstance level loop found!")));
+		TSoftObjectPtr<UWorld> LoopStartAsset(LoopStart->GetLevel()->GetTypedOuter<UWorld>());
+		Error->AddToken(FAssetNameToken::Create(LoopStartAsset.GetLongPackageName(), FText::FromString(LoopStartAsset.GetAssetName())));
+		Error->AddToken(FTextToken::Create(FText::FromString(TEXT(":"))));
+		Error->AddToken(FUObjectToken::Create(LoopStart));
+		for (int32 i = LoopInfo.Num() - 1; i >= 0; --i)
+		{
+			Error->AddToken(FTextToken::Create(LoopInfo[i].Key));
+			TSoftObjectPtr<UWorld> LevelInstancePtr = LoopInfo[i].Value;
+			Error->AddToken(FAssetNameToken::Create(LevelInstancePtr.GetLongPackageName(), FText::FromString(LevelInstancePtr.GetAssetName())));
+		}
+
+		Error->AddToken(FMapErrorToken::Create(FName(TEXT("LevelInstanceActor_Loop_CheckForErrors"))));
+	}
+}
+
+bool ALevelInstance::CheckForLoop(TSoftObjectPtr<UWorld> InLevelInstance, TArray<TPair<FText,TSoftObjectPtr<UWorld>>>* LoopInfo, const ALevelInstance** LoopStart) const
+{
+	bool bValid = true;
+	if (ULevelInstanceSubsystem* LevelInstanceSubsystem = GetLevelInstanceSubsystem())
+	{
+		LevelInstanceSubsystem->ForEachLevelInstanceAncestorsAndSelf(this, [&bValid, &InLevelInstance, LoopInfo, LoopStart, this](const ALevelInstance* LevelInstanceActor)
+		{
+			// Check the level we are spawned in to detect the loop (this will handle loops caused by LevelInstances and by regular level streaming)
+			if (LevelInstanceActor->GetLevel()->GetPackage()->FileName == FName(*InLevelInstance.GetLongPackageName()))
+			{
+				bValid = false;
+				if (LoopStart)
+				{
+					*LoopStart = LevelInstanceActor;
+				}
+			}
+
+			if (LoopInfo)
+			{
+				TSoftObjectPtr<UWorld> LevelInstancePtr = LevelInstanceActor == this ? InLevelInstance : LevelInstanceActor->GetWorldAsset();
+				FText LevelInstanceName = FText::FromString(LevelInstanceActor->GetPathName());
+				FText Description = FText::Format(LOCTEXT("LevelInstanceLoopLink", "-> Actor: {0} loads"), LevelInstanceName);
+				LoopInfo->Emplace(Description, LevelInstancePtr);
+			}
+				
+			return bValid;
+		});
+	}
+
+	return bValid;
+}
+
+bool ALevelInstance::CanSetValue(TSoftObjectPtr<UWorld> InLevelInstance, FString* Reason) const
+{
+	TArray<TPair<FText, TSoftObjectPtr<UWorld>>> LoopInfo;
+	const ALevelInstance* LoopStart = nullptr;
+
+	if (!CheckForLoop(InLevelInstance, Reason ? &LoopInfo : nullptr, Reason ? &LoopStart : nullptr))
+	{
+		if (Reason)
+		{
+			if (ensure(LoopStart))
+			{
+				TSoftObjectPtr<UWorld> LoopStartAsset(LoopStart->GetLevel()->GetTypedOuter<UWorld>());
+				*Reason = FString::Format(TEXT("Setting LevelInstance to {0} would cause loop {1}:{2}\n"), { InLevelInstance.GetLongPackageName(), LoopStart->GetName(), LoopStartAsset.GetLongPackageName() });
+				for (int32 i = LoopInfo.Num() - 1; i >= 0; --i)
+				{
+					Reason->Append(FString::Format(TEXT("{0} {1}\n"), { *LoopInfo[i].Key.ToString(), *LoopInfo[i].Value.GetLongPackageName() }));
+				}
+			}
+		}
+
+		return false;
+	}
+
+	return true;
+}
+
+bool ALevelInstance::SetWorldAsset(TSoftObjectPtr<UWorld> InWorldAsset)
+{
+	FString Reason;
+	if (!CanSetValue(InWorldAsset, &Reason))
+	{
+		UE_LOG(LogLevelInstance, Warning, TEXT("%s"), *Reason);
+		return false;
+	}
+
+	WorldAsset = InWorldAsset;
+	return true;
+}
+
+void ALevelInstance::PostEditChangeProperty(FPropertyChangedEvent& PropertyChangedEvent)
+{
+	{
+		TGuardValue<bool> LoadUnloadGuard(bGuardLoadUnload, true);
+		Super::PostEditChangeProperty(PropertyChangedEvent);
+	}
+
+	if (FProperty* const PropertyThatChanged = PropertyChangedEvent.Property)
+	{
+		if (PropertyThatChanged->GetFName() == GET_MEMBER_NAME_CHECKED(ALevelInstance, WorldAsset))
+		{
+			if (ULevelInstanceSubsystem* LevelInstanceSubsystem = GetLevelInstanceSubsystem())
+			{
+				FString Reason;
+				if (!CanSetValue(GetWorldAsset(), &Reason))
+				{
+					UE_LOG(LogLevelInstance, Warning, TEXT("%s"), *Reason);
+					WorldAsset = CachedWorldAsset;
+				}
+				else
+				{
+					UpdateLevelInstance();
+				}
+				CachedWorldAsset.Reset();
+			}
+		}
+	}
+}
+
+bool ALevelInstance::CanEditChange(const FProperty* InProperty) const
+{
+	if (!Super::CanEditChange(InProperty))
+	{
+		return false;
+	}
+
+	if (IsEditing())
+	{
+		return false;
+	}
+
+	if (HasDirtyChildren())
+	{
+		return false;
+	}
+
+	return true;
+}
+
+void ALevelInstance::PostEditImport()
+{
+	{
+		TGuardValue<bool> LoadUnloadGuard(bGuardLoadUnload, true);
+		Super::PostEditImport();
+	}
+	UpdateLevelInstance();
+}
+
+bool ALevelInstance::CanDeleteSelectedActor(FText& OutReason) const
+{
+	if (!Super::CanDeleteSelectedActor(OutReason))
+	{
+		return false;
+	}
+
+	if (IsDirty())
+	{
+		OutReason = LOCTEXT("HasDirtyLevel", "Can't delete LevelInstance because it is dirty!");
+		return false;
+	}
+
+	if (HasDirtyChildren())
+	{
+		OutReason = LOCTEXT("HasDirtryChildLevel", "Can't delete LevelInstance because it has dirty child LevelInstances!");
+		return false;
+	}
+	return true;
+}
+
+void ALevelInstance::SetIsTemporarilyHiddenInEditor(bool bIsHidden)
+{
+	{
+		TGuardValue<bool> LoadUnloadGuard(bGuardLoadUnload, true);
+		Super::SetIsTemporarilyHiddenInEditor(bIsHidden);
+	}
+
+	if (ULevelInstanceSubsystem* LevelInstanceSubsystem = GetLevelInstanceSubsystem())
+	{
+		LevelInstanceSubsystem->SetIsTemporarilyHiddenInEditor(this, bIsHidden);
+	}
+}
+
+void ALevelInstance::EditorGetUnderlyingActors(TSet<AActor*>& OutUnderlyingActors)
+{
+	Super::EditorGetUnderlyingActors(OutUnderlyingActors);
+
+	if (ULevelInstanceSubsystem* LevelInstanceSubsystem = GetLevelInstanceSubsystem())
+	{
+		LevelInstanceSubsystem->ForEachActorInLevelInstance(this, [&](AActor* LevelActor)
+		{
+			bool bAlreadySet = false;
+			OutUnderlyingActors.Add(LevelActor, &bAlreadySet);
+			if (!bAlreadySet)
+			{
+				LevelActor->EditorGetUnderlyingActors(OutUnderlyingActors);
+			}
+			return true;
+		});
+	}
+}
+
+void ALevelInstance::UpdateLevelInstance()
+{
+	if (HasValidLevelInstanceID())
+	{
+		if (ULevelInstanceSubsystem* LevelInstanceSubsystem = GetLevelInstanceSubsystem())
+		{
+			if (IsLevelInstancePathValid())
+			{
+				const bool bForceUpdate = true;
+				LevelInstanceSubsystem->RequestLoadLevelInstance(this, bForceUpdate);
+			}
+			else if(IsLoaded())
+			{
+				UnloadLevelInstance();
+			}
+		}
+	}
+}
+
+bool ALevelInstance::IsLoaded() const
+{
+	if (ULevelInstanceSubsystem* LevelInstanceSubsystem = GetLevelInstanceSubsystem())
+	{
+		return LevelInstanceSubsystem->IsLoaded(this);
+	}
+
+	return false;
+}
+
+void ALevelInstance::OnLevelInstanceLoaded()
+{
+	if (!GetWorld()->IsPlayInEditor())
+	{
+		// Propagate bounds dirtyness up and check if we need to hide our LevelInstance because self or ancestor is hidden
+		bool bHiddenInEditor = false;
+		GetLevelInstanceSubsystem()->ForEachLevelInstanceAncestorsAndSelf(this, [&bHiddenInEditor](const ALevelInstance* AncestorOrSelf)
+		{
+			AncestorOrSelf->GetLevel()->MarkLevelBoundsDirty();
+			bHiddenInEditor |= AncestorOrSelf->IsTemporarilyHiddenInEditor();
+			return true;
+		});
+
+		if (bHiddenInEditor)
+		{
+			SetIsTemporarilyHiddenInEditor(bHiddenInEditor);
+		}
+
+		UWorldPartitionSubsystem* WorldPartitionSubsystem = GetWorld()->GetSubsystem<UWorldPartitionSubsystem>();
+		if (WorldPartitionSubsystem && WorldPartitionSubsystem->IsEnabled())
+		{
+			WorldPartitionSubsystem->UpdateActorDesc(this);
+		}
+	}
+}
+
+void ALevelInstance::GetActorLocationBounds(bool bOnlyCollidingComponents, FVector& Origin, FVector& BoxExtent, bool bIncludeFromChildActors) const
+{
+	Super::GetActorLocationBounds(bOnlyCollidingComponents, Origin, BoxExtent, bIncludeFromChildActors);
+
+	// Add Level Bounds
+	if (ULevelInstanceSubsystem* LevelInstanceSubsystem = GetLevelInstanceSubsystem())
+	{
+		FBox LevelInstanceBounds;
+		if (LevelInstanceSubsystem->GetLevelInstanceBounds(this, LevelInstanceBounds))
+		{
+			LevelInstanceBounds.GetCenterAndExtents(Origin, BoxExtent);
+		}
+	}
+}
+
+FBox ALevelInstance::GetComponentsBoundingBox(bool bNonColliding, bool bIncludeFromChildActors) const
+{
+	FBox Box = Super::GetComponentsBoundingBox(bNonColliding, bIncludeFromChildActors);
+	if (ULevelInstanceSubsystem* LevelInstanceSubsystem = GetLevelInstanceSubsystem())
+	{
+		FBox LevelInstanceBounds;
+		if(LevelInstanceSubsystem->GetLevelInstanceBounds(this, LevelInstanceBounds))
+		{
+			Box += LevelInstanceBounds;
+		}
+	}
+
+	return Box;
+}
+
+bool ALevelInstance::CanEdit(FText* OutReason) const
+{
+	if (ULevelInstanceSubsystem* LevelInstanceSubsystem = GetLevelInstanceSubsystem())
+	{
+		return LevelInstanceSubsystem->CanEditLevelInstance(this, OutReason);
+	}
+
+	return false;
+}
+
+bool ALevelInstance::CanCommit(FText* OutReason) const
+{
+	if (ULevelInstanceSubsystem* LevelInstanceSubsystem = GetLevelInstanceSubsystem())
+	{
+		return LevelInstanceSubsystem->CanCommitLevelInstance(this, OutReason);
+	}
+
+	return false;
+}
+
+bool ALevelInstance::IsEditing() const
+{
+	if (ULevelInstanceSubsystem* LevelInstanceSubsystem = GetLevelInstanceSubsystem())
+	{
+		return LevelInstanceSubsystem->IsEditingLevelInstance(this);
+	}
+
+	return false;
+}
+
+void ALevelInstance::Edit(AActor* ContextActor)
+{
+	ULevelInstanceSubsystem* LevelInstanceSubsystem = GetLevelInstanceSubsystem();
+	check(LevelInstanceSubsystem);
+	LevelInstanceSubsystem->EditLevelInstance(this, ContextActor);
+}
+
+void ALevelInstance::Commit()
+{
+	ULevelInstanceSubsystem* LevelInstanceSubsystem = GetLevelInstanceSubsystem();
+	check(LevelInstanceSubsystem);
+	LevelInstanceSubsystem->CommitLevelInstance(this);
+}
+
+void ALevelInstance::Discard()
+{
+	ULevelInstanceSubsystem* LevelInstanceSubsystem = GetLevelInstanceSubsystem();
+	check(LevelInstanceSubsystem);
+	const bool bDiscardEdits = true;
+	LevelInstanceSubsystem->CommitLevelInstance(this, bDiscardEdits);
+}
+
+void ALevelInstance::SaveAs()
+{
+	ULevelInstanceSubsystem* LevelInstanceSubsystem = GetLevelInstanceSubsystem();
+	check(LevelInstanceSubsystem);
+	LevelInstanceSubsystem->SaveLevelInstanceAs(this);
+}
+
+bool ALevelInstance::HasDirtyChildren() const
+{
+	if (ULevelInstanceSubsystem* LevelInstanceSubsystem = GetLevelInstanceSubsystem())
+	{
+		return LevelInstanceSubsystem->HasDirtyChildrenLevelInstances(this);
+	}
+
+	return false;
+}
+
+bool ALevelInstance::IsDirty() const
+{
+	if (ULevelInstanceSubsystem* LevelInstanceSubsystem = GetLevelInstanceSubsystem())
+	{
+		return LevelInstanceSubsystem->IsEditingLevelInstanceDirty(this);
+	}
+
+	return false;
+}
+
+bool ALevelInstance::SetCurrent()
+{
+	if (ULevelInstanceSubsystem* LevelInstanceSubsystem = GetLevelInstanceSubsystem())
+	{
+		return LevelInstanceSubsystem->SetCurrent(this);
+	}
+
+	return false;
+}
+
+bool ALevelInstance::IsCurrent() const
+{
+	if (ULevelInstanceSubsystem* LevelInstanceSubsystem = GetLevelInstanceSubsystem())
+	{
+		return LevelInstanceSubsystem->IsCurrent(this);
+	}
+
+	return false;
+}
+
+void ALevelInstance::PushSelectionToProxies()
+{
+	Super::PushSelectionToProxies();
+
+	// Actors of the LevelInstance need to reflect the LevelInstance actor's selected state
+	if (ULevelInstanceSubsystem* LevelInstanceSubsystem = GetLevelInstanceSubsystem())
+	{
+		LevelInstanceSubsystem->ForEachActorInLevelInstance(this, [](AActor* LevelActor)
+		{
+			if(ALevelInstanceEditorInstanceActor* EditorInstanceActor = Cast<ALevelInstanceEditorInstanceActor>(LevelActor))
+			{
+				EditorInstanceActor->PushSelectionToProxies();
+				return false;
+			}
+			return true;
+		});
+	}
+}
+
+#endif
+
+#undef LOCTEXT_NAMESPACE
diff --git a/Engine/Source/Runtime/Engine/Private/LevelInstance/LevelInstanceComponent.cpp b/Engine/Source/Runtime/Engine/Private/LevelInstance/LevelInstanceComponent.cpp
new file mode 100644
--- /dev/null
+++ b/Engine/Source/Runtime/Engine/Private/LevelInstance/LevelInstanceComponent.cpp
@@ -0,0 +1,73 @@
+// Copyright Epic Games, Inc. All Rights Reserved.
+#include "LevelInstance/LevelInstanceComponent.h"
+#include "LevelInstance/LevelInstanceActor.h"
+#include "Engine/Texture2D.h"
+#include "Engine/World.h"
+
+ULevelInstanceComponent::ULevelInstanceComponent(const FObjectInitializer& ObjectInitializer)
+	: Super(ObjectInitializer)
+{
+#if WITH_EDITOR
+	bWantsOnUpdateTransform = true;
+#endif
+}
+
+#if WITH_EDITOR
+void ULevelInstanceComponent::OnRegister()
+{
+#if WITH_EDITORONLY_DATA
+	// Prevents USceneComponent from creating the SpriteComponent in OnRegister because we want to provide a different texture and condition
+	bVisualizeComponent = false;
+#endif
+
+	Super::OnRegister();
+
+#if WITH_EDITORONLY_DATA
+	AActor* Owner = GetOwner();
+	// Only show Sprite for non-instanced LevelInstances
+	if (Owner && Owner->GetLevel() && !GetOwner()->GetLevel()->IsInstancedLevel() && !GetWorld()->IsGameWorld())
+	{
+		// Re-enable before calling CreateSpriteComponent
+		bVisualizeComponent = true;
+		CreateSpriteComponent(LoadObject<UTexture2D>(nullptr, TEXT("/Engine/EditorResources/LevelInstance")));
+	}
+#endif //WITH_EDITORONLY_DATA
+}
+
+void ULevelInstanceComponent::PostEditUndo()
+{
+	Super::PostEditUndo();
+		
+	UpdateComponentToWorld();
+	UpdateEditorInstanceActor();
+}
+
+void ULevelInstanceComponent::PostEditChangeProperty(FPropertyChangedEvent& PropertyChangedEvent)
+{
+	Super::PostEditChangeProperty(PropertyChangedEvent);
+
+	UpdateEditorInstanceActor();
+}
+
+void ULevelInstanceComponent::OnUpdateTransform(EUpdateTransformFlags UpdateTransformFlags, ETeleportType Teleport)
+{
+	Super::OnUpdateTransform(UpdateTransformFlags, Teleport);
+
+	UpdateEditorInstanceActor();
+}
+
+void ULevelInstanceComponent::UpdateEditorInstanceActor()
+{
+	if (!CachedEditorInstanceActorPtr.IsValid())
+	{
+		check(GetOuterALevelInstance());
+		CachedEditorInstanceActorPtr = GetOuterALevelInstance()->FindEditorInstanceActor();
+	}
+
+	if (AActor* EditorInstanceActor = CachedEditorInstanceActorPtr.Get())
+	{
+		EditorInstanceActor->GetRootComponent()->SetWorldTransform(GetComponentTransform());
+	}
+}
+
+#endif
\ No newline at end of file
diff --git a/Engine/Source/Runtime/Engine/Private/LevelInstance/LevelInstanceComponent.h b/Engine/Source/Runtime/Engine/Private/LevelInstance/LevelInstanceComponent.h
new file mode 100644
--- /dev/null
+++ b/Engine/Source/Runtime/Engine/Private/LevelInstance/LevelInstanceComponent.h
@@ -0,0 +1,37 @@
+// Copyright Epic Games, Inc. All Rights Reserved.
+
+#pragma once
+
+#include "CoreMinimal.h"
+#include "Components/SceneComponent.h"
+
+#include "LevelInstanceComponent.generated.h"
+
+/**
+ * ULevelInstanceComponent subclasses USceneComponent for Editing purposes so that we can have a proxy to the LevelInstanceActor's RootComponent transform without attaching to it.
+ *
+ * It is responsible for updating the transform of the ALevelInstanceEditorInstanceActor which is created when loading a LevelInstance Instance Level
+ *
+ * We use this method to avoid attaching the Instance Level Actors to the ALevelInstance. (Cross level attachment and undo/redo pain)
+ * 
+ * The LevelInstance Level Actors are attached to this ALevelInstanceEditorInstanceActor keeping the attachment local to the Instance Level and shielded from the transaction buffer.
+ *
+ * Avoiding those Level Actors from being part of the Transaction system allows us to unload that level without clearing the transaction buffer. It also allows BP Reinstancing without having to update attachements.
+ */
+UCLASS(Within=LevelInstance)
+class ENGINE_API ULevelInstanceComponent : public USceneComponent
+{
+	GENERATED_UCLASS_BODY()
+public:
+#if WITH_EDITOR
+	// Those are the methods that need overriding to be able to properly update the AttachComponent
+	virtual void OnRegister() override;
+	virtual void PostEditUndo() override;
+	virtual void PostEditChangeProperty(FPropertyChangedEvent& PropertyChangedEvent) override;
+	virtual void OnUpdateTransform(EUpdateTransformFlags UpdateTransformFlags, ETeleportType Teleport) override;
+
+	void UpdateEditorInstanceActor();
+private:
+	TWeakObjectPtr<AActor> CachedEditorInstanceActorPtr;
+#endif
+};
\ No newline at end of file
diff --git a/Engine/Source/Runtime/Engine/Private/LevelInstance/LevelInstanceEditorInstanceActor.cpp b/Engine/Source/Runtime/Engine/Private/LevelInstance/LevelInstanceEditorInstanceActor.cpp
new file mode 100644
--- /dev/null
+++ b/Engine/Source/Runtime/Engine/Private/LevelInstance/LevelInstanceEditorInstanceActor.cpp
@@ -0,0 +1,52 @@
+// Copyright Epic Games, Inc. All Rights Reserved.
+
+#include "LevelInstance/LevelInstanceEditorInstanceActor.h"
+#include "LevelInstance/LevelInstanceActor.h"
+#include "LevelInstance/LevelInstanceSubsystem.h"
+#include "Engine/World.h"
+
+ALevelInstanceEditorInstanceActor::ALevelInstanceEditorInstanceActor(const FObjectInitializer& ObjectInitializer)
+	: Super(ObjectInitializer)
+#if WITH_EDITOR
+	, LevelInstanceID(InvalidLevelInstanceID)
+#endif
+{
+	RootComponent = CreateDefaultSubobject<USceneComponent>(TEXT("RootComponent"));
+	RootComponent->Mobility = EComponentMobility::Static;
+}
+
+#if WITH_EDITOR
+AActor* ALevelInstanceEditorInstanceActor::GetSelectionParent() const
+{
+	if (ULevelInstanceSubsystem* LevelInstanceSubsystem = GetWorld()->GetSubsystem<ULevelInstanceSubsystem>())
+	{
+		return LevelInstanceSubsystem->GetLevelInstance(LevelInstanceID);
+	}
+
+	return nullptr;
+}
+
+ALevelInstanceEditorInstanceActor* ALevelInstanceEditorInstanceActor::Create(ALevelInstance* LevelInstanceActor, ULevel* LoadedLevel)
+{
+	FActorSpawnParameters SpawnParams;
+	SpawnParams.OverrideLevel = LoadedLevel;
+	SpawnParams.bHideFromSceneOutliner = true;
+	SpawnParams.bCreateActorPackage = false;
+	SpawnParams.ObjectFlags = RF_Transient;
+	SpawnParams.bNoFail = true;
+	ALevelInstanceEditorInstanceActor* InstanceActor = LevelInstanceActor->GetWorld()->SpawnActor<ALevelInstanceEditorInstanceActor>(LevelInstanceActor->GetActorLocation(), LevelInstanceActor->GetActorRotation(), SpawnParams);
+
+	InstanceActor->SetLevelInstanceID(LevelInstanceActor->GetLevelInstanceID());
+	
+	for (AActor* LevelActor : LoadedLevel->Actors)
+	{
+		if (LevelActor && LevelActor->GetAttachParentActor() == nullptr && !LevelActor->IsChildActor() && LevelActor != InstanceActor)
+		{
+			LevelActor->AttachToActor(InstanceActor, FAttachmentTransformRules::KeepWorldTransform);
+		}
+	}
+
+	return InstanceActor;
+}
+
+#endif
\ No newline at end of file
diff --git a/Engine/Source/Runtime/Engine/Private/LevelInstance/LevelInstanceEditorLevelStreaming.cpp b/Engine/Source/Runtime/Engine/Private/LevelInstance/LevelInstanceEditorLevelStreaming.cpp
new file mode 100644
--- /dev/null
+++ b/Engine/Source/Runtime/Engine/Private/LevelInstance/LevelInstanceEditorLevelStreaming.cpp
@@ -0,0 +1,58 @@
+// Copyright Epic Games, Inc. All Rights Reserved.
+
+#include "LevelInstance/LevelInstanceEditorLevelStreaming.h"
+
+#if WITH_EDITOR
+#include "LevelInstance/LevelInstanceActor.h"
+#include "LevelInstance/LevelInstanceSubsystem.h"
+#include "EditorLevelUtils.h"
+#include "Editor.h"
+#endif
+
+#if WITH_EDITOR
+static FLevelInstanceID EditLevelInstanceID = InvalidLevelInstanceID;
+#endif
+
+ULevelStreamingLevelInstanceEditor::ULevelStreamingLevelInstanceEditor(const FObjectInitializer& ObjectInitializer)
+	: Super(ObjectInitializer)
+#if WITH_EDITOR
+	, LevelInstanceID(EditLevelInstanceID)
+#endif
+{
+#if WITH_EDITOR
+	SetShouldBeVisibleInEditor(true);
+#endif
+}
+
+#if WITH_EDITOR
+ALevelInstance* ULevelStreamingLevelInstanceEditor::GetLevelInstanceActor() const
+{
+	if (ULevelInstanceSubsystem* LevelInstanceSubsystem = GetWorld()->GetSubsystem<ULevelInstanceSubsystem>())
+	{
+		return LevelInstanceSubsystem->GetLevelInstance(LevelInstanceID);
+	}
+
+	return nullptr;
+}
+
+ULevelStreamingLevelInstanceEditor* ULevelStreamingLevelInstanceEditor::Load(ALevelInstance* LevelInstanceActor)
+{
+	UWorld* CurrentWorld = LevelInstanceActor->GetWorld();
+
+	TGuardValue<FLevelInstanceID> GuardEditLevelInstanceID(EditLevelInstanceID, LevelInstanceActor->GetLevelInstanceID());
+	ULevelStreamingLevelInstanceEditor* LevelStreaming = Cast<ULevelStreamingLevelInstanceEditor>(EditorLevelUtils::AddLevelToWorld(CurrentWorld, *LevelInstanceActor->GetWorldAssetPackage(), ULevelStreamingLevelInstanceEditor::StaticClass(), LevelInstanceActor->GetTransform()));
+	check(LevelStreaming);
+	check(LevelStreaming->LevelInstanceID == LevelInstanceActor->GetLevelInstanceID());
+
+	GEngine->BlockTillLevelStreamingCompleted(LevelInstanceActor->GetWorld());
+	
+	return LevelStreaming;
+}
+
+void ULevelStreamingLevelInstanceEditor::Unload(ULevelStreamingLevelInstanceEditor* LevelStreaming)
+{
+	// No need to clear the whole editor selection since actor of this level will be removed from the selection by: UEditorEngine::OnLevelRemovedFromWorld
+	const bool bClearSelection = false;
+	EditorLevelUtils::RemoveLevelFromWorld(LevelStreaming->GetLoadedLevel(), bClearSelection);
+}
+#endif
\ No newline at end of file
diff --git a/Engine/Source/Runtime/Engine/Private/LevelInstance/LevelInstanceLevelStreaming.cpp b/Engine/Source/Runtime/Engine/Private/LevelInstance/LevelInstanceLevelStreaming.cpp
new file mode 100644
--- /dev/null
+++ b/Engine/Source/Runtime/Engine/Private/LevelInstance/LevelInstanceLevelStreaming.cpp
@@ -0,0 +1,149 @@
+// Copyright Epic Games, Inc. All Rights Reserved.
+
+#include "LevelInstance/LevelInstanceLevelStreaming.h"
+#include "LevelInstance/LevelInstanceActor.h"
+#include "LevelInstance/LevelInstancePrivate.h"
+#include "LevelInstance/LevelInstanceSubsystem.h"
+#include "Engine/LevelBounds.h"
+#include "Engine/World.h"
+#include "Misc/PackageName.h"
+
+#if WITH_EDITOR
+#include "LevelInstance/LevelInstanceEditorInstanceActor.h"
+#include "Editor.h"
+#include "EditorLevelUtils.h"
+#include "LevelUtils.h"
+#endif
+
+ULevelStreamingLevelInstance::ULevelStreamingLevelInstance(const FObjectInitializer& ObjectInitializer)
+	: Super(ObjectInitializer)
+	, LevelInstanceID(InvalidLevelInstanceID)
+{
+#if WITH_EDITOR
+	SetShouldBeVisibleInEditor(true);
+#endif
+}
+
+ALevelInstance* ULevelStreamingLevelInstance::GetLevelInstanceActor() const
+{
+	if (ULevelInstanceSubsystem* LevelInstanceSubsystem = GetWorld()->GetSubsystem<ULevelInstanceSubsystem>())
+	{
+		return LevelInstanceSubsystem->GetLevelInstance(LevelInstanceID);
+	}
+
+	return nullptr;
+}
+
+#if WITH_EDITOR
+FBox ULevelStreamingLevelInstance::GetBounds() const
+{
+	check(GetLoadedLevel());
+	return ALevelBounds::CalculateLevelBounds(GetLoadedLevel());
+}
+#endif
+
+ULevelStreamingLevelInstance* ULevelStreamingLevelInstance::LoadInstance(ALevelInstance* LevelInstanceActor)
+{
+#if WITH_EDITOR
+	if (!LevelInstanceActor->CheckForLoop(LevelInstanceActor->GetWorldAsset()))
+	{
+		UE_LOG(LogLevelInstance, Error, TEXT("Failed to load LevelInstance Actor '%s' because that would cause a loop. Run Map Check for more details."), *LevelInstanceActor->GetPathName());
+		return nullptr;
+	}
+#endif
+
+	bool bOutSuccess = false;
+
+	FString ShortPackageName = FPackageName::GetShortName(LevelInstanceActor->GetWorldAsset().GetLongPackageName());
+	// Build a unique and deterministic LevelInstance level instance name by using LevelInstanceID. 
+	// Distinguish game from editor since we don't want to duplicate for PIE already loaded editor instances (not yet supported).
+	FString Suffix = FString::Printf(TEXT("%s_LevelInstance_%08X_%d"), *ShortPackageName, LevelInstanceActor->GetLevelInstanceID(), LevelInstanceActor->GetWorld()->IsGameWorld() ? 1 : 0);
+	ULevelStreamingLevelInstance* LevelStreaming = Cast<ULevelStreamingLevelInstance>(ULevelStreamingDynamic::LoadLevelInstanceBySoftObjectPtr(LevelInstanceActor->GetWorld(), LevelInstanceActor->GetWorldAsset(), LevelInstanceActor->GetActorLocation(), LevelInstanceActor->GetActorRotation(), bOutSuccess, Suffix, ULevelStreamingLevelInstance::StaticClass()));
+	if (bOutSuccess)
+	{
+		LevelStreaming->LevelInstanceID = LevelInstanceActor->GetLevelInstanceID();
+		
+#if WITH_EDITOR
+		if (!LevelInstanceActor->GetWorld()->IsPlayInEditor())
+		{
+			GEngine->BlockTillLevelStreamingCompleted(LevelInstanceActor->GetWorld());
+
+			// Most of the code here is meant to allow partial support for undo/redo of LevelInstance Instance Loading:
+			// by setting the objects RF_Transient and !RF_Transactional we can check when unloading if those flags
+			// have been changed and figure out if we need to clear the transaction buffer or not.
+			// It might not be the final solution to support Undo/Redo in LevelInstances but it handles most of the non-editing part
+			ULevel* Level = LevelStreaming->GetLoadedLevel();
+			check(Level);
+			check(LevelStreaming->GetCurrentState() == ULevelStreaming::ECurrentState::LoadedVisible);
+
+			UWorld* OuterWorld = Level->GetTypedOuter<UWorld>();
+			OuterWorld->ClearFlags(RF_Transactional);
+			OuterWorld->SetFlags(RF_Transient);
+			ResetLoaders(OuterWorld->GetPackage());
+
+			OuterWorld->GetPackage()->ClearFlags(RF_Transactional);
+			OuterWorld->GetPackage()->SetFlags(RF_Transient);
+
+			ForEachObjectWithOuter(OuterWorld, [&](UObject* Obj)
+			{
+				Obj->ClearFlags(RF_Transactional);
+				Obj->SetFlags(RF_Transient);
+			}, true);
+
+			for (AActor* LevelActor : Level->Actors)
+			{
+				if (LevelActor)
+				{
+					if (LevelActor->IsPackageExternal())
+					{
+						ResetLoaders(LevelActor->GetExternalPackage());
+						LevelActor->GetPackage()->SetFlags(RF_Transient);
+					}		
+				}
+			}
+
+			// Create special actor that will handle selection and transform
+			ALevelInstanceEditorInstanceActor::Create(LevelInstanceActor, Level);
+
+			// Make sure selection is reflected after load
+			LevelInstanceActor->PushSelectionToProxies();
+		}
+#endif
+		return LevelStreaming;
+	}
+
+	return nullptr;
+}
+
+void ULevelStreamingLevelInstance::UnloadInstance(ULevelStreamingLevelInstance* LevelStreaming)
+{
+	if (LevelStreaming->GetWorld()->IsGameWorld())
+	{
+		LevelStreaming->SetShouldBeLoaded(false);
+		LevelStreaming->SetShouldBeVisible(false);
+		LevelStreaming->SetIsRequestingUnloadAndRemoval(true);
+	}
+#if WITH_EDITOR
+	else
+	{
+		// Check if we need to flush the Trans buffer...
+		ULevel* LoadedLevel = LevelStreaming->GetLoadedLevel();
+		UWorld* OuterWorld = LoadedLevel->GetTypedOuter<UWorld>();
+		bool bResetTrans = false;
+		ForEachObjectWithOuterBreakable(OuterWorld, [&bResetTrans](UObject* Obj)
+		{
+			if(Obj->HasAnyFlags(RF_Transactional))
+			{
+				bResetTrans = true;
+				return false;
+			}
+			return true;
+		}, true);
+
+		// No need to clear the whole editor selection since actor of this level will be removed from the selection by: UEditorEngine::OnLevelRemovedFromWorld
+		const bool bClearSelection = false;
+		LevelStreaming->GetWorld()->GetSubsystem<ULevelInstanceSubsystem>()->RemoveLevelFromWorld(LevelStreaming->GetLoadedLevel(), bResetTrans);
+	}
+#endif 
+}
+
diff --git a/Engine/Source/Runtime/Engine/Private/LevelInstance/LevelInstancePrivate.h b/Engine/Source/Runtime/Engine/Private/LevelInstance/LevelInstancePrivate.h
new file mode 100644
--- /dev/null
+++ b/Engine/Source/Runtime/Engine/Private/LevelInstance/LevelInstancePrivate.h
@@ -0,0 +1,8 @@
+// Copyright Epic Games, Inc. All Rights Reserved.
+
+#pragma once
+
+#include "CoreMinimal.h"
+
+DECLARE_LOG_CATEGORY_EXTERN(LogLevelInstance, Verbose, All);
+
diff --git a/Engine/Source/Runtime/Engine/Private/LevelInstance/LevelInstanceSubsystem.cpp b/Engine/Source/Runtime/Engine/Private/LevelInstance/LevelInstanceSubsystem.cpp
new file mode 100644
--- /dev/null
+++ b/Engine/Source/Runtime/Engine/Private/LevelInstance/LevelInstanceSubsystem.cpp
@@ -0,0 +1,1279 @@
+// Copyright Epic Games, Inc. All Rights Reserved.
+
+#include "LevelInstance/LevelInstanceSubsystem.h"
+#include "LevelInstance/LevelInstanceActor.h"
+#include "LevelInstance/LevelInstanceLevelStreaming.h"
+#include "WorldPartition/WorldPartitionSubsystem.h"
+#include "WorldPartition/WorldPartitionLevelStreamingDynamic.h"
+#include "UObject/UObjectHash.h"
+#include "UObject/UObjectGlobals.h"
+#include "Engine/World.h"
+#include "Engine/Engine.h"
+#include "EngineUtils.h"
+#include "LevelInstancePrivate.h"
+#include "LevelUtils.h"
+#include "Misc/HashBuilder.h"
+
+#if WITH_EDITOR
+#include "LevelInstance/LevelInstanceEditorLevelStreaming.h"
+#include "WorldPartition/LevelInstance/LevelInstanceActorDescFactory.h"
+#include "Misc/ScopedSlowTask.h"
+#include "AssetRegistryModule.h"
+#include "ActorRegistry.h"
+#include "AssetData.h"
+#include "FileHelpers.h"
+#include "Editor.h"
+#include "EditorLevelUtils.h"
+#include "LevelInstance/ILevelInstanceEditorModule.h"
+#include "HAL/PlatformTime.h"
+#include "Engine/Selection.h"
+#include "Engine/LevelBounds.h"
+#include "LevelInstance/LevelInstanceEditorInstanceActor.h"
+#endif
+
+#include "HAL/IConsoleManager.h"
+
+#define LOCTEXT_NAMESPACE "LevelInstanceSubsystem"
+
+DEFINE_LOG_CATEGORY(LogLevelInstance);
+
+ULevelInstanceSubsystem::ULevelInstanceSubsystem()
+	: UWorldSubsystem()
+{
+
+}
+
+ULevelInstanceSubsystem::~ULevelInstanceSubsystem()
+{
+
+}
+
+void ULevelInstanceSubsystem::Initialize(FSubsystemCollectionBase& Collection)
+{
+	Super::Initialize(Collection);
+
+#if WITH_EDITOR
+	LevelInstanceActorDescFactory.Reset(new FLevelInstanceActorDescFactory());
+
+	Collection.InitializeDependency(UWorldPartitionSubsystem::StaticClass());
+	if (UWorldPartitionSubsystem* WorldPartitionSubsystem = GetWorld()->GetSubsystem<UWorldPartitionSubsystem>())
+	{
+		RegisterActorDescFactories(WorldPartitionSubsystem);
+	}
+
+	if (GEditor)
+	{
+		FModuleManager::LoadModuleChecked<ILevelInstanceEditorModule>("LevelInstanceEditor");
+	}
+#endif
+}
+
+void ULevelInstanceSubsystem::Deinitialize()
+{
+	Super::Deinitialize();
+}
+
+ALevelInstance* ULevelInstanceSubsystem::GetLevelInstance(FLevelInstanceID LevelInstanceID) const
+{
+	if (ALevelInstance*const* LevelInstanceActor = RegisteredLevelInstances.Find(LevelInstanceID))
+	{
+		return *LevelInstanceActor;
+	}
+
+	return nullptr;
+}
+
+FLevelInstanceID ULevelInstanceSubsystem::ComputeLevelInstanceID(ALevelInstance* LevelInstanceActor) const
+{
+	FLevelInstanceID LevelInstanceID = InvalidLevelInstanceID;
+	FHashBuilder HashBuilder;
+	ForEachLevelInstanceAncestorsAndSelf(LevelInstanceActor, [&HashBuilder](const ALevelInstance* AncestorOrSelf)
+	{
+		HashBuilder << AncestorOrSelf->GetLevelInstanceActorGuid();
+		return true;
+	});
+	LevelInstanceID = HashBuilder.GetHash();
+	return LevelInstanceID;
+}
+
+FLevelInstanceID ULevelInstanceSubsystem::RegisterLevelInstance(ALevelInstance* LevelInstanceActor)
+{
+	FLevelInstanceID LevelInstanceID = ComputeLevelInstanceID(LevelInstanceActor);
+	check(LevelInstanceID != InvalidLevelInstanceID);
+	ALevelInstance*& Value = RegisteredLevelInstances.FindOrAdd(LevelInstanceID);
+	check(GIsReinstancing || Value == nullptr || Value == LevelInstanceActor);
+	Value = LevelInstanceActor;
+	return LevelInstanceID;
+}
+
+void ULevelInstanceSubsystem::UnregisterLevelInstance(ALevelInstance* LevelInstanceActor)
+{
+	RegisteredLevelInstances.Remove(LevelInstanceActor->GetLevelInstanceID());
+}
+
+void ULevelInstanceSubsystem::RequestLoadLevelInstance(ALevelInstance* LevelInstanceActor, bool bForce /* = false */)
+{
+	check(LevelInstanceActor && !LevelInstanceActor->IsPendingKillOrUnreachable());
+	if (LevelInstanceActor->IsLevelInstancePathValid())
+	{
+#if WITH_EDITOR
+		if (!IsEditingLevelInstance(LevelInstanceActor))
+#endif
+		{
+			LevelInstancesToUnload.Remove(LevelInstanceActor->GetLevelInstanceID());
+
+			bool* bForcePtr = LevelInstancesToLoadOrUpdate.Find(LevelInstanceActor);
+
+			// Avoid loading if already loaded. Can happen if actor requests unload/load in same frame. Without the force it means its not necessary.
+			if (IsLoaded(LevelInstanceActor) && !bForce && (bForcePtr == nullptr || !(*bForcePtr)))
+			{
+				return;
+			}
+
+			if (bForcePtr != nullptr)
+			{
+				*bForcePtr |= bForce;
+			}
+			else
+			{
+				LevelInstancesToLoadOrUpdate.Add(LevelInstanceActor, bForce);
+			}
+		}
+	}
+}
+
+void ULevelInstanceSubsystem::RequestUnloadLevelInstance(ALevelInstance* LevelInstanceActor)
+{
+	const FLevelInstanceID& LevelInstanceID = LevelInstanceActor->GetLevelInstanceID();
+	if (LevelInstances.Contains(LevelInstanceID))
+	{
+		// LevelInstancesToUnload uses FLevelInstanceID because LevelInstanceActor* can be destroyed in later Tick and we don't need it.
+		LevelInstancesToUnload.Add(LevelInstanceID);
+	}
+	LevelInstancesToLoadOrUpdate.Remove(LevelInstanceActor);
+}
+
+bool ULevelInstanceSubsystem::IsLoaded(const ALevelInstance* LevelInstanceActor) const
+{
+	return LevelInstanceActor->HasValidLevelInstanceID() && LevelInstances.Contains(LevelInstanceActor->GetLevelInstanceID());
+}
+
+void ULevelInstanceSubsystem::Tick(float DeltaSeconds)
+{
+#if WITH_EDITOR
+	// For non-game world, Tick is responsible of processing LevelInstances to update/load/unload
+	if (!GetWorld()->IsGameWorld())
+	{
+		UpdateStreamingState();
+
+		// Begin editing the pending LevelInstance when loads complete
+		if (PendingLevelInstanceToEdit != InvalidLevelInstanceID && !LevelInstancesToLoadOrUpdate.Num())
+		{
+			if (ALevelInstance** LevelInstanceActor = RegisteredLevelInstances.Find(PendingLevelInstanceToEdit))
+			{
+				EditLevelInstance(*LevelInstanceActor);
+			}
+		}
+	}
+#endif
+}
+
+bool ULevelInstanceSubsystem::IsTickableInEditor() const
+{
+	return true;
+}
+
+UWorld* ULevelInstanceSubsystem::GetTickableGameObjectWorld() const
+{
+	return GetWorld();
+}
+
+ETickableTickType ULevelInstanceSubsystem::GetTickableTickType() const
+{
+#if WITH_EDITOR
+	return IsTemplate() ? ETickableTickType::Never : ETickableTickType::Always;
+#else
+	return ETickableTickType::Never;
+#endif
+}
+
+TStatId ULevelInstanceSubsystem::GetStatId() const
+{
+	RETURN_QUICK_DECLARE_CYCLE_STAT(ULevelInstanceSubsystem, STATGROUP_Tickables);
+}
+
+void ULevelInstanceSubsystem::UpdateStreamingState()
+{
+	if (!LevelInstancesToUnload.Num() && !LevelInstancesToLoadOrUpdate.Num())
+	{
+		return;
+	}
+
+#if WITH_EDITOR
+	// Do not update during transaction
+	if (GUndo)
+	{
+		return;
+	}
+
+	FScopedSlowTask SlowTask(LevelInstancesToUnload.Num() + LevelInstancesToLoadOrUpdate.Num() * 2, LOCTEXT("UpdatingLevelInstances", "Updating Level Instances..."), !GetWorld()->IsGameWorld());
+	SlowTask.MakeDialog();
+
+	check(!LevelsToRemoveScope);
+	LevelsToRemoveScope.Reset(new FLevelsToRemoveScope());
+#endif
+
+	if (LevelInstancesToUnload.Num())
+	{
+		TSet<FLevelInstanceID> LevelInstancesToUnloadCopy(MoveTemp(LevelInstancesToUnload));
+		for (const FLevelInstanceID& LevelInstanceID : LevelInstancesToUnloadCopy)
+		{
+#if WITH_EDITOR
+			SlowTask.EnterProgressFrame(1.f, LOCTEXT("UnloadingLevelInstance", "Unloading Level Instance"));
+#endif
+			UnloadLevelInstance(LevelInstanceID);
+		}
+	}
+
+	if (LevelInstancesToLoadOrUpdate.Num())
+	{
+		// Unload levels before doing any loading
+		TMap<ALevelInstance*, bool> LevelInstancesToLoadOrUpdateCopy(MoveTemp(LevelInstancesToLoadOrUpdate));
+		for (const TPair<ALevelInstance*, bool>& Pair : LevelInstancesToLoadOrUpdateCopy)
+		{
+#if WITH_EDITOR
+			SlowTask.EnterProgressFrame(1.f, LOCTEXT("UnloadingLevelInstance", "Unloading Level Instance"));
+#endif
+			ALevelInstance* LevelInstanceActor = Pair.Key;
+			if (Pair.Value)
+			{
+				UnloadLevelInstance(LevelInstanceActor->GetLevelInstanceID());
+			}
+		}
+
+#if WITH_EDITOR
+		LevelsToRemoveScope.Reset();
+		double StartTime = FPlatformTime::Seconds();
+#endif
+		for (const TPair<ALevelInstance*, bool>& Pair : LevelInstancesToLoadOrUpdateCopy)
+		{
+#if WITH_EDITOR
+			SlowTask.EnterProgressFrame(1.f, LOCTEXT("LoadingLevelInstance", "Loading Level Instance"));
+#endif
+			LoadLevelInstance(Pair.Key);
+		}
+#if WITH_EDITOR
+		double ElapsedTime = FPlatformTime::Seconds() - StartTime;
+		UE_LOG(LogLevelInstance, Log, TEXT("Loaded %s levels in %s seconds"), *FText::AsNumber(LevelInstancesToLoadOrUpdateCopy.Num()).ToString(), *FText::AsNumber(ElapsedTime).ToString());
+#endif
+	}
+
+#if WITH_EDITOR
+	LevelsToRemoveScope.Reset();
+#endif
+}
+
+void ULevelInstanceSubsystem::LoadLevelInstance(ALevelInstance* LevelInstanceActor)
+{
+	check(LevelInstanceActor);
+	if (IsLoaded(LevelInstanceActor) || LevelInstanceActor->IsPendingKillOrUnreachable() || !LevelInstanceActor->IsLevelInstancePathValid())
+	{
+		return;
+	}
+
+	const FLevelInstanceID& LevelInstanceID = LevelInstanceActor->GetLevelInstanceID();
+	check(!LevelInstances.Contains(LevelInstanceID));
+
+	if (ULevelStreamingLevelInstance* LevelStreaming = ULevelStreamingLevelInstance::LoadInstance(LevelInstanceActor))
+	{
+		FLevelInstance& LevelInstance = LevelInstances.Add(LevelInstanceID);
+		LevelInstance.LevelStreaming = LevelStreaming;
+#if WITH_EDITOR
+		LevelInstanceActor->OnLevelInstanceLoaded();
+#endif
+	}
+}
+
+void ULevelInstanceSubsystem::UnloadLevelInstance(const FLevelInstanceID& LevelInstanceID)
+{
+#if WITH_EDITOR
+	// Create scope if it doesn't exist
+	bool bReleaseScope = false;
+	if (!LevelsToRemoveScope)
+	{
+		bReleaseScope = true;
+		LevelsToRemoveScope.Reset(new FLevelsToRemoveScope());
+	}
+#endif
+
+	FLevelInstance LevelInstance;
+	if (LevelInstances.RemoveAndCopyValue(LevelInstanceID, LevelInstance))
+	{
+		if (ULevel* LoadedLevel = LevelInstance.LevelStreaming->GetLoadedLevel())
+		{
+			ForEachActorInLevel(LoadedLevel, [this](AActor* LevelActor)
+			{
+				if (ALevelInstance* LevelInstanceActor = Cast<ALevelInstance>(LevelActor))
+				{
+					UnloadLevelInstance(LevelInstanceActor->GetLevelInstanceID());
+				}
+				return true;
+			});
+		}
+
+		ULevelStreamingLevelInstance::UnloadInstance(LevelInstance.LevelStreaming);
+	}
+
+#if WITH_EDITOR
+	if (bReleaseScope)
+	{
+		LevelsToRemoveScope.Reset();
+	}
+#endif
+}
+
+void ULevelInstanceSubsystem::ForEachActorInLevel(ULevel* Level, TFunctionRef<bool(AActor * LevelActor)> Operation) const
+{
+	for (AActor* LevelActor : Level->Actors)
+	{
+		if (LevelActor && !LevelActor->IsPendingKill())
+		{
+			if (!Operation(LevelActor))
+			{
+				return;
+			}
+		}
+	}
+}
+
+void ULevelInstanceSubsystem::ForEachLevelInstanceAncestorsAndSelf(AActor* Actor, TFunctionRef<bool(ALevelInstance*)> Operation) const
+{
+	if (ALevelInstance* LevelInstanceActor = Cast<ALevelInstance>(Actor))
+	{
+		if (!Operation(LevelInstanceActor))
+		{
+			return;
+		}
+	}
+
+	ForEachLevelInstanceAncestors(Actor, Operation);
+}
+
+void ULevelInstanceSubsystem::ForEachLevelInstanceAncestors(AActor* Actor, TFunctionRef<bool(ALevelInstance*)> Operation) const
+{
+	ALevelInstance* ParentLevelInstance = nullptr;
+	do
+	{
+		ParentLevelInstance = GetOwningLevelInstance(Actor->GetLevel());
+		Actor = ParentLevelInstance;
+
+	} while (ParentLevelInstance != nullptr && Operation(ParentLevelInstance));
+}
+
+ALevelInstance* ULevelInstanceSubsystem::GetOwningLevelInstance(const ULevel* Level) const
+{
+	if (ULevelStreaming* BaseLevelStreaming = FLevelUtils::FindStreamingLevel(Level))
+	{
+#if WITH_EDITOR
+		if (ULevelStreamingLevelInstanceEditor* LevelStreamingEditor = Cast<ULevelStreamingLevelInstanceEditor>(BaseLevelStreaming))
+		{
+			return LevelStreamingEditor->GetLevelInstanceActor();
+		}
+		else 
+#endif
+		if (ULevelStreamingLevelInstance* LevelStreaming = Cast<ULevelStreamingLevelInstance>(BaseLevelStreaming))
+		{
+			return LevelStreaming->GetLevelInstanceActor();
+		}
+		else if (UWorldPartitionLevelStreamingDynamic* WorldPartitionLevelStreaming = Cast<UWorldPartitionLevelStreamingDynamic>(BaseLevelStreaming))
+		{
+			return GetOwningLevelInstance(WorldPartitionLevelStreaming->GetOuterWorld()->PersistentLevel);
+		}
+	}
+
+	return nullptr;
+}
+
+#if WITH_EDITOR
+
+bool ULevelInstanceSubsystem::GetLevelInstanceBounds(const ALevelInstance* LevelInstanceActor, FBox& OutBounds) const
+{
+	if (IsLoaded(LevelInstanceActor))
+	{
+		const FLevelInstance& LevelInstance = LevelInstances.FindChecked(LevelInstanceActor->GetLevelInstanceID());
+		OutBounds = LevelInstance.LevelStreaming->GetBounds();
+		return true;
+	}
+	else if(LevelInstanceActor->IsLevelInstancePathValid())
+	{
+		if (ULevel::GetLevelBoundsFromPackage(FName(*LevelInstanceActor->GetWorldAssetPackage()), OutBounds))
+		{
+			return true;
+		}
+	}
+
+	return false;
+}
+
+void ULevelInstanceSubsystem::ForEachActorInLevelInstance(const ALevelInstance* LevelInstanceActor, TFunctionRef<bool(AActor * LevelActor)> Operation) const
+{
+	if (ULevel* LevelInstanceLevel = GetLevelInstanceLevel(LevelInstanceActor))
+	{
+		ForEachActorInLevel(LevelInstanceLevel, Operation);
+	}
+}
+
+void ULevelInstanceSubsystem::ForEachLevelInstanceAncestorsAndSelf(const AActor* Actor, TFunctionRef<bool(const ALevelInstance*)> Operation) const
+{
+	if (const ALevelInstance* LevelInstanceActor = Cast<ALevelInstance>(Actor))
+	{
+		if (!Operation(LevelInstanceActor))
+		{
+			return;
+		}
+	}
+
+	ForEachLevelInstanceAncestors(Actor, Operation);
+}
+
+void ULevelInstanceSubsystem::ForEachLevelInstanceAncestors(const AActor* Actor, TFunctionRef<bool(const ALevelInstance*)> Operation) const
+{
+	const ALevelInstance* ParentLevelInstance = nullptr;
+	do 
+	{
+		ParentLevelInstance = GetOwningLevelInstance(Actor->GetLevel());
+		Actor = ParentLevelInstance;
+	} 
+	while (ParentLevelInstance != nullptr && Operation(ParentLevelInstance));
+}
+
+void ULevelInstanceSubsystem::ForEachLevelInstanceChildren(const ALevelInstance* LevelInstanceActor, bool bRecursive, TFunctionRef<bool(const ALevelInstance*)> Operation) const
+{
+	if (ULevel* LevelInstanceLevel = GetLevelInstanceLevel(LevelInstanceActor))
+	{
+		ForEachActorInLevel(LevelInstanceLevel, [this, Operation,bRecursive](AActor* LevelActor)
+		{
+			if (const ALevelInstance* ChildLevelInstanceActor = Cast<ALevelInstance>(LevelActor))
+			{
+				if (Operation(ChildLevelInstanceActor) && bRecursive)
+				{
+					ForEachLevelInstanceChildren(ChildLevelInstanceActor, bRecursive, Operation);
+				}
+			}
+			return true;
+		});
+	}
+}
+
+void ULevelInstanceSubsystem::ForEachLevelInstanceChildren(ALevelInstance* LevelInstanceActor, bool bRecursive, TFunctionRef<bool(ALevelInstance*)> Operation) const
+{
+	if (ULevel* LevelInstanceLevel = GetLevelInstanceLevel(LevelInstanceActor))
+	{
+		ForEachActorInLevel(LevelInstanceLevel, [this, Operation, bRecursive](AActor* LevelActor)
+		{
+			if (ALevelInstance* ChildLevelInstanceActor = Cast<ALevelInstance>(LevelActor))
+			{
+				if (Operation(ChildLevelInstanceActor) && bRecursive)
+				{
+					ForEachLevelInstanceChildren(ChildLevelInstanceActor, bRecursive, Operation);
+				}
+			}
+			return true;
+		});
+	}
+}
+
+void ULevelInstanceSubsystem::ForEachLevelInstanceEdit(TFunctionRef<bool(ALevelInstance*)> Operation) const
+{
+	for (const auto& Pair : LevelInstanceEdits)
+	{
+		if (!Operation(Pair.Value.LevelStreaming->GetLevelInstanceActor()))
+		{
+			return;
+		}
+	}
+}
+
+bool ULevelInstanceSubsystem::HasDirtyChildrenLevelInstances(const ALevelInstance* LevelInstanceActor) const
+{
+	bool bDirtyChildren = false;
+	ForEachLevelInstanceChildren(LevelInstanceActor, /*bRecursive=*/true, [this, &bDirtyChildren](const ALevelInstance* ChildLevelInstanceActor)
+	{
+		if (IsEditingLevelInstanceDirty(ChildLevelInstanceActor))
+		{
+			bDirtyChildren = true;
+			return false;
+		}
+		return true;
+	});
+	return bDirtyChildren;
+}
+
+void ULevelInstanceSubsystem::SetIsTemporarilyHiddenInEditor(ALevelInstance* LevelInstanceActor, bool bIsHidden)
+{
+	if (ULevel* LevelInstanceLevel = GetLevelInstanceLevel(LevelInstanceActor))
+	{
+		ForEachActorInLevel(LevelInstanceLevel, [bIsHidden](AActor* LevelActor)
+		{
+			LevelActor->SetIsTemporarilyHiddenInEditor(bIsHidden);
+			return true;
+		});
+	}
+}
+
+bool ULevelInstanceSubsystem::SetCurrent(ALevelInstance* LevelInstanceActor) const
+{
+	if (IsEditingLevelInstance(LevelInstanceActor))
+	{
+		return GetWorld()->SetCurrentLevel(GetLevelInstanceLevel(LevelInstanceActor));
+	}
+
+	return false;
+}
+
+bool ULevelInstanceSubsystem::IsCurrent(const ALevelInstance* LevelInstanceActor) const
+{
+	if (IsEditingLevelInstance(LevelInstanceActor))
+	{
+		return GetLevelInstanceLevel(LevelInstanceActor) == GetWorld()->GetCurrentLevel();
+	}
+
+	return false;
+}
+
+bool ULevelInstanceSubsystem::MoveActorsToLevel(const TArray<AActor*>& ActorsToRemove, ULevel* DestinationLevel) const
+{
+	check(DestinationLevel);
+
+	const bool bWarnAboutReferences = true;
+	const bool bWarnAboutRenaming = true;
+	const bool bMoveAllOrFail = true;
+	if (!EditorLevelUtils::MoveActorsToLevel(ActorsToRemove, DestinationLevel, bWarnAboutReferences, bWarnAboutRenaming, bMoveAllOrFail))
+	{
+		UE_LOG(LogLevelInstance, Warning, TEXT("Failed to move actors out of Level Instance because not all actors could be moved"));
+		return false;
+	}
+	return true;
+}
+
+bool ULevelInstanceSubsystem::MoveActorsTo(ALevelInstance* LevelInstanceActor, const TArray<AActor*>& ActorsToMove)
+{
+	check(IsEditingLevelInstance(LevelInstanceActor));
+	ULevel* LevelInstanceLevel = GetLevelInstanceLevel(LevelInstanceActor);
+	check(LevelInstanceLevel);
+
+	return MoveActorsToLevel(ActorsToMove, LevelInstanceLevel);
+}
+
+ALevelInstance* ULevelInstanceSubsystem::CreateLevelInstanceFrom(const TArray<AActor*>& ActorsToMove, UWorld* TemplateWorld)
+{
+	ULevel* CurrentLevel = GetWorld()->GetCurrentLevel();
+		
+	if (ActorsToMove.Num() == 0)
+	{
+		UE_LOG(LogLevelInstance, Warning, TEXT("Failed to create Level Instance from empty actor array"));
+		return nullptr;
+	}
+		
+	FBox ActorLocationBox(ForceInit);
+	for (const AActor* ActorToMove : ActorsToMove)
+	{
+		const bool bNonColliding = false;
+		const bool bIncludeChildren = true;
+		ActorLocationBox += ActorToMove->GetComponentsBoundingBox(bNonColliding, bIncludeChildren);
+
+		if (!CanMoveActorToLevel(ActorToMove))
+		{
+			return nullptr;
+		}
+	}
+
+	FVector LevelInstanceLocation = ActorLocationBox.GetCenter();
+	LevelInstanceLocation.Z = ActorLocationBox.Min.Z;
+
+	ULevelStreamingLevelInstanceEditor* LevelStreaming = StaticCast<ULevelStreamingLevelInstanceEditor*>(EditorLevelUtils::CreateNewStreamingLevelForWorld(*GetWorld(), ULevelStreamingLevelInstanceEditor::StaticClass(), TEXT(""), false, TemplateWorld));
+	if (!LevelStreaming)
+	{
+		UE_LOG(LogLevelInstance, Warning, TEXT("Failed to create new Level Instance level"));
+		return nullptr;
+	}
+
+	ULevel* LoadedLevel = LevelStreaming->GetLoadedLevel();
+	check(LoadedLevel);
+
+	const bool bWarnAboutReferences = true;
+	const bool bWarnAboutRenaming = true;
+	const bool bMoveAllOrFail = true;
+
+	if (!EditorLevelUtils::MoveActorsToLevel(ActorsToMove, LoadedLevel, bWarnAboutReferences, bWarnAboutRenaming, bMoveAllOrFail))
+	{
+		ULevelStreamingLevelInstanceEditor::Unload(LevelStreaming);
+		UE_LOG(LogLevelInstance, Warning, TEXT("Failed to create Level Instance because some actors couldn't be moved"));
+		return nullptr;
+	}
+	
+	FActorSpawnParameters SpawnParams;
+	SpawnParams.bCreateActorPackage = true;
+	SpawnParams.OverrideLevel = CurrentLevel;
+	ALevelInstance* NewLevelInstanceActor = GetWorld()->SpawnActor<ALevelInstance>(ALevelInstance::StaticClass(), SpawnParams);
+	check(NewLevelInstanceActor);
+	NewLevelInstanceActor->SetWorldAsset(LoadedLevel->GetTypedOuter<UWorld>());
+	NewLevelInstanceActor->SetActorLocation(LevelInstanceLocation);
+
+	// Actors were moved and kept their World positions so when saving we want their positions to actually be relative to the FounationActor/LevelTransform
+	// so we set the LevelTransform and we mark the level as having moved its actors. 
+	// On Level save FLevelUtils::RemoveEditorTransform will fixup actor transforms to make them relative to the LevelTransform.
+	LevelStreaming->LevelTransform = NewLevelInstanceActor->GetActorTransform();
+	LoadedLevel->bAlreadyMovedActors = true;
+
+	GEditor->SelectNone(false, true);
+	GEditor->SelectActor(NewLevelInstanceActor, true, true);
+
+	FLevelInstanceEdit& LevelInstanceEdit = LevelInstanceEdits.Add(FName(*NewLevelInstanceActor->GetWorldAssetPackage()));
+	LevelInstanceEdit.LevelStreaming = LevelStreaming;
+	LevelStreaming->LevelInstanceID = NewLevelInstanceActor->GetLevelInstanceID();
+
+	GetWorld()->SetCurrentLevel(LoadedLevel);
+
+	if (!UEditorLoadingAndSavingUtils::SavePackages({ LoadedLevel->GetOutermost() }, true))
+	{
+		UE_LOG(LogLevelInstance, Error, TEXT("Failed to save Level Instance (%s)"), *NewLevelInstanceActor->GetWorldAssetPackage());
+	}
+	else
+	{
+		UE_LOG(LogLevelInstance, Log, TEXT("Level Instance create successfully (%s)"), *NewLevelInstanceActor->GetWorldAssetPackage());
+	}
+
+	// Exit edit
+	CommitLevelInstance(NewLevelInstanceActor);
+
+	return NewLevelInstanceActor;
+}
+
+bool ULevelInstanceSubsystem::BreakLevelInstance(ALevelInstance* LevelInstanceActor, uint32 Levels)
+{
+	if (Levels > 0)
+	{
+		// Can only break the top level LevelInstance
+		check(LevelInstanceActor->GetLevel() == GetWorld()->GetCurrentLevel());
+
+		// need to ensure that LevelInstanceActor has been streamed in fully
+		GEngine->BlockTillLevelStreamingCompleted(LevelInstanceActor->GetWorld());
+
+		TArray<AActor*> ActorsToMove;
+		ForEachActorInLevelInstance(LevelInstanceActor, [this, &ActorsToMove](AActor* Actor)
+			{
+				// Skip some actor types
+				if (!Actor->IsA<ALevelBounds>() && !Actor->IsA<ABrush>() && !Actor->IsA<AWorldSettings>() && !Actor->IsA<ALevelInstanceEditorInstanceActor>())
+				{
+					if (CanMoveActorToLevel(Actor))
+					{
+						FSetActorHiddenInSceneOutliner Show(Actor, false);
+						Actor->DetachFromActor(FDetachmentTransformRules::KeepWorldTransform);
+						ActorsToMove.Add(Actor);
+					}
+				}
+
+				return true;
+			});
+
+		ULevel* DestinationLevel = GetWorld()->GetCurrentLevel();
+		check(DestinationLevel);
+
+		const bool bWarnAboutReferences = true;
+		const bool bWarnAboutRenaming = true;
+		const bool bMoveAllOrFail = true;
+		if (!EditorLevelUtils::MoveActorsToLevel(ActorsToMove, DestinationLevel, bWarnAboutReferences, bWarnAboutRenaming, bMoveAllOrFail))
+		{
+			UE_LOG(LogLevelInstance, Warning, TEXT("Failed to break Level Instance because not all actors could be moved"));
+			return false;
+		}
+
+		// Destroy the old LevelInstance instance actor
+		GetWorld()->DestroyActor(LevelInstanceActor);
+	
+		// Break up any sub LevelInstances if more levels are requested
+		if (Levels > 1)
+		{
+			TArray<ALevelInstance*> Children;
+			for (FSelectionIterator It(GEditor->GetSelectedActorIterator()); It; ++It)
+			{
+				AActor* Actor = Cast<AActor>(*It);
+				if (ALevelInstance* ChildLevelInstance = Cast<ALevelInstance>(Actor))
+				{
+					Children.Add(ChildLevelInstance);
+				}
+			}
+
+			for (auto& Child : Children)
+			{
+				BreakLevelInstance(Child, Levels - 1);
+			}
+		}
+	}
+
+	return true;
+}
+
+ULevel* ULevelInstanceSubsystem::GetLevelInstanceLevel(const ALevelInstance* LevelInstanceActor) const
+{
+	if (LevelInstanceActor->HasValidLevelInstanceID())
+	{
+		if (const FLevelInstanceEdit* LevelInstanceEdit = GetLevelInstanceEdit(LevelInstanceActor))
+		{
+			return LevelInstanceEdit->LevelStreaming->GetLoadedLevel();
+		}
+		else if (const FLevelInstance* LevelInstance = LevelInstances.Find(LevelInstanceActor->GetLevelInstanceID()))
+		{
+			return LevelInstance->LevelStreaming->GetLoadedLevel();
+		}
+	}
+
+	return nullptr;
+}
+
+void ULevelInstanceSubsystem::RemoveLevelFromWorld(ULevel* Level, bool bResetTrans)
+{
+	if (LevelsToRemoveScope)
+	{
+		LevelsToRemoveScope->Levels.AddUnique(Level);
+		LevelsToRemoveScope->bResetTrans |= bResetTrans;
+	}
+	else
+	{
+		EditorLevelUtils::RemoveLevelFromWorld(Level, false, bResetTrans);
+	}
+}
+
+ULevelInstanceSubsystem::FLevelsToRemoveScope::~FLevelsToRemoveScope()
+{
+	if (Levels.Num() > 0)
+	{
+		double StartTime = FPlatformTime::Seconds();
+		const bool bClearSelection = false;
+		// No need to clear the whole editor selection since actor of this level will be removed from the selection by: UEditorEngine::OnLevelRemovedFromWorld
+		EditorLevelUtils::RemoveLevelsFromWorld(Levels, bClearSelection, bResetTrans);
+		double ElapsedTime = FPlatformTime::Seconds() - StartTime;
+		UE_LOG(LogLevelInstance, Log, TEXT("Unloaded %s levels in %s seconds"), *FText::AsNumber(Levels.Num()).ToString(), *FText::AsNumber(ElapsedTime).ToString());
+	}
+}
+
+bool ULevelInstanceSubsystem::CanMoveActorToLevel(const AActor* Actor) const
+{
+	if (Actor->GetWorld() == GetWorld())
+	{
+		if (const ALevelInstance* LevelInstanceActor = Cast<ALevelInstance>(Actor))
+		{
+			if (IsEditingLevelInstance(LevelInstanceActor))
+			{
+				UE_LOG(LogLevelInstance, Warning, TEXT("Can't move Level Instance actor while it is being edited"));
+				return false;
+			}
+
+			bool bEditingChildren = false;
+			ForEachLevelInstanceChildren(LevelInstanceActor, true, [this, &bEditingChildren](const ALevelInstance* ChildLevelInstanceActor)
+			{
+				if (IsEditingLevelInstance(ChildLevelInstanceActor))
+				{
+					bEditingChildren = true;
+					return false;
+				}
+				return true;
+			});
+
+			if (bEditingChildren)
+			{
+				UE_LOG(LogLevelInstance, Warning, TEXT("Can't move Level Instance actor while one of its child Level Instance is being edited"));
+				return false;
+			}
+		}
+	}
+
+	return true;
+}
+
+void ULevelInstanceSubsystem::DiscardEdits()
+{
+	for (const auto& Pair : LevelInstanceEdits)
+	{
+		ULevelStreamingLevelInstanceEditor::Unload(Pair.Value.LevelStreaming);
+	}
+	LevelInstanceEdits.Empty();
+}
+
+void ULevelInstanceSubsystem::OnActorDeleted(AActor* Actor)
+{
+	if (ALevelInstance* LevelInstanceActor = Cast<ALevelInstance>(Actor))
+	{
+		if (Actor->GetClass()->HasAnyClassFlags(CLASS_NewerVersionExists))
+		{
+			// We are receiving this event when destroying the old actor after BP reinstantiation. In this case,
+			// the newly created actor was already added to the list, so we can safely ignore this case.
+			check(GIsReinstancing);
+			return;
+		}
+
+		const bool bAlreadyRooted = LevelInstanceActor->IsRooted();
+		// Unloading LevelInstances leads to GC and Actor can be collected. Add to root temp. It will get collected after the OnActorDeleted callbacks
+		if (!bAlreadyRooted)
+		{
+			LevelInstanceActor->AddToRoot();
+		}
+
+		FScopedSlowTask SlowTask(0, LOCTEXT("UnloadingLevelInstances", "Unloading Level Instances..."), !GetWorld()->IsGameWorld());
+		SlowTask.MakeDialog();
+		check(!IsEditingLevelInstanceDirty(LevelInstanceActor) && !HasDirtyChildrenLevelInstances(LevelInstanceActor));
+		if (IsEditingLevelInstance(LevelInstanceActor))
+		{
+			CommitLevelInstance(LevelInstanceActor);
+		}
+		CommitChildrenLevelInstances(LevelInstanceActor);
+		LevelInstancesToLoadOrUpdate.Remove(LevelInstanceActor);
+				
+		UnloadLevelInstance(LevelInstanceActor->GetLevelInstanceID());
+		
+		// Remove from root so it gets collected on the next GC if it can be.
+		if (!bAlreadyRooted)
+		{
+			LevelInstanceActor->RemoveFromRoot();
+		}
+	}
+}
+
+void ULevelInstanceSubsystem::RegisterActorDescFactories(UWorldPartitionSubsystem* WorldPartitionSubsystem)
+{
+	WorldPartitionSubsystem->RegisterActorDescFactory(ALevelInstance::StaticClass(), LevelInstanceActorDescFactory.Get());
+}
+
+bool ULevelInstanceSubsystem::ShouldIgnoreDirtyPackage(UPackage* DirtyPackage, const UWorld* EditingWorld)
+{
+	if (DirtyPackage == EditingWorld->GetOutermost())
+	{
+		return false;
+	}
+
+	bool bIgnore = true;
+	ForEachObjectWithPackage(DirtyPackage, [&bIgnore, EditingWorld](UObject* Object)
+	{
+		if (Object->GetOutermostObject() == EditingWorld)
+		{
+			bIgnore = false;
+		}
+
+		return bIgnore;
+	});
+
+	return bIgnore;
+}
+
+UWorld* ULevelInstanceSubsystem::FLevelInstanceEdit::GetEditWorld() const
+{
+	if (LevelStreaming && LevelStreaming->GetLoadedLevel())
+	{
+		return LevelStreaming->GetLoadedLevel()->GetTypedOuter<UWorld>();
+	}
+
+	return nullptr;
+}
+
+const ULevelInstanceSubsystem::FLevelInstanceEdit* ULevelInstanceSubsystem::GetLevelInstanceEdit(const ALevelInstance* LevelInstanceActor) const
+{
+	const FLevelInstanceEdit* LevelInstanceEdit = LevelInstanceEdits.Find(FName(*LevelInstanceActor->GetWorldAssetPackage()));
+	if (!LevelInstanceEdit || LevelInstanceEdit->LevelStreaming->GetLevelInstanceActor() != LevelInstanceActor)
+	{
+		return nullptr;
+	}
+
+	return LevelInstanceEdit;
+}
+
+bool ULevelInstanceSubsystem::IsEditingLevelInstanceDirty(const ALevelInstance* LevelInstanceActor) const
+{
+	const FLevelInstanceEdit* LevelInstanceEdit = GetLevelInstanceEdit(LevelInstanceActor);
+	if (!LevelInstanceEdit)
+	{
+		return false;
+	}
+
+	return IsLevelInstanceEditDirty(LevelInstanceEdit);
+}
+
+bool ULevelInstanceSubsystem::IsLevelInstanceEditDirty(const FLevelInstanceEdit* LevelInstanceEdit) const
+{
+	const UWorld* EditingWorld = LevelInstanceEdit->GetEditWorld();
+	check(EditingWorld);
+
+	TArray<UPackage*> OutDirtyPackages;
+	FEditorFileUtils::GetDirtyPackages(OutDirtyPackages, [EditingWorld](UPackage* DirtyPackage)
+	{
+		return ULevelInstanceSubsystem::ShouldIgnoreDirtyPackage(DirtyPackage, EditingWorld);
+	});
+
+	return OutDirtyPackages.Num() > 0;
+}
+
+bool ULevelInstanceSubsystem::CanEditLevelInstance(const ALevelInstance* LevelInstanceActor, FText* OutReason) const
+{
+	if (ULevel* LevelInstanceLevel = GetLevelInstanceLevel(LevelInstanceActor))
+	{
+		if (LevelInstanceLevel->GetWorldPartition())
+		{
+			if (OutReason)
+			{
+				*OutReason = LOCTEXT("CanEditPartitionedLevelInstance", "Can't edit partitioned Level Instance");
+			}
+			return false;
+		}
+	}
+
+	if (const FLevelInstanceEdit* LevelInstanceEdit = LevelInstanceEdits.Find(FName(*LevelInstanceActor->GetWorldAssetPackage())))
+	{
+		if (OutReason)
+		{
+			if (LevelInstanceEdit->LevelStreaming->GetLevelInstanceActor() == LevelInstanceActor)
+			{
+				*OutReason = LOCTEXT("CanEditLevelInstanceAlreadyBeingEdited", "Level Instance already being edited");
+			}
+			else
+			{
+
+				*OutReason = LOCTEXT("CanEditLevelInstanceAlreadyEditing", "Another Level Instance pointing to the same level is being edited");
+			}
+		}
+		return false;
+	}
+
+	// Do not allow multiple LevelInstances of the same hierarchy to be edited... (checking ancestors)
+	bool bAncestorBeingEdited = false;
+	if (LevelInstanceEdits.Num() > 0)
+	{
+		ForEachLevelInstanceAncestors(LevelInstanceActor, [this, &bAncestorBeingEdited, OutReason](const ALevelInstance* AncestorLevelInstance)
+		{
+			if (const FLevelInstanceEdit* LevelInstanceEdit = LevelInstanceEdits.Find(FName(*AncestorLevelInstance->GetWorldAssetPackage())))
+			{
+				// Allow children to be edited if ancestor is clean
+				if (IsLevelInstanceEditDirty(LevelInstanceEdit))
+				{
+					if (OutReason)
+					{
+						*OutReason = LOCTEXT("CanEditLevelInstanceAncestorBeingEdited", "Ancestor Level Instance already being edited");
+					}
+					bAncestorBeingEdited = true;
+					return false;
+				}
+			}
+			return true;
+		});
+	}
+
+	if (bAncestorBeingEdited)
+	{
+		return false;
+	}
+		
+	// Do not allow multiple LevelInstances of the same hierarchy to be edited... (checking children)
+	bool bChildBeingEdited = false;
+	if (LevelInstanceEdits.Num() > 0)
+	{
+		for (const auto& Pair : LevelInstanceEdits)
+		{
+			check(Pair.Value.LevelStreaming);
+			ALevelInstance* LevelInstanceEditActor = Pair.Value.LevelStreaming->GetLevelInstanceActor();
+			check(LevelInstanceEditActor);
+			ForEachLevelInstanceAncestors(LevelInstanceEditActor, [this, &bChildBeingEdited, OutReason, LevelInstanceActor](const ALevelInstance* AncestorLevelInstance)
+			{
+				if (AncestorLevelInstance->GetWorldAsset() == LevelInstanceActor->GetWorldAsset())
+				{
+					if (const FLevelInstanceEdit* AncestorLevelInstanceEdit = LevelInstanceEdits.Find(FName(*AncestorLevelInstance->GetWorldAssetPackage())))
+					{
+						if (IsLevelInstanceEditDirty(AncestorLevelInstanceEdit))
+						{
+							if (OutReason)
+							{
+								*OutReason = LOCTEXT("CanEditLevelInstanceOtherChildren", "Children Level Instance already being edited");
+							}
+							bChildBeingEdited = true;
+							return false;
+						}
+					}
+				}
+
+				return true;
+			});
+
+			if (bChildBeingEdited)
+			{
+				break;
+			}
+		}
+	}
+	
+	if (bChildBeingEdited)
+	{
+		return false;
+	}
+
+	if (!LevelInstanceActor->IsLevelInstancePathValid())
+	{
+		if (OutReason)
+		{
+			*OutReason = LOCTEXT("CanEditLevelInstanceDirtyInvalid", "Level Instance path is invalid");
+		}
+		return false;
+	}
+
+	if (FLevelUtils::FindStreamingLevel(GetWorld(), *LevelInstanceActor->GetWorldAssetPackage()))
+	{
+		if (OutReason)
+		{
+			*OutReason = LOCTEXT("CanEditLevelInstanceAlreadyExists", "The same level was added to world outside of Level Instances");
+		}
+		return false;
+	}
+
+	return true;
+}
+
+bool ULevelInstanceSubsystem::CanCommitLevelInstance(const ALevelInstance* LevelInstanceActor, FText* OutReason) const
+{
+	if (!IsEditingLevelInstance(LevelInstanceActor))
+	{
+		if (OutReason)
+		{
+			*OutReason = LOCTEXT("CanCommitLevelInstanceNotEditing", "Level Instance is not currently being edited");
+		}
+		return false;
+	}
+
+	return true;
+}
+
+void ULevelInstanceSubsystem::EditLevelInstance(ALevelInstance* LevelInstanceActor, TWeakObjectPtr<AActor> ContextActorPtr)
+{
+	check(CanEditLevelInstance(LevelInstanceActor));
+	PendingLevelInstanceToEdit = InvalidLevelInstanceID;
+		
+	FScopedSlowTask SlowTask(0, LOCTEXT("BeginEditLevelInstance", "Loading Level Instance for edit..."), !GetWorld()->IsGameWorld());
+	SlowTask.MakeDialog();
+
+	// Gather information from the context actor to try and select something meaningful after the loading
+	FString ActorNameToSelect;
+	if (AActor* ContextActor = ContextActorPtr.Get())
+	{
+		ActorNameToSelect = ContextActor->GetName();
+		ForEachLevelInstanceAncestorsAndSelf(ContextActor, [&ActorNameToSelect,LevelInstanceActor](const ALevelInstance* AncestorLevelInstanceActor)
+		{
+			// stop when we hit the LevelInstance we are about to edit
+			if (AncestorLevelInstanceActor == LevelInstanceActor)
+			{
+				return false;
+			}
+			
+			ActorNameToSelect = AncestorLevelInstanceActor->GetName();
+			return true;
+		});
+	}
+
+	GEditor->SelectNone(false, true);
+
+	// Check if there is an open (but clean) ancestor and child and unload it before opening the LevelInstance for editing
+	if (LevelInstanceEdits.Num() > 0)
+	{
+		ALevelInstance* LevelInstanceToCommit = nullptr;
+
+		auto GetLevelInstanceToCommit = [this, &LevelInstanceToCommit](ALevelInstance* LevelInstance)
+		{
+			if (const FLevelInstanceEdit* LevelInstanceEdit = LevelInstanceEdits.Find(FName(*LevelInstance->GetWorldAssetPackage())))
+			{
+				check(!IsLevelInstanceEditDirty(LevelInstanceEdit));
+				check(LevelInstanceToCommit == nullptr);
+				LevelInstanceToCommit = LevelInstanceEdit->LevelStreaming->GetLevelInstanceActor();
+				check(LevelInstanceToCommit != nullptr);
+				return false;
+			}
+			return true;
+		};
+
+		ForEachLevelInstanceAncestors(LevelInstanceActor, GetLevelInstanceToCommit);
+		if (!LevelInstanceToCommit)
+		{
+			ForEachLevelInstanceChildren(LevelInstanceActor, true, GetLevelInstanceToCommit);
+		}
+
+		if (LevelInstanceToCommit)
+		{
+			PendingLevelInstanceToEdit = LevelInstanceActor->GetLevelInstanceID();
+			CommitLevelInstance(LevelInstanceToCommit);
+
+			// Stop here. The LevelInstance will be open for editing after an async reload.
+			return;
+		}
+	}
+
+	// Cleanup async requests in case
+	LevelInstancesToUnload.Remove(LevelInstanceActor->GetLevelInstanceID());
+	LevelInstancesToLoadOrUpdate.Remove(LevelInstanceActor);
+	// Unload right away
+	UnloadLevelInstance(LevelInstanceActor->GetLevelInstanceID());
+		
+	// Load Edit LevelInstance level
+	ULevelStreamingLevelInstanceEditor* LevelStreaming = ULevelStreamingLevelInstanceEditor::Load(LevelInstanceActor);
+	FLevelInstanceEdit& LevelInstanceEdit = LevelInstanceEdits.Add(FName(*LevelInstanceActor->GetWorldAssetPackage()));
+	LevelInstanceEdit.LevelStreaming = LevelStreaming;
+
+	// Try and select something meaningful
+	AActor* ActorToSelect = nullptr;
+	if (!ActorNameToSelect.IsEmpty())
+	{		
+		ActorToSelect = FindObject<AActor>(LevelStreaming->GetLoadedLevel(), *ActorNameToSelect);
+	}
+
+	// default to LevelInstance
+	if (!ActorToSelect)
+	{
+		ActorToSelect = LevelInstanceActor;
+	}
+	LevelInstanceActor->SetIsTemporarilyHiddenInEditor(false);
+	GEditor->SelectActor(ActorToSelect, true, true);
+}
+
+void ULevelInstanceSubsystem::CommitChildrenLevelInstances(ALevelInstance* LevelInstanceActor)
+{
+	// We are ending editing. Discard Non dirty child edits
+	ForEachLevelInstanceChildren(LevelInstanceActor, /*bRecursive=*/true, [this](const ALevelInstance* ChildLevelInstanceActor)
+	{
+		if (const FLevelInstanceEdit* ChildLevelInstanceEdit = GetLevelInstanceEdit(ChildLevelInstanceActor))
+		{
+			check(!IsLevelInstanceEditDirty(ChildLevelInstanceEdit));
+			ULevelStreamingLevelInstanceEditor::Unload(ChildLevelInstanceEdit->LevelStreaming);
+			LevelInstanceEdits.Remove(FName(*ChildLevelInstanceActor->GetWorldAssetPackage()));
+		}
+		return true;
+	});
+}
+
+void ULevelInstanceSubsystem::CommitLevelInstance(ALevelInstance* LevelInstanceActor, bool bDiscardEdits)
+{
+	check(CanCommitLevelInstance(LevelInstanceActor));
+
+	const FLevelInstanceEdit* LevelInstanceEdit = GetLevelInstanceEdit(LevelInstanceActor);
+	check(LevelInstanceEdit);
+	UWorld* EditingWorld = LevelInstanceEdit->GetEditWorld();
+	check(EditingWorld);
+			
+	bool bChangesCommitted = false;
+	if (IsLevelInstanceEditDirty(LevelInstanceEdit) && !bDiscardEdits)
+	{
+		const bool bPromptUserToSave = true;
+		const bool bSaveMapPackages = true;
+		const bool bSaveContentPackages = true;
+		const bool bFastSave = false;
+		const bool bNotifyNoPackagesSaved = false;
+		const bool bCanBeDeclined = true;
+
+		if (!FEditorFileUtils::SaveDirtyPackages(bPromptUserToSave, bSaveMapPackages, bSaveContentPackages, bFastSave, bNotifyNoPackagesSaved, bCanBeDeclined, nullptr,
+			[=](UPackage* DirtyPackage)
+			{
+				return ShouldIgnoreDirtyPackage(DirtyPackage, EditingWorld);
+			}))
+		{
+			return;
+		}
+
+		// Validate that we indeed need to refresh instances (user can cancel the changes when prompted)
+		bChangesCommitted = !IsLevelInstanceEditDirty(LevelInstanceEdit);
+
+		if(bChangesCommitted)
+		{
+			// Sync the AssetData so that the updated instances have the latest Actor Registry Data
+			IAssetRegistry& AssetRegistry = FModuleManager::LoadModuleChecked<FAssetRegistryModule>(TEXT("AssetRegistry")).Get();
+			AssetRegistry.ScanPathsSynchronous({ LevelInstanceActor->GetWorldAssetPackage() }, true);
+		}
+	}
+
+	FScopedSlowTask SlowTask(0, LOCTEXT("EndEditLevelInstance", "Unloading edit Level Instance..."), !GetWorld()->IsGameWorld());
+	SlowTask.MakeDialog();
+
+	GEditor->SelectNone(false, true);
+
+	// End edit non dirty child edits
+	CommitChildrenLevelInstances(LevelInstanceActor);
+
+	// Try to find proper LevelInstance to select
+	ALevelInstance* ActorToSelect = LevelInstanceActor;
+	ForEachLevelInstanceAncestors(LevelInstanceActor, [&ActorToSelect](ALevelInstance* AncestorLevelInstance)
+	{
+		// if we find a parent editing LevelInstance this is what we want to select
+		if (AncestorLevelInstance->IsEditing())
+		{
+			ActorToSelect = AncestorLevelInstance;
+			return false;
+		}
+		else // if not we go up the ancestor LevelInstances to the highest level
+		{
+			ActorToSelect = AncestorLevelInstance;
+		}
+		return true;
+	});
+
+	const FString EditPackage = LevelInstanceActor->GetWorldAssetPackage();
+
+	// Remove from streaming level...
+	ULevelStreamingLevelInstanceEditor::Unload(LevelInstanceEdit->LevelStreaming);
+	LevelInstanceEdits.Remove(FName(*EditPackage));
+
+	// Propagate to other instances
+	for (TActorIterator<ALevelInstance> LevelInstanceIt(GetWorld()); LevelInstanceIt; ++LevelInstanceIt)
+	{
+		ALevelInstance* CurrentLevelInstanceActor = *LevelInstanceIt;
+		if (CurrentLevelInstanceActor->GetWorldAssetPackage() == EditPackage && (LevelInstanceActor == CurrentLevelInstanceActor || bChangesCommitted))
+		{
+			CurrentLevelInstanceActor->UpdateLevelInstance();
+		}
+	}
+
+	GEditor->SelectActor(ActorToSelect, true, true);
+}
+
+void ULevelInstanceSubsystem::SaveLevelInstanceAs(ALevelInstance* LevelInstanceActor)
+{
+	check(CanCommitLevelInstance(LevelInstanceActor));
+
+	const FLevelInstanceEdit* OldLevelInstanceEdit = GetLevelInstanceEdit(LevelInstanceActor);
+	check(OldLevelInstanceEdit);
+	UWorld* EditingWorld = OldLevelInstanceEdit->GetEditWorld();
+	check(EditingWorld);
+
+	// Reset the level transform before saving
+	OldLevelInstanceEdit->LevelStreaming->GetLoadedLevel()->ApplyWorldOffset(-LevelInstanceActor->GetTransform().GetLocation(), false);
+
+	TArray<UObject*> OutObjects;
+	FEditorFileUtils::SaveAssetsAs({ EditingWorld }, OutObjects);
+
+	if (OutObjects.Num() == 0 || OutObjects[0] == EditingWorld)
+	{
+		UE_LOG(LogLevelInstance, Warning, TEXT("Failed to save Level Instance as new asset"));
+		return;
+	}
+
+	UWorld* SavedWorld = StaticCast<UWorld*>(OutObjects[0]);
+	// Discard edits and unload streaming level
+	DiscardEdits();
+	
+	LevelInstanceActor->SetWorldAsset(SavedWorld);
+
+	LoadLevelInstance(LevelInstanceActor);
+	GEditor->SelectActor(LevelInstanceActor, true, true);
+}
+
+ALevelInstance* ULevelInstanceSubsystem::GetParentLevelInstance(const AActor* Actor) const
+{
+	check(Actor);
+	const ULevel* OwningLevel = Actor->GetLevel();
+	check(OwningLevel);
+	return GetOwningLevelInstance(OwningLevel);
+}
+#endif
+
+#undef LOCTEXT_NAMESPACE
+
diff --git a/Engine/Source/Runtime/Engine/Private/WorldPartition/Foundation/FoundationActorDesc.cpp b/Engine/Source/Runtime/Engine/Private/WorldPartition/Foundation/FoundationActorDesc.cpp
deleted file mode 100644
diff --git a/Engine/Source/Runtime/Engine/Private/WorldPartition/Foundation/FoundationActorDescFactory.cpp b/Engine/Source/Runtime/Engine/Private/WorldPartition/Foundation/FoundationActorDescFactory.cpp
deleted file mode 100644
diff --git a/Engine/Source/Runtime/Engine/Private/WorldPartition/LevelInstance/LevelInstanceActorDesc.cpp b/Engine/Source/Runtime/Engine/Private/WorldPartition/LevelInstance/LevelInstanceActorDesc.cpp
new file mode 100644
--- /dev/null
+++ b/Engine/Source/Runtime/Engine/Private/WorldPartition/LevelInstance/LevelInstanceActorDesc.cpp
@@ -0,0 +1,29 @@
+// Copyright Epic Games, Inc. All Rights Reserved.
+
+#include "WorldPartition/LevelInstance/LevelInstanceActorDesc.h"
+
+#if WITH_EDITOR
+#include "Engine/Level.h"
+#include "LevelInstance/LevelInstanceActor.h"
+
+FLevelInstanceActorDesc::FLevelInstanceActorDesc(const FWorldPartitionActorDescData& DescData, FName InLevelPackage)
+	: FWorldPartitionActorDesc(DescData)
+	, LevelPackage(InLevelPackage)
+{}
+
+FLevelInstanceActorDesc::FLevelInstanceActorDesc(AActor* InActor)
+	: FWorldPartitionActorDesc(InActor)
+{
+	if (ALevelInstance* LevelInstanceActor = CastChecked<ALevelInstance>(InActor))
+	{
+		LevelPackage = *LevelInstanceActor->GetWorldAssetPackage();
+	}
+}
+
+void FLevelInstanceActorDesc::BuildHash(FHashBuilder& HashBuilder)
+{
+	FWorldPartitionActorDesc::BuildHash(HashBuilder);
+	HashBuilder << LevelPackage;
+}
+
+#endif
diff --git a/Engine/Source/Runtime/Engine/Private/WorldPartition/LevelInstance/LevelInstanceActorDescFactory.cpp b/Engine/Source/Runtime/Engine/Private/WorldPartition/LevelInstance/LevelInstanceActorDescFactory.cpp
new file mode 100644
--- /dev/null
+++ b/Engine/Source/Runtime/Engine/Private/WorldPartition/LevelInstance/LevelInstanceActorDescFactory.cpp
@@ -0,0 +1,59 @@
+// Copyright Epic Games, Inc. All Rights Reserved.
+
+#include "WorldPartition/LevelInstance/LevelInstanceActorDescFactory.h"
+
+#if WITH_EDITOR
+#include "WorldPartition/LevelInstance/LevelInstanceActorDesc.h"
+#include "LevelInstance/LevelInstanceActor.h"
+#include "WorldPartition/WorldPartition.h"
+#include "ActorRegistry.h"
+#include "Editor.h"
+
+FWorldPartitionActorDesc* FLevelInstanceActorDescFactory::CreateInstance(const FWorldPartitionActorDescInitData& ActorDescInitData)
+{
+	FWorldPartitionActorDescData Data;
+	if (!ReadMetaData(ActorDescInitData, Data))
+	{
+		return nullptr;
+	}
+
+	FString LevelPackageNameStr;
+	static const FName NAME_LevelInstancePackage(TEXT("LevelInstancePackage"));
+	if (!FActorRegistry::ReadActorMetaData(NAME_LevelInstancePackage, LevelPackageNameStr, ActorDescInitData.AssetData))
+	{
+		return nullptr;
+	}
+
+	FTransform LevelInstanceTransform;
+	static const FName NAME_LevelInstanceTransform(TEXT("LevelInstanceTransform"));
+	if (!FActorRegistry::ReadActorMetaData(NAME_LevelInstanceTransform, LevelInstanceTransform, ActorDescInitData.AssetData))
+	{
+		return nullptr;
+	}
+	
+	FName LevelPackage = FName(*LevelPackageNameStr);
+
+	if (!LevelPackage.IsNone())
+	{
+		FBox LevelBounds;
+		if (ULevel::GetLevelBoundsFromPackage(LevelPackage, LevelBounds))
+		{
+			FVector LevelBoundsLocation;
+			LevelBounds.GetCenterAndExtents(Data.BoundsLocation, Data.BoundsExtent);
+
+			//@todo_ow: This will result in a new BoundsExtent that is larger than it should. To fix this, we would need the Object Oriented BoundingBox of the actor (the BV of the actor without rotation)
+			const FVector BoundsMin = Data.BoundsLocation - Data.BoundsExtent;
+			const FVector BoundsMax = Data.BoundsLocation + Data.BoundsExtent;
+			const FBox NewBounds = FBox(BoundsMin, BoundsMax).TransformBy(LevelInstanceTransform);
+			NewBounds.GetCenterAndExtents(Data.BoundsLocation, Data.BoundsExtent);
+		}
+	}
+
+	return new FLevelInstanceActorDesc(Data, LevelPackage);
+}
+
+FWorldPartitionActorDesc* FLevelInstanceActorDescFactory::CreateInstance(AActor* InActor)
+{
+	return new FLevelInstanceActorDesc(InActor);
+}
+#endif
\ No newline at end of file
diff --git a/Engine/Source/Runtime/Engine/Public/Foundation/FoundationActor.h b/Engine/Source/Runtime/Engine/Public/Foundation/FoundationActor.h
deleted file mode 100644
diff --git a/Engine/Source/Runtime/Engine/Public/Foundation/FoundationEditorInstanceActor.h b/Engine/Source/Runtime/Engine/Public/Foundation/FoundationEditorInstanceActor.h
deleted file mode 100644
diff --git a/Engine/Source/Runtime/Engine/Public/Foundation/FoundationEditorLevelStreaming.h b/Engine/Source/Runtime/Engine/Public/Foundation/FoundationEditorLevelStreaming.h
deleted file mode 100644
diff --git a/Engine/Source/Runtime/Engine/Public/Foundation/FoundationInstanceLevelStreaming.h b/Engine/Source/Runtime/Engine/Public/Foundation/FoundationInstanceLevelStreaming.h
deleted file mode 100644
diff --git a/Engine/Source/Runtime/Engine/Public/Foundation/FoundationSubsystem.h b/Engine/Source/Runtime/Engine/Public/Foundation/FoundationSubsystem.h
deleted file mode 100644
diff --git a/Engine/Source/Runtime/Engine/Public/Foundation/FoundationTypes.h b/Engine/Source/Runtime/Engine/Public/Foundation/FoundationTypes.h
deleted file mode 100644
diff --git a/Engine/Source/Runtime/Engine/Public/Foundation/IFoundationEditorModule.h b/Engine/Source/Runtime/Engine/Public/Foundation/IFoundationEditorModule.h
deleted file mode 100644
diff --git a/Engine/Source/Runtime/Engine/Public/LevelInstance/ILevelInstanceEditorModule.h b/Engine/Source/Runtime/Engine/Public/LevelInstance/ILevelInstanceEditorModule.h
new file mode 100644
--- /dev/null
+++ b/Engine/Source/Runtime/Engine/Public/LevelInstance/ILevelInstanceEditorModule.h
@@ -0,0 +1,14 @@
+// Copyright Epic Games, Inc. All Rights Reserved.
+#pragma once
+
+#include "CoreMinimal.h"
+#include "Modules/ModuleInterface.h"
+
+/**
+ * The module holding all of the UI related pieces for LevelInstance management
+ */
+class ILevelInstanceEditorModule : public IModuleInterface
+{
+public:
+	virtual ~ILevelInstanceEditorModule() {}
+};
\ No newline at end of file
diff --git a/Engine/Source/Runtime/Engine/Public/LevelInstance/LevelInstanceActor.h b/Engine/Source/Runtime/Engine/Public/LevelInstance/LevelInstanceActor.h
new file mode 100644
--- /dev/null
+++ b/Engine/Source/Runtime/Engine/Public/LevelInstance/LevelInstanceActor.h
@@ -0,0 +1,96 @@
+// Copyright Epic Games, Inc. All Rights Reserved.
+
+#pragma once
+
+#include "CoreMinimal.h"
+#include "UObject/ObjectMacros.h"
+#include "GameFramework/Actor.h"
+#include "Containers/Array.h"
+#include "LevelInstance/LevelInstanceTypes.h"
+#include "LevelInstanceActor.generated.h"
+
+UCLASS()
+class ENGINE_API ALevelInstance : public AActor
+{
+	GENERATED_UCLASS_BODY()
+
+protected:
+	/** Level LevelInstance */
+	UPROPERTY(EditAnywhere, Category = LevelInstance)
+	TSoftObjectPtr<UWorld> WorldAsset;
+public:
+
+	void LoadLevelInstance();
+	void UnloadLevelInstance();
+	bool IsLevelInstancePathValid() const;
+	const TSoftObjectPtr<UWorld>& GetWorldAsset() const { return WorldAsset; }
+	class ULevelInstanceSubsystem* GetLevelInstanceSubsystem() const;
+	
+	const FLevelInstanceID& GetLevelInstanceID() const;
+	bool HasValidLevelInstanceID() const;
+
+	virtual void PostRegisterAllComponents() override;
+	virtual void PostUnregisterAllComponents() override;
+	virtual void Serialize(FArchive& Ar) override;
+
+	const FGuid& GetLevelInstanceActorGuid() const;
+#if WITH_EDITOR
+	// UObject overrides
+	virtual void PreEditUndo() override;
+	virtual void PostEditUndo() override;
+	virtual void PostLoad() override;
+	virtual void GetAssetRegistryTags(TArray<FAssetRegistryTag>& OutTags) const override;
+	virtual bool CanEditChange(const FProperty* InProperty) const override;
+	virtual void PreEditChange(FProperty* PropertyThatWillChange) override;
+	virtual void PostEditChangeProperty(FPropertyChangedEvent& PropertyChangedEvent) override;
+	virtual void PostEditImport() override;
+
+	// AActor overrides
+	virtual void CheckForErrors() override;
+	virtual EActorGridPlacement GetDefaultGridPlacement() const override { return EActorGridPlacement::Bounds; }
+	virtual bool CanDeleteSelectedActor(FText& OutReason) const override;
+	virtual void SetIsTemporarilyHiddenInEditor(bool bIsHidden) override;
+	virtual void EditorGetUnderlyingActors(TSet<AActor*>& OutUnderlyingActors) override;
+	virtual void PushSelectionToProxies() override;
+	virtual FBox GetComponentsBoundingBox(bool bNonColliding = false, bool bIncludeFromChildActors = false) const override;
+	virtual void GetActorLocationBounds(bool bOnlyCollidingComponents, FVector& Origin, FVector& BoxExtent, bool bIncludeFromChildActors = false) const override;
+
+	bool CanEdit(FText* OutReason = nullptr) const;
+	bool CanCommit(FText* OutReason = nullptr) const;
+	bool IsEditing() const;
+	void Edit(AActor* ContextActor = nullptr);
+	void Commit();
+	void Discard();
+	void SaveAs();
+	bool HasDirtyChildren() const;
+	bool IsDirty() const;
+	bool SetCurrent();
+	bool IsCurrent() const;
+	bool IsLoaded() const;
+	void OnLevelInstanceLoaded();
+	FString GetWorldAssetPackage() const;
+	void UpdateLevelInstance();
+	bool SetWorldAsset(TSoftObjectPtr<UWorld> WorldAsset);
+	bool CheckForLoop(TSoftObjectPtr<UWorld> WorldAsset, TArray<TPair<FText, TSoftObjectPtr<UWorld>>>* LoopInfo = nullptr, const ALevelInstance** LoopStart = nullptr) const;
+	AActor* FindEditorInstanceActor() const;
+
+	DECLARE_MULTICAST_DELEGATE_OneParam(FOnLevelInstanceActorPostLoad, ALevelInstance*);
+	static FOnLevelInstanceActorPostLoad OnLevelInstanceActorPostLoad;
+#endif
+
+private:
+
+#if WITH_EDITOR
+	bool CanSetValue(TSoftObjectPtr<UWorld> LevelInstance, FString* Reason = nullptr) const;
+	
+	TSoftObjectPtr<UWorld> CachedWorldAsset;
+	FLevelInstanceID CachedLevelInstanceID;
+	bool bCachedIsTemporarilyHiddenInEditor;
+	bool bGuardLoadUnload;
+#else
+	// This Guid is used to compute the LevelInstanceID. Because in non-editor build we don't have an ActorGuid, we store it at cook time.
+	FGuid LevelInstanceActorGuid;
+#endif
+
+	FLevelInstanceID LevelInstanceID;
+};
\ No newline at end of file
diff --git a/Engine/Source/Runtime/Engine/Public/LevelInstance/LevelInstanceEditorInstanceActor.h b/Engine/Source/Runtime/Engine/Public/LevelInstance/LevelInstanceEditorInstanceActor.h
new file mode 100644
--- /dev/null
+++ b/Engine/Source/Runtime/Engine/Public/LevelInstance/LevelInstanceEditorInstanceActor.h
@@ -0,0 +1,38 @@
+// Copyright Epic Games, Inc. All Rights Reserved.
+
+#pragma once
+
+#include "CoreMinimal.h"
+#include "UObject/ObjectMacros.h"
+#include "GameFramework/Actor.h"
+#include "LevelInstance/LevelInstanceTypes.h"
+#include "LevelInstanceEditorInstanceActor.generated.h"
+
+class ALevelInstance;
+class ULevel;
+
+/**
+ * Editor Only Actor that is spawned inside every LevelInstance Instance Level so that we can update its Actor Transforms through the ALevelInstance's root component(ULevelInstanceComponent)
+ * @see ULevelInstanceComponent
+ */
+UCLASS(transient, notplaceable)
+class ENGINE_API ALevelInstanceEditorInstanceActor : public AActor
+{
+	GENERATED_UCLASS_BODY()
+
+#if WITH_EDITOR
+public:
+	static ALevelInstanceEditorInstanceActor* Create(ALevelInstance* LevelInstanceActor, ULevel* LoadedLevel);
+	
+	void SetLevelInstanceID(const FLevelInstanceID& InLevelInstanceID) { LevelInstanceID = InLevelInstanceID; }
+	const FLevelInstanceID& GetLevelInstanceID() const { return LevelInstanceID; }
+	
+	virtual bool IsSelectionParentOfAttachedActors() const override { return true; }
+	virtual bool IsSelectionChild() const override { return true; }
+	virtual AActor* GetSelectionParent() const override;
+
+private:
+
+	FLevelInstanceID LevelInstanceID;
+#endif
+};
\ No newline at end of file
diff --git a/Engine/Source/Runtime/Engine/Public/LevelInstance/LevelInstanceEditorLevelStreaming.h b/Engine/Source/Runtime/Engine/Public/LevelInstance/LevelInstanceEditorLevelStreaming.h
new file mode 100644
--- /dev/null
+++ b/Engine/Source/Runtime/Engine/Public/LevelInstance/LevelInstanceEditorLevelStreaming.h
@@ -0,0 +1,31 @@
+// Copyright Epic Games, Inc. All Rights Reserved.
+
+#pragma once
+
+#include "CoreMinimal.h"
+#include "UObject/ObjectMacros.h"
+#include "Engine/LevelStreamingAlwaysLoaded.h"
+#include "LevelInstance/LevelInstanceTypes.h"
+#include "LevelInstanceEditorLevelStreaming.generated.h"
+
+class ALevelInstance;
+
+UCLASS(Transient, MinimalAPI)
+class ULevelStreamingLevelInstanceEditor : public ULevelStreamingAlwaysLoaded
+{
+	GENERATED_UCLASS_BODY()
+
+public:
+#if WITH_EDITOR
+	virtual bool ShowInLevelCollection() const override { return false; }
+	ALevelInstance* GetLevelInstanceActor() const;
+
+protected:
+	friend class ULevelInstanceSubsystem;
+
+	static ULevelStreamingLevelInstanceEditor* Load(ALevelInstance* LevelInstanceActor);
+	static void Unload(ULevelStreamingLevelInstanceEditor* LevelStreaming);
+private:
+	FLevelInstanceID LevelInstanceID;
+#endif
+};
\ No newline at end of file
diff --git a/Engine/Source/Runtime/Engine/Public/LevelInstance/LevelInstanceLevelStreaming.h b/Engine/Source/Runtime/Engine/Public/LevelInstance/LevelInstanceLevelStreaming.h
new file mode 100644
--- /dev/null
+++ b/Engine/Source/Runtime/Engine/Public/LevelInstance/LevelInstanceLevelStreaming.h
@@ -0,0 +1,34 @@
+// Copyright Epic Games, Inc. All Rights Reserved.
+
+#pragma once
+
+#include "CoreMinimal.h"
+#include "UObject/ObjectMacros.h"
+#include "Engine/LevelStreamingDynamic.h"
+#include "LevelInstance/LevelInstanceTypes.h"
+#include "LevelInstanceLevelStreaming.generated.h"
+
+class ALevelInstance;
+
+UCLASS(Transient)
+class ENGINE_API ULevelStreamingLevelInstance : public ULevelStreamingDynamic
+{
+	GENERATED_UCLASS_BODY()
+
+public:
+	ALevelInstance* GetLevelInstanceActor() const;
+
+#if WITH_EDITOR
+	virtual bool ShowInLevelCollection() const override { return false; }
+	FBox GetBounds() const;
+#endif
+	
+protected:
+	static ULevelStreamingLevelInstance* LoadInstance(ALevelInstance* LevelInstanceActor);
+	static void UnloadInstance(ULevelStreamingLevelInstance* LevelStreaming);
+
+	friend class ULevelInstanceSubsystem;
+
+private:
+	FLevelInstanceID LevelInstanceID;
+};
\ No newline at end of file
diff --git a/Engine/Source/Runtime/Engine/Public/LevelInstance/LevelInstanceSubsystem.h b/Engine/Source/Runtime/Engine/Public/LevelInstance/LevelInstanceSubsystem.h
new file mode 100644
--- /dev/null
+++ b/Engine/Source/Runtime/Engine/Public/LevelInstance/LevelInstanceSubsystem.h
@@ -0,0 +1,152 @@
+// Copyright Epic Games, Inc. All Rights Reserved.
+#pragma once
+
+#include "Subsystems/WorldSubsystem.h"
+#include "Containers/Array.h"
+#include "Containers/Map.h"
+#include "UObject/NameTypes.h"
+#include "LevelInstance/LevelInstanceTypes.h"
+#include "Tickable.h"
+
+#if WITH_EDITOR
+#include "WorldPartition/LevelInstance/LevelInstanceActorDescFactory.h"
+#endif
+
+#include "LevelInstanceSubsystem.generated.h"
+
+class ALevelInstance;
+class ULevelStreamingLevelInstance;
+class ULevelStreamingLevelInstanceEditor;
+class UWorldPartitionSubsystem;
+
+/**
+ * ULevelInstanceSubsystem
+ */
+UCLASS()
+class ENGINE_API ULevelInstanceSubsystem : public UWorldSubsystem, public FTickableGameObject
+{
+	GENERATED_BODY()
+
+public:
+	ULevelInstanceSubsystem();
+	virtual ~ULevelInstanceSubsystem();
+		
+	//~ Begin USubsystem Interface.
+	virtual void Initialize(FSubsystemCollectionBase& Collection) override;
+	virtual void Deinitialize() override;
+	//~ End USubsystem Interface.
+	
+	//~ Begin UWorldSubsystem Interface.
+	virtual void UpdateStreamingState() override;
+	//~ End UWorldSubsystem Interface.
+
+	//~ Begin FTickableGameObject
+	virtual void Tick(float DeltaSeconds) override;
+	virtual bool IsTickableInEditor() const override;
+	virtual UWorld* GetTickableGameObjectWorld() const override;
+	virtual ETickableTickType GetTickableTickType() const override;
+	virtual TStatId GetStatId() const override;
+	//~End FTickableGameObject
+
+	ALevelInstance* GetLevelInstance(FLevelInstanceID LevelInstanceID) const;
+	FLevelInstanceID RegisterLevelInstance(ALevelInstance* LevelInstanceActor);
+	void UnregisterLevelInstance(ALevelInstance* LevelInstanceActor);
+	void RequestLoadLevelInstance(ALevelInstance* LevelInstanceActor, bool bUpdate);
+	void RequestUnloadLevelInstance(ALevelInstance* LevelInstanceActor);
+	bool IsLoaded(const ALevelInstance* LevelInstanceActor) const;
+	void ForEachLevelInstanceAncestorsAndSelf(AActor* Actor, TFunctionRef<bool(ALevelInstance*)> Operation) const;
+
+#if WITH_EDITOR
+	bool CanEditLevelInstance(const ALevelInstance* LevelInstanceActor, FText* OutReason = nullptr) const;
+	bool CanCommitLevelInstance(const ALevelInstance* LevelInstanceActor, FText* OutReason = nullptr) const;
+	void EditLevelInstance(ALevelInstance* LevelInstanceActor, TWeakObjectPtr<AActor> ContextActorPtr = nullptr);
+	void CommitLevelInstance(ALevelInstance* LevelInstanceActor, bool bDiscardEdits = false);
+	void SaveLevelInstanceAs(ALevelInstance* LevelInstanceActor);
+	bool IsEditingLevelInstanceDirty(const ALevelInstance* LevelInstanceActor) const;
+	bool IsEditingLevelInstance(const ALevelInstance* LevelInstanceActor) const { return GetLevelInstanceEdit(LevelInstanceActor) != nullptr; }
+	
+	FLevelInstanceActorDescFactory* GetActorDescFactory() const { return LevelInstanceActorDescFactory.Get(); }
+	bool GetLevelInstanceBounds(const ALevelInstance* LevelInstanceActor, FBox& OutBounds) const;
+	
+	void ForEachActorInLevelInstance(const ALevelInstance* LevelInstanceActor, TFunctionRef<bool(AActor * LevelActor)> Operation) const;
+	void ForEachLevelInstanceAncestorsAndSelf(const AActor* Actor, TFunctionRef<bool(const ALevelInstance*)> Operation) const;
+	void ForEachLevelInstanceAncestors(const AActor* Actor, TFunctionRef<bool(const ALevelInstance*)> Operation) const;
+	void ForEachLevelInstanceChildren(const ALevelInstance* LevelInstanceActor, bool bRecursive, TFunctionRef<bool(const ALevelInstance*)> Operation) const;
+	void ForEachLevelInstanceChildren(ALevelInstance* LevelInstanceActor, bool bRecursive, TFunctionRef<bool(ALevelInstance*)> Operation) const;
+	void ForEachLevelInstanceEdit(TFunctionRef<bool(ALevelInstance*)> Operation) const;
+	bool HasDirtyChildrenLevelInstances(const ALevelInstance* LevelInstanceActor) const;
+
+	void SetIsTemporarilyHiddenInEditor(ALevelInstance* LevelInstanceActor, bool bIsHidden);
+
+	bool SetCurrent(ALevelInstance* LevelInstanceActor) const;
+	bool IsCurrent(const ALevelInstance* LevelInstanceActor) const;
+	ALevelInstance* CreateLevelInstanceFrom(const TArray<AActor*>& ActorsToMove, UWorld* TemplateWorld = nullptr);
+	bool MoveActorsToLevel(const TArray<AActor*>& ActorsToRemove, ULevel* DestinationLevel) const;
+	bool MoveActorsTo(ALevelInstance* LevelInstanceActor, const TArray<AActor*>& ActorsToMove);
+	bool BreakLevelInstance(ALevelInstance* LevelInstanceActor, uint32 Levels = 1);
+
+	bool CanMoveActorToLevel(const AActor* Actor) const;
+	void DiscardEdits();
+	void OnActorDeleted(AActor* Actor);
+	ULevel* GetLevelInstanceLevel(const ALevelInstance* LevelInstanceActor) const;
+
+	ALevelInstance* GetParentLevelInstance(const AActor* Actor) const;
+#endif
+
+private:
+
+	void LoadLevelInstance(ALevelInstance* LevelInstanceActor);
+	void UnloadLevelInstance(const FLevelInstanceID& LevelInstanceID);
+	void ForEachActorInLevel(ULevel* Level, TFunctionRef<bool(AActor * LevelActor)> Operation) const;
+	void ForEachLevelInstanceAncestors(AActor* Actor, TFunctionRef<bool(ALevelInstance*)> Operation) const;
+	ALevelInstance* GetOwningLevelInstance(const ULevel* Level) const;
+	FLevelInstanceID ComputeLevelInstanceID(ALevelInstance* LevelInstanceActor) const;
+#if WITH_EDITOR
+	void CommitChildrenLevelInstances(ALevelInstance* LevelInstanceActor);
+
+	void RegisterActorDescFactories(UWorldPartitionSubsystem* WorldPartitionSubsystem);
+	static bool ShouldIgnoreDirtyPackage(UPackage* DirtyPackage, const UWorld* EditingWorld);
+
+	struct FLevelInstanceEdit
+	{
+		ULevelStreamingLevelInstanceEditor* LevelStreaming = nullptr;
+
+		UWorld* GetEditWorld() const;
+	};
+
+	const FLevelInstanceEdit* GetLevelInstanceEdit(const ALevelInstance* LevelInstanceActor) const;
+	bool IsLevelInstanceEditDirty(const FLevelInstanceEdit* LevelInstanceEdit) const;
+	
+	struct FLevelsToRemoveScope
+	{
+		~FLevelsToRemoveScope();
+
+		TArray<ULevel*> Levels;
+		bool bResetTrans = false;
+	};
+	
+	friend ULevelStreamingLevelInstance;
+	void RemoveLevelFromWorld(ULevel* Level, bool bResetTrans);
+#endif
+
+
+#if WITH_EDITORONLY_DATA
+	// Optional scope to accelerate level unload by batching them
+	TUniquePtr<FLevelsToRemoveScope> LevelsToRemoveScope;
+	
+	TMap<FName, FLevelInstanceEdit> LevelInstanceEdits;
+
+	TUniquePtr<FLevelInstanceActorDescFactory> LevelInstanceActorDescFactory;
+#endif
+
+	struct FLevelInstance
+	{
+		ULevelStreamingLevelInstance* LevelStreaming = nullptr;
+	};
+
+	TMap<ALevelInstance*, bool> LevelInstancesToLoadOrUpdate;
+	TSet<FLevelInstanceID> LevelInstancesToUnload;
+	TMap<FLevelInstanceID, FLevelInstance> LevelInstances;
+	TMap<FLevelInstanceID, ALevelInstance*> RegisteredLevelInstances;
+	FLevelInstanceID PendingLevelInstanceToEdit = InvalidLevelInstanceID;
+};
\ No newline at end of file
diff --git a/Engine/Source/Runtime/Engine/Public/LevelInstance/LevelInstanceTypes.h b/Engine/Source/Runtime/Engine/Public/LevelInstance/LevelInstanceTypes.h
new file mode 100644
--- /dev/null
+++ b/Engine/Source/Runtime/Engine/Public/LevelInstance/LevelInstanceTypes.h
@@ -0,0 +1,9 @@
+// Copyright Epic Games, Inc. All Rights Reserved.
+
+#pragma once
+
+// FLevelInstanceID is a runtime unique id that is computed from the Hash of LevelInstance Actor Guid and all its ancestor LevelInstance Actor Guids.
+// Resulting in a different ID for all instances whether they load the same level or not.
+using FLevelInstanceID = uint32;
+
+static constexpr FLevelInstanceID InvalidLevelInstanceID = 0;
\ No newline at end of file
diff --git a/Engine/Source/Runtime/Engine/Public/WorldPartition/Foundation/FoundationActorDesc.h b/Engine/Source/Runtime/Engine/Public/WorldPartition/Foundation/FoundationActorDesc.h
deleted file mode 100644
diff --git a/Engine/Source/Runtime/Engine/Public/WorldPartition/Foundation/FoundationActorDescFactory.h b/Engine/Source/Runtime/Engine/Public/WorldPartition/Foundation/FoundationActorDescFactory.h
deleted file mode 100644
diff --git a/Engine/Source/Runtime/Engine/Public/WorldPartition/LevelInstance/LevelInstanceActorDesc.h b/Engine/Source/Runtime/Engine/Public/WorldPartition/LevelInstance/LevelInstanceActorDesc.h
new file mode 100644
--- /dev/null
+++ b/Engine/Source/Runtime/Engine/Public/WorldPartition/LevelInstance/LevelInstanceActorDesc.h
@@ -0,0 +1,33 @@
+// Copyright Epic Games, Inc. All Rights Reserved.
+
+#pragma once
+
+#include "CoreMinimal.h"
+#include "UObject/ObjectMacros.h"
+#include "Containers/Map.h"
+#include "WorldPartition/WorldPartitionActorDesc.h"
+
+class ULevelInstanceSubsystem;
+class UWorldPartition;
+
+/**
+ * ActorDesc for Actors that are part of a LevelInstanceActor Level.
+ */
+class ENGINE_API FLevelInstanceActorDesc : public FWorldPartitionActorDesc
+{
+public:
+#if WITH_EDITOR
+	inline FName GetLevelPackage() const { return LevelPackage; }
+
+protected:
+	FLevelInstanceActorDesc() = delete;
+	FLevelInstanceActorDesc(const FWorldPartitionActorDescData& DescData, FName LevelPackage);
+	FLevelInstanceActorDesc(AActor* InActor);
+
+	virtual void BuildHash(FHashBuilder& HashBuilder) override;
+
+	friend class FLevelInstanceActorDescFactory;
+
+	FName LevelPackage;
+#endif
+};
diff --git a/Engine/Source/Runtime/Engine/Public/WorldPartition/LevelInstance/LevelInstanceActorDescFactory.h b/Engine/Source/Runtime/Engine/Public/WorldPartition/LevelInstance/LevelInstanceActorDescFactory.h
new file mode 100644
--- /dev/null
+++ b/Engine/Source/Runtime/Engine/Public/WorldPartition/LevelInstance/LevelInstanceActorDescFactory.h
@@ -0,0 +1,18 @@
+// Copyright Epic Games, Inc. All Rights Reserved.
+#pragma once
+
+#include "CoreMinimal.h"
+#include "UObject/ObjectMacros.h"
+#include "WorldPartition/WorldPartitionActorDescFactory.h"
+
+class FWorldPartitionActorDesc;
+class ULevelInstanceSubsystem;
+
+class ENGINE_API FLevelInstanceActorDescFactory : public FWorldPartitionActorDescFactory
+{
+public:
+#if WITH_EDITOR
+	virtual FWorldPartitionActorDesc* CreateInstance(const FWorldPartitionActorDescInitData& ActorDescInitData) override;
+	virtual FWorldPartitionActorDesc* CreateInstance(AActor* InActor) override;
+#endif
+};
\ No newline at end of file
