Optimize FLevelUtils::FindStreamingLevel

LevelInstances makes heavy use of StreamingLevel. FindStreamingLevel has been showing up in captures of LevelInstance test levels.

diff --git a/Engine/Source/Editor/UnrealEd/Private/EditorLevelUtils.cpp b/Engine/Source/Editor/UnrealEd/Private/EditorLevelUtils.cpp
--- a/Engine/Source/Editor/UnrealEd/Private/EditorLevelUtils.cpp
+++ b/Engine/Source/Editor/UnrealEd/Private/EditorLevelUtils.cpp
@@ -838,22 +838,9 @@ bool UEditorLevelUtils::RemoveLevelFromWorld(ULevel* InLevel, bool bClearSelecti
 
 void UEditorLevelUtils::PrivateRemoveLevelFromWorld(ULevel* InLevel)
 {
-	int32 StreamingLevelIndex = INDEX_NONE;
-
-	for (int32 LevelIndex = 0; LevelIndex < InLevel->OwningWorld->GetStreamingLevels().Num(); ++LevelIndex)
-	{
-		ULevelStreaming* StreamingLevel = InLevel->OwningWorld->GetStreamingLevels()[LevelIndex];
-		if (StreamingLevel && StreamingLevel->GetLoadedLevel() == InLevel)
-		{
-			StreamingLevelIndex = LevelIndex;
-			break;
-		}
-	}
-
 	bool bIsTransientLevelStreaming = false;
-	if (StreamingLevelIndex != INDEX_NONE)
+	if (ULevelStreaming* StreamingLevel = ULevelStreaming::FindStreamingLevel(InLevel))
 	{
-		ULevelStreaming* StreamingLevel = InLevel->OwningWorld->GetStreamingLevels()[StreamingLevelIndex];
 		bIsTransientLevelStreaming = StreamingLevel->HasAnyFlags(RF_Transient);
 		StreamingLevel->MarkPendingKill();
 		InLevel->OwningWorld->RemoveStreamingLevel(StreamingLevel);
diff --git a/Engine/Source/Runtime/Engine/Classes/Engine/LevelStreaming.h b/Engine/Source/Runtime/Engine/Classes/Engine/LevelStreaming.h
--- a/Engine/Source/Runtime/Engine/Classes/Engine/LevelStreaming.h
+++ b/Engine/Source/Runtime/Engine/Classes/Engine/LevelStreaming.h
@@ -111,6 +111,25 @@ private:
 	static const TCHAR* EnumToString(ETargetState InTargetState);
 
 public:
+	static ULevelStreaming* FindStreamingLevel(const ULevel* Level);
+	static void RemoveLevelAnnotation(const ULevel* Level);
+
+	// Annotation for fast inverse lookup
+	struct FLevelAnnotation
+	{
+		FLevelAnnotation() {}
+		FLevelAnnotation(ULevelStreaming* InLevelStreaming) : LevelStreaming(InLevelStreaming){}
+	
+		FORCEINLINE bool IsDefault()
+		{
+			return !LevelStreaming;
+		}
+
+		ULevelStreaming* LevelStreaming = nullptr;
+	};
+
+	static FUObjectAnnotationSparse<FLevelAnnotation, false> LevelAnnotations;
+
 
 #if WITH_EDITORONLY_DATA
 	/** Deprecated name of the package containing the level to load. Use GetWorldAsset() or GetWorldAssetPackageFName() instead.		*/
diff --git a/Engine/Source/Runtime/Engine/Private/Level.cpp b/Engine/Source/Runtime/Engine/Private/Level.cpp
--- a/Engine/Source/Runtime/Engine/Private/Level.cpp
+++ b/Engine/Source/Runtime/Engine/Private/Level.cpp
@@ -932,6 +932,8 @@ void ULevel::ClearLevelComponents()
 
 void ULevel::BeginDestroy()
 {
+	ULevelStreaming::RemoveLevelAnnotation(this);
+	
 	if (!IStreamingManager::HasShutdown())
 	{
 		// At this time, referenced UTexture2Ds are still in memory.
@@ -1618,13 +1620,9 @@ void ULevel::PostEditUndo()
 		}
 		else
 		{
-			for (const ULevelStreaming* StreamedLevel : OwningWorld->GetStreamingLevels())
+			if(const ULevelStreaming* StreamingLevel = ULevelStreaming::FindStreamingLevel(this))
 			{
-				if (StreamedLevel && StreamedLevel->GetLoadedLevel() == this)
-				{
-					bIsStreamingLevelVisible = FLevelUtils::IsStreamingLevelVisibleInEditor(StreamedLevel);
-					break;
-				}
+				bIsStreamingLevelVisible = FLevelUtils::IsStreamingLevelVisibleInEditor(StreamingLevel);
 			}
 		}
 
diff --git a/Engine/Source/Runtime/Engine/Private/LevelStreaming.cpp b/Engine/Source/Runtime/Engine/Private/LevelStreaming.cpp
--- a/Engine/Source/Runtime/Engine/Private/LevelStreaming.cpp
+++ b/Engine/Source/Runtime/Engine/Private/LevelStreaming.cpp
@@ -792,6 +792,42 @@ void ULevelStreaming::AddLevelToCollectionAfterReload()
 }
 #endif
 
+FUObjectAnnotationSparse<ULevelStreaming::FLevelAnnotation, false> ULevelStreaming::LevelAnnotations;
+
+void ULevelStreaming::RemoveLevelAnnotation(const ULevel* Level)
+{
+	ULevelStreaming::LevelAnnotations.RemoveAnnotation(Level);
+}
+
+ULevelStreaming* ULevelStreaming::FindStreamingLevel(const ULevel* Level)
+{
+	ULevelStreaming* FoundLevelStreaming = nullptr;
+	if (Level && Level->OwningWorld && !Level->IsPersistentLevel())
+	{
+		FLevelAnnotation LevelAnnotation = LevelAnnotations.GetAnnotation(Level);
+		if (LevelAnnotation.LevelStreaming)
+		{
+			check(LevelAnnotation.LevelStreaming->GetLoadedLevel() == Level);
+			FoundLevelStreaming = LevelAnnotation.LevelStreaming;
+		}
+
+		if (!ensure(FoundLevelStreaming))
+		{
+			// fallback search
+			for (ULevelStreaming* CurStreamingLevel : Level->OwningWorld->GetStreamingLevels())
+			{
+				if (CurStreamingLevel && CurStreamingLevel->GetLoadedLevel() == Level)
+				{
+					FoundLevelStreaming = CurStreamingLevel;
+					break;
+				}
+			}
+		}
+	}
+
+	return FoundLevelStreaming;
+}
+
 void ULevelStreaming::SetLoadedLevel(ULevel* Level)
 { 
 	// Pending level should be unloaded at this point
@@ -811,9 +847,15 @@ void ULevelStreaming::SetLoadedLevel(ULevel* Level)
 	FLevelCollection& LC = World->FindOrAddCollectionByType(CollectionType);
 	LC.RemoveLevel(PendingUnloadLevel);
 
+	if (PendingUnloadLevel)
+	{
+		RemoveLevelAnnotation(PendingUnloadLevel);
+	}
+
 	if (LoadedLevel)
 	{
 		LoadedLevel->OwningWorld = World;
+		ULevelStreaming::LevelAnnotations.AddAnnotation(LoadedLevel, FLevelAnnotation(this));
 
 		// Remove the loaded level from its current collection, if any.
 		if (LoadedLevel->GetCachedLevelCollection())
diff --git a/Engine/Source/Runtime/Engine/Private/LevelUtils.cpp b/Engine/Source/Runtime/Engine/Private/LevelUtils.cpp
--- a/Engine/Source/Runtime/Engine/Private/LevelUtils.cpp
+++ b/Engine/Source/Runtime/Engine/Private/LevelUtils.cpp
@@ -58,21 +58,7 @@ bool FLevelUtils::bApplyingLevelTransform = false;
  */
 ULevelStreaming* FLevelUtils::FindStreamingLevel(const ULevel* Level)
 {
-	ULevelStreaming* MatchingLevel = NULL;
-
-	if (Level && Level->OwningWorld)
-	{
-		for (ULevelStreaming* CurStreamingLevel : Level->OwningWorld->GetStreamingLevels())
-		{
-			if( CurStreamingLevel && CurStreamingLevel->GetLoadedLevel() == Level )
-			{
-				MatchingLevel = CurStreamingLevel;
-				break;
-			}
-		}
-	}
-
-	return MatchingLevel;
+	return ULevelStreaming::FindStreamingLevel(Level);
 }
 
 /**
