Fix missing number of intersecting cells counter in WorldPartitionEditorSpatialHash.

diff --git a/Engine/Source/Runtime/Engine/Private/WorldPartition/WorldPartitionEditorSpatialHash.cpp b/Engine/Source/Runtime/Engine/Private/WorldPartition/WorldPartitionEditorSpatialHash.cpp
--- a/Engine/Source/Runtime/Engine/Private/WorldPartition/WorldPartitionEditorSpatialHash.cpp
+++ b/Engine/Source/Runtime/Engine/Private/WorldPartition/WorldPartitionEditorSpatialHash.cpp
@@ -243,7 +243,7 @@ int32 UWorldPartitionEditorSpatialHash::ForEachIntersectingCellInner(const FBox&
 
 				if (Box.Intersect(CellBounds))
 				{
-					ForEachIntersectingCellInner(Box, ChildCellCoord, InOperation);
+					NumIntersecting += ForEachIntersectingCellInner(Box, ChildCellCoord, InOperation);
 				}
 			});
 		}
