Fix Mac

diff --git a/Engine/Source/Runtime/Engine/Engine.Build.cs b/Engine/Source/Runtime/Engine/Engine.Build.cs
--- a/Engine/Source/Runtime/Engine/Engine.Build.cs
+++ b/Engine/Source/Runtime/Engine/Engine.Build.cs
@@ -395,7 +395,7 @@ public class Engine : ModuleRules
 
 		if (Target.Platform == UnrealTargetPlatform.Mac)
 		{
-			PublicFrameworks.AddRange(new string[] { "AVLevelInstance", "CoreVideo", "CoreMedia" });
+			PublicFrameworks.AddRange(new string[] { "AVFoundation", "CoreVideo", "CoreMedia" });
 		}
 
 		if (Target.IsInPlatformGroup(UnrealPlatformGroup.Android))
