Fix crash caused by LevelIinstanceSubsystem being ticked inside the TickTaskManager BeginFrame/EndFrame and deleting FTickTaskLevels objects

diff --git a/Engine/Source/Runtime/Engine/Private/LevelInstance/LevelInstanceSubsystem.cpp b/Engine/Source/Runtime/Engine/Private/LevelInstance/LevelInstanceSubsystem.cpp
--- a/Engine/Source/Runtime/Engine/Private/LevelInstance/LevelInstanceSubsystem.cpp
+++ b/Engine/Source/Runtime/Engine/Private/LevelInstance/LevelInstanceSubsystem.cpp
@@ -158,50 +158,6 @@ bool ULevelInstanceSubsystem::IsLoaded(const ALevelInstance* LevelInstanceActor)
 	return LevelInstanceActor->HasValidLevelInstanceID() && LevelInstances.Contains(LevelInstanceActor->GetLevelInstanceID());
 }
 
-void ULevelInstanceSubsystem::Tick(float DeltaSeconds)
-{
-#if WITH_EDITOR
-	// For non-game world, Tick is responsible of processing LevelInstances to update/load/unload
-	if (!GetWorld()->IsGameWorld())
-	{
-		UpdateStreamingState();
-
-		// Begin editing the pending LevelInstance when loads complete
-		if (PendingLevelInstanceToEdit != InvalidLevelInstanceID && !LevelInstancesToLoadOrUpdate.Num())
-		{
-			if (ALevelInstance** LevelInstanceActor = RegisteredLevelInstances.Find(PendingLevelInstanceToEdit))
-			{
-				EditLevelInstance(*LevelInstanceActor);
-			}
-		}
-	}
-#endif
-}
-
-bool ULevelInstanceSubsystem::IsTickableInEditor() const
-{
-	return true;
-}
-
-UWorld* ULevelInstanceSubsystem::GetTickableGameObjectWorld() const
-{
-	return GetWorld();
-}
-
-ETickableTickType ULevelInstanceSubsystem::GetTickableTickType() const
-{
-#if WITH_EDITOR
-	return IsTemplate() ? ETickableTickType::Never : ETickableTickType::Always;
-#else
-	return ETickableTickType::Never;
-#endif
-}
-
-TStatId ULevelInstanceSubsystem::GetStatId() const
-{
-	RETURN_QUICK_DECLARE_CYCLE_STAT(ULevelInstanceSubsystem, STATGROUP_Tickables);
-}
-
 void ULevelInstanceSubsystem::UpdateStreamingState()
 {
 	if (!LevelInstancesToUnload.Num() && !LevelInstancesToLoadOrUpdate.Num())
@@ -396,6 +352,24 @@ ALevelInstance* ULevelInstanceSubsystem::GetOwningLevelInstance(const ULevel* Le
 
 #if WITH_EDITOR
 
+void ULevelInstanceSubsystem::Tick()
+{
+	// For non-game world, Tick is responsible of processing LevelInstances to update/load/unload
+	if (!GetWorld()->IsGameWorld())
+	{
+		UpdateStreamingState();
+
+		// Begin editing the pending LevelInstance when loads complete
+		if (PendingLevelInstanceToEdit != InvalidLevelInstanceID && !LevelInstancesToLoadOrUpdate.Num())
+		{
+			if (ALevelInstance** LevelInstanceActor = RegisteredLevelInstances.Find(PendingLevelInstanceToEdit))
+			{
+				EditLevelInstance(*LevelInstanceActor);
+			}
+		}
+	}
+}
+
 bool ULevelInstanceSubsystem::GetLevelInstanceBounds(const ALevelInstance* LevelInstanceActor, FBox& OutBounds) const
 {
 	if (IsLoaded(LevelInstanceActor))
diff --git a/Engine/Source/Runtime/Engine/Private/LevelTick.cpp b/Engine/Source/Runtime/Engine/Private/LevelTick.cpp
--- a/Engine/Source/Runtime/Engine/Private/LevelTick.cpp
+++ b/Engine/Source/Runtime/Engine/Private/LevelTick.cpp
@@ -62,6 +62,7 @@
 
 #if WITH_EDITOR
 #include "Editor.h"
+#include "LevelInstance/LevelInstanceSubsystem.h"
 #endif
 
 CSV_DECLARE_CATEGORY_MODULE_EXTERN(CORE_API, Basic);
@@ -1529,6 +1530,14 @@ void UWorld::Tick( ELevelTick TickType, float DeltaSeconds )
 		}
 	}
 
+#if WITH_EDITOR
+	// Tick LevelInstanceSubsystem outside of FTickTaskManagerInterface::StartFrame/EndFrame because it can cause levels to be deleted and invalidate its LevelList
+	if (ULevelInstanceSubsystem* LevelInstanceSubsystem = GetSubsystem<ULevelInstanceSubsystem>())
+	{
+		LevelInstanceSubsystem->Tick();
+	}
+#endif
+
 	if (bDoingActorTicks)
 	{
 		SCOPE_CYCLE_COUNTER(STAT_TickTime);
diff --git a/Engine/Source/Runtime/Engine/Private/TickTaskManager.cpp b/Engine/Source/Runtime/Engine/Private/TickTaskManager.cpp
--- a/Engine/Source/Runtime/Engine/Private/TickTaskManager.cpp
+++ b/Engine/Source/Runtime/Engine/Private/TickTaskManager.cpp
@@ -1442,6 +1442,7 @@ public:
 	/** Free a ticking structure for a ULevel **/
 	virtual void FreeTickTaskLevel(FTickTaskLevel* TickTaskLevel) override
 	{
+		check(!LevelList.Contains(TickTaskLevel));
 		delete TickTaskLevel;
 	}
 
diff --git a/Engine/Source/Runtime/Engine/Public/LevelInstance/LevelInstanceSubsystem.h b/Engine/Source/Runtime/Engine/Public/LevelInstance/LevelInstanceSubsystem.h
--- a/Engine/Source/Runtime/Engine/Public/LevelInstance/LevelInstanceSubsystem.h
+++ b/Engine/Source/Runtime/Engine/Public/LevelInstance/LevelInstanceSubsystem.h
@@ -6,7 +6,6 @@
 #include "Containers/Map.h"
 #include "UObject/NameTypes.h"
 #include "LevelInstance/LevelInstanceTypes.h"
-#include "Tickable.h"
 
 #if WITH_EDITOR
 #include "WorldPartition/LevelInstance/LevelInstanceActorDescFactory.h"
@@ -23,7 +22,7 @@ class UWorldPartitionSubsystem;
  * ULevelInstanceSubsystem
  */
 UCLASS()
-class ENGINE_API ULevelInstanceSubsystem : public UWorldSubsystem, public FTickableGameObject
+class ENGINE_API ULevelInstanceSubsystem : public UWorldSubsystem
 {
 	GENERATED_BODY()
 
@@ -40,14 +39,6 @@ public:
 	virtual void UpdateStreamingState() override;
 	//~ End UWorldSubsystem Interface.
 
-	//~ Begin FTickableGameObject
-	virtual void Tick(float DeltaSeconds) override;
-	virtual bool IsTickableInEditor() const override;
-	virtual UWorld* GetTickableGameObjectWorld() const override;
-	virtual ETickableTickType GetTickableTickType() const override;
-	virtual TStatId GetStatId() const override;
-	//~End FTickableGameObject
-
 	ALevelInstance* GetLevelInstance(FLevelInstanceID LevelInstanceID) const;
 	FLevelInstanceID RegisterLevelInstance(ALevelInstance* LevelInstanceActor);
 	void UnregisterLevelInstance(ALevelInstance* LevelInstanceActor);
@@ -57,6 +48,8 @@ public:
 	void ForEachLevelInstanceAncestorsAndSelf(AActor* Actor, TFunctionRef<bool(ALevelInstance*)> Operation) const;
 
 #if WITH_EDITOR
+	void Tick();
+
 	bool CanEditLevelInstance(const ALevelInstance* LevelInstanceActor, FText* OutReason = nullptr) const;
 	bool CanCommitLevelInstance(const ALevelInstance* LevelInstanceActor, FText* OutReason = nullptr) const;
 	void EditLevelInstance(ALevelInstance* LevelInstanceActor, TWeakObjectPtr<AActor> ContextActorPtr = nullptr);
