World Partition
- Fixed crash when closing the world partition edtitor while having an active selection: selection is now a member of the editor.
- Started removing external references to editor cells classes outside of the world partition core.
- Editor hash is now a completly dynamic, hierarchical 3d spatial grid that adapts to the actors.

diff --git a/Engine/Source/Editor/UnrealEd/Private/Commandlets/WorldPartitionBuildHLODsCommandlet.cpp b/Engine/Source/Editor/UnrealEd/Private/Commandlets/WorldPartitionBuildHLODsCommandlet.cpp
--- a/Engine/Source/Editor/UnrealEd/Private/Commandlets/WorldPartitionBuildHLODsCommandlet.cpp
+++ b/Engine/Source/Editor/UnrealEd/Private/Commandlets/WorldPartitionBuildHLODsCommandlet.cpp
@@ -11,7 +11,6 @@
 #include "Editor.h"
 #include "Logging/LogMacros.h"
 
-#include "WorldPartition/WorldPartitionEditorHash.h"
 #include "WorldPartition/WorldPartition.h"
 #include "WorldPartition/WorldPartitionSubsystem.h"
 #include "WorldPartition/HLOD/HLODActor.h"
@@ -94,10 +93,8 @@ int32 UWorldPartitionBuildHLODsCommandlet::Main(const FString& Params)
 	// For now, load all cells
 	// In the future, we'll want the commandlet to be able to perform partial updates of the map
 	// to allow HLOD rebuild to be distributed on multiple machines.
-	TArray<UWorldPartitionEditorCell*> AllCells;
-	AllCells.Add(WorldPartition->EditorHash->GetAlwaysLoadedCell());
-	WorldPartition->EditorHash->GetAllCells(AllCells);	
-	WorldPartition->LoadCells(AllCells);
+	const FBox LoadBox(FVector(-WORLD_MAX, -WORLD_MAX, -WORLD_MAX), FVector(WORLD_MAX, WORLD_MAX, WORLD_MAX));
+	WorldPartition->LoadEditorCells(LoadBox);
 
 	// Clear all existing HLOD actors
 	for (TActorIterator<AWorldPartitionHLOD> ItHLOD(World); ItHLOD; ++ItHLOD)
diff --git a/Engine/Source/Editor/WorldPartitionEditor/Private/WorldPartition/SWorldPartitionEditorGrid2D.cpp b/Engine/Source/Editor/WorldPartitionEditor/Private/WorldPartition/SWorldPartitionEditorGrid2D.cpp
--- a/Engine/Source/Editor/WorldPartitionEditor/Private/WorldPartition/SWorldPartitionEditorGrid2D.cpp
+++ b/Engine/Source/Editor/WorldPartitionEditor/Private/WorldPartition/SWorldPartitionEditorGrid2D.cpp
@@ -160,14 +160,14 @@ void SWorldPartitionEditorGrid2D::Construct(const FArguments& InArgs)
 
 void SWorldPartitionEditorGrid2D::LoadSelectedCells()
 {
-	WorldPartition->LoadCells(SelectedCells);
+	WorldPartition->LoadEditorCells(SelectedCells.Array());
 	GEditor->RedrawLevelEditingViewports();
 	RefreshSceneOutliner();
 }
 
 void SWorldPartitionEditorGrid2D::UnloadSelectedCells()
 {
-	WorldPartition->UnloadCells(SelectedCells);
+	WorldPartition->UnloadEditorCells(SelectedCells.Array());
 	GEditor->RedrawLevelEditingViewports();
 	RefreshSceneOutliner();
 }
@@ -787,11 +787,6 @@ void SWorldPartitionEditorGrid2D::UpdateSelection(bool bAppend)
 {
 	if (!bAppend)
 	{
-		for (UWorldPartitionEditorCell* SelectedCell: SelectedCells)
-		{
-			check(SelectedCell->bSelected);
-			SelectedCell->bSelected = false;
-		}
 		SelectedCells.Empty();
 	}
 
@@ -805,15 +800,9 @@ void SWorldPartitionEditorGrid2D::UpdateSelection(bool bAppend)
 	{
 		WorldPartition->EditorHash->ForEachIntersectingCell(SelectBox, [&](UWorldPartitionEditorCell* SelectedCell)
 		{
-			if (SelectedCell->bSelected)
-			{
-				check(bAppend);
-			}
-			else
-			{
-				SelectedCells.Add(SelectedCell);
-				SelectedCell->bSelected = true;
-			}
+			bool bWasAlreadyInSet;
+			SelectedCells.Add(SelectedCell, &bWasAlreadyInSet);
+			check(!bWasAlreadyInSet || bAppend);
 		});
 	}
 }
diff --git a/Engine/Source/Editor/WorldPartitionEditor/Private/WorldPartition/SWorldPartitionEditorGridSpatialHash.cpp b/Engine/Source/Editor/WorldPartitionEditor/Private/WorldPartition/SWorldPartitionEditorGridSpatialHash.cpp
--- a/Engine/Source/Editor/WorldPartitionEditor/Private/WorldPartition/SWorldPartitionEditorGridSpatialHash.cpp
+++ b/Engine/Source/Editor/WorldPartitionEditor/Private/WorldPartition/SWorldPartitionEditorGridSpatialHash.cpp
@@ -153,7 +153,7 @@ int32 SWorldPartitionEditorGridSpatialHash::PaintGrid(const FGeometry& AllottedG
 			FCellDesc2D& CellDesc2D = UniqueCells2D.Add(CellHash2D);
 
 			CellDesc2D.Bounds = FBox2D(FVector2D(Cell->Bounds.Min), FVector2D(Cell->Bounds.Max));
-			CellDesc2D.bSelected |= Cell->bSelected;
+			CellDesc2D.bSelected |= SelectedCells.Contains(Cell);
 			CellDesc2D.bLoaded |= Cell->bLoaded;
 			CellDesc2D.bEmpty &= !Cell->Actors.Num();
 		});
diff --git a/Engine/Source/Editor/WorldPartitionEditor/Public/WorldPartition/SWorldPartitionEditorGrid2D.h b/Engine/Source/Editor/WorldPartitionEditor/Public/WorldPartition/SWorldPartitionEditorGrid2D.h
--- a/Engine/Source/Editor/WorldPartitionEditor/Public/WorldPartition/SWorldPartitionEditorGrid2D.h
+++ b/Engine/Source/Editor/WorldPartitionEditor/Public/WorldPartition/SWorldPartitionEditorGrid2D.h
@@ -89,6 +89,6 @@ protected:
 	FVector2D MouseCursorPosWorld;
 	FVector2D SelectionStart;
 	FVector2D SelectionEnd;
-	TArray<UWorldPartitionEditorCell*> SelectedCells;
+	TSet<UWorldPartitionEditorCell*> SelectedCells;
 	FSlateFontInfo SmallLayoutFont;
 };
\ No newline at end of file
diff --git a/Engine/Source/Runtime/Engine/Private/WorldPartition/WorldPartition.cpp b/Engine/Source/Runtime/Engine/Private/WorldPartition/WorldPartition.cpp
--- a/Engine/Source/Runtime/Engine/Private/WorldPartition/WorldPartition.cpp
+++ b/Engine/Source/Runtime/Engine/Private/WorldPartition/WorldPartition.cpp
@@ -776,7 +776,25 @@ bool UWorldPartition::IsSimulating() const
 	return GEditor->bIsSimulatingInEditor || !!GEditor->PlayWorld;
 }
 
-void UWorldPartition::LoadCells(const TArray<UWorldPartitionEditorCell*>& CellsToLoad)
+void UWorldPartition::LoadEditorCells(const FBox& Box)
+{
+	TArray<UWorldPartitionEditorCell*> CellsToLoad;
+	if (EditorHash->GetIntersectingCells(Box, CellsToLoad))
+	{
+		LoadEditorCells(CellsToLoad);
+	}
+}
+
+void UWorldPartition::UnloadEditorCells(const FBox& Box)
+{
+	TArray<UWorldPartitionEditorCell*> CellsToUnload;
+	if (EditorHash->GetIntersectingCells(Box, CellsToUnload))
+	{
+		UnloadEditorCells(CellsToUnload);
+	}
+}
+
+void UWorldPartition::LoadEditorCells(const TArray<UWorldPartitionEditorCell*>& CellsToLoad)
 {
 	FWorldPartionCellUpdateContext CellUpdateContext(this);
 
@@ -794,7 +812,7 @@ void UWorldPartition::LoadCells(const TArray<UWorldPartitionEditorCell*>& CellsT
 	}
 }
 
-void UWorldPartition::UnloadCells(const TArray<UWorldPartitionEditorCell*>& CellsToUnload)
+void UWorldPartition::UnloadEditorCells(const TArray<UWorldPartitionEditorCell*>& CellsToUnload)
 {
 	FWorldPartionCellUpdateContext CellUpdateContext(this);
 
@@ -921,10 +939,6 @@ void UWorldPartition::UpdateLoadingEditorCell(UWorldPartitionEditorCell* Cell, b
 	}
 
 	Cell->bLoaded = bShouldBeLoaded;
-	if (bShouldBeLoaded)
-	{
-		Cell->bDirty = true;
-	}
 }
 
 void UWorldPartition::CreateLayers(const TSet<FName>& LayerNames)
diff --git a/Engine/Source/Runtime/Engine/Private/WorldPartition/WorldPartitionEditorCell.cpp b/Engine/Source/Runtime/Engine/Private/WorldPartition/WorldPartitionEditorCell.cpp
--- a/Engine/Source/Runtime/Engine/Private/WorldPartition/WorldPartitionEditorCell.cpp
+++ b/Engine/Source/Runtime/Engine/Private/WorldPartition/WorldPartitionEditorCell.cpp
@@ -45,8 +45,6 @@ void UWorldPartitionEditorCell::AddActor(FWorldPartitionActorDesc* InActorDesc)
 			UE_LOG(LogWorldPartition, Verbose, TEXT(" ==> Referenced loaded actor %s(%d) [UWorldPartitionEditorCell::AddActor]"), *Actor->GetFullName(), ActorRefCount);
 		}
 	}
-
-	bDirty = true;
 }
 
 void UWorldPartitionEditorCell::RemoveActor(FWorldPartitionActorDesc* InActorDesc)
@@ -59,8 +57,6 @@ void UWorldPartitionEditorCell::RemoveActor(FWorldPartitionActorDesc* InActorDes
 		const uint32 ActorRefCount = InActorDesc->RemoveLoadedRefCount();
 		UE_LOG(LogWorldPartition, Verbose, TEXT(" ==> Unreferenced loaded actor %s(%d) [UWorldPartitionEditorCell::RemoveActor]"), *InActorDesc->GetActor()->GetFullName(), ActorRefCount);
 	}
-
-	bDirty = true;
 }
 
 void UWorldPartitionEditorCell::PreEditUndo()
@@ -85,9 +81,7 @@ UWorldPartitionEditorCell::UWorldPartitionEditorCell(const FObjectInitializer& O
 	: Super(ObjectInitializer)
 #if WITH_EDITOR
 	, Bounds(ForceInitToZero)
-	, bDirty(false)
 	, bLoaded(false)
 	, bWasLoaded(false)
-	, bSelected(false)
 #endif
 {}
diff --git a/Engine/Source/Runtime/Engine/Private/WorldPartition/WorldPartitionEditorSpatialHash.cpp b/Engine/Source/Runtime/Engine/Private/WorldPartition/WorldPartitionEditorSpatialHash.cpp
--- a/Engine/Source/Runtime/Engine/Private/WorldPartition/WorldPartitionEditorSpatialHash.cpp
+++ b/Engine/Source/Runtime/Engine/Private/WorldPartition/WorldPartitionEditorSpatialHash.cpp
@@ -52,9 +52,10 @@ void UWorldPartitionEditorSpatialHash::Tick(float DeltaSeconds)
 			{
 				ForEachIntersectingCells(Bounds, Level, [&](const FCellCoord& CellCoord)
 				{
-					if (HashNodes.Contains(CellCoord))
+					FCellNode CellNode;
+					if (HashCells.RemoveAndCopyValue(CellCoord, CellNode))
 					{
-						HashNodes.Remove(CellCoord);
+						check(!CellNode.Cell);
 					}
 				});
 			}
@@ -74,7 +75,8 @@ void UWorldPartitionEditorSpatialHash::HashActor(FWorldPartitionActorDesc* InAct
 	}
 	else
 	{
-		int32 CurrentLevel = GetLevelForBox(Bounds);
+		const int32 CurrentLevel = GetLevelForBox(Bounds);
+		const int32 ActorLevel = GetLevelForBox(InActorDesc->GetBounds());
 
 		auto UpdateHigherLevels = [this](const FCellCoord& CellCoord, int32 EndLevel)
 		{
@@ -85,7 +87,7 @@ void UWorldPartitionEditorSpatialHash::HashActor(FWorldPartitionActorDesc* InAct
 
 				LevelCellCoord = LevelCellCoord.GetParentCellCoord();
 
-				FCellNode& CellNode = HashNodes.FindOrAdd(LevelCellCoord);
+				FCellNode& CellNode = HashCells.FindOrAdd(LevelCellCoord);
 
 				if (CellNode.HasChildNode(ChildIndex))
 				{
@@ -96,39 +98,34 @@ void UWorldPartitionEditorSpatialHash::HashActor(FWorldPartitionActorDesc* InAct
 			}
 		};
 
-		ForEachIntersectingCells(InActorDesc->GetBounds(), 0, [&](const FCellCoord& CellCoord)
+		ForEachIntersectingCells(InActorDesc->GetBounds(), ActorLevel, [&](const FCellCoord& CellCoord)
 		{
-			UWorldPartitionEditorCell* EditorCell = nullptr;
-			if (UWorldPartitionEditorCell** EditorCellPtr = HashCells.Find(CellCoord))
-			{
-				EditorCell = *EditorCellPtr;
-			}
-			else
+			FCellNode& CellNode = HashCells.FindOrAdd(CellCoord);
+
+			if (!CellNode.Cell)
 			{
-				EditorCell = NewObject<UWorldPartitionEditorCell>(this, NAME_None, RF_Transient);
-				EditorCell->SetFlags(RF_Transactional);
-				EditorCell->Bounds = GetCellBounds(CellCoord);
+				CellNode.Cell = NewObject<UWorldPartitionEditorCell>(this, NAME_None, RF_Transient);
+				CellNode.Cell->SetFlags(RF_Transactional);
+				CellNode.Cell->Bounds = GetCellBounds(CellCoord);
 
-				Cells.Add(EditorCell);
-				HashCells.Add(CellCoord, EditorCell);
+				Cells.Add(CellNode.Cell);
 
 				UpdateHigherLevels(CellCoord, CurrentLevel);
 
-				Bounds += EditorCell->Bounds;
+				Bounds += CellNode.Cell->Bounds;
 			}
 
-			check(EditorCell);
-			EditorCell->AddActor(InActorDesc);
+			CellNode.Cell->AddActor(InActorDesc);
 		});
 
-		int32 NewLevel = GetLevelForBox(Bounds);
+		const int32 NewLevel = GetLevelForBox(Bounds);
 		check(NewLevel >= CurrentLevel);
 
 		if (NewLevel > CurrentLevel)
 		{
 			ForEachIntersectingCells(Bounds, CurrentLevel, [&](const FCellCoord& CellCoord)
 			{
-				if (CurrentLevel ? HashNodes.Contains(CellCoord) : HashCells.Contains(CellCoord))
+				if (HashCells.Contains(CellCoord))
 				{
 					UpdateHigherLevels(CellCoord, NewLevel);
 				}
@@ -146,37 +143,19 @@ void UWorldPartitionEditorSpatialHash::UnhashActor(FWorldPartitionActorDesc* InA
 	}
 	else
 	{
-		int32 CurrentLevel = GetLevelForBox(Bounds);
+		const int32 CurrentLevel = GetLevelForBox(Bounds);
+		const int32 ActorLevel = GetLevelForBox(InActorDesc->GetBounds());
 
-		ForEachIntersectingCells(InActorDesc->GetBounds(), 0, [&](const FCellCoord& CellCoord)
+		ForEachIntersectingCells(InActorDesc->GetBounds(), ActorLevel, [&](const FCellCoord& CellCoord)
 		{
-			UWorldPartitionEditorCell* EditorCell = HashCells.FindChecked(CellCoord);
+			FCellNode& CellNode = HashCells.FindChecked(CellCoord);
 
-			EditorCell->RemoveActor(InActorDesc);
+			CellNode.Cell->RemoveActor(InActorDesc);
 
-			if (!EditorCell->Actors.Num())
+			if (!CellNode.Cell->Actors.Num())
 			{
-				verify(Cells.Remove(EditorCell));
-				verify(HashCells.Remove(CellCoord));
-
-				FCellCoord LevelCellCoord = CellCoord;
-				for (int32 Level=1; Level<=CurrentLevel; Level++)
-				{
-					const uint32 ChildIndex = LevelCellCoord.GetChildIndex();
-
-					LevelCellCoord = LevelCellCoord.GetParentCellCoord();
-
-					FCellNode& CellNode = HashNodes.FindChecked(LevelCellCoord);
-
-					CellNode.RemoveChildNode(ChildIndex);
-
-					if (CellNode.HasChildNodes())
-					{
-						break;
-					}
-
-					HashNodes.Remove(LevelCellCoord);
-				}
+				verify(Cells.Remove(CellNode.Cell));
+				CellNode.Cell = nullptr;
 
 				bBoundsDirty = true;
 			}
@@ -227,35 +206,27 @@ int32 UWorldPartitionEditorSpatialHash::ForEachIntersectingCellInner(const FBox&
 {
 	int32 NumIntersecting = 0;
 
-	if (CellCoord.Level)
+	FCellNode CellNode;
+	if (FCellNode* CellNodePtr = HashCells.Find(CellCoord))
 	{
-		FCellNode CellNode;
-		if (FCellNode* CellNodePtr = HashNodes.Find(CellCoord))
-		{
-			CellNode = *CellNodePtr;
+		CellNode = *CellNodePtr;
 
-			check(CellNode.HasChildNodes());
-
-			CellNode.ForEachChild([&](uint32 ChildIndex)
-			{
-				const FCellCoord ChildCellCoord = CellCoord.GetChildCellCoord(ChildIndex);
-				const FBox CellBounds = GetCellBounds(ChildCellCoord);
-
-				if (Box.Intersect(CellBounds))
-				{
-					NumIntersecting += ForEachIntersectingCellInner(Box, ChildCellCoord, InOperation);
-				}
-			});
-		}
-	}
-	else
-	{
-		UWorldPartitionEditorCell* EditorCell = nullptr;
-		if (UWorldPartitionEditorCell** EditorCellPtr = HashCells.Find(CellCoord))
+		if (CellNode.Cell)
 		{
-			InOperation(*EditorCellPtr);
+			InOperation(CellNode.Cell);
 			NumIntersecting++;
 		}
+
+		CellNode.ForEachChild([&](uint32 ChildIndex)
+		{
+			const FCellCoord ChildCellCoord = CellCoord.GetChildCellCoord(ChildIndex);
+			const FBox CellBounds = GetCellBounds(ChildCellCoord);
+
+			if (Box.Intersect(CellBounds))
+			{
+				NumIntersecting += ForEachIntersectingCellInner(Box, ChildCellCoord, InOperation);
+			}
+		});
 	}
 
 	return NumIntersecting;
@@ -265,9 +236,7 @@ int32 UWorldPartitionEditorSpatialHash::ForEachIntersectingCell(const FBox& Box,
 {
 	int32 NumIntersecting = 0;
 
-	const FBox SearchBox = Box.Overlap(Bounds);
-	const int32 SearchLevel = GetLevelForBox(SearchBox);
-
+	const int32 SearchLevel = GetLevelForBox(Bounds);
 	ForEachIntersectingCells(Box, SearchLevel, [&](const FCellCoord& CellCoord)
 	{
 		NumIntersecting += ForEachIntersectingCellInner(Box, CellCoord, InOperation);
diff --git a/Engine/Source/Runtime/Engine/Private/WorldPartition/WorldPartitionVolume.cpp b/Engine/Source/Runtime/Engine/Private/WorldPartition/WorldPartitionVolume.cpp
--- a/Engine/Source/Runtime/Engine/Private/WorldPartition/WorldPartitionVolume.cpp
+++ b/Engine/Source/Runtime/Engine/Private/WorldPartition/WorldPartitionVolume.cpp
@@ -20,16 +20,9 @@ void AWorldPartitionVolume::LoadIntersectingCells()
 	{
 		FVector Origin;
 		FVector Extent;
-		this->GetActorBounds(false, Origin, Extent);
+		GetActorBounds(false, Origin, Extent);
 
-		TArray<UWorldPartitionEditorCell*> IntersectingCells;
-		if (WorldPartition->EditorHash->GetIntersectingCells(FBox(Origin - Extent, Origin + Extent), IntersectingCells))
-		{
-			FScopedSlowTask SlowTask(IntersectingCells.Num(), LOCTEXT("LoadingCells", "Loading cells..."));
-			SlowTask.MakeDialog();
-
-			WorldPartition->LoadCells(IntersectingCells);
-		}
+		WorldPartition->LoadEditorCells(FBox(Origin - Extent, Origin + Extent));
 	}
 }
 #endif
diff --git a/Engine/Source/Runtime/Engine/Public/WorldPartition/WorldPartition.h b/Engine/Source/Runtime/Engine/Public/WorldPartition/WorldPartition.h
--- a/Engine/Source/Runtime/Engine/Public/WorldPartition/WorldPartition.h
+++ b/Engine/Source/Runtime/Engine/Public/WorldPartition/WorldPartition.h
@@ -72,7 +72,6 @@ private:
 
 	void FlushStreaming();
 
-protected:
 	// PIE Methods
 	void OnPreBeginPIE(bool bStartSimulate);
 	void OnEndPIE(bool bStartSimulate);
@@ -121,8 +120,11 @@ public:
 	FWorldPartitionActorDesc* GetActorDesc(const FGuid& Guid);
 	const FWorldPartitionActorDesc* GetActorDesc(const FGuid& Guid) const;
 
-	void LoadCells(const TArray<UWorldPartitionEditorCell*>& Cells);
-	void UnloadCells(const TArray<UWorldPartitionEditorCell*>& Cells);
+	void LoadEditorCells(const FBox& Box);
+	void UnloadEditorCells(const FBox& Box);
+
+	void LoadEditorCells(const TArray<UWorldPartitionEditorCell*>& Cells);
+	void UnloadEditorCells(const TArray<UWorldPartitionEditorCell*>& Cells);
 
 	void HashActorDesc(FWorldPartitionActorDesc* ActorDesc);
 	void UnhashActorDesc(FWorldPartitionActorDesc* ActorDesc);
diff --git a/Engine/Source/Runtime/Engine/Public/WorldPartition/WorldPartitionEditorCell.h b/Engine/Source/Runtime/Engine/Public/WorldPartition/WorldPartitionEditorCell.h
--- a/Engine/Source/Runtime/Engine/Public/WorldPartition/WorldPartitionEditorCell.h
+++ b/Engine/Source/Runtime/Engine/Public/WorldPartition/WorldPartitionEditorCell.h
@@ -8,7 +8,7 @@ class FWorldPartitionActorDesc;
 class UWorldPartition;
 
 /**
- * Represents an edition cell (editor-only)
+ * Represents an editing cell (editor-only)
  */
 UCLASS(Within = WorldPartitionEditorHash)
 class UWorldPartitionEditorCell: public UObject
@@ -29,12 +29,13 @@ public:
 	virtual void PostEditUndo() override;
 
 	FBox						Bounds;
-	bool						bDirty : 1;
 
 	/** Tells if the cell was manually loaded in the editor */
 	bool						bLoaded : 1;
+
+	/** Tells if the cell loaded state changed during transactions */
 	bool						bWasLoaded : 1;
-	bool						bSelected : 1;
+
 	TSet<FWorldPartitionActorDesc*> Actors;
 	TSet<FWorldPartitionActorDesc*> LoadedActors;
 #endif
diff --git a/Engine/Source/Runtime/Engine/Public/WorldPartition/WorldPartitionEditorSpatialHash.h b/Engine/Source/Runtime/Engine/Public/WorldPartition/WorldPartitionEditorSpatialHash.h
--- a/Engine/Source/Runtime/Engine/Public/WorldPartition/WorldPartitionEditorSpatialHash.h
+++ b/Engine/Source/Runtime/Engine/Public/WorldPartition/WorldPartitionEditorSpatialHash.h
@@ -125,6 +125,7 @@ class ENGINE_API UWorldPartitionEditorSpatialHash : public UWorldPartitionEditor
 	{
 		FCellNode()
 			: ChildNodesMask(0)
+			, Cell(nullptr)
 		{}
 
 		inline bool HasChildNodes() const
@@ -170,6 +171,7 @@ class ENGINE_API UWorldPartitionEditorSpatialHash : public UWorldPartitionEditor
 		}
 
 		uint8 ChildNodesMask;
+		UWorldPartitionEditorCell* Cell;
 	};
 
 public:
@@ -196,8 +198,7 @@ private:
 	UPROPERTY(Config)
 	int32 CellSize;
 
-	TMap<FCellCoord, FCellNode> HashNodes;
-	TMap<FCellCoord, UWorldPartitionEditorCell*> HashCells;
+	TMap<FCellCoord, FCellNode> HashCells;
 
 	UPROPERTY(Transient)
 	TSet<UWorldPartitionEditorCell*> Cells;
