Show a warning when loading a map with World Composition to consider switching to World Partition.

diff --git a/Engine/Source/Runtime/Engine/Private/UnrealEngine.cpp b/Engine/Source/Runtime/Engine/Private/UnrealEngine.cpp
--- a/Engine/Source/Runtime/Engine/Private/UnrealEngine.cpp
+++ b/Engine/Source/Runtime/Engine/Private/UnrealEngine.cpp
@@ -10061,6 +10061,15 @@ float DrawMapWarnings(UWorld* World, FViewport* Viewport, FCanvas* Canvas, UCanv
 		}
 	}
 
+	if (World->WorldComposition)
+	{
+		SmallTextItem.SetColor(FLinearColor::Red);
+		SmallTextItem.Text = LOCTEXT("WORLD_COMPOSITION_DEPRECATED_FMT", "WORLD COMPOSITION IS DEPRECATED, CONSIDER CONVERTING TO WORLD PARTITION");
+
+		Canvas->DrawItem(SmallTextItem, FVector2D(MessageX, MessageY));
+		MessageY += FontSizeY;
+	}
+
 #if !(UE_BUILD_SHIPPING || UE_BUILD_TEST)
 	// Check HLOD clusters and show warning if unbuilt
 	{
