Blueprint Editor: Support for loading LevelInstances (LevelInstanceSubystem supports EditorPreview worlds)

diff --git a/Engine/Source/Runtime/Engine/Private/LevelInstance/LevelInstanceActor.cpp b/Engine/Source/Runtime/Engine/Private/LevelInstance/LevelInstanceActor.cpp
--- a/Engine/Source/Runtime/Engine/Private/LevelInstance/LevelInstanceActor.cpp
+++ b/Engine/Source/Runtime/Engine/Private/LevelInstance/LevelInstanceActor.cpp
@@ -114,15 +114,15 @@ void ALevelInstance::PostUnregisterAllComponents()
 void ALevelInstance::LoadLevelInstance()
 {
 #if WITH_EDITOR
-	if (!bGuardLoadUnload)
+	if (!bGuardLoadUnload && !bIsEditorPreviewActor)
 #endif
 	{
 		if (ULevelInstanceSubsystem* LevelInstanceSubsystem = GetLevelInstanceSubsystem())
 		{
 			bool bForce = false;
 #if WITH_EDITOR
-			// When reinstancing. Avoid reloading the level but if the package changed, force the load
-			bForce = GIsReinstancing && IsLoaded() && LevelInstanceSubsystem->GetLevelInstanceLevel(this)->GetPackage()->FileName != FName(*GetWorldAssetPackage());
+			// When reinstancing or world wasn't ticked between changes. Avoid reloading the level but if the package changed, force the load
+			bForce = IsLoaded() && LevelInstanceSubsystem->GetLevelInstanceLevel(this)->GetPackage()->FileName != FName(*GetWorldAssetPackage());
 #endif
 			LevelInstanceSubsystem->RequestLoadLevelInstance(this, bForce);
 		}
@@ -132,7 +132,7 @@ void ALevelInstance::LoadLevelInstance()
 void ALevelInstance::UnloadLevelInstance()
 {
 #if WITH_EDITOR
-	if (!bGuardLoadUnload)
+	if (!bGuardLoadUnload && !bIsEditorPreviewActor)
 #endif
 	{
 		if (ULevelInstanceSubsystem* LevelInstanceSubsystem = GetLevelInstanceSubsystem())
@@ -327,8 +327,10 @@ bool ALevelInstance::CheckForLoop(TSoftObjectPtr<UWorld> InLevelInstance, TArray
 	{
 		LevelInstanceSubsystem->ForEachLevelInstanceAncestorsAndSelf(this, [&bValid, &InLevelInstance, LoopInfo, LoopStart, this](const ALevelInstance* LevelInstanceActor)
 		{
+			FName LongPackageName(*InLevelInstance.GetLongPackageName());
+			// Check to exclude NAME_None since Preview Levels are in the transient package
 			// Check the level we are spawned in to detect the loop (this will handle loops caused by LevelInstances and by regular level streaming)
-			if (LevelInstanceActor->GetLevel()->GetPackage()->FileName == FName(*InLevelInstance.GetLongPackageName()))
+			if (LongPackageName != NAME_None && LevelInstanceActor->GetLevel()->GetPackage()->FileName == LongPackageName)
 			{
 				bValid = false;
 				if (LoopStart)
diff --git a/Engine/Source/Runtime/Engine/Private/LevelInstance/LevelInstanceSubsystem.cpp b/Engine/Source/Runtime/Engine/Private/LevelInstance/LevelInstanceSubsystem.cpp
--- a/Engine/Source/Runtime/Engine/Private/LevelInstance/LevelInstanceSubsystem.cpp
+++ b/Engine/Source/Runtime/Engine/Private/LevelInstance/LevelInstanceSubsystem.cpp
@@ -73,6 +73,11 @@ void ULevelInstanceSubsystem::Deinitialize()
 	Super::Deinitialize();
 }
 
+bool ULevelInstanceSubsystem::DoesSupportWorldType(EWorldType::Type WorldType) const
+{
+	return Super::DoesSupportWorldType(WorldType) || WorldType == EWorldType::EditorPreview;
+}
+
 ALevelInstance* ULevelInstanceSubsystem::GetLevelInstance(FLevelInstanceID LevelInstanceID) const
 {
 	if (ALevelInstance*const* LevelInstanceActor = RegisteredLevelInstances.Find(LevelInstanceID))
@@ -888,6 +893,12 @@ bool ULevelInstanceSubsystem::IsLevelInstanceEditDirty(const FLevelInstanceEdit*
 
 bool ULevelInstanceSubsystem::CanEditLevelInstance(const ALevelInstance* LevelInstanceActor, FText* OutReason) const
 {
+	// Only allow Editing in Editor World
+	if (GetWorld()->WorldType != EWorldType::Editor)
+	{
+		return false;
+	}
+
 	if (ULevel* LevelInstanceLevel = GetLevelInstanceLevel(LevelInstanceActor))
 	{
 		if (LevelInstanceLevel->GetWorldPartition())
@@ -1197,12 +1208,19 @@ void ULevelInstanceSubsystem::CommitLevelInstance(ALevelInstance* LevelInstanceA
 	LevelInstanceEdits.Remove(FName(*EditPackage));
 
 	// Propagate to other instances
-	for (TActorIterator<ALevelInstance> LevelInstanceIt(GetWorld()); LevelInstanceIt; ++LevelInstanceIt)
+	for (TObjectIterator<UWorld> It(RF_ClassDefaultObject|RF_ArchetypeObject, true, EInternalObjectFlags::PendingKill); It; ++It)
 	{
-		ALevelInstance* CurrentLevelInstanceActor = *LevelInstanceIt;
-		if (CurrentLevelInstanceActor->GetWorldAssetPackage() == EditPackage && (LevelInstanceActor == CurrentLevelInstanceActor || bChangesCommitted))
+		UWorld* CurrentWorld = *It;
+		if (CurrentWorld->GetSubsystem<ULevelInstanceSubsystem>() != nullptr)
 		{
-			CurrentLevelInstanceActor->UpdateLevelInstance();
+			for (TActorIterator<ALevelInstance> LevelInstanceIt(CurrentWorld); LevelInstanceIt; ++LevelInstanceIt)
+			{
+				ALevelInstance* CurrentLevelInstanceActor = *LevelInstanceIt;
+				if (CurrentLevelInstanceActor->GetWorldAssetPackage() == EditPackage && (LevelInstanceActor == CurrentLevelInstanceActor || bChangesCommitted))
+				{
+					CurrentLevelInstanceActor->UpdateLevelInstance();
+				}
+			}
 		}
 	}
 
diff --git a/Engine/Source/Runtime/Engine/Public/LevelInstance/LevelInstanceSubsystem.h b/Engine/Source/Runtime/Engine/Public/LevelInstance/LevelInstanceSubsystem.h
--- a/Engine/Source/Runtime/Engine/Public/LevelInstance/LevelInstanceSubsystem.h
+++ b/Engine/Source/Runtime/Engine/Public/LevelInstance/LevelInstanceSubsystem.h
@@ -33,6 +33,7 @@ public:
 	//~ Begin USubsystem Interface.
 	virtual void Initialize(FSubsystemCollectionBase& Collection) override;
 	virtual void Deinitialize() override;
+	virtual bool DoesSupportWorldType(EWorldType::Type WorldType) const override;
 	//~ End USubsystem Interface.
 	
 	//~ Begin UWorldSubsystem Interface.
