Optimus: Factor out the value type portion of the type definition into its own class. Ensure that definitions are shared and can be compared for equality via simple pointer comparison.

#rnx

diff --git a/Engine/Source/Runtime/Engine/Private/ComputeFramework/ShaderParamTypeDefinition.cpp b/Engine/Source/Runtime/Engine/Private/ComputeFramework/ShaderParamTypeDefinition.cpp
new file mode 100644
--- /dev/null
+++ b/Engine/Source/Runtime/Engine/Private/ComputeFramework/ShaderParamTypeDefinition.cpp
@@ -0,0 +1,333 @@
+// Copyright Epic Games, Inc. All Rights Reserved.
+
+#include "ComputeFramework/ShaderParamTypeDefinition.h"
+
+#include "Serialization/Archive.h"
+#include "Templates/TypeHash.h"
+
+
+// Storage for shared shader value types.
+inline uint32 GetTypeHash(const FShaderValueTypeHandle& InTypeHandle)
+{
+	return GetTypeHash(*InTypeHandle.ValueTypePtr);
+}
+
+struct HandleKeyFuncs : BaseKeyFuncs<FShaderValueTypeHandle,FShaderValueTypeHandle,false>
+{
+	static KeyInitType GetSetKey(ElementInitType Element)
+	{
+		return Element;
+	}
+
+	static bool Matches(KeyInitType A, KeyInitType B)
+	{
+		// The handle value can never be null here.
+		return *A.ValueTypePtr == *B.ValueTypePtr;
+	}
+
+	static uint32 GetKeyHash(KeyInitType Key)
+	{
+		return GetTypeHash(*Key.ValueTypePtr);
+	}
+};
+
+
+static TSet<FShaderValueTypeHandle, HandleKeyFuncs> GloballyKnownValueTypes;
+
+
+
+FShaderValueTypeHandle FShaderValueType::Get(EShaderFundamentalType InType)
+{
+	if (InType == EShaderFundamentalType::Struct)
+	{
+		return {};
+	}
+
+	FShaderValueType ValueType;
+	ValueType.Type = InType;
+	ValueType.DimensionType = EShaderFundamentalDimensionType::Scalar;
+
+	return GetOrCreate(MoveTemp(ValueType));
+}
+
+
+FShaderValueTypeHandle FShaderValueType::Get(EShaderFundamentalType InType, int32 InElemCount)
+{
+	if (InType == EShaderFundamentalType::Struct || InElemCount < 1 || InElemCount > 4)
+	{
+		return {};
+	}
+
+	FShaderValueType ValueType;
+	ValueType.Type = InType;
+	ValueType.DimensionType = EShaderFundamentalDimensionType::Vector;
+	ValueType.VectorElemCount = InElemCount;
+
+	return GetOrCreate(MoveTemp(ValueType));
+}
+
+
+FShaderValueTypeHandle FShaderValueType::Get(EShaderFundamentalType InType, int32 InRowCount, int32 InColumnCount)
+{
+	if (InType == EShaderFundamentalType::Struct || 
+		InRowCount < 1 || InRowCount > 4 ||
+		InColumnCount < 1 || InColumnCount > 4)
+	{
+		return {};
+	}
+
+	FShaderValueType ValueType;
+	ValueType.Type = InType;
+	ValueType.DimensionType = EShaderFundamentalDimensionType::Matrix;
+	ValueType.MatrixRowCount = InRowCount;
+	ValueType.MatrixColumnCount = InColumnCount;
+
+	return GetOrCreate(MoveTemp(ValueType));
+}
+
+
+FShaderValueTypeHandle FShaderValueType::Get(
+	FName InName, 
+	std::initializer_list<FStructElement> InStructElements)
+{
+	if (InName == NAME_None)
+	{
+		return {};
+	}
+
+	// TODO: Check if the name and the element names are valid HLSL identifiers 
+	// (i.e. identifier characters only, not a reserved keyword, and no duplicates).
+	// TODO: Check if the name matches another struct with a different layout.
+
+	FShaderValueType ValueType;
+	ValueType.Name = InName;
+	ValueType.Type = EShaderFundamentalType::Struct;
+	ValueType.DimensionType = EShaderFundamentalDimensionType::Scalar;
+
+	for (const FStructElement &StructElement : InStructElements)
+	{
+		// FIXME: We don't allow nested structs for now to avoid complicating the 
+		// GetTypeDeclaration call too much.
+		if (StructElement.Type.ValueTypePtr == nullptr || 
+		    StructElement.Type.ValueTypePtr->Type == EShaderFundamentalType::Struct)
+		{
+			return {};
+		}
+
+		ValueType.StructElements.Add(StructElement);
+	}
+
+	// We don't allow empty structs.
+	if (ValueType.StructElements.IsEmpty())
+	{
+		return {};
+	}
+
+	return GetOrCreate(MoveTemp(ValueType));
+}
+
+
+FShaderValueTypeHandle FShaderValueType::GetOrCreate(FShaderValueType&& InValueType)
+{
+	FShaderValueTypeHandle Handle;
+	Handle.ValueTypePtr = &InValueType;
+
+	FShaderValueTypeHandle *FoundHandle = GloballyKnownValueTypes.Find(Handle);
+	if (FoundHandle)
+	{
+		return *FoundHandle;
+	}
+
+	Handle.ValueTypePtr = new FShaderValueType(MoveTemp(InValueType));
+	GloballyKnownValueTypes.Add(Handle);
+	return Handle;
+}
+
+
+bool FShaderValueType::operator==(const FShaderValueType& InOtherType) const
+{
+	if (Type != InOtherType.Type)
+	{
+		return false;
+	}
+
+	if (Type == EShaderFundamentalType::Struct)
+	{
+		if (Name != InOtherType.Name || StructElements.Num() != InOtherType.StructElements.Num())
+		{
+			return false;
+		}
+
+		for (int32 Index = 0; Index < StructElements.Num(); Index++)
+		{
+			if (StructElements[Index] != InOtherType.StructElements[Index])
+			{
+				return false;
+			}
+		}
+	}
+	else
+	{
+		if (DimensionType != InOtherType.DimensionType)
+		{
+			return false;
+		}
+
+		if (DimensionType == EShaderFundamentalDimensionType::Vector && 
+		    VectorElemCount != InOtherType.VectorElemCount)
+		{
+			return false;
+		}
+		else if (DimensionType == EShaderFundamentalDimensionType::Matrix &&
+		    (MatrixRowCount != InOtherType.MatrixRowCount || 
+			 MatrixColumnCount != InOtherType.MatrixColumnCount))
+		{
+			return false;
+		}
+	}
+
+	return true;
+}
+
+uint32 GetTypeHash(const FShaderValueType& InShaderValueType)
+{
+	uint32 Hash = GetTypeHash(int32(InShaderValueType.Type));
+
+	if (InShaderValueType.Type == EShaderFundamentalType::Struct)
+	{
+		Hash = HashCombine(Hash, GetTypeHash(InShaderValueType.Name));
+		Hash = HashCombine(Hash, GetTypeHash(InShaderValueType.StructElements.Num()));
+		for (const FShaderValueType::FStructElement& StructElement : InShaderValueType.StructElements)
+		{
+			Hash = HashCombine(Hash, GetTypeHash(StructElement.Name));
+			Hash = HashCombine(Hash, GetTypeHash(*StructElement.Type.ValueTypePtr));
+		}
+	}
+	else
+	{
+		Hash = HashCombine(Hash, GetTypeHash(int32(InShaderValueType.DimensionType)));
+
+		if (InShaderValueType.DimensionType == EShaderFundamentalDimensionType::Vector)
+		{
+			Hash = HashCombine(Hash, GetTypeHash(int32(InShaderValueType.VectorElemCount)));
+		}
+		else if (InShaderValueType.DimensionType == EShaderFundamentalDimensionType::Vector)
+		{
+			Hash = HashCombine(Hash, GetTypeHash(int32(InShaderValueType.MatrixRowCount)));
+			Hash = HashCombine(Hash, GetTypeHash(int32(InShaderValueType.MatrixColumnCount)));
+		}
+	}
+	return Hash;
+}
+
+
+FString FShaderValueType::ToString() const
+{
+	// FIXME: Cache on create?
+	FString BaseName;
+	switch (Type)
+	{
+	case EShaderFundamentalType::Bool:
+		BaseName = "bool";
+		break;
+	case EShaderFundamentalType::Int:
+		BaseName = "int";
+		break;
+	case EShaderFundamentalType::Uint:
+		BaseName = "uint";
+		break;
+	case EShaderFundamentalType::Float:
+		BaseName = "float";
+		break;
+	case EShaderFundamentalType::Double:
+		BaseName = "double";
+		break;
+
+	case EShaderFundamentalType::Struct:
+		return Name.ToString();
+	}
+
+	if (DimensionType == EShaderFundamentalDimensionType::Vector)
+	{
+		BaseName.Appendf(TEXT("%d"), VectorElemCount);
+	}
+	else if (DimensionType == EShaderFundamentalDimensionType::Matrix)
+	{
+		BaseName.Appendf(TEXT("%dx%d"), MatrixRowCount, MatrixColumnCount);
+	}
+
+	return BaseName;
+}
+
+
+FString FShaderValueType::GetTypeDeclaration() const
+{
+	// FIXME: Cache on create?
+	if (Type != EShaderFundamentalType::Struct)
+	{
+		return {};
+	}
+
+	TArray<FString> Elements;
+	Elements.Reserve(StructElements.Num());
+	for (const FStructElement& StructElement : StructElements)
+	{
+		Elements.Add(FString::Printf(TEXT("    %s %s;\n"), *StructElement.Type.ValueTypePtr->ToString(), *StructElement.Name.ToString()));
+	}
+
+	return FString::Printf(TEXT("struct %s {\n%s}"), 
+		*Name.ToString(),
+		*FString::Join(Elements, TEXT("")));
+}
+
+FArchive& operator<<(FArchive& InArchive, FShaderValueTypeHandle& InHandle)
+{
+	FShaderValueType ValueTypeTemp;
+	FShaderValueType *ValueTypePtr;
+
+	if (InArchive.IsLoading())
+	{
+		ValueTypePtr = &ValueTypeTemp;
+	}
+	else
+	{
+		ValueTypePtr = const_cast<FShaderValueType *>(InHandle.ValueTypePtr);
+	}
+
+	InArchive << ValueTypePtr->Type;
+
+	if (ValueTypePtr->Type == EShaderFundamentalType::Struct)
+	{
+		InArchive << ValueTypePtr->Name;
+		InArchive << ValueTypePtr->StructElements;
+	}
+	else
+	{
+		InArchive << ValueTypePtr->DimensionType;
+
+		if (ValueTypePtr->DimensionType == EShaderFundamentalDimensionType::Vector)
+		{
+			InArchive << ValueTypePtr->VectorElemCount;
+		}
+		else if (ValueTypePtr->DimensionType == EShaderFundamentalDimensionType::Vector)
+		{
+			InArchive << ValueTypePtr->MatrixRowCount;
+			InArchive << ValueTypePtr->MatrixColumnCount;
+		}
+	}
+
+	if (InArchive.IsLoading())
+	{
+		InHandle = FShaderValueType::GetOrCreate(MoveTemp(ValueTypeTemp));
+	}
+
+	return InArchive;
+}
+
+
+FArchive& operator<<(FArchive& InArchive, FShaderValueType::FStructElement& InElement)
+{
+	InArchive << InElement.Name;
+	InArchive << InElement.Type;
+	return InArchive;
+}
diff --git a/Engine/Source/Runtime/Engine/Public/ComputeFramework/ShaderParamTypeDefinition.h b/Engine/Source/Runtime/Engine/Public/ComputeFramework/ShaderParamTypeDefinition.h
--- a/Engine/Source/Runtime/Engine/Public/ComputeFramework/ShaderParamTypeDefinition.h
+++ b/Engine/Source/Runtime/Engine/Public/ComputeFramework/ShaderParamTypeDefinition.h
@@ -2,8 +2,16 @@
 
 #pragma once
 
+#include "CoreMinimal.h"
+#include "UObject/NameTypes.h"
+
 #include "ShaderParamTypeDefinition.generated.h"
 
+
+class FArchive;
+struct FShaderValueType;
+
+
 /* The base types of data that shaders can consume/expose */
 UENUM()
 enum class EShaderFundamentalType : uint8
@@ -13,7 +21,7 @@ enum class EShaderFundamentalType : uint8
 	Uint,
 	Float,
 	Double,
-	Struct,
+	Struct
 };
 
 /*
@@ -50,6 +58,129 @@ enum class EShaderResourceType : uint8
 	ByteAddressBuffer,
 };
 
+USTRUCT()
+struct FShaderValueTypeHandle
+{
+	GENERATED_BODY()
+
+	const FShaderValueType *ValueTypePtr = nullptr;
+
+	bool IsValid() const
+	{
+		return ValueTypePtr != nullptr;
+	}
+
+	explicit operator bool() const
+	{
+		return IsValid();
+	}
+
+	const FShaderValueType *operator->() const
+	{
+		return ValueTypePtr;
+	}
+
+	const FShaderValueType &operator*() const
+	{
+		return *ValueTypePtr;
+	}
+
+};
+
+USTRUCT()
+struct ENGINE_API FShaderValueType
+{
+	GENERATED_BODY()
+
+	// A simple container representing a single, named element in a shader value struct.
+	struct FStructElement
+	{
+		FStructElement() = default;
+		FStructElement(FName InName, FShaderValueTypeHandle InType)
+		    : Name(InName), Type(InType) {}
+
+		bool operator==(const FStructElement &InOther) const
+		{
+			return Name == InOther.Name && Type.ValueTypePtr == InOther.Type.ValueTypePtr;
+		}
+
+		bool operator!=(const FStructElement& InOther) const { return !(*this == InOther); }
+
+		friend FArchive& operator<<(FArchive& InArchive, FStructElement& InElement);
+
+		FName Name;
+		FShaderValueTypeHandle Type;
+	};
+
+
+
+	/** Returns a scalar value type. If the fundamental type given is invalid for scalar values 
+	  * (e.g. struct), then this function returns a nullptr. 
+	  */
+	static FShaderValueTypeHandle Get(EShaderFundamentalType InType);
+
+	/** Returns a vector value type. InElemCount can be any value between 1-4. If the type 
+	  * given is invalid for scalar values (e.g. struct) or InElemCount is out of range, then 
+	  * this function returns a nullptr. 
+	  */
+	static FShaderValueTypeHandle Get(EShaderFundamentalType InType, int32 InElemCount);
+
+	/** Constructor for vector values. InElemCount can be any value between 1-4 */
+	static FShaderValueTypeHandle Get(EShaderFundamentalType InType, int32 InRowCount, int32 InColumnCount);
+
+	static FShaderValueTypeHandle Get(FName InName, std::initializer_list<FStructElement> InStructElements);
+
+	/** Returns true if this type and the other type are exactly equal. */
+	bool operator==(const FShaderValueType &InOtherType) const;
+
+	/** Returns true if this type and the other type are not equal. */
+	bool operator!=(const FShaderValueType& InOtherType) const
+	{
+		return !(*this == InOtherType);
+	}
+
+	friend FArchive& operator<<(FArchive& InArchive, FShaderValueTypeHandle& InHandle);
+
+	/** Returns the type name as a string (e.g. 'vector2', 'matrix2x3' or 'struct_name') for 
+	    use in variable declarations. */
+	FString ToString() const;
+
+	/** Returns the type declaration if this type is a struct, or the empty string if not. */
+	FString GetTypeDeclaration() const;
+
+	UPROPERTY()
+	EShaderFundamentalType Type;
+
+	UPROPERTY()
+	EShaderFundamentalDimensionType DimensionType;
+
+	union
+	{
+		uint8 VectorElemCount;
+
+		struct
+		{
+			uint8 MatrixRowCount;
+			uint8 MatrixColumnCount;
+		};
+	};
+
+	UPROPERTY()
+	FName Name;
+
+	TArray<FStructElement> StructElements;
+
+private:
+	static FShaderValueTypeHandle GetOrCreate(FShaderValueType &&InValueType);
+};
+
+/**
+* A hashing function to allow the FShaderValueType class to be used with hashing containers (e.g.
+* TSet or TMap).
+*/
+inline uint32 GetTypeHash(const FShaderValueType& InShaderValueType);
+
+
 /* Fully describes the name and type of a parameter a shader exposes. */
 USTRUCT()
 struct FShaderParamTypeDefinition
@@ -60,6 +191,10 @@ public:
 	UPROPERTY(EditAnywhere, Category = "Kernel")
 	FString	Name;
 
+	// The value type for this definition.
+	UPROPERTY()
+	FShaderValueTypeHandle ValueType;
+
 // #TODO_ZABIR: Shader reflection needs to extract much richer type info. Currently skip param type validation
 #if 0
 	UPROPERTY()
