Fix bad merge.

diff --git a/Engine/Source/Runtime/Engine/Private/WorldPartition/WorldPartition.cpp b/Engine/Source/Runtime/Engine/Private/WorldPartition/WorldPartition.cpp
--- a/Engine/Source/Runtime/Engine/Private/WorldPartition/WorldPartition.cpp
+++ b/Engine/Source/Runtime/Engine/Private/WorldPartition/WorldPartition.cpp
@@ -826,7 +826,6 @@ void UWorldPartition::UnloadEditorCells(const TArray<UWorldPartitionEditorCell*>
 	for (UWorldPartitionEditorCell* Cell: CellsToUnload)
 	{
 		SlowTask.EnterProgressFrame(Cell->LoadedActors.Num(), FText::Format(LOCTEXT("UnloadingActors", "Unloading actors {0}/{1}"), (int32)SlowTask.CompletedWork, NumActorsToUnload));
-		SlowTask.EnterProgressFrame(Cell->LoadedActors.Num());
 		UpdateLoadingEditorCell(Cell, false);
 	}
 }
