CIS fix for Linux.

#rnx

diff --git a/Engine/Source/Runtime/Engine/Private/ComputeFramework/ShaderParamTypeDefinition.cpp b/Engine/Source/Runtime/Engine/Private/ComputeFramework/ShaderParamTypeDefinition.cpp
--- a/Engine/Source/Runtime/Engine/Private/ComputeFramework/ShaderParamTypeDefinition.cpp
+++ b/Engine/Source/Runtime/Engine/Private/ComputeFramework/ShaderParamTypeDefinition.cpp
@@ -7,7 +7,7 @@
 
 
 // Storage for shared shader value types.
-inline uint32 GetTypeHash(const FShaderValueTypeHandle& InTypeHandle)
+static uint32 GetTypeHash(const FShaderValueTypeHandle& InTypeHandle)
 {
 	return GetTypeHash(*InTypeHandle.ValueTypePtr);
 }
diff --git a/Engine/Source/Runtime/Engine/Public/ComputeFramework/ShaderParamTypeDefinition.h b/Engine/Source/Runtime/Engine/Public/ComputeFramework/ShaderParamTypeDefinition.h
--- a/Engine/Source/Runtime/Engine/Public/ComputeFramework/ShaderParamTypeDefinition.h
+++ b/Engine/Source/Runtime/Engine/Public/ComputeFramework/ShaderParamTypeDefinition.h
@@ -87,6 +87,7 @@ struct FShaderValueTypeHandle
 
 };
 
+
 USTRUCT()
 struct ENGINE_API FShaderValueType
 {
@@ -178,7 +179,7 @@ private:
 * A hashing function to allow the FShaderValueType class to be used with hashing containers (e.g.
 * TSet or TMap).
 */
-inline uint32 GetTypeHash(const FShaderValueType& InShaderValueType);
+uint32 GetTypeHash(const FShaderValueType& InShaderValueType);
 
 
 /* Fully describes the name and type of a parameter a shader exposes. */
