Changed showdebug default behavior, now need to pass showdebug TargetActor to show debug of target actor (like all others showdebug options)

diff --git a/Engine/Config/BaseInput.ini b/Engine/Config/BaseInput.ini
--- a/Engine/Config/BaseInput.ini
+++ b/Engine/Config/BaseInput.ini
@@ -184,6 +184,7 @@ MaxScrollbackSize=1024
 +ManualAutoCompleteList=(Command="CaptureMode",Desc="Toggles display state of all on-screen warnings/messages")
 +ManualAutoCompleteList=(Command="ShowDebug None",Desc="Toggles ShowDebug w/ current debug type selection")
 +ManualAutoCompleteList=(Command="ShowDebug Reset",Desc="Turns off ShowDebug, and clears debug type selection")
++ManualAutoCompleteList=(Command="ShowDebug TargetActor",Desc="Toggles display of target actor")
 +ManualAutoCompleteList=(Command="ShowDebug NET",Desc=)
 +ManualAutoCompleteList=(Command="ShowDebug PHYSICS",Desc=)
 +ManualAutoCompleteList=(Command="ShowDebug COLLISION",Desc=)
diff --git a/Engine/Source/Runtime/Engine/Private/HUD.cpp b/Engine/Source/Runtime/Engine/Private/HUD.cpp
--- a/Engine/Source/Runtime/Engine/Private/HUD.cpp
+++ b/Engine/Source/Runtime/Engine/Private/HUD.cpp
@@ -349,6 +349,7 @@ void AHUD::ShowHUD()
 namespace ShowDebugNames
 {
 	static const FName Reset(TEXT("Reset"));
+	static const FName TargetActor(TEXT("TargetActor"));
 	static const FName HitBox(TEXT("HitBox"));
 	static const FName Animation(TEXT("Animation"));
 	static const FName Physics(TEXT("Physics"));
@@ -438,31 +439,34 @@ void AHUD::ShowDebugInfo(float& YL, float& YPos)
 		FDebugDisplayInfo DisplayInfo(DebugDisplay, ToggledDebugCategories);
 		ShowDebugTargetActor = GetCurrentDebugTargetActor();
 
-		// Draw Header.
+		if (ShouldDisplayDebug(ShowDebugNames::TargetActor))
 		{
-			FDisplayDebugManager& DisplayDebugManager = DebugCanvas->DisplayDebugManager;
-			DisplayDebugManager.SetDrawColor(FColor(255, 0, 0));
-			DisplayDebugManager.DrawString(FString::Printf(TEXT("Showing Debug for %s, Press [PageUp] and [PageDown] to cycle between targets."), *GetNameSafe(ShowDebugTargetActor)));
-		}
+			// Draw Header.
+			{
+				FDisplayDebugManager& DisplayDebugManager = DebugCanvas->DisplayDebugManager;
+				DisplayDebugManager.SetDrawColor(FColor(255, 0, 0));
+				DisplayDebugManager.DrawString(FString::Printf(TEXT("Showing Debug for %s, Press [PageUp] and [PageDown] to cycle between targets."), *GetNameSafe(ShowDebugTargetActor)));
+			}
 
-		if (ShowDebugTargetActor && !ShowDebugTargetActor->IsPendingKill())
-		{
-			// Draw box around Actor being debugged.
-#if ENABLE_DRAW_DEBUG
+			if (ShowDebugTargetActor && !ShowDebugTargetActor->IsPendingKill())
 			{
-				FVector BoundsOrigin, BoundsExtent;
-				ShowDebugTargetActor->GetActorBounds(true, BoundsOrigin, BoundsExtent);
+				// Draw box around Actor being debugged.
+#if ENABLE_DRAW_DEBUG
+				{
+					FVector BoundsOrigin, BoundsExtent;
+					ShowDebugTargetActor->GetActorBounds(true, BoundsOrigin, BoundsExtent);
 
-				// Expand extent a little bit
-				BoundsExtent *= 1.1f;
-				DrawDebugBox(GetWorld(), BoundsOrigin, BoundsExtent, FColor::Green, false, -1.f, 0, 2.f);
-			}
+					// Expand extent a little bit
+					BoundsExtent *= 1.1f;
+					DrawDebugBox(GetWorld(), BoundsOrigin, BoundsExtent, FColor::Green, false, -1.f, 0, 2.f);
+				}
 #endif
-			ShowDebugTargetActor->DisplayDebug(DebugCanvas, DisplayInfo, YL, YPos);
+				ShowDebugTargetActor->DisplayDebug(DebugCanvas, DisplayInfo, YL, YPos);
 
-			if (!bShowDebugForReticleTarget && ShowDebugTargetActor->GetLocalRole() == ROLE_SimulatedProxy)
-			{
-				PlayerOwner->DisplayDebug(DebugCanvas, DisplayInfo, YL, YPos);
+				if (!bShowDebugForReticleTarget && ShowDebugTargetActor->GetLocalRole() == ROLE_SimulatedProxy)
+				{
+					PlayerOwner->DisplayDebug(DebugCanvas, DisplayInfo, YL, YPos);
+				}
 			}
 		}
 
