Added a way to access the element handle associated with a hit proxy item

This allows you to query which element you've hit, both at a point and within a region, via the FViewport::GetElementHandleAtPoint and FViewport::GetElementHandlesInRect functions.

diff --git a/Engine/Source/Runtime/Engine/Private/EngineUtils.cpp b/Engine/Source/Runtime/Engine/Private/EngineUtils.cpp
--- a/Engine/Source/Runtime/Engine/Private/EngineUtils.cpp
+++ b/Engine/Source/Runtime/Engine/Private/EngineUtils.cpp
@@ -40,6 +40,20 @@ IMPLEMENT_HIT_PROXY(HTranslucentActor,HActor)
 
 #define LOCTEXT_NAMESPACE "EngineUtils"
 
+FTypedElementHandle HActor::GetElementHandle() const
+{
+#if WITH_EDITOR
+	if (PrimComponent)
+	{
+		return PrimComponent->AcquireEditorElementHandle();
+	}
+	if (Actor)
+	{
+		return Actor->AcquireEditorElementHandle();
+	}
+#endif	// WITH_EDITOR
+	return FTypedElementHandle();
+}
 
 #if !UE_BUILD_SHIPPING
 FContentComparisonHelper::FContentComparisonHelper()
diff --git a/Engine/Source/Runtime/Engine/Private/HitProxies.cpp b/Engine/Source/Runtime/Engine/Private/HitProxies.cpp
--- a/Engine/Source/Runtime/Engine/Private/HitProxies.cpp
+++ b/Engine/Source/Runtime/Engine/Private/HitProxies.cpp
@@ -121,6 +121,11 @@ bool HHitProxy::IsA(HHitProxyType* TestType) const
 	return bIsInstance;
 }
 
+FTypedElementHandle HHitProxy::GetElementHandle() const
+{
+	return FTypedElementHandle();
+}
+
 HHitProxy* GetHitProxyById(FHitProxyId Id)
 {
 	return FHitProxyArray::Get().GetHitProxyById(Id.Index);
diff --git a/Engine/Source/Runtime/Engine/Private/UnrealClient.cpp b/Engine/Source/Runtime/Engine/Private/UnrealClient.cpp
--- a/Engine/Source/Runtime/Engine/Private/UnrealClient.cpp
+++ b/Engine/Source/Runtime/Engine/Private/UnrealClient.cpp
@@ -27,6 +27,7 @@
 #include "EngineModule.h"
 #include "Performance/EnginePerformanceTargets.h"
 #include "Templates/UniquePtr.h"
+#include "TypedElementList.h"
 #include "EngineUtils.h"
 
 DEFINE_LOG_CATEGORY_STATIC(LogClient, Log, All);
@@ -1803,44 +1804,87 @@ HHitProxy* FViewport::GetHitProxy(int32 X,int32 Y)
 
 void FViewport::GetActorsAndModelsInHitProxy(FIntRect InRect, TSet<AActor*>& OutActors, TSet<UModel*>& OutModels)
 {
-	const TArray<FColor>& RawHitProxyData = GetRawHitProxyData(InRect);
-
 	OutActors.Empty();
 	OutModels.Empty();
 
+	EnumerateHitProxiesInRect(InRect, [&OutActors, &OutModels](HHitProxy* HitProxy)
+	{
+		if (HitProxy->IsA(HActor::StaticGetType()))
+		{
+			AActor* Actor = ((HActor*)HitProxy)->Actor;
+			if (Actor)
+			{
+				OutActors.Add(Actor);
+			}
+		}
+		else if (HitProxy->IsA(HModel::StaticGetType()))
+		{
+			OutModels.Add(((HModel*)HitProxy)->GetModel());
+		}
+		else if (HitProxy->IsA(HBSPBrushVert::StaticGetType()))
+		{
+			HBSPBrushVert* HitBSPBrushVert = ((HBSPBrushVert*)HitProxy);
+			if (HitBSPBrushVert->Brush.IsValid())
+			{
+				OutActors.Add(HitBSPBrushVert->Brush.Get());
+			}
+		}
+		return true;
+	});
+}
+
+FTypedElementHandle FViewport::GetElementHandleAtPoint(int32 X, int32 Y)
+{
+	if (HHitProxy* HitProxy = GetHitProxy(X, Y))
+	{
+		return HitProxy->GetElementHandle();
+	}
+	return FTypedElementHandle();
+}
+
+void FViewport::GetElementHandlesInRect(FIntRect InRect, FTypedElementList& OutElementHandles)
+{
+	GetElementHandlesInRectImpl(InRect, OutElementHandles);
+}
+
+void FViewport::GetElementHandlesInRect(FIntRect InRect, TSet<FTypedElementHandle>& OutElementHandles)
+{
+	GetElementHandlesInRectImpl(InRect, OutElementHandles);
+}
+
+template <typename ContainerType>
+void FViewport::GetElementHandlesInRectImpl(FIntRect InRect, ContainerType& OutElementHandles)
+{
+	OutElementHandles.Reset();
+
+	EnumerateHitProxiesInRect(InRect, [&OutElementHandles](HHitProxy* HitProxy)
+	{
+		if (FTypedElementHandle ElementHandle = HitProxy->GetElementHandle())
+		{
+			OutElementHandles.Add(MoveTemp(ElementHandle));
+		}
+		return true;
+	});
+}
+
+void FViewport::EnumerateHitProxiesInRect(FIntRect InRect, TFunctionRef<bool(HHitProxy*)> InCallback)
+{
+	const TArray<FColor>& RawHitProxyData = GetRawHitProxyData(InRect);
+
 	// Lower the resolution with massive box selects
 	const int32 Step = (InRect.Width() > 500 && InRect.Height() > 500) ? 4 : 1;
 
-	for (int32 Y = InRect.Min.Y; Y < InRect.Max.Y; Y = Y < InRect.Max.Y - 1 ? FMath::Min(InRect.Max.Y-1, Y+Step) : ++Y )
+	for (int32 Y = InRect.Min.Y; Y < InRect.Max.Y; Y = Y < InRect.Max.Y-1 ? FMath::Min(InRect.Max.Y-1, Y+Step) : ++Y)
 	{
 		const FColor* SourceData = &RawHitProxyData[Y * SizeX];
-		for (int32 X = InRect.Min.X; X < InRect.Max.X; X = X < InRect.Max.X-1 ? FMath::Min(InRect.Max.X-1, X + Step) : ++X )
+		for (int32 X = InRect.Min.X; X < InRect.Max.X; X = X < InRect.Max.X-1 ? FMath::Min(InRect.Max.X-1, X+Step) : ++X)
 		{
 			FHitProxyId HitProxyId(SourceData[X]);
 			HHitProxy* HitProxy = GetHitProxyById(HitProxyId);
 
-			if (HitProxy)
+			if (HitProxy && !InCallback(HitProxy))
 			{
-				if( HitProxy->IsA(HActor::StaticGetType()) )
-				{
-					AActor* Actor = ((HActor*)HitProxy)->Actor;
-					if (Actor)
-					{
-						OutActors.Add(Actor);
-					}
-				}
-				else if( HitProxy->IsA(HModel::StaticGetType()) )
-				{
-					OutModels.Add( ((HModel*)HitProxy)->GetModel() );
-				}
-				else if( HitProxy->IsA(HBSPBrushVert::StaticGetType()) )
-				{
-					HBSPBrushVert* HitBSPBrushVert = ((HBSPBrushVert*)HitProxy);
-					if( HitBSPBrushVert->Brush.IsValid() )
-					{
-						OutActors.Add( HitBSPBrushVert->Brush.Get() );
-					}
-				}
+				return;
 			}
 		}
 	}
diff --git a/Engine/Source/Runtime/Engine/Public/EngineUtils.h b/Engine/Source/Runtime/Engine/Public/EngineUtils.h
--- a/Engine/Source/Runtime/Engine/Public/EngineUtils.h
+++ b/Engine/Source/Runtime/Engine/Public/EngineUtils.h
@@ -77,6 +77,8 @@ struct HActor : public HHitProxy
 	{
 		return EMouseCursor::Crosshairs;
 	}
+
+	ENGINE_API virtual FTypedElementHandle GetElementHandle() const override;
 };
 
 //
diff --git a/Engine/Source/Runtime/Engine/Public/HitProxies.h b/Engine/Source/Runtime/Engine/Public/HitProxies.h
--- a/Engine/Source/Runtime/Engine/Public/HitProxies.h
+++ b/Engine/Source/Runtime/Engine/Public/HitProxies.h
@@ -10,6 +10,7 @@
 #include "UObject/ObjectMacros.h"
 #include "Templates/RefCounting.h"
 #include "GenericPlatform/ICursor.h"
+#include "TypedElementHandle.h"
 
 class FReferenceCollector;
 
@@ -179,6 +180,11 @@ public:
 		return false;
 	}
 
+	/**
+	 * Get the element handle associated with this hit proxy, if any.
+	 */
+	ENGINE_API virtual FTypedElementHandle GetElementHandle() const;
+
 private:
 	void InitHitProxy();
 };
diff --git a/Engine/Source/Runtime/Engine/Public/UnrealClient.h b/Engine/Source/Runtime/Engine/Public/UnrealClient.h
--- a/Engine/Source/Runtime/Engine/Public/UnrealClient.h
+++ b/Engine/Source/Runtime/Engine/Public/UnrealClient.h
@@ -20,6 +20,7 @@ class FCanvas;
 class FViewport;
 class FViewportClient;
 class UModel;
+class FTypedElementList;
 
 /**
  * A render target.
@@ -491,6 +492,19 @@ public:
 	 */
 	ENGINE_API void GetActorsAndModelsInHitProxy(FIntRect InRect, TSet<AActor*>& OutActors, TSet<UModel*>& OutModels);
 
+	/**
+	 * Returns the dominant element handle at a given point.  If X,Y are outside the client area of the viewport, returns an invalid handle.
+	 */
+	ENGINE_API FTypedElementHandle GetElementHandleAtPoint(int32 X, int32 Y);
+
+	/**
+	 * Returns all element handles found within a specified region.
+	 * InRect must be entirely within the viewport's client area.
+	 * If the hit proxies are not cached, this will call ViewportClient->Draw with a hit-testing canvas.
+	 */
+	ENGINE_API void GetElementHandlesInRect(FIntRect InRect, FTypedElementList& OutElementHandles);
+	ENGINE_API void GetElementHandlesInRect(FIntRect InRect, TSet<FTypedElementHandle>& OutElementHandles);
+
 	/**
 	 * Retrieves the interface to the viewport's frame, if it has one.
 	 * @return The viewport's frame interface.
@@ -624,6 +638,23 @@ protected:
 	 */
 	void HighResScreenshot();
 
+private:
+
+	/**
+	 * Returns all element handles found within a specified region.
+	 * InRect must be entirely within the viewport's client area.
+	 * If the hit proxies are not cached, this will call ViewportClient->Draw with a hit-testing canvas.
+	 */
+	template <typename ContainerType>
+	void GetElementHandlesInRectImpl(FIntRect InRect, ContainerType& OutElementHandles);
+
+	/**
+	 * Enumerate all valid hit proxies found within a specified region.
+	 * InRect must be entirely within the viewport's client area.
+	 * If the hit proxies are not cached, this will call ViewportClient->Draw with a hit-testing canvas.
+	 */
+	void EnumerateHitProxiesInRect(FIntRect InRect, TFunctionRef<bool(HHitProxy*)> InCallback);
+
 protected:
 
 	/** A map from 2D coordinates to cached hit proxies. */
