Optimus: Fixed a yyp error.

#rnx

diff --git a/Engine/Source/Runtime/Engine/Private/ComputeFramework/ShaderParamTypeDefinition.cpp b/Engine/Source/Runtime/Engine/Private/ComputeFramework/ShaderParamTypeDefinition.cpp
--- a/Engine/Source/Runtime/Engine/Private/ComputeFramework/ShaderParamTypeDefinition.cpp
+++ b/Engine/Source/Runtime/Engine/Private/ComputeFramework/ShaderParamTypeDefinition.cpp
@@ -211,7 +211,7 @@ uint32 GetTypeHash(const FShaderValueType& InShaderValueType)
 		{
 			Hash = HashCombine(Hash, GetTypeHash(int32(InShaderValueType.VectorElemCount)));
 		}
-		else if (InShaderValueType.DimensionType == EShaderFundamentalDimensionType::Vector)
+		else if (InShaderValueType.DimensionType == EShaderFundamentalDimensionType::Matrix)
 		{
 			Hash = HashCombine(Hash, GetTypeHash(int32(InShaderValueType.MatrixRowCount)));
 			Hash = HashCombine(Hash, GetTypeHash(int32(InShaderValueType.MatrixColumnCount)));
@@ -309,7 +309,7 @@ FArchive& operator<<(FArchive& InArchive, FShaderValueTypeHandle& InHandle)
 		{
 			InArchive << ValueTypePtr->VectorElemCount;
 		}
-		else if (ValueTypePtr->DimensionType == EShaderFundamentalDimensionType::Vector)
+		else if (ValueTypePtr->DimensionType == EShaderFundamentalDimensionType::Matrix)
 		{
 			InArchive << ValueTypePtr->MatrixRowCount;
 			InArchive << ValueTypePtr->MatrixColumnCount;
