Added missing include needed for FUObjectAnnotationSparse

diff --git a/Engine/Source/Runtime/Engine/Classes/Engine/LevelStreaming.h b/Engine/Source/Runtime/Engine/Classes/Engine/LevelStreaming.h
--- a/Engine/Source/Runtime/Engine/Classes/Engine/LevelStreaming.h
+++ b/Engine/Source/Runtime/Engine/Classes/Engine/LevelStreaming.h
@@ -6,6 +6,7 @@
 #include "UObject/ObjectMacros.h"
 #include "UObject/Object.h"
 #include "UObject/ScriptMacros.h"
+#include "UObject/UObjectAnnotation.h"
 #include "Engine/LatentActionManager.h"
 #include "LatentActions.h"
 #include "ProfilingDebugging/ProfilingHelpers.h"
