- Added a log when reloading global shaders so its more apparent when the CookGlobalShaders commandlet tells a target device to reload shaders

diff --git a/Engine/Source/Runtime/Engine/Private/ShaderCompiler/ShaderCompiler.cpp b/Engine/Source/Runtime/Engine/Private/ShaderCompiler/ShaderCompiler.cpp
--- a/Engine/Source/Runtime/Engine/Private/ShaderCompiler/ShaderCompiler.cpp
+++ b/Engine/Source/Runtime/Engine/Private/ShaderCompiler/ShaderCompiler.cpp
@@ -4981,6 +4981,8 @@ void CompileGlobalShaderMap(bool bRefreshShaderMap)
 
 void ReloadGlobalShaders()
 {
+	UE_LOG(LogShaders, Display, TEXT("Reloading global shaders..."));
+
 	// Flush pending accesses to the existing global shaders.
 	FlushRenderingCommands();
 
