Reworked ActorPartitionSubsystem to work with the editor hierarchical spatial hash in World Partition.

diff --git a/Engine/Source/Runtime/Engine/Classes/GameFramework/Actor.h b/Engine/Source/Runtime/Engine/Classes/GameFramework/Actor.h
--- a/Engine/Source/Runtime/Engine/Classes/GameFramework/Actor.h
+++ b/Engine/Source/Runtime/Engine/Classes/GameFramework/Actor.h
@@ -63,9 +63,9 @@ enum class EActorUpdateOverlapsMethod : uint8
 UENUM()
 enum class EActorGridPlacement : uint8
 {
-	Bounds,
-	Location,
-	AlwaysLoaded,
+	Bounds,						// Actor uses its bounds to determine in which runtime cells it's going to be placed.
+	Location,					// Actor uses its location to determine in which runtime cells it's going to be placed.
+	AlwaysLoaded,				// Actor is always loaded (not placed in the grid), also affects editor.
 	None UMETA(Hidden)
 };
 
diff --git a/Engine/Source/Runtime/Engine/Classes/GameFramework/WorldSettings.h b/Engine/Source/Runtime/Engine/Classes/GameFramework/WorldSettings.h
--- a/Engine/Source/Runtime/Engine/Classes/GameFramework/WorldSettings.h
+++ b/Engine/Source/Runtime/Engine/Classes/GameFramework/WorldSettings.h
@@ -542,6 +542,9 @@ protected:
 	UWorldPartition* WorldPartition;
 
 public:
+	/** Size of the grid for instanced foliage actors*/
+	UPROPERTY(EditAnywhere, Category = Foliage)
+	uint32 InstancedFoliageGridSize;
 
 	/** scale of 1uu to 1m in real world measurements, for HMD and other physically tracked devices (e.g. 1uu = 1cm would be 100.0) */
 	UPROPERTY(EditAnywhere, BlueprintReadOnly, Category=VR)
diff --git a/Engine/Source/Runtime/Engine/Private/ActorPartition/ActorPartitionSubsystem.cpp b/Engine/Source/Runtime/Engine/Private/ActorPartition/ActorPartitionSubsystem.cpp
--- a/Engine/Source/Runtime/Engine/Private/ActorPartition/ActorPartitionSubsystem.cpp
+++ b/Engine/Source/Runtime/Engine/Private/ActorPartition/ActorPartitionSubsystem.cpp
@@ -1,13 +1,11 @@
 // Copyright Epic Games, Inc. All Rights Reserved.
 
 #include "ActorPartition/ActorPartitionSubsystem.h"
+#include "ActorPartition/InstancedObjectsActor.h"
+#include "WorldPartition/ActorPartition/InstancedObjectsActorDesc.h"
 #include "WorldPartition/WorldPartitionSubsystem.h"
 #include "Engine/World.h"
 
-#if WITH_EDITOR
-#include "WorldPartition/WorldPartitionEditorCell.h"
-#endif
-
 DEFINE_LOG_CATEGORY_STATIC(LogActorPartitionSubsystem, All, All);
 
 #if WITH_EDITOR
@@ -18,7 +16,8 @@ DEFINE_LOG_CATEGORY_STATIC(LogActorPartitionSubsystem, All, All);
 class FActorPartitionLevel : public FBaseActorPartition
 {
 public:
-	FActorPartitionLevel(UWorld* InWorld) : FBaseActorPartition(InWorld) 
+	FActorPartitionLevel(UWorld* InWorld)
+		: FBaseActorPartition(InWorld) 
 	{
 		 LevelRemovedFromWorldHandle = FWorldDelegates::LevelRemovedFromWorld.AddRaw(this, &FActorPartitionLevel::OnLevelRemovedFromWorld);
 	}
@@ -28,14 +27,13 @@ public:
 		FWorldDelegates::LevelRemovedFromWorld.Remove(LevelRemovedFromWorldHandle);
 	}
 
-	bool GetActorPartitionHash(const FActorPartitionGetParams& GetParams, FActorPartitionHash& OutPartitionHash) const override 
+	UActorPartitionSubsystem::FCellCoord GetActorPartitionHash(const FActorPartitionGetParams& GetParams) const override 
 	{ 
 		ULevel* SpawnLevel = GetSpawnLevel(GetParams.LevelHint, GetParams.LocationHint);
-		OutPartitionHash = GetLevelHash(SpawnLevel);
-		return true;
+		return UActorPartitionSubsystem::FCellCoord(0, 0, 0, SpawnLevel);
 	}
 	
-	AActor* GetActor(const FActorPartitionGetParams& GetParams) override
+	AActor* GetActor(const FActorPartitionGetParams& GetParams, const UActorPartitionSubsystem::FCellCoord& CellCoord) override
 	{
 		check(GetParams.LevelHint);
 		ULevel* SpawnLevel = GetSpawnLevel(GetParams.LevelHint, GetParams.LocationHint);
@@ -65,13 +63,10 @@ public:
 private:
 	void OnLevelRemovedFromWorld(ULevel* InLevel, UWorld* InWorld)
 	{
-		if (InWorld != World)
+		if (InWorld == World)
 		{
-			return;
+			GetOnActorPartitionHashInvalidated().Broadcast(UActorPartitionSubsystem::FCellCoord(0, 0, 0, InLevel));
 		}
-
-		FActorPartitionHash LevelHash = GetLevelHash(InLevel);
-		GetOnActorPartitionHashInvalidated().Broadcast(LevelHash);
 	}
 
 	ULevel* GetSpawnLevel(ULevel* InLevelHint, const FVector& InLocationHint) const
@@ -88,11 +83,6 @@ private:
 		return SpawnLevel;
 	}
 
-	FActorPartitionHash GetLevelHash(ULevel* Level) const
-	{
-		return StaticCast<FActorPartitionHash>(FCrc::MemCrc32(&Level, sizeof(ULevel*)));
-	}
-
 	FDelegateHandle LevelRemovedFromWorldHandle;
 };
 
@@ -102,54 +92,65 @@ private:
 class FActorPartitionWorldPartition : public FBaseActorPartition
 {
 public:
-	FActorPartitionWorldPartition(UWorld* InWorld) : FBaseActorPartition(InWorld)
+	FActorPartitionWorldPartition(UWorld* InWorld)
+		: FBaseActorPartition(InWorld)
 	{
+		WorldSettings = World->GetWorldSettings();
+		check(WorldSettings);
+
 		WorldPartition = InWorld->GetSubsystem<UWorldPartitionSubsystem>();
 		check(WorldPartition);
 	}
 
-	bool GetActorPartitionHash(const FActorPartitionGetParams& GetParams, FActorPartitionHash& OutPartitionHash) const override
-	{ 
-		FVector CellCenter;
-		UWorldPartitionEditorCell* Cell = nullptr;
-		if (!WorldPartition->GetCellAtLocation(GetParams.LocationHint, CellCenter, Cell))
-		{
-			return false;
-		}
-
-		if (!Cell->bLoaded)
-		{
-			return false;
-		}
-
-		OutPartitionHash = StaticCast<FActorPartitionHash>(FCrc::MemCrc32(&CellCenter, sizeof(CellCenter)));
-		return true;
+	UActorPartitionSubsystem::FCellCoord GetActorPartitionHash(const FActorPartitionGetParams& GetParams) const override 
+	{
+		return UActorPartitionSubsystem::FCellCoord(
+			FMath::FloorToInt(GetParams.LocationHint.X / GetParams.GridSize),
+			FMath::FloorToInt(GetParams.LocationHint.Y / GetParams.GridSize),
+			FMath::FloorToInt(GetParams.LocationHint.Z / GetParams.GridSize),
+			World->PersistentLevel
+		);
 	}
 
-	virtual AActor* GetActor(const FActorPartitionGetParams& GetParams)
+	virtual AActor* GetActor(const FActorPartitionGetParams& GetParams, const UActorPartitionSubsystem::FCellCoord& CellCoord)
 	{
 		AActor* FoundActor = nullptr;
-		FVector CellCenter;
-		UWorldPartitionEditorCell* Cell = nullptr;
-		if (!WorldPartition->GetCellAtLocation(GetParams.LocationHint, CellCenter, Cell))
-		{
-			return nullptr;
-		}
-
-		if (!Cell->bLoaded)
-		{
-			return nullptr;
-		}
-
-		TArray<AActor*> CellActors;
-		CellActors.Reserve(Cell->Actors.Num());
-		WorldPartition->GetCellActors(Cell, CellActors);
-		for (AActor* Actor : CellActors)
+		
+		FBox CellBox(
+			FVector(
+				CellCoord.X * GetParams.GridSize, 
+				CellCoord.Y * GetParams.GridSize, 
+				CellCoord.Z * GetParams.GridSize
+			), 
+			FVector(
+				CellCoord.X * GetParams.GridSize + GetParams.GridSize, 
+				CellCoord.Y * GetParams.GridSize + GetParams.GridSize, 
+				CellCoord.Z * GetParams.GridSize + GetParams.GridSize
+			)
+		);
+
+		TArray<const FWorldPartitionActorDesc*> InstancedObjectsActorDescs = WorldPartition->GetIntersectingActorDescs(CellBox, GetParams.ActorClass);
+
+		for (const FWorldPartitionActorDesc* ActorDesc: InstancedObjectsActorDescs)
 		{
-			if (Actor->IsA(GetParams.ActorClass))
+			if (ActorDesc->GetActorClass() == GetParams.ActorClass)
 			{
-				FoundActor = Actor;
-				break;
+				FInstancedObjectsActorDesc* InstancedObjectsActorDesc = (FInstancedObjectsActorDesc*)ActorDesc;
+				if ((InstancedObjectsActorDesc->GridIndexX == CellCoord.X) && 
+					(InstancedObjectsActorDesc->GridIndexY == CellCoord.Y) && 
+					(InstancedObjectsActorDesc->GridIndexZ == CellCoord.Z))
+				{
+					FoundActor = ActorDesc->GetActor();
+
+					if (!FoundActor)
+					{
+						// Actor exists but is not loaded
+						return nullptr;
+					}
+				
+					AInstancedObjectsActor* InstancedObjectsActor = CastChecked<AInstancedObjectsActor>(FoundActor);
+					check(InstancedObjectsActor->GridSize == GetParams.GridSize);
+				}
 			}
 		}
 
@@ -157,12 +158,11 @@ public:
 		{
 			FActorSpawnParameters SpawnParams;
 			SpawnParams.bCreateActorPackage = true;
+			FVector CellCenter(CellBox.GetCenter());
 			FoundActor = World->SpawnActor(GetParams.ActorClass, &CellCenter, nullptr, SpawnParams);
 
-			if (GetParams.ActorClass->GetDefaultObject<AActor>()->GetDefaultGridPlacement() != EActorGridPlacement::Location)
-			{
-				UE_LOG(LogActorPartitionSubsystem, All, TEXT("Spawned non location-based partition actor %s"), *FoundActor->GetName());
-			}
+			AInstancedObjectsActor* InstancedObjectsActor = CastChecked<AInstancedObjectsActor>(FoundActor);
+			InstancedObjectsActor->GridSize = GetParams.GridSize;
 		}
 
 		check(FoundActor || !GetParams.bCreate);
@@ -170,15 +170,14 @@ public:
 	}
 
 private:
+	AWorldSettings* WorldSettings;
 	UWorldPartitionSubsystem* WorldPartition;
 };
 
 #endif // WITH_EDITOR
 
 UActorPartitionSubsystem::UActorPartitionSubsystem()
-{
-
-}
+{}
 
 bool UActorPartitionSubsystem::IsLevelPartition() const
 {
@@ -191,6 +190,11 @@ void UActorPartitionSubsystem::Initialize(FSubsystemCollectionBase& Collection)
 {
 	Collection.InitializeDependency(UWorldPartitionSubsystem::StaticClass());
 
+	if (UWorldPartitionSubsystem* WorldPartitionSubsystem = GetWorld()->GetSubsystem<UWorldPartitionSubsystem>())
+	{
+		WorldPartitionSubsystem->RegisterActorDescFactory(AInstancedObjectsActor::StaticClass(), &InstancedObjectsActorDescFactory);
+	}
+
 	// Will need to register to WorldPartition setup changes events here...
 	InitializeActorPartition();
 }
@@ -198,56 +202,46 @@ void UActorPartitionSubsystem::Initialize(FSubsystemCollectionBase& Collection)
 /** Implement this for deinitialization of instances of the system */
 void UActorPartitionSubsystem::Deinitialize()
 {
-	if (ActorPartitionPtr)
+	if (ActorPartition)
 	{
-		ActorPartitionPtr->GetOnActorPartitionHashInvalidated().Remove(ActorPartitionHashInvalidatedHandle);
+		ActorPartition->GetOnActorPartitionHashInvalidated().Remove(ActorPartitionHashInvalidatedHandle);
 	}
 }
 
-void UActorPartitionSubsystem::OnActorPartitionHashInvalidated(const FActorPartitionHash& Hash)
+void UActorPartitionSubsystem::OnActorPartitionHashInvalidated(const FCellCoord& Hash)
 {
 	PartitionedActors.Remove(Hash);
 }
 
 void UActorPartitionSubsystem::InitializeActorPartition()
 {
-	check(!ActorPartitionPtr);
+	check(!ActorPartition);
 
 	if (IsLevelPartition())
 	{
-		ActorPartitionPtr.Reset(new FActorPartitionLevel(GetWorld()));
+		ActorPartition.Reset(new FActorPartitionLevel(GetWorld()));
 	}
 	else
 	{
-		ActorPartitionPtr.Reset(new FActorPartitionWorldPartition(GetWorld()));
+		ActorPartition.Reset(new FActorPartitionWorldPartition(GetWorld()));
 	}
-	ActorPartitionHashInvalidatedHandle = ActorPartitionPtr->GetOnActorPartitionHashInvalidated().AddUObject(this, &UActorPartitionSubsystem::OnActorPartitionHashInvalidated);
-}
-
-FBaseActorPartition& UActorPartitionSubsystem::GetActorPartition()
-{
-	check(ActorPartitionPtr);
-	return *ActorPartitionPtr;
+	ActorPartitionHashInvalidatedHandle = ActorPartition->GetOnActorPartitionHashInvalidated().AddUObject(this, &UActorPartitionSubsystem::OnActorPartitionHashInvalidated);
 }
 
 AActor* UActorPartitionSubsystem::GetActor(const FActorPartitionGetParams& GetParams)
 {
-	FBaseActorPartition& ActorPartition = GetActorPartition();
+	check(GetParams.ActorClass->IsChildOf(AInstancedObjectsActor::StaticClass()));
 
-	FActorPartitionHash ActorPartitionHash;
-	if (!ActorPartition.GetActorPartitionHash(GetParams, ActorPartitionHash))
-	{
-		return nullptr;
-	}
+	FCellCoord CellCoord = ActorPartition->GetActorPartitionHash(GetParams);
 		
-	TMap<UClass*, TWeakObjectPtr<AActor>>* ActorsPerClass = PartitionedActors.Find(ActorPartitionHash);
+	TMap<UClass*, TWeakObjectPtr<AActor>>* ActorsPerClass = PartitionedActors.Find(CellCoord);
 	AActor* FoundActor = nullptr;
 	if (!ActorsPerClass)
 	{
-		FoundActor = ActorPartition.GetActor(GetParams);
+		FoundActor = ActorPartition->GetActor(GetParams, CellCoord);
 		if (FoundActor)
 		{
-			PartitionedActors.Add(ActorPartitionHash).Add(GetParams.ActorClass, FoundActor);
+			PartitionedActors.Add(CellCoord).Add(GetParams.ActorClass, FoundActor);
 			return FoundActor;
 		}
 		else if (GetParams.bCreate)
@@ -260,7 +254,7 @@ AActor* UActorPartitionSubsystem::GetActor(const FActorPartitionGetParams& GetPa
 	TWeakObjectPtr<AActor>* ActorPtr = ActorsPerClass->Find(GetParams.ActorClass);
 	if (!ActorPtr || !ActorPtr->IsValid())
 	{
-		FoundActor = ActorPartition.GetActor(GetParams);
+		FoundActor = ActorPartition->GetActor(GetParams, CellCoord);
 		if (FoundActor)
 		{
 			if (!ActorPtr)
diff --git a/Engine/Source/Runtime/Engine/Private/ActorPartition/InstancedObjectsActor.cpp b/Engine/Source/Runtime/Engine/Private/ActorPartition/InstancedObjectsActor.cpp
new file mode 100644
--- /dev/null
+++ b/Engine/Source/Runtime/Engine/Private/ActorPartition/InstancedObjectsActor.cpp
@@ -0,0 +1,41 @@
+// Copyright Epic Games, Inc. All Rights Reserved.
+
+#include "ActorPartition/InstancedObjectsActor.h"
+
+#if WITH_EDITOR
+#include "ActorRegistry.h"
+#include "Components/BoxComponent.h"
+#endif
+
+AInstancedObjectsActor::AInstancedObjectsActor(const FObjectInitializer& ObjectInitializer)
+	: Super(ObjectInitializer)
+	, GridSize(0)
+{
+	USceneComponent* SceneComponent = CreateDefaultSubobject<USceneComponent>(TEXT("RootComponent0"));
+	RootComponent = SceneComponent;
+	RootComponent->Mobility = EComponentMobility::Static;
+}
+
+#if WITH_EDITOR	
+void AInstancedObjectsActor::GetAssetRegistryTags(TArray<FAssetRegistryTag>& OutTags) const
+{
+	Super::GetAssetRegistryTags(OutTags);
+
+	static const FName NAME_GridSize(TEXT("GridSize"));
+	FActorRegistry::SaveActorMetaData(NAME_GridSize, GridSize, OutTags);
+
+	const FVector ActorLocation = GetActorLocation();
+	const int64 GridIndexX = FMath::FloorToInt(ActorLocation.X / GridSize);
+	const int64 GridIndexY = FMath::FloorToInt(ActorLocation.Y / GridSize);
+	const int64 GridIndexZ = FMath::FloorToInt(ActorLocation.Z / GridSize);
+
+	static const FName NAME_GridIndexX(TEXT("GridIndexX"));
+	FActorRegistry::SaveActorMetaData(NAME_GridIndexX, GridIndexX, OutTags);
+
+	static const FName NAME_GridIndexY(TEXT("GridIndexY"));
+	FActorRegistry::SaveActorMetaData(NAME_GridIndexY, GridIndexY, OutTags);
+
+	static const FName NAME_GridIndexZ(TEXT("GridIndexZ"));
+	FActorRegistry::SaveActorMetaData(NAME_GridIndexZ, GridIndexZ, OutTags);
+}
+#endif
\ No newline at end of file
diff --git a/Engine/Source/Runtime/Engine/Private/ActorRegistry.cpp b/Engine/Source/Runtime/Engine/Private/ActorRegistry.cpp
--- a/Engine/Source/Runtime/Engine/Private/ActorRegistry.cpp
+++ b/Engine/Source/Runtime/Engine/Private/ActorRegistry.cpp
@@ -40,6 +40,11 @@ void FActorRegistry::SaveActorMetaData(FName Name, int32 Value, TArray<UObject::
 	OutTags.Add(FAssetRegistryTag(Name, *FString::Printf(TEXT("%d"), Value), FAssetRegistryTag::TT_Hidden));
 }
 
+void FActorRegistry::SaveActorMetaData(FName Name, int64 Value, TArray<UObject::FAssetRegistryTag>& OutTags)
+{
+	OutTags.Add(FAssetRegistryTag(Name, *FString::Printf(TEXT("%lld"), Value), FAssetRegistryTag::TT_Hidden));
+}
+
 void FActorRegistry::SaveActorMetaData(FName Name, const FGuid& Value, TArray<UObject::FAssetRegistryTag>& OutTags)
 {
 	OutTags.Add(FAssetRegistryTag(Name, Value.ToString(EGuidFormats::Base36Encoded), FAssetRegistryTag::TT_Hidden));
@@ -97,6 +102,18 @@ bool FActorRegistry::ReadActorMetaData(FName Name, int32& OutValue, const FAsset
 	return false;
 }
 
+bool FActorRegistry::ReadActorMetaData(FName Name, int64& OutValue, const FAssetData& AssetData)
+{
+	FString ValueStr;
+	if (AssetData.GetTagValue(Name, ValueStr))
+	{
+		OutValue = FCString::Atoi64(*ValueStr);
+		return true;
+	}
+
+	return false;
+}
+
 bool FActorRegistry::ReadActorMetaData(FName Name, FGuid& OutValue, const FAssetData& AssetData)
 {
 	FString ValueStr;
diff --git a/Engine/Source/Runtime/Engine/Private/LevelInstance/LevelInstanceActor.cpp b/Engine/Source/Runtime/Engine/Private/LevelInstance/LevelInstanceActor.cpp
--- a/Engine/Source/Runtime/Engine/Private/LevelInstance/LevelInstanceActor.cpp
+++ b/Engine/Source/Runtime/Engine/Private/LevelInstance/LevelInstanceActor.cpp
@@ -12,7 +12,6 @@
 
 #if WITH_EDITOR
 #include "ActorRegistry.h"
-#include "AssetData.h"
 #include "Logging/MessageLog.h"
 #include "Logging/TokenizedMessage.h"
 #include "Misc/MapErrors.h"
diff --git a/Engine/Source/Runtime/Engine/Private/WorldPartition/ActorPartition/InstancedObjectsActorDesc.cpp b/Engine/Source/Runtime/Engine/Private/WorldPartition/ActorPartition/InstancedObjectsActorDesc.cpp
new file mode 100644
--- /dev/null
+++ b/Engine/Source/Runtime/Engine/Private/WorldPartition/ActorPartition/InstancedObjectsActorDesc.cpp
@@ -0,0 +1,28 @@
+// Copyright Epic Games, Inc. All Rights Reserved.
+
+#include "WorldPartition/ActorPartition/InstancedObjectsActorDesc.h"
+#include "ActorPartition/InstancedObjectsActor.h"
+
+#if WITH_EDITOR
+#include "Engine/Level.h"
+FInstancedObjectsActorDesc::FInstancedObjectsActorDesc(const FWorldPartitionActorDescData& DescData, int32 InGridSize, int64 InGridIndexX, int64 InGridIndexY, int64 InGridIndexZ)
+	: FWorldPartitionActorDesc(DescData)
+	, GridSize(InGridSize)
+	, GridIndexX(InGridIndexX)
+	, GridIndexY(InGridIndexY)
+	, GridIndexZ(InGridIndexZ)
+{}
+
+FInstancedObjectsActorDesc::FInstancedObjectsActorDesc(AActor* InActor)
+	: FWorldPartitionActorDesc(InActor)
+{
+	AInstancedObjectsActor* InstancedObjectsActor = CastChecked<AInstancedObjectsActor>(InActor);
+
+	GridSize = InstancedObjectsActor->GridSize;
+
+	const FVector ActorLocation = InActor->GetActorLocation();
+	GridIndexX = FMath::FloorToInt(ActorLocation.X / GridSize);
+	GridIndexY = FMath::FloorToInt(ActorLocation.Y / GridSize);
+	GridIndexZ = FMath::FloorToInt(ActorLocation.Z / GridSize);
+}
+#endif
diff --git a/Engine/Source/Runtime/Engine/Private/WorldPartition/ActorPartition/InstancedObjectsActorDescFactory.cpp b/Engine/Source/Runtime/Engine/Private/WorldPartition/ActorPartition/InstancedObjectsActorDescFactory.cpp
new file mode 100644
--- /dev/null
+++ b/Engine/Source/Runtime/Engine/Private/WorldPartition/ActorPartition/InstancedObjectsActorDescFactory.cpp
@@ -0,0 +1,51 @@
+// Copyright Epic Games, Inc. All Rights Reserved.
+
+#include "WorldPartition/ActorPartition/InstancedObjectsActorDescFactory.h"
+#include "WorldPartition/ActorPartition/InstancedObjectsActorDesc.h"
+#include "ActorRegistry.h"
+
+#if WITH_EDITOR
+FWorldPartitionActorDesc* FInstancedObjectsActorDescFactory::CreateInstance(const FWorldPartitionActorDescInitData& ActorDescInitData)
+{
+	FWorldPartitionActorDescData Data;
+	if (!ReadMetaData(ActorDescInitData, Data))
+	{
+		return nullptr;
+	}
+
+	int32 GridSize;
+	static const FName NAME_GridSize(TEXT("GridSize"));
+	if (!FActorRegistry::ReadActorMetaData(NAME_GridSize, GridSize, ActorDescInitData.AssetData))
+	{
+		return nullptr;
+	}
+
+	int64 GridIndexX;
+	static const FName NAME_GridIndexX(TEXT("GridIndexX"));
+	if (!FActorRegistry::ReadActorMetaData(NAME_GridIndexX, GridIndexX, ActorDescInitData.AssetData))
+	{
+		return nullptr;
+	}
+
+	int64 GridIndexY;
+	static const FName NAME_GridIndexY(TEXT("GridIndexY"));
+	if (!FActorRegistry::ReadActorMetaData(NAME_GridIndexY, GridIndexY, ActorDescInitData.AssetData))
+	{
+		return nullptr;
+	}
+
+	int64 GridIndexZ;
+	static const FName NAME_GridIndexZ(TEXT("GridIndexZ"));
+	if (!FActorRegistry::ReadActorMetaData(NAME_GridIndexZ, GridIndexZ, ActorDescInitData.AssetData))
+	{
+		return nullptr;
+	}
+
+	return new FInstancedObjectsActorDesc(Data, GridSize, GridIndexX, GridIndexY, GridIndexZ);
+}
+
+FWorldPartitionActorDesc* FInstancedObjectsActorDescFactory::CreateInstance(AActor* InActor)
+{
+	return new FInstancedObjectsActorDesc(InActor);
+}
+#endif
\ No newline at end of file
diff --git a/Engine/Source/Runtime/Engine/Private/WorldPartition/WorldPartition.cpp b/Engine/Source/Runtime/Engine/Private/WorldPartition/WorldPartition.cpp
--- a/Engine/Source/Runtime/Engine/Private/WorldPartition/WorldPartition.cpp
+++ b/Engine/Source/Runtime/Engine/Private/WorldPartition/WorldPartition.cpp
@@ -590,20 +590,19 @@ FWorldPartitionActorDescFactory* UWorldPartition::GetActorDescFactory(const AAct
 	return GetActorDescFactory(Actor->GetClass());
 }
 
-bool UWorldPartition::GetCellAtLocation(const FVector& Location, FVector& Center, UWorldPartitionEditorCell*& Cell) const
+TArray<const FWorldPartitionActorDesc*> UWorldPartition::GetIntersectingActorDescs(const FBox& Box, TSubclassOf<AActor> ActorClass) const
 {
-	return EditorHash->GetCellAtLocation(Location, Center, Cell);
-}
+	TArray<const FWorldPartitionActorDesc*> ActorDescs;
 
-void UWorldPartition::GetCellActors(UWorldPartitionEditorCell* Cell, TArray<AActor*>& CellActors) const
-{
-	check(Cell->bLoaded);
-	for (const FWorldPartitionActorDesc* ActorDesc: Cell->Actors)
+	EditorHash->ForEachIntersectingActor(Box, [&ActorDescs, &ActorClass](const FWorldPartitionActorDesc* ActorDesc)
 	{
-		AActor* CellActor = ActorDesc->GetActor();
-		check(CellActor);
-		CellActors.Add(CellActor);
-	}
+		if (ActorDesc->GetActorClass()->IsChildOf(ActorClass))
+		{
+			ActorDescs.Add(ActorDesc);
+		}
+	});
+	
+	return MoveTemp(ActorDescs);
 }
 
 const FWorldPartitionActorDesc* UWorldPartition::GetActorDesc(const FGuid& Guid) const
diff --git a/Engine/Source/Runtime/Engine/Private/WorldPartition/WorldPartitionEditorSpatialHash.cpp b/Engine/Source/Runtime/Engine/Private/WorldPartition/WorldPartitionEditorSpatialHash.cpp
--- a/Engine/Source/Runtime/Engine/Private/WorldPartition/WorldPartitionEditorSpatialHash.cpp
+++ b/Engine/Source/Runtime/Engine/Private/WorldPartition/WorldPartitionEditorSpatialHash.cpp
@@ -21,7 +21,7 @@ void UWorldPartitionEditorSpatialHash::Initialize()
 {
 	check(!AlwaysLoadedCell);
 
-	AlwaysLoadedCell = NewObject<UWorldPartitionEditorCell>(this, NAME_None, RF_Transactional | RF_Transient);
+	AlwaysLoadedCell = NewObject<UWorldPartitionEditorCell>(this, TEXT("AlwaysLoadedCell"), RF_Transactional | RF_Transient);
 	AlwaysLoadedCell->Bounds.Init();
 }
 
@@ -66,6 +66,25 @@ void UWorldPartitionEditorSpatialHash::Tick(float DeltaSeconds)
 	}
 }
 
+// In the editor, actors are always using their bounds for grid placement, which makes more sense from a user standpoint.
+FBox UWorldPartitionEditorSpatialHash::GetActorBounds(FWorldPartitionActorDesc* InActorDesc) const
+{
+	FBox ActorBounds;
+
+	switch(InActorDesc->GetGridPlacement())
+	{
+	case EActorGridPlacement::Location:
+	case EActorGridPlacement::Bounds:
+	{
+		ActorBounds = InActorDesc->GetBounds();
+		break;
+	}
+	}
+
+	check(ActorBounds.IsValid);
+	return ActorBounds;
+}
+
 void UWorldPartitionEditorSpatialHash::HashActor(FWorldPartitionActorDesc* InActorDesc)
 {
 	check(InActorDesc);
@@ -75,8 +94,9 @@ void UWorldPartitionEditorSpatialHash::HashActor(FWorldPartitionActorDesc* InAct
 	}
 	else
 	{
+		const FBox ActorBounds = GetActorBounds(InActorDesc);
 		const int32 CurrentLevel = GetLevelForBox(Bounds);
-		const int32 ActorLevel = GetLevelForBox(InActorDesc->GetBounds());
+		const int32 ActorLevel = GetLevelForBox(ActorBounds);
 
 		auto UpdateHigherLevels = [this](const FCellCoord& CellCoord, int32 EndLevel)
 		{
@@ -98,13 +118,13 @@ void UWorldPartitionEditorSpatialHash::HashActor(FWorldPartitionActorDesc* InAct
 			}
 		};
 
-		ForEachIntersectingCells(InActorDesc->GetBounds(), ActorLevel, [&](const FCellCoord& CellCoord)
+		ForEachIntersectingCells(ActorBounds, ActorLevel, [&](const FCellCoord& CellCoord)
 		{
 			FCellNode& CellNode = HashCells.FindOrAdd(CellCoord);
 
 			if (!CellNode.Cell)
 			{
-				CellNode.Cell = NewObject<UWorldPartitionEditorCell>(this, NAME_None, RF_Transient);
+				CellNode.Cell = NewObject<UWorldPartitionEditorCell>(this, *FString::Printf(TEXT("EditorCell_X%lld_Y%lld_Z%lld_L%d"), CellCoord.X, CellCoord.Y, CellCoord.Z, CellCoord.Level), RF_Transient);
 				CellNode.Cell->SetFlags(RF_Transactional);
 				CellNode.Cell->Bounds = GetCellBounds(CellCoord);
 
@@ -143,10 +163,11 @@ void UWorldPartitionEditorSpatialHash::UnhashActor(FWorldPartitionActorDesc* InA
 	}
 	else
 	{
+		const FBox ActorBounds = GetActorBounds(InActorDesc);
 		const int32 CurrentLevel = GetLevelForBox(Bounds);
-		const int32 ActorLevel = GetLevelForBox(InActorDesc->GetBounds());
+		const int32 ActorLevel = GetLevelForBox(ActorBounds);
 
-		ForEachIntersectingCells(InActorDesc->GetBounds(), ActorLevel, [&](const FCellCoord& CellCoord)
+		ForEachIntersectingCells(ActorBounds, ActorLevel, [&](const FCellCoord& CellCoord)
 		{
 			FCellNode& CellNode = HashCells.FindChecked(CellCoord);
 
@@ -258,19 +279,6 @@ UWorldPartitionEditorCell* UWorldPartitionEditorSpatialHash::GetAlwaysLoadedCell
 {
 	return AlwaysLoadedCell;
 }
-
-bool UWorldPartitionEditorSpatialHash::GetCellAtLocation(const FVector& Location, FVector& Center, UWorldPartitionEditorCell*& Cell) const
-{
-	FCellCoord CellCoord(GetCellCoords(Location, 0));
-	if (UWorldPartitionEditorCell* const* CellPtr = (UWorldPartitionEditorCell* const*)HashCells.Find(CellCoord))
-	{
-		Cell = *CellPtr;
-		Center = Cell->Bounds.GetCenter();
-		return true;
-	}
-
-	return false;
-}
 #endif
 
 #undef LOCTEXT_NAMESPACE
diff --git a/Engine/Source/Runtime/Engine/Private/WorldPartition/WorldPartitionSubsystem.cpp b/Engine/Source/Runtime/Engine/Private/WorldPartition/WorldPartitionSubsystem.cpp
--- a/Engine/Source/Runtime/Engine/Private/WorldPartition/WorldPartitionSubsystem.cpp
+++ b/Engine/Source/Runtime/Engine/Private/WorldPartition/WorldPartitionSubsystem.cpp
@@ -180,22 +180,14 @@ void UWorldPartitionSubsystem::OnShowDebugInfo(class AHUD* HUD, UCanvas* Canvas,
 //
 
 #if WITH_EDITOR
-
-bool UWorldPartitionSubsystem::GetCellAtLocation(const FVector& Location, FVector& Center, UWorldPartitionEditorCell*& Cell) const
-{
-	if (const UWorldPartition* MainPartition = GetMainWorldPartition())
-	{
-		return MainPartition->GetCellAtLocation(Location, Center, Cell);
-	}
-	return false;
-}
-
-void UWorldPartitionSubsystem::GetCellActors(UWorldPartitionEditorCell* Cell, TArray<AActor*>& CellActors) const
+TArray<const FWorldPartitionActorDesc*> UWorldPartitionSubsystem::GetIntersectingActorDescs(const FBox& Box, TSubclassOf<AActor> ActorClass) const
 {
+	TArray<const FWorldPartitionActorDesc*> ActorDescs;
 	if (const UWorldPartition* MainPartition = GetMainWorldPartition())
 	{
-		MainPartition->GetCellActors(Cell, CellActors);
+		ActorDescs = MainPartition->GetIntersectingActorDescs(Box, ActorClass);
 	}
+	return MoveTemp(ActorDescs);
 }
 
 void UWorldPartitionSubsystem::UpdateActorDesc(AActor* Actor)
diff --git a/Engine/Source/Runtime/Engine/Private/WorldSettings.cpp b/Engine/Source/Runtime/Engine/Private/WorldSettings.cpp
--- a/Engine/Source/Runtime/Engine/Private/WorldSettings.cpp
+++ b/Engine/Source/Runtime/Engine/Private/WorldSettings.cpp
@@ -89,6 +89,8 @@ PRAGMA_ENABLE_DEPRECATION_WARNINGS
 	KillZ = -HALF_WORLD_MAX1;
 	KillZDamageType = ConstructorStatics.DmgType_Environmental_Object.Object;
 
+	InstancedFoliageGridSize = 25600;
+
 	WorldToMeters = 100.f;
 
 	DefaultPhysicsVolumeClass = ADefaultPhysicsVolume::StaticClass();
@@ -638,6 +640,10 @@ bool AWorldSettings::CanEditChange(const FProperty* InProperty) const
 		{
 			return !bEnableWorldPartition;
 		}
+		if (PropertyName == GET_MEMBER_NAME_STRING_CHECKED(AWorldSettings, InstancedFoliageGridSize))
+		{
+			return false;
+		}
 	}
 
 	return Super::CanEditChange(InProperty);
diff --git a/Engine/Source/Runtime/Engine/Public/ActorPartition/ActorPartitionSubsystem.h b/Engine/Source/Runtime/Engine/Public/ActorPartition/ActorPartitionSubsystem.h
--- a/Engine/Source/Runtime/Engine/Public/ActorPartition/ActorPartitionSubsystem.h
+++ b/Engine/Source/Runtime/Engine/Public/ActorPartition/ActorPartitionSubsystem.h
@@ -4,51 +4,39 @@
 
 #include "Subsystems/WorldSubsystem.h"
 #include "Templates/SubclassOf.h"
-
+#include "Misc/HashBuilder.h"
+#include "WorldPartition/ActorPartition/InstancedObjectsActorDescFactory.h"
 #include "ActorPartitionSubsystem.generated.h"
 
-using FActorPartitionHash = int32;
+class FBaseActorPartition;
 
 /**
  * FActorPartitionGetParam
  */
 struct FActorPartitionGetParams
 {
-	FActorPartitionGetParams(const TSubclassOf<AActor>& InActorClass, bool bInCreate, ULevel* InLevelHint, const FVector& InLocationHint)
+	FActorPartitionGetParams(const TSubclassOf<AActor>& InActorClass, bool bInCreate, ULevel* InLevelHint, const FVector& InLocationHint, int32 InGridSize)
 		: ActorClass(InActorClass)
 		, bCreate(bInCreate)
 		, LocationHint(InLocationHint)
 		, LevelHint(InLevelHint)
+		, GridSize(InGridSize)
 	{}
 
 	/* Class of Actor we are getting from the subsystem. */
 	TSubclassOf<AActor> ActorClass;
+	
 	/* Tells Subsystem if it needs to create Actor if it doesn't exist. */
 	bool bCreate;
+	
 	/* Depending on the world LocationHint can be used to find/create the Actor. */
 	FVector LocationHint;
+	
 	/* Depending on the world LevelHint can be used to find/create the Actor. */
 	ULevel* LevelHint;
-};
-
-/**
- * FBaseActorPartition
- */
-class FBaseActorPartition
-{
-public:
-	FBaseActorPartition(UWorld* InWorld) : World(InWorld) {}
-	virtual ~FBaseActorPartition() {}
-
-	virtual bool GetActorPartitionHash(const FActorPartitionGetParams& GetParams, FActorPartitionHash& OutPartitionHash) const = 0;
-	virtual AActor* GetActor(const FActorPartitionGetParams& GetParams) = 0;
-
-	DECLARE_EVENT_OneParam(FBaseActorPartition, FOnActorPartitionHashInvalidated, const FActorPartitionHash&);
-	FOnActorPartitionHashInvalidated& GetOnActorPartitionHashInvalidated() { return OnActorPartitionHashInvalidated; }
-protected:
-	UWorld* World;
-
-	FOnActorPartitionHashInvalidated OnActorPartitionHashInvalidated;
+	
+	/* Desired size of the grid to create actors on if partitioned. */
+	int32 GridSize;
 };
 
 /**
@@ -62,6 +50,36 @@ class ENGINE_API UActorPartitionSubsystem : public UWorldSubsystem
 
 public:
 	UActorPartitionSubsystem();
+
+	struct FCellCoord
+	{
+		FCellCoord()
+		{}
+
+		FCellCoord(int64 InX, int64 InY, int64 InZ, const ULevel* InLevel)
+			: X(InX)
+			, Y(InY)
+			, Z(InZ)
+			, Level(InLevel)
+		{}
+
+		int64 X;
+		int64 Y;
+		int64 Z;
+		const ULevel* Level;
+
+		bool operator==(const FCellCoord& Other) const
+		{
+			return (X == Other.X) && (Y == Other.Y) && (Z == Other.Z) && (Level == Other.Level);
+		}
+
+		friend ENGINE_API uint32 GetTypeHash(const FCellCoord& CellCoord)
+		{
+			FHashBuilder HashBuilder;
+			HashBuilder << CellCoord.X << CellCoord.Y << CellCoord.Z << PointerHash(CellCoord.Level);
+			return HashBuilder.GetHash();
+		}
+	};
 		
 #if WITH_EDITOR
 	AActor* GetActor(const FActorPartitionGetParams& GetParam);
@@ -74,16 +92,35 @@ public:
 private:
 
 #if WITH_EDITOR
-	void OnActorPartitionHashInvalidated(const FActorPartitionHash& Hash);
+	void OnActorPartitionHashInvalidated(const FCellCoord& Hash);
 
-	using TPartitionActors = TMap<UClass*, TWeakObjectPtr<AActor>>;
-	
-	FBaseActorPartition& GetActorPartition();
 	void InitializeActorPartition();
 
-	TMap<FActorPartitionHash, TPartitionActors> PartitionedActors;
-	TUniquePtr<FBaseActorPartition> ActorPartitionPtr;
+	TMap<FCellCoord, TMap<UClass*, TWeakObjectPtr<AActor>>> PartitionedActors;
+	TUniquePtr<FBaseActorPartition> ActorPartition;
+	
+	FInstancedObjectsActorDescFactory InstancedObjectsActorDescFactory;
 
 	FDelegateHandle ActorPartitionHashInvalidatedHandle;
 #endif
 };
+
+/**
+ * FBaseActorPartition
+ */
+class FBaseActorPartition
+{
+public:
+	FBaseActorPartition(UWorld* InWorld) : World(InWorld) {}
+	virtual ~FBaseActorPartition() {}
+
+	virtual UActorPartitionSubsystem::FCellCoord GetActorPartitionHash(const FActorPartitionGetParams& GetParams) const = 0;
+	virtual AActor* GetActor(const FActorPartitionGetParams& GetParams, const UActorPartitionSubsystem::FCellCoord& CellCoord) = 0;
+
+	DECLARE_EVENT_OneParam(FBaseActorPartition, FOnActorPartitionHashInvalidated, const UActorPartitionSubsystem::FCellCoord&);
+	FOnActorPartitionHashInvalidated& GetOnActorPartitionHashInvalidated() { return OnActorPartitionHashInvalidated; }
+protected:
+	UWorld* World;
+
+	FOnActorPartitionHashInvalidated OnActorPartitionHashInvalidated;
+};
\ No newline at end of file
diff --git a/Engine/Source/Runtime/Engine/Public/ActorPartition/InstancedObjectsActor.h b/Engine/Source/Runtime/Engine/Public/ActorPartition/InstancedObjectsActor.h
new file mode 100644
--- /dev/null
+++ b/Engine/Source/Runtime/Engine/Public/ActorPartition/InstancedObjectsActor.h
@@ -0,0 +1,36 @@
+// Copyright Epic Games, Inc. All Rights Reserved.
+#pragma once
+
+#include "CoreMinimal.h"
+#include "InstancedObjectsActor.generated.h"
+
+class UBoxComponent;
+
+// Actor base class for instance containers placed on a grid.
+// See UActorPartitionSubsystem.
+UCLASS(notplaceable)
+class ENGINE_API AInstancedObjectsActor: public AActor
+{
+	GENERATED_BODY()
+
+public:
+	AInstancedObjectsActor(const FObjectInitializer& ObjectInitializer = FObjectInitializer::Get());
+
+	//~ Begin UObject Interface
+#if WITH_EDITOR	
+	virtual void GetAssetRegistryTags(TArray<FAssetRegistryTag>& OutTags) const override;
+#endif
+	//~ End UObject Interface
+
+	//~ Begin AActor Interface
+#if WITH_EDITOR	
+	virtual EActorGridPlacement GetDefaultGridPlacement() const override { return EActorGridPlacement::Location; }
+#endif
+	//~ End AActor Interface	
+
+#if WITH_EDITORONLY_DATA
+	/** The grid size this actors was generated for */
+	UPROPERTY()
+	int32 GridSize;
+#endif
+};
\ No newline at end of file
diff --git a/Engine/Source/Runtime/Engine/Public/ActorRegistry.h b/Engine/Source/Runtime/Engine/Public/ActorRegistry.h
--- a/Engine/Source/Runtime/Engine/Public/ActorRegistry.h
+++ b/Engine/Source/Runtime/Engine/Public/ActorRegistry.h
@@ -31,6 +31,7 @@ public:
 	 */
 	static void SaveActorMetaData(FName Name, bool Value, TArray<FAssetRegistryTag>& OutTags);
 	static void SaveActorMetaData(FName Name, int32 Value, TArray<FAssetRegistryTag>& OutTags);
+	static void SaveActorMetaData(FName Name, int64 Value, TArray<FAssetRegistryTag>& OutTags);
 	static void SaveActorMetaData(FName Name, const FGuid& Value, TArray<FAssetRegistryTag>& OutTags);
 	static void SaveActorMetaData(FName Name, const FVector& Value, TArray<FAssetRegistryTag>& OutTags);
 	static void SaveActorMetaData(FName Name, const FTransform& Value, TArray<FAssetRegistryTag>& OutTags);
@@ -46,6 +47,7 @@ public:
 	 */
 	static bool ReadActorMetaData(FName Name, bool& OutValue, const FAssetData& AssetData);
 	static bool ReadActorMetaData(FName Name, int32& OutValue, const FAssetData& AssetData);
+	static bool ReadActorMetaData(FName Name, int64& OutValue, const FAssetData& AssetData);
 	static bool ReadActorMetaData(FName Name, FGuid& OutValue, const FAssetData& AssetData);
 	static bool ReadActorMetaData(FName Name, FVector& OutValue, const FAssetData& AssetData);
 	static bool ReadActorMetaData(FName Name, FTransform& OutValue, const FAssetData& AssetData);
diff --git a/Engine/Source/Runtime/Engine/Public/WorldPartition/ActorPartition/InstancedObjectsActorDesc.h b/Engine/Source/Runtime/Engine/Public/WorldPartition/ActorPartition/InstancedObjectsActorDesc.h
new file mode 100644
--- /dev/null
+++ b/Engine/Source/Runtime/Engine/Public/WorldPartition/ActorPartition/InstancedObjectsActorDesc.h
@@ -0,0 +1,26 @@
+// Copyright Epic Games, Inc. All Rights Reserved.
+
+#pragma once
+
+#include "CoreMinimal.h"
+#include "UObject/ObjectMacros.h"
+#include "Containers/Map.h"
+#include "WorldPartition/WorldPartitionActorDesc.h"
+
+class ENGINE_API FInstancedObjectsActorDesc : public FWorldPartitionActorDesc
+{
+public:
+	int32 GridSize;
+	int64 GridIndexX;
+	int64 GridIndexY;
+	int64 GridIndexZ;
+
+#if WITH_EDITOR
+protected:
+	FInstancedObjectsActorDesc() = delete;
+	FInstancedObjectsActorDesc(const FWorldPartitionActorDescData& DescData, int32 InGridSize, int64 InGridIndexX, int64 InGridIndexY, int64 InGridIndexZ);
+	FInstancedObjectsActorDesc(AActor* InActor);
+
+	friend class FInstancedObjectsActorDescFactory;
+#endif
+};
diff --git a/Engine/Source/Runtime/Engine/Public/WorldPartition/ActorPartition/InstancedObjectsActorDescFactory.h b/Engine/Source/Runtime/Engine/Public/WorldPartition/ActorPartition/InstancedObjectsActorDescFactory.h
new file mode 100644
--- /dev/null
+++ b/Engine/Source/Runtime/Engine/Public/WorldPartition/ActorPartition/InstancedObjectsActorDescFactory.h
@@ -0,0 +1,17 @@
+// Copyright Epic Games, Inc. All Rights Reserved.
+#pragma once
+
+#include "CoreMinimal.h"
+#include "UObject/ObjectMacros.h"
+#include "WorldPartition/WorldPartitionActorDescFactory.h"
+
+class FWorldPartitionActorDesc;
+
+class ENGINE_API FInstancedObjectsActorDescFactory : public FWorldPartitionActorDescFactory
+{
+public:
+#if WITH_EDITOR
+	virtual FWorldPartitionActorDesc* CreateInstance(const FWorldPartitionActorDescInitData& ActorDescInitData) override;
+	virtual FWorldPartitionActorDesc* CreateInstance(AActor* InActor) override;
+#endif
+};
\ No newline at end of file
diff --git a/Engine/Source/Runtime/Engine/Public/WorldPartition/WorldPartition.h b/Engine/Source/Runtime/Engine/Public/WorldPartition/WorldPartition.h
--- a/Engine/Source/Runtime/Engine/Public/WorldPartition/WorldPartition.h
+++ b/Engine/Source/Runtime/Engine/Public/WorldPartition/WorldPartition.h
@@ -93,10 +93,9 @@ private:
 	// UWorldPartitionSubsystem interface+
 	friend class UWorldPartitionSubsystem;
 
-	virtual bool GetCellAtLocation(const FVector& Location, FVector& Center, UWorldPartitionEditorCell*& Cell) const;
-
-	bool GetCellAtLocation(const FVector& Location, FVector& Center, UWorldPartitionEditorCell*& Cell);
-	void GetCellActors(UWorldPartitionEditorCell* Cell, TArray<AActor*>& CellActors) const;
+	// UActorPartitionSubsystem interface+
+	TArray<const FWorldPartitionActorDesc*> GetIntersectingActorDescs(const FBox& Box, TSubclassOf<AActor> ActorClass) const;
+	// UActorPartitionSubsystem interface-
 
 	void AddActor(AActor* InActor);
 	void RemoveActor(AActor* InActor);
diff --git a/Engine/Source/Runtime/Engine/Public/WorldPartition/WorldPartitionEditorHash.h b/Engine/Source/Runtime/Engine/Public/WorldPartition/WorldPartitionEditorHash.h
--- a/Engine/Source/Runtime/Engine/Public/WorldPartition/WorldPartitionEditorHash.h
+++ b/Engine/Source/Runtime/Engine/Public/WorldPartition/WorldPartitionEditorHash.h
@@ -26,7 +26,6 @@ public:
 	virtual int32 ForEachIntersectingCell(const FBox& Box, TFunctionRef<void(UWorldPartitionEditorCell*)> InOperation) PURE_VIRTUAL(UWorldPartitionEditorHash::ForEachIntersectingCell, return 0;);
 	virtual int32 ForEachCell(TFunctionRef<void(UWorldPartitionEditorCell*)> InOperation) PURE_VIRTUAL(UWorldPartitionEditorHash::ForEachCell, return 0;);
 	virtual UWorldPartitionEditorCell* GetAlwaysLoadedCell() PURE_VIRTUAL(UWorldPartitionEditorHash::GetAlwaysLoadedCell, return nullptr;);
-	virtual bool GetCellAtLocation(const FVector& Location, FVector& Center, UWorldPartitionEditorCell*& Cell) const PURE_VIRTUAL(UWorldPartitionEditorHash::GetCellAtLocation, return false;);
 
 	// Helpers
 	inline int32 GetIntersectingActors(const FBox& Box, TArray<FWorldPartitionActorDesc*>& OutActors)
diff --git a/Engine/Source/Runtime/Engine/Public/WorldPartition/WorldPartitionEditorSpatialHash.h b/Engine/Source/Runtime/Engine/Public/WorldPartition/WorldPartitionEditorSpatialHash.h
--- a/Engine/Source/Runtime/Engine/Public/WorldPartition/WorldPartitionEditorSpatialHash.h
+++ b/Engine/Source/Runtime/Engine/Public/WorldPartition/WorldPartitionEditorSpatialHash.h
@@ -92,7 +92,7 @@ class ENGINE_API UWorldPartitionEditorSpatialHash : public UWorldPartitionEditor
 		return FBox(Min, Max);
 	}
 
-	inline int32 GetLevelForBox(const FBox& Box)
+	inline int32 GetLevelForBox(const FBox& Box) const
 	{
 		const FVector Extent = Box.GetExtent();
 		const float MaxLength = Extent.GetMax() * 2.0;
@@ -187,12 +187,11 @@ public:
 	virtual int32 ForEachIntersectingCell(const FBox& Box, TFunctionRef<void(UWorldPartitionEditorCell*)> InOperation) override;
 	virtual int32 ForEachCell(TFunctionRef<void(UWorldPartitionEditorCell*)> InOperation) override;
 	virtual UWorldPartitionEditorCell* GetAlwaysLoadedCell() override;
-	virtual bool GetCellAtLocation(const FVector& Location, FVector& Center, UWorldPartitionEditorCell*& Cell) const override;
-
 #endif
 
 #if WITH_EDITORONLY_DATA
 private:
+	FBox GetActorBounds(FWorldPartitionActorDesc* InActorDesc) const;
 	int32 ForEachIntersectingCellInner(const FBox& Box, const FCellCoord& CellCoord, TFunctionRef<void(UWorldPartitionEditorCell*)> InOperation);
 
 	UPROPERTY(Config)
diff --git a/Engine/Source/Runtime/Engine/Public/WorldPartition/WorldPartitionSubsystem.h b/Engine/Source/Runtime/Engine/Public/WorldPartition/WorldPartitionSubsystem.h
--- a/Engine/Source/Runtime/Engine/Public/WorldPartition/WorldPartitionSubsystem.h
+++ b/Engine/Source/Runtime/Engine/Public/WorldPartition/WorldPartitionSubsystem.h
@@ -10,6 +10,7 @@
 class UWorldPartition;
 class UWorldPartitionEditorCell;
 class FWorldPartitionActorDescFactory;
+class FWorldPartitionActorDesc;
 
 /**
  * UWorldPartitionSubsystem
@@ -43,8 +44,7 @@ public:
 	//~End FTickableGameObject
 
 #if WITH_EDITOR
-	bool GetCellAtLocation(const FVector& Location, FVector& Center, UWorldPartitionEditorCell*& Cell) const;
-	void GetCellActors(UWorldPartitionEditorCell* Cell, TArray<AActor*>& CellActors) const;
+	TArray<const FWorldPartitionActorDesc*> GetIntersectingActorDescs(const FBox& Box, TSubclassOf<AActor> ActorClass) const;
 
 	void UpdateActorDesc(AActor* Actor);
 	void AddActor(AActor* Actor);
diff --git a/Engine/Source/Runtime/Foliage/Private/InstancedFoliage.cpp b/Engine/Source/Runtime/Foliage/Private/InstancedFoliage.cpp
--- a/Engine/Source/Runtime/Foliage/Private/InstancedFoliage.cpp
+++ b/Engine/Source/Runtime/Foliage/Private/InstancedFoliage.cpp
@@ -2336,10 +2336,6 @@ TArray<int32> FFoliageInfo::GetInstancesOverlappingBox(const FBox& Box) const
 AInstancedFoliageActor::AInstancedFoliageActor(const FObjectInitializer& ObjectInitializer)
 	: Super(ObjectInitializer)
 {
-	USceneComponent* SceneComponent = CreateDefaultSubobject<USceneComponent>(TEXT("RootComponent0"));
-	RootComponent = SceneComponent;
-	RootComponent->Mobility = EComponentMobility::Static;
-
 	SetActorEnableCollision(true);
 #if WITH_EDITORONLY_DATA
 	bListedInSceneOutliner = false;
@@ -2492,7 +2488,19 @@ const FFoliageInfo* AInstancedFoliageActor::FindInfo(const UFoliageType* InType)
 #if WITH_EDITOR
 AInstancedFoliageActor* AInstancedFoliageActor::Get(UWorld* InWorld, bool bCreateIfNone, ULevel* InLevelHint, const FVector& InLocationHint)
 {
-	return Cast<AInstancedFoliageActor>(InWorld->GetSubsystem<UActorPartitionSubsystem>()->GetActor(FActorPartitionGetParams(AInstancedFoliageActor::StaticClass(), bCreateIfNone, InLevelHint, InLocationHint)));
+	UActorPartitionSubsystem* ActorPartitionSubsystem = InWorld->GetSubsystem<UActorPartitionSubsystem>();
+
+	return CastChecked<AInstancedFoliageActor>(
+		ActorPartitionSubsystem->GetActor(
+			FActorPartitionGetParams(
+				AInstancedFoliageActor::StaticClass(), 
+				bCreateIfNone, 
+				InLevelHint, 
+				InLocationHint, 
+				InWorld->GetWorldSettings()->InstancedFoliageGridSize
+			)
+		)
+	);
 }
 
 
diff --git a/Engine/Source/Runtime/Foliage/Public/InstancedFoliageActor.h b/Engine/Source/Runtime/Foliage/Public/InstancedFoliageActor.h
--- a/Engine/Source/Runtime/Foliage/Public/InstancedFoliageActor.h
+++ b/Engine/Source/Runtime/Foliage/Public/InstancedFoliageActor.h
@@ -10,6 +10,7 @@
 #include "FoliageType_InstancedStaticMesh.h"
 #include "FoliageInstanceBase.h"
 #include "InstancedFoliage.h"
+#include "ActorPartition/InstancedObjectsActor.h"
 
 #include "InstancedFoliageActor.generated.h"
 
@@ -66,7 +67,7 @@ class UProceduralFoliageComponent;
 typedef TFunction<bool(const UPrimitiveComponent*)> FFoliageTraceFilterFunc;
 
 UCLASS(notplaceable, hidecategories = (Object, Rendering, Mobility), MinimalAPI, NotBlueprintable)
-class AInstancedFoliageActor : public AActor
+class AInstancedFoliageActor : public AInstancedObjectsActor
 {
 	GENERATED_UCLASS_BODY()
 
@@ -272,8 +273,6 @@ public:
 	void RepairDuplicateIFA(AInstancedFoliageActor* InDuplicateIFA);
 
 	void RemoveBaseComponentOnFoliageTypeInstances(UFoliageType* FoliageType);
-
-	virtual EActorGridPlacement GetDefaultGridPlacement() const override { return EActorGridPlacement::Location; }
 #endif	//WITH_EDITOR
 
 private:
