fix non-editor builds.

diff --git a/Engine/Source/Runtime/Engine/Private/ActorPartition/InstancedObjectsActor.cpp b/Engine/Source/Runtime/Engine/Private/ActorPartition/InstancedObjectsActor.cpp
--- a/Engine/Source/Runtime/Engine/Private/ActorPartition/InstancedObjectsActor.cpp
+++ b/Engine/Source/Runtime/Engine/Private/ActorPartition/InstancedObjectsActor.cpp
@@ -9,14 +9,16 @@
 
 AInstancedObjectsActor::AInstancedObjectsActor(const FObjectInitializer& ObjectInitializer)
 	: Super(ObjectInitializer)
+#if WITH_EDITORONLY_DATA
 	, GridSize(0)
+#endif
 {
 	USceneComponent* SceneComponent = CreateDefaultSubobject<USceneComponent>(TEXT("RootComponent0"));
 	RootComponent = SceneComponent;
 	RootComponent->Mobility = EComponentMobility::Static;
 }
 
-#if WITH_EDITOR	
+#if WITH_EDITOR
 void AInstancedObjectsActor::GetAssetRegistryTags(TArray<FAssetRegistryTag>& OutTags) const
 {
 	Super::GetAssetRegistryTags(OutTags);
diff --git a/Engine/Source/Runtime/Engine/Public/ActorPartition/InstancedObjectsActor.h b/Engine/Source/Runtime/Engine/Public/ActorPartition/InstancedObjectsActor.h
--- a/Engine/Source/Runtime/Engine/Public/ActorPartition/InstancedObjectsActor.h
+++ b/Engine/Source/Runtime/Engine/Public/ActorPartition/InstancedObjectsActor.h
@@ -2,6 +2,7 @@
 #pragma once
 
 #include "CoreMinimal.h"
+#include "GameFramework/Actor.h"
 #include "InstancedObjectsActor.generated.h"
 
 class UBoxComponent;
