Make sure HISM, Foliage, Grass are assigned the proper EngineShowFlag

diff --git a/Engine/Source/Runtime/Engine/Classes/Components/HierarchicalInstancedStaticMeshComponent.h b/Engine/Source/Runtime/Engine/Classes/Components/HierarchicalInstancedStaticMeshComponent.h
--- a/Engine/Source/Runtime/Engine/Classes/Components/HierarchicalInstancedStaticMeshComponent.h
+++ b/Engine/Source/Runtime/Engine/Classes/Components/HierarchicalInstancedStaticMeshComponent.h
@@ -14,6 +14,15 @@
 class FClusterBuilder;
 class FStaticLightingTextureMapping_InstancedStaticMesh;
 
+
+UENUM()
+enum class EHISMViewRelevanceType : uint8
+{
+	Grass,
+	Foliage,
+	HISM
+};
+
 // Due to BulkSerialize we can't edit the struct, so we must deprecated this one and create a new one
 USTRUCT()
 struct FClusterNode_DEPRECATED
@@ -253,6 +262,8 @@ public:
 
 	virtual void PropagateLightingScenarioChange() override;
 
+	virtual EHISMViewRelevanceType GetViewRelevanceType() const { return PerInstanceSMData.Num() ? EHISMViewRelevanceType::HISM : EHISMViewRelevanceType::Grass; }
+
 protected:
 	void BuildTree();
 	void BuildTreeAsync();
diff --git a/Engine/Source/Runtime/Engine/Private/HierarchicalInstancedStaticMesh.cpp b/Engine/Source/Runtime/Engine/Private/HierarchicalInstancedStaticMesh.cpp
--- a/Engine/Source/Runtime/Engine/Private/HierarchicalInstancedStaticMesh.cpp
+++ b/Engine/Source/Runtime/Engine/Private/HierarchicalInstancedStaticMesh.cpp
@@ -858,7 +858,7 @@ class FHierarchicalStaticMeshSceneProxy final : public FInstancedStaticMeshScene
 	int32 LastOcclusionNode;
 	TArray<FBoxSphereBounds> OcclusionBounds;
 	TMap<uint32, FFoliageOcclusionResults> OcclusionResults;
-	bool bIsGrass;
+	EHISMViewRelevanceType ViewRelevance;
 	bool bDitheredLODTransitions;
 	uint32 SceneProxyCreatedFrameNumberRenderThread;
 
@@ -877,14 +877,14 @@ public:
 		return reinterpret_cast<size_t>(&UniquePointer);
 	}
 
-	FHierarchicalStaticMeshSceneProxy(bool bInIsGrass, UHierarchicalInstancedStaticMeshComponent* InComponent, ERHIFeatureLevel::Type InFeatureLevel)
+	FHierarchicalStaticMeshSceneProxy(UHierarchicalInstancedStaticMeshComponent* InComponent, ERHIFeatureLevel::Type InFeatureLevel)
 		: FInstancedStaticMeshSceneProxy(InComponent, InFeatureLevel)
 		, ClusterTreePtr(InComponent->ClusterTreePtr.ToSharedRef())
 		, ClusterTree(*InComponent->ClusterTreePtr)
 		, UnbuiltBounds(InComponent->UnbuiltInstanceBoundsList)
 		, FirstUnbuiltIndex(InComponent->NumBuiltInstances > 0 ? InComponent->NumBuiltInstances : InComponent->NumBuiltRenderInstances)
 		, InstanceCountToRender(InComponent->InstanceCountToRender)
-		, bIsGrass(bInIsGrass)
+		, ViewRelevance(InComponent->GetViewRelevanceType())
 		, bDitheredLODTransitions(InComponent->SupportsDitheredLODTransitions(InFeatureLevel))
 		, SceneProxyCreatedFrameNumberRenderThread(UINT32_MAX)
 #if !(UE_BUILD_SHIPPING || UE_BUILD_TEST)
@@ -953,7 +953,22 @@ public:
 	virtual FPrimitiveViewRelevance GetViewRelevance(const FSceneView* View) const override
 	{
 		FPrimitiveViewRelevance Result;
-		if (bIsGrass ? View->Family->EngineShowFlags.InstancedGrass : View->Family->EngineShowFlags.InstancedFoliage)
+		bool bShowInstancedMesh = true;
+		switch (ViewRelevance)
+		{
+		case EHISMViewRelevanceType::Grass:
+			bShowInstancedMesh = View->Family->EngineShowFlags.InstancedGrass;
+			break;
+		case EHISMViewRelevanceType::Foliage:
+			bShowInstancedMesh = View->Family->EngineShowFlags.InstancedFoliage;
+			break;
+		case EHISMViewRelevanceType::HISM:
+			bShowInstancedMesh = View->Family->EngineShowFlags.InstancedStaticMeshes;
+			break;
+		default:
+			break;
+		}
+		if (bShowInstancedMesh)
 		{
 			Result = FStaticMeshSceneProxy::GetViewRelevance(View);
 			Result.bDynamicRelevance = true;
@@ -3018,8 +3033,6 @@ FPrimitiveSceneProxy* UHierarchicalInstancedStaticMeshComponent::CreateSceneProx
 
 		ProxySize = PerInstanceRenderData->ResourceSize;
 		INC_DWORD_STAT_BY(STAT_FoliageInstanceBuffers, ProxySize);
-		
-		bool bIsGrass = !PerInstanceSMData.Num();
 
 		// TODO: Abstract with a common helper
 		if (DoesPlatformSupportNanite(GMaxRHIShaderPlatform) &&
@@ -3031,7 +3044,7 @@ FPrimitiveSceneProxy* UHierarchicalInstancedStaticMeshComponent::CreateSceneProx
 		}
 		else
 		{
-			return ::new FHierarchicalStaticMeshSceneProxy(bIsGrass, this, GetWorld()->FeatureLevel);
+			return ::new FHierarchicalStaticMeshSceneProxy(this, GetWorld()->FeatureLevel);
 		}
 	}
 
diff --git a/Engine/Source/Runtime/Foliage/Public/FoliageInstancedStaticMeshComponent.h b/Engine/Source/Runtime/Foliage/Public/FoliageInstancedStaticMeshComponent.h
--- a/Engine/Source/Runtime/Foliage/Public/FoliageInstancedStaticMeshComponent.h
+++ b/Engine/Source/Runtime/Foliage/Public/FoliageInstancedStaticMeshComponent.h
@@ -44,6 +44,8 @@ public:
 	void SetGenerationGuid(const FGuid& InGuid) { GenerationGuid = InGuid; }
 	const FGuid& GetGenerationGuid() const { return GenerationGuid; }
 
+	virtual EHISMViewRelevanceType GetViewRelevanceType() const override { return EHISMViewRelevanceType::Foliage; }
+
 private:
 
 	UPROPERTY()
