Make sure that the always loaded cell is always loaded, regardless of if we are running from a cmdlet or not

diff --git a/Engine/Source/Runtime/Engine/Private/WorldPartition/WorldPartition.cpp b/Engine/Source/Runtime/Engine/Private/WorldPartition/WorldPartition.cpp
--- a/Engine/Source/Runtime/Engine/Private/WorldPartition/WorldPartition.cpp
+++ b/Engine/Source/Runtime/Engine/Private/WorldPartition/WorldPartition.cpp
@@ -232,10 +232,7 @@ void UWorldPartition::Initialize(UWorld* InWorld, const FTransform& InTransform)
 
 		EditorHash->Initialize();
 
-		if (!IsRunningCommandlet())
-		{
-			LayerSubSystem = GEditor->GetEditorSubsystem<ULayersSubsystem>();
-		}
+		LayerSubSystem = GEditor->GetEditorSubsystem<ULayersSubsystem>();
 
 		if (IsMainWorldPartition())
 		{
@@ -308,21 +305,18 @@ void UWorldPartition::Initialize(UWorld* InWorld, const FTransform& InTransform)
 				HashActorDesc(Pair.Value.Get());
 			}
 
-			if (!IsRunningCommandlet())
-			{
-				CreateLayers(AllLayersNames);
+			CreateLayers(AllLayersNames);
 
-				// Load the always loaded cell, don't call LoadCells to avoid creating a transaction
-				UpdateLoadingEditorCell(EditorHash->GetAlwaysLoadedCell(), true);
+			// Load the always loaded cell, don't call LoadCells to avoid creating a transaction
+			UpdateLoadingEditorCell(EditorHash->GetAlwaysLoadedCell(), true);
 
-				// When loading a subworld, load all actors
-				if (!IsMainWorldPartition())
+			// When loading a subworld, load all actors
+			if (!IsMainWorldPartition())
+			{
+				EditorHash->ForEachCell([this](UWorldPartitionEditorCell* Cell)
 				{
-					EditorHash->ForEachCell([this](UWorldPartitionEditorCell* Cell)
-					{
-						UpdateLoadingEditorCell(Cell, true);
-					});
-				}
+					UpdateLoadingEditorCell(Cell, true);
+				});
 			}
 		}
 	}
