Only register components in ULevel::AddLoadedActor() if the level has bAreComponentsCurrentlyRegistered set

If components from actors belonging to a level were already registered, do the same for the newly loaded actor.
Otherwise, the new actor components will be registered later once the world initialization is completed, through UpdateWorldComponents()

diff --git a/Engine/Source/Runtime/Engine/Private/Level.cpp b/Engine/Source/Runtime/Engine/Private/Level.cpp
--- a/Engine/Source/Runtime/Engine/Private/Level.cpp
+++ b/Engine/Source/Runtime/Engine/Private/Level.cpp
@@ -604,10 +604,15 @@ void ULevel::AddLoadedActor(AActor* Actor)
 	Actors.Add(Actor);
 	ActorsForGC.Add(Actor);
 
-	Actor->RegisterAllComponents();
-	Actor->RerunConstructionScripts();
-	GetWorld()->UpdateCullDistanceVolumes(Actor);
-	Actor->MarkComponentsRenderStateDirty();
+	// If components from actors belonging to this level were already registered, do the same for the newly loaded actor.
+	// Otherwise, the new actor components will be registered later once the world initialization is completed, through UpdateWorldComponents()
+	if (bAreComponentsCurrentlyRegistered)
+	{
+		Actor->RegisterAllComponents();
+		Actor->RerunConstructionScripts();
+		GetWorld()->UpdateCullDistanceVolumes(Actor);
+		Actor->MarkComponentsRenderStateDirty();
+	}
 
 	OnLoadedActorAddedToLevelEvent.Broadcast(*Actor);
 }
