WorldPartition:
- Added 3D debug to show Runtime Hash in world
- UWorldPartitionRuntimeSpatialHash::GetStreamingCells now makes sure to always return always loaded cell
- Fixed bug where streaming grid could create cells without any actor
- Added offset in Runtime Spatial Hash hierarchical parent level grids to fix actor promotion actefacts
- Optimized actor promotion in hierarchical spatial hash
- Optimized UWorldPartitionRuntimeSpatialHash::GetStreamingCells

diff --git a/Engine/Source/Runtime/Core/Public/Math/Grid2D.h b/Engine/Source/Runtime/Core/Public/Math/Grid2D.h
deleted file mode 100644
diff --git a/Engine/Source/Runtime/Engine/Private/WorldPartition/WorldPartition.cpp b/Engine/Source/Runtime/Engine/Private/WorldPartition/WorldPartition.cpp
--- a/Engine/Source/Runtime/Engine/Private/WorldPartition/WorldPartition.cpp
+++ b/Engine/Source/Runtime/Engine/Private/WorldPartition/WorldPartition.cpp
@@ -1182,16 +1182,22 @@ void UWorldPartition::UpdateStreamingState()
 	}
 }
 
-FVector2D UWorldPartition::GetShowDebugDesiredFootprint(const FVector2D& CanvasSize)
+FVector2D UWorldPartition::GetDrawRuntimeHash2DDesiredFootprint(const FVector2D& CanvasSize)
 {
 	check(GetWorld()->IsGameWorld());
-	return GetStreamingPolicy()->GetShowDebugDesiredFootprint(CanvasSize);
+	return GetStreamingPolicy()->GetDrawRuntimeHash2DDesiredFootprint(CanvasSize);
 }
 
-void UWorldPartition::ShowDebugInfo(class UCanvas* Canvas, const FVector2D& PartitionCanvasOffset, const FVector2D& PartitionCanvasSize)
+void UWorldPartition::DrawRuntimeHash2D(class UCanvas* Canvas, const FVector2D& PartitionCanvasOffset, const FVector2D& PartitionCanvasSize)
 {
 	check(GetWorld()->IsGameWorld());
-	GetStreamingPolicy()->ShowDebugInfo(Canvas, PartitionCanvasOffset, PartitionCanvasSize);
+	GetStreamingPolicy()->DrawRuntimeHash2D(Canvas, PartitionCanvasOffset, PartitionCanvasSize);
+}
+
+void UWorldPartition::DrawRuntimeHash3D()
+{
+	check(GetWorld()->IsGameWorld());
+	GetStreamingPolicy()->DrawRuntimeHash3D();
 }
 
 #if WITH_EDITOR
diff --git a/Engine/Source/Runtime/Engine/Private/WorldPartition/WorldPartitionRuntimeSpatialHash.cpp b/Engine/Source/Runtime/Engine/Private/WorldPartition/WorldPartitionRuntimeSpatialHash.cpp
--- a/Engine/Source/Runtime/Engine/Private/WorldPartition/WorldPartitionRuntimeSpatialHash.cpp
+++ b/Engine/Source/Runtime/Engine/Private/WorldPartition/WorldPartitionRuntimeSpatialHash.cpp
@@ -9,7 +9,6 @@
 #include "WorldPartition/WorldPartitionActorDescIterator.h"
 #include "WorldPartition/WorldPartitionStreamingPolicy.h"
 #include "GameFramework/WorldSettings.h"
-#include "Math/Grid2D.h"
 #include "ProfilingDebugging/ScopedTimers.h"
 #include "Engine/World.h"
 #include "Engine/Canvas.h"
@@ -42,7 +41,7 @@ DEFINE_LOG_CATEGORY_STATIC(LogWorldPartitionRuntimeSpatialHash, Log, All);
 class FScopedLoadActorsHelper
 {
 public:
-	FScopedLoadActorsHelper(UWorldPartition* InWorldPartition, const TSet<FGuid>& InActors, bool bSkipEditorOnly)
+	FScopedLoadActorsHelper(UWorldPartition* InWorldPartition, const TArray<FGuid>& InActors, bool bSkipEditorOnly)
 		: WorldPartition(InWorldPartition)
 	{
 		LoadedActors.Reserve(InActors.Num());
@@ -76,11 +75,240 @@ private:
 };
 #endif
 
-/** 
+static int32 GShowRuntimeSpatialHashGridLevel = 0;
+static FAutoConsoleVariableRef CVarShowRuntimeSpatialHashGridLevel(
+	TEXT("WorldPartition.ShowRuntimeSpatialHashGridLevel"),
+	GShowRuntimeSpatialHashGridLevel,
+	TEXT("Used to choose which grid level to display when showing world partition runtime hash."));
+
+static int32 GShowRuntimeSpatialHashGridLevelCount = 1;
+static FAutoConsoleVariableRef CVarShowRuntimeSpatialHashGridLevelCount(
+	TEXT("WorldPartition.ShowRuntimeSpatialHashGridLevelCount"),
+	GShowRuntimeSpatialHashGridLevelCount,
+	TEXT("Used to choose how many grid levels to display when showing world partition runtime hash."));
+
+// ------------------------------------------------------------------------------------------------
+
+/**
   * Square 2D grid helper
   */
 struct FSquare2DGridHelper
 {
+	struct FGrid2D
+	{
+		FVector2D Origin;
+		int32 CellSize;
+		int32 GridSize;
+
+		inline FGrid2D(const FVector2D& InOrigin, int32 InCellSize, int32 InGridSize)
+			: Origin(InOrigin)
+			, CellSize(InCellSize)
+			, GridSize(InGridSize)
+		{}
+
+		/**
+		 * Validate that the coordinates fit the grid size
+		 *
+		 * @return true if the specified coordinates are valid
+		 */
+		inline bool IsValidCoords(const FIntVector2& InCoords) const
+		{
+			return (InCoords.X >= 0) && (InCoords.X < GridSize) && (InCoords.Y >= 0) && (InCoords.Y < GridSize);
+		}
+
+		/**
+		 * Returns the cell bounds
+		 *
+		 * @return true if the specified index was valid
+		 */
+		inline bool GetCellBounds(int32 InIndex, FBox2D& OutBounds) const
+		{
+			if (InIndex >= 0 && InIndex <= (GridSize * GridSize))
+			{
+				const FIntVector2 Coords(InIndex % GridSize, InIndex / GridSize);
+				return GetCellBounds(Coords, OutBounds);
+			}
+
+			return false;
+		}
+
+		/**
+		 * Returns the cell bounds
+		 *
+		 * @return true if the specified coord was valid
+		 */
+		inline bool GetCellBounds(const FIntVector2& InCoords, FBox2D& OutBounds) const
+		{
+			if (IsValidCoords(InCoords))
+			{
+				const FVector2D Min = (FVector2D(Origin) - FVector2D(GridSize * CellSize * 0.5f, GridSize * CellSize * 0.5f)) + FVector2D(InCoords.X * CellSize, InCoords.Y * CellSize);
+				const FVector2D Max = Min + FVector2D(CellSize, CellSize);
+				OutBounds = FBox2D(Min, Max);
+				return true;
+			}
+
+			return false;
+		}
+
+		/**
+		 * Returns the cell coordinates of the provided position
+		 *
+		 * @return true if the position was inside the grid
+		 */
+		inline bool GetCellCoords(const FVector2D& InPos, FIntVector2& OutCoords) const
+		{
+			OutCoords = FIntVector2(
+				FMath::FloorToInt(((InPos.X - Origin.X) / CellSize) + GridSize * 0.5f),
+				FMath::FloorToInt(((InPos.Y - Origin.Y) / CellSize) + GridSize * 0.5f)
+			);
+
+			return IsValidCoords(OutCoords);
+		}
+
+		/**
+		 * Returns the cells coordinates of the provided box
+		 *
+		 * @return true if the bounds was intersecting with the grid
+		 */
+		inline bool GetCellCoords(const FBox2D& InBounds2D, FIntVector2& OutMinCellCoords, FIntVector2& OutMaxCellCoords) const
+		{
+			GetCellCoords(InBounds2D.Min, OutMinCellCoords);
+			if ((OutMinCellCoords.X >= GridSize) || (OutMinCellCoords.Y >= GridSize))
+			{
+				return false;
+			}
+
+			GetCellCoords(InBounds2D.Max, OutMaxCellCoords);
+			if ((OutMaxCellCoords.X < 0) || (OutMaxCellCoords.Y < 0))
+			{
+				return false;
+			}
+
+			OutMinCellCoords.X = FMath::Clamp(OutMinCellCoords.X, 0, GridSize - 1);
+			OutMinCellCoords.Y = FMath::Clamp(OutMinCellCoords.Y, 0, GridSize - 1);
+			OutMaxCellCoords.X = FMath::Clamp(OutMaxCellCoords.X, 0, GridSize - 1);
+			OutMaxCellCoords.Y = FMath::Clamp(OutMaxCellCoords.Y, 0, GridSize - 1);
+
+			return true;
+		}
+
+		/**
+		 * Returns the cell index of the provided coords
+		 *
+		 * @return true if the coords was inside the grid
+		 */
+		inline bool GetCellIndex(const FIntVector2& InCoords, uint32& OutIndex) const
+		{
+			if (IsValidCoords(InCoords))
+			{
+				OutIndex = (InCoords.Y * GridSize) + InCoords.X;
+				return true;
+			}
+
+			return false;
+		}
+
+		/**
+		 * Returns the cell index of the provided position
+		 *
+		 * @return true if the position was inside the grid
+		 */
+		inline bool GetCellIndex(const FVector& InPos, uint32& OutIndex) const
+		{
+			FIntVector2 Coords = FIntVector2(
+				FMath::FloorToInt(((InPos.X - Origin.X) / CellSize) + GridSize * 0.5f),
+				FMath::FloorToInt(((InPos.Y - Origin.Y) / CellSize) + GridSize * 0.5f)
+			);
+
+			return GetCellIndex(Coords, OutIndex);
+		}
+
+		/**
+		 * Get the number of intersecting cells of the provided box
+		 *
+		 * @return the number of intersecting cells
+		 */
+		int32 GetNumIntersectingCells(const FBox& InBox) const
+		{
+			FIntVector2 MinCellCoords;
+			FIntVector2 MaxCellCoords;
+			const FBox2D Bounds2D(FVector2D(InBox.Min), FVector2D(InBox.Max));
+
+			if (GetCellCoords(Bounds2D, MinCellCoords, MaxCellCoords))
+			{
+				return (MaxCellCoords.X - MinCellCoords.X + 1) * (MaxCellCoords.Y - MinCellCoords.Y + 1);
+			}
+
+			return 0;
+		}
+
+		/**
+		 * Runs a function on all intersecting cells for the provided box
+		 *
+		 * @return the number of intersecting cells
+		 */
+		int32 ForEachIntersectingCellsBreakable(const FBox& InBox, TFunctionRef<bool(const FIntVector2&)> InOperation) const
+		{
+			int32 NumCells = 0;
+
+			FIntVector2 MinCellCoords;
+			FIntVector2 MaxCellCoords;
+			const FBox2D Bounds2D(FVector2D(InBox.Min), FVector2D(InBox.Max));
+
+			if (GetCellCoords(Bounds2D, MinCellCoords, MaxCellCoords))
+			{
+				for (int32 y = MinCellCoords.Y; y <= MaxCellCoords.Y; y++)
+				{
+					for (int32 x = MinCellCoords.X; x <= MaxCellCoords.X; x++)
+					{
+						if (!InOperation(FIntVector2(x, y)))
+						{
+							return NumCells;
+						}
+						++NumCells;
+					}
+				}
+			}
+
+			return NumCells;
+		}
+
+		int32 ForEachIntersectingCells(const FBox& InBox, TFunctionRef<void(const FIntVector2&)> InOperation) const
+		{
+			return ForEachIntersectingCellsBreakable(InBox, [InOperation](const FIntVector2& Vector) { InOperation(Vector); return true; });
+		}
+
+		/**
+		 * Runs a function on all intersecting cells for the provided sphere
+		 *
+		 * @return the number of intersecting cells
+		 */
+		int32 ForEachIntersectingCells(const FSphere& InSphere, TFunctionRef<void(const FIntVector2&)> InOperation) const
+		{
+			int32 NumCells = 0;
+
+			// @todo_ow: rasterize circle instead?
+			const FBox Box(InSphere.Center - FVector(InSphere.W), InSphere.Center + FVector(InSphere.W));
+
+			ForEachIntersectingCells(Box, [this, &InSphere, &InOperation, &NumCells](const FIntVector2& Coords)
+				{
+					const int32 CellIndex = Coords.Y * GridSize + Coords.X;
+
+					FBox2D CellBounds;
+					GetCellBounds(CellIndex, CellBounds);
+
+					FVector2D Delta = FVector2D(InSphere.Center) - FVector2D::Max(CellBounds.GetCenter() - CellBounds.GetExtent(), FVector2D::Min(FVector2D(InSphere.Center), CellBounds.GetCenter() + CellBounds.GetExtent()));
+					if ((Delta.X * Delta.X + Delta.Y * Delta.Y) < (InSphere.W * InSphere.W))
+					{
+						InOperation(Coords);
+						NumCells++;
+					}
+				});
+
+			return NumCells;
+		}
+	};
+
 	struct FGridLevel : public FGrid2D
 	{
 		struct FGridCell
@@ -123,56 +351,43 @@ struct FSquare2DGridHelper
 		TArray<FGridCell> Cells;
 	};
 
-	inline FSquare2DGridHelper(int32 InNumLevels, const FVector& InOrigin, int32 InCellSize, int32 InGridSize)
+	FSquare2DGridHelper(int32 InNumLevels, const FVector& InOrigin, int32 InCellSize, int32 InGridSize)
 		: Origin(InOrigin)
 		, CellSize(InCellSize)
 		, GridSize(InGridSize)
 	{
 		Levels.Reserve(InNumLevels);
 
-		int32 CurrentCellSize = InCellSize;
-		int32 CurrentGridSize = InGridSize;
+		int32 CurrentCellSize = CellSize;
+		int32 CurrentGridSize = GridSize;
 
+		const FVector2D BaseLevelOffset = 0.5 * FVector2D(CellSize, CellSize);
 		for (int32 Level = 0; Level < InNumLevels; ++Level)
 		{
-			Levels.Emplace(FVector2D(InOrigin), CurrentCellSize, CurrentGridSize);
+			// Add offset on origin based on level's cell size to break pattern of perfectly aligned cell edges at multiple level.
+			// This will prevent weird artefact during actor promotion.
+			// Apply base level offset so that first level isn't offset.
+			const FVector2D GridLevelOffset = 0.5f * FVector2D(CurrentCellSize, CurrentCellSize) - BaseLevelOffset;
+			const FVector2D LevelOrigin = FVector2D(InOrigin) + GridLevelOffset;
+
+			Levels.Emplace(LevelOrigin, CurrentCellSize, CurrentGridSize);
 
 			CurrentCellSize <<= 1;
 			CurrentGridSize >>= 1;
 		}
 	}
 
-	/**
-	 * Returns the lowest grid level
-	 */
-	inline FGridLevel& GetLowestLevel()
-	{
-		return Levels[0];
-	}
+	// Returns the lowest grid level
+	inline FGridLevel& GetLowestLevel() { return Levels[0]; }
 
-	/**
-	 * Returns the always loaded (top level) cell
-	 */
-	inline FGridLevel::FGridCell& GetAlwaysLoadedCell()
-	{
-		return Levels.Last().Cells[0];
-	}
+	// Returns the always loaded (top level) cell
+	inline FGridLevel::FGridCell& GetAlwaysLoadedCell() { return Levels.Last().Cells[0]; }
 
-	/**
-	 * Returns the always loaded (top level) cell
-	 */
-	inline const FGridLevel::FGridCell& GetAlwaysLoadedCell() const
-	{
-		return Levels.Last().Cells[0];
-	}
+	// Returns the always loaded (top level) cell
+	inline const FGridLevel::FGridCell& GetAlwaysLoadedCell() const { return Levels.Last().Cells[0]; }
 
-	/**
-	 * Returns the cell at the given coord
-	 */
-	inline const FGridLevel::FGridCell& GetCell(const FIntVector& InCoords) const
-	{
-		return Levels[InCoords.Z].GetCell(FIntVector2(InCoords.X, InCoords.Y));
-	}
+	// Returns the cell at the given coord
+	inline const FGridLevel::FGridCell& GetCell(const FIntVector& InCoords) const { return Levels[InCoords.Z].GetCell(FIntVector2(InCoords.X, InCoords.Y)); }
 
 	/**
 	 * Returns the cell bounds
@@ -185,13 +400,10 @@ struct FSquare2DGridHelper
 		{
 			return Levels[InCoords.Z].GetCellBounds(FIntVector2(InCoords.X, InCoords.Y), OutBounds);
 		}
-
 		return false;
 	}
 
-	/**
-	 * Runs a function on all cells
-	 */
+	// Runs a function on all cells
 	void ForEachCells(TFunctionRef<void(const FIntVector&)> InOperation) const
 	{
 		for (int32 Level = 0; Level < Levels.Num(); Level++)
@@ -216,7 +428,7 @@ struct FSquare2DGridHelper
 	int32 ForEachIntersectingCells(const FBox& InBox, TFunctionRef<void(const FIntVector&)> InOperation) const
 	{
 		int32 NumCells = 0;
-	
+
 		for (int32 Level = 0; Level < Levels.Num(); Level++)
 		{
 			NumCells += Levels[Level].ForEachIntersectingCells(InBox, [InBox, Level, InOperation](const FIntVector2& Coord) { InOperation(FIntVector(Coord.X, Coord.Y, Level)); });
@@ -243,67 +455,15 @@ struct FSquare2DGridHelper
 	}
 
 #if WITH_EDITOR
-	/**
-	 * Perform actors promotion
-	 * When an actor is referenced by multiple cells, promote it to the higher level parent cell.
-	 */
-	void PerformActorsPromotion()
+	// Validates that actor is not referenced by multiple cells
+	void ValidateSingleActorReferer()
 	{
-		UE_SCOPED_TIMER(TEXT("PerformActorsPromotion"), LogWorldPartitionRuntimeSpatialHash, Log);
-
-		for (int32 Level = 0; Level < Levels.Num() - 1; Level++)
-		{
-			const int32 CurrentGridSize = Levels[Level].GridSize;
-
-			// First pass: gather actors usage
-			TMap<FGuid, int32> ActorUsage;
-			for (int32 y = 0; y < CurrentGridSize; y++)
-			{
-				for (int32 x = 0; x < CurrentGridSize; x++)
-				{
-					FGridLevel::FGridCell& ThisCell = Levels[Level].GetCell(FIntVector2(x, y));
-					for (const FGuid& ActorGuid : ThisCell.Actors)
-					{
-						ActorUsage.FindOrAdd(ActorGuid, 0)++;
-					}
-				}
-			}
-
-			// Second pass: promote actors to higher level
-			for (int32 y = 0; y < CurrentGridSize; y++)
-			{
-				for (int32 x = 0; x < CurrentGridSize; x++)
-				{
-					FGridLevel::FGridCell& ThisCell = Levels[Level].GetCell(FIntVector2(x, y));
-					for (auto ActorIt = ThisCell.Actors.CreateIterator(); ActorIt; ++ActorIt)
-					{
-						const FGuid& ActorGuid = *ActorIt;
-						const int32 ActorUse = ActorUsage.FindChecked(ActorGuid);
+		UE_SCOPED_TIMER(TEXT("ValidateSingleActorReferer"), LogWorldPartitionRuntimeSpatialHash, Log);
 
-						// @todo_ow: Use StreamingPolicy to get Minimum Actor count for promotion
-						if (ActorUse > 1)
-						{
-							FIntVector2 ParentCellCoord(x/2, y/2);
-							int32 ParentCellLevel = Level + 1;
-							FGridLevel::FGridCell& ParentCell = Levels[ParentCellLevel].GetCell(ParentCellCoord);
-							if (!ParentCell.Actors.Contains(ActorGuid))
-							{
-								ParentCell.Actors.Add(ActorGuid);
-								UE_LOG(LogWorldPartitionRuntimeSpatialHash, Verbose, TEXT("Promoted %s(%d) from cell(X%02d, Y%02d, Level=%d) to cell(X%02d, Y%02d, Level=%d)"), *ActorGuid.ToString(EGuidFormats::UniqueObjectGuid), ActorUse, x, y, Level, ParentCellCoord.X, ParentCellCoord.Y, ParentCellLevel);
-							}
-							ActorIt.RemoveCurrent();
-						}
-					}
-				}
-			}
-		}
-
-		// Third pass: validate actor usage on all levels
+		TSet<FGuid> ActorUsage;
 		for (int32 Level = 0; Level < Levels.Num() - 1; Level++)
 		{
 			const int32 CurrentGridSize = Levels[Level].GridSize;
-
-			TSet<FGuid> ActorUsage;
 			for (int32 y = 0; y < CurrentGridSize; y++)
 			{
 				for (int32 x = 0; x < CurrentGridSize; x++)
@@ -328,6 +488,236 @@ public:
 	TArray<FGridLevel> Levels;
 };
 
+// ------------------------------------------------------------------------------------------------
+FSpatialHashStreamingGrid::FSpatialHashStreamingGrid()
+	: Origin(ForceInitToZero)
+	, CellSize(0)
+	, GridSize(0)
+	, LoadingRange(0.0f)
+#if WITH_EDITOR
+	, DebugColor(ForceInitToZero)
+	, OverrideLoadingRange(0)
+#endif
+	, GridHelper(nullptr)
+{
+}
+
+FSpatialHashStreamingGrid::~FSpatialHashStreamingGrid()
+{
+	if (GridHelper)
+	{
+		delete GridHelper;
+	}
+}
+
+const FSquare2DGridHelper& FSpatialHashStreamingGrid::GetGridHelper() const
+{
+	if (!GridHelper)
+	{
+		GridHelper = new FSquare2DGridHelper(GridLevels.Num(), Origin, CellSize, GridSize);
+	}
+	else
+	{
+		check(GridHelper->Levels.Num() == GridLevels.Num());
+		check(GridHelper->Origin == Origin);
+		check(GridHelper->CellSize == CellSize);
+		check(GridHelper->GridSize == GridSize);
+	}
+	return *GridHelper;
+}
+
+void FSpatialHashStreamingGrid::GetCells(const TArray<FWorldPartitionStreamingSource>& Sources, TSet<const UWorldPartitionRuntimeCell*>& Cells) const
+{
+	const FSquare2DGridHelper& Helper = GetGridHelper();
+	for (const FWorldPartitionStreamingSource& Source : Sources)
+	{
+		const FSphere GridSphere(Source.Location, GetLoadingRange());
+		Helper.ForEachIntersectingCells(GridSphere, [&](const FIntVector& Coords)
+		{
+			if (UWorldPartitionRuntimeSpatialHashCell* Cell = GridLevels[Coords.Z].GridCells[Coords.Y * Helper.Levels[Coords.Z].GridSize + Coords.X])
+			{
+				Cells.Add(Cell);
+			}
+		});
+	}
+
+	GetAlwaysLoadedCells(Cells);
+}
+
+void FSpatialHashStreamingGrid::GetAlwaysLoadedCells(TSet<const UWorldPartitionRuntimeCell*>& Cells) const
+{
+	if (GridLevels.Num() > 0)
+	{
+		const int32 TopLevel = GridLevels.Num() - 1;
+		check(GridLevels[TopLevel].GridCells.Num() == 1);
+		if (UWorldPartitionRuntimeSpatialHashCell* Cell = GridLevels[TopLevel].GridCells[0])
+		{
+			Cells.Add(Cell);
+		}
+	}
+}
+
+void FSpatialHashStreamingGrid::Draw3D(UWorld* World, const TArray<FWorldPartitionStreamingSource>& Sources) const
+{
+	const FSquare2DGridHelper& Helper = GetGridHelper();
+	int32 MinGridLevel = FMath::Clamp<int32>(GShowRuntimeSpatialHashGridLevel, 0, GridLevels.Num() - 1);
+	int32 MaxGridLevel = FMath::Clamp<int32>(MinGridLevel + GShowRuntimeSpatialHashGridLevelCount - 1, 0, GridLevels.Num() - 1);
+	const float GridViewMinimumSizeInCellCount = 5.f;
+	const float GridViewLoadingRangeExtentRatio = 1.5f;
+	const float Radius = GetLoadingRange();
+	const float GridSideDistance = FMath::Max((2.f * Radius * GridViewLoadingRangeExtentRatio), CellSize * GridViewMinimumSizeInCellCount);
+
+	for (const FWorldPartitionStreamingSource& Source : Sources)
+	{
+		FVector StartTrace = Source.Location + FVector(0.f, 0.f, 100.f);
+		FVector EndTrace = StartTrace - FVector(0.f, 0.f, 1000000.f);
+		float Z = Source.Location.Z;
+		FHitResult Hit;
+		if (World->LineTraceSingleByObjectType(Hit, StartTrace, EndTrace, FCollisionObjectQueryParams(ECC_WorldStatic), FCollisionQueryParams(SCENE_QUERY_STAT(DebugWorldPartitionTrace), true)))
+		{
+			Z = Hit.ImpactPoint.Z;
+		}
+
+		FSphere Sphere(Source.Location, GridSideDistance * 0.5f);
+		const FBox Region(Sphere.Center - Sphere.W, Sphere.Center + Sphere.W);
+		for (int32 GridLevel = MinGridLevel; GridLevel <= MaxGridLevel; ++GridLevel)
+		{
+			Helper.Levels[GridLevel].ForEachIntersectingCells(Region, [&](const FIntVector2& Coords)
+			{
+				FBox2D CellWorldBounds;
+				Helper.Levels[GridLevel].GetCellBounds(FIntVector2(Coords.X, Coords.Y), CellWorldBounds);
+				
+				FVector BoundsExtent(CellWorldBounds.GetExtent(), 100.f);
+				const UWorldPartitionRuntimeSpatialHashCell* Cell = Cast<const UWorldPartitionRuntimeSpatialHashCell>(GridLevels[GridLevel].GridCells[Coords.Y * Helper.Levels[GridLevel].GridSize + Coords.X]);
+				FColor CellColor = Cell ? Cell->GetDebugColor().ToFColor(false).WithAlpha(32) : FColor(0, 0, 0, 32);
+				FVector BoundsOrigin(CellWorldBounds.GetCenter(), Z);
+				DrawDebugSolidBox(World, BoundsOrigin, BoundsExtent, CellColor, false, -1.f, 255);
+				DrawDebugBox(World, BoundsOrigin, BoundsExtent, CellColor.WithAlpha(255), false, -1.f, 255, 10.f);
+			});
+		}
+
+		// Draw Loading Ranges
+		FVector SphereLocation(FVector2D(Source.Location), Z);
+		DrawDebugSphere(World, SphereLocation, Radius, 32, FColor::White, false, -1.f, 0, 20.f);
+	}
+}
+
+void FSpatialHashStreamingGrid::Draw2D(UCanvas* Canvas, const TArray<FWorldPartitionStreamingSource>& Sources, const FBox& Region, const FBox2D& GridScreenBounds, TFunctionRef<FVector2D(const FVector2D&)> WorldToScreen) const
+{
+	FCanvas* CanvasObject = Canvas->Canvas;
+	int32 MinGridLevel = FMath::Clamp<int32>(GShowRuntimeSpatialHashGridLevel, 0, GridLevels.Num() - 1);
+	int32 MaxGridLevel = FMath::Clamp<int32>(MinGridLevel + GShowRuntimeSpatialHashGridLevelCount - 1, 0, GridLevels.Num() - 1);
+
+	for (int32 GridLevel = MinGridLevel; GridLevel <= MaxGridLevel; ++GridLevel)
+	{
+		// Draw X/Y Axis
+		{
+			FCanvasLineItem Axis;
+			Axis.LineThickness = 3;
+			{
+				Axis.SetColor(FLinearColor::Green);
+				FVector2D LineStart = WorldToScreen(FVector2D(-163840.f, 0.f));
+				FVector2D LineEnd = WorldToScreen(FVector2D(163840.f, 0.f));
+				LineStart.X = FMath::Clamp(LineStart.X, GridScreenBounds.Min.X, GridScreenBounds.Max.X);
+				LineStart.Y = FMath::Clamp(LineStart.Y, GridScreenBounds.Min.Y, GridScreenBounds.Max.Y);
+				LineEnd.X = FMath::Clamp(LineEnd.X, GridScreenBounds.Min.X, GridScreenBounds.Max.X);
+				LineEnd.Y = FMath::Clamp(LineEnd.Y, GridScreenBounds.Min.Y, GridScreenBounds.Max.Y);
+				Axis.Draw(CanvasObject, LineStart, LineEnd);
+			}
+			{
+				Axis.SetColor(FLinearColor::Red);
+				FVector2D LineStart = WorldToScreen(FVector2D(0.f, -163840.f));
+				FVector2D LineEnd = WorldToScreen(FVector2D(0.f, 163840.f));
+				LineStart.X = FMath::Clamp(LineStart.X, GridScreenBounds.Min.X, GridScreenBounds.Max.X);
+				LineStart.Y = FMath::Clamp(LineStart.Y, GridScreenBounds.Min.Y, GridScreenBounds.Max.Y);
+				LineEnd.X = FMath::Clamp(LineEnd.X, GridScreenBounds.Min.X, GridScreenBounds.Max.X);
+				LineEnd.Y = FMath::Clamp(LineEnd.Y, GridScreenBounds.Min.Y, GridScreenBounds.Max.Y);
+				Axis.Draw(CanvasObject, LineStart, LineEnd);
+			}
+		}
+
+		// Draw Grid cells at desired grid level
+		const FSquare2DGridHelper& Helper = GetGridHelper();
+		Helper.Levels[GridLevel].ForEachIntersectingCells(Region, [&](const FIntVector2& Coords)
+		{
+			FBox2D CellWorldBounds;
+			Helper.Levels[GridLevel].GetCellBounds(FIntVector2(Coords.X, Coords.Y), CellWorldBounds);
+			FBox2D CellScreenBounds = FBox2D(WorldToScreen(CellWorldBounds.Min), WorldToScreen(CellWorldBounds.Max));
+			// Clamp inside grid bounds
+			if (!GridScreenBounds.IsInside(CellScreenBounds))
+			{
+				CellScreenBounds.Min.X = FMath::Clamp(CellScreenBounds.Min.X, GridScreenBounds.Min.X, GridScreenBounds.Max.X);
+				CellScreenBounds.Min.Y = FMath::Clamp(CellScreenBounds.Min.Y, GridScreenBounds.Min.Y, GridScreenBounds.Max.Y);
+				CellScreenBounds.Max.X = FMath::Clamp(CellScreenBounds.Max.X, GridScreenBounds.Min.X, GridScreenBounds.Max.X);
+				CellScreenBounds.Max.Y = FMath::Clamp(CellScreenBounds.Max.Y, GridScreenBounds.Min.Y, GridScreenBounds.Max.Y);
+			}
+			else
+			{
+				FString GridInfoText = FString::Printf(TEXT("X%02d_Y%02d"), Coords.X, Coords.Y);
+				float TextWidth, TextHeight;
+				Canvas->SetDrawColor(255, 255, 0);
+				Canvas->StrLen(GEngine->GetTinyFont(), GridInfoText, TextWidth, TextHeight);
+				FVector2D CellSize = CellScreenBounds.GetSize();
+				if (TextWidth < CellSize.X && TextHeight < CellSize.Y)
+				{
+					FVector2D GridInfoPos = CellScreenBounds.GetCenter() - FVector2D(TextWidth / 2, TextHeight / 2);
+					Canvas->DrawText(GEngine->GetTinyFont(), GridInfoText, GridInfoPos.X, GridInfoPos.Y);
+				}
+			}
+
+			const UWorldPartitionRuntimeSpatialHashCell* Cell = Cast<const UWorldPartitionRuntimeSpatialHashCell>(GridLevels[GridLevel].GridCells[Coords.Y * Helper.Levels[GridLevel].GridSize + Coords.X]);
+			FLinearColor CellColor = Cell ? Cell->GetDebugColor() : FLinearColor(0.f, 0.f, 0.f, 0.25f);
+			FCanvasTileItem Item(CellScreenBounds.Min, GWhiteTexture, CellScreenBounds.GetSize(), CellColor);
+			Item.BlendMode = SE_BLEND_Translucent;
+			Canvas->DrawItem(Item);
+
+			FCanvasBoxItem Box(CellScreenBounds.Min, CellScreenBounds.GetSize());
+			Box.SetColor(CellColor);
+			Box.BlendMode = SE_BLEND_Translucent;
+			Canvas->DrawItem(Box);
+		});
+
+		// Draw Loading Ranges
+		float Range = GetLoadingRange();
+
+		FCanvasLineItem LineItem;
+		LineItem.LineThickness = 2;
+		LineItem.SetColor(FLinearColor::White);
+
+		for (const FWorldPartitionStreamingSource& Source : Sources)
+		{
+			TArray<FVector> LinePoints;
+			LinePoints.SetNum(2);
+
+			float Sin, Cos;
+			FMath::SinCos(&Sin, &Cos, (63.0f / 64.0f) * 2.0f * PI);
+			FVector2D LineStart(Sin * Range, Cos * Range);
+
+			for (int32 i = 0; i < 64; i++)
+			{
+				FMath::SinCos(&Sin, &Cos, (i / 64.0f) * 2.0f * PI);
+				FVector2D LineEnd(Sin * Range, Cos * Range);
+				LineItem.Draw(CanvasObject, WorldToScreen(FVector2D(Source.Location) + LineStart), WorldToScreen(FVector2D(Source.Location) + LineEnd));
+				LineStart = LineEnd;
+			}
+
+			FVector2D SourceDir = FVector2D(Source.Rotation.Vector());
+			if (SourceDir.Size())
+			{
+				SourceDir.Normalize();
+				FVector2D ConeCenter(FVector2D(Source.Location));
+				LineItem.Draw(CanvasObject, WorldToScreen(ConeCenter), WorldToScreen(ConeCenter + SourceDir * Range));
+			}
+		}
+
+		FCanvasBoxItem Box(GridScreenBounds.Min, GridScreenBounds.GetSize());
+		Box.SetColor(DebugColor);
+		Canvas->DrawItem(Box);
+	}
+}
+
+// ------------------------------------------------------------------------------------------------
+
 #if WITH_EDITOR
 
 static FSquare2DGridHelper GetPartitionedActors(const UWorldPartition* WorldPartition, const FBox& WorldBounds, const FSpatialHashRuntimeGrid& Grid, const TArray<UWorldPartition::FActorCluster>& GridActors)
@@ -339,12 +729,17 @@ static FSquare2DGridHelper GetPartitionedActors(const UWorldPartition* WorldPart
 	FVector GridOrigin = FVector::ZeroVector;
 	int32 GridSize = 1;
 	int32 GridLevelCount = 1;
-	
-	// If World bounds is valid, compute Grid's origin, size and level count based on it
-	const float WorldBoundsMaxExtent = WorldBounds.IsValid ? WorldBounds.GetExtent().GetMax() : 0.f;
+
+	float WorldBoundsMaxExtent = 0.f;
+	// If World bounds is valid, compute Grid's size and level count based on it
+	if (WorldBounds.IsValid)
+	{
+		const FVector2D DistMin = FMath::Abs(FVector2D(WorldBounds.Min - GridOrigin));
+		const FVector2D DistMax = FMath::Abs(FVector2D(WorldBounds.Max - GridOrigin));
+		WorldBoundsMaxExtent = FMath::Max(DistMin.GetMax(), DistMax.GetMax());
+	}
 	if (WorldBoundsMaxExtent > 0.f)
 	{
-		GridOrigin = WorldBounds.GetCenter();
 		GridSize = 2.f * FMath::CeilToFloat(WorldBoundsMaxExtent / GridCellSize);
 		if (!FMath::IsPowerOfTwo(GridSize))
 		{
@@ -399,17 +794,30 @@ static FSquare2DGridHelper GetPartitionedActors(const UWorldPartition* WorldPart
 			case EActorGridPlacement::Bounds:
 			{
 				const FBox ActorClusterBounds = WorldPartition->GetActorClusterBounds(ActorCluster);
-				if (!PartitionedActors.GetLowestLevel().ForEachIntersectingCells(ActorClusterBounds, [&](const FIntVector2& Coords)
+				// Find grid level cell that encompasses the actor cluster and put actors in it.
+				bool bFoundCell = false;
+				for (FSquare2DGridHelper::FGridLevel& GridLevel : PartitionedActors.Levels)
 				{
-					PartitionedActors.GetLowestLevel().GetCell(Coords).Actors.Append(ActorCluster.Actors);
-				}))
+					int32 IntersectingCellCount = 0;
+					GridLevel.ForEachIntersectingCellsBreakable(ActorClusterBounds, [&IntersectingCellCount](const FIntVector2& Coords) { return ++IntersectingCellCount <= 1; });
+					if (IntersectingCellCount == 1)
+					{
+						GridLevel.ForEachIntersectingCells(ActorClusterBounds, [&GridLevel,&ActorCluster](const FIntVector2& Coords)
+						{
+							GridLevel.GetCell(Coords).Actors.Append(ActorCluster.Actors);
+						});
+						bFoundCell = true;
+						break;
+					}
+				}
+				if (!bFoundCell)
 				{
 					GridPlacement = EActorGridPlacement::AlwaysLoaded;
 					bAlwaysLoadedPromotedOutOfGrid = true;
 				}
 				break;
 			}
-			default:
+			default:	
 				check(0);
 			}
 		}
@@ -444,8 +852,8 @@ static FSquare2DGridHelper GetPartitionedActors(const UWorldPartition* WorldPart
 		}
 	}
 
-	// Perform actor promotion
-	PartitionedActors.PerformActorsPromotion();
+	// Perform validation
+	PartitionedActors.ValidateSingleActorReferer();
 	
 	return PartitionedActors;
 }
@@ -620,6 +1028,7 @@ bool UWorldPartitionRuntimeSpatialHash::CreateStreamingGrid(const FSpatialHashRu
 	// Move actors into the final streaming grids
 	CurrentStreamingGrid.GridLevels.Reserve(PartionedActors.Levels.Num());
 
+	TArray<FGuid> FilteredActors;
 	int32 Level = INDEX_NONE;
 	for (const FSquare2DGridHelper::FGridLevel& TempLevel : PartionedActors.Levels)
 	{
@@ -633,7 +1042,21 @@ bool UWorldPartitionRuntimeSpatialHash::CreateStreamingGrid(const FSpatialHashRu
 		for (const FSquare2DGridHelper::FGridLevel::FGridCell& TempCell : TempLevel.Cells)
 		{
 			CellIndex++;
-			if (!TempCell.Actors.Num())
+			
+			FilteredActors.SetNum(0, false);
+			FilteredActors.Reset(TempCell.Actors.Num());
+			if (TempCell.Actors.Num())
+			{
+				Algo::TransformIf(TempCell.Actors, FilteredActors, [&WorldPartition](const FGuid& ActorGuid)
+				{
+					const FWorldPartitionActorDesc* ActorDesc = WorldPartition->GetActorDesc(ActorGuid);
+					const bool bShouldStripActorFromStreaming = ActorDesc->GetActorIsEditorOnly();
+					UE_CLOG(bShouldStripActorFromStreaming, LogWorldPartitionRuntimeSpatialHash, Verbose, TEXT("Stripping Actor %s (%s) from streaming grid"), *(ActorDesc->GetActorPath().ToString()), *ActorGuid.ToString(EGuidFormats::UniqueObjectGuid));
+					return !bShouldStripActorFromStreaming;
+				}, [](const FGuid& ActorGuid) { return ActorGuid;});
+			}
+
+			if (!FilteredActors.Num())
 			{
 				GridLevel.GridCells.Add(nullptr);
 				continue;
@@ -653,31 +1076,21 @@ bool UWorldPartitionRuntimeSpatialHash::CreateStreamingGrid(const FSpatialHashRu
 			verify(TempLevel.GetCellBounds(FIntVector2(CellCoordX, CellCoordY), Bounds));
 			StreamingCell->Position = FVector(Bounds.GetCenter(), 0.f);
 
-			UE_LOG(LogWorldPartitionRuntimeSpatialHash, Verbose, TEXT("Cell%s %s Actors = %d"), bIsCellAlwaysLoaded ? TEXT(" (AlwaysLoaded)") : TEXT(""), *StreamingCell->GetName(), TempCell.Actors.Num());
+			UE_LOG(LogWorldPartitionRuntimeSpatialHash, Verbose, TEXT("Cell%s %s Actors = %d"), bIsCellAlwaysLoaded ? TEXT(" (AlwaysLoaded)") : TEXT(""), *StreamingCell->GetName(), FilteredActors.Num());
 
 			// Keep track of all AWorldPartitionHLOD actors referenced by this cell
 			TSet<FGuid> ReferencedHLODActors;
 
-			for (const FGuid& ActorGuid : TempCell.Actors)
+			for (const FGuid& ActorGuid : FilteredActors)
 			{
 				const FWorldPartitionActorDesc* ActorDesc = WorldPartition->GetActorDesc(ActorGuid);
-				const bool bShouldStripActorFromStreaming = ActorDesc->GetActorIsEditorOnly();
-
-				if (!bShouldStripActorFromStreaming)
+				FGuid ParentHLOD = CachedHLODParents.FindRef(ActorGuid);
+				if (ParentHLOD.IsValid())
 				{
-					FGuid ParentHLOD = CachedHLODParents.FindRef(ActorGuid);
-					if (ParentHLOD.IsValid())
-					{
-						ReferencedHLODActors.Add(ParentHLOD);
-					}
-
-					StreamingCell->AddActorToCell(ActorDesc->GetActorPackage(), ActorDesc->GetActorPath());
-					UE_LOG(LogWorldPartitionRuntimeSpatialHash, Verbose, TEXT("  Actor : %s (%s) Origin(%s)"), *(ActorDesc->GetActorPath().ToString()), *ActorGuid.ToString(EGuidFormats::UniqueObjectGuid), *FVector2D(ActorDesc->GetOrigin()).ToString());
-				}
-				else
-				{
-					UE_LOG(LogWorldPartitionRuntimeSpatialHash, Verbose, TEXT("  Stripping Actor %s (%s) from streaming grid"), *(ActorDesc->GetActorPath().ToString()), *ActorGuid.ToString(EGuidFormats::UniqueObjectGuid));
+					ReferencedHLODActors.Add(ParentHLOD);
 				}
+				StreamingCell->AddActorToCell(ActorDesc->GetActorPackage(), ActorDesc->GetActorPath());
+				UE_LOG(LogWorldPartitionRuntimeSpatialHash, Verbose, TEXT("  Actor : %s (%s) Origin(%s)"), *(ActorDesc->GetActorPath().ToString()), *ActorGuid.ToString(EGuidFormats::UniqueObjectGuid), *FVector2D(ActorDesc->GetOrigin()).ToString());
 			}
 
 			if (ReferencedHLODActors.Num() > 0)
@@ -690,7 +1103,7 @@ bool UWorldPartitionRuntimeSpatialHash::CreateStreamingGrid(const FSpatialHashRu
 			
 			if (Mode == EWorldPartitionStreamingMode::RuntimeStreamingCells)
 			{
-				FScopedLoadActorsHelper LoadCellActors(WorldPartition, TempCell.Actors, /*bSkipEditorOnly*/true);
+				FScopedLoadActorsHelper LoadCellActors(WorldPartition, FilteredActors, /*bSkipEditorOnly*/true);
 				UE_LOG(LogWorldPartitionRuntimeSpatialHash, Log, TEXT("Creating runtime streaming cells %s."), *StreamingCell->GetName());
 				if (!StreamingCell->CreateCellForCook())
 				{
@@ -803,14 +1216,8 @@ bool UWorldPartitionRuntimeSpatialHash::GenerateHLOD()
 }
 #endif
 
-static int32 GShowRuntimeSpatialHashGridLevel = 0;
-static FAutoConsoleVariableRef CVarGuardBandMultiplier(
-	TEXT("WorldPartitionRuntimeSpatialHash.ShowGridLevel"),
-	GShowRuntimeSpatialHashGridLevel,
-	TEXT("Used to choose which grid level to display when showing world partition runtime grids."));
-
 FAutoConsoleCommand UWorldPartitionRuntimeSpatialHash::OverrideLoadingRangeCommand(
-	TEXT("WorldPartitionRuntimeSpatialHash.OverrideLoadingRange"),
+	TEXT("WorldPartition.OverrideRuntimeSpatialHashLoadingRange"),
 	TEXT("Sets runtime loading range. Args -grid=[index] -range=[override_loading_range]"),
 	FConsoleCommandWithArgsDelegate::CreateLambda([](const TArray<FString>& Args)
 	{
@@ -870,7 +1277,7 @@ int32 UWorldPartitionRuntimeSpatialHash::GetStreamingCells(const TArray<FWorldPa
 		// Get always loaded cells
 		for (const FSpatialHashStreamingGrid& StreamingGrid : StreamingGrids)
 		{
-			GetAlwaysLoadedStreamingCells(StreamingGrid, Cells);
+			StreamingGrid.GetAlwaysLoadedCells(Cells);
 		}
 	}
 	else
@@ -878,43 +1285,13 @@ int32 UWorldPartitionRuntimeSpatialHash::GetStreamingCells(const TArray<FWorldPa
 		// Get cells based on streaming sources
 		for (const FSpatialHashStreamingGrid& StreamingGrid : StreamingGrids)
 		{
-			for (const FWorldPartitionStreamingSource& Source : Sources)
-			{
-				GetStreamingCells(Source.Location, StreamingGrid, Cells);
-			}
+			StreamingGrid.GetCells(Sources, Cells);
 		}
 	}
 
 	return Cells.Num();
 }
 
-void UWorldPartitionRuntimeSpatialHash::GetStreamingCells(const FVector& Position, const FSpatialHashStreamingGrid& StreamingGrid, TSet<const UWorldPartitionRuntimeCell*>& Cells) const
-{
-	const FSphere GridSphere(Position, StreamingGrid.GetLoadingRange());
-	const FSquare2DGridHelper GridHelper(StreamingGrid.GridLevels.Num(), StreamingGrid.Origin, StreamingGrid.CellSize, StreamingGrid.GridSize);
-
-	GridHelper.ForEachIntersectingCells(GridSphere, [&](const FIntVector& Coords)
-	{
-		if (UWorldPartitionRuntimeSpatialHashCell* Cell = StreamingGrid.GridLevels[Coords.Z].GridCells[Coords.Y * GridHelper.Levels[Coords.Z].GridSize + Coords.X])
-		{
-			Cells.Add(Cell);
-		}
-	});
-}
-
-void UWorldPartitionRuntimeSpatialHash::GetAlwaysLoadedStreamingCells(const FSpatialHashStreamingGrid& StreamingGrid, TSet<const UWorldPartitionRuntimeCell*>& Cells) const
-{
-	if (StreamingGrid.GridLevels.Num() > 0)
-	{
-		const int32 TopLevel = StreamingGrid.GridLevels.Num() - 1;
-		check(StreamingGrid.GridLevels[TopLevel].GridCells.Num() == 1);
-		if (UWorldPartitionRuntimeSpatialHashCell* Cell = StreamingGrid.GridLevels[TopLevel].GridCells[0])
-		{
-			Cells.Add(Cell);
-		}
-	}
-}
-
 void UWorldPartitionRuntimeSpatialHash::SortStreamingCellsByDistance(const TSet<const UWorldPartitionRuntimeCell*>& InCells, const TArray<FWorldPartitionStreamingSource>& InSources, TArray<const UWorldPartitionRuntimeCell*>& OutSortedCells)
 {
 	struct FCellShortestDist
@@ -957,12 +1334,12 @@ void UWorldPartitionRuntimeSpatialHash::SortStreamingCellsByDistance(const TSet<
 	}
 }
 
-FVector2D UWorldPartitionRuntimeSpatialHash::GetShowDebugDesiredFootprint(const FVector2D& CanvasSize) const
+FVector2D UWorldPartitionRuntimeSpatialHash::GetDraw2DDesiredFootprint(const FVector2D& CanvasSize) const
 {
 	return FVector2D(CanvasSize.X * StreamingGrids.Num(), CanvasSize.Y);
 }
 
-void UWorldPartitionRuntimeSpatialHash::ShowDebugInfo(UCanvas* Canvas, const TArray<FWorldPartitionStreamingSource>& Sources, const FVector2D& PartitionCanvasOffset, const FVector2D& PartitionCanvasSize) const
+void UWorldPartitionRuntimeSpatialHash::Draw2D(UCanvas* Canvas, const TArray<FWorldPartitionStreamingSource>& Sources, const FVector2D& PartitionCanvasOffset, const FVector2D& PartitionCanvasSize) const
 {
 	if (StreamingGrids.Num() == 0 || Sources.Num() == 0)
 	{
@@ -984,9 +1361,6 @@ void UWorldPartitionRuntimeSpatialHash::ShowDebugInfo(UCanvas* Canvas, const TAr
 	int32 GridIndex = 0;
 	for (const FSpatialHashStreamingGrid& StreamingGrid : StreamingGrids)
 	{
-		const FVector2D GridScreenOffset = GridScreenInitialOffset + ((float)GridIndex * FVector2D(GridMaxScreenSize, 0.f)) + GridScreenHalfExtent;
-		const FBox2D GridScreenBounds(GridScreenOffset - GridScreenHalfExtent, GridScreenOffset + GridScreenHalfExtent);
-
 		// Display view sides based on extended grid loading range (minimum of N cells)
 		const float GridSideDistance = FMath::Max((2.f * StreamingGrid.GetLoadingRange() * GridViewLoadingRangeExtentRatio), StreamingGrid.CellSize * GridViewMinimumSizeInCellCount);
 		FSphere AverageSphere(ForceInit);
@@ -995,87 +1369,13 @@ void UWorldPartitionRuntimeSpatialHash::ShowDebugInfo(UCanvas* Canvas, const TAr
 			AverageSphere += FSphere(Source.Location, 0.5f * GridSideDistance);
 		}
 		const FVector2D GridReferenceWorldPos = FVector2D(AverageSphere.Center);
+		const FBox Region(AverageSphere.Center - AverageSphere.W, AverageSphere.Center + AverageSphere.W);
+		const FVector2D GridScreenOffset = GridScreenInitialOffset + ((float)GridIndex * FVector2D(GridMaxScreenSize, 0.f)) + GridScreenHalfExtent;
+		const FBox2D GridScreenBounds(GridScreenOffset - GridScreenHalfExtent, GridScreenOffset + GridScreenHalfExtent);
 		const float WorldToScreenScale = (0.5f * GridEffectiveScreenSize) / AverageSphere.W;
-		auto WorldToScreen = [&](FVector2D WorldPos) { return (WorldToScreenScale * (WorldPos - GridReferenceWorldPos)) + GridScreenOffset; };
-
-		// Draw Grid cells at desired grid level
-		int32 GridLevel = FMath::Clamp<int32>(GShowRuntimeSpatialHashGridLevel, 0, StreamingGrid.GridLevels.Num() - 1);
-		const FSquare2DGridHelper GridHelper(StreamingGrid.GridLevels.Num(), StreamingGrid.Origin, StreamingGrid.CellSize, StreamingGrid.GridSize);
-		FBox TestBox(AverageSphere.Center - AverageSphere.W, AverageSphere.Center + AverageSphere.W);
-		GridHelper.Levels[GridLevel].ForEachIntersectingCells(TestBox, [&](const FIntVector2& Coords)
-		{
-			FBox2D CellWorldBounds;
-			GridHelper.Levels[GridLevel].GetCellBounds(FIntVector2(Coords.X, Coords.Y), CellWorldBounds);
-			FBox2D CellScreenBounds = FBox2D(WorldToScreen(CellWorldBounds.Min), WorldToScreen(CellWorldBounds.Max));
-			// Clamp inside grid bounds
-			if (!GridScreenBounds.IsInside(CellScreenBounds))
-			{
-				CellScreenBounds.Min.X = FMath::Clamp(CellScreenBounds.Min.X, GridScreenBounds.Min.X, GridScreenBounds.Max.X);
-				CellScreenBounds.Min.Y = FMath::Clamp(CellScreenBounds.Min.Y, GridScreenBounds.Min.Y, GridScreenBounds.Max.Y);
-				CellScreenBounds.Max.X = FMath::Clamp(CellScreenBounds.Max.X, GridScreenBounds.Min.X, GridScreenBounds.Max.X);
-				CellScreenBounds.Max.Y = FMath::Clamp(CellScreenBounds.Max.Y, GridScreenBounds.Min.Y, GridScreenBounds.Max.Y);
-			}
-			else
-			{
-				FString GridInfoText = FString::Printf(TEXT("X%02d_Y%02d"), Coords.X, Coords.Y);
-				float TextWidth, TextHeight;
-				Canvas->SetDrawColor(255, 255, 0);
-				Canvas->StrLen(GEngine->GetTinyFont(), GridInfoText, TextWidth, TextHeight);
-				FVector2D CellSize = CellScreenBounds.GetSize();
-				if (TextWidth < CellSize.X && TextHeight < CellSize.Y)
-				{
-					FVector2D GridInfoPos = CellScreenBounds.GetCenter() - FVector2D(TextWidth / 2, TextHeight / 2);
-					Canvas->DrawText(GEngine->GetTinyFont(), GridInfoText, GridInfoPos.X, GridInfoPos.Y);
-				}
-			}
-			
-			const UWorldPartitionRuntimeSpatialHashCell* Cell = Cast<const UWorldPartitionRuntimeSpatialHashCell>(StreamingGrid.GridLevels[GridLevel].GridCells[Coords.Y * GridHelper.Levels[GridLevel].GridSize + Coords.X]);
-			FLinearColor CellColor = Cell ? Cell->GetDebugColor() : FLinearColor(0.f, 0.f, 0.f, 0.25f);
-			FCanvasTileItem Item(CellScreenBounds.Min, GWhiteTexture, CellScreenBounds.GetSize(), CellColor);
-			Item.BlendMode = SE_BLEND_Translucent;
-			Canvas->DrawItem(Item);
+		auto WorldToScreen = [&](const FVector2D& WorldPos) { return (WorldToScreenScale * (WorldPos - GridReferenceWorldPos)) + GridScreenOffset; };
 
-			FCanvasBoxItem Box(CellScreenBounds.Min, CellScreenBounds.GetSize());
-			Box.SetColor(CellColor);
-			Box.BlendMode = SE_BLEND_Translucent;
-			Canvas->DrawItem(Box);
-		});
-
-		// Draw Loading Ranges
-		float LoadingRange = StreamingGrid.GetLoadingRange();
-
-		FCanvasLineItem LineItem;
-		LineItem.LineThickness = 2;
-		LineItem.SetColor(FLinearColor::White);
-
-		FCanvas* CanvasObject = Canvas->Canvas;
-
-		for (const FWorldPartitionStreamingSource& Source : Sources)
-		{
-			TArray<FVector> LinePoints;
-			LinePoints.SetNum(2);
-
-			float Sin, Cos;
-			FMath::SinCos(&Sin, &Cos, (63.0f / 64.0f) * 2.0f * PI);
-			FVector2D LineStart(Sin * LoadingRange, Cos * LoadingRange);
-
-			for (int32 i = 0; i < 64; i++)
-			{
-				FMath::SinCos(&Sin, &Cos, (i / 64.0f) * 2.0f * PI);
-				FVector2D LineEnd(Sin * LoadingRange, Cos * LoadingRange);
-				LineItem.Draw(CanvasObject, WorldToScreen(FVector2D(Source.Location) + LineStart), WorldToScreen(FVector2D(Source.Location) + LineEnd));
-				LineStart = LineEnd;
-			}
-
-			FVector2D SourceDir = FVector2D(Source.Rotation.Vector());
-
-			if (SourceDir.Size())
-			{
-				SourceDir.Normalize();
-				FVector2D ConeCenter(FVector2D(Source.Location));
-				LineItem.Draw(CanvasObject, WorldToScreen(ConeCenter), WorldToScreen(ConeCenter + SourceDir * LoadingRange));
-			}
-		}
+		StreamingGrid.Draw2D(Canvas, Sources, Region, GridScreenBounds, WorldToScreen);
 
 		// Draw WorldPartition name
 		FVector2D GridInfoPos = GridScreenOffset - GridScreenHalfExtent;
@@ -1090,15 +1390,21 @@ void UWorldPartitionRuntimeSpatialHash::ShowDebugInfo(UCanvas* Canvas, const TAr
 
 		// Draw Grid name, loading range
 		{
-			FString GridInfoText = FString::Printf(TEXT("%s | %d m"), *StreamingGrid.GridName.ToString(), int32(LoadingRange * 0.01f));
+			FString GridInfoText = FString::Printf(TEXT("%s | %d m"), *StreamingGrid.GridName.ToString(), int32(StreamingGrid.GetLoadingRange() * 0.01f));
 			Canvas->SetDrawColor(255, 255, 0);
 			Canvas->DrawText(GEngine->GetTinyFont(), GridInfoText, GridInfoPos.X, GridInfoPos.Y);
 		}
 
-		FCanvasBoxItem Box(GridScreenBounds.Min, GridScreenBounds.GetSize());
-		Box.SetColor(StreamingGrid.DebugColor);
-		Canvas->DrawItem(Box);
-
 		++GridIndex;
 	}
 }
+
+void UWorldPartitionRuntimeSpatialHash::Draw3D(const TArray<FWorldPartitionStreamingSource>& Sources) const
+{
+	UWorld* World = GetWorld();
+	// Get cells based on streaming sources
+	for (const FSpatialHashStreamingGrid& StreamingGrid : StreamingGrids)
+	{
+		StreamingGrid.Draw3D(World, Sources);
+	}
+}
\ No newline at end of file
diff --git a/Engine/Source/Runtime/Engine/Private/WorldPartition/WorldPartitionStreamingPolicy.cpp b/Engine/Source/Runtime/Engine/Private/WorldPartition/WorldPartitionStreamingPolicy.cpp
--- a/Engine/Source/Runtime/Engine/Private/WorldPartition/WorldPartitionStreamingPolicy.cpp
+++ b/Engine/Source/Runtime/Engine/Private/WorldPartition/WorldPartitionStreamingPolicy.cpp
@@ -138,16 +138,25 @@ void UWorldPartitionStreamingPolicy::UpdateStreamingState()
 	}
 }
 
-FVector2D UWorldPartitionStreamingPolicy::GetShowDebugDesiredFootprint(const FVector2D& CanvasSize)
+FVector2D UWorldPartitionStreamingPolicy::GetDrawRuntimeHash2DDesiredFootprint(const FVector2D& CanvasSize)
 {
-	return WorldPartition->RuntimeHash->GetShowDebugDesiredFootprint(CanvasSize);
+	return WorldPartition->RuntimeHash->GetDraw2DDesiredFootprint(CanvasSize);
 }
 
-void UWorldPartitionStreamingPolicy::ShowDebugInfo(class UCanvas* Canvas, const FVector2D& PartitionCanvasOffset, const FVector2D& PartitionCanvasSize)
+void UWorldPartitionStreamingPolicy::DrawRuntimeHash2D(class UCanvas* Canvas, const FVector2D& PartitionCanvasOffset, const FVector2D& PartitionCanvasSize)
 {
 	if (StreamingSources.Num() > 0)
 	{
-		WorldPartition->RuntimeHash->ShowDebugInfo(Canvas, StreamingSources, PartitionCanvasOffset, PartitionCanvasSize);
+		WorldPartition->RuntimeHash->Draw2D(Canvas, StreamingSources, PartitionCanvasOffset, PartitionCanvasSize);
+	}
+}
+
+void UWorldPartitionStreamingPolicy::DrawRuntimeHash3D()
+{
+	UWorld* World = WorldPartition->GetWorld();
+	if (World->GetNetMode() != NM_DedicatedServer && WorldPartition->IsInitialized())
+	{
+		WorldPartition->RuntimeHash->Draw3D(StreamingSources);
 	}
 }
 
diff --git a/Engine/Source/Runtime/Engine/Private/WorldPartition/WorldPartitionSubsystem.cpp b/Engine/Source/Runtime/Engine/Private/WorldPartition/WorldPartitionSubsystem.cpp
--- a/Engine/Source/Runtime/Engine/Private/WorldPartition/WorldPartitionSubsystem.cpp
+++ b/Engine/Source/Runtime/Engine/Private/WorldPartition/WorldPartitionSubsystem.cpp
@@ -9,7 +9,7 @@
 #include "Engine/Console.h"
 #include "ConsoleSettings.h"
 
-static const FName NAME_WorldPartitionRuntimeGrids("WorldPartitionRuntimeGrids");
+static const FName NAME_WorldPartitionRuntimeHash("WorldPartitionRuntimeHash");
 
 namespace WorldPartitionSubsystemConsole
 {
@@ -19,12 +19,18 @@ namespace WorldPartitionSubsystemConsole
 		AutoCompleteList.AddDefaulted();
 
 		FAutoCompleteCommand& AutoCompleteCommand = AutoCompleteList.Last();
-		AutoCompleteCommand.Command = FString::Printf(TEXT("showdebug %s"), *NAME_WorldPartitionRuntimeGrids.ToString());
-		AutoCompleteCommand.Desc = TEXT("Toggles display of world partition runtime grids");
+		AutoCompleteCommand.Command = FString::Printf(TEXT("showdebug %s"), *NAME_WorldPartitionRuntimeHash.ToString());
+		AutoCompleteCommand.Desc = TEXT("Toggles 2D debug display of world partition runtime hash.");
 		AutoCompleteCommand.Color = ConsoleSettings->AutoCompleteCommandColor;
 	}
 }
 
+static int32 GDrawRuntimeHash3D = 0;
+static FAutoConsoleCommand CVarDrawRuntimeHash3D(
+	TEXT("WorldPartition.ToggleDrawRuntimeHash3D"),
+	TEXT("Toggles 3D debug display of world partition runtime hash."),
+	FConsoleCommandDelegate::CreateLambda([] { GDrawRuntimeHash3D = !GDrawRuntimeHash3D; }));
+
 UWorldPartitionSubsystem::UWorldPartitionSubsystem()
 {}
 
@@ -103,6 +109,11 @@ void UWorldPartitionSubsystem::Tick(float DeltaSeconds)
 	for (UWorldPartition* Partition : RegisteredWorldPartitions)
 	{
 		Partition->Tick(DeltaSeconds);
+		
+		if (GDrawRuntimeHash3D && GetWorld()->IsGameWorld())
+		{
+			Partition->DrawRuntimeHash3D();
+		}
 	}
 }
 
@@ -144,7 +155,7 @@ void UWorldPartitionSubsystem::UpdateStreamingState()
 
 void UWorldPartitionSubsystem::OnShowDebugInfo(class AHUD* HUD, UCanvas* Canvas, const class FDebugDisplayInfo& DisplayInfo, float& YL, float& YPos)
 {
-	if (!Canvas || !HUD || !HUD->ShouldDisplayDebug(NAME_WorldPartitionRuntimeGrids))
+	if (!Canvas || !HUD || !HUD->ShouldDisplayDebug(NAME_WorldPartitionRuntimeHash))
 	{
 		return;
 	}
@@ -157,7 +168,7 @@ void UWorldPartitionSubsystem::OnShowDebugInfo(class AHUD* HUD, UCanvas* Canvas,
 	FVector2D TotalFootprint(ForceInitToZero);
 	for (UWorldPartition* Partition : RegisteredWorldPartitions)
 	{
-		FVector2D Footprint = Partition->GetShowDebugDesiredFootprint(CanvasMaxScreenSize);
+		FVector2D Footprint = Partition->GetDrawRuntimeHash2DDesiredFootprint(CanvasMaxScreenSize);
 		TotalFootprint.X += Footprint.X;
 	}
 
@@ -166,10 +177,10 @@ void UWorldPartitionSubsystem::OnShowDebugInfo(class AHUD* HUD, UCanvas* Canvas,
 		FVector2D PartitionCanvasOffset(CanvasTopLeftPadding);
 		for (UWorldPartition* Partition : RegisteredWorldPartitions)
 		{
-			FVector2D Footprint = Partition->GetShowDebugDesiredFootprint(CanvasMaxScreenSize);
+			FVector2D Footprint = Partition->GetDrawRuntimeHash2DDesiredFootprint(CanvasMaxScreenSize);
 			float FootprintRatio = Footprint.X / TotalFootprint.X;
 			FVector2D PartitionCanvasSize = FVector2D(CanvasMaxScreenSize.X * FootprintRatio, CanvasMaxScreenSize.Y);
-			Partition->ShowDebugInfo(Canvas, PartitionCanvasOffset, PartitionCanvasSize);
+			Partition->DrawRuntimeHash2D(Canvas, PartitionCanvasOffset, PartitionCanvasSize);
 			PartitionCanvasOffset.X += PartitionCanvasSize.X;
 		}
 	}
diff --git a/Engine/Source/Runtime/Engine/Public/WorldPartition/WorldPartition.h b/Engine/Source/Runtime/Engine/Public/WorldPartition/WorldPartition.h
--- a/Engine/Source/Runtime/Engine/Public/WorldPartition/WorldPartition.h
+++ b/Engine/Source/Runtime/Engine/Public/WorldPartition/WorldPartition.h
@@ -163,8 +163,9 @@ public:
 	void Tick(float DeltaSeconds);
 	void UpdateStreamingState();
 
-	FVector2D GetShowDebugDesiredFootprint(const FVector2D& CanvasSize);
-	void ShowDebugInfo(class UCanvas* Canvas, const FVector2D& PartitionCanvasOffset, const FVector2D& PartitionCanvasSize);
+	FVector2D GetDrawRuntimeHash2DDesiredFootprint(const FVector2D& CanvasSize);
+	void DrawRuntimeHash2D(class UCanvas* Canvas, const FVector2D& PartitionCanvasOffset, const FVector2D& PartitionCanvasSize);
+	void DrawRuntimeHash3D();
 
 	virtual UWorld* GetWorld() const override;
 
diff --git a/Engine/Source/Runtime/Engine/Public/WorldPartition/WorldPartitionRuntimeHash.h b/Engine/Source/Runtime/Engine/Public/WorldPartition/WorldPartitionRuntimeHash.h
--- a/Engine/Source/Runtime/Engine/Public/WorldPartition/WorldPartitionRuntimeHash.h
+++ b/Engine/Source/Runtime/Engine/Public/WorldPartition/WorldPartitionRuntimeHash.h
@@ -29,12 +29,13 @@ class ENGINE_API UWorldPartitionRuntimeHash : public UObject
 	virtual int32 GetStreamingCells(const TArray<FWorldPartitionStreamingSource>& Sources, TSet<const UWorldPartitionRuntimeCell*>& Cells) const { return 0; };
 	virtual void SortStreamingCellsByDistance(const TSet<const UWorldPartitionRuntimeCell*>& InCells, const TArray<FWorldPartitionStreamingSource>& InSources, TArray<const UWorldPartitionRuntimeCell*>& OutSortedCells) {}
 
-	/* Returns desired footprint that ShowDebugInfo should take relative to given Canvas size (the value can exceed the given size).
-	 * UWorldPartitionSubSystem will re-adapt the size relative to all others UWorldPartitionRuntimeHash and provide the correct size to ShowDebugInfo.
+	/* Returns desired footprint that Draw2D should take relative to given Canvas size (the value can exceed the given size).
+	 * UWorldPartitionSubSystem will re-adapt the size relative to all others UWorldPartitionRuntimeHash and provide the correct size to Draw2D.
 	 *
-	 * Return ShowDebugInfo's desired footprint.
+	 * Return Draw2D's desired footprint.
 	 */
-	virtual FVector2D GetShowDebugDesiredFootprint(const FVector2D& CanvasSize) const { return FVector2D::ZeroVector; }
+	virtual FVector2D GetDraw2DDesiredFootprint(const FVector2D& CanvasSize) const { return FVector2D::ZeroVector; }
 
-	virtual void ShowDebugInfo(class UCanvas* Canvas, const TArray<FWorldPartitionStreamingSource>& Sources, const FVector2D& PartitionCanvasOffset, const FVector2D& PartitionCanvasSize) const {}
+	virtual void Draw2D(class UCanvas* Canvas, const TArray<FWorldPartitionStreamingSource>& Sources, const FVector2D& PartitionCanvasOffset, const FVector2D& PartitionCanvasSize) const {}
+	virtual void Draw3D(const TArray<FWorldPartitionStreamingSource>& Sources) const {}
 };
\ No newline at end of file
diff --git a/Engine/Source/Runtime/Engine/Public/WorldPartition/WorldPartitionRuntimeSpatialHash.h b/Engine/Source/Runtime/Engine/Public/WorldPartition/WorldPartitionRuntimeSpatialHash.h
--- a/Engine/Source/Runtime/Engine/Public/WorldPartition/WorldPartitionRuntimeSpatialHash.h
+++ b/Engine/Source/Runtime/Engine/Public/WorldPartition/WorldPartitionRuntimeSpatialHash.h
@@ -37,6 +37,8 @@ struct FSpatialHashStreamingGridLevel
 // Disable copy of this USTRUCT by the autogenerated code.
 template<> struct TStructOpsTypeTraits<FSpatialHashStreamingGridLevel> : public TStructOpsTypeTraitsBase2<FSpatialHashStreamingGridLevel> { enum { WithCopy = false }; };
 
+struct FSquare2DGridHelper;
+
 /**
  * Represents a PIE/Game streaming grid
  */
@@ -45,16 +47,8 @@ struct FSpatialHashStreamingGrid
 {
 	GENERATED_USTRUCT_BODY()
 
-	FSpatialHashStreamingGrid()
-		: Origin(ForceInitToZero)
-		, CellSize(0)
-		, GridSize(0)
-		, LoadingRange(0.0f)
-#if WITH_EDITOR
-		, DebugColor(ForceInitToZero)
-		, OverrideLoadingRange(0)
-#endif
-	{}
+	FSpatialHashStreamingGrid();
+	~FSpatialHashStreamingGrid();
 
 	// Move constructor.
 	FSpatialHashStreamingGrid(const FSpatialHashStreamingGrid& Other) = delete;
@@ -95,7 +89,19 @@ struct FSpatialHashStreamingGrid
 		}
 		return LoadingRange;
 	}
+
+	// Used by PIE/Game
+	void GetCells(const TArray<FWorldPartitionStreamingSource>& Sources, TSet<const UWorldPartitionRuntimeCell*>& Cells) const;
+	void GetAlwaysLoadedCells(TSet<const UWorldPartitionRuntimeCell*>& Cells) const;
+	void Draw2D(UCanvas* Canvas, const TArray<FWorldPartitionStreamingSource>& Sources, const FBox& Region, const FBox2D& GridScreenBounds, TFunctionRef<FVector2D(const FVector2D&)> WorldToScreen) const;
+	void Draw3D(UWorld* World, const TArray<FWorldPartitionStreamingSource>& Sources) const;
+
+private:
+	mutable TMap<uint32, FVector> CachedLineTrace;
+	const FSquare2DGridHelper& GetGridHelper() const;
+	mutable FSquare2DGridHelper* GridHelper;
 };
+
 // Disable copy of this USTRUCT by the autogenerated code.
 template<> struct TStructOpsTypeTraits<FSpatialHashStreamingGrid> : public TStructOpsTypeTraitsBase2<FSpatialHashStreamingGrid> { enum { WithCopy = false }; };
 
@@ -130,8 +136,6 @@ struct FSpatialHashRuntimeGrid
 #endif
 };
 
-struct FSquare2DGridHelper;
-
 UCLASS()
 class ENGINE_API UWorldPartitionRuntimeSpatialHash : public UWorldPartitionRuntimeHash
 {
@@ -183,8 +187,9 @@ private:
 	UPROPERTY(NonPIEDuplicateTransient)
 	TArray<FSpatialHashStreamingGrid> StreamingGrids;
 
-	virtual FVector2D GetShowDebugDesiredFootprint(const FVector2D& CanvasSize) const override;
-	virtual void ShowDebugInfo(class UCanvas* Canvas, const TArray<FWorldPartitionStreamingSource>& Sources, const FVector2D& PartitionCanvasOffset, const FVector2D& PartitionCanvasSize) const override;
+	virtual FVector2D GetDraw2DDesiredFootprint(const FVector2D& CanvasSize) const override;
+	virtual void Draw2D(class UCanvas* Canvas, const TArray<FWorldPartitionStreamingSource>& Sources, const FVector2D& PartitionCanvasOffset, const FVector2D& PartitionCanvasSize) const override;
+	virtual void Draw3D(const TArray<FWorldPartitionStreamingSource>& Sources) const override;
 
 private:
 	void GetAlwaysLoadedStreamingCells(const FSpatialHashStreamingGrid& StreamingGrid, TSet<const UWorldPartitionRuntimeCell*>& Cells) const;
diff --git a/Engine/Source/Runtime/Engine/Public/WorldPartition/WorldPartitionStreamingPolicy.h b/Engine/Source/Runtime/Engine/Public/WorldPartition/WorldPartitionStreamingPolicy.h
--- a/Engine/Source/Runtime/Engine/Public/WorldPartition/WorldPartitionStreamingPolicy.h
+++ b/Engine/Source/Runtime/Engine/Public/WorldPartition/WorldPartitionStreamingPolicy.h
@@ -37,8 +37,9 @@ public:
 	virtual void UnloadCells(const TSet<const UWorldPartitionRuntimeCell*>& ToUnloadCells);
 	virtual void LoadCell(const UWorldPartitionRuntimeCell* Cell) PURE_VIRTUAL(UWorldPartitionStreamingPolicy::LoadCell, );
 	virtual void UnloadCell(const UWorldPartitionRuntimeCell* Cell) PURE_VIRTUAL(UWorldPartitionStreamingPolicy::UnloadCell, );
-	FVector2D GetShowDebugDesiredFootprint(const FVector2D& CanvasSize);
-	void ShowDebugInfo(class UCanvas* Canvas, const FVector2D& PartitionCanvasOffset, const FVector2D& PartitionCanvasSize);
+	virtual FVector2D GetDrawRuntimeHash2DDesiredFootprint(const FVector2D& CanvasSize);
+	virtual void DrawRuntimeHash2D(class UCanvas* Canvas, const FVector2D& PartitionCanvasOffset, const FVector2D& PartitionCanvasSize);
+	virtual void DrawRuntimeHash3D();
 
 #if WITH_EDITOR
 	virtual TSubclassOf<class UWorldPartitionRuntimeCell> GetRuntimeCellClass() const PURE_VIRTUAL(UWorldPartitionStreamingPolicy::GetRuntimeCellClass, return UWorldPartitionRuntimeCell::StaticClass(); );
