Avoid conditional postload getting hung when textures are being compiled while in paused mode

diff --git a/Engine/Source/Runtime/Engine/Private/Texture2D.cpp b/Engine/Source/Runtime/Engine/Private/Texture2D.cpp
--- a/Engine/Source/Runtime/Engine/Private/Texture2D.cpp
+++ b/Engine/Source/Runtime/Engine/Private/Texture2D.cpp
@@ -793,6 +793,13 @@ void UTexture2D::CancelPendingTextureStreaming()
 
 bool UTexture2D::IsReadyForAsyncPostLoad() const
 {
+#if WITH_EDITOR
+	if (IsDefaultTexture())
+	{
+		return true;
+	}
+#endif
+
 	return !PrivatePlatformData || PrivatePlatformData->IsReadyForAsyncPostLoad();
 }
 
@@ -813,8 +820,8 @@ void UTexture2D::PostLoad()
 		}
 		else
 		{
-		FinishCachePlatformData();
-	}
+			FinishCachePlatformData();
+		}
 	}
 #endif // #if WITH_EDITOR
 
