WorldPartition - Always force a full rescan of the level path when fetching level actors from the asset registry

diff --git a/Engine/Source/Runtime/Engine/Private/ActorRegistry.cpp b/Engine/Source/Runtime/Engine/Private/ActorRegistry.cpp
--- a/Engine/Source/Runtime/Engine/Private/ActorRegistry.cpp
+++ b/Engine/Source/Runtime/Engine/Private/ActorRegistry.cpp
@@ -16,7 +16,8 @@ void FActorRegistry::GetLevelActors(const FName& LevelPath, TArray<FAssetData>&
 	IAssetRegistry& AssetRegistry = FModuleManager::LoadModuleChecked<FAssetRegistryModule>(TEXT("AssetRegistry")).Get();
 	
 	// Do a synchronous scan of the level external actors path.
-	AssetRegistry.ScanPathsSynchronous({ULevel::GetExternalActorsPath(LevelPath.ToString())});
+	const bool bForceRescan = true;
+	AssetRegistry.ScanPathsSynchronous({ULevel::GetExternalActorsPath(LevelPath.ToString())}, bForceRescan);
 
 	static const FName NAME_LevelPackage(TEXT("LevelPackage"));
 	FARFilter Filter;
