Fixed variable shadowing

diff --git a/Engine/Source/Runtime/Engine/Private/WorldPartition/WorldPartitionRuntimeSpatialHash.cpp b/Engine/Source/Runtime/Engine/Private/WorldPartition/WorldPartitionRuntimeSpatialHash.cpp
--- a/Engine/Source/Runtime/Engine/Private/WorldPartition/WorldPartitionRuntimeSpatialHash.cpp
+++ b/Engine/Source/Runtime/Engine/Private/WorldPartition/WorldPartitionRuntimeSpatialHash.cpp
@@ -657,8 +657,8 @@ void FSpatialHashStreamingGrid::Draw2D(UCanvas* Canvas, const TArray<FWorldParti
 				float TextWidth, TextHeight;
 				Canvas->SetDrawColor(255, 255, 0);
 				Canvas->StrLen(GEngine->GetTinyFont(), GridInfoText, TextWidth, TextHeight);
-				FVector2D CellSize = CellScreenBounds.GetSize();
-				if (TextWidth < CellSize.X && TextHeight < CellSize.Y)
+				FVector2D CellBoundsSize = CellScreenBounds.GetSize();
+				if (TextWidth < CellBoundsSize.X && TextHeight < CellBoundsSize.Y)
 				{
 					FVector2D GridInfoPos = CellScreenBounds.GetCenter() - FVector2D(TextWidth / 2, TextHeight / 2);
 					Canvas->DrawText(GEngine->GetTinyFont(), GridInfoText, GridInfoPos.X, GridInfoPos.Y);
