WorldPartition: Hook debug 2D view to DebugDrawService instead of HUD.

diff --git a/Engine/Source/Runtime/Engine/Private/WorldPartition/WorldPartitionSubsystem.cpp b/Engine/Source/Runtime/Engine/Private/WorldPartition/WorldPartitionSubsystem.cpp
--- a/Engine/Source/Runtime/Engine/Private/WorldPartition/WorldPartitionSubsystem.cpp
+++ b/Engine/Source/Runtime/Engine/Private/WorldPartition/WorldPartitionSubsystem.cpp
@@ -4,33 +4,25 @@
 #include "WorldPartition/WorldPartition.h"
 #include "WorldPartition/WorldPartitionStreamingPolicy.h"
 #include "WorldPartition/WorldPartitionActorDescFactory.h"
-#include "GameFramework/HUD.h"
 #include "Engine/Canvas.h"
 #include "Engine/Console.h"
 #include "ConsoleSettings.h"
+#include "Debug/DebugDrawService.h"
 
 static const FName NAME_WorldPartitionRuntimeHash("WorldPartitionRuntimeHash");
 
-namespace WorldPartitionSubsystemConsole
-{
-	static void PopulateAutoCompleteEntries(TArray<FAutoCompleteCommand>& AutoCompleteList)
-	{
-		const UConsoleSettings* ConsoleSettings = GetDefault<UConsoleSettings>();
-		AutoCompleteList.AddDefaulted();
-
-		FAutoCompleteCommand& AutoCompleteCommand = AutoCompleteList.Last();
-		AutoCompleteCommand.Command = FString::Printf(TEXT("showdebug %s"), *NAME_WorldPartitionRuntimeHash.ToString());
-		AutoCompleteCommand.Desc = TEXT("Toggles 2D debug display of world partition runtime hash.");
-		AutoCompleteCommand.Color = ConsoleSettings->AutoCompleteCommandColor;
-	}
-}
-
 static int32 GDrawRuntimeHash3D = 0;
 static FAutoConsoleCommand CVarDrawRuntimeHash3D(
 	TEXT("WorldPartition.ToggleDrawRuntimeHash3D"),
 	TEXT("Toggles 3D debug display of world partition runtime hash."),
 	FConsoleCommandDelegate::CreateLambda([] { GDrawRuntimeHash3D = !GDrawRuntimeHash3D; }));
 
+static int32 GDrawRuntimeHash2D = 0;
+static FAutoConsoleCommand CVarDrawRuntimeHash2D(
+	TEXT("WorldPartition.ToggleDrawRuntimeHash2D"),
+	TEXT("Toggles 2D debug display of world partition runtime hash."),
+	FConsoleCommandDelegate::CreateLambda([] { GDrawRuntimeHash2D = !GDrawRuntimeHash2D; }));
+
 UWorldPartitionSubsystem::UWorldPartitionSubsystem()
 {}
 
@@ -53,26 +45,24 @@ void UWorldPartitionSubsystem::PostInitialize()
 {
 	Super::PostInitialize();
 
-	if (GetWorld()->IsGameWorld() && (GetWorld()->GetNetMode() != NM_DedicatedServer))
-	{
-		AHUD::OnShowDebugInfo.AddUObject(this, &UWorldPartitionSubsystem::OnShowDebugInfo);
-		static bool bRegisteredToConsole = false;
-		if (!bRegisteredToConsole)
-		{
-			UConsole::RegisterConsoleAutoCompleteEntries.AddStatic(&WorldPartitionSubsystemConsole::PopulateAutoCompleteEntries);
-			bRegisteredToConsole = true;
-		}
-	}
-
 	if (UWorldPartition* MainPartition = GetMainWorldPartition())
 	{
 		MainPartition->Initialize(GetWorld(), FTransform::Identity);
+
+		if (GetWorld()->IsGameWorld() && (GetWorld()->GetNetMode() != NM_DedicatedServer))
+		{
+			DrawRuntimeHash2DHandle = UDebugDrawService::Register(TEXT("Game"), FDebugDrawDelegate::CreateUObject(this, &UWorldPartitionSubsystem::DrawRuntimeHash2D));
+		}
 	}
 }
 
 void UWorldPartitionSubsystem::Deinitialize()
 {
-	AHUD::OnShowDebugInfo.RemoveAll(this);
+	if (DrawRuntimeHash2DHandle.IsValid())
+	{
+		UDebugDrawService::Unregister(DrawRuntimeHash2DHandle);
+		DrawRuntimeHash2DHandle.Reset();
+	}
 
 	while (RegisteredWorldPartitions.Num() > 0)
 	{
@@ -153,14 +143,14 @@ void UWorldPartitionSubsystem::UpdateStreamingState()
 	}
 }
 
-void UWorldPartitionSubsystem::OnShowDebugInfo(class AHUD* HUD, UCanvas* Canvas, const class FDebugDisplayInfo& DisplayInfo, float& YL, float& YPos)
+void UWorldPartitionSubsystem::DrawRuntimeHash2D(UCanvas* Canvas, class APlayerController* PC)
 {
-	if (!Canvas || !HUD || !HUD->ShouldDisplayDebug(NAME_WorldPartitionRuntimeHash))
+	if (!GDrawRuntimeHash2D || !Canvas || !Canvas->SceneView)
 	{
 		return;
 	}
 
-	const FVector2D CanvasTopLeftPadding(10.f, YPos);
+	const FVector2D CanvasTopLeftPadding(10.f, 10.f);
 	const FVector2D CanvasBottomRightPadding(10.f, 10.f);
 	const FVector2D CanvasMinimumSize(100.f, 100.f);
 	const FVector2D CanvasMaxScreenSize = FVector2D::Max(FVector2D(Canvas->ClipX, Canvas->ClipY) - CanvasBottomRightPadding - CanvasTopLeftPadding, CanvasMinimumSize);
diff --git a/Engine/Source/Runtime/Engine/Public/WorldPartition/WorldPartitionSubsystem.h b/Engine/Source/Runtime/Engine/Public/WorldPartition/WorldPartitionSubsystem.h
--- a/Engine/Source/Runtime/Engine/Public/WorldPartition/WorldPartitionSubsystem.h
+++ b/Engine/Source/Runtime/Engine/Public/WorldPartition/WorldPartitionSubsystem.h
@@ -55,15 +55,18 @@ public:
 	FBox GetWorldBounds();
 #endif
 
+	void ToggleDrawRuntimeHash2D();
+
 private:
 	UWorldPartition* GetMainWorldPartition();
 	const UWorldPartition* GetMainWorldPartition() const;
-
-	void OnShowDebugInfo(class AHUD* HUD, class UCanvas* Canvas, const class FDebugDisplayInfo& DisplayInfo, float& YL, float& YPos);
 	void RegisterWorldPartition(UWorldPartition* WorldPartition);
 	void UnregisterWorldPartition(UWorldPartition* WorldPartition);
+	void DrawRuntimeHash2D(class UCanvas* Canvas, class APlayerController* PC);
 	friend class UWorldPartition;
 
 	UPROPERTY()
 	TArray<UWorldPartition*> RegisteredWorldPartitions;
+
+	FDelegateHandle	DrawRuntimeHash2DHandle;
 };
