World Partition
- Automatically treat actors that are too big as always loaded in the editor, to avoid spanning more than 32K editing cells.
- Always store actors in level 0 of the octree to avoid having editing cells of different sizes.
- Cleaned up UWorldPartition/UWorldPartitionSubSystem interface of unused functions.
- Minor visual changes to the World Partition topview editor.
- Changed default editing cell sizes to be 512M.

diff --git a/Engine/Source/Editor/WorldPartitionEditor/Private/WorldPartition/SWorldPartitionEditorGrid2D.cpp b/Engine/Source/Editor/WorldPartitionEditor/Private/WorldPartition/SWorldPartitionEditorGrid2D.cpp
--- a/Engine/Source/Editor/WorldPartitionEditor/Private/WorldPartition/SWorldPartitionEditorGrid2D.cpp
+++ b/Engine/Source/Editor/WorldPartitionEditor/Private/WorldPartition/SWorldPartitionEditorGrid2D.cpp
@@ -342,7 +342,7 @@ int32 SWorldPartitionEditorGrid2D::PaintGrid(const FGeometry& AllottedGeometry,
 		FLinearColor YAxisColor = FLinearColor::Green;
 		YAxisColor.A = 0.4f;
 		
-		FSlateDrawElement::MakeLines(OutDrawElements, LayerId, AllottedGeometry.ToPaintGeometry(), LinePoints, ESlateDrawEffect::None, YAxisColor);
+		FSlateDrawElement::MakeLines(OutDrawElements, LayerId, AllottedGeometry.ToPaintGeometry(), LinePoints, ESlateDrawEffect::None, YAxisColor, true, 2.0f);
 	}
 
 	// World X-axis
@@ -354,7 +354,7 @@ int32 SWorldPartitionEditorGrid2D::PaintGrid(const FGeometry& AllottedGeometry,
 		FLinearColor XAxisColor = FLinearColor::Red;
 		XAxisColor.A = 0.4f;
 		
-		FSlateDrawElement::MakeLines(OutDrawElements, LayerId, AllottedGeometry.ToPaintGeometry(), LinePoints, ESlateDrawEffect::None, XAxisColor);
+		FSlateDrawElement::MakeLines(OutDrawElements, LayerId, AllottedGeometry.ToPaintGeometry(), LinePoints, ESlateDrawEffect::None, XAxisColor, true, 2.0f);
 	}
 
 	return LayerId + 1;
@@ -723,6 +723,25 @@ uint32 SWorldPartitionEditorGrid2D::PaintSelection(const FGeometry& AllottedGeom
 		LinePoints[3] = BottomLeft;
 		LinePoints[4] = TopLeft;
 
+		{
+			FSlateColorBrush CellBrush(FLinearColor::White);
+			FLinearColor CellColor(FLinearColor(1, 1, 1, 0.25f));
+
+			FPaintGeometry CellGeometry = AllottedGeometry.ToPaintGeometry(
+				TopLeft,
+				BottomRight - TopLeft
+			);
+
+			FSlateDrawElement::MakeBox(
+				OutDrawElements,
+				LayerId,
+				CellGeometry,
+				&CellBrush,
+				ESlateDrawEffect::None,
+				CellColor
+			);
+		}
+
 		FSlateDrawElement::MakeLines(OutDrawElements, LayerId, AllottedGeometry.ToPaintGeometry(), LinePoints, ESlateDrawEffect::None, FLinearColor::White, false, 1.0f);
 	}
 
diff --git a/Engine/Source/Editor/WorldPartitionEditor/Private/WorldPartition/SWorldPartitionEditorGridSpatialHash.cpp b/Engine/Source/Editor/WorldPartitionEditor/Private/WorldPartition/SWorldPartitionEditorGridSpatialHash.cpp
--- a/Engine/Source/Editor/WorldPartitionEditor/Private/WorldPartition/SWorldPartitionEditorGridSpatialHash.cpp
+++ b/Engine/Source/Editor/WorldPartitionEditor/Private/WorldPartition/SWorldPartitionEditorGridSpatialHash.cpp
@@ -84,7 +84,7 @@ int32 SWorldPartitionEditorGridSpatialHash::PaintGrid(const FGeometry& AllottedG
 	// Shadow whole grid area
 	{
 		FSlateColorBrush ShadowBrush(FLinearColor::Black);
-		FLinearColor ShadowColor(FLinearColor(0.0f, 0.0f, 0.0f, 0.5f));
+		FLinearColor ShadowColor(FLinearColor(0.0f, 0.0f, 0.0f, 0.25f));
 
 		FPaintGeometry GridGeometry = AllottedGeometry.ToPaintGeometry(
 			WorldToScreen.TransformPoint(FVector2D(VisibleGridRectWorld.Min)),
@@ -169,8 +169,21 @@ int32 SWorldPartitionEditorGridSpatialHash::PaintGrid(const FGeometry& AllottedG
 					WorldToScreen.TransformPoint(Cell.Bounds.Max) - WorldToScreen.TransformPoint(Cell.Bounds.Min)
 				);
 
+				FLinearColor CellColor;
+				if(Cell.bEmpty)
+				{
+					CellColor = FLinearColor(0, 0, 0, 1.0f);
+				}
+				else if (Cell.bSelected)
+				{
+					CellColor = FLinearColor(1, 1, 1, 0.25f);
+				}
+				else // !Cell.bLoaded
+				{
+					 CellColor = FLinearColor(0, 0, 0, 0.5f);
+				}
+
 				FSlateColorBrush CellBrush(FLinearColor::White);
-				FLinearColor CellColor(Cell.bSelected ? FLinearColor(1, 1, 1, 0.25f) : FLinearColor(0, 0, 0, Cell.bEmpty ? 1.0f : 0.5f));
 
 				FSlateDrawElement::MakeBox(
 					OutDrawElements,
diff --git a/Engine/Source/Runtime/Engine/Private/WorldPartition/WorldPartition.cpp b/Engine/Source/Runtime/Engine/Private/WorldPartition/WorldPartition.cpp
--- a/Engine/Source/Runtime/Engine/Private/WorldPartition/WorldPartition.cpp
+++ b/Engine/Source/Runtime/Engine/Private/WorldPartition/WorldPartition.cpp
@@ -506,33 +506,6 @@ void UWorldPartition::UpdateActorDesc(AActor* InActor)
 	}
 }
 
-void UWorldPartition::AddActor(AActor* InActor)
-{
-	check(!InActor->IsPendingKill());
-	
-	if (InActor->GetLevel() == World->PersistentLevel && InActor->IsPackageExternal() && !InActor->IsChildActor())
-	{
-		check(!GetActorDesc(InActor->GetActorGuid()))
-		verify(InActor->GetLevel()->Actors.Find(InActor) != INDEX_NONE);
-
-		FWorldPartitionActorDesc* NewDesc = GetActorDescFactory(InActor)->Create(InActor);
-
-		Actors.Add(InActor->GetActorGuid(), TUniquePtr<FWorldPartitionActorDesc>(NewDesc));
-		
-		HashActorDesc(NewDesc);
-		RefreshLoadedCells();
-	}
-}
-
-void UWorldPartition::RemoveActor(AActor* InActor)
-{
-	if (FWorldPartitionActorDesc* ActorDesc = GetActorDesc(InActor->GetActorGuid()))
-	{
-		UnhashActorDesc(ActorDesc);
-		Actors.Remove(ActorDesc->GetGuid());
-	}
-}
-
 void UWorldPartition::ApplyActorTransform(AActor* InActor, const FTransform& InTransform)
 {
 	if (!InTransform.Equals(FTransform::Identity))
diff --git a/Engine/Source/Runtime/Engine/Private/WorldPartition/WorldPartitionEditorSpatialHash.cpp b/Engine/Source/Runtime/Engine/Private/WorldPartition/WorldPartitionEditorSpatialHash.cpp
--- a/Engine/Source/Runtime/Engine/Private/WorldPartition/WorldPartitionEditorSpatialHash.cpp
+++ b/Engine/Source/Runtime/Engine/Private/WorldPartition/WorldPartitionEditorSpatialHash.cpp
@@ -27,7 +27,7 @@ void UWorldPartitionEditorSpatialHash::Initialize()
 
 void UWorldPartitionEditorSpatialHash::SetDefaultValues()
 {
-	CellSize = 32768;
+	CellSize = 51200;
 }
 
 FName UWorldPartitionEditorSpatialHash::GetWorldPartitionEditorName()
@@ -52,11 +52,7 @@ void UWorldPartitionEditorSpatialHash::Tick(float DeltaSeconds)
 			{
 				ForEachIntersectingCells(Bounds, Level, [&](const FCellCoord& CellCoord)
 				{
-					FCellNode CellNode;
-					if (HashCells.RemoveAndCopyValue(CellCoord, CellNode))
-					{
-						check(!CellNode.Cell);
-					}
+					HashNodes.Remove(CellCoord);
 				});
 			}
 		}
@@ -85,10 +81,24 @@ FBox UWorldPartitionEditorSpatialHash::GetActorBounds(FWorldPartitionActorDesc*
 	return ActorBounds;
 }
 
+bool UWorldPartitionEditorSpatialHash::IsActorAlwaysLoaded(FWorldPartitionActorDesc* InActorDesc) const
+{
+	if (InActorDesc->GetGridPlacement() == EActorGridPlacement::AlwaysLoaded)
+	{
+		return true;
+	}
+
+	// If an actor covers more that 4 levels in the octree (which means 32K cells), treat it as always loaded
+	const FBox ActorBounds = GetActorBounds(InActorDesc);
+	const int32 ActorLevel = GetLevelForBox(ActorBounds);
+	return (ActorLevel > 4);
+}
+
 void UWorldPartitionEditorSpatialHash::HashActor(FWorldPartitionActorDesc* InActorDesc)
 {
 	check(InActorDesc);
-	if (InActorDesc->GetGridPlacement() == EActorGridPlacement::AlwaysLoaded)
+
+	if (IsActorAlwaysLoaded(InActorDesc))
 	{
 		AlwaysLoadedCell->AddActor(InActorDesc);
 	}
@@ -96,7 +106,6 @@ void UWorldPartitionEditorSpatialHash::HashActor(FWorldPartitionActorDesc* InAct
 	{
 		const FBox ActorBounds = GetActorBounds(InActorDesc);
 		const int32 CurrentLevel = GetLevelForBox(Bounds);
-		const int32 ActorLevel = GetLevelForBox(ActorBounds);
 
 		auto UpdateHigherLevels = [this](const FCellCoord& CellCoord, int32 EndLevel)
 		{
@@ -107,7 +116,7 @@ void UWorldPartitionEditorSpatialHash::HashActor(FWorldPartitionActorDesc* InAct
 
 				LevelCellCoord = LevelCellCoord.GetParentCellCoord();
 
-				FCellNode& CellNode = HashCells.FindOrAdd(LevelCellCoord);
+				FCellNode& CellNode = HashNodes.FindOrAdd(LevelCellCoord);
 
 				if (CellNode.HasChildNode(ChildIndex))
 				{
@@ -118,24 +127,29 @@ void UWorldPartitionEditorSpatialHash::HashActor(FWorldPartitionActorDesc* InAct
 			}
 		};
 
-		ForEachIntersectingCells(ActorBounds, ActorLevel, [&](const FCellCoord& CellCoord)
+		ForEachIntersectingCells(ActorBounds, 0, [&](const FCellCoord& CellCoord)
 		{
-			FCellNode& CellNode = HashCells.FindOrAdd(CellCoord);
-
-			if (!CellNode.Cell)
+			UWorldPartitionEditorCell* EditorCell = nullptr;
+			if (UWorldPartitionEditorCell** EditorCellPtr = HashCells.Find(CellCoord))
 			{
-				CellNode.Cell = NewObject<UWorldPartitionEditorCell>(this, *FString::Printf(TEXT("EditorCell_X%lld_Y%lld_Z%lld_L%d"), CellCoord.X, CellCoord.Y, CellCoord.Z, CellCoord.Level), RF_Transient);
-				CellNode.Cell->SetFlags(RF_Transactional);
-				CellNode.Cell->Bounds = GetCellBounds(CellCoord);
+				EditorCell = *EditorCellPtr;
+			}
+			else
+			{
+				EditorCell = NewObject<UWorldPartitionEditorCell>(this, *FString::Printf(TEXT("EditorCell_X%lld_Y%lld_Z%lld_L%d"), CellCoord.X, CellCoord.Y, CellCoord.Z, CellCoord.Level), RF_Transient);
+				EditorCell->SetFlags(RF_Transactional);
+				EditorCell->Bounds = GetCellBounds(CellCoord);
 
-				Cells.Add(CellNode.Cell);
+				Cells.Add(EditorCell);
+				HashCells.Add(CellCoord, EditorCell);
 
 				UpdateHigherLevels(CellCoord, CurrentLevel);
 
-				Bounds += CellNode.Cell->Bounds;
+				Bounds += EditorCell->Bounds;
 			}
 
-			CellNode.Cell->AddActor(InActorDesc);
+			check(EditorCell);
+			EditorCell->AddActor(InActorDesc);
 		});
 
 		const int32 NewLevel = GetLevelForBox(Bounds);
@@ -157,26 +171,44 @@ void UWorldPartitionEditorSpatialHash::HashActor(FWorldPartitionActorDesc* InAct
 void UWorldPartitionEditorSpatialHash::UnhashActor(FWorldPartitionActorDesc* InActorDesc)
 {
 	check(InActorDesc);
-	if (InActorDesc->GetGridPlacement() == EActorGridPlacement::AlwaysLoaded)
+
+	if (IsActorAlwaysLoaded(InActorDesc))
 	{
 		AlwaysLoadedCell->RemoveActor(InActorDesc);
 	}
 	else
 	{
-		const FBox ActorBounds = GetActorBounds(InActorDesc);
-		const int32 CurrentLevel = GetLevelForBox(Bounds);
-		const int32 ActorLevel = GetLevelForBox(ActorBounds);
+		int32 CurrentLevel = GetLevelForBox(Bounds);
 
-		ForEachIntersectingCells(ActorBounds, ActorLevel, [&](const FCellCoord& CellCoord)
+		ForEachIntersectingCells(InActorDesc->GetBounds(), 0, [&](const FCellCoord& CellCoord)
 		{
-			FCellNode& CellNode = HashCells.FindChecked(CellCoord);
+			UWorldPartitionEditorCell* EditorCell = HashCells.FindChecked(CellCoord);
 
-			CellNode.Cell->RemoveActor(InActorDesc);
+			EditorCell->RemoveActor(InActorDesc);
 
-			if (!CellNode.Cell->Actors.Num())
+			if (!EditorCell->Actors.Num())
 			{
-				verify(Cells.Remove(CellNode.Cell));
-				CellNode.Cell = nullptr;
+				verify(Cells.Remove(EditorCell));
+				verify(HashCells.Remove(CellCoord));
+
+				FCellCoord LevelCellCoord = CellCoord;
+				for (int32 Level=1; Level<=CurrentLevel; Level++)
+				{
+					const uint32 ChildIndex = LevelCellCoord.GetChildIndex();
+
+					LevelCellCoord = LevelCellCoord.GetParentCellCoord();
+
+					FCellNode& CellNode = HashNodes.FindChecked(LevelCellCoord);
+
+					CellNode.RemoveChildNode(ChildIndex);
+
+					if (CellNode.HasChildNodes())
+					{
+						break;
+					}
+
+					HashNodes.Remove(LevelCellCoord);
+				}
 
 				bBoundsDirty = true;
 			}
@@ -227,27 +259,35 @@ int32 UWorldPartitionEditorSpatialHash::ForEachIntersectingCellInner(const FBox&
 {
 	int32 NumIntersecting = 0;
 
-	FCellNode CellNode;
-	if (FCellNode* CellNodePtr = HashCells.Find(CellCoord))
+	if (CellCoord.Level)
 	{
-		CellNode = *CellNodePtr;
-
-		if (CellNode.Cell)
+		FCellNode CellNode;
+		if (FCellNode* CellNodePtr = HashNodes.Find(CellCoord))
 		{
-			InOperation(CellNode.Cell);
-			NumIntersecting++;
-		}
+			CellNode = *CellNodePtr;
 
-		CellNode.ForEachChild([&](uint32 ChildIndex)
-		{
-			const FCellCoord ChildCellCoord = CellCoord.GetChildCellCoord(ChildIndex);
-			const FBox CellBounds = GetCellBounds(ChildCellCoord);
+			check(CellNode.HasChildNodes());
 
-			if (Box.Intersect(CellBounds))
+			CellNode.ForEachChild([&](uint32 ChildIndex)
 			{
-				NumIntersecting += ForEachIntersectingCellInner(Box, ChildCellCoord, InOperation);
-			}
-		});
+				const FCellCoord ChildCellCoord = CellCoord.GetChildCellCoord(ChildIndex);
+				const FBox CellBounds = GetCellBounds(ChildCellCoord);
+
+				if (Box.Intersect(CellBounds))
+				{
+					NumIntersecting += ForEachIntersectingCellInner(Box, ChildCellCoord, InOperation);
+				}
+			});
+		}
+	}
+	else
+	{
+		UWorldPartitionEditorCell* EditorCell = nullptr;
+		if (UWorldPartitionEditorCell** EditorCellPtr = HashCells.Find(CellCoord))
+		{
+			InOperation(*EditorCellPtr);
+			NumIntersecting++;
+		}
 	}
 
 	return NumIntersecting;
@@ -257,7 +297,9 @@ int32 UWorldPartitionEditorSpatialHash::ForEachIntersectingCell(const FBox& Box,
 {
 	int32 NumIntersecting = 0;
 
-	const int32 SearchLevel = GetLevelForBox(Bounds);
+	const FBox SearchBox = Box.Overlap(Bounds);
+	const int32 SearchLevel = GetLevelForBox(SearchBox);
+
 	ForEachIntersectingCells(Box, SearchLevel, [&](const FCellCoord& CellCoord)
 	{
 		NumIntersecting += ForEachIntersectingCellInner(Box, CellCoord, InOperation);
diff --git a/Engine/Source/Runtime/Engine/Private/WorldPartition/WorldPartitionSubsystem.cpp b/Engine/Source/Runtime/Engine/Private/WorldPartition/WorldPartitionSubsystem.cpp
--- a/Engine/Source/Runtime/Engine/Private/WorldPartition/WorldPartitionSubsystem.cpp
+++ b/Engine/Source/Runtime/Engine/Private/WorldPartition/WorldPartitionSubsystem.cpp
@@ -199,22 +199,6 @@ void UWorldPartitionSubsystem::UpdateActorDesc(AActor* Actor)
 	}
 }
 
-void UWorldPartitionSubsystem::AddActor(AActor* Actor)
-{
-	if (UWorldPartition* MainPartition = GetMainWorldPartition())
-	{
-		MainPartition->AddActor(Actor);
-	}
-}
-
-void UWorldPartitionSubsystem::RemoveActor(AActor* Actor)
-{
-	if (UWorldPartition* MainPartition = GetMainWorldPartition())
-	{
-		MainPartition->RemoveActor(Actor);
-	}
-}
-
 void UWorldPartitionSubsystem::RegisterActorDescFactory(TSubclassOf<AActor> Class, FWorldPartitionActorDescFactory* Factory)
 {
 	if (IsEnabled())
diff --git a/Engine/Source/Runtime/Engine/Public/WorldPartition/WorldPartition.h b/Engine/Source/Runtime/Engine/Public/WorldPartition/WorldPartition.h
--- a/Engine/Source/Runtime/Engine/Public/WorldPartition/WorldPartition.h
+++ b/Engine/Source/Runtime/Engine/Public/WorldPartition/WorldPartition.h
@@ -97,9 +97,6 @@ private:
 	TArray<const FWorldPartitionActorDesc*> GetIntersectingActorDescs(const FBox& Box, TSubclassOf<AActor> ActorClass) const;
 	// UActorPartitionSubsystem interface-
 
-	void AddActor(AActor* InActor);
-	void RemoveActor(AActor* InActor);
-
 	void ApplyActorTransform(AActor* InActor, const FTransform& InTransform);
 
 	static void RegisterActorDescFactory(TSubclassOf<AActor> Class, FWorldPartitionActorDescFactory* Factory);
diff --git a/Engine/Source/Runtime/Engine/Public/WorldPartition/WorldPartitionEditorSpatialHash.h b/Engine/Source/Runtime/Engine/Public/WorldPartition/WorldPartitionEditorSpatialHash.h
--- a/Engine/Source/Runtime/Engine/Public/WorldPartition/WorldPartitionEditorSpatialHash.h
+++ b/Engine/Source/Runtime/Engine/Public/WorldPartition/WorldPartitionEditorSpatialHash.h
@@ -125,7 +125,6 @@ class ENGINE_API UWorldPartitionEditorSpatialHash : public UWorldPartitionEditor
 	{
 		FCellNode()
 			: ChildNodesMask(0)
-			, Cell(nullptr)
 		{}
 
 		inline bool HasChildNodes() const
@@ -171,7 +170,6 @@ class ENGINE_API UWorldPartitionEditorSpatialHash : public UWorldPartitionEditor
 		}
 
 		uint8 ChildNodesMask;
-		UWorldPartitionEditorCell* Cell;
 	};
 
 public:
@@ -192,12 +190,14 @@ public:
 #if WITH_EDITORONLY_DATA
 private:
 	FBox GetActorBounds(FWorldPartitionActorDesc* InActorDesc) const;
+	bool IsActorAlwaysLoaded(FWorldPartitionActorDesc* InActorDesc) const;
 	int32 ForEachIntersectingCellInner(const FBox& Box, const FCellCoord& CellCoord, TFunctionRef<void(UWorldPartitionEditorCell*)> InOperation);
 
 	UPROPERTY(Config)
 	int32 CellSize;
 
-	TMap<FCellCoord, FCellNode> HashCells;
+	TMap<FCellCoord, FCellNode> HashNodes;
+	TMap<FCellCoord, UWorldPartitionEditorCell*> HashCells;
 
 	UPROPERTY(Transient)
 	TSet<UWorldPartitionEditorCell*> Cells;
diff --git a/Engine/Source/Runtime/Engine/Public/WorldPartition/WorldPartitionSubsystem.h b/Engine/Source/Runtime/Engine/Public/WorldPartition/WorldPartitionSubsystem.h
--- a/Engine/Source/Runtime/Engine/Public/WorldPartition/WorldPartitionSubsystem.h
+++ b/Engine/Source/Runtime/Engine/Public/WorldPartition/WorldPartitionSubsystem.h
@@ -47,8 +47,6 @@ public:
 	TArray<const FWorldPartitionActorDesc*> GetIntersectingActorDescs(const FBox& Box, TSubclassOf<AActor> ActorClass) const;
 
 	void UpdateActorDesc(AActor* Actor);
-	void AddActor(AActor* Actor);
-	void RemoveActor(AActor* Actor);
 
 	void RegisterActorDescFactory(TSubclassOf<AActor> Class, FWorldPartitionActorDescFactory* Factory);
 
