Made USkinnedMeshComponent::GetCPUSkinnedVertices() const
* Found issue where an existing forced LOD would get overriden by this call

diff --git a/Engine/Source/Runtime/Engine/Classes/Components/SkinnedMeshComponent.h b/Engine/Source/Runtime/Engine/Classes/Components/SkinnedMeshComponent.h
--- a/Engine/Source/Runtime/Engine/Classes/Components/SkinnedMeshComponent.h
+++ b/Engine/Source/Runtime/Engine/Classes/Components/SkinnedMeshComponent.h
@@ -350,7 +350,7 @@ public:
 	 * @param	OutVertices		The skinned vertices
 	 * @param	InLODIndex		The LOD we want to export
 	 */
-	void GetCPUSkinnedVertices(TArray<struct FFinalSkinVertex>& OutVertices, int32 InLODIndex);
+	void GetCPUSkinnedVertices(TArray<struct FFinalSkinVertex>& OutVertices, int32 InLODIndex) const;
 
 #if !(UE_BUILD_SHIPPING || UE_BUILD_TEST)
 	/** Get whether to draw this mesh's debug skeleton */
diff --git a/Engine/Source/Runtime/Engine/Private/Components/SkinnedMeshComponent.cpp b/Engine/Source/Runtime/Engine/Private/Components/SkinnedMeshComponent.cpp
--- a/Engine/Source/Runtime/Engine/Private/Components/SkinnedMeshComponent.cpp
+++ b/Engine/Source/Runtime/Engine/Private/Components/SkinnedMeshComponent.cpp
@@ -3137,25 +3137,23 @@ void USkinnedMeshComponent::SetComponentSpaceTransformsDoubleBuffering(bool bInD
 	}
 }
 
-void USkinnedMeshComponent::GetCPUSkinnedVertices(TArray<FFinalSkinVertex>& OutVertices, int32 InLODIndex)
+void USkinnedMeshComponent::GetCPUSkinnedVertices(TArray<FFinalSkinVertex>& OutVertices, int32 InLODIndex) const
 {
-	if (USkinnedMeshComponent* MasterPoseComponentPtr = MasterPoseComponent.Get())
-	{
-		MasterPoseComponentPtr->SetForcedLOD(InLODIndex + 1);
-		MasterPoseComponentPtr->UpdateLODStatus();
-		MasterPoseComponentPtr->RefreshBoneTransforms(nullptr);
-	}
-	else
-	{
-		SetForcedLOD(InLODIndex + 1);
-		UpdateLODStatus();
-		RefreshBoneTransforms(nullptr);
-	}
+	// Work around the fact that non-const methods must be called to perform the skinning
+	// Component state should be left unchanged at the end of this call.
+	USkinnedMeshComponent* MutableThis = const_cast<USkinnedMeshComponent*>(this);
+
+	USkinnedMeshComponent* PoseComponent = MutableThis->MasterPoseComponent.Get() ? MutableThis->MasterPoseComponent.Get() : MutableThis;
 
+	int32 CachedForcedLOD = PoseComponent->GetForcedLOD();
+	PoseComponent->SetForcedLOD(InLODIndex + 1);
+	PoseComponent->UpdateLODStatus();
+	PoseComponent->RefreshBoneTransforms(nullptr);
+	
 	// switch to CPU skinning
 	const bool bCachedCPUSkinning = GetCPUSkinningEnabled();
 	constexpr bool bRecreateRenderStateImmediately = true;
-	SetCPUSkinningEnabled(true, bRecreateRenderStateImmediately);
+	MutableThis->SetCPUSkinningEnabled(true, bRecreateRenderStateImmediately);
 
 	check(MeshObject);
 	check(MeshObject->IsCPUSkinned());
@@ -3164,8 +3162,8 @@ void USkinnedMeshComponent::GetCPUSkinnedVertices(TArray<FFinalSkinVertex>& OutV
 	OutVertices = static_cast<FSkeletalMeshObjectCPUSkin*>(MeshObject)->GetCachedFinalVertices();
 	
 	// switch skinning mode, LOD etc. back
-	SetForcedLOD(0);
-	SetCPUSkinningEnabled(bCachedCPUSkinning, bRecreateRenderStateImmediately);
+	PoseComponent->SetForcedLOD(CachedForcedLOD);
+	MutableThis->SetCPUSkinningEnabled(bCachedCPUSkinning, bRecreateRenderStateImmediately);
 }
 
 void USkinnedMeshComponent::ReleaseResources()
