PR #7231: Fix conditional vlog macro log output (Contributed by [REDACTED_PERSON])

diff --git a/Engine/Source/Runtime/Engine/Public/VisualLogger/VisualLogger.h b/Engine/Source/Runtime/Engine/Public/VisualLogger/VisualLogger.h
--- a/Engine/Source/Runtime/Engine/Public/VisualLogger/VisualLogger.h
+++ b/Engine/Source/Runtime/Engine/Public/VisualLogger/VisualLogger.h
@@ -25,7 +25,7 @@
 #define UE_CVLOG(Condition, LogOwner, CategoryName, Verbosity, Format, ...)  if(FVisualLogger::IsRecording() && Condition) {UE_VLOG(LogOwner, CategoryName, Verbosity, Format, ##__VA_ARGS__);} 
 // Text, log with output to regular unreal logs too
 #define UE_VLOG_UELOG(LogOwner, CategoryName, Verbosity, Format, ...) { if(FVisualLogger::IsRecording()) FVisualLogger::CategorizedLogf(LogOwner, CategoryName, ELogVerbosity::Verbosity, Format, ##__VA_ARGS__); UE_LOG(CategoryName, Verbosity, Format, ##__VA_ARGS__); }
-#define UE_CVLOG_UELOG(Condition, LogOwner, CategoryName, Verbosity, Format, ...)  if(FVisualLogger::IsRecording() && Condition) {UE_VLOG_UELOG(LogOwner, CategoryName, Verbosity, Format, ##__VA_ARGS__);} 
+#define UE_CVLOG_UELOG(Condition, LogOwner, CategoryName, Verbosity, Format, ...)  if(Condition) {UE_VLOG_UELOG(LogOwner, CategoryName, Verbosity, Format, ##__VA_ARGS__);} 
 // Segment shape
 #define UE_VLOG_SEGMENT(LogOwner, CategoryName, Verbosity, SegmentStart, SegmentEnd, Color, Format, ...) if(FVisualLogger::IsRecording()) FVisualLogger::GeometryShapeLogf(LogOwner, CategoryName, ELogVerbosity::Verbosity, SegmentStart, SegmentEnd, Color, 0, Format, ##__VA_ARGS__)
 #define UE_CVLOG_SEGMENT(Condition, LogOwner, CategoryName, Verbosity, SegmentStart, SegmentEnd, Color, Format, ...) if(FVisualLogger::IsRecording() && Condition) {UE_VLOG_SEGMENT(LogOwner, CategoryName, Verbosity, SegmentStart, SegmentEnd, Color, Format, ##__VA_ARGS__);}
