Remove unused private fields - Blueprints

diff --git a/Engine/Source/Runtime/Engine/Private/SimpleConstructionScript.cpp b/Engine/Source/Runtime/Engine/Private/SimpleConstructionScript.cpp
--- a/Engine/Source/Runtime/Engine/Private/SimpleConstructionScript.cpp
+++ b/Engine/Source/Runtime/Engine/Private/SimpleConstructionScript.cpp
@@ -336,7 +336,6 @@ void USimpleConstructionScript::FixupSceneNodeHierarchy()
 		FSceneHierarchyMapper(TArray<USCS_Node*>& RootNodesIn)
 			: RootNodeList(RootNodesIn)
 			, PendingParent(nullptr)
-			, bBreak(false)
 		{}
 
 		/** Identifies orphan (root) nodes, and fixes up broken/cyclic tree linkages */
@@ -508,7 +507,6 @@ void USimpleConstructionScript::FixupSceneNodeHierarchy()
 		TSet<USCS_Node*> VisitedNodes;
 		TSet<USCS_Node*> OrphanedNodes;
 		USCS_Node* PendingParent;
-		bool bBreak;
 	};
 
 	FSceneHierarchyMapper HierarchyMapper(RootNodes);
