World Partition - Added support for cell preview actors

Currently applying a post process effect to give a red tint to the preview meshes.
Ultimately, we'll probably want to come up with a better solution since the post process effect doesn't work in all visualization modes.

diff --git a/Engine/Source/Editor/WorldPartitionEditor/Private/WorldPartition/SWorldPartitionEditorGrid2D.cpp b/Engine/Source/Editor/WorldPartitionEditor/Private/WorldPartition/SWorldPartitionEditorGrid2D.cpp
--- a/Engine/Source/Editor/WorldPartitionEditor/Private/WorldPartition/SWorldPartitionEditorGrid2D.cpp
+++ b/Engine/Source/Editor/WorldPartitionEditor/Private/WorldPartition/SWorldPartitionEditorGrid2D.cpp
@@ -33,6 +33,7 @@ void SWorldPartitionEditorGrid2D::FEditorCommands::RegisterCommands()
 	UI_COMMAND(LoadSelectedCells, "Load Selected Cells", "Load the selected cells.", EUserInterfaceActionType::Button, FInputChord());
 	UI_COMMAND(UnloadSelectedCells, "Unload Selected Cells", "Unload the selected cells.", EUserInterfaceActionType::Button, FInputChord());
 	UI_COMMAND(MoveCameraHere, "Move Camera Here", "MOve the camera to the selected position.", EUserInterfaceActionType::Button, FInputChord());
+	UI_COMMAND(UpdateSelectedCellsPreviewMesh, "Update Cells Preview Mesh", "Update the selected cells preview mesh.", EUserInterfaceActionType::Button, FInputChord());
 }
 
 SWorldPartitionEditorGrid2D::SWorldPartitionEditorGrid2D()
@@ -156,6 +157,7 @@ void SWorldPartitionEditorGrid2D::Construct(const FArguments& InArgs)
 	ActionList.MapAction(Commands.LoadSelectedCells, FExecuteAction::CreateSP(this, &SWorldPartitionEditorGrid2D::LoadSelectedCells), FCanExecuteAction::CreateLambda(HasSelectedCells));
 	ActionList.MapAction(Commands.UnloadSelectedCells, FExecuteAction::CreateSP(this, &SWorldPartitionEditorGrid2D::UnloadSelectedCells), FCanExecuteAction::CreateLambda(HasSelectedCells));
 	ActionList.MapAction(Commands.MoveCameraHere, FExecuteAction::CreateSP(this, &SWorldPartitionEditorGrid2D::MoveCameraHere));
+	ActionList.MapAction(Commands.UpdateSelectedCellsPreviewMesh, FExecuteAction::CreateSP(this, &SWorldPartitionEditorGrid2D::UpdateSelectedCellsPreviewMesh), FCanExecuteAction::CreateLambda(HasSelectedCells));
 }
 
 void SWorldPartitionEditorGrid2D::LoadSelectedCells()
@@ -172,6 +174,12 @@ void SWorldPartitionEditorGrid2D::UnloadSelectedCells()
 	RefreshSceneOutliner();
 }
 
+void SWorldPartitionEditorGrid2D::UpdateSelectedCellsPreviewMesh()
+{
+	WorldPartition->UpdateCellsPreviewMesh(SelectedCells.Array());
+	GEditor->RedrawLevelEditingViewports();
+}
+
 void SWorldPartitionEditorGrid2D::MoveCameraHere()
 {
 	for (FLevelEditorViewportClient* LevelVC : GEditor->GetLevelViewportClients())
@@ -244,6 +252,7 @@ FReply SWorldPartitionEditorGrid2D::OnMouseButtonUp(const FGeometry& MyGeometry,
 				MenuBuilder.AddMenuEntry(Commands.LoadSelectedCells);
 				MenuBuilder.AddMenuEntry(Commands.UnloadSelectedCells);
 				MenuBuilder.AddMenuEntry(Commands.MoveCameraHere);
+				MenuBuilder.AddMenuEntry(Commands.UpdateSelectedCellsPreviewMesh);
 			}
 			MenuBuilder.EndSection();
 
diff --git a/Engine/Source/Editor/WorldPartitionEditor/Public/WorldPartition/SWorldPartitionEditorGrid2D.h b/Engine/Source/Editor/WorldPartitionEditor/Public/WorldPartition/SWorldPartitionEditorGrid2D.h
--- a/Engine/Source/Editor/WorldPartitionEditor/Public/WorldPartition/SWorldPartitionEditorGrid2D.h
+++ b/Engine/Source/Editor/WorldPartitionEditor/Public/WorldPartition/SWorldPartitionEditorGrid2D.h
@@ -21,6 +21,7 @@ protected:
 		TSharedPtr<FUICommandInfo> LoadSelectedCells;
 		TSharedPtr<FUICommandInfo> UnloadSelectedCells;
 		TSharedPtr<FUICommandInfo> MoveCameraHere;
+		TSharedPtr<FUICommandInfo> UpdateSelectedCellsPreviewMesh;
 
 		/**
 		 * Initialize commands
@@ -37,6 +38,7 @@ public:
 	void LoadSelectedCells();
 	void UnloadSelectedCells();
 	void MoveCameraHere();
+	void UpdateSelectedCellsPreviewMesh();
 
 	// IWorldPartitionEditor interface
 	virtual void InvalidatePartition() override;
diff --git a/Engine/Source/Runtime/Engine/Private/WorldPartition/HLOD/HLODBuilder.cpp b/Engine/Source/Runtime/Engine/Private/WorldPartition/HLOD/HLODBuilder.cpp
--- a/Engine/Source/Runtime/Engine/Private/WorldPartition/HLOD/HLODBuilder.cpp
+++ b/Engine/Source/Runtime/Engine/Private/WorldPartition/HLOD/HLODBuilder.cpp
@@ -7,6 +7,7 @@
 #include "WorldPartition/WorldPartition.h"
 #include "WorldPartition/HLOD/HLODActor.h"
 #include "WorldPartition/HLOD/HLODLayer.h"
+#include "WorldPartition/WorldPartitionEditorCellPreviewActor.h"
 
 #include "Algo/ForEach.h"
 #include "Algo/Transform.h"
@@ -18,6 +19,16 @@
 #include "Engine/StaticMesh.h"
 #include "Materials/Material.h"
 
+#include "Engine/StaticMesh.h"
+#include "Engine/StaticMeshActor.h"
+#include "IMeshDescriptionModule.h"
+#include "Materials/Material.h"
+#include "StaticMeshAttributes.h"
+#include "StaticMeshOperations.h"
+#include "StaticMeshComponentAdapter.h"
+#include "MeshDescription.h"
+#include "Async/ParallelFor.h"
+
 /**
  * Base class for all HLODBuilders
  */
@@ -267,4 +278,177 @@ void FHLODBuilderUtilities::BuildHLODs(UWorldPartition* InWorldPartition, FName
 	}
 }
 
+AWorldPartitionEditorCellPreview* FHLODBuilderUtilities::BuildCellPreviewMesh(const TArray<AActor*>& InCellActors, const FBox& InCellBounds)
+{
+	TRACE_CPUPROFILER_EVENT_SCOPE(FHLODBuilderUtilities::BuildCellPreviewMesh);
+
+	TMap<UHLODLayer*, TArray<UStaticMeshComponent*>> CellLODPrimitives;
+
+	// Gather all components that will be used to create the preview mesh
+	for (AActor* Actor : InCellActors)
+	{
+		UHLODLayer* ActorHLODLayer = UHLODLayer::GetHLODLayer(Actor);
+
+		if (ActorHLODLayer && UHLODLayer::ShouldIncludeInHLOD(Actor))
+		{
+			TArray<UStaticMeshComponent*> StaticMeshComponent;
+			Actor->GetComponents(StaticMeshComponent);
+
+			for (UStaticMeshComponent* Component : StaticMeshComponent)
+			{
+				if (UHLODLayer::ShouldIncludeInHLOD(Component, ActorHLODLayer->GetLevels().Num() - 1))
+				{
+					CellLODPrimitives.FindOrAdd(ActorHLODLayer).Add(Component);
+				}
+			}
+		}
+	}
+
+	if (CellLODPrimitives.IsEmpty())
+	{
+		return nullptr;
+	}
+
+	UWorld* World = InCellActors[0]->GetWorld();
+
+	const IMeshDescriptionModule& MeshDescriptionModule = IMeshDescriptionModule::Get();
+	const IMeshMergeUtilities& MergeUtilities = FModuleManager::Get().LoadModuleChecked<IMeshMergeModule>("MeshMergeUtilities").GetUtilities();
+
+	// Spawn a cell preview actor
+	FActorSpawnParameters SpawnParams;
+	SpawnParams.bDeferConstruction = true;
+	SpawnParams.bCreateActorPackage = true;
+	AWorldPartitionEditorCellPreview* CellPreviewActor = World->SpawnActor<AWorldPartitionEditorCellPreview>(AWorldPartitionEditorCellPreview::StaticClass(), FTransform(InCellBounds.GetCenter()), SpawnParams);
+	CellPreviewActor->SetCellBounds(InCellBounds);
+	
+	TArray<UStaticMeshComponent*> ComponentsToAdd;
+
+	for (const auto& Primitives : CellLODPrimitives)
+	{
+		const UHLODLayer* HLODLayer = Primitives.Key;
+		const TArray<UStaticMeshComponent*>& StaticMeshComponents = Primitives.Value;
+
+		for (const FHLODLevelSettings& LevelSettings : HLODLayer->GetLevels())
+		{
+			switch (LevelSettings.LevelType)
+			{
+			case EHLODLevelType::Instancing:
+			{
+				// Gather all meshes to instantiate along with transforms
+				TMap<UStaticMesh*, TArray<UStaticMeshComponent*>> Instances;
+				for (UStaticMeshComponent* SMC : StaticMeshComponents)
+				{
+					Instances.FindOrAdd(SMC->GetStaticMesh()).Add(SMC);
+				}
+
+				for (const auto& Entry : Instances)
+				{
+					// Get or create preview mesh
+					UInstancedStaticMeshComponent* Component = NewObject<UInstancedStaticMeshComponent>(CellPreviewActor);
+					Component->SetStaticMesh(Entry.Key);
+					Component->SetForcedLodModel(Entry.Key->GetNumLODs());
+
+					for (UStaticMeshComponent* SMC : Entry.Value)
+					{
+						Component->AddInstanceWorldSpace(SMC->GetComponentTransform());
+					}
+
+					ComponentsToAdd.Add(Component);
+				}
+			}
+			break;
+
+			default:
+			{
+				FMeshDescription MergedMeshDescription;
+				FStaticMeshAttributes(MergedMeshDescription).Register();
+
+				TArray<FMeshDescription> SourceMeshesDescriptions;
+				SourceMeshesDescriptions.SetNum(StaticMeshComponents.Num());
+
+				{
+					TRACE_CPUPROFILER_EVENT_SCOPE(Gather Meshes);
+					ParallelFor(SourceMeshesDescriptions.Num(), [&MergeUtilities, &StaticMeshComponents, &SourceMeshesDescriptions](uint32 Index)
+					{
+						FMeshDescription& SourceMeshDescription = SourceMeshesDescriptions[Index];
+						FStaticMeshAttributes(SourceMeshDescription).Register();
+
+						UStaticMeshComponent* SMComponent = StaticMeshComponents[Index];
+					
+						MergeUtilities.RetrieveMeshDescription(SMComponent, SMComponent->GetStaticMesh()->GetNumLODs() - 1, SourceMeshDescription, false);
+					});
+				}
+
+				{
+					TRACE_CPUPROFILER_EVENT_SCOPE(Merge Meshes);
+
+					FStaticMeshOperations::FAppendSettings AppendSettings;
+					AppendSettings.bMergeVertexColor = false;
+					for (int32 ChannelIdx = 0; ChannelIdx < FStaticMeshOperations::FAppendSettings::MAX_NUM_UV_CHANNELS; ++ChannelIdx)
+					{
+						AppendSettings.bMergeUVChannels[ChannelIdx] = false;
+					}
+				
+					MergedMeshDescription.CreatePolygonGroupWithID(FPolygonGroupID(0));
+
+					AppendSettings.PolygonGroupsDelegate = FAppendPolygonGroupsDelegate::CreateLambda([](const FMeshDescription& SourceMesh, FMeshDescription& TargetMesh, PolygonGroupMap& RemapPolygonGroups)
+					{
+						for (FPolygonGroupID SourcePolygonGroupID : SourceMesh.PolygonGroups().GetElementIDs())
+						{
+							RemapPolygonGroups.Add(SourcePolygonGroupID, FPolygonGroupID(0));
+						}
+					});
+
+					for (const FMeshDescription& SourceMeshDescription : SourceMeshesDescriptions)
+					{
+						FStaticMeshOperations::AppendMeshDescription(SourceMeshDescription, MergedMeshDescription, AppendSettings);
+					}
+				}
+
+				FStaticMeshComponentRecreateRenderStateContext RecreateRenderStateContext(FindObject<UStaticMesh>(CellPreviewActor->GetPackage(), TEXT("EditorCellPreviewMesh")));
+
+				UStaticMesh* StaticMesh = NewObject<UStaticMesh>(CellPreviewActor->GetPackage(), TEXT("EditorCellPreviewMesh"));
+				StaticMesh->BuildFromMeshDescriptions( { &MergedMeshDescription } );
+				
+				UStaticMeshComponent* Component = NewObject<UStaticMeshComponent>(CellPreviewActor);
+				Component->SetStaticMesh(StaticMesh);
+
+				ComponentsToAdd.Add(Component);
+			}
+			
+			}
+		}
+	}
+
+	// If for some reason no components were created, delete the actor.
+	if (ComponentsToAdd.IsEmpty())
+	{
+		World->DestroyActor(CellPreviewActor);
+		return nullptr;
+	}
+
+	// Setup the newly created components
+	for (UPrimitiveComponent* ComponentToAdd : ComponentsToAdd)
+	{
+		ComponentToAdd->AttachToComponent(CellPreviewActor->GetRootComponent(), FAttachmentTransformRules::KeepWorldTransform);
+
+		// Setup custom depth rendering to achieve a red tint using a post process material
+		// Temporary until a cleaner solution is implemented
+		const int32 CellPreviewStencilValue = 180;
+		ComponentToAdd->bRenderCustomDepth = true;
+		ComponentToAdd->CustomDepthStencilValue = CellPreviewStencilValue;
+
+		ComponentToAdd->SetMobility(EComponentMobility::Static);
+		ComponentToAdd->SetCollisionEnabled(ECollisionEnabled::NoCollision);
+
+		ComponentToAdd->RegisterComponent();
+		ComponentToAdd->MarkRenderStateDirty();
+
+		CellPreviewActor->AddInstanceComponent(ComponentToAdd);
+	}
+
+	CellPreviewActor->SetVisibility(false);
+	return CellPreviewActor;
+}
+
 #endif
diff --git a/Engine/Source/Runtime/Engine/Private/WorldPartition/HLOD/HLODBuilder.h b/Engine/Source/Runtime/Engine/Private/WorldPartition/HLOD/HLODBuilder.h
--- a/Engine/Source/Runtime/Engine/Private/WorldPartition/HLOD/HLODBuilder.h
+++ b/Engine/Source/Runtime/Engine/Private/WorldPartition/HLOD/HLODBuilder.h
@@ -10,6 +10,7 @@
 class AActor;
 class UWorldPartition;
 class UHLODLayer;
+class AWorldPartitionEditorCellPreview;
 
 /**
  * Tools for building HLODs in WorldPartition
@@ -27,6 +28,15 @@ public:
 	 * @param 	InSubActors			The actors from which we'll gather geometry to generate an HLOD mesh
 	 */
 	static void BuildHLODs(UWorldPartition* InWorldPartition, FName InCellName, const UHLODLayer* InHLODLayer, const TArray<AActor*>& InSubActors);
+
+	/**
+	 * Generate a cell preview mesh from the given actors.
+	 * 
+	 * @param	InCellActors		Actors from which meshes will be gathered
+	 * @param	InCellBounds		Bounds of the cell
+	 * @return A new cell preview actor
+	 */
+	static AWorldPartitionEditorCellPreview* BuildCellPreviewMesh(const TArray<AActor*>& InCellActors, const FBox& InCellBounds);
 };
 
 #endif
\ No newline at end of file
diff --git a/Engine/Source/Runtime/Engine/Private/WorldPartition/WorldPartition.cpp b/Engine/Source/Runtime/Engine/Private/WorldPartition/WorldPartition.cpp
--- a/Engine/Source/Runtime/Engine/Private/WorldPartition/WorldPartition.cpp
+++ b/Engine/Source/Runtime/Engine/Private/WorldPartition/WorldPartition.cpp
@@ -4,7 +4,6 @@
 	WorldPartition.cpp: UWorldPartition implementation
 =============================================================================*/
 #include "WorldPartition/WorldPartition.h"
-#include "WorldPartition/WorldPartitionEditorCell.h"
 #include "WorldPartition/WorldPartitionRuntimeCell.h"
 #include "WorldPartition/WorldPartitionStreamingPolicy.h"
 #include "WorldPartition/WorldPartitionLevelStreamingPolicy.h"
@@ -33,9 +32,12 @@
 #include "UnrealEdMisc.h"
 #include "WorldPartition/WorldPartitionActorDescFactory.h"
 #include "WorldPartition/WorldPartitionLevelStreamingDynamic.h"
+#include "WorldPartition/WorldPartitionEditorCell.h"
+#include "WorldPartition/WorldPartitionEditorCellPreviewActor.h"
 #include "WorldPartition/WorldPartitionEditorHash.h"
 #include "WorldPartition/WorldPartitionEditorSpatialHash.h"
 #include "WorldPartition/WorldPartitionRuntimeHash.h"
+#include "WorldPartition/HLOD/HLODBuilder.h"
 #endif
 
 DEFINE_LOG_CATEGORY(LogWorldPartition);
@@ -318,6 +320,24 @@ void UWorldPartition::Initialize(UWorld* InWorld, const FTransform& InTransform)
 					UpdateLoadingEditorCell(Cell, true);
 				});
 			}
+
+			// Bind editor cells to their preview actors
+			EditorHash->ForEachCell([this](UWorldPartitionEditorCell* Cell)
+			{
+				TArray<const FWorldPartitionActorDesc*> ActorsDescs = GetIntersectingActorDescs(Cell->Bounds, AWorldPartitionEditorCellPreview::StaticClass());
+				
+				for(const FWorldPartitionActorDesc* ActorDesc : ActorsDescs)
+				{
+					if (ActorDesc->GetBounds() == Cell->Bounds)
+					{
+						AWorldPartitionEditorCellPreview* CellPreviewActor = CastChecked<AWorldPartitionEditorCellPreview>(ActorDesc->GetActor());
+						Cell->CellPreviewActor = CellPreviewActor;
+
+						CellPreviewActor->SetVisibility(true);
+						break;
+					}
+				}
+			});
 		}
 	}
 #endif
@@ -334,6 +354,9 @@ void UWorldPartition::Initialize(UWorld* InWorld, const FTransform& InTransform)
 	{
 		PrepareForPIE();
 	}
+	
+	// Ensure editor cells preview meshes are rendered with a red tint
+	World->SpawnActor<AWorldPartitionUnloadedCellPreviewPostProcessVolume>();
 #endif
 }
 
@@ -904,6 +927,35 @@ void UWorldPartition::UpdateLoadingEditorCell(UWorldPartitionEditorCell* Cell, b
 	}
 
 	Cell->bLoaded = bShouldBeLoaded;
+
+	if (Cell->CellPreviewActor)
+	{
+		Cell->CellPreviewActor->SetVisibility(!Cell->bLoaded);
+	}
+}
+
+void UWorldPartition::UpdateCellsPreviewMesh(const TArray<UWorldPartitionEditorCell*>& CellsToUpdate)
+{
+	TRACE_CPUPROFILER_EVENT_SCOPE(UWorldPartition::UpdateCellsPreviewMesh);
+
+	LoadEditorCells(CellsToUpdate);
+
+	for (UWorldPartitionEditorCell* Cell : CellsToUpdate)
+	{
+		TArray<AActor*> CellActors;
+		Algo::Transform(Cell->LoadedActors, CellActors, [](FWorldPartitionActorDesc* ActorDesc) { return ActorDesc->GetActor(); });
+
+		if (!CellActors.IsEmpty())
+		{
+			if (Cell->CellPreviewActor)
+			{
+				World->DestroyActor(Cell->CellPreviewActor);
+				Cell->CellPreviewActor = nullptr;
+			}
+
+			Cell->CellPreviewActor = FHLODBuilderUtilities::BuildCellPreviewMesh(CellActors, Cell->Bounds);
+		}
+	}
 }
 
 void UWorldPartition::CreateLayers(const TSet<FName>& LayerNames)
diff --git a/Engine/Source/Runtime/Engine/Private/WorldPartition/WorldPartitionEditorCellPreviewActor.cpp b/Engine/Source/Runtime/Engine/Private/WorldPartition/WorldPartitionEditorCellPreviewActor.cpp
new file mode 100644
--- /dev/null
+++ b/Engine/Source/Runtime/Engine/Private/WorldPartition/WorldPartitionEditorCellPreviewActor.cpp
@@ -0,0 +1,83 @@
+// Copyright Epic Games, Inc. All Rights Reserved.
+
+#include "WorldPartition/WorldPartitionEditorCellPreviewActor.h"
+
+#if WITH_EDITOR
+#include "UObject/ConstructorHelpers.h"
+#include "Engine/World.h"
+#include "Materials/Material.h"
+#include "WorldPartition/WorldPartitionEditorHash.h"
+#include "WorldPartition/WorldPartition.h"
+#include "WorldPartition/WorldPartitionActorDesc.h"
+#include "WorldPartition/WorldPartitionSubsystem.h"
+#endif
+
+AWorldPartitionEditorCellPreview::AWorldPartitionEditorCellPreview(const FObjectInitializer& ObjectInitializer)
+	: Super(ObjectInitializer)
+{
+#if WITH_EDITOR
+	bVisible = false;
+
+	bIsEditorOnlyActor = true;
+
+	SetActorEnableCollision(false);
+
+	RootComponent = CreateDefaultSubobject<USceneComponent>(USceneComponent::GetDefaultSceneRootVariableName());
+	RootComponent->Mobility = EComponentMobility::Static;
+
+	SetActorHiddenInGame(true);
+#endif
+}
+
+#if WITH_EDITOR
+EActorGridPlacement AWorldPartitionEditorCellPreview::GetDefaultGridPlacement() const
+{
+	return EActorGridPlacement::AlwaysLoaded;
+}
+
+void AWorldPartitionEditorCellPreview::SetVisibility(bool bInVisible)
+{
+	bVisible = bInVisible;
+
+	if (HasValidRootComponent())
+	{
+		GetRootComponent()->SetVisibility(bVisible, true);
+	}
+}
+
+bool AWorldPartitionEditorCellPreview::IsSelectable() const
+{
+	return false;
+}
+
+void AWorldPartitionEditorCellPreview::GetActorBounds(bool bOnlyCollidingComponents, FVector& Origin, FVector& BoxExtent, bool bIncludeFromChildActors) const
+{
+	Origin = CellBounds.GetCenter();
+	BoxExtent = CellBounds.GetExtent();
+}
+
+void AWorldPartitionEditorCellPreview::PostRegisterAllComponents()
+{
+	Super::PostRegisterAllComponents();
+	SetVisibility(bVisible);
+}
+
+#endif // WITH_EDITOR
+
+AWorldPartitionUnloadedCellPreviewPostProcessVolume::AWorldPartitionUnloadedCellPreviewPostProcessVolume(const FObjectInitializer& ObjectInitializer)
+	: Super(ObjectInitializer)
+{
+#if WITH_EDITOR
+	if (!IsTemplate())
+	{
+		static ConstructorHelpers::FObjectFinder<UMaterial> UnloadedCellPreviewMat(TEXT("/Engine/EditorMaterials/WorldPartition/UnloadedCellPreview_PP"));
+		if (ensure(UnloadedCellPreviewMat.Object))
+		{
+			SetFlags(RF_Transient);
+			bIsEditorOnlyActor = true;
+			bUnbound = true;
+			Settings.WeightedBlendables.Array.Emplace(1.0f, UnloadedCellPreviewMat.Object);
+		}
+	}
+#endif // WITH_EDITOR
+}
diff --git a/Engine/Source/Runtime/Engine/Public/WorldPartition/WorldPartition.h b/Engine/Source/Runtime/Engine/Public/WorldPartition/WorldPartition.h
--- a/Engine/Source/Runtime/Engine/Public/WorldPartition/WorldPartition.h
+++ b/Engine/Source/Runtime/Engine/Public/WorldPartition/WorldPartition.h
@@ -128,10 +128,11 @@ public:
 	bool IsPreCooked() const { return bIsPreCooked; }
 	void SetIsPreCooked(bool bInIsPreCooked) { bIsPreCooked = bInIsPreCooked; }
 
-	bool GenerateStreaming(EWorldPartitionStreamingMode Mode);
 	FBox GetWorldBounds() const;
 
+	bool GenerateStreaming(EWorldPartitionStreamingMode Mode);
 	void GenerateHLOD();
+	void UpdateCellsPreviewMesh(const TArray<UWorldPartitionEditorCell*>& Cells);
 
 	// Clustering
 	struct FActorCluster
@@ -240,4 +241,4 @@ private:
 
 	UWorldPartitionStreamingPolicy* GetStreamingPolicy() const;
 	friend class UWorldPartitionStreamingPolicy;
-};
\ No newline at end of file
+};
diff --git a/Engine/Source/Runtime/Engine/Public/WorldPartition/WorldPartitionEditorCell.h b/Engine/Source/Runtime/Engine/Public/WorldPartition/WorldPartitionEditorCell.h
--- a/Engine/Source/Runtime/Engine/Public/WorldPartition/WorldPartitionEditorCell.h
+++ b/Engine/Source/Runtime/Engine/Public/WorldPartition/WorldPartitionEditorCell.h
@@ -6,6 +6,7 @@
 
 class FWorldPartitionActorDesc;
 class UWorldPartition;
+class AWorldPartitionEditorCellPreview;
 
 /**
  * Represents an editing cell (editor-only)
@@ -38,5 +39,7 @@ public:
 
 	TSet<FWorldPartitionActorDesc*> Actors;
 	TSet<FWorldPartitionActorDesc*> LoadedActors;
+
+	AWorldPartitionEditorCellPreview* CellPreviewActor;
 #endif
-};
\ No newline at end of file
+};
diff --git a/Engine/Source/Runtime/Engine/Public/WorldPartition/WorldPartitionEditorCellPreviewActor.h b/Engine/Source/Runtime/Engine/Public/WorldPartition/WorldPartitionEditorCellPreviewActor.h
new file mode 100644
--- /dev/null
+++ b/Engine/Source/Runtime/Engine/Public/WorldPartition/WorldPartitionEditorCellPreviewActor.h
@@ -0,0 +1,36 @@
+// Copyright Epic Games, Inc. All Rights Reserved.
+#pragma once
+
+#include "CoreMinimal.h"
+#include "GameFramework/Actor.h"
+#include "Engine/PostProcessVolume.h"
+#include "WorldPartitionEditorCellPreviewActor.generated.h"
+
+UCLASS(NotPlaceable)
+class AWorldPartitionEditorCellPreview : public AActor
+{
+	GENERATED_UCLASS_BODY()
+
+#if WITH_EDITOR
+public:
+	//~ Begin AActor interface
+	virtual EActorGridPlacement GetDefaultGridPlacement() const override;
+	virtual bool IsSelectable() const override;
+	virtual void GetActorBounds(bool bOnlyCollidingComponents, FVector& Origin, FVector& BoxExtent, bool bIncludeFromChildActors = false) const override;
+	virtual void PostRegisterAllComponents() override;
+	//~ End AActor interface
+
+	void SetVisibility(bool bVisible);
+	void SetCellBounds(const FBox& InCellBounds) { CellBounds = InCellBounds; }
+
+private:
+	bool bVisible;
+	FBox CellBounds;
+#endif
+};
+
+UCLASS(NotPlaceable, Transient)
+class AWorldPartitionUnloadedCellPreviewPostProcessVolume : public APostProcessVolume
+{
+	GENERATED_UCLASS_BODY()
+};
