Fix crash where TickFunctions with Intervals and with TickWhenPaused are trying to unregister themselves on a destroyed Level (LevelTask). This happen when the level are destroyed and the paused.

#rnx

diff --git a/Engine/Source/Runtime/Engine/Private/TickTaskManager.cpp b/Engine/Source/Runtime/Engine/Private/TickTaskManager.cpp
--- a/Engine/Source/Runtime/Engine/Private/TickTaskManager.cpp
+++ b/Engine/Source/Runtime/Engine/Private/TickTaskManager.cpp
@@ -728,6 +728,10 @@ public:
 			CoolingDownNode->InternalData->bRegistered = false;
 			CoolingDownNode = CoolingDownNode->InternalData->Next;
 		}
+		for (FTickScheduleDetails& TickDetails : TickFunctionsToReschedule)
+		{
+			TickDetails.TickFunction->InternalData->bRegistered = false;
+		}
 	}
 
 	/**
