Create Kernel Data interface and compute graph

diff --git a/Engine/Source/Runtime/Engine/Private/ComputeFramework/ComputeGraph.cpp b/Engine/Source/Runtime/Engine/Private/ComputeFramework/ComputeGraph.cpp
new file mode 100644
--- /dev/null
+++ b/Engine/Source/Runtime/Engine/Private/ComputeFramework/ComputeGraph.cpp
@@ -0,0 +1,16 @@
+// Copyright Epic Games, Inc. All Rights Reserved.
+
+#include "ComputeFramework/ComputeGraph.h"
+#include "ComputeFramework/ComputeKernel.h"
+#include "UObject/UObjectGlobals.h"
+
+
+FComputeKernelInvocationHandle UComputeGraph::AddKernel(UComputeKernelSource* KernelSource)
+{
+	return {};
+}
+
+bool UComputeGraph::RemoveKernel(FComputeKernelInvocationHandle KernelInvocation)
+{
+	return false;
+}
diff --git a/Engine/Source/Runtime/Engine/Public/ComputeFramework/ComputeGraph.h b/Engine/Source/Runtime/Engine/Public/ComputeFramework/ComputeGraph.h
new file mode 100644
--- /dev/null
+++ b/Engine/Source/Runtime/Engine/Public/ComputeFramework/ComputeGraph.h
@@ -0,0 +1,182 @@
+// Copyright Epic Games, Inc. All Rights Reserved.
+
+#pragma once
+
+#include "ComputeGraph.generated.h"
+
+class UComputeKernelSource;
+class UComputeKernel;
+
+/* Identifies a specific kernel invocation within a compute graph. */
+struct FComputeKernelInvocationHandle
+{
+	uint16 InvocationIdx = 0xFFFF;
+	uint16 GenerationIdx = 0xFFFF;
+};
+
+/* Identifies a specific resource within a compute graph. */
+struct FComputeResourceHandle
+{
+	uint16 ResourceIdx = 0xFFFF;
+	uint16 GenerationIdx = 0xFFFF;
+};
+
+/* Identifies a specific binding within a compute graph. */
+struct FComputeBindingHandle
+{
+	uint16 BindingIdx = 0xFFFF;
+	uint16 GenerationIdx = 0xFFFF;
+};
+
+UENUM()
+enum class EComputeGraphResourceType : uint8
+{
+	Buffer,
+	Texture1D,
+	Texture2D,
+	Texture3D,
+};
+
+/* Defines creation format and dimensions of a resource. */
+struct FComputeResourceDesc
+{
+	EPixelFormat ResourceFormat = EPixelFormat::PF_Unknown;
+	uint16 X = 0;
+	uint16 Y = 0;	// Only used for Tex2D and Tex3D
+	uint16 Z = 0;	// Only used for Tex3D
+	EComputeGraphResourceType Type = EComputeGraphResourceType::Buffer;
+	uint8 ArrayCount = 0;
+	uint8 MipLevels = 0;
+	uint8 Flags = 0;
+	uint8 SampleCount = 0;
+	uint8 SampleQuality = 0;
+};
+
+/*
+ * Defines view into an externally allocated resouse and valid uses.
+ * Less params are necessary than FComputeResourceDesc since only
+ * bindings/types are validated with external resources.
+ */
+struct FComputeResourceExternalDesc
+{
+	EPixelFormat ResourceFormat = EPixelFormat::PF_Unknown;
+	EComputeGraphResourceType Type = EComputeGraphResourceType::Buffer;
+	uint8 Flags = 0;
+};
+
+USTRUCT()
+struct FComputeKernelInvocation
+{
+	GENERATED_BODY()
+
+	static constexpr uint32 MAX_NAME_LENGTH = 32;
+
+	UPROPERTY()
+	UComputeKernel* ComputeKernel = nullptr;
+	
+	UPROPERTY()
+	uint16 GenerationIdx = 0xFFFF;
+
+	FComputeKernelInvocation() = default;
+	FComputeKernelInvocation(UComputeKernel* InComputeKernel, uint16 InGenerationIdx)
+		: ComputeKernel(InComputeKernel)
+		, GenerationIdx(InGenerationIdx)
+	{}
+};
+
+USTRUCT()
+struct FComputeResource
+{
+	GENERATED_BODY()
+
+	static constexpr uint32 MAX_NAME_LENGTH = 32;
+
+	UPROPERTY()
+	uint16 GenerationIdx = 0xFFFF;
+};
+
+/* Identifies a specific kernel invocation AND parameter slot within a compute graph. */
+struct FComputeKernelInvocationBindPoint
+{
+	FComputeKernelInvocationHandle KernelInvocation;
+	FString BindPointName;
+};
+
+/* Define is the binding data flows from resource into kernel or kernel into resource. */
+enum class EComputeGraphBindingType : uint8
+{
+	Input,	// Resource       ->   Kernel Input
+	Output, // Kernel Output  ->   Resource
+};
+
+/* Identifies a specific data flow between a kernel and resource with the compute graph. */
+USTRUCT()
+struct FKernelBinding
+{
+	GENERATED_BODY()
+
+	EComputeGraphBindingType BindingType = EComputeGraphBindingType::Input;
+	FComputeKernelInvocationBindPoint InvocationBindPoint;
+	FComputeResourceHandle Resource;
+};
+
+/*
+ * Graph representing the logical execution and data flow topology as a DAG. The underlying
+ * scheduler will extract parallel/serial execution from this DAG description.
+ */
+UCLASS()
+class UComputeGraph : public UObject
+{
+	GENERATED_BODY()
+
+public:
+	//===================================
+	// Add/Remove "nodes" in the graph
+	//===================================
+
+	/* May fail and return nullptr. */
+	FComputeKernelInvocationHandle AddKernel(UComputeKernelSource* KernelSource);
+
+	/* Returns true if kernel successfully removed. When a kernel is removed all associated bindings are also removed. */
+	bool RemoveKernel(FComputeKernelInvocationHandle KernelInvocation);
+
+	/*
+	 * Compute resource are parameters/buffers/textures that hold data during the execution of the 
+	 * compute graph. They can either be transient - their lifetime does not out live the execution 
+	 * of the graph, or external - the resource out lives the graph and acts as the mechanism to 
+	 * bring data in or move data out of the compute graph.
+	 */
+	FComputeResourceHandle AddResource(FString Name, FComputeResourceDesc Desc);
+	FComputeResourceHandle AddResource(FString Name, FComputeResourceExternalDesc Desc);
+
+	/* Returns true if resource successfully removed. When a resource is removed all associated bindings are also removed. */
+	bool RemoveResource(FComputeResourceHandle Resource);
+
+
+
+	//===================================
+	// Add/Remove edges in the graph
+	//===================================
+
+	/* Creates a binding, aka data flow, from src to dest. */
+	FComputeBindingHandle AddBinding(
+		EComputeGraphBindingType BindingType,
+		FComputeKernelInvocationBindPoint InvocationBindPoint,
+		FComputeResourceHandle Resource
+		);
+
+	void RemoveBinding(FComputeBindingHandle Binding);
+
+private:
+	/* Represents all the invocation nodes of the graph. */
+	UPROPERTY()
+	TArray<FComputeKernelInvocation> KernelInvocations;
+
+	/* Represents all the data nodes of the graph */
+	UPROPERTY()
+	TArray<FComputeResource> Resources;
+
+	/* Represents all the edges of the graph. */
+	UPROPERTY()
+	TArray<FKernelBinding> 	Bindings;
+};
diff --git a/Engine/Source/Runtime/Engine/Public/ComputeFramework/ComputeKernel.h b/Engine/Source/Runtime/Engine/Public/ComputeFramework/ComputeKernel.h
--- a/Engine/Source/Runtime/Engine/Public/ComputeFramework/ComputeKernel.h
+++ b/Engine/Source/Runtime/Engine/Public/ComputeFramework/ComputeKernel.h
@@ -109,9 +109,7 @@ public:
 	FShaderResetEvent ShaderResetSignal;
 #endif
 
-	/* Retrieve the IO bind points for the kernel - only way to get a valid bind point. */
-	//FShaderBindPoint GetInputBindPoint(const FString& Name) const;
-	//FShaderBindPoint GetOutputBindPoint(const FString& Name) const;
+	void SetKernelSource(UComputeKernelSource* KernelSource);
 
 	/*
 	 * Implemented by derived classes to specify how many invocations of the kernel are necessary 
