Fix failing ensure on non-persistent non-streaming levels (ex: TransLevelMoveBuffer) in ULevel::FindStreamingLevel

diff --git a/Engine/Source/Runtime/Engine/Private/LevelStreaming.cpp b/Engine/Source/Runtime/Engine/Private/LevelStreaming.cpp
--- a/Engine/Source/Runtime/Engine/Private/LevelStreaming.cpp
+++ b/Engine/Source/Runtime/Engine/Private/LevelStreaming.cpp
@@ -811,7 +811,7 @@ ULevelStreaming* ULevelStreaming::FindStreamingLevel(const ULevel* Level)
 			FoundLevelStreaming = LevelAnnotation.LevelStreaming;
 		}
 
-		if (!ensure(FoundLevelStreaming))
+		if (!FoundLevelStreaming)
 		{
 			// fallback search
 			for (ULevelStreaming* CurStreamingLevel : Level->OwningWorld->GetStreamingLevels())
@@ -822,6 +822,8 @@ ULevelStreaming* ULevelStreaming::FindStreamingLevel(const ULevel* Level)
 					break;
 				}
 			}
+			// we shouldn't have found a streaming level here
+			ensure(!FoundLevelStreaming);
 		}
 	}
 
