Fix Level Instances with External Actors not loading those.

diff --git a/Engine/Source/Runtime/Engine/Classes/Engine/World.h b/Engine/Source/Runtime/Engine/Classes/Engine/World.h
--- a/Engine/Source/Runtime/Engine/Classes/Engine/World.h
+++ b/Engine/Source/Runtime/Engine/Classes/Engine/World.h
@@ -1234,6 +1234,9 @@ public:
 #if WITH_EDITOR
 	/** Hierarchical LOD System. */
 	struct FHierarchicalLODBuilder*				HierarchicalLODBuilder;
+
+	/** Original World Name before PostLoad rename. Used to get external actors on disk. */
+	FName OriginalWorldName;
 #endif // WITH_EDITOR
 
 	/** Called from DemoNetDriver when playing back a replay and the timeline is successfully scrubbed */
diff --git a/Engine/Source/Runtime/Engine/Private/Level.cpp b/Engine/Source/Runtime/Engine/Private/Level.cpp
--- a/Engine/Source/Runtime/Engine/Private/Level.cpp
+++ b/Engine/Source/Runtime/Engine/Private/Level.cpp
@@ -2281,7 +2281,7 @@ TArray<FString> ULevel::GetOnDiskExternalActorPackages() const
 {
 	TArray<FString> ActorPackageNames;
 	UWorld* World = GetTypedOuter<UWorld>();
-	FString ExternalActorsPath = ULevel::GetExternalActorsPath(World->GetPackage(), World->GetName());
+	FString ExternalActorsPath = ULevel::GetExternalActorsPath(World->GetPackage(), World->OriginalWorldName == NAME_None ? World->GetName() : World->OriginalWorldName.ToString());
 	if (!ExternalActorsPath.IsEmpty())
 	{
 		IFileManager::Get().IterateDirectoryRecursively(*FPackageName::LongPackageNameToFilename(ExternalActorsPath), [&ActorPackageNames](const TCHAR* FilenameOrDirectory, bool bIsDirectory)
diff --git a/Engine/Source/Runtime/Engine/Private/World.cpp b/Engine/Source/Runtime/Engine/Private/World.cpp
--- a/Engine/Source/Runtime/Engine/Private/World.cpp
+++ b/Engine/Source/Runtime/Engine/Private/World.cpp
@@ -1070,6 +1070,7 @@ void UWorld::PostLoad()
 		{
 			// Needed for VER_UE4_WORLD_NAMED_AFTER_PACKAGE. If this file was manually renamed outside of the editor, this is needed anyway
 			const FString ShortPackageName = FPackageName::GetLongPackageAssetName(GetOutermost()->GetName());
+			OriginalWorldName = GetFName();
 			if (GetName() != ShortPackageName)
 			{
 				Rename(*ShortPackageName, NULL, REN_NonTransactional | REN_ForceNoResetLoaders | REN_DontCreateRedirectors);
