Pass over viewport sprite icons
Also deleted a lot of UE3 icons which are no longer used

diff --git a/Engine/Config/BaseEngine.ini b/Engine/Config/BaseEngine.ini
--- a/Engine/Config/BaseEngine.ini
+++ b/Engine/Config/BaseEngine.ini
@@ -1158,7 +1158,6 @@ AutoSaveIndex=0
 +PackagesToBeFullyLoadedAtStartup=/Engine/EditorResources/SceneManager
 +PackagesToBeFullyLoadedAtStartup=/Engine/EditorResources/SmallFont
 +PackagesToBeFullyLoadedAtStartup=/Engine/EditorResources/S_Actor
-+PackagesToBeFullyLoadedAtStartup=/Engine/EditorResources/S_DecalActorIcon
 +PackagesToBeFullyLoadedAtStartup=/Engine/EditorResources/S_TextRenderActorIcon
 +PackagesToBeFullyLoadedAtStartup=/Engine/EditorResources/S_Emitter
 +PackagesToBeFullyLoadedAtStartup=/Engine/EditorResources/S_ExpoHeightFog
@@ -1171,7 +1170,6 @@ AutoSaveIndex=0
 +PackagesToBeFullyLoadedAtStartup=/Engine/EditorResources/S_Player
 +PackagesToBeFullyLoadedAtStartup=/Engine/EditorResources/S_RadForce
 +PackagesToBeFullyLoadedAtStartup=/Engine/EditorResources/S_ReflActorIcon
-+PackagesToBeFullyLoadedAtStartup=/Engine/EditorResources/S_Terrain
 +PackagesToBeFullyLoadedAtStartup=/Engine/EditorResources/S_Thruster
 +PackagesToBeFullyLoadedAtStartup=/Engine/EditorResources/S_Trigger
 +PackagesToBeFullyLoadedAtStartup=/Engine/EditorResources/S_VectorFieldVol
@@ -1183,8 +1181,6 @@ AutoSaveIndex=0
 +PackagesToBeFullyLoadedAtStartup=/Engine/EditorResources/LightIcons/S_LightPointMove
 +PackagesToBeFullyLoadedAtStartup=/Engine/EditorResources/LightIcons/S_LightSpot
 +PackagesToBeFullyLoadedAtStartup=/Engine/EditorResources/LightIcons/S_LightSpotMove
-+PackagesToBeFullyLoadedAtStartup=/Engine/EditorResources/Spline/T_Loft_Spline
-
 +PackagesToBeFullyLoadedAtStartup=/Engine/EngineMaterials/BlinkingCaret
 +PackagesToBeFullyLoadedAtStartup=/Engine/EngineMaterials/DefaultBokeh
 +PackagesToBeFullyLoadedAtStartup=/Engine/EngineMaterials/DefaultBloomKernel
diff --git a/Engine/Source/Editor/StaticMeshEditor/Private/StaticMeshEditorViewportClient.cpp b/Engine/Source/Editor/StaticMeshEditor/Private/StaticMeshEditorViewportClient.cpp
--- a/Engine/Source/Editor/StaticMeshEditor/Private/StaticMeshEditorViewportClient.cpp
+++ b/Engine/Source/Editor/StaticMeshEditor/Private/StaticMeshEditorViewportClient.cpp
@@ -51,8 +51,6 @@ FStaticMeshEditorViewportClient::FStaticMeshEditorViewportClient(TWeakPtr<IStati
 	, StaticMeshEditorPtr(InStaticMeshEditor)
 	, StaticMeshEditorViewportPtr(InStaticMeshEditorViewport)
 {
-	SimplygonLogo = LoadObject<UTexture2D>(NULL, TEXT("/Engine/EditorResources/SimplygonLogo.SimplygonLogo"), NULL, LOAD_None, NULL);
-
 	// Setup defaults for the common draw helper.
 	DrawHelper.bDrawPivot = false;
 	DrawHelper.bDrawWorldBox = false;
@@ -731,30 +729,6 @@ static void DrawAngles(FCanvas* Canvas, int32 XPos, int32 YPos, EAxisList::Type
 
 void FStaticMeshEditorViewportClient::DrawCanvas( FViewport& InViewport, FSceneView& View, FCanvas& Canvas )
 {
-#ifdef TODO_STATICMESH
-	if ( StaticMesh->bHasBeenSimplified && SimplygonLogo && SimplygonLogo->Resource )
-	{
-		const float LogoSizeX = 64.0f;
-		const float LogoSizeY = 40.65f;
-		const float Padding = 6.0f;
-		const float LogoX = Viewport->GetSizeXY().X - Padding - LogoSizeX;
-		const float LogoY = Viewport->GetSizeXY().Y - Padding - LogoSizeY;
-
-		Canvas->DrawTile(
-			LogoX,
-			LogoY,
-			LogoSizeX,
-			LogoSizeY,
-			0.0f,
-			0.0f,
-			1.0f,
-			1.0f,
-			FLinearColor::White,
-			SimplygonLogo->Resource,
-			SE_BLEND_Opaque );
-	}
-#endif // #if TODO_STATICMESH
-
 	auto StaticMeshEditor = StaticMeshEditorPtr.Pin();
 	auto StaticMeshEditorViewport = StaticMeshEditorViewportPtr.Pin();
 	if (!StaticMeshEditor.IsValid() || !StaticMeshEditorViewport.IsValid())
diff --git a/Engine/Source/Runtime/Engine/Private/Atmosphere/AtmosphericFog.cpp b/Engine/Source/Runtime/Engine/Private/Atmosphere/AtmosphericFog.cpp
--- a/Engine/Source/Runtime/Engine/Private/Atmosphere/AtmosphericFog.cpp
+++ b/Engine/Source/Runtime/Engine/Private/Atmosphere/AtmosphericFog.cpp
@@ -37,7 +37,7 @@ AAtmosphericFog::AAtmosphericFog(const FObjectInitializer& ObjectInitializer)
 			FName ID_Fog;
 			FText NAME_Fog;
 			FConstructorStatics()
-				: FogTextureObject(TEXT("/Engine/EditorResources/S_ExpoHeightFog"))
+				: FogTextureObject(TEXT("/Engine/EditorResources/S_AtmosphericHeightFog"))
 				, ID_Fog(TEXT("Fog"))
 				, NAME_Fog(NSLOCTEXT("SpriteCategory", "Fog", "Fog"))
 			{
diff --git a/Engine/Source/Runtime/Engine/Private/Components/BillboardComponent.cpp b/Engine/Source/Runtime/Engine/Private/Components/BillboardComponent.cpp
--- a/Engine/Source/Runtime/Engine/Private/Components/BillboardComponent.cpp
+++ b/Engine/Source/Runtime/Engine/Private/Components/BillboardComponent.cpp
@@ -15,6 +15,7 @@
 #include "Engine/Engine.h"
 #include "Engine/LevelStreaming.h"
 #include "LevelUtils.h"
+#include "TextureCompiler.h"
 
 namespace BillboardConstants
 {
@@ -66,6 +67,10 @@ public:
 		if(InComponent->Sprite)
 		{
 			Texture = InComponent->Sprite;
+
+			// Finish the compilation as we rely on the actual size of the texture to function properly
+			FTextureCompilingManager::Get().FinishCompilation({(UTexture*)Texture});
+
 			// Set UL and VL to the size of the texture if they are set to 0.0, otherwise use the given value
 			UL = InComponent->UL == 0.0f ? InComponent->Sprite->GetSurfaceWidth() : InComponent->UL;
 			VL = InComponent->VL == 0.0f ? InComponent->Sprite->GetSurfaceHeight() : InComponent->VL;	
diff --git a/Engine/Source/Runtime/Engine/Private/Components/RectLightComponent.cpp b/Engine/Source/Runtime/Engine/Private/Components/RectLightComponent.cpp
--- a/Engine/Source/Runtime/Engine/Private/Components/RectLightComponent.cpp
+++ b/Engine/Source/Runtime/Engine/Private/Components/RectLightComponent.cpp
@@ -30,8 +30,8 @@ URectLightComponent::URectLightComponent(const FObjectInitializer& ObjectInitial
 #if WITH_EDITORONLY_DATA
 	if (!IsRunningCommandlet())
 	{
-		static ConstructorHelpers::FObjectFinder<UTexture2D> StaticTexture(TEXT("/Engine/EditorResources/LightIcons/S_LightPoint"));
-		static ConstructorHelpers::FObjectFinder<UTexture2D> DynamicTexture(TEXT("/Engine/EditorResources/LightIcons/S_LightPointMove"));
+		static ConstructorHelpers::FObjectFinder<UTexture2D> StaticTexture(TEXT("/Engine/EditorResources/LightIcons/S_LightRect"));
+		static ConstructorHelpers::FObjectFinder<UTexture2D> DynamicTexture(TEXT("/Engine/EditorResources/LightIcons/S_LightRect"));
 
 		StaticEditorTexture = StaticTexture.Object;
 		StaticEditorTextureScale = 0.5f;
diff --git a/Engine/Source/Runtime/Engine/Private/Components/ReflectionCaptureComponent.cpp b/Engine/Source/Runtime/Engine/Private/Components/ReflectionCaptureComponent.cpp
--- a/Engine/Source/Runtime/Engine/Private/Components/ReflectionCaptureComponent.cpp
+++ b/Engine/Source/Runtime/Engine/Private/Components/ReflectionCaptureComponent.cpp
@@ -247,11 +247,29 @@ ABoxReflectionCapture::ABoxReflectionCapture(const FObjectInitializer& ObjectIni
 	if (GetSpriteComponent())
 	{
 		GetSpriteComponent()->SetupAttachment(BoxComponent);
+
+
+		// Structure to hold one-time initialization
+		struct FConstructorStatics
+		{
+			FName NAME_ReflectionCapture;
+			ConstructorHelpers::FObjectFinderOptional<UTexture2D> Texture;
+			FConstructorStatics()
+				: NAME_ReflectionCapture(TEXT("ReflectionCapture"))
+				, Texture(TEXT("/Engine/EditorResources/S_BoxReflectionCapture"))
+			{
+			}
+		};
+		static FConstructorStatics ConstructorStatics;
+
+		GetSpriteComponent()->Sprite = ConstructorStatics.Texture.Get();
 	}
+
 	if (GetCaptureOffsetComponent())
 	{
 		GetCaptureOffsetComponent()->SetupAttachment(BoxComponent);
 	}
+
 #endif	//WITH_EDITORONLY_DATA
 	UBoxComponent* DrawInfluenceBox = CreateDefaultSubobject<UBoxComponent>(TEXT("DrawBox0"));
 	DrawInfluenceBox->SetupAttachment(GetCaptureComponent());
@@ -949,6 +967,14 @@ void UReflectionCaptureComponent::UpdatePreviewShape()
 	if (CaptureOffsetComponent)
 	{
 		CaptureOffsetComponent->SetRelativeLocation_Direct(CaptureOffset / GetComponentTransform().GetScale3D());
+		if (CaptureOffset.IsNearlyZero())
+		{
+			CaptureOffsetComponent->SetVisibility(false);
+		}
+		else
+		{
+			CaptureOffsetComponent->SetVisibility(true);
+		}
 	}
 }
 
diff --git a/Engine/Source/Runtime/Engine/Private/Components/SceneComponent.cpp b/Engine/Source/Runtime/Engine/Private/Components/SceneComponent.cpp
--- a/Engine/Source/Runtime/Engine/Private/Components/SceneComponent.cpp
+++ b/Engine/Source/Runtime/Engine/Private/Components/SceneComponent.cpp
@@ -711,6 +711,7 @@ void USceneComponent::CreateSpriteComponent(UTexture2D* SpriteTexture)
 		SpriteComponent->CreationMethod = CreationMethod;
 		SpriteComponent->bIsScreenSizeScaled = true;
 		SpriteComponent->bUseInEditorScaling = true;
+		SpriteComponent->OpacityMaskRefVal = .3f;
 
 		SpriteComponent->SetupAttachment(this);
 		SpriteComponent->RegisterComponent();
diff --git a/Engine/Source/Runtime/Engine/Private/Components/SkyLightComponent.cpp b/Engine/Source/Runtime/Engine/Private/Components/SkyLightComponent.cpp
--- a/Engine/Source/Runtime/Engine/Private/Components/SkyLightComponent.cpp
+++ b/Engine/Source/Runtime/Engine/Private/Components/SkyLightComponent.cpp
@@ -1063,6 +1063,7 @@ ASkyLight::ASkyLight(const FObjectInitializer& ObjectInitializer)
 			GetSpriteComponent()->Sprite = ConstructorStatics.SkyLightTextureObject.Get();
 			GetSpriteComponent()->SpriteInfo.Category = ConstructorStatics.ID_Sky;
 			GetSpriteComponent()->SpriteInfo.DisplayName = ConstructorStatics.NAME_Sky;
+			GetSpriteComponent()->SetRelativeScale3D(FVector(0.5f, 0.5f, 0.5f));
 			GetSpriteComponent()->SetupAttachment(LightComponent);
 		}
 	}
diff --git a/Engine/Source/Runtime/Engine/Private/Interpolation.cpp b/Engine/Source/Runtime/Engine/Private/Interpolation.cpp
--- a/Engine/Source/Runtime/Engine/Private/Interpolation.cpp
+++ b/Engine/Source/Runtime/Engine/Private/Interpolation.cpp
@@ -8941,6 +8941,7 @@ AMaterialInstanceActor::AMaterialInstanceActor(const FObjectInitializer& ObjectI
 		SpriteComponent->Sprite = ConstructorStatics.MaterialInstanceSpriteObject.Get();
 		SpriteComponent->SpriteInfo.Category = ConstructorStatics.ID_Materials;
 		SpriteComponent->SpriteInfo.DisplayName = ConstructorStatics.NAME_Materials;
+		SpriteComponent->SetRelativeScale3D(FVector(0.5f, 0.5f, 0.5f));
 		SpriteComponent->SetupAttachment(SceneComponent);
 		SpriteComponent->bIsScreenSizeScaled = true;
 	}
diff --git a/Engine/Source/Runtime/Engine/Private/TargetPoint.cpp b/Engine/Source/Runtime/Engine/Private/TargetPoint.cpp
--- a/Engine/Source/Runtime/Engine/Private/TargetPoint.cpp
+++ b/Engine/Source/Runtime/Engine/Private/TargetPoint.cpp
@@ -28,7 +28,7 @@ ATargetPoint::ATargetPoint(const FObjectInitializer& ObjectInitializer)
 			FText NAME_TargetPoint;
 			FConstructorStatics()
 				: TargetIconSpawnObject(TEXT("/Engine/EditorMaterials/TargetIconSpawn"))
-				, TargetIconObject(TEXT("/Engine/EditorMaterials/TargetIcon"))
+				, TargetIconObject(TEXT("/Engine/EditorResources/S_TargetPoint"))
 				, ID_TargetPoint(TEXT("TargetPoint"))
 				, NAME_TargetPoint(NSLOCTEXT("SpriteCategory", "TargetPoint", "Target Points"))
 			{
@@ -39,7 +39,7 @@ ATargetPoint::ATargetPoint(const FObjectInitializer& ObjectInitializer)
 		if (SpriteComponent)
 		{
 			SpriteComponent->Sprite = ConstructorStatics.TargetIconObject.Get();
-			SpriteComponent->SetRelativeScale3D_Direct(FVector(0.35f, 0.35f, 0.35f));
+			SpriteComponent->SetRelativeScale3D_Direct(FVector(0.5f, 0.5f, 0.5f));
 			SpriteComponent->SpriteInfo.Category = ConstructorStatics.ID_TargetPoint;
 			SpriteComponent->SpriteInfo.DisplayName = ConstructorStatics.NAME_TargetPoint;
 			SpriteComponent->bIsScreenSizeScaled = true;
diff --git a/Engine/Source/Runtime/Engine/Private/TriggerActors.cpp b/Engine/Source/Runtime/Engine/Private/TriggerActors.cpp
--- a/Engine/Source/Runtime/Engine/Private/TriggerActors.cpp
+++ b/Engine/Source/Runtime/Engine/Private/TriggerActors.cpp
@@ -14,6 +14,7 @@
 #include "Components/BoxComponent.h"
 #include "Components/SphereComponent.h"
 #include "Components/CapsuleComponent.h"
+#include "UObject/ConstructorHelpers.h"
 namespace
 {
 	static const FColor TriggerBaseColor(100, 255, 100, 255);
@@ -34,6 +35,26 @@ ATriggerCapsule::ATriggerCapsule(const FObjectInitializer& ObjectInitializer)
 	{
 		TriggerSpriteComponent->SetupAttachment(CapsuleCollisionComponent);
 	}
+
+	if (GetSpriteComponent())
+	{
+		// Structure to hold one-time initialization
+		struct FConstructorStatics
+		{
+			ConstructorHelpers::FObjectFinderOptional<UTexture2D> TriggerTextureObject;
+			FName ID_Triggers;
+			FText NAME_Triggers;
+			FConstructorStatics()
+				: TriggerTextureObject(TEXT("/Engine/EditorResources/S_TriggerCapsule"))
+				, ID_Triggers(TEXT("Triggers"))
+				, NAME_Triggers(NSLOCTEXT("SpriteCategory", "Triggers", "Triggers"))
+			{
+			}
+		};
+		static FConstructorStatics ConstructorStatics;
+
+		GetSpriteComponent()->Sprite = ConstructorStatics.TriggerTextureObject.Get();
+	}
 #endif
 }
 
@@ -100,6 +121,26 @@ ATriggerBox::ATriggerBox(const FObjectInitializer& ObjectInitializer)
 	{
 		TriggerSpriteComponent->SetupAttachment(BoxCollisionComponent);
 	}
+
+	if (GetSpriteComponent())
+	{
+		// Structure to hold one-time initialization
+		struct FConstructorStatics
+		{
+			ConstructorHelpers::FObjectFinderOptional<UTexture2D> TriggerTextureObject;
+			FName ID_Triggers;
+			FText NAME_Triggers;
+			FConstructorStatics()
+				: TriggerTextureObject(TEXT("/Engine/EditorResources/S_TriggerBox"))
+					, ID_Triggers(TEXT("Triggers"))
+					, NAME_Triggers(NSLOCTEXT("SpriteCategory", "Triggers", "Triggers"))
+			{
+			}
+		};
+		static FConstructorStatics ConstructorStatics;
+
+		GetSpriteComponent()->Sprite = ConstructorStatics.TriggerTextureObject.Get();
+	}
 #endif
 }
 
@@ -142,6 +183,26 @@ ATriggerSphere::ATriggerSphere(const FObjectInitializer& ObjectInitializer)
 	{
 		TriggerSpriteComponent->SetupAttachment(SphereCollisionComponent);
 	}
+
+	if (GetSpriteComponent())
+	{
+		// Structure to hold one-time initialization
+		struct FConstructorStatics
+		{
+			ConstructorHelpers::FObjectFinderOptional<UTexture2D> TriggerTextureObject;
+			FName ID_Triggers;
+			FText NAME_Triggers;
+			FConstructorStatics()
+				: TriggerTextureObject(TEXT("/Engine/EditorResources/S_TriggerSphere"))
+				, ID_Triggers(TEXT("Triggers"))
+				, NAME_Triggers(NSLOCTEXT("SpriteCategory", "Triggers", "Triggers"))
+			{
+			}
+		};
+		static FConstructorStatics ConstructorStatics;
+
+		GetSpriteComponent()->Sprite = ConstructorStatics.TriggerTextureObject.Get();
+	}
 #endif
 }
 
