Fix non-editor

diff --git a/Engine/Source/Runtime/Engine/Private/Components/BillboardComponent.cpp b/Engine/Source/Runtime/Engine/Private/Components/BillboardComponent.cpp
--- a/Engine/Source/Runtime/Engine/Private/Components/BillboardComponent.cpp
+++ b/Engine/Source/Runtime/Engine/Private/Components/BillboardComponent.cpp
@@ -68,8 +68,10 @@ public:
 		{
 			Texture = InComponent->Sprite;
 
+#if WITH_EDITOR
 			// Finish the compilation as we rely on the actual size of the texture to function properly
 			FTextureCompilingManager::Get().FinishCompilation({(UTexture*)Texture});
+#endif
 
 			// Set UL and VL to the size of the texture if they are set to 0.0, otherwise use the given value
 			UL = InComponent->UL == 0.0f ? InComponent->Sprite->GetSurfaceWidth() : InComponent->UL;
