Correct indexing when using multiple section + soft selection.

diff --git a/Engine/Source/Runtime/Engine/Private/SkeletalMeshComponentPhysics.cpp b/Engine/Source/Runtime/Engine/Private/SkeletalMeshComponentPhysics.cpp
--- a/Engine/Source/Runtime/Engine/Private/SkeletalMeshComponentPhysics.cpp
+++ b/Engine/Source/Runtime/Engine/Private/SkeletalMeshComponentPhysics.cpp
@@ -2135,7 +2135,7 @@ void USkeletalMeshComponent::ComputeSkinnedTangentBasis(USkeletalMeshComponent*
 		const uint32 NumSoftVerts = Section.GetNumVertices();
 		for (uint32 SoftIdx = 0; SoftIdx < NumSoftVerts; ++SoftIdx)
 		{
-			const uint32 TangentOffset = SoftOffset + SoftIdx * 2;
+			const uint32 TangentOffset = (SoftOffset + SoftIdx) * 2;
 			GetTypedSkinnedTangentBasis(Component, Section, LODData.StaticVertexBuffers, SkinWeightBuffer, SoftIdx, CachedRefToLocals, OutTangenXZ[TangentOffset + 0], OutTangenXZ[TangentOffset + 1]);
 		}
 	}
