Initial scafolding for how we will support runtime execution of deformer graph.

diff --git a/Engine/Source/Runtime/Engine/Private/ComputeFramework/ComputeFramework.cpp b/Engine/Source/Runtime/Engine/Private/ComputeFramework/ComputeFramework.cpp
new file mode 100644
--- /dev/null
+++ b/Engine/Source/Runtime/Engine/Private/ComputeFramework/ComputeFramework.cpp
@@ -0,0 +1,110 @@
+// Copyright Epic Games, Inc. All Rights Reserved.
+
+#include "ComputeFramework/ComputeFramework.h"
+#include "ComputeFramework/ComputeGraph.h"
+#include "RenderGraphBuilder.h"
+#include "RenderGraphUtils.h"
+#include "RHIDefinitions.h"
+
+static int32 GComputeFrameworkMode = 0;
+static FAutoConsoleVariableRef CVarComputeFrameworkMode(
+	TEXT("r.ComputeFramework.mode"),
+	GComputeFrameworkMode,
+	TEXT("The mode Compute Framework should operate.\n")
+	TEXT("    0: disabled (Default)\n")
+	TEXT("    1: enabled\n"),
+	ECVF_Scalability | ECVF_RenderThreadSafe
+	);
+
+bool SupportsComputeFramework(ERHIFeatureLevel::Type FeatureLevel, EShaderPlatform ShaderPlatform)
+{
+	return GComputeFrameworkMode > 0
+		&& FeatureLevel >= ERHIFeatureLevel::SM5
+		&& FDataDrivenShaderPlatformInfo::GetSupportsComputeFramework(ShaderPlatform);
+}
+
+struct FComputeExecutionBuffer
+{
+	FComputeExecutionBuffer() = default;
+
+	FComputeExecutionBuffer(FName InName, FRDGBufferRef InBufferRef)
+		: Name(InName)
+		, BufferRef(InBufferRef)
+	{
+	}
+
+	FName Name;
+	FRDGBufferRef BufferRef;
+};
+
+void FComputeFramework::ExecuteBatches(
+	FRHICommandListImmediate& RHICmdList
+	)
+{
+	FRDGBuilder GraphBuilder(RHICmdList);
+
+	TArray<FComputeExecutionBuffer> ExecutionBuffers;
+
+	for (auto& ExtRes : ExternalBuffers)
+	{
+		TCHAR ResName[FComputeResource::MAX_NAME_LENGTH];
+		ExtRes.Name.ToString(ResName);
+
+		ExecutionBuffers.Emplace(
+			ExtRes.Name,
+			GraphBuilder.RegisterExternalBuffer(ExtRes.Buffer, ResName)
+			);
+	}
+
+	for (auto& TransRes : TransientBuffers)
+	{
+		TCHAR ResName[FComputeResource::MAX_NAME_LENGTH];
+		TransRes.Name.ToString(ResName);
+
+		FRDGBufferDesc BufferDesc;
+		
+		EComputeExecutionBufferType BufferType = (EComputeExecutionBufferType)TransRes.BufferType;
+		switch (BufferType)
+		{
+		case EComputeExecutionBufferType::Buffer:
+			BufferDesc = FRDGBufferDesc::CreateBufferDesc(TransRes.BytesPerElement, TransRes.ElementCount);
+			break;
+
+		case EComputeExecutionBufferType::StructuredBuffer:
+			BufferDesc = FRDGBufferDesc::CreateStructuredDesc(TransRes.BytesPerElement, TransRes.ElementCount);
+			break;
+
+		case EComputeExecutionBufferType::ByteAddressBuffer:
+			BufferDesc = FRDGBufferDesc::CreateByteAddressDesc(TransRes.BytesPerElement * TransRes.ElementCount);
+			break;
+
+		default:
+			check(!"Unsupported buffer type");
+			break;
+		}
+
+		ExecutionBuffers.Emplace(
+			TransRes.Name,
+			GraphBuilder.CreateBuffer(BufferDesc, ResName)
+			);
+	}
+
+	for (auto& Compute : ComputeShaders)
+	{
+		TCHAR KernelName[FComputeKernelInvocation::MAX_NAME_LENGTH];
+		Compute.KernelName.ToString(KernelName);
+
+		TCHAR InvocationName[FComputeKernelInvocation::MAX_NAME_LENGTH];
+		Compute.InvocationName.ToString(InvocationName);
+
+		/*
+		FComputeShaderUtils::AddPass(
+			GraphBuilder,
+			RDG_EVENT_NAME("Compute[%s]: %s", KernelName, InvocationName),
+			Compute.Shader,
+			{},
+			Compute.DispatchDim
+			);
+		*/
+	}
+}
diff --git a/Engine/Source/Runtime/Engine/Private/LevelTick.cpp b/Engine/Source/Runtime/Engine/Private/LevelTick.cpp
--- a/Engine/Source/Runtime/Engine/Private/LevelTick.cpp
+++ b/Engine/Source/Runtime/Engine/Private/LevelTick.cpp
@@ -59,6 +59,7 @@
 #include "Streaming/TextureStreamingHelpers.h"
 #include "ProfilingDebugging/CsvProfiler.h"
 #include "GPUSkinCache.h"
+#include "ComputeFramework/ComputeFramework.h"
 
 #if WITH_EDITOR
 #include "Editor.h"
@@ -927,16 +928,22 @@ bool UWorld::HasEndOfFrameUpdates() const
 
 struct FSendAllEndOfFrameUpdates
 {
-	FGPUSkinCache* GPUSkinCache;
+	FGPUSkinCache* GPUSkinCache = nullptr;
+	FComputeFramework* ComputeFramework = nullptr;
+
 #if WANTS_DRAW_MESH_EVENTS
 	FDrawEvent DrawEvent;
 #endif
 };
 
-FSendAllEndOfFrameUpdates* BeginSendEndOfFrameUpdatesDrawEvent(FGPUSkinCache* GPUSkinCache)
+FSendAllEndOfFrameUpdates* BeginSendEndOfFrameUpdatesDrawEvent(
+	FGPUSkinCache* GPUSkinCache,
+	FComputeFramework* ComputeFramework
+	)
 {
 	FSendAllEndOfFrameUpdates* SendAllEndOfFrameUpdates = new FSendAllEndOfFrameUpdates;
 	SendAllEndOfFrameUpdates->GPUSkinCache = GPUSkinCache;
+	SendAllEndOfFrameUpdates->ComputeFramework = ComputeFramework;
 
 #if WANTS_DRAW_MESH_EVENTS
 	ENQUEUE_RENDER_COMMAND(BeginDrawEventCommand)(
@@ -956,7 +963,7 @@ FSendAllEndOfFrameUpdates* BeginSendEndOfFrameUpdatesDrawEvent(FGPUSkinCache* GP
 void EndSendEndOfFrameUpdatesDrawEvent(FSendAllEndOfFrameUpdates* SendAllEndOfFrameUpdates)
 {
 	ENQUEUE_RENDER_COMMAND(EndDrawEventCommand)(
-		[SendAllEndOfFrameUpdates](FRHICommandList& RHICmdList)
+		[SendAllEndOfFrameUpdates](FRHICommandListImmediate& RHICmdList)
 	{
 		if (SendAllEndOfFrameUpdates->GPUSkinCache)
 		{
@@ -968,6 +975,11 @@ void EndSendEndOfFrameUpdatesDrawEvent(FSendAllEndOfFrameUpdates* SendAllEndOfFr
 		#endif
 		}
 
+		if (SendAllEndOfFrameUpdates->ComputeFramework)
+		{
+			SendAllEndOfFrameUpdates->ComputeFramework->ExecuteBatches(RHICmdList);
+		}
+
 		delete SendAllEndOfFrameUpdates;
 	});
 }
@@ -1009,7 +1021,10 @@ void UWorld::SendAllEndOfFrameUpdates()
 	}
 
 	// Issue a GPU event to wrap GPU work done during SendAllEndOfFrameUpdates, like skin cache updates
-	FSendAllEndOfFrameUpdates* SendAllEndOfFrameUpdates = BeginSendEndOfFrameUpdatesDrawEvent(Scene ? Scene->GetGPUSkinCache() : nullptr);
+	FSendAllEndOfFrameUpdates* SendAllEndOfFrameUpdates = BeginSendEndOfFrameUpdatesDrawEvent(
+		Scene ? Scene->GetGPUSkinCache() : nullptr,
+		Scene ? Scene->GetComputeFramework() : nullptr
+		);
 
 	// update all dirty components. 
 	FGuardValue_Bitfield(bPostTickComponentUpdate, true); 
diff --git a/Engine/Source/Runtime/Engine/Public/ComputeFramework/ComputeKernelShader.h b/Engine/Source/Runtime/Engine/Public/ComputeFramework/ComputeKernelShader.h
new file mode 100644
--- /dev/null
+++ b/Engine/Source/Runtime/Engine/Public/ComputeFramework/ComputeKernelShader.h
@@ -0,0 +1,7 @@
+// Copyright Epic Games, Inc. All Rights Reserved.
+
+#pragma once
+
+class FComputeKernelShader : public FShader
+{
+};
diff --git a/Engine/Source/Runtime/Engine/Public/SceneInterface.h b/Engine/Source/Runtime/Engine/Public/SceneInterface.h
--- a/Engine/Source/Runtime/Engine/Public/SceneInterface.h
+++ b/Engine/Source/Runtime/Engine/Public/SceneInterface.h
@@ -32,6 +32,7 @@ class UStaticMeshComponent;
 class UTextureCube;
 class FViewInfo;
 class FSceneRenderer;
+class FComputeFramework;
 
 enum EBasePassDrawListType
 {
@@ -398,6 +399,11 @@ public:
 		return nullptr;
 	}
 
+	virtual FComputeFramework* GetComputeFramework()
+	{
+		return nullptr;
+	}
+
 	/**
 	 * Sets the FX system associated with the scene.
 	 */
diff --git a/Engine/Source/Runtime/RHI/Private/RHI.cpp b/Engine/Source/Runtime/RHI/Private/RHI.cpp
--- a/Engine/Source/Runtime/RHI/Private/RHI.cpp
+++ b/Engine/Source/Runtime/RHI/Private/RHI.cpp
@@ -1562,6 +1562,7 @@ void FGenericDataDrivenShaderPlatformInfo::ParseDataDrivenShaderInfo(const FConf
 	Info.bSupportsGen5TemporalAA = GetSectionBool(Section, "bSupportsGen5TemporalAA");
 	Info.bTargetsTiledGPU = GetSectionBool(Section, "bTargetsTiledGPU");
 	Info.bNeedsOfflineCompiler = GetSectionBool(Section, "bNeedsOfflineCompiler");
+	Info.bSupportsComputeFramework = GetSectionBool(Section, "bSupportsComputeFramework");
 }
 
 void FGenericDataDrivenShaderPlatformInfo::Initialize()
diff --git a/Engine/Source/Runtime/RHI/Public/RHIDefinitions.h b/Engine/Source/Runtime/RHI/Public/RHIDefinitions.h
--- a/Engine/Source/Runtime/RHI/Public/RHIDefinitions.h
+++ b/Engine/Source/Runtime/RHI/Public/RHIDefinitions.h
@@ -345,9 +345,9 @@ class RHI_API FGenericDataDrivenShaderPlatformInfo
 	uint32 bSupportsConservativeRasterization : 1;
 	uint32 bRequiresExplicit128bitRT : 1;
 	uint32 bSupportsGen5TemporalAA : 1;
-
 	uint32 bTargetsTiledGPU: 1;
 	uint32 bNeedsOfflineCompiler: 1;
+	uint32 bSupportsComputeFramework : 1;
 
 	// NOTE: When adding fields, you must also add to ParseDataDrivenShaderInfo!
 	uint32 bContainsValidPlatformInfo : 1;
@@ -517,6 +517,11 @@ public:
 		return Infos[Platform].bSupportsGPUSkinCache;
 	}
 
+	static FORCEINLINE_DEBUGGABLE const bool GetSupportsComputeFramework(const FStaticShaderPlatform Platform)
+	{
+		return Infos[Platform].bSupportsComputeFramework;
+	}
+
 	static FORCEINLINE_DEBUGGABLE const bool GetTargetsTiledGPU(const FStaticShaderPlatform Platform)
 	{
 		return Infos[Platform].bTargetsTiledGPU;
diff --git a/Engine/Source/Runtime/Renderer/Private/RendererScene.cpp b/Engine/Source/Runtime/Renderer/Private/RendererScene.cpp
--- a/Engine/Source/Runtime/Renderer/Private/RendererScene.cpp
+++ b/Engine/Source/Runtime/Renderer/Private/RendererScene.cpp
@@ -55,6 +55,7 @@
 #include "PlanarReflectionSceneProxy.h"
 #include "Engine/StaticMesh.h"
 #include "GPUSkinCache.h"
+#include "ComputeFramework/ComputeFramework.h"
 #include "DynamicShadowMapChannelBindingHelper.h"
 #include "GPUScene.h"
 #include "HAL/LowLevelMemTracker.h"
@@ -1205,6 +1206,11 @@ FScene::FScene(UWorld* InWorld, bool bInRequiresHitProxies, bool bInIsEditorScen
 		GPUSkinCache = new FGPUSkinCache(bRequiresMemoryLimit);
 	}
 
+	if (SupportsComputeFramework(InFeatureLevel, GetFeatureLevelShaderPlatform(InFeatureLevel)))
+	{
+		ComputeFramework = new FComputeFramework();
+	}
+
 #if RHI_RAYTRACING
 	if (IsRayTracingEnabled())
 	{
diff --git a/Engine/Source/Runtime/Renderer/Private/ScenePrivate.h b/Engine/Source/Runtime/Renderer/Private/ScenePrivate.h
--- a/Engine/Source/Runtime/Renderer/Private/ScenePrivate.h
+++ b/Engine/Source/Runtime/Renderer/Private/ScenePrivate.h
@@ -86,6 +86,7 @@ class FRHIGPUTextureReadback;
 class FRuntimeVirtualTextureSceneProxy;
 class FLumenSceneData;
 class FVirtualShadowMapArrayCacheManager;
+class FComputeFramework;
 
 /** Holds information about a single primitive's occlusion. */
 class FPrimitiveOcclusionHistory
@@ -2702,6 +2703,9 @@ public:
 	/** GPU Skinning cache, if enabled */
 	class FGPUSkinCache* GPUSkinCache;
 
+	/* Root system that handles runtime execution of Compute Framework tasks */
+	FComputeFramework* ComputeFramework = nullptr;
+
 	/** Uniform buffers for parameter collections with the corresponding Ids. */
 	TMap<FGuid, FUniformBufferRHIRef> ParameterCollections;
 
