Temporariy reverting 14259363 until we get the missing header file

diff --git a/Engine/Source/Runtime/Engine/Private/ComputeFramework/ComputeFramework.cpp b/Engine/Source/Runtime/Engine/Private/ComputeFramework/ComputeFramework.cpp
deleted file mode 100644
diff --git a/Engine/Source/Runtime/Engine/Private/LevelTick.cpp b/Engine/Source/Runtime/Engine/Private/LevelTick.cpp
--- a/Engine/Source/Runtime/Engine/Private/LevelTick.cpp
+++ b/Engine/Source/Runtime/Engine/Private/LevelTick.cpp
@@ -59,7 +59,6 @@
 #include "Streaming/TextureStreamingHelpers.h"
 #include "ProfilingDebugging/CsvProfiler.h"
 #include "GPUSkinCache.h"
-#include "ComputeFramework/ComputeFramework.h"
 
 #if WITH_EDITOR
 #include "Editor.h"
@@ -928,22 +927,16 @@ bool UWorld::HasEndOfFrameUpdates() const
 
 struct FSendAllEndOfFrameUpdates
 {
-	FGPUSkinCache* GPUSkinCache = nullptr;
-	FComputeFramework* ComputeFramework = nullptr;
-
+	FGPUSkinCache* GPUSkinCache;
 #if WANTS_DRAW_MESH_EVENTS
 	FDrawEvent DrawEvent;
 #endif
 };
 
-FSendAllEndOfFrameUpdates* BeginSendEndOfFrameUpdatesDrawEvent(
-	FGPUSkinCache* GPUSkinCache,
-	FComputeFramework* ComputeFramework
-	)
+FSendAllEndOfFrameUpdates* BeginSendEndOfFrameUpdatesDrawEvent(FGPUSkinCache* GPUSkinCache)
 {
 	FSendAllEndOfFrameUpdates* SendAllEndOfFrameUpdates = new FSendAllEndOfFrameUpdates;
 	SendAllEndOfFrameUpdates->GPUSkinCache = GPUSkinCache;
-	SendAllEndOfFrameUpdates->ComputeFramework = ComputeFramework;
 
 #if WANTS_DRAW_MESH_EVENTS
 	ENQUEUE_RENDER_COMMAND(BeginDrawEventCommand)(
@@ -963,7 +956,7 @@ FSendAllEndOfFrameUpdates* BeginSendEndOfFrameUpdatesDrawEvent(
 void EndSendEndOfFrameUpdatesDrawEvent(FSendAllEndOfFrameUpdates* SendAllEndOfFrameUpdates)
 {
 	ENQUEUE_RENDER_COMMAND(EndDrawEventCommand)(
-		[SendAllEndOfFrameUpdates](FRHICommandListImmediate& RHICmdList)
+		[SendAllEndOfFrameUpdates](FRHICommandList& RHICmdList)
 	{
 		if (SendAllEndOfFrameUpdates->GPUSkinCache)
 		{
@@ -975,11 +968,6 @@ void EndSendEndOfFrameUpdatesDrawEvent(FSendAllEndOfFrameUpdates* SendAllEndOfFr
 		#endif
 		}
 
-		if (SendAllEndOfFrameUpdates->ComputeFramework)
-		{
-			SendAllEndOfFrameUpdates->ComputeFramework->ExecuteBatches(RHICmdList);
-		}
-
 		delete SendAllEndOfFrameUpdates;
 	});
 }
@@ -1021,10 +1009,7 @@ void UWorld::SendAllEndOfFrameUpdates()
 	}
 
 	// Issue a GPU event to wrap GPU work done during SendAllEndOfFrameUpdates, like skin cache updates
-	FSendAllEndOfFrameUpdates* SendAllEndOfFrameUpdates = BeginSendEndOfFrameUpdatesDrawEvent(
-		Scene ? Scene->GetGPUSkinCache() : nullptr,
-		Scene ? Scene->GetComputeFramework() : nullptr
-		);
+	FSendAllEndOfFrameUpdates* SendAllEndOfFrameUpdates = BeginSendEndOfFrameUpdatesDrawEvent(Scene ? Scene->GetGPUSkinCache() : nullptr);
 
 	// update all dirty components. 
 	FGuardValue_Bitfield(bPostTickComponentUpdate, true); 
diff --git a/Engine/Source/Runtime/Engine/Public/ComputeFramework/ComputeKernelShader.h b/Engine/Source/Runtime/Engine/Public/ComputeFramework/ComputeKernelShader.h
deleted file mode 100644
diff --git a/Engine/Source/Runtime/Engine/Public/SceneInterface.h b/Engine/Source/Runtime/Engine/Public/SceneInterface.h
--- a/Engine/Source/Runtime/Engine/Public/SceneInterface.h
+++ b/Engine/Source/Runtime/Engine/Public/SceneInterface.h
@@ -32,7 +32,6 @@ class UStaticMeshComponent;
 class UTextureCube;
 class FViewInfo;
 class FSceneRenderer;
-class FComputeFramework;
 
 enum EBasePassDrawListType
 {
@@ -399,11 +398,6 @@ public:
 		return nullptr;
 	}
 
-	virtual FComputeFramework* GetComputeFramework()
-	{
-		return nullptr;
-	}
-
 	/**
 	 * Sets the FX system associated with the scene.
 	 */
diff --git a/Engine/Source/Runtime/RHI/Private/RHI.cpp b/Engine/Source/Runtime/RHI/Private/RHI.cpp
--- a/Engine/Source/Runtime/RHI/Private/RHI.cpp
+++ b/Engine/Source/Runtime/RHI/Private/RHI.cpp
@@ -1562,7 +1562,6 @@ void FGenericDataDrivenShaderPlatformInfo::ParseDataDrivenShaderInfo(const FConf
 	Info.bSupportsGen5TemporalAA = GetSectionBool(Section, "bSupportsGen5TemporalAA");
 	Info.bTargetsTiledGPU = GetSectionBool(Section, "bTargetsTiledGPU");
 	Info.bNeedsOfflineCompiler = GetSectionBool(Section, "bNeedsOfflineCompiler");
-	Info.bSupportsComputeFramework = GetSectionBool(Section, "bSupportsComputeFramework");
 }
 
 void FGenericDataDrivenShaderPlatformInfo::Initialize()
diff --git a/Engine/Source/Runtime/RHI/Public/RHIDefinitions.h b/Engine/Source/Runtime/RHI/Public/RHIDefinitions.h
--- a/Engine/Source/Runtime/RHI/Public/RHIDefinitions.h
+++ b/Engine/Source/Runtime/RHI/Public/RHIDefinitions.h
@@ -345,9 +345,9 @@ class RHI_API FGenericDataDrivenShaderPlatformInfo
 	uint32 bSupportsConservativeRasterization : 1;
 	uint32 bRequiresExplicit128bitRT : 1;
 	uint32 bSupportsGen5TemporalAA : 1;
+
 	uint32 bTargetsTiledGPU: 1;
 	uint32 bNeedsOfflineCompiler: 1;
-	uint32 bSupportsComputeFramework : 1;
 
 	// NOTE: When adding fields, you must also add to ParseDataDrivenShaderInfo!
 	uint32 bContainsValidPlatformInfo : 1;
@@ -517,11 +517,6 @@ public:
 		return Infos[Platform].bSupportsGPUSkinCache;
 	}
 
-	static FORCEINLINE_DEBUGGABLE const bool GetSupportsComputeFramework(const FStaticShaderPlatform Platform)
-	{
-		return Infos[Platform].bSupportsComputeFramework;
-	}
-
 	static FORCEINLINE_DEBUGGABLE const bool GetTargetsTiledGPU(const FStaticShaderPlatform Platform)
 	{
 		return Infos[Platform].bTargetsTiledGPU;
diff --git a/Engine/Source/Runtime/Renderer/Private/RendererScene.cpp b/Engine/Source/Runtime/Renderer/Private/RendererScene.cpp
--- a/Engine/Source/Runtime/Renderer/Private/RendererScene.cpp
+++ b/Engine/Source/Runtime/Renderer/Private/RendererScene.cpp
@@ -55,7 +55,6 @@
 #include "PlanarReflectionSceneProxy.h"
 #include "Engine/StaticMesh.h"
 #include "GPUSkinCache.h"
-#include "ComputeFramework/ComputeFramework.h"
 #include "DynamicShadowMapChannelBindingHelper.h"
 #include "GPUScene.h"
 #include "HAL/LowLevelMemTracker.h"
@@ -1206,11 +1205,6 @@ FScene::FScene(UWorld* InWorld, bool bInRequiresHitProxies, bool bInIsEditorScen
 		GPUSkinCache = new FGPUSkinCache(bRequiresMemoryLimit);
 	}
 
-	if (SupportsComputeFramework(InFeatureLevel, GetFeatureLevelShaderPlatform(InFeatureLevel)))
-	{
-		ComputeFramework = new FComputeFramework();
-	}
-
 #if RHI_RAYTRACING
 	if (IsRayTracingEnabled())
 	{
diff --git a/Engine/Source/Runtime/Renderer/Private/ScenePrivate.h b/Engine/Source/Runtime/Renderer/Private/ScenePrivate.h
--- a/Engine/Source/Runtime/Renderer/Private/ScenePrivate.h
+++ b/Engine/Source/Runtime/Renderer/Private/ScenePrivate.h
@@ -86,7 +86,6 @@ class FRHIGPUTextureReadback;
 class FRuntimeVirtualTextureSceneProxy;
 class FLumenSceneData;
 class FVirtualShadowMapArrayCacheManager;
-class FComputeFramework;
 
 /** Holds information about a single primitive's occlusion. */
 class FPrimitiveOcclusionHistory
@@ -2703,9 +2702,6 @@ public:
 	/** GPU Skinning cache, if enabled */
 	class FGPUSkinCache* GPUSkinCache;
 
-	/* Root system that handles runtime execution of Compute Framework tasks */
-	FComputeFramework* ComputeFramework = nullptr;
-
 	/** Uniform buffers for parameter collections with the corresponding Ids. */
 	TMap<FGuid, FUniformBufferRHIRef> ParameterCollections;
 
