- Added a checkf to try to track down an Android cooktime error

diff --git a/Engine/Source/Runtime/Engine/Private/TextureLODSettings.cpp b/Engine/Source/Runtime/Engine/Private/TextureLODSettings.cpp
--- a/Engine/Source/Runtime/Engine/Private/TextureLODSettings.cpp
+++ b/Engine/Source/Runtime/Engine/Private/TextureLODSettings.cpp
@@ -136,6 +136,8 @@ int32 UTextureLODSettings::CalculateLODBias(const UTexture* Texture, bool bIncCi
 
 int32 UTextureLODSettings::CalculateLODBias(int32 Width, int32 Height, int32 MaxSize, int32 LODGroup, int32 LODBias, int32 NumCinematicMipLevels, TextureMipGenSettings InMipGenSetting, bool bVirtualTexture ) const
 {	
+	checkf(LODGroup < TextureLODGroups.Num(), TEXT("A texture had passed a bad LODGroup to UTextureLODSettings::CalculateLODBias (%d, out of %d groups). This code does not have the texture name. The LODSettings object is '%'"), LODGroup, TextureLODGroups.Num(), *GetPathName());
+
 	// Find LOD group.
 	const FTextureLODGroup& LODGroupInfo = TextureLODGroups[LODGroup];
 
