- Fixed cooking crashes in Client configs (and Android)

diff --git a/Engine/Source/Runtime/Engine/Private/DeviceProfiles/DeviceProfileManager.cpp b/Engine/Source/Runtime/Engine/Private/DeviceProfiles/DeviceProfileManager.cpp
--- a/Engine/Source/Runtime/Engine/Private/DeviceProfiles/DeviceProfileManager.cpp
+++ b/Engine/Source/Runtime/Engine/Private/DeviceProfiles/DeviceProfileManager.cpp
@@ -594,7 +594,7 @@ void UDeviceProfileManager::LoadProfiles()
 				ITargetPlatform* Platform = TargetPlatforms[PlatformIndex];
 
 				// Set TextureLODSettings
-				const UTextureLODSettings* TextureLODSettingsObj = FindProfile(*Platform->PlatformName(), false);
+				const UTextureLODSettings* TextureLODSettingsObj = FindProfile(*Platform->IniPlatformName(), false);
 				Platform->RegisterTextureLODSettings(TextureLODSettingsObj);
 			}
 		}
