New UWorld AssetRegistry Tag "LevelIsUsingExternalActors"

diff --git a/Engine/Source/Runtime/Engine/Classes/Engine/Level.h b/Engine/Source/Runtime/Engine/Classes/Engine/Level.h
--- a/Engine/Source/Runtime/Engine/Classes/Engine/Level.h
+++ b/Engine/Source/Runtime/Engine/Classes/Engine/Level.h
@@ -709,6 +709,7 @@ public:
 #if WITH_EDITOR
 	ENGINE_API static bool GetLevelBoundsFromPackage(FName LevelPackage, FBox& OutLevelBounds);
 	ENGINE_API static bool GetIsLevelPartitionedFromPackage(FName LevelPackage);
+	ENGINE_API static bool GetIsLevelUsingExternalActorsFromPackage(FName LevelPackage);
 #endif
 
 private:
diff --git a/Engine/Source/Runtime/Engine/Private/Level.cpp b/Engine/Source/Runtime/Engine/Private/Level.cpp
--- a/Engine/Source/Runtime/Engine/Private/Level.cpp
+++ b/Engine/Source/Runtime/Engine/Private/Level.cpp
@@ -1745,6 +1745,21 @@ bool ULevel::GetIsLevelPartitionedFromPackage(FName LevelPackage)
 	});
 }
 
+bool ULevel::GetIsLevelUsingExternalActorsFromPackage(FName LevelPackage)
+{
+	return LevelAssetRegistryHelper::GetLevelInfoFromAssetRegistry(LevelPackage, [](const FAssetData& Asset)
+	{
+		FString LevelIsUsingExternalActorsStr;
+		static const FName NAME_LevelIsUsingExternalActors(TEXT("LevelIsUsingExternalActors"));
+		if (Asset.GetTagValue(NAME_LevelIsUsingExternalActors, LevelIsUsingExternalActorsStr))
+		{
+			check(LevelIsUsingExternalActorsStr == TEXT("1"));
+			return true;
+		}
+		return false;
+	});
+}
+
 bool ULevel::GetLevelBoundsFromPackage(FName LevelPackage, FBox& OutLevelBounds)
 {
 	return LevelAssetRegistryHelper::GetLevelInfoFromAssetRegistry(LevelPackage, [&OutLevelBounds](const FAssetData& Asset)
diff --git a/Engine/Source/Runtime/Engine/Private/World.cpp b/Engine/Source/Runtime/Engine/Private/World.cpp
--- a/Engine/Source/Runtime/Engine/Private/World.cpp
+++ b/Engine/Source/Runtime/Engine/Private/World.cpp
@@ -7725,6 +7725,12 @@ void UWorld::GetAssetRegistryTags(TArray<FAssetRegistryTag>& OutTags) const
 			static const FName NAME_LevelIsPartitioned(TEXT("LevelIsPartitioned"));
 			OutTags.Add(FAssetRegistryTag(NAME_LevelIsPartitioned, TEXT("1"), FAssetRegistryTag::TT_Hidden));
 		}
+
+		if (PersistentLevel->IsUsingExternalActors())
+		{
+			static const FName NAME_LevelIsUsingExternalActors(TEXT("LevelIsUsingExternalActors"));
+			OutTags.Add(FAssetRegistryTag(NAME_LevelIsUsingExternalActors, TEXT("1"), FAssetRegistryTag::TT_Hidden));
+		}
 	}
 
 	// Get the full file path with extension
