WorldPartition: Added optional NextPreferredLevelPendingVisibility in world to give a hint of which of the MakingVisible StreamingLevels should be used as CurrentLevelPendingVisibility.
Added angle between streaming source forward vector and streaming cell to source vector as another criteria to sort spatial hash runtime cells by importance (grid cell level, grid cell distance and now angle contribute to sort cells).

diff --git a/Engine/Source/Runtime/Engine/Classes/Engine/World.h b/Engine/Source/Runtime/Engine/Classes/Engine/World.h
--- a/Engine/Source/Runtime/Engine/Classes/Engine/World.h
+++ b/Engine/Source/Runtime/Engine/Classes/Engine/World.h
@@ -961,6 +961,9 @@ public:
 	/** Returns the level, if any, in the process of being made visible */
 	ULevel* GetCurrentLevelPendingVisibility() const { return CurrentLevelPendingVisibility; }
 
+	/** Returns the best candidate level, if any, to become the CurrentLevelPendingVisibility */
+	ULevel* GetNextPreferredLevelPendingVisibility() { return NextPreferredLevelPendingVisibility; }
+
 	/** Returns the level, if any, in the process of being made invisible */
 	ULevel* GetCurrentLevelPendingInvisibility() const { return CurrentLevelPendingInvisibility; }
 
@@ -1018,6 +1021,10 @@ private:
 	UPROPERTY(Transient)
 	class ULevel*								CurrentLevelPendingVisibility;
 
+	/** Pointer to the next best candidate level to become the CurrentLevelPendingVisibility. */
+	UPROPERTY(Transient)
+	class ULevel*								NextPreferredLevelPendingVisibility;
+
 	/** Pointer to the current level in the queue to be made invisible, NULL if none are pending. */
 	UPROPERTY(Transient)
 	class ULevel*								CurrentLevelPendingInvisibility;
@@ -3310,6 +3317,9 @@ private:
 	/** Updates world's required streaming levels */
 	void InternalUpdateStreamingState();
 
+	/** Determine next making visible best candidate to AddToWorld */
+	void UpdateNextPreferredLevelPendingVisibility();
+
 #if WITH_EDITOR
 public:
 	void SetPlayInEditorInitialNetMode(ENetMode InNetMode)
diff --git a/Engine/Source/Runtime/Engine/Private/LevelStreaming.cpp b/Engine/Source/Runtime/Engine/Private/LevelStreaming.cpp
--- a/Engine/Source/Runtime/Engine/Private/LevelStreaming.cpp
+++ b/Engine/Source/Runtime/Engine/Private/LevelStreaming.cpp
@@ -570,23 +570,26 @@ void ULevelStreaming::UpdateStreamingState(bool& bOutUpdateAgain, bool& bOutRede
 	case ECurrentState::MakingVisible:
 		if (ensure(LoadedLevel))
 		{
-			World->AddToWorld(LoadedLevel, LevelTransform, !bShouldBlockOnLoad);
-
-			if (LoadedLevel->bIsVisible)
+			if (!World->GetNextPreferredLevelPendingVisibility() || World->GetNextPreferredLevelPendingVisibility() == LoadedLevel)
 			{
-				// immediately discard previous level
-				DiscardPendingUnloadLevel(World);
+				World->AddToWorld(LoadedLevel, LevelTransform, !bShouldBlockOnLoad);
 
-				if (World->Scene)
+				if (LoadedLevel->bIsVisible)
 				{
-					QUICK_SCOPE_CYCLE_COUNTER(STAT_UpdateLevelStreamingInner_OnLevelAddedToWorld);
-					// Notify the new level has been added after the old has been discarded
-					World->Scene->OnLevelAddedToWorld(LoadedLevel->GetOutermost()->GetFName(), World, LoadedLevel->bIsLightingScenario);
-				}
+					// immediately discard previous level
+					DiscardPendingUnloadLevel(World);
 
-				CurrentState = ECurrentState::LoadedVisible;
-				bOutUpdateAgain = true;
-				bOutRedetermineTarget = true;
+					if (World->Scene)
+					{
+						QUICK_SCOPE_CYCLE_COUNTER(STAT_UpdateLevelStreamingInner_OnLevelAddedToWorld);
+						// Notify the new level has been added after the old has been discarded
+						World->Scene->OnLevelAddedToWorld(LoadedLevel->GetOutermost()->GetFName(), World, LoadedLevel->bIsLightingScenario);
+					}
+
+					CurrentState = ECurrentState::LoadedVisible;
+					bOutUpdateAgain = true;
+					bOutRedetermineTarget = true;
+				}
 			}
 		}
 		break;
diff --git a/Engine/Source/Runtime/Engine/Private/World.cpp b/Engine/Source/Runtime/Engine/Private/World.cpp
--- a/Engine/Source/Runtime/Engine/Private/World.cpp
+++ b/Engine/Source/Runtime/Engine/Private/World.cpp
@@ -371,6 +371,7 @@ UWorld::UWorld( const FObjectInitializer& ObjectInitializer )
 ,	NextTravelType(TRAVEL_Relative)
 ,	CleanupWorldTag(0)
 {
+	NextPreferredLevelPendingVisibility = nullptr;
 	TimerManager = new FTimerManager();
 #if WITH_EDITOR
 	SetPlayInEditorInitialNetMode(ENetMode::NM_Standalone);
@@ -2405,6 +2406,7 @@ void UWorld::AddToWorld( ULevel* Level, const FTransform& LevelTransform, bool b
 		
 		// Mark level as being the one in process of being made visible.
 		CurrentLevelPendingVisibility = Level;
+		NextPreferredLevelPendingVisibility = nullptr;
 
 		// Add to the UWorld's array of levels, which causes it to be rendered et al.
 		Levels.AddUnique( Level );
@@ -2606,7 +2608,7 @@ void UWorld::AddToWorld( ULevel* Level, const FTransform& LevelTransform, bool b
 		Level->bAlreadyClearedActorsSeamlessTravelFlag	= false;
 
 		// Finished making level visible - allow other levels to be added to the world.
-		CurrentLevelPendingVisibility					= NULL;
+		CurrentLevelPendingVisibility = nullptr;
 
 		// notify server that the client has finished making this level visible
 		if (!Level->bClientOnlyVisible)
@@ -3486,6 +3488,15 @@ void UWorld::InternalUpdateStreamingState()
 	}
 }
 
+void UWorld::UpdateNextPreferredLevelPendingVisibility()
+{
+	if (CurrentLevelPendingVisibility == nullptr)
+	{
+		const UWorldPartition* WorldPartition = GetWorldPartition();
+		NextPreferredLevelPendingVisibility = WorldPartition ? WorldPartition->GetPreferredLoadedLevelToAddToWorld() : nullptr;
+	}
+}
+
 void UWorld::UpdateLevelStreaming()
 {
 	SCOPE_CYCLE_COUNTER(STAT_UpdateLevelStreamingTime);
@@ -3507,6 +3518,8 @@ void UWorld::UpdateLevelStreaming()
 	{
 		if (ULevelStreaming* StreamingLevel = StreamingLevelsToConsider.GetStreamingLevels()[Index])
 		{
+			UpdateNextPreferredLevelPendingVisibility();
+
 			bool bUpdateAgain = true;
 			bool bShouldContinueToConsider = true;
 			while (bUpdateAgain && bShouldContinueToConsider)
diff --git a/Engine/Source/Runtime/Engine/Private/WorldPartition/WorldPartition.cpp b/Engine/Source/Runtime/Engine/Private/WorldPartition/WorldPartition.cpp
--- a/Engine/Source/Runtime/Engine/Private/WorldPartition/WorldPartition.cpp
+++ b/Engine/Source/Runtime/Engine/Private/WorldPartition/WorldPartition.cpp
@@ -1155,6 +1155,15 @@ void UWorldPartition::UpdateStreamingState()
 	}
 }
 
+class ULevel* UWorldPartition::GetPreferredLoadedLevelToAddToWorld() const
+{
+	if (GetWorld()->IsGameWorld())
+	{
+		return GetStreamingPolicy()->GetPreferredLoadedLevelToAddToWorld();
+	}
+	return nullptr;
+}
+
 FVector2D UWorldPartition::GetDrawRuntimeHash2DDesiredFootprint(const FVector2D& CanvasSize)
 {
 	check(GetWorld()->IsGameWorld());
diff --git a/Engine/Source/Runtime/Engine/Private/WorldPartition/WorldPartitionLevelStreamingPolicy.cpp b/Engine/Source/Runtime/Engine/Private/WorldPartition/WorldPartitionLevelStreamingPolicy.cpp
--- a/Engine/Source/Runtime/Engine/Private/WorldPartition/WorldPartitionLevelStreamingPolicy.cpp
+++ b/Engine/Source/Runtime/Engine/Private/WorldPartition/WorldPartitionLevelStreamingPolicy.cpp
@@ -11,6 +11,7 @@
 #include "WorldPartition/WorldPartitionLevelStreamingDynamic.h"
 #include "WorldPartition/WorldPartitionRuntimeHash.h"
 #include "WorldPartition/WorldPartitionLevelHelper.h"
+#include "Engine/Level.h"
 #include "Engine/World.h"
 
 #if WITH_EDITOR
@@ -54,9 +55,9 @@ void UWorldPartitionLevelStreamingPolicy::LoadCells(const TSet<const UWorldParti
 		return;
 	}
 
-	// Sort cells based on shortest distance
-	TArray<const UWorldPartitionRuntimeCell*> SortedCells;
-	WorldPartition->RuntimeHash->SortStreamingCellsByDistance(ToLoadCells, StreamingSources, SortedCells);
+	// Sort cells based on importance
+	TArray<const UWorldPartitionRuntimeCell*, TInlineAllocator<256>> SortedCells;
+	WorldPartition->RuntimeHash->SortStreamingCellsByImportance(ToLoadCells, StreamingSources, SortedCells);
 
 	// Trigger cell loading until we hit the maximum
 	for (const UWorldPartitionRuntimeCell* Cell : SortedCells)
@@ -73,6 +74,34 @@ void UWorldPartitionLevelStreamingPolicy::LoadCells(const TSet<const UWorldParti
 	}
 }
 
+ULevel* UWorldPartitionLevelStreamingPolicy::GetPreferredLoadedLevelToAddToWorld() const
+{
+	check(WorldPartition->IsInitialized());
+	UWorld* World = WorldPartition->GetWorld();
+	if (World->GetNetMode() == NM_DedicatedServer)
+	{
+		return nullptr;
+	}
+
+	// Sort loaded cells based on importance (only those with a streaming level in MakingVisible state)
+	TArray<const UWorldPartitionRuntimeCell*, TInlineAllocator<256>> SortedMakingVisibileCells;
+	TSet<const UWorldPartitionRuntimeCell*> MakingVisibileCells;
+	MakingVisibileCells.Reserve(LoadedCells.Num());
+	for (const UWorldPartitionRuntimeCell* LoadedCell : LoadedCells)
+	{
+		if (const UWorldPartitionRuntimeLevelStreamingCell* LevelSreamingCell = Cast<const UWorldPartitionRuntimeLevelStreamingCell>(LoadedCell))
+		{
+			ULevelStreaming* LevelStreaming = LevelSreamingCell->GetLevelStreaming();
+			if (LevelStreaming && LevelStreaming->GetLoadedLevel() && (LevelStreaming->GetCurrentState() == ULevelStreaming::ECurrentState::MakingVisible))
+			{
+				MakingVisibileCells.Add(LoadedCell);
+			}
+		}
+	}
+	WorldPartition->RuntimeHash->SortStreamingCellsByImportance(MakingVisibileCells, StreamingSources, SortedMakingVisibileCells);
+	return SortedMakingVisibileCells.Num() > 0 ? CastChecked<UWorldPartitionRuntimeLevelStreamingCell>(SortedMakingVisibileCells[0])->GetLevelStreaming()->GetLoadedLevel() : nullptr;
+}
+
 void UWorldPartitionLevelStreamingPolicy::LoadCell(const UWorldPartitionRuntimeCell* InCell)
 {
 	UE_LOG(LogWorldPartition, Verbose, TEXT("UWorldPartitionLevelStreamingPolicy::LoadCell %s"), *InCell->GetName());
diff --git a/Engine/Source/Runtime/Engine/Private/WorldPartition/WorldPartitionRuntimeSpatialHash.cpp b/Engine/Source/Runtime/Engine/Private/WorldPartition/WorldPartitionRuntimeSpatialHash.cpp
--- a/Engine/Source/Runtime/Engine/Private/WorldPartition/WorldPartitionRuntimeSpatialHash.cpp
+++ b/Engine/Source/Runtime/Engine/Private/WorldPartition/WorldPartitionRuntimeSpatialHash.cpp
@@ -87,6 +87,18 @@ static FAutoConsoleVariableRef CVarShowRuntimeSpatialHashGridLevelCount(
 	GShowRuntimeSpatialHashGridLevelCount,
 	TEXT("Used to choose how many grid levels to display when showing world partition runtime hash."));
 
+static int32 GShowRuntimeSpatialHashGridIndex = 0;
+static FAutoConsoleVariableRef CVarShowRuntimeSpatialHashGridIndex(
+	TEXT("WorldPartition.ShowRuntimeSpatialHashGridIndex"),
+	GShowRuntimeSpatialHashGridIndex,
+	TEXT("Used to show only one particual grid when showing world partition runtime hash (invalid index will show all)."));
+
+static float GRuntimeSpatialHashCellToSourceAngleContributionToCellImportance = 0.4f; // Value between [0, 1]
+static FAutoConsoleVariableRef CVarRuntimeSpatialHashCellToSourceAngleContributionToCellImportance(
+	TEXT("WorldPartition.RuntimeSpatialHashCellToSourceAngleContributionToCellImportance"),
+	GRuntimeSpatialHashCellToSourceAngleContributionToCellImportance,
+	TEXT("Value between 0 and 1 that modulates the contribution of the angle between streaming source-to-cell vector and source-forward vector to the cell importance. The closest to 0, the less the angle will contribute to the cell importance."));
+
 // ------------------------------------------------------------------------------------------------
 
 /**
@@ -589,7 +601,7 @@ void FSpatialHashStreamingGrid::Draw3D(UWorld* World, const TArray<FWorldPartiti
 				
 				FVector BoundsExtent(CellWorldBounds.GetExtent(), 100.f);
 				const UWorldPartitionRuntimeSpatialHashCell* Cell = Cast<const UWorldPartitionRuntimeSpatialHashCell>(GridLevels[GridLevel].GridCells[Coords.Y * Helper.Levels[GridLevel].GridSize + Coords.X]);
-				FColor CellColor = Cell ? Cell->GetDebugColor().ToFColor(false).WithAlpha(32) : FColor(0, 0, 0, 32);
+				FColor CellColor = Cell ? Cell->GetDebugColor().ToFColor(false).WithAlpha(16) : FColor(0, 0, 0, 16);
 				FVector BoundsOrigin(CellWorldBounds.GetCenter(), Z);
 				DrawDebugSolidBox(World, BoundsOrigin, BoundsExtent, CellColor, false, -1.f, 255);
 				DrawDebugBox(World, BoundsOrigin, BoundsExtent, CellColor.WithAlpha(255), false, -1.f, 255, 10.f);
@@ -1292,7 +1304,7 @@ int32 UWorldPartitionRuntimeSpatialHash::GetStreamingCells(const TArray<FWorldPa
 	return Cells.Num();
 }
 
-void UWorldPartitionRuntimeSpatialHash::SortStreamingCellsByDistance(const TSet<const UWorldPartitionRuntimeCell*>& InCells, const TArray<FWorldPartitionStreamingSource>& InSources, TArray<const UWorldPartitionRuntimeCell*>& OutSortedCells)
+void UWorldPartitionRuntimeSpatialHash::SortStreamingCellsByImportance(const TSet<const UWorldPartitionRuntimeCell*>& InCells, const TArray<FWorldPartitionStreamingSource>& InSources, TArray<const UWorldPartitionRuntimeCell*, TInlineAllocator<256>>& OutSortedCells) const
 {
 	struct FCellShortestDist
 	{
@@ -1311,10 +1323,21 @@ void UWorldPartitionRuntimeSpatialHash::SortStreamingCellsByDistance(const TSet<
 		const UWorldPartitionRuntimeSpatialHashCell* Cell = Cast<const UWorldPartitionRuntimeSpatialHashCell>(ToLoadCell);
 		FCellShortestDist& SortedCell = SortedCells.Emplace_GetRef(Cell, FLT_MAX);
 
+		const float AngleContribution = FMath::Clamp(GRuntimeSpatialHashCellToSourceAngleContributionToCellImportance, 0.f, 1.f);
 		for (const FWorldPartitionStreamingSource& Source : InSources)
 		{
-			float SqrDistance = FVector::DistSquared(Source.Location, Cell->Position);
-			SortedCell.SourceMinDistance = FMath::Min(SqrDistance, SortedCell.SourceMinDistance);
+			const float SqrDistance = FVector::DistSquared(Source.Location, Cell->Position);
+			float AngleFactor = 1.f;
+			if (!FMath::IsNearlyZero(AngleContribution))
+			{
+				const FVector2D SourceForward(Source.Rotation.Quaternion().GetForwardVector());
+				const FVector2D SourceToCell(Cell->Position - Source.Location);
+				const float Dot = FVector2D::DotProduct(SourceForward.GetSafeNormal(), SourceToCell.GetSafeNormal());
+				const float NormalizedAngle = FMath::Clamp(FMath::Abs(FMath::Acos(Dot)/PI), 0.f, 1.f);
+				AngleFactor = FMath::Pow(NormalizedAngle, AngleContribution);
+			}
+			// Modulate distance to cell by angle relative to source forward vector (to prioritize cells in front)
+			SortedCell.SourceMinDistance = FMath::Min(SqrDistance * AngleFactor, SortedCell.SourceMinDistance);
 		}
 	}
 
@@ -1402,9 +1425,16 @@ void UWorldPartitionRuntimeSpatialHash::Draw2D(UCanvas* Canvas, const TArray<FWo
 void UWorldPartitionRuntimeSpatialHash::Draw3D(const TArray<FWorldPartitionStreamingSource>& Sources) const
 {
 	UWorld* World = GetWorld();
-	// Get cells based on streaming sources
-	for (const FSpatialHashStreamingGrid& StreamingGrid : StreamingGrids)
+
+	if (StreamingGrids.IsValidIndex(GShowRuntimeSpatialHashGridIndex))
 	{
-		StreamingGrid.Draw3D(World, Sources);
+		StreamingGrids[GShowRuntimeSpatialHashGridIndex].Draw3D(World, Sources);
+	}
+	else
+	{
+		for (const FSpatialHashStreamingGrid& StreamingGrid : StreamingGrids)
+		{
+			StreamingGrid.Draw3D(World, Sources);
+		}
 	}
 }
\ No newline at end of file
diff --git a/Engine/Source/Runtime/Engine/Public/WorldPartition/WorldPartition.h b/Engine/Source/Runtime/Engine/Public/WorldPartition/WorldPartition.h
--- a/Engine/Source/Runtime/Engine/Public/WorldPartition/WorldPartition.h
+++ b/Engine/Source/Runtime/Engine/Public/WorldPartition/WorldPartition.h
@@ -159,6 +159,7 @@ public:
 
 	void Tick(float DeltaSeconds);
 	void UpdateStreamingState();
+	class ULevel* GetPreferredLoadedLevelToAddToWorld() const;
 
 	FVector2D GetDrawRuntimeHash2DDesiredFootprint(const FVector2D& CanvasSize);
 	void DrawRuntimeHash2D(class UCanvas* Canvas, const FVector2D& PartitionCanvasOffset, const FVector2D& PartitionCanvasSize);
diff --git a/Engine/Source/Runtime/Engine/Public/WorldPartition/WorldPartitionLevelStreamingPolicy.h b/Engine/Source/Runtime/Engine/Public/WorldPartition/WorldPartitionLevelStreamingPolicy.h
--- a/Engine/Source/Runtime/Engine/Public/WorldPartition/WorldPartitionLevelStreamingPolicy.h
+++ b/Engine/Source/Runtime/Engine/Public/WorldPartition/WorldPartitionLevelStreamingPolicy.h
@@ -23,6 +23,7 @@ public:
 	virtual void LoadCells(const TSet<const UWorldPartitionRuntimeCell*>& ToLoadCells) override;
 	virtual void LoadCell(const UWorldPartitionRuntimeCell* Cell) override;
 	virtual void UnloadCell(const UWorldPartitionRuntimeCell* Cell) override;
+	virtual class ULevel* GetPreferredLoadedLevelToAddToWorld() const override;
 
 #if WITH_EDITOR
 	virtual TSubclassOf<class UWorldPartitionRuntimeCell> GetRuntimeCellClass() const override;
diff --git a/Engine/Source/Runtime/Engine/Public/WorldPartition/WorldPartitionRuntimeHash.h b/Engine/Source/Runtime/Engine/Public/WorldPartition/WorldPartitionRuntimeHash.h
--- a/Engine/Source/Runtime/Engine/Public/WorldPartition/WorldPartitionRuntimeHash.h
+++ b/Engine/Source/Runtime/Engine/Public/WorldPartition/WorldPartitionRuntimeHash.h
@@ -27,7 +27,7 @@ class ENGINE_API UWorldPartitionRuntimeHash : public UObject
 	// Streaming interface
 	virtual int32 GetAllStreamingCells(TSet<const UWorldPartitionRuntimeCell*>& Cells) const { return 0; }
 	virtual int32 GetStreamingCells(const TArray<FWorldPartitionStreamingSource>& Sources, TSet<const UWorldPartitionRuntimeCell*>& Cells) const { return 0; };
-	virtual void SortStreamingCellsByDistance(const TSet<const UWorldPartitionRuntimeCell*>& InCells, const TArray<FWorldPartitionStreamingSource>& InSources, TArray<const UWorldPartitionRuntimeCell*>& OutSortedCells) {}
+	virtual void SortStreamingCellsByImportance(const TSet<const UWorldPartitionRuntimeCell*>& InCells, const TArray<FWorldPartitionStreamingSource>& InSources, TArray<const UWorldPartitionRuntimeCell*, TInlineAllocator<256>>& OutSortedCells) const {}
 
 	/* Returns desired footprint that Draw2D should take relative to given Canvas size (the value can exceed the given size).
 	 * UWorldPartitionSubSystem will re-adapt the size relative to all others UWorldPartitionRuntimeHash and provide the correct size to Draw2D.
diff --git a/Engine/Source/Runtime/Engine/Public/WorldPartition/WorldPartitionRuntimeSpatialHash.h b/Engine/Source/Runtime/Engine/Public/WorldPartition/WorldPartitionRuntimeSpatialHash.h
--- a/Engine/Source/Runtime/Engine/Public/WorldPartition/WorldPartitionRuntimeSpatialHash.h
+++ b/Engine/Source/Runtime/Engine/Public/WorldPartition/WorldPartitionRuntimeSpatialHash.h
@@ -154,7 +154,7 @@ public:
 	// streaming interface
 	virtual int32 GetAllStreamingCells(TSet<const UWorldPartitionRuntimeCell*>& Cells) const override;
 	virtual int32 GetStreamingCells(const TArray<FWorldPartitionStreamingSource>& Sources, TSet<const UWorldPartitionRuntimeCell*>& Cells) const override;
-	virtual void SortStreamingCellsByDistance(const TSet<const UWorldPartitionRuntimeCell*>& InCells, const TArray<FWorldPartitionStreamingSource>& InSources, TArray<const UWorldPartitionRuntimeCell*>& OutSortedCells) override;
+	virtual void SortStreamingCellsByImportance(const TSet<const UWorldPartitionRuntimeCell*>& InCells, const TArray<FWorldPartitionStreamingSource>& InSources, TArray<const UWorldPartitionRuntimeCell*, TInlineAllocator<256>>& OutSortedCells) const override;
 
 protected:
 	// Used to convert from a world using World Composition
@@ -197,4 +197,5 @@ private:
 #if WITH_EDITOR
 	bool CreateStreamingGrid(const FSpatialHashRuntimeGrid& RuntimeGrid, const FSquare2DGridHelper& PartionedActors, EWorldPartitionStreamingMode Mode, UWorldPartitionStreamingPolicy* StreamingPolicy);
 #endif
+friend class UWorldPartitionSubsystem;
 };
\ No newline at end of file
diff --git a/Engine/Source/Runtime/Engine/Public/WorldPartition/WorldPartitionStreamingPolicy.h b/Engine/Source/Runtime/Engine/Public/WorldPartition/WorldPartitionStreamingPolicy.h
--- a/Engine/Source/Runtime/Engine/Public/WorldPartition/WorldPartitionStreamingPolicy.h
+++ b/Engine/Source/Runtime/Engine/Public/WorldPartition/WorldPartitionStreamingPolicy.h
@@ -37,6 +37,7 @@ public:
 	virtual void UnloadCells(const TSet<const UWorldPartitionRuntimeCell*>& ToUnloadCells);
 	virtual void LoadCell(const UWorldPartitionRuntimeCell* Cell) PURE_VIRTUAL(UWorldPartitionStreamingPolicy::LoadCell, );
 	virtual void UnloadCell(const UWorldPartitionRuntimeCell* Cell) PURE_VIRTUAL(UWorldPartitionStreamingPolicy::UnloadCell, );
+	virtual class ULevel* GetPreferredLoadedLevelToAddToWorld() const { return nullptr; }
 	virtual FVector2D GetDrawRuntimeHash2DDesiredFootprint(const FVector2D& CanvasSize);
 	virtual void DrawRuntimeHash2D(class UCanvas* Canvas, const FVector2D& PartitionCanvasOffset, const FVector2D& PartitionCanvasSize);
 	virtual void DrawRuntimeHash3D();
