Fix missing forward declare

diff --git a/Engine/Source/Runtime/Engine/Public/ComputeFramework/ComputeFramework.h b/Engine/Source/Runtime/Engine/Public/ComputeFramework/ComputeFramework.h
--- a/Engine/Source/Runtime/Engine/Public/ComputeFramework/ComputeFramework.h
+++ b/Engine/Source/Runtime/Engine/Public/ComputeFramework/ComputeFramework.h
@@ -6,6 +6,7 @@ class FComputeKernelShader;
 
 extern ENGINE_API bool SupportsComputeFramework(ERHIFeatureLevel::Type FeatureLevel, EShaderPlatform ShaderPlatform);
 
+struct FPooledRDGBuffer;
 struct FComputeExecutionExternalBufferDesc
 {
 	FName Name;
