World Partition Editor
- Octree nodes now maintain a loaded childs mask to know if there is loaded editing cells down the tree.
- Reworked how we shadow unloaded cells in the editor to hierarchically shadow unloaded regions based on the loaded child mask.

diff --git a/Engine/Source/Editor/WorldPartitionEditor/Private/WorldPartition/SWorldPartitionEditorGridSpatialHash.cpp b/Engine/Source/Editor/WorldPartitionEditor/Private/WorldPartition/SWorldPartitionEditorGridSpatialHash.cpp
--- a/Engine/Source/Editor/WorldPartitionEditor/Private/WorldPartition/SWorldPartitionEditorGridSpatialHash.cpp
+++ b/Engine/Source/Editor/WorldPartitionEditor/Private/WorldPartition/SWorldPartitionEditorGridSpatialHash.cpp
@@ -24,6 +24,8 @@
 
 #define LOCTEXT_NAMESPACE "WorldPartitionEditor"
 
+TAutoConsoleVariable<bool> CVarDebugDrawOctree(TEXT("wp.Editor.DebugDrawOctree"), false, TEXT("Whether to debug draw the World Partition octree"), ECVF_Default);
+
 SWorldPartitionEditorGridSpatialHash::SWorldPartitionEditorGridSpatialHash()
 {
 }
@@ -84,7 +86,7 @@ int32 SWorldPartitionEditorGridSpatialHash::PaintGrid(const FGeometry& AllottedG
 	// Shadow whole grid area
 	{
 		FSlateColorBrush ShadowBrush(FLinearColor::Black);
-		FLinearColor ShadowColor(FLinearColor(0.0f, 0.0f, 0.0f, 0.25f));
+		FLinearColor ShadowColor(FLinearColor(0.0f, 0.0f, 0.0f, 0.5f));
 
 		FPaintGeometry GridGeometry = AllottedGeometry.ToPaintGeometry(
 			WorldToScreen.TransformPoint(FVector2D(VisibleGridRectWorld.Min)),
@@ -122,78 +124,101 @@ int32 SWorldPartitionEditorGridSpatialHash::PaintGrid(const FGeometry& AllottedG
 		);
 	}
 
+	struct FCellDesc2D
 	{
-		// Flatten cells in 2D
-		struct FCellDesc2D
-		{
-			FCellDesc2D()
-				: Bounds(ForceInitToZero)
-				, bSelected(false)
-				, bLoaded(false)
-				, bEmpty(true)
-			{}
-
-			FBox2D Bounds;
-			bool bSelected : 1;
-			bool bLoaded : 1;
-			bool bEmpty : 1;
-		};
+		FCellDesc2D()
+			: Bounds(ForceInitToZero)
+		{}
 
+		FBox2D Bounds;
+	};
+
+	// Draw shadowed regions
+	{
 		TMap<uint32, FCellDesc2D> UniqueCells2D;
-		WorldPartition->EditorHash->ForEachIntersectingCell(VisibleGridRectWorld, [&](UWorldPartitionEditorCell* Cell)
+		EditorSpatialHash->ForEachIntersectingUnloadedRegion(VisibleGridRectWorld, [&](const UWorldPartitionEditorSpatialHash::FCellCoord& CellCoord)
 		{
-			// Compute a hash of the 2D bounds, snapped to centimeters
+			FBox CellBounds = EditorSpatialHash->GetCellBounds(CellCoord);
+			CellBounds = CellBounds.Overlap(EditorSpatialHash->Bounds);
+
 			FHashBuilder HashBuilder;
-			HashBuilder << FMath::RoundToInt(Cell->Bounds.Min.X) 
-						<< FMath::RoundToInt(Cell->Bounds.Min.Y) 
-						<< FMath::RoundToInt(Cell->Bounds.Max.X) 
-						<< FMath::RoundToInt(Cell->Bounds.Max.Y);
+			HashBuilder << CellCoord.X << CellCoord.Y << CellCoord.Level;
 			uint32 CellHash2D = HashBuilder.GetHash();
 
 			FCellDesc2D& CellDesc2D = UniqueCells2D.Add(CellHash2D);
-
-			CellDesc2D.Bounds = FBox2D(FVector2D(Cell->Bounds.Min), FVector2D(Cell->Bounds.Max));
-			CellDesc2D.bSelected |= SelectedCells.Contains(Cell);
-			CellDesc2D.bLoaded |= Cell->bLoaded;
-			CellDesc2D.bEmpty &= !Cell->Actors.Num();
+			CellDesc2D.Bounds = FBox2D(FVector2D(CellBounds.Min), FVector2D(CellBounds.Max));
 		});
 
+		const bool bDebugDrawOctree = CVarDebugDrawOctree.GetValueOnAnyThread();
+
 		for(auto& UniqueCell: UniqueCells2D)
 		{
 			const FCellDesc2D& Cell = UniqueCell.Value;
 
-			if (!Cell.bLoaded || Cell.bEmpty || Cell.bSelected)
+			FPaintGeometry CellGeometry = AllottedGeometry.ToPaintGeometry(
+				WorldToScreen.TransformPoint(Cell.Bounds.Min),
+				WorldToScreen.TransformPoint(Cell.Bounds.Max) - WorldToScreen.TransformPoint(Cell.Bounds.Min)
+			);
+
+			FSlateColorBrush CellBrush(FLinearColor::White);
+			FLinearColor CellColor(0, 0, 0, 0.5f);
+
+			if (bDebugDrawOctree)
+			{
+				CellColor = FLinearColor::MakeFromHSV8((uint8)UniqueCell.Key, 255, 255);
+			}
+			
+			FSlateDrawElement::MakeBox(
+				OutDrawElements,
+				++LayerId,
+				CellGeometry,
+				&CellBrush,
+				ESlateDrawEffect::None,
+				CellColor
+			);
+		}
+	}
+
+	// Draw selected cells
+	if (SelectedCells.Num())
+	{
+		TMap<uint32, FCellDesc2D> UniqueCells2D;
+		WorldPartition->EditorHash->ForEachIntersectingCell(VisibleGridRectWorld, [&](UWorldPartitionEditorCell* Cell)
+		{
+			if (SelectedCells.Contains(Cell))
 			{
-				FPaintGeometry CellGeometry = AllottedGeometry.ToPaintGeometry(
-					WorldToScreen.TransformPoint(Cell.Bounds.Min),
-					WorldToScreen.TransformPoint(Cell.Bounds.Max) - WorldToScreen.TransformPoint(Cell.Bounds.Min)
-				);
-
-				FLinearColor CellColor;
-				if(Cell.bEmpty)
-				{
-					CellColor = FLinearColor(0, 0, 0, 1.0f);
-				}
-				else if (Cell.bSelected)
-				{
-					CellColor = FLinearColor(1, 1, 1, 0.25f);
-				}
-				else // !Cell.bLoaded
-				{
-					 CellColor = FLinearColor(0, 0, 0, 0.5f);
-				}
-
-				FSlateColorBrush CellBrush(FLinearColor::White);
-
-				FSlateDrawElement::MakeBox(
-					OutDrawElements,
-					++LayerId,
-					CellGeometry,
-					&CellBrush,
-					ESlateDrawEffect::None,
-					CellColor
-				);
+				UWorldPartitionEditorSpatialHash::FCellCoord CellCoord = EditorSpatialHash->GetCellCoords(Cell->Bounds.GetCenter(), 0);
+
+				FHashBuilder HashBuilder;
+				HashBuilder << CellCoord.X << CellCoord.Y;
+				uint32 CellHash2D = HashBuilder.GetHash();
+
+				FCellDesc2D& CellDesc2D = UniqueCells2D.Add(CellHash2D);
+
+				CellDesc2D.Bounds = FBox2D(FVector2D(Cell->Bounds.Min), FVector2D(Cell->Bounds.Max));
 			}
+		});
+
+		for(auto& UniqueCell: UniqueCells2D)
+		{
+			const FCellDesc2D& Cell = UniqueCell.Value;
+
+			FPaintGeometry CellGeometry = AllottedGeometry.ToPaintGeometry(
+				WorldToScreen.TransformPoint(Cell.Bounds.Min),
+				WorldToScreen.TransformPoint(Cell.Bounds.Max) - WorldToScreen.TransformPoint(Cell.Bounds.Min)
+			);
+
+			FSlateColorBrush CellBrush(FLinearColor::White);
+			FLinearColor CellColor(1, 1, 1, 0.25f);
+
+			FSlateDrawElement::MakeBox(
+				OutDrawElements,
+				++LayerId,
+				CellGeometry,
+				&CellBrush,
+				ESlateDrawEffect::None,
+				CellColor
+			);
 		}
 	}
 
diff --git a/Engine/Source/Runtime/Engine/Private/WorldPartition/WorldPartition.cpp b/Engine/Source/Runtime/Engine/Private/WorldPartition/WorldPartition.cpp
--- a/Engine/Source/Runtime/Engine/Private/WorldPartition/WorldPartition.cpp
+++ b/Engine/Source/Runtime/Engine/Private/WorldPartition/WorldPartition.cpp
@@ -903,7 +903,19 @@ void UWorldPartition::UpdateLoadingEditorCell(UWorldPartitionEditorCell* Cell, b
 		}
 	}
 
-	Cell->bLoaded = bShouldBeLoaded;
+	if (Cell->bLoaded != bShouldBeLoaded)
+	{
+		Cell->bLoaded = bShouldBeLoaded;
+
+		if (Cell->bLoaded)
+		{
+			EditorHash->OnCellLoaded(Cell);
+		}
+		else
+		{
+			EditorHash->OnCellUnloaded(Cell);
+		}
+	}
 }
 
 void UWorldPartition::CreateLayers(const TSet<FName>& LayerNames)
diff --git a/Engine/Source/Runtime/Engine/Private/WorldPartition/WorldPartitionEditorCell.cpp b/Engine/Source/Runtime/Engine/Private/WorldPartition/WorldPartitionEditorCell.cpp
--- a/Engine/Source/Runtime/Engine/Private/WorldPartition/WorldPartitionEditorCell.cpp
+++ b/Engine/Source/Runtime/Engine/Private/WorldPartition/WorldPartitionEditorCell.cpp
@@ -72,7 +72,8 @@ void UWorldPartitionEditorCell::PostEditUndo()
 
 	if (bWasLoaded != bLoaded)
 	{
-		GetWorldPartition()->UpdateLoadingEditorCell(this, bLoaded);
+		bLoaded = bWasLoaded;
+		GetWorldPartition()->UpdateLoadingEditorCell(this, !bLoaded);
 	}
 }
 #endif
diff --git a/Engine/Source/Runtime/Engine/Private/WorldPartition/WorldPartitionEditorSpatialHash.cpp b/Engine/Source/Runtime/Engine/Private/WorldPartition/WorldPartitionEditorSpatialHash.cpp
--- a/Engine/Source/Runtime/Engine/Private/WorldPartition/WorldPartitionEditorSpatialHash.cpp
+++ b/Engine/Source/Runtime/Engine/Private/WorldPartition/WorldPartitionEditorSpatialHash.cpp
@@ -8,6 +8,10 @@
 
 #define LOCTEXT_NAMESPACE "WorldPartitionEditor"
 
+#if WITH_EDITOR
+TAutoConsoleVariable<bool> CVarEnableSpatialHashValidation(TEXT("wp.Editor.EnableSpatialHashValidation"), false, TEXT("Whether to enable World Partition editor spatial hash validation"), ECVF_Default);
+#endif
+
 UWorldPartitionEditorSpatialHash::UWorldPartitionEditorSpatialHash(const FObjectInitializer& ObjectInitializer)
 	: Super(ObjectInitializer)
 #if WITH_EDITOR
@@ -60,6 +64,37 @@ void UWorldPartitionEditorSpatialHash::Tick(float DeltaSeconds)
 		Bounds = NewBounds;
 		bBoundsDirty = false;
 	}
+
+	if (CVarEnableSpatialHashValidation.GetValueOnAnyThread())
+	{	
+		const int32 CurrentLevel = GetLevelForBox(Bounds);
+		ForEachIntersectingCells(Bounds, 0, [&](const FCellCoord& CellCoord)
+		{
+			UWorldPartitionEditorCell* EditorCell = nullptr;
+			if (UWorldPartitionEditorCell** EditorCellPtr = HashCells.Find(CellCoord))
+			{
+				if ((*EditorCellPtr)->bLoaded)
+				{
+					FCellCoord LevelCellCoord = CellCoord;
+					while (LevelCellCoord.Level < CurrentLevel)
+					{
+						const uint32 ChildIndex = LevelCellCoord.GetChildIndex();
+						LevelCellCoord = LevelCellCoord.GetParentCellCoord();
+						const FCellNode& CellNode = HashNodes.FindChecked(LevelCellCoord);
+						check(CellNode.HasChildLoadedNode(ChildIndex));
+					}
+				}
+				else
+				{
+					FCellCoord LevelCellCoord = CellCoord;
+					const uint32 ChildIndex = LevelCellCoord.GetChildIndex();
+					LevelCellCoord = LevelCellCoord.GetParentCellCoord();
+					const FCellNode& CellNode = HashNodes.FindChecked(LevelCellCoord);
+					check(!CellNode.HasChildLoadedNode(ChildIndex));
+				}
+			}
+		});
+	}
 }
 
 // In the editor, actors are always using their bounds for grid placement, which makes more sense from a user standpoint.
@@ -110,10 +145,22 @@ void UWorldPartitionEditorSpatialHash::HashActor(FWorldPartitionActorDesc* InAct
 		auto UpdateHigherLevels = [this](const FCellCoord& CellCoord, int32 EndLevel)
 		{
 			FCellCoord LevelCellCoord = CellCoord;
-			for (int32 Level=CellCoord.Level+1; Level<=EndLevel; Level++)
+			while (LevelCellCoord.Level < EndLevel)
 			{
 				const uint32 ChildIndex = LevelCellCoord.GetChildIndex();
 
+				bool bChildHasLoadedNodes;
+				if (LevelCellCoord.Level)
+				{
+					FCellNode& ChildCellNode = HashNodes.FindChecked(LevelCellCoord);
+					bChildHasLoadedNodes = ChildCellNode.HasChildLoadedNodes();
+				}
+				else
+				{
+					UWorldPartitionEditorCell* EditorCell = HashCells.FindChecked(LevelCellCoord);
+					bChildHasLoadedNodes = EditorCell->bLoaded;
+				}
+
 				LevelCellCoord = LevelCellCoord.GetParentCellCoord();
 
 				FCellNode& CellNode = HashNodes.FindOrAdd(LevelCellCoord);
@@ -124,6 +171,11 @@ void UWorldPartitionEditorSpatialHash::HashActor(FWorldPartitionActorDesc* InAct
 				}
 
 				CellNode.AddChildNode(ChildIndex);
+				
+				if (bChildHasLoadedNodes)
+				{
+					CellNode.AddChildLoadedNode(ChildIndex);
+				}
 			}
 		};
 
@@ -192,8 +244,9 @@ void UWorldPartitionEditorSpatialHash::UnhashActor(FWorldPartitionActorDesc* InA
 				verify(Cells.Remove(EditorCell));
 				verify(HashCells.Remove(CellCoord));
 
+				bool bClearParentLoadedChild = true;
 				FCellCoord LevelCellCoord = CellCoord;
-				for (int32 Level=1; Level<=CurrentLevel; Level++)
+				while (LevelCellCoord.Level < CurrentLevel)
 				{
 					const uint32 ChildIndex = LevelCellCoord.GetChildIndex();
 
@@ -201,7 +254,13 @@ void UWorldPartitionEditorSpatialHash::UnhashActor(FWorldPartitionActorDesc* InA
 
 					FCellNode& CellNode = HashNodes.FindChecked(LevelCellCoord);
 
-					CellNode.RemoveChildNode(ChildIndex);
+					if (bClearParentLoadedChild && CellNode.HasChildLoadedNode(ChildIndex))
+					{
+						CellNode.RemoveChildLoadedNode(ChildIndex);
+						bClearParentLoadedChild = !CellNode.HasChildLoadedNodes();
+					}
+
+
 
 					if (CellNode.HasChildNodes())
 					{
@@ -217,6 +276,63 @@ void UWorldPartitionEditorSpatialHash::UnhashActor(FWorldPartitionActorDesc* InA
 	}
 }
 
+void UWorldPartitionEditorSpatialHash::OnCellLoaded(const UWorldPartitionEditorCell* Cell)
+{
+	if (Cell == AlwaysLoadedCell)
+	{
+		return;
+	}
+
+	check(Cell->bLoaded);
+
+	const FCellCoord CellCoord = GetCellCoords(Cell->Bounds.GetCenter(), 0);
+	const int32 CurrentLevel = GetLevelForBox(Bounds);
+
+	FCellCoord LevelCellCoord = CellCoord;
+	while (LevelCellCoord.Level < CurrentLevel)
+	{
+		const uint32 ChildIndex = LevelCellCoord.GetChildIndex();
+
+		LevelCellCoord = LevelCellCoord.GetParentCellCoord();
+
+		FCellNode& CellNode = HashNodes.FindChecked(LevelCellCoord);
+
+		check(!CellNode.HasChildLoadedNode(ChildIndex) || (CellCoord.Level != 1));
+
+		if (CellNode.HasChildLoadedNode(ChildIndex))
+		{
+			break;
+		}
+
+		CellNode.AddChildLoadedNode(ChildIndex);
+	}
+}
+
+void UWorldPartitionEditorSpatialHash::OnCellUnloaded(const UWorldPartitionEditorCell* Cell)
+{
+	check(!Cell->bLoaded);
+
+	const FCellCoord CellCoord = GetCellCoords(Cell->Bounds.GetCenter(), 0);
+	const int32 CurrentLevel = GetLevelForBox(Bounds);
+
+	FCellCoord LevelCellCoord = CellCoord;
+	while (LevelCellCoord.Level < CurrentLevel)
+	{
+		const uint32 ChildIndex = LevelCellCoord.GetChildIndex();
+
+		LevelCellCoord = LevelCellCoord.GetParentCellCoord();
+
+		FCellNode& CellNode = HashNodes.FindChecked(LevelCellCoord);
+
+		CellNode.RemoveChildLoadedNode(ChildIndex);
+
+		if (CellNode.HasChildLoadedNodes())
+		{
+			break;
+		}
+	}
+}
+
 int32 UWorldPartitionEditorSpatialHash::ForEachIntersectingActor(const FBox& Box, TFunctionRef<void(FWorldPartitionActorDesc*)> InOperation)
 {
 	int32 NumIntersecting = 0;
@@ -309,6 +425,69 @@ int32 UWorldPartitionEditorSpatialHash::ForEachIntersectingCell(const FBox& Box,
 	return NumIntersecting;
 }
 
+int32 UWorldPartitionEditorSpatialHash::ForEachIntersectingUnloadedRegionInner(const FBox& Box, const FCellCoord& CellCoord, TFunctionRef<void(const FCellCoord&)> InOperation)
+{
+	int32 NumIntersecting = 0;
+
+	if (CellCoord.Level)
+	{
+		FCellNode CellNode;
+		if (FCellNode* CellNodePtr = HashNodes.Find(CellCoord))
+		{
+			CellNode = *CellNodePtr;
+
+			check(CellNode.HasChildNodes());
+
+			if (!CellNode.HasChildLoadedNodes())
+			{
+				InOperation(CellCoord);
+			}
+			else
+			{
+				CellNode.ForEachChild([&](uint32 ChildIndex)
+				{
+					const FCellCoord ChildCellCoord = CellCoord.GetChildCellCoord(ChildIndex);
+					const FBox CellBounds = GetCellBounds(ChildCellCoord);
+
+					if (Box.Intersect(CellBounds))
+					{
+						NumIntersecting += ForEachIntersectingUnloadedRegionInner(Box, ChildCellCoord, InOperation);
+					}
+				});
+			}
+		}
+	}
+	else
+	{
+		UWorldPartitionEditorCell* EditorCell = nullptr;
+		if (UWorldPartitionEditorCell** EditorCellPtr = HashCells.Find(CellCoord))
+		{
+			if (!(*EditorCellPtr)->bLoaded)
+			{
+				InOperation(CellCoord);
+				NumIntersecting++;
+			}			
+		}
+	}
+
+	return NumIntersecting;
+}
+
+int32 UWorldPartitionEditorSpatialHash::ForEachIntersectingUnloadedRegion(const FBox& Box, TFunctionRef<void(const FCellCoord&)> InOperation)
+{
+	int32 NumIntersecting = 0;
+
+	const FBox SearchBox = Box.Overlap(Bounds);
+	const int32 SearchLevel = GetLevelForBox(SearchBox);
+
+	ForEachIntersectingCells(SearchBox, SearchLevel, [&](const FCellCoord& CellCoord)
+	{
+		NumIntersecting += ForEachIntersectingUnloadedRegionInner(Box, CellCoord, InOperation);
+	});
+
+	return NumIntersecting;
+}
+
 int32 UWorldPartitionEditorSpatialHash::ForEachCell(TFunctionRef<void(UWorldPartitionEditorCell*)> InOperation)
 {
 	for (UWorldPartitionEditorCell* Cell: Cells)
diff --git a/Engine/Source/Runtime/Engine/Public/WorldPartition/WorldPartitionEditorHash.h b/Engine/Source/Runtime/Engine/Public/WorldPartition/WorldPartitionEditorHash.h
--- a/Engine/Source/Runtime/Engine/Public/WorldPartition/WorldPartitionEditorHash.h
+++ b/Engine/Source/Runtime/Engine/Public/WorldPartition/WorldPartitionEditorHash.h
@@ -20,8 +20,13 @@ public:
 	virtual void SetDefaultValues() PURE_VIRTUAL(UWorldPartitionEditorHash::SetDefaultValues, return;);
 	virtual FName GetWorldPartitionEditorName() PURE_VIRTUAL(UWorldPartitionEditorHash::GetWorldPartitionEditorName, return FName(NAME_None););
 	virtual void Tick(float DeltaSeconds) PURE_VIRTUAL(UWorldPartitionEditorHash::Tick, return;);
+
 	virtual void HashActor(FWorldPartitionActorDesc* InActorDesc) PURE_VIRTUAL(UWorldPartitionEditorHash::HashActor, ;);
 	virtual void UnhashActor(FWorldPartitionActorDesc* InActorDesc) PURE_VIRTUAL(UWorldPartitionEditorHash::UnhashActor, ;);
+
+	virtual void OnCellLoaded(const UWorldPartitionEditorCell* Cell) PURE_VIRTUAL(UWorldPartitionEditorHash::OnCellLoaded, ;);
+	virtual void OnCellUnloaded(const UWorldPartitionEditorCell* Cell) PURE_VIRTUAL(UWorldPartitionEditorHash::OnCellLoaded, ;);
+
 	virtual int32 ForEachIntersectingActor(const FBox& Box, TFunctionRef<void(FWorldPartitionActorDesc*)> InOperation) PURE_VIRTUAL(UWorldPartitionEditorHash::ForEachIntForEachIntersectingActorersectingCell, return 0;);
 	virtual int32 ForEachIntersectingCell(const FBox& Box, TFunctionRef<void(UWorldPartitionEditorCell*)> InOperation) PURE_VIRTUAL(UWorldPartitionEditorHash::ForEachIntersectingCell, return 0;);
 	virtual int32 ForEachCell(TFunctionRef<void(UWorldPartitionEditorCell*)> InOperation) PURE_VIRTUAL(UWorldPartitionEditorHash::ForEachCell, return 0;);
diff --git a/Engine/Source/Runtime/Engine/Public/WorldPartition/WorldPartitionEditorSpatialHash.h b/Engine/Source/Runtime/Engine/Public/WorldPartition/WorldPartitionEditorSpatialHash.h
--- a/Engine/Source/Runtime/Engine/Public/WorldPartition/WorldPartitionEditorSpatialHash.h
+++ b/Engine/Source/Runtime/Engine/Public/WorldPartition/WorldPartitionEditorSpatialHash.h
@@ -125,6 +125,7 @@ class ENGINE_API UWorldPartitionEditorSpatialHash : public UWorldPartitionEditor
 	{
 		FCellNode()
 			: ChildNodesMask(0)
+			, ChildNodesLoadedMask(0)
 		{}
 
 		inline bool HasChildNodes() const
@@ -169,26 +170,66 @@ class ENGINE_API UWorldPartitionEditorSpatialHash : public UWorldPartitionEditor
 			}
 		}
 
+		inline bool HasChildLoadedNodes() const
+		{
+			return !!ChildNodesLoadedMask;
+		}
+
+		inline bool HasChildLoadedNode(uint32 ChildIndex) const
+		{
+			check(ChildIndex < 8);
+			return !!(ChildNodesLoadedMask & (1 << ChildIndex));
+		}
+
+		inline void AddChildLoadedNode(uint32 ChildIndex)
+		{
+			check(ChildIndex < 8);
+			uint32 ChildMask = 1 << ChildIndex;
+			check(ChildNodesMask & ChildMask);
+			check(!(ChildNodesLoadedMask & ChildMask));
+			ChildNodesLoadedMask |= ChildMask;
+		}
+
+		inline void RemoveChildLoadedNode(uint32 ChildIndex)
+		{
+			check(ChildIndex < 8);
+			uint32 ChildMask = 1 << ChildIndex;
+			check(ChildNodesMask & ChildMask);
+			check(ChildNodesLoadedMask & ChildMask);
+			ChildNodesLoadedMask &= ~ChildMask;
+		}
+
 		uint8 ChildNodesMask;
+		uint8 ChildNodesLoadedMask;
 	};
 
 public:
 	virtual ~UWorldPartitionEditorSpatialHash() {}
 
+	// UWorldPartitionEditorHash interface begin
 	virtual void Initialize() override;
 	virtual void SetDefaultValues() override;
 	virtual FName GetWorldPartitionEditorName() override;
 	virtual void Tick(float DeltaSeconds) override;
+
 	virtual void HashActor(FWorldPartitionActorDesc* InActorDesc) override;
 	virtual void UnhashActor(FWorldPartitionActorDesc* InActorDesc) override;
+
+	virtual void OnCellLoaded(const UWorldPartitionEditorCell* Cell) override;
+	virtual void OnCellUnloaded(const UWorldPartitionEditorCell* Cell) override;
+
 	virtual int32 ForEachIntersectingActor(const FBox& Box, TFunctionRef<void(FWorldPartitionActorDesc*)> InOperation) override;
 	virtual int32 ForEachIntersectingCell(const FBox& Box, TFunctionRef<void(UWorldPartitionEditorCell*)> InOperation) override;
 	virtual int32 ForEachCell(TFunctionRef<void(UWorldPartitionEditorCell*)> InOperation) override;
 	virtual UWorldPartitionEditorCell* GetAlwaysLoadedCell() override;
+	// UWorldPartitionEditorHash interface end
 #endif
 
 #if WITH_EDITORONLY_DATA
 private:
+	int32 ForEachIntersectingUnloadedRegion(const FBox& Box, TFunctionRef<void(const FCellCoord&)> InOperation);
+	int32 ForEachIntersectingUnloadedRegionInner(const FBox& Box, const FCellCoord& CellCoord, TFunctionRef<void(const FCellCoord&)> InOperation);
+
 	FBox GetActorBounds(FWorldPartitionActorDesc* InActorDesc) const;
 	bool IsActorAlwaysLoaded(FWorldPartitionActorDesc* InActorDesc) const;
 	int32 ForEachIntersectingCellInner(const FBox& Box, const FCellCoord& CellCoord, TFunctionRef<void(UWorldPartitionEditorCell*)> InOperation);
