Fix copy pasta for BP-exposed arg for AudioBus API

diff --git a/Engine/Source/Runtime/Engine/Classes/Components/AudioComponent.h b/Engine/Source/Runtime/Engine/Classes/Components/AudioComponent.h
--- a/Engine/Source/Runtime/Engine/Classes/Components/AudioComponent.h
+++ b/Engine/Source/Runtime/Engine/Classes/Components/AudioComponent.h
@@ -538,7 +538,7 @@ public:
 	/** Sets how much audio the sound should send to the given Audio Bus (POST Source Effects).
 		if the Audio Bus Send doesn't already exist, it will be added to the overrides on the active sound */
 	UFUNCTION(BlueprintCallable, Category = "Audio|Components|Audio")
-	void SetAudioBusSendPostEffect(UAudioBus* AudioBus, float SourceBusSendLevel);
+	void SetAudioBusSendPostEffect(UAudioBus* AudioBus, float AudioBusSendLevel);
 
 	/** Sets whether or not the low pass filter is enabled on the audio component. */
 	UFUNCTION(BlueprintCallable, Category = "Audio|Components|Audio")
