Minor member spelling error fix-up

diff --git a/Engine/Source/Runtime/Engine/Classes/Components/AudioComponent.h b/Engine/Source/Runtime/Engine/Classes/Components/AudioComponent.h
--- a/Engine/Source/Runtime/Engine/Classes/Components/AudioComponent.h
+++ b/Engine/Source/Runtime/Engine/Classes/Components/AudioComponent.h
@@ -597,7 +597,7 @@ private:
 	/** Whether or not the sound is audible. */
 	bool IsInAudibleRange(float* OutMaxDistance) const;
 
-	void SetBusSendffectInternal(USoundSourceBus* InSourceBus, UAudioBus* InAudioBus, float SendLevel, EBusSendType InBusSendType);
+	void SetBusSendEffectInternal(USoundSourceBus* InSourceBus, UAudioBus* InAudioBus, float SendLevel, EBusSendType InBusSendType);
 
 	void BroadcastPlayState();
 
diff --git a/Engine/Source/Runtime/Engine/Private/Components/AudioComponent.cpp b/Engine/Source/Runtime/Engine/Private/Components/AudioComponent.cpp
--- a/Engine/Source/Runtime/Engine/Private/Components/AudioComponent.cpp
+++ b/Engine/Source/Runtime/Engine/Private/Components/AudioComponent.cpp
@@ -1375,7 +1375,7 @@ void UAudioComponent::SetSubmixSend(USoundSubmixBase* Submix, float SendLevel)
 	}
 }
 
-void UAudioComponent::SetBusSendffectInternal(USoundSourceBus* InSourceBus, UAudioBus* InAudioBus, float SendLevel, EBusSendType InBusSendType)
+void UAudioComponent::SetBusSendEffectInternal(USoundSourceBus* InSourceBus, UAudioBus* InAudioBus, float SendLevel, EBusSendType InBusSendType)
 {
 	if (FAudioDevice* AudioDevice = GetAudioDevice())
 	{
@@ -1398,22 +1398,22 @@ void UAudioComponent::SetBusSendffectInternal(USoundSourceBus* InSourceBus, UAud
 
 void UAudioComponent::SetSourceBusSendPreEffect(USoundSourceBus* SoundSourceBus, float SourceBusSendLevel)
 {
-	SetBusSendffectInternal(SoundSourceBus, nullptr, SourceBusSendLevel, EBusSendType::PreEffect);
+	SetBusSendEffectInternal(SoundSourceBus, nullptr, SourceBusSendLevel, EBusSendType::PreEffect);
 }
 
-void UAudioComponent::SetSourceBusSendPostEffect(USoundSourceBus * SoundSourceBus, float SourceBusSendLevel)
+void UAudioComponent::SetSourceBusSendPostEffect(USoundSourceBus* SoundSourceBus, float SourceBusSendLevel)
 {
-	SetBusSendffectInternal(SoundSourceBus, nullptr, SourceBusSendLevel, EBusSendType::PostEffect);
+	SetBusSendEffectInternal(SoundSourceBus, nullptr, SourceBusSendLevel, EBusSendType::PostEffect);
 }
 
 void UAudioComponent::SetAudioBusSendPreEffect(UAudioBus* AudioBus, float AudioBusSendLevel)
 {
-	SetBusSendffectInternal(nullptr, AudioBus, AudioBusSendLevel, EBusSendType::PreEffect);
+	SetBusSendEffectInternal(nullptr, AudioBus, AudioBusSendLevel, EBusSendType::PreEffect);
 }
 
 void UAudioComponent::SetAudioBusSendPostEffect(UAudioBus* AudioBus, float AudioBusSendLevel)
 {
-	SetBusSendffectInternal(nullptr, AudioBus, AudioBusSendLevel, EBusSendType::PostEffect);
+	SetBusSendEffectInternal(nullptr, AudioBus, AudioBusSendLevel, EBusSendType::PostEffect);
 }
 
 void UAudioComponent::SetLowPassFilterEnabled(bool InLowPassFilterEnabled)
