World Partition - Added persistent actor clusters info which is now built incrementally as actors are added/removed

diff --git a/Engine/Source/Editor/WorldPartitionEditor/Private/WorldPartition/SWorldPartitionEditorGrid2D.cpp b/Engine/Source/Editor/WorldPartitionEditor/Private/WorldPartition/SWorldPartitionEditorGrid2D.cpp
--- a/Engine/Source/Editor/WorldPartitionEditor/Private/WorldPartition/SWorldPartitionEditorGrid2D.cpp
+++ b/Engine/Source/Editor/WorldPartitionEditor/Private/WorldPartition/SWorldPartitionEditorGrid2D.cpp
@@ -397,38 +397,24 @@ uint32 SWorldPartitionEditorGrid2D::PaintActors(const FGeometry& AllottedGeometr
 
 	if (bShowActorClusters)
 	{
-		int32 AllActorsHash = 0;
-		for (const auto& Pair : WorldPartition->Actors)
-		{
-			const FWorldPartitionActorDesc* ActorDesc = Pair.Value.Get();
-			AllActorsHash ^= ActorDesc->GetHash();
-		}
-	
-		static uint32 AllActorsCurrentHash = 0;
-		if (AllActorsHash != AllActorsCurrentHash)
-		{
-			ActorClusters = WorldPartition->CreateActorClusters();
-			AllActorsCurrentHash = AllActorsHash;
-		}
-
-		for (const UWorldPartition::FActorCluster& ActorCluster : ActorClusters)
+		for (const UWorldPartition::FActorCluster* ActorCluster : WorldPartition->GetActorClusters())
 		{
 			if (bShowActors)
 			{
-				ActorClustersToShow.Add(&ActorCluster);
+				ActorClustersToShow.Add(ActorCluster);
 			}
 			else
 			{
 				bool bShowClusterActors = false;
 
-				for (const FGuid& ActorGuid: ActorCluster.Actors)
+				for (const FGuid& ActorGuid: ActorCluster->Actors)
 				{
 					FWorldPartitionActorDesc* ActorDesc = WorldPartition->GetActorDesc(ActorGuid);
 					check(ActorDesc);
 
 					if (ActorDesc->Tag == FWorldPartitionActorDesc::GlobalTag)
 					{
-						ActorClustersToShow.Add(&ActorCluster);
+						ActorClustersToShow.Add(ActorCluster);
 						bShowClusterActors = true;
 						break;
 					}
@@ -436,7 +422,7 @@ uint32 SWorldPartitionEditorGrid2D::PaintActors(const FGeometry& AllottedGeometr
 
 				if (bShowClusterActors)
 				{
-					for (const FGuid& ActorGuid: ActorCluster.Actors)
+					for (const FGuid& ActorGuid: ActorCluster->Actors)
 					{
 						FWorldPartitionActorDesc* ActorDesc = WorldPartition->GetActorDesc(ActorGuid);
 						check(ActorDesc);
@@ -471,10 +457,8 @@ uint32 SWorldPartitionEditorGrid2D::PaintActors(const FGeometry& AllottedGeometr
 		{
 			if ((ActorCluster->Actors.Num() > 1) && (ActorCluster->GridPlacement != EActorGridPlacement::AlwaysLoaded))
 			{
-				FBox ActorBounds = WorldPartition->GetActorClusterBounds(*ActorCluster);
-
 				FVector Origin, Extent;
-				ActorBounds.GetCenterAndExtents(Origin, Extent);
+				ActorCluster->Bounds.GetCenterAndExtents(Origin, Extent);
 
 				FVector2D TopLeftW = FVector2D(Origin - Extent);
 				FVector2D BottomRightW = FVector2D(Origin + Extent);
diff --git a/Engine/Source/Runtime/Engine/Private/WorldPartition/WorldPartition.cpp b/Engine/Source/Runtime/Engine/Private/WorldPartition/WorldPartition.cpp
--- a/Engine/Source/Runtime/Engine/Private/WorldPartition/WorldPartition.cpp
+++ b/Engine/Source/Runtime/Engine/Private/WorldPartition/WorldPartition.cpp
@@ -302,7 +302,7 @@ void UWorldPartition::Initialize(UWorld* InWorld, const FTransform& InTransform)
 		{
 			for (const auto& Pair : Actors)
 			{
-				HashActorDesc(Pair.Value.Get());
+				AddToPartition(Pair.Value.Get());
 			}
 
 			CreateLayers(AllLayersNames);
@@ -388,6 +388,10 @@ void UWorldPartition::Uninitialize()
 		}
 
 		EditorHash = nullptr;
+
+		for (FActorCluster* ActorCluster : ActorClustersSet) { delete ActorCluster; }
+		ActorClustersSet.Empty();
+		ActorToActorCluster.Empty();
 #endif		
 
 		if (UWorldPartitionSubsystem* WorldPartitionSubsystem = World->GetSubsystem<UWorldPartitionSubsystem>())
@@ -498,9 +502,9 @@ void UWorldPartition::UpdateActorDesc(AActor* InActor)
 			
 			if (NewDesc->GetHash() != ActorDescPtr->Get()->GetHash())
 			{
-				UnhashActorDesc(ActorDescPtr->Get());
+				RemoveFromPartition(ActorDescPtr->Get(), false);
 				*ActorDescPtr = MoveTemp(NewDesc);
-				HashActorDesc(ActorDescPtr->Get());
+				AddToPartition(ActorDescPtr->Get());
 			}
 		}
 	}
@@ -584,14 +588,14 @@ FWorldPartitionActorDesc* UWorldPartition::GetActorDesc(const FGuid& Guid)
 	return ActorDescObj != nullptr ? ActorDescObj->Get() : nullptr;
 }
 
-void UWorldPartition::CreateActorCluster(const FGuid& ActorGuid, const FWorldPartitionActorDesc* ActorDesc, TMap<FGuid, UWorldPartition::FActorCluster*>& ActorToActorCluster, TSet<UWorldPartition::FActorCluster*>& ActorClustersSet) const
+void UWorldPartition::AddToClusters(const FWorldPartitionActorDesc* ActorDesc)
 {
-	FActorCluster* ActorCluster = ActorToActorCluster.FindRef(ActorGuid);
+	FActorCluster* ActorCluster = ActorToActorCluster.FindRef(ActorDesc->GetGuid());
 	if (!ActorCluster)
 	{
 		ActorCluster = new FActorCluster(ActorDesc);
 		ActorClustersSet.Add(ActorCluster);
-		ActorToActorCluster.Add(ActorGuid, ActorCluster);
+		ActorToActorCluster.Add(ActorDesc->GetGuid(), ActorCluster);
 	}
 
 	// Don't include references from editor-only actors
@@ -632,46 +636,32 @@ void UWorldPartition::CreateActorCluster(const FGuid& ActorGuid, const FWorldPar
 	}
 }
 
-TArray<UWorldPartition::FActorCluster> UWorldPartition::CreateActorClusters() const
+void UWorldPartition::RemoveFromClusters(const FWorldPartitionActorDesc* ActorDesc)
 {
-	TMap<FGuid, FActorCluster*> ActorToActorCluster;
-	TSet<FActorCluster*> ActorClustersSet;
+	FActorCluster* ActorCluster = ActorToActorCluster.FindChecked(ActorDesc->GetGuid());
 
-	for (const auto& Pair : Actors)
+	// Break up this cluster and reinsert all actors
+	ActorClustersSet.Remove(ActorCluster);
+
+	for (const FGuid& Guid : ActorCluster->Actors)
 	{
-		const FGuid& ActorGuid = Pair.Key;
-		const FWorldPartitionActorDesc* ActorDesc = Pair.Value.Get();
-		CreateActorCluster(ActorGuid, ActorDesc, ActorToActorCluster, ActorClustersSet);
+		ActorToActorCluster[Guid] = nullptr;
 	}
 
-	TArray<FActorCluster> ActorClusters;
-	ActorClusters.Reserve(ActorClustersSet.Num());
-	Algo::Transform(ActorClustersSet, ActorClusters, [](FActorCluster* ActorCluster) { return MoveTemp(*ActorCluster); });
-	for (FActorCluster* ActorCluster : ActorClustersSet) { delete ActorCluster; }
-	return ActorClusters;
-}
-
-TArray<UWorldPartition::FActorCluster> UWorldPartition::CreateActorClusters(const TSet<FGuid>& InActors) const
-{
-	TMap<FGuid, FActorCluster*> ActorToActorCluster;
-	TSet<FActorCluster*> ActorClustersSet;
+	ActorToActorCluster.Remove(ActorDesc->GetGuid());
 
-	for (const FGuid& ActorGuid: InActors)
+	for (const FGuid& Guid : ActorCluster->Actors)
 	{
-		const FWorldPartitionActorDesc* ActorDesc = GetActorDesc(ActorGuid);
-		CreateActorCluster(ActorGuid, ActorDesc, ActorToActorCluster, ActorClustersSet);
+		AddToClusters(GetActorDesc(Guid));
 	}
 
-	TArray<FActorCluster> ActorClusters;
-	ActorClusters.Reserve(ActorClustersSet.Num());
-	Algo::Transform(ActorClustersSet, ActorClusters, [](FActorCluster* ActorCluster) { return MoveTemp(*ActorCluster); });
-	for (FActorCluster* ActorCluster : ActorClustersSet) { delete ActorCluster; }
-	return ActorClusters;
+	delete ActorCluster;
 }
 
 UWorldPartition::FActorCluster::FActorCluster(const FWorldPartitionActorDesc* ActorDesc)
 	: GridPlacement(ActorDesc->GetGridPlacement())
 	, RuntimeGrid(ActorDesc->GetRuntimeGrid())
+	, Bounds(ActorDesc->GetBounds())
 {
 	check(GridPlacement != EActorGridPlacement::None);
 	Actors.Add(ActorDesc->GetGuid());
@@ -685,6 +675,9 @@ void UWorldPartition::FActorCluster::Add(const UWorldPartition::FActorCluster& A
 	// Merge RuntimeGrid
 	RuntimeGrid = RuntimeGrid == ActorCluster.RuntimeGrid ? RuntimeGrid : NAME_None;
 
+	// Merge Bounds
+	Bounds += ActorCluster.Bounds;
+
 	// Merge GridPlacement
 	// If currently None, will always stay None
 	if (GridPlacement != EActorGridPlacement::None)
@@ -711,15 +704,14 @@ void UWorldPartition::FActorCluster::Add(const UWorldPartition::FActorCluster& A
 	}
 }
 
-FBox UWorldPartition::GetActorClusterBounds(const FActorCluster& ActorCluster) const
+const TSet<UWorldPartition::FActorCluster*>& UWorldPartition::GetActorClusters() const
 {
-	FBox Bounds(ForceInit);
-	for (const FGuid& ActorGuid : ActorCluster.Actors)
-	{
-		const FWorldPartitionActorDesc* ActorDesc = GetActorDesc(ActorGuid);
-		Bounds += ActorDesc->GetBounds();
-	}
-	return Bounds;
+	return ActorClustersSet;
+}
+
+const UWorldPartition::FActorCluster* UWorldPartition::GetClusterForActor(const FGuid& InActorGuid) const
+{
+	return ActorToActorCluster.FindRef(InActorGuid);
 }
 
 void UWorldPartition::RefreshLoadedCells()
@@ -1032,7 +1024,7 @@ void UWorldPartition::OnActorAdded(AActor* InActor)
 			// old one before adding the new one.
 			check(GIsReinstancing);
 
-			UnhashActorDesc(ActorDescPtr->Get());
+			RemoveFromPartition(ActorDescPtr->Get(), false);
 
 			TUniquePtr<FWorldPartitionActorDesc> NewDesc(GetActorDescFactory(InActor)->Create(InActor));
 			*ActorDescPtr = MoveTemp(NewDesc);
@@ -1046,7 +1038,7 @@ void UWorldPartition::OnActorAdded(AActor* InActor)
 			ActorDescPtr = &Actors.Add(InActor->GetActorGuid(), TUniquePtr<FWorldPartitionActorDesc>(NewDesc));
 		}
 
-		HashActorDesc(ActorDescPtr->Get());
+		AddToPartition(ActorDescPtr->Get());
 	}
 }
 
@@ -1071,12 +1063,7 @@ void UWorldPartition::OnActorDeleted(AActor* InActor)
 				// Validate that this actor is already removed from the level
 				verify(InActor->GetLevel()->Actors.Find(InActor) == INDEX_NONE);
 
-				// Unhash this actor from the editor hash
-				UnhashActorDesc(ActorDesc);
-				check(!ActorDesc->GetLoadedRefCount());
-
-				// Remove this actor descriptor
-				Actors.Remove(ActorDesc->GetGuid());
+				RemoveFromPartition(ActorDesc);
 			}
 		}
 	}
@@ -1098,8 +1085,7 @@ void UWorldPartition::OnActorOuterChanged(AActor* InActor, UObject* InOldOuter)
 	{
 		if (FWorldPartitionActorDesc* ActorDesc = GetActorDesc(InActor->GetActorGuid()))
 		{
-			UnhashActorDesc(ActorDesc);
-			Actors.Remove(ActorDesc->GetGuid());
+			RemoveFromPartition(ActorDesc);
 		}
 	}
 }
@@ -1120,6 +1106,28 @@ void UWorldPartition::UnhashActorDesc(FWorldPartitionActorDesc* ActorDesc)
 
 	check(!ActorDesc->GetLoadedRefCount());
 }
+
+void UWorldPartition::AddToPartition(FWorldPartitionActorDesc* ActorDesc)
+{
+	AddToClusters(ActorDesc);
+
+	HashActorDesc(ActorDesc);	
+}
+
+void UWorldPartition::RemoveFromPartition(FWorldPartitionActorDesc* ActorDesc, bool bRemoveDescriptorFromArray)
+{
+	// Unhash this actor from the editor hash
+	UnhashActorDesc(ActorDesc);
+
+	RemoveFromClusters(ActorDesc);
+
+	// Remove this actor descriptor
+	if (bRemoveDescriptorFromArray)
+	{
+		Actors.Remove(ActorDesc->GetGuid());
+	}
+}
+
 #endif
 
 void UWorldPartition::Serialize(FArchive& Ar)
diff --git a/Engine/Source/Runtime/Engine/Private/WorldPartition/WorldPartitionRuntimeSpatialHash.cpp b/Engine/Source/Runtime/Engine/Private/WorldPartition/WorldPartitionRuntimeSpatialHash.cpp
--- a/Engine/Source/Runtime/Engine/Private/WorldPartition/WorldPartitionRuntimeSpatialHash.cpp
+++ b/Engine/Source/Runtime/Engine/Private/WorldPartition/WorldPartitionRuntimeSpatialHash.cpp
@@ -732,7 +732,7 @@ void FSpatialHashStreamingGrid::Draw2D(UCanvas* Canvas, const TArray<FWorldParti
 
 #if WITH_EDITOR
 
-static FSquare2DGridHelper GetPartitionedActors(const UWorldPartition* WorldPartition, const FBox& WorldBounds, const FSpatialHashRuntimeGrid& Grid, const TArray<UWorldPartition::FActorCluster>& GridActors)
+static FSquare2DGridHelper GetPartitionedActors(const UWorldPartition* WorldPartition, const FBox& WorldBounds, const FSpatialHashRuntimeGrid& Grid, const TArray<const UWorldPartition::FActorCluster*>& GridActors)
 {
 	UE_SCOPED_TIMER(TEXT("GetPartitionedActors"), LogWorldPartitionRuntimeSpatialHash, Log);
 
@@ -769,11 +769,11 @@ static FSquare2DGridHelper GetPartitionedActors(const UWorldPartition* WorldPart
 	//	
 	FSquare2DGridHelper PartitionedActors(GridLevelCount, GridOrigin, GridCellSize, GridSize);
 
-	for (const UWorldPartition::FActorCluster& ActorCluster : GridActors)
+	for (const UWorldPartition::FActorCluster* ActorCluster : GridActors)
 	{
-		check(ActorCluster.Actors.Num() > 0);
+		check(ActorCluster->Actors.Num() > 0);
 
-		EActorGridPlacement GridPlacement = ActorCluster.GridPlacement;
+		EActorGridPlacement GridPlacement = ActorCluster->GridPlacement;
 		bool bAlwaysLoadedPromotedCluster = (GridPlacement == EActorGridPlacement::None);
 		bool bAlwaysLoadedPromotedOutOfGrid = false;
 
@@ -789,8 +789,8 @@ static FSquare2DGridHelper GetPartitionedActors(const UWorldPartition* WorldPart
 			case EActorGridPlacement::Location:
 			{
 				FIntVector2 CellCoords;
-				check(ActorCluster.Actors.Num() == 1);
-				const FGuid& ActorGuid = *ActorCluster.Actors.CreateConstIterator();
+				check(ActorCluster->Actors.Num() == 1);
+				const FGuid& ActorGuid = *ActorCluster->Actors.CreateConstIterator();
 				const FWorldPartitionActorDesc& ActorDesc = *WorldPartition->GetActorDesc(ActorGuid);
 				if (PartitionedActors.GetLowestLevel().GetCellCoords(FVector2D(ActorDesc.GetOrigin()), CellCoords))
 				{
@@ -805,18 +805,17 @@ static FSquare2DGridHelper GetPartitionedActors(const UWorldPartition* WorldPart
 			}
 			case EActorGridPlacement::Bounds:
 			{
-				const FBox ActorClusterBounds = WorldPartition->GetActorClusterBounds(ActorCluster);
 				// Find grid level cell that encompasses the actor cluster and put actors in it.
 				bool bFoundCell = false;
 				for (FSquare2DGridHelper::FGridLevel& GridLevel : PartitionedActors.Levels)
 				{
 					int32 IntersectingCellCount = 0;
-					GridLevel.ForEachIntersectingCellsBreakable(ActorClusterBounds, [&IntersectingCellCount](const FIntVector2& Coords) { return ++IntersectingCellCount <= 1; });
+					GridLevel.ForEachIntersectingCellsBreakable(ActorCluster->Bounds, [&IntersectingCellCount](const FIntVector2& Coords) { return ++IntersectingCellCount <= 1; });
 					if (IntersectingCellCount == 1)
 					{
-						GridLevel.ForEachIntersectingCells(ActorClusterBounds, [&GridLevel,&ActorCluster](const FIntVector2& Coords)
+						GridLevel.ForEachIntersectingCells(ActorCluster->Bounds, [&GridLevel,&ActorCluster](const FIntVector2& Coords)
 						{
-							GridLevel.GetCell(Coords).Actors.Append(ActorCluster.Actors);
+							GridLevel.GetCell(Coords).Actors.Append(ActorCluster->Actors);
 						});
 						bFoundCell = true;
 						break;
@@ -836,25 +835,24 @@ static FSquare2DGridHelper GetPartitionedActors(const UWorldPartition* WorldPart
 			
 		if (GridPlacement == EActorGridPlacement::AlwaysLoaded)
 		{
-			PartitionedActors.GetAlwaysLoadedCell().Actors.Append(ActorCluster.Actors);
+			PartitionedActors.GetAlwaysLoadedCell().Actors.Append(ActorCluster->Actors);
 		}
 
 		if (!LogWorldPartitionRuntimeSpatialHash.IsSuppressed(ELogVerbosity::Verbose))
 		{
-			if (ActorCluster.Actors.Num() > 1)
+			if (ActorCluster->Actors.Num() > 1)
 			{
 				static UEnum* ActorGridPlacementEnum = FindObject<UEnum>(ANY_PACKAGE, TEXT("EActorGridPlacement"));
-				const FBox ActorClusterBounds = WorldPartition->GetActorClusterBounds(ActorCluster);
 
 				UE_LOG(LogWorldPartitionRuntimeSpatialHash, Verbose, TEXT("Clustered %d actors (%s%s%s), generated shared BV of [%d x %d] (meters)"), 
-					ActorCluster.Actors.Num(),
+					ActorCluster->Actors.Num(),
 					*ActorGridPlacementEnum->GetNameStringByValue((int64)GridPlacement),
 					bAlwaysLoadedPromotedCluster ? TEXT(":PromotedCluster") : TEXT(""),
 					bAlwaysLoadedPromotedOutOfGrid ? TEXT(":PromotedOutOfGrid") : TEXT(""),
-					(int)(0.01f * ActorClusterBounds.GetSize().X),
-					(int)(0.01f * ActorClusterBounds.GetSize().Y));
+					(int)(0.01f * ActorCluster->Bounds.GetSize().X),
+					(int)(0.01f * ActorCluster->Bounds.GetSize().Y));
 
-				for (const FGuid& ActorGuid : ActorCluster.Actors)
+				for (const FGuid& ActorGuid : ActorCluster->Actors)
 				{
 					const FWorldPartitionActorDesc& Desc = *WorldPartition->GetActorDesc(ActorGuid);
 					UE_LOG(LogWorldPartitionRuntimeSpatialHash, Verbose, TEXT("   - Actor: %s (%s)"), *Desc.GetActorPath().ToString(), *ActorGuid.ToString(EGuidFormats::UniqueObjectGuid));
@@ -966,21 +964,19 @@ bool UWorldPartitionRuntimeSpatialHash::GenerateStreaming(EWorldPartitionStreami
 	}
 
 	// Create actor clusters
-	TArray<UWorldPartition::FActorCluster> ActorClusters = WorldPartition->CreateActorClusters();
-
-	TArray<TArray<UWorldPartition::FActorCluster>> GridActors;
+	TArray<TArray<const UWorldPartition::FActorCluster*>> GridActors;
 	GridActors.InsertDefaulted(0, AllGrids.Num());
 
-	for (UWorldPartition::FActorCluster& ActorCluster : ActorClusters)
+	for (const UWorldPartition::FActorCluster* ActorCluster : WorldPartition->GetActorClusters())
 	{
-		int32* FoundIndex = GridsMapping.Find(ActorCluster.RuntimeGrid);
+		int32* FoundIndex = GridsMapping.Find(ActorCluster->RuntimeGrid);
 		if (!FoundIndex)
 		{
-			UE_LOG(LogWorldPartitionRuntimeSpatialHash, Error, TEXT("Invalid partition grid '%s' referenced by actor cluster"), *ActorCluster.RuntimeGrid.ToString());
+			UE_LOG(LogWorldPartitionRuntimeSpatialHash, Error, TEXT("Invalid partition grid '%s' referenced by actor cluster"), *ActorCluster->RuntimeGrid.ToString());
 		}
 
 		int32 GridIndex = FoundIndex ? *FoundIndex : 0;
-		GridActors[GridIndex].Add(MoveTemp(ActorCluster));
+		GridActors[GridIndex].Add(ActorCluster);
 	}
 	
 	const FBox WorldBounds = WorldPartition->GetWorldBounds();
@@ -1186,22 +1182,19 @@ bool UWorldPartitionRuntimeSpatialHash::GenerateHLOD()
 
 	UWorldPartition* WorldPartition = GetOuterUWorldPartition();
 
-	// Create actor clusters
-	TArray<UWorldPartition::FActorCluster> ActorClusters = WorldPartition->CreateActorClusters();
-
-	TArray<TArray<UWorldPartition::FActorCluster>> GridActors;
+	TArray<TArray<const UWorldPartition::FActorCluster*>> GridActors;
 	GridActors.InsertDefaulted(0, Grids.Num());
 
-	for (UWorldPartition::FActorCluster& ActorCluster : ActorClusters)
+	for (const UWorldPartition::FActorCluster* ActorCluster : WorldPartition->GetActorClusters())
 	{
-		int32* FoundIndex = GridsMapping.Find(ActorCluster.RuntimeGrid);
+		int32* FoundIndex = GridsMapping.Find(ActorCluster->RuntimeGrid);
 		if (!FoundIndex)
 		{
-			UE_LOG(LogWorldPartitionRuntimeSpatialHash, Error, TEXT("Invalid partition grid '%s' referenced by actor cluster"), *ActorCluster.RuntimeGrid.ToString());
+			UE_LOG(LogWorldPartitionRuntimeSpatialHash, Error, TEXT("Invalid partition grid '%s' referenced by actor cluster"), *ActorCluster->RuntimeGrid.ToString());
 		}
 
 		int32 GridIndex = FoundIndex ? *FoundIndex : 0;
-		GridActors[GridIndex].Add(MoveTemp(ActorCluster));
+		GridActors[GridIndex].Add(ActorCluster);
 	}
 
 	const FBox WorldBounds = WorldPartition->GetWorldBounds();
diff --git a/Engine/Source/Runtime/Engine/Public/WorldPartition/WorldPartition.h b/Engine/Source/Runtime/Engine/Public/WorldPartition/WorldPartition.h
--- a/Engine/Source/Runtime/Engine/Public/WorldPartition/WorldPartition.h
+++ b/Engine/Source/Runtime/Engine/Public/WorldPartition/WorldPartition.h
@@ -122,8 +122,6 @@ public:
 	void LoadEditorCells(const TArray<UWorldPartitionEditorCell*>& Cells);
 	void UnloadEditorCells(const TArray<UWorldPartitionEditorCell*>& Cells);
 
-	void HashActorDesc(FWorldPartitionActorDesc* ActorDesc);
-	void UnhashActorDesc(FWorldPartitionActorDesc* ActorDesc);
 	void UpdateActorDesc(AActor* InActor);
 	bool IsPreCooked() const { return bIsPreCooked; }
 	void SetIsPreCooked(bool bInIsPreCooked) { bIsPreCooked = bInIsPreCooked; }
@@ -139,15 +137,14 @@ public:
 		TSet<FGuid>			Actors;
 		EActorGridPlacement	GridPlacement;
 		FName				RuntimeGrid;
+		FBox				Bounds;
 
 		FActorCluster(const FWorldPartitionActorDesc* ActorDesc);
 		void Add(const FActorCluster& ActorCluster);
 	};
 
-	TArray<FActorCluster> CreateActorClusters() const;
-	TArray<FActorCluster> CreateActorClusters(const TSet<FGuid>& InActors) const;
-	FBox GetActorClusterBounds(const FActorCluster& ActorCluster) const;
-
+	const TSet<FActorCluster*>& GetActorClusters() const;
+	const FActorCluster* GetClusterForActor(const FGuid& InActorGuid) const;
 	// public interface-
 #endif
 
@@ -234,9 +231,20 @@ private:
 #if WITH_EDITOR
 	void UpdateLoadingEditorCell(UWorldPartitionEditorCell* Cell, bool bShouldBeLoaded);
 	void CreateLayers(const TSet<FName>& LayerNames);
-	void CreateActorCluster(const FGuid& ActorGuid, const FWorldPartitionActorDesc* ActorDesc, TMap<FGuid, FActorCluster*>& ActorToActorCluster, TSet<FActorCluster*>& ActorClustersSet) const;
 
 	void RefreshLoadedCells();
+
+	void HashActorDesc(FWorldPartitionActorDesc* ActorDesc);
+	void UnhashActorDesc(FWorldPartitionActorDesc* ActorDesc);
+
+	void AddToClusters(const FWorldPartitionActorDesc* ActorDesc);
+	void RemoveFromClusters(const FWorldPartitionActorDesc* ActorDesc);
+
+	void AddToPartition(FWorldPartitionActorDesc* ActorDesc);
+	void RemoveFromPartition(FWorldPartitionActorDesc* ActorDesc, bool bRemoveDescriptorFromArray = true);
+
+	TMap<FGuid, UWorldPartition::FActorCluster*> ActorToActorCluster;
+	TSet<UWorldPartition::FActorCluster*> ActorClustersSet;
 #endif
 
 	UWorldPartitionStreamingPolicy* GetStreamingPolicy() const;
