Put back missing line from last submit.

diff --git a/Engine/Source/Runtime/Engine/Private/WorldPartition/WorldPartitionEditorSpatialHash.cpp b/Engine/Source/Runtime/Engine/Private/WorldPartition/WorldPartitionEditorSpatialHash.cpp
--- a/Engine/Source/Runtime/Engine/Private/WorldPartition/WorldPartitionEditorSpatialHash.cpp
+++ b/Engine/Source/Runtime/Engine/Private/WorldPartition/WorldPartitionEditorSpatialHash.cpp
@@ -260,7 +260,7 @@ void UWorldPartitionEditorSpatialHash::UnhashActor(FWorldPartitionActorDesc* InA
 						bClearParentLoadedChild = !CellNode.HasChildLoadedNodes();
 					}
 
-
+					CellNode.RemoveChildNode(ChildIndex);
 
 					if (CellNode.HasChildNodes())
 					{
