Fixed bug in UWorldPartition::RemoveFromClusters()
- Actor being removed from clusters would be re-added as it was still in ActorCluster->Actors
- Merged a Find() & Remove() to a FindAndRemoveChecked()

diff --git a/Engine/Source/Runtime/Engine/Private/WorldPartition/WorldPartition.cpp b/Engine/Source/Runtime/Engine/Private/WorldPartition/WorldPartition.cpp
--- a/Engine/Source/Runtime/Engine/Private/WorldPartition/WorldPartition.cpp
+++ b/Engine/Source/Runtime/Engine/Private/WorldPartition/WorldPartition.cpp
@@ -638,7 +638,8 @@ void UWorldPartition::AddToClusters(const FWorldPartitionActorDesc* ActorDesc)
 
 void UWorldPartition::RemoveFromClusters(const FWorldPartitionActorDesc* ActorDesc)
 {
-	FActorCluster* ActorCluster = ActorToActorCluster.FindChecked(ActorDesc->GetGuid());
+	FActorCluster* ActorCluster = ActorToActorCluster.FindAndRemoveChecked(ActorDesc->GetGuid());
+	ActorCluster->Actors.Remove(ActorDesc->GetGuid());
 
 	// Break up this cluster and reinsert all actors
 	ActorClustersSet.Remove(ActorCluster);
@@ -648,8 +649,6 @@ void UWorldPartition::RemoveFromClusters(const FWorldPartitionActorDesc* ActorDe
 		ActorToActorCluster[Guid] = nullptr;
 	}
 
-	ActorToActorCluster.Remove(ActorDesc->GetGuid());
-
 	for (const FGuid& Guid : ActorCluster->Actors)
 	{
 		AddToClusters(GetActorDesc(Guid));
