WorldPartition HLOD Rebuild Cmdlet
* Remove "-NoSourceControl" option from world partition commandlets, the standard option "-SCCProvider=None" should now be used instead
* FCommandletPackageHelper - Unload packages before deleting them
* WorldPartitionBuildHLODsCommandlet - Ensure we save the map itself as we currently inject streaming grids in UWorldPartitionRuntimeSpatialHash::HLODGrids
* Rebuilding HLODs will reuse the same actors (& packages)

diff --git a/Engine/Source/Editor/UnrealEd/Classes/Commandlets/CommandletPackageHelper.h b/Engine/Source/Editor/UnrealEd/Classes/Commandlets/CommandletPackageHelper.h
--- a/Engine/Source/Editor/UnrealEd/Classes/Commandlets/CommandletPackageHelper.h
+++ b/Engine/Source/Editor/UnrealEd/Classes/Commandlets/CommandletPackageHelper.h
@@ -5,23 +5,28 @@
 #include "CoreMinimal.h"
 #include "UObject/ObjectMacros.h"
 #include "SourceControlHelpers.h"
+#include "AssetRegistryModule.h"
 
 class UNREALED_API FCommandletPackageHelper
 {
 public:
 	FCommandletPackageHelper();
+	~FCommandletPackageHelper();
 
-	void SetSourceControlEnabled(bool bWithSourceControl);
-	bool UseSourceControl() const { return SourceControlProvider != nullptr; }
-	ISourceControlProvider& GetSourceControlProvider() const { check(UseSourceControl()); return *SourceControlProvider; }
+	bool UseSourceControl() const;
 
 	bool Delete(const FString& PackageName) const;
 	bool Delete(UPackage* Package) const;
+	bool Delete(const TArray<UPackage*>& Packages) const;
+	bool Delete(const TArray<FAssetData>& Assets) const;
 	bool AddToSourceControl(UPackage* Package) const;
 	bool Checkout(UPackage* Package) const;
 	bool Save(UPackage* Package) const;
 
 private:
-	FScopedSourceControl SourceControl;
-	mutable ISourceControlProvider* SourceControlProvider;
+	ISourceControlProvider& GetSourceControlProvider() const;
+
+	bool DeleteInternal(const FString& PackageName) const;
+
+	FScopedSourceControl SourceControl;	
 };
diff --git a/Engine/Source/Editor/UnrealEd/Private/Commandlets/CommandletPackageHelper.cpp b/Engine/Source/Editor/UnrealEd/Private/Commandlets/CommandletPackageHelper.cpp
--- a/Engine/Source/Editor/UnrealEd/Private/Commandlets/CommandletPackageHelper.cpp
+++ b/Engine/Source/Editor/UnrealEd/Private/Commandlets/CommandletPackageHelper.cpp
@@ -9,6 +9,7 @@
 #include "UObject/Package.h"
 #include "HAL/FileManager.h"
 #include "HAL/PlatformFileManager.h"
+#include "PackageTools.h"
 #include "ISourceControlOperation.h"
 #include "ISourceControlModule.h"
 #include "ISourceControlState.h"
@@ -18,13 +19,22 @@
 DEFINE_LOG_CATEGORY_STATIC(LogCommandletPackageHelper, Log, All);
 
 FCommandletPackageHelper::FCommandletPackageHelper()
-	: SourceControlProvider(nullptr)
 {
 }
 
-void FCommandletPackageHelper::SetSourceControlEnabled(bool bWithSourceControl)
+FCommandletPackageHelper::~FCommandletPackageHelper()
 {
-	SourceControlProvider = bWithSourceControl ? &ISourceControlModule::Get().GetProvider() : nullptr;
+	UPackage::WaitForAsyncFileWrites();
+}
+
+bool FCommandletPackageHelper::UseSourceControl() const
+{
+	return ISourceControlModule::Get().GetProvider().IsEnabled();
+}
+
+ISourceControlProvider& FCommandletPackageHelper::GetSourceControlProvider() const
+{ 
+	return ISourceControlModule::Get().GetProvider();
 }
 
 bool FCommandletPackageHelper::Delete(const FString& PackageName) const
@@ -102,7 +112,50 @@ bool FCommandletPackageHelper::Delete(const FString& PackageName) const
 
 bool FCommandletPackageHelper::Delete(UPackage* Package) const
 {
-	return Delete(Package->GetName());
+	TArray<UPackage*> Packages = { Package };
+	return Delete(Packages);
+}
+
+bool FCommandletPackageHelper::Delete(const TArray<FAssetData>& Assets) const
+{
+	bool bAllPackagesDeleted = true;
+	for (const FAssetData& AssetToDelete : Assets)
+	{
+		bAllPackagesDeleted &= Delete(AssetToDelete.PackageName.ToString());
+	}
+
+	return bAllPackagesDeleted;
+}
+
+bool FCommandletPackageHelper::Delete(const TArray<UPackage*>& Packages) const
+{
+	// Store all packages names as we won't be able to retrieve them from the UPackages once they are unloaded
+	TArray<FString> PackagesNames;
+	PackagesNames.Reserve(Packages.Num());
+
+	for (UPackage* Package : Packages)
+	{
+		PackagesNames.Add(Package->GetName());
+
+		// Must clear dirty flag befor unloading
+		Package->SetDirtyFlag(false);
+	}
+
+	// Unload packages so we can delete them
+	FText ErrorMessage;
+	if (!UPackageTools::UnloadPackages(Packages, ErrorMessage))
+	{
+		UE_LOG(LogCommandletPackageHelper, Error, TEXT("Error unloading package: %s"), *ErrorMessage.ToString());
+		return false;
+	}
+
+	bool bAllPackagesDeleted = true;
+	for (const FString& PackageName : PackagesNames)
+	{
+		bAllPackagesDeleted &= Delete(PackageName);
+	}
+
+	return bAllPackagesDeleted;
 }
 
 bool FCommandletPackageHelper::AddToSourceControl(UPackage* Package) const
diff --git a/Engine/Source/Editor/UnrealEd/Private/Commandlets/WorldPartitionBuildHLODsCommandlet.cpp b/Engine/Source/Editor/UnrealEd/Private/Commandlets/WorldPartitionBuildHLODsCommandlet.cpp
--- a/Engine/Source/Editor/UnrealEd/Private/Commandlets/WorldPartitionBuildHLODsCommandlet.cpp
+++ b/Engine/Source/Editor/UnrealEd/Private/Commandlets/WorldPartitionBuildHLODsCommandlet.cpp
@@ -27,9 +27,6 @@ int32 UWorldPartitionBuildHLODsCommandlet::Main(const FString& Params)
 	TArray<FString> Tokens, Switches;
 	ParseCommandLine(*Params, Tokens, Switches);
 
-	bool bNoSourceControl = Switches.Contains(TEXT("NoSourceControl"));
-	PackageHelper.SetSourceControlEnabled(!bNoSourceControl);
-
 	// Retrieve map name from the commandline
 	FString MapName;
 	bool bMapProvided = FParse::Value(*Params, TEXT("Map="), MapName);
@@ -99,41 +96,47 @@ int32 UWorldPartitionBuildHLODsCommandlet::Main(const FString& Params)
 	const FBox LoadBox(FVector(-WORLD_MAX, -WORLD_MAX, -WORLD_MAX), FVector(WORLD_MAX, WORLD_MAX, WORLD_MAX));
 	WorldPartition->LoadEditorCells(LoadBox);
 
-	TSet<UPackage*> PackagesToDelete;
-	TSet<UPackage*> PackagesToSave;
+	TSet<UPackage*> HLODActorPackages;
 
-	// Clear all existing HLOD actors
+	// Gather all packages before any HLOD actor is deleted
 	for (TActorIterator<AWorldPartitionHLOD> ItHLOD(World); ItHLOD; ++ItHLOD)
 	{
-		PackagesToDelete.Add(ItHLOD->GetPackage());
-		World->DestroyActor(*ItHLOD);
+		HLODActorPackages.Add(ItHLOD->GetPackage());
 	}
 
 	// Rebuild HLOD for the whole world
 	WorldPartition->GenerateHLOD();
 	
-	// Required to clear any deleted actors from the level
-	CollectGarbage(RF_Standalone);
-
-	// Gather all packages
+	// Gather all packages again to include newly created HLOD actors
 	for (TActorIterator<AWorldPartitionHLOD> ItHLOD(World); ItHLOD; ++ItHLOD)
 	{
-		UPackage* Package = ItHLOD->GetPackage();
-		if (Package && Package->IsDirty())
+		HLODActorPackages.Add(ItHLOD->GetPackage());
+	}
+
+	TArray<UPackage*> PackagesToSave;
+	TArray<UPackage*> PackagesToDelete;
+
+	for (UPackage* HLODActorPackage : HLODActorPackages)
+	{
+		if (HLODActorPackage && HLODActorPackage->IsDirty())
 		{
-			PackagesToSave.Add(Package);
+			if (UPackage::IsEmptyPackage(HLODActorPackage))
+			{
+				PackagesToDelete.Add(HLODActorPackage);
+			}
+			else
+			{
+				PackagesToSave.Add(HLODActorPackage);
+			}
 		}
 	}
 
 	// Delete packages
 	UE_LOG(LogWorldPartitionBuildHLODsCommandlet, Log, TEXT("Deleting %d packages."), PackagesToDelete.Num());
-	for (UPackage* Package : PackagesToDelete)
+	if (!PackageHelper.Delete(PackagesToDelete))
 	{
-		if (!PackageHelper.Delete(Package->GetName()))
-		{
-			UE_LOG(LogWorldPartitionBuildHLODsCommandlet, Error, TEXT("Error deleting package %s."), *Package->GetName());
-			return 1;
-		}
+		UE_LOG(LogWorldPartitionBuildHLODsCommandlet, Error, TEXT("Error deleting packages."));
+		return 1;
 	}
 
 	// Checkout packages
@@ -169,6 +172,13 @@ int32 UWorldPartitionBuildHLODsCommandlet::Main(const FString& Params)
 		}
 	}
 
+	// Save the map itself
+	if (!PackageHelper.Save(MapPackage))
+	{
+		UE_LOG(LogWorldPartitionBuildHLODsCommandlet, Error, TEXT("Error saving map package %s."), *MapPackage->GetName());
+		return 1;
+	}
+
 	// Cleanup
 	World->RemoveFromRoot();
 	WorldContext.SetCurrentWorld(nullptr);
diff --git a/Engine/Source/Editor/UnrealEd/Private/Commandlets/WorldPartitionConvertCommandlet.cpp b/Engine/Source/Editor/UnrealEd/Private/Commandlets/WorldPartitionConvertCommandlet.cpp
--- a/Engine/Source/Editor/UnrealEd/Private/Commandlets/WorldPartitionConvertCommandlet.cpp
+++ b/Engine/Source/Editor/UnrealEd/Private/Commandlets/WorldPartitionConvertCommandlet.cpp
@@ -668,9 +668,6 @@ int32 UWorldPartitionConvertCommandlet::Main(const FString& Params)
 
 	bConversionSuffix = Switches.Contains(TEXT("ConversionSuffix"));
 
-	bool bNoSourceControl = Switches.Contains(TEXT("NoSourceControl"));
-	PackageHelper.SetSourceControlEnabled(!bNoSourceControl);
-
 	// Load configuration file
 	FString LevelLongPackageName;
 	if (FPackageName::SearchForPackageOnDisk(Tokens[0], nullptr, &LevelLongPackageName))
@@ -718,17 +715,19 @@ int32 UWorldPartitionConvertCommandlet::Main(const FString& Params)
 		FString OldLevelName = Tokens[0] + ConversionSuffix;
 		FActorRegistry::GetLevelActors(*OldLevelName, Assets);
 
-		for (const FAssetData& Asset : Assets)
+		if (!PackageHelper.Delete(Assets))
 		{
-			if (!PackageHelper.Delete(Asset.PackageName.ToString()))
-			{
-				return 1;
-			}
+			UE_LOG(LogWorldPartitionConvertCommandlet, Error, TEXT("Failed to delete external actor package(s)"));
+			return 1;
 		}
 
 		if (FPackageName::SearchForPackageOnDisk(OldLevelName, &OldLevelName))
 		{
-			PackageHelper.Delete(OldLevelName);
+			if (!PackageHelper.Delete(OldLevelName))
+			{
+				UE_LOG(LogWorldPartitionConvertCommandlet, Error, TEXT("Failed to delete previously converted level '%s'"), *Tokens[0]);
+				return 1;
+			}
 		}
 	}
 
diff --git a/Engine/Source/Runtime/Engine/Private/WorldPartition/HLOD/HLODBuilder.cpp b/Engine/Source/Runtime/Engine/Private/WorldPartition/HLOD/HLODBuilder.cpp
--- a/Engine/Source/Runtime/Engine/Private/WorldPartition/HLOD/HLODBuilder.cpp
+++ b/Engine/Source/Runtime/Engine/Private/WorldPartition/HLOD/HLODBuilder.cpp
@@ -26,25 +26,34 @@ class FHLODBuilder
 public:
 	virtual ~FHLODBuilder() {}
 
-	virtual void Build(const TArray<AActor*>& InSubActors) const = 0;
+	virtual void Build(const TArray<AActor*>& InSubActors) = 0;
 
 	typedef TFunction<UPrimitiveComponent*(AWorldPartitionHLOD*)> FCreateComponentFunction;
 		
-	AWorldPartitionHLOD* SpawnHLODActor(const TCHAR* InName, const TArray<UPrimitiveComponent*>& InSubComponents, FCreateComponentFunction InCreateComponentFunc) const
+	void SpawnHLODActor(const TCHAR* InName, const TArray<UPrimitiveComponent*>& InSubComponents, FCreateComponentFunction InCreateComponentFunc)
 	{
-		FActorSpawnParameters SpawnParams;
-		SpawnParams.bDeferConstruction = true;
-		SpawnParams.bCreateActorPackage = true;
-		AWorldPartitionHLOD* HLODActor = World->SpawnActor<AWorldPartitionHLOD>(SpawnParams);
+		FString HLODActorName = FString::Printf(TEXT("%s_%s_%s"), *HLODLayer->GetName(), *CellName.ToString(), InName);
+		FString HLODActorPath = FString::Printf(TEXT("%s.%s"), *World->PersistentLevel->GetPathName(), *HLODActorName);
+
+		AWorldPartitionHLOD* HLODActor = FindObject<AWorldPartitionHLOD>(ANY_PACKAGE, *HLODActorPath);
+		if (!HLODActor)
+		{
+			FActorSpawnParameters SpawnParams;
+			SpawnParams.bDeferConstruction = true;
+			SpawnParams.bCreateActorPackage = true;
+			SpawnParams.Name = *HLODActorName;
+			HLODActor = World->SpawnActor<AWorldPartitionHLOD>(SpawnParams);
+		}
 
 		UPrimitiveComponent* ParentComponent = InCreateComponentFunc(HLODActor);
 
 		if (ParentComponent)
 		{
+			HLODActor->Modify();
 			HLODActor->SetHLODLayer(HLODLayer, iLevel);
-			HLODActor->SetParentPrimitive(ParentComponent);
+			HLODActor->SetHLODPrimitive(ParentComponent);
 			HLODActor->SetChildrenPrimitives(InSubComponents);
-			HLODActor->SetActorLabel(FString::Printf(TEXT("%s_%s_%s"), *HLODLayer->GetName(), *CellName.ToString(), InName));
+			HLODActor->SetActorLabel(HLODActorName);
 			HLODActor->RuntimeGrid = GetHLODLevelSettings().TargetGrid;
 
 			WorldPartition->UpdateActorDesc(HLODActor);
@@ -55,7 +64,10 @@ public:
 			HLODActor = nullptr;
 		}
 
-		return HLODActor;
+		if (HLODActor)
+		{
+			HLODActors.Add(HLODActor);
+		}
 	}
 
 	TArray<UPrimitiveComponent*> GatherPrimitiveComponents(int32 iHLODLevel, const TArray<AActor*> InActors) const
@@ -88,6 +100,8 @@ public:
 	const UHLODLayer*	HLODLayer;
 	int32				iLevel;
 	FName				CellName;
+
+	TArray<AWorldPartitionHLOD*> HLODActors;
 };
 
 
@@ -96,7 +110,7 @@ public:
  */
 class FHLODBuilder_Instancing : public FHLODBuilder
 {
-	virtual void Build(const TArray<AActor*>& InSubActors) const override
+	virtual void Build(const TArray<AActor*>& InSubActors) override
 	{
 		TRACE_CPUPROFILER_EVENT_SCOPE(FHLODBuilder_Instancing::BuildHLOD);
 
@@ -145,7 +159,7 @@ class FHLODBuilder_Instancing : public FHLODBuilder
  */
 class FHLODBuilder_MeshMerge : public FHLODBuilder
 {
-	virtual void Build(const TArray<AActor*>& InSubActors) const override
+	virtual void Build(const TArray<AActor*>& InSubActors) override
 	{
 		TRACE_CPUPROFILER_EVENT_SCOPE(UHLODLayer::BuildHLOD_MeshMerge);
 
@@ -185,7 +199,7 @@ class FHLODBuilder_MeshMerge : public FHLODBuilder
  */
 class FHLODBuilder_MeshSimplify : public FHLODBuilder
 {
-	virtual void Build(const TArray<AActor*>& InSubActors) const override
+	virtual void Build(const TArray<AActor*>& InSubActors) override
 	{
 		TRACE_CPUPROFILER_EVENT_SCOPE(UHLODLayer::BuildHLOD_MeshProxy);
 
@@ -227,10 +241,12 @@ class FHLODBuilder_MeshSimplify : public FHLODBuilder
 /**
  *
  */
-void FHLODBuilderUtilities::BuildHLODs(UWorldPartition* InWorldPartition, FName InCellName, const UHLODLayer* InHLODLayer, const TArray<AActor*>& InSubActors)
+TArray<AWorldPartitionHLOD*> FHLODBuilderUtilities::BuildHLODs(UWorldPartition* InWorldPartition, FName InCellName, const UHLODLayer* InHLODLayer, const TArray<AActor*>& InSubActors)
 {
 	const TArray<FHLODLevelSettings>& LevelsSettings = InHLODLayer->GetLevels();
 
+	TArray<AWorldPartitionHLOD*> HLODActors;
+
 	for (int32 iLevel = 0; iLevel < LevelsSettings.Num(); ++iLevel)
 	{
 		TUniquePtr<FHLODBuilder> HLODBuilder = nullptr;
@@ -263,8 +279,12 @@ void FHLODBuilderUtilities::BuildHLODs(UWorldPartition* InWorldPartition, FName
 			HLODBuilder->CellName = InCellName;
 
 			HLODBuilder->Build(InSubActors);
+			
+			HLODActors += HLODBuilder->HLODActors;
 		}
 	}
+
+	return HLODActors;
 }
 
 #endif
diff --git a/Engine/Source/Runtime/Engine/Private/WorldPartition/HLOD/HLODBuilder.h b/Engine/Source/Runtime/Engine/Private/WorldPartition/HLOD/HLODBuilder.h
--- a/Engine/Source/Runtime/Engine/Private/WorldPartition/HLOD/HLODBuilder.h
+++ b/Engine/Source/Runtime/Engine/Private/WorldPartition/HLOD/HLODBuilder.h
@@ -10,6 +10,7 @@
 class AActor;
 class UWorldPartition;
 class UHLODLayer;
+class AWorldPartitionHLOD;
 
 /**
  * Tools for building HLODs in WorldPartition
@@ -25,8 +26,9 @@ public:
 	 * @param 	InCellName			The name of the cell
 	 * @param 	InHLODLayer			The HLODLayer which will provide the HLOD building parameters.
 	 * @param 	InSubActors			The actors from which we'll gather geometry to generate an HLOD mesh
+	 * @return The list of HLOD actors for this cell.
 	 */
-	static void BuildHLODs(UWorldPartition* InWorldPartition, FName InCellName, const UHLODLayer* InHLODLayer, const TArray<AActor*>& InSubActors);
+	static TArray<AWorldPartitionHLOD*> BuildHLODs(UWorldPartition* InWorldPartition, FName InCellName, const UHLODLayer* InHLODLayer, const TArray<AActor*>& InSubActors);
 };
 
 #endif
\ No newline at end of file
diff --git a/Engine/Source/Runtime/Engine/Private/WorldPartition/HLOD/HLODLayer.cpp b/Engine/Source/Runtime/Engine/Private/WorldPartition/HLOD/HLODLayer.cpp
--- a/Engine/Source/Runtime/Engine/Private/WorldPartition/HLOD/HLODLayer.cpp
+++ b/Engine/Source/Runtime/Engine/Private/WorldPartition/HLOD/HLODLayer.cpp
@@ -20,7 +20,7 @@ UHLODLayer::UHLODLayer(const FObjectInitializer& ObjectInitializer)
 
 #if WITH_EDITOR
 
-void UHLODLayer::GenerateHLODForCell(UWorldPartition* InWorldPartition, FName InCellName, const TSet<FGuid>& InCellActors)
+TArray<AWorldPartitionHLOD*> UHLODLayer::GenerateHLODForCell(UWorldPartition* InWorldPartition, FName InCellName, const TSet<FGuid>& InCellActors)
 {
 	TRACE_CPUPROFILER_EVENT_SCOPE(UHLODLayer::GenerateHLODForCell);
 
@@ -45,6 +45,7 @@ void UHLODLayer::GenerateHLODForCell(UWorldPartition* InWorldPartition, FName In
 		}
 	}
 
+	TArray<AWorldPartitionHLOD*> HLODActors;
 	for (const auto& HLODLayerActors : HLODLayersActors)
 	{
 		UHLODLayer* HLODLayer = HLODLayerActors.Key;
@@ -54,8 +55,9 @@ void UHLODLayer::GenerateHLODForCell(UWorldPartition* InWorldPartition, FName In
 			continue;
 		}
 
-		FHLODBuilderUtilities::BuildHLODs(InWorldPartition, InCellName, HLODLayer, HLODLayerActors.Value);
+		HLODActors += FHLODBuilderUtilities::BuildHLODs(InWorldPartition, InCellName, HLODLayer, HLODLayerActors.Value);
 	}
+	return HLODActors;
 }
 
 bool UHLODLayer::ShouldIncludeInHLOD(AActor* InActor)
diff --git a/Engine/Source/Runtime/Engine/Private/WorldPartition/WorldPartitionRuntimeSpatialHash.cpp b/Engine/Source/Runtime/Engine/Private/WorldPartition/WorldPartitionRuntimeSpatialHash.cpp
--- a/Engine/Source/Runtime/Engine/Private/WorldPartition/WorldPartitionRuntimeSpatialHash.cpp
+++ b/Engine/Source/Runtime/Engine/Private/WorldPartition/WorldPartitionRuntimeSpatialHash.cpp
@@ -21,6 +21,7 @@
 #include "Misc/Parse.h"
 #include "Algo/Find.h"
 #include "Algo/Transform.h"
+#include "EngineUtils.h"
 
 #include "WorldPartition/HLOD/HLODLayer.h"
 #include "WorldPartition/HLOD/HLODActor.h"
@@ -1197,13 +1198,16 @@ bool UWorldPartitionRuntimeSpatialHash::GenerateHLOD()
 		GridActors[GridIndex].Add(ActorCluster);
 	}
 
+	// Keep track of all valid HLOD actors
+	TSet<AWorldPartitionHLOD*> ValidHLODActors;
+
 	const FBox WorldBounds = WorldPartition->GetWorldBounds();
 	for (int32 GridIndex = 0; GridIndex < Grids.Num(); GridIndex++)
 	{
 		const FSpatialHashRuntimeGrid& RuntimeGrid = Grids[GridIndex];
 		const FSquare2DGridHelper PartionedActors = GetPartitionedActors(WorldPartition, WorldBounds, RuntimeGrid, GridActors[GridIndex]);
 
-		PartionedActors.ForEachCells([PartionedActors, RuntimeGrid, WorldPartition, this](const FIntVector& CellCoord)
+		PartionedActors.ForEachCells([PartionedActors, &ValidHLODActors, RuntimeGrid, WorldPartition, this](const FIntVector& CellCoord)
 		{
 			const FSquare2DGridHelper::FGridLevel::FGridCell& GridCell = PartionedActors.GetCell(CellCoord);
 			if (GridCell.Actors.Num() != 0)
@@ -1212,11 +1216,24 @@ bool UWorldPartitionRuntimeSpatialHash::GenerateHLOD()
 				PartionedActors.GetCellBounds(CellCoord, CellBounds);
 
 				FName CellName = GetCellName(RuntimeGrid.GridName, CellCoord.Z, CellCoord.X, CellCoord.Y);
-				UHLODLayer::GenerateHLODForCell(WorldPartition, CellName, GridCell.Actors);
+
+				TArray<AWorldPartitionHLOD*> HLODActors;
+				HLODActors = UHLODLayer::GenerateHLODForCell(WorldPartition, CellName, GridCell.Actors);
+
+				ValidHLODActors.Append(HLODActors);
 			}
 		});
 	}
 
+	// Destroy all invalid HLOD actors
+	for (TActorIterator<AWorldPartitionHLOD> ItHLOD(GetWorld()); ItHLOD; ++ItHLOD)
+	{
+		if (!ValidHLODActors.Contains(*ItHLOD))
+		{
+			GetWorld()->DestroyActor(*ItHLOD);
+		}
+	}
+
 	return true;
 }
 #endif
diff --git a/Engine/Source/Runtime/Engine/Public/WorldPartition/HLOD/HLODLayer.h b/Engine/Source/Runtime/Engine/Public/WorldPartition/HLOD/HLODLayer.h
--- a/Engine/Source/Runtime/Engine/Public/WorldPartition/HLOD/HLODLayer.h
+++ b/Engine/Source/Runtime/Engine/Public/WorldPartition/HLOD/HLODLayer.h
@@ -59,7 +59,7 @@ class ENGINE_API UHLODLayer : public UObject
 public:
 	const TArray<FHLODLevelSettings>& GetLevels() const { return Levels; }
 
-	static void GenerateHLODForCell(UWorldPartition* InWorldPartition, FName InCellName, const TSet<FGuid>& InCellActors);
+	static TArray<AWorldPartitionHLOD*> GenerateHLODForCell(UWorldPartition* InWorldPartition, FName InCellName, const TSet<FGuid>& InCellActors);
 
 	static UHLODLayer* GetHLODLayer(const AActor* InActor);
 	static bool ShouldIncludeInHLOD(AActor* InActor);
