Missing files for previous CL

diff --git a/Engine/Source/Runtime/Engine/Private/WorldPartition/HLOD/HLODActor.cpp b/Engine/Source/Runtime/Engine/Private/WorldPartition/HLOD/HLODActor.cpp
--- a/Engine/Source/Runtime/Engine/Private/WorldPartition/HLOD/HLODActor.cpp
+++ b/Engine/Source/Runtime/Engine/Private/WorldPartition/HLOD/HLODActor.cpp
@@ -20,17 +20,15 @@ UWorldPartitionRuntimeHLODCellData::UWorldPartitionRuntimeHLODCellData(const FOb
 AWorldPartitionHLOD::AWorldPartitionHLOD(const FObjectInitializer& ObjectInitializer)
 	: Super(ObjectInitializer)
 {
-	if (!IsTemplate())
-	{
-		USceneComponent* Root = CreateDefaultSubobject<USceneComponent>(USceneComponent::GetDefaultSceneRootVariableName());
-		Root->SetMobility(EComponentMobility::Static);
-		SetRootComponent(Root);
-	}
-
 	SetCanBeDamaged(false);
 	SetActorEnableCollision(false);
 }
 
+UPrimitiveComponent* AWorldPartitionHLOD::GetHLODComponent()
+{
+	return CastChecked<UPrimitiveComponent>(RootComponent);
+}
+
 void AWorldPartitionHLOD::LinkCell(FName InCellName)
 {
 	TRACE_CPUPROFILER_EVENT_SCOPE(AWorldPartitionHLOD::LinkCell);
@@ -39,7 +37,7 @@ void AWorldPartitionHLOD::LinkCell(FName InCellName)
 	{
 		if (SubPrimitiveComponent.IsValid())
 		{
-			SubPrimitiveComponent->SetCachedLODParentPrimitive(HLODComponent);
+			SubPrimitiveComponent->SetCachedLODParentPrimitive(GetHLODComponent());
 		}
 	}
 }
@@ -89,7 +87,7 @@ void AWorldPartitionHLOD::UpdateLODParent(AActor& InActor, bool bInClear)
 		{
 			if (UHLODLayer::ShouldIncludeInHLOD(PrimitiveComponent, SubActorsHLODLevel))
 			{
-				PrimitiveComponent->SetCachedLODParentPrimitive(!bInClear ? HLODComponent : nullptr);
+				PrimitiveComponent->SetCachedLODParentPrimitive(!bInClear ? GetHLODComponent() : nullptr);
 			}
 		}
 	}
@@ -97,24 +95,32 @@ void AWorldPartitionHLOD::UpdateLODParent(AActor& InActor, bool bInClear)
 
 const float HLODACTOR_DEFAULT_MIN_DRAW_DISTANCE = 5000.0f;
 
-void AWorldPartitionHLOD::SetParentPrimitive(UPrimitiveComponent* InParentPrimitive)
+void AWorldPartitionHLOD::SetHLODPrimitive(UPrimitiveComponent* InHLODPrimitive)
 {
 	TRACE_CPUPROFILER_EVENT_SCOPE(AWorldPartitionHLOD::SetHLODPrimitive);
 
-	HLODComponent = InParentPrimitive;
+	USceneComponent* OldRootComponent = GetRootComponent();
+
+	SetRootComponent(InHLODPrimitive);
+	AddInstanceComponent(InHLODPrimitive);
+
+	// Setup custom depth rendering to achieve a red tint using a post process material
+	const int32 CellPreviewStencilValue = 180;
+	InHLODPrimitive->bRenderCustomDepth = true;
+	InHLODPrimitive->CustomDepthStencilValue = CellPreviewStencilValue;
+	
+	InHLODPrimitive->SetCollisionEnabled(ECollisionEnabled::NoCollision);
+	InHLODPrimitive->SetMobility(EComponentMobility::Static);
 
-	HLODComponent->AttachToComponent(GetRootComponent(), FAttachmentTransformRules::KeepWorldTransform);
-	AddInstanceComponent(HLODComponent);
+	InHLODPrimitive->MinDrawDistance = HLODACTOR_DEFAULT_MIN_DRAW_DISTANCE;
 
-	if (GetRootComponent()->IsRegistered())
+	InHLODPrimitive->RegisterComponent();
+	InHLODPrimitive->MarkRenderStateDirty();
+
+	if (OldRootComponent)
 	{
-		HLODComponent->RegisterComponent();
+		OldRootComponent->DestroyComponent();
 	}
-
-	HLODComponent->MinDrawDistance = HLODACTOR_DEFAULT_MIN_DRAW_DISTANCE;
-	HLODComponent->SetCollisionEnabled(ECollisionEnabled::NoCollision);
-	HLODComponent->SetMobility(EComponentMobility::Static);
-	HLODComponent->MarkRenderStateDirty();
 }
 
 void AWorldPartitionHLOD::SetChildrenPrimitives(const TArray<UPrimitiveComponent*>& InChildrenPrimitives)
diff --git a/Engine/Source/Runtime/Engine/Public/WorldPartition/HLOD/HLODActor.h b/Engine/Source/Runtime/Engine/Public/WorldPartition/HLOD/HLODActor.h
--- a/Engine/Source/Runtime/Engine/Public/WorldPartition/HLOD/HLODActor.h
+++ b/Engine/Source/Runtime/Engine/Public/WorldPartition/HLOD/HLODActor.h
@@ -16,6 +16,8 @@ class ENGINE_API AWorldPartitionHLOD : public AActor
 	GENERATED_UCLASS_BODY()
 
 public:
+	UPrimitiveComponent* GetHLODComponent();
+
 	void LinkCell(FName InCellName);
 	void UnlinkCell(FName InCellName);
 
@@ -33,7 +35,7 @@ protected:
 
 public:
 #if WITH_EDITOR
-	void SetParentPrimitive(UPrimitiveComponent* InParentPrimitive);
+	void SetHLODPrimitive(UPrimitiveComponent* InHLODPrimitive);
 	void SetChildrenPrimitives(const TArray<UPrimitiveComponent*>& InChildrenPrimitives);
 
 	const TArray<FGuid>& GetSubActors() const;
