Momentum: Removing SourceAnimationData and doing animation layering bake during compression

diff --git a/Engine/Source/Editor/AnimationEditor/Private/AnimationEditor.cpp b/Engine/Source/Editor/AnimationEditor/Private/AnimationEditor.cpp
--- a/Engine/Source/Editor/AnimationEditor/Private/AnimationEditor.cpp
+++ b/Engine/Source/Editor/AnimationEditor/Private/AnimationEditor.cpp
@@ -198,10 +198,6 @@ void FAnimationEditor::BindCommands()
 		FExecuteAction::CreateSP(this, &FAnimationEditor::OnReimportAnimation),
 		FCanExecuteAction::CreateSP(this, &FAnimationEditor::HasValidAnimationSequence));
 
-	ToolkitCommands->MapAction(FAnimationEditorCommands::Get().ApplyAnimation,
-		FExecuteAction::CreateSP(this, &FAnimationEditor::OnApplyRawAnimChanges),
-		FCanExecuteAction::CreateSP(this, &FAnimationEditor::CanApplyRawAnimChanges));
-
 	ToolkitCommands->MapAction(FAnimationEditorCommands::Get().ExportToFBX_AnimData,
 		FExecuteAction::CreateSP(this, &FAnimationEditor::OnExportToFBX, EExportSourceOption::CurrentAnimation_AnimData),
 		FCanExecuteAction::CreateSP(this, &FAnimationEditor::HasValidAnimationSequence));
@@ -264,7 +260,7 @@ void FAnimationEditor::ExtendToolbar()
 			ToolbarBuilder.BeginSection("Animation");
 			{
 				ToolbarBuilder.AddToolBarButton(FAnimationEditorCommands::Get().ReimportAnimation);
-				ToolbarBuilder.AddToolBarButton(FAnimationEditorCommands::Get().ApplyCompression, NAME_None, LOCTEXT("Toolbar_ApplyCompression", "Compression"));
+				ToolbarBuilder.AddToolBarButton(FAnimationEditorCommands::Get().ApplyCompression, NAME_None, LOCTEXT("Toolbar_ApplyCompression", "Apply Compression"));
 
 				{
 					ToolbarBuilder.AddComboButton(
@@ -281,7 +277,6 @@ void FAnimationEditor::ExtendToolbar()
 			ToolbarBuilder.BeginSection("Editing");
 			{
 				ToolbarBuilder.AddToolBarButton(FAnimationEditorCommands::Get().SetKey, NAME_None, LOCTEXT("Toolbar_SetKey", "Key"));
-				ToolbarBuilder.AddToolBarButton(FAnimationEditorCommands::Get().ApplyAnimation, NAME_None, LOCTEXT("Toolbar_ApplyAnimation", "Apply"));
 			}
 			ToolbarBuilder.EndSection();
 
@@ -577,37 +572,6 @@ void FAnimationEditor::OnSetKey()
 	}
 }
 
-bool FAnimationEditor::CanApplyRawAnimChanges() const
-{
-	UAnimSequence* AnimSequence = Cast<UAnimSequence>(AnimationAsset);
-
-	// ideally would be great if we can only show if something changed
-	return (AnimSequence && (AnimSequence->DoesNeedRebake() || AnimSequence->DoesNeedRecompress()));
-}
-
-void FAnimationEditor::OnApplyRawAnimChanges()
-{
-	UAnimSequence* AnimSequence = Cast<UAnimSequence>(AnimationAsset);
-	if (AnimSequence)
-	{
-		if (AnimSequence->DoesNeedRebake() || AnimSequence->DoesNeedRecompress())
-		{
-			FScopedTransaction ScopedTransaction(LOCTEXT("BakeAnimation", "Bake Animation"));
-			if (AnimSequence->DoesNeedRebake())
-			{
-				AnimSequence->Modify(true);
-				AnimSequence->BakeTrackCurvesToRawAnimation();
-			}
-
-			if (AnimSequence->DoesNeedRecompress())
-			{
-				AnimSequence->Modify(true);
-				AnimSequence->RequestSyncAnimRecompression(false);
-			}
-		}
-	}
-}
-
 void FAnimationEditor::OnReimportAnimation()
 {
 	UAnimSequence* AnimSequence = Cast<UAnimSequence>(AnimationAsset);
diff --git a/Engine/Source/Editor/AnimationEditor/Private/AnimationEditor.h b/Engine/Source/Editor/AnimationEditor/Private/AnimationEditor.h
--- a/Engine/Source/Editor/AnimationEditor/Private/AnimationEditor.h
+++ b/Engine/Source/Editor/AnimationEditor/Private/AnimationEditor.h
@@ -129,10 +129,6 @@ private:
 
 	void OnSetKey();
 
-	bool CanApplyRawAnimChanges() const;
-
-	void OnApplyRawAnimChanges();
-
 	void OnReimportAnimation();
 
 	void OnApplyCompression();
diff --git a/Engine/Source/Editor/AnimationEditor/Private/AnimationEditorCommands.cpp b/Engine/Source/Editor/AnimationEditor/Private/AnimationEditorCommands.cpp
--- a/Engine/Source/Editor/AnimationEditor/Private/AnimationEditorCommands.cpp
+++ b/Engine/Source/Editor/AnimationEditor/Private/AnimationEditorCommands.cpp
@@ -12,8 +12,7 @@ void FAnimationEditorCommands::RegisterCommands()
 	UI_COMMAND(ExportToFBX_AnimData, "Animation Data", "Export current animation source data to FBX", EUserInterfaceActionType::Button, FInputChord());
 	UI_COMMAND(ExportToFBX_PreviewMesh, "Preview Mesh", "Export current animation by playing on the current preview mesh, including retargeting, post process graph, or anything you see on the preview mesh.", EUserInterfaceActionType::Button, FInputChord());
 	UI_COMMAND(AddLoopingInterpolation, "Add Looping Interpolation", "Add an extra first frame at the end of the animation to create interpolation when looping", EUserInterfaceActionType::Button, FInputChord());
-	UI_COMMAND(SetKey, "Set Key", "Add Bone Transform to Additive Layer Tracks", EUserInterfaceActionType::Button, FInputChord(EKeys::S));
-	UI_COMMAND(ApplyAnimation, "Apply Animation", "Apply Additive Layer Tracks to Runtime Animation Data", EUserInterfaceActionType::Button, FInputChord());
+	UI_COMMAND(SetKey, "Set Key", "Add Bone Transform to Additive Layer Tracks", EUserInterfaceActionType::Button, FInputChord(EKeys::S));	
 	UI_COMMAND(RemoveBoneTracks, "Remove All Bone Tracks", "Remove all bone tracks. This doesn't remove notifies or curves.", EUserInterfaceActionType::Button, FInputChord());
 }
 
diff --git a/Engine/Source/Editor/AnimationEditor/Private/AnimationEditorCommands.h b/Engine/Source/Editor/AnimationEditor/Private/AnimationEditorCommands.h
--- a/Engine/Source/Editor/AnimationEditor/Private/AnimationEditorCommands.h
+++ b/Engine/Source/Editor/AnimationEditor/Private/AnimationEditorCommands.h
@@ -31,8 +31,6 @@ public:
 	TSharedPtr<FUICommandInfo> AddLoopingInterpolation;
 	// set key for bone track
 	TSharedPtr<FUICommandInfo> SetKey;
-	// bake bone track curves to animation
-	TSharedPtr<FUICommandInfo> ApplyAnimation;
 	// Remove bone tracks
 	TSharedPtr<FUICommandInfo> RemoveBoneTracks;
 };
diff --git a/Engine/Source/Editor/AnimationModifiers/Private/AnimationBlueprintLibrary.cpp b/Engine/Source/Editor/AnimationModifiers/Private/AnimationBlueprintLibrary.cpp
--- a/Engine/Source/Editor/AnimationModifiers/Private/AnimationBlueprintLibrary.cpp
+++ b/Engine/Source/Editor/AnimationModifiers/Private/AnimationBlueprintLibrary.cpp
@@ -1195,8 +1195,6 @@ void UAnimationBlueprintLibrary::RemoveAllCurveData(UAnimSequence* AnimationSequ
 		AnimationSequence->RawCurveData.DeleteAllCurveData(ERawCurveTrackTypes::RCT_Float);
 		AnimationSequence->RawCurveData.DeleteAllCurveData(ERawCurveTrackTypes::RCT_Vector);
 		AnimationSequence->RawCurveData.DeleteAllCurveData(ERawCurveTrackTypes::RCT_Transform);
-
-		AnimationSequence->bNeedsRebake = true;
 	}
 	else
 	{
@@ -1215,7 +1213,6 @@ void UAnimationBlueprintLibrary::AddTransformationCurveKey(UAnimSequence* Animat
 		TransformArray.Add(Transform);
 
 		AddCurveKeysInternal<FTransform, FTransformCurve>(AnimationSequence, CurveName, TimeArray, TransformArray, ERawCurveTrackTypes::RCT_Transform);
-		AnimationSequence->bNeedsRebake = true;
 	}
 	else
 	{
@@ -1231,7 +1228,6 @@ void UAnimationBlueprintLibrary::AddTransformationCurveKeys(UAnimSequence* Anima
 		if (Times.Num() == Transforms.Num())
 		{
 			AddCurveKeysInternal<FTransform, FTransformCurve>(AnimationSequence, CurveName, Times, Transforms, ERawCurveTrackTypes::RCT_Transform);
-			AnimationSequence->bNeedsRebake = true;
 		}
 		else
 		{
@@ -1346,8 +1342,6 @@ void UAnimationBlueprintLibrary::AddCurveKeysInternal(UAnimSequence* AnimationSe
 			{
 				Curve->UpdateOrAddKey(KeyData[KeyIndex], Times[KeyIndex]);
 			}
-
-			AnimationSequence->BakeTrackCurvesToRawAnimation();
 		}
 	}
 }
@@ -1387,7 +1381,6 @@ bool UAnimationBlueprintLibrary::RemoveCurveInternal(UAnimSequence* AnimationSeq
 			if (ContainerName == USkeleton::AnimTrackCurveMappingName)
 			{
 				bRemoved = AnimationSequence->RawCurveData.DeleteCurveData(SmartCurveName, ERawCurveTrackTypes::RCT_Transform);
-				AnimationSequence->bNeedsRebake = true;
 			}
 			else
 			{
diff --git a/Engine/Source/Editor/AnimationModifiers/Private/AnimationModifier.cpp b/Engine/Source/Editor/AnimationModifiers/Private/AnimationModifier.cpp
--- a/Engine/Source/Editor/AnimationModifiers/Private/AnimationModifier.cpp
+++ b/Engine/Source/Editor/AnimationModifiers/Private/AnimationModifier.cpp
@@ -110,17 +110,9 @@ void UAnimationModifier::ApplyToAnimationSequence(class UAnimSequence* InAnimati
 
 void UAnimationModifier::UpdateCompressedAnimationData()
 {
-	if (CurrentAnimSequence->DoesNeedRebake() || CurrentAnimSequence->DoesNeedRecompress())
+	if (CurrentAnimSequence->DoesNeedRecompress())
 	{
-		if (CurrentAnimSequence->DoesNeedRebake())
-		{
-			CurrentAnimSequence->BakeTrackCurvesToRawAnimation();
-		}
-
-		if (CurrentAnimSequence->DoesNeedRecompress())
-		{
-			CurrentAnimSequence->RequestSyncAnimRecompression(false);
-		}
+		CurrentAnimSequence->RequestSyncAnimRecompression(false);
 	}
 }
 
diff --git a/Engine/Source/Editor/Persona/Private/AnimTimeline/AnimModel_AnimSequenceBase.cpp b/Engine/Source/Editor/Persona/Private/AnimTimeline/AnimModel_AnimSequenceBase.cpp
--- a/Engine/Source/Editor/Persona/Private/AnimTimeline/AnimModel_AnimSequenceBase.cpp
+++ b/Engine/Source/Editor/Persona/Private/AnimTimeline/AnimModel_AnimSequenceBase.cpp
@@ -345,12 +345,7 @@ void FAnimModel_AnimSequenceBase::RemoveSelectedCurves()
 					OnStopEditingCurves.ExecuteIfBound(CurveEditInfo);
 
 					AnimSequenceBase->RawCurveData.DeleteCurveData(CurveToDelete, ERawCurveTrackTypes::RCT_Transform);
-
-					if(UAnimSequence* AnimSequence = Cast<UAnimSequence>(AnimSequenceBase))
-					{
-						AnimSequence->bNeedsRebake = true;
-					}
-
+					
 					bDeletedCurve = true;
 				}
 			}	
diff --git a/Engine/Source/Editor/Persona/Private/AnimTimeline/AnimTimelineTrack_Curve.cpp b/Engine/Source/Editor/Persona/Private/AnimTimeline/AnimTimelineTrack_Curve.cpp
--- a/Engine/Source/Editor/Persona/Private/AnimTimeline/AnimTimelineTrack_Curve.cpp
+++ b/Engine/Source/Editor/Persona/Private/AnimTimeline/AnimTimelineTrack_Curve.cpp
@@ -372,11 +372,6 @@ void FAnimTimelineTrack_Curve::HandleCurveChanged()
 	}
 
 	GetModel()->GetAnimSequenceBase()->MarkRawDataAsModified();
-
-	if(UAnimSequence* AnimSequence = Cast<UAnimSequence>(GetModel()->GetAnimSequenceBase()))
-	{
-		AnimSequence->bNeedsRebake = true;
-	}
 }
 
 void FAnimTimelineTrack_Curve::PostUndoRedo()
diff --git a/Engine/Source/Editor/Persona/Private/AnimTimeline/AnimTimelineTrack_TransformCurve.cpp b/Engine/Source/Editor/Persona/Private/AnimTimeline/AnimTimelineTrack_TransformCurve.cpp
--- a/Engine/Source/Editor/Persona/Private/AnimTimeline/AnimTimelineTrack_TransformCurve.cpp
+++ b/Engine/Source/Editor/Persona/Private/AnimTimeline/AnimTimelineTrack_TransformCurve.cpp
@@ -143,10 +143,7 @@ void FAnimTimelineTrack_TransformCurve::DeleteTrack()
 			AnimSequenceBase->Modify();
 			AnimSequenceBase->RawCurveData.DeleteCurveData(CurveToDelete, ERawCurveTrackTypes::RCT_Transform);
 
-			if(UAnimSequence* AnimSequence = Cast<UAnimSequence>(AnimSequenceBase))
-			{
-				AnimSequence->bNeedsRebake = true;
-			}
+			AnimSequenceBase->MarkRawDataAsModified();
 
 			GetModel()->RefreshTracks();
 
@@ -179,10 +176,7 @@ void FAnimTimelineTrack_TransformCurve::ToggleEnabled()
 
 		Curve->SetCurveTypeFlag(AACF_Disabled, bEnabled);
 
-		if(UAnimSequence* AnimSequence = Cast<UAnimSequence>(AnimSequenceBase))
-		{
-			AnimSequence->bNeedsRebake = true;
-		}
+		AnimSequenceBase->MarkRawDataAsModified();
 
 		// need to update curves, otherwise they're not disabled
 		if (GetModel()->GetPreviewScene()->GetPreviewMeshComponent()->PreviewInstance != nullptr)
diff --git a/Engine/Source/Editor/Persona/Private/AnimationEditorViewportClient.cpp b/Engine/Source/Editor/Persona/Private/AnimationEditorViewportClient.cpp
--- a/Engine/Source/Editor/Persona/Private/AnimationEditorViewportClient.cpp
+++ b/Engine/Source/Editor/Persona/Private/AnimationEditorViewportClient.cpp
@@ -782,11 +782,6 @@ FText FAnimationViewportClient::GetDisplayInfo(bool bDisplayAllInfo) const
 		UAnimSequence* Sequence = Cast<UAnimSequence>(PreviewInstance->GetCurrentAsset());
 		if (Sequence)
 		{
-			if (Sequence->DoesNeedRebake())
-			{
-				TextValue = ConcatenateLine(TextValue, LOCTEXT("ApplyRawAnimationDataWarning", "<AnimViewport.WarningText>Animation is being edited. To apply to raw animation data, click \"Apply\"</>"));
-			}
-
 			if (Sequence->DoesNeedRecompress())
 			{
 				TextValue = ConcatenateLine(TextValue, LOCTEXT("ApplyToCompressedDataWarning", "<AnimViewport.WarningText>Animation is being edited. To apply to compressed data (and recalculate baked additives), click \"Apply\"</>"));
diff --git a/Engine/Source/Editor/UnrealEd/Private/Animation/DebugSkelMeshComponent.cpp b/Engine/Source/Editor/UnrealEd/Private/Animation/DebugSkelMeshComponent.cpp
--- a/Engine/Source/Editor/UnrealEd/Private/Animation/DebugSkelMeshComponent.cpp
+++ b/Engine/Source/Editor/UnrealEd/Private/Animation/DebugSkelMeshComponent.cpp
@@ -542,11 +542,13 @@ void UDebugSkelMeshComponent::RefreshBoneTransforms(FActorComponentTickFunction*
 
 			BoneContainer.SetUseSourceData(false);
 			BoneContainer.SetUseRAWData(true);
+			PreviewInstance->EnableControllers(false);
 
 			GenSpaceBases(UncompressedSpaceBases);
-
+			
 			BoneContainer.SetUseRAWData(bUseRaw);
 			BoneContainer.SetUseSourceData(bUseSource);
+			PreviewInstance->EnableControllers(true);
 		}
 
 		// Non retargeted pose.
diff --git a/Engine/Source/Editor/UnrealEd/Private/EditorAnimUtils.cpp b/Engine/Source/Editor/UnrealEd/Private/EditorAnimUtils.cpp
--- a/Engine/Source/Editor/UnrealEd/Private/EditorAnimUtils.cpp
+++ b/Engine/Source/Editor/UnrealEd/Private/EditorAnimUtils.cpp
@@ -256,7 +256,6 @@ namespace EditorAnimUtils
 					// clear transform curves since those curves won't work in new skeleton
 					// since we're deleting curves, mark this rebake flag off
 					AnimSequenceToRetarget->RawCurveData.TransformCurves.Empty();
-					AnimSequenceToRetarget->bNeedsRebake = false;
 					// I can't copy transform curves yet because transform curves need retargeting. 
 					//EditorAnimUtils::CopyAnimCurves(OldSkeleton, NewSkeleton, AssetToRetarget, USkeleton::AnimTrackCurveMappingName, FRawCurveTracks::TransformType);
 				}
diff --git a/Engine/Source/Editor/UnrealEd/Private/SkeletalMeshEdit.cpp b/Engine/Source/Editor/UnrealEd/Private/SkeletalMeshEdit.cpp
--- a/Engine/Source/Editor/UnrealEd/Private/SkeletalMeshEdit.cpp
+++ b/Engine/Source/Editor/UnrealEd/Private/SkeletalMeshEdit.cpp
@@ -1504,8 +1504,6 @@ bool UnFbx::FFbxImporter::ImportAnimation(USkeleton* Skeleton, UAnimSequence * D
 	}
 
 	// importing custom attribute END
-	
-	const bool bSourceDataExists = DestSeq->HasSourceRawData();
 	TArray<AnimationTransformDebug::FAnimationTransformDebugData> TransformDebugData;
 	int32 TotalNumKeys = 0;
 	const FReferenceSkeleton& RefSkeleton = Skeleton->GetReferenceSkeleton();
@@ -1697,16 +1695,8 @@ bool UnFbx::FFbxImporter::ImportAnimation(USkeleton* Skeleton, UAnimSequence * D
 	{
 		GWarn->BeginSlowTask( LOCTEXT("BeginCompressAnimation", "Compress Animation"), true);
 		GWarn->StatusForceUpdate(1, 1, LOCTEXT("CompressAnimation", "Compressing Animation"));
-		// if source data exists, you should bake it to Raw to apply
-		if(bSourceDataExists)
-		{
-			DestSeq->BakeTrackCurvesToRawAnimation();
-		}
-		else
-		{
-			// otherwise just compress
-			DestSeq->PostProcessSequence();
-		}
+		
+		DestSeq->PostProcessSequence();
 
 		// run debug mode
 		AnimationTransformDebug::OutputAnimationTransformDebugData(TransformDebugData, TotalNumKeys, RefSkeleton);
diff --git a/Engine/Source/Runtime/Core/Public/UObject/UE5MainStreamObjectVersion.h b/Engine/Source/Runtime/Core/Public/UObject/UE5MainStreamObjectVersion.h
--- a/Engine/Source/Runtime/Core/Public/UObject/UE5MainStreamObjectVersion.h
+++ b/Engine/Source/Runtime/Core/Public/UObject/UE5MainStreamObjectVersion.h
@@ -17,6 +17,9 @@ struct CORE_API FUE5MainStreamObjectVersion
 
 		// Nanite Geometry Collection data moved to DDC
 		GeometryCollectionNaniteDDC,
+		
+		// Removing SourceAnimationData, animation layering is now applied during compression
+		RemovingSourceAnimationData,
 
 		// -----<new versions can be added above this line>-------------------------------------------------
 		VersionPlusOne,
diff --git a/Engine/Source/Runtime/Engine/Classes/Animation/AnimSequence.h b/Engine/Source/Runtime/Engine/Classes/Animation/AnimSequence.h
--- a/Engine/Source/Runtime/Engine/Classes/Animation/AnimSequence.h
+++ b/Engine/Source/Runtime/Engine/Classes/Animation/AnimSequence.h
@@ -241,7 +241,7 @@ protected:
 	/**
 	 * Source RawAnimationData. Only can be overridden by when transform curves are added first time OR imported
 	 */
-	TArray<struct FRawAnimSequenceTrack> SourceRawAnimationData;
+	TArray<struct FRawAnimSequenceTrack> SourceRawAnimationData_DEPRECATED;
 #endif
 
 public:
@@ -346,6 +346,7 @@ public:
 	UPROPERTY()
 	FString SourceFileTimestamp_DEPRECATED;
 
+	UE_DEPRECATED(5.0, "bNeedsRebake has been deprecated, transform curves are now baked during compression")
 	UPROPERTY(transient)
 	bool bNeedsRebake;
 
@@ -370,7 +371,6 @@ public:
 	virtual void GetPreloadDependencies(TArray<UObject*>& OutDeps) override;
 #if WITH_EDITOR
 	virtual void PostEditChangeProperty(FPropertyChangedEvent& PropertyChangedEvent) override;
-	virtual void PostDuplicate(bool bDuplicateForPIE) override;
 #endif // WITH_EDITOR
 	virtual void BeginDestroy() override;
 	virtual void GetAssetRegistryTags(TArray<FAssetRegistryTag>& OutTags) const override;
@@ -446,7 +446,8 @@ public:
 	const TArray<FRawAnimSequenceTrack>& GetRawAnimationData() const { return RawAnimationData; }
 
 #if WITH_EDITORONLY_DATA
-	bool  HasSourceRawData() const { return SourceRawAnimationData.Num() > 0; }
+	UE_DEPRECATED(5.0, "SourceRawAnimationData has been deprecated")
+	bool HasSourceRawData() const { return false; }
 	const TArray<FName>& GetAnimationTrackNames() const { return AnimationTrackNames; }
 	
 	void  UpdateCompressedCurveName(SmartName::UID_Type CurveUID, const struct FSmartName& NewCurveName);
@@ -585,6 +586,7 @@ public:
 
 protected:
 	void ApplyCompressedData(const TArray<uint8>& Data);
+	void SanitizeAnimationTrackData(FRawAnimSequenceTrack& RawAnim);
 
 public:
 	bool IsCompressedDataValid() const;
@@ -661,13 +663,17 @@ public:
 	/** 
 	 * Bake Transform Curves.TransformCurves to RawAnimation after making a back up of current RawAnimation
 	 */
-	void BakeTrackCurvesToRawAnimation();
+	UE_DEPRECATED(5.0, "BakeTrackCurvesToRawAnimation has been deprecated, transform curves are now baked during compression")
+	void BakeTrackCurvesToRawAnimation() {}
+
+	void BakeTrackCurvesToRawAnimationTracks(TArray<FRawAnimSequenceTrack>& NewRawTracks, TArray<FName>& NewTrackNames, TArray<FTrackToSkeletonMap>& NewTrackToSkeletonMapTable);
 
 	/**
 	 * Sometimes baked data gets invalidated. For example, if you retarget this from another animation
 	 * It won't matter anymore, so in any case, when the data is not valid anymore
 	 * We clear Source Raw Animation Data as well as Transform Curve
 	 */
+	UE_DEPRECATED(5.0, "ClearBakedTransformData has been deprecated, transform curves are now baked during compression")
 	void ClearBakedTransformData();
 	/**
 	 * Add Key to Transform Curves
@@ -676,7 +682,8 @@ public:
 	/**
 	 * Return true if it needs to re-bake
 	 */
-	bool DoesNeedRebake() const;
+	UE_DEPRECATED(5.0, "DoesNeedRebake has been deprecated, transform curves are now baked during compression")
+	bool DoesNeedRebake() const { return false; }
 	/**
 	 * Return true if it contains transform curves
 	 */
@@ -685,12 +692,14 @@ public:
 	/**
 	 * Returns whether this animation has baked transform curves (i.e. has the raw data been modified)
 	 */
-	bool HasBakedTransformCurves() const;
+	UE_DEPRECATED(5.0, "HasBakedTransformCurves has been deprecated, transform curves are now baked during compression")
+	bool HasBakedTransformCurves() const { return false; }
 
 	/**
 	 * Restore the pre baked transform curve raw data
 	 */
-	void RestoreSourceData();
+	UE_DEPRECATED(5.0, "RestoreSourceData has been deprecated, transform curves are now baked during compression")
+	void RestoreSourceData() {}
 
 	/**
 	* Return true if compressed data is out of date / missing and so animation needs to use raw data
diff --git a/Engine/Source/Runtime/Engine/Private/Animation/AnimCompressionTypes.cpp b/Engine/Source/Runtime/Engine/Private/Animation/AnimCompressionTypes.cpp
--- a/Engine/Source/Runtime/Engine/Private/Animation/AnimCompressionTypes.cpp
+++ b/Engine/Source/Runtime/Engine/Private/Animation/AnimCompressionTypes.cpp
@@ -38,13 +38,35 @@ void UpdateWithData(FSHA1& Sha, const DataType& Data)
 
 void UpdateSHAWithCurves(FSHA1& Sha, const FRawCurveTracks& InRawCurveData) 
 {
+	auto UpdateWithFloatCurve = [&Sha](const FRichCurve& Curve)
+	{
+		UpdateWithData(Sha, Curve.DefaultValue);
+		UpdateSHAWithArray(Sha, Curve.GetConstRefOfKeys());
+		UpdateWithData(Sha, Curve.PreInfinityExtrap);
+		UpdateWithData(Sha, Curve.PostInfinityExtrap);
+	};
+
 	for (const FFloatCurve& Curve : InRawCurveData.FloatCurves)
 	{
 		UpdateWithData(Sha, Curve.Name.UID);
-		UpdateWithData(Sha, Curve.FloatCurve.DefaultValue);
-		UpdateSHAWithArray(Sha, Curve.FloatCurve.GetConstRefOfKeys());
-		UpdateWithData(Sha, Curve.FloatCurve.PreInfinityExtrap);
-		UpdateWithData(Sha, Curve.FloatCurve.PostInfinityExtrap);
+		UpdateWithFloatCurve(Curve.FloatCurve);
+	}
+
+	for (const FTransformCurve& Curve : InRawCurveData.TransformCurves)
+	{
+		UpdateWithData(Sha, Curve.Name.UID);
+
+		auto UpdateWithComponent = [&Sha, &UpdateWithFloatCurve](const FVectorCurve& VectorCurve)
+		{
+			for (int32 ChannelIndex = 0; ChannelIndex < 3; ++ChannelIndex)
+			{
+				UpdateWithFloatCurve(VectorCurve.FloatCurves[ChannelIndex]);
+			}
+		};
+		
+		UpdateWithComponent(Curve.TranslationCurve);
+		UpdateWithComponent(Curve.RotationCurve);
+		UpdateWithComponent(Curve.ScaleCurve);
 	}
 }
 
@@ -151,22 +173,33 @@ FCompressibleAnimData::FCompressibleAnimData(class UAnimSequence* InSeq, const b
 
 	const bool bHasVirtualBones = InSeq->GetSkeleton()->GetVirtualBones().Num() > 0;
 
+	TArray<FName> TempTrackNames;
+
 	if (InSeq->CanBakeAdditive())
 	{
-		TArray<FName> TempTrackNames;
 		InSeq->BakeOutAdditiveIntoRawData(RawAnimationData, TempTrackNames, TrackToSkeletonMapTable, RawCurveData, AdditiveBaseAnimationData);
 	}
-	else if (bHasVirtualBones)// If we aren't additive we must bake virtual bones
-	{
-		TArray<FName> TempTrackNames;
-		InSeq->BakeOutVirtualBoneTracks(RawAnimationData, TempTrackNames, TrackToSkeletonMapTable);
-		RawCurveData = InSeq->RawCurveData;
-	}
 	else
 	{
-		RawAnimationData = InSeq->GetRawAnimationData();
-		TrackToSkeletonMapTable = InSeq->GetRawTrackToSkeletonMapTable();
+		// In case we require baking down transform curves, do so now meaning Virtual Bone baking will incorporate the correct bone transforms
+		if (InSeq->DoesContainTransformCurves())
+		{
+			InSeq->BakeTrackCurvesToRawAnimationTracks(RawAnimationData, TempTrackNames, TrackToSkeletonMapTable);			
+		}
+		else
+		{
+			RawAnimationData = InSeq->GetRawAnimationData();
+			TrackToSkeletonMapTable = InSeq->GetRawTrackToSkeletonMapTable();
+			TempTrackNames = InSeq->GetAnimationTrackNames();
+		}
+
 		RawCurveData = InSeq->RawCurveData;
+
+		// If we aren't additive we must bake virtual bones
+		if (bHasVirtualBones)
+		{
+			InSeq->BakeOutVirtualBoneTracks(RawAnimationData, TempTrackNames, TrackToSkeletonMapTable);
+		}
 	}
 
 	if (bPerformStripping)
diff --git a/Engine/Source/Runtime/Engine/Private/Animation/AnimSequence.cpp b/Engine/Source/Runtime/Engine/Private/Animation/AnimSequence.cpp
--- a/Engine/Source/Runtime/Engine/Private/Animation/AnimSequence.cpp
+++ b/Engine/Source/Runtime/Engine/Private/Animation/AnimSequence.cpp
@@ -8,6 +8,7 @@
 #include "Misc/MessageDialog.h"
 #include "Logging/LogScopedVerbosityOverride.h"
 #include "UObject/FrameworkObjectVersion.h"
+#include "UObject/UE5MainStreamObjectVersion.h"
 #include "Serialization/MemoryReader.h"
 #include "UObject/UObjectIterator.h"
 #include "UObject/PropertyPortFlags.h"
@@ -551,35 +552,6 @@ float GetIntervalPerKey(int32 NumFrames, float SequenceLength)
 	return (NumFrames > 1) ? (SequenceLength / (NumFrames-1)) : MINIMUM_ANIMATION_LENGTH;
 }
 
-#if WITH_EDITOR
-// Handles keeping source raw data in sync when modifying raw data
-struct FModifyRawDataSourceGuard
-{
-private:
-	UAnimSequence* ModifyingSequence;
-
-public:
-	FModifyRawDataSourceGuard(UAnimSequence* AnimToModify)
-		: ModifyingSequence(nullptr)
-	{
-		check(AnimToModify)
-		if (AnimToModify->HasBakedTransformCurves())
-		{
-			ModifyingSequence = AnimToModify;
-			ModifyingSequence->RestoreSourceData();
-		}
-	}
-
-	~FModifyRawDataSourceGuard()
-	{
-		if (ModifyingSequence)
-		{
-			ModifyingSequence->BakeTrackCurvesToRawAnimation();
-		}
-	}
-};
-#endif
-
 /////////////////////////////////////////////////////
 // UAnimSequence
 
@@ -711,6 +683,7 @@ void UAnimSequence::Serialize(FArchive& Ar)
 	LLM_SCOPE(ELLMTag::Animation);
 	
 	Ar.UsingCustomVersion(FFrameworkObjectVersion::GUID);
+	Ar.UsingCustomVersion(FUE5MainStreamObjectVersion::GUID);
 
 	FRawCurveTracks RawCurveCache;
 
@@ -748,14 +721,21 @@ void UAnimSequence::Serialize(FArchive& Ar)
 		{
 			if (Ar.UE4Ver() >= VER_UE4_ANIMATION_ADD_TRACKCURVES)
 			{
-				Ar << SourceRawAnimationData;
-			}
-		}
+				if (Ar.CustomVer(FUE5MainStreamObjectVersion::GUID) < FUE5MainStreamObjectVersion::RemovingSourceAnimationData)
+				{
+					PRAGMA_DISABLE_DEPRECATION_WARNINGS
+					Ar << SourceRawAnimationData_DEPRECATED;
 
-		// If we have transform curves but no SourceRawAnimationData then we need to rebake
-		if (DoesContainTransformCurves() && RawAnimationData.Num() > 0 && SourceRawAnimationData.Num() == 0)
-		{
-			bNeedsRebake = true;
+					if (SourceRawAnimationData_DEPRECATED.Num())
+					{
+						// Set RawAnimationData to Source
+						RawAnimationData = SourceRawAnimationData_DEPRECATED;
+						SourceRawAnimationData_DEPRECATED.Empty();
+						MarkRawDataAsModified();
+					}
+					PRAGMA_ENABLE_DEPRECATION_WARNINGS
+				}
+			}
 		}
 #endif // WITH_EDITORONLY_DATA
 	}
@@ -866,19 +846,10 @@ void UAnimSequence::PreSave(const class ITargetPlatform* TargetPlatform)
 	// Could already be compressing
 	WaitOnExistingCompression();
 
-	// we have to bake it if it's not baked
-	if (DoesNeedRebake())
-	{
-		BakeTrackCurvesToRawAnimation();
-	}
-
 	// make sure if it does contain transform curvesm it contains source data
 	// empty track animation still can be made by retargeting to invalid skeleton
 	// make sure to not trigger ensure if RawAnimationData is also null
 	
-	// Why should we not be able to have empty transform curves?
-	ensureMsgf(!DoesContainTransformCurves() || (RawAnimationData.Num()==0 || SourceRawAnimationData.Num() != 0), TEXT("Animation : %s"), *GetName());
-
 	if (DoesNeedRecompress())
 	{
 		RequestSyncAnimRecompression(); // Update Normal data
@@ -1229,17 +1200,6 @@ void UAnimSequence::PostEditChangeProperty(FPropertyChangedEvent& PropertyChange
 		}
 	}
 }
-
-void UAnimSequence::PostDuplicate(bool bDuplicateForPIE)
-{
-	// if transform curve exists, mark as bake
-	if (DoesContainTransformCurves())
-	{
-		bNeedsRebake = true;
-	}
-
-	Super::PostDuplicate(bDuplicateForPIE);
-}
 #endif // WITH_EDITOR
 
 // @todo DB: Optimize!
@@ -1609,11 +1569,10 @@ void UAnimSequence::GetBonePose(FCompactPose& OutPose, FBlendedCurve& OutCurve,
 	// Slower path for disable retargeting, that's only used in editor and for debugging.
 	if (bUseRawDataForPoseExtraction)
 	{
-		const bool bShouldUseSourceData = (RequiredBones.ShouldUseSourceData() && SourceRawAnimationData.Num() > 0);
-		const TArray<FRawAnimSequenceTrack>& AnimationData = bShouldUseSourceData ? SourceRawAnimationData : RawAnimationData;
+		const bool bShouldUseSourceData = RequiredBones.ShouldUseSourceData();
+		const TArray<FRawAnimSequenceTrack>& AnimationData = RawAnimationData;
 
 		// Warning if we have invalid data
-
 		for (int32 TrackIndex = 0; TrackIndex < NumTracks; TrackIndex++)
 		{
 			const FRawAnimSequenceTrack& TrackToExtract = AnimationData[TrackIndex];
@@ -1625,7 +1584,26 @@ void UAnimSequence::GetBonePose(FCompactPose& OutPose, FBlendedCurve& OutCurve,
 			}
 		}
 
-		BuildPoseFromRawData(AnimationData, TrackToSkeletonMapTable, OutPose, ExtractionContext.CurrentTime, Interpolation, NumFrames, SequenceLength, RetargetSource);
+		TMap<int32, const FTransformCurve*> ActiveBoneCurves;
+		if (!bShouldUseSourceData)
+		{
+			const FSmartNameMapping* NameMapping = MySkeleton->GetSmartNameContainer(USkeleton::AnimTrackCurveMappingName);
+			for (const FTransformCurve& Curve : RawCurveData.TransformCurves)
+			{
+				FName BoneName;
+				if (Curve.GetCurveTypeFlag(AACF_Disabled) == false &&
+					ensureAlways(NameMapping->GetName(Curve.Name.UID, BoneName)))
+				{
+					const int32 BoneIndex = MySkeleton->GetReferenceSkeleton().FindBoneIndex(BoneName);
+					if (BoneIndex != INDEX_NONE)
+					{
+						ActiveBoneCurves.Add(BoneIndex, &Curve);
+					}
+				}
+			}
+		}
+
+		BuildPoseFromRawData(AnimationData, TrackToSkeletonMapTable, OutPose, ExtractionContext.CurrentTime, Interpolation, NumFrames, SequenceLength, RetargetSource, &ActiveBoneCurves);
 
 		if ((ExtractionContext.bExtractRootMotion && RootMotionReset.bEnableRootMotion) || RootMotionReset.bForceRootLock)
 		{
@@ -1967,10 +1945,6 @@ bool UAnimSequence::InsertFramesToRawAnimData( int32 StartFrame, int32 EndFrame,
 	int32 NumFramesToInsert = EndFrame-StartFrame;
 	if ((CopyFrame>=0 && CopyFrame<NumFrames) && (StartFrame >= 0 && StartFrame <=NumFrames) && NumFramesToInsert > 0)
 	{
-#if WITH_EDITOR
-		FModifyRawDataSourceGuard Modify(this);
-#endif
-
 		for (auto& RawData : RawAnimationData)
 		{
 			if (RawData.PosKeys.Num() > 1 && RawData.PosKeys.IsValidIndex(CopyFrame))
@@ -2069,11 +2043,6 @@ bool UAnimSequence::CropRawAnimData( float CurrentTime, bool bFromStart )
 		CropRawTrack(RawAnimationData[i], StartKey, NumKeys, TotalNumOfFrames);
 	}
 
-	for (int32 i = 0; i < SourceRawAnimationData.Num(); ++i)
-	{
-		CropRawTrack(SourceRawAnimationData[i], StartKey, NumKeys, TotalNumOfFrames);
-	}
-
 	// Double check that everything is fine
 	for(int32 i=0; i<RawAnimationData.Num(); i++)
 	{
@@ -2346,6 +2315,36 @@ void UAnimSequence::ApplyCompressedData(const TArray<uint8>& Data)
 	}
 }
 
+void UAnimSequence::SanitizeAnimationTrackData(FRawAnimSequenceTrack& RawAnim)
+{
+	for (auto ScaleIter = RawAnim.ScaleKeys.CreateIterator(); ScaleIter; ++ScaleIter)
+	{
+		FVector& Scale3D = *ScaleIter;
+		if (FMath::IsNearlyZero(Scale3D.X))
+		{
+			Scale3D.X = 0.f;
+		}
+		if (FMath::IsNearlyZero(Scale3D.Y))
+		{
+			Scale3D.Y = 0.f;
+		}
+		if (FMath::IsNearlyZero(Scale3D.Z))
+		{
+			Scale3D.Z = 0.f;
+		}
+	}
+
+	// make sure Rotation part is normalized before compress
+	for (auto RotIter = RawAnim.RotKeys.CreateIterator(); RotIter; ++RotIter)
+	{
+		FQuat& Rotation = *RotIter;
+		if (!Rotation.IsNormalized())
+		{
+			Rotation.Normalize();
+		}
+	}
+}
+
 void UAnimSequence::SerializeCompressedData(FArchive& Ar, bool bDDCData)
 {
 	if (!HasAnyFlags(RF_ClassDefaultObject))
@@ -2382,7 +2381,6 @@ bool UAnimSequence::DoesSequenceContainZeroScale()
 FGuid UAnimSequence::GenerateGuidFromRawData() const
 {
 	return GenerateGuidFromRawAnimData(RawAnimationData, RawCurveData);
-
 }
 
 void CopyTransformToRawAnimationData(const FTransform& BoneTransform, FRawAnimSequenceTrack& Track, int32 Frame)
@@ -2431,11 +2429,19 @@ void UAnimSequence::BakeOutVirtualBoneTracks(TArray<FRawAnimSequenceTrack>& NewR
 	const int32 NumVirtualBones = GetSkeleton()->GetVirtualBones().Num();
 	check( (RawAnimationData.Num() == TrackToSkeletonMapTable.Num()) && (RawAnimationData.Num() == AnimationTrackNames.Num()) ); //Make sure starting data is valid
 
-	NewRawTracks = TArray<FRawAnimSequenceTrack>(RawAnimationData, NumVirtualBones);
+	auto SetupData = [NumVirtualBones](auto& InOutData, const auto& BaseData)
+	{
+		if (InOutData.Num() == 0)
+		{
+			InOutData = BaseData;
+		}
 
-	NewTrackToSkeletonMapTable = TArray<FTrackToSkeletonMap>(TrackToSkeletonMapTable, NumVirtualBones);
+		InOutData.Reserve(InOutData.Num() + NumVirtualBones);
+	};
 
-	NewAnimationTrackNames = TArray<FName>(AnimationTrackNames, NumVirtualBones);
+	SetupData(NewRawTracks, RawAnimationData);
+	SetupData(NewTrackToSkeletonMapTable, TrackToSkeletonMapTable);
+	SetupData(NewAnimationTrackNames, AnimationTrackNames);
 
 	for (int32 VBIndex = 0; VBIndex < NumVirtualBones; ++VBIndex)
 	{
@@ -2729,7 +2735,6 @@ void UAnimSequence::RecycleAnimSequence()
 	RawDataGuid.Invalidate();
 	AnimationTrackNames.Empty();
 	TrackToSkeletonMapTable.Empty();
-	SourceRawAnimationData.Empty(0);
 	RawCurveData.Empty();
 
 	ClearCompressedBoneData();
@@ -2752,7 +2757,6 @@ void UAnimSequence::CleanAnimSequenceForImport()
 	TrackToSkeletonMapTable.Empty();
 	ClearCompressedBoneData();
 	ClearCompressedCurveData();
-	SourceRawAnimationData.Empty(0);
 }
 #endif // WITH_EDITOR
 
@@ -3511,33 +3515,7 @@ void UAnimSequence::PostProcessSequence(bool bForceNewRawDatGuid)
 	for (auto Iter = RawAnimationData.CreateIterator(); Iter; ++Iter)
 	{
 		FRawAnimSequenceTrack& RawAnim = (*Iter);
-
-		for (auto ScaleIter = RawAnim.ScaleKeys.CreateIterator(); ScaleIter; ++ScaleIter)
-		{
-			FVector& Scale3D = *ScaleIter;
-			if ( FMath::IsNearlyZero(Scale3D.X) )
-			{
-				Scale3D.X = 0.f;
-			}
-			if ( FMath::IsNearlyZero(Scale3D.Y) )
-			{
-				Scale3D.Y = 0.f;
-			}
-			if ( FMath::IsNearlyZero(Scale3D.Z) )
-			{
-				Scale3D.Z = 0.f;
-			}
-		}
-
-		// make sure Rotation part is normalized before compress
-		for(auto RotIter = RawAnim.RotKeys.CreateIterator(); RotIter; ++RotIter)
-		{
-			FQuat& Rotation = *RotIter;
-			if( !Rotation.IsNormalized() )
-			{
-				Rotation.Normalize();
-			}
-		}
+		SanitizeAnimationTrackData(RawAnim);
 	}
 
 	CompressRawAnimData();
@@ -3606,7 +3584,6 @@ void UAnimSequence::RemoveAllTracks()
 	RawAnimationData.Empty();
 	AnimationTrackNames.Empty();
 	TrackToSkeletonMapTable.Empty();
-	SourceRawAnimationData.Empty();
 	// clear all transform tracks
 	// not deleting curve names from skeleton 
 	// since we don't know if that name is used by other assets
@@ -3623,11 +3600,6 @@ void UAnimSequence::RemoveTrack(int32 TrackIndex)
 		RawAnimationData.RemoveAt(TrackIndex);
 		AnimationTrackNames.RemoveAt(TrackIndex);
 		TrackToSkeletonMapTable.RemoveAt(TrackIndex);
-		// source raw animation only exists if edited
-		if (SourceRawAnimationData.IsValidIndex(TrackIndex))
-		{
-			SourceRawAnimationData.RemoveAt(TrackIndex);
-		}
 
 		check (RawAnimationData.Num() == AnimationTrackNames.Num() && AnimationTrackNames.Num() == TrackToSkeletonMapTable.Num() );
 	}
@@ -3690,10 +3662,6 @@ int32 FindFirstChildTrack(const USkeleton* MySkeleton, const FReferenceSkeleton&
 
 int32 UAnimSequence::InsertTrack(const FName& BoneName)
 {
-#if WITH_EDITOR
-	FModifyRawDataSourceGuard Modify(this);
-#endif
-
 	// first verify if it doesn't exists, if it does, return
 	int32 CurrentTrackIndex = AnimationTrackNames.Find(BoneName);
 	if (CurrentTrackIndex != INDEX_NONE)
@@ -3723,7 +3691,6 @@ int32 UAnimSequence::InsertTrack(const FName& BoneName)
 		// now insert to the track
 		RawAnimationData.Insert(RawTrack, NewTrackIndex);
 		AnimationTrackNames.Insert(BoneName, NewTrackIndex);
-		SourceRawAnimationData.Insert(RawTrack, NewTrackIndex);
 
 		RefreshTrackMapFromAnimTrackNames();
 
@@ -4194,11 +4161,8 @@ bool UAnimSequence::ConvertRiggingDataToAnimationData(FAnimSequenceTrackContaine
 		RawAnimationData.Empty(ValidNumNodes);
 		RawAnimationData.AddZeroed(ValidNumNodes);
 
-		// if source animation exists, clear it, it won't matter anymore
-		if (SourceRawAnimationData.Num() > 0)
-		{
-			ClearBakedTransformData();
-		}
+		// Clear Transform curve data
+		RawCurveData.DeleteAllCurveData(ERawCurveTrackTypes::RCT_Transform);
 
 		const FReferenceSkeleton& RefSkeleton = MySkeleton->GetReferenceSkeleton();
 		const URig* Rig = MySkeleton->GetRig();
@@ -4272,52 +4236,33 @@ bool UAnimSequence::ConvertRiggingDataToAnimationData(FAnimSequenceTrackContaine
 void UAnimSequence::ClearBakedTransformData()
 {
 	UE_LOG(LogAnimation, Warning, TEXT("[%s] Detected previous edited data is invalidated. Clearing transform curve data and Source Data. This can happen if you do retarget another animation to this. If not, please report back to Epic. "), *GetName());
-	SourceRawAnimationData.Empty();
 	//Clear Transform curve data
 	RawCurveData.DeleteAllCurveData(ERawCurveTrackTypes::RCT_Transform);
 }
 
-void UAnimSequence::BakeTrackCurvesToRawAnimation()
+bool UAnimSequence::DoesContainTransformCurves() const
 {
-	// now bake the curves to the RawAnimationData
-	if(NumFrames == 0)
-	{
-		// fail error?
-		return;
-	}
+	return (RawCurveData.TransformCurves.Num() > 0);
+}
 
-	if (!DoesContainTransformCurves())
-	{
-		if (SourceRawAnimationData.Num() > 0)
-		{
-			// if curve doesn't exists, we just bring back Source to Raw, and clears Source
-			RawAnimationData = MoveTemp(SourceRawAnimationData);
-			PostProcessSequence();
-		}
-	}
-	else
-	{
-		if(SourceRawAnimationData.Num() == 0)
-		{
-			// if source data is empty, this is first time
-			// copies the data
-			SourceRawAnimationData = RawAnimationData;
-		}
-		else
-		{
-			// we copy SourceRawAnimationData because we'd need to create additive on top of current one
-			RawAnimationData = SourceRawAnimationData;
-		}
 
-		USkeleton * CurSkeleton = GetSkeleton();
+void UAnimSequence::BakeTrackCurvesToRawAnimationTracks(TArray<FRawAnimSequenceTrack>& NewRawTracks, TArray<FName>& NewTrackNames, TArray<FTrackToSkeletonMap>& NewTrackToSkeletonMapTable)
+{
+	if (DoesContainTransformCurves())
+	{
+		USkeleton* CurSkeleton = GetSkeleton();
 		check(CurSkeleton);
-		
+
+		NewRawTracks = RawAnimationData;
+		NewTrackToSkeletonMapTable = TrackToSkeletonMapTable;
+		NewTrackNames = AnimationTrackNames;
+			
 		VerifyCurveNames<FTransformCurve>(*CurSkeleton, USkeleton::AnimTrackCurveMappingName, RawCurveData.TransformCurves);
 		const FSmartNameMapping*  NameMapping = CurSkeleton->GetSmartNameContainer(USkeleton::AnimTrackCurveMappingName);
-		
+
 		// since now I'm about to modify Scale Keys. I should add all of them here at least one key. 
 		// if all turns out to be same, it will clear it up. 
-		for (auto & RawTrack: RawAnimationData)
+		for (FRawAnimSequenceTrack& RawTrack : NewRawTracks)
 		{
 			if (RawTrack.ScaleKeys.Num() == 0)
 			{
@@ -4327,22 +4272,49 @@ void UAnimSequence::BakeTrackCurvesToRawAnimation()
 			}
 		}
 
-		for(const auto& Curve : RawCurveData.TransformCurves)
+		auto InsertNewBoneTrack = [&NewRawTracks, &NewTrackToSkeletonMapTable, &NewTrackNames, CurSkeleton](const FName& BoneName) -> int32
+		{
+			const FReferenceSkeleton& RefSkeleton = CurSkeleton->GetReferenceSkeleton();
+			const int32 NewTrackIndex = FindFirstChildTrack(CurSkeleton, RefSkeleton, NewTrackNames, BoneName);
+			const int32 BoneIndex = RefSkeleton.FindBoneIndex(BoneName);
+			if (NewTrackIndex != INDEX_NONE)
+			{
+				const TArray<FTransform>& RefPose = RefSkeleton.GetRefBonePose();
+
+				FRawAnimSequenceTrack& RawTrack = NewRawTracks.InsertDefaulted_GetRef(NewTrackIndex);
+				RawTrack.PosKeys.Add(RefPose[BoneIndex].GetTranslation());
+				RawTrack.RotKeys.Add(RefPose[BoneIndex].GetRotation());
+				RawTrack.ScaleKeys.Add(RefPose[BoneIndex].GetScale3D());
+
+				NewTrackToSkeletonMapTable.Insert(BoneIndex, NewTrackIndex);
+				NewTrackNames.Insert(BoneName, NewTrackIndex);
+
+				check(NewRawTracks.Num() == NewTrackToSkeletonMapTable.Num());
+			}
+
+			return NewTrackIndex;
+		};
+
+
+		// NumFrames can't be zero (filtered earlier)
+		const float Interval = GetIntervalPerKey(NumFrames, SequenceLength);
+
+		for (const FTransformCurve& Curve : RawCurveData.TransformCurves)
 		{
 			// find curves first, and then see what is index of this curve
 			FName BoneName;
 
-			if(Curve.GetCurveTypeFlag(AACF_Disabled)== false &&
+			if (Curve.GetCurveTypeFlag(AACF_Disabled) == false &&
 				ensureAlways(NameMapping->GetName(Curve.Name.UID, BoneName)))
 			{
-				int32 TrackIndex = AnimationTrackNames.Find(BoneName);
+				int32 TrackIndex = NewTrackNames.Find(BoneName);
 
 				// the animation data doesn't have this track, so insert it
-				if(TrackIndex == INDEX_NONE)
+				if (TrackIndex == INDEX_NONE)
 				{
-					TrackIndex = InsertTrack(BoneName);
+					TrackIndex = InsertNewBoneTrack(BoneName);
 					// if it still didn't find, something went horribly wrong
-					if(ensure(TrackIndex != INDEX_NONE) == false)
+					if (ensure(TrackIndex != INDEX_NONE) == false)
 					{
 						UE_LOG(LogAnimation, Warning, TEXT("Animation Baking : Error adding %s track."), *BoneName.ToString());
 						// I can't do anything about it
@@ -4351,11 +4323,11 @@ void UAnimSequence::BakeTrackCurvesToRawAnimation()
 				}
 
 				// now modify data
-				auto& RawTrack = RawAnimationData[TrackIndex];
+				FRawAnimSequenceTrack& RawTrack = NewRawTracks[TrackIndex];
 
 				// since now we're editing keys, 
 				// if 1 (which meant constant), just expands to # of frames
-				if(RawTrack.PosKeys.Num() == 1)
+				if (RawTrack.PosKeys.Num() == 1)
 				{
 					FVector OneKey = RawTrack.PosKeys[0];
 					RawTrack.PosKeys.Init(OneKey, NumFrames);
@@ -4365,7 +4337,7 @@ void UAnimSequence::BakeTrackCurvesToRawAnimation()
 					ensure(RawTrack.PosKeys.Num() == NumFrames);
 				}
 
-				if(RawTrack.RotKeys.Num() == 1)
+				if (RawTrack.RotKeys.Num() == 1)
 				{
 					FQuat OneKey = RawTrack.RotKeys[0];
 					RawTrack.RotKeys.Init(OneKey, NumFrames);
@@ -4378,7 +4350,7 @@ void UAnimSequence::BakeTrackCurvesToRawAnimation()
 				// although we don't allow edit of scale
 				// it is important to consider scale when apply transform
 				// so make sure this also is included
-				if(RawTrack.ScaleKeys.Num() == 1)
+				if (RawTrack.ScaleKeys.Num() == 1)
 				{
 					FVector OneKey = RawTrack.ScaleKeys[0];
 					RawTrack.ScaleKeys.Init(OneKey, NumFrames);
@@ -4388,24 +4360,16 @@ void UAnimSequence::BakeTrackCurvesToRawAnimation()
 					ensure(RawTrack.ScaleKeys.Num() == NumFrames);
 				}
 
-				// NumFrames can't be zero (filtered earlier)
-				float Interval = GetIntervalPerKey(NumFrames, SequenceLength);
-
 				// now we have all data ready to apply
-				for(int32 KeyIndex=0; KeyIndex < NumFrames; ++KeyIndex)
+				for (int32 KeyIndex = 0; KeyIndex < NumFrames; ++KeyIndex)
 				{
 					// now evaluate
 					FTransformCurve* TransformCurve = static_cast<FTransformCurve*>(RawCurveData.GetCurveData(Curve.Name.UID, ERawCurveTrackTypes::RCT_Transform));
 
-					if(ensure(TransformCurve))
+					if (ensure(TransformCurve))
 					{
 						FTransform AdditiveTransform = TransformCurve->Evaluate(KeyIndex * Interval, 1.0);
 						FTransform LocalTransform(RawTrack.RotKeys[KeyIndex], RawTrack.PosKeys[KeyIndex], RawTrack.ScaleKeys[KeyIndex]);
-						//  						LocalTransform = LocalTransform * AdditiveTransform;
-						//  						RawTrack.RotKeys[KeyIndex] = LocalTransform.GetRotation();
-						//  						RawTrack.PosKeys[KeyIndex] = LocalTransform.GetTranslation();
-						//  						RawTrack.ScaleKeys[KeyIndex] = LocalTransform.GetScale3D();
-
 						RawTrack.RotKeys[KeyIndex] = LocalTransform.GetRotation() * AdditiveTransform.GetRotation();
 						RawTrack.PosKeys[KeyIndex] = LocalTransform.TransformPosition(AdditiveTransform.GetTranslation());
 						RawTrack.ScaleKeys[KeyIndex] = LocalTransform.GetScale3D() * AdditiveTransform.GetScale3D();
@@ -4415,52 +4379,17 @@ void UAnimSequence::BakeTrackCurvesToRawAnimation()
 						UE_LOG(LogAnimation, Warning, TEXT("Animation Baking : Missing Curve for %s."), *BoneName.ToString());
 					}
 				}
+
+				// Apply PostProcess behaviour in-place
+				SanitizeAnimationTrackData(RawTrack);
+				StaticCompressRawAnimData(NewRawTracks, NumFrames, GetFName());
 			}
 		}
-
-		PostProcessSequence();
 	}
-
-	bNeedsRebake = false;
 }
 
-bool UAnimSequence::DoesNeedRebake() const
-{
-	return (bNeedsRebake);
-}
-
-bool UAnimSequence::DoesContainTransformCurves() const
-{
-	return (RawCurveData.TransformCurves.Num() > 0);
-}
-
-#if WITH_EDITOR
-bool  UAnimSequence::HasBakedTransformCurves() const
-{
-	return DoesContainTransformCurves() && SourceRawAnimationData.Num() > 0;
-}
-
-void  UAnimSequence::RestoreSourceData() 
-{
-	if (HasBakedTransformCurves())
-	{
-		RawAnimationData = MoveTemp(SourceRawAnimationData);
-		bNeedsRebake = true;
-	}
-}
-#endif
-
 void UAnimSequence::AddKeyToSequence(float Time, const FName& BoneName, const FTransform& AdditiveTransform)
 {
-	// if source animation exists, but doesn't match with raw animation number, it's possible this has been retargetted
-	// or for any other reason, track has been modified. Just log here. 
-	if (SourceRawAnimationData.Num()>0 && SourceRawAnimationData.Num() != RawAnimationData.Num())
-	{
-		// currently it contains invalid data to edit
-		// clear and start over
-		ClearBakedTransformData();
-	}
-
 	// find if this already exists, then just add curve data only
 	FName CurveName = BoneName;
 	USkeleton * CurrentSkeleton = GetSkeleton();
@@ -4476,9 +4405,8 @@ void UAnimSequence::AddKeyToSequence(float Time, const FName& BoneName, const FT
 	FTransformCurve* TransformCurve = static_cast<FTransformCurve*>(RawCurveData.GetCurveData(NewCurveName.UID, ERawCurveTrackTypes::RCT_Transform));
 	check(TransformCurve);
 
-	TransformCurve->UpdateOrAddKey(AdditiveTransform, Time);	
-
-	bNeedsRebake = true;
+	TransformCurve->UpdateOrAddKey(AdditiveTransform, Time);
+	MarkRawDataAsModified();
 }
 
 void UAnimSequence::ResetAnimation()
@@ -4487,7 +4415,6 @@ void UAnimSequence::ResetAnimation()
 	NumFrames = 0;
 	SequenceLength = 0.f;
 	RawAnimationData.Empty();
-	SourceRawAnimationData.Empty();
 	AnimationTrackNames.Empty();
 	TrackToSkeletonMapTable.Empty();
 
@@ -4636,9 +4563,6 @@ bool UAnimSequence::CreateAnimation(UAnimSequence* Sequence)
 		Notifies = Sequence->Notifies;
 		AnimNotifyTracks = Sequence->AnimNotifyTracks;
 		RawCurveData = Sequence->RawCurveData;
-		// keep the same setting as source
-		bNeedsRebake = Sequence->DoesNeedRebake();
-		SourceRawAnimationData = Sequence->SourceRawAnimationData;
 
 		// refresh TrackToskeletonMapIndex
 		RefreshTrackMapFromAnimTrackNames();
diff --git a/Engine/Source/Runtime/Engine/Private/Animation/BonePose.cpp b/Engine/Source/Runtime/Engine/Private/Animation/BonePose.cpp
--- a/Engine/Source/Runtime/Engine/Private/Animation/BonePose.cpp
+++ b/Engine/Source/Runtime/Engine/Private/Animation/BonePose.cpp
@@ -161,7 +161,7 @@ struct FBuildRawPoseScratchArea : public TThreadSingleton<FBuildRawPoseScratchAr
 
 
 template<bool bInterpolateT>
-void BuildPoseFromRawDataInternal(const TArray<FRawAnimSequenceTrack>& InAnimationData, const TArray<struct FTrackToSkeletonMap>& TrackToSkeletonMapTable, FCompactPose& InOutPose, int32 KeyIndex1, int32 KeyIndex2, float Alpha)
+void BuildPoseFromRawDataInternal(const TArray<FRawAnimSequenceTrack>& InAnimationData, const TArray<struct FTrackToSkeletonMap>& TrackToSkeletonMapTable, FCompactPose& InOutPose, int32 KeyIndex1, int32 KeyIndex2, float Alpha, float TimePerKey, const TMap<int32, const FTransformCurve*>* AdditiveBoneTransformCurves)
 {
 	const int32 NumTracks = InAnimationData.Num();
 	const FBoneContainer& RequiredBones = InOutPose.GetBoneContainer();
@@ -201,6 +201,8 @@ void BuildPoseFromRawDataInternal(const TArray<FRawAnimSequenceTrack>& InAnimati
 
 				const FRawAnimSequenceTrack& TrackToExtract = InAnimationData[TrackIndex];
 
+				const FTransformCurve* const * AdditiveBoneTransformCurve = AdditiveBoneTransformCurves ? AdditiveBoneTransformCurves->Find(SkeletonBoneIndex) : nullptr;
+
 				// Bail out (with rather wacky data) if data is empty for some reason.
 				if (TrackToExtract.PosKeys.Num() == 0 || TrackToExtract.RotKeys.Num() == 0)
 				{
@@ -219,6 +221,23 @@ void BuildPoseFromRawDataInternal(const TArray<FRawAnimSequenceTrack>& InAnimati
 					{
 						Key2Pose[PoseBoneIndex] = ExtractTransformForKey(KeyIndex2, TrackToExtract);
 					}
+
+
+					if (AdditiveBoneTransformCurve)
+					{
+						const FTransform PoseOneAdditive = (*AdditiveBoneTransformCurve)->Evaluate(KeyIndex1 * TimePerKey, 1.f);
+						InOutPose[PoseBoneIndex].SetRotation(InOutPose[PoseBoneIndex].GetRotation() * PoseOneAdditive.GetRotation());
+						InOutPose[PoseBoneIndex].SetTranslation(InOutPose[PoseBoneIndex].TransformPosition(PoseOneAdditive.GetTranslation()));
+						InOutPose[PoseBoneIndex].SetScale3D(InOutPose[PoseBoneIndex].GetScale3D() * PoseOneAdditive.GetScale3D());
+
+						if (bInterpolateT)
+						{
+							const FTransform PoseTwoAdditive = (*AdditiveBoneTransformCurve)->Evaluate(KeyIndex2 * TimePerKey, 1.f);
+							Key2Pose[PoseBoneIndex].SetRotation(Key2Pose[PoseBoneIndex].GetRotation() * PoseTwoAdditive.GetRotation());
+							Key2Pose[PoseBoneIndex].SetTranslation(Key2Pose[PoseBoneIndex].TransformPosition(PoseTwoAdditive.GetTranslation()));
+							Key2Pose[PoseBoneIndex].SetScale3D(Key2Pose[PoseBoneIndex].GetScale3D() * PoseTwoAdditive.GetScale3D());
+						}
+					}
 				}
 
 				RetargetTracking.Add(FRetargetTracking(PoseBoneIndex, SkeletonBoneIndex));
@@ -262,7 +281,7 @@ void BuildPoseFromRawDataInternal(const TArray<FRawAnimSequenceTrack>& InAnimati
 	}
 }
 
-void BuildPoseFromRawData(const TArray<FRawAnimSequenceTrack>& InAnimationData, const TArray<struct FTrackToSkeletonMap>& TrackToSkeletonMapTable, FCompactPose& InOutPose, float InTime, EAnimInterpolationType Interpolation, int32 NumFrames, float SequenceLength, FName RetargetSource)
+void BuildPoseFromRawData(const TArray<FRawAnimSequenceTrack>& InAnimationData, const TArray<struct FTrackToSkeletonMap>& TrackToSkeletonMapTable, FCompactPose& InOutPose, float InTime, EAnimInterpolationType Interpolation, int32 NumFrames, float SequenceLength, FName RetargetSource, const TMap<int32, const FTransformCurve*>* AdditiveBoneTransformCurves /*= nullptr*/)
 {
 	int32 KeyIndex1, KeyIndex2;
 	float Alpha;
@@ -286,13 +305,15 @@ void BuildPoseFromRawData(const TArray<FRawAnimSequenceTrack>& InAnimationData,
 		KeyIndex1 = KeyIndex2;
 	}
 
+	const float TimePerFrame = SequenceLength / (float)FMath::Max(NumFrames - 1, 1);
+
 	if (bInterpolate)
 	{
-		BuildPoseFromRawDataInternal<true>(InAnimationData, TrackToSkeletonMapTable, InOutPose, KeyIndex1, KeyIndex2, Alpha);
+		BuildPoseFromRawDataInternal<true>(InAnimationData, TrackToSkeletonMapTable, InOutPose, KeyIndex1, KeyIndex2, Alpha, TimePerFrame, AdditiveBoneTransformCurves);
 	}
 	else
 	{
-		BuildPoseFromRawDataInternal<false>(InAnimationData, TrackToSkeletonMapTable, InOutPose, KeyIndex1, KeyIndex2, Alpha);
+		BuildPoseFromRawDataInternal<false>(InAnimationData, TrackToSkeletonMapTable, InOutPose, KeyIndex1, KeyIndex2, Alpha, TimePerFrame, AdditiveBoneTransformCurves);
 	}
 
 	const FBoneContainer& RequiredBones = InOutPose.GetBoneContainer();
@@ -309,7 +330,6 @@ void BuildPoseFromRawData(const TArray<FRawAnimSequenceTrack>& InAnimationData,
 			FAnimationRuntime::RetargetBoneTransform(Skeleton, RetargetSource, InOutPose[RT.PoseBoneIndex], RT.SkeletonBoneIndex, RT.PoseBoneIndex, RequiredBones, false);
 		}
 	}
-
 }
 
 
diff --git a/Engine/Source/Runtime/Engine/Public/BonePose.h b/Engine/Source/Runtime/Engine/Public/BonePose.h
--- a/Engine/Source/Runtime/Engine/Public/BonePose.h
+++ b/Engine/Source/Runtime/Engine/Public/BonePose.h
@@ -877,4 +877,4 @@ void FCSPose<PoseType>::ConvertComponentPosesToLocalPoses(FCSPose<PoseType>&& In
 }
 
 // Populate InOutPose based on raw animation data. 
-extern void BuildPoseFromRawData(const TArray<FRawAnimSequenceTrack>& InAnimationData, const TArray<struct FTrackToSkeletonMap>& TrackToSkeletonMapTable, FCompactPose& InOutPose, float InTime, EAnimInterpolationType Interpolation, int32 NumFrames, float SequenceLength, FName RetargetSource);
\ No newline at end of file
+extern void BuildPoseFromRawData(const TArray<FRawAnimSequenceTrack>& InAnimationData, const TArray<struct FTrackToSkeletonMap>& TrackToSkeletonMapTable, FCompactPose& InOutPose, float InTime, EAnimInterpolationType Interpolation, int32 NumFrames, float SequenceLength, FName RetargetSource, const TMap<int32, const FTransformCurve*>* AdditiveBoneTransformCurves = nullptr);
\ No newline at end of file
