Non-pch CIS fix

diff --git a/Engine/Source/Runtime/Engine/Private/Animation/AnimCompressionTypes.cpp b/Engine/Source/Runtime/Engine/Private/Animation/AnimCompressionTypes.cpp
--- a/Engine/Source/Runtime/Engine/Private/Animation/AnimCompressionTypes.cpp
+++ b/Engine/Source/Runtime/Engine/Private/Animation/AnimCompressionTypes.cpp
@@ -52,6 +52,7 @@ void UpdateSHAWithCurves(FSHA1& Sha, const FRawCurveTracks& InRawCurveData)
 		UpdateWithFloatCurve(Curve.FloatCurve);
 	}
 
+#if WITH_EDITOR
 	for (const FTransformCurve& Curve : InRawCurveData.TransformCurves)
 	{
 		UpdateWithData(Sha, Curve.Name.UID);
@@ -68,6 +69,7 @@ void UpdateSHAWithCurves(FSHA1& Sha, const FRawCurveTracks& InRawCurveData)
 		UpdateWithComponent(Curve.RotationCurve);
 		UpdateWithComponent(Curve.ScaleCurve);
 	}
+#endif // WITH_EDITOR
 }
 
 FGuid GenerateGuidFromRawAnimData(const TArray<FRawAnimSequenceTrack>& RawAnimationData, const FRawCurveTracks& RawCurveData)
diff --git a/Engine/Source/Runtime/Engine/Public/BonePose.h b/Engine/Source/Runtime/Engine/Public/BonePose.h
--- a/Engine/Source/Runtime/Engine/Public/BonePose.h
+++ b/Engine/Source/Runtime/Engine/Public/BonePose.h
@@ -877,4 +877,4 @@ void FCSPose<PoseType>::ConvertComponentPosesToLocalPoses(FCSPose<PoseType>&& In
 }
 
 // Populate InOutPose based on raw animation data. 
-extern void BuildPoseFromRawData(const TArray<FRawAnimSequenceTrack>& InAnimationData, const TArray<struct FTrackToSkeletonMap>& TrackToSkeletonMapTable, FCompactPose& InOutPose, float InTime, EAnimInterpolationType Interpolation, int32 NumFrames, float SequenceLength, FName RetargetSource, const TMap<int32, const FTransformCurve*>* AdditiveBoneTransformCurves = nullptr);
\ No newline at end of file
+extern void BuildPoseFromRawData(const TArray<FRawAnimSequenceTrack>& InAnimationData, const TArray<struct FTrackToSkeletonMap>& TrackToSkeletonMapTable, FCompactPose& InOutPose, float InTime, EAnimInterpolationType Interpolation, int32 NumFrames, float SequenceLength, FName RetargetSource, const TMap<int32, const struct FTransformCurve*>* AdditiveBoneTransformCurves = nullptr);
\ No newline at end of file
