Momentum: Deprecating public access to SequenceLength
    + Added SetSequenceLength
    + Moved GetPlayLength virtaul from AnimSequenceBase to AnimationAsset and deprecated GetMaxCurrentTime (only used in a single place)

diff --git a/Engine/Plugins/Developer/AnimationSharing/Source/AnimationSharing/Private/AnimationSharingManager.cpp b/Engine/Plugins/Developer/AnimationSharing/Source/AnimationSharing/Private/AnimationSharingManager.cpp
--- a/Engine/Plugins/Developer/AnimationSharing/Source/AnimationSharing/Private/AnimationSharingManager.cpp
+++ b/Engine/Plugins/Developer/AnimationSharing/Source/AnimationSharing/Private/AnimationSharingManager.cpp
@@ -934,7 +934,7 @@ void UAnimSharingInstance::SetupState(FPerStateData& StateData, const FAnimation
 							AnimInstance->AnimationToPlay = AnimSequence;
 							if (InstanceIndex > 0)
 							{
-								const float Steps = (AnimSequence->SequenceLength * 0.9f) / (NumInstances);
+								const float Steps = (AnimSequence->GetPlayLength() * 0.9f) / (NumInstances);
 								const float StartTimeOffset = Steps * InstanceIndex;
 								AnimInstance->PermutationTimeOffset = StartTimeOffset;
 							}
@@ -946,7 +946,7 @@ void UAnimSharingInstance::SetupState(FPerStateData& StateData, const FAnimation
 							AnimInstance->ComponentIndex = Components.Num();
 
 							/** Set the current animation length length */
-							StateData.AnimationLengths.Add(AnimSequence->SequenceLength);
+							StateData.AnimationLengths.Add(AnimSequence->GetPlayLength());
 						}
 					}
 					else if (AnimSequence != nullptr)
@@ -962,14 +962,14 @@ void UAnimSharingInstance::SetupState(FPerStateData& StateData, const FAnimation
 						{
 							if (InstanceIndex > 0)
 							{
-								const float Steps = (AnimSequence->SequenceLength * 0.9f) / (NumInstances);
+								const float Steps = (AnimSequence->GetPlayLength() * 0.9f) / (NumInstances);
 								const float StartTimeOffset = Steps * InstanceIndex;
 								Component->SetPosition(StartTimeOffset, false);
 							}
 						}
 
 						/** Set the current animation length length */
-						StateData.AnimationLengths.Add(AnimSequence->SequenceLength);
+						StateData.AnimationLengths.Add(AnimSequence->GetPlayLength());
 					}
 
 					/** Set material to red to indicate that it's not in use*/
@@ -995,7 +995,7 @@ void UAnimSharingInstance::SetupState(FPerStateData& StateData, const FAnimation
 					AdditiveInstanceStack.AddInstance(AdditiveInstance);
 					
 					/** Set the current animation length length */
-					StateData.AnimationLengths.Add(AnimSequence->SequenceLength);
+					StateData.AnimationLengths.Add(AnimSequence->GetPlayLength());
 					Components.Add(AdditiveComponent);
 				}
 				
@@ -2034,7 +2034,7 @@ uint32 UAnimSharingInstance::SetupAdditiveInstance(uint8 StateIndex, uint8 FromS
 		Instance.AdditiveAnimationInstance = AnimationInstance;
 		Instance.BaseComponent = PerStateData[FromState].Components[StateComponentIndex];
 		const float WorldTimeSeconds = GetWorld()->GetTimeSeconds();
-		Instance.EndTime = WorldTimeSeconds + StateData.AdditiveAnimationSequence->SequenceLength;
+		Instance.EndTime = WorldTimeSeconds + StateData.AdditiveAnimationSequence->GetPlayLength();
 		Instance.State = StateIndex;
 		Instance.UsedPerStateComponentIndex = PerStateData[StateIndex].Components.IndexOfByKey(AnimationInstance->GetComponent());
 
diff --git a/Engine/Plugins/Experimental/AlembicImporter/Source/AlembicLibrary/Private/AbcImporter.cpp b/Engine/Plugins/Experimental/AlembicImporter/Source/AlembicLibrary/Private/AbcImporter.cpp
--- a/Engine/Plugins/Experimental/AlembicImporter/Source/AlembicLibrary/Private/AbcImporter.cpp
+++ b/Engine/Plugins/Experimental/AlembicImporter/Source/AlembicLibrary/Private/AbcImporter.cpp
@@ -637,7 +637,7 @@ TArray<UObject*> FAbcImporter::ImportAsSkeletalMesh(UObject* InParent, EObjectFl
 		// Create animation sequence for the skeleton
 		UAnimSequence* Sequence = CreateObjectInstance<UAnimSequence>(InParent, FString::Printf(TEXT("%s_Animation"), *SkeletalMesh->GetName()), Flags);
 		Sequence->SetSkeleton(Skeleton);
-		Sequence->SequenceLength = AbcFile->GetImportLength();
+		Sequence->SetSequenceLength(AbcFile->GetImportLength());
 		Sequence->ImportFileFramerate = AbcFile->GetFramerate();
 		Sequence->ImportResampleFramerate = AbcFile->GetFramerate();
 
diff --git a/Engine/Plugins/Runtime/GameplayAbilities/Source/GameplayAbilities/Private/AbilitySystemComponent_Abilities.cpp b/Engine/Plugins/Runtime/GameplayAbilities/Source/GameplayAbilities/Private/AbilitySystemComponent_Abilities.cpp
--- a/Engine/Plugins/Runtime/GameplayAbilities/Source/GameplayAbilities/Private/AbilitySystemComponent_Abilities.cpp
+++ b/Engine/Plugins/Runtime/GameplayAbilities/Source/GameplayAbilities/Private/AbilitySystemComponent_Abilities.cpp
@@ -2942,12 +2942,12 @@ float UAbilitySystemComponent::GetCurrentMontageSectionLength() const
 			// Otherwise we are the last section, so take delta with Montage total time.
 			else
 			{
-				return (CurrentAnimMontage->SequenceLength - CompositeSections[CurrentSectionID].GetTime());
+				return (CurrentAnimMontage->GetPlayLength() - CompositeSections[CurrentSectionID].GetTime());
 			}
 		}
 
 		// if we have no sections, just return total length of Montage.
-		return CurrentAnimMontage->SequenceLength;
+		return CurrentAnimMontage->GetPlayLength();
 	}
 
 	return 0.f;
diff --git a/Engine/Plugins/VirtualProduction/Takes/Source/TakeTrackRecorders/Private/MovieSceneAnimationTrackRecorder.cpp b/Engine/Plugins/VirtualProduction/Takes/Source/TakeTrackRecorders/Private/MovieSceneAnimationTrackRecorder.cpp
--- a/Engine/Plugins/VirtualProduction/Takes/Source/TakeTrackRecorders/Private/MovieSceneAnimationTrackRecorder.cpp
+++ b/Engine/Plugins/VirtualProduction/Takes/Source/TakeTrackRecorders/Private/MovieSceneAnimationTrackRecorder.cpp
@@ -358,7 +358,7 @@ bool UMovieSceneAnimationTrackRecorder::LoadRecordedFile(const FString& FileName
 								CreateAnimationAssetAndSequence(Actor,AnimationDirectory);
 
 								AnimSequence->RecycleAnimSequence();
-								AnimSequence->SequenceLength = 0.f;
+								AnimSequence->SetSequenceLength(0.f);
 								AnimSequence->SetRawNumberOfFrame(0);
 								for (int32 TrackIndex = 0; TrackIndex < Header.AnimationTrackNames.Num(); ++TrackIndex)
 								{
@@ -379,7 +379,7 @@ bool UMovieSceneAnimationTrackRecorder::LoadRecordedFile(const FString& FileName
 									}
 								}
 								AnimSequence->SetRawNumberOfFrame( InFrames.Num() );
-								AnimSequence->SequenceLength = (AnimSequence->GetRawNumberOfFrames() > 1) ? (AnimSequence->GetRawNumberOfFrames() - 1) * Header.IntervalTime : MINIMUM_ANIMATION_LENGTH;
+								AnimSequence->SetSequenceLength((AnimSequence->GetRawNumberOfFrames() > 1) ? (AnimSequence->GetRawNumberOfFrames() - 1) * Header.IntervalTime : MINIMUM_ANIMATION_LENGTH);
 
 								//fix up notifies todo notifies mz
 								AnimSequence->PostProcessSequence();
diff --git a/Engine/Source/Editor/AnimGraph/Private/AnimPreviewInstance.cpp b/Engine/Source/Editor/AnimGraph/Private/AnimPreviewInstance.cpp
--- a/Engine/Source/Editor/AnimGraph/Private/AnimPreviewInstance.cpp
+++ b/Engine/Source/Editor/AnimGraph/Private/AnimPreviewInstance.cpp
@@ -706,9 +706,9 @@ void UAnimPreviewInstance::MontagePreview_StepForward()
 		// Advance a single frame, leaving it paused afterwards
 		int32 NumFrames = Montage->GetNumberOfFrames();
 		// Add DELTA to prefer next frame when we're close to the boundary
-		float CurrentFraction = Proxy.GetCurrentTime() / Montage->SequenceLength + DELTA;
+		float CurrentFraction = Proxy.GetCurrentTime() / Montage->GetPlayLength() + DELTA;
 		float NextFrame = FMath::Clamp<float>(FMath::FloorToFloat(CurrentFraction * NumFrames) + 1.0f, 0, NumFrames);
-		float NewTime = Montage->SequenceLength * (NextFrame / (NumFrames-1));
+		float NewTime = Montage->GetPlayLength() * (NextFrame / (NumFrames-1));
 
 		GetSkelMeshComponent()->GlobalAnimRateScale = 1.0f;
 		GetSkelMeshComponent()->TickAnimation(NewTime - Proxy.GetCurrentTime(), false);
@@ -762,9 +762,9 @@ void UAnimPreviewInstance::MontagePreview_StepBackward()
 		// Advance a single frame, leaving it paused afterwards
 		int32 NumFrames = Montage->GetNumberOfFrames();
 		// Add DELTA to prefer next frame when we're close to the boundary
-		float CurrentFraction = Proxy.GetCurrentTime() / Montage->SequenceLength + DELTA;
+		float CurrentFraction = Proxy.GetCurrentTime() / Montage->GetPlayLength() + DELTA;
 		float NextFrame = FMath::Clamp<float>(FMath::FloorToFloat(CurrentFraction * NumFrames) - 1.0f, 0, NumFrames);
-		float NewTime = Montage->SequenceLength * (NextFrame / (NumFrames-1));
+		float NewTime = Montage->GetPlayLength() * (NextFrame / (NumFrames-1));
 
 		GetSkelMeshComponent()->GlobalAnimRateScale = 1.0f;
 		GetSkelMeshComponent()->TickAnimation(FMath::Abs(NewTime - Proxy.GetCurrentTime()), false);
@@ -897,7 +897,7 @@ void UAnimPreviewInstance::MontagePreview_RemoveBlendOut()
 void UAnimPreviewInstance::MontagePreview_PreviewNormal(int32 FromSectionIdx, bool bPlay)
 {
 	UAnimMontage* Montage = Cast<UAnimMontage>(CurrentAsset);
-	if (Montage && Montage->SequenceLength > 0.0f)
+	if (Montage && Montage->GetPlayLength() > 0.0f)
 	{
 		FAnimPreviewInstanceProxy& Proxy = GetProxyOnGameThread<FAnimPreviewInstanceProxy>();
 
@@ -933,7 +933,7 @@ void UAnimPreviewInstance::MontagePreview_PreviewNormal(int32 FromSectionIdx, bo
 void UAnimPreviewInstance::MontagePreview_PreviewAllSections(bool bPlay)
 {
 	UAnimMontage* Montage = Cast<UAnimMontage>(CurrentAsset);
-	if (Montage && Montage->SequenceLength > 0.0f)
+	if (Montage && Montage->GetPlayLength() > 0.0f)
 	{
 		FAnimPreviewInstanceProxy& Proxy = GetProxyOnGameThread<FAnimPreviewInstanceProxy>();
 
diff --git a/Engine/Source/Editor/AnimationBlueprintEditor/Private/AnimationNodes/SGraphNodeSequencePlayer.cpp b/Engine/Source/Editor/AnimationBlueprintEditor/Private/AnimationNodes/SGraphNodeSequencePlayer.cpp
--- a/Engine/Source/Editor/AnimationBlueprintEditor/Private/AnimationNodes/SGraphNodeSequencePlayer.cpp
+++ b/Engine/Source/Editor/AnimationBlueprintEditor/Private/AnimationNodes/SGraphNodeSequencePlayer.cpp
@@ -170,7 +170,7 @@ void SGraphNodeSequencePlayer::SetSequencePositionRatio(float NewRatio)
 	{
 		if (SequencePlayer->Sequence != NULL)
 		{
-			const float NewTime = NewRatio * SequencePlayer->Sequence->SequenceLength;
+			const float NewTime = NewRatio * SequencePlayer->Sequence->GetPlayLength();
 			SequencePlayer->SetAccumulatedTime(NewTime);
 		}
 	}
diff --git a/Engine/Source/Editor/AnimationModifiers/Private/AnimationBlueprintLibrary.cpp b/Engine/Source/Editor/AnimationModifiers/Private/AnimationBlueprintLibrary.cpp
--- a/Engine/Source/Editor/AnimationModifiers/Private/AnimationBlueprintLibrary.cpp
+++ b/Engine/Source/Editor/AnimationModifiers/Private/AnimationBlueprintLibrary.cpp
@@ -1955,7 +1955,7 @@ void UAnimationBlueprintLibrary::GetSequenceLength(const UAnimSequence* Animatio
 	Length = 0.0f;
 	if (AnimationSequence)
 	{
-		Length = AnimationSequence->SequenceLength;
+		Length = AnimationSequence->GetPlayLength();
 	}
 	else
 	{
@@ -2034,7 +2034,7 @@ void UAnimationBlueprintLibrary::IsValidTime(const UAnimSequence* AnimationSeque
 
 bool UAnimationBlueprintLibrary::IsValidTimeInternal(const UAnimSequence* AnimationSequence, const float Time)
 {
-	return FMath::IsWithinInclusive(Time, 0.0f, AnimationSequence->SequenceLength);
+	return FMath::IsWithinInclusive(Time, 0.0f, AnimationSequence->GetPlayLength());
 }
 
 void UAnimationBlueprintLibrary::FindBonePathToRoot(const UAnimSequence* AnimationSequence, FName BoneName, TArray<FName>& BonePath)
diff --git a/Engine/Source/Editor/DetailCustomizations/Private/AnimMontageSegmentDetails.cpp b/Engine/Source/Editor/DetailCustomizations/Private/AnimMontageSegmentDetails.cpp
--- a/Engine/Source/Editor/DetailCustomizations/Private/AnimMontageSegmentDetails.cpp
+++ b/Engine/Source/Editor/DetailCustomizations/Private/AnimMontageSegmentDetails.cpp
@@ -564,7 +564,7 @@ float SAnimationSegmentScrubPanel::GetSequenceLength() const
 	}
 	else if (LockedSequence)
 	{
-		return LockedSequence->SequenceLength;
+		return LockedSequence->GetPlayLength();
 	}
 	return 0.f;
 }
diff --git a/Engine/Source/Editor/DetailCustomizations/Private/AnimSequenceDetails.cpp b/Engine/Source/Editor/DetailCustomizations/Private/AnimSequenceDetails.cpp
--- a/Engine/Source/Editor/DetailCustomizations/Private/AnimSequenceDetails.cpp
+++ b/Engine/Source/Editor/DetailCustomizations/Private/AnimSequenceDetails.cpp
@@ -486,7 +486,7 @@ void SAnimationRefPoseViewport::Tick( const FGeometry& AllottedGeometry, const d
 					int RefFrameIndex;
 					RefFrameIndexPropertyHandle->GetValue( RefFrameIndex );
 					float Fraction = ( AnimRef->GetRawNumberOfFrames() > 0 ) ? FMath::Clamp<float>( (float)RefFrameIndex / (float)AnimRef->GetRawNumberOfFrames(), 0.f, 1.f ) : 0.f;
-					float RefTime = AnimRef->SequenceLength * Fraction;
+					float RefTime = AnimRef->GetPlayLength() * Fraction;
 					PreviewComponent->PreviewInstance->SetPosition( RefTime, false );
 					PreviewComponent->PreviewInstance->SetPlaying( false );
 					LevelViewportClient->Invalidate();
@@ -571,7 +571,7 @@ TArray<float> SAnimationRefPoseViewport::GetBars() const
 		int RefFrameIndex;
 		RefFrameIndexPropertyHandle->GetValue(RefFrameIndex);
 		float Fraction = (AnimRef->GetRawNumberOfFrames() > 0)? FMath::Clamp<float>((float)RefFrameIndex/(float)AnimRef->GetRawNumberOfFrames(), 0.f, 1.f) : 0.f;
-		Bars.Add(AnimRef->SequenceLength * Fraction);
+		Bars.Add(AnimRef->GetPlayLength() * Fraction);
 	}
 	else
 	{
@@ -584,7 +584,7 @@ void SAnimationRefPoseViewport::OnBarDrag(int32 Index, float Position)
 {
 	if (AnimRef)
 	{
-		int RefFrameIndex = FMath::Clamp<int>(AnimRef->SequenceLength > 0.0f? (int)(Position * (float)AnimRef->GetRawNumberOfFrames() / AnimRef->SequenceLength + 0.5f) : 0.0f, 0, AnimRef->GetRawNumberOfFrames() - 1);
+		int RefFrameIndex = FMath::Clamp<int>(AnimRef->GetPlayLength() > 0.0f? (int)(Position * (float)AnimRef->GetRawNumberOfFrames() / AnimRef->GetPlayLength() + 0.5f) : 0.0f, 0, AnimRef->GetRawNumberOfFrames() - 1);
 		RefFrameIndexPropertyHandle->SetValue(RefFrameIndex);
 	}
 }
diff --git a/Engine/Source/Editor/Matinee/Private/MatineeMenus.cpp b/Engine/Source/Editor/Matinee/Private/MatineeMenus.cpp
--- a/Engine/Source/Editor/Matinee/Private/MatineeMenus.cpp
+++ b/Engine/Source/Editor/Matinee/Private/MatineeMenus.cpp
@@ -4231,7 +4231,7 @@ void FMatinee::OnExportAnimationInfoCommand()
 							
 							float AnimStartTime = Key.AnimStartOffset;
 							
-							float AnimEndTime = (Seq) ? (Seq->SequenceLength - Key.AnimEndOffset) : 0.0f;
+							float AnimEndTime = (Seq) ? (Seq->GetPlayLength() - Key.AnimEndOffset) : 0.0f;
 					
 							float AnimPlayRate = Key.AnimPlayRate;
 							bool bLooping = Key.bLooping;
diff --git a/Engine/Source/Editor/MovieSceneTools/Private/MatineeImportTools.cpp b/Engine/Source/Editor/MovieSceneTools/Private/MatineeImportTools.cpp
--- a/Engine/Source/Editor/MovieSceneTools/Private/MatineeImportTools.cpp
+++ b/Engine/Source/Editor/MovieSceneTools/Private/MatineeImportTools.cpp
@@ -729,7 +729,7 @@ bool FMatineeImportTools::CopyInterpAnimControlTrack( UInterpTrackAnimControl* M
 		}
 		else
 		{
-			EndTime = AnimSeq.StartTime + ( ( ( AnimSeq.AnimSeq->SequenceLength - AnimSeq.AnimEndOffset ) - AnimSeq.AnimStartOffset ) / AnimSeq.AnimPlayRate );
+			EndTime = AnimSeq.StartTime + ( ( ( AnimSeq.AnimSeq->GetPlayLength() - AnimSeq.AnimEndOffset ) - AnimSeq.AnimStartOffset ) / AnimSeq.AnimPlayRate );
 
 			// Clamp to next clip's start time
 			if (i+1 < MatineeAnimControlTrack->AnimSeqs.Num())
diff --git a/Engine/Source/Editor/Persona/Private/AnimTimeline/AnimModel.cpp b/Engine/Source/Editor/Persona/Private/AnimTimeline/AnimModel.cpp
--- a/Engine/Source/Editor/Persona/Private/AnimTimeline/AnimModel.cpp
+++ b/Engine/Source/Editor/Persona/Private/AnimTimeline/AnimModel.cpp
@@ -190,7 +190,7 @@ float FAnimModel::CalculateSequenceLengthOfEditorObject() const
 {
 	if(UAnimSequenceBase* AnimSequenceBase = GetAnimSequenceBase())
 	{
-		return AnimSequenceBase->SequenceLength;
+		return AnimSequenceBase->GetPlayLength();
 	}
 
 	return 0.0f;
diff --git a/Engine/Source/Editor/Persona/Private/AnimTimeline/AnimModel_AnimComposite.cpp b/Engine/Source/Editor/Persona/Private/AnimTimeline/AnimModel_AnimComposite.cpp
--- a/Engine/Source/Editor/Persona/Private/AnimTimeline/AnimModel_AnimComposite.cpp
+++ b/Engine/Source/Editor/Persona/Private/AnimTimeline/AnimModel_AnimComposite.cpp
@@ -74,7 +74,7 @@ void FAnimModel_AnimComposite::RecalculateSequenceLength()
 		AnimComposite->InvalidateRecursiveAsset();
 
 		float NewSequenceLength = CalculateSequenceLengthOfEditorObject();
-		if (NewSequenceLength != AnimComposite->SequenceLength)
+		if (NewSequenceLength != AnimComposite->GetPlayLength())
 		{
 			ClampToEndTime(NewSequenceLength);
 
diff --git a/Engine/Source/Editor/Persona/Private/AnimTimeline/AnimModel_AnimMontage.cpp b/Engine/Source/Editor/Persona/Private/AnimTimeline/AnimModel_AnimMontage.cpp
--- a/Engine/Source/Editor/Persona/Private/AnimTimeline/AnimModel_AnimMontage.cpp
+++ b/Engine/Source/Editor/Persona/Private/AnimTimeline/AnimModel_AnimMontage.cpp
@@ -213,7 +213,7 @@ void FAnimModel_AnimMontage::RecalculateSequenceLength()
 		AnimMontage->InvalidateRecursiveAsset();
 
 		float NewSequenceLength = CalculateSequenceLengthOfEditorObject();
-		if (NewSequenceLength != AnimMontage->SequenceLength)
+		if (NewSequenceLength != AnimMontage->GetPlayLength())
 		{
 			ClampToEndTime(NewSequenceLength);
 
@@ -238,12 +238,12 @@ void FAnimModel_AnimMontage::RecalculateSequenceLength()
 
 bool FAnimModel_AnimMontage::ClampToEndTime(float NewEndTime)
 {
-	float SequenceLength = AnimMontage->SequenceLength;
+	float SequenceLength = AnimMontage->GetPlayLength();
 
 	bool bClampingNeeded = (SequenceLength > 0.f && NewEndTime < SequenceLength);
 	if(bClampingNeeded)
 	{
-		float ratio = NewEndTime / AnimMontage->SequenceLength;
+		float ratio = NewEndTime / AnimMontage->GetPlayLength();
 
 		for(int32 i=0; i < AnimMontage->CompositeSections.Num(); i++)
 		{
diff --git a/Engine/Source/Editor/Persona/Private/AnimTimeline/AnimTimelineTrack_Montage.cpp b/Engine/Source/Editor/Persona/Private/AnimTimeline/AnimTimelineTrack_Montage.cpp
--- a/Engine/Source/Editor/Persona/Private/AnimTimeline/AnimTimelineTrack_Montage.cpp
+++ b/Engine/Source/Editor/Persona/Private/AnimTimeline/AnimTimelineTrack_Montage.cpp
@@ -361,7 +361,7 @@ class SMontageSections : public SLeafWidget
 							SNew(SNumericEntryBox<float>)
 							.Font(FEditorStyle::GetFontStyle(TEXT("MenuItem.Font")))
 							.MinValue(0.0f)
-							.MaxValue(AnimMontage->SequenceLength)
+							.MaxValue(AnimMontage->GetPlayLength())
 							.Value(Section.GetTime())
 							.AllowSpin(true)
 							.OnValueCommitted_Lambda([this, SectionIndex](float InValue, ETextCommit::Type InCommitType)
@@ -398,7 +398,7 @@ class SMontageSections : public SLeafWidget
 							{
 								if (AnimMontage->CompositeSections.IsValidIndex(SectionIndex))
 								{
-									float NewTime = FMath::Clamp(AnimMontage->GetTimeAtFrame(InValue), 0.0f, AnimMontage->SequenceLength);
+									float NewTime = FMath::Clamp(AnimMontage->GetTimeAtFrame(InValue), 0.0f, AnimMontage->GetPlayLength());
 									WeakModel.Pin()->GetMontagePanel()->GetAnimMontagePanel()->SetSectionTime(SectionIndex, NewTime);
 								}
 
diff --git a/Engine/Source/Editor/Persona/Private/SAnimCompositePanel.cpp b/Engine/Source/Editor/Persona/Private/SAnimCompositePanel.cpp
--- a/Engine/Source/Editor/Persona/Private/SAnimCompositePanel.cpp
+++ b/Engine/Source/Editor/Persona/Private/SAnimCompositePanel.cpp
@@ -84,7 +84,7 @@ void SAnimCompositePanel::Update()
 				.ViewInputMin(ViewInputMin)
 				.ViewInputMax(ViewInputMax)
 				.OnGetNodeColor(this,  &SAnimCompositePanel::HandleGetNodeColor)
-				.TrackMaxValue(Composite->SequenceLength)
+				.TrackMaxValue(Composite->GetPlayLength())
 				.TrackNumDiscreteValues(Composite->GetNumberOfFrames())
 				.bChildAnimMontage(AnimMontage && AnimMontage->HasParentAsset())
 				.OnAnimSegmentNodeClicked( this, &SAnimCompositePanel::ShowSegmentInDetailsView )
@@ -210,9 +210,9 @@ void SAnimCompositePanel::PostUndoRedo()
 	}
 
 	// when undo or redo happens, we still have to recalculate length, so we can't rely on sequence length changes or not
-	if (Composite->SequenceLength)
+	if (Composite->GetPlayLength())
 	{
-		Composite->SequenceLength = 0.f;
+		Composite->SetSequenceLength(0.f);
 	}
 }
 
diff --git a/Engine/Source/Editor/Persona/Private/SAnimEditorBase.cpp b/Engine/Source/Editor/Persona/Private/SAnimEditorBase.cpp
--- a/Engine/Source/Editor/Persona/Private/SAnimEditorBase.cpp
+++ b/Engine/Source/Editor/Persona/Private/SAnimEditorBase.cpp
@@ -272,7 +272,7 @@ float SAnimEditorBase::GetSequenceLength() const
 {
 	if (UAnimSequenceBase* AnimSeqBase = Cast<UAnimSequenceBase>(GetEditorObject()))
 	{
-		return AnimSeqBase->SequenceLength;
+		return AnimSeqBase->GetPlayLength();
 	}
 	else if (UBlendSpaceBase* BlendSpaceBase = Cast<UBlendSpaceBase>(GetEditorObject()))
 	{
diff --git a/Engine/Source/Editor/Persona/Private/SAnimMontagePanel.cpp b/Engine/Source/Editor/Persona/Private/SAnimMontagePanel.cpp
--- a/Engine/Source/Editor/Persona/Private/SAnimMontagePanel.cpp
+++ b/Engine/Source/Editor/Persona/Private/SAnimMontagePanel.cpp
@@ -360,7 +360,7 @@ bool SAnimMontagePanel::ClampToEndTime(float NewEndTime)
 	bool bClampingNeeded = (SequenceLength > 0.f && NewEndTime < SequenceLength);
 	if(bClampingNeeded)
 	{
-		float ratio = NewEndTime / Montage->SequenceLength;
+		float ratio = NewEndTime / Montage->GetPlayLength();
 
 		for(int32 i=0; i < Montage->CompositeSections.Num(); i++)
 		{
@@ -582,9 +582,9 @@ void SAnimMontagePanel::PostRedo( bool bSuccess )
 void SAnimMontagePanel::PostRedoUndo()
 {
 	// when undo or redo happens, we still have to recalculate length, so we can't rely on sequence length changes or not
-	if (Montage->SequenceLength)
+	if (Montage->GetPlayLength())
 	{
-		Montage->SequenceLength = 0.f;
+		Montage->SetSequenceLength(0.f);
 	}
 
 	RebuildMontagePanel(); //Rebuild here, undoing adds can cause slate to crash later on if we don't (using dummy args since they aren't used by the method
@@ -795,7 +795,7 @@ void SAnimMontagePanel::OnNewSectionClicked(float DataPosX)
 bool SAnimMontagePanel::CanAddNewSection()
 {
 	// Can't add sections if there isn't a montage, or that montage is of zero length
-	return Montage && Montage->SequenceLength > 0.0f;
+	return Montage && Montage->GetPlayLength() > 0.0f;
 }
 
 void SAnimMontagePanel::CreateNewSection(const FText& NewSectionName, ETextCommit::Type CommitInfo, float StartTime)
@@ -1076,7 +1076,7 @@ float SAnimMontagePanel::GetSequenceLength() const
 {
 	if(Montage != nullptr)
 	{
-		return Montage->SequenceLength;
+		return Montage->GetPlayLength();
 	}
 	return 0.0f;
 }
diff --git a/Engine/Source/Editor/Persona/Private/SAnimNotifyPanel.cpp b/Engine/Source/Editor/Persona/Private/SAnimNotifyPanel.cpp
--- a/Engine/Source/Editor/Persona/Private/SAnimNotifyPanel.cpp
+++ b/Engine/Source/Editor/Persona/Private/SAnimNotifyPanel.cpp
@@ -1218,7 +1218,7 @@ public:
 			// Look for a snap on the first scrub handle
 			FVector2D TrackNodePos = TrackGeom.AbsoluteToLocal(EventPosition);
 			const FVector2D OriginalNodePosition = TrackNodePos;
-			float SequenceEnd = TrackScaleInfo.InputToLocalX(Sequence->SequenceLength);
+			float SequenceEnd = TrackScaleInfo.InputToLocalX(Sequence->GetPlayLength());
 
 			// Always clamp the Y to the current track
 			SelectionBeginPosition.Y = SelectionPositionClampInfo->TrackPos - 1.0f;
@@ -1331,7 +1331,7 @@ public:
 		SnapPosition = ScaleInfo.InputToLocalX(SnapPosition);
 
 		float WidgetSpaceStartPosition = ScaleInfo.InputToLocalX(0.0f);
-		float WidgetSpaceEndPosition = ScaleInfo.InputToLocalX(Sequence->SequenceLength);
+		float WidgetSpaceEndPosition = ScaleInfo.InputToLocalX(Sequence->GetPlayLength());
 
 		if(!bOutSnapped)
 		{
@@ -1779,7 +1779,7 @@ int32 SAnimNotifyNode::OnPaint(const FPaintArgs& Args, const FGeometry& Allotted
 		if(AnimNotifyEvent && AnimNotifyEvent->EndTriggerTimeOffset != 0.f) //Do we have an offset to render?
 		{
 			float EndTime = AnimNotifyEvent->GetTime() + AnimNotifyEvent->GetDuration();
-			if(EndTime != Sequence->SequenceLength) //Don't render offset when we are at the end of the sequence, doesnt help the user
+			if(EndTime != Sequence->GetPlayLength()) //Don't render offset when we are at the end of the sequence, doesnt help the user
 			{
 				// ScrubHandle
 				float HandleCentre = NotifyDurationSizeX + (ScrubHandleSize.X - 2.0f);
@@ -1862,7 +1862,7 @@ int32 SAnimNotifyNode::OnPaint(const FPaintArgs& Args, const FGeometry& Allotted
 	if(AnimNotifyEvent && AnimNotifyEvent->TriggerTimeOffset != 0.f) //Do we have an offset to render?
 	{
 		float NotifyTime = AnimNotifyEvent->GetTime();
-		if(NotifyTime != 0.f && NotifyTime != Sequence->SequenceLength) //Don't render offset when we are at the start/end of the sequence, doesn't help the user
+		if(NotifyTime != 0.f && NotifyTime != Sequence->GetPlayLength()) //Don't render offset when we are at the start/end of the sequence, doesn't help the user
 		{
 			float HandleCentre = NotifyScrubHandleCentre;
 			float &Offset = AnimNotifyEvent->TriggerTimeOffset;
@@ -2296,7 +2296,7 @@ int32 SAnimNotifyTrack::OnPaint(const FPaintArgs& Args, const FGeometry& Allotte
 
 FCursorReply SAnimNotifyTrack::OnCursorQuery(const FGeometry& MyGeometry, const FPointerEvent& CursorEvent) const
 {
-	if (ViewInputMin.Get() > 0.f || ViewInputMax.Get() < Sequence->SequenceLength)
+	if (ViewInputMin.Get() > 0.f || ViewInputMax.Get() < Sequence->GetPlayLength())
 	{
 		return FCursorReply::Cursor(EMouseCursor::GrabHand);
 	}
@@ -2820,7 +2820,7 @@ TSharedPtr<SWidget> SAnimNotifyTrack::SummonContextMenu(const FGeometry& MyGeome
 							SNew(SNumericEntryBox<float>)
 							.Font(FEditorStyle::GetFontStyle(TEXT("MenuItem.Font")))
 							.MinValue(0.0f)
-							.MaxValue(Sequence->SequenceLength)
+							.MaxValue(Sequence->GetPlayLength())
 							.Value(NodeObject->GetTime())
 							.AllowSpin(false)
 							.OnValueCommitted_Lambda([this, NodeIndex](float InValue, ETextCommit::Type InCommitType)
@@ -2829,7 +2829,7 @@ TSharedPtr<SWidget> SAnimNotifyTrack::SummonContextMenu(const FGeometry& MyGeome
 								{
 									INodeObjectInterface* LocalNodeObject = NotifyNodes[NodeIndex]->NodeObjectInterface;
 
-									float NewTime = FMath::Clamp(InValue, 0.0f, Sequence->SequenceLength - LocalNodeObject->GetDuration());
+									float NewTime = FMath::Clamp(InValue, 0.0f, Sequence->GetPlayLength() - LocalNodeObject->GetDuration());
 									LocalNodeObject->SetTime(NewTime);
 
 									if (FAnimNotifyEvent* Event = LocalNodeObject->GetNotifyEvent())
@@ -2872,7 +2872,7 @@ TSharedPtr<SWidget> SAnimNotifyTrack::SummonContextMenu(const FGeometry& MyGeome
 								{
 									INodeObjectInterface* LocalNodeObject = NotifyNodes[NodeIndex]->NodeObjectInterface;
 
-									float NewTime = FMath::Clamp(Sequence->GetTimeAtFrame(InValue), 0.0f, Sequence->SequenceLength - LocalNodeObject->GetDuration());
+									float NewTime = FMath::Clamp(Sequence->GetTimeAtFrame(InValue), 0.0f, Sequence->GetPlayLength() - LocalNodeObject->GetDuration());
 									LocalNodeObject->SetTime(NewTime);
 
 									if (FAnimNotifyEvent* Event = LocalNodeObject->GetNotifyEvent())
@@ -2951,7 +2951,7 @@ TSharedPtr<SWidget> SAnimNotifyTrack::SummonContextMenu(const FGeometry& MyGeome
 										if ( InCommitType == ETextCommit::OnEnter && AnimNotifies.IsValidIndex(NotifyIndex) )
 										{
 											float NewDuration = FMath::Max(InValue, SAnimNotifyNode::MinimumStateDuration);
-											float MaxDuration = Sequence->SequenceLength - AnimNotifies[NotifyIndex]->GetTime();
+											float MaxDuration = Sequence->GetPlayLength() - AnimNotifies[NotifyIndex]->GetTime();
 											NewDuration = FMath::Min(NewDuration, MaxDuration);
 											AnimNotifies[NotifyIndex]->SetDuration(NewDuration);
 
@@ -3070,7 +3070,7 @@ TSharedPtr<SWidget> SAnimNotifyTrack::SummonContextMenu(const FGeometry& MyGeome
 					MenuBuilder.AddMenuEntry(LOCTEXT("PasteMultAbs", "Paste Multiple Absolute"), LOCTEXT("PasteMultAbsToolTip", "Paste multiple notifies beginning at the mouse cursor, maintaining absolute spacing."), FSlateIcon(), NewAction);
 				}
 
-				if(OriginalTime < Sequence->SequenceLength)
+				if(OriginalTime < Sequence->GetPlayLength())
 				{
 					NewAction.ExecuteAction.BindRaw(
 						this, &SAnimNotifyTrack::OnPasteNotifyClicked, ENotifyPasteMode::OriginalTime, ENotifyPasteMultipleMode::Absolute);
@@ -3507,7 +3507,7 @@ float SAnimNotifyTrack::CalculateTime(const FGeometry& MyGeometry, FVector2D Nod
 		NodePos = MyGeometry.AbsoluteToLocal(NodePos);
 	}
 	FTrackScaleInfo ScaleInfo(ViewInputMin.Get(), ViewInputMax.Get(), 0, 0, MyGeometry.Size);
-	return FMath::Clamp<float>(ScaleInfo.LocalXToInput(NodePos.X), 0.f, Sequence->SequenceLength);
+	return FMath::Clamp<float>(ScaleInfo.LocalXToInput(NodePos.X), 0.f, Sequence->GetPlayLength());
 }
 
 FReply SAnimNotifyTrack::OnDrop(const FGeometry& MyGeometry, const FDragDropEvent& DragDropEvent)
@@ -3598,7 +3598,7 @@ void SAnimNotifyTrack::PasteSingleNotify(FString& NotifyString, float PasteTime)
 		// We have to link to the montage / sequence again, we need a correct time set and we could be pasting to a new montage / sequence
 		int32 NewSlotIndex = 0;
 		float NewNotifyTime = PasteTime != 1.0f ? PasteTime : NewNotify.GetTime();
-		NewNotifyTime = FMath::Clamp(NewNotifyTime, 0.0f, Sequence->SequenceLength);
+		NewNotifyTime = FMath::Clamp(NewNotifyTime, 0.0f, Sequence->GetPlayLength());
 
 		if(UAnimMontage* Montage = Cast<UAnimMontage>(Sequence))
 		{
@@ -3630,7 +3630,7 @@ void SAnimNotifyTrack::PasteSingleNotify(FString& NotifyString, float PasteTime)
 			NewNotify.NotifyStateClass = NewNotifyStateObject;
 			bValidNotify = NewNotifyStateObject->CanBePlaced(Sequence);
 			// Clamp duration into the sequence
-			NewNotify.SetDuration(FMath::Clamp(NewNotify.GetDuration(), 1 / 30.0f, Sequence->SequenceLength - NewNotify.GetTime()));
+			NewNotify.SetDuration(FMath::Clamp(NewNotify.GetDuration(), 1 / 30.0f, Sequence->GetPlayLength() - NewNotify.GetTime()));
 			NewNotify.EndTriggerTimeOffset = GetTriggerTimeOffsetForType(Sequence->CalculateOffsetForNotify(NewNotify.GetTime() + NewNotify.GetDuration()));
 			NewNotify.EndLink.Link(Sequence, NewNotify.EndLink.GetTime());
 		}
@@ -3677,7 +3677,7 @@ void SAnimNotifyTrack::PasteSingleSyncMarker(FString& MarkerString, float PasteT
 			}
 
 			// Make sure the notify is within the track area
-			SyncMarker.Time = FMath::Clamp(SyncMarker.Time, 0.0f, Sequence->SequenceLength);
+			SyncMarker.Time = FMath::Clamp(SyncMarker.Time, 0.0f, Sequence->GetPlayLength());
 			SyncMarker.TrackIndex = TrackIndex;
 
 			SyncMarker.Guid = FGuid::NewGuid();
@@ -4234,7 +4234,7 @@ FReply SAnimNotifyPanel::OnNotifyNodeDragStarted(TArray<TSharedPtr<SAnimNotifyNo
 
 float SAnimNotifyPanel::GetSequenceLength() const
 {
-	return Sequence->SequenceLength;
+	return Sequence->GetPlayLength();
 }
 
 void SAnimNotifyPanel::PostUndo( bool bSuccess )
@@ -4396,7 +4396,7 @@ void SAnimNotifyPanel::CopySelectedNodesToClipboard() const
 		int32 TrackSpan = MaxTrack - MinTrack + 1;
 
 		StrValue += FString::Printf(TEXT("OriginalTime=%f,"), MinTime);
-		StrValue += FString::Printf(TEXT("OriginalLength=%f,"), Sequence->SequenceLength);
+		StrValue += FString::Printf(TEXT("OriginalLength=%f,"), Sequence->GetPlayLength());
 		StrValue += FString::Printf(TEXT("TrackSpan=%d"), TrackSpan);
 
 		for(const INodeObjectInterface* NodeObject : SelectedNodes)
@@ -4578,7 +4578,7 @@ void SAnimNotifyPanel::OnPasteNodes(SAnimNotifyTrack* RequestTrack, float ClickT
 		// Scaling for relative paste
 		if(MultiplePasteType == ENotifyPasteMultipleMode::Relative)
 		{
-			ScaleMultiplier = Sequence->SequenceLength / OrigLength;
+			ScaleMultiplier = Sequence->GetPlayLength() / OrigLength;
 		}
 
 		// Process each line of the paste buffer and spawn notifies
@@ -4908,21 +4908,21 @@ void SAnimNotifyPanel::OnNotifyObjectChanged(UObject* EditorBaseObj, bool bRebui
 
 void SAnimNotifyPanel::OnNotifyTrackScrolled(float InScrollOffsetFraction)
 {
-	float Ratio = (ViewInputMax.Get() - ViewInputMin.Get()) / Sequence->SequenceLength;
+	float Ratio = (ViewInputMax.Get() - ViewInputMin.Get()) / Sequence->GetPlayLength();
 	float MaxOffset = (Ratio < 1.0f) ? 1.0f - Ratio : 0.0f;
 	InScrollOffsetFraction = FMath::Clamp(InScrollOffsetFraction, 0.0f, MaxOffset);
 
 	// Calculate new view ranges
-	float NewMin = InScrollOffsetFraction * Sequence->SequenceLength;
-	float NewMax = (InScrollOffsetFraction + Ratio) * Sequence->SequenceLength;
+	float NewMin = InScrollOffsetFraction * Sequence->GetPlayLength();
+	float NewMax = (InScrollOffsetFraction + Ratio) * Sequence->GetPlayLength();
 	
 	InputViewRangeChanged(NewMin, NewMax);
 }
 
 void SAnimNotifyPanel::InputViewRangeChanged(float ViewMin, float ViewMax)
 {
-	float Ratio = (ViewMax - ViewMin) / Sequence->SequenceLength;
-	float OffsetFraction = ViewMin / Sequence->SequenceLength;
+	float Ratio = (ViewMax - ViewMin) / Sequence->GetPlayLength();
+	float OffsetFraction = ViewMin / Sequence->GetPlayLength();
 	if(NotifyTrackScrollBar.IsValid())
 	{
 		NotifyTrackScrollBar->SetState(OffsetFraction, Ratio);
diff --git a/Engine/Source/Editor/Persona/Private/SAnimSegmentsPanel.cpp b/Engine/Source/Editor/Persona/Private/SAnimSegmentsPanel.cpp
--- a/Engine/Source/Editor/Persona/Private/SAnimSegmentsPanel.cpp
+++ b/Engine/Source/Editor/Persona/Private/SAnimSegmentsPanel.cpp
@@ -320,7 +320,7 @@ void SAnimSegmentsPanel::AddAnimSegment( UAnimSequenceBase* NewSequenceBase, flo
 	FAnimSegment NewSegment;
 	NewSegment.AnimReference = NewSequenceBase;
 	NewSegment.AnimStartTime = 0.f;
-	NewSegment.AnimEndTime = NewSequenceBase->SequenceLength;
+	NewSegment.AnimEndTime = NewSequenceBase->GetPlayLength();
 	NewSegment.AnimPlayRate = 1.f;
 	NewSegment.LoopingCount = 1;
 	NewSegment.StartPos = NewStartPos;
diff --git a/Engine/Source/Editor/Persona/Private/SAnimationScrubPanel.cpp b/Engine/Source/Editor/Persona/Private/SAnimationScrubPanel.cpp
--- a/Engine/Source/Editor/Persona/Private/SAnimationScrubPanel.cpp
+++ b/Engine/Source/Editor/Persona/Private/SAnimationScrubPanel.cpp
@@ -353,7 +353,7 @@ float SAnimationScrubPanel::GetSequenceLength() const
 	}
 	else if (LockedSequence)
 	{
-		return LockedSequence->SequenceLength;
+		return LockedSequence->GetPlayLength();
 	}
 	else
 	{
@@ -480,7 +480,7 @@ void SAnimationScrubPanel::OnCropAnimSequence( bool bFromStart, float CurrentTim
 				//Resetting slider position to the first frame
 				PreviewInstance->SetPosition( 0.0f, false );
 
-				OnSetInputViewRange.ExecuteIfBound(0, AnimSequence->SequenceLength);
+				OnSetInputViewRange.ExecuteIfBound(0, AnimSequence->GetPlayLength());
 			}
 		}
 	}
@@ -508,7 +508,7 @@ void SAnimationScrubPanel::OnAppendAnimSequence( bool bFromStart, int32 NumOfFra
 			int32 CopyFrame = StartFrame;
 			AnimSequence->InsertFramesToRawAnimData(StartFrame, EndFrame, CopyFrame);
 
-			OnSetInputViewRange.ExecuteIfBound(0, AnimSequence->SequenceLength);
+			OnSetInputViewRange.ExecuteIfBound(0, AnimSequence->GetPlayLength());
 		}
 	}
 }
@@ -534,7 +534,7 @@ void SAnimationScrubPanel::OnInsertAnimSequence( bool bBefore, int32 CurrentFram
 			int32 EndFrame = StartFrame + 1;
 			AnimSequence->InsertFramesToRawAnimData(StartFrame, EndFrame, CurrentFrame);
 
-			OnSetInputViewRange.ExecuteIfBound(0, AnimSequence->SequenceLength);
+			OnSetInputViewRange.ExecuteIfBound(0, AnimSequence->GetPlayLength());
 		}
 	}
 }
diff --git a/Engine/Source/Editor/Persona/Private/SAnimationSequenceBrowser.cpp b/Engine/Source/Editor/Persona/Private/SAnimationSequenceBrowser.cpp
--- a/Engine/Source/Editor/Persona/Private/SAnimationSequenceBrowser.cpp
+++ b/Engine/Source/Editor/Persona/Private/SAnimationSequenceBrowser.cpp
@@ -1017,7 +1017,9 @@ void SAnimationSequenceBrowser::Construct(const FArguments& InArgs, const TShare
 	// Create the ignore set for asset registry tags
 	// Making Skeleton to be private, and now GET_MEMBER_NAME_CHECKED doesn't work
 	AssetRegistryTagsToIgnore.Add(TEXT("Skeleton"));
+	PRAGMA_DISABLE_DEPRECATION_WARNINGS
 	AssetRegistryTagsToIgnore.Add(GET_MEMBER_NAME_CHECKED(UAnimSequenceBase, SequenceLength));
+	PRAGMA_ENABLE_DEPRECATION_WARNINGS
 	AssetRegistryTagsToIgnore.Add(GET_MEMBER_NAME_CHECKED(UAnimSequenceBase, RateScale));
 }
 
diff --git a/Engine/Source/Editor/SequenceRecorder/Private/AnimationRecorder.cpp b/Engine/Source/Editor/SequenceRecorder/Private/AnimationRecorder.cpp
--- a/Engine/Source/Editor/SequenceRecorder/Private/AnimationRecorder.cpp
+++ b/Engine/Source/Editor/SequenceRecorder/Private/AnimationRecorder.cpp
@@ -224,7 +224,7 @@ void FAnimationRecorder::StartRecord(USkeletalMeshComponent* Component, UAnimSeq
 	PreviousComponentToWorld = Component->GetComponentTransform();
 
 	LastFrame = 0;
-	AnimationObject->SequenceLength = 0.f;
+	AnimationObject->SetSequenceLength(0.f);
 	AnimationObject->SetRawNumberOfFrame(0);
 
 	RecordedCurves.Reset();
@@ -292,7 +292,7 @@ UAnimSequence* FAnimationRecorder::StopRecord(bool bShowMessage)
 		AnimationObject->SetRawNumberOfFrame(NumFrames);
 
 		// can't use TimePassed. That is just total time that has been passed, not necessarily match with frame count
-		AnimationObject->SequenceLength = (NumFrames>1) ? (NumFrames-1) * IntervalTime : MINIMUM_ANIMATION_LENGTH;
+		AnimationObject->SetSequenceLength((NumFrames>1) ? (NumFrames-1) * IntervalTime : MINIMUM_ANIMATION_LENGTH);
 
 		FixupNotifies();
 
@@ -396,7 +396,7 @@ UAnimSequence* FAnimationRecorder::StopRecord(bool bShowMessage)
 			const FText NotificationText = FText::Format(LOCTEXT("RecordAnimation", "'{0}' has been successfully recorded [{1} frames : {2} sec(s) @ {3} Hz]"),
 				FText::FromString(AnimationObject->GetName()),
 				FText::AsNumber(AnimationObject->GetRawNumberOfFrames()),
-				FText::AsNumber(AnimationObject->SequenceLength),
+				FText::AsNumber(AnimationObject->GetPlayLength()),
 				FText::AsNumber(1.f / IntervalTime)
 				);
 					
diff --git a/Engine/Source/Editor/UnrealEd/Private/Factories/ActorFactory.cpp b/Engine/Source/Editor/UnrealEd/Private/Factories/ActorFactory.cpp
--- a/Engine/Source/Editor/UnrealEd/Private/Factories/ActorFactory.cpp
+++ b/Engine/Source/Editor/UnrealEd/Private/Factories/ActorFactory.cpp
@@ -924,8 +924,8 @@ void UActorFactoryAnimationAsset::PostSpawnActor( UObject* Asset, AActor* NewAct
 				//If we have a negative play rate, default initial position to sequence end
 				if (AnimSeq->RateScale < 0.f)
 				{
-					NewSASComponent->AnimationData.SavedPosition = AnimSeq->SequenceLength;
-					NewSASComponent->SetPosition(AnimSeq->SequenceLength, false);
+					NewSASComponent->AnimationData.SavedPosition = AnimSeq->GetPlayLength();
+					NewSASComponent->SetPosition(AnimSeq->GetPlayLength(), false);
 				}
 			}
 			
diff --git a/Engine/Source/Editor/UnrealEd/Private/Factories/AnimCompositeFactory.cpp b/Engine/Source/Editor/UnrealEd/Private/Factories/AnimCompositeFactory.cpp
--- a/Engine/Source/Editor/UnrealEd/Private/Factories/AnimCompositeFactory.cpp
+++ b/Engine/Source/Editor/UnrealEd/Private/Factories/AnimCompositeFactory.cpp
@@ -83,13 +83,13 @@ UObject* UAnimCompositeFactory::FactoryCreateNew(UClass* Class, UObject* InParen
 			FAnimSegment NewSegment;
 			NewSegment.AnimReference = SourceAnimation;
 			NewSegment.AnimStartTime = 0.f;
-			NewSegment.AnimEndTime = SourceAnimation->SequenceLength;
+			NewSegment.AnimEndTime = SourceAnimation->GetPlayLength();
 			NewSegment.AnimPlayRate = 1.f;
 			NewSegment.LoopingCount = 1;
 			NewSegment.StartPos = 0.f;
 
 			AnimComposite->AnimationTrack.AnimSegments.Add(NewSegment);
-			AnimComposite->SetSequenceLength( AnimComposite->AnimationTrack.GetLength() );
+			AnimComposite->SetSequenceLength(AnimComposite->AnimationTrack.GetLength());
 		}
 		
 		AnimComposite->SetSkeleton( TargetSkeleton );
diff --git a/Engine/Source/Editor/UnrealEd/Private/Factories/AnimMontageFactory.cpp b/Engine/Source/Editor/UnrealEd/Private/Factories/AnimMontageFactory.cpp
--- a/Engine/Source/Editor/UnrealEd/Private/Factories/AnimMontageFactory.cpp
+++ b/Engine/Source/Editor/UnrealEd/Private/Factories/AnimMontageFactory.cpp
@@ -83,7 +83,7 @@ UObject* UAnimMontageFactory::FactoryCreateNew(UClass* Class, UObject* InParent,
 			FAnimSegment NewSegment;
 			NewSegment.AnimReference = SourceAnimation;
 			NewSegment.AnimStartTime = 0.f;
-			NewSegment.AnimEndTime = SourceAnimation->SequenceLength;
+			NewSegment.AnimEndTime = SourceAnimation->GetPlayLength();
 			NewSegment.AnimPlayRate = 1.f;
 			NewSegment.LoopingCount = 1;
 			NewSegment.StartPos = 0.f;
@@ -91,7 +91,7 @@ UObject* UAnimMontageFactory::FactoryCreateNew(UClass* Class, UObject* InParent,
 			FSlotAnimationTrack& NewTrack = AnimMontage->SlotAnimTracks[0];
 			NewTrack.AnimTrack.AnimSegments.Add(NewSegment);
 
-			AnimMontage->SetSequenceLength(SourceAnimation->SequenceLength);
+			AnimMontage->SetSequenceLength(SourceAnimation->GetPlayLength());
 		}
 
 		AnimMontage->SetSkeleton(TargetSkeleton);
diff --git a/Engine/Source/Editor/UnrealEd/Private/Factories/AnimSequenceFactory.cpp b/Engine/Source/Editor/UnrealEd/Private/Factories/AnimSequenceFactory.cpp
--- a/Engine/Source/Editor/UnrealEd/Private/Factories/AnimSequenceFactory.cpp
+++ b/Engine/Source/Editor/UnrealEd/Private/Factories/AnimSequenceFactory.cpp
@@ -70,7 +70,7 @@ UObject* UAnimSequenceFactory::FactoryCreateNew(UClass* Class, UObject* InParent
 	UAnimSequence* AnimSequence = NewObject<UAnimSequence>(InParent, Class, Name, Flags);
 
 	// @todo I think this will crash, we should support differentoptions
-	AnimSequence->SequenceLength = 0.f;
+	AnimSequence->SetSequenceLength(0.f);
 	AnimSequence->SetRawNumberOfFrame(0);
 	
 	if (TargetSkeleton)
diff --git a/Engine/Source/Editor/UnrealEd/Private/Fbx/FbxAnimationExport.cpp b/Engine/Source/Editor/UnrealEd/Private/Fbx/FbxAnimationExport.cpp
--- a/Engine/Source/Editor/UnrealEd/Private/Fbx/FbxAnimationExport.cpp
+++ b/Engine/Source/Editor/UnrealEd/Private/Fbx/FbxAnimationExport.cpp
@@ -25,13 +25,13 @@ namespace UnFbx
 
 	bool FFbxExporter::SetupAnimStack(const UAnimSequence* AnimSeq)
 	{
-		if (AnimSeq->SequenceLength == 0.f)
+		if (AnimSeq->GetPlayLength() == 0.f)
 		{
 			// something is wrong
 			return false;
 		}
 
-		const float FrameRate = FMath::TruncToFloat(((AnimSeq->GetRawNumberOfFrames() - 1) / AnimSeq->SequenceLength) + 0.5f);
+		const float FrameRate = FMath::TruncToFloat(((AnimSeq->GetRawNumberOfFrames() - 1) / AnimSeq->GetPlayLength()) + 0.5f);
 		//Configure the scene time line
 		{
 			FbxGlobalSettings& SceneGlobalSettings = Scene->GetGlobalSettings();
@@ -52,7 +52,7 @@ namespace UnFbx
 		// set time correctly
 		FbxTime ExportedStartTime, ExportedStopTime;
 		ExportedStartTime.SetSecondDouble(0.f);
-		ExportedStopTime.SetSecondDouble(AnimSeq->SequenceLength);
+		ExportedStopTime.SetSecondDouble(AnimSeq->GetPlayLength());
 
 		FbxTimeSpan ExportedTimeSpan;
 		ExportedTimeSpan.Set(ExportedStartTime, ExportedStopTime);
@@ -249,9 +249,9 @@ void FFbxExporter::ExportCustomAnimCurvesToFbx(const TMap<FName, FbxAnimCurve*>&
 void FFbxExporter::IterateInsideAnimSequence(const UAnimSequence* AnimSeq, float AnimStartOffset, float AnimEndOffset, float AnimPlayRate, float StartTime, TFunctionRef<void(float, FbxTime, bool)> IterationLambda)
 {
 	float AnimTime = AnimStartOffset;
-	float AnimEndTime = (AnimSeq->SequenceLength - AnimEndOffset);
+	float AnimEndTime = (AnimSeq->GetPlayLength() - AnimEndOffset);
 	// Subtracts 1 because NumFrames includes an initial pose for 0.0 second
-	double TimePerKey = (AnimSeq->SequenceLength / (AnimSeq->GetRawNumberOfFrames() - 1));
+	double TimePerKey = (AnimSeq->GetPlayLength() / (AnimSeq->GetRawNumberOfFrames() - 1));
 	const float AnimTimeIncrement = TimePerKey * AnimPlayRate;
 	uint32 AnimFrameIndex = 0;
 
diff --git a/Engine/Source/Editor/UnrealEd/Private/SkeletalMeshEdit.cpp b/Engine/Source/Editor/UnrealEd/Private/SkeletalMeshEdit.cpp
--- a/Engine/Source/Editor/UnrealEd/Private/SkeletalMeshEdit.cpp
+++ b/Engine/Source/Editor/UnrealEd/Private/SkeletalMeshEdit.cpp
@@ -1292,16 +1292,16 @@ bool UnFbx::FFbxImporter::ImportAnimation(USkeleton* Skeleton, UAnimSequence * D
 
 	// @todo : the length might need to change w.r.t. sampling keys
 	FbxTime SequenceLength = AnimTimeSpan.GetDuration();
-	float PreviousSequenceLength = DestSeq->SequenceLength;
+	float PreviousSequenceLength = DestSeq->GetPlayLength();
 
 	// if you have one pose(thus 0.f duration), it still contains animation, so we'll need to consider that as MINIMUM_ANIMATION_LENGTH time length
-	DestSeq->SequenceLength = FGenericPlatformMath::Max<float>(SequenceLength.GetSecondDouble(), MINIMUM_ANIMATION_LENGTH);
+	DestSeq->SetSequenceLength(FGenericPlatformMath::Max<float>(SequenceLength.GetSecondDouble(), MINIMUM_ANIMATION_LENGTH));
 
 	if(PreviousSequenceLength > MINIMUM_ANIMATION_LENGTH && DestSeq->RawCurveData.FloatCurves.Num() > 0)
 	{
 		// The sequence already existed when we began the import. We need to scale the key times for all curves to match the new 
 		// duration before importing over them. This is to catch any user-added curves
-		float ScaleFactor = DestSeq->SequenceLength / PreviousSequenceLength;
+		float ScaleFactor = DestSeq->GetPlayLength() / PreviousSequenceLength;
 		if (!FMath::IsNearlyEqual(ScaleFactor, 1.f))
 		{
 			for (FFloatCurve& Curve : DestSeq->RawCurveData.FloatCurves)
@@ -1349,9 +1349,9 @@ bool UnFbx::FFbxImporter::ImportAnimation(USkeleton* Skeleton, UAnimSequence * D
 
 	const bool bReimportWarnings = GetDefault<UEditorPerProjectUserSettings>()->bAnimationReimportWarnings;
 	
-	if (bReimportWarnings && !FMath::IsNearlyZero(PreviousSequenceLength) && !FMath::IsNearlyEqual(DestSeq->SequenceLength, PreviousSequenceLength))
+	if (bReimportWarnings && !FMath::IsNearlyZero(PreviousSequenceLength) && !FMath::IsNearlyEqual(DestSeq->GetPlayLength(), PreviousSequenceLength))
 	{
-		AddTokenizedErrorMessage(FTokenizedMessage::Create(EMessageSeverity::Warning, FText::Format(LOCTEXT("Warning_SequenceLengthChanged", "Animation Sequence ({0}) length {1} is different from previous {2}."), FText::FromName(DestSeq->GetFName()), DestSeq->SequenceLength, PreviousSequenceLength)), FFbxErrors::Animation_DifferentLength);
+		AddTokenizedErrorMessage(FTokenizedMessage::Create(EMessageSeverity::Warning, FText::Format(LOCTEXT("Warning_SequenceLengthChanged", "Animation Sequence ({0}) length {1} is different from previous {2}."), FText::FromName(DestSeq->GetFName()), DestSeq->GetPlayLength(), PreviousSequenceLength)), FFbxErrors::Animation_DifferentLength);
 	}
 
 	FbxNode *SkeletalMeshRootNode = NodeArray.Num() > 0 ? NodeArray[0] : nullptr;
diff --git a/Engine/Source/Editor/UnrealEd/Private/ThumbnailHelpers.cpp b/Engine/Source/Editor/UnrealEd/Private/ThumbnailHelpers.cpp
--- a/Engine/Source/Editor/UnrealEd/Private/ThumbnailHelpers.cpp
+++ b/Engine/Source/Editor/UnrealEd/Private/ThumbnailHelpers.cpp
@@ -605,7 +605,7 @@ bool FAnimationSequenceThumbnailScene::SetAnimation(UAnimSequenceBase* InAnimati
 				if (InAnimation->IsValidToPlay())
 				{
 					// Handle posing the mesh at the middle of the animation
-					const float AnimPosition = InAnimation->SequenceLength / 2.f;
+					const float AnimPosition = InAnimation->GetPlayLength() / 2.f;
 
 					UDebugSkelMeshComponent* MeshComponent = CastChecked<UDebugSkelMeshComponent>(PreviewActor->GetSkeletalMeshComponent());
 
diff --git a/Engine/Source/Runtime/AIModule/Private/BehaviorTree/Tasks/BTTask_PlayAnimation.cpp b/Engine/Source/Runtime/AIModule/Private/BehaviorTree/Tasks/BTTask_PlayAnimation.cpp
--- a/Engine/Source/Runtime/AIModule/Private/BehaviorTree/Tasks/BTTask_PlayAnimation.cpp
+++ b/Engine/Source/Runtime/AIModule/Private/BehaviorTree/Tasks/BTTask_PlayAnimation.cpp
@@ -50,7 +50,7 @@ EBTNodeResult::Type UBTTask_PlayAnimation::ExecuteTask(UBehaviorTreeComponent& O
 			CachedSkelMesh = SkelMesh;
 
 			SkelMesh->PlayAnimation(AnimationToPlay, bLooping);
-			const float FinishDelay = AnimationToPlay->GetMaxCurrentTime();
+			const float FinishDelay = AnimationToPlay->GetPlayLength();
 
 			if (bNonBlocking == false && FinishDelay > 0)
 			{
diff --git a/Engine/Source/Runtime/AnimGraphRuntime/Private/AnimNodes/AnimNode_BlendSpacePlayer.cpp b/Engine/Source/Runtime/AnimGraphRuntime/Private/AnimNodes/AnimNode_BlendSpacePlayer.cpp
--- a/Engine/Source/Runtime/AnimGraphRuntime/Private/AnimNodes/AnimNode_BlendSpacePlayer.cpp
+++ b/Engine/Source/Runtime/AnimGraphRuntime/Private/AnimNodes/AnimNode_BlendSpacePlayer.cpp
@@ -46,7 +46,7 @@ float FAnimNode_BlendSpacePlayer::GetCurrentAssetLength()
 		if (BlendSpace != nullptr)
 		{
 			const FBlendSample& Sample = BlendSpace->GetBlendSample(HighestWeightedSample->SampleDataIndex);
-			return Sample.Animation->SequenceLength;
+			return Sample.Animation->GetPlayLength();
 		}
 	}
 
@@ -158,7 +158,9 @@ void FAnimNode_BlendSpacePlayer::GatherDebugData(FNodeDebugData& DebugData)
 
 float FAnimNode_BlendSpacePlayer::GetTimeFromEnd(float CurrentTime)
 {
-	return BlendSpace != nullptr ? BlendSpace->GetMaxCurrentTime() - CurrentTime : 0.0f;
+	// Blend-spaces use normalized time value
+	const float PlayLength = 1.0f;
+	return BlendSpace != nullptr ? PlayLength - CurrentTime : 0.0f;
 }
 
 UAnimationAsset* FAnimNode_BlendSpacePlayer::GetAnimAsset()
diff --git a/Engine/Source/Runtime/AnimGraphRuntime/Private/AnimNodes/AnimNode_RandomPlayer.cpp b/Engine/Source/Runtime/AnimGraphRuntime/Private/AnimNodes/AnimNode_RandomPlayer.cpp
--- a/Engine/Source/Runtime/AnimGraphRuntime/Private/AnimNodes/AnimNode_RandomPlayer.cpp
+++ b/Engine/Source/Runtime/AnimGraphRuntime/Private/AnimNodes/AnimNode_RandomPlayer.cpp
@@ -131,7 +131,7 @@ void FAnimNode_RandomPlayer::Update_AnyThread(const FAnimationUpdateContext& Con
 	float AdjustedPreviousPlayTime = CurrentData->PreviousPlayTime;
 	if (CurrentData->CurrentPlayTime < AdjustedPreviousPlayTime)
 	{
-		AdjustedPreviousPlayTime -= CurrentSequence->SequenceLength;
+		AdjustedPreviousPlayTime -= CurrentSequence->GetPlayLength();
 	}
 
 	// Did we cross the play start time? Decrement the loop counter. Once we're on the last loop, we can
@@ -170,10 +170,10 @@ void FAnimNode_RandomPlayer::Update_AnyThread(const FAnimationUpdateContext& Con
 				float AmountPlayedSoFar = CurrentData->CurrentPlayTime - CurrentData->PlayStartTime;
 				if (AmountPlayedSoFar < 0.0f)
 				{
-					AmountPlayedSoFar += CurrentSequence->SequenceLength;
+					AmountPlayedSoFar += CurrentSequence->GetPlayLength();
 				}
 
-				float TimeRemaining = CurrentSequence->SequenceLength - AmountPlayedSoFar;
+				float TimeRemaining = CurrentSequence->GetPlayLength() - AmountPlayedSoFar;
 
 				if (TimeRemaining <= NextSequenceEntry.BlendIn.GetBlendTime())
 				{
diff --git a/Engine/Source/Runtime/AnimGraphRuntime/Private/AnimNodes/AnimNode_SequenceEvaluator.cpp b/Engine/Source/Runtime/AnimGraphRuntime/Private/AnimNodes/AnimNode_SequenceEvaluator.cpp
--- a/Engine/Source/Runtime/AnimGraphRuntime/Private/AnimNodes/AnimNode_SequenceEvaluator.cpp
+++ b/Engine/Source/Runtime/AnimGraphRuntime/Private/AnimNodes/AnimNode_SequenceEvaluator.cpp
@@ -11,7 +11,7 @@ float FAnimNode_SequenceEvaluator::GetCurrentAssetTime()
 
 float FAnimNode_SequenceEvaluator::GetCurrentAssetLength()
 {
-	return Sequence ? Sequence->SequenceLength : 0.0f;
+	return Sequence ? Sequence->GetPlayLength() : 0.0f;
 }
 
 /////////////////////////////////////////////////////
@@ -36,7 +36,7 @@ void FAnimNode_SequenceEvaluator::UpdateAssetPlayer(const FAnimationUpdateContex
 	if (Sequence)
 	{
 		// Clamp input to a valid position on this sequence's time line.
-		ExplicitTime = FMath::Clamp(ExplicitTime, 0.f, Sequence->SequenceLength);
+		ExplicitTime = FMath::Clamp(ExplicitTime, 0.f, Sequence->GetPlayLength());
 
 		if ((!bTeleportToExplicitTime || (GroupIndex != INDEX_NONE)) && (Context.AnimInstanceProxy->IsSkeletonCompatible(Sequence->GetSkeleton())))
 		{
@@ -48,21 +48,21 @@ void FAnimNode_SequenceEvaluator::UpdateAssetPlayer(const FAnimationUpdateContex
 					case ESequenceEvalReinit::ExplicitTime: InternalTimeAccumulator = ExplicitTime; break;
 				}
 
-				InternalTimeAccumulator = FMath::Clamp(InternalTimeAccumulator, 0.f, Sequence->SequenceLength);
+				InternalTimeAccumulator = FMath::Clamp(InternalTimeAccumulator, 0.f, Sequence->GetPlayLength());
 			}
 
 			float TimeJump = ExplicitTime - InternalTimeAccumulator;
 			if (bShouldLoop)
 			{
-				if (FMath::Abs(TimeJump) > (Sequence->SequenceLength * 0.5f))
+				if (FMath::Abs(TimeJump) > (Sequence->GetPlayLength() * 0.5f))
 				{
 					if (TimeJump > 0.f)
 					{
-						TimeJump -= Sequence->SequenceLength;
+						TimeJump -= Sequence->GetPlayLength();
 					}
 					else
 					{
-						TimeJump += Sequence->SequenceLength;
+						TimeJump += Sequence->GetPlayLength();
 					}
 				}
 			}
diff --git a/Engine/Source/Runtime/Engine/Classes/Animation/AnimCompositeBase.h b/Engine/Source/Runtime/Engine/Classes/Animation/AnimCompositeBase.h
--- a/Engine/Source/Runtime/Engine/Classes/Animation/AnimCompositeBase.h
+++ b/Engine/Source/Runtime/Engine/Classes/Animation/AnimCompositeBase.h
@@ -292,11 +292,6 @@ class UAnimCompositeBase : public UAnimSequenceBase
 {
 	GENERATED_UCLASS_BODY()
 
-#if WITH_EDITOR
-	/** Set Sequence Length */
-	ENGINE_API void SetSequenceLength(float InSequenceLength);
-#endif
-
 	//~ Begin UObject Interface
 	virtual void PostLoad() override;
 	//~ End UObject Interface
diff --git a/Engine/Source/Runtime/Engine/Classes/Animation/AnimSequence.h b/Engine/Source/Runtime/Engine/Classes/Animation/AnimSequence.h
--- a/Engine/Source/Runtime/Engine/Classes/Animation/AnimSequence.h
+++ b/Engine/Source/Runtime/Engine/Classes/Animation/AnimSequence.h
@@ -410,7 +410,7 @@ public:
 	//~ End UAnimSequenceBase Interface
 
 	// Returns the framerate of the animation
-	float GetFrameRate() const { return (float)(FMath::Max(NumFrames - 1, 1)) / (SequenceLength > 0.f ? SequenceLength : 1.f); }
+	float GetFrameRate() const { return (float)(FMath::Max(NumFrames - 1, 1)) / (GetPlayLength() > 0.f ? GetPlayLength() : 1.f); }
 
 	// Extract Root Motion transform from the animation
 	FTransform ExtractRootMotion(float StartTime, float DeltaTime, bool bAllowLooping) const;
diff --git a/Engine/Source/Runtime/Engine/Classes/Animation/AnimSequenceBase.h b/Engine/Source/Runtime/Engine/Classes/Animation/AnimSequenceBase.h
--- a/Engine/Source/Runtime/Engine/Classes/Animation/AnimSequenceBase.h
+++ b/Engine/Source/Runtime/Engine/Classes/Animation/AnimSequenceBase.h
@@ -28,14 +28,23 @@ class ENGINE_API UAnimSequenceBase : public UAnimationAsset
 {
 	GENERATED_UCLASS_BODY()
 
+public:
 	/** Animation notifies, sorted by time (earliest notification first). */
 	UPROPERTY()
 	TArray<struct FAnimNotifyEvent> Notifies;
 
 	/** Length (in seconds) of this AnimSequence if played back with a speed of 1.0. */
+	UE_DEPRECATED(5.0, "Direct access to SequenceLength should be removed, use GetPlayLength or SetSequenceLength instead")
 	UPROPERTY(Category=Length, AssetRegistrySearchable, VisibleAnywhere, BlueprintReadOnly)
 	float SequenceLength;
 
+	void SetSequenceLength(float NewLength) 	
+	{
+		PRAGMA_DISABLE_DEPRECATION_WARNINGS
+		SequenceLength = NewLength; 
+		PRAGMA_ENABLE_DEPRECATION_WARNINGS
+	}
+
 	/** Number for tweaking playback rate of this animation globally. */
 	UPROPERTY(EditAnywhere, Category=Animation)
 	float RateScale;
@@ -57,8 +66,7 @@ class ENGINE_API UAnimSequenceBase : public UAnimationAsset
 	//~ End UObject Interface
 
 	/** Returns the total play length of the montage, if played back with a speed of 1.0. */
-	UFUNCTION(BlueprintCallable, Category = "Animation")
-	virtual float GetPlayLength();
+	virtual float GetPlayLength() const override;
 
 	/** Sort the Notifies array by time, earliest first. */
 	void SortNotifies();	
@@ -156,10 +164,6 @@ public:
 
 	void TickByMarkerAsLeader(FMarkerTickRecord& Instance, FMarkerTickContext& MarkerContext, float& CurrentTime, float& OutPreviousTime, const float MoveDelta, const bool bLooping) const;
 
-	// this is used in editor only when used for transition getter
-	// this doesn't mean max time. In Sequence, this is SequenceLength,
-	// but for BlendSpace CurrentTime is normalized [0,1], so this is 1
-	virtual float GetMaxCurrentTime() override { return SequenceLength; }
 	//~ End UAnimationAsset Interface
 
 	/**
@@ -237,8 +241,6 @@ public:
 	void UnregisterOnAnimTrackCurvesChanged(void* Unregister);
 #endif
 
-
-
 protected:
 	template <typename DataType>
 	void VerifyCurveNames(USkeleton& Skeleton, const FName& NameContainer, TArray<DataType>& CurveList)
diff --git a/Engine/Source/Runtime/Engine/Classes/Animation/AnimationAsset.h b/Engine/Source/Runtime/Engine/Classes/Animation/AnimationAsset.h
--- a/Engine/Source/Runtime/Engine/Classes/Animation/AnimationAsset.h
+++ b/Engine/Source/Runtime/Engine/Classes/Animation/AnimationAsset.h
@@ -881,7 +881,11 @@ public:
 	// this is used in editor only when used for transition getter
 	// this doesn't mean max time. In Sequence, this is SequenceLength,
 	// but for BlendSpace CurrentTime is normalized [0,1], so this is 1
-	virtual float GetMaxCurrentTime() { return 0.f; }
+	UE_DEPRECATED(5.0, "Use GetPlayLength instead")
+	virtual float GetMaxCurrentTime() { return GetPlayLength(); }
+
+	UFUNCTION(BlueprintCallable, Category = "Animation")
+	virtual float GetPlayLength() const { return 0.f; };
 
 	void SetSkeleton(USkeleton* NewSkeleton);
 	void ResetSkeleton(USkeleton* NewSkeleton);
diff --git a/Engine/Source/Runtime/Engine/Classes/Animation/BlendSpaceBase.h b/Engine/Source/Runtime/Engine/Classes/Animation/BlendSpaceBase.h
--- a/Engine/Source/Runtime/Engine/Classes/Animation/BlendSpaceBase.h
+++ b/Engine/Source/Runtime/Engine/Classes/Animation/BlendSpaceBase.h
@@ -244,7 +244,7 @@ class UBlendSpaceBase : public UAnimationAsset, public IInterpolationIndexProvid
 	// this is used in editor only when used for transition getter
 	// this doesn't mean max time. In Sequence, this is SequenceLength,
 	// but for BlendSpace CurrentTime is normalized [0,1], so this is 1
-	virtual float GetMaxCurrentTime() override { return 1.f; }	
+	virtual float GetPlayLength() const override { return 1.f; }
 	virtual TArray<FName>* GetUniqueMarkerNames() override;
 	virtual bool IsValidAdditive() const override;
 #if WITH_EDITOR
diff --git a/Engine/Source/Runtime/Engine/Private/Animation/AnimComposite.cpp b/Engine/Source/Runtime/Engine/Private/Animation/AnimComposite.cpp
--- a/Engine/Source/Runtime/Engine/Private/Animation/AnimComposite.cpp
+++ b/Engine/Source/Runtime/Engine/Private/Animation/AnimComposite.cpp
@@ -26,7 +26,7 @@ void UAnimComposite::ReplaceReferredAnimations(const TMap<UAnimationAsset*, UAni
 
 bool UAnimComposite::IsNotifyAvailable() const
 {
-	return (SequenceLength > 0.f && (Super::IsNotifyAvailable() || AnimationTrack.IsNotifyAvailable()));
+	return (GetPlayLength() > 0.f && (Super::IsNotifyAvailable() || AnimationTrack.IsNotifyAvailable()));
 }
 
 void UAnimComposite::GetAnimNotifiesFromDeltaPositions(const float& PreviousPosition, const float & CurrentPosition, TArray<FAnimNotifyEventReference>& OutActiveNotifies) const
@@ -43,7 +43,7 @@ void UAnimComposite::GetAnimNotifiesFromDeltaPositions(const float& PreviousPosi
 		}
 		else
 		{
-			AnimationTrack.GetAnimNotifiesFromTrackPositions(PreviousPosition, SequenceLength, OutActiveNotifies);
+			AnimationTrack.GetAnimNotifiesFromTrackPositions(PreviousPosition, GetPlayLength(), OutActiveNotifies);
 			AnimationTrack.GetAnimNotifiesFromTrackPositions(0.f, CurrentPosition, OutActiveNotifies);
 		}
 	}
@@ -56,7 +56,7 @@ void UAnimComposite::GetAnimNotifiesFromDeltaPositions(const float& PreviousPosi
 		else
 		{
 			AnimationTrack.GetAnimNotifiesFromTrackPositions(PreviousPosition, 0.f, OutActiveNotifies);
-			AnimationTrack.GetAnimNotifiesFromTrackPositions(SequenceLength, CurrentPosition, OutActiveNotifies);
+			AnimationTrack.GetAnimNotifiesFromTrackPositions(GetPlayLength(), CurrentPosition, OutActiveNotifies);
 		}
 	}
 }
diff --git a/Engine/Source/Runtime/Engine/Private/Animation/AnimCompositeBase.cpp b/Engine/Source/Runtime/Engine/Private/Animation/AnimCompositeBase.cpp
--- a/Engine/Source/Runtime/Engine/Private/Animation/AnimCompositeBase.cpp
+++ b/Engine/Source/Runtime/Engine/Private/Animation/AnimCompositeBase.cpp
@@ -406,9 +406,9 @@ void FAnimTrack::ValidateSegmentTimes()
 				Segment.StartPos = AnimSegments[J - 1].StartPos + AnimSegments[J - 1].GetLength();
 			}
 
-			if(Segment.AnimReference && Segment.AnimEndTime > Segment.AnimReference->SequenceLength)
+			if(Segment.AnimReference && Segment.AnimEndTime > Segment.AnimReference->GetPlayLength())
 			{
-				Segment.AnimEndTime = Segment.AnimReference->SequenceLength;
+				Segment.AnimEndTime = Segment.AnimReference->GetPlayLength();
 			}
 		}
 	}
@@ -678,7 +678,7 @@ bool FAnimTrack::IsValidToAdd(const UAnimSequenceBase* SequenceBase) const
 	// remove asset if invalid
 	if (SequenceBase)
 	{
-		if (SequenceBase->SequenceLength <= 0.f)
+		if (SequenceBase->GetPlayLength() <= 0.f)
 		{
 			UE_LOG(LogAnimation, Warning, TEXT("Remove Empty Sequence (%s)"), *SequenceBase->GetFullName());
 		}
@@ -711,13 +711,6 @@ UAnimCompositeBase::UAnimCompositeBase(const FObjectInitializer& ObjectInitializ
 {
 }
 
-#if WITH_EDITOR
-void UAnimCompositeBase::SetSequenceLength(float InSequenceLength)
-{
-	SequenceLength = InSequenceLength;
-}
-#endif
-
 void UAnimCompositeBase::ExtractRootMotionFromTrack(const FAnimTrack &SlotAnimTrack, float StartTrackPosition, float EndTrackPosition, FRootMotionMovementParams &RootMotion) const
 {
 	TArray<FRootMotionExtractionStep> RootMotionExtractionSteps;
diff --git a/Engine/Source/Runtime/Engine/Private/Animation/AnimCompressionTypes.cpp b/Engine/Source/Runtime/Engine/Private/Animation/AnimCompressionTypes.cpp
--- a/Engine/Source/Runtime/Engine/Private/Animation/AnimCompressionTypes.cpp
+++ b/Engine/Source/Runtime/Engine/Private/Animation/AnimCompressionTypes.cpp
@@ -154,7 +154,7 @@ FCompressibleAnimData::FCompressibleAnimData(class UAnimSequence* InSeq, const b
 	, BoneCompressionSettings(InSeq->BoneCompressionSettings)
 	, TrackToSkeletonMapTable(InSeq->GetRawTrackToSkeletonMapTable())
 	, Interpolation(InSeq->Interpolation)
-	, SequenceLength(InSeq->SequenceLength)
+	, SequenceLength(InSeq->GetPlayLength())
 	, NumFrames(InSeq->GetRawNumberOfFrames())
 	, bIsValidAdditive(InSeq->IsValidAdditive())
 #if WITH_EDITORONLY_DATA
diff --git a/Engine/Source/Runtime/Engine/Private/Animation/AnimInstance.cpp b/Engine/Source/Runtime/Engine/Private/Animation/AnimInstance.cpp
--- a/Engine/Source/Runtime/Engine/Private/Animation/AnimInstance.cpp
+++ b/Engine/Source/Runtime/Engine/Private/Animation/AnimInstance.cpp
@@ -793,7 +793,7 @@ void OutputTickRecords(const TArray<FAnimTickRecord>& Records, UCanvas* Canvas,
 		if (UAnimSequenceBase* AnimSeqBase = Cast<UAnimSequenceBase>(Player.SourceAsset))
 		{
 			PlayerEntry += FString::Printf(TEXT(" P(%.2f/%.2f)"), 
-				Player.TimeAccumulator != nullptr ? *Player.TimeAccumulator : 0.f, AnimSeqBase->SequenceLength);
+				Player.TimeAccumulator != nullptr ? *Player.TimeAccumulator : 0.f, AnimSeqBase->GetPlayLength());
 		}
 		else
 		{
@@ -1766,11 +1766,11 @@ float UAnimInstance::PlaySlotAnimation(UAnimSequenceBase* Asset, FName SlotNodeN
 	FAnimSegment NewSegment;
 	NewSegment.AnimReference = Asset;
 	NewSegment.AnimStartTime = 0.f;
-	NewSegment.AnimEndTime = Asset->SequenceLength;
+	NewSegment.AnimEndTime = Asset->GetPlayLength();
 	NewSegment.AnimPlayRate = 1.f;
 	NewSegment.StartPos = 0.f;
 	NewSegment.LoopingCount = LoopCount;
-	NewMontage->SequenceLength = NewSegment.GetLength();
+	NewMontage->SetSequenceLength(NewSegment.GetLength());
 	NewTrack.AnimTrack.AnimSegments.Add(NewSegment);
 		
 	FCompositeSection NewSection;
@@ -1877,7 +1877,7 @@ float UAnimInstance::Montage_Play(UAnimMontage* MontageToPlay, float InPlayRate/
 {
 	LLM_SCOPE(ELLMTag::Animation);
 
-	if (MontageToPlay && (MontageToPlay->SequenceLength > 0.f) && MontageToPlay->HasValidSlotSetup())
+	if (MontageToPlay && (MontageToPlay->GetPlayLength() > 0.f) && MontageToPlay->HasValidSlotSetup())
 	{
 		if (CurrentSkeleton && CurrentSkeleton->IsCompatible(MontageToPlay->GetSkeleton()))
 		{
@@ -1901,7 +1901,7 @@ float UAnimInstance::Montage_Play(UAnimMontage* MontageToPlay, float InPlayRate/
 			FAnimMontageInstance* NewInstance = new FAnimMontageInstance(this);
 			check(NewInstance);
 
-			const float MontageLength = MontageToPlay->SequenceLength;
+			const float MontageLength = MontageToPlay->GetPlayLength();
 
 			NewInstance->Initialize(MontageToPlay);
 			NewInstance->Play(InPlayRate);
diff --git a/Engine/Source/Runtime/Engine/Private/Animation/AnimLinkableElement.cpp b/Engine/Source/Runtime/Engine/Private/Animation/AnimLinkableElement.cpp
--- a/Engine/Source/Runtime/Engine/Private/Animation/AnimLinkableElement.cpp
+++ b/Engine/Source/Runtime/Engine/Private/Animation/AnimLinkableElement.cpp
@@ -39,14 +39,14 @@ void FAnimLinkableElement::LinkMontage(UAnimMontage* Montage, float AbsMontageTi
 
 void FAnimLinkableElement::LinkSequence(UAnimSequenceBase* Sequence, float AbsSequenceTime)
 {
-	if(Sequence && Sequence->SequenceLength > 0)
+	if(Sequence && Sequence->GetPlayLength() > 0)
 	{
 		LinkedMontage = nullptr;
 		LinkedSequence = Sequence;
 		SegmentIndex = 0;
 
 		SegmentBeginTime = 0.0f;
-		SegmentLength = Sequence->SequenceLength;
+		SegmentLength = Sequence->GetPlayLength();
 
 		SetTime(AbsSequenceTime);
 	}
diff --git a/Engine/Source/Runtime/Engine/Private/Animation/AnimMontage.cpp b/Engine/Source/Runtime/Engine/Private/Animation/AnimMontage.cpp
--- a/Engine/Source/Runtime/Engine/Private/Animation/AnimMontage.cpp
+++ b/Engine/Source/Runtime/Engine/Private/Animation/AnimMontage.cpp
@@ -108,14 +108,14 @@ bool UAnimMontage::IsWithinPos(int32 FirstIndex, int32 SecondIndex, float Curren
 	{
 		// @todo anim, I don't know if using SequenceLength is better or BIG_NUMBER
 		// I don't think that'd matter. 
-		EndTime = SequenceLength;
+		EndTime = GetPlayLength();
 	}
 
 	// since we do range of [StartTime, EndTime) (excluding EndTime) 
 	// there is blindspot of when CurrentTime becomes >= SequenceLength
 	// include that frame if CurrentTime gets there. 
 	// Otherwise, we continue to use [StartTime, EndTime)
-	if (CurrentTime >= SequenceLength)
+	if (CurrentTime >= GetPlayLength())
 	{
 		return (StartTime <= CurrentTime && EndTime >= CurrentTime);
 	}
@@ -173,7 +173,7 @@ float UAnimMontage::GetSectionTimeLeftFromPos(float Position)
 		}
 		else
 		{
-			return (SequenceLength - Position);
+			return (GetPlayLength() - Position);
 		}
 	}
 
@@ -230,7 +230,7 @@ bool UAnimMontage::IsValidSectionIndex(int32 SectionIndex) const
 void UAnimMontage::GetSectionStartAndEndTime(int32 SectionIndex, float& OutStartTime, float& OutEndTime) const
 {
 	OutStartTime = 0.f;
-	OutEndTime = SequenceLength;
+	OutEndTime = GetPlayLength();	
 	if ( IsValidSectionIndex(SectionIndex) )
 	{
 		OutStartTime = GetAnimCompositeSection(SectionIndex).GetTime();		
@@ -245,7 +245,7 @@ void UAnimMontage::GetSectionStartAndEndTime(int32 SectionIndex, float& OutStart
 float UAnimMontage::GetSectionLength(int32 SectionIndex) const
 {
 	float StartTime = 0.f;
-	float EndTime = SequenceLength;
+	float EndTime = GetPlayLength();
 	if ( IsValidSectionIndex(SectionIndex) )
 	{
 		StartTime = GetAnimCompositeSection(SectionIndex).GetTime();		
@@ -353,11 +353,10 @@ void UAnimMontage::PostLoad()
 		Track.ValidateSegmentTimes();
 
 		const float CurrentCalculatedLength = CalculateSequenceLength();
-
-		if(CurrentCalculatedLength != SequenceLength)
+		if(CurrentCalculatedLength != GetPlayLength())		
 		{
 			UE_LOG(LogAnimMontage, Display, TEXT("UAnimMontage::PostLoad: The actual sequence length for %s does not match the length stored in the asset, please resave the asset."), *GetFullName());
-			SequenceLength = CurrentCalculatedLength;
+			 SetSequenceLength(CurrentCalculatedLength);
 		}
 	}
 
@@ -1095,7 +1094,7 @@ void UAnimMontage::RefreshParentAssetData()
 			UAnimSequenceBase* TargetReference = Cast<UAnimSequenceBase>(AssetMappingTable->GetMappedAsset(SourceReference));
 			Segment.AnimReference = TargetReference;
 
-			float LengthChange = FMath::IsNearlyZero(SourceReference->SequenceLength) ? 0.f : TargetReference->SequenceLength / SourceReference->SequenceLength;
+			float LengthChange = FMath::IsNearlyZero(SourceReference->GetPlayLength()) ? 0.f : TargetReference->GetPlayLength() / SourceReference->GetPlayLength();
 			float RateChange = FMath::IsNearlyZero(SourceReference->RateScale) ? 0.f : FMath::Abs(TargetReference->RateScale / SourceReference->RateScale);
 			float TotalRateChange = FMath::IsNearlyZero(RateChange)? 0.f : (LengthChange / RateChange);
 			Segment.AnimPlayRate *= TotalRateChange;
@@ -1128,7 +1127,7 @@ void UAnimMontage::TickAssetPlayer(FAnimTickRecord& Instance, struct FAnimNotify
 		const float MoveDelta = Instance.Montage.MoveDelta;
 
 		Context.SetLeaderDelta(MoveDelta);
-		Context.SetPreviousAnimationPositionRatio(PreviousTime / SequenceLength);
+		Context.SetPreviousAnimationPositionRatio(PreviousTime / GetPlayLength());
 
 		if (MoveDelta != 0.f)
 		{
@@ -1174,7 +1173,7 @@ void UAnimMontage::TickAssetPlayer(FAnimTickRecord& Instance, struct FAnimNotify
 			}
 		}
 
-		Context.SetAnimationPositionRatio(CurrentTime / SequenceLength);
+		Context.SetAnimationPositionRatio(CurrentTime / GetPlayLength());
 	}
 
 	if (bRecordNeedsResetting && Instance.MarkerTickRecord)
@@ -1225,12 +1224,12 @@ void UAnimMontage::CollectMarkers()
 
 void UAnimMontage::GetMarkerIndicesForTime(float CurrentTime, bool bLooping, const TArray<FName>& ValidMarkerNames, FMarkerPair& OutPrevMarker, FMarkerPair& OutNextMarker) const
 {
-	MarkerData.GetMarkerIndicesForTime(CurrentTime, bLooping, ValidMarkerNames, OutPrevMarker, OutNextMarker, SequenceLength);
+	MarkerData.GetMarkerIndicesForTime(CurrentTime, bLooping, ValidMarkerNames, OutPrevMarker, OutNextMarker, GetPlayLength());
 }
 
 FMarkerSyncAnimPosition UAnimMontage::GetMarkerSyncPositionfromMarkerIndicies(int32 PrevMarker, int32 NextMarker, float CurrentTime) const
 {
-	return MarkerData.GetMarkerSyncPositionfromMarkerIndicies(PrevMarker, NextMarker, CurrentTime, SequenceLength);
+	return MarkerData.GetMarkerSyncPositionfromMarkerIndicies(PrevMarker, NextMarker, CurrentTime, GetPlayLength());
 }
 
 void UAnimMontage::InvalidateRecursiveAsset()
@@ -2843,16 +2842,16 @@ UAnimMontage* UAnimMontage::CreateSlotAnimationAsDynamicMontage(UAnimSequenceBas
 	FAnimSegment NewSegment;
 	NewSegment.AnimReference = Asset;
 	NewSegment.AnimStartTime = 0.f;
-	NewSegment.AnimEndTime = Asset->SequenceLength;
+	NewSegment.AnimEndTime = Asset->GetPlayLength();
 	NewSegment.AnimPlayRate = 1.f;
 	NewSegment.StartPos = 0.f;
 	NewSegment.LoopingCount = LoopCount;
-	NewMontage->SequenceLength = NewSegment.GetLength();
+	NewMontage->SetSequenceLength(NewSegment.GetLength());
 	NewTrack.AnimTrack.AnimSegments.Add(NewSegment);
 
 	FCompositeSection NewSection;
 	NewSection.SectionName = TEXT("Default");
-	NewSection.LinkSequence(Asset, Asset->SequenceLength);
+	NewSection.LinkSequence(Asset, Asset->GetPlayLength());
 	NewSection.SetTime(0.0f);
 
 	// add new section
@@ -2891,5 +2890,7 @@ void UAnimMontage::BakeTimeStretchCurve()
 		return;
 	}
 	
+	PRAGMA_DISABLE_DEPRECATION_WARNINGS
 	TimeStretchCurve.BakeFromFloatCurve(*TimeStretchFloatCurve, SequenceLength);
+	PRAGMA_ENABLE_DEPRECATION_WARNINGS
 }
diff --git a/Engine/Source/Runtime/Engine/Private/Animation/AnimNode_SequencePlayer.cpp b/Engine/Source/Runtime/Engine/Private/Animation/AnimNode_SequencePlayer.cpp
--- a/Engine/Source/Runtime/Engine/Private/Animation/AnimNode_SequencePlayer.cpp
+++ b/Engine/Source/Runtime/Engine/Private/Animation/AnimNode_SequencePlayer.cpp
@@ -26,7 +26,7 @@ float FAnimNode_SequencePlayer::GetCurrentAssetTimePlayRateAdjusted()
 
 float FAnimNode_SequencePlayer::GetCurrentAssetLength()
 {
-	return Sequence ? Sequence->SequenceLength : 0.0f;
+	return Sequence ? Sequence->GetPlayLength() : 0.0f;
 }
 
 void FAnimNode_SequencePlayer::Initialize_AnyThread(const FAnimationInitializeContext& Context)
@@ -47,12 +47,12 @@ void FAnimNode_SequencePlayer::Initialize_AnyThread(const FAnimationInitializeCo
 
 	if (Sequence != nullptr)
 	{
-		InternalTimeAccumulator = FMath::Clamp(StartPosition, 0.f, Sequence->SequenceLength);
+		InternalTimeAccumulator = FMath::Clamp(StartPosition, 0.f, Sequence->GetPlayLength());
 		const float AdjustedPlayRate = PlayRateScaleBiasClamp.ApplyTo(FMath::IsNearlyZero(PlayRateBasis) ? 0.f : (PlayRate / PlayRateBasis), 0.f);
 		const float EffectivePlayrate = Sequence->RateScale * AdjustedPlayRate;
 		if ((StartPosition == 0.f) && (EffectivePlayrate < 0.f))
 		{
-			InternalTimeAccumulator = Sequence->SequenceLength;
+			InternalTimeAccumulator = Sequence->GetPlayLength();
 		}
 	}
 }
@@ -75,7 +75,7 @@ void FAnimNode_SequencePlayer::UpdateAssetPlayer(const FAnimationUpdateContext&
 
 	if ((Sequence != nullptr) && (Context.AnimInstanceProxy->IsSkeletonCompatible(Sequence->GetSkeleton())))
 	{
-		InternalTimeAccumulator = FMath::Clamp(InternalTimeAccumulator, 0.f, Sequence->SequenceLength);
+		InternalTimeAccumulator = FMath::Clamp(InternalTimeAccumulator, 0.f, Sequence->GetPlayLength());
 		const float AdjustedPlayRate = PlayRateScaleBiasClamp.ApplyTo(FMath::IsNearlyZero(PlayRateBasis) ? 0.f : (PlayRate / PlayRateBasis), Context.GetDeltaTime());
 		CreateTickRecordForNode(Context, Sequence, bLoopAnimation, AdjustedPlayRate);
 	}
@@ -83,7 +83,7 @@ void FAnimNode_SequencePlayer::UpdateAssetPlayer(const FAnimationUpdateContext&
 #if ANIM_NODE_IDS_AVAILABLE && WITH_EDITORONLY_DATA
 	if (FAnimBlueprintDebugData* DebugData = Context.AnimInstanceProxy->GetAnimBlueprintDebugData())
 	{
-		DebugData->RecordSequencePlayer(Context.GetCurrentNodeId(), GetAccumulatedTime(), Sequence != nullptr ? Sequence->SequenceLength : 0.0f, Sequence != nullptr ? Sequence->GetNumberOfFrames() : 0);
+		DebugData->RecordSequencePlayer(Context.GetCurrentNodeId(), GetAccumulatedTime(), Sequence != nullptr ? Sequence->GetPlayLength() : 0.0f, Sequence != nullptr ? Sequence->GetNumberOfFrames() : 0);
 	}
 #endif
 
@@ -134,7 +134,7 @@ void FAnimNode_SequencePlayer::GatherDebugData(FNodeDebugData& DebugData)
 
 float FAnimNode_SequencePlayer::GetTimeFromEnd(float CurrentNodeTime)
 {
-	return Sequence->GetMaxCurrentTime() - CurrentNodeTime;
+	return Sequence->GetPlayLength() - CurrentNodeTime;
 }
 
 #undef LOCTEXT_NAMESPACE
\ No newline at end of file
diff --git a/Engine/Source/Runtime/Engine/Private/Animation/AnimNode_StateMachine.cpp b/Engine/Source/Runtime/Engine/Private/Animation/AnimNode_StateMachine.cpp
--- a/Engine/Source/Runtime/Engine/Private/Animation/AnimNode_StateMachine.cpp
+++ b/Engine/Source/Runtime/Engine/Private/Animation/AnimNode_StateMachine.cpp
@@ -683,7 +683,7 @@ bool FAnimNode_StateMachine::FindValidTransition(const FAnimationUpdateContext&
 			{
 				if (UAnimationAsset* AnimAsset = RelevantPlayer->GetAnimAsset())
 				{
-					const float AnimTimeRemaining = AnimAsset->GetMaxCurrentTime() - RelevantPlayer->GetAccumulatedTime();
+					const float AnimTimeRemaining = AnimAsset->GetPlayLength() - RelevantPlayer->GetAccumulatedTime();
 					const FAnimationTransitionBetweenStates& TransitionInfo = GetTransitionInfo(TransitionRule.TransitionIndex);
 					CrossfadeTimeAdjustment = TransitionInfo.CrossfadeDuration - AnimTimeRemaining;
 					bCanEnterTransition = (CrossfadeTimeAdjustment >= 0.f);
diff --git a/Engine/Source/Runtime/Engine/Private/Animation/AnimSequence.cpp b/Engine/Source/Runtime/Engine/Private/Animation/AnimSequence.cpp
--- a/Engine/Source/Runtime/Engine/Private/Animation/AnimSequence.cpp
+++ b/Engine/Source/Runtime/Engine/Private/Animation/AnimSequence.cpp
@@ -129,7 +129,7 @@ void OnCVarsChanged()
 		const FCompressedAnimSequence& AnimData = Seq->CompressedData;
 		const FUECompressedAnimData& UEAnimData = AnimData.CompressedDataStructure;
 		const int32 Additive = Seq->IsValidAdditive() ? 1 : 0;
-		OutputMessage += FString::Printf(TEXT("%s - %.2f Fr:%i Add:%i TO:%i SO:%i CBS:%i\n"), *Seq->GetName(), Seq->SequenceLength, Seq->GetRawNumberOfFrames(), Additive, UEAnimData.CompressedTrackOffsets.Num(), UEAnimData.CompressedScaleOffsets.OffsetData.Num(), UEAnimData.CompressedByteStream.Num());
+		OutputMessage += FString::Printf(TEXT("%s - %.2f Fr:%i Add:%i TO:%i SO:%i CBS:%i\n"), *Seq->GetName(), Seq->GetPlayLength(), Seq->GetRawNumberOfFrames(), Additive, UEAnimData.CompressedTrackOffsets.Num(), UEAnimData.CompressedScaleOffsets.OffsetData.Num(), UEAnimData.CompressedByteStream.Num());
 		OutputMessage += FString::Printf(TEXT("\t K:%i (%i : %i : %i)\n"), (int32)UEAnimData.KeyEncodingFormat, (int32)UEAnimData.TranslationCompressionFormat, (int32)UEAnimData.RotationCompressionFormat, (int32)UEAnimData.ScaleCompressionFormat);
 		OutputMessage += FString::Printf(TEXT("\t Curve Codec:%s\n"), AnimData.CurveCompressionCodec ? *AnimData.CurveCompressionCodec->GetPathName() : TEXT("nullptr"));
 		OutputMessage += FString::Printf(TEXT("\t TrackOff:%s\n"), *GetArrayGuid<int32>(UEAnimData.CompressedTrackOffsets).ToString());
@@ -804,7 +804,7 @@ void UAnimSequence::Serialize(FArchive& Ar)
 bool UAnimSequence::IsValidToPlay() const
 {
 	// make sure sequence length is valid and raw animation data exists, and compressed
-	return ( SequenceLength > 0.f);
+	return (GetPlayLength() > 0.f);
 }
 #endif
 
@@ -813,7 +813,7 @@ void UAnimSequence::SortSyncMarkers()
 	// First make sure all SyncMarkers are within a valid range
 	for (auto& SyncMarker : AuthoredSyncMarkers)
 	{
-		SyncMarker.Time = FMath::Clamp(SyncMarker.Time, 0.f, SequenceLength);
+		SyncMarker.Time = FMath::Clamp(SyncMarker.Time, 0.f, GetPlayLength());
 	}
 
 	// Then sort
@@ -953,10 +953,10 @@ void UAnimSequence::PostLoad()
 	}
 	// @remove temp hack for fixing length
 	// @todo need to fix importer/editing feature
-	else if ( SequenceLength == 0.f )
+	else if (GetPlayLength() == 0.f )
 	{
 		ensure(NumFrames == 1);
-		SequenceLength = MINIMUM_ANIMATION_LENGTH;
+		SetSequenceLength(MINIMUM_ANIMATION_LENGTH);
 	}
 	// Raw data exists, but missing compress animation data
 	else if( GetSkeleton() && !IsCompressedDataValid()
@@ -1227,7 +1227,7 @@ void UAnimSequence::GetBoneTransform(FTransform& OutAtom, int32 TrackIndex, floa
 	// If the caller didn't request that raw animation data be used . . .
 	if ( !bUseRawData && IsCompressedDataValid() )
 	{
-		FAnimSequenceDecompressionContext DecompContext(SequenceLength, Interpolation, GetFName(), *CompressedData.CompressedDataStructure);
+		FAnimSequenceDecompressionContext DecompContext(GetPlayLength(), Interpolation, GetFName(), *CompressedData.CompressedDataStructure);
 		DecompContext.Seek(Time);
 		CompressedData.BoneCompressionCodec->DecompressBone(DecompContext, TrackIndex, OutAtom);
 		return;
@@ -1290,7 +1290,7 @@ void UAnimSequence::ExtractBoneTransform(const struct FRawAnimSequenceTrack& Raw
 
 void UAnimSequence::ExtractBoneTransform(const struct FRawAnimSequenceTrack& RawTrack, FTransform& OutAtom, float Time) const
 {
-	FAnimationUtils::ExtractTransformFromTrack(Time, NumFrames, SequenceLength, RawTrack, Interpolation, OutAtom);
+	FAnimationUtils::ExtractTransformFromTrack(Time, NumFrames, GetPlayLength(), RawTrack, Interpolation, OutAtom);
 }
 
 void UAnimSequence::HandleAssetPlayerTickedInternal(FAnimAssetTickContext &Context, const float PreviousTime, const float MoveDelta, const FAnimTickRecord &Instance, struct FAnimNotifyQueue& NotifyQueue) const
@@ -1356,7 +1356,7 @@ FTransform UAnimSequence::ExtractRootMotion(float StartTime, float DeltaTime, bo
 		do
 		{
 			// Disable looping here. Advance to desired position, or beginning / end of animation 
-			const ETypeAdvanceAnim AdvanceType = FAnimationRuntime::AdvanceTime(false, DesiredDeltaMove, CurrentPosition, SequenceLength);
+			const ETypeAdvanceAnim AdvanceType = FAnimationRuntime::AdvanceTime(false, DesiredDeltaMove, CurrentPosition, GetPlayLength());
 
 			// Verify position assumptions
 			ensureMsgf(bPlayingBackwards ? (CurrentPosition <= PreviousPosition) : (CurrentPosition >= PreviousPosition), TEXT("in Animation %s(Skeleton %s) : bPlayingBackwards(%d), PreviousPosition(%0.2f), Current Position(%0.2f)"),
@@ -1370,7 +1370,7 @@ FTransform UAnimSequence::ExtractRootMotion(float StartTime, float DeltaTime, bo
 				const float ActualDeltaMove = (CurrentPosition - PreviousPosition);
 				DesiredDeltaMove -= ActualDeltaMove;
 
-				PreviousPosition = bPlayingBackwards ? SequenceLength : 0.f;
+				PreviousPosition = bPlayingBackwards ? GetPlayLength() : 0.f;
 				CurrentPosition = PreviousPosition;
 			}
 			else
@@ -1603,7 +1603,7 @@ void UAnimSequence::GetBonePose(FCompactPose& OutPose, FBlendedCurve& OutCurve,
 			}
 		}
 
-		BuildPoseFromRawData(AnimationData, TrackToSkeletonMapTable, OutPose, ExtractionContext.CurrentTime, Interpolation, NumFrames, SequenceLength, RetargetSource, &ActiveBoneCurves);
+		BuildPoseFromRawData(AnimationData, TrackToSkeletonMapTable, OutPose, ExtractionContext.CurrentTime, Interpolation, NumFrames, GetPlayLength(), RetargetSource, &ActiveBoneCurves);
 
 		if ((ExtractionContext.bExtractRootMotion && RootMotionReset.bEnableRootMotion) || RootMotionReset.bForceRootLock)
 		{
@@ -1613,7 +1613,7 @@ void UAnimSequence::GetBonePose(FCompactPose& OutPose, FBlendedCurve& OutCurve,
 	}
 #endif // WITH_EDITOR
 
-	DecompressPose(OutPose, CompressedData, ExtractionContext, GetSkeleton(), SequenceLength, Interpolation, bIsBakedAdditive, RetargetSource, GetFName(), RootMotionReset);
+	DecompressPose(OutPose, CompressedData, ExtractionContext, GetSkeleton(), GetPlayLength(), Interpolation, bIsBakedAdditive, RetargetSource, GetFName(), RootMotionReset);
 }
 
 #if WITH_EDITORONLY_DATA
@@ -1716,8 +1716,8 @@ void UAnimSequence::GetAdditiveBasePose(FCompactPose& OutPose, FBlendedCurve& Ou
 		case ABPT_AnimScaled:
 		{
 			// normalize time to fit base seq
-			const float Fraction = (SequenceLength > 0.f)? FMath::Clamp<float>(ExtractionContext.CurrentTime / SequenceLength, 0.f, 1.f) : 0.f;
-			const float BasePoseTime = RefPoseSeq->SequenceLength * Fraction;
+			const float Fraction = (GetPlayLength() > 0.f)? FMath::Clamp<float>(ExtractionContext.CurrentTime / GetPlayLength(), 0.f, 1.f) : 0.f;
+			const float BasePoseTime = RefPoseSeq->GetPlayLength() * Fraction;
 
 			FAnimExtractContext BasePoseExtractionContext(ExtractionContext);
 			BasePoseExtractionContext.CurrentTime = BasePoseTime;
@@ -1728,7 +1728,7 @@ void UAnimSequence::GetAdditiveBasePose(FCompactPose& OutPose, FBlendedCurve& Ou
 		case ABPT_AnimFrame:
 		{
 			const float Fraction = (RefPoseSeq->NumFrames > 0) ? FMath::Clamp<float>((float)RefFrameIndex / (float)RefPoseSeq->NumFrames, 0.f, 1.f) : 0.f;
-			const float BasePoseTime = RefPoseSeq->SequenceLength * Fraction;
+			const float BasePoseTime = RefPoseSeq->GetPlayLength() * Fraction;
 
 			FAnimExtractContext BasePoseExtractionContext(ExtractionContext);
 			BasePoseExtractionContext.CurrentTime = BasePoseTime;
@@ -1818,15 +1818,15 @@ void UAnimSequence::ResizeSequence(float NewLength, int32 NewNumFrames, bool bIn
 	check (StartFrame < EndFrame);
 
 	int32 OldNumFrames = NumFrames;
-	float OldSequenceLength = SequenceLength;
+	float OldSequenceLength = GetPlayLength();
 
 	// verify condition
 	NumFrames = NewNumFrames;
 	// Update sequence length to match new number of frames.
-	SequenceLength = NewLength;
+	SetSequenceLength(NewLength);
 
 	float Interval = OldSequenceLength / OldNumFrames;
-	ensure (Interval == SequenceLength/NumFrames);
+	ensure (Interval == GetPlayLength()/NumFrames);
 
 	float OldStartTime = StartFrame * Interval;
 	float OldEndTime = EndFrame * Interval;
@@ -1896,7 +1896,7 @@ void UAnimSequence::ResizeSequence(float NewLength, int32 NewNumFrames, bool bIn
 			}
 		}
 
-		float ClampedCurrentTime = FMath::Clamp(CurrentTime, 0.f, SequenceLength);
+		float ClampedCurrentTime = FMath::Clamp(CurrentTime, 0.f, GetPlayLength());
 		Notify.LinkSequence(this, ClampedCurrentTime);
 		Notify.SetDuration(NewDuration);
 
@@ -1904,7 +1904,7 @@ void UAnimSequence::ResizeSequence(float NewLength, int32 NewNumFrames, bool bIn
 		{
 			Notify.TriggerTimeOffset = GetTriggerTimeOffsetForType(EAnimEventTriggerOffsets::OffsetAfter);
 		}
-		else if (ClampedCurrentTime == SequenceLength)
+		else if (ClampedCurrentTime == GetPlayLength())
 		{
 			Notify.TriggerTimeOffset = GetTriggerTimeOffsetForType(EAnimEventTriggerOffsets::OffsetBefore);
 		}
@@ -1933,7 +1933,7 @@ void UAnimSequence::ResizeSequence(float NewLength, int32 NewNumFrames, bool bIn
 				CurrentTime -= Duration;
 			}
 		}
-		Marker.Time = FMath::Clamp(CurrentTime, 0.f, SequenceLength);
+		Marker.Time = FMath::Clamp(CurrentTime, 0.f, GetPlayLength());
 	}
 	// resize curves
 	RawCurveData.Resize(NewLength, bInsert, OldStartTime, OldEndTime);
@@ -1979,12 +1979,12 @@ bool UAnimSequence::InsertFramesToRawAnimData( int32 StartFrame, int32 EndFrame,
 			}
 		}
 
-		float const FrameTime = SequenceLength / ((float)NumFrames);
+		float const FrameTime = GetPlayLength() / ((float)NumFrames);
 
 		int32 NewNumFrames = NumFrames + NumFramesToInsert;
 		ResizeSequence((float)NewNumFrames * FrameTime, NewNumFrames, true, StartFrame, EndFrame);
 
-		UE_LOG(LogAnimation, Log, TEXT("\tSequenceLength: %f, NumFrames: %d"), SequenceLength, NumFrames);
+		UE_LOG(LogAnimation, Log, TEXT("\tGetPlayLength(): %f, NumFrames: %d"), GetPlayLength(), NumFrames);
 		
 		MarkRawDataAsModified();
 		MarkPackageDirty();
@@ -1998,7 +1998,7 @@ bool UAnimSequence::InsertFramesToRawAnimData( int32 StartFrame, int32 EndFrame,
 bool UAnimSequence::CropRawAnimData( float CurrentTime, bool bFromStart )
 {
 	// Length of one frame.
-	float const FrameTime = SequenceLength / ((float)NumFrames);
+	float const FrameTime = GetPlayLength() / ((float)NumFrames);
 	// Save Total Number of Frames before crop
 	int32 TotalNumOfFrames = NumFrames;
 
@@ -2014,7 +2014,7 @@ bool UAnimSequence::CropRawAnimData( float CurrentTime, bool bFromStart )
 	// causing if you were in below position, it will still crop 1 frame. 
 	// To be clearer, it seems better if we reject those inputs. 
 	// If you're a bit before/after, we assume that you'd like to crop
-	if ( CurrentTime == 0.f || CurrentTime == SequenceLength )
+	if ( CurrentTime == 0.f || CurrentTime == GetPlayLength() )
 	{
 		return false;
 	}
@@ -2022,8 +2022,8 @@ bool UAnimSequence::CropRawAnimData( float CurrentTime, bool bFromStart )
 	// Find the right key to cut at.
 	// This assumes that all keys are equally spaced (ie. won't work if we have dropped unimportant frames etc).
 	// The reason I'm changing to TotalNumOfFrames is CT/SL = KeyIndexWithFraction/TotalNumOfFrames
-	// To play TotalNumOfFrames, it takes SequenceLength. Each key will take SequenceLength/TotalNumOfFrames
-	float const KeyIndexWithFraction = (CurrentTime * (float)(TotalNumOfFrames)) / SequenceLength;
+	// To play TotalNumOfFrames, it takes GetPlayLength(). Each key will take GetPlayLength()/TotalNumOfFrames
+	float const KeyIndexWithFraction = (CurrentTime * (float)(TotalNumOfFrames)) / GetPlayLength();
 	int32 KeyIndex = bFromStart ? FMath::FloorToInt(KeyIndexWithFraction) : FMath::CeilToInt(KeyIndexWithFraction);
 	// Ensure KeyIndex is in range.
 	KeyIndex = FMath::Clamp<int32>(KeyIndex, 1, TotalNumOfFrames-1); 
@@ -2054,7 +2054,7 @@ bool UAnimSequence::CropRawAnimData( float CurrentTime, bool bFromStart )
 	// Update sequence length to match new number of frames.
 	ResizeSequence((float)NewNumFrames * FrameTime, NewNumFrames, false, StartKey, StartKey+NumKeys);
 
-	UE_LOG(LogAnimation, Log, TEXT("\tSequenceLength: %f, NumFrames: %d"), SequenceLength, NumFrames);
+	UE_LOG(LogAnimation, Log, TEXT("\tGetPlayLength(): %f, NumFrames: %d"), GetPlayLength(), NumFrames);
 
 	MarkRawDataAsModified();
 	OnRawDataChanged();
@@ -2402,11 +2402,11 @@ public:
 	TArray<FBoneIndexType> RequiredBoneIndexArray;
 
 	FByFramePoseEvalContext(const UAnimSequence* InAnimToEval)
-		: FByFramePoseEvalContext(InAnimToEval->SequenceLength, InAnimToEval->GetRawNumberOfFrames(), InAnimToEval->GetSkeleton())
+		: FByFramePoseEvalContext(InAnimToEval->GetPlayLength(), InAnimToEval->GetRawNumberOfFrames(), InAnimToEval->GetSkeleton())
 	{}
 		
-	FByFramePoseEvalContext(float InSequenceLength, int32 InRawNumOfFrames, USkeleton* InSkeleton)
-		: IntervalTime(InSequenceLength / ((float)InRawNumOfFrames - 1))
+	FByFramePoseEvalContext(float SequenceLength, int32 InRawNumOfFrames, USkeleton* InSkeleton)
+		: IntervalTime(SequenceLength / ((float)InRawNumOfFrames - 1))
 	{
 		// Initialize RequiredBones for pose evaluation
 		RequiredBones.SetUseRAWData(true);
@@ -2773,13 +2773,13 @@ bool UAnimSequence::CopyNotifies(UAnimSequence* SourceAnimSeq, UAnimSequence* De
 	// If the destination sequence is shorter than the source sequence, we'll be dropping notifies that
 	// occur at later times than the dest sequence is long.  Give the user a chance to abort if we
 	// find any notifies that won't be copied over.
-	if(bShowDialogs && DestAnimSeq->SequenceLength < SourceAnimSeq->SequenceLength)
+	if(bShowDialogs && DestAnimSeq->GetPlayLength() < SourceAnimSeq->GetPlayLength())
 	{
 		for(int32 NotifyIndex=0; NotifyIndex<SourceAnimSeq->Notifies.Num(); ++NotifyIndex)
 		{
 			// If a notify is found which occurs off the end of the destination sequence, prompt the user to continue.
 			const FAnimNotifyEvent& SrcNotifyEvent = SourceAnimSeq->Notifies[NotifyIndex];
-			if( SrcNotifyEvent.GetTriggerTime() > DestAnimSeq->SequenceLength )
+			if( SrcNotifyEvent.GetTriggerTime() > DestAnimSeq->GetPlayLength() )
 			{
 				const bool bProceed = EAppReturnType::Yes == FMessageDialog::Open( EAppMsgType::YesNo, NSLOCTEXT("UnrealEd", "SomeNotifiesWillNotBeCopiedQ", "Some notifies will not be copied because the destination sequence is not long enough.  Proceed?") );
 				if( !bProceed )
@@ -2815,7 +2815,7 @@ bool UAnimSequence::CopyNotifies(UAnimSequence* SourceAnimSeq, UAnimSequence* De
 		const FAnimNotifyEvent& SrcNotifyEvent = SourceAnimSeq->Notifies[NotifyIndex];
 
 		// Skip notifies which occur at times later than the destination sequence is long.
-		if( SrcNotifyEvent.GetTriggerTime() > DestAnimSeq->SequenceLength )
+		if( SrcNotifyEvent.GetTriggerTime() > DestAnimSeq->GetPlayLength() )
 		{
 			++NumNotifiesThatWereNotCopied;
 			continue;
@@ -3116,7 +3116,7 @@ void UAnimSequence::RemapTracksToNewSkeleton( USkeleton* NewSkeleton, bool bConv
 			ConvertedLocalSpaceAnimations.AddZeroed(SrcNumTracks);
 
 			int32 NumKeys = NumFrames;
-			float Interval = GetIntervalPerKey(NumFrames, SequenceLength);
+			float Interval = GetIntervalPerKey(NumFrames, GetPlayLength());
 
 			// allocate arrays
 			for(int32 SrcTrackIndex=0; SrcTrackIndex<SrcNumTracks; ++SrcTrackIndex)
@@ -3340,7 +3340,7 @@ void UAnimSequence::RemapTracksToNewSkeleton( USkeleton* NewSkeleton, bool bConv
 			ConvertedSpaceBases.AddZeroed(NumBones);
 
 			int32 NumKeys = NumFrames;
-			float Interval = GetIntervalPerKey(NumFrames, SequenceLength);
+			float Interval = GetIntervalPerKey(NumFrames, GetPlayLength());
 
 			// allocate arrays
 			for(int32 BoneIndex=0; BoneIndex<NumBones; ++BoneIndex)
@@ -3727,7 +3727,7 @@ void UAnimSequence::ReplaceReferredAnimations(const TMap<UAnimationAsset*, UAnim
 bool UAnimSequence::AddLoopingInterpolation()
 {
 	int32 NumTracks = AnimationTrackNames.Num();
-	float Interval = GetIntervalPerKey(NumFrames, SequenceLength);
+	float Interval = GetIntervalPerKey(NumFrames, GetPlayLength());
 
 	if(NumFrames > 0)
 	{
@@ -3757,7 +3757,7 @@ bool UAnimSequence::AddLoopingInterpolation()
 			}
 		}
 
-		SequenceLength += Interval;
+		SetSequenceLength(GetPlayLength() + Interval);
 		NumFrames = NewNumKeys;
 
 		PostProcessSequence();
@@ -3787,7 +3787,7 @@ int32 UAnimSequence::GetSpaceBasedAnimationData(TArray< TArray<FTransform> >& An
 
 	// 2d array of animated time [boneindex][time key]
 	int32 NumKeys = NumFrames;
-	float Interval = GetIntervalPerKey(NumFrames, SequenceLength);
+	float Interval = GetIntervalPerKey(NumFrames, GetPlayLength());
 
 	// allocate arrays
 	for (int32 BoneIndex = 0; BoneIndex < NumBones; ++BoneIndex)
@@ -4413,7 +4413,7 @@ void UAnimSequence::ResetAnimation()
 {
 	// clear everything. Making new animation, so need to reset all the things that belong here
 	NumFrames = 0;
-	SequenceLength = 0.f;
+	SetSequenceLength(0.f);
 	RawAnimationData.Empty();
 	AnimationTrackNames.Empty();
 	TrackToSkeletonMapTable.Empty();
@@ -4473,7 +4473,7 @@ bool UAnimSequence::CreateAnimation(USkeletalMesh* Mesh)
 		ResetAnimation();
 
 		const FReferenceSkeleton& RefSkeleton = Mesh->RefSkeleton;
-		SequenceLength = MINIMUM_ANIMATION_LENGTH;
+		SetSequenceLength(MINIMUM_ANIMATION_LENGTH);
 		NumFrames = 1;
 
 		const int32 NumBones = RefSkeleton.GetRawBoneNum();
@@ -4515,7 +4515,7 @@ bool UAnimSequence::CreateAnimation(USkeletalMeshComponent* MeshComponent)
 		ResetAnimation();
 
 		const FReferenceSkeleton& RefSkeleton = Mesh->RefSkeleton;
-		SequenceLength = MINIMUM_ANIMATION_LENGTH;
+		SetSequenceLength(MINIMUM_ANIMATION_LENGTH);
 		NumFrames = 1;
 
 		const int32 NumBones = RefSkeleton.GetRawBoneNum();
@@ -4554,7 +4554,7 @@ bool UAnimSequence::CreateAnimation(UAnimSequence* Sequence)
 	{
 		ResetAnimation();
 
-		SequenceLength = Sequence->SequenceLength;
+		SetSequenceLength(Sequence->GetPlayLength());
 		NumFrames = Sequence->NumFrames;
 
 		RawAnimationData = Sequence->RawAnimationData;
@@ -4708,9 +4708,9 @@ void UAnimSequence::AdvanceMarkerPhaseAsLeader(bool bLooping, float MoveDelta, c
 	float MarkerTimeOffset = 0.f;
 
 	// Hard to reproduce issue triggering this, ensure & clamp for now
-	ensureMsgf(CurrentTime >= 0.f && CurrentTime <= SequenceLength, TEXT("Current time inside of AdvanceMarkerPhaseAsLeader is out of range %.3f of 0.0 to %.3f\n    Sequence: %s"), CurrentTime, SequenceLength, *GetFullName());
+	ensureMsgf(CurrentTime >= 0.f && CurrentTime <= GetPlayLength(), TEXT("Current time inside of AdvanceMarkerPhaseAsLeader is out of range %.3f of 0.0 to %.3f\n    Sequence: %s"), CurrentTime, GetPlayLength(), *GetFullName());
 
-	CurrentTime = FMath::Clamp(CurrentTime, 0.f, SequenceLength);
+	CurrentTime = FMath::Clamp(CurrentTime, 0.f, GetPlayLength());
 
 	if (bPlayingForwards)
 	{
@@ -4719,8 +4719,8 @@ void UAnimSequence::AdvanceMarkerPhaseAsLeader(bool bLooping, float MoveDelta, c
 			if (NextMarker.MarkerIndex == -1)
 			{
 				float PrevCurrentTime = CurrentTime;
-				CurrentTime = FMath::Min(CurrentTime + CurrentMoveDelta, SequenceLength);
-				NextMarker.TimeToMarker = SequenceLength - CurrentTime;
+				CurrentTime = FMath::Min(CurrentTime + CurrentMoveDelta, GetPlayLength());
+				NextMarker.TimeToMarker = GetPlayLength() - CurrentTime;
 				PrevMarker.TimeToMarker -= CurrentTime - PrevCurrentTime; //Add how far we moved to distance from previous marker
 				break;
 			}
@@ -4731,7 +4731,7 @@ void UAnimSequence::AdvanceMarkerPhaseAsLeader(bool bLooping, float MoveDelta, c
 				bOffsetInitialized = true;
 				if (NextSyncMarker.Time < CurrentTime)
 				{
-					MarkerTimeOffset = SequenceLength;
+					MarkerTimeOffset = GetPlayLength();
 				}
 			}
 
@@ -4760,7 +4760,7 @@ void UAnimSequence::AdvanceMarkerPhaseAsLeader(bool bLooping, float MoveDelta, c
 							break;
 						}
 						NextMarker.MarkerIndex = 0;
-						MarkerTimeOffset += SequenceLength;
+						MarkerTimeOffset += GetPlayLength();
 					}
 				} while (!ValidMarkerNames.Contains(AuthoredSyncMarkers[NextMarker.MarkerIndex].MarkerName));
 				if (NextMarker.MarkerIndex != -1)
@@ -4770,10 +4770,10 @@ void UAnimSequence::AdvanceMarkerPhaseAsLeader(bool bLooping, float MoveDelta, c
 			}
 			else
 			{
-				CurrentTime = FMath::Fmod(CurrentTime + CurrentMoveDelta, SequenceLength);
+				CurrentTime = FMath::Fmod(CurrentTime + CurrentMoveDelta, GetPlayLength());
 				if (CurrentTime < 0.f)
 				{
-					CurrentTime += SequenceLength;
+					CurrentTime += GetPlayLength();
 				}
 				NextMarker.TimeToMarker = TimeToMarker - CurrentMoveDelta;
 				PrevMarker.TimeToMarker -= CurrentMoveDelta;
@@ -4800,7 +4800,7 @@ void UAnimSequence::AdvanceMarkerPhaseAsLeader(bool bLooping, float MoveDelta, c
 				bOffsetInitialized = true;
 				if (PrevSyncMarker.Time > CurrentTime)
 				{
-					MarkerTimeOffset = -SequenceLength;
+					MarkerTimeOffset = -GetPlayLength();
 				}
 			}
 
@@ -4829,7 +4829,7 @@ void UAnimSequence::AdvanceMarkerPhaseAsLeader(bool bLooping, float MoveDelta, c
 							break;
 						}
 						PrevMarker.MarkerIndex = AuthoredSyncMarkers.Num() - 1;
-						MarkerTimeOffset -= SequenceLength;
+						MarkerTimeOffset -= GetPlayLength();
 					}
 				} while (!ValidMarkerNames.Contains(AuthoredSyncMarkers[PrevMarker.MarkerIndex].MarkerName));
 				if (PrevMarker.MarkerIndex != -1)
@@ -4839,10 +4839,10 @@ void UAnimSequence::AdvanceMarkerPhaseAsLeader(bool bLooping, float MoveDelta, c
 			}
 			else
 			{
-				CurrentTime = FMath::Fmod(CurrentTime + CurrentMoveDelta, SequenceLength);
+				CurrentTime = FMath::Fmod(CurrentTime + CurrentMoveDelta, GetPlayLength());
 				if (CurrentTime < 0.f)
 				{
-					CurrentTime += SequenceLength;
+					CurrentTime += GetPlayLength();
 				}
 				PrevMarker.TimeToMarker = TimeToMarker - CurrentMoveDelta;
 				NextMarker.TimeToMarker -= CurrentMoveDelta;
@@ -4851,7 +4851,7 @@ void UAnimSequence::AdvanceMarkerPhaseAsLeader(bool bLooping, float MoveDelta, c
 		}
 	}
 
-	check(CurrentTime >= 0.f && CurrentTime <= SequenceLength);
+	check(CurrentTime >= 0.f && CurrentTime <= GetPlayLength());
 }
 
 void AdvanceMarkerForwards(int32& Marker, FName MarkerToFind, bool bLooping, const TArray<FAnimSyncMarker>& AuthoredSyncMarkers)
@@ -4973,7 +4973,7 @@ void UAnimSequence::AdvanceMarkerPhaseAsFollower(const FMarkerTickContext& Conte
 			if (NextMarker.MarkerIndex == -1)
 			{
 				check(!bLooping || Context.GetMarkerSyncEndPosition().NextMarkerName == NAME_None); // shouldnt have an end of anim marker if looping
-				CurrentTime = FMath::Min(CurrentTime + DeltaRemaining, SequenceLength);
+				CurrentTime = FMath::Min(CurrentTime + DeltaRemaining, GetPlayLength());
 				break;
 			}
 			else if (PassedMarkersIndex < Context.MarkersPassedThisTick.Num())
@@ -5075,11 +5075,11 @@ void UAnimSequence::GetMarkerIndicesForTime(float CurrentTime, bool bLooping, co
 	OutPrevMarker.MarkerIndex = -1;
 	OutPrevMarker.TimeToMarker = -CurrentTime;
 	OutNextMarker.MarkerIndex = -1;
-	OutNextMarker.TimeToMarker = SequenceLength - CurrentTime;
+	OutNextMarker.TimeToMarker = GetPlayLength() - CurrentTime;
 
 	for (int32 LoopMod = LoopModStart; LoopMod < LoopModEnd; ++LoopMod)
 	{
-		const float LoopModTime = LoopMod * SequenceLength;
+		const float LoopModTime = LoopMod * GetPlayLength();
 		for (int Idx = 0; Idx < AuthoredSyncMarkers.Num(); ++Idx)
 		{
 			const FAnimSyncMarker& Marker = AuthoredSyncMarkers[Idx];
@@ -5130,23 +5130,23 @@ FMarkerSyncAnimPosition UAnimSequence::GetMarkerSyncPositionfromMarkerIndicies(i
 	}
 	else
 	{
-		NextTime = SequenceLength;
+		NextTime = GetPlayLength();
 	}
 
 	// Account for looping
 	if(PrevTime > NextTime)
 	{
-		PrevTime = (PrevTime > CurrentTime) ? PrevTime - SequenceLength : PrevTime;
-		NextTime = (NextTime < CurrentTime) ? NextTime + SequenceLength : NextTime;
+		PrevTime = (PrevTime > CurrentTime) ? PrevTime - GetPlayLength() : PrevTime;
+		NextTime = (NextTime < CurrentTime) ? NextTime + GetPlayLength() : NextTime;
 	}
 	else if (PrevTime > CurrentTime)
 	{
-		CurrentTime += SequenceLength;
+		CurrentTime += GetPlayLength();
 	}
 
 	if (PrevTime == NextTime)
 	{
-		PrevTime -= SequenceLength;
+		PrevTime -= GetPlayLength();
 	}
 
 	check(NextTime > PrevTime);
@@ -5158,18 +5158,18 @@ FMarkerSyncAnimPosition UAnimSequence::GetMarkerSyncPositionfromMarkerIndicies(i
 float UAnimSequence::GetCurrentTimeFromMarkers(FMarkerPair& PrevMarker, FMarkerPair& NextMarker, float PositionBetweenMarkers) const
 {
 	float PrevTime = (PrevMarker.MarkerIndex != -1) ? AuthoredSyncMarkers[PrevMarker.MarkerIndex].Time : 0.f;
-	float NextTime = (NextMarker.MarkerIndex != -1) ? AuthoredSyncMarkers[NextMarker.MarkerIndex].Time : SequenceLength;
+	float NextTime = (NextMarker.MarkerIndex != -1) ? AuthoredSyncMarkers[NextMarker.MarkerIndex].Time : GetPlayLength();
 
 	if (PrevTime >= NextTime)
 	{
-		PrevTime -= SequenceLength; //Account for looping
+		PrevTime -= GetPlayLength(); //Account for looping
 	}
 	float CurrentTime = PrevTime + PositionBetweenMarkers * (NextTime - PrevTime);
 	if (CurrentTime < 0.f)
 	{
-		CurrentTime += SequenceLength;
+		CurrentTime += GetPlayLength();
 	}
-	CurrentTime = FMath::Clamp<float>(CurrentTime, 0, SequenceLength);
+	CurrentTime = FMath::Clamp<float>(CurrentTime, 0, GetPlayLength());
 
 	PrevMarker.TimeToMarker = PrevTime - CurrentTime;
 	NextMarker.TimeToMarker = NextTime - CurrentTime;
@@ -5202,7 +5202,7 @@ void UAnimSequence::GetMarkerIndicesForPosition(const FMarkerSyncAnimPosition& S
 			}
 		}
 
-		ensureMsgf(OutCurrentTime >= 0.f && OutCurrentTime <= SequenceLength, TEXT("Current time inside of GetMarkerIndicesForPosition is out of range %.3f of 0.0 to %.3f\n    Sequence: %s"), OutCurrentTime, SequenceLength, *GetFullName());
+		ensureMsgf(OutCurrentTime >= 0.f && OutCurrentTime <= GetPlayLength(), TEXT("Current time inside of GetMarkerIndicesForPosition is out of range %.3f of 0.0 to %.3f\n    Sequence: %s"), OutCurrentTime, GetPlayLength(), *GetFullName());
 		return;
 	}
 
@@ -5264,12 +5264,12 @@ void UAnimSequence::GetMarkerIndicesForPosition(const FMarkerSyncAnimPosition& S
 					float NextMarkerTime = AuthoredSyncMarkers[NextMarkerIdx].Time;
 					if (NextMarkerTime < PrevMarker.Time)
 					{
-						NextMarkerTime += SequenceLength;
+						NextMarkerTime += GetPlayLength();
 					}
 					float ThisCurrentTime = PrevMarker.Time + SyncPosition.PositionBetweenMarkers * (NextMarkerTime - PrevMarker.Time);
-					if (ThisCurrentTime > SequenceLength)
+					if (ThisCurrentTime > GetPlayLength())
 					{
-						ThisCurrentTime -= SequenceLength;
+						ThisCurrentTime -= GetPlayLength();
 					}
 					float ThisDiff = FMath::Abs(ThisCurrentTime - CurrentInputTime);
 					if (ThisDiff < DiffToCurrentTime)
diff --git a/Engine/Source/Runtime/Engine/Private/Animation/AnimSequenceBase.cpp b/Engine/Source/Runtime/Engine/Private/Animation/AnimSequenceBase.cpp
--- a/Engine/Source/Runtime/Engine/Private/Animation/AnimSequenceBase.cpp
+++ b/Engine/Source/Runtime/Engine/Private/Animation/AnimSequenceBase.cpp
@@ -102,9 +102,11 @@ void UAnimSequenceBase::PostLoad()
 	}
 }
 
-float UAnimSequenceBase::GetPlayLength()
+float UAnimSequenceBase::GetPlayLength() const
 {
+	PRAGMA_DISABLE_DEPRECATION_WARNINGS
 	return SequenceLength;
+	PRAGMA_ENABLE_DEPRECATION_WARNINGS
 }
 
 void UAnimSequenceBase::SortNotifies()
@@ -140,7 +142,7 @@ bool UAnimSequenceBase::RemoveNotifies(const TArray<FName>& NotifiesToRemove)
 
 bool UAnimSequenceBase::IsNotifyAvailable() const
 {
-	return (Notifies.Num() != 0) && (SequenceLength > 0.f);
+	return (Notifies.Num() != 0) && (GetPlayLength() > 0.f);
 }
 
 #if WITH_EDITOR
@@ -200,8 +202,8 @@ void UAnimSequenceBase::GetAnimNotifies(const float& StartTime, const float& Del
 
 	do 
 	{
-		// Disable looping here. Advance to desired position, or beginning / end of animation 
-		const ETypeAdvanceAnim AdvanceType = FAnimationRuntime::AdvanceTime(false, DesiredDeltaMove, CurrentPosition, SequenceLength);
+		// Disable looping here. Advance to desired position, or beginning / end of animation
+		const ETypeAdvanceAnim AdvanceType = FAnimationRuntime::AdvanceTime(false, DesiredDeltaMove, CurrentPosition, GetPlayLength());
 
 		// Verify position assumptions
 		ensureMsgf(bPlayingBackwards ? (CurrentPosition <= PreviousPosition) : (CurrentPosition >= PreviousPosition), TEXT("in Animation %s(Skeleton %s) : bPlayingBackwards(%d), PreviousPosition(%0.2f), Current Position(%0.2f)"), 
@@ -215,7 +217,7 @@ void UAnimSequenceBase::GetAnimNotifies(const float& StartTime, const float& Del
 			const float ActualDeltaMove = (CurrentPosition - PreviousPosition);
 			DesiredDeltaMove -= ActualDeltaMove; 
 
-			PreviousPosition = bPlayingBackwards ? SequenceLength : 0.f;
+			PreviousPosition = bPlayingBackwards ? GetPlayLength() : 0.f;
 			CurrentPosition = PreviousPosition;
 		}
 		else
@@ -305,7 +307,7 @@ void UAnimSequenceBase::TickAssetPlayer(FAnimTickRecord& Instance, struct FAnimN
 		MoveDelta = PlayRate * DeltaTime;
 
 		Context.SetLeaderDelta(MoveDelta);
-		Context.SetPreviousAnimationPositionRatio(PreviousTime / SequenceLength);
+		Context.SetPreviousAnimationPositionRatio(PreviousTime / GetPlayLength());
 
 		if (MoveDelta != 0.f)
 		{
@@ -316,12 +318,12 @@ void UAnimSequenceBase::TickAssetPlayer(FAnimTickRecord& Instance, struct FAnimN
 			else
 			{
 				// Advance time
-				FAnimationRuntime::AdvanceTime(Instance.bLooping, MoveDelta, CurrentTime, SequenceLength);
+				FAnimationRuntime::AdvanceTime(Instance.bLooping, MoveDelta, CurrentTime, GetPlayLength());
 				UE_LOG(LogAnimMarkerSync, Log, TEXT("Leader (%s) (normal advance)  - PreviousTime (%0.2f), CurrentTime (%0.2f), MoveDelta (%0.2f), Looping (%d) "), *GetName(), PreviousTime, CurrentTime, MoveDelta, Instance.bLooping ? 1 : 0);
 			}
 		}
 
-		Context.SetAnimationPositionRatio(CurrentTime / SequenceLength);
+		Context.SetAnimationPositionRatio(CurrentTime / GetPlayLength());
 	}
 	else
 	{
@@ -337,14 +339,14 @@ void UAnimSequenceBase::TickAssetPlayer(FAnimTickRecord& Instance, struct FAnimN
 				const float DeltaTime = Context.GetDeltaTime();
 				const float MyMoveDelta = PlayRate * DeltaTime;
 				// If leader is not valid, advance time as normal, do not jump position and pop.
-				FAnimationRuntime::AdvanceTime(Instance.bLooping, MyMoveDelta, CurrentTime, SequenceLength);
+				FAnimationRuntime::AdvanceTime(Instance.bLooping, MyMoveDelta, CurrentTime, GetPlayLength());
 				UE_LOG(LogAnimMarkerSync, Log, TEXT("Follower (%s) (normal advance)  - PreviousTime (%0.2f), CurrentTime (%0.2f), MoveDelta (%0.2f), Looping (%d) "), *GetName(), PreviousTime, CurrentTime, MoveDelta, Instance.bLooping ? 1 : 0);
 			}
 		}
 		else
 		{
-			PreviousTime = Context.GetPreviousAnimationPositionRatio() * SequenceLength;
-			CurrentTime = Context.GetAnimationPositionRatio() * SequenceLength;
+			PreviousTime = Context.GetPreviousAnimationPositionRatio() * GetPlayLength();
+			CurrentTime = Context.GetAnimationPositionRatio() * GetPlayLength();
 			UE_LOG(LogAnimMarkerSync, Log, TEXT("Follower (%s) (normalized position advance) - PreviousTime (%0.2f), CurrentTime (%0.2f), MoveDelta (%0.2f), Looping (%d) "), *GetName(), PreviousTime, CurrentTime, MoveDelta, Instance.bLooping ? 1 : 0);
 		}
 
@@ -357,7 +359,7 @@ void UAnimSequenceBase::TickAssetPlayer(FAnimTickRecord& Instance, struct FAnimN
 			// if we went against play rate, then loop around.
 			if( (MoveDelta * PlayRate) < 0.f )
 			{
-				MoveDelta += FMath::Sign<float>(PlayRate) * SequenceLength;
+				MoveDelta += FMath::Sign<float>(PlayRate) * GetPlayLength();
 			}
 		}
 	}
@@ -526,7 +528,7 @@ int32 UAnimSequenceBase::GetNumberOfFrames() const
 {
 	static float DefaultSampleRateInterval = 1.f / DEFAULT_SAMPLERATE;
 	// because of float error, add small margin at the end, so it can clamp correctly
-	return (int32)(SequenceLength / DefaultSampleRateInterval + KINDA_SMALL_NUMBER) + 1;
+	return (int32)(GetPlayLength() / DefaultSampleRateInterval + KINDA_SMALL_NUMBER) + 1;
 }
 
 #if WITH_EDITOR
@@ -546,14 +548,14 @@ void UAnimSequenceBase::InitializeNotifyTrack()
 
 int32 UAnimSequenceBase::GetFrameAtTime(const float Time) const
 {
-	const float FrameTime = GetNumberOfFrames() > 1 ? SequenceLength / (float)(GetNumberOfFrames() - 1) : 0.0f;
+	const float FrameTime = GetNumberOfFrames() > 1 ? GetPlayLength() / (float)(GetNumberOfFrames() - 1) : 0.0f;
 	return FMath::Clamp(FMath::RoundToInt(Time / FrameTime), 0, GetNumberOfFrames() - 1);
 }
 
 float UAnimSequenceBase::GetTimeAtFrame(const int32 Frame) const
 {
-	const float FrameTime = GetNumberOfFrames() > 1 ? SequenceLength / (float)(GetNumberOfFrames() - 1) : 0.0f;
-	return FMath::Clamp(FrameTime * Frame, 0.0f, SequenceLength);
+	const float FrameTime = GetNumberOfFrames() > 1 ? GetPlayLength() / (float)(GetNumberOfFrames() - 1) : 0.0f;
+	return FMath::Clamp(FrameTime * Frame, 0.0f, GetPlayLength());
 }
 
 void UAnimSequenceBase::RegisterOnNotifyChanged(const FOnNotifyChanged& Delegate)
@@ -567,10 +569,10 @@ void UAnimSequenceBase::UnregisterOnNotifyChanged(void* Unregister)
 
 void UAnimSequenceBase::ClampNotifiesAtEndOfSequence()
 {
-	const float NotifyClampTime = SequenceLength - 0.01f; //Slight offset so that notify is still draggable
+	const float NotifyClampTime = GetPlayLength() - 0.01f; //Slight offset so that notify is still draggable
 	for(int i = 0; i < Notifies.Num(); ++ i)
 	{
-		if(Notifies[i].GetTime() >= SequenceLength)
+		if(Notifies[i].GetTime() >= GetPlayLength())
 		{
 			Notifies[i].SetTime(NotifyClampTime);
 			Notifies[i].TriggerTimeOffset = GetTriggerTimeOffsetForType(EAnimEventTriggerOffsets::OffsetBefore);
@@ -584,7 +586,7 @@ EAnimEventTriggerOffsets::Type UAnimSequenceBase::CalculateOffsetForNotify(float
 	{
 		return EAnimEventTriggerOffsets::OffsetAfter;
 	}
-	else if(NotifyDisplayTime == SequenceLength)
+	else if(NotifyDisplayTime == GetPlayLength())
 	{
 		return EAnimEventTriggerOffsets::OffsetBefore;
 	}
@@ -693,7 +695,7 @@ void UAnimSequenceBase::RefreshParentAssetData()
 		NotifyEvent.EndLink.Link(this, NotifyEvent.GetTime() + NotifyEvent.Duration, NotifyEvent.GetSlotIndex());
 	}
 
-	SequenceLength = ParentSeqBase->SequenceLength;
+	SetSequenceLength(ParentSeqBase->GetPlayLength());
 	RateScale = ParentSeqBase->RateScale;
 	RawCurveData = ParentSeqBase->RawCurveData;
 
diff --git a/Engine/Source/Runtime/Engine/Private/Animation/AnimSingleNodeInstance.cpp b/Engine/Source/Runtime/Engine/Private/Animation/AnimSingleNodeInstance.cpp
--- a/Engine/Source/Runtime/Engine/Private/Animation/AnimSingleNodeInstance.cpp
+++ b/Engine/Source/Runtime/Engine/Private/Animation/AnimSingleNodeInstance.cpp
@@ -390,7 +390,7 @@ float UAnimSingleNodeInstance::GetLength()
 		}
 		else if (UAnimSequenceBase* SequenceBase = Cast<UAnimSequenceBase>(CurrentAsset))
 		{
-			return SequenceBase->SequenceLength;
+			return SequenceBase->GetPlayLength();
 		}
 	}
 
@@ -402,13 +402,13 @@ void UAnimSingleNodeInstance::StepForward()
 	if (UAnimSequenceBase* Sequence = Cast<UAnimSequenceBase>(CurrentAsset))
 	{
 		FAnimSingleNodeInstanceProxy& Proxy = GetProxyOnGameThread<FAnimSingleNodeInstanceProxy>();
-		const FAnimKeyHelper Helper(Sequence->SequenceLength, Sequence->GetNumberOfFrames());
+		const FAnimKeyHelper Helper(Sequence->GetPlayLength(), Sequence->GetNumberOfFrames());
 		float KeyLength = Helper.TimePerKey() + SMALL_NUMBER;
-		float Fraction = (Proxy.GetCurrentTime()+KeyLength)/Sequence->SequenceLength;
+		float Fraction = (Proxy.GetCurrentTime()+KeyLength)/Sequence->GetPlayLength();
 		int32 Frames = (float)Helper.LastKey() * Fraction;
 		if (IsLooping())
 		{
-			int32 PreviousFrame = (float)Helper.LastKey() * (Proxy.GetCurrentTime() / Sequence->SequenceLength);
+			int32 PreviousFrame = (float)Helper.LastKey() * (Proxy.GetCurrentTime() / Sequence->GetPlayLength());
 			Frames = (PreviousFrame == Helper.LastKey()) ? 0 : Frames;
 		}
 		else
@@ -425,9 +425,9 @@ void UAnimSingleNodeInstance::StepBackward()
 	{
 		FAnimSingleNodeInstanceProxy& Proxy = GetProxyOnGameThread<FAnimSingleNodeInstanceProxy>();
 
-		const FAnimKeyHelper Helper(Sequence->SequenceLength, Sequence->GetNumberOfFrames());
+		const FAnimKeyHelper Helper(Sequence->GetPlayLength(), Sequence->GetNumberOfFrames());
 		float KeyLength = Helper.TimePerKey() + SMALL_NUMBER;
-		float Fraction = (Proxy.GetCurrentTime()-KeyLength)/Sequence->SequenceLength;
+		float Fraction = (Proxy.GetCurrentTime()-KeyLength)/Sequence->GetPlayLength();
 		int32 Frames = (float)Helper.LastKey() * Fraction;
 		if (IsLooping())
 		{
diff --git a/Engine/Source/Runtime/Engine/Private/Animation/AnimSingleNodeInstanceProxy.cpp b/Engine/Source/Runtime/Engine/Private/Animation/AnimSingleNodeInstanceProxy.cpp
--- a/Engine/Source/Runtime/Engine/Private/Animation/AnimSingleNodeInstanceProxy.cpp
+++ b/Engine/Source/Runtime/Engine/Private/Animation/AnimSingleNodeInstanceProxy.cpp
@@ -383,7 +383,7 @@ void FAnimNode_SingleNode::Evaluate_AnyThread(FPoseContext& Output)
 				{
 #if WITH_EDITORONLY_DATA
 					// if montage is additive, we need to have base pose for the slot pose evaluate
-					if (bCanProcessAdditiveAnimationsLocal && Montage->PreviewBasePose && Montage->SequenceLength > 0.f)
+					if (bCanProcessAdditiveAnimationsLocal && Montage->PreviewBasePose && Montage->GetPlayLength() > 0.f)
 					{
 						Montage->PreviewBasePose->GetBonePose(LocalSourcePose, LocalSourceCurve, FAnimExtractContext(Proxy->CurrentTime));
 					}
@@ -522,7 +522,7 @@ void FAnimNode_SingleNode::Update_AnyThread(const FAnimationUpdateContext& Conte
 			if (!Proxy->bLooping)
 			{
 				const float CombinedPlayRate = NewPlayRate*Sequence->RateScale;
-				if ((CombinedPlayRate < 0.f && Proxy->CurrentTime <= 0.f) || (CombinedPlayRate > 0.f && Proxy->CurrentTime >= Sequence->SequenceLength))
+				if ((CombinedPlayRate < 0.f && Proxy->CurrentTime <= 0.f) || (CombinedPlayRate > 0.f && Proxy->CurrentTime >= Sequence->GetPlayLength()))
 				{
 					Proxy->SetPlaying(false);
 				}
@@ -539,7 +539,7 @@ void FAnimNode_SingleNode::Update_AnyThread(const FAnimationUpdateContext& Conte
 			if (!Proxy->bLooping)
 			{
 				const float CombinedPlayRate = NewPlayRate * Streamable->RateScale;
-				if ((CombinedPlayRate < 0.f && Proxy->CurrentTime <= 0.f) || (CombinedPlayRate > 0.f && Proxy->CurrentTime >= Streamable->SequenceLength))
+				if ((CombinedPlayRate < 0.f && Proxy->CurrentTime <= 0.f) || (CombinedPlayRate > 0.f && Proxy->CurrentTime >= Streamable->GetPlayLength()))
 				{
 					Proxy->SetPlaying(false);
 				}
@@ -556,7 +556,7 @@ void FAnimNode_SingleNode::Update_AnyThread(const FAnimationUpdateContext& Conte
 			if (!Proxy->bLooping)
 			{
 				const float CombinedPlayRate = NewPlayRate*Composite->RateScale;
-				if ((CombinedPlayRate < 0.f && Proxy->CurrentTime <= 0.f) || (CombinedPlayRate > 0.f && Proxy->CurrentTime >= Composite->SequenceLength))
+				if ((CombinedPlayRate < 0.f && Proxy->CurrentTime <= 0.f) || (CombinedPlayRate > 0.f && Proxy->CurrentTime >= Composite->GetPlayLength()))
 				{
 					Proxy->SetPlaying(false);
 				}
@@ -599,7 +599,7 @@ void FAnimNode_SingleNode::Update_AnyThread(const FAnimationUpdateContext& Conte
 		float MoveDelta = Proxy->GetDeltaSeconds() * NewPlayRate;
 		const bool bIsPreviewPoseLooping = true;
 
-		FAnimationRuntime::AdvanceTime(bIsPreviewPoseLooping, MoveDelta, Proxy->PreviewPoseCurrentTime, PreviewBasePose->SequenceLength);
+		FAnimationRuntime::AdvanceTime(bIsPreviewPoseLooping, MoveDelta, Proxy->PreviewPoseCurrentTime, PreviewBasePose->GetPlayLength());
 	}
 #endif
 }
diff --git a/Engine/Source/Runtime/Engine/Private/Animation/AnimStreamable.cpp b/Engine/Source/Runtime/Engine/Private/Animation/AnimStreamable.cpp
--- a/Engine/Source/Runtime/Engine/Private/Animation/AnimStreamable.cpp
+++ b/Engine/Source/Runtime/Engine/Private/Animation/AnimStreamable.cpp
@@ -275,7 +275,7 @@ void UAnimStreamable::GetAnimationPose(FCompactPose& OutPose, FBlendedCurve& Out
 			}
 		}
 
-		BuildPoseFromRawData(RawAnimationData, TrackToSkeletonMapTable, OutPose, ExtractionContext.CurrentTime, Interpolation, NumFrames, SequenceLength, RetargetSource);
+		BuildPoseFromRawData(RawAnimationData, TrackToSkeletonMapTable, OutPose, ExtractionContext.CurrentTime, Interpolation, NumFrames, GetPlayLength(), RetargetSource);
 
 		if ((ExtractionContext.bExtractRootMotion && RootMotionReset.bEnableRootMotion) || RootMotionReset.bForceRootLock)
 		{
@@ -442,7 +442,7 @@ void UAnimStreamable::InitFrom(const UAnimSequence* InSourceSequence)
 	AnimationTrackNames = InSourceSequence->GetAnimationTrackNames();
 	
 	NumFrames = InSourceSequence->GetNumberOfFrames();
-	SequenceLength = InSourceSequence->SequenceLength;
+	SetSequenceLength(InSourceSequence->GetPlayLength());
 	
 	RateScale = InSourceSequence->RateScale;
 
@@ -534,7 +534,7 @@ void UAnimStreamable::RequestCompressedData(const ITargetPlatform* Platform)
 	if (!Platform->IsServerOnly() && ChunkSizeSeconds > 0.f) // <= 0.f signifies to not chunk & don't have chunks on a server
 	{
 		ChunkSizeSeconds = FMath::Max(ChunkSizeSeconds, MINIMUM_CHUNK_SIZE);
-		const int32 InitialNumChunks = FMath::FloorToInt(SequenceLength / ChunkSizeSeconds);
+		const int32 InitialNumChunks = FMath::FloorToInt(GetPlayLength() / ChunkSizeSeconds);
 		NumChunks = FMath::Max(InitialNumChunks, 1);
 	}
 
@@ -620,7 +620,7 @@ void UAnimStreamable::RequestCompressedDataForChunk(const FString& ChunkDDCKey,
 		const bool bSkipDDC = false;
 
 		const int32 ChunkNumFrames = FrameEnd - FrameStart;
-		const float FrameLength = SequenceLength / (float)(NumFrames - 1);
+		const float FrameLength = GetPlayLength() / (float)(NumFrames - 1);
 		Chunk.StartTime = FrameStart * FrameLength;
 		Chunk.SequenceLength = ChunkNumFrames * FrameLength;
 
diff --git a/Engine/Source/Runtime/Engine/Private/Animation/AnimTrace.cpp b/Engine/Source/Runtime/Engine/Private/Animation/AnimTrace.cpp
--- a/Engine/Source/Runtime/Engine/Private/Animation/AnimTrace.cpp
+++ b/Engine/Source/Runtime/Engine/Private/Animation/AnimTrace.cpp
@@ -973,7 +973,7 @@ void FAnimTrace::OutputAnimSequencePlayer(const FAnimationBaseContext& InContext
 		<< AnimSequencePlayer.AnimInstanceId(FObjectTrace::GetObjectId(AnimInstance))
 		<< AnimSequencePlayer.NodeId(InContext.GetCurrentNodeId())
 		<< AnimSequencePlayer.Position(InNode.GetAccumulatedTime())
-		<< AnimSequencePlayer.Length(InNode.Sequence ? InNode.Sequence->SequenceLength : 0.0f)
+		<< AnimSequencePlayer.Length(InNode.Sequence ? InNode.Sequence->GetPlayLength() : 0.0f)
 		<< AnimSequencePlayer.FrameCounter(InNode.Sequence ? InNode.Sequence->GetNumberOfFrames() : 0);
 }
 
diff --git a/Engine/Source/Runtime/Engine/Private/Animation/BlendSpaceBase.cpp b/Engine/Source/Runtime/Engine/Private/Animation/BlendSpaceBase.cpp
--- a/Engine/Source/Runtime/Engine/Private/Animation/BlendSpaceBase.cpp
+++ b/Engine/Source/Runtime/Engine/Private/Animation/BlendSpaceBase.cpp
@@ -307,7 +307,7 @@ void UBlendSpaceBase::TickAssetPlayer(FAnimTickRecord& Instance, struct FAnimNot
 					{
 						SampleDataItem.MarkerTickRecord.Reset();
 						bResetMarkerDataOnFollowers = true;
-						SampleDataItem.Time = NormalizedCurrentTime * Sample.Animation->SequenceLength;
+						SampleDataItem.Time = NormalizedCurrentTime * Sample.Animation->GetPlayLength();
 					}
 					else if (!SampleDataItem.MarkerTickRecord.IsValid(Instance.bLooping) && Context.MarkerTickContext.GetMarkerSyncStartPosition().IsValid())
 					{
@@ -322,7 +322,7 @@ void UBlendSpaceBase::TickAssetPlayer(FAnimTickRecord& Instance, struct FAnimNot
 						check(!Instance.bLooping || Context.MarkerTickContext.IsMarkerSyncStartValid());
 						TickFollowerSamples(SampleDataList, HighestMarkerSyncWeightIndex, Context, bResetMarkerDataOnFollowers);
 					}
-					NormalizedCurrentTime = SampleDataItem.Time / Sample.Animation->SequenceLength;
+					NormalizedCurrentTime = SampleDataItem.Time / Sample.Animation->GetPlayLength();
 					*Instance.MarkerTickRecord = SampleDataItem.MarkerTickRecord;
 				}
 				else
@@ -353,13 +353,13 @@ void UBlendSpaceBase::TickAssetPlayer(FAnimTickRecord& Instance, struct FAnimNot
 					{
 						if(!Instance.MarkerTickRecord->IsValid(Instance.bLooping))
 						{
-							SampleDataItem.Time = NormalizedCurrentTime * Sample.Animation->SequenceLength;
+							SampleDataItem.Time = NormalizedCurrentTime * Sample.Animation->GetPlayLength();
 						}
 
 						TickFollowerSamples(SampleDataList, -1, Context, false);
 					}
 					*Instance.MarkerTickRecord = SampleDataItem.MarkerTickRecord;
-					NormalizedCurrentTime = SampleDataItem.Time / Sample.Animation->SequenceLength;
+					NormalizedCurrentTime = SampleDataItem.Time / Sample.Animation->GetPlayLength();
 				}
 				else
 				{
@@ -401,8 +401,8 @@ void UBlendSpaceBase::TickAssetPlayer(FAnimTickRecord& Instance, struct FAnimNot
 							{
 								const float SampleNormalizedPreviousTime = MultipliedSampleRateScale >= 0.f ? ClampedNormalizedPreviousTime : 1.f - ClampedNormalizedPreviousTime;
 								const float SampleNormalizedCurrentTime = MultipliedSampleRateScale >= 0.f ? ClampedNormalizedCurrentTime : 1.f - ClampedNormalizedCurrentTime;
-								PrevSampleDataTime = SampleNormalizedPreviousTime * Sample.Animation->SequenceLength;
-								CurrentSampleDataTime = SampleNormalizedCurrentTime * Sample.Animation->SequenceLength;
+								PrevSampleDataTime = SampleNormalizedPreviousTime * Sample.Animation->GetPlayLength();
+								CurrentSampleDataTime = SampleNormalizedCurrentTime * Sample.Animation->GetPlayLength();
 							}
 							else
 							{
@@ -416,7 +416,7 @@ void UBlendSpaceBase::TickAssetPlayer(FAnimTickRecord& Instance, struct FAnimNot
 							// if we went against play rate, then loop around.
 							if ((SampleMoveDelta * DeltaTimePosition) < 0.f)
 							{
-								DeltaTimePosition += FMath::Sign<float>(SampleMoveDelta) * Sample.Animation->SequenceLength;
+								DeltaTimePosition += FMath::Sign<float>(SampleMoveDelta) * Sample.Animation->GetPlayLength();
 							}
 
 							if( bGenerateNotifies && (!bTriggerNotifyHighestWeightedAnim || (I == HighestWeightIndex)))
@@ -610,7 +610,7 @@ void UBlendSpaceBase::GetAnimationPose(TArray<FBlendSampleData>& BlendSampleData
 #endif // WITH_EDITOR
 			)
 			{
-				const float Time = FMath::Clamp<float>(BlendSampleDataCache[I].Time, 0.f, Sample.Animation->SequenceLength);
+				const float Time = FMath::Clamp<float>(BlendSampleDataCache[I].Time, 0.f, Sample.Animation->GetPlayLength());
 
 				// first one always fills up the source one
 				Sample.Animation->GetAnimationPose(Pose, ChildrenCurves[I], FAnimExtractContext(Time, true));
@@ -807,10 +807,10 @@ void UBlendSpaceBase::ValidateSampleData()
 
 		if (Sample.bIsValid)
 		{
-			if (Sample.Animation->SequenceLength > AnimLength)
+			if (Sample.Animation->GetPlayLength() > AnimLength)
 			{
 				// @todo : should apply scale? If so, we'll need to apply also when blend
-				AnimLength = Sample.Animation->SequenceLength;
+				AnimLength = Sample.Animation->GetPlayLength();
 			}
 
 			Sample.CachedMarkerDataUpdateCounter = Sample.Animation->GetMarkerUpdateCounter();
@@ -1118,7 +1118,7 @@ float UBlendSpaceBase::GetAnimationLengthFromSampleData(const TArray<FBlendSampl
 				//Multiple samples contribution which we would otherwise lose
 				const float MultipliedSampleRateScale = Sample.Animation->RateScale * SampleDataList[I].SamplePlayRate;
 				// apply rate scale to get actual playback time
-				BlendAnimLength += (Sample.Animation->SequenceLength / ((MultipliedSampleRateScale) != 0.0f ? FMath::Abs(MultipliedSampleRateScale) : 1.0f))*SampleDataList[I].GetWeight();
+				BlendAnimLength += (Sample.Animation->GetPlayLength() / ((MultipliedSampleRateScale) != 0.0f ? FMath::Abs(MultipliedSampleRateScale) : 1.0f))*SampleDataList[I].GetWeight();
 				UE_LOG(LogAnimation, Verbose, TEXT("[%d] - Sample Animation(%s) : Weight(%0.5f) "), I+1, *Sample.Animation->GetName(), SampleDataList[I].GetWeight());
 			}
 		}
diff --git a/Engine/Source/Runtime/Engine/Private/Animation/PoseAsset.cpp b/Engine/Source/Runtime/Engine/Private/Animation/PoseAsset.cpp
--- a/Engine/Source/Runtime/Engine/Private/Animation/PoseAsset.cpp
+++ b/Engine/Source/Runtime/Engine/Private/Animation/PoseAsset.cpp
@@ -1164,7 +1164,7 @@ void UPoseAsset::CreatePoseFromAnimation(class UAnimSequence* AnimSequence, cons
 
 				// @Todo fill up curve data
 				TArray<float> CurveData;
-				float IntervalBetweenKeys = (NumPoses > 1)? AnimSequence->SequenceLength / (NumPoses -1 ) : 0.f;
+				float IntervalBetweenKeys = (NumPoses > 1)? AnimSequence->GetPlayLength() / (NumPoses -1 ) : 0.f;
 
 				// add curves - only float curves
 				const int32 TotalFloatCurveCount = AnimSequence->RawCurveData.FloatCurves.Num();
diff --git a/Engine/Source/Runtime/Engine/Private/Interpolation.cpp b/Engine/Source/Runtime/Engine/Private/Interpolation.cpp
--- a/Engine/Source/Runtime/Engine/Private/Interpolation.cpp
+++ b/Engine/Source/Runtime/Engine/Private/Interpolation.cpp
@@ -7250,7 +7250,7 @@ bool UInterpTrackAnimControl::GetClosestSnapPosition(float InPosition, TArray<in
 			UAnimSequence* Seq = AnimSeqs[i].AnimSeq;
 			if(Seq)
 			{
-				SeqLength = FMath::Max((Seq->SequenceLength - (AnimSeqs[i].AnimStartOffset + AnimSeqs[i].AnimEndOffset)) / AnimSeqs[i].AnimPlayRate, 0.01f);
+				SeqLength = FMath::Max((Seq->GetPlayLength() - (AnimSeqs[i].AnimStartOffset + AnimSeqs[i].AnimEndOffset)) / AnimSeqs[i].AnimPlayRate, 0.01f);
 				SeqEndTime += SeqLength;
 			}
 
@@ -7308,7 +7308,7 @@ float UInterpTrackAnimControl::GetTrackEndTime() const
 		{
 			// When calculating the end time, we do not consider the AnimStartOffset since we 
 			// are not calculating the length of the anim key. We just want the time where it ends.
-			EndTime += FMath::Max( (AnimSequence->SequenceLength - AnimKey.AnimEndOffset) / AnimKey.AnimPlayRate, 0.01f );
+			EndTime += FMath::Max( (AnimSequence->GetPlayLength() - AnimKey.AnimEndOffset) / AnimKey.AnimPlayRate, 0.01f );
 		}
 
 	}
@@ -7358,7 +7358,7 @@ bool UInterpTrackAnimControl::GetAnimForTime(float InTime, UAnimSequence** OutAn
 			UAnimSequence *Seq = AnimSeqs[i].AnimSeq;
 			if(Seq)
 			{
-				float SeqLength = FMath::Max(Seq->SequenceLength - (AnimSeqs[i].AnimStartOffset + AnimSeqs[i].AnimEndOffset), 0.01f);
+				float SeqLength = FMath::Max(Seq->GetPlayLength() - (AnimSeqs[i].AnimStartOffset + AnimSeqs[i].AnimEndOffset), 0.01f);
 
 				if(AnimSeqs[i].bLooping)
 				{
@@ -7367,14 +7367,14 @@ bool UInterpTrackAnimControl::GetAnimForTime(float InTime, UAnimSequence** OutAn
 				}
 				else
 				{
-					OutPosition = FMath::Clamp(OutPosition + AnimSeqs[i].AnimStartOffset, 0.f, (Seq->SequenceLength - AnimSeqs[i].AnimEndOffset) + (float)KINDA_SMALL_NUMBER);
+					OutPosition = FMath::Clamp(OutPosition + AnimSeqs[i].AnimStartOffset, 0.f, (Seq->GetPlayLength() - AnimSeqs[i].AnimEndOffset) + (float)KINDA_SMALL_NUMBER);
 				}
 
 				// Reverse position if the key is set to be reversed.
 				if(AnimSeqs[i].bReverse)
 				{
 					OutPosition = ConditionallyReversePosition(AnimSeqs[i], Seq, OutPosition);
-					bResetTime = (OutPosition == (Seq->SequenceLength - AnimSeqs[i].AnimEndOffset));
+					bResetTime = (OutPosition == (Seq->GetPlayLength() - AnimSeqs[i].AnimEndOffset));
 				}
 				else
 				{
@@ -7411,7 +7411,7 @@ float UInterpTrackAnimControl::ConditionallyReversePosition(FAnimControlTrackKey
 		// Reverse the clip.
 		if(Seq)
 		{
-			const float RealLength = Seq->SequenceLength - (SeqKey.AnimStartOffset+SeqKey.AnimEndOffset);
+			const float RealLength = Seq->GetPlayLength() - (SeqKey.AnimStartOffset+SeqKey.AnimEndOffset);
 			Result = (RealLength - (InPosition-SeqKey.AnimStartOffset))+SeqKey.AnimStartOffset;	// Mirror the cropped clip.
 		}
 	}
@@ -7546,10 +7546,10 @@ void UInterpTrackAnimControl::UpdateTrack(float NewPosition, UInterpTrackInst* T
 					if(Seq)
 					{
 						// Find position we should not play beyond in this sequence.
-						float SeqEnd = (Seq->SequenceLength - CurrentEndOffset);
+						float SeqEnd = (Seq->GetPlayLength() - CurrentEndOffset);
 
 						// Find time this sequence will take to play
-						float SeqLength = FMath::Max(Seq->SequenceLength - (CurrentStartOffset + CurrentEndOffset), 0.01f);
+						float SeqLength = FMath::Max(Seq->GetPlayLength() - (CurrentStartOffset + CurrentEndOffset), 0.01f);
 
 						// Find the number of loops we make. 
 						// @todo: This will need to be updated if we decide to support notifies in reverse.
@@ -7588,7 +7588,7 @@ void UInterpTrackAnimControl::UpdateTrack(float NewPosition, UInterpTrackInst* T
 					UAnimSequence* Seq = CurrentAnimSequence;
 					if( Seq )
 					{
-						float SeqEnd = (Seq->SequenceLength - CurrentEndOffset);
+						float SeqEnd = (Seq->GetPlayLength() - CurrentEndOffset);
 						AnimPos = FMath::Clamp( AnimPos, 0.f, SeqEnd + (float)KINDA_SMALL_NUMBER );
 					}
 
@@ -7652,7 +7652,7 @@ int32 UInterpTrackAnimControl::SplitKeyAtPosition(float InPosition)
 
 	// Check we are over an actual chunk of sequence.
 	float SplitAnimPos = ((InPosition - SplitKey.StartTime) * SplitKey.AnimPlayRate) + SplitKey.AnimStartOffset;
-	if(SplitAnimPos <= SplitKey.AnimStartOffset || SplitAnimPos >= (Seq->SequenceLength - SplitKey.AnimEndOffset))
+	if(SplitAnimPos <= SplitKey.AnimStartOffset || SplitAnimPos >= (Seq->GetPlayLength() - SplitKey.AnimEndOffset))
 	{
 		return INDEX_NONE;
 	}
@@ -7666,7 +7666,7 @@ int32 UInterpTrackAnimControl::SplitKeyAtPosition(float InPosition)
 	NewKey.AnimStartOffset = SplitAnimPos; // Start position in the new animation wants to be the place we are currently at.
 	NewKey.AnimEndOffset = SplitKey.AnimEndOffset; // End place is the same as the one we are splitting.
 
-	SplitKey.AnimEndOffset = Seq->SequenceLength - SplitAnimPos; // New end position is where we are.
+	SplitKey.AnimEndOffset = Seq->GetPlayLength() - SplitAnimPos; // New end position is where we are.
 	SplitKey.bLooping = false; // Disable looping for section before the cut.
 
 	// Add new key to track.
@@ -7697,7 +7697,7 @@ int32 UInterpTrackAnimControl::CropKeyAtPosition(float InPosition, bool bCutArea
 
 	// Check we are over an actual chunk of sequence.
 	float SplitAnimPos = ((InPosition - SplitKey.StartTime) * SplitKey.AnimPlayRate) + SplitKey.AnimStartOffset;
-	if(SplitAnimPos <= SplitKey.AnimStartOffset || SplitAnimPos >= (Seq->SequenceLength - SplitKey.AnimEndOffset))
+	if(SplitAnimPos <= SplitKey.AnimStartOffset || SplitAnimPos >= (Seq->GetPlayLength() - SplitKey.AnimEndOffset))
 	{
 		return INDEX_NONE;
 	}
@@ -7710,7 +7710,7 @@ int32 UInterpTrackAnimControl::CropKeyAtPosition(float InPosition, bool bCutArea
 	}
 	else
 	{
-		SplitKey.AnimEndOffset = Seq->SequenceLength - SplitAnimPos; // New end position is where we are.
+		SplitKey.AnimEndOffset = Seq->GetPlayLength() - SplitAnimPos; // New end position is where we are.
 	}
 
 	return SplitSeqIndex;
diff --git a/Engine/Source/Runtime/Engine/Private/InterpolationDraw.cpp b/Engine/Source/Runtime/Engine/Private/InterpolationDraw.cpp
--- a/Engine/Source/Runtime/Engine/Private/InterpolationDraw.cpp
+++ b/Engine/Source/Runtime/Engine/Private/InterpolationDraw.cpp
@@ -1024,7 +1024,7 @@ void UInterpTrackAnimControl::ObjectDragged(FInterpEdInputData& InputData)
 		UAnimSequence* Seq = AnimSeq.AnimSeq;
 		if(Seq)
 		{
-			float ActualLength = (Seq->SequenceLength - (OriginalKey->AnimStartOffset+OriginalKey->AnimEndOffset));
+			float ActualLength = (Seq->GetPlayLength() - (OriginalKey->AnimStartOffset+OriginalKey->AnimEndOffset));
 			float ActualLengthScaled =  ActualLength / OriginalKey->AnimPlayRate;
 			switch(InputData.InputType)
 			{
@@ -1041,7 +1041,7 @@ void UInterpTrackAnimControl::ObjectDragged(FInterpEdInputData& InputData)
 				{
 					// We are changing the offset but then scaling the animation proportionately so that the start and end times don't change
 					AnimSeq.AnimStartOffset = OriginalKey->AnimStartOffset + TimeDelta * AnimSeq.AnimPlayRate;
-					AnimSeq.AnimStartOffset = FMath::Clamp<float>(AnimSeq.AnimStartOffset, 0, Seq->SequenceLength-AnimSeq.AnimEndOffset);
+					AnimSeq.AnimStartOffset = FMath::Clamp<float>(AnimSeq.AnimStartOffset, 0, Seq->GetPlayLength() -AnimSeq.AnimEndOffset);
 
 					// Fix the play rate to keep the start and end times the same depending on how much the length of the clip actually changed by.
 					float ActualTimeChange = (AnimSeq.AnimStartOffset - OriginalKey->AnimStartOffset) / AnimSeq.AnimPlayRate;
@@ -1051,7 +1051,7 @@ void UInterpTrackAnimControl::ObjectDragged(FInterpEdInputData& InputData)
 				else
 				{
 					AnimSeq.AnimStartOffset = OriginalKey->AnimStartOffset + TimeDelta * AnimSeq.AnimPlayRate;
-					AnimSeq.AnimStartOffset = FMath::Clamp<float>(AnimSeq.AnimStartOffset, 0, Seq->SequenceLength-AnimSeq.AnimEndOffset);
+					AnimSeq.AnimStartOffset = FMath::Clamp<float>(AnimSeq.AnimStartOffset, 0, Seq->GetPlayLength() -AnimSeq.AnimEndOffset);
 					AnimSeq.StartTime = OriginalKey->StartTime + (AnimSeq.AnimStartOffset - OriginalKey->AnimStartOffset) / AnimSeq.AnimPlayRate;
 				}
 				break;
@@ -1067,7 +1067,7 @@ void UInterpTrackAnimControl::ObjectDragged(FInterpEdInputData& InputData)
 				{
 					// We are changing the offset but then scaling the animation proportionately so that the start and end times don't change
 					AnimSeq.AnimEndOffset = OriginalKey->AnimEndOffset - TimeDelta * AnimSeq.AnimPlayRate;
-					AnimSeq.AnimEndOffset = FMath::Clamp<float>(AnimSeq.AnimEndOffset, 0, Seq->SequenceLength-AnimSeq.AnimStartOffset);
+					AnimSeq.AnimEndOffset = FMath::Clamp<float>(AnimSeq.AnimEndOffset, 0, Seq->GetPlayLength() -AnimSeq.AnimStartOffset);
 
 					// Fix the play rate to keep the start and end times the same depending on how much the length of the clip actually changed by.
 					float ActualTimeChange = (AnimSeq.AnimEndOffset - OriginalKey->AnimEndOffset) / AnimSeq.AnimPlayRate;
@@ -1077,7 +1077,7 @@ void UInterpTrackAnimControl::ObjectDragged(FInterpEdInputData& InputData)
 				else
 				{
 					AnimSeq.AnimEndOffset = OriginalKey->AnimEndOffset - TimeDelta * AnimSeq.AnimPlayRate;
-					AnimSeq.AnimEndOffset = FMath::Clamp<float>(AnimSeq.AnimEndOffset, 0, Seq->SequenceLength-AnimSeq.AnimStartOffset);
+					AnimSeq.AnimEndOffset = FMath::Clamp<float>(AnimSeq.AnimEndOffset, 0, Seq->GetPlayLength() -AnimSeq.AnimStartOffset);
 				}
 				break;
 			}
@@ -1116,7 +1116,7 @@ void UInterpTrackAnimControl::DrawTrack( FCanvas* Canvas, UInterpGroup* Group, c
 		UAnimSequence* Seq = CurKey.AnimSeq;
 		if(Seq)
 		{
-			SeqLength = FMath::Max((Seq->SequenceLength - (CurKey.AnimStartOffset + CurKey.AnimEndOffset)) / CurKey.AnimPlayRate, 0.01f);
+			SeqLength = FMath::Max((Seq->GetPlayLength() - (CurKey.AnimStartOffset + CurKey.AnimEndOffset)) / CurKey.AnimPlayRate, 0.01f);
 			SeqEndTime += SeqLength;
 		}
 
@@ -1272,12 +1272,12 @@ void UInterpTrackAnimControl::DrawTrack( FCanvas* Canvas, UInterpGroup* Group, c
 				{
 					// Convert to the animation time values to frame numbers
 					const int32 AnimFrameOffsetFromStart = FMath::RoundToInt( CurKey.AnimStartOffset / Params.SnapAmount );
-					const int32 AnimFrameOffsetFromEnd = FMath::RoundToInt( ( Seq->SequenceLength - CurKey.AnimEndOffset ) / Params.SnapAmount );
+					const int32 AnimFrameOffsetFromEnd = FMath::RoundToInt( ( Seq->GetPlayLength() - CurKey.AnimEndOffset ) / Params.SnapAmount );
 					SeqString += FString::Printf( TEXT(" (%i->%i)"), AnimFrameOffsetFromStart, AnimFrameOffsetFromEnd );
 				}
 				else
 				{
-					SeqString += FString::Printf( TEXT(" (%2.2f->%2.2f)"), CurKey.AnimStartOffset, Seq->SequenceLength - CurKey.AnimEndOffset );
+					SeqString += FString::Printf( TEXT(" (%2.2f->%2.2f)"), CurKey.AnimStartOffset, Seq->GetPlayLength() - CurKey.AnimEndOffset );
 				}
 			}
 
diff --git a/Engine/Source/Runtime/Engine/Private/SingleAnimationPlayData.cpp b/Engine/Source/Runtime/Engine/Private/SingleAnimationPlayData.cpp
--- a/Engine/Source/Runtime/Engine/Private/SingleAnimationPlayData.cpp
+++ b/Engine/Source/Runtime/Engine/Private/SingleAnimationPlayData.cpp
@@ -32,7 +32,7 @@ void FSingleAnimationPlayData::ValidatePosition()
 		UAnimSequenceBase* SequenceBase = Cast<UAnimSequenceBase>(AnimToPlay);
 		if (SequenceBase)
 		{
-			Max = SequenceBase->SequenceLength;
+			Max = SequenceBase->GetPlayLength();
 		}
 	}
 
diff --git a/Engine/Source/Runtime/MovieSceneTracks/Private/Tracks/MovieSceneSkeletalAnimationTrack.cpp b/Engine/Source/Runtime/MovieSceneTracks/Private/Tracks/MovieSceneSkeletalAnimationTrack.cpp
--- a/Engine/Source/Runtime/MovieSceneTracks/Private/Tracks/MovieSceneSkeletalAnimationTrack.cpp
+++ b/Engine/Source/Runtime/MovieSceneTracks/Private/Tracks/MovieSceneSkeletalAnimationTrack.cpp
@@ -43,7 +43,7 @@ UMovieSceneSection* UMovieSceneSkeletalAnimationTrack::AddNewAnimationOnRow(FFra
 {
 	UMovieSceneSkeletalAnimationSection* NewSection = Cast<UMovieSceneSkeletalAnimationSection>(CreateNewSection());
 	{
-		FFrameTime AnimationLength = AnimSequence->SequenceLength * GetTypedOuter<UMovieScene>()->GetTickResolution();
+		FFrameTime AnimationLength = AnimSequence->GetPlayLength() * GetTypedOuter<UMovieScene>()->GetTickResolution();
 		int32 IFrameNumber = AnimationLength.FrameNumber.Value + (int)(AnimationLength.GetSubFrame() + 0.5f) + 1;
 		NewSection->InitialPlacementOnRow(AnimationSections, KeyTime, IFrameNumber, RowIndex);
 		NewSection->Params.Animation = AnimSequence;
diff --git a/Engine/Source/Runtime/MovieSceneTracks/Public/Sections/MovieSceneSkeletalAnimationSection.h b/Engine/Source/Runtime/MovieSceneTracks/Public/Sections/MovieSceneSkeletalAnimationSection.h
--- a/Engine/Source/Runtime/MovieSceneTracks/Public/Sections/MovieSceneSkeletalAnimationSection.h
+++ b/Engine/Source/Runtime/MovieSceneTracks/Public/Sections/MovieSceneSkeletalAnimationSection.h
@@ -17,10 +17,10 @@ struct FMovieSceneSkeletalAnimationParams
 	FMovieSceneSkeletalAnimationParams();
 
 	/** Gets the animation duration, modified by play rate */
-	float GetDuration() const { return FMath::IsNearlyZero(PlayRate) || Animation == nullptr ? 0.f : Animation->SequenceLength / PlayRate; }
+	float GetDuration() const { return FMath::IsNearlyZero(PlayRate) || Animation == nullptr ? 0.f : Animation->GetPlayLength() / PlayRate; }
 
 	/** Gets the animation sequence length, not modified by play rate */
-	float GetSequenceLength() const { return Animation != nullptr ? Animation->SequenceLength : 0.f; }
+	float GetSequenceLength() const { return Animation != nullptr ? Animation->GetPlayLength() : 0.f; }
 
 	/** The animation this section plays */
 	UPROPERTY(BlueprintReadWrite, EditAnywhere, Category="Animation", meta=(AllowedClasses = "AnimSequence,AnimComposite,AnimStreamable"))
