Adjust billboard scales when the texture size changes

diff --git a/Engine/Source/Runtime/Engine/Private/Components/BillboardComponent.cpp b/Engine/Source/Runtime/Engine/Private/Components/BillboardComponent.cpp
--- a/Engine/Source/Runtime/Engine/Private/Components/BillboardComponent.cpp
+++ b/Engine/Source/Runtime/Engine/Private/Components/BillboardComponent.cpp
@@ -60,7 +60,7 @@ public:
 		bWillEverBeLit = false;
 
 		// Calculate the scale factor for the sprite.
-		float Scale = InComponent->GetComponentTransform().GetMaximumAxisScale();
+		Scale = InComponent->GetComponentTransform().GetMaximumAxisScale() * SpriteScale * 0.25f;
 
 		OpacityMaskRefVal = InComponent->OpacityMaskRefVal;
 
@@ -74,15 +74,13 @@ public:
 #endif
 
 			// Set UL and VL to the size of the texture if they are set to 0.0, otherwise use the given value
-			UL = InComponent->UL == 0.0f ? InComponent->Sprite->GetSurfaceWidth() : InComponent->UL;
-			VL = InComponent->VL == 0.0f ? InComponent->Sprite->GetSurfaceHeight() : InComponent->VL;	
-			SizeX = Scale * UL * SpriteScale * 0.25f;
-			SizeY = Scale * VL * SpriteScale * 0.25f;
+			ComponentUL = InComponent->UL;
+			ComponentVL = InComponent->VL;
 		}
 		else
 		{
 			Texture = NULL;
-			SizeX = SizeY = UL = VL = 0;
+			ComponentUL = ComponentVL = 0;
 		}
 
 		if (AActor* Owner = InComponent->GetOwner())
@@ -136,6 +134,9 @@ public:
 		FTexture* TextureResource = Texture ? Texture->Resource : nullptr;
 		if (TextureResource)
 		{
+			const float UL = ComponentUL == 0.0f ? TextureResource->GetSizeX() : ComponentUL;
+			const float VL = ComponentVL == 0.0f ? TextureResource->GetSizeY() : ComponentVL;
+
 			for (int32 ViewIndex = 0; ViewIndex < Views.Num(); ViewIndex++)
 			{
 				if (VisibilityMap & (1 << ViewIndex))
@@ -143,8 +144,8 @@ public:
 					const FSceneView* View = Views[ViewIndex];
 
 					// Calculate the view-dependent scaling factor.
-					float ViewedSizeX = SizeX;
-					float ViewedSizeY = SizeY;
+					float ViewedSizeX = Scale * UL;
+					float ViewedSizeY = Scale * VL;
 
 					if (bIsScreenSizeScaled && (View->ViewMatrices.GetProjectionMatrix().M[3][3] != 1.0f))
 					{
@@ -257,14 +258,13 @@ public:
 
 private:
 	FVector Origin;
-	float SizeX;
-	float SizeY;
 	const float ScreenSize;
 	const UTexture2D* Texture;
+	float Scale;
 	const float U;
-	float UL;
+	float ComponentUL;
 	const float V;
-	float VL;
+	float ComponentVL;
 	float OpacityMaskRefVal;
 	FLinearColor Color;
 	const uint32 bIsScreenSizeScaled : 1;
