Wait for the compilation of some engine textures at startup

diff --git a/Engine/Source/Runtime/Engine/Private/UnrealEngine.cpp b/Engine/Source/Runtime/Engine/Private/UnrealEngine.cpp
--- a/Engine/Source/Runtime/Engine/Private/UnrealEngine.cpp
+++ b/Engine/Source/Runtime/Engine/Private/UnrealEngine.cpp
@@ -152,6 +152,7 @@ UnrealEngine.cpp: Implements the UEngine class and helpers.
 #include "Settings/LevelEditorPlaySettings.h"
 #include "LandscapeSubsystem.h"
 #include "TextureCompiler.h"
+#include "StaticMeshCompiler.h"
 #endif
 // @todo this is here only due to circular dependency to AIModule. To be removed
 
@@ -2647,6 +2648,11 @@ void UEngine::InitializeObjectReferences()
 	LoadEngineTexture(BlueNoiseTexture, *BlueNoiseTextureName.ToString());
 #endif
 
+#if WITH_EDITOR
+	// Avoid breaking some engine textures that might be cached very early (i.e. BlueNoise)
+	FTextureCompilingManager::Get().FinishAllCompilation();
+#endif
+
 	if ( DefaultPhysMaterial == NULL )
 	{
 		DefaultPhysMaterial = LoadObject<UPhysicalMaterial>(NULL, *DefaultPhysMaterialName.ToString(), NULL, LOAD_None, NULL);
