Remove include file placed there by mistake

diff --git a/Engine/Source/Runtime/Engine/Private/UnrealEngine.cpp b/Engine/Source/Runtime/Engine/Private/UnrealEngine.cpp
--- a/Engine/Source/Runtime/Engine/Private/UnrealEngine.cpp
+++ b/Engine/Source/Runtime/Engine/Private/UnrealEngine.cpp
@@ -152,7 +152,6 @@ UnrealEngine.cpp: Implements the UEngine class and helpers.
 #include "Settings/LevelEditorPlaySettings.h"
 #include "LandscapeSubsystem.h"
 #include "TextureCompiler.h"
-#include "StaticMeshCompiler.h"
 #endif
 // @todo this is here only due to circular dependency to AIModule. To be removed
 
