Fixed crash loading cooked static mesh in the editor

diff --git a/Engine/Source/Runtime/Engine/Private/StaticMesh.cpp b/Engine/Source/Runtime/Engine/Private/StaticMesh.cpp
--- a/Engine/Source/Runtime/Engine/Private/StaticMesh.cpp
+++ b/Engine/Source/Runtime/Engine/Private/StaticMesh.cpp
@@ -1432,18 +1432,13 @@ void FStaticMeshRenderData::SerializeInlineDataRepresentations(FArchive& Ar, USt
 
 			if (bValid)
 			{
-#if WITH_EDITOR
-				check(LOD.CardRepresentationData != nullptr);
-
-				Ar << *(LOD.CardRepresentationData);
-#else
 				if (LOD.CardRepresentationData == nullptr)
 				{
+					check(Ar.IsLoading());
 					LOD.CardRepresentationData = new FCardRepresentationData();
 				}
 
 				Ar << *(LOD.CardRepresentationData);
-#endif
 			}
 		}
 	}
