Remove async texture wait that is not needed anymore

diff --git a/Engine/Source/Runtime/Engine/Private/Components/BillboardComponent.cpp b/Engine/Source/Runtime/Engine/Private/Components/BillboardComponent.cpp
--- a/Engine/Source/Runtime/Engine/Private/Components/BillboardComponent.cpp
+++ b/Engine/Source/Runtime/Engine/Private/Components/BillboardComponent.cpp
@@ -68,11 +68,6 @@ public:
 		{
 			Texture = InComponent->Sprite;
 
-#if WITH_EDITOR
-			// Finish the compilation as we rely on the actual size of the texture to function properly
-			FTextureCompilingManager::Get().FinishCompilation({(UTexture*)Texture});
-#endif
-
 			// Set UL and VL to the size of the texture if they are set to 0.0, otherwise use the given value
 			ComponentUL = InComponent->UL;
 			ComponentVL = InComponent->VL;
