- WorldPartitionConversionCommandlet: Conversion of Foliage
- Renamed InstancedObjectsActor to PartitionActor
- Foliage InstanceHash now a 3D grid hash (was 2D grid)
- ActorPartitionSubsystem actors now named with CellCoord (ex: InstancedFoliageActor_x_y_z)

diff --git a/Engine/Source/Editor/UnrealEd/Private/Commandlets/WorldPartitionConvertCommandlet.cpp b/Engine/Source/Editor/UnrealEd/Private/Commandlets/WorldPartitionConvertCommandlet.cpp
--- a/Engine/Source/Editor/UnrealEd/Private/Commandlets/WorldPartitionConvertCommandlet.cpp
+++ b/Engine/Source/Editor/UnrealEd/Private/Commandlets/WorldPartitionConvertCommandlet.cpp
@@ -36,6 +36,8 @@
 #include "AssetRegistryModule.h"
 #include "FoliageHelper.h"
 #include "Engine/WorldComposition.h"
+#include "ActorPartition/ActorPartitionSubsystem.h"
+#include "InstancedFoliage.h"
 
 DEFINE_LOG_CATEGORY(LogWorldPartitionConvertCommandlet);
 
@@ -297,8 +299,7 @@ bool UWorldPartitionConvertCommandlet::ShouldDeleteActor(AActor* Actor, bool bMa
 {
 	if (Actor->HasAllFlags(RF_Transient) ||
 		Actor->IsA<ALODActor>() ||
-		Actor->IsA<ALevelBounds>() ||
-		Actor->IsA<AInstancedFoliageActor>())
+		Actor->IsA<ALevelBounds>())
 	{
 		return true;
 	}
@@ -757,9 +758,50 @@ int32 UWorldPartitionConvertCommandlet::Main(const FString& Params)
 	{
 		return 1;
 	}
+
+	// Make sure ActorPartitionSubsystem uses WorldPartition
+	UActorPartitionSubsystem* ActorPartitionSubsystem = MainWorld->GetSubsystem<UActorPartitionSubsystem>();
+	ActorPartitionSubsystem->InitializeForWorldPartitionConversion();
 	
-	auto PrepareLevelActors = [this](ULevel* Level, bool bMainLevel, EActorGridPlacement DefaultGridPlacement)
+	auto PartitionFoliage = [this, MainWorld, ActorPartitionSubsystem](AInstancedFoliageActor* IFA)
 	{
+		for (auto& Pair : IFA->FoliageInfos)
+		{
+			FFoliageInfo* FoliageInfo = &Pair.Value.Get();
+			if (FoliageInfo)
+			{
+				if (FoliageInfo->Type == EFoliageImplType::Actor)
+				{
+					// We don't support Actor Foliage in WP
+					FoliageInfo->ExcludeActors();
+				}
+				else if (FoliageInfo->Instances.Num())
+				{
+					FBox InstanceBounds = FoliageInfo->GetApproximatedInstanceBounds();
+					FActorPartitionGridHelper::ForEachIntersectingCell(AInstancedFoliageActor::StaticClass(), InstanceBounds, MainWorld->PersistentLevel, [IFA, FoliageInfo, ActorPartitionSubsystem](const UActorPartitionSubsystem::FCellCoord& InCellCoord, const FBox& InCellBounds)
+					{
+						TArray<int32> Indices;
+						FoliageInfo->GetInstancesInsideBounds(InCellBounds, Indices);
+						if (Indices.Num())
+						{
+							AInstancedFoliageActor* CellIFA = Cast<AInstancedFoliageActor>(ActorPartitionSubsystem->GetActor(AInstancedFoliageActor::StaticClass(), InCellCoord, true));
+							check(CellIFA);
+							FoliageInfo->MoveInstances(IFA, CellIFA, TSet<int32>(Indices), false);
+						}
+
+						return true;
+					});
+					check(!FoliageInfo->Instances.Num());
+				}
+			}
+		}
+
+		IFA->GetLevel()->GetWorld()->DestroyActor(IFA);
+	};
+
+	auto PrepareLevelActors = [this, PartitionFoliage](ULevel* Level, bool bMainLevel, EActorGridPlacement DefaultGridPlacement)
+	{
+		TArray<AInstancedFoliageActor*> IFAs;
 		for (AActor* Actor: Level->Actors)
 		{
 			if (Actor && !Actor->IsPendingKill())
@@ -768,6 +810,10 @@ int32 UWorldPartitionConvertCommandlet::Main(const FString& Params)
 				{
 					Level->GetWorld()->DestroyActor(Actor);
 				}
+				else if (AInstancedFoliageActor* IFA = Cast<AInstancedFoliageActor>(Actor))
+				{
+					IFAs.Add(IFA);
+				}
 				// Only override default grid placement on actors that are not marked as always loaded
 				else if (Actor->GridPlacement != EActorGridPlacement::AlwaysLoaded)
 				{
@@ -775,6 +821,12 @@ int32 UWorldPartitionConvertCommandlet::Main(const FString& Params)
 				}
 			}
 		}
+
+		// do loop after as it may modify Level->Actors
+		for (AInstancedFoliageActor* IFA : IFAs)
+		{
+			PartitionFoliage(IFA);
+		}
 	};
 
 	// Gather and load sublevels
@@ -883,12 +935,6 @@ int32 UWorldPartitionConvertCommandlet::Main(const FString& Params)
 				FArchiveGatherPrivateImports Ar(Actor, PrivateRefsMap, ActorsReferencesToActors);
 				Actor->Serialize(Ar);
 
-				//@todo_ow: fix this
-				if (FFoliageHelper::IsOwnedByFoliage(Actor))
-				{
-					FFoliageHelper::SetIsOwnedByFoliage(Actor, false);
-				}
-
 				ChangeObjectOuter(Actor, MainLevel);
 
 				// Migrate blueprint classes
diff --git a/Engine/Source/Runtime/Engine/Private/ActorPartition/ActorPartitionSubsystem.cpp b/Engine/Source/Runtime/Engine/Private/ActorPartition/ActorPartitionSubsystem.cpp
--- a/Engine/Source/Runtime/Engine/Private/ActorPartition/ActorPartitionSubsystem.cpp
+++ b/Engine/Source/Runtime/Engine/Private/ActorPartition/ActorPartitionSubsystem.cpp
@@ -1,8 +1,8 @@
 // Copyright Epic Games, Inc. All Rights Reserved.
 
 #include "ActorPartition/ActorPartitionSubsystem.h"
-#include "ActorPartition/InstancedObjectsActor.h"
-#include "WorldPartition/ActorPartition/InstancedObjectsActorDesc.h"
+#include "ActorPartition/PartitionActor.h"
+#include "WorldPartition/ActorPartition/PartitionActorDesc.h"
 #include "WorldPartition/WorldPartitionSubsystem.h"
 #include "Engine/World.h"
 
@@ -10,6 +10,44 @@ DEFINE_LOG_CATEGORY_STATIC(LogActorPartitionSubsystem, All, All);
 
 #if WITH_EDITOR
 
+FActorPartitionGetParams::FActorPartitionGetParams(const TSubclassOf<APartitionActor>& InActorClass, bool bInCreate, ULevel* InLevelHint, const FVector& InLocationHint)
+	: ActorClass(InActorClass)
+	, bCreate(bInCreate)
+	, LocationHint(InLocationHint)
+	, LevelHint(InLevelHint)
+{
+}
+
+void FActorPartitionGridHelper::ForEachIntersectingCell(const TSubclassOf<APartitionActor>& InActorClass, const FBox& InBounds, ULevel* InLevel, TFunctionRef<bool(const UActorPartitionSubsystem::FCellCoord&, const FBox&)> InOperation)
+{
+	uint32 GridSize = InActorClass->GetDefaultObject<APartitionActor>()->GetDefaultGridSize(InLevel->GetWorld());
+	UActorPartitionSubsystem::FCellCoord MinCellCoords = UActorPartitionSubsystem::FCellCoord::GetCellCoord(InBounds.Min, InLevel, GridSize);
+	UActorPartitionSubsystem::FCellCoord MaxCellCoords = UActorPartitionSubsystem::FCellCoord::GetCellCoord(InBounds.Max, InLevel, GridSize);
+
+	for (int32 z = MinCellCoords.Z; z <= MaxCellCoords.Z; z++)
+	{
+		for (int32 y = MinCellCoords.Y; y <= MaxCellCoords.Y; y++)
+		{
+			for (int32 x = MinCellCoords.X; x <= MaxCellCoords.X; x++)
+			{
+				UActorPartitionSubsystem::FCellCoord CellCoords(x, y, z, InLevel);
+				const FVector Min = FVector(
+					CellCoords.X * GridSize,
+					CellCoords.Y * GridSize,
+					CellCoords.Z * GridSize
+				);
+				const FVector Max = Min + FVector(GridSize);
+				FBox CellBounds(Min, Max);
+
+				if (!InOperation(MoveTemp(CellCoords), MoveTemp(CellBounds)))
+				{
+					return;
+				}
+			}
+		}
+	}
+}
+
 /**
  * FActorPartitionLevel
  */
@@ -33,30 +71,29 @@ public:
 		return UActorPartitionSubsystem::FCellCoord(0, 0, 0, SpawnLevel);
 	}
 	
-	AActor* GetActor(const FActorPartitionGetParams& GetParams, const UActorPartitionSubsystem::FCellCoord& CellCoord) override
+	APartitionActor* GetActor(const TSubclassOf<APartitionActor>& InActorClass, bool bInCreate, const UActorPartitionSubsystem::FCellCoord& InCellCoord) override
 	{
-		check(GetParams.LevelHint);
-		ULevel* SpawnLevel = GetSpawnLevel(GetParams.LevelHint, GetParams.LocationHint);
-
-		AActor* FoundActor = nullptr;
-		for (AActor* Actor : SpawnLevel->Actors)
+		check(InCellCoord.Level);
+		
+		APartitionActor* FoundActor = nullptr;
+		for (AActor* Actor : InCellCoord.Level->Actors)
 		{
-			if (Actor && Actor->IsA(GetParams.ActorClass))
+			if (APartitionActor* PartitionActor = Cast<APartitionActor>(Actor))
 			{
-				FoundActor = Actor;
+				FoundActor = PartitionActor;
 				break;
 			}
 		}
 		
-		if (!FoundActor && GetParams.bCreate)
+		if (!FoundActor && bInCreate)
 		{
 			FActorSpawnParameters SpawnParams;
-			SpawnParams.OverrideLevel = SpawnLevel;
+			SpawnParams.OverrideLevel = InCellCoord.Level;
 			SpawnParams.bCreateActorPackage = true;
-			FoundActor = World->SpawnActor(GetParams.ActorClass, &GetParams.LocationHint, nullptr, SpawnParams);
+			FoundActor = CastChecked<APartitionActor>(World->SpawnActor(InActorClass, nullptr, nullptr, SpawnParams));
 		}
 
-		check(FoundActor || !GetParams.bCreate);
+		check(FoundActor || !bInCreate);
 		return FoundActor;
 	}
 
@@ -95,82 +132,66 @@ public:
 	FActorPartitionWorldPartition(UWorld* InWorld)
 		: FBaseActorPartition(InWorld)
 	{
-		WorldSettings = World->GetWorldSettings();
-		check(WorldSettings);
-
 		WorldPartition = InWorld->GetSubsystem<UWorldPartitionSubsystem>();
 		check(WorldPartition);
 	}
 
 	UActorPartitionSubsystem::FCellCoord GetActorPartitionHash(const FActorPartitionGetParams& GetParams) const override 
 	{
-		return UActorPartitionSubsystem::FCellCoord(
-			FMath::FloorToInt(GetParams.LocationHint.X / GetParams.GridSize),
-			FMath::FloorToInt(GetParams.LocationHint.Y / GetParams.GridSize),
-			FMath::FloorToInt(GetParams.LocationHint.Z / GetParams.GridSize),
-			World->PersistentLevel
-		);
+		uint32 GridSize = GetParams.ActorClass->GetDefaultObject<APartitionActor>()->GetDefaultGridSize(World);
+
+		return UActorPartitionSubsystem::FCellCoord::GetCellCoord(GetParams.LocationHint, World->PersistentLevel, GridSize);
 	}
 
-	virtual AActor* GetActor(const FActorPartitionGetParams& GetParams, const UActorPartitionSubsystem::FCellCoord& CellCoord)
+	virtual APartitionActor* GetActor(const TSubclassOf<APartitionActor>& InActorClass, bool bInCreate, const UActorPartitionSubsystem::FCellCoord& InCellCoord)
 	{
-		AActor* FoundActor = nullptr;
-		
-		FBox CellBox(
-			FVector(
-				CellCoord.X * GetParams.GridSize, 
-				CellCoord.Y * GetParams.GridSize, 
-				CellCoord.Z * GetParams.GridSize
-			), 
-			FVector(
-				CellCoord.X * GetParams.GridSize + GetParams.GridSize, 
-				CellCoord.Y * GetParams.GridSize + GetParams.GridSize, 
-				CellCoord.Z * GetParams.GridSize + GetParams.GridSize
-			)
-		);
-
-		TArray<const FWorldPartitionActorDesc*> InstancedObjectsActorDescs = WorldPartition->GetIntersectingActorDescs(CellBox, GetParams.ActorClass);
+		APartitionActor* FoundActor = nullptr;
+		const uint32 GridSize = InActorClass->GetDefaultObject<APartitionActor>()->GetDefaultGridSize(World);
+		FBox CellBounds(UActorPartitionSubsystem::FCellCoord::GetCellBounds(InCellCoord, GridSize));
+
+		TArray<const FWorldPartitionActorDesc*> InstancedObjectsActorDescs = WorldPartition->GetIntersectingActorDescs(CellBounds, InActorClass);
 
 		for (const FWorldPartitionActorDesc* ActorDesc: InstancedObjectsActorDescs)
 		{
-			if (ActorDesc->GetActorClass() == GetParams.ActorClass)
+			if (ActorDesc->GetActorClass() == InActorClass)
 			{
-				FInstancedObjectsActorDesc* InstancedObjectsActorDesc = (FInstancedObjectsActorDesc*)ActorDesc;
-				if ((InstancedObjectsActorDesc->GridIndexX == CellCoord.X) && 
-					(InstancedObjectsActorDesc->GridIndexY == CellCoord.Y) && 
-					(InstancedObjectsActorDesc->GridIndexZ == CellCoord.Z))
+				FPartitionActorDesc* PartitionActorDesc = (FPartitionActorDesc*)ActorDesc;
+				if ((PartitionActorDesc->GridIndexX == InCellCoord.X) &&
+					(PartitionActorDesc->GridIndexY == InCellCoord.Y) &&
+					(PartitionActorDesc->GridIndexZ == InCellCoord.Z))
 				{
-					FoundActor = ActorDesc->GetActor();
+					AActor* DescActor = ActorDesc->GetActor();
 
-					if (!FoundActor)
+					if (!DescActor)
 					{
 						// Actor exists but is not loaded
 						return nullptr;
 					}
 				
-					AInstancedObjectsActor* InstancedObjectsActor = CastChecked<AInstancedObjectsActor>(FoundActor);
-					check(InstancedObjectsActor->GridSize == GetParams.GridSize);
+					FoundActor = CastChecked<APartitionActor>(DescActor);
+					check(FoundActor->GridSize == GridSize);
 				}
 			}
 		}
 
-		if (!FoundActor && GetParams.bCreate)
+		if (!FoundActor && bInCreate)
 		{
 			FActorSpawnParameters SpawnParams;
 			SpawnParams.bCreateActorPackage = true;
-			FVector CellCenter(CellBox.GetCenter());
-			FoundActor = World->SpawnActor(GetParams.ActorClass, &CellCenter, nullptr, SpawnParams);
-
-			AInstancedObjectsActor* InstancedObjectsActor = CastChecked<AInstancedObjectsActor>(FoundActor);
-			InstancedObjectsActor->GridSize = GetParams.GridSize;
+			SpawnParams.Name = FName(*FString::Printf(TEXT("%s_%d_%d_%d"), *InActorClass->GetName(), InCellCoord.X, InCellCoord.Y, InCellCoord.Z));
+			SpawnParams.NameMode = FActorSpawnParameters::ESpawnActorNameMode::Requested;
+							
+			FVector CellCenter(CellBounds.GetCenter());
+			FoundActor = CastChecked<APartitionActor>(World->SpawnActor(InActorClass, &CellCenter, nullptr, SpawnParams));
+			FoundActor->GridSize = GridSize;
+			FoundActor->bLockLocation = true;
 		}
 
-		check(FoundActor || !GetParams.bCreate);
+		check(FoundActor || !bInCreate);
 		return FoundActor;
 	}
 
 private:
-	AWorldSettings* WorldSettings;
 	UWorldPartitionSubsystem* WorldPartition;
 };
 
@@ -186,13 +207,19 @@ bool UActorPartitionSubsystem::IsLevelPartition() const
 }
 
 #if WITH_EDITOR
+void UActorPartitionSubsystem::InitializeForWorldPartitionConversion()
+{
+	check(IsRunningCommandlet());
+	ActorPartition.Reset(new FActorPartitionWorldPartition(GetWorld()));
+}
+
 void UActorPartitionSubsystem::Initialize(FSubsystemCollectionBase& Collection)
 {
 	Collection.InitializeDependency(UWorldPartitionSubsystem::StaticClass());
 
 	if (UWorldPartitionSubsystem* WorldPartitionSubsystem = GetWorld()->GetSubsystem<UWorldPartitionSubsystem>())
 	{
-		WorldPartitionSubsystem->RegisterActorDescFactory(AInstancedObjectsActor::StaticClass(), &InstancedObjectsActorDescFactory);
+		WorldPartitionSubsystem->RegisterActorDescFactory(APartitionActor::StaticClass(), &PartitionActorDescFactory);
 	}
 
 	// Will need to register to WorldPartition setup changes events here...
@@ -228,38 +255,41 @@ void UActorPartitionSubsystem::InitializeActorPartition()
 	ActorPartitionHashInvalidatedHandle = ActorPartition->GetOnActorPartitionHashInvalidated().AddUObject(this, &UActorPartitionSubsystem::OnActorPartitionHashInvalidated);
 }
 
-AActor* UActorPartitionSubsystem::GetActor(const FActorPartitionGetParams& GetParams)
+APartitionActor* UActorPartitionSubsystem::GetActor(const FActorPartitionGetParams& GetParams)
 {
-	check(GetParams.ActorClass->IsChildOf(AInstancedObjectsActor::StaticClass()));
-
 	FCellCoord CellCoord = ActorPartition->GetActorPartitionHash(GetParams);
-		
-	TMap<UClass*, TWeakObjectPtr<AActor>>* ActorsPerClass = PartitionedActors.Find(CellCoord);
-	AActor* FoundActor = nullptr;
+
+	return GetActor(GetParams.ActorClass, CellCoord, GetParams.bCreate);
+}
+
+APartitionActor* UActorPartitionSubsystem::GetActor(const TSubclassOf<APartitionActor>& InActorClass, const FCellCoord& InCellCoords, bool bInCreate)
+{
+	TMap<UClass*, TWeakObjectPtr<APartitionActor>>* ActorsPerClass = PartitionedActors.Find(InCellCoords);
+	APartitionActor* FoundActor = nullptr;
 	if (!ActorsPerClass)
 	{
-		FoundActor = ActorPartition->GetActor(GetParams, CellCoord);
+		FoundActor = ActorPartition->GetActor(InActorClass, bInCreate, InCellCoords);
 		if (FoundActor)
 		{
-			PartitionedActors.Add(CellCoord).Add(GetParams.ActorClass, FoundActor);
+			PartitionedActors.Add(InCellCoords).Add(InActorClass, FoundActor);
 			return FoundActor;
 		}
-		else if (GetParams.bCreate)
+		else if (bInCreate)
 		{
 			// Actor wasn't found and couldn't be created
 			return nullptr;
 		}
 	}
 
-	TWeakObjectPtr<AActor>* ActorPtr = ActorsPerClass->Find(GetParams.ActorClass);
+	TWeakObjectPtr<APartitionActor>* ActorPtr = ActorsPerClass->Find(InActorClass);
 	if (!ActorPtr || !ActorPtr->IsValid())
 	{
-		FoundActor = ActorPartition->GetActor(GetParams, CellCoord);
+		FoundActor = ActorPartition->GetActor(InActorClass, bInCreate, InCellCoords);
 		if (FoundActor)
 		{
 			if (!ActorPtr)
 			{
-				ActorsPerClass->Add(GetParams.ActorClass, FoundActor);
+				ActorsPerClass->Add(InActorClass, FoundActor);
 			}
 			else
 			{
diff --git a/Engine/Source/Runtime/Engine/Private/ActorPartition/InstancedObjectsActor.cpp b/Engine/Source/Runtime/Engine/Private/ActorPartition/InstancedObjectsActor.cpp
deleted file mode 100644
diff --git a/Engine/Source/Runtime/Engine/Private/ActorPartition/PartitionActor.cpp b/Engine/Source/Runtime/Engine/Private/ActorPartition/PartitionActor.cpp
new file mode 100644
--- /dev/null
+++ b/Engine/Source/Runtime/Engine/Private/ActorPartition/PartitionActor.cpp
@@ -0,0 +1,43 @@
+// Copyright Epic Games, Inc. All Rights Reserved.
+
+#include "ActorPartition/PartitionActor.h"
+
+#if WITH_EDITOR
+#include "ActorRegistry.h"
+#include "Components/BoxComponent.h"
+#endif
+
+APartitionActor::APartitionActor(const FObjectInitializer& ObjectInitializer)
+	: Super(ObjectInitializer)
+#if WITH_EDITORONLY_DATA
+	, GridSize(0)
+#endif
+{
+	USceneComponent* SceneComponent = CreateDefaultSubobject<USceneComponent>(TEXT("RootComponent0"));
+	RootComponent = SceneComponent;
+	RootComponent->Mobility = EComponentMobility::Static;
+}
+
+#if WITH_EDITOR
+void APartitionActor::GetAssetRegistryTags(TArray<FAssetRegistryTag>& OutTags) const
+{
+	Super::GetAssetRegistryTags(OutTags);
+
+	static const FName NAME_GridSize(TEXT("GridSize"));
+	FActorRegistry::SaveActorMetaData(NAME_GridSize, (int32)GridSize, OutTags);
+
+	const FVector ActorLocation = GetActorLocation();
+	const int64 GridIndexX = FMath::FloorToInt(ActorLocation.X / GridSize);
+	const int64 GridIndexY = FMath::FloorToInt(ActorLocation.Y / GridSize);
+	const int64 GridIndexZ = FMath::FloorToInt(ActorLocation.Z / GridSize);
+
+	static const FName NAME_GridIndexX(TEXT("GridIndexX"));
+	FActorRegistry::SaveActorMetaData(NAME_GridIndexX, GridIndexX, OutTags);
+
+	static const FName NAME_GridIndexY(TEXT("GridIndexY"));
+	FActorRegistry::SaveActorMetaData(NAME_GridIndexY, GridIndexY, OutTags);
+
+	static const FName NAME_GridIndexZ(TEXT("GridIndexZ"));
+	FActorRegistry::SaveActorMetaData(NAME_GridIndexZ, GridIndexZ, OutTags);
+}
+#endif
\ No newline at end of file
diff --git a/Engine/Source/Runtime/Engine/Private/WorldPartition/ActorPartition/InstancedObjectsActorDesc.cpp b/Engine/Source/Runtime/Engine/Private/WorldPartition/ActorPartition/InstancedObjectsActorDesc.cpp
deleted file mode 100644
diff --git a/Engine/Source/Runtime/Engine/Private/WorldPartition/ActorPartition/InstancedObjectsActorDescFactory.cpp b/Engine/Source/Runtime/Engine/Private/WorldPartition/ActorPartition/InstancedObjectsActorDescFactory.cpp
deleted file mode 100644
diff --git a/Engine/Source/Runtime/Engine/Private/WorldPartition/ActorPartition/PartitionActorDesc.cpp b/Engine/Source/Runtime/Engine/Private/WorldPartition/ActorPartition/PartitionActorDesc.cpp
new file mode 100644
--- /dev/null
+++ b/Engine/Source/Runtime/Engine/Private/WorldPartition/ActorPartition/PartitionActorDesc.cpp
@@ -0,0 +1,28 @@
+// Copyright Epic Games, Inc. All Rights Reserved.
+
+#include "WorldPartition/ActorPartition/PartitionActorDesc.h"
+#include "ActorPartition/PartitionActor.h"
+
+#if WITH_EDITOR
+#include "Engine/Level.h"
+FPartitionActorDesc::FPartitionActorDesc(const FWorldPartitionActorDescData& DescData, uint32 InGridSize, int64 InGridIndexX, int64 InGridIndexY, int64 InGridIndexZ)
+	: FWorldPartitionActorDesc(DescData)
+	, GridSize(InGridSize)
+	, GridIndexX(InGridIndexX)
+	, GridIndexY(InGridIndexY)
+	, GridIndexZ(InGridIndexZ)
+{}
+
+FPartitionActorDesc::FPartitionActorDesc(AActor* InActor)
+	: FWorldPartitionActorDesc(InActor)
+{
+	APartitionActor* PartitionActor = CastChecked<APartitionActor>(InActor);
+
+	GridSize = PartitionActor->GridSize;
+
+	const FVector ActorLocation = InActor->GetActorLocation();
+	GridIndexX = FMath::FloorToInt(ActorLocation.X / GridSize);
+	GridIndexY = FMath::FloorToInt(ActorLocation.Y / GridSize);
+	GridIndexZ = FMath::FloorToInt(ActorLocation.Z / GridSize);
+}
+#endif
diff --git a/Engine/Source/Runtime/Engine/Private/WorldPartition/ActorPartition/PartitionActorDescFactory.cpp b/Engine/Source/Runtime/Engine/Private/WorldPartition/ActorPartition/PartitionActorDescFactory.cpp
new file mode 100644
--- /dev/null
+++ b/Engine/Source/Runtime/Engine/Private/WorldPartition/ActorPartition/PartitionActorDescFactory.cpp
@@ -0,0 +1,51 @@
+// Copyright Epic Games, Inc. All Rights Reserved.
+
+#include "WorldPartition/ActorPartition/PartitionActorDescFactory.h"
+#include "WorldPartition/ActorPartition/PartitionActorDesc.h"
+#include "ActorRegistry.h"
+
+#if WITH_EDITOR
+FWorldPartitionActorDesc* FPartitionActorDescFactory::CreateInstance(const FWorldPartitionActorDescInitData& ActorDescInitData)
+{
+	FWorldPartitionActorDescData Data;
+	if (!ReadMetaData(ActorDescInitData, Data))
+	{
+		return nullptr;
+	}
+
+	int32 GridSize;
+	static const FName NAME_GridSize(TEXT("GridSize"));
+	if (!FActorRegistry::ReadActorMetaData(NAME_GridSize, GridSize, ActorDescInitData.AssetData))
+	{
+		return nullptr;
+	}
+
+	int64 GridIndexX;
+	static const FName NAME_GridIndexX(TEXT("GridIndexX"));
+	if (!FActorRegistry::ReadActorMetaData(NAME_GridIndexX, GridIndexX, ActorDescInitData.AssetData))
+	{
+		return nullptr;
+	}
+
+	int64 GridIndexY;
+	static const FName NAME_GridIndexY(TEXT("GridIndexY"));
+	if (!FActorRegistry::ReadActorMetaData(NAME_GridIndexY, GridIndexY, ActorDescInitData.AssetData))
+	{
+		return nullptr;
+	}
+
+	int64 GridIndexZ;
+	static const FName NAME_GridIndexZ(TEXT("GridIndexZ"));
+	if (!FActorRegistry::ReadActorMetaData(NAME_GridIndexZ, GridIndexZ, ActorDescInitData.AssetData))
+	{
+		return nullptr;
+	}
+
+	return new FPartitionActorDesc(Data, (uint32)GridSize, GridIndexX, GridIndexY, GridIndexZ);
+}
+
+FWorldPartitionActorDesc* FPartitionActorDescFactory::CreateInstance(AActor* InActor)
+{
+	return new FPartitionActorDesc(InActor);
+}
+#endif
\ No newline at end of file
diff --git a/Engine/Source/Runtime/Engine/Public/ActorPartition/ActorPartitionSubsystem.h b/Engine/Source/Runtime/Engine/Public/ActorPartition/ActorPartitionSubsystem.h
--- a/Engine/Source/Runtime/Engine/Public/ActorPartition/ActorPartitionSubsystem.h
+++ b/Engine/Source/Runtime/Engine/Public/ActorPartition/ActorPartitionSubsystem.h
@@ -5,26 +5,22 @@
 #include "Subsystems/WorldSubsystem.h"
 #include "Templates/SubclassOf.h"
 #include "Misc/HashBuilder.h"
-#include "WorldPartition/ActorPartition/InstancedObjectsActorDescFactory.h"
+#include "WorldPartition/ActorPartition/PartitionActorDescFactory.h"
 #include "ActorPartitionSubsystem.generated.h"
 
 class FBaseActorPartition;
+class APartitionActor;
 
+#if WITH_EDITOR
 /**
  * FActorPartitionGetParam
  */
-struct FActorPartitionGetParams
+struct ENGINE_API FActorPartitionGetParams
 {
-	FActorPartitionGetParams(const TSubclassOf<AActor>& InActorClass, bool bInCreate, ULevel* InLevelHint, const FVector& InLocationHint, int32 InGridSize)
-		: ActorClass(InActorClass)
-		, bCreate(bInCreate)
-		, LocationHint(InLocationHint)
-		, LevelHint(InLevelHint)
-		, GridSize(InGridSize)
-	{}
+	FActorPartitionGetParams(const TSubclassOf<APartitionActor>& InActorClass, bool bInCreate, ULevel* InLevelHint, const FVector& InLocationHint);
 
 	/* Class of Actor we are getting from the subsystem. */
-	TSubclassOf<AActor> ActorClass;
+	TSubclassOf<APartitionActor> ActorClass;
 	
 	/* Tells Subsystem if it needs to create Actor if it doesn't exist. */
 	bool bCreate;
@@ -34,11 +30,10 @@ struct FActorPartitionGetParams
 	
 	/* Depending on the world LevelHint can be used to find/create the Actor. */
 	ULevel* LevelHint;
-	
-	/* Desired size of the grid to create actors on if partitioned. */
-	int32 GridSize;
 };
 
+#endif
+
 /**
  * UActorPartitionSubsystem
  */
@@ -56,7 +51,7 @@ public:
 		FCellCoord()
 		{}
 
-		FCellCoord(int64 InX, int64 InY, int64 InZ, const ULevel* InLevel)
+		FCellCoord(int64 InX, int64 InY, int64 InZ, ULevel* InLevel)
 			: X(InX)
 			, Y(InY)
 			, Z(InZ)
@@ -66,7 +61,7 @@ public:
 		int64 X;
 		int64 Y;
 		int64 Z;
-		const ULevel* Level;
+		ULevel* Level;
 
 		bool operator==(const FCellCoord& Other) const
 		{
@@ -79,11 +74,40 @@ public:
 			HashBuilder << CellCoord.X << CellCoord.Y << CellCoord.Z << PointerHash(CellCoord.Level);
 			return HashBuilder.GetHash();
 		}
+
+		static FCellCoord GetCellCoord(FVector InPos, ULevel* InLevel, uint32 InGridSize)
+		{
+			return FCellCoord(
+				FMath::FloorToInt(InPos.X / InGridSize),
+				FMath::FloorToInt(InPos.Y / InGridSize),
+				FMath::FloorToInt(InPos.Z / InGridSize),
+				InLevel
+			);
+		}
+
+		static FBox GetCellBounds(const FCellCoord& InCellCoord, uint32 InGridSize)
+		{
+			return FBox(
+				FVector(
+					InCellCoord.X * InGridSize,
+					InCellCoord.Y * InGridSize,
+					InCellCoord.Z * InGridSize
+				),
+				FVector(
+					InCellCoord.X * InGridSize + InGridSize,
+					InCellCoord.Y * InGridSize + InGridSize,
+					InCellCoord.Z * InGridSize + InGridSize
+				)
+			);
+		}
 	};
-		
+
 #if WITH_EDITOR
-	AActor* GetActor(const FActorPartitionGetParams& GetParam);
-		
+	APartitionActor* GetActor(const FActorPartitionGetParams& GetParam);
+	APartitionActor* GetActor(const TSubclassOf<APartitionActor>& InActorClass, const FCellCoord& InCellCoords, bool bInCreate);
+
+	virtual void InitializeForWorldPartitionConversion();
+
 	virtual void Initialize(FSubsystemCollectionBase& Collection) override;
 	virtual void Deinitialize() override;
 #endif
@@ -96,10 +120,10 @@ private:
 
 	void InitializeActorPartition();
 
-	TMap<FCellCoord, TMap<UClass*, TWeakObjectPtr<AActor>>> PartitionedActors;
+	TMap<FCellCoord, TMap<UClass*, TWeakObjectPtr<APartitionActor>>> PartitionedActors;
 	TUniquePtr<FBaseActorPartition> ActorPartition;
 	
-	FInstancedObjectsActorDescFactory InstancedObjectsActorDescFactory;
+	FPartitionActorDescFactory PartitionActorDescFactory;
 
 	FDelegateHandle ActorPartitionHashInvalidatedHandle;
 #endif
@@ -115,7 +139,7 @@ public:
 	virtual ~FBaseActorPartition() {}
 
 	virtual UActorPartitionSubsystem::FCellCoord GetActorPartitionHash(const FActorPartitionGetParams& GetParams) const = 0;
-	virtual AActor* GetActor(const FActorPartitionGetParams& GetParams, const UActorPartitionSubsystem::FCellCoord& CellCoord) = 0;
+	virtual APartitionActor* GetActor(const TSubclassOf<APartitionActor>& InActorClass, bool bInCreate, const UActorPartitionSubsystem::FCellCoord& InCellCoord) = 0;
 
 	DECLARE_EVENT_OneParam(FBaseActorPartition, FOnActorPartitionHashInvalidated, const UActorPartitionSubsystem::FCellCoord&);
 	FOnActorPartitionHashInvalidated& GetOnActorPartitionHashInvalidated() { return OnActorPartitionHashInvalidated; }
@@ -123,4 +147,14 @@ protected:
 	UWorld* World;
 
 	FOnActorPartitionHashInvalidated OnActorPartitionHashInvalidated;
-};
\ No newline at end of file
+
+	friend class FActorPartitionGridHelper;
+};
+
+#if WITH_EDITOR
+class ENGINE_API FActorPartitionGridHelper
+{
+public:
+	static void ForEachIntersectingCell(const TSubclassOf<APartitionActor>& InActorClass, const FBox& InBounds, ULevel* InLevel, TFunctionRef<bool(const UActorPartitionSubsystem::FCellCoord&, const FBox&)> InOperation);
+};
+#endif
\ No newline at end of file
diff --git a/Engine/Source/Runtime/Engine/Public/ActorPartition/InstancedObjectsActor.h b/Engine/Source/Runtime/Engine/Public/ActorPartition/InstancedObjectsActor.h
deleted file mode 100644
diff --git a/Engine/Source/Runtime/Engine/Public/ActorPartition/PartitionActor.h b/Engine/Source/Runtime/Engine/Public/ActorPartition/PartitionActor.h
new file mode 100644
--- /dev/null
+++ b/Engine/Source/Runtime/Engine/Public/ActorPartition/PartitionActor.h
@@ -0,0 +1,38 @@
+// Copyright Epic Games, Inc. All Rights Reserved.
+#pragma once
+
+#include "CoreMinimal.h"
+#include "GameFramework/Actor.h"
+#include "PartitionActor.generated.h"
+
+class UBoxComponent;
+
+// Actor base class for instance containers placed on a grid.
+// See UActorPartitionSubsystem.
+UCLASS(Abstract)
+class ENGINE_API APartitionActor: public AActor
+{
+	GENERATED_BODY()
+
+public:
+	APartitionActor(const FObjectInitializer& ObjectInitializer = FObjectInitializer::Get());
+
+	//~ Begin UObject Interface
+#if WITH_EDITOR	
+	virtual void GetAssetRegistryTags(TArray<FAssetRegistryTag>& OutTags) const override;
+#endif
+	//~ End UObject Interface
+
+	//~ Begin AActor Interface
+#if WITH_EDITOR	
+	virtual EActorGridPlacement GetDefaultGridPlacement() const override { return EActorGridPlacement::Location; }
+	virtual uint32 GetDefaultGridSize(UWorld* InWorld) const PURE_VIRTUAL(APartitionActor, return 0;)
+#endif
+	//~ End AActor Interface	
+
+#if WITH_EDITORONLY_DATA
+	/** The grid size this actors was generated for */
+	UPROPERTY()
+	uint32 GridSize;
+#endif
+};
\ No newline at end of file
diff --git a/Engine/Source/Runtime/Engine/Public/WorldPartition/ActorPartition/InstancedObjectsActorDesc.h b/Engine/Source/Runtime/Engine/Public/WorldPartition/ActorPartition/InstancedObjectsActorDesc.h
deleted file mode 100644
diff --git a/Engine/Source/Runtime/Engine/Public/WorldPartition/ActorPartition/InstancedObjectsActorDescFactory.h b/Engine/Source/Runtime/Engine/Public/WorldPartition/ActorPartition/InstancedObjectsActorDescFactory.h
deleted file mode 100644
diff --git a/Engine/Source/Runtime/Engine/Public/WorldPartition/ActorPartition/PartitionActorDesc.h b/Engine/Source/Runtime/Engine/Public/WorldPartition/ActorPartition/PartitionActorDesc.h
new file mode 100644
--- /dev/null
+++ b/Engine/Source/Runtime/Engine/Public/WorldPartition/ActorPartition/PartitionActorDesc.h
@@ -0,0 +1,26 @@
+// Copyright Epic Games, Inc. All Rights Reserved.
+
+#pragma once
+
+#include "CoreMinimal.h"
+#include "UObject/ObjectMacros.h"
+#include "Containers/Map.h"
+#include "WorldPartition/WorldPartitionActorDesc.h"
+
+class ENGINE_API FPartitionActorDesc : public FWorldPartitionActorDesc
+{
+public:
+	uint32 GridSize;
+	int64 GridIndexX;
+	int64 GridIndexY;
+	int64 GridIndexZ;
+
+#if WITH_EDITOR
+protected:
+	FPartitionActorDesc() = delete;
+	FPartitionActorDesc(const FWorldPartitionActorDescData& DescData, uint32 InGridSize, int64 InGridIndexX, int64 InGridIndexY, int64 InGridIndexZ);
+	FPartitionActorDesc(AActor* InActor);
+
+	friend class FPartitionActorDescFactory;
+#endif
+};
diff --git a/Engine/Source/Runtime/Engine/Public/WorldPartition/ActorPartition/PartitionActorDescFactory.h b/Engine/Source/Runtime/Engine/Public/WorldPartition/ActorPartition/PartitionActorDescFactory.h
new file mode 100644
--- /dev/null
+++ b/Engine/Source/Runtime/Engine/Public/WorldPartition/ActorPartition/PartitionActorDescFactory.h
@@ -0,0 +1,17 @@
+// Copyright Epic Games, Inc. All Rights Reserved.
+#pragma once
+
+#include "CoreMinimal.h"
+#include "UObject/ObjectMacros.h"
+#include "WorldPartition/WorldPartitionActorDescFactory.h"
+
+class FWorldPartitionActorDesc;
+
+class ENGINE_API FPartitionActorDescFactory : public FWorldPartitionActorDescFactory
+{
+public:
+#if WITH_EDITOR
+	virtual FWorldPartitionActorDesc* CreateInstance(const FWorldPartitionActorDescInitData& ActorDescInitData) override;
+	virtual FWorldPartitionActorDesc* CreateInstance(AActor* InActor) override;
+#endif
+};
\ No newline at end of file
diff --git a/Engine/Source/Runtime/Foliage/Private/InstancedFoliage.cpp b/Engine/Source/Runtime/Foliage/Private/InstancedFoliage.cpp
--- a/Engine/Source/Runtime/Foliage/Private/InstancedFoliage.cpp
+++ b/Engine/Source/Runtime/Foliage/Private/InstancedFoliage.cpp
@@ -2106,6 +2106,18 @@ void FFoliageInfo::ReallocateClusters(AInstancedFoliageActor* InIFA, UFoliageTyp
 	Refresh(InIFA, true, true);
 }
 
+void FFoliageInfo::GetInstancesInsideBounds(const FBox& Box, TArray<int32>& OutInstances)
+{
+	auto TempInstances = InstanceHash->GetInstancesOverlappingBox(Box);
+	for (int32 Idx : TempInstances)
+	{
+		if (Box.IsInsideOrOn(Instances[Idx].Location))
+		{
+			OutInstances.Add(Idx);
+		}
+	}
+}
+
 void FFoliageInfo::GetInstancesInsideSphere(const FSphere& Sphere, TArray<int32>& OutInstances)
 {
 	auto TempInstances = InstanceHash->GetInstancesOverlappingBox(FBox::BuildAABB(Sphere.Center, FVector(Sphere.W)));
@@ -2328,6 +2340,11 @@ TArray<int32> FFoliageInfo::GetInstancesOverlappingBox(const FBox& Box) const
 {
 	return InstanceHash->GetInstancesOverlappingBox(Box);
 }
+
+FBox FFoliageInfo::GetApproximatedInstanceBounds() const
+{
+	return InstanceHash->GetBounds();
+}
 #endif	//WITH_EDITOR
 
 //
@@ -2496,8 +2513,7 @@ AInstancedFoliageActor* AInstancedFoliageActor::Get(UWorld* InWorld, bool bCreat
 				AInstancedFoliageActor::StaticClass(), 
 				bCreateIfNone, 
 				InLevelHint, 
-				InLocationHint, 
-				InWorld->GetWorldSettings()->InstancedFoliageGridSize
+				InLocationHint
 			)
 		)
 	);
@@ -4275,6 +4291,11 @@ void AInstancedFoliageActor::DetectFoliageTypeChangeAndUpdate()
 	}
 }
 
+uint32 AInstancedFoliageActor::GetDefaultGridSize(UWorld* InWorld) const
+{
+	return InWorld->GetWorldSettings()->InstancedFoliageGridSize;
+}
+
 void AInstancedFoliageActor::OnApplyLevelTransform(const FTransform& InTransform)
 {
 #if WITH_EDITORONLY_DATA
diff --git a/Engine/Source/Runtime/Foliage/Public/InstancedFoliage.h b/Engine/Source/Runtime/Foliage/Public/InstancedFoliage.h
--- a/Engine/Source/Runtime/Foliage/Public/InstancedFoliage.h
+++ b/Engine/Source/Runtime/Foliage/Public/InstancedFoliage.h
@@ -7,6 +7,7 @@ InstancedFoliage.h: Instanced foliage type definitions.
 
 #include "CoreMinimal.h"
 #include "Misc/Guid.h"
+#include "Misc/HashBuilder.h"
 #include "FoliageInstanceBase.h"
 
 class AInstancedFoliageActor;
@@ -295,6 +296,7 @@ struct FFoliageInfo
 	FOLIAGE_API void PostMoveInstances(AInstancedFoliageActor* InIFA, const TArray<int32>& InInstancesMoved, bool bFinished = false);
 	FOLIAGE_API void PostUpdateInstances(AInstancedFoliageActor* InIFA, const TArray<int32>& InInstancesUpdated, bool bReAddToHash = false, bool InUpdateSelection = false);
 	FOLIAGE_API void DuplicateInstances(AInstancedFoliageActor* InIFA, UFoliageType* InSettings, const TArray<int32>& InInstancesToDuplicate);
+	FOLIAGE_API void GetInstancesInsideBounds(const FBox& Box, TArray<int32>& OutInstances);
 	FOLIAGE_API void GetInstancesInsideSphere(const FSphere& Sphere, TArray<int32>& OutInstances);
 	FOLIAGE_API void GetInstanceAtLocation(const FVector& Location, int32& OutInstance, bool& bOutSucess);
 	FOLIAGE_API bool CheckForOverlappingSphere(const FSphere& Sphere);
@@ -328,6 +330,8 @@ struct FFoliageInfo
 	FOLIAGE_API void RemoveBaseComponentOnInstances();
 	FOLIAGE_API void IncludeActor(AInstancedFoliageActor* IFA, const UFoliageType* FoliageType, AActor* InActor);
 	FOLIAGE_API void ExcludeActors();
+
+	FOLIAGE_API FBox GetApproximatedInstanceBounds() const;
 #endif
 
 	friend FArchive& operator<<(FArchive& Ar, FFoliageInfo& MeshInfo);
@@ -351,22 +355,54 @@ private:
 // FFoliageInstanceHash
 //
 
-#define FOLIAGE_HASH_CELL_BITS 9	// 512x512 grid
+#define FOLIAGE_HASH_CELL_BITS 9	// 512x512x512 grid
 
 struct FFoliageInstanceHash
 {
 private:
-	const int32 HashCellBits;
-	TMap<uint64, TSet<int32>> CellMap;
+	struct FKey
+	{
+		int32 X;
+		int32 Y;
+		int32 Z;
+
+		FKey() = default;
+
+		FKey(int32 InX, int32 InY, int32 InZ)
+			: X(InX), Y(InY), Z(InZ) {}
+
+		bool operator==(const FKey& Other) const
+		{
+			return (X == Other.X) && (Y == Other.Y) && (Z == Other.Z);
+		}
 
-	uint64 MakeKey(int32 CellX, int32 CellY) const
+		friend uint32 GetTypeHash(const FKey& Key)
+		{
+			FHashBuilder HashBuilder;
+			HashBuilder << Key.X << Key.Y << Key.Z;
+			return HashBuilder.GetHash();
+		}
+				
+		friend FArchive& operator<<(FArchive& Ar, FKey& Key)
+		{
+			Ar << Key.X;
+			Ar << Key.Y;
+			Ar << Key.Z;
+			return Ar;
+		}
+	};
+
+	const int32 HashCellBits;
+	TMap<FKey, TSet<int32>> CellMap;
+		
+	FKey MakeKey(const FVector& Location) const
 	{
-		return ((uint64)(*(uint32*)(&CellX)) << 32) | (*(uint32*)(&CellY) & 0xffffffff);
+		return FKey(FMath::FloorToInt(Location.X) >> HashCellBits, FMath::FloorToInt(Location.Y) >> HashCellBits, FMath::FloorToInt(Location.Z) >> HashCellBits);
 	}
 
-	uint64 MakeKey(const FVector& Location) const
+	FVector MakeLocation(FKey CellKey) const
 	{
-		return  MakeKey(FMath::FloorToInt(Location.X) >> HashCellBits, FMath::FloorToInt(Location.Y) >> HashCellBits);
+		return FVector(CellKey.X << HashCellBits, CellKey.Y << HashCellBits, CellKey.Z << HashCellBits);
 	}
 
 public:
@@ -376,14 +412,14 @@ public:
 
 	void InsertInstance(const FVector& InstanceLocation, int32 InstanceIndex)
 	{
-		uint64 Key = MakeKey(InstanceLocation);
+		FKey Key = MakeKey(InstanceLocation);
 
 		CellMap.FindOrAdd(Key).Add(InstanceIndex);
 	}
 
 	void RemoveInstance(const FVector& InstanceLocation, int32 InstanceIndex, bool bChecked = true)
 	{
-		uint64 Key = MakeKey(InstanceLocation);
+		FKey Key = MakeKey(InstanceLocation);
 
 		if (bChecked)
 		{
@@ -398,20 +434,20 @@ public:
 
 	void GetInstancesOverlappingBox(const FBox& InBox, TArray<int32>& OutInstanceIndices) const
 	{
-		int32 MinX = FMath::FloorToInt(InBox.Min.X) >> HashCellBits;
-		int32 MinY = FMath::FloorToInt(InBox.Min.Y) >> HashCellBits;
-		int32 MaxX = FMath::FloorToInt(InBox.Max.X) >> HashCellBits;
-		int32 MaxY = FMath::FloorToInt(InBox.Max.Y) >> HashCellBits;
+		FKey MinKey = MakeKey(InBox.Min);
+		FKey MaxKey = MakeKey(InBox.Max);
 
-		for (int32 y = MinY; y <= MaxY; y++)
+		for (int32 z = MinKey.Z; z <= MaxKey.Z; ++z)
 		{
-			for (int32 x = MinX; x <= MaxX; x++)
+			for (int32 y = MinKey.Y; y <= MaxKey.Y; y++)
 			{
-				uint64 Key = MakeKey(x, y);
-				auto* SetPtr = CellMap.Find(Key);
-				if (SetPtr)
+				for (int32 x = MinKey.X; x <= MaxKey.X; x++)
 				{
-					OutInstanceIndices.Append(SetPtr->Array());
+					auto* SetPtr = CellMap.Find(FKey(x, y, z));
+					if (SetPtr)
+					{
+						OutInstanceIndices.Append(SetPtr->Array());
+					}
 				}
 			}
 		}
@@ -435,6 +471,17 @@ public:
 		check(HashCount == InCount);
 	}
 
+	FBox GetBounds() const
+	{
+		FBox HashBounds(ForceInit);
+		for (const auto& Pair : CellMap)
+		{
+			HashBounds += MakeLocation(Pair.Key);
+		}
+
+		return HashBounds;
+	}
+
 	void Empty()
 	{
 		CellMap.Empty();
diff --git a/Engine/Source/Runtime/Foliage/Public/InstancedFoliageActor.h b/Engine/Source/Runtime/Foliage/Public/InstancedFoliageActor.h
--- a/Engine/Source/Runtime/Foliage/Public/InstancedFoliageActor.h
+++ b/Engine/Source/Runtime/Foliage/Public/InstancedFoliageActor.h
@@ -10,7 +10,7 @@
 #include "FoliageType_InstancedStaticMesh.h"
 #include "FoliageInstanceBase.h"
 #include "InstancedFoliage.h"
-#include "ActorPartition/InstancedObjectsActor.h"
+#include "ActorPartition/PartitionActor.h"
 
 #include "InstancedFoliageActor.generated.h"
 
@@ -67,7 +67,7 @@ class UProceduralFoliageComponent;
 typedef TFunction<bool(const UPrimitiveComponent*)> FFoliageTraceFilterFunc;
 
 UCLASS(notplaceable, hidecategories = (Object, Rendering, Mobility), MinimalAPI, NotBlueprintable)
-class AInstancedFoliageActor : public AInstancedObjectsActor
+class AInstancedFoliageActor : public APartitionActor
 {
 	GENERATED_UCLASS_BODY()
 
@@ -107,6 +107,8 @@ public:
 	FOLIAGE_API void CleanupDeletedFoliageType();
 	FOLIAGE_API void DetectFoliageTypeChangeAndUpdate();
 
+	virtual uint32 GetDefaultGridSize(UWorld* InWorld) const override;
+
 	// Delegate type for selection change events
 	DECLARE_MULTICAST_DELEGATE_TwoParams(FOnSelectionChanged, bool, const TArray<AActor*>&);
 	FOLIAGE_API static FOnSelectionChanged SelectionChanged;
