WorldPartition HLOD cleanup/fixes
* Do not build HLODs for always loaded cells
* HLODActor Bounds is now set to the bounds of the source runtime grid cell
* HLODActor MinDrawDistance is now set to the loading range of the source runtime grid
* HLODActor - Default grid placement is now EActorGridPlacement::Location
* Only create the WP HLOD subsystem if the World is using a WorldPartition
* Removed overzealous ensure triggered when encountering Actors with no HLODLayer assigned in a world without a default HLODLayer

diff --git a/Engine/Source/Runtime/Engine/Private/WorldPartition/HLOD/HLODActor.cpp b/Engine/Source/Runtime/Engine/Private/WorldPartition/HLOD/HLODActor.cpp
--- a/Engine/Source/Runtime/Engine/Private/WorldPartition/HLOD/HLODActor.cpp
+++ b/Engine/Source/Runtime/Engine/Private/WorldPartition/HLOD/HLODActor.cpp
@@ -26,33 +26,17 @@ AWorldPartitionHLOD::AWorldPartitionHLOD(const FObjectInitializer& ObjectInitial
 
 UPrimitiveComponent* AWorldPartitionHLOD::GetHLODComponent()
 {
-	return CastChecked<UPrimitiveComponent>(RootComponent);
+	return Cast<UPrimitiveComponent>(RootComponent);
 }
 
-void AWorldPartitionHLOD::LinkCell(FName InCellName)
+void AWorldPartitionHLOD::OnCellShown(FName InCellName)
 {
-	TRACE_CPUPROFILER_EVENT_SCOPE(AWorldPartitionHLOD::LinkCell);
-
-	for (const TSoftObjectPtr<UPrimitiveComponent>& SubPrimitiveComponent : SubPrimitivesComponents)
-	{
-		if (SubPrimitiveComponent.IsValid())
-		{
-			SubPrimitiveComponent->SetCachedLODParentPrimitive(GetHLODComponent());
-		}
-	}
+	GetRootComponent()->SetVisibility(false, true);
 }
 
-void AWorldPartitionHLOD::UnlinkCell(FName InCellName)
+void AWorldPartitionHLOD::OnCellHidden(FName InCellName)
 {
-	TRACE_CPUPROFILER_EVENT_SCOPE(AWorldPartitionHLOD::UnlinkCell);
-
-	for (const TSoftObjectPtr<UPrimitiveComponent>& SubPrimitiveComponent : SubPrimitivesComponents)
-	{
-		if (SubPrimitiveComponent.IsValid())
-		{
-			SubPrimitiveComponent->SetCachedLODParentPrimitive(nullptr);
-		}
-	}
+	GetRootComponent()->SetVisibility(true, true);
 }
 
 void AWorldPartitionHLOD::BeginPlay()
@@ -69,6 +53,11 @@ void AWorldPartitionHLOD::EndPlay(const EEndPlayReason::Type EndPlayReason)
 
 #if WITH_EDITOR
 
+EActorGridPlacement AWorldPartitionHLOD::GetDefaultGridPlacement() const
+{
+	return EActorGridPlacement::Location;
+}
+
 void AWorldPartitionHLOD::SetLODParent(AActor& InActor)
 {
 	UpdateLODParent(InActor, false);
@@ -93,44 +82,83 @@ void AWorldPartitionHLOD::UpdateLODParent(AActor& InActor, bool bInClear)
 	}
 }
 
-const float HLODACTOR_DEFAULT_MIN_DRAW_DISTANCE = 5000.0f;
-
-void AWorldPartitionHLOD::SetHLODPrimitive(UPrimitiveComponent* InHLODPrimitive)
+void AWorldPartitionHLOD::SetHLODPrimitives(const TArray<UPrimitiveComponent*>& InHLODPrimitives, float InFadeOutDistance)
 {
 	TRACE_CPUPROFILER_EVENT_SCOPE(AWorldPartitionHLOD::SetHLODPrimitive);
+	check(!InHLODPrimitives.IsEmpty());
+
+	TArray<USceneComponent*> ComponentsToRemove;
+	GetComponents<USceneComponent>(ComponentsToRemove);
+
+	SetRootComponent(InHLODPrimitives[0]);
+
+	for(UPrimitiveComponent* InHLODPrimitive : InHLODPrimitives)
+	{
+		ComponentsToRemove.Remove(InHLODPrimitive);
 
-	USceneComponent* OldRootComponent = GetRootComponent();
+		AddInstanceComponent(InHLODPrimitive);
 
-	SetRootComponent(InHLODPrimitive);
-	AddInstanceComponent(InHLODPrimitive);
+		if (InHLODPrimitive != RootComponent)
+		{
+			InHLODPrimitive->SetupAttachment(RootComponent);
+		}
 
-	// Setup custom depth rendering to achieve a red tint using a post process material
-	const int32 CellPreviewStencilValue = 180;
-	InHLODPrimitive->bRenderCustomDepth = true;
-	InHLODPrimitive->CustomDepthStencilValue = CellPreviewStencilValue;
+		// Setup custom depth rendering to achieve a red tint using a post process material
+		const int32 CellPreviewStencilValue = 180;
+		InHLODPrimitive->bRenderCustomDepth = true;
+		InHLODPrimitive->CustomDepthStencilValue = CellPreviewStencilValue;
 	
-	InHLODPrimitive->SetCollisionEnabled(ECollisionEnabled::NoCollision);
-	InHLODPrimitive->SetMobility(EComponentMobility::Static);
+		InHLODPrimitive->SetCollisionEnabled(ECollisionEnabled::NoCollision);
+		InHLODPrimitive->SetMobility(EComponentMobility::Static);
 
-	InHLODPrimitive->MinDrawDistance = HLODACTOR_DEFAULT_MIN_DRAW_DISTANCE;
+		InHLODPrimitive->MinDrawDistance = InFadeOutDistance;
 
-	InHLODPrimitive->RegisterComponent();
-	InHLODPrimitive->MarkRenderStateDirty();
+		InHLODPrimitive->RegisterComponent();
+		InHLODPrimitive->MarkRenderStateDirty();
+	}
 
-	if (OldRootComponent)
+	for (USceneComponent* ComponentToRemove : ComponentsToRemove)
 	{
-		OldRootComponent->DestroyComponent();
+		ComponentToRemove->DestroyComponent();
 	}
 }
 
+const FBox& AWorldPartitionHLOD::GetHLODBounds() const
+{
+	return HLODBounds;
+}
+
+void AWorldPartitionHLOD::SetHLODBounds(const FBox& InBounds)
+{
+	HLODBounds = InBounds;
+}
+
+void AWorldPartitionHLOD::GetActorBounds(bool bOnlyCollidingComponents, FVector& Origin, FVector& BoxExtent, bool bIncludeFromChildActors) const
+{
+	Super::GetActorBounds(bOnlyCollidingComponents, Origin, BoxExtent, bIncludeFromChildActors);
+
+	FBox Bounds = FBox(Origin - BoxExtent, Origin + BoxExtent);
+	Bounds += HLODBounds;
+	Bounds.GetCenterAndExtents(Origin, BoxExtent);
+}
+
+void AWorldPartitionHLOD::GetActorLocationBounds(bool bOnlyCollidingComponents, FVector& Origin, FVector& BoxExtent, bool bIncludeFromChildActors) const
+{
+	GetActorBounds(bOnlyCollidingComponents, Origin, BoxExtent, bIncludeFromChildActors);
+}
+
 void AWorldPartitionHLOD::SetChildrenPrimitives(const TArray<UPrimitiveComponent*>& InChildrenPrimitives)
 {
 	TSet<FGuid> SubActorsSet;
 
+	UPrimitiveComponent* HLODComponent = GetHLODComponent();
+	check(HLODComponent);
+
 	for (UPrimitiveComponent* ChildPrimitive : InChildrenPrimitives)
 	{
-		SubPrimitivesComponents.Add(ChildPrimitive);
 		SubActorsSet.Add(ChildPrimitive->GetOwner()->GetActorGuid());
+
+		ChildPrimitive->SetCachedLODParentPrimitive(HLODComponent);
 	}
 
 	SubActors = SubActorsSet.Array();
@@ -147,18 +175,6 @@ void AWorldPartitionHLOD::SetHLODLayer(const UHLODLayer* InSubActorsHLODLayer, i
 	SubActorsHLODLevel = InSubActorsHLODLevel;
 }
 
-void AWorldPartitionHLOD::OnWorldPartitionActorRegistered(AActor& InActor, bool bInLoaded)
-{
-	TRACE_CPUPROFILER_EVENT_SCOPE(AWorldPartitionHLOD::OnWorldPartitionActorRegistered);
-
-	UpdateLODParent(InActor, !bInLoaded);
-}
-
-EActorGridPlacement AWorldPartitionHLOD::GetDefaultGridPlacement() const
-{
-	return EActorGridPlacement::Bounds;
-}
-
 void AWorldPartitionHLOD::GetAssetRegistryTags(TArray<FAssetRegistryTag>& OutTags) const
 {
 	Super::GetAssetRegistryTags(OutTags);
@@ -177,6 +193,9 @@ void AWorldPartitionHLOD::GetAssetRegistryTags(TArray<FAssetRegistryTag>& OutTag
 			static const FName NAME_HLODSubActors(TEXT("HLODSubActors"));
 			FActorRegistry::SaveActorMetaData(NAME_HLODSubActors, SubActorsGUIDsStr, OutTags);
 		}
+
+		static const FName NAME_HLODSubActors(TEXT("HLODLayer"));
+		FActorRegistry::SaveActorMetaData(NAME_HLODSubActors, FSoftObjectPath(SubActorsHLODLayer).ToString(), OutTags);
 	}
 }
 
@@ -186,67 +205,6 @@ void AWorldPartitionHLOD::PostActorCreated()
 	HLODGuid = GetActorGuid();
 }
 
-void AWorldPartitionHLOD::RegisterAllComponents()
-{
-	TRACE_CPUPROFILER_EVENT_SCOPE(AWorldPartitionHLOD::RegisterAllComponents);
-
-	Super::RegisterAllComponents();
-
-	UWorld* World = GetWorld();
-	if (!World->IsGameWorld())
-	{
-		UWorldPartition* WorldPartition = World->GetWorldPartition();
-		check(WorldPartition);
-
-		check(!ActorRegisteredDelegateHandle.IsValid());
-		ActorRegisteredDelegateHandle = WorldPartition->OnActorRegisteredEvent.AddUObject(this, &AWorldPartitionHLOD::OnWorldPartitionActorRegistered);
-
-		for (const FGuid& SubActorGuid : SubActors)
-		{
-			const FWorldPartitionActorDesc* ActorDesc = WorldPartition->GetActorDesc(SubActorGuid);
-			if (ActorDesc != nullptr)
-			{
-				if (AActor* Actor = ActorDesc->GetActor())
-				{
-					SetLODParent(*Actor);
-				}
-			}
-		}
-	}
-}
-
-void AWorldPartitionHLOD::UnregisterAllComponents(const bool bForReregister)
-{
-	TRACE_CPUPROFILER_EVENT_SCOPE(AWorldPartitionHLOD::UnregisterAllComponents);
-
-	UWorld* World = GetWorld();
-	if (World && !World->IsPendingKillOrUnreachable())
-	{
-		if (!World->IsGameWorld() && ActorRegisteredDelegateHandle.IsValid())
-		{
-			UWorldPartition* WorldPartition = World->GetWorldPartition();
-			check(WorldPartition);
-
-			WorldPartition->OnActorRegisteredEvent.Remove(ActorRegisteredDelegateHandle);
-			ActorRegisteredDelegateHandle.Reset();
-
-			for (const FGuid& SubActorGuid : SubActors)
-			{
-				const FWorldPartitionActorDesc* ActorDesc = WorldPartition->GetActorDesc(SubActorGuid);
-				if (ActorDesc != nullptr)
-				{
-					if (AActor* Actor = ActorDesc->GetActor())
-					{
-						ClearLODParent(*Actor);
-					}
-				}
-			}
-		}
-	}
-
-	Super::UnregisterAllComponents();
-}
-
 void UWorldPartitionRuntimeHLODCellData::SetReferencedHLODActors(TArray<FGuid>&& InReferencedHLODActors)
 {
 	ReferencedHLODActors = MoveTemp(InReferencedHLODActors);
diff --git a/Engine/Source/Runtime/Engine/Private/WorldPartition/HLOD/HLODActorDesc.cpp b/Engine/Source/Runtime/Engine/Private/WorldPartition/HLOD/HLODActorDesc.cpp
--- a/Engine/Source/Runtime/Engine/Private/WorldPartition/HLOD/HLODActorDesc.cpp
+++ b/Engine/Source/Runtime/Engine/Private/WorldPartition/HLOD/HLODActorDesc.cpp
@@ -4,11 +4,13 @@
 
 #if WITH_EDITOR
 #include "WorldPartition/HLOD/HLODActor.h"
+#include "WorldPartition/HLOD/HLODLayer.h"
 
-FHLODActorDesc::FHLODActorDesc(const FWorldPartitionActorDescData& DescData, const TArray<FGuid>& InSubActors)
+FHLODActorDesc::FHLODActorDesc(const FWorldPartitionActorDescData& DescData, const TArray<FGuid>& InSubActors, const FSoftObjectPath& InHLODLayer)
 	: FWorldPartitionActorDesc(DescData)
+	, SubActors(InSubActors)
+	, HLODLayer(InHLODLayer)
 {
-	SubActors = InSubActors;
 }
 
 FHLODActorDesc::FHLODActorDesc(AActor* InActor)
@@ -17,6 +19,7 @@ FHLODActorDesc::FHLODActorDesc(AActor* InActor)
 	if (AWorldPartitionHLOD* HLODActor = CastChecked<AWorldPartitionHLOD>(InActor))
 	{
 		SubActors = HLODActor->GetSubActors();
+		HLODLayer = HLODActor->GetHLODLayer();
 	}
 }
 
@@ -24,6 +27,7 @@ void FHLODActorDesc::BuildHash(FHashBuilder& HashBuilder)
 {
 	FWorldPartitionActorDesc::BuildHash(HashBuilder);
 	HashBuilder << SubActors;
+	HashBuilder << HLODLayer.ToString();
 }
 
 #endif
diff --git a/Engine/Source/Runtime/Engine/Private/WorldPartition/HLOD/HLODActorDescFactory.cpp b/Engine/Source/Runtime/Engine/Private/WorldPartition/HLOD/HLODActorDescFactory.cpp
--- a/Engine/Source/Runtime/Engine/Private/WorldPartition/HLOD/HLODActorDescFactory.cpp
+++ b/Engine/Source/Runtime/Engine/Private/WorldPartition/HLOD/HLODActorDescFactory.cpp
@@ -19,8 +19,9 @@ FWorldPartitionActorDesc* FHLODActorDescFactory::CreateInstance(const FWorldPart
 		return nullptr;
 	}
 
-	TArray<FGuid> SubActorsGUIDs;
 	FString MetadataStr;
+
+	TArray<FGuid> SubActorsGUIDs;
 	static const FName NAME_HLODSubActors(TEXT("HLODSubActors"));
 	if (!FActorRegistry::ReadActorMetaData(NAME_HLODSubActors, MetadataStr, ActorDescInitData.AssetData))
 	{
@@ -46,7 +47,15 @@ FWorldPartitionActorDesc* FHLODActorDescFactory::CreateInstance(const FWorldPart
 		});
 	}
 
-	return new FHLODActorDesc(Data, SubActorsGUIDs);
+	FSoftObjectPath HLODLayerPath;
+	static const FName NAME_HLODLayer(TEXT("HLODLayer"));
+	if (!FActorRegistry::ReadActorMetaData(NAME_HLODLayer, MetadataStr, ActorDescInitData.AssetData))
+	{
+		return nullptr;
+	}
+	HLODLayerPath = MetadataStr;
+
+	return new FHLODActorDesc(Data, SubActorsGUIDs, HLODLayerPath);
 }
 
 FWorldPartitionActorDesc* FHLODActorDescFactory::CreateInstance(AActor* InActor)
diff --git a/Engine/Source/Runtime/Engine/Private/WorldPartition/HLOD/HLODBuilder.cpp b/Engine/Source/Runtime/Engine/Private/WorldPartition/HLOD/HLODBuilder.cpp
--- a/Engine/Source/Runtime/Engine/Private/WorldPartition/HLOD/HLODBuilder.cpp
+++ b/Engine/Source/Runtime/Engine/Private/WorldPartition/HLOD/HLODBuilder.cpp
@@ -9,6 +9,7 @@
 #include "WorldPartition/HLOD/HLODLayer.h"
 
 #include "Algo/ForEach.h"
+#include "Algo/RemoveIf.h"
 #include "Algo/Transform.h"
 #include "Templates/UniquePtr.h"
 
@@ -26,11 +27,11 @@ class FHLODBuilder
 public:
 	virtual ~FHLODBuilder() {}
 
-	virtual void Build(const TArray<AActor*>& InSubActors) = 0;
+	virtual void Build(const TArray<UPrimitiveComponent*>& InSubComponents) = 0;
 
-	typedef TFunction<UPrimitiveComponent*(AWorldPartitionHLOD*)> FCreateComponentFunction;
+	typedef TFunction<TArray<UPrimitiveComponent*>(AWorldPartitionHLOD*)> FCreateComponentsFunction;
 		
-	void SpawnHLODActor(const TCHAR* InName, const TArray<UPrimitiveComponent*>& InSubComponents, FCreateComponentFunction InCreateComponentFunc)
+	void SpawnHLODActor(const TCHAR* InName, const TArray<UPrimitiveComponent*>& InSubComponents, FCreateComponentsFunction InCreateComponentsFunc)
 	{
 		FString HLODActorName = FString::Printf(TEXT("%s_%s_%s"), *HLODLayer->GetName(), *CellName.ToString(), InName);
 		FString HLODActorPath = FString::Printf(TEXT("%s.%s"), *World->PersistentLevel->GetPathName(), *HLODActorName);
@@ -39,19 +40,20 @@ public:
 		if (!HLODActor)
 		{
 			FActorSpawnParameters SpawnParams;
-			SpawnParams.bDeferConstruction = true;
 			SpawnParams.bCreateActorPackage = true;
 			SpawnParams.Name = *HLODActorName;
 			HLODActor = World->SpawnActor<AWorldPartitionHLOD>(SpawnParams);
 		}
 
-		UPrimitiveComponent* ParentComponent = InCreateComponentFunc(HLODActor);
+		TArray<UPrimitiveComponent*> HLODPrimitives = InCreateComponentsFunc(HLODActor);
+		Algo::RemoveIf(HLODPrimitives, [](const UPrimitiveComponent* HLODPrimitive) { return !HLODPrimitive; });
 
-		if (ParentComponent)
+		if (!HLODPrimitives.IsEmpty())
 		{
 			HLODActor->Modify();
 			HLODActor->SetHLODLayer(HLODLayer, iLevel);
-			HLODActor->SetHLODPrimitive(ParentComponent);
+			HLODActor->SetHLODPrimitives(HLODPrimitives, CellLoadingRange);
+			HLODActor->SetHLODBounds(CellBounds);
 			HLODActor->SetChildrenPrimitives(InSubComponents);
 			HLODActor->SetActorLabel(HLODActorName);
 			HLODActor->RuntimeGrid = GetHLODLevelSettings().TargetGrid;
@@ -70,7 +72,7 @@ public:
 		}
 	}
 
-	TArray<UPrimitiveComponent*> GatherPrimitiveComponents(int32 iHLODLevel, const TArray<AActor*> InActors) const
+	static TArray<UPrimitiveComponent*> GatherPrimitiveComponents(int32 iHLODLevel, const TArray<AActor*> InActors)
 	{
 		TArray<UPrimitiveComponent*> PrimitiveComponents;
 		for (AActor* SubActor : InActors)
@@ -100,6 +102,8 @@ public:
 	const UHLODLayer*	HLODLayer;
 	int32				iLevel;
 	FName				CellName;
+	FBox				CellBounds;
+	float				CellLoadingRange;
 
 	TArray<AWorldPartitionHLOD*> HLODActors;
 };
@@ -110,46 +114,49 @@ public:
  */
 class FHLODBuilder_Instancing : public FHLODBuilder
 {
-	virtual void Build(const TArray<AActor*>& InSubActors) override
+	virtual void Build(const TArray<UPrimitiveComponent*>& InSubComponents) override
 	{
 		TRACE_CPUPROFILER_EVENT_SCOPE(FHLODBuilder_Instancing::BuildHLOD);
 
 		const FHLODLevelSettings& LevelSettings = GetHLODLevelSettings();
 
-		// Gather all meshes to instantiate along with their transforms
-		TArray<UPrimitiveComponent*> PrimitiveComponents = GatherPrimitiveComponents(iLevel, InSubActors);
-		TMap<UStaticMesh*, TArray<UPrimitiveComponent*>> Instances;
-		for (UPrimitiveComponent* Primitive : PrimitiveComponents)
+		FCreateComponentsFunction CreateComponentLambda = [&LevelSettings, &InSubComponents](AWorldPartitionHLOD* HLODActor)
 		{
-			if (UStaticMeshComponent* SMC = Cast<UStaticMeshComponent>(Primitive))
+			TArray<UPrimitiveComponent*> Components;
+
+			// Gather all meshes to instantiate along with their transforms
+			TMap<UStaticMesh*, TArray<UPrimitiveComponent*>> Instances;
+			for (UPrimitiveComponent* Primitive : InSubComponents)
 			{
-				Instances.FindOrAdd(SMC->GetStaticMesh()).Add(SMC);
+				if (UStaticMeshComponent* SMC = Cast<UStaticMeshComponent>(Primitive))
+				{
+					Instances.FindOrAdd(SMC->GetStaticMesh()).Add(SMC);
+				}
 			}
-		}
 
-		// Now, create an ISMC for each SM asset we found
-		for (const auto& Entry : Instances)
-		{
-			UStaticMesh* StaticMesh = Entry.Key;
-			const TArray<UPrimitiveComponent*>& SubComponents = Entry.Value;
-
-			FCreateComponentFunction CreateComponentLambda = [&LevelSettings, StaticMesh, SubComponents](AWorldPartitionHLOD* HLODActor)
+			// Create an ISMC for each SM asset we found
+			for (const auto& Entry : Instances)
 			{
+				UStaticMesh* EntryStaticMesh = Entry.Key;
+				const TArray<UPrimitiveComponent*>& EntryComponents = Entry.Value;
+			
 				UInstancedStaticMeshComponent* Component = NewObject<UInstancedStaticMeshComponent>(HLODActor);
-				Component->SetStaticMesh(StaticMesh);
-				Component->SetForcedLodModel(StaticMesh->GetNumLODs());
+				Component->SetStaticMesh(EntryStaticMesh);
+				Component->SetForcedLodModel(EntryStaticMesh->GetNumLODs());
 
 				// Add all instances
-				for (UPrimitiveComponent* SMC : SubComponents)
+				for (UPrimitiveComponent* SMC : EntryComponents)
 				{
 					Component->AddInstanceWorldSpace(SMC->GetComponentTransform());
 				}
 
-				return Component;
+				Components.Add(Component);
 			};
 
-			SpawnHLODActor(*StaticMesh->GetName(), SubComponents, CreateComponentLambda);
-		}
+			return Components;
+		};
+
+		SpawnHLODActor(TEXT("InstancedMeshes"), InSubComponents, CreateComponentLambda);
 	}
 };
 
@@ -159,21 +166,19 @@ class FHLODBuilder_Instancing : public FHLODBuilder
  */
 class FHLODBuilder_MeshMerge : public FHLODBuilder
 {
-	virtual void Build(const TArray<AActor*>& InSubActors) override
+	virtual void Build(const TArray<UPrimitiveComponent*>& InSubComponents) override
 	{
 		TRACE_CPUPROFILER_EVENT_SCOPE(UHLODLayer::BuildHLOD_MeshMerge);
 
 		const FHLODLevelSettings& LevelSettings = GetHLODLevelSettings();
 
-		TArray<UPrimitiveComponent*> SubComponents = GatherPrimitiveComponents(iLevel, InSubActors);
-
-		FCreateComponentFunction CreateComponentLambda = [&LevelSettings, SubComponents, this](AWorldPartitionHLOD* HLODActor)
+		FCreateComponentsFunction CreateComponentLambda = [&LevelSettings, &InSubComponents, this](AWorldPartitionHLOD* HLODActor)
 		{
 			TArray<UObject*> Assets;
 			FVector MergedActorLocation;
 
 			const IMeshMergeUtilities& MeshMergeUtilities = FModuleManager::Get().LoadModuleChecked<IMeshMergeModule>("MeshMergeUtilities").GetUtilities();
-			MeshMergeUtilities.MergeComponentsToStaticMesh(SubComponents, HLODActor->GetWorld(), LevelSettings.MergeSetting, LevelSettings.FlattenMaterial.LoadSynchronous(), HLODActor->GetPackage(), CellName.ToString(), Assets, MergedActorLocation, 0.25f, false);
+			MeshMergeUtilities.MergeComponentsToStaticMesh(InSubComponents, HLODActor->GetWorld(), LevelSettings.MergeSetting, LevelSettings.FlattenMaterial.LoadSynchronous(), HLODActor->GetPackage(), CellName.ToString(), Assets, MergedActorLocation, 0.25f, false);
 
 			// All merged mesh assets are stored in the HLOD Actor package
 			Algo::ForEach(Assets, [](UObject* Asset) { Asset->ClearFlags(RF_Public | RF_Standalone); });
@@ -187,10 +192,10 @@ class FHLODBuilder_MeshMerge : public FHLODBuilder
 				Component->SetWorldLocation(MergedActorLocation);
 			}
 
-			return Component;
+			return TArray<UPrimitiveComponent*>({ Component });
 		};
 		
-		SpawnHLODActor(TEXT("MergedMesh"), SubComponents, CreateComponentLambda);
+		SpawnHLODActor(TEXT("MergedMesh"), InSubComponents, CreateComponentLambda);
 	}
 };
 
@@ -199,22 +204,20 @@ class FHLODBuilder_MeshMerge : public FHLODBuilder
  */
 class FHLODBuilder_MeshSimplify : public FHLODBuilder
 {
-	virtual void Build(const TArray<AActor*>& InSubActors) override
+	virtual void Build(const TArray<UPrimitiveComponent*>& InSubComponents) override
 	{
 		TRACE_CPUPROFILER_EVENT_SCOPE(UHLODLayer::BuildHLOD_MeshProxy);
 
 		const FHLODLevelSettings& LevelSettings = GetHLODLevelSettings();
 
-		TArray<UPrimitiveComponent*> SubComponents = GatherPrimitiveComponents(iLevel, InSubActors);
-
-		FCreateComponentFunction CreateComponentLambda = [&LevelSettings, SubComponents, this] (AWorldPartitionHLOD* HLODActor)
+		FCreateComponentsFunction CreateComponentLambda = [&LevelSettings, &InSubComponents, this] (AWorldPartitionHLOD* HLODActor)
 		{
 			TArray<UObject*> Assets;
 			FCreateProxyDelegate ProxyDelegate;
 			ProxyDelegate.BindLambda([&Assets](const FGuid Guid, TArray<UObject*>& InAssetsCreated) { Assets = InAssetsCreated; });
 
 			TArray<UStaticMeshComponent*> StaticMeshComponents;
-			Algo::TransformIf(SubComponents, StaticMeshComponents, [](UPrimitiveComponent* InPrimitiveComponent) { return InPrimitiveComponent->IsA<UStaticMeshComponent>(); }, [](UPrimitiveComponent* InPrimitiveComponent) { return Cast<UStaticMeshComponent>(InPrimitiveComponent); });
+			Algo::TransformIf(InSubComponents, StaticMeshComponents, [](UPrimitiveComponent* InPrimitiveComponent) { return InPrimitiveComponent->IsA<UStaticMeshComponent>(); }, [](UPrimitiveComponent* InPrimitiveComponent) { return Cast<UStaticMeshComponent>(InPrimitiveComponent); });
 
 			const IMeshMergeUtilities& MeshMergeUtilities = FModuleManager::Get().LoadModuleChecked<IMeshMergeModule>("MeshMergeUtilities").GetUtilities();
 			MeshMergeUtilities.CreateProxyMesh(StaticMeshComponents, LevelSettings.ProxySetting, LevelSettings.FlattenMaterial.LoadSynchronous(), HLODActor->GetPackage(), CellName.ToString(), FGuid::NewGuid(), ProxyDelegate, true);
@@ -230,18 +233,14 @@ class FHLODBuilder_MeshSimplify : public FHLODBuilder
 				Component->SetStaticMesh(StaticMesh);
 			}
 
-			return Component;
+			return TArray<UPrimitiveComponent*>({ Component });
 		};
 		
-		SpawnHLODActor(TEXT("SimplifiedMesh"), SubComponents, CreateComponentLambda);
+		SpawnHLODActor(TEXT("SimplifiedMesh"), InSubComponents, CreateComponentLambda);
 	}
 };
 
-
-/**
- *
- */
-TArray<AWorldPartitionHLOD*> FHLODBuilderUtilities::BuildHLODs(UWorldPartition* InWorldPartition, FName InCellName, const UHLODLayer* InHLODLayer, const TArray<AActor*>& InSubActors)
+TArray<AWorldPartitionHLOD*> FHLODBuilderUtilities::BuildHLODs(UWorldPartition* InWorldPartition, FName InCellName, FBox InCellBounds, float InCellLoadingRange, const UHLODLayer* InHLODLayer, const TArray<AActor*>& InSubActors)
 {
 	const TArray<FHLODLevelSettings>& LevelsSettings = InHLODLayer->GetLevels();
 
@@ -270,15 +269,19 @@ TArray<AWorldPartitionHLOD*> FHLODBuilderUtilities::BuildHLODs(UWorldPartition*
 			checkf(false, TEXT("Unsupported type"));
 		}
 
-		if (HLODBuilder)
+		TArray<UPrimitiveComponent*> SubComponents = FHLODBuilder::GatherPrimitiveComponents(iLevel, InSubActors);
+
+		if (HLODBuilder && !SubComponents.IsEmpty())
 		{
 			HLODBuilder->World = InWorldPartition->GetWorld();
 			HLODBuilder->WorldPartition = InWorldPartition;
 			HLODBuilder->HLODLayer = InHLODLayer;
 			HLODBuilder->iLevel = iLevel;
 			HLODBuilder->CellName = InCellName;
+			HLODBuilder->CellBounds = InCellBounds;
+			HLODBuilder->CellLoadingRange = InCellLoadingRange;
 
-			HLODBuilder->Build(InSubActors);
+			HLODBuilder->Build(SubComponents);
 			
 			HLODActors += HLODBuilder->HLODActors;
 		}
diff --git a/Engine/Source/Runtime/Engine/Private/WorldPartition/HLOD/HLODBuilder.h b/Engine/Source/Runtime/Engine/Private/WorldPartition/HLOD/HLODBuilder.h
--- a/Engine/Source/Runtime/Engine/Private/WorldPartition/HLOD/HLODBuilder.h
+++ b/Engine/Source/Runtime/Engine/Private/WorldPartition/HLOD/HLODBuilder.h
@@ -24,11 +24,13 @@ public:
 	 *
 	 * @param 	InWorldPartition	The WorldPartition for which we are building HLODs
 	 * @param 	InCellName			The name of the cell
+	 * @param 	InCellBounds		Bounds of the cell, will be assigned to the created HLOD actors.
+	 * @param 	InCellLoadingRange	The range at which the cell should be loaded. This is the distance where the HLOD actors should fade out.
 	 * @param 	InHLODLayer			The HLODLayer which will provide the HLOD building parameters.
 	 * @param 	InSubActors			The actors from which we'll gather geometry to generate an HLOD mesh
 	 * @return The list of HLOD actors for this cell.
 	 */
-	static TArray<AWorldPartitionHLOD*> BuildHLODs(UWorldPartition* InWorldPartition, FName InCellName, const UHLODLayer* InHLODLayer, const TArray<AActor*>& InSubActors);
+	static TArray<AWorldPartitionHLOD*> BuildHLODs(UWorldPartition* InWorldPartition, FName InCellName, FBox InCellBounds, float InCellLoadingRange, const UHLODLayer* InHLODLayer, const TArray<AActor*>& InSubActors);
 };
 
 #endif
\ No newline at end of file
diff --git a/Engine/Source/Runtime/Engine/Private/WorldPartition/HLOD/HLODLayer.cpp b/Engine/Source/Runtime/Engine/Private/WorldPartition/HLOD/HLODLayer.cpp
--- a/Engine/Source/Runtime/Engine/Private/WorldPartition/HLOD/HLODLayer.cpp
+++ b/Engine/Source/Runtime/Engine/Private/WorldPartition/HLOD/HLODLayer.cpp
@@ -5,11 +5,11 @@
 =============================================================================*/
 
 #include "WorldPartition/HLOD/HLODLayer.h"
+#include "WorldPartition/HLOD/HLODActor.h"
 
 #if WITH_EDITOR
 #include "WorldPartition/WorldPartition.h"
 #include "WorldPartition/WorldPartitionActorDesc.h"
-#include "WorldPartition/HLOD/HLODActor.h"
 #include "WorldPartition/HLOD/HLODBuilder.h"
 #endif
 
@@ -20,7 +20,7 @@
 
 #if WITH_EDITOR
 
-TArray<AWorldPartitionHLOD*> UHLODLayer::GenerateHLODForCell(UWorldPartition* InWorldPartition, FName InCellName, const TSet<FGuid>& InCellActors)
+TArray<AWorldPartitionHLOD*> UHLODLayer::GenerateHLODForCell(UWorldPartition* InWorldPartition, FName InCellName, FBox InCellBounds, float InCellLoadingRange, const TSet<FGuid>& InCellActors)
 {
 	TRACE_CPUPROFILER_EVENT_SCOPE(UHLODLayer::GenerateHLODForCell);
 
@@ -41,72 +41,21 @@
 		if (ShouldIncludeInHLOD(Actor))
 		{
 			UHLODLayer* HLODLayer = UHLODLayer::GetHLODLayer(Actor);
-			HLODLayersActors.FindOrAdd(HLODLayer).Add(Actor);
+			if (HLODLayer)
+			{
+				HLODLayersActors.FindOrAdd(HLODLayer).Add(Actor);
+			}
 		}
 	}
 
 	TArray<AWorldPartitionHLOD*> HLODActors;
 	for (const auto& HLODLayerActors : HLODLayersActors)
 	{
-		UHLODLayer* HLODLayer = HLODLayerActors.Key;
-		if (!ensure(HLODLayer))
-		{
-			// No default HLOD layer, can't generate HLODs for those actors.
-			continue;
-		}
-
-		HLODActors += FHLODBuilderUtilities::BuildHLODs(InWorldPartition, InCellName, HLODLayer, HLODLayerActors.Value);
+		HLODActors += FHLODBuilderUtilities::BuildHLODs(InWorldPartition, InCellName, InCellBounds, InCellLoadingRange, HLODLayerActors.Key, HLODLayerActors.Value);
 	}
 	return HLODActors;
 }
 
-bool UHLODLayer::ShouldIncludeInHLOD(AActor* InActor)
-{
-	if (!InActor)
-	{
-		return false;
-	}
-
-	if (InActor->IsHidden())
-	{
-		return false;
-	}
-
-	if (InActor->IsEditorOnly())
-	{
-		return false;
-	}
-
-	if (InActor->HasAnyFlags(RF_Transient))
-	{
-		return false;
-	}
-
-	if (InActor->IsTemplate())
-	{
-		return false;
-	}
-
-	if (InActor->IsPendingKill())
-	{
-		return false;
-	}
-
-	if (!InActor->bEnableAutoLODGeneration)
-	{
-		return false;
-	}
-
-	FVector Origin, Extent;
-	InActor->GetActorBounds(false, Origin, Extent);
-	if (Extent.SizeSquared() <= 0.1)
-	{
-		return false;
-	}
-
-	return true;
-}
-
 bool UHLODLayer::ShouldIncludeInHLOD(UPrimitiveComponent* InComponent, int32 InLevelIndex)
 {
 	if (UStaticMeshComponent* SMC = Cast<UStaticMeshComponent>(InComponent))
@@ -149,3 +98,51 @@
 }
 
 #endif // WITH_EDITOR
+
+
+bool UHLODLayer::ShouldIncludeInHLOD(AActor* InActor)
+{
+	if (!InActor)
+	{
+		return false;
+	}
+
+	if (InActor->IsHidden())
+	{
+		return false;
+	}
+
+	if (InActor->IsEditorOnly())
+	{
+		return false;
+	}
+
+	if (InActor->HasAnyFlags(RF_Transient))
+	{
+		return false;
+	}
+
+	if (InActor->IsTemplate())
+	{
+		return false;
+	}
+
+	if (InActor->IsPendingKill())
+	{
+		return false;
+	}
+
+	if (!InActor->bEnableAutoLODGeneration)
+	{
+		return false;
+	}
+
+	FVector Origin, Extent;
+	InActor->GetActorBounds(false, Origin, Extent);
+	if (Extent.SizeSquared() <= 0.1)
+	{
+		return false;
+	}
+
+	return true;
+}
diff --git a/Engine/Source/Runtime/Engine/Private/WorldPartition/HLOD/HLODSubsystem.cpp b/Engine/Source/Runtime/Engine/Private/WorldPartition/HLOD/HLODSubsystem.cpp
--- a/Engine/Source/Runtime/Engine/Private/WorldPartition/HLOD/HLODSubsystem.cpp
+++ b/Engine/Source/Runtime/Engine/Private/WorldPartition/HLOD/HLODSubsystem.cpp
@@ -2,6 +2,9 @@
 
 #include "WorldPartition/HLOD/HLODSubsystem.h"
 #include "WorldPartition/HLOD/HLODActor.h"
+#include "WorldPartition/HLOD/HLODActorDesc.h"
+#include "WorldPartition/HLOD/HLODLayer.h"
+#include "WorldPartition/WorldPartition.h"
 #include "WorldPartition/WorldPartitionSubsystem.h"
 #include "WorldPartition/WorldPartitionRuntimeCell.h"
 #include "UObject/UObjectHash.h"
@@ -21,28 +24,153 @@ UHLODSubsystem::~UHLODSubsystem()
 {
 }
 
+bool UHLODSubsystem::ShouldCreateSubsystem(UObject* Outer) const
+{
+	if (UWorld* WorldOuter = Cast<UWorld>(Outer))
+	{
+		return WorldOuter->GetWorldPartition() != nullptr;
+	}
+	return false;
+}
+
 void UHLODSubsystem::Initialize(FSubsystemCollectionBase& Collection)
 {
 	Super::Initialize(Collection);
 
+	UWorldPartitionSubsystem* WorldPartitionSubsystem = Collection.InitializeDependency<UWorldPartitionSubsystem>();
+
 #if WITH_EDITOR
 	HLODActorDescFactory.Reset(new FHLODActorDescFactory());
 
-	Collection.InitializeDependency(UWorldPartitionSubsystem::StaticClass());
+	RegisterActorDescFactories(WorldPartitionSubsystem);
+
+	GetWorld()->GetWorldPartition()->OnActorRegisteredEvent.AddUObject(this, &UHLODSubsystem::OnWorldPartitionActorRegistered);
+#endif
+}
+
+#if WITH_EDITOR
+
+void UHLODSubsystem::OnWorldPartitionActorRegistered(AActor& InActor, bool bInLoaded)
+{
+	TRACE_CPUPROFILER_EVENT_SCOPE(UHLODSubsystem::OnWorldPartitionActorRegistered);
+
+	const FGuid& ActorGuid = InActor.GetActorGuid();
 
-	if (UWorldPartitionSubsystem* WorldPartitionSubsystem = GetWorld()->GetSubsystem<UWorldPartitionSubsystem>())
+	if (AWorldPartitionHLOD* HLODActor = Cast<AWorldPartitionHLOD>(&InActor))
 	{
-		RegisterActorDescFactories(WorldPartitionSubsystem);
+		if (bInLoaded)
+		{
+			UWorldPartition* WorldPartition = GetWorld()->GetWorldPartition();
+
+			TArray<FGuid> SubActors;
+			PendingHLODAssignment.MultiFind(ActorGuid, SubActors);
+
+			if (!SubActors.IsEmpty())
+			{
+				for (const FGuid& SubActorGuid : SubActors)
+				{
+					const FWorldPartitionActorDesc* ActorDesc = WorldPartition->GetActorDesc(SubActorGuid);
+					check(ActorDesc);
+
+					AActor* Actor = Cast<AWorldPartitionHLOD>(ActorDesc->GetActor());
+					if (Actor)
+					{
+						HLODActor->UpdateLODParent(*Actor, !bInLoaded);
+					}
+				}
+
+				PendingHLODAssignment.Remove(ActorGuid);
+			}
+		}
+
+		return;
+	}
+
+	const FWorldPartitionActorDesc* HLODActorDesc = GetHLODActorForActor(&InActor);
+	if (HLODActorDesc)
+	{
+		AWorldPartitionHLOD* HLODActor = Cast<AWorldPartitionHLOD>(HLODActorDesc->GetActor());
+		if (HLODActor)
+		{
+			HLODActor->UpdateLODParent(InActor, !bInLoaded);
+		}
+		else if (bInLoaded)
+		{
+			PendingHLODAssignment.Add(HLODActorDesc->GetGuid(), ActorGuid);
+		}
+		else
+		{
+			PendingHLODAssignment.Remove(HLODActorDesc->GetGuid(), ActorGuid);
+		}
 	}
-#endif
 }
 
-void UHLODSubsystem::Deinitialize()
+const FWorldPartitionActorDesc* UHLODSubsystem::GetHLODActorForActor(const AActor* InActor) const
 {
-	Super::Deinitialize();
+	UHLODLayer* HLODLayer = UHLODLayer::GetHLODLayer(InActor);
+	if (!HLODLayer)
+	{
+		return nullptr;
+	}
+
+	UWorldPartitionSubsystem* WorldPartitionSubsystem = GetWorld()->GetSubsystem<UWorldPartitionSubsystem>();
+	UWorldPartition* WorldPartition = GetWorld()->GetWorldPartition();
+
+	const FWorldPartitionActorDesc* ActorDesc = WorldPartition->GetActorDesc(InActor->GetActorGuid());
+	check(ActorDesc);
+
+	if (ActorDesc->GetGridPlacement() == EActorGridPlacement::AlwaysLoaded)
+	{
+		return nullptr;
+	}
+
+	FVector ActorLocation = ActorDesc->GetOrigin();
+	FBox ActorBox(ActorLocation, ActorLocation);
+
+	// Find all HLODActors at that location
+	TArray<const FWorldPartitionActorDesc*> HLODActorsDescs = WorldPartitionSubsystem->GetIntersectingActorDescs(ActorBox, AWorldPartitionHLOD::StaticClass());
+
+	// Only keep the HLODActors matching our HLODLayer
+	HLODActorsDescs.RemoveAll([HLODLayer, ActorDesc](const FWorldPartitionActorDesc* InActorDesc)
+	{
+		const FHLODActorDesc* HLODActorDesc = (FHLODActorDesc*)InActorDesc;
+		return ActorDesc == HLODActorDesc || HLODActorDesc->GetHLODLayer() != HLODLayer;
+	});
+	
+	if (HLODActorsDescs.IsEmpty())
+	{
+		return nullptr;
+	}
+
+	// Sort by bounds size
+	HLODActorsDescs.Sort([](const FWorldPartitionActorDesc& A, const FWorldPartitionActorDesc& B)
+	{
+		return A.GetBounds().GetExtent().Size() < B.GetBounds().GetExtent().Size();
+	});
+
+	const UWorldPartition::FActorCluster* ActorCluster = WorldPartition->GetClusterForActor(ActorDesc->GetGuid());
+	switch (ActorCluster->GridPlacement)
+	{
+	case EActorGridPlacement::Location:
+		return HLODActorsDescs[0];
+	
+	case EActorGridPlacement::Bounds:
+		for (const FWorldPartitionActorDesc* HLODActorDesc : HLODActorsDescs)
+		{
+			if (HLODActorDesc->GetBounds().IsInsideXY(ActorCluster->Bounds))
+			{
+				return HLODActorDesc;
+			}
+		}
+		break;
+	
+	default:
+		check(0);
+	}
+
+	return nullptr;
 }
 
-#if WITH_EDITOR
 void UHLODSubsystem::RegisterActorDescFactories(UWorldPartitionSubsystem* WorldPartitionSubsystem)
 {
 	WorldPartitionSubsystem->RegisterActorDescFactory(AWorldPartitionHLOD::StaticClass(), HLODActorDescFactory.Get());
@@ -57,14 +185,14 @@ void UHLODSubsystem::RegisterHLODActor(AWorldPartitionHLOD* InWorldPartitionHLOD
 	RegisteredHLODActors.Emplace(HLODActorGUID, InWorldPartitionHLOD);
 
 	TArray<FName> Cells;
-	PendingTransitionsToHLOD.MultiFind(HLODActorGUID, Cells);
+	PendingCellsShown.MultiFind(HLODActorGUID, Cells);
 	if (Cells.Num())
 	{
 		for (FName Cell : Cells)
 		{ 
-			InWorldPartitionHLOD->LinkCell(Cell);
+			InWorldPartitionHLOD->OnCellShown(Cell);
 		}
-		PendingTransitionsToHLOD.Remove(HLODActorGUID);
+		PendingCellsShown.Remove(HLODActorGUID);
 	}
 }
 
@@ -75,11 +203,9 @@ void UHLODSubsystem::UnregisterHLODActor(AWorldPartitionHLOD* InWorldPartitionHL
 	const FGuid& HLODActorGUID = InWorldPartitionHLOD->GetHLODGuid();
 	int32 NumRemoved = RegisteredHLODActors.Remove(HLODActorGUID);
 	check(NumRemoved == 1);
-
-	check(!PendingTransitionsToHLOD.Find(HLODActorGUID));
 }
 
-void UHLODSubsystem::TransitionToHLOD(const UWorldPartitionRuntimeCell* InCell)
+void UHLODSubsystem::OnCellShown(const UWorldPartitionRuntimeCell* InCell)
 {
 	const UWorldPartitionRuntimeHLODCellData* HLODCellData = InCell->GetCellData<UWorldPartitionRuntimeHLODCellData>();
 	if (!ensure(HLODCellData))
@@ -92,17 +218,17 @@ void UHLODSubsystem::TransitionToHLOD(const UWorldPartitionRuntimeCell* InCell)
 		AWorldPartitionHLOD* HLODActor = RegisteredHLODActors.FindRef(HLODActorGUID);
 		if (HLODActor)
 		{
-			HLODActor->LinkCell(InCell->GetFName());
+			HLODActor->OnCellShown(InCell->GetFName());
 		}
 		else
 		{
-			// Cell was loaded before the HLOD
-			PendingTransitionsToHLOD.Add(HLODActorGUID, InCell->GetFName());
+			// Cell was shown before the HLOD
+			PendingCellsShown.Add(HLODActorGUID, InCell->GetFName());
 		}
 	}
 }
 
-void UHLODSubsystem::TransitionFromHLOD(const UWorldPartitionRuntimeCell* InCell)
+void UHLODSubsystem::OnCellHidden(const UWorldPartitionRuntimeCell* InCell)
 {
 	const UWorldPartitionRuntimeHLODCellData* HLODCellData = InCell->GetCellData<UWorldPartitionRuntimeHLODCellData>();
 	if (!ensure(HLODCellData))
@@ -115,11 +241,12 @@ void UHLODSubsystem::TransitionFromHLOD(const UWorldPartitionRuntimeCell* InCell
 		AWorldPartitionHLOD* HLODActor = RegisteredHLODActors.FindRef(HLODActorGUID);
 		if (HLODActor)
 		{
-			HLODActor->UnlinkCell(InCell->GetFName());
+			HLODActor->OnCellHidden(InCell->GetFName());
+			check(!PendingCellsShown.Find(HLODActorGUID));
 		}
-		else 
+		else
 		{
-			PendingTransitionsToHLOD.Remove(HLODActorGUID, InCell->GetFName());
+			PendingCellsShown.Remove(HLODActorGUID, InCell->GetFName());
 		}
 	}
 }
diff --git a/Engine/Source/Runtime/Engine/Private/WorldPartition/WorldPartitionRuntimeLevelStreamingCell.cpp b/Engine/Source/Runtime/Engine/Private/WorldPartition/WorldPartitionRuntimeLevelStreamingCell.cpp
--- a/Engine/Source/Runtime/Engine/Private/WorldPartition/WorldPartitionRuntimeLevelStreamingCell.cpp
+++ b/Engine/Source/Runtime/Engine/Private/WorldPartition/WorldPartitionRuntimeLevelStreamingCell.cpp
@@ -132,12 +132,6 @@ void UWorldPartitionRuntimeLevelStreamingCell::Activate() const
 		if (!LevelStreaming)
 		{
 			LevelStreaming = Cast<UWorldPartitionLevelStreamingDynamic>(CreateLevelStreaming());
-
-			if (const UWorldPartitionRuntimeHLODCellData* HLODCellData = GetCellData<UWorldPartitionRuntimeHLODCellData>())
-			{
-				LevelStreaming->OnLevelShown.AddDynamic(this, &UWorldPartitionRuntimeLevelStreamingCell::OnLevelShown);
-				LevelStreaming->OnLevelHidden.AddDynamic(this, &UWorldPartitionRuntimeLevelStreamingCell::OnLevelHidden);
-			}
 		}
 		check(LevelStreaming);
 	}
@@ -177,6 +171,12 @@ void UWorldPartitionRuntimeLevelStreamingCell::Activate() const
 
 	if (LevelStreaming)
 	{
+		if (const UWorldPartitionRuntimeHLODCellData* HLODCellData = GetCellData<UWorldPartitionRuntimeHLODCellData>())
+		{
+			LevelStreaming->OnLevelShown.AddUniqueDynamic(this, &UWorldPartitionRuntimeLevelStreamingCell::OnLevelShown);
+			LevelStreaming->OnLevelHidden.AddUniqueDynamic(this, &UWorldPartitionRuntimeLevelStreamingCell::OnLevelHidden);
+		}
+
 		LevelStreaming->Activate();
 	}
 }
@@ -206,11 +206,11 @@ void UWorldPartitionRuntimeLevelStreamingCell::Deactivate() const
 void UWorldPartitionRuntimeLevelStreamingCell::OnLevelShown()
 {
 	check(LevelStreaming);
-	LevelStreaming->GetWorld()->GetSubsystem<UHLODSubsystem>()->TransitionToHLOD(this);
+	LevelStreaming->GetWorld()->GetSubsystem<UHLODSubsystem>()->OnCellShown(this);
 }
 
 void UWorldPartitionRuntimeLevelStreamingCell::OnLevelHidden()
 {
 	check(LevelStreaming);
-	LevelStreaming->GetWorld()->GetSubsystem<UHLODSubsystem>()->TransitionFromHLOD(this);
+	LevelStreaming->GetWorld()->GetSubsystem<UHLODSubsystem>()->OnCellHidden(this);
 }
diff --git a/Engine/Source/Runtime/Engine/Private/WorldPartition/WorldPartitionRuntimeSpatialHash.cpp b/Engine/Source/Runtime/Engine/Private/WorldPartition/WorldPartitionRuntimeSpatialHash.cpp
--- a/Engine/Source/Runtime/Engine/Private/WorldPartition/WorldPartitionRuntimeSpatialHash.cpp
+++ b/Engine/Source/Runtime/Engine/Private/WorldPartition/WorldPartitionRuntimeSpatialHash.cpp
@@ -1181,6 +1181,13 @@ bool UWorldPartitionRuntimeSpatialHash::GenerateHLOD()
 		GridsMapping.Add(Grid.GridName, i);
 	}
 
+	// Ignore HLODD clusters as we are rebuilding them
+	for (const FSpatialHashRuntimeGrid& HLODGrid : HLODGrids)
+	{
+		check(!GridsMapping.Contains(HLODGrid.GridName));
+		GridsMapping.Add(HLODGrid.GridName, INDEX_NONE);
+	}
+
 	UWorldPartition* WorldPartition = GetOuterUWorldPartition();
 
 	TArray<TArray<const UWorldPartition::FActorCluster*>> GridActors;
@@ -1195,7 +1202,10 @@ bool UWorldPartitionRuntimeSpatialHash::GenerateHLOD()
 		}
 
 		int32 GridIndex = FoundIndex ? *FoundIndex : 0;
-		GridActors[GridIndex].Add(ActorCluster);
+		if (GridIndex != INDEX_NONE)
+		{
+			GridActors[GridIndex].Add(ActorCluster);
+		}
 	}
 
 	// Keep track of all valid HLOD actors
@@ -1207,18 +1217,22 @@ bool UWorldPartitionRuntimeSpatialHash::GenerateHLOD()
 		const FSpatialHashRuntimeGrid& RuntimeGrid = Grids[GridIndex];
 		const FSquare2DGridHelper PartionedActors = GetPartitionedActors(WorldPartition, WorldBounds, RuntimeGrid, GridActors[GridIndex]);
 
-		PartionedActors.ForEachCells([PartionedActors, &ValidHLODActors, RuntimeGrid, WorldPartition, this](const FIntVector& CellCoord)
+		PartionedActors.ForEachCells([PartionedActors, &ValidHLODActors, RuntimeGrid, WorldPartition, WorldBounds, this](const FIntVector& CellCoord)
 		{
 			const FSquare2DGridHelper::FGridLevel::FGridCell& GridCell = PartionedActors.GetCell(CellCoord);
-			if (GridCell.Actors.Num() != 0)
+			const bool bIsCellAlwaysLoaded = &GridCell == &PartionedActors.GetAlwaysLoadedCell();
+
+			if (!bIsCellAlwaysLoaded && GridCell.Actors.Num() != 0)
 			{
-				FBox2D CellBounds;
-				PartionedActors.GetCellBounds(CellCoord, CellBounds);
+				FBox2D CellBounds2D;
+				PartionedActors.GetCellBounds(CellCoord, CellBounds2D);
 
 				FName CellName = GetCellName(RuntimeGrid.GridName, CellCoord.Z, CellCoord.X, CellCoord.Y);
+				float CellLoadingRange = RuntimeGrid.LoadingRange;
+				FBox CellBounds = FBox(FVector(CellBounds2D.Min, WorldBounds.Min.Z), FVector(CellBounds2D.Max, WorldBounds.Max.Z));
 
 				TArray<AWorldPartitionHLOD*> HLODActors;
-				HLODActors = UHLODLayer::GenerateHLODForCell(WorldPartition, CellName, GridCell.Actors);
+				HLODActors = UHLODLayer::GenerateHLODForCell(WorldPartition, CellName, CellBounds, CellLoadingRange, GridCell.Actors);
 
 				ValidHLODActors.Append(HLODActors);
 			}
@@ -1234,6 +1248,8 @@ bool UWorldPartitionRuntimeSpatialHash::GenerateHLOD()
 		}
 	}
 
+	CollectGarbage(RF_Standalone);
+
 	return true;
 }
 #endif
diff --git a/Engine/Source/Runtime/Engine/Public/WorldPartition/HLOD/HLODActor.h b/Engine/Source/Runtime/Engine/Public/WorldPartition/HLOD/HLODActor.h
--- a/Engine/Source/Runtime/Engine/Public/WorldPartition/HLOD/HLODActor.h
+++ b/Engine/Source/Runtime/Engine/Public/WorldPartition/HLOD/HLODActor.h
@@ -9,17 +9,18 @@
 #include "HLODActor.generated.h"
 
 class UHLODLayer;
+class UHLODSubsystem;
 
 UCLASS(NotPlaceable)
 class ENGINE_API AWorldPartitionHLOD : public AActor
 {
 	GENERATED_UCLASS_BODY()
 
-public:
-	UPrimitiveComponent* GetHLODComponent();
+	friend class UHLODSubsystem;
 
-	void LinkCell(FName InCellName);
-	void UnlinkCell(FName InCellName);
+public:
+	void OnCellShown(FName InCellName);
+	void OnCellHidden(FName InCellName);
 
 	const FGuid& GetHLODGuid() const { return HLODGuid; }
 
@@ -29,28 +30,33 @@ protected:
 	virtual void EndPlay(const EEndPlayReason::Type EndPlayReason) override;
 	//~ End AActor Interface.
 
+#if WITH_EDITOR
 	void SetLODParent(AActor& InActor);
 	void ClearLODParent(AActor& InActor);
 	void UpdateLODParent(AActor& InActor, bool bInClear);
+#endif
+
+	UPrimitiveComponent* GetHLODComponent();
 
 public:
 #if WITH_EDITOR
-	void SetHLODPrimitive(UPrimitiveComponent* InHLODPrimitive);
+	void SetHLODPrimitives(const TArray<UPrimitiveComponent*>& InHLODPrimitives, float InFadeOutDistance);
 	void SetChildrenPrimitives(const TArray<UPrimitiveComponent*>& InChildrenPrimitives);
 
 	const TArray<FGuid>& GetSubActors() const;
 
 	void SetHLODLayer(const UHLODLayer* InSubActorsHLODLayer, int32 InSubActorsHLODLevel);
+	
+	const FBox& GetHLODBounds() const;
+	void SetHLODBounds(const FBox& InBounds);
 
 protected:
-	void OnWorldPartitionActorRegistered(AActor& InActor, bool bInLoaded);
-
 	//~ Begin AActor Interface.
 	virtual EActorGridPlacement GetDefaultGridPlacement() const override;
 	virtual void GetAssetRegistryTags(TArray<FAssetRegistryTag>& OutTags) const override;
 	virtual void PostActorCreated() override;
-	virtual void RegisterAllComponents() override;
-	virtual void UnregisterAllComponents(const bool bForReregister) override;
+	virtual void GetActorBounds(bool bOnlyCollidingComponents, FVector& Origin, FVector& BoxExtent, bool bIncludeFromChildActors) const override;
+	virtual void GetActorLocationBounds(bool bOnlyCollidingComponents, FVector& Origin, FVector& BoxExtent, bool bIncludeFromChildActors) const override;
 	//~ End AActor Interface.
 #endif // WITH_EDITOR
 
@@ -62,6 +68,9 @@ private:
 	UPROPERTY()
 	const UHLODLayer* SubActorsHLODLayer;
 
+	UPROPERTY()
+	FBox HLODBounds;
+
 	FDelegateHandle ActorRegisteredDelegateHandle;
 #endif
 	
@@ -70,12 +79,6 @@ private:
 
 	UPROPERTY()
 	FGuid HLODGuid;
-
-	UPROPERTY()
-	UPrimitiveComponent* HLODComponent;
-
-	UPROPERTY()
-	TArray<TSoftObjectPtr<UPrimitiveComponent>> SubPrimitivesComponents;
 };
 
 UCLASS()
diff --git a/Engine/Source/Runtime/Engine/Public/WorldPartition/HLOD/HLODActorDesc.h b/Engine/Source/Runtime/Engine/Public/WorldPartition/HLOD/HLODActorDesc.h
--- a/Engine/Source/Runtime/Engine/Public/WorldPartition/HLOD/HLODActorDesc.h
+++ b/Engine/Source/Runtime/Engine/Public/WorldPartition/HLOD/HLODActorDesc.h
@@ -7,8 +7,7 @@
 #include "Containers/Map.h"
 #include "WorldPartition/WorldPartitionActorDesc.h"
 
-class UHLODSubsystem;
-class UWorldPartition;
+class UHLODLayer;
 
 /**
  * ActorDesc for AWorldPartitionHLOD.
@@ -18,10 +17,11 @@ class ENGINE_API FHLODActorDesc : public FWorldPartitionActorDesc
 public:
 #if WITH_EDITOR
 	inline const TArray<FGuid>& GetSubActors() const { return SubActors; }
+	inline const FSoftObjectPath& GetHLODLayer() const { return HLODLayer; }
 
 protected:
 	FHLODActorDesc() = delete;
-	FHLODActorDesc(const FWorldPartitionActorDescData& DescData, const TArray<FGuid>& InSubActors);
+	FHLODActorDesc(const FWorldPartitionActorDescData& DescData, const TArray<FGuid>& InSubActors, const FSoftObjectPath& HLODLayer);
 	FHLODActorDesc(AActor* InActor);
 
 	virtual void BuildHash(FHashBuilder& HashBuilder) override;
@@ -29,5 +29,6 @@ protected:
 	friend class FHLODActorDescFactory;
 
 	TArray<FGuid> SubActors;
+	FSoftObjectPath	HLODLayer;
 #endif
 };
diff --git a/Engine/Source/Runtime/Engine/Public/WorldPartition/HLOD/HLODLayer.h b/Engine/Source/Runtime/Engine/Public/WorldPartition/HLOD/HLODLayer.h
--- a/Engine/Source/Runtime/Engine/Public/WorldPartition/HLOD/HLODLayer.h
+++ b/Engine/Source/Runtime/Engine/Public/WorldPartition/HLOD/HLODLayer.h
@@ -59,13 +59,14 @@ class ENGINE_API UHLODLayer : public UObject
 public:
 	const TArray<FHLODLevelSettings>& GetLevels() const { return Levels; }
 
-	static TArray<AWorldPartitionHLOD*> GenerateHLODForCell(UWorldPartition* InWorldPartition, FName InCellName, const TSet<FGuid>& InCellActors);
+	static TArray<AWorldPartitionHLOD*> GenerateHLODForCell(UWorldPartition* InWorldPartition, FName InCellName, FBox InCellBounds, float InCellLoadingRange, const TSet<FGuid>& InCellActors);
 
 	static UHLODLayer* GetHLODLayer(const AActor* InActor);
-	static bool ShouldIncludeInHLOD(AActor* InActor);
 	static bool ShouldIncludeInHLOD(UPrimitiveComponent* InComponent, int32 InLevelIndex);
 #endif
 
+	static bool ShouldIncludeInHLOD(AActor* InActor);
+
 
 #if WITH_EDITORONLY_DATA
 private:
diff --git a/Engine/Source/Runtime/Engine/Public/WorldPartition/HLOD/HLODSubsystem.h b/Engine/Source/Runtime/Engine/Public/WorldPartition/HLOD/HLODSubsystem.h
--- a/Engine/Source/Runtime/Engine/Public/WorldPartition/HLOD/HLODSubsystem.h
+++ b/Engine/Source/Runtime/Engine/Public/WorldPartition/HLOD/HLODSubsystem.h
@@ -27,15 +27,17 @@ class ENGINE_API UHLODSubsystem : public UWorldSubsystem
 public:
 	UHLODSubsystem();
 	virtual ~UHLODSubsystem();
-		
+
+	//~ Begin USubsystem Interface.
+	virtual bool ShouldCreateSubsystem(UObject* Outer) const override;
 	virtual void Initialize(FSubsystemCollectionBase& Collection) override;
-	virtual void Deinitialize() override;
+	//~ End USubsystem Interface.
 
 	void RegisterHLODActor(AWorldPartitionHLOD* InWorldPartitionHLOD);
 	void UnregisterHLODActor(AWorldPartitionHLOD* InWorldPartitionHLOD);
 
-	void TransitionToHLOD(const UWorldPartitionRuntimeCell* InCell);
-	void TransitionFromHLOD(const UWorldPartitionRuntimeCell* InCell);
+	void OnCellShown(const UWorldPartitionRuntimeCell* InCell);
+	void OnCellHidden(const UWorldPartitionRuntimeCell* InCell);
 	
 #if WITH_EDITOR
 	FHLODActorDescFactory* GetActorDescFactory() const { return HLODActorDescFactory.Get(); }
@@ -44,6 +46,8 @@ public:
 private:
 #if WITH_EDITOR
 	void RegisterActorDescFactories(UWorldPartitionSubsystem* WorldPartitionSubsystem);
+	void OnWorldPartitionActorRegistered(AActor& InActor, bool bInLoaded);
+	const FWorldPartitionActorDesc* GetHLODActorForActor(const AActor* InActor) const;
 #endif
 
 private:
@@ -54,6 +58,10 @@ private:
 	// Mapping between an HLOD GUID & the loaded actor.
 	TMap<FGuid, AWorldPartitionHLOD*> RegisteredHLODActors;
 
-	// Cells that are waiting for their HLOD to load.
-	TMultiMap<FGuid, FName> PendingTransitionsToHLOD;
+	// Cells that where shown before their HLOD was loaded.
+	TMultiMap<FGuid, FName> PendingCellsShown;
+
+#if WITH_EDITOR
+	TMultiMap<FGuid, FGuid> PendingHLODAssignment;
+#endif
 };
