Added support for WP HLODs to "viewmode hlodcoloration"

diff --git a/Engine/Source/Runtime/Engine/Private/StaticMeshRender.cpp b/Engine/Source/Runtime/Engine/Private/StaticMeshRender.cpp
--- a/Engine/Source/Runtime/Engine/Private/StaticMeshRender.cpp
+++ b/Engine/Source/Runtime/Engine/Private/StaticMeshRender.cpp
@@ -41,6 +41,7 @@
 #include "ComponentRecreateRenderStateContext.h"
 #include "PhysicsEngine/BodySetup.h"
 #include "Engine/LODActor.h"
+#include "WorldPartition/HLOD/HLODActor.h"
 
 #include "UnrealEngine.h"
 #include "RayTracingInstance.h"
@@ -316,6 +317,10 @@ FStaticMeshSceneProxy::FStaticMeshSceneProxy(UStaticMeshComponent* InComponent,
 		// An HLOD cluster (they count from 1, but the colors for HLOD levels start at index 2)
 		HierarchicalLODIndex = LODActorOwner->LODLevel + 1;
 	}
+	else if (AWorldPartitionHLOD* WorldPartitionHLODOwner = Cast<AWorldPartitionHLOD>(Owner))
+	{
+		HierarchicalLODIndex = 2;
+	}
 	else if (InComponent->GetLODParentPrimitive())
 	{
 		// Part of a HLOD cluster but still a plain mesh
