Remove unrequired UModel/UPolys PreEditChange overloads since 14268352.

diff --git a/Engine/Source/Runtime/Engine/Classes/Engine/Polys.h b/Engine/Source/Runtime/Engine/Classes/Engine/Polys.h
--- a/Engine/Source/Runtime/Engine/Classes/Engine/Polys.h
+++ b/Engine/Source/Runtime/Engine/Classes/Engine/Polys.h
@@ -609,7 +609,6 @@ class UPolys : public UObject
 	//~ Begin UObject Interface
 #if WITH_EDITOR
 	ENGINE_API virtual bool Modify(bool bAlwaysMarkDirty = false) override;	
-	virtual void PreEditChange(FProperty*) override;
 #endif
 	ENGINE_API virtual void Serialize( FArchive& Ar ) override;
 	virtual bool IsAsset() const override { return false; }
diff --git a/Engine/Source/Runtime/Engine/Private/Model.cpp b/Engine/Source/Runtime/Engine/Private/Model.cpp
--- a/Engine/Source/Runtime/Engine/Private/Model.cpp
+++ b/Engine/Source/Runtime/Engine/Private/Model.cpp
@@ -585,12 +585,6 @@ bool UModel::Modify( bool bAlwaysMarkDirty/*=false*/ )
 
 	return bSavedToTransactionBuffer;
 }
-
-void UModel::PreEditChange(FProperty*)
-{
-	// Do not call Super! Override PreEditChange to keep the same behavior as the Modify which change the default of `bAlwaysMarkDirty`...
-	Modify(false);
-}
 #endif
 
 //
diff --git a/Engine/Source/Runtime/Engine/Private/Polygon.cpp b/Engine/Source/Runtime/Engine/Private/Polygon.cpp
--- a/Engine/Source/Runtime/Engine/Private/Polygon.cpp
+++ b/Engine/Source/Runtime/Engine/Private/Polygon.cpp
@@ -985,12 +985,6 @@ bool UPolys::Modify(bool bAlwaysMarkDirty)
 
 	return !!GUndo; // we will make a broad assumption that if we have an undo buffer, we were saved in it
 }
-
-void UPolys::PreEditChange(FProperty*)
-{
-	// Do not call Super! Override PreEditChange to keep the same behavior as the Modify which change the default of `bAlwaysMarkDirty`...
-	Modify(false);
-}
 #endif
 
 void UPolys::AddReferencedObjects(UObject* InThis, FReferenceCollector& Collector)
diff --git a/Engine/Source/Runtime/Engine/Public/Model.h b/Engine/Source/Runtime/Engine/Public/Model.h
--- a/Engine/Source/Runtime/Engine/Public/Model.h
+++ b/Engine/Source/Runtime/Engine/Public/Model.h
@@ -476,7 +476,6 @@ public:
 #if WITH_EDITOR
 	virtual void PostEditUndo() override;
 	virtual bool Modify(bool bAlwaysMarkDirty = false) override;
-	virtual void PreEditChange(FProperty*) override;
 #endif // WITH_EDITOR
 	virtual bool Rename( const TCHAR* InName=NULL, UObject* NewOuter=NULL, ERenameFlags Flags=REN_None ) override;
 	static void AddReferencedObjects(UObject* InThis, FReferenceCollector& Collector);
