Removed leftover debug code

diff --git a/Engine/Source/Runtime/Engine/Private/WorldPartition/HLOD/HLODActor.cpp b/Engine/Source/Runtime/Engine/Private/WorldPartition/HLOD/HLODActor.cpp
--- a/Engine/Source/Runtime/Engine/Private/WorldPartition/HLOD/HLODActor.cpp
+++ b/Engine/Source/Runtime/Engine/Private/WorldPartition/HLOD/HLODActor.cpp
@@ -102,11 +102,6 @@ void AWorldPartitionHLOD::SetHLODPrimitives(const TArray<UPrimitiveComponent*>&
 		{
 			InHLODPrimitive->SetupAttachment(RootComponent);
 		}
-
-		// Setup custom depth rendering to achieve a red tint using a post process material
-		const int32 CellPreviewStencilValue = 180;
-		InHLODPrimitive->bRenderCustomDepth = true;
-		InHLODPrimitive->CustomDepthStencilValue = CellPreviewStencilValue;
 	
 		InHLODPrimitive->SetCollisionEnabled(ECollisionEnabled::NoCollision);
 		InHLODPrimitive->SetMobility(EComponentMobility::Static);
