Changed FActorSpawnParameters::bCreateActorPackage default value to true. This way, any code that spawns an actor *only in a level that supports it* will create a package for this actor.

diff --git a/Engine/Source/Editor/UnrealEd/Private/ActorGroupingUtils.cpp b/Engine/Source/Editor/UnrealEd/Private/ActorGroupingUtils.cpp
--- a/Engine/Source/Editor/UnrealEd/Private/ActorGroupingUtils.cpp
+++ b/Engine/Source/Editor/UnrealEd/Private/ActorGroupingUtils.cpp
@@ -81,7 +81,6 @@ void UActorGroupingUtils::GroupActors(const TArray<AActor*>& ActorsToGroup)
 
 					FActorSpawnParameters SpawnInfo;
 					SpawnInfo.OverrideLevel = ActorLevel;
-					SpawnInfo.bCreateActorPackage = true;
 					AGroupActor* SpawnedGroupActor = World->SpawnActor<AGroupActor>(SpawnInfo);
 
 					bool bActorsInSameFolder = true;
diff --git a/Engine/Source/Editor/UnrealEd/Private/EditorEngine.cpp b/Engine/Source/Editor/UnrealEd/Private/EditorEngine.cpp
--- a/Engine/Source/Editor/UnrealEd/Private/EditorEngine.cpp
+++ b/Engine/Source/Editor/UnrealEd/Private/EditorEngine.cpp
@@ -6443,7 +6443,6 @@ AActor* UEditorEngine::AddActor(ULevel* InLevel, UClass* Class, const FTransform
 		FActorSpawnParameters SpawnInfo;
 		SpawnInfo.OverrideLevel = DesiredLevel;
 		SpawnInfo.SpawnCollisionHandlingOverride = ESpawnActorCollisionHandlingMethod::AlwaysSpawn;
-		SpawnInfo.bCreateActorPackage = true;
 		SpawnInfo.ObjectFlags = InObjectFlags;
 		const auto Location = Transform.GetLocation();
 		const auto Rotation = Transform.GetRotation().Rotator();
diff --git a/Engine/Source/Editor/UnrealEd/Private/Factories/ActorFactory.cpp b/Engine/Source/Editor/UnrealEd/Private/Factories/ActorFactory.cpp
--- a/Engine/Source/Editor/UnrealEd/Private/Factories/ActorFactory.cpp
+++ b/Engine/Source/Editor/UnrealEd/Private/Factories/ActorFactory.cpp
@@ -313,7 +313,6 @@ AActor* UActorFactory::SpawnActor( UObject* Asset, ULevel* InLevel, const FTrans
 		FActorSpawnParameters SpawnInfo;
 		SpawnInfo.OverrideLevel = InLevel;
 		SpawnInfo.ObjectFlags = InObjectFlags;
-		SpawnInfo.bCreateActorPackage = true;
 		SpawnInfo.Name = Name;
 #if WITH_EDITOR
 		SpawnInfo.bTemporaryEditorActor = FLevelEditorViewportClient::IsDroppingPreviewActor();
diff --git a/Engine/Source/Editor/UnrealEd/Private/Factories/EditorFactories.cpp b/Engine/Source/Editor/UnrealEd/Private/Factories/EditorFactories.cpp
--- a/Engine/Source/Editor/UnrealEd/Private/Factories/EditorFactories.cpp
+++ b/Engine/Source/Editor/UnrealEd/Private/Factories/EditorFactories.cpp
@@ -974,7 +974,6 @@ UObject* ULevelFactory::FactoryCreateText
 						SpawnInfo.Name = ActorUniqueName;
 						SpawnInfo.Template = Archetype;
 						SpawnInfo.SpawnCollisionHandlingOverride = ESpawnActorCollisionHandlingMethod::AlwaysSpawn;
-						SpawnInfo.bCreateActorPackage = true;
 						AActor* NewActor = World->SpawnActor( TempClass, nullptr, nullptr, SpawnInfo );
 						
 						if( NewActor )
diff --git a/Engine/Source/Editor/UnrealEd/Private/Kismet2/KismetReinstanceUtilities.cpp b/Engine/Source/Editor/UnrealEd/Private/Kismet2/KismetReinstanceUtilities.cpp
--- a/Engine/Source/Editor/UnrealEd/Private/Kismet2/KismetReinstanceUtilities.cpp
+++ b/Engine/Source/Editor/UnrealEd/Private/Kismet2/KismetReinstanceUtilities.cpp
@@ -1815,6 +1815,7 @@ static void ReplaceActorHelper(AActor* OldActor, UClass* OldClass, UObject*& New
 	SpawnInfo.SpawnCollisionHandlingOverride = ESpawnActorCollisionHandlingMethod::AlwaysSpawn;
 	SpawnInfo.bDeferConstruction = true;
 	SpawnInfo.Name = OldActor->GetFName();
+	SpawnInfo.bCreateActorPackage = false;
 
 	if (!OldActor->IsListedInSceneOutliner())
 	{
diff --git a/Engine/Source/Runtime/Engine/Classes/Engine/World.h b/Engine/Source/Runtime/Engine/Classes/Engine/World.h
--- a/Engine/Source/Runtime/Engine/Classes/Engine/World.h
+++ b/Engine/Source/Runtime/Engine/Classes/Engine/World.h
@@ -568,7 +568,7 @@ public:
 	/* Determines whether or not the actor should be hidden from the Scene Outliner */
 	uint8	bHideFromSceneOutliner:1;
 
-	/** Determines whether to create a new package for the actor or not. */
+	/** Determines whether to create a new package for the actor or not, if the level supports it. */
 	uint16	bCreateActorPackage:1;
 #endif
 
diff --git a/Engine/Source/Runtime/Engine/Private/ActorPartition/ActorPartitionSubsystem.cpp b/Engine/Source/Runtime/Engine/Private/ActorPartition/ActorPartitionSubsystem.cpp
--- a/Engine/Source/Runtime/Engine/Private/ActorPartition/ActorPartitionSubsystem.cpp
+++ b/Engine/Source/Runtime/Engine/Private/ActorPartition/ActorPartitionSubsystem.cpp
@@ -89,7 +89,6 @@ public:
 		{
 			FActorSpawnParameters SpawnParams;
 			SpawnParams.OverrideLevel = InCellCoord.Level;
-			SpawnParams.bCreateActorPackage = true;
 			FoundActor = CastChecked<APartitionActor>(World->SpawnActor(InActorClass, nullptr, nullptr, SpawnParams));
 		}
 
@@ -177,7 +176,6 @@ public:
 		if (!FoundActor && bInCreate)
 		{
 			FActorSpawnParameters SpawnParams;
-			SpawnParams.bCreateActorPackage = true;
 			SpawnParams.Name = FName(*FString::Printf(TEXT("%s_%d_%d_%d"), *InActorClass->GetName(), InCellCoord.X, InCellCoord.Y, InCellCoord.Z));
 			SpawnParams.NameMode = FActorSpawnParameters::ESpawnActorNameMode::Requested;
 							
diff --git a/Engine/Source/Runtime/Engine/Private/Components/ChildActorComponent.cpp b/Engine/Source/Runtime/Engine/Private/Components/ChildActorComponent.cpp
--- a/Engine/Source/Runtime/Engine/Private/Components/ChildActorComponent.cpp
+++ b/Engine/Source/Runtime/Engine/Private/Components/ChildActorComponent.cpp
@@ -616,6 +616,7 @@ void UChildActorComponent::CreateChildActor()
 				Params.Name = ChildActorName;
 				Params.NameMode = FActorSpawnParameters::ESpawnActorNameMode::Requested;
 #if WITH_EDITOR
+				Params.bCreateActorPackage = false;
 				Params.OverridePackage = GetOwner()->GetExternalPackage();
 				Params.OverrideParentComponent = this;
 				Params.OverrideActorGuid = CachedInstanceData ? CachedInstanceData->ChildActorGUID : FGuid();
diff --git a/Engine/Source/Runtime/Engine/Private/LevelInstance/LevelInstanceSubsystem.cpp b/Engine/Source/Runtime/Engine/Private/LevelInstance/LevelInstanceSubsystem.cpp
--- a/Engine/Source/Runtime/Engine/Private/LevelInstance/LevelInstanceSubsystem.cpp
+++ b/Engine/Source/Runtime/Engine/Private/LevelInstance/LevelInstanceSubsystem.cpp
@@ -592,7 +592,6 @@ ALevelInstance* ULevelInstanceSubsystem::CreateLevelInstanceFrom(const TArray<AA
 	}
 	
 	FActorSpawnParameters SpawnParams;
-	SpawnParams.bCreateActorPackage = true;
 	SpawnParams.OverrideLevel = CurrentLevel;
 	ALevelInstance* NewLevelInstanceActor = GetWorld()->SpawnActor<ALevelInstance>(ALevelInstance::StaticClass(), SpawnParams);
 	check(NewLevelInstanceActor);
diff --git a/Engine/Source/Runtime/Engine/Private/World.cpp b/Engine/Source/Runtime/Engine/Private/World.cpp
--- a/Engine/Source/Runtime/Engine/Private/World.cpp
+++ b/Engine/Source/Runtime/Engine/Private/World.cpp
@@ -180,7 +180,7 @@ FActorSpawnParameters::FActorSpawnParameters()
 #if WITH_EDITOR
 , bTemporaryEditorActor(false)
 , bHideFromSceneOutliner(false)
-, bCreateActorPackage(false)
+, bCreateActorPackage(true)
 #endif
 , NameMode(ESpawnActorNameMode::Required_Fatal)
 , ObjectFlags(RF_Transactional)
diff --git a/Engine/Source/Runtime/Engine/Private/WorldPartition/HLOD/HLODBuilder.cpp b/Engine/Source/Runtime/Engine/Private/WorldPartition/HLOD/HLODBuilder.cpp
--- a/Engine/Source/Runtime/Engine/Private/WorldPartition/HLOD/HLODBuilder.cpp
+++ b/Engine/Source/Runtime/Engine/Private/WorldPartition/HLOD/HLODBuilder.cpp
@@ -40,7 +40,6 @@ public:
 		if (!HLODActor)
 		{
 			FActorSpawnParameters SpawnParams;
-			SpawnParams.bCreateActorPackage = true;
 			SpawnParams.Name = *HLODActorName;
 			HLODActor = World->SpawnActor<AWorldPartitionHLOD>(SpawnParams);
 		}
