WorldPartition Edition Cell Loading - Fixed invalid progress count

diff --git a/Engine/Source/Runtime/Engine/Private/WorldPartition/WorldPartition.cpp b/Engine/Source/Runtime/Engine/Private/WorldPartition/WorldPartition.cpp
--- a/Engine/Source/Runtime/Engine/Private/WorldPartition/WorldPartition.cpp
+++ b/Engine/Source/Runtime/Engine/Private/WorldPartition/WorldPartition.cpp
@@ -764,7 +764,7 @@ void UWorldPartition::LoadEditorCells(const TArray<UWorldPartitionEditorCell*>&
 
 	for (UWorldPartitionEditorCell* Cell : CellsToLoad)
 	{
-		SlowTask.EnterProgressFrame(Cell->Actors.Num() - Cell->LoadedActors.Num(), FText::Format(LOCTEXT("LoadingActors", "Loading actors {0}/{1}"), (int32)SlowTask.CompletedWork, NumActorsToLoad));
+		SlowTask.EnterProgressFrame(Cell->Actors.Num() - Cell->LoadedActors.Num());
 		UpdateLoadingEditorCell(Cell, true);
 	}
 }
@@ -782,7 +782,7 @@ void UWorldPartition::UnloadEditorCells(const TArray<UWorldPartitionEditorCell*>
 
 	for (UWorldPartitionEditorCell* Cell: CellsToUnload)
 	{
-		SlowTask.EnterProgressFrame(Cell->LoadedActors.Num(), FText::Format(LOCTEXT("UnloadingActors", "Unloading actors {0}/{1}"), (int32)SlowTask.CompletedWork, NumActorsToUnload));
+		SlowTask.EnterProgressFrame(Cell->LoadedActors.Num());
 		UpdateLoadingEditorCell(Cell, false);
 	}
 }
