Do not list HLODActors in the scene outliner

diff --git a/Engine/Source/Runtime/Engine/Private/WorldPartition/HLOD/HLODActor.cpp b/Engine/Source/Runtime/Engine/Private/WorldPartition/HLOD/HLODActor.cpp
--- a/Engine/Source/Runtime/Engine/Private/WorldPartition/HLOD/HLODActor.cpp
+++ b/Engine/Source/Runtime/Engine/Private/WorldPartition/HLOD/HLODActor.cpp
@@ -12,16 +12,15 @@
 #include "WorldPartition/HLOD/HLODLayer.h"
 #endif
 
-UWorldPartitionRuntimeHLODCellData::UWorldPartitionRuntimeHLODCellData(const FObjectInitializer& ObjectInitializer)
-	: Super(ObjectInitializer)
-{
-}
-
 AWorldPartitionHLOD::AWorldPartitionHLOD(const FObjectInitializer& ObjectInitializer)
 	: Super(ObjectInitializer)
 {
 	SetCanBeDamaged(false);
 	SetActorEnableCollision(false);
+
+#if WITH_EDITORONLY_DATA
+	bListedInSceneOutliner = false;
+#endif
 }
 
 UPrimitiveComponent* AWorldPartitionHLOD::GetHLODComponent()
@@ -200,6 +199,16 @@ void AWorldPartitionHLOD::PostActorCreated()
 	HLODGuid = GetActorGuid();
 }
 
+#endif // WITH_EDITOR
+
+
+UWorldPartitionRuntimeHLODCellData::UWorldPartitionRuntimeHLODCellData(const FObjectInitializer& ObjectInitializer)
+	: Super(ObjectInitializer)
+{
+}
+
+#if WITH_EDITOR
+
 void UWorldPartitionRuntimeHLODCellData::SetReferencedHLODActors(TArray<FGuid>&& InReferencedHLODActors)
 {
 	ReferencedHLODActors = MoveTemp(InReferencedHLODActors);
