HLODLayer cleanup
* Removed FHLODLevelSettings to store a single setting level per HLODLayer
* Multiple HLOD  levels will ultimately be achieved by stacking HLODLayers
* Renamed fields & added documentation

diff --git a/Engine/Source/Runtime/Engine/Private/WorldPartition/HLOD/HLODBuilder.cpp b/Engine/Source/Runtime/Engine/Private/WorldPartition/HLOD/HLODBuilder.cpp
--- a/Engine/Source/Runtime/Engine/Private/WorldPartition/HLOD/HLODBuilder.cpp
+++ b/Engine/Source/Runtime/Engine/Private/WorldPartition/HLOD/HLODBuilder.cpp
@@ -50,12 +50,12 @@ public:
 		if (!HLODPrimitives.IsEmpty())
 		{
 			HLODActor->Modify();
-			HLODActor->SetHLODLayer(HLODLayer, iLevel);
+			HLODActor->SetHLODLayer(HLODLayer, 0);
 			HLODActor->SetHLODPrimitives(HLODPrimitives, CellLoadingRange);
 			HLODActor->SetHLODBounds(CellBounds);
 			HLODActor->SetChildrenPrimitives(InSubComponents);
 			HLODActor->SetActorLabel(HLODActorName);
-			HLODActor->RuntimeGrid = GetHLODLevelSettings().TargetGrid;
+			HLODActor->RuntimeGrid = HLODLayer->GetRuntimeGrid();
 
 			WorldPartition->UpdateActorDesc(HLODActor);
 		}
@@ -90,11 +90,6 @@ public:
 		return PrimitiveComponents;
 	}
 
-	const FHLODLevelSettings& GetHLODLevelSettings() const
-	{
-		return HLODLayer->GetLevels()[iLevel];
-	}
-
 public:
 	UWorld*				World;
 	UWorldPartition*	WorldPartition;
@@ -117,9 +112,7 @@ class FHLODBuilder_Instancing : public FHLODBuilder
 	{
 		TRACE_CPUPROFILER_EVENT_SCOPE(FHLODBuilder_Instancing::BuildHLOD);
 
-		const FHLODLevelSettings& LevelSettings = GetHLODLevelSettings();
-
-		FCreateComponentsFunction CreateComponentLambda = [&LevelSettings, &InSubComponents](AWorldPartitionHLOD* HLODActor)
+		FCreateComponentsFunction CreateComponentLambda = [&InSubComponents](AWorldPartitionHLOD* HLODActor)
 		{
 			TArray<UPrimitiveComponent*> Components;
 
@@ -169,15 +162,13 @@ class FHLODBuilder_MeshMerge : public FHLODBuilder
 	{
 		TRACE_CPUPROFILER_EVENT_SCOPE(UHLODLayer::BuildHLOD_MeshMerge);
 
-		const FHLODLevelSettings& LevelSettings = GetHLODLevelSettings();
-
-		FCreateComponentsFunction CreateComponentLambda = [&LevelSettings, &InSubComponents, this](AWorldPartitionHLOD* HLODActor)
+		FCreateComponentsFunction CreateComponentLambda = [&InSubComponents, this](AWorldPartitionHLOD* HLODActor)
 		{
 			TArray<UObject*> Assets;
 			FVector MergedActorLocation;
 
 			const IMeshMergeUtilities& MeshMergeUtilities = FModuleManager::Get().LoadModuleChecked<IMeshMergeModule>("MeshMergeUtilities").GetUtilities();
-			MeshMergeUtilities.MergeComponentsToStaticMesh(InSubComponents, HLODActor->GetWorld(), LevelSettings.MergeSetting, LevelSettings.FlattenMaterial.LoadSynchronous(), HLODActor->GetPackage(), CellName.ToString(), Assets, MergedActorLocation, 0.25f, false);
+			MeshMergeUtilities.MergeComponentsToStaticMesh(InSubComponents, HLODActor->GetWorld(), HLODLayer->GetMeshMergeSettings(), HLODLayer->GetHLODMaterial().LoadSynchronous(), HLODActor->GetPackage(), CellName.ToString(), Assets, MergedActorLocation, 0.25f, false);
 
 			// All merged mesh assets are stored in the HLOD Actor package
 			Algo::ForEach(Assets, [](UObject* Asset) { Asset->ClearFlags(RF_Public | RF_Standalone); });
@@ -207,9 +198,7 @@ class FHLODBuilder_MeshSimplify : public FHLODBuilder
 	{
 		TRACE_CPUPROFILER_EVENT_SCOPE(UHLODLayer::BuildHLOD_MeshProxy);
 
-		const FHLODLevelSettings& LevelSettings = GetHLODLevelSettings();
-
-		FCreateComponentsFunction CreateComponentLambda = [&LevelSettings, &InSubComponents, this] (AWorldPartitionHLOD* HLODActor)
+		FCreateComponentsFunction CreateComponentLambda = [&InSubComponents, this] (AWorldPartitionHLOD* HLODActor)
 		{
 			TArray<UObject*> Assets;
 			FCreateProxyDelegate ProxyDelegate;
@@ -219,7 +208,7 @@ class FHLODBuilder_MeshSimplify : public FHLODBuilder
 			Algo::TransformIf(InSubComponents, StaticMeshComponents, [](UPrimitiveComponent* InPrimitiveComponent) { return InPrimitiveComponent->IsA<UStaticMeshComponent>(); }, [](UPrimitiveComponent* InPrimitiveComponent) { return Cast<UStaticMeshComponent>(InPrimitiveComponent); });
 
 			const IMeshMergeUtilities& MeshMergeUtilities = FModuleManager::Get().LoadModuleChecked<IMeshMergeModule>("MeshMergeUtilities").GetUtilities();
-			MeshMergeUtilities.CreateProxyMesh(StaticMeshComponents, LevelSettings.ProxySetting, LevelSettings.FlattenMaterial.LoadSynchronous(), HLODActor->GetPackage(), CellName.ToString(), FGuid::NewGuid(), ProxyDelegate, true);
+			MeshMergeUtilities.CreateProxyMesh(StaticMeshComponents, HLODLayer->GetMeshSimplifySettings(), HLODLayer->GetHLODMaterial().LoadSynchronous(), HLODActor->GetPackage(), CellName.ToString(), FGuid::NewGuid(), ProxyDelegate, true);
 
 			// All merged mesh assets are stored in the HLOD Actor package
 			Algo::ForEach(Assets, [](UObject* Asset) { Asset->ClearFlags(RF_Public | RF_Standalone); });
@@ -241,49 +230,42 @@ class FHLODBuilder_MeshSimplify : public FHLODBuilder
 
 TArray<AWorldPartitionHLOD*> FHLODBuilderUtilities::BuildHLODs(UWorldPartition* InWorldPartition, FName InCellName, FBox InCellBounds, float InCellLoadingRange, const UHLODLayer* InHLODLayer, const TArray<AActor*>& InSubActors)
 {
-	const TArray<FHLODLevelSettings>& LevelsSettings = InHLODLayer->GetLevels();
-
-	TArray<AWorldPartitionHLOD*> HLODActors;
+	TUniquePtr<FHLODBuilder> HLODBuilder = nullptr;
 
-	for (int32 iLevel = 0; iLevel < LevelsSettings.Num(); ++iLevel)
+	EHLODLayerType HLODLevelType = InHLODLayer->GetLayerType();
+	switch (HLODLevelType)
 	{
-		TUniquePtr<FHLODBuilder> HLODBuilder = nullptr;
+	case EHLODLayerType::Instancing:
+		HLODBuilder = TUniquePtr<FHLODBuilder>(new FHLODBuilder_Instancing());
+		break;
 
-		EHLODLevelType HLODLevelType = LevelsSettings[iLevel].LevelType;
-		switch (HLODLevelType)
-		{
-		case EHLODLevelType::Instancing:
-			HLODBuilder = TUniquePtr<FHLODBuilder>(new FHLODBuilder_Instancing());
-			break;
-
-		case EHLODLevelType::MeshMerge:
-			HLODBuilder = TUniquePtr<FHLODBuilder>(new FHLODBuilder_MeshMerge());
-			break;
+	case EHLODLayerType::MeshMerge:
+		HLODBuilder = TUniquePtr<FHLODBuilder>(new FHLODBuilder_MeshMerge());
+		break;
 
-		case EHLODLevelType::MeshSimplify:
-			HLODBuilder = TUniquePtr<FHLODBuilder>(new FHLODBuilder_MeshSimplify());
-			break;
+	case EHLODLayerType::MeshSimplify:
+		HLODBuilder = TUniquePtr<FHLODBuilder>(new FHLODBuilder_MeshSimplify());
+		break;
 
-		default:
-			checkf(false, TEXT("Unsupported type"));
-		}
+	default:
+		checkf(false, TEXT("Unsupported type"));
+	}
 
-		TArray<UPrimitiveComponent*> SubComponents = FHLODBuilder::GatherPrimitiveComponents(iLevel, InSubActors);
+	TArray<UPrimitiveComponent*> SubComponents = FHLODBuilder::GatherPrimitiveComponents(0, InSubActors);
+	TArray<AWorldPartitionHLOD*> HLODActors;
 
-		if (HLODBuilder && !SubComponents.IsEmpty())
-		{
-			HLODBuilder->World = InWorldPartition->GetWorld();
-			HLODBuilder->WorldPartition = InWorldPartition;
-			HLODBuilder->HLODLayer = InHLODLayer;
-			HLODBuilder->iLevel = iLevel;
-			HLODBuilder->CellName = InCellName;
-			HLODBuilder->CellBounds = InCellBounds;
-			HLODBuilder->CellLoadingRange = InCellLoadingRange;
-
-			HLODBuilder->Build(SubComponents);
+	if (HLODBuilder && !SubComponents.IsEmpty())
+	{
+		HLODBuilder->World = InWorldPartition->GetWorld();
+		HLODBuilder->WorldPartition = InWorldPartition;
+		HLODBuilder->HLODLayer = InHLODLayer;
+		HLODBuilder->CellName = InCellName;
+		HLODBuilder->CellBounds = InCellBounds;
+		HLODBuilder->CellLoadingRange = InCellLoadingRange;
+
+		HLODBuilder->Build(SubComponents);
 			
-			HLODActors += HLODBuilder->HLODActors;
-		}
+		HLODActors = HLODBuilder->HLODActors;
 	}
 
 	return HLODActors;
diff --git a/Engine/Source/Runtime/Engine/Private/WorldPartition/HLOD/HLODLayer.cpp b/Engine/Source/Runtime/Engine/Private/WorldPartition/HLOD/HLODLayer.cpp
--- a/Engine/Source/Runtime/Engine/Private/WorldPartition/HLOD/HLODLayer.cpp
+++ b/Engine/Source/Runtime/Engine/Private/WorldPartition/HLOD/HLODLayer.cpp
@@ -15,6 +15,7 @@
 
 UHLODLayer::UHLODLayer(const FObjectInitializer& ObjectInitializer)
 	: Super(ObjectInitializer)
+	, LoadingRange(30000)
 {
 }
 
diff --git a/Engine/Source/Runtime/Engine/Private/WorldPartition/WorldPartitionRuntimeSpatialHash.cpp b/Engine/Source/Runtime/Engine/Private/WorldPartition/WorldPartitionRuntimeSpatialHash.cpp
--- a/Engine/Source/Runtime/Engine/Private/WorldPartition/WorldPartitionRuntimeSpatialHash.cpp
+++ b/Engine/Source/Runtime/Engine/Private/WorldPartition/WorldPartitionRuntimeSpatialHash.cpp
@@ -1145,8 +1145,7 @@ bool UWorldPartitionRuntimeSpatialHash::GenerateHLOD()
 		return false;
 	}
 
-	// Recreate HLODGrids from known layers
-	HLODGrids.Reset();
+	TArray<FSpatialHashRuntimeGrid> NewHLODGrids;
 
 	// Gather up all HLODLayer assets
 	FAssetRegistryModule& AssetRegistryModule = FModuleManager::LoadModuleChecked<FAssetRegistryModule>(TEXT("AssetRegistry"));
@@ -1156,22 +1155,42 @@ bool UWorldPartitionRuntimeSpatialHash::GenerateHLOD()
 	// Ensure all assets are loaded
 	for (const FAssetData& HLODLayerAsset : HLODLayerAssets)
 	{
-		if (UHLODLayer* HLODLayer = Cast<UHLODLayer>(HLODLayerAsset.GetAsset()))
+		if (const UHLODLayer* HLODLayer = Cast<UHLODLayer>(HLODLayerAsset.GetAsset()))
 		{
-			for (const FHLODLevelSettings& HLODLevel : HLODLayer->GetLevels())
+			check(HLODLayer->GetLoadingRange() > 0);
+
+			FName RuntimeGrid = HLODLayer->GetRuntimeGrid();
+			if (RuntimeGrid == NAME_None)
+			{
+				UE_LOG(LogWorldPartitionRuntimeSpatialHash, Error, TEXT("Invalid partition grids setup"));
+				return false;
+			}
+
+			const FSpatialHashRuntimeGrid* ExistingGrid = NewHLODGrids.FindByPredicate([RuntimeGrid](const FSpatialHashRuntimeGrid& ExistingGrid) { return ExistingGrid.GridName == RuntimeGrid; });
+			if (ExistingGrid)
+			{
+				UE_LOG(LogWorldPartitionRuntimeSpatialHash, Error, TEXT("Found two HLOD grids using the same RuntimeGrid name: %s"), *RuntimeGrid.ToString());
+				return false;
+			}
+
+			ExistingGrid = Grids.FindByPredicate([RuntimeGrid](const FSpatialHashRuntimeGrid& ExistingGrid) { return ExistingGrid.GridName == RuntimeGrid; });
+			if (ExistingGrid)
 			{
-				check(HLODLevel.TargetGrid != NAME_None);
-				check(HLODLevel.LoadingRange > 0);
-
-				FSpatialHashRuntimeGrid& HLODGrid = HLODGrids.AddDefaulted_GetRef();
-				HLODGrid.GridName = HLODLevel.TargetGrid;
-				HLODGrid.CellSize = HLODLevel.LoadingRange * 2;	// @todo_ow: Proper setup
-				HLODGrid.LoadingRange = HLODLevel.LoadingRange;
-				HLODGrid.DebugColor = FLinearColor::Red;
+				UE_LOG(LogWorldPartitionRuntimeSpatialHash, Error, TEXT("HLODLayer RuntimeGrid must not match an existing grid name: %s"), *RuntimeGrid.ToString());
+				return false;
 			}
+
+			FSpatialHashRuntimeGrid& HLODGrid = NewHLODGrids.AddDefaulted_GetRef();
+			HLODGrid.GridName = HLODLayer->GetRuntimeGrid();
+			HLODGrid.CellSize = HLODLayer->GetLoadingRange() * 2;	// @todo_ow: Proper setup
+			HLODGrid.LoadingRange = HLODLayer->GetLoadingRange();
+			HLODGrid.DebugColor = FLinearColor::Red;
 		}
 	}
 
+	// Recreate HLODGrids from known layers
+	HLODGrids = NewHLODGrids;
+
 	TMap<FName, int32> GridsMapping;
 	GridsMapping.Add(NAME_None, 0);
 	for (int32 i = 0; i < Grids.Num(); i++)
@@ -1181,7 +1200,7 @@ bool UWorldPartitionRuntimeSpatialHash::GenerateHLOD()
 		GridsMapping.Add(Grid.GridName, i);
 	}
 
-	// Ignore HLODD clusters as we are rebuilding them
+	// Ignore HLOD clusters as we are rebuilding them
 	for (const FSpatialHashRuntimeGrid& HLODGrid : HLODGrids)
 	{
 		check(!GridsMapping.Contains(HLODGrid.GridName));
diff --git a/Engine/Source/Runtime/Engine/Public/WorldPartition/HLOD/HLODLayer.h b/Engine/Source/Runtime/Engine/Public/WorldPartition/HLOD/HLODLayer.h
--- a/Engine/Source/Runtime/Engine/Public/WorldPartition/HLOD/HLODLayer.h
+++ b/Engine/Source/Runtime/Engine/Public/WorldPartition/HLOD/HLODLayer.h
@@ -17,37 +17,11 @@ class UWorldPartition;
 class AWorldPartitionHLOD;
 
 UENUM()
-enum class EHLODLevelType : uint8
+enum class EHLODLayerType : uint8
 {
-	Instancing = 0		UMETA(DisplayName = "Use instancing"),
-	MeshMerge = 1		UMETA(DisplayName = "Generate merged mesh"),
-	MeshSimplify = 2	UMETA(DisplayName = "Generate simplified mesh")
-};
-
-USTRUCT()
-struct FHLODLevelSettings
-{
-	GENERATED_USTRUCT_BODY()
-
-#if WITH_EDITORONLY_DATA
-	UPROPERTY(EditAnywhere, Category=HLOD)
-	EHLODLevelType LevelType;
-
-	UPROPERTY(EditAnywhere, Category=HLOD, meta = (editcondition = "LevelType == EHLODLevelType::MeshMerge", ShowOnlyInnerProperties))
-	FMeshMergingSettings MergeSetting;
-
-	UPROPERTY(EditAnywhere, Category=HLOD, meta = (editcondition = "LevelType == EHLODLevelType::MeshSimplify", ShowOnlyInnerProperties))
-	FMeshProxySettings ProxySetting;
-
-	UPROPERTY(EditAnywhere, Category=HLOD, meta = (editcondition = "LevelType == EHLODLevelType::MeshMerge || LevelType == EHLODLevelType::MeshSimplify", ShowOnlyInnerProperties))
-	TSoftObjectPtr<UMaterial> FlattenMaterial;
-
-	UPROPERTY(EditAnywhere, Category=HLOD)
-	FName TargetGrid;
-
-	UPROPERTY(EditAnywhere, Category=HLOD)
-	float LoadingRange;
-#endif // WITH_EDITORONLY_DATA
+	Instancing = 0		UMETA(DisplayName = "Instancing"),
+	MeshMerge = 1		UMETA(DisplayName = "Merged Mesh"),
+	MeshSimplify = 2	UMETA(DisplayName = "Simplified Mesh")
 };
 
 UCLASS(Blueprintable, Config=Engine, PerObjectConfig)
@@ -57,20 +31,49 @@ class ENGINE_API UHLODLayer : public UObject
 	
 #if WITH_EDITOR
 public:
-	const TArray<FHLODLevelSettings>& GetLevels() const { return Levels; }
-
 	static TArray<AWorldPartitionHLOD*> GenerateHLODForCell(UWorldPartition* InWorldPartition, FName InCellName, FBox InCellBounds, float InCellLoadingRange, const TSet<FGuid>& InCellActors);
 
 	static UHLODLayer* GetHLODLayer(const AActor* InActor);
 	static bool ShouldIncludeInHLOD(UPrimitiveComponent* InComponent, int32 InLevelIndex);
 #endif
 
-	static bool ShouldIncludeInHLOD(AActor* InActor);
+#if WITH_EDITORONLY_DATA
+	EHLODLayerType GetLayerType() const { return LayerType; }
+	const FMeshMergingSettings& GetMeshMergeSettings() const { return MeshMergeSettings; }
+	const FMeshProxySettings& GetMeshSimplifySettings() const { return MeshSimplifySettings; }
+	const TSoftObjectPtr<UMaterial>& GetHLODMaterial() const { return HLODMaterial; }
+	FName GetRuntimeGrid() const { return RuntimeGrid; }
+	float GetLoadingRange() const { return LoadingRange; }
+#endif
 
+	static bool ShouldIncludeInHLOD(AActor* InActor);
 
 #if WITH_EDITORONLY_DATA
 private:
-	UPROPERTY(EditAnywhere, Config, Category = HLOD)
-	TArray<FHLODLevelSettings> Levels;
+	/** Type of HLOD generation to use */
+	UPROPERTY(EditAnywhere, Config, Category=HLOD)
+	EHLODLayerType LayerType;
+
+	/** Merged mesh generation settings - Used when this layer is of type EHLODLayerType::MeshMerge */
+	UPROPERTY(EditAnywhere, Config, Category=HLOD, meta = (editcondition = "LayerType == EHLODLayerType::MeshMerge"))
+	FMeshMergingSettings MeshMergeSettings;
+
+	/** Simplified mesh generation settings - Used when this layer is of type EHLODLayerType::MeshSimplify */
+	UPROPERTY(EditAnywhere, Config, Category=HLOD, meta = (editcondition = "LayerType == EHLODLayerType::MeshSimplify"))
+	FMeshProxySettings MeshSimplifySettings;
+
+	/** Material that will be used by the generated HLOD static mesh */
+	UPROPERTY(EditAnywhere, Config, Category=HLOD, meta = (editcondition = "LayerType == EHLODLayerType::MeshMerge || LayerType == EHLODLayerType::MeshSimplify"))
+	TSoftObjectPtr<UMaterial> HLODMaterial;
+
+	/**
+	 * Runtime grid that will be used by the generated HLOD. 
+	 * For now, this grid will be autogenerated and must not be manually added to the world partition. */
+	UPROPERTY(EditAnywhere, Config, Category=HLOD)
+	FName RuntimeGrid;
+
+	/** Loading range of the RuntimeGrid */
+	UPROPERTY(EditAnywhere, Config, Category=HLOD)
+	float LoadingRange;
 #endif
 };
