Modified LevelInstanceSubsystem::CanMoveActorToLevel so that it now returns the reason as an optional out param rather than logging every time.

diff --git a/Engine/Source/Editor/SceneOutliner/Private/ActorBrowsingMode.cpp b/Engine/Source/Editor/SceneOutliner/Private/ActorBrowsingMode.cpp
--- a/Engine/Source/Editor/SceneOutliner/Private/ActorBrowsingMode.cpp
+++ b/Engine/Source/Editor/SceneOutliner/Private/ActorBrowsingMode.cpp
@@ -821,9 +821,8 @@ FSceneOutlinerDragValidationInfo FActorBrowsingMode::ValidateDrop(const ISceneOu
 							}
 						}
 
-						if (!LevelInstanceSubsystem->CanMoveActorToLevel(DragActor))
+						if (!LevelInstanceSubsystem->CanMoveActorToLevel(DragActor, &AttachErrorMsg))
 						{
-							AttachErrorMsg = LOCTEXT("Error_MoveActorToLevelInstance", "Cannot move LevelInstance while it or its children are being edited");
 							bCanAttach = bDraggedOntoAttachmentParent = false;
 							break;
 						}
@@ -963,9 +962,10 @@ FSceneOutlinerDragValidationInfo FActorBrowsingMode::ValidateDrop(const ISceneOu
 
 					if (const ALevelInstance* LevelInstanceActor = Cast<ALevelInstance>(Actor))
 					{
-						if (!LevelInstanceSubsystem->CanMoveActorToLevel(LevelInstanceActor))
+						FText Reason;
+						if (!LevelInstanceSubsystem->CanMoveActorToLevel(LevelInstanceActor, &Reason))
 						{
-							return FSceneOutlinerDragValidationInfo(ESceneOutlinerDropCompatibility::IncompatibleGeneric, LOCTEXT("Error_RemoveEditingLevelInstance", "Cannot detach a LevelInstance which is currently being edited"));
+							return FSceneOutlinerDragValidationInfo(ESceneOutlinerDropCompatibility::IncompatibleGeneric, Reason);
 						}
 					}
 				}
diff --git a/Engine/Source/Runtime/Engine/Private/LevelInstance/LevelInstanceSubsystem.cpp b/Engine/Source/Runtime/Engine/Private/LevelInstance/LevelInstanceSubsystem.cpp
--- a/Engine/Source/Runtime/Engine/Private/LevelInstance/LevelInstanceSubsystem.cpp
+++ b/Engine/Source/Runtime/Engine/Private/LevelInstance/LevelInstanceSubsystem.cpp
@@ -561,8 +561,10 @@ ALevelInstance* ULevelInstanceSubsystem::CreateLevelInstanceFrom(const TArray<AA
 		const bool bIncludeChildren = true;
 		ActorLocationBox += ActorToMove->GetComponentsBoundingBox(bNonColliding, bIncludeChildren);
 
-		if (!CanMoveActorToLevel(ActorToMove))
+		FText Reason;
+		if (!CanMoveActorToLevel(ActorToMove, &Reason))
 		{
+			UE_LOG(LogLevelInstance, Warning, TEXT("%s"), *Reason.ToString());
 			return nullptr;
 		}
 	}
@@ -736,7 +738,7 @@ ULevelInstanceSubsystem::FLevelsToRemoveScope::~FLevelsToRemoveScope()
 	}
 }
 
-bool ULevelInstanceSubsystem::CanMoveActorToLevel(const AActor* Actor) const
+bool ULevelInstanceSubsystem::CanMoveActorToLevel(const AActor* Actor, FText* OutReason) const
 {
 	if (Actor->GetWorld() == GetWorld())
 	{
@@ -744,7 +746,10 @@ bool ULevelInstanceSubsystem::CanMoveActorToLevel(const AActor* Actor) const
 		{
 			if (IsEditingLevelInstance(LevelInstanceActor))
 			{
-				UE_LOG(LogLevelInstance, Warning, TEXT("Can't move Level Instance actor while it is being edited"));
+				if (OutReason != nullptr)
+				{
+					*OutReason = LOCTEXT("CanMoveActorLevelEditing", "Can't move Level Instance actor while it is being edited");
+				}
 				return false;
 			}
 
@@ -761,7 +766,10 @@ bool ULevelInstanceSubsystem::CanMoveActorToLevel(const AActor* Actor) const
 
 			if (bEditingChildren)
 			{
-				UE_LOG(LogLevelInstance, Warning, TEXT("Can't move Level Instance actor while one of its child Level Instance is being edited"));
+				if (OutReason != nullptr)
+				{
+					*OutReason = LOCTEXT("CanMoveActorToLevelChildEditing", "Can't move Level Instance actor while one of its child Level Instance is being edited");
+				}
 				return false;
 			}
 		}
diff --git a/Engine/Source/Runtime/Engine/Public/LevelInstance/LevelInstanceSubsystem.h b/Engine/Source/Runtime/Engine/Public/LevelInstance/LevelInstanceSubsystem.h
--- a/Engine/Source/Runtime/Engine/Public/LevelInstance/LevelInstanceSubsystem.h
+++ b/Engine/Source/Runtime/Engine/Public/LevelInstance/LevelInstanceSubsystem.h
@@ -79,7 +79,7 @@ public:
 	bool MoveActorsTo(ALevelInstance* LevelInstanceActor, const TArray<AActor*>& ActorsToMove);
 	bool BreakLevelInstance(ALevelInstance* LevelInstanceActor, uint32 Levels = 1);
 
-	bool CanMoveActorToLevel(const AActor* Actor) const;
+	bool CanMoveActorToLevel(const AActor* Actor, FText* OutReason = nullptr) const;
 	void DiscardEdits();
 	void OnActorDeleted(AActor* Actor);
 	ULevel* GetLevelInstanceLevel(const ALevelInstance* LevelInstanceActor) const;
