Creating a LevelInstance now resets the folder path of all actors belonging to it.

diff --git a/Engine/Source/Runtime/Engine/Private/LevelInstance/LevelInstanceSubsystem.cpp b/Engine/Source/Runtime/Engine/Private/LevelInstance/LevelInstanceSubsystem.cpp
--- a/Engine/Source/Runtime/Engine/Private/LevelInstance/LevelInstanceSubsystem.cpp
+++ b/Engine/Source/Runtime/Engine/Private/LevelInstance/LevelInstanceSubsystem.cpp
@@ -593,6 +593,15 @@ ALevelInstance* ULevelInstanceSubsystem::CreateLevelInstanceFrom(const TArray<AA
 		return nullptr;
 	}
 	
+	// Take all actors out of any folders they may have been in since we don't support folders inside of level instances
+	for (AActor* Actor : LoadedLevel->Actors)
+	{
+		if (Actor)
+		{
+			Actor->SetFolderPath_Recursively(NAME_None);
+		}
+	}
+
 	FActorSpawnParameters SpawnParams;
 	SpawnParams.OverrideLevel = CurrentLevel;
 	ALevelInstance* NewLevelInstanceActor = GetWorld()->SpawnActor<ALevelInstance>(ALevelInstance::StaticClass(), SpawnParams);
