Disable bool hiding modulation

diff --git a/Engine/Source/Runtime/Engine/Classes/Sound/SoundBase.h b/Engine/Source/Runtime/Engine/Classes/Sound/SoundBase.h
--- a/Engine/Source/Runtime/Engine/Classes/Sound/SoundBase.h
+++ b/Engine/Source/Runtime/Engine/Classes/Sound/SoundBase.h
@@ -151,19 +151,19 @@ public:
 	USoundAttenuation* AttenuationSettings;
 
 	/** Volume modulation */
-	UPROPERTY(EditAnywhere, BlueprintReadWrite, Category = "Modulation", meta = (DisplayName = "Volume", AudioParam = "Volume", EditCondition = "bModulationEnabled", EditConditionHides))
+	UPROPERTY(EditAnywhere, BlueprintReadWrite, Category = "Modulation", meta = (DisplayName = "Volume", AudioParam = "Volume"))
 	FSoundModulationDestinationSettings VolumeModulationDestination;
 
 	/** Pitch modulation */
-	UPROPERTY(EditAnywhere, BlueprintReadWrite, Category = "Modulation", meta = (DisplayName = "Pitch", AudioParam = "Pitch", EditCondition = "bModulationEnabled", EditConditionHides))
+	UPROPERTY(EditAnywhere, BlueprintReadWrite, Category = "Modulation", meta = (DisplayName = "Pitch", AudioParam = "Pitch"))
 	FSoundModulationDestinationSettings PitchModulationDestination;
 
 	/** Highpass modulation */
-	UPROPERTY(EditAnywhere, BlueprintReadWrite, Category = "Modulation", meta = (DisplayName = "Highpass", AudioParam = "HPFCutoffFrequency", EditCondition = "bModulationEnabled", EditConditionHides))
+	UPROPERTY(EditAnywhere, BlueprintReadWrite, Category = "Modulation", meta = (DisplayName = "Highpass", AudioParam = "HPFCutoffFrequency"))
 	FSoundModulationDestinationSettings HighpassModulationDestination;
 
 	/** Lowpass modulation */
-	UPROPERTY(EditAnywhere, BlueprintReadWrite, Category = "Modulation", meta = (DisplayName = "Lowpass", AudioParam = "LPFCutoffFrequency", EditCondition = "bModulationEnabled", EditConditionHides))
+	UPROPERTY(EditAnywhere, BlueprintReadWrite, Category = "Modulation", meta = (DisplayName = "Lowpass", AudioParam = "LPFCutoffFrequency"))
 	FSoundModulationDestinationSettings LowpassModulationDestination;
 
 	/** Submix to route sound output to. If unset, falls back to referenced SoundClass submix.
