Actor Partition
- Add missing actor descriptor update.
- Fix crash when converting to WP through commandlet.

diff --git a/Engine/Source/Runtime/Engine/Private/ActorPartition/ActorPartitionSubsystem.cpp b/Engine/Source/Runtime/Engine/Private/ActorPartition/ActorPartitionSubsystem.cpp
--- a/Engine/Source/Runtime/Engine/Private/ActorPartition/ActorPartitionSubsystem.cpp
+++ b/Engine/Source/Runtime/Engine/Private/ActorPartition/ActorPartitionSubsystem.cpp
@@ -132,7 +132,7 @@ public:
 		: FBaseActorPartition(InWorld)
 	{
 		WorldPartition = InWorld->GetSubsystem<UWorldPartitionSubsystem>();
-		check(WorldPartition);
+		check(WorldPartition || IsRunningCommandlet());
 	}
 
 	UActorPartitionSubsystem::FCellCoord GetActorPartitionHash(const FActorPartitionGetParams& GetParams) const override 
@@ -184,6 +184,8 @@ public:
 			FoundActor = CastChecked<APartitionActor>(World->SpawnActor(InActorClass, &CellCenter, nullptr, SpawnParams));
 			FoundActor->GridSize = GridSize;
 			FoundActor->bLockLocation = true;
+
+			WorldPartition->UpdateActorDesc(FoundActor);
 		}
 
 		check(FoundActor || !bInCreate);
