Bump VT DDC version for bad data in cache.

diff --git a/Engine/Source/Runtime/Engine/Private/TextureDerivedData.cpp b/Engine/Source/Runtime/Engine/Private/TextureDerivedData.cpp
--- a/Engine/Source/Runtime/Engine/Private/TextureDerivedData.cpp
+++ b/Engine/Source/Runtime/Engine/Private/TextureDerivedData.cpp
@@ -53,7 +53,7 @@
 
 // This GUID is mixed into DDC version for virtual textures only, this allows updating DDC version for VT without invalidating DDC for all textures
 // This is useful during development, but once large numbers of VT are present in shipped content, it will have the same problem as TEXTURE_DERIVEDDATA_VER
-#define TEXTURE_VT_DERIVEDDATA_VER	TEXT("0AA83222404A4BA48593D334268AF806")
+#define TEXTURE_VT_DERIVEDDATA_VER	TEXT("1B3FCA846FB74FD784CA5E08A524EB25")
 
 #if ENABLE_COOK_STATS
 namespace TextureCookStats
