Fix maintaining the world partition editor hash loaded node mask logic.

diff --git a/Engine/Source/Runtime/Engine/Private/WorldPartition/WorldPartitionEditorSpatialHash.cpp b/Engine/Source/Runtime/Engine/Private/WorldPartition/WorldPartitionEditorSpatialHash.cpp
--- a/Engine/Source/Runtime/Engine/Private/WorldPartition/WorldPartitionEditorSpatialHash.cpp
+++ b/Engine/Source/Runtime/Engine/Private/WorldPartition/WorldPartitionEditorSpatialHash.cpp
@@ -81,6 +81,7 @@ void UWorldPartitionEditorSpatialHash::Tick(float DeltaSeconds)
 						const uint32 ChildIndex = LevelCellCoord.GetChildIndex();
 						LevelCellCoord = LevelCellCoord.GetParentCellCoord();
 						const FCellNode& CellNode = HashNodes.FindChecked(LevelCellCoord);
+						check(CellNode.HasChildNode(ChildIndex));
 						check(CellNode.HasChildLoadedNode(ChildIndex));
 					}
 				}
@@ -142,43 +143,6 @@ void UWorldPartitionEditorSpatialHash::HashActor(FWorldPartitionActorDesc* InAct
 		const FBox ActorBounds = GetActorBounds(InActorDesc);
 		const int32 CurrentLevel = GetLevelForBox(Bounds);
 
-		auto UpdateHigherLevels = [this](const FCellCoord& CellCoord, int32 EndLevel)
-		{
-			FCellCoord LevelCellCoord = CellCoord;
-			while (LevelCellCoord.Level < EndLevel)
-			{
-				const uint32 ChildIndex = LevelCellCoord.GetChildIndex();
-
-				bool bChildHasLoadedNodes;
-				if (LevelCellCoord.Level)
-				{
-					FCellNode& ChildCellNode = HashNodes.FindChecked(LevelCellCoord);
-					bChildHasLoadedNodes = ChildCellNode.HasChildLoadedNodes();
-				}
-				else
-				{
-					UWorldPartitionEditorCell* EditorCell = HashCells.FindChecked(LevelCellCoord);
-					bChildHasLoadedNodes = EditorCell->bLoaded;
-				}
-
-				LevelCellCoord = LevelCellCoord.GetParentCellCoord();
-
-				FCellNode& CellNode = HashNodes.FindOrAdd(LevelCellCoord);
-
-				if (CellNode.HasChildNode(ChildIndex))
-				{
-					break;
-				}
-
-				CellNode.AddChildNode(ChildIndex);
-				
-				if (bChildHasLoadedNodes)
-				{
-					CellNode.AddChildLoadedNode(ChildIndex);
-				}
-			}
-		};
-
 		ForEachIntersectingCells(ActorBounds, 0, [&](const FCellCoord& CellCoord)
 		{
 			UWorldPartitionEditorCell* EditorCell = nullptr;
@@ -195,9 +159,25 @@ void UWorldPartitionEditorSpatialHash::HashActor(FWorldPartitionActorDesc* InAct
 				Cells.Add(EditorCell);
 				HashCells.Add(CellCoord, EditorCell);
 
-				UpdateHigherLevels(CellCoord, CurrentLevel);
-
 				Bounds += EditorCell->Bounds;
+
+				// Update parent cells
+				FCellCoord LevelCellCoord = CellCoord;
+				while (LevelCellCoord.Level < CurrentLevel)
+				{
+					const uint32 ChildIndex = LevelCellCoord.GetChildIndex();
+
+					LevelCellCoord = LevelCellCoord.GetParentCellCoord();
+
+					FCellNode& CellNode = HashNodes.FindOrAdd(LevelCellCoord);
+
+					if (CellNode.HasChildNode(ChildIndex))
+					{
+						break;
+					}
+
+					CellNode.AddChildNode(ChildIndex);
+				}
 			}
 
 			check(EditorCell);
@@ -213,7 +193,47 @@ void UWorldPartitionEditorSpatialHash::HashActor(FWorldPartitionActorDesc* InAct
 			{
 				if (CurrentLevel ? HashNodes.Contains(CellCoord) : HashCells.Contains(CellCoord))
 				{
-					UpdateHigherLevels(CellCoord, NewLevel);
+					bool bSetLoadedMask;
+					if (CellCoord.Level)
+					{
+						FCellNode& ChildCellNode = HashNodes.FindChecked(CellCoord);
+						bSetLoadedMask = ChildCellNode.HasChildLoadedNodes();
+					}
+					else
+					{
+						UWorldPartitionEditorCell* EditorCell = HashCells.FindChecked(CellCoord);
+						bSetLoadedMask = EditorCell->bLoaded;
+					}
+
+					FCellCoord LevelCellCoord = CellCoord;
+					while (LevelCellCoord.Level < NewLevel)
+					{
+						const uint32 ChildIndex = LevelCellCoord.GetChildIndex();
+
+						LevelCellCoord = LevelCellCoord.GetParentCellCoord();
+
+						FCellNode& CellNode = HashNodes.FindOrAdd(LevelCellCoord);
+
+						// We can break updating when aggregated flags are already properly set for parent nodes
+						const bool bShouldBreak = CellNode.HasChildNodes() && (!bSetLoadedMask || CellNode.HasChildLoadedNodes());
+
+						// Propagate the child mask
+						if (!CellNode.HasChildNode(ChildIndex))
+						{
+							CellNode.AddChildNode(ChildIndex);
+						}				
+				
+						// Propagate child loaded mask
+						if (bSetLoadedMask && !CellNode.HasChildLoadedNode(ChildIndex))
+						{
+							CellNode.AddChildLoadedNode(ChildIndex);
+						}
+
+						if (bShouldBreak)
+						{
+							break;
+						}
+					}
 				}
 			});
 		}
@@ -244,7 +264,9 @@ void UWorldPartitionEditorSpatialHash::UnhashActor(FWorldPartitionActorDesc* InA
 				verify(Cells.Remove(EditorCell));
 				verify(HashCells.Remove(CellCoord));
 
-				bool bClearParentLoadedChild = true;
+				bool bClearChildMask = true;
+				bool bClearLoadedMask = EditorCell->bLoaded;
+
 				FCellCoord LevelCellCoord = CellCoord;
 				while (LevelCellCoord.Level < CurrentLevel)
 				{
@@ -254,20 +276,34 @@ void UWorldPartitionEditorSpatialHash::UnhashActor(FWorldPartitionActorDesc* InA
 
 					FCellNode& CellNode = HashNodes.FindChecked(LevelCellCoord);
 
-					if (bClearParentLoadedChild && CellNode.HasChildLoadedNode(ChildIndex))
+					if (bClearChildMask)
 					{
-						CellNode.RemoveChildLoadedNode(ChildIndex);
-						bClearParentLoadedChild = !CellNode.HasChildLoadedNodes();
+						CellNode.RemoveChildNode(ChildIndex);
+
+						if (CellNode.HasChildNodes())
+						{
+							bClearChildMask = false;
+						}
+						else
+						{
+							HashNodes.Remove(LevelCellCoord);
+						}
 					}
 
-					CellNode.RemoveChildNode(ChildIndex);
+					if (bClearLoadedMask)
+					{
+						CellNode.RemoveChildLoadedNode(ChildIndex);
 
-					if (CellNode.HasChildNodes())
+						if (CellNode.HasChildLoadedNodes())
+						{
+							bClearLoadedMask = false;
+						}
+					}
+
+					if (!bClearChildMask && !bClearLoadedMask)
 					{
 						break;
 					}
-
-					HashNodes.Remove(LevelCellCoord);
 				}
 
 				bBoundsDirty = true;
