World Partition
- Unified actor descriptors serialization code to avoid code duplication.
- Actors don't need to overload GetAssetRegistryTags anymore to serialize custom metadata.
- WorldPartitionConvertCommandlet: don't rely on the asset registry anymore to delete old conversion results.
- Fixed FPartitionActorDesc not updating its hash properly.
- Removed asset registry specific code from FPackageSourceControlHelper.
- Deleted unused class FActorRegistry.

diff --git a/Engine/Source/Editor/UnrealEd/Private/Commandlets/WorldPartitionConvertCommandlet.cpp b/Engine/Source/Editor/UnrealEd/Private/Commandlets/WorldPartitionConvertCommandlet.cpp
--- a/Engine/Source/Editor/UnrealEd/Private/Commandlets/WorldPartitionConvertCommandlet.cpp
+++ b/Engine/Source/Editor/UnrealEd/Private/Commandlets/WorldPartitionConvertCommandlet.cpp
@@ -16,7 +16,6 @@
 #include "Engine/LODActor.h"
 #include "Engine/LevelStreaming.h"
 #include "Engine/MapBuildDataRegistry.h"
-#include "ActorRegistry.h"
 #include "WorldPartition/WorldPartition.h"
 #include "WorldPartition/WorldPartitionSubsystem.h"
 #include "WorldPartition/HLOD/HLODActor.h"
@@ -709,18 +708,23 @@ int32 UWorldPartitionConvertCommandlet::Main(const FString& Params)
 	{
 		UE_SCOPED_TIMER(TEXT("Deleting existing conversion results"), LogWorldPartitionConvertCommandlet, Display);
 
-		// Update AssetRegistry
-		{
-			UE_SCOPED_TIMER(TEXT("Waiting for asset registry"), LogWorldPartitionConvertCommandlet, Display);
-			IAssetRegistry& AssetRegistry = FModuleManager::LoadModuleChecked<FAssetRegistryModule>(TEXT("AssetRegistry")).Get();
-			AssetRegistry.SearchAllAssets(true);
-		}
-
-		TArray<FAssetData> Assets;
 		FString OldLevelName = Tokens[0] + ConversionSuffix;
-		FActorRegistry::GetLevelActors(*OldLevelName, Assets);
+		FString ExternalActorsPath = ULevel::GetExternalActorsPath(OldLevelName);
+
+		bool bResult = IFileManager::Get().IterateDirectoryRecursively(*FPackageName::LongPackageNameToFilename(ExternalActorsPath), [this](const TCHAR* FilenameOrDirectory, bool bIsDirectory)
+		{
+			if (!bIsDirectory)
+			{
+				FString Filename(FilenameOrDirectory);
+				if (Filename.EndsWith(FPackageName::GetAssetPackageExtension()))
+				{
+					return PackageHelper.Delete(Filename);
+				}
+			}
+			return true;
+		});
 
-		if (!PackageHelper.Delete(Assets))
+		if (!bResult)
 		{
 			UE_LOG(LogWorldPartitionConvertCommandlet, Error, TEXT("Failed to delete external actor package(s)"));
 			return 1;
diff --git a/Engine/Source/Editor/UnrealEd/Private/PackageSourceControlHelper.cpp b/Engine/Source/Editor/UnrealEd/Private/PackageSourceControlHelper.cpp
--- a/Engine/Source/Editor/UnrealEd/Private/PackageSourceControlHelper.cpp
+++ b/Engine/Source/Editor/UnrealEd/Private/PackageSourceControlHelper.cpp
@@ -107,17 +107,6 @@ bool FPackageSourceControlHelper::Delete(UPackage* Package) const
 	return Delete(Packages);
 }
 
-bool FPackageSourceControlHelper::Delete(const TArray<FAssetData>& Assets) const
-{
-	bool bAllPackagesDeleted = true;
-	for (const FAssetData& AssetToDelete : Assets)
-	{
-		bAllPackagesDeleted &= Delete(AssetToDelete.PackageName.ToString());
-	}
-
-	return bAllPackagesDeleted;
-}
-
 bool FPackageSourceControlHelper::Delete(const TArray<UPackage*>& Packages) const
 {
 	if (Packages.IsEmpty())
diff --git a/Engine/Source/Editor/UnrealEd/Public/PackageSourceControlHelper.h b/Engine/Source/Editor/UnrealEd/Public/PackageSourceControlHelper.h
--- a/Engine/Source/Editor/UnrealEd/Public/PackageSourceControlHelper.h
+++ b/Engine/Source/Editor/UnrealEd/Public/PackageSourceControlHelper.h
@@ -5,7 +5,6 @@
 #include "CoreMinimal.h"
 #include "UObject/ObjectMacros.h"
 #include "SourceControlHelpers.h"
-#include "AssetRegistryModule.h"
 
 class UNREALED_API FPackageSourceControlHelper
 {
@@ -14,7 +13,6 @@ public:
 	bool Delete(const FString& PackageName) const;
 	bool Delete(UPackage* Package) const;
 	bool Delete(const TArray<UPackage*>& Packages) const;
-	bool Delete(const TArray<FAssetData>& Assets) const;
 	bool AddToSourceControl(UPackage* Package) const;
 	bool Checkout(UPackage* Package) const;
 
diff --git a/Engine/Source/Runtime/Engine/Private/Actor.cpp b/Engine/Source/Runtime/Engine/Private/Actor.cpp
--- a/Engine/Source/Runtime/Engine/Private/Actor.cpp
+++ b/Engine/Source/Runtime/Engine/Private/Actor.cpp
@@ -46,10 +46,12 @@
 #include "ObjectTrace.h"
 #include "Net/Core/PushModel/PushModel.h"
 #include "Engine/AutoDestroySubsystem.h"
-#include "ActorRegistry.h"
 #include "LevelUtils.h"
 #include "Components/DecalComponent.h"
 
+#include "WorldPartition/WorldPartition.h"
+#include "WorldPartition/WorldPartitionActorDesc.h"
+
 DEFINE_LOG_CATEGORY(LogActor);
 
 DEFINE_STAT(STAT_GetComponentsTime);
@@ -515,69 +517,17 @@ void AActor::GetAssetRegistryTags(TArray<FAssetRegistryTag>& OutTags) const
 
 	if (IsPackageExternal() && !IsChildActor())
 	{
-		check(ActorGuid.IsValid());
-		static const FName NAME_ActorGuid(TEXT("ActorGuid"));
-		FActorRegistry::SaveActorMetaData(NAME_ActorGuid, ActorGuid, OutTags);
-
-		// Get the first native class in the hierarchy, to be able to recreate the actor descriptor without having to load a BP class
-		UClass* NativeClass = GetParentNativeClass(GetClass());
-		static const FName NAME_ActorClass(TEXT("ActorClass"));
-		FActorRegistry::SaveActorMetaData(NAME_ActorClass, NativeClass->GetFName(), OutTags);
-
-		static const FName NAME_LevelPackage(TEXT("LevelPackage"));
-		FString LevelPackageName = GetOuter()->GetOutermost()->GetName();
-		FActorRegistry::SaveActorMetaData(NAME_LevelPackage, LevelPackageName, OutTags);
-
-		FVector BoundsLocation;
-		FVector BoundsExtent;		
-		GetActorLocationBounds(false, BoundsLocation, BoundsExtent);
-
-		static const FName NAME_BoundsLocation(TEXT("BoundsLocation"));
-		FActorRegistry::SaveActorMetaData(NAME_BoundsLocation, BoundsLocation, OutTags);
-
-		static const FName NAME_BoundsExtent(TEXT("BoundsExtent"));
-		FActorRegistry::SaveActorMetaData(NAME_BoundsExtent, BoundsExtent, OutTags);
-
-		static const FName NAME_GridPlacement(TEXT("GridPlacement"));
-		FActorRegistry::SaveActorMetaData(NAME_GridPlacement, (int32)GridPlacement, OutTags);
-
-		static const FName NAME_RuntimeGrid(TEXT("RuntimeGrid"));
-		FActorRegistry::SaveActorMetaData(NAME_RuntimeGrid, RuntimeGrid, OutTags);
-
-		static const FName NAME_IsEditorOnly(TEXT("IsEditorOnly"));
-		FActorRegistry::SaveActorMetaData(NAME_IsEditorOnly, IsEditorOnly(), OutTags);
-
-		static const FName NAME_IsLevelBoundsRelevant(TEXT("IsLevelBoundsRelevant"));
-		FActorRegistry::SaveActorMetaData(NAME_IsLevelBoundsRelevant, IsLevelBoundsRelevant(), OutTags);		
-
-		if (Layers.Num())
+		if (UWorldPartition* WorldPartition = GetWorldSettings()->GetWorldPartition())
 		{
-			FString ActorLayers;
-			for(FName Layer: Layers)
-			{
-				ActorLayers += Layer.ToString() + TEXT(";");
-			}
-			ActorLayers.RemoveFromEnd(TEXT(";"));
+			TUniquePtr<FWorldPartitionActorDesc> ActorDesc(WorldPartition->GetActorDescFactory(this)->Create());
+			ActorDesc->Init(this);
 
-			static const FName NAME_Layers(TEXT("Layers"));
-			FActorRegistry::SaveActorMetaData(NAME_Layers, ActorLayers, OutTags);
-		}
+			FActorMetaDataWriter ActorMetaDataSerializer(OutTags);
 
-		TSet<AActor*> ActorReferences;
-		FArchiveGetActorRefs GetActorRefsAr((AActor*)this, ActorReferences);
-		((AActor*)this)->Serialize(GetActorRefsAr);
-
-		if (ActorReferences.Num())
-		{
-			FString ActorRefs;
-			for(AActor* ActorReference: ActorReferences)
-			{
-				ActorRefs += ActorReference->ActorGuid.ToString() + TEXT(";");
-			}
-			ActorRefs.RemoveFromEnd(TEXT(";"));
+			FString LevelPackageName = GetOuter()->GetOutermost()->GetName();
+			ActorMetaDataSerializer.Serialize(TEXT("LevelPackage"), LevelPackageName);
 
-			static const FName NAME_ActorReferences(TEXT("ActorReferences"));
-			FActorRegistry::SaveActorMetaData(NAME_ActorReferences, ActorRefs, OutTags);
+			ActorDesc->SerializeMetaData(&ActorMetaDataSerializer);
 		}
 	}
 }
diff --git a/Engine/Source/Runtime/Engine/Private/ActorPartition/PartitionActor.cpp b/Engine/Source/Runtime/Engine/Private/ActorPartition/PartitionActor.cpp
--- a/Engine/Source/Runtime/Engine/Private/ActorPartition/PartitionActor.cpp
+++ b/Engine/Source/Runtime/Engine/Private/ActorPartition/PartitionActor.cpp
@@ -3,41 +3,16 @@
 #include "ActorPartition/PartitionActor.h"
 
 #if WITH_EDITOR
-#include "ActorRegistry.h"
 #include "Components/BoxComponent.h"
 #endif
 
 APartitionActor::APartitionActor(const FObjectInitializer& ObjectInitializer)
 	: Super(ObjectInitializer)
 #if WITH_EDITORONLY_DATA
-	, GridSize(0)
+	, GridSize(1)
 #endif
 {
 	USceneComponent* SceneComponent = CreateDefaultSubobject<USceneComponent>(TEXT("RootComponent0"));
 	RootComponent = SceneComponent;
 	RootComponent->Mobility = EComponentMobility::Static;
-}
-
-#if WITH_EDITOR
-void APartitionActor::GetAssetRegistryTags(TArray<FAssetRegistryTag>& OutTags) const
-{
-	Super::GetAssetRegistryTags(OutTags);
-
-	static const FName NAME_GridSize(TEXT("GridSize"));
-	FActorRegistry::SaveActorMetaData(NAME_GridSize, (int32)GridSize, OutTags);
-
-	const FVector ActorLocation = GetActorLocation();
-	const int64 GridIndexX = FMath::FloorToInt(ActorLocation.X / GridSize);
-	const int64 GridIndexY = FMath::FloorToInt(ActorLocation.Y / GridSize);
-	const int64 GridIndexZ = FMath::FloorToInt(ActorLocation.Z / GridSize);
-
-	static const FName NAME_GridIndexX(TEXT("GridIndexX"));
-	FActorRegistry::SaveActorMetaData(NAME_GridIndexX, GridIndexX, OutTags);
-
-	static const FName NAME_GridIndexY(TEXT("GridIndexY"));
-	FActorRegistry::SaveActorMetaData(NAME_GridIndexY, GridIndexY, OutTags);
-
-	static const FName NAME_GridIndexZ(TEXT("GridIndexZ"));
-	FActorRegistry::SaveActorMetaData(NAME_GridIndexZ, GridIndexZ, OutTags);
-}
-#endif
\ No newline at end of file
+}
\ No newline at end of file
diff --git a/Engine/Source/Runtime/Engine/Private/ActorRegistry.cpp b/Engine/Source/Runtime/Engine/Private/ActorRegistry.cpp
deleted file mode 100644
diff --git a/Engine/Source/Runtime/Engine/Private/Level.cpp b/Engine/Source/Runtime/Engine/Private/Level.cpp
--- a/Engine/Source/Runtime/Engine/Private/Level.cpp
+++ b/Engine/Source/Runtime/Engine/Private/Level.cpp
@@ -47,7 +47,6 @@ Level.cpp: Level-related functions
 #include "Engine/LevelBounds.h"
 #include "Async/ParallelFor.h"
 #include "UnrealEngine.h"
-#include "ActorRegistry.h"
 #if WITH_EDITOR
 #include "Kismet2/KismetEditorUtilities.h"
 #include "Kismet2/BlueprintEditorUtils.h"
diff --git a/Engine/Source/Runtime/Engine/Private/LevelInstance/LevelInstanceActor.cpp b/Engine/Source/Runtime/Engine/Private/LevelInstance/LevelInstanceActor.cpp
--- a/Engine/Source/Runtime/Engine/Private/LevelInstance/LevelInstanceActor.cpp
+++ b/Engine/Source/Runtime/Engine/Private/LevelInstance/LevelInstanceActor.cpp
@@ -11,7 +11,6 @@
 #include "LevelInstancePrivate.h"
 
 #if WITH_EDITOR
-#include "ActorRegistry.h"
 #include "Logging/MessageLog.h"
 #include "Logging/TokenizedMessage.h"
 #include "Misc/MapErrors.h"
@@ -267,21 +266,6 @@ FString ALevelInstance::GetWorldAssetPackage() const
 	return WorldAsset.GetUniqueID().GetLongPackageName();
 }
 
-void ALevelInstance::GetAssetRegistryTags(TArray<FAssetRegistryTag>& OutTags) const
-{
-	Super::GetAssetRegistryTags(OutTags);
-
-	if (IsPackageExternal())
-	{
-		static const FName NAME_LevelInstancePackage(TEXT("LevelInstancePackage"));
-		FActorRegistry::SaveActorMetaData(NAME_LevelInstancePackage, GetWorldAssetPackage(), OutTags);
-
-		static const FName NAME_LevelInstanceTransform(TEXT("LevelInstanceTransform"));
-		FTransform LevelInstanceTransform(GetActorRotation(), GetActorLocation());
-		FActorRegistry::SaveActorMetaData(NAME_LevelInstanceTransform, LevelInstanceTransform, OutTags);
-	}
-}
-
 void ALevelInstance::PreEditChange(FProperty* PropertyThatWillChange)
 {
 	{
diff --git a/Engine/Source/Runtime/Engine/Private/LevelInstance/LevelInstanceSubsystem.cpp b/Engine/Source/Runtime/Engine/Private/LevelInstance/LevelInstanceSubsystem.cpp
--- a/Engine/Source/Runtime/Engine/Private/LevelInstance/LevelInstanceSubsystem.cpp
+++ b/Engine/Source/Runtime/Engine/Private/LevelInstance/LevelInstanceSubsystem.cpp
@@ -19,7 +19,6 @@
 #include "WorldPartition/LevelInstance/LevelInstanceActorDescFactory.h"
 #include "Misc/ScopedSlowTask.h"
 #include "AssetRegistryModule.h"
-#include "ActorRegistry.h"
 #include "AssetData.h"
 #include "FileHelpers.h"
 #include "Editor.h"
diff --git a/Engine/Source/Runtime/Engine/Private/WorldPartition/ActorPartition/PartitionActorDesc.cpp b/Engine/Source/Runtime/Engine/Private/WorldPartition/ActorPartition/PartitionActorDesc.cpp
--- a/Engine/Source/Runtime/Engine/Private/WorldPartition/ActorPartition/PartitionActorDesc.cpp
+++ b/Engine/Source/Runtime/Engine/Private/WorldPartition/ActorPartition/PartitionActorDesc.cpp
@@ -5,24 +5,40 @@
 
 #if WITH_EDITOR
 #include "Engine/Level.h"
-FPartitionActorDesc::FPartitionActorDesc(const FWorldPartitionActorDescData& DescData, uint32 InGridSize, int64 InGridIndexX, int64 InGridIndexY, int64 InGridIndexZ)
-	: FWorldPartitionActorDesc(DescData)
-	, GridSize(InGridSize)
-	, GridIndexX(InGridIndexX)
-	, GridIndexY(InGridIndexY)
-	, GridIndexZ(InGridIndexZ)
-{}
-
-FPartitionActorDesc::FPartitionActorDesc(AActor* InActor)
-	: FWorldPartitionActorDesc(InActor)
+
+bool FPartitionActorDesc::Init(const AActor* InActor)
+{
+	if (FWorldPartitionActorDesc::Init(InActor))
+	{
+		const APartitionActor* PartitionActor = CastChecked<APartitionActor>(InActor);
+
+		GridSize = PartitionActor->GridSize;
+
+		const FVector ActorLocation = InActor->GetActorLocation();
+		GridIndexX = FMath::FloorToInt(ActorLocation.X / GridSize);
+		GridIndexY = FMath::FloorToInt(ActorLocation.Y / GridSize);
+		GridIndexZ = FMath::FloorToInt(ActorLocation.Z / GridSize);
+
+		UpdateHash();
+		return true;
+	}
+
+	return false;
+}
+
+void FPartitionActorDesc::SerializeMetaData(FActorMetaDataSerializer* Serializer)
 {
-	APartitionActor* PartitionActor = CastChecked<APartitionActor>(InActor);
+	FWorldPartitionActorDesc::SerializeMetaData(Serializer);
 
-	GridSize = PartitionActor->GridSize;
+	Serializer->Serialize(TEXT("GridSize"), (int32&)GridSize);
+	Serializer->Serialize(TEXT("GridIndexX"), GridIndexX);
+	Serializer->Serialize(TEXT("GridIndexY"), GridIndexY);
+	Serializer->Serialize(TEXT("GridIndexZ"), GridIndexZ);
+}
 
-	const FVector ActorLocation = InActor->GetActorLocation();
-	GridIndexX = FMath::FloorToInt(ActorLocation.X / GridSize);
-	GridIndexY = FMath::FloorToInt(ActorLocation.Y / GridSize);
-	GridIndexZ = FMath::FloorToInt(ActorLocation.Z / GridSize);
+void FPartitionActorDesc::BuildHash(FHashBuilder& HashBuilder)
+{
+	FWorldPartitionActorDesc::BuildHash(HashBuilder);
+	HashBuilder << GridSize << GridIndexX << GridIndexY << GridIndexZ;
 }
 #endif
diff --git a/Engine/Source/Runtime/Engine/Private/WorldPartition/ActorPartition/PartitionActorDescFactory.cpp b/Engine/Source/Runtime/Engine/Private/WorldPartition/ActorPartition/PartitionActorDescFactory.cpp
--- a/Engine/Source/Runtime/Engine/Private/WorldPartition/ActorPartition/PartitionActorDescFactory.cpp
+++ b/Engine/Source/Runtime/Engine/Private/WorldPartition/ActorPartition/PartitionActorDescFactory.cpp
@@ -1,51 +1,12 @@
 // Copyright Epic Games, Inc. All Rights Reserved.
 
 #include "WorldPartition/ActorPartition/PartitionActorDescFactory.h"
-#include "WorldPartition/ActorPartition/PartitionActorDesc.h"
-#include "ActorRegistry.h"
 
 #if WITH_EDITOR
-FWorldPartitionActorDesc* FPartitionActorDescFactory::CreateInstance(const FWorldPartitionActorDescInitData& ActorDescInitData)
-{
-	FWorldPartitionActorDescData Data;
-	if (!ReadMetaData(ActorDescInitData, Data))
-	{
-		return nullptr;
-	}
-
-	int32 GridSize;
-	static const FName NAME_GridSize(TEXT("GridSize"));
-	if (!FActorRegistry::ReadActorMetaData(NAME_GridSize, GridSize, ActorDescInitData.AssetData))
-	{
-		return nullptr;
-	}
-
-	int64 GridIndexX;
-	static const FName NAME_GridIndexX(TEXT("GridIndexX"));
-	if (!FActorRegistry::ReadActorMetaData(NAME_GridIndexX, GridIndexX, ActorDescInitData.AssetData))
-	{
-		return nullptr;
-	}
-
-	int64 GridIndexY;
-	static const FName NAME_GridIndexY(TEXT("GridIndexY"));
-	if (!FActorRegistry::ReadActorMetaData(NAME_GridIndexY, GridIndexY, ActorDescInitData.AssetData))
-	{
-		return nullptr;
-	}
-
-	int64 GridIndexZ;
-	static const FName NAME_GridIndexZ(TEXT("GridIndexZ"));
-	if (!FActorRegistry::ReadActorMetaData(NAME_GridIndexZ, GridIndexZ, ActorDescInitData.AssetData))
-	{
-		return nullptr;
-	}
-
-	return new FPartitionActorDesc(Data, (uint32)GridSize, GridIndexX, GridIndexY, GridIndexZ);
-}
+#include "WorldPartition/ActorPartition/PartitionActorDesc.h"
 
-FWorldPartitionActorDesc* FPartitionActorDescFactory::CreateInstance(AActor* InActor)
+FWorldPartitionActorDesc* FPartitionActorDescFactory::Create()
 {
-	return new FPartitionActorDesc(InActor);
+	return new FPartitionActorDesc;
 }
 #endif
\ No newline at end of file
diff --git a/Engine/Source/Runtime/Engine/Private/WorldPartition/HLOD/HLODActor.cpp b/Engine/Source/Runtime/Engine/Private/WorldPartition/HLOD/HLODActor.cpp
--- a/Engine/Source/Runtime/Engine/Private/WorldPartition/HLOD/HLODActor.cpp
+++ b/Engine/Source/Runtime/Engine/Private/WorldPartition/HLOD/HLODActor.cpp
@@ -5,8 +5,6 @@
 #include "Components/PrimitiveComponent.h"
 
 #if WITH_EDITOR
-#include "ActorRegistry.h"
-#include "AssetData.h"
 #include "WorldPartition/WorldPartition.h"
 #include "WorldPartition/WorldPartitionActorDesc.h"
 #include "WorldPartition/HLOD/HLODLayer.h"
@@ -169,30 +167,6 @@ void AWorldPartitionHLOD::SetHLODLayer(const UHLODLayer* InSubActorsHLODLayer, i
 	SubActorsHLODLevel = InSubActorsHLODLevel;
 }
 
-void AWorldPartitionHLOD::GetAssetRegistryTags(TArray<FAssetRegistryTag>& OutTags) const
-{
-	Super::GetAssetRegistryTags(OutTags);
-
-	if (IsPackageExternal())
-	{
-		if (SubActors.Num())
-		{
-			FString SubActorsGUIDsStr;
-			for (FGuid ActorGUID : SubActors)
-			{
-				SubActorsGUIDsStr += ActorGUID.ToString() + TEXT(";");
-			}
-			SubActorsGUIDsStr.RemoveFromEnd(TEXT(";"));
-
-			static const FName NAME_HLODSubActors(TEXT("HLODSubActors"));
-			FActorRegistry::SaveActorMetaData(NAME_HLODSubActors, SubActorsGUIDsStr, OutTags);
-		}
-
-		static const FName NAME_HLODSubActors(TEXT("HLODLayer"));
-		FActorRegistry::SaveActorMetaData(NAME_HLODSubActors, FSoftObjectPath(SubActorsHLODLayer).ToString(), OutTags);
-	}
-}
-
 void AWorldPartitionHLOD::PostActorCreated()
 {
 	Super::PostActorCreated();
diff --git a/Engine/Source/Runtime/Engine/Private/WorldPartition/HLOD/HLODActorDesc.cpp b/Engine/Source/Runtime/Engine/Private/WorldPartition/HLOD/HLODActorDesc.cpp
--- a/Engine/Source/Runtime/Engine/Private/WorldPartition/HLOD/HLODActorDesc.cpp
+++ b/Engine/Source/Runtime/Engine/Private/WorldPartition/HLOD/HLODActorDesc.cpp
@@ -3,23 +3,78 @@
 #include "WorldPartition/HLOD/HLODActorDesc.h"
 
 #if WITH_EDITOR
+#include "Algo/Transform.h"
+#include "Algo/RemoveIf.h"
 #include "WorldPartition/HLOD/HLODActor.h"
 #include "WorldPartition/HLOD/HLODLayer.h"
 
-FHLODActorDesc::FHLODActorDesc(const FWorldPartitionActorDescData& DescData, const TArray<FGuid>& InSubActors, const FSoftObjectPath& InHLODLayer)
-	: FWorldPartitionActorDesc(DescData)
-	, SubActors(InSubActors)
-	, HLODLayer(InHLODLayer)
+bool FHLODActorDesc::Init(const AActor* InActor)
 {
+	if (FWorldPartitionActorDesc::Init(InActor))
+	{
+		const AWorldPartitionHLOD* HLODActor = CastChecked<AWorldPartitionHLOD>(InActor);
+		SubActors = HLODActor->GetSubActors();
+		HLODLayer = HLODActor->GetHLODLayer();
+
+		UpdateHash();
+		return true;
+	}
+
+	return false;
 }
 
-FHLODActorDesc::FHLODActorDesc(AActor* InActor)
-	: FWorldPartitionActorDesc(InActor)
+void FHLODActorDesc::SerializeMetaData(FActorMetaDataSerializer* Serializer)
 {
-	if (AWorldPartitionHLOD* HLODActor = CastChecked<AWorldPartitionHLOD>(InActor))
+	FWorldPartitionActorDesc::SerializeMetaData(Serializer);
+
+	FString SubActorsGUIDsStr;
+	if (Serializer->IsWriting())
 	{
-		SubActors = HLODActor->GetSubActors();
-		HLODLayer = HLODActor->GetHLODLayer();
+		if (SubActors.Num())
+		{
+			for (const FGuid& ActorGUID : SubActors)
+			{
+				SubActorsGUIDsStr += ActorGUID.ToString() + TEXT(";");
+			}
+			SubActorsGUIDsStr.RemoveFromEnd(TEXT(";"));
+		}
+	}
+
+	Serializer->Serialize(TEXT("HLODSubActors"), SubActorsGUIDsStr);
+
+	if (Serializer->IsReading())
+	{
+		TArray<FString> SubActorsStr;
+		if (SubActorsGUIDsStr.ParseIntoArray(SubActorsStr, TEXT(";")))
+		{
+			Algo::Transform(SubActorsStr, SubActors, [](const FString& ActorGuidStr)
+			{
+				FGuid ActorGuid;
+				if (!FGuid::Parse(ActorGuidStr, ActorGuid))
+				{
+					ActorGuid.Invalidate();
+				}
+				return ActorGuid;
+			});
+
+			Algo::RemoveIf(SubActors, [](const FGuid& InGuid)
+			{
+				return !InGuid.IsValid();
+			});
+		}
+	}
+
+	FString HLODLayerStr;
+	if (Serializer->IsWriting())
+	{
+		HLODLayerStr = HLODLayer.ToString();
+	}
+	
+	Serializer->Serialize(TEXT("HLODSubActors"), HLODLayerStr);
+	
+	if (Serializer->IsReading())
+	{
+		HLODLayer = HLODLayerStr;
 	}
 }
 
diff --git a/Engine/Source/Runtime/Engine/Private/WorldPartition/HLOD/HLODActorDescFactory.cpp b/Engine/Source/Runtime/Engine/Private/WorldPartition/HLOD/HLODActorDescFactory.cpp
--- a/Engine/Source/Runtime/Engine/Private/WorldPartition/HLOD/HLODActorDescFactory.cpp
+++ b/Engine/Source/Runtime/Engine/Private/WorldPartition/HLOD/HLODActorDescFactory.cpp
@@ -4,62 +4,9 @@
 
 #if WITH_EDITOR
 #include "WorldPartition/HLOD/HLODActorDesc.h"
-#include "WorldPartition/HLOD/HLODActor.h"
-#include "WorldPartition/WorldPartition.h"
-#include "ActorRegistry.h"
-#include "Editor.h"
-#include "Algo/Transform.h"
-#include "Algo/RemoveIf.h"
 
-FWorldPartitionActorDesc* FHLODActorDescFactory::CreateInstance(const FWorldPartitionActorDescInitData& ActorDescInitData)
+FWorldPartitionActorDesc* FHLODActorDescFactory::Create()
 {
-	FWorldPartitionActorDescData Data;
-	if (!ReadMetaData(ActorDescInitData, Data))
-	{
-		return nullptr;
-	}
-
-	FString MetadataStr;
-
-	TArray<FGuid> SubActorsGUIDs;
-	static const FName NAME_HLODSubActors(TEXT("HLODSubActors"));
-	if (!FActorRegistry::ReadActorMetaData(NAME_HLODSubActors, MetadataStr, ActorDescInitData.AssetData))
-	{
-		return nullptr;
-	}
-
-	TArray<FString> SubActorsStr;
-	if (MetadataStr.ParseIntoArray(SubActorsStr, TEXT(";")))
-	{
-		Algo::Transform(SubActorsStr, SubActorsGUIDs, [](const FString& ActorGuidStr)
-		{
-			FGuid ActorGuid;
-			if (!FGuid::Parse(ActorGuidStr, ActorGuid))
-			{
-				ActorGuid.Invalidate();
-			}
-			return ActorGuid;
-		});
-
-		Algo::RemoveIf(SubActorsGUIDs, [](const FGuid& Guid)
-		{
-			return !Guid.IsValid();
-		});
-	}
-
-	FSoftObjectPath HLODLayerPath;
-	static const FName NAME_HLODLayer(TEXT("HLODLayer"));
-	if (!FActorRegistry::ReadActorMetaData(NAME_HLODLayer, MetadataStr, ActorDescInitData.AssetData))
-	{
-		return nullptr;
-	}
-	HLODLayerPath = MetadataStr;
-
-	return new FHLODActorDesc(Data, SubActorsGUIDs, HLODLayerPath);
-}
-
-FWorldPartitionActorDesc* FHLODActorDescFactory::CreateInstance(AActor* InActor)
-{
-	return new FHLODActorDesc(InActor);
+	return new FHLODActorDesc;
 }
 #endif
\ No newline at end of file
diff --git a/Engine/Source/Runtime/Engine/Private/WorldPartition/LevelInstance/LevelInstanceActorDesc.cpp b/Engine/Source/Runtime/Engine/Private/WorldPartition/LevelInstance/LevelInstanceActorDesc.cpp
--- a/Engine/Source/Runtime/Engine/Private/WorldPartition/LevelInstance/LevelInstanceActorDesc.cpp
+++ b/Engine/Source/Runtime/Engine/Private/WorldPartition/LevelInstance/LevelInstanceActorDesc.cpp
@@ -6,24 +6,49 @@
 #include "Engine/Level.h"
 #include "LevelInstance/LevelInstanceActor.h"
 
-FLevelInstanceActorDesc::FLevelInstanceActorDesc(const FWorldPartitionActorDescData& DescData, FName InLevelPackage)
-	: FWorldPartitionActorDesc(DescData)
-	, LevelPackage(InLevelPackage)
-{}
-
-FLevelInstanceActorDesc::FLevelInstanceActorDesc(AActor* InActor)
-	: FWorldPartitionActorDesc(InActor)
+bool FLevelInstanceActorDesc::Init(const AActor* InActor)
 {
-	if (ALevelInstance* LevelInstanceActor = CastChecked<ALevelInstance>(InActor))
+	if (FWorldPartitionActorDesc::Init(InActor))
 	{
+		const ALevelInstance* LevelInstanceActor = CastChecked<ALevelInstance>(InActor);
 		LevelPackage = *LevelInstanceActor->GetWorldAssetPackage();
+		LevelInstanceTransform = FTransform(InActor->GetActorRotation(), InActor->GetActorLocation());
+		UpdateHash();
+		return true;
+	}
+	return false;
+}
+
+void FLevelInstanceActorDesc::SerializeMetaData(FActorMetaDataSerializer* Serializer)
+{
+	FWorldPartitionActorDesc::SerializeMetaData(Serializer);
+
+	Serializer->Serialize(TEXT("LevelInstancePackage"), LevelPackage);
+	Serializer->Serialize(TEXT("LevelInstanceTransform"), LevelInstanceTransform);
+
+	if (Serializer->IsReading())
+	{
+		if (!LevelPackage.IsNone())
+		{
+			FBox LevelBounds;
+			if (ULevel::GetLevelBoundsFromPackage(LevelPackage, LevelBounds))
+			{
+				FVector LevelBoundsLocation;
+				LevelBounds.GetCenterAndExtents(BoundsLocation, BoundsExtent);
+
+				//@todo_ow: This will result in a new BoundsExtent that is larger than it should. To fix this, we would need the Object Oriented BoundingBox of the actor (the BV of the actor without rotation)
+				const FVector BoundsMin = BoundsLocation - BoundsExtent;
+				const FVector BoundsMax = BoundsLocation + BoundsExtent;
+				const FBox NewBounds = FBox(BoundsMin, BoundsMax).TransformBy(LevelInstanceTransform);
+				NewBounds.GetCenterAndExtents(BoundsLocation, BoundsExtent);
+			}
+		}
 	}
 }
 
 void FLevelInstanceActorDesc::BuildHash(FHashBuilder& HashBuilder)
 {
 	FWorldPartitionActorDesc::BuildHash(HashBuilder);
-	HashBuilder << LevelPackage;
+	HashBuilder << LevelPackage << LevelInstanceTransform;
 }
-
 #endif
diff --git a/Engine/Source/Runtime/Engine/Private/WorldPartition/LevelInstance/LevelInstanceActorDescFactory.cpp b/Engine/Source/Runtime/Engine/Private/WorldPartition/LevelInstance/LevelInstanceActorDescFactory.cpp
--- a/Engine/Source/Runtime/Engine/Private/WorldPartition/LevelInstance/LevelInstanceActorDescFactory.cpp
+++ b/Engine/Source/Runtime/Engine/Private/WorldPartition/LevelInstance/LevelInstanceActorDescFactory.cpp
@@ -4,56 +4,9 @@
 
 #if WITH_EDITOR
 #include "WorldPartition/LevelInstance/LevelInstanceActorDesc.h"
-#include "LevelInstance/LevelInstanceActor.h"
-#include "WorldPartition/WorldPartition.h"
-#include "ActorRegistry.h"
-#include "Editor.h"
 
-FWorldPartitionActorDesc* FLevelInstanceActorDescFactory::CreateInstance(const FWorldPartitionActorDescInitData& ActorDescInitData)
+FWorldPartitionActorDesc* FLevelInstanceActorDescFactory::Create()
 {
-	FWorldPartitionActorDescData Data;
-	if (!ReadMetaData(ActorDescInitData, Data))
-	{
-		return nullptr;
-	}
-
-	FString LevelPackageNameStr;
-	static const FName NAME_LevelInstancePackage(TEXT("LevelInstancePackage"));
-	if (!FActorRegistry::ReadActorMetaData(NAME_LevelInstancePackage, LevelPackageNameStr, ActorDescInitData.AssetData))
-	{
-		return nullptr;
-	}
-
-	FTransform LevelInstanceTransform;
-	static const FName NAME_LevelInstanceTransform(TEXT("LevelInstanceTransform"));
-	if (!FActorRegistry::ReadActorMetaData(NAME_LevelInstanceTransform, LevelInstanceTransform, ActorDescInitData.AssetData))
-	{
-		return nullptr;
-	}
-	
-	FName LevelPackage = FName(*LevelPackageNameStr);
-
-	if (!LevelPackage.IsNone())
-	{
-		FBox LevelBounds;
-		if (ULevel::GetLevelBoundsFromPackage(LevelPackage, LevelBounds))
-		{
-			FVector LevelBoundsLocation;
-			LevelBounds.GetCenterAndExtents(Data.BoundsLocation, Data.BoundsExtent);
-
-			//@todo_ow: This will result in a new BoundsExtent that is larger than it should. To fix this, we would need the Object Oriented BoundingBox of the actor (the BV of the actor without rotation)
-			const FVector BoundsMin = Data.BoundsLocation - Data.BoundsExtent;
-			const FVector BoundsMax = Data.BoundsLocation + Data.BoundsExtent;
-			const FBox NewBounds = FBox(BoundsMin, BoundsMax).TransformBy(LevelInstanceTransform);
-			NewBounds.GetCenterAndExtents(Data.BoundsLocation, Data.BoundsExtent);
-		}
-	}
-
-	return new FLevelInstanceActorDesc(Data, LevelPackage);
-}
-
-FWorldPartitionActorDesc* FLevelInstanceActorDescFactory::CreateInstance(AActor* InActor)
-{
-	return new FLevelInstanceActorDesc(InActor);
+	return new FLevelInstanceActorDesc;
 }
 #endif
\ No newline at end of file
diff --git a/Engine/Source/Runtime/Engine/Private/WorldPartition/WorldPartition.cpp b/Engine/Source/Runtime/Engine/Private/WorldPartition/WorldPartition.cpp
--- a/Engine/Source/Runtime/Engine/Private/WorldPartition/WorldPartition.cpp
+++ b/Engine/Source/Runtime/Engine/Private/WorldPartition/WorldPartition.cpp
@@ -24,13 +24,13 @@
 #include "Layers/Layer.h"
 #include "Layers/LayersSubsystem.h"
 #include "Editor/GroupActor.h"
-#include "ActorRegistry.h"
 #include "EditorLevelUtils.h"
 #include "FileHelpers.h"
 #include "Misc/ScopedSlowTask.h"
 #include "Misc/ScopeExit.h"
 #include "ScopedTransaction.h"
 #include "UnrealEdMisc.h"
+#include "AssetRegistryModule.h"
 #include "WorldPartition/WorldPartitionActorDescFactory.h"
 #include "WorldPartition/WorldPartitionLevelStreamingDynamic.h"
 #include "WorldPartition/WorldPartitionEditorHash.h"
@@ -244,9 +244,28 @@ void UWorldPartition::Initialize(UWorld* InWorld, const FTransform& InTransform)
 	{
 		TSet<FName> AllLayersNames;
 
+		auto GetLevelActors = [](const FName& LevelPath, TArray<FAssetData>& OutAssets)
+		{
+			if (!LevelPath.IsNone())
+			{
+				FString LevelPathStr = LevelPath.ToString();
+				IAssetRegistry& AssetRegistry = FModuleManager::LoadModuleChecked<FAssetRegistryModule>(TEXT("AssetRegistry")).Get();
+	
+				// Do a synchronous scan of the level external actors path.
+				const bool bForceRescan = true;
+				AssetRegistry.ScanPathsSynchronous({ULevel::GetExternalActorsPath(LevelPath.ToString())}, bForceRescan);
+
+				static const FName NAME_LevelPackage(TEXT("LevelPackage"));
+				FARFilter Filter;
+				Filter.TagsAndValues.Add(NAME_LevelPackage, LevelPathStr);
+				Filter.bIncludeOnlyOnDiskAssets = true;
+				AssetRegistry.GetAssets(Filter, OutAssets);
+			}
+		};
+
 		TArray<FAssetData> Assets;
 		UPackage* LevelPackage = OuterWorld->PersistentLevel->GetOutermost();
-		FActorRegistry::GetLevelActors(LevelPackage->FileName, Assets);
+		GetLevelActors(LevelPackage->FileName, Assets);
 
 		bool bIsInstanced = (bEditorOnly && !IsRunningCommandlet()) ? OuterWorld->PersistentLevel->IsInstancedLevel() : false;
 
@@ -266,9 +285,10 @@ void UWorldPartition::Initialize(UWorld* InWorld, const FTransform& InTransform)
 
 		for (const FAssetData& Asset : Assets)
 		{
+			FActorMetaDataReader ActorMetaDataSerializer(Asset);
+
 			FName ActorClass;
-			static const FName NAME_ActorClass(TEXT("ActorClass"));
-			FActorRegistry::ReadActorMetaData(NAME_ActorClass, ActorClass, Asset);
+			ActorMetaDataSerializer.Serialize(TEXT("ActorClass"), ActorClass);
 
 			FWorldPartitionActorDescInitData ActorDescInitData;
 			ActorDescInitData.NativeClass = FindObjectChecked<UClass>(ANY_PACKAGE, *ActorClass.ToString(), true);
@@ -276,6 +296,7 @@ void UWorldPartition::Initialize(UWorld* InWorld, const FTransform& InTransform)
 			ActorDescInitData.AssetData = Asset;
 			ActorDescInitData.ActorPath = Asset.ObjectPath;
 			ActorDescInitData.Transform = bIsInstanced ? InstanceTransform : FTransform::Identity;
+			ActorDescInitData.Serializer = &ActorMetaDataSerializer;
 
 			if (bIsInstanced)
 			{
@@ -288,13 +309,16 @@ void UWorldPartition::Initialize(UWorld* InWorld, const FTransform& InTransform)
 				ActorDescInitData.ActorPath = *ActorDescInitData.ActorPath.ToString().Replace(*ReplaceFrom, *ReplaceTo);
 			}
 
-			if (FWorldPartitionActorDesc* NewActorDesc = GetActorDescFactory(ActorDescInitData.NativeClass)->Create(ActorDescInitData))
+			TUniquePtr<FWorldPartitionActorDesc> NewActorDesc(GetActorDescFactory(ActorDescInitData.NativeClass)->Create());
+
+			if (NewActorDesc->Init(ActorDescInitData))
 			{
 				if (bEditorOnly)
 				{
 					AllLayersNames.Append(NewActorDesc->GetLayers());
 				}
-				Actors.Add(NewActorDesc->GetGuid(), TUniquePtr<FWorldPartitionActorDesc>(NewActorDesc));
+
+				Actors.Add(NewActorDesc->GetGuid(), MoveTemp(NewActorDesc));
 			}
 		}
 
@@ -500,7 +524,8 @@ void UWorldPartition::UpdateActorDesc(AActor* InActor)
 		}
 		else if (ActorDescPtr && InActor->GetLevel() == World->PersistentLevel)
 		{
-			TUniquePtr<FWorldPartitionActorDesc> NewDesc(GetActorDescFactory(InActor)->Create(InActor));
+			TUniquePtr<FWorldPartitionActorDesc> NewDesc(GetActorDescFactory(InActor)->Create());
+			NewDesc->Init(InActor);
 			
 			if (NewDesc->GetHash() != ActorDescPtr->Get()->GetHash())
 			{
@@ -1018,6 +1043,9 @@ void UWorldPartition::OnActorAdded(AActor* InActor)
 	{
 		TUniquePtr<FWorldPartitionActorDesc>* ActorDescPtr = Actors.Find(InActor->GetActorGuid());
 
+		TUniquePtr<FWorldPartitionActorDesc> NewDesc(GetActorDescFactory(InActor)->Create());
+		NewDesc->Init(InActor);
+
 		if (ActorDescPtr)
 		{
 			// The only case where this is valid is when reinstancing BP actors. In this case, we'll receive the call
@@ -1027,16 +1055,13 @@ void UWorldPartition::OnActorAdded(AActor* InActor)
 
 			RemoveFromPartition(ActorDescPtr->Get(), false);
 
-			TUniquePtr<FWorldPartitionActorDesc> NewDesc(GetActorDescFactory(InActor)->Create(InActor));
 			*ActorDescPtr = MoveTemp(NewDesc);
 		}
 		else
 		{
 			checkSlow(InActor->GetLevel()->Actors.Find(InActor) != INDEX_NONE);
 
-			FWorldPartitionActorDesc* NewDesc(GetActorDescFactory(InActor)->Create(InActor));
-
-			ActorDescPtr = &Actors.Add(InActor->GetActorGuid(), TUniquePtr<FWorldPartitionActorDesc>(NewDesc));
+			ActorDescPtr = &Actors.Add(InActor->GetActorGuid(), MoveTemp(NewDesc));
 		}
 
 		AddToPartition(ActorDescPtr->Get());
diff --git a/Engine/Source/Runtime/Engine/Private/WorldPartition/WorldPartitionActorDesc.cpp b/Engine/Source/Runtime/Engine/Private/WorldPartition/WorldPartitionActorDesc.cpp
--- a/Engine/Source/Runtime/Engine/Private/WorldPartition/WorldPartitionActorDesc.cpp
+++ b/Engine/Source/Runtime/Engine/Private/WorldPartition/WorldPartitionActorDesc.cpp
@@ -6,6 +6,8 @@
 #include "Misc/HashBuilder.h"
 #include "UObject/LinkerInstancingContext.h"
 #include "UObject/UObjectHash.h"
+#include "Algo/Transform.h"
+#include "Algo/RemoveIf.h"
 #include "Engine/World.h"
 #include "Engine/Level.h"
 #include "Editor/GroupActor.h"
@@ -15,9 +17,12 @@
 #if WITH_EDITOR
 uint32 FWorldPartitionActorDesc::GlobalTag = 0;
 
-FWorldPartitionActorDesc::FWorldPartitionActorDesc(AActor* InActor)
+FWorldPartitionActorDesc::FWorldPartitionActorDesc()
 	: LoadedRefCount(0)
 	, Tag(0)
+{}
+
+bool FWorldPartitionActorDesc::Init(const AActor* InActor)
 {	
 	check(InActor->IsPackageExternal());
 
@@ -53,8 +58,8 @@ FWorldPartitionActorDesc::FWorldPartitionActorDesc(AActor* InActor)
 	ActorPath = *InActor->GetPathName();
 	
 	TSet<AActor*> ActorReferences;
-	FArchiveGetActorRefs GetActorRefsAr(InActor, ActorReferences);
-	InActor->Serialize(GetActorRefsAr);
+	FArchiveGetActorRefs GetActorRefsAr((AActor*)InActor, ActorReferences);
+	((AActor*)InActor)->Serialize(GetActorRefsAr);
 
 	if (ActorReferences.Num())
 	{
@@ -64,26 +69,45 @@ FWorldPartitionActorDesc::FWorldPartitionActorDesc(AActor* InActor)
 			References.Add(ActorReference->GetActorGuid());
 		}
 	}
+
+	UpdateHash();
+	return true;
 }
 
-FWorldPartitionActorDesc::FWorldPartitionActorDesc(const FWorldPartitionActorDescData& DescData)
-	: LoadedRefCount(0)
-	, Tag(0)
+bool FWorldPartitionActorDesc::Init(const FWorldPartitionActorDescInitData& DescData)
 {
-	Guid = DescData.Guid;
-	Class = DescData.Class;
-	BoundsLocation = DescData.BoundsLocation;
-	BoundsExtent = DescData.BoundsExtent;
-	GridPlacement = DescData.GridPlacement;
-	RuntimeGrid = DescData.RuntimeGrid;
-	bActorIsEditorOnly = DescData.bActorIsEditorOnly;
-	bLevelBoundsRelevant = DescData.bLevelBoundsRelevant;
-	Layers = DescData.Layers;
-	References = DescData.References;
-	ActorPackage = DescData.ActorPackage;
+	ActorPackage = DescData.PackageName;
 	ActorPath = DescData.ActorPath;
+	ActorClass = DescData.NativeClass;
+	Class = DescData.NativeClass->GetFName();
+
+	// Serialize actor metadata
+	SerializeMetaData(DescData.Serializer);
+
+	if (!DescData.Serializer->GetHasErrors())
+	{	
+		// Override grid placement by default class value
+		const EActorGridPlacement DefaultGridPlacement = ActorClass->GetDefaultObject<AActor>()->GetDefaultGridPlacement();
+		if (DefaultGridPlacement != EActorGridPlacement::None)
+		{
+			GridPlacement = DefaultGridPlacement;
+		}
+
+		// Transform BoundsLocation and BoundsExtent if necessary
+		if (!DescData.Transform.Equals(FTransform::Identity))
+		{
+			//@todo_ow: This will result in a new BoundsExtent that is larger than it should. To fix this, we would need the Object Oriented BoundingBox of the actor (the BV of the actor without rotation)
+			const FVector BoundsMin = BoundsLocation - BoundsExtent;
+			const FVector BoundsMax = BoundsLocation + BoundsExtent;
+			const FBox NewBounds = FBox(BoundsMin, BoundsMax).TransformBy(DescData.Transform);
+			NewBounds.GetCenterAndExtents(BoundsLocation, BoundsExtent);
+		}
 
-	ActorClass = FindObjectChecked<UClass>(ANY_PACKAGE, *Class.ToString(), true);
+		UpdateHash();
+		return true;
+	}
+
+	return false;
 }
 
 FString FWorldPartitionActorDesc::ToString() const
@@ -100,7 +124,81 @@ void FWorldPartitionActorDesc::UpdateHash()
 
 void FWorldPartitionActorDesc::BuildHash(FHashBuilder& HashBuilder)
 {
-	HashBuilder << Guid << Class << BoundsLocation << BoundsExtent << GridPlacement << RuntimeGrid << bActorIsEditorOnly << bLevelBoundsRelevant << Layers << References << ActorPackage << ActorPath;
+	HashBuilder << Guid << Class << ActorPackage << ActorPath << BoundsLocation << BoundsExtent << GridPlacement << RuntimeGrid << bActorIsEditorOnly << bLevelBoundsRelevant << Layers << References;
+}
+
+void FWorldPartitionActorDesc::SerializeMetaData(FActorMetaDataSerializer* Serializer)
+{
+	bool bValidSerialization = true;
+
+	bValidSerialization &= Serializer->Serialize(TEXT("ActorClass"), Class);
+	bValidSerialization &= Serializer->Serialize(TEXT("ActorGuid"), Guid);
+	bValidSerialization &= Serializer->Serialize(TEXT("BoundsLocation"), BoundsLocation);
+	bValidSerialization &= Serializer->Serialize(TEXT("BoundsExtent"), BoundsExtent);
+	bValidSerialization &= Serializer->Serialize(TEXT("GridPlacement"), (int8&)GridPlacement);
+	bValidSerialization &= Serializer->Serialize(TEXT("RuntimeGrid"), RuntimeGrid);
+	bValidSerialization &= Serializer->Serialize(TEXT("IsEditorOnly"), bActorIsEditorOnly);
+	bValidSerialization &= Serializer->Serialize(TEXT("IsLevelBoundsRelevant"), bLevelBoundsRelevant);
+
+	if (!bValidSerialization)
+	{
+		Serializer->SetHasErrors();
+	}
+
+	FString LayersStr;
+	FString ActorsRefsStr;
+
+	if (Serializer->IsWriting())
+	{
+		if (Layers.Num())
+		{
+			for(FName Layer: Layers)
+			{
+				LayersStr += Layer.ToString() + TEXT(";");
+			}
+			LayersStr.RemoveFromEnd(TEXT(";"));
+		}
+
+		if (References.Num())
+		{
+			for(const FGuid& ActoGuid: References)
+			{
+				ActorsRefsStr += ActoGuid.ToString() + TEXT(";");
+			}
+			ActorsRefsStr.RemoveFromEnd(TEXT(";"));
+		}
+	}
+
+	Serializer->Serialize(TEXT("Layers"), LayersStr);
+	Serializer->Serialize(TEXT("ActorReferences"), ActorsRefsStr);
+
+	if (Serializer->IsReading())
+	{
+		TArray<FString> NewLayers;
+		if (LayersStr.ParseIntoArray(NewLayers, TEXT(";")))
+		{
+			Algo::Transform(NewLayers, Layers, [&](const FString& Layer) { return FName(*Layer); });
+		}
+
+		TArray<FString> ActorsRefs;
+		if (ActorsRefsStr.ParseIntoArray(ActorsRefs, TEXT(";")))
+		{
+			Algo::Transform(ActorsRefs, References, [](const FString& ActorGuidStr)
+			{
+				FGuid ActorGuid;
+				if (!FGuid::Parse(ActorGuidStr, ActorGuid))
+				{
+					ActorGuid.Invalidate();
+				}
+				return ActorGuid;
+			});
+
+			Algo::RemoveIf(References, [](const FGuid& InGuid)
+			{
+				return !InGuid.IsValid();
+			});
+		}
+	}
 }
 
 FBox FWorldPartitionActorDesc::GetBounds() const
diff --git a/Engine/Source/Runtime/Engine/Private/WorldPartition/WorldPartitionActorDescFactory.cpp b/Engine/Source/Runtime/Engine/Private/WorldPartition/WorldPartitionActorDescFactory.cpp
--- a/Engine/Source/Runtime/Engine/Private/WorldPartition/WorldPartitionActorDescFactory.cpp
+++ b/Engine/Source/Runtime/Engine/Private/WorldPartition/WorldPartitionActorDescFactory.cpp
@@ -4,148 +4,9 @@
 
 #if WITH_EDITOR
 #include "WorldPartition/WorldPartitionActorDesc.h"
-#include "ActorRegistry.h"
-#include "Algo/Transform.h"
-#include "Algo/RemoveIf.h"
 
-bool FWorldPartitionActorDescFactory::ReadMetaData(const FWorldPartitionActorDescInitData& ActorDescInitData, FWorldPartitionActorDescData& OutData)
+FWorldPartitionActorDesc* FWorldPartitionActorDescFactory::Create()
 {
-	OutData.ActorPackage = ActorDescInitData.PackageName;
-	OutData.ActorPath = ActorDescInitData.ActorPath;
-	OutData.Class = ActorDescInitData.NativeClass->GetFName();
-
-	static FName NAME_ActorGuid(TEXT("ActorGuid"));
-	if (!FActorRegistry::ReadActorMetaData(NAME_ActorGuid, OutData.Guid, ActorDescInitData.AssetData))
-	{
-		return false;
-	}
-		
-	static FName NAME_BoundsLocation(TEXT("BoundsLocation"));
-	if (!FActorRegistry::ReadActorMetaData(NAME_BoundsLocation, OutData.BoundsLocation, ActorDescInitData.AssetData))
-	{
-		return false;
-	}
-
-	static FName NAME_BoundsExtent(TEXT("BoundsExtent"));
-	if (!FActorRegistry::ReadActorMetaData(NAME_BoundsExtent, OutData.BoundsExtent, ActorDescInitData.AssetData))
-	{
-		return false;
-	}
-
-	// Transform BoundsLocation and BoundsExtent if necessary
-	if (!ActorDescInitData.Transform.Equals(FTransform::Identity))
-	{
-		//@todo_ow: This will result in a new BoundsExtent that is larger than it should. To fix this, we would need the Object Oriented BoundingBox of the actor (the BV of the actor without rotation)
-		const FVector BoundsMin = OutData.BoundsLocation - OutData.BoundsExtent;
-		const FVector BoundsMax = OutData.BoundsLocation + OutData.BoundsExtent;
-		const FBox NewBounds = FBox(BoundsMin, BoundsMax).TransformBy(ActorDescInitData.Transform);
-		NewBounds.GetCenterAndExtents(OutData.BoundsLocation, OutData.BoundsExtent);
-	}
-
-	const EActorGridPlacement DefaultGridPlacement = ActorDescInitData.NativeClass->GetDefaultObject<AActor>()->GetDefaultGridPlacement();
-	if (DefaultGridPlacement != EActorGridPlacement::None)
-	{
-		OutData.GridPlacement = DefaultGridPlacement;
-	}
-	else
-	{
-		int32 GridPlacementValue = (int32)EActorGridPlacement::AlwaysLoaded;
-		static const FName NAME_GridPlacement(TEXT("GridPlacement"));
-		FActorRegistry::ReadActorMetaData(NAME_GridPlacement, GridPlacementValue, ActorDescInitData.AssetData);
-		OutData.GridPlacement = (EActorGridPlacement)GridPlacementValue;
-	}
-
-	static const FName NAME_RuntimeGrid(TEXT("RuntimeGrid"));
-	FActorRegistry::ReadActorMetaData(NAME_RuntimeGrid, OutData.RuntimeGrid, ActorDescInitData.AssetData);
-
-	static const FName NAME_IsEditorOnly(TEXT("IsEditorOnly"));
-	FActorRegistry::ReadActorMetaData(NAME_IsEditorOnly, OutData.bActorIsEditorOnly, ActorDescInitData.AssetData);
-
-	static const FName NAME_IsLevelBoundsRelevant(TEXT("IsLevelBoundsRelevant"));
-	if (!FActorRegistry::ReadActorMetaData(NAME_IsLevelBoundsRelevant, OutData.bLevelBoundsRelevant, ActorDescInitData.AssetData))
-	{
-		OutData.bLevelBoundsRelevant = true;
-	}
-	
-	FString LayersStr;
-	static FName NAME_Layers(TEXT("Layers"));
-	if (FActorRegistry::ReadActorMetaData(NAME_Layers, LayersStr, ActorDescInitData.AssetData))
-	{
-		TArray<FString> Layers;
-		if (LayersStr.ParseIntoArray(Layers, TEXT(";")))
-		{
-			Algo::Transform(Layers, OutData.Layers, [&](const FString& Layer) { return FName(*Layer); });
-		}
-	}
-
-	FString ActorsRefsStr;
-	static const FName NAME_ActorReferences(TEXT("ActorReferences"));
-	if (FActorRegistry::ReadActorMetaData(NAME_ActorReferences, ActorsRefsStr, ActorDescInitData.AssetData))
-	{
-		TArray<FString> ActorsRefs;
-		if (ActorsRefsStr.ParseIntoArray(ActorsRefs, TEXT(";")))
-		{
-			Algo::Transform(ActorsRefs, OutData.References, [](const FString& ActorGuidStr)
-			{
-				FGuid ActorGuid;
-				if (!FGuid::Parse(ActorGuidStr, ActorGuid))
-				{
-					ActorGuid.Invalidate();
-				}
-				return ActorGuid;
-			});
-
-			Algo::RemoveIf(OutData.References, [](const FGuid& Guid)
-			{
-				return !Guid.IsValid();
-			});
-		}
-	}
-
-	return true;
+	return new FWorldPartitionActorDesc;
 }
-
-FWorldPartitionActorDesc* FWorldPartitionActorDescFactory::Create(const FWorldPartitionActorDescInitData& ActorDescInitData)
-{
-	if (FWorldPartitionActorDesc* NewDesc = CreateInstance(ActorDescInitData))
-	{
-		PostCreate(NewDesc);
-		return NewDesc;
-	}
-
-	return nullptr;
-}
-
-FWorldPartitionActorDesc* FWorldPartitionActorDescFactory::Create(AActor* InActor)
-{
-	if (FWorldPartitionActorDesc* NewDesc = CreateInstance(InActor))
-	{
-		PostCreate(NewDesc);
-		return NewDesc;
-	}
-
-	return nullptr;
-}
-
-void FWorldPartitionActorDescFactory::PostCreate(FWorldPartitionActorDesc* ActorDesc)
-{
-	ActorDesc->UpdateHash();
-}
-
-FWorldPartitionActorDesc* FWorldPartitionActorDescFactory::CreateInstance(const FWorldPartitionActorDescInitData& ActorDescInitData)
-{
-	FWorldPartitionActorDescData OutData;
-	if (!ReadMetaData(ActorDescInitData, OutData))
-	{
-		return nullptr;
-	}
-
-	return new FWorldPartitionActorDesc(OutData);
-}
-
-FWorldPartitionActorDesc* FWorldPartitionActorDescFactory::CreateInstance(AActor* InActor)
-{
-	return new FWorldPartitionActorDesc(InActor);
-}
-
 #endif
\ No newline at end of file
diff --git a/Engine/Source/Runtime/Engine/Private/WorldPartition/WorldPartitionActorDescMetaDataSerializer.cpp b/Engine/Source/Runtime/Engine/Private/WorldPartition/WorldPartitionActorDescMetaDataSerializer.cpp
new file mode 100644
--- /dev/null
+++ b/Engine/Source/Runtime/Engine/Private/WorldPartition/WorldPartitionActorDescMetaDataSerializer.cpp
@@ -0,0 +1,165 @@
+// Copyright Epic Games, Inc. All Rights Reserved.
+
+#include "WorldPartition/WorldPartitionActorMetaDataSerializer.h"
+
+#if WITH_EDITOR
+
+bool FActorMetaDataReader::Serialize(FName Name, bool& Value)
+{
+	FString ValueStr;
+	if (ReadTag(Name, ValueStr))
+	{
+		if (ValueStr == TEXT("0"))
+		{
+			Value = false;
+			return true;
+		}
+		else if (ValueStr == TEXT("1"))
+		{
+			Value = true;
+			return true;
+		}
+	}
+
+	return false;
+}
+
+bool FActorMetaDataReader::Serialize(FName Name, int8& Value)
+{
+	FString ValueStr;
+	if (ReadTag(Name, ValueStr))
+	{
+		Value = FCString::Atoi(*ValueStr);
+		return true;
+	}
+
+	return false;
+}
+
+bool FActorMetaDataReader::Serialize(FName Name, int32& Value)
+{
+	FString ValueStr;
+	if (ReadTag(Name, ValueStr))
+	{
+		Value = FCString::Atoi(*ValueStr);
+		return true;
+	}
+
+	return false;
+}
+
+bool FActorMetaDataReader::Serialize(FName Name, int64& Value)
+{
+	FString ValueStr;
+	if (ReadTag(Name, ValueStr))
+	{
+		Value = FCString::Atoi64(*ValueStr);
+		return true;
+	}
+
+	return false;
+}
+
+bool FActorMetaDataReader::Serialize(FName Name, FGuid& Value)
+{
+	FString ValueStr;
+	return ReadTag(Name, ValueStr) && FGuid::Parse(ValueStr, Value);
+}
+
+bool FActorMetaDataReader::Serialize(FName Name, FVector& Value)
+{
+	FString ValueStr;
+	if (ReadTag(Name, ValueStr))
+	{
+		Value.InitFromString(ValueStr);
+		return true;
+	}
+
+	return false;
+}
+
+bool FActorMetaDataReader::Serialize(FName Name, FTransform& Value)
+{
+	FString ValueStr;
+	if (ReadTag(Name, ValueStr))
+	{
+		Value.InitFromString(ValueStr);
+		return true;
+	}
+
+	return false;
+}
+
+bool FActorMetaDataReader::Serialize(FName Name, FString& Value)
+{
+	return ReadTag(Name, Value);
+}
+
+bool FActorMetaDataReader::Serialize(FName Name, FName& Value)
+{
+	FString ValueStr;
+	if (ReadTag(Name, ValueStr))
+	{
+		Value = *ValueStr;
+		return true;
+	}
+
+	return false;
+}
+
+bool FActorMetaDataReader::ReadTag(FName Name, FString& Value)
+{
+	return AssetData.GetTagValue(Name, Value);
+}
+
+bool FActorMetaDataWriter::Serialize(FName Name, bool& Value)
+{
+	return WriteTag(Name, Value ? TEXT("1") : TEXT("0"));
+}
+
+bool FActorMetaDataWriter::Serialize(FName Name, int8& Value)
+{
+	return WriteTag(Name, *FString::Printf(TEXT("%d"), Value));
+}
+
+bool FActorMetaDataWriter::Serialize(FName Name, int32& Value)
+{
+	return WriteTag(Name, *FString::Printf(TEXT("%d"), Value));
+}
+
+bool FActorMetaDataWriter::Serialize(FName Name, int64& Value)
+{
+	return WriteTag(Name, *FString::Printf(TEXT("%lld"), Value));
+}
+
+bool FActorMetaDataWriter::Serialize(FName Name, FGuid& Value)
+{
+	return WriteTag(Name, Value.ToString(EGuidFormats::Base36Encoded));
+}
+
+bool FActorMetaDataWriter::Serialize(FName Name, FVector& Value)
+{
+	return WriteTag(Name, Value.ToCompactString());
+}
+
+bool FActorMetaDataWriter::Serialize(FName Name, FTransform& Value)
+{
+	return WriteTag(Name, Value.ToString());
+}
+
+bool FActorMetaDataWriter::Serialize(FName Name, FString& Value)
+{
+	return WriteTag(Name, Value);
+}
+
+bool FActorMetaDataWriter::Serialize(FName Name, FName& Value)
+{
+	return WriteTag(Name, Value.ToString());
+}
+
+bool FActorMetaDataWriter::WriteTag(FName Name, const FString& Value)
+{
+	Tags.Add(UObject::FAssetRegistryTag(Name, Value, UObject::FAssetRegistryTag::TT_Hidden));
+	return true;
+}
+#endif
\ No newline at end of file
diff --git a/Engine/Source/Runtime/Engine/Public/ActorPartition/PartitionActor.h b/Engine/Source/Runtime/Engine/Public/ActorPartition/PartitionActor.h
--- a/Engine/Source/Runtime/Engine/Public/ActorPartition/PartitionActor.h
+++ b/Engine/Source/Runtime/Engine/Public/ActorPartition/PartitionActor.h
@@ -17,12 +17,6 @@ class ENGINE_API APartitionActor: public AActor
 public:
 	APartitionActor(const FObjectInitializer& ObjectInitializer = FObjectInitializer::Get());
 
-	//~ Begin UObject Interface
-#if WITH_EDITOR	
-	virtual void GetAssetRegistryTags(TArray<FAssetRegistryTag>& OutTags) const override;
-#endif
-	//~ End UObject Interface
-
 	//~ Begin AActor Interface
 #if WITH_EDITOR	
 	virtual EActorGridPlacement GetDefaultGridPlacement() const override { return EActorGridPlacement::Location; }
diff --git a/Engine/Source/Runtime/Engine/Public/ActorRegistry.h b/Engine/Source/Runtime/Engine/Public/ActorRegistry.h
deleted file mode 100644
diff --git a/Engine/Source/Runtime/Engine/Public/LevelInstance/LevelInstanceActor.h b/Engine/Source/Runtime/Engine/Public/LevelInstance/LevelInstanceActor.h
--- a/Engine/Source/Runtime/Engine/Public/LevelInstance/LevelInstanceActor.h
+++ b/Engine/Source/Runtime/Engine/Public/LevelInstance/LevelInstanceActor.h
@@ -39,7 +39,6 @@ public:
 	virtual void PreEditUndo() override;
 	virtual void PostEditUndo() override;
 	virtual void PostLoad() override;
-	virtual void GetAssetRegistryTags(TArray<FAssetRegistryTag>& OutTags) const override;
 	virtual bool CanEditChange(const FProperty* InProperty) const override;
 	virtual void PreEditChange(FProperty* PropertyThatWillChange) override;
 	virtual void PostEditChangeProperty(FPropertyChangedEvent& PropertyChangedEvent) override;
diff --git a/Engine/Source/Runtime/Engine/Public/WorldPartition/ActorPartition/PartitionActorDesc.h b/Engine/Source/Runtime/Engine/Public/WorldPartition/ActorPartition/PartitionActorDesc.h
--- a/Engine/Source/Runtime/Engine/Public/WorldPartition/ActorPartition/PartitionActorDesc.h
+++ b/Engine/Source/Runtime/Engine/Public/WorldPartition/ActorPartition/PartitionActorDesc.h
@@ -9,18 +9,20 @@
 
 class ENGINE_API FPartitionActorDesc : public FWorldPartitionActorDesc
 {
+#if WITH_EDITOR
+	friend class FPartitionActorDescFactory;
+
 public:
 	uint32 GridSize;
 	int64 GridIndexX;
 	int64 GridIndexY;
 	int64 GridIndexZ;
 
-#if WITH_EDITOR
 protected:
-	FPartitionActorDesc() = delete;
-	FPartitionActorDesc(const FWorldPartitionActorDescData& DescData, uint32 InGridSize, int64 InGridIndexX, int64 InGridIndexY, int64 InGridIndexZ);
-	FPartitionActorDesc(AActor* InActor);
+	virtual bool Init(const AActor* InActor) override;
+	
+	virtual void BuildHash(FHashBuilder& HashBuilder) override;
 
-	friend class FPartitionActorDescFactory;
+	virtual void SerializeMetaData(FActorMetaDataSerializer* Serializer) override;
 #endif
 };
diff --git a/Engine/Source/Runtime/Engine/Public/WorldPartition/ActorPartition/PartitionActorDescFactory.h b/Engine/Source/Runtime/Engine/Public/WorldPartition/ActorPartition/PartitionActorDescFactory.h
--- a/Engine/Source/Runtime/Engine/Public/WorldPartition/ActorPartition/PartitionActorDescFactory.h
+++ b/Engine/Source/Runtime/Engine/Public/WorldPartition/ActorPartition/PartitionActorDescFactory.h
@@ -2,16 +2,12 @@
 #pragma once
 
 #include "CoreMinimal.h"
-#include "UObject/ObjectMacros.h"
 #include "WorldPartition/WorldPartitionActorDescFactory.h"
 
-class FWorldPartitionActorDesc;
-
 class ENGINE_API FPartitionActorDescFactory : public FWorldPartitionActorDescFactory
 {
 public:
 #if WITH_EDITOR
-	virtual FWorldPartitionActorDesc* CreateInstance(const FWorldPartitionActorDescInitData& ActorDescInitData) override;
-	virtual FWorldPartitionActorDesc* CreateInstance(AActor* InActor) override;
+	virtual FWorldPartitionActorDesc* Create() override;
 #endif
 };
\ No newline at end of file
diff --git a/Engine/Source/Runtime/Engine/Public/WorldPartition/HLOD/HLODActor.h b/Engine/Source/Runtime/Engine/Public/WorldPartition/HLOD/HLODActor.h
--- a/Engine/Source/Runtime/Engine/Public/WorldPartition/HLOD/HLODActor.h
+++ b/Engine/Source/Runtime/Engine/Public/WorldPartition/HLOD/HLODActor.h
@@ -53,7 +53,6 @@ public:
 protected:
 	//~ Begin AActor Interface.
 	virtual EActorGridPlacement GetDefaultGridPlacement() const override;
-	virtual void GetAssetRegistryTags(TArray<FAssetRegistryTag>& OutTags) const override;
 	virtual void PostActorCreated() override;
 	virtual void GetActorBounds(bool bOnlyCollidingComponents, FVector& Origin, FVector& BoxExtent, bool bIncludeFromChildActors) const override;
 	virtual void GetActorLocationBounds(bool bOnlyCollidingComponents, FVector& Origin, FVector& BoxExtent, bool bIncludeFromChildActors) const override;
diff --git a/Engine/Source/Runtime/Engine/Public/WorldPartition/HLOD/HLODActorDesc.h b/Engine/Source/Runtime/Engine/Public/WorldPartition/HLOD/HLODActorDesc.h
--- a/Engine/Source/Runtime/Engine/Public/WorldPartition/HLOD/HLODActorDesc.h
+++ b/Engine/Source/Runtime/Engine/Public/WorldPartition/HLOD/HLODActorDesc.h
@@ -14,19 +14,19 @@ class UHLODLayer;
  */
 class ENGINE_API FHLODActorDesc : public FWorldPartitionActorDesc
 {
-public:
 #if WITH_EDITOR
+	friend class FHLODActorDescFactory;
+
+public:
 	inline const TArray<FGuid>& GetSubActors() const { return SubActors; }
 	inline const FSoftObjectPath& GetHLODLayer() const { return HLODLayer; }
 
 protected:
-	FHLODActorDesc() = delete;
-	FHLODActorDesc(const FWorldPartitionActorDescData& DescData, const TArray<FGuid>& InSubActors, const FSoftObjectPath& HLODLayer);
-	FHLODActorDesc(AActor* InActor);
+	virtual bool Init(const AActor* InActor) override;
 
 	virtual void BuildHash(FHashBuilder& HashBuilder) override;
 
-	friend class FHLODActorDescFactory;
+	virtual void SerializeMetaData(FActorMetaDataSerializer* Serializer) override;
 
 	TArray<FGuid> SubActors;
 	FSoftObjectPath	HLODLayer;
diff --git a/Engine/Source/Runtime/Engine/Public/WorldPartition/HLOD/HLODActorDescFactory.h b/Engine/Source/Runtime/Engine/Public/WorldPartition/HLOD/HLODActorDescFactory.h
--- a/Engine/Source/Runtime/Engine/Public/WorldPartition/HLOD/HLODActorDescFactory.h
+++ b/Engine/Source/Runtime/Engine/Public/WorldPartition/HLOD/HLODActorDescFactory.h
@@ -2,17 +2,12 @@
 #pragma once
 
 #include "CoreMinimal.h"
-#include "UObject/ObjectMacros.h"
 #include "WorldPartition/WorldPartitionActorDescFactory.h"
 
-class FWorldPartitionActorDesc;
-class UHLODSubsystem;
-
 class ENGINE_API FHLODActorDescFactory : public FWorldPartitionActorDescFactory
 {
 public:
 #if WITH_EDITOR
-	virtual FWorldPartitionActorDesc* CreateInstance(const FWorldPartitionActorDescInitData& ActorDescInitData) override;
-	virtual FWorldPartitionActorDesc* CreateInstance(AActor* InActor) override;
+	virtual FWorldPartitionActorDesc* Create() override;
 #endif
 };
diff --git a/Engine/Source/Runtime/Engine/Public/WorldPartition/LevelInstance/LevelInstanceActorDesc.h b/Engine/Source/Runtime/Engine/Public/WorldPartition/LevelInstance/LevelInstanceActorDesc.h
--- a/Engine/Source/Runtime/Engine/Public/WorldPartition/LevelInstance/LevelInstanceActorDesc.h
+++ b/Engine/Source/Runtime/Engine/Public/WorldPartition/LevelInstance/LevelInstanceActorDesc.h
@@ -15,19 +15,20 @@ class UWorldPartition;
  */
 class ENGINE_API FLevelInstanceActorDesc : public FWorldPartitionActorDesc
 {
-public:
 #if WITH_EDITOR
+	friend class FLevelInstanceActorDescFactory;
+
+public:
 	inline FName GetLevelPackage() const { return LevelPackage; }
 
 protected:
-	FLevelInstanceActorDesc() = delete;
-	FLevelInstanceActorDesc(const FWorldPartitionActorDescData& DescData, FName LevelPackage);
-	FLevelInstanceActorDesc(AActor* InActor);
+	virtual bool Init(const AActor* InActor) override;
 
 	virtual void BuildHash(FHashBuilder& HashBuilder) override;
 
-	friend class FLevelInstanceActorDescFactory;
+	virtual void SerializeMetaData(FActorMetaDataSerializer* Serializer) override;
 
 	FName LevelPackage;
+	FTransform LevelInstanceTransform;
 #endif
 };
diff --git a/Engine/Source/Runtime/Engine/Public/WorldPartition/LevelInstance/LevelInstanceActorDescFactory.h b/Engine/Source/Runtime/Engine/Public/WorldPartition/LevelInstance/LevelInstanceActorDescFactory.h
--- a/Engine/Source/Runtime/Engine/Public/WorldPartition/LevelInstance/LevelInstanceActorDescFactory.h
+++ b/Engine/Source/Runtime/Engine/Public/WorldPartition/LevelInstance/LevelInstanceActorDescFactory.h
@@ -2,17 +2,12 @@
 #pragma once
 
 #include "CoreMinimal.h"
-#include "UObject/ObjectMacros.h"
 #include "WorldPartition/WorldPartitionActorDescFactory.h"
 
-class FWorldPartitionActorDesc;
-class ULevelInstanceSubsystem;
-
 class ENGINE_API FLevelInstanceActorDescFactory : public FWorldPartitionActorDescFactory
 {
 public:
 #if WITH_EDITOR
-	virtual FWorldPartitionActorDesc* CreateInstance(const FWorldPartitionActorDescInitData& ActorDescInitData) override;
-	virtual FWorldPartitionActorDesc* CreateInstance(AActor* InActor) override;
+	virtual FWorldPartitionActorDesc* Create() override;
 #endif
 };
\ No newline at end of file
diff --git a/Engine/Source/Runtime/Engine/Public/WorldPartition/WorldPartition.h b/Engine/Source/Runtime/Engine/Public/WorldPartition/WorldPartition.h
--- a/Engine/Source/Runtime/Engine/Public/WorldPartition/WorldPartition.h
+++ b/Engine/Source/Runtime/Engine/Public/WorldPartition/WorldPartition.h
@@ -4,7 +4,9 @@
 #include "CoreMinimal.h"
 #include "UObject/ObjectMacros.h"
 #include "UObject/Object.h"
+#include "GameFramework/Actor.h"
 #include "Templates/SubclassOf.h"
+#include "WorldPartition/WorldPartitionActorDesc.h"
 #include "WorldPartition/WorldPartitionActorDescFactory.h"
 #include "WorldPartition.generated.h"
 
diff --git a/Engine/Source/Runtime/Engine/Public/WorldPartition/WorldPartitionActorDesc.h b/Engine/Source/Runtime/Engine/Public/WorldPartition/WorldPartitionActorDesc.h
--- a/Engine/Source/Runtime/Engine/Public/WorldPartition/WorldPartitionActorDesc.h
+++ b/Engine/Source/Runtime/Engine/Public/WorldPartition/WorldPartitionActorDesc.h
@@ -2,37 +2,38 @@
 #pragma once
 
 #include "CoreMinimal.h"
+#include "AssetData.h"
 #include "UObject/NameTypes.h"
 #include "UObject/ObjectMacros.h"
 #include "GameFramework/Actor.h"
 #include "Containers/Set.h"
 #include "Misc/Guid.h"
 #include "Misc/HashBuilder.h"
+#include "WorldPartition/WorldPartitionActorMetaDataSerializer.h"
+#include "WorldPartition/WorldPartitionActorDescFactory.h"
 
-class FWorldPartitionActorDescData
+#if WITH_EDITOR
+// Struct used to create actor descriptor
+struct FWorldPartitionActorDescInitData
 {
-public:
-	FGuid Guid;
-	FName Class;
-	FName ActorPackage;
+	UClass* NativeClass;
+	FName PackageName;
 	FName ActorPath;
-	FVector BoundsLocation;
-	FVector BoundsExtent;
-	EActorGridPlacement GridPlacement;
-	FName RuntimeGrid;
-	bool bActorIsEditorOnly;
-	bool bLevelBoundsRelevant;
-	TArray<FName> Layers;
-	TArray<FGuid> References;
+	FAssetData AssetData;
+	FTransform Transform;
+	FActorMetaDataSerializer* Serializer;
 };
+#endif
 
 /**
  * Represents a potentially unloaded actor (editor-only)
  */
-class ENGINE_API FWorldPartitionActorDesc : protected FWorldPartitionActorDescData
+class ENGINE_API FWorldPartitionActorDesc 
 {
 #if WITH_EDITOR
-	FWorldPartitionActorDesc() = delete;
+	friend class AActor;
+	friend class UWorldPartition;
+	friend class FWorldPartitionActorDescFactory;
 
 public:
 	virtual ~FWorldPartitionActorDesc() {}
@@ -84,12 +85,27 @@ public:
 	void Unload();
 
 protected:
-	FWorldPartitionActorDesc(AActor* InActor);
-	FWorldPartitionActorDesc(const FWorldPartitionActorDescData& DescData);
-	friend class FWorldPartitionActorDescFactory;
+	FWorldPartitionActorDesc();
+	virtual bool Init(const AActor* InActor);
+	bool Init(const FWorldPartitionActorDescInitData& DescData);
 		
 	void UpdateHash();
 	virtual void BuildHash(FHashBuilder& HashBuilder);
+
+	virtual void SerializeMetaData(FActorMetaDataSerializer* Serializer);
+
+	FGuid						Guid;
+	FName						Class;
+	FName						ActorPackage;
+	FName						ActorPath;
+	FVector						BoundsLocation;
+	FVector						BoundsExtent;
+	EActorGridPlacement			GridPlacement;
+	FName						RuntimeGrid;
+	bool						bActorIsEditorOnly;
+	bool						bLevelBoundsRelevant;
+	TArray<FName>				Layers;
+	TArray<FGuid>				References;
 	
 	mutable uint32				LoadedRefCount;
 	mutable uint32				Hash;
diff --git a/Engine/Source/Runtime/Engine/Public/WorldPartition/WorldPartitionActorDescFactory.h b/Engine/Source/Runtime/Engine/Public/WorldPartition/WorldPartitionActorDescFactory.h
--- a/Engine/Source/Runtime/Engine/Public/WorldPartition/WorldPartitionActorDescFactory.h
+++ b/Engine/Source/Runtime/Engine/Public/WorldPartition/WorldPartitionActorDescFactory.h
@@ -2,40 +2,13 @@
 #pragma once
 
 #include "CoreMinimal.h"
-#include "AssetData.h"
-#include "UObject/ObjectMacros.h"
-#include "WorldPartition/WorldPartitionActorDesc.h"
 
-#if WITH_EDITOR
-enum class EActorGridPlacement : uint8;
-
-// Struct used to create actor descriptor
-struct FWorldPartitionActorDescInitData
-{
-	UClass* NativeClass;
-	FName PackageName;
-	FName ActorPath;
-	FAssetData AssetData;
-	FTransform Transform;
-};
-#endif
+class FWorldPartitionActorDesc;
 
 class ENGINE_API FWorldPartitionActorDescFactory
 {
-#if WITH_EDITOR
-private:
-	void PostCreate(FWorldPartitionActorDesc* ActorDesc);
-
-protected:
-	bool ReadMetaData(const FWorldPartitionActorDescInitData& ActorDescInitData, FWorldPartitionActorDescData& OutData);
-
-	virtual FWorldPartitionActorDesc* CreateInstance(const FWorldPartitionActorDescInitData& ActorDescInitData);
-	virtual FWorldPartitionActorDesc* CreateInstance(AActor* InActor);
-
 public:
-	virtual ~FWorldPartitionActorDescFactory() = default;
-
-	FWorldPartitionActorDesc* Create(const FWorldPartitionActorDescInitData& ActorDescInitData);
-	FWorldPartitionActorDesc* Create(AActor* InActor);
+#if WITH_EDITOR
+	virtual FWorldPartitionActorDesc* Create();
 #endif
 };
\ No newline at end of file
diff --git a/Engine/Source/Runtime/Engine/Public/WorldPartition/WorldPartitionActorMetaDataSerializer.h b/Engine/Source/Runtime/Engine/Public/WorldPartition/WorldPartitionActorMetaDataSerializer.h
new file mode 100644
--- /dev/null
+++ b/Engine/Source/Runtime/Engine/Public/WorldPartition/WorldPartitionActorMetaDataSerializer.h
@@ -0,0 +1,85 @@
+// Copyright Epic Games, Inc. All Rights Reserved.
+#pragma once
+
+#include "CoreMinimal.h"
+#include "AssetData.h"
+#include "Misc/Guid.h"
+
+#if WITH_EDITOR
+class FActorMetaDataSerializer
+{
+public:
+	FActorMetaDataSerializer()
+		: bHasErrors(false)
+	{}
+
+	virtual bool IsReading() const =0;
+	inline bool IsWriting() const { return !IsReading(); }
+
+	virtual bool Serialize(FName Name, bool& Value)=0;
+	virtual bool Serialize(FName Name, int8& Value)=0;
+	virtual bool Serialize(FName Name, int32& Value)=0;
+	virtual bool Serialize(FName Name, int64& Value)=0;
+	virtual bool Serialize(FName Name, FGuid& Value)=0;
+	virtual bool Serialize(FName Name, FVector& Value)=0;
+	virtual bool Serialize(FName Name, FTransform& Value)=0;
+	virtual bool Serialize(FName Name, FString& Value)=0;
+	virtual bool Serialize(FName Name, FName& Value)=0;
+
+	inline void SetHasErrors() { bHasErrors = true; }
+	inline bool GetHasErrors() const { return bHasErrors; }
+
+private:
+	bool bHasErrors;
+};
+
+class FActorMetaDataReader: public FActorMetaDataSerializer
+{
+public:
+	FActorMetaDataReader(const FAssetData& InAssetData)
+		: AssetData(InAssetData)
+	{}
+
+	virtual bool IsReading() const override { return true; }
+
+	virtual bool Serialize(FName Name, bool& Value) override;
+	virtual bool Serialize(FName Name, int8& Value) override;
+	virtual bool Serialize(FName Name, int32& Value) override;
+	virtual bool Serialize(FName Name, int64& Value) override;
+	virtual bool Serialize(FName Name, FGuid& Value) override;
+	virtual bool Serialize(FName Name, FVector& Value) override;
+	virtual bool Serialize(FName Name, FTransform& Value) override;
+	virtual bool Serialize(FName Name, FString& Value) override;
+	virtual bool Serialize(FName Name, FName& Value) override;
+
+private:
+	bool ReadTag(FName Name, FString& Value);
+
+	const FAssetData& AssetData;
+};
+
+class FActorMetaDataWriter: public FActorMetaDataSerializer
+{
+public:
+	FActorMetaDataWriter(TArray<UObject::FAssetRegistryTag>& InTags)
+		: Tags(InTags)
+	{}
+
+	virtual bool IsReading() const override { return false; }
+
+	virtual bool Serialize(FName Name, bool& Value) override;
+	virtual bool Serialize(FName Name, int8& Value) override;
+	virtual bool Serialize(FName Name, int32& Value) override;
+	virtual bool Serialize(FName Name, int64& Value) override;
+	virtual bool Serialize(FName Name, FGuid& Value) override;
+	virtual bool Serialize(FName Name, FVector& Value) override;
+	virtual bool Serialize(FName Name, FTransform& Value) override;
+	virtual bool Serialize(FName Name, FString& Value) override;
+	virtual bool Serialize(FName Name, FName& Value) override;
+
+private:
+	bool WriteTag(FName Name, const FString& Value);
+
+	TArray<UObject::FAssetRegistryTag>& Tags;
+};
+#endif
\ No newline at end of file
