Fix deadlock caused by listeners attempting to access AudioDevice amidst device destruction sequence where it should no longer be accessed.

diff --git a/Engine/Plugins/FX/Niagara/Source/Niagara/Private/NiagaraDataInterfaceAudio.cpp b/Engine/Plugins/FX/Niagara/Source/Niagara/Private/NiagaraDataInterfaceAudio.cpp
--- a/Engine/Plugins/FX/Niagara/Source/Niagara/Private/NiagaraDataInterfaceAudio.cpp
+++ b/Engine/Plugins/FX/Niagara/Source/Niagara/Private/NiagaraDataInterfaceAudio.cpp
@@ -73,10 +73,9 @@ void FNiagaraSubmixListener::UnregisterFromSubmix()
 	if (bIsRegistered)
 	{
 		FAudioDeviceHandle DeviceHandle = FAudioDeviceManager::Get()->GetAudioDevice(AudioDeviceId);
-		USoundSubmix* RawSubmixPtr = Submix;
-
 		check(DeviceHandle);
 		DeviceHandle->UnregisterSubmixBufferListener(this, Submix);
+
 		bIsRegistered = false;
 
 		// UnregisterSubmixBufferListener lazily enqueues the unregistration on the audio thread,
diff --git a/Engine/Source/Runtime/Engine/Private/AudioDeviceManager.cpp b/Engine/Source/Runtime/Engine/Private/AudioDeviceManager.cpp
--- a/Engine/Source/Runtime/Engine/Private/AudioDeviceManager.cpp
+++ b/Engine/Source/Runtime/Engine/Private/AudioDeviceManager.cpp
@@ -277,30 +277,33 @@ bool FAudioDeviceManager::ToggleAudioMixer()
 	// Only need to toggle if we have 2 device module names loaded at init
 	if (AudioDeviceModule && AudioDeviceModuleName.Len() > 0 && AudioMixerModuleName.Len() > 0)
 	{
-		// Check to see if there are still any active handles using this audio device, and either warn or exit early.
-		for (auto& DeviceContainer : Devices)
+		FScopeLock ScopeLock(&DeviceMapCriticalSection);
 		{
-			if (DeviceContainer.Value.NumberOfHandlesToThisDevice > 0)
+			// Check to see if there are still any active handles using this audio device, and either warn or exit early.
+			for (auto& DeviceContainer : Devices)
 			{
+				if (DeviceContainer.Value.NumberOfHandlesToThisDevice > 0)
+				{
 
-				UE_LOG(LogAudio, Warning, TEXT("Attempted to toggle the audio mixer while that audio device was in use."));
+					UE_LOG(LogAudio, Warning, TEXT("Attempted to toggle the audio mixer while that audio device was in use."));
 
-#if INSTRUMENT_AUDIODEVICE_HANDLES
-				FString ActiveDeviceHandles;
-				for (auto& StackWalkString : DeviceContainer.Value.HandleCreationStackWalks)
-				{
-					ActiveDeviceHandles += StackWalkString.Value;
-					ActiveDeviceHandles += TEXT("\n\n");
-				}
+	#if INSTRUMENT_AUDIODEVICE_HANDLES
+					FString ActiveDeviceHandles;
+					for (auto& StackWalkString : DeviceContainer.Value.HandleCreationStackWalks)
+					{
+						ActiveDeviceHandles += StackWalkString.Value;
+						ActiveDeviceHandles += TEXT("\n\n");
+					}
 
-				UE_LOG(LogAudio, Warning, TEXT("List Of Active Handles to this device: \n%s"), *ActiveDeviceHandles);
-#else
-				UE_LOG(LogAudio, Warning, TEXT("For more information compile with INSTRUMENT_AUDIODEVICE_HANDLES."));
-#endif
+					UE_LOG(LogAudio, Warning, TEXT("List Of Active Handles to this device: \n%s"), *ActiveDeviceHandles);
+	#else
+					UE_LOG(LogAudio, Warning, TEXT("For more information compile with INSTRUMENT_AUDIODEVICE_HANDLES."));
+	#endif
 
-				if (!GCvarAllowUnsafeAudioMixerToggling)
-				{
-					return false;
+					if (!GCvarAllowUnsafeAudioMixerToggling)
+					{
+						return false;
+					}
 				}
 			}
 		}
@@ -709,20 +712,34 @@ void FAudioDeviceManager::DecrementDevice(Audio::FDeviceId DeviceID, UWorld* InW
 
 		FAudioDeviceContainer& Container = Devices[DeviceID];
 		check(Container.NumberOfHandlesToThisDevice > 0);
-		Container.NumberOfHandlesToThisDevice--;
 
-		// If there is no longer anyone using this device, shut it down.
-		if (!Container.NumberOfHandlesToThisDevice)
+		// Report device being destroyed before actual destruction
+		// to allow listeners to access and respond where applicable.
+		bool bDestroyingDevice = false;
+		if (Container.NumberOfHandlesToThisDevice == 1)
 		{
+			bDestroyingDevice = true;
+			FAudioDeviceManagerDelegates::OnAudioDeviceDestroyed.Broadcast(DeviceID);
+
 			// If this is the active device and being destroyed, set the main device as the active device.
 			if (DeviceID == ActiveAudioDeviceID)
 			{
 				SetActiveDevice(MainAudioDeviceHandle.GetDeviceID());
 			}
 
-			FAudioDeviceManagerDelegates::OnAudioDeviceDestroyed.Broadcast(DeviceID);
-			Swap(DeviceToTearDown, Container.Device);
 			UnregisterWorld(InWorld, DeviceID);
+		}
+
+		Container.NumberOfHandlesToThisDevice--;
+
+		// If there is no longer any users of this device, destroy it.
+		if (Container.NumberOfHandlesToThisDevice)
+		{
+			ensureMsgf(!bDestroyingDevice, TEXT("AudioDevice Destruction Failure: 'OnAudioDeviceDestroyed' listener generated new persistent handle(s) to AudioDevice."));
+		}
+		else
+		{
+			Swap(DeviceToTearDown, Container.Device);
 			Devices.Remove(DeviceID);
 		}
 	}
