Fix double-DPI scaling for debug text

diff --git a/Engine/Source/Runtime/Engine/Private/HUD.cpp b/Engine/Source/Runtime/Engine/Private/HUD.cpp
--- a/Engine/Source/Runtime/Engine/Private/HUD.cpp
+++ b/Engine/Source/Runtime/Engine/Private/HUD.cpp
@@ -735,7 +735,8 @@ void AHUD::DrawDebugTextList()
 			// don't draw text behind the camera
 			if (((WorldTextLoc - CameraLoc) | CameraRot.Vector()) > 0.f)
 			{
-				FVector ScreenLoc = Canvas->Project(WorldTextLoc);
+				// Project() will give a DPI-scaled position, but DrawItem will also DPI-scale before drawing. So unscale it here to get the right position.
+				FVector ScreenLoc = Canvas->Project(WorldTextLoc) / DebugCanvas->GetDPIScale();
 				TextItem.SetColor(DebugTextList[Idx].TextColor);
 				TextItem.Text = FText::FromString(DebugTextList[Idx].DebugText);
 				TextItem.Scale = FVector2D(DebugTextList[Idx].FontScale, DebugTextList[Idx].FontScale);
