Fix Linux build.

diff --git a/Engine/Source/Runtime/Engine/Public/WorldPartition/WorldPartitionActorDescFactory.h b/Engine/Source/Runtime/Engine/Public/WorldPartition/WorldPartitionActorDescFactory.h
--- a/Engine/Source/Runtime/Engine/Public/WorldPartition/WorldPartitionActorDescFactory.h
+++ b/Engine/Source/Runtime/Engine/Public/WorldPartition/WorldPartitionActorDescFactory.h
@@ -8,6 +8,8 @@ class FWorldPartitionActorDesc;
 class ENGINE_API FWorldPartitionActorDescFactory
 {
 public:
+	virtual ~FWorldPartitionActorDescFactory() = default;
+
 #if WITH_EDITOR
 	virtual FWorldPartitionActorDesc* Create();
 #endif
