WorldPartition: Unified all worldpartition console commands to start with wp.Runtime or wp.Editor.

diff --git a/Engine/Source/Runtime/Engine/Private/WorldPartition/WorldPartition.cpp b/Engine/Source/Runtime/Engine/Private/WorldPartition/WorldPartition.cpp
--- a/Engine/Source/Runtime/Engine/Private/WorldPartition/WorldPartition.cpp
+++ b/Engine/Source/Runtime/Engine/Private/WorldPartition/WorldPartition.cpp
@@ -40,6 +40,8 @@
 
 DEFINE_LOG_CATEGORY(LogWorldPartition);
 
+PRAGMA_DISABLE_OPTIMIZATION
+
 #define LOCTEXT_NAMESPACE "WorldPartitionEditor"
 
 #if WITH_EDITOR
@@ -59,7 +61,7 @@ static void GenerateHLOD(const TArray<FString>& Args)
 }
 
 static FAutoConsoleCommand GenerateHLODCmd(
-	TEXT("WorldPartition.GenerateHLOD"),
+	TEXT("wp.Editor.GenerateHLOD"),
 	TEXT("Generates HLOD data for runtime."),
 	FConsoleCommandWithArgsDelegate::CreateStatic(&GenerateHLOD)
 );
@@ -633,7 +635,7 @@ void UWorldPartition::AddToClusters(const FWorldPartitionActorDesc* ActorDesc)
 			const FWorldPartitionActorDesc* ReferenceActorDesc = GetActorDesc(ReferenceGuid);
 
 			// Don't include references to editor-only actors
-			if (!ReferenceActorDesc->GetActorIsEditorOnly())
+			if (ReferenceActorDesc && !ReferenceActorDesc->GetActorIsEditorOnly())
 			{
 				FActorCluster* ReferenceCluster = ActorToActorCluster.FindRef(ReferenceGuid);
 				if (ReferenceCluster)
@@ -678,7 +680,10 @@ void UWorldPartition::RemoveFromClusters(const FWorldPartitionActorDesc* ActorDe
 
 	for (const FGuid& Guid : ActorCluster->Actors)
 	{
-		AddToClusters(GetActorDesc(Guid));
+		if (FWorldPartitionActorDesc* ClusterActorDesc = GetActorDesc(Guid))
+		{
+			AddToClusters(ClusterActorDesc);
+		}
 	}
 
 	delete ActorCluster;
@@ -1309,4 +1314,5 @@ FBox UWorldPartition::GetWorldBounds() const
 }
 #endif
 
-#undef LOCTEXT_NAMESPACE
\ No newline at end of file
+#undef LOCTEXT_NAMESPACE
+PRAGMA_ENABLE_OPTIMIZATION
\ No newline at end of file
diff --git a/Engine/Source/Runtime/Engine/Private/WorldPartition/WorldPartitionLevelStreamingPolicy.cpp b/Engine/Source/Runtime/Engine/Private/WorldPartition/WorldPartitionLevelStreamingPolicy.cpp
--- a/Engine/Source/Runtime/Engine/Private/WorldPartition/WorldPartitionLevelStreamingPolicy.cpp
+++ b/Engine/Source/Runtime/Engine/Private/WorldPartition/WorldPartitionLevelStreamingPolicy.cpp
@@ -20,7 +20,7 @@
 
 static int32 GMaxLoadingLevelStreamingCells = 4;
 static FAutoConsoleVariableRef CMaxLoadingLevelStreamingCells(
-	TEXT("WorldPartition.MaxLoadingLevelStreamingCells"),
+	TEXT("wp.Runtime.MaxLoadingLevelStreamingCells"),
 	GMaxLoadingLevelStreamingCells,
 	TEXT("Used to limit the number of concurrent loading world partition streaming cells."));
 
diff --git a/Engine/Source/Runtime/Engine/Private/WorldPartition/WorldPartitionRuntimeSpatialHash.cpp b/Engine/Source/Runtime/Engine/Private/WorldPartition/WorldPartitionRuntimeSpatialHash.cpp
--- a/Engine/Source/Runtime/Engine/Private/WorldPartition/WorldPartitionRuntimeSpatialHash.cpp
+++ b/Engine/Source/Runtime/Engine/Private/WorldPartition/WorldPartitionRuntimeSpatialHash.cpp
@@ -78,25 +78,25 @@ private:
 
 static int32 GShowRuntimeSpatialHashGridLevel = 0;
 static FAutoConsoleVariableRef CVarShowRuntimeSpatialHashGridLevel(
-	TEXT("WorldPartition.ShowRuntimeSpatialHashGridLevel"),
+	TEXT("wp.Runtime.ShowRuntimeSpatialHashGridLevel"),
 	GShowRuntimeSpatialHashGridLevel,
 	TEXT("Used to choose which grid level to display when showing world partition runtime hash."));
 
 static int32 GShowRuntimeSpatialHashGridLevelCount = 1;
 static FAutoConsoleVariableRef CVarShowRuntimeSpatialHashGridLevelCount(
-	TEXT("WorldPartition.ShowRuntimeSpatialHashGridLevelCount"),
+	TEXT("wp.Runtime.ShowRuntimeSpatialHashGridLevelCount"),
 	GShowRuntimeSpatialHashGridLevelCount,
 	TEXT("Used to choose how many grid levels to display when showing world partition runtime hash."));
 
 static int32 GShowRuntimeSpatialHashGridIndex = 0;
 static FAutoConsoleVariableRef CVarShowRuntimeSpatialHashGridIndex(
-	TEXT("WorldPartition.ShowRuntimeSpatialHashGridIndex"),
+	TEXT("wp.Runtime.ShowRuntimeSpatialHashGridIndex"),
 	GShowRuntimeSpatialHashGridIndex,
-	TEXT("Used to show only one particual grid when showing world partition runtime hash (invalid index will show all)."));
+	TEXT("Used to show only one particular grid when showing world partition runtime hash (invalid index will show all)."));
 
 static float GRuntimeSpatialHashCellToSourceAngleContributionToCellImportance = 0.4f; // Value between [0, 1]
 static FAutoConsoleVariableRef CVarRuntimeSpatialHashCellToSourceAngleContributionToCellImportance(
-	TEXT("WorldPartition.RuntimeSpatialHashCellToSourceAngleContributionToCellImportance"),
+	TEXT("wp.Runtime.RuntimeSpatialHashCellToSourceAngleContributionToCellImportance"),
 	GRuntimeSpatialHashCellToSourceAngleContributionToCellImportance,
 	TEXT("Value between 0 and 1 that modulates the contribution of the angle between streaming source-to-cell vector and source-forward vector to the cell importance. The closest to 0, the less the angle will contribute to the cell importance."));
 
@@ -1274,7 +1274,7 @@ bool UWorldPartitionRuntimeSpatialHash::GenerateHLOD()
 #endif
 
 FAutoConsoleCommand UWorldPartitionRuntimeSpatialHash::OverrideLoadingRangeCommand(
-	TEXT("WorldPartition.OverrideRuntimeSpatialHashLoadingRange"),
+	TEXT("wp.Runtime.OverrideRuntimeSpatialHashLoadingRange"),
 	TEXT("Sets runtime loading range. Args -grid=[index] -range=[override_loading_range]"),
 	FConsoleCommandWithArgsDelegate::CreateLambda([](const TArray<FString>& Args)
 	{
diff --git a/Engine/Source/Runtime/Engine/Private/WorldPartition/WorldPartitionSubsystem.cpp b/Engine/Source/Runtime/Engine/Private/WorldPartition/WorldPartitionSubsystem.cpp
--- a/Engine/Source/Runtime/Engine/Private/WorldPartition/WorldPartitionSubsystem.cpp
+++ b/Engine/Source/Runtime/Engine/Private/WorldPartition/WorldPartitionSubsystem.cpp
@@ -13,13 +13,13 @@ static const FName NAME_WorldPartitionRuntimeHash("WorldPartitionRuntimeHash");
 
 static int32 GDrawRuntimeHash3D = 0;
 static FAutoConsoleCommand CVarDrawRuntimeHash3D(
-	TEXT("WorldPartition.ToggleDrawRuntimeHash3D"),
+	TEXT("wp.Runtime.ToggleDrawRuntimeHash3D"),
 	TEXT("Toggles 3D debug display of world partition runtime hash."),
 	FConsoleCommandDelegate::CreateLambda([] { GDrawRuntimeHash3D = !GDrawRuntimeHash3D; }));
 
 static int32 GDrawRuntimeHash2D = 0;
 static FAutoConsoleCommand CVarDrawRuntimeHash2D(
-	TEXT("WorldPartition.ToggleDrawRuntimeHash2D"),
+	TEXT("wp.Runtime.ToggleDrawRuntimeHash2D"),
 	TEXT("Toggles 2D debug display of world partition runtime hash."),
 	FConsoleCommandDelegate::CreateLambda([] { GDrawRuntimeHash2D = !GDrawRuntimeHash2D; }));
 
