Fix missing material updates during async texture compilation

diff --git a/Engine/Source/Runtime/Engine/Private/TextureCompiler.cpp b/Engine/Source/Runtime/Engine/Private/TextureCompiler.cpp
--- a/Engine/Source/Runtime/Engine/Private/TextureCompiler.cpp
+++ b/Engine/Source/Runtime/Engine/Private/TextureCompiler.cpp
@@ -77,16 +77,19 @@ namespace TextureCompilingManagerImpl
 			
 			// Fix in CL 13480995 broke GetReferencedTextures() ability to return all referenced textures
 			// so we manually gather them instead...
+			// Recursion is required for MIC to resolve all TextureParameterValues in the hierarchy.
 			UMaterialInstance* MaterialInstance = Cast<UMaterialInstance>(MaterialInterface);
-			if (MaterialInstance)
+			while (MaterialInstance)
 			{
 				for (const FTextureParameterValue& TextureParam : MaterialInstance->TextureParameterValues)
 				{
 					if (TextureParam.ParameterValue)
 					{
-						TexturesRequiringMaterialUpdate.Emplace(TextureParam.ParameterValue, MaterialInterface);
+						TexturesRequiringMaterialUpdate.Add(TextureParam.ParameterValue, MaterialInterface);
 					}
 				}
+
+				MaterialInstance = Cast<UMaterialInstance>(MaterialInstance->Parent);
 			}
 		}
 
