Blueprints: Add Double Comparison functions

diff --git a/Engine/Source/Runtime/Engine/Classes/Kismet/KismetMathLibrary.h b/Engine/Source/Runtime/Engine/Classes/Kismet/KismetMathLibrary.h
--- a/Engine/Source/Runtime/Engine/Classes/Kismet/KismetMathLibrary.h
+++ b/Engine/Source/Runtime/Engine/Classes/Kismet/KismetMathLibrary.h
@@ -541,6 +541,26 @@ class ENGINE_API UKismetMathLibrary : public UBlueprintFunctionLibrary
 	UFUNCTION(BlueprintPure, meta = (DisplayName = "double / double", CompactNodeTitle = "/", Keywords = "/ divide division"), Category = "Math|Double")
 	static double Divide_DoubleDouble(double A, double B = 1.0);
 
+	/** Returns true if A is Less than B (A < B) */
+	UFUNCTION(BlueprintPure, meta = (DisplayName = "double < double", CompactNodeTitle = "<", Keywords = "< less"), Category = "Math|double")
+	static bool Less_DoubleDouble(double A, double B);
+
+	/** Returns true if A is greater than B (A > B) */
+	UFUNCTION(BlueprintPure, meta = (DisplayName = "double > double", CompactNodeTitle = ">", Keywords = "> greater"), Category = "Math|double")
+	static bool Greater_DoubleDouble(double A, double B);
+
+	/** Returns true if A is Less than or equal to B (A <= B) */
+	UFUNCTION(BlueprintPure, meta = (DisplayName = "double <= double", CompactNodeTitle = "<=", Keywords = "<= less"), Category = "Math|double")
+	static bool LessEqual_DoubleDouble(double A, double B);
+
+	/** Returns true if A is greater than or equal to B (A >= B) */
+	UFUNCTION(BlueprintPure, meta = (DisplayName = "double >= double", CompactNodeTitle = ">=", Keywords = ">= greater"), Category = "Math|double")
+	static bool GreaterEqual_DoubleDouble(double A, double B);
+
+	/** Returns true if A is exactly equal to B (A == B)*/
+	UFUNCTION(BlueprintPure, meta = (DisplayName = "Equal (double)", CompactNodeTitle = "==", Keywords = "== equal"), Category = "Math|double")
+	static bool EqualEqual_DoubleDouble(double A, double B);
+
 	/** Addition (A + B) */
 	UFUNCTION(BlueprintPure, meta=(DisplayName = "float + float", CompactNodeTitle = "+", Keywords = "+ add plus", CommutativeAssociativeBinaryOperator = "true"), Category="Math|Float")
 	static float Add_FloatFloat(float A, float B = 1.f);
diff --git a/Engine/Source/Runtime/Engine/Classes/Kismet/KismetMathLibrary.inl b/Engine/Source/Runtime/Engine/Classes/Kismet/KismetMathLibrary.inl
--- a/Engine/Source/Runtime/Engine/Classes/Kismet/KismetMathLibrary.inl
+++ b/Engine/Source/Runtime/Engine/Classes/Kismet/KismetMathLibrary.inl
@@ -491,6 +491,36 @@ double UKismetMathLibrary::Divide_DoubleDouble(double A, double B)
 	return A / B;
 }
 
+KISMET_MATH_FORCEINLINE
+bool UKismetMathLibrary::Less_DoubleDouble(double A, double B)
+{
+	return A < B;
+}
+
+KISMET_MATH_FORCEINLINE
+bool UKismetMathLibrary::Greater_DoubleDouble(double A, double B)
+{
+	return A > B;
+}
+
+KISMET_MATH_FORCEINLINE
+bool UKismetMathLibrary::LessEqual_DoubleDouble(double A, double B)
+{
+	return A <= B;
+}
+
+KISMET_MATH_FORCEINLINE
+bool UKismetMathLibrary::GreaterEqual_DoubleDouble(double A, double B)
+{
+	return A >= B;
+}
+
+KISMET_MATH_FORCEINLINE
+bool UKismetMathLibrary::EqualEqual_DoubleDouble(double A, double B)
+{
+	return A == B;
+}
+
 KISMET_MATH_FORCEINLINE
 float UKismetMathLibrary::Add_FloatFloat(float A, float B)
 {
