ShaderCompiler timing improvements:
- When logging the time it took to compile a shader job, include the worker's reported time as well.
- When logging timing info, report how long a worker was idle.

diff --git a/Engine/Source/Runtime/Engine/Private/ShaderCompiler/ShaderCompiler.cpp b/Engine/Source/Runtime/Engine/Private/ShaderCompiler/ShaderCompiler.cpp
--- a/Engine/Source/Runtime/Engine/Private/ShaderCompiler/ShaderCompiler.cpp
+++ b/Engine/Source/Runtime/Engine/Private/ShaderCompiler/ShaderCompiler.cpp
@@ -1114,6 +1114,9 @@ struct FShaderCompileWorkerInfo
 	/** Time at which the worker started the most recent batch of tasks. */
 	double StartTime;
 
+	/** Time at which the worker finished its most recent batch of tasks. */
+	double FinishTime = 0.0;
+
 	/** Jobs that this worker is responsible for compiling. */
 	TArray<TSharedRef<FShaderCommonCompileJob, ESPMode::ThreadSafe>> QueuedJobs;
 
@@ -1287,6 +1290,12 @@ int32 FShaderCompileThreadRunnable::PullTasksFromQueue()
 					CurrentWorkerInfo.StartTime = FPlatformTime::Seconds();
 					NumActiveThreads++;
 					Manager->CompileQueue.RemoveAt(0, JobIndex);
+
+					if (Manager->bLogJobCompletionTimes && CurrentWorkerInfo.FinishTime > 0.0)
+					{
+						const double WorkerIdleTime = CurrentWorkerInfo.StartTime - CurrentWorkerInfo.FinishTime;
+						UE_LOG(LogShaders, Display, TEXT("  Worker was idle for %fs"), WorkerIdleTime);
+					}
 				}
 			}
 			else
@@ -1321,7 +1330,7 @@ int32 FShaderCompileThreadRunnable::PullTasksFromQueue()
 							const FShaderCommonCompileJob& Job = *CurrentWorkerInfo.QueuedJobs[JobIndex];
 							if (auto* SingleJob = Job.GetSingleShaderJob())
 							{
-								JobNames += FString(SingleJob->ShaderType->GetName()) + TEXT(" Instructions = ") + FString::FromInt(SingleJob->Output.NumInstructions);
+								JobNames += FString::Printf(TEXT("%s [Instructions=%d WorkerTime=%.3fs]"), SingleJob->ShaderType->GetName(), SingleJob->Output.NumInstructions, SingleJob->Output.CompileTime);
 							}
 							else
 							{
@@ -1338,12 +1347,13 @@ int32 FShaderCompileThreadRunnable::PullTasksFromQueue()
 							}
 						}
 
-						UE_LOG(LogShaders, Display, TEXT("Finished batch of %u jobs in %.3fs, %s"), CurrentWorkerInfo.QueuedJobs.Num(), ElapsedTime, *JobNames);
+						UE_LOG(LogShaders, Display, TEXT("Worker (%d/%d) finished batch of %u jobs in %.3fs, %s"), WorkerIndex + 1, WorkerInfos.Num(), CurrentWorkerInfo.QueuedJobs.Num(), ElapsedTime, *JobNames);
 					}
 
 					// Using atomics to update NumOutstandingJobs since it is read outside of the critical section
 					FPlatformAtomics::InterlockedAdd(&Manager->NumOutstandingJobs, -CurrentWorkerInfo.QueuedJobs.Num());
 
+					CurrentWorkerInfo.FinishTime = FPlatformTime::Seconds();
 					CurrentWorkerInfo.bComplete = false;
 					CurrentWorkerInfo.QueuedJobs.Empty();
 				}
