WorldPartition: Partial backout of CL14336146 (removes disabled optimization)

diff --git a/Engine/Source/Runtime/Engine/Private/WorldPartition/WorldPartition.cpp b/Engine/Source/Runtime/Engine/Private/WorldPartition/WorldPartition.cpp
--- a/Engine/Source/Runtime/Engine/Private/WorldPartition/WorldPartition.cpp
+++ b/Engine/Source/Runtime/Engine/Private/WorldPartition/WorldPartition.cpp
@@ -40,8 +40,6 @@
 
 DEFINE_LOG_CATEGORY(LogWorldPartition);
 
-PRAGMA_DISABLE_OPTIMIZATION
-
 #define LOCTEXT_NAMESPACE "WorldPartitionEditor"
 
 #if WITH_EDITOR
@@ -635,7 +633,7 @@ void UWorldPartition::AddToClusters(const FWorldPartitionActorDesc* ActorDesc)
 			const FWorldPartitionActorDesc* ReferenceActorDesc = GetActorDesc(ReferenceGuid);
 
 			// Don't include references to editor-only actors
-			if (ReferenceActorDesc && !ReferenceActorDesc->GetActorIsEditorOnly())
+			if (!ReferenceActorDesc->GetActorIsEditorOnly())
 			{
 				FActorCluster* ReferenceCluster = ActorToActorCluster.FindRef(ReferenceGuid);
 				if (ReferenceCluster)
@@ -680,10 +678,7 @@ void UWorldPartition::RemoveFromClusters(const FWorldPartitionActorDesc* ActorDe
 
 	for (const FGuid& Guid : ActorCluster->Actors)
 	{
-		if (FWorldPartitionActorDesc* ClusterActorDesc = GetActorDesc(Guid))
-		{
-			AddToClusters(ClusterActorDesc);
-		}
+		AddToClusters(GetActorDesc(Guid));
 	}
 
 	delete ActorCluster;
@@ -1314,5 +1309,4 @@ FBox UWorldPartition::GetWorldBounds() const
 }
 #endif
 
-#undef LOCTEXT_NAMESPACE
-PRAGMA_ENABLE_OPTIMIZATION
\ No newline at end of file
+#undef LOCTEXT_NAMESPACE
\ No newline at end of file
