Added check in APlayerController::HasClientLoadedCurrentWorld to ensure client has completed seamless travel.

diff --git a/Engine/Source/Runtime/Engine/Private/PlayerController.cpp b/Engine/Source/Runtime/Engine/Private/PlayerController.cpp
--- a/Engine/Source/Runtime/Engine/Private/PlayerController.cpp
+++ b/Engine/Source/Runtime/Engine/Private/PlayerController.cpp
@@ -643,8 +643,16 @@ bool APlayerController::HasClientLoadedCurrentWorld()
 	if (Connection != NULL)
 	{
 		// NOTE: To prevent exploits, child connections must not use the parent connections ClientWorldPackageName value at all.
-
-		return (Connection->GetClientWorldPackageName() == GetWorld()->GetOutermost()->GetFName());
+		bool bInCorrectWorld = (Connection->GetClientWorldPackageName() == GetWorld()->GetOutermost()->GetFName());
+		if (SeamlessTravelCount > 0)
+		{
+			// In the case where seamless travel has occurred, make sure the client has actually completed the travel
+			return bInCorrectWorld && (LastCompletedSeamlessTravelCount == SeamlessTravelCount);
+		}
+		else
+		{
+			return bInCorrectWorld;
+		}
 	}
 	else
 	{
