Revert spotlight shadow map resolution estimation to pre-reverb demo (now matches UE4 again)
- reduce appearence changes in existing content & test cases.

diff --git a/Engine/Source/Runtime/Engine/Private/Components/SpotLightComponent.cpp b/Engine/Source/Runtime/Engine/Private/Components/SpotLightComponent.cpp
--- a/Engine/Source/Runtime/Engine/Private/Components/SpotLightComponent.cpp
+++ b/Engine/Source/Runtime/Engine/Private/Components/SpotLightComponent.cpp
@@ -143,14 +143,17 @@ public:
 
 	virtual float GetEffectiveScreenRadius(const FViewMatrices& ShadowViewMatrices) const override
 	{
-		// Heuristic: use the radius of the inscribed sphere at the closest point on the spotlight axis to the viewer as the light's effective screen radius
+		// Heuristic: use the radius of the inscribed sphere at the cone's end as the light's effective screen radius
 		// We do so because we do not want to use the light's radius directly, which will make us overestimate the shadow map resolution greatly for a spot light
 
-		float ViewOriginProjectionOntoSpotlightAxis = FVector::DotProduct(ShadowViewMatrices.GetViewOrigin() - GetOrigin(), GetDirection());
-		ViewOriginProjectionOntoSpotlightAxis = FMath::Clamp(ViewOriginProjectionOntoSpotlightAxis, 0.0f, GetRadius());
-		const FVector InscribedSpherePosition = GetOrigin() + ViewOriginProjectionOntoSpotlightAxis * GetDirection();
+		// In the correct form,
+		//   InscribedSpherePosition = GetOrigin() + GetDirection() * GetRadius() / CosOuterCone
+		//   InscribedSphereRadius = GetRadius() / SinOuterCone
+		// Do it incorrectly to avoid division which is more expensive and risks division by zero
+		const FVector InscribedSpherePosition = GetOrigin() + GetDirection() * GetRadius() * CosOuterCone;
+		const float InscribedSphereRadius = GetRadius() * SinOuterCone;
+
 		const float SphereDistanceFromViewOrigin = (InscribedSpherePosition - ShadowViewMatrices.GetViewOrigin()).Size();
-		const float InscribedSphereRadius = FMath::Max(ViewOriginProjectionOntoSpotlightAxis, GetRadius() / 10.0f) / InvTanOuterCone;
 
 		return ShadowViewMatrices.GetScreenScale() * InscribedSphereRadius / FMath::Max(SphereDistanceFromViewOrigin, 1.0f);
 	}
