WorldPartition: Loading/Unloading cell no longer supports undo/redo until we can fix all related issues. Unloading cell asks the user to save or discard changes

diff --git a/Engine/Source/Editor/UnrealEd/Private/FileHelpers.cpp b/Engine/Source/Editor/UnrealEd/Private/FileHelpers.cpp
--- a/Engine/Source/Editor/UnrealEd/Private/FileHelpers.cpp
+++ b/Engine/Source/Editor/UnrealEd/Private/FileHelpers.cpp
@@ -3734,6 +3734,8 @@ FEditorFileUtils::EPromptReturnCode InternalPromptForCheckoutAndSave(const TArra
  * @param		PackagesToSave				The list of packages to save.  Both map and content packages are supported 
  * @param		bCheckDirty					If true, only packages that are dirty in PackagesToSave will be saved	
  * @param		bPromptToSave				If true the user will be prompted with a list of packages to save, otherwise all passed in packages are saved
+ * @param		Title						If bPromptToSave true provides a dialog title
+ * @param		Message						If bPromptToSave true provides a dialog message
  * @param		OutFailedPackages			[out] If specified, will be filled in with all of the packages that failed to save successfully
  * @param		bAlreadyCheckedOut			If true, the user will not be prompted with the source control dialog
  * @param		bCanBeDeclined				If true, offer a "Don't Save" option in addition to "Cancel", which will not result in a cancellation return code.
@@ -3744,7 +3746,7 @@ FEditorFileUtils::EPromptReturnCode InternalPromptForCheckoutAndSave(const TArra
  *				Save" option on the dialog, the return code will indicate the user has declined out of the prompt. This way calling code can distinguish between a decline and a cancel
  *				and then proceed as planned, or abort its operation accordingly.
  */
-FEditorFileUtils::EPromptReturnCode FEditorFileUtils::PromptForCheckoutAndSave( const TArray<UPackage*>& InPackages, bool bCheckDirty, bool bPromptToSave, TArray<UPackage*>* OutFailedPackages, bool bAlreadyCheckedOut, bool bCanBeDeclined )
+FEditorFileUtils::EPromptReturnCode FEditorFileUtils::PromptForCheckoutAndSave( const TArray<UPackage*>& InPackages, bool bCheckDirty, bool bPromptToSave, const FText& Title, const FText& Message, TArray<UPackage*>* OutFailedPackages, bool bAlreadyCheckedOut, bool bCanBeDeclined)
 {
 	// Check for re-entrance into this function
 	if ( bIsPromptingForCheckoutAndSave )
@@ -3796,7 +3798,7 @@ FEditorFileUtils::EPromptReturnCode FEditorFileUtils::PromptForCheckoutAndSave(
 	{
 		// Set up the save package dialog
 		FPackagesDialogModule& PackagesDialogModule = FModuleManager::LoadModuleChecked<FPackagesDialogModule>( TEXT("PackagesDialog") );
-		PackagesDialogModule.CreatePackagesDialog(NSLOCTEXT("PackagesDialogModule", "PackagesDialogTitle", "Save Content"), NSLOCTEXT("PackagesDialogModule", "PackagesDialogMessage", "Select content to save."));
+		PackagesDialogModule.CreatePackagesDialog(Title, Message);
 		PackagesDialogModule.AddButton(DRT_Save, NSLOCTEXT("PackagesDialogModule", "SaveSelectedButton", "Save Selected"), NSLOCTEXT("PackagesDialogModule", "SaveSelectedButtonTip", "Attempt to save the selected content"));
 		if (bCanBeDeclined)
 		{
@@ -3970,6 +3972,19 @@ FEditorFileUtils::EPromptReturnCode FEditorFileUtils::PromptForCheckoutAndSave(
 	return ReturnResponse;
 }
 
+FEditorFileUtils::EPromptReturnCode FEditorFileUtils::PromptForCheckoutAndSave(const TArray<UPackage*>& InPackages, bool bCheckDirty, bool bPromptToSave, TArray<UPackage*>* OutFailedPackages, bool bAlreadyCheckedOut, bool bCanBeDeclined)
+{
+	return PromptForCheckoutAndSave(
+		InPackages,
+		bCheckDirty,
+		bPromptToSave,
+		NSLOCTEXT("PackagesDialogModule", "PackagesDialogTitle", "Save Content"),
+		NSLOCTEXT("PackagesDialogModule", "PackagesDialogMessage", "Select content to save."),
+		OutFailedPackages,
+		bAlreadyCheckedOut,
+		bCanBeDeclined);
+}
+
 bool FEditorFileUtils::SaveWorlds(UWorld* InWorld, const FString& RootPath, const TCHAR* Prefix, TArray<FString>& OutFilenames)
 {
 	const FScopedBusyCursor BusyCursor;
diff --git a/Engine/Source/Editor/UnrealEd/Public/FileHelpers.h b/Engine/Source/Editor/UnrealEd/Public/FileHelpers.h
--- a/Engine/Source/Editor/UnrealEd/Public/FileHelpers.h
+++ b/Engine/Source/Editor/UnrealEd/Public/FileHelpers.h
@@ -409,6 +409,8 @@ public:
 	 * @param		PackagesToSave				The list of packages to save.  Both map and content packages are supported 
 	 * @param		bCheckDirty					If true, only packages that are dirty in PackagesToSave will be saved	
 	 * @param		bPromptToSave				If true the user will be prompted with a list of packages to save, otherwise all passed in packages are saved
+	 * @param		Title						If bPromptToSave true provides a dialog title
+	 * @param		Message						If bPromptToSave true provides a dialog message
 	 * @param		OutFailedPackages			[out] If specified, will be filled in with all of the packages that failed to save successfully
 	 * @param		bAlreadyCheckedOut			If true, the user will not be prompted with the source control dialog
 	 * @param		bCanBeDeclined				If true, offer a "Don't Save" option in addition to "Cancel", which will not result in a cancellation return code.
@@ -419,7 +421,9 @@ public:
 	 *				Save" option on the dialog, the return code will indicate the user has declined out of the prompt. This way calling code can distinguish between a decline and a cancel
 	 *				and then proceed as planned, or abort its operation accordingly.
 	 */
+	UNREALED_API static EPromptReturnCode PromptForCheckoutAndSave(const TArray<UPackage*>& PackagesToSave, bool bCheckDirty, bool bPromptToSave, const FText& Title, const FText& Message, TArray<UPackage*>* OutFailedPackages = NULL, bool bAlreadyCheckedOut = false, bool bCanBeDeclined = true);
 	UNREALED_API static EPromptReturnCode PromptForCheckoutAndSave( const TArray<UPackage*>& PackagesToSave, bool bCheckDirty, bool bPromptToSave, TArray<UPackage*>* OutFailedPackages = NULL, bool bAlreadyCheckedOut = false, bool bCanBeDeclined = true );
+	
 
 	////////////////////////////////////////////////////////////////////////////
 	// Import/Export
diff --git a/Engine/Source/Runtime/Engine/Private/Level.cpp b/Engine/Source/Runtime/Engine/Private/Level.cpp
--- a/Engine/Source/Runtime/Engine/Private/Level.cpp
+++ b/Engine/Source/Runtime/Engine/Private/Level.cpp
@@ -610,6 +610,7 @@ bool ULevel::IsNetActor(const AActor* Actor)
 void ULevel::AddLoadedActor(AActor* Actor)
 {
 	check(Actor->GetLevel() == this);
+	check(!Actor->IsPendingKill());
 
 	Actors.Add(Actor);
 	ActorsForGC.Add(Actor);
@@ -642,6 +643,9 @@ void ULevel::RemoveLoadedActor(AActor* Actor)
 	ActorsForGC.Remove(Actor);
 
 	OnLoadedActorRemovedFromLevelEvent.Broadcast(*Actor);
+
+	Actor->MarkPendingKill();
+	Actor->MarkComponentsAsPendingKill();
 }
 #endif
 
diff --git a/Engine/Source/Runtime/Engine/Private/WorldPartition/WorldPartition.cpp b/Engine/Source/Runtime/Engine/Private/WorldPartition/WorldPartition.cpp
--- a/Engine/Source/Runtime/Engine/Private/WorldPartition/WorldPartition.cpp
+++ b/Engine/Source/Runtime/Engine/Private/WorldPartition/WorldPartition.cpp
@@ -128,12 +128,6 @@ UWorldPartition::UWorldPartition(const FObjectInitializer& ObjectInitializer)
 }
 
 #if WITH_EDITOR
-void UWorldPartition::OnPostGarbageCollect()
-{
-	// Clean-up ModifiedActors (deleted/garbage collected actors)
-	ModifiedActors.Remove(nullptr);
-}
-
 void UWorldPartition::OnPreBeginPIE(bool bStartSimulate)
 {
 	check(IsMainWorldPartition());
@@ -161,37 +155,6 @@ void UWorldPartition::OnObjectPropertyChanged(UObject* Object, FPropertyChangedE
 		}
 	}
 }
-
-void UWorldPartition::OnObjectModified(UObject* Object)
-{
-	if (Object->GetPackage()->IsDirty())
-	{
-		// Don't rely on undo stack and keep references on modified actor to avoid loosing changes when unloading editor cells
-		if (AActor* Actor = Cast<AActor>(Object))
-		{
-			if (FWorldPartitionActorDesc* ActorDesc = GetActorDesc(Actor->GetActorGuid()))
-			{
-				UPackage* Package = Actor->GetOutermost();
-				if (Package && Package->IsDirty())
-				{
-					ModifiedActors.Add(Actor);
-				}
-			}
-		}
-	}
-}
-
-void UWorldPartition::OnObjectSaved(UObject* Object)
-{
-	if (FUnrealEdMisc::Get().GetAutosaveState() == FUnrealEdMisc::EAutosaveState::Inactive)
-	{
-		// Once modified actor is saved, we can remove reference on it
-		if (AActor* Actor = Cast<AActor>(Object))
-		{
-			ModifiedActors.Remove(Actor);
-		}
-	}
-}
 #endif
 
 void UWorldPartition::Initialize(UWorld* InWorld, const FTransform& InTransform)
@@ -454,9 +417,6 @@ void UWorldPartition::RegisterDelegates()
 		GEditor->OnLevelActorDeleted().AddUObject(this, &UWorldPartition::OnActorDeleted);
 
 		FCoreUObjectDelegates::OnObjectPropertyChanged.AddUObject(this, &UWorldPartition::OnObjectPropertyChanged);
-		FCoreUObjectDelegates::OnObjectModified.AddUObject(this, &UWorldPartition::OnObjectModified);
-		FCoreUObjectDelegates::OnObjectSaved.AddUObject(this, &UWorldPartition::OnObjectSaved);
-		FCoreUObjectDelegates::GetPostGarbageCollect().AddUObject(this, &UWorldPartition::OnPostGarbageCollect);
 
 		FEditorDelegates::PreBeginPIE.AddUObject(this, &UWorldPartition::OnPreBeginPIE);
 		FEditorDelegates::EndPIE.AddUObject(this, &UWorldPartition::OnEndPIE);
@@ -500,9 +460,6 @@ void UWorldPartition::UnregisterDelegates()
 		GEditor->OnLevelActorDeleted().RemoveAll(this);
 
 		FCoreUObjectDelegates::OnObjectPropertyChanged.RemoveAll(this);
-		FCoreUObjectDelegates::OnObjectModified.RemoveAll(this);
-		FCoreUObjectDelegates::OnObjectSaved.RemoveAll(this);
-		FCoreUObjectDelegates::GetPostGarbageCollect().RemoveAll(this);
 	}
 }
 
@@ -782,8 +739,6 @@ void UWorldPartition::LoadEditorCells(const TArray<UWorldPartitionEditorCell*>&
 {
 	FWorldPartionCellUpdateContext CellUpdateContext(this);
 
-	FScopedTransaction Transaction(LOCTEXT("UWorldPartition_LoadCells", "Load Cells"));
-
 	int32 NumActorsToLoad = Algo::TransformAccumulate(CellsToLoad, [](UWorldPartitionEditorCell* Cell) { return Cell->Actors.Num() - Cell->LoadedActors.Num();}, 0);
 
 	FScopedSlowTask SlowTask(NumActorsToLoad, LOCTEXT("LoadingCells", "Loading cells..."));
@@ -799,11 +754,57 @@ void UWorldPartition::LoadEditorCells(const TArray<UWorldPartitionEditorCell*>&
 void UWorldPartition::UnloadEditorCells(const TArray<UWorldPartitionEditorCell*>& CellsToUnload)
 {
 	FWorldPartionCellUpdateContext CellUpdateContext(this);
+		
+	TSet<UPackage*> ModifiedPackages;
+	TMap<FWorldPartitionActorDesc*, int32> UnloadCount;
+	
+	int32 NumActorsToUnload = Algo::TransformAccumulate(CellsToUnload, [&UnloadCount](UWorldPartitionEditorCell* Cell)
+	{ 
+		for (FWorldPartitionActorDesc* ActorDesc : Cell->LoadedActors)
+		{
+			UnloadCount.FindOrAdd(ActorDesc, 0)++;
+		}
+						
+		return Cell->LoadedActors.Num(); 
+	}, 0);
+		
+	for (const TPair<FWorldPartitionActorDesc*, int32> Pair : UnloadCount)
+	{
+		FWorldPartitionActorDesc* ActorDesc = Pair.Key;
+		// Only prompt if the actor will get unloaded by the unloading cells
+		if (ActorDesc->GetLoadedRefCount() == Pair.Value)
+		{
+			AActor* LoadedActor = ActorDesc->GetActor();
+			check(LoadedActor);
+			UPackage* ActorPackage = LoadedActor->GetExternalPackage();
+			if (ActorPackage && ActorPackage->IsDirty())
+			{
+				ModifiedPackages.Add(ActorPackage);
+			}
+		}
+	}
+		
+	// Make sure we save modified actor packages before unloading
+	if (ModifiedPackages.Num())
+	{
+		const bool bCheckDirty = false;
+		const bool bAlreadyCheckedOut = false;
+		const bool bCanBeDeclined = true;
+		const bool bPromptToSave = true;
+		const FText Title = LOCTEXT("SaveActorsTitle", "Save Actor(s)");
+		const FText Message = LOCTEXT("SaveActorsMessage", "Save Actor(s) before unloading them.");
+
+		FEditorFileUtils::EPromptReturnCode RetCode = FEditorFileUtils::PromptForCheckoutAndSave(ModifiedPackages.Array(), bCheckDirty, bPromptToSave, Title, Message, nullptr, bAlreadyCheckedOut, bCanBeDeclined);
+		if (RetCode == FEditorFileUtils::PR_Cancelled)
+		{
+			return;
+		}
 
-	FScopedTransaction Transaction(LOCTEXT("UWorldPartition_UnloadCells", "Unload Cells"));
-
-	int32 NumActorsToUnload = Algo::TransformAccumulate(CellsToUnload, [](UWorldPartitionEditorCell* Cell) { return Cell->LoadedActors.Num(); }, 0);
+		check(RetCode != FEditorFileUtils::PR_Failure);
+	}
 
+	GEditor->SelectNone(true, true);
+	
 	FScopedSlowTask SlowTask(NumActorsToUnload, LOCTEXT("UnloadingCells", "Unloading cells..."));
 	SlowTask.MakeDialog();
 
@@ -812,6 +813,13 @@ void UWorldPartition::UnloadEditorCells(const TArray<UWorldPartitionEditorCell*>
 		SlowTask.EnterProgressFrame(Cell->LoadedActors.Num());
 		UpdateLoadingEditorCell(Cell, false);
 	}
+		
+	GEditor->ResetTransaction(LOCTEXT("UnloadingEditorCellsResetTrans", "Unloading Cells"));
+	
+	for (UPackage* ModifiedPackage : ModifiedPackages)
+	{
+		ModifiedPackage->ClearDirtyFlag();
+	}
 }
 
 // Loads actors in Editor cell
@@ -1111,7 +1119,9 @@ void UWorldPartition::OnActorOuterChanged(AActor* InActor, UObject* InOldOuter)
 	{
 		if (FWorldPartitionActorDesc* ActorDesc = GetActorDesc(InActor->GetActorGuid()))
 		{
-			RemoveFromPartition(ActorDesc);
+			const bool bRemoveDescriptor = true;
+			const bool bUnloadDescriptor = false;
+			RemoveFromPartition(ActorDesc, bRemoveDescriptor, bUnloadDescriptor);
 		}
 	}
 }
@@ -1140,7 +1150,7 @@ void UWorldPartition::AddToPartition(FWorldPartitionActorDesc* ActorDesc)
 	HashActorDesc(ActorDesc);	
 }
 
-void UWorldPartition::RemoveFromPartition(FWorldPartitionActorDesc* ActorDesc, bool bRemoveDescriptorFromArray)
+void UWorldPartition::RemoveFromPartition(FWorldPartitionActorDesc* ActorDesc, bool bRemoveDescriptorFromArray, bool bUnloadRemovedDescriptor)
 {
 	// Unhash this actor from the editor hash
 	UnhashActorDesc(ActorDesc);
@@ -1150,7 +1160,12 @@ void UWorldPartition::RemoveFromPartition(FWorldPartitionActorDesc* ActorDesc, b
 	// Remove this actor descriptor
 	if (bRemoveDescriptorFromArray)
 	{
-		Actors.Remove(ActorDesc->GetGuid());
+		TUniquePtr<FWorldPartitionActorDesc> RemovedActorDesc = Actors.FindAndRemoveChecked(ActorDesc->GetGuid());
+		check(RemovedActorDesc->GetLoadedRefCount() == 0);
+		if (bUnloadRemovedDescriptor)
+		{
+			RemovedActorDesc->Unload();
+		}
 	}
 }
 
diff --git a/Engine/Source/Runtime/Engine/Private/WorldPartition/WorldPartitionEditorCell.cpp b/Engine/Source/Runtime/Engine/Private/WorldPartition/WorldPartitionEditorCell.cpp
--- a/Engine/Source/Runtime/Engine/Private/WorldPartition/WorldPartitionEditorCell.cpp
+++ b/Engine/Source/Runtime/Engine/Private/WorldPartition/WorldPartitionEditorCell.cpp
@@ -58,24 +58,6 @@ void UWorldPartitionEditorCell::RemoveActor(FWorldPartitionActorDesc* InActorDes
 		UE_LOG(LogWorldPartition, Verbose, TEXT(" ==> Unreferenced loaded actor %s(%d) [UWorldPartitionEditorCell::RemoveActor]"), *InActorDesc->GetActor()->GetFullName(), ActorRefCount);
 	}
 }
-
-void UWorldPartitionEditorCell::PreEditUndo()
-{
-	bWasLoaded = bLoaded;
-
-	Super::PreEditUndo();
-}
-
-void UWorldPartitionEditorCell::PostEditUndo()
-{
-	Super::PostEditUndo();
-
-	if (bWasLoaded != bLoaded)
-	{
-		bLoaded = bWasLoaded;
-		GetWorldPartition()->UpdateLoadingEditorCell(this, !bLoaded);
-	}
-}
 #endif
 
 UWorldPartitionEditorCell::UWorldPartitionEditorCell(const FObjectInitializer& ObjectInitializer)
@@ -83,6 +65,5 @@ UWorldPartitionEditorCell::UWorldPartitionEditorCell(const FObjectInitializer& O
 #if WITH_EDITOR
 	, Bounds(ForceInitToZero)
 	, bLoaded(false)
-	, bWasLoaded(false)
 #endif
 {}
diff --git a/Engine/Source/Runtime/Engine/Private/WorldPartition/WorldPartitionEditorSpatialHash.cpp b/Engine/Source/Runtime/Engine/Private/WorldPartition/WorldPartitionEditorSpatialHash.cpp
--- a/Engine/Source/Runtime/Engine/Private/WorldPartition/WorldPartitionEditorSpatialHash.cpp
+++ b/Engine/Source/Runtime/Engine/Private/WorldPartition/WorldPartitionEditorSpatialHash.cpp
@@ -25,7 +25,7 @@ void UWorldPartitionEditorSpatialHash::Initialize()
 {
 	check(!AlwaysLoadedCell);
 
-	AlwaysLoadedCell = NewObject<UWorldPartitionEditorCell>(this, TEXT("AlwaysLoadedCell"), RF_Transactional | RF_Transient);
+	AlwaysLoadedCell = NewObject<UWorldPartitionEditorCell>(this, TEXT("AlwaysLoadedCell"), RF_Transient);
 	AlwaysLoadedCell->Bounds.Init();
 }
 
@@ -153,7 +153,6 @@ void UWorldPartitionEditorSpatialHash::HashActor(FWorldPartitionActorDesc* InAct
 			else
 			{
 				EditorCell = NewObject<UWorldPartitionEditorCell>(this, *FString::Printf(TEXT("EditorCell_X%lld_Y%lld_Z%lld_L%d"), CellCoord.X, CellCoord.Y, CellCoord.Z, CellCoord.Level), RF_Transient);
-				EditorCell->SetFlags(RF_Transactional);
 				EditorCell->Bounds = GetCellBounds(CellCoord);
 
 				Cells.Add(EditorCell);
diff --git a/Engine/Source/Runtime/Engine/Public/WorldPartition/WorldPartition.h b/Engine/Source/Runtime/Engine/Public/WorldPartition/WorldPartition.h
--- a/Engine/Source/Runtime/Engine/Public/WorldPartition/WorldPartition.h
+++ b/Engine/Source/Runtime/Engine/Public/WorldPartition/WorldPartition.h
@@ -65,9 +65,6 @@ private:
 	static FWorldPartitionChangedEvent WorldPartitionChangedEvent;
 
 	void OnObjectPropertyChanged(UObject* Object, FPropertyChangedEvent& InPropertyChangedEvent);
-	void OnObjectModified(UObject* Object);
-	void OnObjectSaved(UObject* Object);
-	void OnPostGarbageCollect();
 
 	static FWorldPartitionActorDescFactory* GetActorDescFactory(TSubclassOf<AActor> Class);
 	static FWorldPartitionActorDescFactory* GetActorDescFactory(const AActor* Actor);
@@ -213,9 +210,6 @@ private:
 	UPROPERTY()
 	bool bIsPreCooked;
 
-	UPROPERTY(Transient)
-	TSet<AActor*> ModifiedActors;
-
 	FLinkerInstancingContext InstancingContext;
 #endif
 
@@ -243,7 +237,7 @@ private:
 	void RemoveFromClusters(const FWorldPartitionActorDesc* ActorDesc);
 
 	void AddToPartition(FWorldPartitionActorDesc* ActorDesc);
-	void RemoveFromPartition(FWorldPartitionActorDesc* ActorDesc, bool bRemoveDescriptorFromArray = true);
+	void RemoveFromPartition(FWorldPartitionActorDesc* ActorDesc, bool bRemoveDescriptorFromArray = true, bool bUnloadRemovedDescriptor = true);
 
 	TMap<FGuid, UWorldPartition::FActorCluster*> ActorToActorCluster;
 	TSet<UWorldPartition::FActorCluster*> ActorClustersSet;
diff --git a/Engine/Source/Runtime/Engine/Public/WorldPartition/WorldPartitionEditorCell.h b/Engine/Source/Runtime/Engine/Public/WorldPartition/WorldPartitionEditorCell.h
--- a/Engine/Source/Runtime/Engine/Public/WorldPartition/WorldPartitionEditorCell.h
+++ b/Engine/Source/Runtime/Engine/Public/WorldPartition/WorldPartitionEditorCell.h
@@ -25,17 +25,11 @@ public:
 	void AddActor(FWorldPartitionActorDesc* InActorDesc);
 	void RemoveActor(FWorldPartitionActorDesc* InActorDesc);
 
-	virtual void PreEditUndo() override;
-	virtual void PostEditUndo() override;
-
 	FBox						Bounds;
 
 	/** Tells if the cell was manually loaded in the editor */
 	bool						bLoaded : 1;
 
-	/** Tells if the cell loaded state changed during transactions */
-	bool						bWasLoaded : 1;
-
 	TSet<FWorldPartitionActorDesc*> Actors;
 	TSet<FWorldPartitionActorDesc*> LoadedActors;
 #endif
