Actor metadata is now saved in binary format in a single asset registry section (encoded in Base64).
This avoids the floating point truncation we had before when exporting vectors and transforms as text.
The underlying metadata archive supports object custom version for future changes to the stream.

diff --git a/Engine/Source/Runtime/Engine/Private/Actor.cpp b/Engine/Source/Runtime/Engine/Private/Actor.cpp
--- a/Engine/Source/Runtime/Engine/Private/Actor.cpp
+++ b/Engine/Source/Runtime/Engine/Private/Actor.cpp
@@ -522,12 +522,19 @@ void AActor::GetAssetRegistryTags(TArray<FAssetRegistryTag>& OutTags) const
 			TUniquePtr<FWorldPartitionActorDesc> ActorDesc(WorldPartition->GetActorDescFactory(this)->Create());
 			ActorDesc->Init(this);
 
-			FActorMetaDataWriter ActorMetaDataSerializer(OutTags);
-
-			FString LevelPackageName = GetOuter()->GetOutermost()->GetName();
-			ActorMetaDataSerializer.Serialize(TEXT("LevelPackage"), LevelPackageName);
-
-			ActorDesc->SerializeMetaData(&ActorMetaDataSerializer);
+			const FString LevelPackageName = GetOuter()->GetOutermost()->GetName();
+			static FName NAME_LevelPackage(TEXT("LevelPackage"));
+			OutTags.Add(UObject::FAssetRegistryTag(NAME_LevelPackage, LevelPackageName, UObject::FAssetRegistryTag::TT_Hidden));
+
+			const FString ActorMetaDataClass = GetParentNativeClass(GetClass())->GetName();
+			static FName NAME_ActorMetaDataClass(TEXT("ActorMetaDataClass"));
+			OutTags.Add(UObject::FAssetRegistryTag(NAME_ActorMetaDataClass, ActorMetaDataClass, UObject::FAssetRegistryTag::TT_Hidden));
+
+			TArray<uint8> SerializedData;
+			ActorDesc->SerializeTo(SerializedData);
+			const FString ActorMetaData = FBase64::Encode(SerializedData);
+			static FName NAME_ActorMetaData(TEXT("ActorMetaData"));
+			OutTags.Add(UObject::FAssetRegistryTag(NAME_ActorMetaData, ActorMetaData, UObject::FAssetRegistryTag::TT_Hidden));
 		}
 	}
 }
diff --git a/Engine/Source/Runtime/Engine/Private/WorldPartition/ActorPartition/PartitionActorDesc.cpp b/Engine/Source/Runtime/Engine/Private/WorldPartition/ActorPartition/PartitionActorDesc.cpp
--- a/Engine/Source/Runtime/Engine/Private/WorldPartition/ActorPartition/PartitionActorDesc.cpp
+++ b/Engine/Source/Runtime/Engine/Private/WorldPartition/ActorPartition/PartitionActorDesc.cpp
@@ -6,34 +6,25 @@
 #if WITH_EDITOR
 #include "Engine/Level.h"
 
-bool FPartitionActorDesc::Init(const AActor* InActor)
+void FPartitionActorDesc::InitFrom(const AActor* InActor)
 {
-	if (FWorldPartitionActorDesc::Init(InActor))
-	{
-		const APartitionActor* PartitionActor = CastChecked<APartitionActor>(InActor);
+	FWorldPartitionActorDesc::InitFrom(InActor);
 
-		GridSize = PartitionActor->GridSize;
+	const APartitionActor* PartitionActor = CastChecked<APartitionActor>(InActor);
 
-		const FVector ActorLocation = InActor->GetActorLocation();
-		GridIndexX = FMath::FloorToInt(ActorLocation.X / GridSize);
-		GridIndexY = FMath::FloorToInt(ActorLocation.Y / GridSize);
-		GridIndexZ = FMath::FloorToInt(ActorLocation.Z / GridSize);
+	GridSize = PartitionActor->GridSize;
 
-		UpdateHash();
-		return true;
-	}
-
-	return false;
+	const FVector ActorLocation = InActor->GetActorLocation();
+	GridIndexX = FMath::FloorToInt(ActorLocation.X / GridSize);
+	GridIndexY = FMath::FloorToInt(ActorLocation.Y / GridSize);
+	GridIndexZ = FMath::FloorToInt(ActorLocation.Z / GridSize);
 }
 
-void FPartitionActorDesc::SerializeMetaData(FActorMetaDataSerializer* Serializer)
+void FPartitionActorDesc::Serialize(FArchive& Ar)
 {
-	FWorldPartitionActorDesc::SerializeMetaData(Serializer);
+	FWorldPartitionActorDesc::Serialize(Ar);
 
-	Serializer->Serialize(TEXT("GridSize"), (int32&)GridSize);
-	Serializer->Serialize(TEXT("GridIndexX"), GridIndexX);
-	Serializer->Serialize(TEXT("GridIndexY"), GridIndexY);
-	Serializer->Serialize(TEXT("GridIndexZ"), GridIndexZ);
+	Ar << GridSize << GridIndexX << GridIndexY << GridIndexZ;
 }
 
 void FPartitionActorDesc::BuildHash(FHashBuilder& HashBuilder)
diff --git a/Engine/Source/Runtime/Engine/Private/WorldPartition/HLOD/HLODActorDesc.cpp b/Engine/Source/Runtime/Engine/Private/WorldPartition/HLOD/HLODActorDesc.cpp
--- a/Engine/Source/Runtime/Engine/Private/WorldPartition/HLOD/HLODActorDesc.cpp
+++ b/Engine/Source/Runtime/Engine/Private/WorldPartition/HLOD/HLODActorDesc.cpp
@@ -8,71 +8,30 @@
 #include "WorldPartition/HLOD/HLODActor.h"
 #include "WorldPartition/HLOD/HLODLayer.h"
 
-bool FHLODActorDesc::Init(const AActor* InActor)
+void FHLODActorDesc::InitFrom(const AActor* InActor)
 {
-	if (FWorldPartitionActorDesc::Init(InActor))
-	{
-		const AWorldPartitionHLOD* HLODActor = CastChecked<AWorldPartitionHLOD>(InActor);
-		SubActors = HLODActor->GetSubActors();
-		HLODLayer = HLODActor->GetHLODLayer();
-
-		UpdateHash();
-		return true;
-	}
+	FWorldPartitionActorDesc::InitFrom(InActor);
 
-	return false;
+	const AWorldPartitionHLOD* HLODActor = CastChecked<AWorldPartitionHLOD>(InActor);
+	SubActors = HLODActor->GetSubActors();
+	HLODLayer = HLODActor->GetHLODLayer();
 }
 
-void FHLODActorDesc::SerializeMetaData(FActorMetaDataSerializer* Serializer)
+void FHLODActorDesc::Serialize(FArchive& Ar)
 {
-	FWorldPartitionActorDesc::SerializeMetaData(Serializer);
-
-	FString SubActorsGUIDsStr;
-	if (Serializer->IsWriting())
-	{
-		if (SubActors.Num())
-		{
-			for (const FGuid& ActorGUID : SubActors)
-			{
-				SubActorsGUIDsStr += ActorGUID.ToString() + TEXT(";");
-			}
-			SubActorsGUIDsStr.RemoveFromEnd(TEXT(";"));
-		}
-	}
-
-	Serializer->Serialize(TEXT("HLODSubActors"), SubActorsGUIDsStr);
+	FWorldPartitionActorDesc::Serialize(Ar);
 
-	if (Serializer->IsReading())
-	{
-		TArray<FString> SubActorsStr;
-		if (SubActorsGUIDsStr.ParseIntoArray(SubActorsStr, TEXT(";")))
-		{
-			Algo::Transform(SubActorsStr, SubActors, [](const FString& ActorGuidStr)
-			{
-				FGuid ActorGuid;
-				if (!FGuid::Parse(ActorGuidStr, ActorGuid))
-				{
-					ActorGuid.Invalidate();
-				}
-				return ActorGuid;
-			});
-
-			Algo::RemoveIf(SubActors, [](const FGuid& InGuid)
-			{
-				return !InGuid.IsValid();
-			});
-		}
-	}
+	Ar << SubActors;
 
 	FString HLODLayerStr;
-	if (Serializer->IsWriting())
+	if (Ar.IsSaving())
 	{
 		HLODLayerStr = HLODLayer.ToString();
 	}
 	
-	Serializer->Serialize(TEXT("HLODSubActors"), HLODLayerStr);
+	Ar << HLODLayerStr;
 	
-	if (Serializer->IsReading())
+	if (Ar.IsLoading())
 	{
 		HLODLayer = HLODLayerStr;
 	}
diff --git a/Engine/Source/Runtime/Engine/Private/WorldPartition/LevelInstance/LevelInstanceActorDesc.cpp b/Engine/Source/Runtime/Engine/Private/WorldPartition/LevelInstance/LevelInstanceActorDesc.cpp
--- a/Engine/Source/Runtime/Engine/Private/WorldPartition/LevelInstance/LevelInstanceActorDesc.cpp
+++ b/Engine/Source/Runtime/Engine/Private/WorldPartition/LevelInstance/LevelInstanceActorDesc.cpp
@@ -6,27 +6,22 @@
 #include "Engine/Level.h"
 #include "LevelInstance/LevelInstanceActor.h"
 
-bool FLevelInstanceActorDesc::Init(const AActor* InActor)
+void FLevelInstanceActorDesc::InitFrom(const AActor* InActor)
 {
-	if (FWorldPartitionActorDesc::Init(InActor))
-	{
-		const ALevelInstance* LevelInstanceActor = CastChecked<ALevelInstance>(InActor);
-		LevelPackage = *LevelInstanceActor->GetWorldAssetPackage();
-		LevelInstanceTransform = FTransform(InActor->GetActorRotation(), InActor->GetActorLocation());
-		UpdateHash();
-		return true;
-	}
-	return false;
+	FWorldPartitionActorDesc::InitFrom(InActor);
+
+	const ALevelInstance* LevelInstanceActor = CastChecked<ALevelInstance>(InActor);
+	LevelPackage = *LevelInstanceActor->GetWorldAssetPackage();
+	LevelInstanceTransform = FTransform(InActor->GetActorRotation(), InActor->GetActorLocation());
 }
 
-void FLevelInstanceActorDesc::SerializeMetaData(FActorMetaDataSerializer* Serializer)
+void FLevelInstanceActorDesc::Serialize(FArchive& Ar)
 {
-	FWorldPartitionActorDesc::SerializeMetaData(Serializer);
+	FWorldPartitionActorDesc::Serialize(Ar);
 
-	Serializer->Serialize(TEXT("LevelInstancePackage"), LevelPackage);
-	Serializer->Serialize(TEXT("LevelInstanceTransform"), LevelInstanceTransform);
+	Ar << LevelPackage << LevelInstanceTransform;
 
-	if (Serializer->IsReading())
+	if (Ar.IsLoading())
 	{
 		if (!LevelPackage.IsNone())
 		{
diff --git a/Engine/Source/Runtime/Engine/Private/WorldPartition/WorldPartition.cpp b/Engine/Source/Runtime/Engine/Private/WorldPartition/WorldPartition.cpp
--- a/Engine/Source/Runtime/Engine/Private/WorldPartition/WorldPartition.cpp
+++ b/Engine/Source/Runtime/Engine/Private/WorldPartition/WorldPartition.cpp
@@ -248,18 +248,20 @@ void UWorldPartition::Initialize(UWorld* InWorld, const FTransform& InTransform)
 
 		for (const FAssetData& Asset : Assets)
 		{
-			FActorMetaDataReader ActorMetaDataSerializer(Asset);
+			FString ActorClass;
+			static FName NAME_ActorMetaDataClass(TEXT("ActorMetaDataClass"));
+			Asset.GetTagValue(NAME_ActorMetaDataClass, ActorClass);
 
-			FName ActorClass;
-			ActorMetaDataSerializer.Serialize(TEXT("ActorClass"), ActorClass);
+			FString ActorMetaDataStr;
+			static FName NAME_ActorMetaData(TEXT("ActorMetaData"));
+			Asset.GetTagValue(NAME_ActorMetaData, ActorMetaDataStr);
 
 			FWorldPartitionActorDescInitData ActorDescInitData;
-			ActorDescInitData.NativeClass = FindObjectChecked<UClass>(ANY_PACKAGE, *ActorClass.ToString(), true);
+			ActorDescInitData.NativeClass = FindObjectChecked<UClass>(ANY_PACKAGE, *ActorClass, true);
 			ActorDescInitData.PackageName = Asset.PackageName;
-			ActorDescInitData.AssetData = Asset;
 			ActorDescInitData.ActorPath = Asset.ObjectPath;
 			ActorDescInitData.Transform = bIsInstanced ? InstanceTransform : FTransform::Identity;
-			ActorDescInitData.Serializer = &ActorMetaDataSerializer;
+			FBase64::Decode(ActorMetaDataStr, ActorDescInitData.SerializedData);			
 
 			if (bIsInstanced)
 			{
@@ -273,16 +275,14 @@ void UWorldPartition::Initialize(UWorld* InWorld, const FTransform& InTransform)
 			}
 
 			TUniquePtr<FWorldPartitionActorDesc> NewActorDesc(GetActorDescFactory(ActorDescInitData.NativeClass)->Create());
+			NewActorDesc->Init(ActorDescInitData);
 
-			if (NewActorDesc->Init(ActorDescInitData))
+			if (bEditorOnly)
 			{
-				if (bEditorOnly)
-				{
-					AllLayersNames.Append(NewActorDesc->GetLayers());
-				}
-
-				Actors.Add(NewActorDesc->GetGuid(), MoveTemp(NewActorDesc));
+				AllLayersNames.Append(NewActorDesc->GetLayers());
 			}
+
+			Actors.Add(NewActorDesc->GetGuid(), MoveTemp(NewActorDesc));
 		}
 
 		if (bEditorOnly)
diff --git a/Engine/Source/Runtime/Engine/Private/WorldPartition/WorldPartitionActorDesc.cpp b/Engine/Source/Runtime/Engine/Private/WorldPartition/WorldPartitionActorDesc.cpp
--- a/Engine/Source/Runtime/Engine/Private/WorldPartition/WorldPartitionActorDesc.cpp
+++ b/Engine/Source/Runtime/Engine/Private/WorldPartition/WorldPartitionActorDesc.cpp
@@ -22,7 +22,13 @@ FWorldPartitionActorDesc::FWorldPartitionActorDesc()
 	, Tag(0)
 {}
 
-bool FWorldPartitionActorDesc::Init(const AActor* InActor)
+void FWorldPartitionActorDesc::Init(const AActor* InActor)
+{
+	InitFrom(InActor);
+	UpdateHash();
+}
+
+void FWorldPartitionActorDesc::InitFrom(const AActor* InActor)
 {	
 	check(InActor->IsPackageExternal());
 
@@ -69,12 +75,9 @@ bool FWorldPartitionActorDesc::Init(const AActor* InActor)
 			References.Add(ActorReference->GetActorGuid());
 		}
 	}
-
-	UpdateHash();
-	return true;
 }
 
-bool FWorldPartitionActorDesc::Init(const FWorldPartitionActorDescInitData& DescData)
+void FWorldPartitionActorDesc::Init(const FWorldPartitionActorDescInitData& DescData)
 {
 	ActorPackage = DescData.PackageName;
 	ActorPath = DescData.ActorPath;
@@ -82,32 +85,52 @@ bool FWorldPartitionActorDesc::Init(const FWorldPartitionActorDescInitData& Desc
 	Class = DescData.NativeClass->GetFName();
 
 	// Serialize actor metadata
-	SerializeMetaData(DescData.Serializer);
+	FMemoryReader MetadataAr(DescData.SerializedData);
 
-	if (!DescData.Serializer->GetHasErrors())
-	{	
-		// Override grid placement by default class value
-		const EActorGridPlacement DefaultGridPlacement = ActorClass->GetDefaultObject<AActor>()->GetDefaultGridPlacement();
-		if (DefaultGridPlacement != EActorGridPlacement::None)
-		{
-			GridPlacement = DefaultGridPlacement;
-		}
+	// Serialize metadata custom versions
+	FCustomVersionContainer CustomVersions;
+	CustomVersions.Serialize(MetadataAr);
+	MetadataAr.SetCustomVersions(CustomVersions);
+	
+	// Serialize metadata payload
+	Serialize(MetadataAr);
 
-		// Transform BoundsLocation and BoundsExtent if necessary
-		if (!DescData.Transform.Equals(FTransform::Identity))
-		{
-			//@todo_ow: This will result in a new BoundsExtent that is larger than it should. To fix this, we would need the Object Oriented BoundingBox of the actor (the BV of the actor without rotation)
-			const FVector BoundsMin = BoundsLocation - BoundsExtent;
-			const FVector BoundsMax = BoundsLocation + BoundsExtent;
-			const FBox NewBounds = FBox(BoundsMin, BoundsMax).TransformBy(DescData.Transform);
-			NewBounds.GetCenterAndExtents(BoundsLocation, BoundsExtent);
-		}
+	// Override grid placement by default class value
+	const EActorGridPlacement DefaultGridPlacement = ActorClass->GetDefaultObject<AActor>()->GetDefaultGridPlacement();
+	if (DefaultGridPlacement != EActorGridPlacement::None)
+	{
+		GridPlacement = DefaultGridPlacement;
+	}
 
-		UpdateHash();
-		return true;
+	// Transform BoundsLocation and BoundsExtent if necessary
+	if (!DescData.Transform.Equals(FTransform::Identity))
+	{
+		//@todo_ow: This will result in a new BoundsExtent that is larger than it should. To fix this, we would need the Object Oriented BoundingBox of the actor (the BV of the actor without rotation)
+		const FVector BoundsMin = BoundsLocation - BoundsExtent;
+		const FVector BoundsMax = BoundsLocation + BoundsExtent;
+		const FBox NewBounds = FBox(BoundsMin, BoundsMax).TransformBy(DescData.Transform);
+		NewBounds.GetCenterAndExtents(BoundsLocation, BoundsExtent);
 	}
 
-	return false;
+	UpdateHash();
+}
+
+void FWorldPartitionActorDesc::SerializeTo(TArray<uint8>& OutData)
+{
+	// Serialize to archive and gather custom versions
+	TArray<uint8> PayloadData;
+	FMemoryWriter PayloadAr(PayloadData);
+	Serialize(PayloadAr);
+
+	// Serialize custom versions
+	TArray<uint8> HeaderData;
+	FMemoryWriter HeaderAr(HeaderData);
+	FCustomVersionContainer CustomVersions = PayloadAr.GetCustomVersions();
+	CustomVersions.Serialize(HeaderAr);
+
+	// Append data
+	OutData = MoveTemp(HeaderData);
+	OutData.Append(PayloadData);
 }
 
 FString FWorldPartitionActorDesc::ToString() const
@@ -127,78 +150,9 @@ void FWorldPartitionActorDesc::BuildHash(FHashBuilder& HashBuilder)
 	HashBuilder << Guid << Class << ActorPackage << ActorPath << BoundsLocation << BoundsExtent << GridPlacement << RuntimeGrid << bActorIsEditorOnly << bLevelBoundsRelevant << Layers << References;
 }
 
-void FWorldPartitionActorDesc::SerializeMetaData(FActorMetaDataSerializer* Serializer)
+void FWorldPartitionActorDesc::Serialize(FArchive& Ar)
 {
-	bool bValidSerialization = true;
-
-	bValidSerialization &= Serializer->Serialize(TEXT("ActorClass"), Class);
-	bValidSerialization &= Serializer->Serialize(TEXT("ActorGuid"), Guid);
-	bValidSerialization &= Serializer->Serialize(TEXT("BoundsLocation"), BoundsLocation);
-	bValidSerialization &= Serializer->Serialize(TEXT("BoundsExtent"), BoundsExtent);
-	bValidSerialization &= Serializer->Serialize(TEXT("GridPlacement"), (int8&)GridPlacement);
-	bValidSerialization &= Serializer->Serialize(TEXT("RuntimeGrid"), RuntimeGrid);
-	bValidSerialization &= Serializer->Serialize(TEXT("IsEditorOnly"), bActorIsEditorOnly);
-	bValidSerialization &= Serializer->Serialize(TEXT("IsLevelBoundsRelevant"), bLevelBoundsRelevant);
-
-	if (!bValidSerialization)
-	{
-		Serializer->SetHasErrors();
-	}
-
-	FString LayersStr;
-	FString ActorsRefsStr;
-
-	if (Serializer->IsWriting())
-	{
-		if (Layers.Num())
-		{
-			for(FName Layer: Layers)
-			{
-				LayersStr += Layer.ToString() + TEXT(";");
-			}
-			LayersStr.RemoveFromEnd(TEXT(";"));
-		}
-
-		if (References.Num())
-		{
-			for(const FGuid& ActoGuid: References)
-			{
-				ActorsRefsStr += ActoGuid.ToString() + TEXT(";");
-			}
-			ActorsRefsStr.RemoveFromEnd(TEXT(";"));
-		}
-	}
-
-	Serializer->Serialize(TEXT("Layers"), LayersStr);
-	Serializer->Serialize(TEXT("ActorReferences"), ActorsRefsStr);
-
-	if (Serializer->IsReading())
-	{
-		TArray<FString> NewLayers;
-		if (LayersStr.ParseIntoArray(NewLayers, TEXT(";")))
-		{
-			Algo::Transform(NewLayers, Layers, [&](const FString& Layer) { return FName(*Layer); });
-		}
-
-		TArray<FString> ActorsRefs;
-		if (ActorsRefsStr.ParseIntoArray(ActorsRefs, TEXT(";")))
-		{
-			Algo::Transform(ActorsRefs, References, [](const FString& ActorGuidStr)
-			{
-				FGuid ActorGuid;
-				if (!FGuid::Parse(ActorGuidStr, ActorGuid))
-				{
-					ActorGuid.Invalidate();
-				}
-				return ActorGuid;
-			});
-
-			Algo::RemoveIf(References, [](const FGuid& InGuid)
-			{
-				return !InGuid.IsValid();
-			});
-		}
-	}
+	Ar << Class << Guid << BoundsLocation << BoundsExtent << GridPlacement << RuntimeGrid << bActorIsEditorOnly << bLevelBoundsRelevant << Layers << References;
 }
 
 FBox FWorldPartitionActorDesc::GetBounds() const
diff --git a/Engine/Source/Runtime/Engine/Private/WorldPartition/WorldPartitionActorDescMetaDataSerializer.cpp b/Engine/Source/Runtime/Engine/Private/WorldPartition/WorldPartitionActorDescMetaDataSerializer.cpp
deleted file mode 100644
diff --git a/Engine/Source/Runtime/Engine/Public/WorldPartition/ActorPartition/PartitionActorDesc.h b/Engine/Source/Runtime/Engine/Public/WorldPartition/ActorPartition/PartitionActorDesc.h
--- a/Engine/Source/Runtime/Engine/Public/WorldPartition/ActorPartition/PartitionActorDesc.h
+++ b/Engine/Source/Runtime/Engine/Public/WorldPartition/ActorPartition/PartitionActorDesc.h
@@ -19,10 +19,9 @@ public:
 	int64 GridIndexZ;
 
 protected:
-	virtual bool Init(const AActor* InActor) override;
-	
 	virtual void BuildHash(FHashBuilder& HashBuilder) override;
 
-	virtual void SerializeMetaData(FActorMetaDataSerializer* Serializer) override;
+	virtual void InitFrom(const AActor* InActor) override;
+	virtual void Serialize(FArchive& Ar) override;
 #endif
 };
diff --git a/Engine/Source/Runtime/Engine/Public/WorldPartition/HLOD/HLODActorDesc.h b/Engine/Source/Runtime/Engine/Public/WorldPartition/HLOD/HLODActorDesc.h
--- a/Engine/Source/Runtime/Engine/Public/WorldPartition/HLOD/HLODActorDesc.h
+++ b/Engine/Source/Runtime/Engine/Public/WorldPartition/HLOD/HLODActorDesc.h
@@ -22,11 +22,10 @@ public:
 	inline const FSoftObjectPath& GetHLODLayer() const { return HLODLayer; }
 
 protected:
-	virtual bool Init(const AActor* InActor) override;
-
 	virtual void BuildHash(FHashBuilder& HashBuilder) override;
 
-	virtual void SerializeMetaData(FActorMetaDataSerializer* Serializer) override;
+	virtual void InitFrom(const AActor* InActor) override;
+	virtual void Serialize(FArchive& Ar) override;
 
 	TArray<FGuid> SubActors;
 	FSoftObjectPath	HLODLayer;
diff --git a/Engine/Source/Runtime/Engine/Public/WorldPartition/LevelInstance/LevelInstanceActorDesc.h b/Engine/Source/Runtime/Engine/Public/WorldPartition/LevelInstance/LevelInstanceActorDesc.h
--- a/Engine/Source/Runtime/Engine/Public/WorldPartition/LevelInstance/LevelInstanceActorDesc.h
+++ b/Engine/Source/Runtime/Engine/Public/WorldPartition/LevelInstance/LevelInstanceActorDesc.h
@@ -22,11 +22,10 @@ public:
 	inline FName GetLevelPackage() const { return LevelPackage; }
 
 protected:
-	virtual bool Init(const AActor* InActor) override;
-
 	virtual void BuildHash(FHashBuilder& HashBuilder) override;
 
-	virtual void SerializeMetaData(FActorMetaDataSerializer* Serializer) override;
+	virtual void InitFrom(const AActor* InActor) override;
+	virtual void Serialize(FArchive& Ar) override;
 
 	FName LevelPackage;
 	FTransform LevelInstanceTransform;
diff --git a/Engine/Source/Runtime/Engine/Public/WorldPartition/WorldPartitionActorDesc.h b/Engine/Source/Runtime/Engine/Public/WorldPartition/WorldPartitionActorDesc.h
--- a/Engine/Source/Runtime/Engine/Public/WorldPartition/WorldPartitionActorDesc.h
+++ b/Engine/Source/Runtime/Engine/Public/WorldPartition/WorldPartitionActorDesc.h
@@ -9,7 +9,6 @@
 #include "Containers/Set.h"
 #include "Misc/Guid.h"
 #include "Misc/HashBuilder.h"
-#include "WorldPartition/WorldPartitionActorMetaDataSerializer.h"
 #include "WorldPartition/WorldPartitionActorDescFactory.h"
 
 #if WITH_EDITOR
@@ -19,9 +18,8 @@ struct FWorldPartitionActorDescInitData
 	UClass* NativeClass;
 	FName PackageName;
 	FName ActorPath;
-	FAssetData AssetData;
 	FTransform Transform;
-	FActorMetaDataSerializer* Serializer;
+	TArray<uint8> SerializedData;
 };
 #endif
 
@@ -31,8 +29,6 @@ struct FWorldPartitionActorDescInitData
 class ENGINE_API FWorldPartitionActorDesc 
 {
 #if WITH_EDITOR
-	friend class AActor;
-	friend class UWorldPartition;
 	friend class FWorldPartitionActorDescFactory;
 
 public:
@@ -84,15 +80,19 @@ public:
 	AActor* Load(const FLinkerInstancingContext* InstancingContext = nullptr);
 	void Unload();
 
+	void Init(const AActor* InActor);
+	void Init(const FWorldPartitionActorDescInitData& DescData);
+
+	void SerializeTo(TArray<uint8>& OutData);
+
 protected:
 	FWorldPartitionActorDesc();
-	virtual bool Init(const AActor* InActor);
-	bool Init(const FWorldPartitionActorDescInitData& DescData);
-		
+	
 	void UpdateHash();
 	virtual void BuildHash(FHashBuilder& HashBuilder);
 
-	virtual void SerializeMetaData(FActorMetaDataSerializer* Serializer);
+	virtual void InitFrom(const AActor* Actor);
+	virtual void Serialize(FArchive& Ar);
 
 	FGuid						Guid;
 	FName						Class;
diff --git a/Engine/Source/Runtime/Engine/Public/WorldPartition/WorldPartitionActorMetaDataSerializer.h b/Engine/Source/Runtime/Engine/Public/WorldPartition/WorldPartitionActorMetaDataSerializer.h
deleted file mode 100644
