Get rid of the 'LevelPackage' asset registry tag for actors.

diff --git a/Engine/Source/Runtime/Engine/Private/Actor.cpp b/Engine/Source/Runtime/Engine/Private/Actor.cpp
--- a/Engine/Source/Runtime/Engine/Private/Actor.cpp
+++ b/Engine/Source/Runtime/Engine/Private/Actor.cpp
@@ -522,10 +522,6 @@ void AActor::GetAssetRegistryTags(TArray<FAssetRegistryTag>& OutTags) const
 			TUniquePtr<FWorldPartitionActorDesc> ActorDesc(WorldPartition->GetActorDescFactory(this)->Create());
 			ActorDesc->Init(this);
 
-			const FString LevelPackageName = GetOuter()->GetOutermost()->GetName();
-			static FName NAME_LevelPackage(TEXT("LevelPackage"));
-			OutTags.Add(UObject::FAssetRegistryTag(NAME_LevelPackage, LevelPackageName, UObject::FAssetRegistryTag::TT_Hidden));
-
 			const FString ActorMetaDataClass = GetParentNativeClass(GetClass())->GetName();
 			static FName NAME_ActorMetaDataClass(TEXT("ActorMetaDataClass"));
 			OutTags.Add(UObject::FAssetRegistryTag(NAME_ActorMetaDataClass, ActorMetaDataClass, UObject::FAssetRegistryTag::TT_Hidden));
diff --git a/Engine/Source/Runtime/Engine/Private/WorldPartition/WorldPartition.cpp b/Engine/Source/Runtime/Engine/Private/WorldPartition/WorldPartition.cpp
--- a/Engine/Source/Runtime/Engine/Private/WorldPartition/WorldPartition.cpp
+++ b/Engine/Source/Runtime/Engine/Private/WorldPartition/WorldPartition.cpp
@@ -211,17 +211,18 @@ void UWorldPartition::Initialize(UWorld* InWorld, const FTransform& InTransform)
 		{
 			if (!LevelPath.IsNone())
 			{
-				FString LevelPathStr = LevelPath.ToString();
-				IAssetRegistry& AssetRegistry = FModuleManager::LoadModuleChecked<FAssetRegistryModule>(TEXT("AssetRegistry")).Get();
-	
+				const FString LevelPathStr = LevelPath.ToString();
+				const FString LevelExternalActorsPath = ULevel::GetExternalActorsPath(LevelPathStr);
+
 				// Do a synchronous scan of the level external actors path.
-				const bool bForceRescan = true;
-				AssetRegistry.ScanPathsSynchronous({ULevel::GetExternalActorsPath(LevelPath.ToString())}, bForceRescan);
+				IAssetRegistry& AssetRegistry = FModuleManager::LoadModuleChecked<FAssetRegistryModule>(TEXT("AssetRegistry")).Get();
+				AssetRegistry.ScanPathsSynchronous({LevelExternalActorsPath}, /*bForceRescan*/true);
 
-				static const FName NAME_LevelPackage(TEXT("LevelPackage"));
 				FARFilter Filter;
-				Filter.TagsAndValues.Add(NAME_LevelPackage, LevelPathStr);
+				Filter.bRecursivePaths = true;
 				Filter.bIncludeOnlyOnDiskAssets = true;
+				Filter.PackagePaths.Add(*LevelExternalActorsPath);
+
 				AssetRegistry.GetAssets(Filter, OutAssets);
 			}
 		};
