Blueprints: Add ArePinTypesEquivalent function the schema for a quicker way to determine if two pin types are the same

#rnx
#okfrogithub public

diff --git a/Engine/Source/Editor/BlueprintGraph/Classes/EdGraphSchema_K2.h b/Engine/Source/Editor/BlueprintGraph/Classes/EdGraphSchema_K2.h
--- a/Engine/Source/Editor/BlueprintGraph/Classes/EdGraphSchema_K2.h
+++ b/Engine/Source/Editor/BlueprintGraph/Classes/EdGraphSchema_K2.h
@@ -1007,6 +1007,16 @@ public:
 	 */
 	virtual bool ArePinTypesCompatible(const FEdGraphPinType& Output, const FEdGraphPinType& Input, const UClass* CallingContext = NULL, bool bIgnoreArray = false) const;
 
+	/**
+	 * Returns true if the types are schema Equivalent. 
+	 *
+	 * @param	PinA		  	The type of Pin A.
+	 * @param	PinB		  	The type of Pin B.
+	 *
+	 * @return	true if the pin types and directions are compatible.
+	 */
+	virtual bool ArePinTypesEquivalent(const FEdGraphPinType& PinA, const FEdGraphPinType& PinB) const;
+
 	/** Sets the autogenerated default value for a pin, optionally using the passed in function and parameter. This will also reset the current default value to the autogenerated one */
 	virtual void SetPinAutogeneratedDefaultValue(UEdGraphPin* Pin, const FString& NewValue) const;
 
diff --git a/Engine/Source/Editor/BlueprintGraph/Private/EdGraphSchema_K2.cpp b/Engine/Source/Editor/BlueprintGraph/Private/EdGraphSchema_K2.cpp
--- a/Engine/Source/Editor/BlueprintGraph/Private/EdGraphSchema_K2.cpp
+++ b/Engine/Source/Editor/BlueprintGraph/Private/EdGraphSchema_K2.cpp
@@ -4400,6 +4400,16 @@ bool UEdGraphSchema_K2::ArePinTypesCompatible(const FEdGraphPinType& Output, con
 	return false;
 }
 
+bool UEdGraphSchema_K2::ArePinTypesEquivalent(const FEdGraphPinType& PinA, const FEdGraphPinType& PinB) const
+{
+	return 
+		PinA.PinCategory == PinB.PinCategory &&
+		PinA.PinSubCategory == PinB.PinSubCategory &&
+		PinA.PinSubCategoryObject == PinB.PinSubCategoryObject &&
+		PinA.ContainerType == PinB.ContainerType &&
+		PinA.bIsWeakPointer == PinB.bIsWeakPointer;
+}
+
 void UEdGraphSchema_K2::BreakNodeLinks(UEdGraphNode& TargetNode) const
 {
 	UBlueprint* Blueprint = FBlueprintEditorUtils::FindBlueprintForNodeChecked(&TargetNode);
diff --git a/Engine/Source/Runtime/Engine/Classes/EdGraph/EdGraphSchema.h b/Engine/Source/Runtime/Engine/Classes/EdGraph/EdGraphSchema.h
--- a/Engine/Source/Runtime/Engine/Classes/EdGraph/EdGraphSchema.h
+++ b/Engine/Source/Runtime/Engine/Classes/EdGraph/EdGraphSchema.h
@@ -1049,7 +1049,17 @@ class ENGINE_API UEdGraphSchema : public UObject
 	 *
 	 * @return	true if the pin types and directions are compatible.
 	 */
-	virtual bool ArePinsCompatible(const UEdGraphPin* PinA, const UEdGraphPin* PinB, const UClass* CallingContext = NULL, bool bIgnoreArray = false) const { return true; }
+	virtual bool ArePinsCompatible(const UEdGraphPin* PinA, const UEdGraphPin* PinB, const UClass* CallingContext = nullptr, bool bIgnoreArray = false) const { return true; }
+
+	/**
+	 * Returns true if the types are schema Equivalent. 
+	 *
+	 * @param	PinA		  	The type of Pin A.
+	 * @param	PinB		  	The type of Pin B.
+	 *
+	 * @return	true if the pin types and directions are compatible.
+	 */
+	virtual bool ArePinTypesEquivalent(const FEdGraphPinType& PinA, const FEdGraphPinType& PinB) const { return true; }
 
 	/**
 	 * Returns true if the schema wants to overdrive the behaviour of dirtying the blueprint on new node creation.
