Fix bug where Skeletal Mesh RegisterEndPhysicsTick is registering before its dependant World->EndPhysicsTickFunction, so the prerequisite is ingored. Meaning it uses stale physics data and causing jitter between skeletal mesh and follow camera.

diff --git a/Engine/Source/Runtime/Engine/Private/Components/SkeletalMeshComponent.cpp b/Engine/Source/Runtime/Engine/Private/Components/SkeletalMeshComponent.cpp
--- a/Engine/Source/Runtime/Engine/Private/Components/SkeletalMeshComponent.cpp
+++ b/Engine/Source/Runtime/Engine/Private/Components/SkeletalMeshComponent.cpp
@@ -369,17 +369,18 @@ void USkeletalMeshComponent::RegisterComponentTickFunctions(bool bRegister)
 	UpdateClothTickRegisteredState();
 }
 
+
 void USkeletalMeshComponent::RegisterEndPhysicsTick(bool bRegister)
 {
 	if (bRegister != EndPhysicsTickFunction.IsTickFunctionRegistered())
 	{
 		if (bRegister)
 		{
-			if (SetupActorComponentTickFunction(&EndPhysicsTickFunction))
+			UWorld* World = GetWorld();
+			if (World->EndPhysicsTickFunction.IsTickFunctionRegistered() && SetupActorComponentTickFunction(&EndPhysicsTickFunction))
 			{
 				EndPhysicsTickFunction.Target = this;
 				// Make sure our EndPhysicsTick gets called after physics simulation is finished
-				UWorld* World = GetWorld();
 				if (World != nullptr)
 				{
 					EndPhysicsTickFunction.AddPrerequisite(World, World->EndPhysicsTickFunction);
