Improved ZoneGraph visualization
- Moved shared draw functions to separate file (rendering utils)
- Draw lane direction arrows for lanes
- Fixed and brushed up FPrimitiveSceneProxy::DrawArrowHead()

diff --git a/Engine/Source/Runtime/Engine/Private/PrimitiveSceneProxy.cpp b/Engine/Source/Runtime/Engine/Private/PrimitiveSceneProxy.cpp
--- a/Engine/Source/Runtime/Engine/Private/PrimitiveSceneProxy.cpp
+++ b/Engine/Source/Runtime/Engine/Private/PrimitiveSceneProxy.cpp
@@ -940,28 +940,11 @@ void FPrimitiveSceneProxy::DrawArc(FPrimitiveDrawInterface* PDI, const FVector&
 
 void FPrimitiveSceneProxy::DrawArrowHead(FPrimitiveDrawInterface* PDI, const FVector& Tip, const FVector& Origin, const float Size, const FLinearColor& Color, uint8 DepthPriorityGroup, const float Thickness, const bool bScreenSpace)
 {
-	//float ax[3], ay[3] = {0,1,0}, az[3];
-	FVector Ax, Ay, Az(0,1,0);
-	// dtVsub(az, q, p);
-	Ay = Origin - Tip;
-	// dtVnormalize(az);
-	Ay.Normalize();
-	// dtVcross(ax, ay, az);
-	Ax = FVector::CrossProduct(Az, Ay);
-	// dtVcross(ay, az, ax);
-	//Az = FVector::CrossProduct(Ay, Ax);
-	////dtVnormalize(ay);
-	//Az.Normalize();
-
-	PDI->DrawLine( Tip
-		//, FVector(p[0]+az[0]*s+ax[0]*s/3, p[1]+az[1]*s+ax[1]*s/3, p[2]+az[2]*s+ax[2]*s/3)
-		, FVector(Tip.X + Ay.X*Size + Ax.X*Size/3, Tip.Y + Ay.Y*Size + Ax.Y*Size/3, Tip.Z + Ay.Z*Size + Ax.Z*Size/3)
-		, Color, SDPG_World, Thickness, bScreenSpace);
-
-	PDI->DrawLine( Tip
-		//, FVector(p[0]+az[0]*s-ax[0]*s/3, p[1]+az[1]*s-ax[1]*s/3, p[2]+az[2]*s-ax[2]*s/3)
-		, FVector(Tip.X + Ay.X*Size - Ax.X*Size/3, Tip.Y + Ay.Y*Size - Ax.Y*Size/3, Tip.Z + Ay.Z*Size - Ax.Z*Size/3)
-		, Color, SDPG_World, Thickness, bScreenSpace);
+	const FVector Forward = (Origin - Tip).GetUnsafeNormal();
+	const FVector Right = FVector::CrossProduct(Forward, FVector::UpVector);
+	const float HalfWidth = Size / 3.0f;
+	PDI->DrawLine(Tip, Tip + Forward * Size + Right * HalfWidth, Color, DepthPriorityGroup, Thickness, bScreenSpace);
+	PDI->DrawLine(Tip, Tip + Forward * Size - Right * HalfWidth, Color, DepthPriorityGroup, Thickness, bScreenSpace);
 }
 
 
