[APawn] Improved documentation for bCanAffectNavigationGeneration

diff --git a/Engine/Source/Runtime/Engine/Classes/GameFramework/Pawn.h b/Engine/Source/Runtime/Engine/Classes/GameFramework/Pawn.h
--- a/Engine/Source/Runtime/Engine/Classes/GameFramework/Pawn.h
+++ b/Engine/Source/Runtime/Engine/Classes/GameFramework/Pawn.h
@@ -64,11 +64,15 @@ public:
 	UPROPERTY(EditAnywhere, BlueprintReadWrite, Category=Pawn)
 	uint32 bUseControllerRotationRoll:1;
 
-	/** If set to false (default) given pawn instance will never affect navigation generation. 
+	/**
+	 *	If set to false (default) given pawn instance will never affect navigation generation (but components could).
 	 *	Setting it to true will result in using regular AActor's navigation relevancy 
-	 *	calculation to check if this pawn instance should affect navigation generation
-	 *	Use SetCanAffectNavigationGeneration to change this value at runtime.
-	 *	Note that modifying this value at runtime will result in any navigation change only if runtime navigation generation is enabled. */
+	 *	calculation to check if this pawn instance should affect navigation generation.
+	 *	@note Use SetCanAffectNavigationGeneration() to change this value at runtime.
+	 *	@note Modifying this value at runtime will result in any navigation change only if runtime navigation generation is enabled.
+	 *	@note Override UpdateNavigationRelevance() to propagate the flag to the desired components.
+	 *	@see SetCanAffectNavigationGeneration(), UpdateNavigationRelevance()
+	 */
 	UPROPERTY(EditAnywhere, BlueprintReadOnly, Category = Pawn)
 	uint32 bCanAffectNavigationGeneration : 1;
 
