Fix stale ray tracing geometry when skin cache entry is released

diff --git a/Engine/Source/Runtime/Engine/Private/SkeletalRenderGPUSkin.cpp b/Engine/Source/Runtime/Engine/Private/SkeletalRenderGPUSkin.cpp
--- a/Engine/Source/Runtime/Engine/Private/SkeletalRenderGPUSkin.cpp
+++ b/Engine/Source/Runtime/Engine/Private/SkeletalRenderGPUSkin.cpp
@@ -399,74 +399,83 @@ void FSkeletalMeshObjectGPUSkin::UpdateDynamicData_RenderThread(FGPUSkinCache* G
 #if RHI_RAYTRACING
 	if (IsRayTracingEnabled())
 	{
-		if (GEnableGPUSkinCache && SkinCacheEntry)
+		if (GEnableGPUSkinCache)
 		{
-			if (DynamicData->LODIndex >= SkeletalMeshRenderData->CurrentFirstLODIdx) // According to GetMeshElementsConditionallySelectable(), non-resident LODs should just be skipped
+			if (SkinCacheEntry)
 			{
-				if (bRequireRecreatingRayTracingGeometry)
+				if (DynamicData->LODIndex >= SkeletalMeshRenderData->CurrentFirstLODIdx) // According to GetMeshElementsConditionallySelectable(), non-resident LODs should just be skipped
 				{
-					FSkeletalMeshLODRenderData& LODModel = this->SkeletalMeshRenderData->LODRenderData[DynamicData->LODIndex];
-					FIndexBufferRHIRef IndexBufferRHI = LODModel.MultiSizeIndexContainer.GetIndexBuffer()->IndexBufferRHI;
-					uint32 VertexBufferStride = LODModel.StaticVertexBuffers.PositionVertexBuffer.GetStride();
-
-					//#dxr_todo: do we need support for separate sections in FRayTracingGeometryData?
-					uint32 TrianglesCount = 0;
-					for (int32 SectionIndex = 0; SectionIndex < LODModel.RenderSections.Num(); SectionIndex++)
+					if (bRequireRecreatingRayTracingGeometry)
 					{
-						const FSkelMeshRenderSection& Section = LODModel.RenderSections[SectionIndex];
-						TrianglesCount += Section.NumTriangles;
-					}
+						FSkeletalMeshLODRenderData& LODModel = this->SkeletalMeshRenderData->LODRenderData[DynamicData->LODIndex];
+						FIndexBufferRHIRef IndexBufferRHI = LODModel.MultiSizeIndexContainer.GetIndexBuffer()->IndexBufferRHI;
+						uint32 VertexBufferStride = LODModel.StaticVertexBuffers.PositionVertexBuffer.GetStride();
 
-					FRayTracingGeometryInitializer Initializer;
-					static const FName DebugName("FSkeletalMeshObjectGPUSkin");
-					static int32 DebugNumber = 0;
-					Initializer.DebugName = FName(DebugName, DebugNumber++);
+						//#dxr_todo: do we need support for separate sections in FRayTracingGeometryData?
+						uint32 TrianglesCount = 0;
+						for (int32 SectionIndex = 0; SectionIndex < LODModel.RenderSections.Num(); SectionIndex++)
+						{
+							const FSkelMeshRenderSection& Section = LODModel.RenderSections[SectionIndex];
+							TrianglesCount += Section.NumTriangles;
+						}
 
-					FRHIResourceCreateInfo CreateInfo;
+						FRayTracingGeometryInitializer Initializer;
+						static const FName DebugName("FSkeletalMeshObjectGPUSkin");
+						static int32 DebugNumber = 0;
+						Initializer.DebugName = FName(DebugName, DebugNumber++);
 
-					Initializer.IndexBuffer = IndexBufferRHI;
-					Initializer.TotalPrimitiveCount = TrianglesCount;
-					Initializer.GeometryType = RTGT_Triangles;
-					Initializer.bFastBuild = true;
-					Initializer.bAllowUpdate = true;
+						FRHIResourceCreateInfo CreateInfo;
 
-					Initializer.Segments.Reserve(LODModel.RenderSections.Num());
-					for (const FSkelMeshRenderSection& Section : LODModel.RenderSections)
-					{
-						FRayTracingGeometrySegment Segment;
-						Segment.VertexBuffer = nullptr;
-						Segment.VertexBufferElementType = VET_Float3;
-						Segment.VertexBufferStride = VertexBufferStride;
-						Segment.VertexBufferOffset = 0;
-						Segment.FirstPrimitive = Section.BaseIndex / 3;
-						Segment.NumPrimitives = Section.NumTriangles;
-						Segment.bEnabled = !Section.bDisabled;
-						Initializer.Segments.Add(Segment);
-					}
+						Initializer.IndexBuffer = IndexBufferRHI;
+						Initializer.TotalPrimitiveCount = TrianglesCount;
+						Initializer.GeometryType = RTGT_Triangles;
+						Initializer.bFastBuild = true;
+						Initializer.bAllowUpdate = true;
 
-					FGPUSkinCache::GetRayTracingSegmentVertexBuffers(*SkinCacheEntry, Initializer.Segments);
+						Initializer.Segments.Reserve(LODModel.RenderSections.Num());
+						for (const FSkelMeshRenderSection& Section : LODModel.RenderSections)
+						{
+							FRayTracingGeometrySegment Segment;
+							Segment.VertexBuffer = nullptr;
+							Segment.VertexBufferElementType = VET_Float3;
+							Segment.VertexBufferStride = VertexBufferStride;
+							Segment.VertexBufferOffset = 0;
+							Segment.FirstPrimitive = Section.BaseIndex / 3;
+							Segment.NumPrimitives = Section.NumTriangles;
+							Segment.bEnabled = !Section.bDisabled;
+							Initializer.Segments.Add(Segment);
+						}
 
-					// Flush pending resource barriers before BVH is built for the first time
-					GPUSkinCache->TransitionAllToReadable(RHICmdList);
+						FGPUSkinCache::GetRayTracingSegmentVertexBuffers(*SkinCacheEntry, Initializer.Segments);
 
-					RayTracingGeometry.SetInitializer(Initializer);
-					RayTracingGeometry.UpdateRHI();
-				}
-				else
-				{
-					// If we are not using world position offset in material, handle BLAS refit here
-					if (!DynamicData->bAnySegmentUsesWorldPositionOffset)
-					{
-						// Refit BLAS with new vertex buffer data
-						FGPUSkinCache::GetRayTracingSegmentVertexBuffers(*SkinCacheEntry, RayTracingGeometry.Initializer.Segments);
-						GPUSkinCache->AddRayTracingGeometryToUpdate(&RayTracingGeometry);
+						// Flush pending resource barriers before BVH is built for the first time
+						GPUSkinCache->TransitionAllToReadable(RHICmdList);
+
+						RayTracingGeometry.SetInitializer(Initializer);
+						RayTracingGeometry.UpdateRHI();
 					}
 					else
 					{
-						// Otherwise, we will run the dynamic ray tracing geometry path, i.e. runnning VSinCS and refit geometry there, so do nothing here
+						// If we are not using world position offset in material, handle BLAS refit here
+						if (!DynamicData->bAnySegmentUsesWorldPositionOffset)
+						{
+							// Refit BLAS with new vertex buffer data
+							FGPUSkinCache::GetRayTracingSegmentVertexBuffers(*SkinCacheEntry, RayTracingGeometry.Initializer.Segments);
+							GPUSkinCache->AddRayTracingGeometryToUpdate(&RayTracingGeometry);
+						}
+						else
+						{
+							// Otherwise, we will run the dynamic ray tracing geometry path, i.e. runnning VSinCS and refit geometry there, so do nothing here
+						}
 					}
 				}
 			}
+			else
+			{
+				// When SkinCacheEntry is gone, clear geometry
+				RayTracingGeometry.ReleaseRHI();
+				RayTracingGeometry.SetInitializer(FRayTracingGeometryInitializer {});
+			}
 		}
 	}
 #endif
