Fixing recompile issues caused by the mesh materials being in a different map type that doesn't override GetOutdatedTypes. This should allow 'recompileshaders changed' to detect that mesh material shaders are out of date.

diff --git a/Engine/Source/Runtime/Engine/Private/Materials/MaterialShader.cpp b/Engine/Source/Runtime/Engine/Private/Materials/MaterialShader.cpp
--- a/Engine/Source/Runtime/Engine/Private/Materials/MaterialShader.cpp
+++ b/Engine/Source/Runtime/Engine/Private/Materials/MaterialShader.cpp
@@ -2582,6 +2582,19 @@ void FMaterialShaderMap::DumpDebugInfo()
 }
 
 #if WITH_EDITOR
+void FMaterialShaderMap::GetOutdatedTypes(TArray<const FShaderType*>& OutdatedShaderTypes, TArray<const FShaderPipelineType*>& OutdatedShaderPipelineTypes, TArray<const FVertexFactoryType*>& OutdatedFactoryTypes) const
+{
+	FShaderMapBase::GetOutdatedTypes(OutdatedShaderTypes, OutdatedShaderPipelineTypes, OutdatedFactoryTypes);
+
+	for (FMeshMaterialShaderMap* MeshShaderMap : GetContent()->OrderedMeshShaderMaps)
+	{
+		if (MeshShaderMap)
+		{
+			MeshShaderMap->GetOutdatedTypes(*this, OutdatedShaderTypes, OutdatedShaderPipelineTypes, OutdatedFactoryTypes);
+		}
+	}
+}
+
 void FMaterialShaderMap::SaveShaderStableKeys(EShaderPlatform TargetShaderPlatform, const FStableShaderKeyAndValue& SaveKeyVal)
 {
 	FShaderMapBase::SaveShaderStableKeys(TargetShaderPlatform, SaveKeyVal);
diff --git a/Engine/Source/Runtime/Engine/Public/MaterialShared.h b/Engine/Source/Runtime/Engine/Public/MaterialShared.h
--- a/Engine/Source/Runtime/Engine/Public/MaterialShared.h
+++ b/Engine/Source/Runtime/Engine/Public/MaterialShared.h
@@ -1178,6 +1178,7 @@ public:
 	/** Returns the maximum number of texture samplers used by any shader in this shader map. */
 	uint32 GetMaxTextureSamplers() const;
 
+	void GetOutdatedTypes(TArray<const FShaderType*>& OutdatedShaderTypes, TArray<const FShaderPipelineType*>& OutdatedShaderPipelineTypes, TArray<const FVertexFactoryType*>& OutdatedFactoryTypes) const;
 	void SaveShaderStableKeys(EShaderPlatform TargetShaderPlatform, const struct FStableShaderKeyAndValue& SaveKeyVal);
 #endif
 
