Fix crash when recompiling a global shader (UE-99977).

diff --git a/Engine/Source/Runtime/Engine/Private/ShaderCompiler/ShaderCompiler.cpp b/Engine/Source/Runtime/Engine/Private/ShaderCompiler/ShaderCompiler.cpp
--- a/Engine/Source/Runtime/Engine/Private/ShaderCompiler/ShaderCompiler.cpp
+++ b/Engine/Source/Runtime/Engine/Private/ShaderCompiler/ShaderCompiler.cpp
@@ -2649,6 +2649,28 @@ void FShaderCompilingManager::ProcessCompiledShaderMaps(
 
 				ProcessCompiledGlobalShaders(CompilationResults);
 
+				// Re-register everything to work around FShader lifetime issues - it currently lives and dies with the
+				// shadermap it is stored in, while cached MDCs can reference its memory. Re-registering will
+				// re-create the cache.
+				TIndirectArray<FComponentRecreateRenderStateContext> ComponentContexts;
+				for (TObjectIterator<UPrimitiveComponent> PrimitiveIt; PrimitiveIt; ++PrimitiveIt)
+				{
+					UPrimitiveComponent* PrimitiveComponent = *PrimitiveIt;
+
+					if (PrimitiveComponent->IsRenderStateCreated())
+					{
+						ComponentContexts.Add(new FComponentRecreateRenderStateContext(PrimitiveComponent));
+#if WITH_EDITOR
+						if (PrimitiveComponent->HasValidSettingsForStaticLighting(false))
+						{
+							FStaticLightingSystemInterface::OnPrimitiveComponentUnregistered.Broadcast(PrimitiveComponent);
+							FStaticLightingSystemInterface::OnPrimitiveComponentRegistered.Broadcast(PrimitiveComponent);
+						}
+#endif
+					}
+				}
+				ComponentContexts.Empty();
+
 				CompilationResults.Empty();
 				ProcessIt.RemoveCurrent();
 			}
