Adding lightweight instances subsystem for gameplay objects.
Lightweight instances allow a low memory representation of objects. Instances can be converted into full actors when necessary. Actors can also be reduced down to light weight instances when the full representation is no longer needed.
Hit results now return a handle that may refer to either an actor or lightweight instance.

diff --git a/Engine/Plugins/2D/Paper2D/Source/Paper2DEditor/Private/MeshPainting/MeshPaintSpriteAdapter.cpp b/Engine/Plugins/2D/Paper2D/Source/Paper2DEditor/Private/MeshPainting/MeshPaintSpriteAdapter.cpp
--- a/Engine/Plugins/2D/Paper2D/Source/Paper2DEditor/Private/MeshPainting/MeshPaintSpriteAdapter.cpp
+++ b/Engine/Plugins/2D/Paper2D/Source/Paper2DEditor/Private/MeshPainting/MeshPaintSpriteAdapter.cpp
@@ -83,7 +83,7 @@ bool FMeshPaintSpriteAdapter::LineTraceComponent(struct FHitResult& OutHit, cons
 					OutHit.ImpactNormal = WorldNormal;
 					OutHit.TraceStart = Start;
 					OutHit.TraceEnd = End;
-					OutHit.Actor = SpriteComponent->GetOwner();
+					OutHit.HitObjectHandle = FActorInstanceHandle(SpriteComponent->GetOwner());
 					OutHit.Component = SpriteComponent;
 					OutHit.FaceIndex = VertexIndex / 3;
 
diff --git a/Engine/Plugins/Developer/Concert/ConcertSync/ConcertSyncClient/Source/ConcertSyncClient/Private/ConcertClientPresenceMode.cpp b/Engine/Plugins/Developer/Concert/ConcertSync/ConcertSyncClient/Source/ConcertSyncClient/Private/ConcertClientPresenceMode.cpp
--- a/Engine/Plugins/Developer/Concert/ConcertSync/ConcertSyncClient/Source/ConcertSyncClient/Private/ConcertClientPresenceMode.cpp
+++ b/Engine/Plugins/Developer/Concert/ConcertSync/ConcertSyncClient/Source/ConcertSyncClient/Private/ConcertClientPresenceMode.cpp
@@ -183,7 +183,7 @@ void FConcertClientDesktopPresenceMode::SendEvents(IConcertClientSession& Sessio
 
 
 						// chose closest point to origin.  Always favor the trace if the proxy and the trace reached the same actor since the trace will be more accurate
-						if ((ActorProxy && ActorProxy->Actor == Result.Actor) || FVector::DistSquared(LineCheckStart, LineTrace) < FVector::DistSquared(LineCheckStart, HitProxyTrace))
+						if ((ActorProxy && Result.GetHitObjectHandle() == ActorProxy->Actor) || FVector::DistSquared(LineCheckStart, LineTrace) < FVector::DistSquared(LineCheckStart, HitProxyTrace))
 						{
 							Event.TraceEnd = LineTrace;
 						}
diff --git a/Engine/Plugins/Experimental/VirtualCamera/Source/VirtualCamera/Private/VirtualCameraPlayerControllerBase.cpp b/Engine/Plugins/Experimental/VirtualCamera/Source/VirtualCamera/Private/VirtualCameraPlayerControllerBase.cpp
--- a/Engine/Plugins/Experimental/VirtualCamera/Source/VirtualCamera/Private/VirtualCameraPlayerControllerBase.cpp
+++ b/Engine/Plugins/Experimental/VirtualCamera/Source/VirtualCamera/Private/VirtualCameraPlayerControllerBase.cpp
@@ -377,7 +377,7 @@ void AVirtualCameraPlayerControllerBase::SetFocusDistanceToActor(const ETouchInd
 		VCPawn->SetFocusDistance(Hit.Distance);
 
 		// Set this actor as the tracked actor, and move the exact focus tracked point to where the touch occurred
-		if (AActor* HitActor = Hit.GetActor())
+		if (AActor* HitActor = Hit.GetHitObjectHandle().FetchActor())
 		{
 			FVector TrackingPointOffset = HitActor->GetActorRotation().UnrotateVector(Hit.ImpactPoint - HitActor->GetActorLocation());
 			TrackingPointOffset /= HitActor->GetActorScale(); // Adjust for non-standard scales when we rotate the vector
@@ -430,7 +430,7 @@ void AVirtualCameraPlayerControllerBase::SetFocusDistanceThroughPoint(const FVec
 	FVector const TraceEnd = CameraWorldLocation + TraceDirection * FocusTraceDist;
 	bool const bHit = GetWorld()->LineTraceSingleByChannel(Hit, CameraWorldLocation, TraceEnd, ECC_Visibility, TraceParams);
 
-	if (bHit && Hit.GetActor())
+	if (bHit && Hit.GetHitObjectHandle().IsValid())
 	{
 		VCPawn->SetFocusDistance(Hit.Distance);
 	}	
diff --git a/Engine/Plugins/Runtime/GameplayAbilities/Source/GameplayAbilities/Private/Abilities/GameplayAbilityTargetActor_GroundTrace.cpp b/Engine/Plugins/Runtime/GameplayAbilities/Source/GameplayAbilities/Private/Abilities/GameplayAbilityTargetActor_GroundTrace.cpp
--- a/Engine/Plugins/Runtime/GameplayAbilities/Source/GameplayAbilities/Private/Abilities/GameplayAbilityTargetActor_GroundTrace.cpp
+++ b/Engine/Plugins/Runtime/GameplayAbilities/Source/GameplayAbilities/Private/Abilities/GameplayAbilityTargetActor_GroundTrace.cpp
@@ -84,7 +84,7 @@ bool AGameplayAbilityTargetActor_GroundTrace::AdjustCollisionResultForShape(cons
 		SweepWithFilter(LocalResult, ThisWorld, Filter, TraceStart, TraceEnd, FQuat::Identity, CollisionShape, TraceProfile.Name, Params);
 		if (!LocalResult.bStartPenetrating)
 		{
-			if (!LocalResult.bBlockingHit || (LocalResult.Actor.IsValid() && Cast<APawn>(LocalResult.Actor.Get())))
+			if (!LocalResult.bBlockingHit || (LocalResult.HitObjectHandle.IsValid() && LocalResult.HitObjectHandle.DoesRepresentClass(APawn::StaticClass())))
 			{
 				//Off the map, or hit an actor
 #if ENABLE_DRAW_DEBUG
diff --git a/Engine/Plugins/Runtime/GameplayAbilities/Source/GameplayAbilities/Private/Abilities/GameplayAbilityTargetActor_Radius.cpp b/Engine/Plugins/Runtime/GameplayAbilities/Source/GameplayAbilities/Private/Abilities/GameplayAbilityTargetActor_Radius.cpp
--- a/Engine/Plugins/Runtime/GameplayAbilities/Source/GameplayAbilities/Private/Abilities/GameplayAbilityTargetActor_Radius.cpp
+++ b/Engine/Plugins/Runtime/GameplayAbilities/Source/GameplayAbilities/Private/Abilities/GameplayAbilityTargetActor_Radius.cpp
@@ -63,7 +63,7 @@ TArray<TWeakObjectPtr<AActor> >	AGameplayAbilityTargetActor_Radius::PerformOverl
 	for (int32 i = 0; i < Overlaps.Num(); ++i)
 	{
 		//Should this check to see if these pawns are in the AimTarget list?
-		APawn* PawnActor = Cast<APawn>(Overlaps[i].GetActor());
+		APawn* PawnActor = Overlaps[i].OverlapObjectHandle.FetchActor<APawn>();
 		if (PawnActor && !HitActors.Contains(PawnActor) && Filter.FilterPassesForActor(PawnActor))
 		{
 			HitActors.Add(PawnActor);
diff --git a/Engine/Plugins/Runtime/GameplayAbilities/Source/GameplayAbilities/Private/Abilities/GameplayAbilityTargetActor_SingleLineTrace.cpp b/Engine/Plugins/Runtime/GameplayAbilities/Source/GameplayAbilities/Private/Abilities/GameplayAbilityTargetActor_SingleLineTrace.cpp
--- a/Engine/Plugins/Runtime/GameplayAbilities/Source/GameplayAbilities/Private/Abilities/GameplayAbilityTargetActor_SingleLineTrace.cpp
+++ b/Engine/Plugins/Runtime/GameplayAbilities/Source/GameplayAbilities/Private/Abilities/GameplayAbilityTargetActor_SingleLineTrace.cpp
@@ -2,6 +2,7 @@
 
 #include "Abilities/GameplayAbilityTargetActor_SingleLineTrace.h"
 #include "DrawDebugHelpers.h"
+#include "GameFramework/LightWeightInstanceSubsystem.h"
 
 // --------------------------------------------------------------------------------------------------------------------------------------------------------
 //
@@ -40,8 +41,8 @@ FHitResult AGameplayAbilityTargetActor_SingleLineTrace::PerformTrace(AActor* InS
 	}
 	if (AGameplayAbilityWorldReticle* LocalReticleActor = ReticleActor.Get())
 	{
-		const bool bHitActor = (ReturnHitResult.bBlockingHit && (ReturnHitResult.Actor != NULL));
-		const FVector ReticleLocation = (bHitActor && LocalReticleActor->bSnapToTargetedActor) ? ReturnHitResult.Actor->GetActorLocation() : ReturnHitResult.Location;
+		const bool bHitActor = (ReturnHitResult.bBlockingHit && (ReturnHitResult.HitObjectHandle.IsValid()));
+		const FVector ReticleLocation = (bHitActor && LocalReticleActor->bSnapToTargetedActor) ? FLightWeightInstanceSubsystem::Get().GetLocation(ReturnHitResult.HitObjectHandle) : ReturnHitResult.Location;
 
 		LocalReticleActor->SetActorLocation(ReticleLocation);
 		LocalReticleActor->SetIsTargetAnActor(bHitActor);
diff --git a/Engine/Plugins/Runtime/GameplayAbilities/Source/GameplayAbilities/Private/Abilities/GameplayAbilityTargetActor_Trace.cpp b/Engine/Plugins/Runtime/GameplayAbilities/Source/GameplayAbilities/Private/Abilities/GameplayAbilityTargetActor_Trace.cpp
--- a/Engine/Plugins/Runtime/GameplayAbilities/Source/GameplayAbilities/Private/Abilities/GameplayAbilityTargetActor_Trace.cpp
+++ b/Engine/Plugins/Runtime/GameplayAbilities/Source/GameplayAbilities/Private/Abilities/GameplayAbilityTargetActor_Trace.cpp
@@ -45,7 +45,7 @@ void AGameplayAbilityTargetActor_Trace::LineTraceWithFilter(FHitResult& OutHitRe
 	{
 		const FHitResult& Hit = HitResults[HitIdx];
 
-		if (!Hit.Actor.IsValid() || FilterHandle.FilterPassesForActor(Hit.Actor))
+		if (!Hit.HitObjectHandle.IsValid() || FilterHandle.FilterPassesForActor(Hit.HitObjectHandle.FetchActor()))
 		{
 			OutHitResult = Hit;
 			OutHitResult.bBlockingHit = true; // treat it as a blocking hit
@@ -68,7 +68,7 @@ void AGameplayAbilityTargetActor_Trace::SweepWithFilter(FHitResult& OutHitResult
 	{
 		const FHitResult& Hit = HitResults[HitIdx];
 
-		if (!Hit.Actor.IsValid() || FilterHandle.FilterPassesForActor(Hit.Actor))
+		if (!Hit.HitObjectHandle.IsValid() || FilterHandle.FilterPassesForActor(Hit.HitObjectHandle.FetchActor()))
 		{
 			OutHitResult = Hit;
 			OutHitResult.bBlockingHit = true; // treat it as a blocking hit
diff --git a/Engine/Plugins/Runtime/GameplayAbilities/Source/GameplayAbilities/Public/Abilities/GameplayAbilityTargetTypes.h b/Engine/Plugins/Runtime/GameplayAbilities/Source/GameplayAbilities/Public/Abilities/GameplayAbilityTargetTypes.h
--- a/Engine/Plugins/Runtime/GameplayAbilities/Source/GameplayAbilities/Public/Abilities/GameplayAbilityTargetTypes.h
+++ b/Engine/Plugins/Runtime/GameplayAbilities/Source/GameplayAbilities/Public/Abilities/GameplayAbilityTargetTypes.h
@@ -576,9 +576,9 @@ struct GAMEPLAYABILITIES_API FGameplayAbilityTargetData_SingleTargetHit : public
 	virtual TArray<TWeakObjectPtr<AActor> >	GetActors() const override
 	{
 		TArray<TWeakObjectPtr<AActor> >	Actors;
-		if (HitResult.Actor.IsValid())
+		if (HitResult.HasValidHitObjectHandle())
 		{
-			Actors.Push(HitResult.Actor);
+			Actors.Push(HitResult.HitObjectHandle.FetchActor());
 		}
 		return Actors;
 	}
diff --git a/Engine/Plugins/Tests/EditorTests/Source/EditorTests/Private/UnrealEd/CollisionAutomationTests.cpp b/Engine/Plugins/Tests/EditorTests/Source/EditorTests/Private/UnrealEd/CollisionAutomationTests.cpp
--- a/Engine/Plugins/Tests/EditorTests/Source/EditorTests/Private/UnrealEd/CollisionAutomationTests.cpp
+++ b/Engine/Plugins/Tests/EditorTests/Source/EditorTests/Private/UnrealEd/CollisionAutomationTests.cpp
@@ -199,7 +199,7 @@ bool FComponentSweepMultiTest::RunTest(const FString& Parameters)
 			{
 				for (int32 iHits = 0; iHits < OutHits.Num(); iHits++)
 				{
-					AActor* EachActor = OutHits[iHits].GetActor();
+					AActor* EachActor = OutHits[iHits].GetHitObjectHandle().FetchActor();
 					if (EachActor == TestRayMeshActor)
 					{
 						BlockedBySpecified = true;	
@@ -271,7 +271,7 @@ bool FLineTraceSingleByChannel::RunTest(const FString& Parameters)
 			bool BlockedBySpecified = false;
 			if (WasBlocked == true)
 			{
-				if (OutHit.GetActor() == TestRayMeshActor)
+				if (OutHit.GetHitObjectHandle() == TestRayMeshActor)
 				{
 					BlockedBySpecified = true;
 					// This generates a snippet you can copy/paste into the ini file for test validation
diff --git a/Engine/Source/Developer/AssetTools/Private/AssetTools.cpp b/Engine/Source/Developer/AssetTools/Private/AssetTools.cpp
--- a/Engine/Source/Developer/AssetTools/Private/AssetTools.cpp
+++ b/Engine/Source/Developer/AssetTools/Private/AssetTools.cpp
@@ -74,6 +74,7 @@
 #include "AssetTypeActions/AssetTypeActions_InterpData.h"
 #include "AssetTypeActions/AssetTypeActions_LandscapeLayer.h"
 #include "AssetTypeActions/AssetTypeActions_LandscapeGrassType.h"
+#include "AssetTypeActions/AssetTypeActions_LightWeightInstance.h"
 #include "AssetTypeActions/AssetTypeActions_Material.h"
 #include "AssetTypeActions/AssetTypeActions_MaterialFunction.h"
 #include "AssetTypeActions/AssetTypeActions_MaterialFunctionInstance.h"
@@ -246,6 +247,7 @@ UAssetToolsImpl::UAssetToolsImpl(const FObjectInitializer& ObjectInitializer)
 	RegisterAssetTypeActions(MakeShareable(new FAssetTypeActions_InterpData));
 	RegisterAssetTypeActions(MakeShareable(new FAssetTypeActions_LandscapeLayer));
 	RegisterAssetTypeActions(MakeShareable(new FAssetTypeActions_LandscapeGrassType(FoliageCategoryBit)));
+	RegisterAssetTypeActions(MakeShareable(new FAssetTypeActions_LightWeightInstance));
 	RegisterAssetTypeActions(MakeShareable(new FAssetTypeActions_Material(BlendablesCategoryBit)));
 	RegisterAssetTypeActions(MakeShareable(new FAssetTypeActions_MaterialFunction));
 	RegisterAssetTypeActions(MakeShareable(new FAssetTypeActions_MaterialFunctionLayer));
diff --git a/Engine/Source/Developer/AssetTools/Private/AssetTypeActions/AssetTypeActions_LightWeightInstance.cpp b/Engine/Source/Developer/AssetTools/Private/AssetTypeActions/AssetTypeActions_LightWeightInstance.cpp
new file mode 100644
--- /dev/null
+++ b/Engine/Source/Developer/AssetTools/Private/AssetTypeActions/AssetTypeActions_LightWeightInstance.cpp
@@ -0,0 +1,15 @@
+// Copyright Epic Games, Inc. All Rights Reserved.
+
+#include "AssetTypeActions/AssetTypeActions_LightWeightInstance.h"
+#include "ToolMenus.h"
+#include "Misc/FileHelper.h"
+#include "EditorFramework/AssetImportData.h"
+#include "Misc/MessageDialog.h"
+#include "Framework/Application/SlateApplication.h"
+
+#include "DesktopPlatformModule.h"
+#include "AssetToolsModule.h"
+
+#define LOCTEXT_NAMESPACE "AssetTypeActions"
+
+#undef LOCTEXT_NAMESPACE
diff --git a/Engine/Source/Developer/AssetTools/Private/AssetTypeActions/AssetTypeActions_LightWeightInstance.h b/Engine/Source/Developer/AssetTools/Private/AssetTypeActions/AssetTypeActions_LightWeightInstance.h
new file mode 100644
--- /dev/null
+++ b/Engine/Source/Developer/AssetTools/Private/AssetTypeActions/AssetTypeActions_LightWeightInstance.h
@@ -0,0 +1,25 @@
+// Copyright Epic Games, Inc. All Rights Reserved.
+
+#pragma once
+
+#include "CoreMinimal.h"
+#include "Toolkits/IToolkitHost.h"
+#include "AssetTypeActions_Base.h"
+#include "GameFramework/LightWeightInstanceManager.h"
+
+
+class FAssetTypeActions_LightWeightInstance : public FAssetTypeActions_Base
+{
+public:
+	virtual FColor GetTypeColor() const override { return FColor(255, 105, 180); }
+	virtual uint32 GetCategories() override { return EAssetTypeCategories::Gameplay; }
+
+	// IAssetTypeActions Implementation
+	virtual FText GetName() const override { return NSLOCTEXT("AssetTypeActions", "AssetTypeActions_LightWeightInstance", "Light Weight Instance"); }
+	virtual UClass* GetSupportedClass() const override { return ALightWeightInstanceManager::StaticClass(); }
+//	virtual void GetActions(const TArray<UObject*>& InObjects, struct FToolMenuSection& Section) override;
+//	virtual void OpenAssetEditor( const TArray<UObject*>& InObjects, TSharedPtr<class IToolkitHost> EditWithinLevelEditor = TSharedPtr<IToolkitHost>() ) override;
+//	virtual void GetResolvedSourceFilePaths(const TArray<UObject*>& TypeAssets, TArray<FString>& OutSourceFilePaths) const override;
+	// End IAssetTypeActions
+};
+
diff --git a/Engine/Source/Developer/CollisionAnalyzer/Private/CollisionAnalyzer.cpp b/Engine/Source/Developer/CollisionAnalyzer/Private/CollisionAnalyzer.cpp
--- a/Engine/Source/Developer/CollisionAnalyzer/Private/CollisionAnalyzer.cpp
+++ b/Engine/Source/Developer/CollisionAnalyzer/Private/CollisionAnalyzer.cpp
@@ -37,7 +37,7 @@ void SerializeHitResult(FArchive& Ar, FHitResult& Result)
 	Ar << Result.PenetrationDepth;
 	Ar << Result.BoneName;
 	Ar << Result.PhysMaterial;
-	Ar << Result.Actor;
+	Ar << Result.HitObjectHandle;
 	Ar << Result.Component;
 	Ar << Result.FaceIndex;
 }
diff --git a/Engine/Source/Developer/FunctionalTesting/Private/TraceQueryTestResults.cpp b/Engine/Source/Developer/FunctionalTesting/Private/TraceQueryTestResults.cpp
--- a/Engine/Source/Developer/FunctionalTesting/Private/TraceQueryTestResults.cpp
+++ b/Engine/Source/Developer/FunctionalTesting/Private/TraceQueryTestResults.cpp
@@ -47,11 +47,10 @@ FString UTraceQueryTestResults::ToString()
 void CaptureNameHelper(FTraceQueryTestNames& Names, FHitResult& HitResult)
 {
 	UPrimitiveComponent* HitComp = HitResult.GetComponent();
-	AActor* HitActor = HitResult.GetActor();
 	UPhysicalMaterial* PhysMat = HitResult.PhysMaterial.Get();
 
 	Names.ComponentName = HitComp ? HitComp->GetFName() : NAME_None;
-	Names.ActorName = HitActor ? HitActor->GetFName() : NAME_None;
+	Names.ActorName = HitResult.HitObjectHandle.IsValid() ? HitResult.HitObjectHandle.GetFName() : NAME_None;
 	Names.PhysicalMaterialName = PhysMat ? PhysMat->GetFName() : NAME_None;
 }
 
diff --git a/Engine/Source/Developer/LogVisualizer/Private/VisualLoggerHUD.cpp b/Engine/Source/Developer/LogVisualizer/Private/VisualLoggerHUD.cpp
--- a/Engine/Source/Developer/LogVisualizer/Private/VisualLoggerHUD.cpp
+++ b/Engine/Source/Developer/LogVisualizer/Private/VisualLoggerHUD.cpp
@@ -51,7 +51,7 @@ void AVisualLoggerHUD::PostRender()
 			bool bHit = GetWorld()->LineTraceSingleByChannel(Hit, CamLoc, CamRot.Vector() * 100000.f + CamLoc, ECC_Pawn, TraceParams);
 			if( bHit )
 			{
-				TextItem.Text = FText::FromString(FString::Printf(TEXT("Under cursor: '%s'"), *Hit.GetActor()->GetName()));
+				TextItem.Text = FText::FromString(FString::Printf(TEXT("Under cursor: '%s'"), *Hit.HitObjectHandle.GetName()));
 				Canvas->DrawItem( TextItem, X, Y );
 				
 				DrawDebugLine( GetWorld(), Hit.Location, Hit.Location+Hit.Normal*30.f, FColor::White );
diff --git a/Engine/Source/Editor/Experimental/EditorInteractiveToolsFramework/Private/EdModeInteractiveToolsContext.cpp b/Engine/Source/Editor/Experimental/EditorInteractiveToolsFramework/Private/EdModeInteractiveToolsContext.cpp
--- a/Engine/Source/Editor/Experimental/EditorInteractiveToolsFramework/Private/EdModeInteractiveToolsContext.cpp
+++ b/Engine/Source/Editor/Experimental/EditorInteractiveToolsFramework/Private/EdModeInteractiveToolsContext.cpp
@@ -272,7 +272,7 @@ public:
 					// if we found a valid snap, return it
 					if (SmallestAngle < Request.VisualAngleThresholdDegrees)
 					{
-						SnapResult.TargetActor = HitResult.Actor.Get();
+						SnapResult.TargetActor = HitResult.HitObjectHandle.FetchActor();
 						SnapResult.TargetComponent = HitResult.Component.Get();
 						Results.Add(SnapResult);
 						FoundResultCount++;
diff --git a/Engine/Source/Editor/UnrealEd/Classes/Factories/LightWeightInstanceFactory.h b/Engine/Source/Editor/UnrealEd/Classes/Factories/LightWeightInstanceFactory.h
new file mode 100644
--- /dev/null
+++ b/Engine/Source/Editor/UnrealEd/Classes/Factories/LightWeightInstanceFactory.h
@@ -0,0 +1,28 @@
+// Copyright Epic Games, Inc. All Rights Reserved.
+
+#pragma once
+
+#include "CoreMinimal.h"
+#include "UObject/ObjectMacros.h"
+#include "Factories/Factory.h"
+#include "LightWeightInstanceFactory.generated.h"
+
+class ALightWeightInstanceManager;
+
+UCLASS(hidecategories=Object)
+class UNREALED_API ULightWeightInstanceFactory : public UFactory
+{
+	GENERATED_UCLASS_BODY()
+
+	//~ Begin UFactory Interface
+	virtual bool ConfigureProperties() override;
+	virtual UObject* FactoryCreateNew(UClass* Class, UObject* InParent, FName Name, EObjectFlags Flags, UObject* Context, FFeedbackContext* Warn) override;
+	//~ Begin UFactory Interface
+
+protected:
+	virtual ALightWeightInstanceManager* MakeNewLightWeightInstanceManager(UObject* InParent, FName Name, EObjectFlags Flags);
+
+	// The parent class of the created blueprint
+	UPROPERTY()
+	TSubclassOf<class UObject> ParentClass;
+};
diff --git a/Engine/Source/Editor/UnrealEd/Private/DebugToolExec.cpp b/Engine/Source/Editor/UnrealEd/Private/DebugToolExec.cpp
--- a/Engine/Source/Editor/UnrealEd/Private/DebugToolExec.cpp
+++ b/Engine/Source/Editor/UnrealEd/Private/DebugToolExec.cpp
@@ -214,7 +214,7 @@ bool FDebugToolExec::Exec( UWorld* InWorld, const TCHAR* Cmd, FOutputDevice& Ar
 				PlayerController->GetPlayerViewPoint(PlayerLocation, PlayerRotation);
 				FHitResult Hit(1.0f);
 				PlayerController->GetWorld()->LineTraceSingleByChannel(Hit, PlayerLocation, PlayerLocation + PlayerRotation.Vector() * 10000.f, ECC_Pawn, FCollisionQueryParams(NAME_None, FCollisionQueryParams::GetUnknownStatId(), true, PlayerController->GetPawn()));
-				Found = Hit.GetActor();
+				Found = Hit.GetHitObjectHandle().FetchActor();
 			}
 		}
 		// Search by class.
diff --git a/Engine/Source/Editor/UnrealEd/Private/Editor/ActorPositioning.cpp b/Engine/Source/Editor/UnrealEd/Private/Editor/ActorPositioning.cpp
--- a/Engine/Source/Editor/UnrealEd/Private/Editor/ActorPositioning.cpp
+++ b/Engine/Source/Editor/UnrealEd/Private/Editor/ActorPositioning.cpp
@@ -56,7 +56,8 @@ FActorPositionTraceResult FActorPositioning::TraceWorldForPosition(const FViewpo
 /** Check to see if the specified hit result should be ignored from actor positioning calculations for the specified scene view */
 bool IsHitIgnored(const FHitResult& InHit, const FSceneView& InSceneView)
 {
-	const auto* Actor = InHit.GetActor();
+	const FActorInstanceHandle& HitObjHandle = InHit.HitObjectHandle;
+	AActor* Actor = HitObjHandle.FetchActor();
 	
 	// Try and find a primitive component for the hit
 	const UPrimitiveComponent* PrimitiveComponent = Actor ? Cast<UPrimitiveComponent>(Actor->GetRootComponent()) : nullptr;
@@ -76,7 +77,7 @@ bool IsHitIgnored(const FHitResult& InHit, const FSceneView& InSceneView)
 	}
 
 	// Ignore volumes and shapes
-	if (Actor && Actor->IsA(AVolume::StaticClass()))
+	if (HitObjHandle.DoesRepresentClass(AVolume::StaticClass()))
 	{
 		return true;
 	}
@@ -126,7 +127,7 @@ FActorPositionTraceResult FActorPositioning::TraceWorldForPosition(const UWorld&
 		// Go through all hits and find closest
 		float ClosestHitDistanceSqr = TNumericLimits<float>::Max();
 
-		for (const auto& Hit : Hits)
+		for (const FHitResult& Hit : Hits)
 		{
 			const float DistanceToHitSqr = (Hit.ImpactPoint - RayStart).SizeSquared();
 			if (DistanceToHitSqr < ClosestHitDistanceSqr)
@@ -135,7 +136,7 @@ FActorPositionTraceResult FActorPositioning::TraceWorldForPosition(const UWorld&
 				Results.Location = Hit.Location;
 				Results.SurfaceNormal = Hit.Normal.GetSafeNormal();
 				Results.State = FActorPositionTraceResult::HitSuccess;
-				Results.HitActor = Hit.Actor;
+				Results.HitActor = Hit.HitObjectHandle.FetchActor();
 			}
 		}
 	}
diff --git a/Engine/Source/Editor/UnrealEd/Private/Factories/LightWeightInstanceFactory.cpp b/Engine/Source/Editor/UnrealEd/Private/Factories/LightWeightInstanceFactory.cpp
new file mode 100644
--- /dev/null
+++ b/Engine/Source/Editor/UnrealEd/Private/Factories/LightWeightInstanceFactory.cpp
@@ -0,0 +1,119 @@
+// Copyright Epic Games, Inc. All Rights Reserved.
+
+#include "Factories/LightWeightInstanceFactory.h"
+#include "GameFramework/LightWeightInstanceManager.h"
+#include "Editor.h"
+
+#include "KismetCompilerModule.h"
+#include "ClassViewerFilter.h"
+
+#include "Widgets/SWindow.h"
+#include "Widgets/Input/SButton.h"
+#include "Widgets/Input/SComboButton.h"
+#include "Widgets/Text/STextBlock.h"
+#include "Widgets/Layout/SBox.h"
+#include "EditorStyleSet.h"
+#include "Input/Reply.h"
+#include "Kismet2/SClassPickerDialog.h"
+
+#define LOCTEXT_NAMESPACE "LightWeightInstanceFactory"
+
+class FLWIFilter : public IClassViewerFilter
+{
+public:
+	/** Classes to not allow any children of into the Class Viewer/Picker. */
+	TSet< const UClass* > AllowedChildrenOfClasses;
+
+	virtual bool IsClassAllowed(const FClassViewerInitializationOptions& InInitOptions, const UClass* InClass, TSharedRef< FClassViewerFilterFuncs > InFilterFuncs) override
+	{
+		return InFilterFuncs->IfInChildOfClassesSet(AllowedChildrenOfClasses, InClass) == EFilterReturn::Passed && !InClass->HasAnyClassFlags(CLASS_Deprecated);
+	}
+
+	virtual bool IsUnloadedClassAllowed(const FClassViewerInitializationOptions& InInitOptions, const TSharedRef< const IUnloadedBlueprintData > InUnloadedClassData, TSharedRef< FClassViewerFilterFuncs > InFilterFuncs) override
+	{
+		return InFilterFuncs->IfInChildOfClassesSet(AllowedChildrenOfClasses, InUnloadedClassData) == EFilterReturn::Passed && !InUnloadedClassData->HasAnyClassFlags(CLASS_Deprecated);
+	}
+};
+
+ULightWeightInstanceFactory::ULightWeightInstanceFactory(const FObjectInitializer& ObjectInitializer)
+	: Super(ObjectInitializer)
+{
+	SupportedClass = ALightWeightInstanceManager::StaticClass();
+	bCreateNew = true;
+	bEditAfterNew = true;
+}
+
+bool ULightWeightInstanceFactory::ConfigureProperties()
+{
+	// Null the parent class to ensure one is selected
+	ParentClass = nullptr;
+
+	// Fill in options
+	FClassViewerInitializationOptions Options;
+	Options.Mode = EClassViewerMode::ClassPicker;
+	Options.DisplayMode = EClassViewerDisplayMode::TreeView;
+	Options.bShowObjectRootClass = true;
+	Options.bIsBlueprintBaseOnly = false;
+
+	// This will allow unloaded blueprints to be shown.
+	Options.bShowUnloadedBlueprints = true;
+
+	// Enable Class Dynamic Loading
+	Options.bEnableClassDynamicLoading = true;
+
+	Options.NameTypeToDisplay = EClassViewerNameTypeToDisplay::Dynamic;
+
+	// Prevent creating blueprints of classes that require special setup (they'll be allowed in the corresponding factories / via other means)
+	TSharedPtr<FLWIFilter> Filter = MakeShareable(new FLWIFilter);
+	Options.ClassFilter = Filter;
+
+	Filter->AllowedChildrenOfClasses.Add(ALightWeightInstanceManager::StaticClass());
+
+	const FText TitleText = LOCTEXT("CreateBlueprintOptions", "Pick Parent Class");
+	UClass* ChosenClass = nullptr;
+	const bool bPressedOk = SClassPickerDialog::PickClass(TitleText, Options, ChosenClass, ALightWeightInstanceManager::StaticClass());
+
+	if (bPressedOk)
+	{
+		ParentClass = ChosenClass;
+
+		FEditorDelegates::OnFinishPickingBlueprintClass.Broadcast(ChosenClass);
+	}
+
+	return bPressedOk;
+};
+
+UObject* ULightWeightInstanceFactory::FactoryCreateNew(UClass* Class, UObject* InParent, FName Name, EObjectFlags Flags, UObject* Context, FFeedbackContext* Warn)
+{
+	UObject* LWIManager = nullptr;
+	if (ensure(SupportedClass == Class))
+	{
+		ensure(0 != (RF_Public & Flags));
+
+		UClass* BlueprintClass = nullptr;
+		UClass* BlueprintGeneratedClass = nullptr;
+
+		IKismetCompilerInterface& KismetCompilerModule = FModuleManager::LoadModuleChecked<IKismetCompilerInterface>("KismetCompiler");
+		KismetCompilerModule.GetBlueprintTypesForClass(ParentClass, BlueprintClass, BlueprintGeneratedClass);
+
+		LWIManager = FKismetEditorUtilities::CreateBlueprint(ParentClass, InParent, Name, BPTYPE_Normal, BlueprintClass, BlueprintGeneratedClass);
+	}
+	return LWIManager;
+}
+
+ALightWeightInstanceManager* ULightWeightInstanceFactory::MakeNewLightWeightInstanceManager(UObject* InParent, FName Name, EObjectFlags Flags)
+{
+	UClass* BlueprintClass = nullptr;
+	UClass* BlueprintGeneratedClass = nullptr;
+
+	IKismetCompilerInterface& KismetCompilerModule = FModuleManager::LoadModuleChecked<IKismetCompilerInterface>("KismetCompiler");
+	KismetCompilerModule.GetBlueprintTypesForClass(ParentClass, BlueprintClass, BlueprintGeneratedClass);
+
+	UBlueprint* NewBP = FKismetEditorUtilities::CreateBlueprint(ParentClass, InParent, Name, BPTYPE_Normal, BlueprintClass, BlueprintGeneratedClass);
+
+	return Cast<ALightWeightInstanceManager>(FKismetEditorUtilities::CreateBlueprint(ParentClass, InParent, Name, BPTYPE_Normal, BlueprintClass, BlueprintGeneratedClass));
+//	return FKismetEditorUtilities::CreateBlueprint(ParentClass, InParent, Name, BPTYPE_Normal, UEditorUtilityBlueprint::StaticClass(), UBlueprintGeneratedClass::StaticClass());
+//	return NewObject<ALightWeightInstanceManager>(InParent, Name, Flags);
+}
+
+#undef LOCTEXT_NAMESPACE // "LightWeightInstanceFactory"
diff --git a/Engine/Source/Editor/VREditor/UI/VREditorUISystem.cpp b/Engine/Source/Editor/VREditor/UI/VREditorUISystem.cpp
--- a/Engine/Source/Editor/VREditor/UI/VREditorUISystem.cpp
+++ b/Engine/Source/Editor/VREditor/UI/VREditorUISystem.cpp
@@ -396,7 +396,7 @@ void UVREditorUISystem::OnPreviewInputAction(FEditorViewportClient& ViewportClie
 		if (VREditorInteractor->GetLaserPointer(LaserPointerStart, LaserPointerEnd))
 		{
 			FHitResult HitResult = VREditorInteractor->GetHitResultFromLaserPointer();
-			if (HitResult.Actor.IsValid())
+			if (HitResult.HitObjectHandle.IsValid())
 			{
 				UWidgetComponent* WidgetComponent = Cast<UWidgetComponent>(HitResult.GetComponent());
 
@@ -593,10 +593,10 @@ void UVREditorUISystem::OnVRHoverUpdate(UViewportInteractor* Interactor, FVector
 			if( Interactor->GetLaserPointer( LaserPointerStart, LaserPointerEnd ) )
 			{
 				FHitResult HitResult = Interactor->GetHitResultFromLaserPointer();
-				if( HitResult.Actor.IsValid() )
+				if( HitResult.HitObjectHandle.IsValid() )
 				{
 					// The laser should make the quick radial menu stay active
-					if ((HitResult.Actor == QuickRadialMenu) && (QuickRadialMenu != nullptr))
+					if ((HitResult.HitObjectHandle == QuickRadialMenu) && (QuickRadialMenu != nullptr))
 					{
 						RadialMenuModifierSpawnTime = FTimespan::FromSeconds(FPlatformTime::Seconds());
 						FVector2D ReturnToCenter = FVector2D::ZeroVector;
@@ -1907,9 +1907,9 @@ bool UVREditorUISystem::ShouldPreviewPanel()
 	{
 		// If we are clicking with the laser interactor, instantly spawn the panel
 		FHitResult HitResult = LaserInteractor->GetHitResultFromLaserPointer();
-		if (HitResult.Actor.IsValid())
+		if (HitResult.HitObjectHandle.IsValid())
 		{
-			if (HitResult.Actor == QuickRadialMenu)
+			if (HitResult.HitObjectHandle == QuickRadialMenu)
 			{
 				bDragFromOpen = !LaserInteractor->IsClickingOnUI();
 			}
diff --git a/Engine/Source/Editor/ViewportInteraction/ViewportInteractor.cpp b/Engine/Source/Editor/ViewportInteraction/ViewportInteractor.cpp
--- a/Engine/Source/Editor/ViewportInteraction/ViewportInteractor.cpp
+++ b/Engine/Source/Editor/ViewportInteraction/ViewportInteractor.cpp
@@ -146,8 +146,8 @@ bool UViewportInteractor::HandleInputKey( FEditorViewportClient& ViewportClient,
 			// Prefer transform gizmo interactions over anything else
 			const bool bPressedTransformGizmo =
 				Event == IE_Pressed &&
-				HitResult.Actor.IsValid() &&
-				HitResult.Actor == WorldInteraction->GetTransformGizmoActor() &&
+				HitResult.HitObjectHandle.IsValid() &&
+				HitResult.HitObjectHandle == WorldInteraction->GetTransformGizmoActor() &&
 				Action->ActionType == ViewportWorldActionTypes::SelectAndMove;
 			if( !bPressedTransformGizmo )
 			{
@@ -177,13 +177,13 @@ bool UViewportInteractor::HandleInputKey( FEditorViewportClient& ViewportClient,
 				if ( Event == IE_Pressed )
 				{
 					// No clicking while we're dragging the world.  (No laser pointers are visible, anyway.)
-					if ( !bIsDraggingWorldWithTwoHands && HitResult.Actor.IsValid() )
+					if ( !bIsDraggingWorldWithTwoHands && HitResult.HitObjectHandle.IsValid() )
 					{
 						if ( WorldInteraction->IsInteractableComponent( HitResult.GetComponent() ) )
 						{
 							InteractorData.ClickingOnComponent = HitResult.GetComponent();	// @todo gizmo: Should be changed to store only gizmo components?
 
-							AActor* Actor = HitResult.Actor.Get();
+							AActor* Actor = HitResult.HitObjectHandle.FetchActor();
 
 							FVector LaserPointerStart, LaserPointerEnd;
 							const bool bHaveLaserPointer = GetLaserPointer( /* Out */ LaserPointerStart, /* Out */ LaserPointerEnd );
@@ -675,7 +675,7 @@ FHitResult UViewportInteractor::GetHitResultFromLaserPointer( TArray<AActor*>* O
 {
 	FHitResult BestHitResult;
 
-	if (SavedHitResult.IsSet() && OptionalListOfIgnoredActors != nullptr && !OptionalListOfIgnoredActors->Contains(SavedHitResult->Actor) && (SavedHitResultFilterMode.IsSet() && (GizmoFilterMode == SavedHitResultFilterMode.GetValue())))
+	if (SavedHitResult.IsSet() && OptionalListOfIgnoredActors != nullptr && !OptionalListOfIgnoredActors->Contains(SavedHitResult->HitObjectHandle.FetchActor()) && (SavedHitResultFilterMode.IsSet() && (GizmoFilterMode == SavedHitResultFilterMode.GetValue())))
 	{
 		BestHitResult = SavedHitResult.GetValue();
 	}
@@ -760,7 +760,7 @@ FHitResult UViewportInteractor::GetHitResultFromLaserPointer( TArray<AActor*>* O
 								bool bClassHasPriority = false;
 								bClassHasPriority =
 									( HitResult.GetComponent() != nullptr && HitResult.GetComponent()->IsA( CurrentClass ) ) ||
-									( HitResult.GetActor() != nullptr && HitResult.GetActor()->IsA( CurrentClass ) );
+									( HitResult.HitObjectHandle.IsValid() && HitResult.HitObjectHandle.DoesRepresentClass( CurrentClass ) );
 
 								if ( bClassHasPriority )
 								{
@@ -770,8 +770,8 @@ FHitResult UViewportInteractor::GetHitResultFromLaserPointer( TArray<AActor*>* O
 							}
 						}
 
-						const bool bHitResultIsGizmo = HitResult.GetActor() != nullptr && HitResult.GetActor() == WorldInteraction->GetTransformGizmoActor();
-						if ( BestHitResult.GetActor() == nullptr ||
+						const bool bHitResultIsGizmo = HitResult.HitObjectHandle.IsValid() && HitResult.HitObjectHandle == WorldInteraction->GetTransformGizmoActor();
+						if ( !BestHitResult.HitObjectHandle.IsValid() ||
 							 InteractorData.bHitResultIsPriorityType ||
 							 bHitResultIsGizmo )
 						{
diff --git a/Engine/Source/Editor/ViewportInteraction/ViewportWorldInteraction.cpp b/Engine/Source/Editor/ViewportInteraction/ViewportWorldInteraction.cpp
--- a/Engine/Source/Editor/ViewportInteraction/ViewportWorldInteraction.cpp
+++ b/Engine/Source/Editor/ViewportInteraction/ViewportWorldInteraction.cpp
@@ -830,8 +830,8 @@ void UViewportWorldInteraction::HoverTick( const float DeltaTime )
 		FHitResult HitHoverResult = Interactor->GetHitResultFromLaserPointer();
 
 		const bool bIsHoveringOverTransformGizmo =
-			HitHoverResult.Actor.IsValid() &&
-			HitHoverResult.Actor == TransformGizmoActor;
+			HitHoverResult.HitObjectHandle.IsValid() &&
+			HitHoverResult.HitObjectHandle == TransformGizmoActor;
 
 		// Prefer transform gizmo hover over everything else
 		if( !bIsHoveringOverTransformGizmo )
@@ -848,13 +848,13 @@ void UViewportWorldInteraction::HoverTick( const float DeltaTime )
 		{
 			UActorComponent* NewHoveredActorComponent = nullptr;
 
-			if ( HitHoverResult.Actor.IsValid() )
+			if ( HitHoverResult.HitObjectHandle.IsValid() )
 			{
-				USceneComponent* HoveredActorComponent = HitHoverResult.GetComponent();
-				AActor* Actor = HitHoverResult.Actor.Get();
+				USceneComponent* HoveredActorComponent = HitHoverResult.GetComponent();	
 
 				if ( HoveredActorComponent && IsInteractableComponent( HoveredActorComponent ) )
 				{
+					AActor* Actor = HitHoverResult.HitObjectHandle.FetchActor();
 					HoveredObjects.Add( FViewportHoverTarget( Actor ) );
 
 					Interactor->SetHoverLocation(HitHoverResult.ImpactPoint);
@@ -2540,7 +2540,7 @@ bool UViewportWorldInteraction::FindPlacementPointUnderLaser( UViewportInteracto
 	const EHitResultGizmoFilterMode GizmoFilterMode = EHitResultGizmoFilterMode::NoGizmos; // Never place on top of gizmos, just ignore them
 	const bool bEvenIfUIIsInFront = true;	// Don't let the UI block placement
 	FHitResult HitResult = Interactor->GetHitResultFromLaserPointer( &IgnoredActors, GizmoFilterMode, nullptr, bEvenIfUIIsInFront );
-	if( HitResult.Actor.IsValid() )
+	if( HitResult.HitObjectHandle.IsValid() )
 	{
 		bHitSomething = true;
 		HitLocation = HitResult.ImpactPoint;
@@ -3454,7 +3454,7 @@ FVector UViewportWorldInteraction::FindTransformGizmoAlignPoint(const FTransform
 		{
 			for (FOverlapResult OverlapResult : OutOverlaps)
 			{
-				const AActor* PossibleCandidateActor = OverlapResult.GetActor();
+				const AActor* PossibleCandidateActor = OverlapResult.OverlapObjectHandle.FetchActor();
 
 				// Don't align to yourself, the entire world, or any actors hidden in the editor
 				if (PossibleCandidateActor != nullptr
diff --git a/Engine/Source/Runtime/AIModule/Private/Perception/AISense_Sight.cpp b/Engine/Source/Runtime/AIModule/Private/Perception/AISense_Sight.cpp
--- a/Engine/Source/Runtime/AIModule/Private/Perception/AISense_Sight.cpp
+++ b/Engine/Source/Runtime/AIModule/Private/Perception/AISense_Sight.cpp
@@ -331,7 +331,7 @@ float UAISense_Sight::Update()
 
 						auto HitResultActorIsOwnedByTargetActor = [&HitResult, TargetActor]()
 						{
-							AActor* HitResultActor = HitResult.Actor.Get();
+							AActor* HitResultActor = HitResult.HitObjectHandle.FetchActor();
 							return (HitResultActor ? HitResultActor->IsOwnedBy(TargetActor) : false);
 						};
 
diff --git a/Engine/Source/Runtime/Engine/Classes/Components/HierarchicalInstancedStaticMeshComponent.h b/Engine/Source/Runtime/Engine/Classes/Components/HierarchicalInstancedStaticMeshComponent.h
--- a/Engine/Source/Runtime/Engine/Classes/Components/HierarchicalInstancedStaticMeshComponent.h
+++ b/Engine/Source/Runtime/Engine/Classes/Components/HierarchicalInstancedStaticMeshComponent.h
@@ -137,6 +137,8 @@ class ENGINE_API UHierarchicalInstancedStaticMeshComponent : public UInstancedSt
 {
 	GENERATED_UCLASS_BODY()
 
+	friend class ALightWeightInstanceStaticMeshManager;
+
 	~UHierarchicalInstancedStaticMeshComponent();
 
 	TSharedPtr<TArray<FClusterNode>, ESPMode::ThreadSafe> ClusterTreePtr;
diff --git a/Engine/Source/Runtime/Engine/Classes/Components/InstancedStaticMeshComponent.h b/Engine/Source/Runtime/Engine/Classes/Components/InstancedStaticMeshComponent.h
--- a/Engine/Source/Runtime/Engine/Classes/Components/InstancedStaticMeshComponent.h
+++ b/Engine/Source/Runtime/Engine/Classes/Components/InstancedStaticMeshComponent.h
@@ -124,6 +124,8 @@ UCLASS(ClassGroup = Rendering, meta = (BlueprintSpawnableComponent), Blueprintab
 class ENGINE_API UInstancedStaticMeshComponent : public UStaticMeshComponent
 {
 	GENERATED_UCLASS_BODY()
+
+	friend class ALightWeightInstanceManager;
 	
 	/** Needs implementation in InstancedStaticMesh.cpp to compile UniquePtr for forward declared class */
 	UInstancedStaticMeshComponent(FVTableHelper& Helper);
diff --git a/Engine/Source/Runtime/Engine/Classes/Engine/EngineTypes.h b/Engine/Source/Runtime/Engine/Classes/Engine/EngineTypes.h
--- a/Engine/Source/Runtime/Engine/Classes/Engine/EngineTypes.h
+++ b/Engine/Source/Runtime/Engine/Classes/Engine/EngineTypes.h
@@ -1917,6 +1917,74 @@ struct FPrimitiveMaterialRef
 	{ 	}
 };
 
+// Handle to a unique object. This may specify a full weigh actor or it may only specify the light weight instance that represents the same object.
+USTRUCT(BlueprintType)
+struct ENGINE_API FActorInstanceHandle
+{
+	GENERATED_BODY()
+
+	friend struct FLightWeightInstanceSubsystem;
+	friend class ALightWeightInstanceManager;
+
+	FActorInstanceHandle();
+
+	explicit FActorInstanceHandle(AActor* InActor);
+	explicit FActorInstanceHandle(int32 InManagerIndex, int32 InInstanceIndex);
+	explicit FActorInstanceHandle(class ALightWeightInstanceManager* Manager, int32 InInstanceIndex);
+
+	FActorInstanceHandle(const FActorInstanceHandle& Other);
+
+	bool IsValid() const;
+
+	bool DoesRepresentClass(const UClass* OtherClass) const;
+
+	UClass* GetRepresentedClass() const;
+
+	FVector GetLocation() const;
+	FRotator GetRotation() const;
+
+	ULevel* GetLevel() const;
+	bool IsInLevel(ULevel* InLevel) const;
+
+	FName GetFName() const;
+	FString GetName() const;
+
+	// Returns the actor specified by this handle. This may require loading and creating the actor object.
+	AActor* FetchActor() const;
+	template <typename T>
+	T* FetchActor() const;
+
+	FORCEINLINE int32 GetInstanceIndex() const { return InstanceIndex; }
+
+	FActorInstanceHandle& operator=(const FActorInstanceHandle& Other) = default;
+	FActorInstanceHandle& operator=(FActorInstanceHandle&& Other) = default;
+
+	bool operator==(const FActorInstanceHandle& Other) const;
+	bool operator!=(const FActorInstanceHandle& Other) const;
+
+	bool operator==(const AActor* OtherActor) const;
+	bool operator!=(const AActor* OtherActor) const;
+
+	friend ENGINE_API FArchive& operator<<(FArchive& Ar, FActorInstanceHandle& Handle);
+
+private:
+	// this is cached here for convenience
+	UPROPERTY()
+	mutable TWeakObjectPtr<AActor> Actor;
+
+	// Identifies the light weight instance manager to use
+	int32 ManagerIndex;
+
+	// Identifies the instance within the manager
+	int32 InstanceIndex;
+};
+
+template <typename T>
+T* FActorInstanceHandle::FetchActor() const
+{
+	return Cast<T>(FetchActor());
+}
+
 /**
  * Structure containing information about one hit of a trace, such as point of impact and surface normal at that point.
  */
@@ -2022,9 +2090,9 @@ struct ENGINE_API FHitResult
 	UPROPERTY()
 	TWeakObjectPtr<class UPhysicalMaterial> PhysMaterial;
 
-	/** Actor hit by the trace. */
+	/** Handle to the object hit by the trace. */
 	UPROPERTY()
-	TWeakObjectPtr<class AActor> Actor;
+	FActorInstanceHandle HitObjectHandle;
 
 	/** PrimitiveComponent hit by the trace. */
 	UPROPERTY()
@@ -2097,10 +2165,14 @@ struct ENGINE_API FHitResult
 		}
 	}
 
-	/** Utility to return the Actor that owns the Component that was hit. */
-	FORCEINLINE AActor* GetActor() const
+	FORCEINLINE FActorInstanceHandle GetHitObjectHandle() const
 	{
-		return Actor.Get();
+		return HitObjectHandle;
+	}
+
+	FORCEINLINE bool HasValidHitObjectHandle() const
+	{
+		return HitObjectHandle.IsValid();
 	}
 
 	/** Utility to return the Component that was hit. */
@@ -2203,9 +2275,8 @@ struct ENGINE_API FOverlapResult
 {
 	GENERATED_BODY()
 
-	/** Actor that the check hit. */
 	UPROPERTY()
-	TWeakObjectPtr<class AActor> Actor;
+	FActorInstanceHandle OverlapObjectHandle;
 
 	/** PrimitiveComponent that the check hit. */
 	UPROPERTY()
@@ -2216,9 +2287,6 @@ struct ENGINE_API FOverlapResult
 		For SkeletalMeshComponents this is the Body index or INDEX_NONE for single body */
 	int32 ItemIndex;
 
-	/** Utility to return the Actor that owns the Component that was hit */
-	AActor* GetActor() const;
-
 	/** Utility to return the Component that was hit */
 	UPrimitiveComponent* GetComponent() const;
 
diff --git a/Engine/Source/Runtime/Engine/Classes/GameFramework/Actor.h b/Engine/Source/Runtime/Engine/Classes/GameFramework/Actor.h
--- a/Engine/Source/Runtime/Engine/Classes/GameFramework/Actor.h
+++ b/Engine/Source/Runtime/Engine/Classes/GameFramework/Actor.h
@@ -2610,6 +2610,7 @@ public:
 	bool IsInLevel(const class ULevel *TestLevel) const;
 
 	/** Return the ULevel that this Actor is part of. */
+	UFUNCTION(BlueprintCallable, Category=Level)
 	ULevel* GetLevel() const;
 
 	/**	Do anything needed to clear out cross level references; Called from ULevel::PreSave	 */
diff --git a/Engine/Source/Runtime/Engine/Classes/GameFramework/LightWeightInstanceBlueprintFunctionLibrary.h b/Engine/Source/Runtime/Engine/Classes/GameFramework/LightWeightInstanceBlueprintFunctionLibrary.h
new file mode 100644
--- /dev/null
+++ b/Engine/Source/Runtime/Engine/Classes/GameFramework/LightWeightInstanceBlueprintFunctionLibrary.h
@@ -0,0 +1,17 @@
+// Copyright Epic Games, Inc. All Rights Reserved.
+
+#pragma once
+
+#include "Kismet/BlueprintFunctionLibrary.h"
+#include "LightWeightInstanceBlueprintFunctionLibrary.generated.h"
+
+UCLASS()
+class ENGINE_API ULightWeightInstanceBlueprintFunctionLibrary : public UBlueprintFunctionLibrary
+{
+	GENERATED_BODY()
+
+public:
+
+	UFUNCTION(BlueprintCallable, Category = "Light Weight Instance")
+	static FActorInstanceHandle CreateNewLightWeightInstance(UClass* ActorClass, FTransform Transform, ULevel* Level);
+};
\ No newline at end of file
diff --git a/Engine/Source/Runtime/Engine/Classes/GameFramework/LightWeightInstanceManager.h b/Engine/Source/Runtime/Engine/Classes/GameFramework/LightWeightInstanceManager.h
new file mode 100644
--- /dev/null
+++ b/Engine/Source/Runtime/Engine/Classes/GameFramework/LightWeightInstanceManager.h
@@ -0,0 +1,140 @@
+// Copyright Epic Games, Inc. All Rights Reserved.
+
+#pragma once
+
+#include "CoreMinimal.h"
+#include "UObject/ObjectMacros.h"
+#include "UObject/UObjectBaseUtility.h"
+#include "UObject/Object.h"
+#include "Engine/EngineTypes.h"
+#include "Engine/EngineBaseTypes.h"
+#include "Templates/SharedPointer.h"
+
+#include "Actor.h"
+
+#include "LightWeightInstanceManager.generated.h"
+
+//DECLARE_DYNAMIC_DELEGATE_RetVal_OneParam(FActorInstanceHandle, FOnActorReady, FActorInstanceHandle, InHandle);
+
+// Used for initializing light weight instances.
+struct ENGINE_API FLWIData
+{
+	FTransform Transform;
+};
+
+UCLASS(BlueprintType, Blueprintable)
+class ENGINE_API ALightWeightInstanceManager : public AActor
+{
+	GENERATED_UCLASS_BODY()
+
+	friend struct FLightWeightInstanceSubsystem;
+	friend struct FActorInstanceHandle;
+	friend class ULightWeightInstanceBlueprintFunctionLibrary;
+
+public:
+	virtual ~ALightWeightInstanceManager();
+
+	virtual void Tick(float DeltaSeconds);
+
+	// Returns the location of the instance specified by Handle
+	FVector GetLocation(const FActorInstanceHandle& Handle) const;
+
+	// Returns the name of the instance specified by Handle
+	FString GetName(const FActorInstanceHandle& Handle) const;
+
+	// Returns true if this manager stores instances that can be turned into full weight objects of class OtherClass
+	bool DoesRepresentClass(const UClass* OtherClass) const;
+
+	// Returns true if this manager is capable of representing objects of type OtherClass
+	bool DoesAcceptClass(const UClass* OtherClass) const;
+
+	// Returns the specific class that this manages
+	UClass* GetRepresentedClass() const;
+
+	// Returns the base class of types that this can manage
+	UClass* GetAcceptedClass() const;
+
+	// Sets the specific class that this manages
+	virtual void SetRepresentedClass(UClass* ActorClass);
+
+	// Returns the actor associated with Handle if one exists
+	AActor* GetActorFromHandle(const FActorInstanceHandle& Handle);
+
+	// Returns the index of the light weight instance associated with InActor if one exists; otherwise we return INDEX_NONE
+	int32 FindIndexForActor(const AActor* InActor) const;
+
+	virtual int32 ConvertCollisionIndexToLightWeightIndex(int32 InIndex) const;
+
+protected:
+	// Takes a polymorphic struct to set the initial data for a new instance
+	int32 AddNewInstance(FLWIData* InitData);
+
+	// Adds a new instance at the specified index. This function should only be called by AddNewInstance.
+	// Provided to allow subclasses to update their per instance data
+	virtual void AddNewInstanceAt(FLWIData* InitData, int32 Index);
+
+	// Removes the instance
+	virtual void RemoveInstance(int32 Index);
+
+	// Returns true if we have current information for an instance at Index
+	bool IsIndexValid(int32 Index) const;
+
+	// Checks if we already have an actor for this handle. If an actor already exists we set the Actor variable on Handle and return true.
+	bool FindActorForHandle(const FActorInstanceHandle& Handle) const;
+
+	// Sets the parameters for actor spawning.
+	virtual void SetSpawnParameters(FActorSpawnParameters& SpawnParams);
+
+	// Called after spawning a new actor from a light weight instance
+	virtual void PostActorSpawn(const FActorInstanceHandle& Handle);
+
+	//
+	// Data and replication functions
+	//
+
+public:
+
+	virtual void GetLifetimeReplicatedProps(TArray<FLifetimeProperty>& OutLifetimeProps) const override;
+
+protected:
+
+	FString BaseInstanceName;
+
+	UPROPERTY(EditDefaultsOnly, Replicated, Category=LightWeightInstance)
+	TSubclassOf<AActor> RepresentedClass;
+
+	UPROPERTY()
+	TSubclassOf<AActor> AcceptedClass;
+
+	//
+	// Per instance data
+	// Stored in separate arrays to make ticking more efficient when we need to update everything
+	// 
+
+	// Current per instance transforms
+	UPROPERTY(ReplicatedUsing=OnRep_Transforms)
+	TArray<FTransform> InstanceTransforms;
+	UFUNCTION()
+	virtual void OnRep_Transforms();
+
+	//
+	// Bookkeeping info
+	//
+
+	// keep track of which instances are currently represented by an actor
+	TMap<int32, AActor*> Actors;
+
+	// list of indices that we are no longer using
+	UPROPERTY(Replicated)
+	TArray<int32> FreeIndices;
+
+	// handy way to check indices quickly so we don't need to iterate through the free indices list
+	UPROPERTY(Replicated)
+	TArray<bool> ValidIndices;
+
+	//
+	// Editor functions
+	//
+protected:
+	virtual void PostEditChangeProperty(struct FPropertyChangedEvent& PropertyChangedEvent) override;	
+};
diff --git a/Engine/Source/Runtime/Engine/Classes/GameFramework/LightWeightInstanceStaticMeshManager.h b/Engine/Source/Runtime/Engine/Classes/GameFramework/LightWeightInstanceStaticMeshManager.h
new file mode 100644
--- /dev/null
+++ b/Engine/Source/Runtime/Engine/Classes/GameFramework/LightWeightInstanceStaticMeshManager.h
@@ -0,0 +1,80 @@
+// Copyright Epic Games, Inc. All Rights Reserved.
+
+#pragma once
+
+#include "CoreMinimal.h"
+#include "UObject/ObjectMacros.h"
+#include "UObject/UObjectBaseUtility.h"
+#include "UObject/Object.h"
+#include "Engine/EngineTypes.h"
+#include "Engine/EngineBaseTypes.h"
+#include "Kismet/BlueprintFunctionLibrary.h"
+#include "Templates/SharedPointer.h"
+#include "GameFramework/LightWeightInstanceManager.h"
+
+#include "Actor.h"
+
+#include "LightWeightInstanceStaticMeshManager.generated.h"
+
+DECLARE_DYNAMIC_DELEGATE_RetVal_OneParam(FActorInstanceHandle, FOnActorReady, FActorInstanceHandle, InHandle);
+
+
+
+
+
+UCLASS(BlueprintType, Blueprintable)
+class ENGINE_API ALightWeightInstanceStaticMeshManager : public ALightWeightInstanceManager
+{
+	GENERATED_UCLASS_BODY()
+
+	virtual void SetRepresentedClass(UClass* ActorClass) override;
+
+	// Sets the static mesh to use based on the info contained in InActor
+	virtual void SetStaticMeshFromActor(AActor* InActor);
+
+	// Clears the static mesh used for rendering instances
+	void ClearStaticMesh();
+
+	virtual int32 ConvertCollisionIndexToLightWeightIndex(int32 InIndex) const override;
+
+protected:
+	virtual void AddNewInstanceAt(FLWIData* InitData, int32 Index) override;
+
+	virtual void RemoveInstance(int32 Index) override;
+
+	void RemoveInstanceFromRendering(int32 DataIndex);
+
+	// sets the parameters on the instanced static mesh component
+	virtual void SetInstancedStaticMeshParams();
+
+	// Called when we set the static mesh
+	void OnStaticMeshSet();
+
+	virtual void OnRep_Transforms() override;
+
+	virtual void PostActorSpawn(const FActorInstanceHandle& Handle) override;
+
+public:
+
+	virtual void GetLifetimeReplicatedProps(TArray<FLifetimeProperty>& OutLifetimeProps) const override;
+
+protected:
+
+	UPROPERTY(ReplicatedUsing = OnRep_StaticMesh)
+	TSoftObjectPtr<UStaticMesh> StaticMesh;
+	UFUNCTION()
+	void OnRep_StaticMesh();
+
+	UPROPERTY(EditAnywhere, BlueprintReadOnly, Category = Debug, AdvancedDisplay, meta = (BlueprintProtected = "true", AllowPrivateAccess = "true"))
+	class UHierarchicalInstancedStaticMeshComponent* InstancedStaticMeshComponent;
+
+	//
+	// Bookkeeping info
+	//
+
+	// keep track of the relationship between our data and the rendering data stored in the instanced static mesh component
+	UPROPERTY(Replicated)
+	TArray<int32> RenderingIndicesToDataIndices;
+	UPROPERTY(Replicated)
+	TArray<int32> DataIndicesToRenderingIndices;
+};
diff --git a/Engine/Source/Runtime/Engine/Classes/GameFramework/LightWeightInstanceSubsystem.h b/Engine/Source/Runtime/Engine/Classes/GameFramework/LightWeightInstanceSubsystem.h
new file mode 100644
--- /dev/null
+++ b/Engine/Source/Runtime/Engine/Classes/GameFramework/LightWeightInstanceSubsystem.h
@@ -0,0 +1,80 @@
+// Copyright Epic Games, Inc. All Rights Reserved.
+
+#pragma once
+
+#include "CoreMinimal.h"
+#include "UObject/ObjectMacros.h"
+#include "UObject/UObjectBaseUtility.h"
+#include "UObject/Object.h"
+#include "Engine/EngineTypes.h"
+#include "Engine/EngineBaseTypes.h"
+#include "Templates/SharedPointer.h"
+
+#include "Actor.h"
+
+//#include "LightWeightInstanceSubsystem.generated.h"
+
+ENGINE_API DECLARE_LOG_CATEGORY_EXTERN(LogLightWeightInstance, Log, Warning);
+
+//DECLARE_DYNAMIC_DELEGATE_RetVal_OneParam(FActorInstanceHandle, FOnActorReady, FActorInstanceHandle, InHandle);
+
+struct ENGINE_API FLightWeightInstanceSubsystem
+{
+
+	friend struct FActorInstanceHandle;
+	friend class ALightWeightInstanceManager;
+
+	static FLightWeightInstanceSubsystem& Get()
+	{
+		if (!LWISubsystem)
+		{
+			LWISubsystem = MakeShareable(new FLightWeightInstanceSubsystem());
+		}
+		return *LWISubsystem;
+	}
+
+	// returns the instance manager that handles the given handle
+	ALightWeightInstanceManager* FindLightWeightInstanceManager(const FActorInstanceHandle& Handle) const;
+
+	// returns the instance manager that handles actors of type ActorClass in level Level
+	ALightWeightInstanceManager* FindLightWeightInstanceManager(UClass* ActorClass, ULevel* Level) const;
+
+	// returns the instance manager that handles instances of type Class that live in Level
+	UFUNCTION(Server, Unreliable)
+	ALightWeightInstanceManager* FindOrAddLightWeightInstanceManager(UClass* ActorClass, ULevel* Level);
+
+	// Returns the actor specified by Handle. This may require loading and creating the actor object.
+	AActor* GetActor(const FActorInstanceHandle& Handle);
+
+	// Returns the actor specified by Handle if it exists. Returns nullptr if it doesn't
+	AActor* GetActor_NoCreate(const FActorInstanceHandle& Handle) const;
+
+	// Returns the class of the actor specified by Handle.
+	UClass* GetActorClass(const FActorInstanceHandle& Handle);
+
+	FVector GetLocation(const FActorInstanceHandle& Handle);
+
+	FString GetName(const FActorInstanceHandle& Handle);
+
+	ULevel* GetLevel(const FActorInstanceHandle& Handle);
+
+	// returns true if the object represented by Handle is in InLevel
+	bool IsInLevel(const FActorInstanceHandle& Handle, const ULevel* InLevel);
+
+protected:
+	UClass* FindBestInstanceManagerClass(const UClass* ActorClass);
+
+	// returns the index associated with Manager
+	int32 GetManagerIndex(const ALightWeightInstanceManager* Manager) const;
+
+	// returns the light weight instance manager at index Index
+	const ALightWeightInstanceManager* GetManagerAt(int32 Index) const;
+
+private:
+	/** Application singleton */
+	static TSharedPtr<FLightWeightInstanceSubsystem> LWISubsystem;
+
+	// TODO: preallocate the size of this based on a config variable
+	UPROPERTY()
+	TArray<ALightWeightInstanceManager*> LWInstanceManagers;
+};
diff --git a/Engine/Source/Runtime/Engine/Private/Actor.cpp b/Engine/Source/Runtime/Engine/Private/Actor.cpp
--- a/Engine/Source/Runtime/Engine/Private/Actor.cpp
+++ b/Engine/Source/Runtime/Engine/Private/Actor.cpp
@@ -4503,7 +4503,7 @@ void AActor::DispatchPhysicsCollisionHit(const FRigidBodyCollisionInfo& MyInfo,
 	Result.Normal = Result.ImpactNormal = ContactInfo.ContactNormal;
 	Result.PenetrationDepth = ContactInfo.ContactPenetration;
 	Result.PhysMaterial = ContactInfo.PhysMaterial[1];
-	Result.Actor = OtherInfo.Actor;
+	Result.HitObjectHandle = FActorInstanceHandle(OtherInfo.Actor.Get());
 	Result.Component = OtherInfo.Component;
 	Result.Item = OtherInfo.BodyIndex;
 	Result.BoneName = OtherInfo.BoneName;
diff --git a/Engine/Source/Runtime/Engine/Private/CheatManager.cpp b/Engine/Source/Runtime/Engine/Private/CheatManager.cpp
--- a/Engine/Source/Runtime/Engine/Private/CheatManager.cpp
+++ b/Engine/Source/Runtime/Engine/Private/CheatManager.cpp
@@ -718,7 +718,7 @@ void UCheatManager::TickCollisionDebug()
 				if(bHit)
 				{
 					AddCapsuleSweepDebugInfo(ViewLoc, End, Result.ImpactPoint, Result.Normal, Result.ImpactNormal, Result.Location, DebugCapsuleHalfHeight, DebugCapsuleRadius, false, (Result.bStartPenetrating && Result.bBlockingHit)? true: false);
-					UE_LOG(LogCollision, Log, TEXT("Collision component (%s) : Actor (%s)"), *GetNameSafe(Result.Component.Get()), *GetNameSafe(Result.GetActor()));
+					UE_LOG(LogCollision, Log, TEXT("Collision component (%s) : Actor (%s)"), *GetNameSafe(Result.Component.Get()), *Result.HitObjectHandle.GetName());
 				}
 			}
 		}
@@ -1214,8 +1214,8 @@ AActor* UCheatManager::GetTarget(APlayerController* PlayerController, struct FHi
     bool bHit = GetWorld()->LineTraceSingleByChannel(OutHit, CamLoc, CamRot.Vector() * 100000.f + CamLoc, ECC_Pawn, TraceParams);
     if (bHit)
     {
-        check(OutHit.GetActor() != NULL);
-        return OutHit.GetActor();
+        check(OutHit.HitObjectHandle.FetchActor() != nullptr);
+		return OutHit.HitObjectHandle.FetchActor();
     }
     return NULL;
 }
diff --git a/Engine/Source/Runtime/Engine/Private/Collision/Collision.cpp b/Engine/Source/Runtime/Engine/Private/Collision/Collision.cpp
--- a/Engine/Source/Runtime/Engine/Private/Collision/Collision.cpp
+++ b/Engine/Source/Runtime/Engine/Private/Collision/Collision.cpp
@@ -35,7 +35,7 @@ FHitResult::FHitResult(class AActor* InActor, class UPrimitiveComponent* InCompo
 	ImpactPoint = HitLoc;
 	Normal = HitNorm;
 	ImpactNormal = HitNorm;
-	Actor = InActor;
+	HitObjectHandle = FActorInstanceHandle(InActor);
 	Component = InComponent;
 }
 
@@ -130,7 +130,7 @@ bool FHitResult::NetSerialize(FArchive& Ar, class UPackageMap* Map, bool& bOutSu
 	}
 
 	Ar << PhysMaterial;
-	Ar << Actor;
+	Ar << HitObjectHandle;
 	Ar << Component;
 	Ar << BoneName;
 	if (!bInvalidFaceIndex)
@@ -157,11 +157,6 @@ bool FHitResult::NetSerialize(FArchive& Ar, class UPackageMap* Map, bool& bOutSu
 //////////////////////////////////////////////////////////////////////////
 // FOverlapResult
 
-AActor* FOverlapResult::GetActor() const
-{
-	return Actor.Get();
-}
-
 UPrimitiveComponent* FOverlapResult::GetComponent() const
 {
 	return Component.Get();
diff --git a/Engine/Source/Runtime/Engine/Private/Collision/CollisionConversions.cpp b/Engine/Source/Runtime/Engine/Private/Collision/CollisionConversions.cpp
--- a/Engine/Source/Runtime/Engine/Private/Collision/CollisionConversions.cpp
+++ b/Engine/Source/Runtime/Engine/Private/Collision/CollisionConversions.cpp
@@ -21,6 +21,7 @@
 #include "Chaos/ParticleHandle.h"
 #endif
 #include "PhysicsEngine/CollisionQueryFilterCallback.h"
+#include "GameFramework/LightWeightInstanceManager.h"
 
 // Used to place overlaps into a TMap when deduplicating them
 struct FOverlapKey
@@ -201,8 +202,16 @@ static void SetHitResultFromShapeAndFaceIndex(const FPhysicsShape& Shape,  const
 	// Grab actor/component
 	if( OwningComponent )
 	{
-		OutResult.Actor = OwningComponent->GetOwner();
 		OutResult.Component = OwningComponent;
+		AActor* Owner = OwningComponent->GetOwner();
+		if (ALightWeightInstanceManager* LWIManager = Cast<ALightWeightInstanceManager>(Owner))
+		{
+			OutResult.HitObjectHandle = FActorInstanceHandle(LWIManager, OutResult.Item);
+		}
+		else
+		{
+			OutResult.HitObjectHandle = FActorInstanceHandle(OwningComponent->GetOwner());
+		}
 
 		if (bReturnPhysMat)
 		{
@@ -287,7 +296,7 @@ EConvertQueryResult ConvertQueryImpactHit(const UWorld* World, const FHitLocatio
 #if ENABLE_NAN_DIAGNOSTIC
 			SetHitResultFromShapeAndFaceIndex(HitShape, HitActor, InternalFaceIndex, OutResult.ImpactPoint, OutResult, bReturnPhysMat);
 			UE_LOG(LogCore, Error, TEXT("ConvertQueryImpactHit() NaN details:\n>> Actor:%s (%s)\n>> Component:%s\n>> Item:%d\n>> BoneName:%s\n>> Time:%f\n>> Distance:%f\n>> Location:%s\n>> bIsBlocking:%d\n>> bStartPenetrating:%d"),
-				*GetNameSafe(OutResult.GetActor()), OutResult.Actor.IsValid() ? *OutResult.GetActor()->GetPathName() : TEXT("no path"),
+				*OutResult.GetHitObjectHandle().GetName(), OutResult.HasValidHitObjectHandle() ? *OutResult.GetHitObjectHandle().FetchActor()->GetPathName() : TEXT("no path"),
 				*GetNameSafe(OutResult.GetComponent()), OutResult.Item, *OutResult.BoneName.ToString(),
 				OutResult.Time, OutResult.Distance, *OutResult.Location.ToString(), OutResult.bBlockingHit ? 1 : 0, OutResult.bStartPenetrating ? 1 : 0);
 #endif // ENABLE_NAN_DIAGNOSTIC
@@ -307,7 +316,7 @@ EConvertQueryResult ConvertQueryImpactHit(const UWorld* World, const FHitLocatio
 #if ENABLE_NAN_DIAGNOSTIC
 		SetHitResultFromShapeAndFaceIndex(HitShape, HitActor, InternalFaceIndex, OutResult.ImpactPoint, OutResult, bReturnPhysMat);
 		UE_LOG(LogCore, Error, TEXT("ConvertQueryImpactHit() NaN details:\n>> Actor:%s (%s)\n>> Component:%s\n>> Item:%d\n>> BoneName:%s\n>> Time:%f\n>> Distance:%f\n>> Location:%s\n>> bIsBlocking:%d\n>> bStartPenetrating:%d"),
-			*GetNameSafe(OutResult.GetActor()), OutResult.Actor.IsValid() ? *OutResult.GetActor()->GetPathName() : TEXT("no path"),
+			*OutResult.GetHitObjectHandle().GetName(), OutResult.HasValidHitObjectHandle() ? *OutResult.GetHitObjectHandle().FetchActor()->GetPathName() : TEXT("no path"),
 			*GetNameSafe(OutResult.GetComponent()), OutResult.Item, *OutResult.BoneName.ToString(),
 			OutResult.Time, OutResult.Distance, *OutResult.Location.ToString(), OutResult.bBlockingHit ? 1 : 0, OutResult.bStartPenetrating ? 1 : 0);
 #endif // ENABLE_NAN_DIAGNOSTIC
@@ -546,7 +555,7 @@ void ConvertQueryOverlap(const FPhysicsShape& Shape, const FPhysicsActor& Actor,
         BodyInst = FPhysicsInterface::ShapeToOriginalBodyInstance(BodyInst, &Shape);
 		if (const UPrimitiveComponent* OwnerComponent = BodyInst->OwnerComponent.Get())
 		{
-			OutOverlap.Actor = OwnerComponent->GetOwner();
+			OutOverlap.OverlapObjectHandle = FActorInstanceHandle(OwnerComponent->GetOwner());
 			OutOverlap.Component = BodyInst->OwnerComponent; // Copying weak pointer is faster than assigning raw pointer.
 			OutOverlap.ItemIndex = OwnerComponent->bMultiBodyOverlap ? BodyInst->InstanceBodyIndex : INDEX_NONE;
 		}
@@ -574,7 +583,7 @@ void ConvertQueryOverlap(const FPhysicsShape& Shape, const FPhysicsActor& Actor,
 		if(PossibleOwner)
 		{
 			OutOverlap.Component = PossibleOwner;
-			OutOverlap.Actor = OutOverlap.Component->GetOwner();
+			OutOverlap.OverlapObjectHandle = FActorInstanceHandle(OutOverlap.Component->GetOwner());
 			OutOverlap.ItemIndex = INDEX_NONE;
 		}
 		else
@@ -599,7 +608,7 @@ static void AddUniqueOverlap(TArray<FOverlapResult>& OutOverlaps, const FOverlap
 		if (Overlap.ItemIndex == NewOverlap.ItemIndex && Overlap.Component == NewOverlap.Component)
 		{
 			// These should be the same if the component matches!
-			checkSlow(Overlap.Actor == NewOverlap.Actor);
+			checkSlow(Overlap.OverlapObjectHandle == NewOverlap.OverlapObjectHandle);
 
 			// If we had a non-blocking overlap with this component, but now we have a blocking one, use that one instead!
 			if(!Overlap.bBlockingHit && NewOverlap.bBlockingHit)
diff --git a/Engine/Source/Runtime/Engine/Private/Components/CharacterMovementComponent.cpp b/Engine/Source/Runtime/Engine/Private/Components/CharacterMovementComponent.cpp
--- a/Engine/Source/Runtime/Engine/Private/Components/CharacterMovementComponent.cpp
+++ b/Engine/Source/Runtime/Engine/Private/Components/CharacterMovementComponent.cpp
@@ -2958,8 +2958,7 @@ FVector UCharacterMovementComponent::GetPenetrationAdjustment(const FHitResult&
 	{
 		const bool bIsProxy = (CharacterOwner->GetLocalRole() == ROLE_SimulatedProxy);
 		float MaxDistance = bIsProxy ? MaxDepenetrationWithGeometryAsProxy : MaxDepenetrationWithGeometry;
-		const AActor* HitActor = Hit.GetActor();
-		if (Cast<APawn>(HitActor))
+		if (Hit.HitObjectHandle.DoesRepresentClass(APawn::StaticClass()))
 		{
 			MaxDistance = bIsProxy ? MaxDepenetrationWithPawnAsProxy : MaxDepenetrationWithPawn;
 		}
@@ -4089,12 +4088,11 @@ FVector UCharacterMovementComponent::FindWaterLine(FVector InWater, FVector Outo
 	TArray<FHitResult> Hits;
 	GetWorld()->LineTraceMultiByChannel(Hits, OutofWater, InWater, UpdatedComponent->GetCollisionObjectType(), FCollisionQueryParams(SCENE_QUERY_STAT(FindWaterLine), true, CharacterOwner));
 
-	for( int32 HitIdx = 0; HitIdx < Hits.Num(); HitIdx++ )
+	for (const FHitResult& Check : Hits)
 	{
-		const FHitResult& Check = Hits[HitIdx];
-		if ( !CharacterOwner->IsOwnedBy(Check.GetActor()) && !Check.Component.Get()->IsWorldGeometry() )
+		if ( !CharacterOwner->IsOwnedBy(Check.HitObjectHandle.FetchActor()) && !Check.Component.Get()->IsWorldGeometry() )
 		{
-			APhysicsVolume *W = Cast<APhysicsVolume>(Check.GetActor());
+			APhysicsVolume *W = Check.HitObjectHandle.FetchActor<APhysicsVolume>();
 			if ( W && W->bWaterVolume )
 			{
 				FVector Dir = (InWater - OutofWater).GetSafeNormal();
@@ -4685,7 +4683,7 @@ void UCharacterMovementComponent::OnCharacterStuckInGeometry(const FHitResult* H
 					   Hit->Location.X, Hit->Location.Y, Hit->Location.Z,
 					   Hit->Normal.X, Hit->Normal.Y, Hit->Normal.Z,
 					   Hit->PenetrationDepth,
-					   *GetNameSafe(Hit->GetActor()),
+					   *Hit->HitObjectHandle.GetName(),
 					   *GetNameSafe(Hit->GetComponent()),
 					   Hit->BoneName.IsValid() ? *Hit->BoneName.ToString() : TEXT("None"),
 					   StuckWarningCountSinceNotify
@@ -4746,7 +4744,7 @@ void UCharacterMovementComponent::MoveAlongFloor(const FVector& InVelocity, floa
 
 		if (Hit.IsValidBlockingHit())
 		{
-			if (CanStepUp(Hit) || (CharacterOwner->GetMovementBase() != NULL && CharacterOwner->GetMovementBase()->GetOwner() == Hit.GetActor()))
+			if (CanStepUp(Hit) || (CharacterOwner->GetMovementBase() != nullptr && Hit.HitObjectHandle == CharacterOwner->GetMovementBase()->GetOwner()))
 			{
 				// hit a barrier, try to step up
 				const FVector GravDir(0.f, 0.f, -1.f);
@@ -5852,7 +5850,7 @@ bool UCharacterMovementComponent::CheckWaterJump(FVector CheckPoint, FVector& Wa
 	const ECollisionChannel CollisionChannel = UpdatedComponent->GetCollisionObjectType();
 	bool bHit = GetWorld()->SweepSingleByChannel( HitInfo, UpdatedComponent->GetComponentLocation(), CheckPoint, FQuat::Identity, CollisionChannel, CapsuleShape, CapsuleParams, ResponseParam);
 	
-	if ( bHit && !Cast<APawn>(HitInfo.GetActor()) )
+	if ( bHit && !HitInfo.HitObjectHandle.DoesRepresentClass(APawn::StaticClass()) )
 	{
 		// hit a wall - check if it is low enough
 		WallNormal = -1.f * HitInfo.ImpactNormal;
@@ -6519,12 +6517,13 @@ bool UCharacterMovementComponent::CanStepUp(const FHitResult& Hit) const
 	}
 
 	// No actor for "fake" hits when we are on a known good base.
-	const AActor* HitActor = Hit.GetActor();
-	if (!HitActor)
+	
+	if (!Hit.HitObjectHandle.IsValid())
 	{
 		 return true;
 	}
 
+	const AActor* HitActor = Hit.HitObjectHandle.FetchActor();
 	if (!HitActor->CanBeBaseForCharacter(CharacterOwner))
 	{
 		return false;
@@ -6766,7 +6765,7 @@ void UCharacterMovementComponent::HandleImpact(const FHitResult& Impact, float T
 		PFAgent->OnMoveBlockedBy(Impact);
 	}
 
-	APawn* OtherPawn = Cast<APawn>(Impact.GetActor());
+	APawn* OtherPawn = Impact.HitObjectHandle.FetchActor<APawn>();
 	if (OtherPawn)
 	{
 		NotifyBumpedPawn(OtherPawn);
diff --git a/Engine/Source/Runtime/Engine/Private/Components/MovementComponent.cpp b/Engine/Source/Runtime/Engine/Private/Components/MovementComponent.cpp
--- a/Engine/Source/Runtime/Engine/Private/Components/MovementComponent.cpp
+++ b/Engine/Source/Runtime/Engine/Private/Components/MovementComponent.cpp
@@ -635,7 +635,7 @@ bool UMovementComponent::ResolvePenetrationImpl(const FVector& ProposedAdjustmen
 			   *ActorOwner->GetName(),
 			   *UpdatedComponent->GetName(),
 			   *UpdatedComponent->GetComponentLocation().ToString(),
-			   *GetNameSafe(Hit.GetActor()),
+			   *Hit.GetHitObjectHandle().GetName(),
 			   *GetNameSafe(Hit.GetComponent()),
 			   Hit.Component.IsValid() ? *Hit.GetComponent()->GetComponentLocation().ToString() : TEXT("<unknown>"),
 			   Hit.PenetrationDepth,
diff --git a/Engine/Source/Runtime/Engine/Private/Components/PrimitiveComponent.cpp b/Engine/Source/Runtime/Engine/Private/Components/PrimitiveComponent.cpp
--- a/Engine/Source/Runtime/Engine/Private/Components/PrimitiveComponent.cpp
+++ b/Engine/Source/Runtime/Engine/Private/Components/PrimitiveComponent.cpp
@@ -121,7 +121,7 @@ struct FPredicateOverlapHasSameActor
 	bool operator() (const FOverlapInfo& Info)
 	{
 		// MyOwnerPtr is always valid, so we don't need the IsValid() checks in the WeakObjectPtr comparison operator.
-		return MyOwnerPtr.HasSameIndexAndSerialNumber(Info.OverlapInfo.Actor);
+		return MyOwnerPtr.HasSameIndexAndSerialNumber(Info.OverlapInfo.HitObjectHandle.FetchActor());
 	}
 
 private:
@@ -139,7 +139,7 @@ struct FPredicateOverlapHasDifferentActor
 	bool operator() (const FOverlapInfo& Info)
 	{
 		// MyOwnerPtr is always valid, so we don't need the IsValid() checks in the WeakObjectPtr comparison operator.
-		return !MyOwnerPtr.HasSameIndexAndSerialNumber(Info.OverlapInfo.Actor);
+		return !MyOwnerPtr.HasSameIndexAndSerialNumber(Info.OverlapInfo.HitObjectHandle.FetchActor());
 	}
 
 private:
@@ -1916,7 +1916,7 @@ static bool ShouldIgnoreHitResult(const UWorld* InWorld, FHitResult const& TestH
 		if ( (MoveFlags & MOVECOMP_IgnoreBases) && MovingActor )	//we let overlap components go through because their overlap is still needed and will cause beginOverlap/endOverlap events
 		{
 			// ignore if there's a base relationship between moving actor and hit actor
-			AActor const* const HitActor = TestHit.GetActor();
+			AActor const* const HitActor = TestHit.HitObjectHandle.FetchActor();
 			if (HitActor)
 			{
 				if (MovingActor->IsBasedOnActor(HitActor) || HitActor->IsBasedOnActor(MovingActor))
@@ -2194,7 +2194,7 @@ bool UPrimitiveComponent::MoveComponentImpl( const FVector& Delta, const FQuat&
 						UPrimitiveComponent* OverlapComponent = TestHit.Component.Get();
 						if (OverlapComponent && (OverlapComponent->GetGenerateOverlapEvents() || bForceGatherOverlaps))
 						{
-							if (!ShouldIgnoreOverlapResult(MyWorld, Actor, *this, TestHit.GetActor(), *OverlapComponent, /*bCheckOverlapFlags=*/ !bForceGatherOverlaps))
+							if (!ShouldIgnoreOverlapResult(MyWorld, Actor, *this, TestHit.HitObjectHandle.FetchActor(), *OverlapComponent, /*bCheckOverlapFlags=*/ !bForceGatherOverlaps))
 							{
 								// don't process touch events after initial blocking hits
 								if (BlockingHitIndex >= 0 && TestHit.Time > Hits[BlockingHitIndex].Time)
@@ -2392,7 +2392,7 @@ void UPrimitiveComponent::DispatchBlockingHit(AActor& Owner, FHitResult const& B
 		if (!BlockingHitComponent->IsPendingKill())
 		{
 			// BlockingHit.GetActor() could be marked for deletion in DispatchBlockingHit(), which would make the weak pointer return NULL.
-			if (AActor* const BlockingHitActor = BlockingHit.GetActor())
+			if (AActor* const BlockingHitActor = BlockingHit.HitObjectHandle.FetchActor())
 			{
 				BlockingHitActor->DispatchBlockingHit(BlockingHitComponent, this, false, BlockingHit);
 			}
@@ -3184,7 +3184,7 @@ bool UPrimitiveComponent::UpdateOverlapsImpl(const TOverlapArrayView* NewPending
 						if (HitComp && (HitComp != this) && HitComp->GetGenerateOverlapEvents())
 						{
 							const bool bCheckOverlapFlags = false; // Already checked above
-							if (!ShouldIgnoreOverlapResult(MyWorld, MyActor, *this, Result.GetActor(), *HitComp, bCheckOverlapFlags))
+							if (!ShouldIgnoreOverlapResult(MyWorld, MyActor, *this, Result.OverlapObjectHandle.FetchActor(), *HitComp, bCheckOverlapFlags))
 							{
 								OverlapMultiResult.Emplace(HitComp, Result.ItemIndex);		// don't need to add unique unless the overlap check can return dupes
 							}
diff --git a/Engine/Source/Runtime/Engine/Private/Components/ProjectileMovementComponent.cpp b/Engine/Source/Runtime/Engine/Private/Components/ProjectileMovementComponent.cpp
--- a/Engine/Source/Runtime/Engine/Private/Components/ProjectileMovementComponent.cpp
+++ b/Engine/Source/Runtime/Engine/Private/Components/ProjectileMovementComponent.cpp
@@ -538,7 +538,7 @@ UProjectileMovementComponent::EHandleBlockingHitResult UProjectileMovementCompon
 
 	if (Hit.bStartPenetrating)
 	{
-		UE_LOG(LogProjectileMovement, Verbose, TEXT("Projectile %s is stuck inside %s.%s with velocity %s!"), *GetNameSafe(ActorOwner), *GetNameSafe(Hit.GetActor()), *GetNameSafe(Hit.GetComponent()), *Velocity.ToString());
+		UE_LOG(LogProjectileMovement, Verbose, TEXT("Projectile %s is stuck inside %s.%s with velocity %s!"), *GetNameSafe(ActorOwner), *Hit.HitObjectHandle.GetName(), *GetNameSafe(Hit.GetComponent()), *Velocity.ToString());
 		return EHandleBlockingHitResult::Abort;
 	}
 
diff --git a/Engine/Source/Runtime/Engine/Private/Components/SceneComponent.cpp b/Engine/Source/Runtime/Engine/Private/Components/SceneComponent.cpp
--- a/Engine/Source/Runtime/Engine/Private/Components/SceneComponent.cpp
+++ b/Engine/Source/Runtime/Engine/Private/Components/SceneComponent.cpp
@@ -57,7 +57,7 @@ DECLARE_CYCLE_STAT(TEXT("Component PostUpdateNavData"), STAT_ComponentPostUpdate
 FOverlapInfo::FOverlapInfo(UPrimitiveComponent* InComponent, int32 InBodyIndex)
 	: bFromSweep(false)
 {
-	OverlapInfo.Actor = InComponent ? InComponent->GetOwner() : nullptr;
+	OverlapInfo.HitObjectHandle = FActorInstanceHandle(InComponent ? InComponent->GetOwner() : nullptr);
 	OverlapInfo.Component = InComponent;
 	OverlapInfo.Item = InBodyIndex;
 }
@@ -2626,10 +2626,9 @@ void USceneComponent::UpdatePhysicsVolume( bool bTriggerNotifiers )
 						MyWorld->OverlapMultiByChannel(Hits, GetComponentLocation(), FQuat::Identity, GetCollisionObjectType(), FCollisionShape::MakeSphere(0.f), Params);
 					}
 
-					for (int32 HitIdx = 0; HitIdx < Hits.Num(); HitIdx++)
+					for (const FOverlapResult& Link : Hits)
 					{
-						const FOverlapResult& Link = Hits[HitIdx];
-						APhysicsVolume* const V = Cast<APhysicsVolume>(Link.GetActor());
+						APhysicsVolume* const V = Link.OverlapObjectHandle.FetchActor<APhysicsVolume>();
 						if (V && (V->Priority > NewVolume->Priority))
 						{
 							if (bOverlappedOrigin || V->IsOverlapInVolume(*this))
diff --git a/Engine/Source/Runtime/Engine/Private/DebugCameraController.cpp b/Engine/Source/Runtime/Engine/Private/DebugCameraController.cpp
--- a/Engine/Source/Runtime/Engine/Private/DebugCameraController.cpp
+++ b/Engine/Source/Runtime/Engine/Private/DebugCameraController.cpp
@@ -252,7 +252,7 @@ AActor* ADebugCameraController::GetSelectedActor() const
 void ADebugCameraController::Select( FHitResult const& Hit )
 {
 	// store selection
-	SelectedActor = Hit.GetActor();
+	SelectedActor = Hit.HitObjectHandle.FetchActor();
 	SelectedComponent = Hit.Component.Get();
 	SelectedHitPoint = Hit;
 
diff --git a/Engine/Source/Runtime/Engine/Private/DebugCameraHUD.cpp b/Engine/Source/Runtime/Engine/Private/DebugCameraHUD.cpp
--- a/Engine/Source/Runtime/Engine/Private/DebugCameraHUD.cpp
+++ b/Engine/Source/Runtime/Engine/Private/DebugCameraHUD.cpp
@@ -12,6 +12,7 @@
 #include "DrawDebugHelpers.h"
 #include "GameFramework/PawnMovementComponent.h"
 #include "GameFramework/SpectatorPawn.h"
+#include "GameFramework/LightWeightInstanceSubsystem.h"
 #include "Engine/Canvas.h"
 #include "Engine/DebugCameraController.h"
 #include "PhysicalMaterials/PhysicalMaterial.h"
@@ -145,18 +146,19 @@ void ADebugCameraHUD::PostRender()
 
 			if( bHit )
 			{
-				AActor* HitActor = Hit.GetActor();
+				const FActorInstanceHandle& HitHandle = Hit.HitObjectHandle;
 				yl += Y;
 				Canvas->DrawText(RenderFont, FString::Printf(TEXT("HitLoc:%s HitNorm:%s"), *Hit.Location.ToString(), *Hit.Normal.ToString() ), X, yl, 1.f, 1.f, FontRenderInfo);
 				yl += Y;
 				Canvas->DrawText(RenderFont, FString::Printf(TEXT("HitDist: %f"), Hit.Distance), X, yl, 1.f, 1.f, FontRenderInfo);
 				yl += Y;
-				Canvas->DrawText(RenderFont, FString::Printf(TEXT("HitActor: '%s'"), HitActor ? *HitActor->GetFName().ToString() : TEXT("<NULL>")), X, yl, 1.f, 1.f, FontRenderInfo);
+				Canvas->DrawText(RenderFont, FString::Printf(TEXT("HitObject: '%s'"), *FLightWeightInstanceSubsystem::Get().GetName(HitHandle)), X, yl, 1.f, 1.f, FontRenderInfo);
 				yl += Y;
 				Canvas->DrawText(RenderFont, FString::Printf(TEXT("HitComponent: '%s'"), Hit.Component.Get() ? *Hit.Component.Get()->GetFName().ToString() : TEXT("<NULL>")), X, yl, 1.f, 1.f, FontRenderInfo);
 				yl += Y;
-				Canvas->DrawText(RenderFont, FString::Printf(TEXT("HitActor Class: '%s'"), HitActor && HitActor->GetClass() ? *HitActor->GetClass()->GetName() : TEXT("<Not Found>") ), X, yl, 1.f, 1.f, FontRenderInfo);
-				yl += Y;
+				UClass* HitClass = HitHandle.GetRepresentedClass();
+				Canvas->DrawText(RenderFont, FString::Printf(TEXT("HitHandle Class: '%s'"), HitClass ? *HitClass->GetName() : TEXT("<Not Found>") ), X, yl, 1.f, 1.f, FontRenderInfo);
+				yl += Y; AActor* HitActor = HitHandle.FetchActor();
 				Canvas->DrawText(RenderFont, FString::Printf(TEXT("HitActorPath: '%s'"), HitActor ? *HitActor->GetPathName() : TEXT("<Not Found>")), X, yl, 1.f, 1.f, FontRenderInfo);
 				yl += Y;
 				Canvas->DrawText(RenderFont, FString::Printf(TEXT("HitPhysMat: '%s'"), Hit.PhysMaterial.Get() ? *Hit.PhysMaterial.Get()->GetPathName() : TEXT("<Not Found>")), X, yl, 1.f, 1.f, FontRenderInfo);
diff --git a/Engine/Source/Runtime/Engine/Private/EngineTypes.cpp b/Engine/Source/Runtime/Engine/Private/EngineTypes.cpp
--- a/Engine/Source/Runtime/Engine/Private/EngineTypes.cpp
+++ b/Engine/Source/Runtime/Engine/Private/EngineTypes.cpp
@@ -6,6 +6,8 @@
 #include "Engine/EngineBaseTypes.h"
 #include "Components/SceneComponent.h"
 #include "GameFramework/Actor.h"
+#include "GameFramework/LightWeightInstanceSubsystem.h"
+#include "GameFramework/LightWeightInstanceManager.h"
 #include "Engine/World.h"
 #include "Components/PrimitiveComponent.h"
 #include "Engine/MeshMerging.h"
@@ -115,7 +117,7 @@ void FDamageEvent::GetBestHitInfo(AActor const* HitActor, AActor const* HitInsti
 	if (HitActor)
 	{
 		// fill out the hitinfo as best we can
-		OutHitInfo.Actor = const_cast<AActor*>(HitActor);
+		OutHitInfo.HitObjectHandle = FActorInstanceHandle(const_cast<AActor*>(HitActor));
 		OutHitInfo.bBlockingHit = true;
 		OutHitInfo.BoneName = NAME_None;
 		OutHitInfo.Component = Cast<UPrimitiveComponent>(HitActor->GetRootComponent());
@@ -235,6 +237,210 @@ UActorComponent* FComponentReference::GetComponent(AActor* OwningActor) const
 	return Result;
 }
 
+FActorInstanceHandle::FActorInstanceHandle()
+	: Actor(nullptr)
+	, ManagerIndex(INDEX_NONE)
+	, InstanceIndex(INDEX_NONE)
+{
+	// do nothing
+}
+
+FActorInstanceHandle::FActorInstanceHandle(AActor* InActor)
+	: Actor(InActor)
+	, ManagerIndex(INDEX_NONE)
+	, InstanceIndex(INDEX_NONE)
+{
+	if (InActor)
+	{
+		if (ALightWeightInstanceManager* LWIManager = FLightWeightInstanceSubsystem::Get().FindLightWeightInstanceManager(InActor->StaticClass(), InActor->GetLevel()))
+		{
+			InstanceIndex = LWIManager->FindIndexForActor(InActor);
+			if (InstanceIndex != INDEX_NONE)
+			{
+				ManagerIndex = FLightWeightInstanceSubsystem::Get().GetManagerIndex(LWIManager);
+			}
+		}
+	}
+}
+
+FActorInstanceHandle::FActorInstanceHandle(int32 InManagerIndex, int32 InInstanceIndex)
+	: Actor(nullptr)
+	, ManagerIndex(InManagerIndex)
+{
+	const ALightWeightInstanceManager* Manager = FLightWeightInstanceSubsystem::Get().GetManagerAt(InManagerIndex);
+
+	if (ensure(Manager))
+	{
+		InstanceIndex = Manager->ConvertCollisionIndexToLightWeightIndex(InInstanceIndex);
+
+		if (AActor*const* FoundActor = Manager->Actors.Find(InInstanceIndex))
+		{
+			Actor = *FoundActor;
+		}
+	}
+	else
+	{
+		ManagerIndex = INDEX_NONE;
+		InstanceIndex = INDEX_NONE;
+	}
+}
+
+FActorInstanceHandle::FActorInstanceHandle(ALightWeightInstanceManager* Manager, int32 InInstanceIndex)
+	: Actor(nullptr)
+	, ManagerIndex(INDEX_NONE)
+	, InstanceIndex(InInstanceIndex)
+{
+	if (ensure(Manager))
+	{
+		InstanceIndex = Manager->ConvertCollisionIndexToLightWeightIndex(InInstanceIndex);
+		if (AActor*const* FoundActor = Manager->Actors.Find(InstanceIndex))
+		{
+			Actor = *FoundActor;
+		}
+
+		ManagerIndex = FLightWeightInstanceSubsystem::Get().GetManagerIndex(Manager);
+		ensure(ManagerIndex != INDEX_NONE);
+	}
+}
+
+FActorInstanceHandle::FActorInstanceHandle(const FActorInstanceHandle& Other)
+{
+	Actor = Other.Actor;
+	ManagerIndex = Other.ManagerIndex;
+	InstanceIndex = Other.InstanceIndex;
+}
+
+bool FActorInstanceHandle::IsValid() const
+{
+	return (ManagerIndex != INDEX_NONE && InstanceIndex != INDEX_NONE) || Actor.IsValid();
+}
+
+bool FActorInstanceHandle::DoesRepresentClass(const UClass* OtherClass) const
+{
+	if (OtherClass == nullptr)
+	{
+		return false;
+	}
+
+	if (Actor.IsValid())
+	{
+		return Actor->IsA(OtherClass);
+	}
+
+	if (ALightWeightInstanceManager* Manager = FLightWeightInstanceSubsystem::Get().FindLightWeightInstanceManager(*this))
+	{
+		return Manager->DoesRepresentClass(OtherClass);
+	}
+
+	return false;
+}
+
+UClass* FActorInstanceHandle::GetRepresentedClass() const
+{
+	if (!IsValid())
+	{
+		return nullptr;
+	}
+
+	if (Actor.IsValid())
+	{
+		return Actor->GetClass();
+	}
+
+	if (ALightWeightInstanceManager* Manager = FLightWeightInstanceSubsystem::Get().FindLightWeightInstanceManager(*this))
+	{
+		return Manager->GetRepresentedClass();
+	}
+
+	return nullptr;
+}
+
+FVector FActorInstanceHandle::GetLocation() const
+{
+	return FVector();
+}
+
+FRotator FActorInstanceHandle::GetRotation() const
+{
+	return FRotator();
+}
+
+ULevel* FActorInstanceHandle::GetLevel() const
+{
+	return nullptr;
+}
+
+bool FActorInstanceHandle::IsInLevel(ULevel* Level) const
+{
+	return false;
+}
+
+FName FActorInstanceHandle::GetFName() const
+{
+	return NAME_None;
+}
+
+FString FActorInstanceHandle::GetName() const
+{
+	return FString();
+}
+
+AActor* FActorInstanceHandle::FetchActor() const
+{
+	if (Actor.IsValid())
+	{
+		return Actor.Get();
+	}
+
+	return FLightWeightInstanceSubsystem::Get().GetActor(*this);
+}
+
+bool FActorInstanceHandle::operator==(const FActorInstanceHandle& Other) const
+{
+	return ManagerIndex == Other.ManagerIndex && InstanceIndex == Other.InstanceIndex;
+}
+
+bool FActorInstanceHandle::operator!=(const FActorInstanceHandle& Other) const
+{
+	return !(*this == Other);
+}
+
+bool FActorInstanceHandle::operator==(const AActor* OtherActor) const
+{
+	// if we have an actor, compare the two actors
+	if (Actor.IsValid())
+	{
+		return Actor.Get() == OtherActor;
+	}
+
+	// we don't have an actor so see if we can look up an instance associated with OtherActor and see if we refer to the same instance
+	if (ALightWeightInstanceManager* Manager = FLightWeightInstanceSubsystem::Get().FindLightWeightInstanceManager(OtherActor->StaticClass(), OtherActor->GetLevel()))
+	{
+		if (FLightWeightInstanceSubsystem::Get().GetManagerIndex(Manager) != ManagerIndex)
+		{
+			return false;
+		}
+
+		return Manager->FindIndexForActor(OtherActor) == InstanceIndex;
+	}
+	
+	return false;
+}
+
+bool FActorInstanceHandle::operator!=(const AActor* OtherActor) const
+{
+	return !(*this == OtherActor);
+}
+
+FArchive& operator<<(FArchive& Ar, FActorInstanceHandle& Handle)
+{
+	Ar << Handle.Actor;
+	Ar << Handle.ManagerIndex;
+	Ar << Handle.InstanceIndex;
+
+	return Ar;
+}
+
 FString FHitResult::ToString() const
 {
 	return FString::Printf(TEXT("bBlockingHit:%s bStartPenetrating:%s Time:%f Location:%s ImpactPoint:%s Normal:%s ImpactNormal:%s TraceStart:%s TraceEnd:%s PenetrationDepth:%f Item:%d PhysMaterial:%s Actor:%s Component:%s BoneName:%s FaceIndex:%d"),
@@ -250,7 +456,7 @@ FString FHitResult::ToString() const
 		PenetrationDepth,
 		Item,
 		PhysMaterial.IsValid() ? *PhysMaterial->GetName() : TEXT("None"),
-		Actor.IsValid() ? *Actor->GetName() : TEXT("None"),
+		*FLightWeightInstanceSubsystem::Get().GetName(HitObjectHandle),
 		Component.IsValid() ? *Component->GetName() : TEXT("None"),
 		BoneName.IsValid() ? *BoneName.ToString() : TEXT("None"),
 		FaceIndex);
diff --git a/Engine/Source/Runtime/Engine/Private/EngineUtils.cpp b/Engine/Source/Runtime/Engine/Private/EngineUtils.cpp
--- a/Engine/Source/Runtime/Engine/Private/EngineUtils.cpp
+++ b/Engine/Source/Runtime/Engine/Private/EngineUtils.cpp
@@ -10,6 +10,7 @@
 #include "Misc/PackageName.h"
 #include "Misc/EngineVersion.h"
 #include "GameFramework/PlayerController.h"
+#include "GameFramework/LightWeightInstanceSubsystem.h"
 #include "EngineGlobals.h"
 #include "Components/PrimitiveComponent.h"
 #include "Engine/Engine.h"
@@ -430,9 +431,9 @@ TArray<FSubLevelStatus> GetSubLevelsStatus( UWorld* World )
 
 				ULevel* LevelPlayerIsIn = nullptr;
 
-				if (AActor* HitActor = Hit.GetActor())
+				if (Hit.HitObjectHandle.IsValid())
 				{
-					LevelPlayerIsIn = HitActor->GetLevel();
+					LevelPlayerIsIn = FLightWeightInstanceSubsystem::Get().GetLevel(Hit.HitObjectHandle);
 				}
 				else if (UPrimitiveComponent* HitComponent = Hit.Component.Get())
 				{
diff --git a/Engine/Source/Runtime/Engine/Private/GameplayStatics.cpp b/Engine/Source/Runtime/Engine/Private/GameplayStatics.cpp
--- a/Engine/Source/Runtime/Engine/Private/GameplayStatics.cpp
+++ b/Engine/Source/Runtime/Engine/Private/GameplayStatics.cpp
@@ -489,7 +489,7 @@ static bool ComponentIsDamageableFrom(UPrimitiveComponent* VictimComp, FVector c
 			else
 			{
 				// if we hit something else blocking, it's not
-				UE_LOG(LogDamage, Log, TEXT("Radial Damage to %s blocked by %s (%s)"), *GetNameSafe(VictimComp), *GetNameSafe(OutHitResult.GetActor()), *GetNameSafe(OutHitResult.Component.Get()));
+				UE_LOG(LogDamage, Log, TEXT("Radial Damage to %s blocked by %s (%s)"), *GetNameSafe(VictimComp), *OutHitResult.GetHitObjectHandle().GetName(), *GetNameSafe(OutHitResult.Component.Get()));
 				return false;
 			}
 		}
@@ -528,10 +528,9 @@ bool UGameplayStatics::ApplyRadialDamageWithFalloff(const UObject* WorldContextO
 
 	// collate into per-actor list of hit components
 	TMap<AActor*, TArray<FHitResult> > OverlapComponentMap;
-	for (int32 Idx = 0; Idx < Overlaps.Num(); ++Idx)
+	for (const FOverlapResult& Overlap : Overlaps)
 	{
-		FOverlapResult const& Overlap = Overlaps[Idx];
-		AActor* const OverlapActor = Overlap.GetActor();
+		AActor* const OverlapActor = Overlap.OverlapObjectHandle.FetchActor();
 
 		if (OverlapActor &&
 			OverlapActor->CanBeDamaged() &&
@@ -1132,6 +1131,7 @@ UParticleSystemComponent* UGameplayStatics::SpawnEmitterAttached(UParticleSystem
 	return PSC;
 }
 
+// FRED_TODO: propagate hit object handles further instead of converting to an actor here
 void UGameplayStatics::BreakHitResult(const FHitResult& Hit, bool& bBlockingHit, bool& bInitialOverlap, float& Time, float& Distance, FVector& Location, FVector& ImpactPoint, FVector& Normal, FVector& ImpactNormal, UPhysicalMaterial*& PhysMat, AActor*& HitActor, UPrimitiveComponent*& HitComponent, FName& HitBoneName, int32& HitItem, int32& FaceIndex, FVector& TraceStart, FVector& TraceEnd)
 {
 	SCOPE_CYCLE_COUNTER(STAT_BreakHitResult);
@@ -1144,7 +1144,7 @@ void UGameplayStatics::BreakHitResult(const FHitResult& Hit, bool& bBlockingHit,
 	Normal = Hit.Normal;
 	ImpactNormal = Hit.ImpactNormal;	
 	PhysMat = Hit.PhysMaterial.Get();
-	HitActor = Hit.GetActor();
+	HitActor = Hit.GetHitObjectHandle().FetchActor();
 	HitComponent = Hit.GetComponent();
 	HitBoneName = Hit.BoneName;
 	HitItem = Hit.Item;
@@ -1166,7 +1166,7 @@ FHitResult UGameplayStatics::MakeHitResult(bool bBlockingHit, bool bInitialOverl
 	Hit.Normal = Normal;
 	Hit.ImpactNormal = ImpactNormal;
 	Hit.PhysMaterial = PhysMat;
-	Hit.Actor = HitActor;
+	Hit.HitObjectHandle = FActorInstanceHandle(HitActor);
 	Hit.Component = HitComponent;
 	Hit.BoneName = HitBoneName;
 	Hit.Item = HitItem;
diff --git a/Engine/Source/Runtime/Engine/Private/HUD.cpp b/Engine/Source/Runtime/Engine/Private/HUD.cpp
--- a/Engine/Source/Runtime/Engine/Private/HUD.cpp
+++ b/Engine/Source/Runtime/Engine/Private/HUD.cpp
@@ -492,7 +492,7 @@ AActor* AHUD::GetCurrentDebugTargetActor()
 		bool bHit = GetWorld()->LineTraceSingleByChannel(Hit, CamLoc, CamRot.Vector() * 100000.f + CamLoc, ECC_WorldDynamic, TraceParams);
 		if (bHit)
 		{
-			AActor* HitActor = Hit.Actor.Get();
+			AActor* HitActor = Hit.HitObjectHandle.FetchActor();
 			if (HitActor && ((ShowDebugTargetDesiredClass == nullptr) || HitActor->IsA(ShowDebugTargetDesiredClass)))
 			{
 				DebugTargetActor = HitActor;
diff --git a/Engine/Source/Runtime/Engine/Private/LightWeightInstanceBlueprintFunctionLibrary.cpp b/Engine/Source/Runtime/Engine/Private/LightWeightInstanceBlueprintFunctionLibrary.cpp
new file mode 100644
--- /dev/null
+++ b/Engine/Source/Runtime/Engine/Private/LightWeightInstanceBlueprintFunctionLibrary.cpp
@@ -0,0 +1,21 @@
+// Copyright Epic Games, Inc. All Rights Reserved.
+
+#include "GameFramework/LightWeightInstanceBlueprintFunctionLibrary.h"
+#include "GameFramework/LightWeightInstanceSubsystem.h"
+#include "GameFramework/LightWeightInstanceManager.h"
+
+FActorInstanceHandle ULightWeightInstanceBlueprintFunctionLibrary::CreateNewLightWeightInstance(UClass* InActorClass, FTransform InTransform, ULevel* InLevel)
+{
+	// Get or create a light weight instance for this class and level
+	if (ALightWeightInstanceManager* LightWeightInstance = FLightWeightInstanceSubsystem::Get().FindOrAddLightWeightInstanceManager(InActorClass, InLevel))
+	{
+		// create an instance with the given transform
+		FLWIData PerInstanceData;
+		PerInstanceData.Transform = InTransform;
+		int32 InstanceIdx = LightWeightInstance->AddNewInstance(&PerInstanceData);
+
+		return FActorInstanceHandle(LightWeightInstance, InstanceIdx);
+	}
+
+	return FActorInstanceHandle();
+}
diff --git a/Engine/Source/Runtime/Engine/Private/LightWeightInstanceManager.cpp b/Engine/Source/Runtime/Engine/Private/LightWeightInstanceManager.cpp
new file mode 100644
--- /dev/null
+++ b/Engine/Source/Runtime/Engine/Private/LightWeightInstanceManager.cpp
@@ -0,0 +1,234 @@
+// Copyright Epic Games, Inc. All Rights Reserved.
+
+#include "GameFramework/LightWeightInstanceManager.h"
+#include "GameFramework/LightWeightInstanceSubsystem.h"
+#include "Net/UnrealNetwork.h"
+
+ALightWeightInstanceManager::ALightWeightInstanceManager(const FObjectInitializer& ObjectInitializer)
+{
+	bReplicates = true;
+	AcceptedClass = AActor::StaticClass();
+	if (!HasAnyFlags(RF_ClassDefaultObject))
+	{
+		FLightWeightInstanceSubsystem::Get().LWInstanceManagers.Add(this);
+	}
+}
+
+ALightWeightInstanceManager::~ALightWeightInstanceManager()
+{
+	FLightWeightInstanceSubsystem::Get().LWInstanceManagers.Remove(this);
+}
+
+void ALightWeightInstanceManager::SetRepresentedClass(UClass* ActorClass)
+{
+	// we should not be changing this once it has been set
+	ensure(RepresentedClass == nullptr);
+
+	RepresentedClass = ActorClass;
+}
+
+void ALightWeightInstanceManager::Tick(float DeltaSeconds)
+{
+	// do nothing
+}
+
+AActor* ALightWeightInstanceManager::GetActorFromHandle(const FActorInstanceHandle& Handle)
+{
+	// make sure the handle doesn't have an actor already
+	if (ensure(!Handle.Actor.IsValid()))
+	{
+		// check if we already have an actor for this handle
+		AActor* const* FoundActor = nullptr;
+		if ((FoundActor = Actors.Find(Handle.GetInstanceIndex())) == nullptr)
+		{
+			// spawn a new actor
+			FActorSpawnParameters SpawnParams;
+			SetSpawnParameters(SpawnParams);
+
+			AActor* Actor = GetLevel()->GetWorld()->SpawnActor<AActor>(RepresentedClass, InstanceTransforms[Handle.GetInstanceIndex()], SpawnParams);
+			check(Actor);
+
+			Handle.Actor = Actor;
+			Actors.Add(Handle.GetInstanceIndex(), Actor);
+
+			PostActorSpawn(Handle);
+		}
+		else
+		{
+			Handle.Actor = *FoundActor;
+		}
+	}
+
+	ensure(Handle.Actor.IsValid());
+	return Handle.Actor.Get();
+}
+
+int32 ALightWeightInstanceManager::FindIndexForActor(const AActor* InActor) const
+{
+	for (const TPair<int32, AActor*>& ActorPair : Actors)
+	{
+		if (ActorPair.Value == InActor)
+		{
+			return ActorPair.Key;
+		}
+	}
+
+	return INDEX_NONE;
+}
+
+int32 ALightWeightInstanceManager::ConvertCollisionIndexToLightWeightIndex(int32 InIndex) const
+{
+	return InIndex;
+}
+
+void ALightWeightInstanceManager::SetSpawnParameters(FActorSpawnParameters& SpawnParams)
+{
+	SpawnParams.OverrideLevel = GetLevel();
+	SpawnParams.SpawnCollisionHandlingOverride = ESpawnActorCollisionHandlingMethod::AlwaysSpawn;
+	SpawnParams.ObjectFlags = RF_Transactional;
+}
+
+void ALightWeightInstanceManager::PostActorSpawn(const FActorInstanceHandle& Handle)
+{
+	// do nothing
+}
+
+bool ALightWeightInstanceManager::IsIndexValid(int32 Index) const
+{
+	if (Index < ValidIndices.Num())
+	{
+		return ValidIndices[Index];
+	}
+
+	return false;
+}
+
+bool ALightWeightInstanceManager::FindActorForHandle(const FActorInstanceHandle& Handle) const
+{
+	ensure(!Handle.Actor.IsValid());
+
+	AActor* const* FoundActor = Actors.Find(Handle.GetInstanceIndex());
+	Handle.Actor = FoundActor ? *FoundActor : nullptr;
+	return Handle.Actor != nullptr;
+}
+
+FVector ALightWeightInstanceManager::GetLocation(const FActorInstanceHandle& Handle) const
+{
+	if (FindActorForHandle(Handle))
+	{
+		return Handle.Actor->GetActorLocation();
+	}
+
+	if (ensure(IsIndexValid(Handle.GetInstanceIndex())))
+	{
+		return InstanceTransforms[Handle.GetInstanceIndex()].GetTranslation();
+	}
+
+	return FVector();
+}
+
+FString ALightWeightInstanceManager::GetName(const FActorInstanceHandle& Handle) const
+{
+	if (FindActorForHandle(Handle))
+	{
+		return Handle.Actor->GetName();
+	}
+
+	return FString::Printf(TEXT("%s_%u"), *BaseInstanceName, Handle.GetInstanceIndex());
+}
+
+bool ALightWeightInstanceManager::DoesRepresentClass(const UClass* OtherClass) const
+{
+	return OtherClass ? OtherClass->IsChildOf(RepresentedClass) : false;
+}
+
+bool ALightWeightInstanceManager::DoesAcceptClass(const UClass* OtherClass) const
+{
+	return OtherClass ? OtherClass->IsChildOf(AcceptedClass) : false;
+}
+
+UClass* ALightWeightInstanceManager::GetRepresentedClass() const
+{
+	return RepresentedClass;
+}
+
+UClass* ALightWeightInstanceManager::GetAcceptedClass() const
+{
+	return AcceptedClass;
+}
+
+int32 ALightWeightInstanceManager::AddNewInstance(FLWIData* InitData)
+{
+	if (!InitData)
+	{
+		return INDEX_NONE;
+	}
+
+	// use one of the free indices if any are available; otherwise grow the size of the array
+	const int32 DataIdx = FreeIndices.Num() > 0 ? FreeIndices.Pop(false) : ValidIndices.Num();
+	
+	// Update the rest of our per instance data
+	AddNewInstanceAt(InitData, DataIdx);
+
+	return DataIdx;
+}
+
+void ALightWeightInstanceManager::AddNewInstanceAt(FLWIData* InitData, int32 Index)
+{
+	// allocate space on the end of the array if we need to
+	if (Index >= ValidIndices.Num())
+	{
+		InstanceTransforms.AddUninitialized();
+		ValidIndices.AddUninitialized();
+	}
+	ensure(Index < ValidIndices.Num());
+	ensure(InstanceTransforms.Num() == ValidIndices.Num());
+
+	// update our data
+	InstanceTransforms[Index] = InitData->Transform;
+	ValidIndices[Index] = true;
+}
+
+void ALightWeightInstanceManager::RemoveInstance(const int32 Index)
+{
+	if (ensure(IsIndexValid(Index)))
+	{
+		// mark the index as no longer in use
+		FreeIndices.Add(Index);
+		ValidIndices[Index] = false;
+
+		// destroy the associated actor if one existed
+		if (AActor** FoundActor = Actors.Find(Index))
+		{
+			AActor* ActorToDestroy = *FoundActor;
+			ActorToDestroy->Destroy();
+		}
+	}
+}
+
+void ALightWeightInstanceManager::GetLifetimeReplicatedProps(TArray<FLifetimeProperty>& OutLifetimeProps) const
+{
+	Super::GetLifetimeReplicatedProps(OutLifetimeProps);
+
+	DOREPLIFETIME(ALightWeightInstanceManager, RepresentedClass);
+	DOREPLIFETIME(ALightWeightInstanceManager, InstanceTransforms);
+	DOREPLIFETIME(ALightWeightInstanceManager, FreeIndices);
+	DOREPLIFETIME(ALightWeightInstanceManager, ValidIndices);
+}
+
+void ALightWeightInstanceManager::OnRep_Transforms()
+{
+	// do nothing
+}
+
+void ALightWeightInstanceManager::PostEditChangeProperty(struct FPropertyChangedEvent& PropertyChangedEvent)
+{
+	Super::PostEditChangeProperty(PropertyChangedEvent);
+
+	const FName Name = (PropertyChangedEvent.Property != nullptr) ? PropertyChangedEvent.Property->GetFName() : NAME_None;
+
+	if (Name == GET_MEMBER_NAME_CHECKED(ALightWeightInstanceManager, RepresentedClass))
+	{
+		SetRepresentedClass(RepresentedClass);
+	}
+}
\ No newline at end of file
diff --git a/Engine/Source/Runtime/Engine/Private/LightWeightInstanceStaticMeshManager.cpp b/Engine/Source/Runtime/Engine/Private/LightWeightInstanceStaticMeshManager.cpp
new file mode 100644
--- /dev/null
+++ b/Engine/Source/Runtime/Engine/Private/LightWeightInstanceStaticMeshManager.cpp
@@ -0,0 +1,193 @@
+// Copyright Epic Games, Inc. All Rights Reserved.
+
+#include "GameFramework/LightWeightInstanceStaticMeshManager.h"
+#include "Components/HierarchicalInstancedStaticMeshComponent.h"
+#include "Net/UnrealNetwork.h"
+
+ALightWeightInstanceStaticMeshManager::ALightWeightInstanceStaticMeshManager(const FObjectInitializer& ObjectInitializer)
+{
+	InstancedStaticMeshComponent = ObjectInitializer.CreateDefaultSubobject<UHierarchicalInstancedStaticMeshComponent>(this, TEXT("InstancedStaticMeshComponent0"));
+
+	if (StaticMesh.IsValid())
+	{
+		OnStaticMeshSet();
+	}
+	SetInstancedStaticMeshParams();
+
+	RootComponent = InstancedStaticMeshComponent;
+	AddInstanceComponent(InstancedStaticMeshComponent);
+	if (GetWorld())
+	{
+		InstancedStaticMeshComponent->RegisterComponent();
+	}
+}
+
+void ALightWeightInstanceStaticMeshManager::SetRepresentedClass(UClass* ActorClass)
+{
+	Super::SetRepresentedClass(ActorClass);
+
+	AActor* ActorCDO = Cast<AActor>(RepresentedClass->GetDefaultObject());
+	if (ActorCDO)
+	{
+		BaseInstanceName = ActorCDO->GetName();
+		SetStaticMeshFromActor(ActorCDO);
+	}
+	else
+	{
+		BaseInstanceName.Reset();
+		ClearStaticMesh();
+	}
+
+	if (InstancedStaticMeshComponent)
+	{
+		InstancedStaticMeshComponent->OnPostLoadPerInstanceData();
+	}
+}
+
+int32 ALightWeightInstanceStaticMeshManager::ConvertCollisionIndexToLightWeightIndex(int32 InIndex) const
+{
+	return RenderingIndicesToDataIndices[InIndex];
+}
+
+void ALightWeightInstanceStaticMeshManager::AddNewInstanceAt(FLWIData* InitData, int32 Index)
+{
+	Super::AddNewInstanceAt(InitData, Index);
+
+	// The rendering indices are tightly packed so we know it's going on the end of the array
+	const int32 RenderingIdx = RenderingIndicesToDataIndices.Add(Index);
+
+	// Now that we know the rendering index we can fill in the other side of the map
+	if (Index >= DataIndicesToRenderingIndices.Num())
+	{
+		ensure(Index == DataIndicesToRenderingIndices.Add(RenderingIdx));
+	}
+	else
+	{
+		DataIndicesToRenderingIndices[Index] = RenderingIdx;
+	}
+
+	// Update the HISMC
+	if (InstancedStaticMeshComponent)
+	{
+		ensure(RenderingIdx == InstancedStaticMeshComponent->AddInstanceWorldSpace(InitData->Transform));
+	}
+}
+
+void ALightWeightInstanceStaticMeshManager::RemoveInstance(const int32 Index)
+{
+	Super::RemoveInstance(Index);
+
+	RemoveInstanceFromRendering(Index);
+}
+
+void ALightWeightInstanceStaticMeshManager::RemoveInstanceFromRendering(int32 DataIndex)
+{
+	if (IsIndexValid(DataIndex))
+	{
+		const int32 RenderingIndex = DataIndicesToRenderingIndices[DataIndex];
+		DataIndicesToRenderingIndices[DataIndex] = INDEX_NONE;
+
+		if (RenderingIndex != INDEX_NONE)
+		{
+			if (InstancedStaticMeshComponent)
+			{
+				InstancedStaticMeshComponent->RemoveInstance(RenderingIndex);
+			}
+
+			// update the map to match what is now in the instanced static mesh component
+			RenderingIndicesToDataIndices.RemoveAtSwap(RenderingIndex);
+
+			// fix up the other side of the map to match the change we just made
+			// if we removed the last element than nothing was moved so we're done
+			if (RenderingIndex < RenderingIndicesToDataIndices.Num())
+			{
+				// find the data index that corresponds with the changed rendering index
+				const int32 ShiftedDataIndex = RenderingIndicesToDataIndices[RenderingIndex];
+
+				DataIndicesToRenderingIndices[ShiftedDataIndex] = RenderingIndex;
+			}
+		}
+	}
+}
+
+void ALightWeightInstanceStaticMeshManager::GetLifetimeReplicatedProps(TArray<FLifetimeProperty>& OutLifetimeProps) const
+{
+	Super::GetLifetimeReplicatedProps(OutLifetimeProps);
+
+	DOREPLIFETIME(ALightWeightInstanceStaticMeshManager, StaticMesh);
+	DOREPLIFETIME(ALightWeightInstanceStaticMeshManager, RenderingIndicesToDataIndices);
+	DOREPLIFETIME(ALightWeightInstanceStaticMeshManager, DataIndicesToRenderingIndices);
+}
+
+void ALightWeightInstanceStaticMeshManager::OnRep_StaticMesh()
+{
+	OnStaticMeshSet();
+}
+
+void ALightWeightInstanceStaticMeshManager::OnRep_Transforms()
+{
+	Super::OnRep_Transforms();
+
+	if (InstancedStaticMeshComponent)
+	{
+		InstancedStaticMeshComponent->AddInstanceWorldSpace(InstanceTransforms.Last());
+	}
+}
+
+void ALightWeightInstanceStaticMeshManager::PostActorSpawn(const FActorInstanceHandle& Handle)
+{
+	Super::PostActorSpawn(Handle);
+
+	// remove the rendered instance from the HISMC
+	RemoveInstanceFromRendering(Handle.GetInstanceIndex());
+}
+
+void ALightWeightInstanceStaticMeshManager::SetInstancedStaticMeshParams()
+{
+	FName CollisionProfileName(TEXT("LightWeightInstancedStaticMeshPhysics"));
+	InstancedStaticMeshComponent->SetCollisionProfileName(CollisionProfileName);
+
+	InstancedStaticMeshComponent->CanCharacterStepUpOn = ECB_Owner;
+	InstancedStaticMeshComponent->CastShadow = true;
+	InstancedStaticMeshComponent->bCastDynamicShadow = true;
+	InstancedStaticMeshComponent->bCastStaticShadow = true;
+	InstancedStaticMeshComponent->PrimaryComponentTick.bCanEverTick = false;
+	// Allows updating in game, while optimizing rendering for the case that it is not modified
+	InstancedStaticMeshComponent->Mobility = EComponentMobility::Movable;
+}
+
+void ALightWeightInstanceStaticMeshManager::SetStaticMeshFromActor(AActor* InActor)
+{
+	ensureMsgf(false, TEXT("ALightWeightInstanceManager::SetStaticMeshFromActor was called on %s. Projects should override this function in subclasses."), *GetNameSafe(InActor));
+	ClearStaticMesh();
+}
+
+void ALightWeightInstanceStaticMeshManager::OnStaticMeshSet()
+{
+	if (InstancedStaticMeshComponent)
+	{
+		EComponentMobility::Type Mobility = InstancedStaticMeshComponent->Mobility;
+		if (Mobility == EComponentMobility::Static)
+		{
+			InstancedStaticMeshComponent->SetMobility(EComponentMobility::Stationary);
+			InstancedStaticMeshComponent->SetStaticMesh(StaticMesh.Get());
+			InstancedStaticMeshComponent->SetMobility(Mobility);
+		}
+		else
+		{
+			InstancedStaticMeshComponent->SetStaticMesh(StaticMesh.Get());
+		}
+
+		UStaticMesh* Mesh = StaticMesh.Get();
+		for (int32 Idx = 0; Idx < Mesh->StaticMaterials.Num(); ++Idx)
+		{
+			InstancedStaticMeshComponent->SetMaterial(Idx, Mesh->GetMaterial(Idx));
+		}
+	}
+}
+
+void ALightWeightInstanceStaticMeshManager::ClearStaticMesh()
+{
+	StaticMesh = nullptr;
+	OnStaticMeshSet();
+}
\ No newline at end of file
diff --git a/Engine/Source/Runtime/Engine/Private/LightWeightInstanceSubsystem.cpp b/Engine/Source/Runtime/Engine/Private/LightWeightInstanceSubsystem.cpp
new file mode 100644
--- /dev/null
+++ b/Engine/Source/Runtime/Engine/Private/LightWeightInstanceSubsystem.cpp
@@ -0,0 +1,266 @@
+// Copyright Epic Games, Inc. All Rights Reserved.
+
+#include "GameFramework/LightWeightInstanceSubsystem.h"
+#include "GameFramework/LightWeightInstanceManager.h"
+
+DEFINE_LOG_CATEGORY(LogLightWeightInstance);
+
+TSharedPtr<FLightWeightInstanceSubsystem> FLightWeightInstanceSubsystem::LWISubsystem;
+
+int32 FLightWeightInstanceSubsystem::GetManagerIndex(const ALightWeightInstanceManager* Manager) const
+{
+	for (int32 Idx = 0; Idx < LWInstanceManagers.Num(); ++Idx)
+	{
+		if (LWInstanceManagers[Idx] == Manager)
+		{
+			return Idx;
+		}
+	}
+
+	return INDEX_NONE;
+}
+
+const ALightWeightInstanceManager* FLightWeightInstanceSubsystem::GetManagerAt(int32 Index) const
+{
+	return Index < LWInstanceManagers.Num() ? LWInstanceManagers[Index] : nullptr;
+}
+
+ALightWeightInstanceManager* FLightWeightInstanceSubsystem::FindLightWeightInstanceManager(const FActorInstanceHandle& Handle) const
+{
+	if (Handle.ManagerIndex != INDEX_NONE)
+	{
+		if (ensure(Handle.ManagerIndex < LWInstanceManagers.Num()))
+		{
+			return LWInstanceManagers[Handle.ManagerIndex];
+		}
+	}
+
+	if (Handle.Actor.IsValid())
+	{
+		for (ALightWeightInstanceManager* LWInstance : LWInstanceManagers)
+		{
+			if (Handle.Actor->GetClass() == LWInstance->GetRepresentedClass() && Handle.Actor->GetLevel() == LWInstance->GetLevel())
+			{
+				return LWInstance;
+			}
+		}
+	}
+
+	return nullptr;
+}
+
+ALightWeightInstanceManager* FLightWeightInstanceSubsystem::FindLightWeightInstanceManager(UClass* ActorClass, ULevel* Level) const
+{
+	if (ActorClass == nullptr || Level == nullptr)
+	{
+		return nullptr;
+	}
+
+	// see if we already have a match
+	for (ALightWeightInstanceManager* Instance : LWInstanceManagers)
+	{
+		if (Instance->GetRepresentedClass() == ActorClass && Instance->GetLevel() == Level)
+		{
+			return Instance;
+		}
+	}
+
+	return nullptr;
+}
+
+ALightWeightInstanceManager* FLightWeightInstanceSubsystem::FindOrAddLightWeightInstanceManager(UClass* ActorClass, ULevel* Level)
+{
+	if (ActorClass == nullptr || Level == nullptr)
+	{
+		return nullptr;
+	}
+	
+	// see if we already have a match
+	for (ALightWeightInstanceManager* Instance : LWInstanceManagers)
+	{
+		if (Instance->GetRepresentedClass() == ActorClass && Instance->GetLevel() == Level)
+		{
+			return Instance;
+		}
+	}
+
+	// we didn't find a match so we should add one.
+
+	// Find the best base class to start from
+	UClass* BestMatchingClass = FindBestInstanceManagerClass(ActorClass);
+	if (BestMatchingClass == nullptr)
+	{
+		return nullptr;
+	}
+
+	FActorSpawnParameters SpawnParams;
+	SpawnParams.OverrideLevel = Level;
+	SpawnParams.SpawnCollisionHandlingOverride = ESpawnActorCollisionHandlingMethod::AlwaysSpawn;
+	SpawnParams.ObjectFlags = RF_Transactional;
+
+	ALightWeightInstanceManager* NewInstance = Level->GetWorld()->SpawnActor<ALightWeightInstanceManager>(BestMatchingClass, FTransform::Identity, SpawnParams);
+	ULevel* InstanceLvl = NewInstance->GetLevel();
+	NewInstance->SetRepresentedClass(ActorClass);
+
+	check(LWInstanceManagers.Find(NewInstance) != INDEX_NONE);
+
+	return NewInstance;
+}
+
+UClass* FLightWeightInstanceSubsystem::FindBestInstanceManagerClass(const UClass* InActorClass)
+{
+	// Get every light weight instance class
+	// FRED_TODO: we need to search through unloaded blueprints as well
+	// FRED_TODO: this should be done once and cached. In the editor we can add a listener for when new BP classes are added
+	TArray<UClass*> ManagerClasses;
+	for (TObjectIterator<UClass> It; It; ++It)
+	{
+		if (It->IsChildOf(ALightWeightInstanceManager::StaticClass()) && !It->HasAnyClassFlags(CLASS_Abstract))
+		{
+			// Skip SKEL and REINST classes.
+			if (It->GetName().StartsWith(TEXT("SKEL_")) || It->GetName().StartsWith(TEXT("REINST_")))
+			{
+				continue;
+			}
+			ManagerClasses.Add(*It);
+		}
+	}
+
+	// Figure out which one is the closest fit for InActorClass
+	UClass* BestManagerClass = nullptr;
+	int32 BestDistance = INT_MAX;
+	for (UClass* ManagerClass : ManagerClasses)
+	{
+		if (ManagerClass->GetDefaultObject<ALightWeightInstanceManager>()->DoesAcceptClass(InActorClass))
+		{
+			const UClass* HandledClass = ManagerClass->GetDefaultObject<ALightWeightInstanceManager>()->GetRepresentedClass();
+			if (!HandledClass)
+			{
+				HandledClass = ManagerClass->GetDefaultObject<ALightWeightInstanceManager>()->GetAcceptedClass();
+			}
+			if (InActorClass == HandledClass)
+			{
+				BestManagerClass = ManagerClass;
+				break;
+			}
+			const UClass* ActorClass = InActorClass;
+			int32 Distance = 0;
+			while (ActorClass && ActorClass != HandledClass)
+			{
+				++Distance;
+				ActorClass = ActorClass->GetSuperClass();
+			}
+			if (ActorClass && Distance < BestDistance)
+			{
+				BestDistance = Distance;
+				BestManagerClass = ManagerClass;
+			}
+		}
+	}
+
+	return BestManagerClass;
+}
+
+AActor* FLightWeightInstanceSubsystem::GetActor(const FActorInstanceHandle& Handle)
+{
+	// if the actor is valid return it
+	if (Handle.Actor.IsValid())
+	{
+		return Handle.Actor.Get();
+	}
+
+	if (ALightWeightInstanceManager* LWIManager = FindLightWeightInstanceManager(Handle))
+	{
+		return LWIManager->GetActorFromHandle(Handle);
+	}
+
+	return nullptr;
+}
+
+AActor* FLightWeightInstanceSubsystem::GetActor_NoCreate(const FActorInstanceHandle& Handle) const
+{
+	return Handle.Actor.Get();
+}
+
+UClass* FLightWeightInstanceSubsystem::GetActorClass(const FActorInstanceHandle& Handle)
+{
+	if (Handle.Actor.IsValid())
+	{
+		return Handle.Actor->StaticClass();
+	}
+
+	if (ALightWeightInstanceManager* LWIManager = FindLightWeightInstanceManager(Handle))
+	{
+		return LWIManager->GetRepresentedClass();
+	}
+
+	return nullptr;
+}
+
+FVector FLightWeightInstanceSubsystem::GetLocation(const FActorInstanceHandle& Handle)
+{
+	ensure(Handle.IsValid());
+
+	if (Handle.Actor.IsValid())
+	{
+		return Handle.Actor->GetActorLocation();
+	}
+
+	if (ALightWeightInstanceManager* InstanceManager = FindLightWeightInstanceManager(Handle))
+	{
+		return InstanceManager->GetLocation(Handle);
+	}
+
+	return FVector::ZeroVector;
+}
+
+FString FLightWeightInstanceSubsystem::GetName(const FActorInstanceHandle& Handle)
+{
+	ensure(Handle.IsValid());
+
+	if (Handle.Actor.IsValid())
+	{
+		return Handle.Actor->GetName();
+	}
+
+	if (ALightWeightInstanceManager* InstanceManager = FindLightWeightInstanceManager(Handle))
+	{
+		return InstanceManager->GetName(Handle);
+	}
+
+	return TEXT("None");
+}
+
+ULevel* FLightWeightInstanceSubsystem::GetLevel(const FActorInstanceHandle& Handle)
+{
+	ensure(Handle.IsValid());
+
+	if (Handle.Actor.IsValid())
+	{
+		return Handle.Actor->GetLevel();
+	}
+
+	if (ALightWeightInstanceManager* InstanceManager = FindLightWeightInstanceManager(Handle))
+	{
+		return InstanceManager->GetLevel();
+	}
+
+	return nullptr;
+}
+
+bool FLightWeightInstanceSubsystem::IsInLevel(const FActorInstanceHandle& Handle, const ULevel* InLevel)
+{
+	ensure(Handle.IsValid());
+
+	if (Handle.Actor.IsValid())
+	{
+		return Handle.Actor->IsInLevel(InLevel);
+	}
+
+	if (ALightWeightInstanceManager* InstanceManager = FindLightWeightInstanceManager(Handle))
+	{
+		return InstanceManager->GetLevel() == InLevel;
+	}
+
+	return false;
+}
\ No newline at end of file
diff --git a/Engine/Source/Runtime/Engine/Private/Particles/ParticleModules_Collision.cpp b/Engine/Source/Runtime/Engine/Private/Particles/ParticleModules_Collision.cpp
--- a/Engine/Source/Runtime/Engine/Private/Particles/ParticleModules_Collision.cpp
+++ b/Engine/Source/Runtime/Engine/Private/Particles/ParticleModules_Collision.cpp
@@ -385,10 +385,10 @@ void UParticleModuleCollision::Update(FParticleEmitterInstance* Owner, int32 Off
 		{
 			bool bDecrementMaxCount = true;
 			bool bIgnoreCollision = false;
-			if (Hit.GetActor())
+			if (Hit.HitObjectHandle.IsValid())
 			{
-				bDecrementMaxCount = !bPawnsDoNotDecrementCount || !Cast<APawn>(Hit.GetActor());
-				bIgnoreCollision = bIgnoreTriggerVolumes && Hit.GetActor()->IsA(ATriggerBase::StaticClass());
+				bDecrementMaxCount = !bPawnsDoNotDecrementCount || !Hit.HitObjectHandle.DoesRepresentClass(APawn::StaticClass());
+				bIgnoreCollision = bIgnoreTriggerVolumes && Hit.HitObjectHandle.DoesRepresentClass(ATriggerBase::StaticClass());
 				[REDACTED_EMAIL]. Allow for PSys to say what it wants to collide w/?
 			}
 
diff --git a/Engine/Source/Runtime/Engine/Private/PhysicsEngine/CollisionAnalyzerCapture.h b/Engine/Source/Runtime/Engine/Private/PhysicsEngine/CollisionAnalyzerCapture.h
--- a/Engine/Source/Runtime/Engine/Private/PhysicsEngine/CollisionAnalyzerCapture.h
+++ b/Engine/Source/Runtime/Engine/Private/PhysicsEngine/CollisionAnalyzerCapture.h
@@ -187,7 +187,7 @@ inline void CaptureOverlap(const UWorld* World, const FPhysicsGeometryCollection
 	{
 		FHitResult NewResult = FHitResult(0.f);
 		NewResult.bBlockingHit = OverlapResult.bBlockingHit;
-		NewResult.Actor = OverlapResult.Actor;
+		NewResult.HitObjectHandle = OverlapResult.OverlapObjectHandle;
 		NewResult.Component = OverlapResult.Component;
 		NewResult.Item = OverlapResult.ItemIndex;
 		HitResults.Add(NewResult);
@@ -216,7 +216,7 @@ inline void CaptureOverlap(const UWorld* World, const FCollisionShape& PGeom, co
 	{
 		FHitResult NewResult = FHitResult(0.f);
 		NewResult.bBlockingHit = OverlapResult.bBlockingHit;
-		NewResult.Actor = OverlapResult.Actor;
+		NewResult.HitObjectHandle = OverlapResult.OverlapObjectHandle;
 		NewResult.Component = OverlapResult.Component;
 		NewResult.Item = OverlapResult.ItemIndex;
 		HitResults.Add(NewResult);
diff --git a/Engine/Source/Runtime/Engine/Private/StaticMeshActor.cpp b/Engine/Source/Runtime/Engine/Private/StaticMeshActor.cpp
--- a/Engine/Source/Runtime/Engine/Private/StaticMeshActor.cpp
+++ b/Engine/Source/Runtime/Engine/Private/StaticMeshActor.cpp
@@ -184,20 +184,25 @@ void AStaticMeshActor::CheckForErrors()
 		TArray<FOverlapResult> Overlaps;
 		GetWorld()->OverlapMultiByChannel(Overlaps, GetActorLocation(), FQuat::Identity, ECC_Pawn, FCollisionShape::MakeSphere(1.f), SphereParams);
 
-		for ( int32 OverlapIdx=0; OverlapIdx<Overlaps.Num(); OverlapIdx++ )
+		for ( const FOverlapResult& Overlap : Overlaps )
 		{
-			AStaticMeshActor *A = Cast<AStaticMeshActor>(Overlaps[OverlapIdx].GetActor());
-			if ( A && (A != this) && (A->GetActorLocation() - GetActorLocation()).IsNearlyZero() && A->StaticMeshComponent
-				&& (A->StaticMeshComponent->GetStaticMesh() == StaticMeshComponent->GetStaticMesh()) && (A->GetActorRotation() == GetActorRotation())
-				&& (A->StaticMeshComponent->GetRelativeScale3D() == StaticMeshComponent->GetRelativeScale3D()) )
+			if (Overlap.OverlapObjectHandle.IsValid() && Overlap.OverlapObjectHandle != this && Overlap.OverlapObjectHandle.DoesRepresentClass(AStaticMeshActor::StaticClass()) &&
+				(Overlap.OverlapObjectHandle.GetLocation() - GetActorLocation()).IsNearlyZero() &&
+				(Overlap.OverlapObjectHandle.GetRotation() == GetActorRotation()))
 			{
-				FFormatNamedArguments Arguments;
-				Arguments.Add(TEXT("ActorName0"), FText::FromString(GetName()));
-				Arguments.Add(TEXT("ActorName1"), FText::FromString(A->GetName()));
-				MapCheck.Warning()
-					->AddToken(FUObjectToken::Create(this))
-					->AddToken(FTextToken::Create(FText::Format( LOCTEXT( "MapCheck_Message_SameLocation", "{ActorName0} is in the same location as {ActorName1}" ), Arguments ) ))
-					->AddToken(FMapErrorToken::Create(FMapErrors::SameLocation));
+				AStaticMeshActor* A = Overlap.OverlapObjectHandle.FetchActor<AStaticMeshActor>();
+				check(A);
+
+				if (A->StaticMeshComponent && (A->StaticMeshComponent->GetStaticMesh() == StaticMeshComponent->GetStaticMesh()) && (A->StaticMeshComponent->GetRelativeScale3D() == StaticMeshComponent->GetRelativeScale3D()))
+				{
+					FFormatNamedArguments Arguments;
+						Arguments.Add(TEXT("ActorName0"), FText::FromString(GetName()));
+						Arguments.Add(TEXT("ActorName1"), FText::FromString(A->GetName()));
+						MapCheck.Warning()
+						->AddToken(FUObjectToken::Create(this))
+						->AddToken(FTextToken::Create(FText::Format(LOCTEXT("MapCheck_Message_SameLocation", "{ActorName0} is in the same location as {ActorName1}"), Arguments)))
+						->AddToken(FMapErrorToken::Create(FMapErrors::SameLocation));
+				}
 			}
 		}
 
diff --git a/Engine/Source/Runtime/Experimental/ChaosSolverEngine/Private/Chaos/ChaosNotifyHandlerInterface.cpp b/Engine/Source/Runtime/Experimental/ChaosSolverEngine/Private/Chaos/ChaosNotifyHandlerInterface.cpp
--- a/Engine/Source/Runtime/Experimental/ChaosSolverEngine/Private/Chaos/ChaosNotifyHandlerInterface.cpp
+++ b/Engine/Source/Runtime/Experimental/ChaosSolverEngine/Private/Chaos/ChaosNotifyHandlerInterface.cpp
@@ -22,7 +22,7 @@ FHitResult UChaosSolverEngineBlueprintLibrary::ConvertPhysicsCollisionToHitResul
 	FHitResult Hit(0.f);
 	
 	Hit.Component = PhysicsCollision.OtherComponent;
-	Hit.Actor = Hit.Component.IsValid() ? Hit.Component->GetOwner() : nullptr;
+	Hit.HitObjectHandle = FActorInstanceHandle(Hit.Component.IsValid() ? Hit.Component->GetOwner() : nullptr);
 	Hit.bBlockingHit = true;
 	Hit.Normal = PhysicsCollision.Normal;
 	Hit.ImpactNormal = PhysicsCollision.Normal;
diff --git a/Engine/Source/Runtime/Foliage/Private/InstancedFoliage.cpp b/Engine/Source/Runtime/Foliage/Private/InstancedFoliage.cpp
--- a/Engine/Source/Runtime/Foliage/Private/InstancedFoliage.cpp
+++ b/Engine/Source/Runtime/Foliage/Private/InstancedFoliage.cpp
@@ -4446,13 +4446,13 @@ bool AInstancedFoliageActor::FoliageTrace(const UWorld* InWorld, FHitResult& Out
 
 	for (const FHitResult& Hit : Hits)
 	{
-		const AActor* HitActor = Hit.GetActor();
+		const FActorInstanceHandle& HitObjectHandle = Hit.HitObjectHandle;
 				
 		// don't place procedural foliage inside an AProceduralFoliageBlockingVolume
 		// this test is first because two of the tests below would otherwise cause the trace to ignore AProceduralFoliageBlockingVolume
 		if (DesiredInstance.PlacementMode == EFoliagePlacementMode::Procedural)
 		{
-			if (const AProceduralFoliageBlockingVolume* ProceduralFoliageBlockingVolume = Cast<AProceduralFoliageBlockingVolume>(HitActor))
+			if (const AProceduralFoliageBlockingVolume* ProceduralFoliageBlockingVolume = HitObjectHandle.FetchActor<AProceduralFoliageBlockingVolume>())
 			{
 				const AProceduralFoliageVolume* ProceduralFoliageVolume = ProceduralFoliageBlockingVolume->ProceduralFoliageVolume;
 				if (ProceduralFoliageVolume == nullptr || ProceduralFoliageVolume->ProceduralComponent == nullptr || ProceduralFoliageVolume->ProceduralComponent->GetProceduralGuid() == DesiredInstance.ProceduralGuid)
@@ -4460,7 +4460,7 @@ bool AInstancedFoliageActor::FoliageTrace(const UWorld* InWorld, FHitResult& Out
 					return false;
 				}
 			}
-			else if (HitActor && HitActor->IsA<AProceduralFoliageVolume>()) //we never want to collide with our spawning volume
+			else if (HitObjectHandle.IsValid() && HitObjectHandle.DoesRepresentClass(AProceduralFoliageVolume::StaticClass())) //we never want to collide with our spawning volume
 			{
 				continue;
 			}
@@ -4482,6 +4482,7 @@ bool AInstancedFoliageActor::FoliageTrace(const UWorld* InWorld, FHitResult& Out
 		}
 
 		// Don't place foliage on itself
+		const AActor* HitActor = Hit.HitObjectHandle.FetchActor();
 		const AInstancedFoliageActor* FoliageActor = Cast<AInstancedFoliageActor>(HitActor);
 		if (!FoliageActor && HitActor && FFoliageHelper::IsOwnedByFoliage(HitActor))
 		{
diff --git a/Engine/Source/Runtime/NavigationSystem/Private/NavLinkCustomComponent.cpp b/Engine/Source/Runtime/NavigationSystem/Private/NavLinkCustomComponent.cpp
--- a/Engine/Source/Runtime/NavigationSystem/Private/NavLinkCustomComponent.cpp
+++ b/Engine/Source/Runtime/NavigationSystem/Private/NavLinkCustomComponent.cpp
@@ -330,7 +330,7 @@ void UNavLinkCustomComponent::CollectNearbyAgents(TArray<UObject*>& NotifyList)
 	TArray<AController*> ControllerList;
 	for (int32 i = 0; i < OverlapsL.Num(); i++)
 	{
-		APawn* MovingPawn = Cast<APawn>(OverlapsL[i].GetActor());
+		APawn* MovingPawn = OverlapsL[i].OverlapObjectHandle.FetchActor<APawn>();
 		if (MovingPawn && MovingPawn->GetController())
 		{
 			ControllerList.Add(MovingPawn->GetController());
@@ -338,7 +338,7 @@ void UNavLinkCustomComponent::CollectNearbyAgents(TArray<UObject*>& NotifyList)
 	}
 	for (int32 i = 0; i < OverlapsR.Num(); i++)
 	{
-		APawn* MovingPawn = Cast<APawn>(OverlapsR[i].GetActor());
+		APawn* MovingPawn = OverlapsR[i].OverlapObjectHandle.FetchActor<APawn>();
 		if (MovingPawn && MovingPawn->GetController())
 		{
 			ControllerList.Add(MovingPawn->GetController());
diff --git a/Engine/Source/Runtime/NavigationSystem/Private/NavigationSystemTypes.cpp b/Engine/Source/Runtime/NavigationSystem/Private/NavigationSystemTypes.cpp
--- a/Engine/Source/Runtime/NavigationSystem/Private/NavigationSystemTypes.cpp
+++ b/Engine/Source/Runtime/NavigationSystem/Private/NavigationSystemTypes.cpp
@@ -27,7 +27,7 @@ namespace
 		if (bHit)
 		{
 			UE_VLOG_LOCATION(Querier, LogNavigation, Log, Hit.Location, 15, FColor::Red, TEXT("%s")
-				, Hit.Actor.IsValid() ? *Hit.Actor->GetName() : TEXT("NULL"));
+				, *Hit.HitObjectHandle.GetName());
 
 			if (Cast<UStaticMeshComponent>(Hit.Component.Get()))
 			{
