Re-add GetActor to FHitResult to fix broken projects.

diff --git a/Engine/Source/Runtime/Engine/Classes/Engine/EngineTypes.h b/Engine/Source/Runtime/Engine/Classes/Engine/EngineTypes.h
--- a/Engine/Source/Runtime/Engine/Classes/Engine/EngineTypes.h
+++ b/Engine/Source/Runtime/Engine/Classes/Engine/EngineTypes.h
@@ -2165,6 +2165,12 @@ struct ENGINE_API FHitResult
 		}
 	}
 
+	/** Utility to return the Actor that owns the Component that was hit. */
+	FORCEINLINE AActor* GetActor() const
+	{
+		HitObjectHandle.FetchActor();
+	}
+
 	FORCEINLINE FActorInstanceHandle GetHitObjectHandle() const
 	{
 		return HitObjectHandle;
