Added missing return statement.

diff --git a/Engine/Source/Runtime/Engine/Classes/Engine/EngineTypes.h b/Engine/Source/Runtime/Engine/Classes/Engine/EngineTypes.h
--- a/Engine/Source/Runtime/Engine/Classes/Engine/EngineTypes.h
+++ b/Engine/Source/Runtime/Engine/Classes/Engine/EngineTypes.h
@@ -2168,7 +2168,7 @@ struct ENGINE_API FHitResult
 	/** Utility to return the Actor that owns the Component that was hit. */
 	FORCEINLINE AActor* GetActor() const
 	{
-		HitObjectHandle.FetchActor();
+		return HitObjectHandle.FetchActor();
 	}
 
 	FORCEINLINE FActorInstanceHandle GetHitObjectHandle() const
