re-add OverlapResult::GetActor() with updated implementation to use handle.

diff --git a/Engine/Source/Runtime/Engine/Classes/Engine/EngineTypes.h b/Engine/Source/Runtime/Engine/Classes/Engine/EngineTypes.h
--- a/Engine/Source/Runtime/Engine/Classes/Engine/EngineTypes.h
+++ b/Engine/Source/Runtime/Engine/Classes/Engine/EngineTypes.h
@@ -2293,6 +2293,9 @@ struct ENGINE_API FOverlapResult
 		For SkeletalMeshComponents this is the Body index or INDEX_NONE for single body */
 	int32 ItemIndex;
 
+	/** Utility to return the Actor that owns the Component that was hit */
+	AActor* GetActor() const;
+
 	/** Utility to return the Component that was hit */
 	UPrimitiveComponent* GetComponent() const;
 
diff --git a/Engine/Source/Runtime/Engine/Private/Collision/Collision.cpp b/Engine/Source/Runtime/Engine/Private/Collision/Collision.cpp
--- a/Engine/Source/Runtime/Engine/Private/Collision/Collision.cpp
+++ b/Engine/Source/Runtime/Engine/Private/Collision/Collision.cpp
@@ -157,6 +157,11 @@ bool FHitResult::NetSerialize(FArchive& Ar, class UPackageMap* Map, bool& bOutSu
 //////////////////////////////////////////////////////////////////////////
 // FOverlapResult
 
+AActor* FOverlapResult::GetActor() const
+{
+	return OverlapObjectHandle.FetchActor();
+}
+
 UPrimitiveComponent* FOverlapResult::GetComponent() const
 {
 	return Component.Get();
