WorldPartition: Fix Crash when loading cell with unsaved (discarded) actor

diff --git a/Engine/Source/Runtime/Engine/Private/WorldPartition/WorldPartition.cpp b/Engine/Source/Runtime/Engine/Private/WorldPartition/WorldPartition.cpp
--- a/Engine/Source/Runtime/Engine/Private/WorldPartition/WorldPartition.cpp
+++ b/Engine/Source/Runtime/Engine/Private/WorldPartition/WorldPartition.cpp
@@ -786,6 +786,7 @@ void UWorldPartition::UnloadEditorCells(const TArray<UWorldPartitionEditorCell*>
 	}
 		
 	// Make sure we save modified actor packages before unloading
+	FEditorFileUtils::EPromptReturnCode RetCode = FEditorFileUtils::PR_Success;
 	if (ModifiedPackages.Num())
 	{
 		const bool bCheckDirty = false;
@@ -795,7 +796,7 @@ void UWorldPartition::UnloadEditorCells(const TArray<UWorldPartitionEditorCell*>
 		const FText Title = LOCTEXT("SaveActorsTitle", "Save Actor(s)");
 		const FText Message = LOCTEXT("SaveActorsMessage", "Save Actor(s) before unloading them.");
 
-		FEditorFileUtils::EPromptReturnCode RetCode = FEditorFileUtils::PromptForCheckoutAndSave(ModifiedPackages.Array(), bCheckDirty, bPromptToSave, Title, Message, nullptr, bAlreadyCheckedOut, bCanBeDeclined);
+		RetCode = FEditorFileUtils::PromptForCheckoutAndSave(ModifiedPackages.Array(), bCheckDirty, bPromptToSave, Title, Message, nullptr, bAlreadyCheckedOut, bCanBeDeclined);
 		if (RetCode == FEditorFileUtils::PR_Cancelled)
 		{
 			return;
@@ -817,9 +818,24 @@ void UWorldPartition::UnloadEditorCells(const TArray<UWorldPartitionEditorCell*>
 		
 	GEditor->ResetTransaction(LOCTEXT("UnloadingEditorCellsResetTrans", "Unloading Cells"));
 	
-	for (UPackage* ModifiedPackage : ModifiedPackages)
+	// When save is declined make sure we don't keep unloaded/unsaved actor descs
+	// also make sure all modified packages are no longer dirty so they don't show up in save prompts anymore.
+	if (RetCode == FEditorFileUtils::PR_Declined)
 	{
-		ModifiedPackage->ClearDirtyFlag();
+		for (UPackage* ModifiedPackage : ModifiedPackages)
+		{
+			ModifiedPackage->ClearDirtyFlag();
+		}
+
+		for (const TPair<FWorldPartitionActorDesc*, int32>& Pair : UnloadCount)
+		{
+			FWorldPartitionActorDesc* ActorDesc = Pair.Key;
+			if (!ActorDesc->GetLoadedRefCount() && !FPackageName::DoesPackageExist(ActorDesc->GetActorPackage().ToString()))
+			{
+				// Already unloaded just remove from descriptor array
+				RemoveFromPartition(ActorDesc, /*bRemoveDescriptorFromArray=*/true, /*bUnloadRemovedDescriptor=*/false);
+			}
+		}
 	}
 }
 
