Fix Non-Editor build

diff --git a/Engine/Source/Runtime/Engine/Private/DistanceFieldAtlas.cpp b/Engine/Source/Runtime/Engine/Private/DistanceFieldAtlas.cpp
--- a/Engine/Source/Runtime/Engine/Private/DistanceFieldAtlas.cpp
+++ b/Engine/Source/Runtime/Engine/Private/DistanceFieldAtlas.cpp
@@ -21,13 +21,16 @@
 #include "GlobalShader.h"
 #include "RenderGraph.h"
 #include "MeshCardRepresentation.h"
-#include "IMeshBuilderModule.h"
 
 #if WITH_EDITOR
 #include "DerivedDataCacheInterface.h"
 #include "MeshUtilities.h"
 #endif
 
+#if WITH_EDITORONLY_DATA
+#include "IMeshBuilderModule.h"
+#endif
+
 CSV_DEFINE_CATEGORY(DistanceField, false);
 
 #if ENABLE_COOK_STATS
diff --git a/Engine/Source/Runtime/Engine/Private/MeshCardRepresentation.cpp b/Engine/Source/Runtime/Engine/Private/MeshCardRepresentation.cpp
--- a/Engine/Source/Runtime/Engine/Private/MeshCardRepresentation.cpp
+++ b/Engine/Source/Runtime/Engine/Private/MeshCardRepresentation.cpp
@@ -17,13 +17,16 @@
 #include "Engine/StaticMesh.h"
 #include "Misc/AutomationTest.h"
 #include "Async/ParallelFor.h"
-#include "IMeshBuilderModule.h"
 
 #if WITH_EDITOR
 #include "DerivedDataCacheInterface.h"
 #include "MeshUtilities.h"
 #endif
 
+#if WITH_EDITORONLY_DATA
+#include "IMeshBuilderModule.h"
+#endif
+
 #if ENABLE_COOK_STATS
 namespace CardRepresentationCookStats
 {
