World Partition - Remove automatic LOD parenting of loaded actors to their matching HLOD (editor only)

diff --git a/Engine/Source/Runtime/Engine/Private/WorldPartition/HLOD/HLODActor.cpp b/Engine/Source/Runtime/Engine/Private/WorldPartition/HLOD/HLODActor.cpp
--- a/Engine/Source/Runtime/Engine/Private/WorldPartition/HLOD/HLODActor.cpp
+++ b/Engine/Source/Runtime/Engine/Private/WorldPartition/HLOD/HLODActor.cpp
@@ -55,30 +55,6 @@ EActorGridPlacement AWorldPartitionHLOD::GetDefaultGridPlacement() const
 	return EActorGridPlacement::Location;
 }
 
-void AWorldPartitionHLOD::SetLODParent(AActor& InActor)
-{
-	UpdateLODParent(InActor, false);
-}
-
-void AWorldPartitionHLOD::ClearLODParent(AActor& InActor)
-{
-	UpdateLODParent(InActor, true);
-}
-
-void AWorldPartitionHLOD::UpdateLODParent(AActor& InActor, bool bInClear)
-{
-	for (UActorComponent* Component : InActor.GetComponents())
-	{
-		if (UPrimitiveComponent* PrimitiveComponent = Cast<UPrimitiveComponent>(Component))
-		{
-			if (UHLODLayer::ShouldIncludeInHLOD(PrimitiveComponent, SubActorsHLODLevel))
-			{
-				PrimitiveComponent->SetCachedLODParentPrimitive(!bInClear ? GetHLODComponent() : nullptr);
-			}
-		}
-	}
-}
-
 void AWorldPartitionHLOD::SetHLODPrimitives(const TArray<UPrimitiveComponent*>& InHLODPrimitives, float InFadeOutDistance)
 {
 	TRACE_CPUPROFILER_EVENT_SCOPE(AWorldPartitionHLOD::SetHLODPrimitive);
@@ -149,8 +125,6 @@ void AWorldPartitionHLOD::SetChildrenPrimitives(const TArray<UPrimitiveComponent
 	for (UPrimitiveComponent* ChildPrimitive : InChildrenPrimitives)
 	{
 		SubActorsSet.Add(ChildPrimitive->GetOwner()->GetActorGuid());
-
-		ChildPrimitive->SetCachedLODParentPrimitive(HLODComponent);
 	}
 
 	SubActors = SubActorsSet.Array();
diff --git a/Engine/Source/Runtime/Engine/Private/WorldPartition/HLOD/HLODSubsystem.cpp b/Engine/Source/Runtime/Engine/Private/WorldPartition/HLOD/HLODSubsystem.cpp
--- a/Engine/Source/Runtime/Engine/Private/WorldPartition/HLOD/HLODSubsystem.cpp
+++ b/Engine/Source/Runtime/Engine/Private/WorldPartition/HLOD/HLODSubsystem.cpp
@@ -48,134 +48,10 @@ void UHLODSubsystem::Initialize(FSubsystemCollectionBase& Collection)
 	HLODActorDescFactory.Reset(new FHLODActorDescFactory());
 
 	RegisterActorDescFactories(WorldPartitionSubsystem);
-
-	GetWorld()->GetWorldPartition()->OnActorRegisteredEvent.AddUObject(this, &UHLODSubsystem::OnWorldPartitionActorRegistered);
 #endif
 }
 
 #if WITH_EDITOR
-
-void UHLODSubsystem::OnWorldPartitionActorRegistered(AActor& InActor, bool bInLoaded)
-{
-	TRACE_CPUPROFILER_EVENT_SCOPE(UHLODSubsystem::OnWorldPartitionActorRegistered);
-
-	const FGuid& ActorGuid = InActor.GetActorGuid();
-
-	if (AWorldPartitionHLOD* HLODActor = Cast<AWorldPartitionHLOD>(&InActor))
-	{
-		if (bInLoaded)
-		{
-			UWorldPartition* WorldPartition = GetWorld()->GetWorldPartition();
-
-			TArray<FGuid> SubActors;
-			PendingHLODAssignment.MultiFind(ActorGuid, SubActors);
-
-			if (!SubActors.IsEmpty())
-			{
-				for (const FGuid& SubActorGuid : SubActors)
-				{
-					const FWorldPartitionActorDesc* ActorDesc = WorldPartition->GetActorDesc(SubActorGuid);
-					check(ActorDesc);
-
-					AActor* Actor = Cast<AWorldPartitionHLOD>(ActorDesc->GetActor());
-					if (Actor)
-					{
-						HLODActor->UpdateLODParent(*Actor, !bInLoaded);
-					}
-				}
-
-				PendingHLODAssignment.Remove(ActorGuid);
-			}
-		}
-
-		return;
-	}
-
-	const FWorldPartitionActorDesc* HLODActorDesc = GetHLODActorForActor(&InActor);
-	if (HLODActorDesc)
-	{
-		AWorldPartitionHLOD* HLODActor = Cast<AWorldPartitionHLOD>(HLODActorDesc->GetActor());
-		if (HLODActor)
-		{
-			HLODActor->UpdateLODParent(InActor, !bInLoaded);
-		}
-		else if (bInLoaded)
-		{
-			PendingHLODAssignment.Add(HLODActorDesc->GetGuid(), ActorGuid);
-		}
-		else
-		{
-			PendingHLODAssignment.Remove(HLODActorDesc->GetGuid(), ActorGuid);
-		}
-	}
-}
-
-const FWorldPartitionActorDesc* UHLODSubsystem::GetHLODActorForActor(const AActor* InActor) const
-{
-	UHLODLayer* HLODLayer = UHLODLayer::GetHLODLayer(InActor);
-	if (!HLODLayer)
-	{
-		return nullptr;
-	}
-
-	UWorldPartitionSubsystem* WorldPartitionSubsystem = GetWorld()->GetSubsystem<UWorldPartitionSubsystem>();
-	UWorldPartition* WorldPartition = GetWorld()->GetWorldPartition();
-
-	const FWorldPartitionActorDesc* ActorDesc = WorldPartition->GetActorDesc(InActor->GetActorGuid());
-	check(ActorDesc);
-
-	if (ActorDesc->GetGridPlacement() == EActorGridPlacement::AlwaysLoaded)
-	{
-		return nullptr;
-	}
-
-	FVector ActorLocation = ActorDesc->GetOrigin();
-	FBox ActorBox(ActorLocation, ActorLocation);
-
-	// Find all HLODActors at that location
-	TArray<const FWorldPartitionActorDesc*> HLODActorsDescs = WorldPartitionSubsystem->GetIntersectingActorDescs(ActorBox, AWorldPartitionHLOD::StaticClass());
-
-	// Only keep the HLODActors matching our HLODLayer
-	HLODActorsDescs.RemoveAll([HLODLayer, ActorDesc](const FWorldPartitionActorDesc* InActorDesc)
-	{
-		const FHLODActorDesc* HLODActorDesc = (FHLODActorDesc*)InActorDesc;
-		return ActorDesc == HLODActorDesc || HLODActorDesc->GetHLODLayer() != HLODLayer;
-	});
-	
-	if (HLODActorsDescs.IsEmpty())
-	{
-		return nullptr;
-	}
-
-	// Sort by bounds size
-	HLODActorsDescs.Sort([](const FWorldPartitionActorDesc& A, const FWorldPartitionActorDesc& B)
-	{
-		return A.GetBounds().GetExtent().Size() < B.GetBounds().GetExtent().Size();
-	});
-
-	const UWorldPartition::FActorCluster* ActorCluster = WorldPartition->GetClusterForActor(ActorDesc->GetGuid());
-	switch (ActorCluster->GridPlacement)
-	{
-	case EActorGridPlacement::Location:
-		return HLODActorsDescs[0];
-	
-	case EActorGridPlacement::Bounds:
-		for (const FWorldPartitionActorDesc* HLODActorDesc : HLODActorsDescs)
-		{
-			if (HLODActorDesc->GetBounds().IsInsideXY(ActorCluster->Bounds))
-			{
-				return HLODActorDesc;
-			}
-		}
-		break;
-	
-	default:
-		check(0);
-	}
-
-	return nullptr;
-}
-
 void UHLODSubsystem::RegisterActorDescFactories(UWorldPartitionSubsystem* WorldPartitionSubsystem)
 {
 	WorldPartitionSubsystem->RegisterActorDescFactory(AWorldPartitionHLOD::StaticClass(), HLODActorDescFactory.Get());
diff --git a/Engine/Source/Runtime/Engine/Public/WorldPartition/HLOD/HLODActor.h b/Engine/Source/Runtime/Engine/Public/WorldPartition/HLOD/HLODActor.h
--- a/Engine/Source/Runtime/Engine/Public/WorldPartition/HLOD/HLODActor.h
+++ b/Engine/Source/Runtime/Engine/Public/WorldPartition/HLOD/HLODActor.h
@@ -30,12 +30,6 @@ protected:
 	virtual void EndPlay(const EEndPlayReason::Type EndPlayReason) override;
 	//~ End AActor Interface.
 
-#if WITH_EDITOR
-	void SetLODParent(AActor& InActor);
-	void ClearLODParent(AActor& InActor);
-	void UpdateLODParent(AActor& InActor, bool bInClear);
-#endif
-
 	UPrimitiveComponent* GetHLODComponent();
 
 public:
diff --git a/Engine/Source/Runtime/Engine/Public/WorldPartition/HLOD/HLODSubsystem.h b/Engine/Source/Runtime/Engine/Public/WorldPartition/HLOD/HLODSubsystem.h
--- a/Engine/Source/Runtime/Engine/Public/WorldPartition/HLOD/HLODSubsystem.h
+++ b/Engine/Source/Runtime/Engine/Public/WorldPartition/HLOD/HLODSubsystem.h
@@ -46,8 +46,6 @@ public:
 private:
 #if WITH_EDITOR
 	void RegisterActorDescFactories(UWorldPartitionSubsystem* WorldPartitionSubsystem);
-	void OnWorldPartitionActorRegistered(AActor& InActor, bool bInLoaded);
-	const FWorldPartitionActorDesc* GetHLODActorForActor(const AActor* InActor) const;
 #endif
 
 private:
