Animation Runtime: Deprecated FAnimationRuntime::BlendTransformsByWeight

diff --git a/Engine/Source/Runtime/Engine/Public/AnimationRuntime.h b/Engine/Source/Runtime/Engine/Public/AnimationRuntime.h
--- a/Engine/Source/Runtime/Engine/Public/AnimationRuntime.h
+++ b/Engine/Source/Runtime/Engine/Public/AnimationRuntime.h
@@ -361,6 +361,7 @@ public:
 	 *
 	 * @param OutTransform : result
 	 */
+	UE_DEPRECATED(4.26, "Please use BlendTransform() for weighted blending")
 	static void BlendTransformsByWeight(FTransform& OutTransform, const TArray<FTransform>& Transforms, const TArray<float>& Weights);
 
 	/** 
