#Fixed root-motion extraction to properly handle the case where the root is translated and rotated from the origin on frame 0

diff --git a/Engine/Source/Runtime/Engine/Private/Animation/AnimSequence.cpp b/Engine/Source/Runtime/Engine/Private/Animation/AnimSequence.cpp
--- a/Engine/Source/Runtime/Engine/Private/Animation/AnimSequence.cpp
+++ b/Engine/Source/Runtime/Engine/Private/Animation/AnimSequence.cpp
@@ -1408,9 +1408,8 @@ FTransform UAnimSequence::ExtractRootMotionFromRange(float StartTrackPosition, f
 	}
 
 	// Transform to Component Space Rotation (inverse root transform from first frame)
-	const FTransform RootToComponentRot = FTransform(InitialTransform.GetRotation().Inverse());
-	StartTransform = RootToComponentRot * StartTransform;
-	EndTransform = RootToComponentRot * EndTransform;
+	StartTransform = StartTransform.GetRelativeTransform(InitialTransform);
+    EndTransform = EndTransform.GetRelativeTransform(InitialTransform);
 
 	return EndTransform.GetRelativeTransform(StartTransform);
 }
