WorldPartition:
- Added support for New Landscape with a Grid Size setting to split new landscape into StreamingProxies
- Support Add Component tool so that it creates new Grid proxies when needed
ActorPartitionSubsystem:
- Support for Partition Guid to allow multiple actors of the same type (different grids)

diff --git a/Engine/Source/Editor/LandscapeEditor/Private/LandscapeEdMode.cpp b/Engine/Source/Editor/LandscapeEditor/Private/LandscapeEdMode.cpp
--- a/Engine/Source/Editor/LandscapeEditor/Private/LandscapeEdMode.cpp
+++ b/Engine/Source/Editor/LandscapeEditor/Private/LandscapeEdMode.cpp
@@ -14,6 +14,7 @@
 #include "LandscapeEditorObject.h"
 #include "Landscape.h"
 #include "LandscapeStreamingProxy.h"
+#include "LandscapeSubsystem.h"
 
 #include "EditorSupportDelegates.h"
 #include "ScopedTransaction.h"
@@ -2303,6 +2304,11 @@ int32 FEdModeLandscape::UpdateLandscapeList()
 	return CurrentIndex;
 }
 
+bool FEdModeLandscape::IsGridBased() const
+{
+	return GetWorld()->GetSubsystem<ULandscapeSubsystem>()->IsGridBased();
+}
+
 bool FEdModeLandscape::HasValidLandscapeEditLayerSelection() const
 {
 	return !CanHaveLandscapeLayersContent() || GetCurrentLayer();
diff --git a/Engine/Source/Editor/LandscapeEditor/Private/LandscapeEdMode.h b/Engine/Source/Editor/LandscapeEditor/Private/LandscapeEdMode.h
--- a/Engine/Source/Editor/LandscapeEditor/Private/LandscapeEdMode.h
+++ b/Engine/Source/Editor/LandscapeEditor/Private/LandscapeEdMode.h
@@ -545,6 +545,8 @@ public:
 	void OnLandscapeMaterialChangedDelegate();
 	void RefreshDetailPanel();
 
+	bool IsGridBased() const;
+
 	// Edit Layers
 	bool HasValidLandscapeEditLayerSelection() const;
 	bool CanHaveLandscapeLayersContent() const;
diff --git a/Engine/Source/Editor/LandscapeEditor/Private/LandscapeEdModeComponentTools.cpp b/Engine/Source/Editor/LandscapeEditor/Private/LandscapeEdModeComponentTools.cpp
--- a/Engine/Source/Editor/LandscapeEditor/Private/LandscapeEdModeComponentTools.cpp
+++ b/Engine/Source/Editor/LandscapeEditor/Private/LandscapeEdModeComponentTools.cpp
@@ -22,6 +22,7 @@
 #include "PhysicalMaterials/PhysicalMaterial.h"
 #include "Materials/MaterialExpressionLandscapeVisibilityMask.h"
 #include "Algo/Copy.h"
+#include "LandscapeSubsystem.h"
 
 #define LOCTEXT_NAMESPACE "Landscape"
 
@@ -526,8 +527,7 @@ public:
 
 	virtual void Apply(FEditorViewportClient* ViewportClient, FLandscapeBrush* Brush, const ULandscapeEditorObject* UISettings, const TArray<FLandscapeToolInteractorPosition>& InteractorPositions)
 	{
-		ALandscapeProxy* LandscapeProxy = LandscapeInfo ? LandscapeInfo->GetCurrentLevelLandscapeProxy(true) : nullptr;
-		if (LandscapeProxy && EdMode->LandscapeRenderAddCollision)
+		if (LandscapeInfo && EdMode->LandscapeRenderAddCollision)
 		{
 			check(Brush->GetBrushType() == ELandscapeBrushType::Component);
 
@@ -544,7 +544,7 @@ public:
 
 			// Find component range for this block of data, non shared vertices
 			int32 ComponentIndexX1, ComponentIndexY1, ComponentIndexX2, ComponentIndexY2;
-			ALandscape::CalcComponentIndicesNoOverlap(X1, Y1, X2, Y2, LandscapeProxy->ComponentSizeQuads, ComponentIndexX1, ComponentIndexY1, ComponentIndexX2, ComponentIndexY2);
+			ALandscape::CalcComponentIndicesNoOverlap(X1, Y1, X2, Y2, LandscapeInfo->ComponentSizeQuads, ComponentIndexX1, ComponentIndexY1, ComponentIndexX2, ComponentIndexY2);
 
 			// expand the area by one vertex in each direction to ensure normals are calculated correctly
 			X1 -= 1;
@@ -559,9 +559,11 @@ public:
 			HeightCache.GetCachedData(X1, Y1, X2, Y2, Data);
 			bool bHasXYOffset = XYOffsetCache.GetCachedData(X1, Y1, X2, Y2, XYOffsetData);
 
+			UWorld* World = ViewportClient->GetScene()->GetWorld();
+			
 			TArray<ULandscapeComponent*> NewComponents;
-			LandscapeProxy->Modify();
 			LandscapeInfo->Modify();
+			ULandscapeSubsystem* LandscapeSubsystem = World->GetSubsystem<ULandscapeSubsystem>();
 			for (int32 ComponentIndexY = ComponentIndexY1; ComponentIndexY <= ComponentIndexY2; ComponentIndexY++)
 			{
 				for (int32 ComponentIndexX = ComponentIndexX1; ComponentIndexX <= ComponentIndexX2; ComponentIndexX++)
@@ -570,34 +572,39 @@ public:
 					if (!LandscapeComponent)
 					{
 						// Add New component...
-						FIntPoint ComponentBase = FIntPoint(ComponentIndexX, ComponentIndexY)*LandscapeProxy->ComponentSizeQuads;
-						LandscapeComponent = NewObject<ULandscapeComponent>(LandscapeProxy, NAME_None, RF_Transactional);
-						LandscapeProxy->LandscapeComponents.Add(LandscapeComponent);
-						NewComponents.Add(LandscapeComponent);
-						LandscapeComponent->Init(
-							ComponentBase.X, ComponentBase.Y,
-							LandscapeProxy->ComponentSizeQuads,
-							LandscapeProxy->NumSubsections,
-							LandscapeProxy->SubsectionSizeQuads
+						FIntPoint ComponentBase = FIntPoint(ComponentIndexX, ComponentIndexY)* LandscapeInfo->ComponentSizeQuads;
+
+						ALandscapeProxy* LandscapeProxy = LandscapeSubsystem->FindOrAddLandscapeProxy(LandscapeInfo, ComponentBase);
+						if (LandscapeProxy)
+						{
+							LandscapeComponent = NewObject<ULandscapeComponent>(LandscapeProxy, NAME_None, RF_Transactional);
+							LandscapeProxy->LandscapeComponents.Add(LandscapeComponent);
+							NewComponents.Add(LandscapeComponent);
+							LandscapeComponent->Init(
+								ComponentBase.X, ComponentBase.Y,
+								LandscapeProxy->ComponentSizeQuads,
+								LandscapeProxy->NumSubsections,
+								LandscapeProxy->SubsectionSizeQuads
 							);
-						LandscapeComponent->AttachToComponent(LandscapeProxy->GetRootComponent(), FAttachmentTransformRules::KeepRelativeTransform);
+							LandscapeComponent->AttachToComponent(LandscapeProxy->GetRootComponent(), FAttachmentTransformRules::KeepRelativeTransform);
 
-						// Assign shared properties
-						LandscapeComponent->UpdatedSharedPropertiesFromActor();
+							// Assign shared properties
+							LandscapeComponent->UpdatedSharedPropertiesFromActor();
 
-						int32 ComponentVerts = (LandscapeProxy->SubsectionSizeQuads + 1) * LandscapeProxy->NumSubsections;
-						// Update Weightmap Scale Bias
-						LandscapeComponent->WeightmapScaleBias = FVector4(1.0f / (float)ComponentVerts, 1.0f / (float)ComponentVerts, 0.5f / (float)ComponentVerts, 0.5f / (float)ComponentVerts);
-						LandscapeComponent->WeightmapSubsectionOffset = (float)(LandscapeComponent->SubsectionSizeQuads + 1) / (float)ComponentVerts;
+							int32 ComponentVerts = (LandscapeProxy->SubsectionSizeQuads + 1) * LandscapeProxy->NumSubsections;
+							// Update Weightmap Scale Bias
+							LandscapeComponent->WeightmapScaleBias = FVector4(1.0f / (float)ComponentVerts, 1.0f / (float)ComponentVerts, 0.5f / (float)ComponentVerts, 0.5f / (float)ComponentVerts);
+							LandscapeComponent->WeightmapSubsectionOffset = (float)(LandscapeComponent->SubsectionSizeQuads + 1) / (float)ComponentVerts;
 
-						TArray<FColor> HeightData;
-						HeightData.Empty(FMath::Square(ComponentVerts));
-						HeightData.AddZeroed(FMath::Square(ComponentVerts));
-						LandscapeComponent->InitHeightmapData(HeightData, true);
-						LandscapeComponent->UpdateMaterialInstances();
+							TArray<FColor> HeightData;
+							HeightData.Empty(FMath::Square(ComponentVerts));
+							HeightData.AddZeroed(FMath::Square(ComponentVerts));
+							LandscapeComponent->InitHeightmapData(HeightData, true);
+							LandscapeComponent->UpdateMaterialInstances();
 
-						LandscapeInfo->XYtoComponentMap.Add(FIntPoint(ComponentIndexX, ComponentIndexY), LandscapeComponent);
-						LandscapeInfo->XYtoAddCollisionMap.Remove(FIntPoint(ComponentIndexX, ComponentIndexY));
+							LandscapeInfo->XYtoComponentMap.Add(FIntPoint(ComponentIndexX, ComponentIndexY), LandscapeComponent);
+							LandscapeInfo->XYtoAddCollisionMap.Remove(FIntPoint(ComponentIndexX, ComponentIndexY));
+						}
 					}
 				}
 			}
@@ -661,7 +668,8 @@ public:
 				TMap<ULandscapeLayerInfoObject*, int32> NeighbourLayerInfoObjectCount;
 
 				{
-					FScopedSetLandscapeEditingLayer Scope(Landscape, Landscape ? Landscape->GetLayer(0)->Guid : FGuid(), [=] { });
+					FLandscapeLayer* LandscapeLayer = Landscape ? Landscape->GetLayer(0) : nullptr;
+					FScopedSetLandscapeEditingLayer Scope(Landscape, LandscapeLayer ? LandscapeLayer->Guid : FGuid(), [=] { });
 
 					// Cover 9 tiles around us to determine which object should we use by default
 					for (int32 ComponentIndexX = ComponentIndexX1 - 1; ComponentIndexX <= ComponentIndexX2 + 1; ++ComponentIndexX)
diff --git a/Engine/Source/Editor/LandscapeEditor/Private/LandscapeEditorDetailCustomization_NewLandscape.cpp b/Engine/Source/Editor/LandscapeEditor/Private/LandscapeEditorDetailCustomization_NewLandscape.cpp
--- a/Engine/Source/Editor/LandscapeEditor/Private/LandscapeEditorDetailCustomization_NewLandscape.cpp
+++ b/Engine/Source/Editor/LandscapeEditor/Private/LandscapeEditorDetailCustomization_NewLandscape.cpp
@@ -42,6 +42,7 @@
 #include "LandscapeDataAccess.h"
 #include "Settings/EditorExperimentalSettings.h"
 #include "Editor.h"
+#include "LandscapeSubsystem.h"
 
 #define LOCTEXT_NAMESPACE "LandscapeEditor.NewLandscape"
 
@@ -91,7 +92,17 @@ void FLandscapeEditorDetailCustomization_NewLandscape::CustomizeDetails(IDetailL
 
 	TSharedRef<IPropertyHandle> PropertyHandle_CanHaveLayersContent = DetailBuilder.GetProperty(GET_MEMBER_NAME_CHECKED(ULandscapeEditorObject, bCanHaveLayersContent));
 	NewLandscapeCategory.AddProperty(PropertyHandle_CanHaveLayersContent);
+	TSharedRef<IPropertyHandle> PropertyHandle_GridSize = DetailBuilder.GetProperty(GET_MEMBER_NAME_CHECKED(ULandscapeEditorObject, WorldPartitionGridSize));
 	
+	FEdModeLandscape* LandscapeEdMode = GetEditorMode();
+	if (LandscapeEdMode->GetWorld()->GetSubsystem<ULandscapeSubsystem>()->IsGridBased())
+	{
+		NewLandscapeCategory.AddProperty(PropertyHandle_GridSize);
+	}
+	else
+	{
+		DetailBuilder.HideProperty(PropertyHandle_GridSize);
+	}
 	TSharedRef<IPropertyHandle> PropertyHandle_HeightmapFilename = DetailBuilder.GetProperty(GET_MEMBER_NAME_CHECKED(ULandscapeEditorObject, ImportLandscape_HeightmapFilename));
 	TSharedRef<IPropertyHandle> PropertyHandle_HeightmapImportResult = DetailBuilder.GetProperty(GET_MEMBER_NAME_CHECKED(ULandscapeEditorObject, ImportLandscape_HeightmapImportResult));
 	TSharedRef<IPropertyHandle> PropertyHandle_HeightmapErrorMessage = DetailBuilder.GetProperty(GET_MEMBER_NAME_CHECKED(ULandscapeEditorObject, ImportLandscape_HeightmapErrorMessage));
@@ -879,6 +890,8 @@ FReply FLandscapeEditorDetailCustomization_NewLandscape::OnCreateButtonClicked()
 		LandscapeEdMode->SetCurrentTool("Sculpt"); // change to sculpting mode and tool
 		LandscapeEdMode->SetCurrentLayer(0);
 
+		LandscapeEdMode->GetWorld()->GetSubsystem<ULandscapeSubsystem>()->UpdateGrid(LandscapeInfo, UISettings->WorldPartitionGridSize);
+
 		if (LandscapeEdMode->CurrentToolTarget.LandscapeInfo.IsValid())
 		{
 			ALandscapeProxy* LandscapeProxy = LandscapeEdMode->CurrentToolTarget.LandscapeInfo->GetLandscapeProxy();
diff --git a/Engine/Source/Editor/LandscapeEditor/Private/SLandscapeEditor.cpp b/Engine/Source/Editor/LandscapeEditor/Private/SLandscapeEditor.cpp
--- a/Engine/Source/Editor/LandscapeEditor/Private/SLandscapeEditor.cpp
+++ b/Engine/Source/Editor/LandscapeEditor/Private/SLandscapeEditor.cpp
@@ -198,9 +198,13 @@ void FLandscapeToolKit::BuildToolPalette(FName PaletteName, class FToolBarBuilde
 		ToolBarBuilder.AddToolBarButton(Commands.SelectComponentTool);
 		ToolBarBuilder.AddToolBarButton(Commands.AddComponentTool);
 		ToolBarBuilder.AddToolBarButton(Commands.DeleteComponentTool);
-		ToolBarBuilder.AddToolBarButton(Commands.MoveToLevelTool);
-		ToolBarBuilder.AddToolBarButton(Commands.ResizeLandscape);
-
+		// MoveToLevel isn't supported because in GridBased worlds don't support Proxies in different Levels
+		// Resize isn't supported and instead should be done through a Commandlet for GridBased worlds
+		if (!LandscapeEdMode->IsGridBased())
+		{
+			ToolBarBuilder.AddToolBarButton(Commands.MoveToLevelTool);
+			ToolBarBuilder.AddToolBarButton(Commands.ResizeLandscape);
+		}
 		ToolBarBuilder.AddToolBarButton(Commands.SplineTool);
 	}
 
diff --git a/Engine/Source/Editor/LandscapeEditor/Public/LandscapeEditorObject.h b/Engine/Source/Editor/LandscapeEditor/Public/LandscapeEditorObject.h
--- a/Engine/Source/Editor/LandscapeEditor/Public/LandscapeEditorObject.h
+++ b/Engine/Source/Editor/LandscapeEditor/Public/LandscapeEditorObject.h
@@ -529,6 +529,9 @@ public:
 	UPROPERTY(Category = "New Landscape", EditAnywhere, NonTransactional, meta = (DisplayName= "Enable Edit Layers", ToolTip="Enable support for landscape edit layers.", ShowForTools= "NewLandscape"))
 	bool bCanHaveLayersContent = false;
 
+	UPROPERTY(Category = "New Landscape", EditAnywhere, NonTransactional, meta = (DisplayName= "World Partition Grid Size", ToolTip="Number of components per landscape streaming proxies per axis", ShowForTools="NewLandscape", ClampMin=1, ClampMax=16, UIMin=1, UIMax=16))
+	uint32 WorldPartitionGridSize = 2;
+
 	// Whether the imported alpha maps are to be interpreted as "layered" or "additive" (UE4 uses additive internally)
 	UPROPERTY(Category="New Landscape", EditAnywhere, NonTransactional, meta=(DisplayName="Layer Alphamap Type", ShowForTools="NewLandscape"))
 	ELandscapeImportAlphamapType ImportLandscape_AlphamapType;
diff --git a/Engine/Source/Runtime/Core/Public/UObject/UE5MainStreamObjectVersion.h b/Engine/Source/Runtime/Core/Public/UObject/UE5MainStreamObjectVersion.h
--- a/Engine/Source/Runtime/Core/Public/UObject/UE5MainStreamObjectVersion.h
+++ b/Engine/Source/Runtime/Core/Public/UObject/UE5MainStreamObjectVersion.h
@@ -25,6 +25,9 @@ struct CORE_API FUE5MainStreamObjectVersion
 		// This is the correct versioning for MeshDescription changes which were added to ReleaseObjectVersion.
 		MeshDescriptionNewFormat,
 
+		// Serialize GridGuid in PartitionActorDesc
+		PartitionActorDescSerializeGridGuid,
+
 		// -----<new versions can be added above this line>-------------------------------------------------
 		VersionPlusOne,
 		LatestVersion = VersionPlusOne - 1
diff --git a/Engine/Source/Runtime/Engine/Private/ActorPartition/ActorPartitionSubsystem.cpp b/Engine/Source/Runtime/Engine/Private/ActorPartition/ActorPartitionSubsystem.cpp
--- a/Engine/Source/Runtime/Engine/Private/ActorPartition/ActorPartitionSubsystem.cpp
+++ b/Engine/Source/Runtime/Engine/Private/ActorPartition/ActorPartitionSubsystem.cpp
@@ -10,19 +10,21 @@ DEFINE_LOG_CATEGORY_STATIC(LogActorPartitionSubsystem, All, All);
 
 #if WITH_EDITOR
 
-FActorPartitionGetParams::FActorPartitionGetParams(const TSubclassOf<APartitionActor>& InActorClass, bool bInCreate, ULevel* InLevelHint, const FVector& InLocationHint)
+FActorPartitionGetParams::FActorPartitionGetParams(const TSubclassOf<APartitionActor>& InActorClass, bool bInCreate, ULevel* InLevelHint, const FVector& InLocationHint, uint32 InGridSize, const FGuid& InGuidHint)
 	: ActorClass(InActorClass)
 	, bCreate(bInCreate)
 	, LocationHint(InLocationHint)
 	, LevelHint(InLevelHint)
+	, GuidHint(InGuidHint)
+	, GridSize(InGridSize)
 {
 }
 
-void FActorPartitionGridHelper::ForEachIntersectingCell(const TSubclassOf<APartitionActor>& InActorClass, const FBox& InBounds, ULevel* InLevel, TFunctionRef<bool(const UActorPartitionSubsystem::FCellCoord&, const FBox&)> InOperation)
+void FActorPartitionGridHelper::ForEachIntersectingCell(const TSubclassOf<APartitionActor>& InActorClass, const FBox& InBounds, ULevel* InLevel, TFunctionRef<bool(const UActorPartitionSubsystem::FCellCoord&, const FBox&)> InOperation, uint32 InGridSize)
 {
-	uint32 GridSize = InActorClass->GetDefaultObject<APartitionActor>()->GetDefaultGridSize(InLevel->GetWorld());
-	UActorPartitionSubsystem::FCellCoord MinCellCoords = UActorPartitionSubsystem::FCellCoord::GetCellCoord(InBounds.Min, InLevel, GridSize);
-	UActorPartitionSubsystem::FCellCoord MaxCellCoords = UActorPartitionSubsystem::FCellCoord::GetCellCoord(InBounds.Max, InLevel, GridSize);
+	const uint32 GridSize = InGridSize > 0 ? InGridSize : InActorClass->GetDefaultObject<APartitionActor>()->GetDefaultGridSize(InLevel->GetWorld());
+	const UActorPartitionSubsystem::FCellCoord MinCellCoords = UActorPartitionSubsystem::FCellCoord::GetCellCoord(InBounds.Min, InLevel, GridSize);
+	const UActorPartitionSubsystem::FCellCoord MaxCellCoords = UActorPartitionSubsystem::FCellCoord::GetCellCoord(InBounds.Max, InLevel, GridSize);
 
 	for (int32 z = MinCellCoords.Z; z <= MaxCellCoords.Z; z++)
 	{
@@ -48,6 +50,29 @@ void FActorPartitionGridHelper::ForEachIntersectingCell(const TSubclassOf<AParti
 	}
 }
 
+void FActorPartitionGridHelper::ForEachIntersectingCell(const TSubclassOf<APartitionActor>& InActorClass, const FIntRect& InRect, ULevel* InLevel, TFunctionRef<bool(const UActorPartitionSubsystem::FCellCoord&, const FIntRect&)> InOperation, uint32 InGridSize)
+{
+	const uint32 GridSize = InGridSize > 0 ? InGridSize : InActorClass->GetDefaultObject<APartitionActor>()->GetDefaultGridSize(InLevel->GetWorld());
+	const UActorPartitionSubsystem::FCellCoord MinCellCoords = UActorPartitionSubsystem::FCellCoord::GetCellCoord(InRect.Min, InLevel, GridSize);
+	const UActorPartitionSubsystem::FCellCoord MaxCellCoords = UActorPartitionSubsystem::FCellCoord::GetCellCoord(InRect.Max, InLevel, GridSize);
+
+	for (int32 y = MinCellCoords.Y; y <= MaxCellCoords.Y; y++)
+	{
+		for (int32 x = MinCellCoords.X; x <= MaxCellCoords.X; x++)
+		{
+			UActorPartitionSubsystem::FCellCoord CellCoords(x, y, 0, InLevel);
+			const FIntPoint Min = FIntPoint(CellCoords.X * GridSize, CellCoords.Y * GridSize);
+			const FIntPoint Max = Min + FIntPoint(GridSize);
+			FIntRect CellBounds(Min, Max);
+
+			if (!InOperation(MoveTemp(CellCoords), MoveTemp(CellBounds)))
+			{
+				return;
+			}
+		}
+	}
+}
+
 /**
  * FActorPartitionLevel
  */
@@ -71,7 +96,7 @@ public:
 		return UActorPartitionSubsystem::FCellCoord(0, 0, 0, SpawnLevel);
 	}
 	
-	APartitionActor* GetActor(const TSubclassOf<APartitionActor>& InActorClass, bool bInCreate, const UActorPartitionSubsystem::FCellCoord& InCellCoord) override
+	APartitionActor* GetActor(const TSubclassOf<APartitionActor>& InActorClass, bool bInCreate, const UActorPartitionSubsystem::FCellCoord& InCellCoord, const FGuid& InGuid, uint32 InGridSize, bool bInBoundsSearch, TFunctionRef<void(APartitionActor*)> InActorCreated) override
 	{
 		check(InCellCoord.Level);
 		
@@ -80,8 +105,11 @@ public:
 		{
 			if (APartitionActor* PartitionActor = Cast<APartitionActor>(Actor))
 			{
-				FoundActor = PartitionActor;
-				break;
+				if (PartitionActor->GetGridGuid() == InGuid)
+				{
+					FoundActor = PartitionActor;
+					break;
+				}
 			}
 		}
 		
@@ -90,6 +118,7 @@ public:
 			FActorSpawnParameters SpawnParams;
 			SpawnParams.OverrideLevel = InCellCoord.Level;
 			FoundActor = CastChecked<APartitionActor>(World->SpawnActor(InActorClass, nullptr, nullptr, SpawnParams));
+			InActorCreated(FoundActor);
 		}
 
 		check(FoundActor || !bInCreate);
@@ -137,42 +166,56 @@ public:
 
 	UActorPartitionSubsystem::FCellCoord GetActorPartitionHash(const FActorPartitionGetParams& GetParams) const override 
 	{
-		uint32 GridSize = GetParams.ActorClass->GetDefaultObject<APartitionActor>()->GetDefaultGridSize(World);
+		uint32 GridSize = GetParams.GridSize > 0 ? GetParams.GridSize : GetParams.ActorClass->GetDefaultObject<APartitionActor>()->GetDefaultGridSize(World);
 
 		return UActorPartitionSubsystem::FCellCoord::GetCellCoord(GetParams.LocationHint, World->PersistentLevel, GridSize);
 	}
 
-	virtual APartitionActor* GetActor(const TSubclassOf<APartitionActor>& InActorClass, bool bInCreate, const UActorPartitionSubsystem::FCellCoord& InCellCoord)
+	virtual APartitionActor* GetActor(const TSubclassOf<APartitionActor>& InActorClass, bool bInCreate, const UActorPartitionSubsystem::FCellCoord& InCellCoord, const FGuid& InGuid, uint32 InGridSize, bool bInBoundsSearch, TFunctionRef<void(APartitionActor*)> InActorCreated)
 	{
 		APartitionActor* FoundActor = nullptr;
-		const uint32 GridSize = InActorClass->GetDefaultObject<APartitionActor>()->GetDefaultGridSize(World);
-		FBox CellBounds(UActorPartitionSubsystem::FCellCoord::GetCellBounds(InCellCoord, GridSize));
-
-		TArray<const FWorldPartitionActorDesc*> InstancedObjectsActorDescs = WorldPartition->GetIntersectingActorDescs(CellBounds, InActorClass);
-
-		for (const FWorldPartitionActorDesc* ActorDesc: InstancedObjectsActorDescs)
+		bool bUnloadedActorExists = false;
+		auto FindActor = [&FoundActor, &bUnloadedActorExists, InActorClass, InCellCoord, InGuid, InGridSize](const FWorldPartitionActorDesc* ActorDesc)
 		{
-			if (ActorDesc->GetActorClass() == InActorClass)
+			check(ActorDesc->GetActorClass()->IsChildOf(InActorClass));
+			FPartitionActorDesc* PartitionActorDesc = (FPartitionActorDesc*)ActorDesc;
+			if ((PartitionActorDesc->GridIndexX == InCellCoord.X) &&
+				(PartitionActorDesc->GridIndexY == InCellCoord.Y) &&
+				(PartitionActorDesc->GridIndexZ == InCellCoord.Z) &&
+				(PartitionActorDesc->GridSize == InGridSize) &&
+				(PartitionActorDesc->GridGuid == InGuid))
 			{
-				FPartitionActorDesc* PartitionActorDesc = (FPartitionActorDesc*)ActorDesc;
-				if ((PartitionActorDesc->GridIndexX == InCellCoord.X) &&
-					(PartitionActorDesc->GridIndexY == InCellCoord.Y) &&
-					(PartitionActorDesc->GridIndexZ == InCellCoord.Z))
-				{
-					AActor* DescActor = ActorDesc->GetActor();
+				AActor* DescActor = ActorDesc->GetActor();
 
-					if (!DescActor)
-					{
-						// Actor exists but is not loaded
-						return nullptr;
-					}
-				
-					FoundActor = CastChecked<APartitionActor>(DescActor);
-					check(FoundActor->GridSize == GridSize);
+				if (!DescActor)
+				{
+					// Actor exists but is not loaded
+					bUnloadedActorExists = true;
+					return false;
 				}
+
+				FoundActor = CastChecked<APartitionActor>(DescActor);
+				check(FoundActor->GridSize == InGridSize && FoundActor->GetGridGuid() == InGuid);
+				return false;
 			}
-		}
+			return true;
+		};
 
+		FBox CellBounds = UActorPartitionSubsystem::FCellCoord::GetCellBounds(InCellCoord, InGridSize);
+		if (bInBoundsSearch)
+		{
+			WorldPartition->ForEachIntersectingActorDesc(CellBounds, InActorClass, FindActor);
+		}
+		else
+		{
+			WorldPartition->ForEachActorDesc(InActorClass, FindActor);
+		}
+				
+		if (bUnloadedActorExists)
+		{
+			return nullptr;
+		}
+				
 		if (!FoundActor && bInCreate)
 		{
 			FActorSpawnParameters SpawnParams;
@@ -182,9 +225,12 @@ public:
 							
 			FVector CellCenter(CellBounds.GetCenter());
 			FoundActor = CastChecked<APartitionActor>(World->SpawnActor(InActorClass, &CellCenter, nullptr, SpawnParams));
-			FoundActor->GridSize = GridSize;
+			FoundActor->GridSize = InGridSize;
 			FoundActor->bLockLocation = true;
+			
+			InActorCreated(FoundActor);
 
+			FoundActor->SetActorLabel(FoundActor->GetName());
 			WorldPartition->UpdateActorDesc(FoundActor);
 		}
 
@@ -251,48 +297,58 @@ void UActorPartitionSubsystem::InitializeActorPartition()
 APartitionActor* UActorPartitionSubsystem::GetActor(const FActorPartitionGetParams& GetParams)
 {
 	FCellCoord CellCoord = ActorPartition->GetActorPartitionHash(GetParams);
-
-	return GetActor(GetParams.ActorClass, CellCoord, GetParams.bCreate);
+	return GetActor(GetParams.ActorClass, CellCoord, GetParams.bCreate, GetParams.GuidHint, GetParams.GridSize);
 }
 
-APartitionActor* UActorPartitionSubsystem::GetActor(const TSubclassOf<APartitionActor>& InActorClass, const FCellCoord& InCellCoords, bool bInCreate)
+APartitionActor* UActorPartitionSubsystem::GetActor(const TSubclassOf<APartitionActor>& InActorClass, const FCellCoord& InCellCoords, bool bInCreate, const FGuid& InGuid, uint32 InGridSize, bool bInBoundsSearch, TFunctionRef<void(APartitionActor*)> InActorCreated)
 {
-	TMap<UClass*, TWeakObjectPtr<APartitionActor>>* ActorsPerClass = PartitionedActors.Find(InCellCoords);
+	const uint32 GridSize = InGridSize > 0 ? InGridSize : InActorClass->GetDefaultObject<APartitionActor>()->GetDefaultGridSize(GetWorld());
+	
+	TMap<UClass*, TMap<FGuid, TWeakObjectPtr<APartitionActor>>>* ActorsPerClass = PartitionedActors.Find(InCellCoords);
 	APartitionActor* FoundActor = nullptr;
 	if (!ActorsPerClass)
 	{
-		FoundActor = ActorPartition->GetActor(InActorClass, bInCreate, InCellCoords);
+		FoundActor = ActorPartition->GetActor(InActorClass, bInCreate, InCellCoords, InGuid, GridSize, bInBoundsSearch, InActorCreated);
 		if (FoundActor)
 		{
-			PartitionedActors.Add(InCellCoords).Add(InActorClass, FoundActor);
-			return FoundActor;
-		}
-		else if (bInCreate)
-		{
-			// Actor wasn't found and couldn't be created
-			return nullptr;
+			PartitionedActors.Add(InCellCoords).Add(InActorClass).Add(InGuid, FoundActor);
 		}
 	}
-
-	TWeakObjectPtr<APartitionActor>* ActorPtr = ActorsPerClass->Find(InActorClass);
-	if (!ActorPtr || !ActorPtr->IsValid())
+	else
 	{
-		FoundActor = ActorPartition->GetActor(InActorClass, bInCreate, InCellCoords);
-		if (FoundActor)
+		TMap<FGuid, TWeakObjectPtr<APartitionActor>>* ActorsPerGuid = ActorsPerClass->Find(InActorClass);
+		if (!ActorsPerGuid)
+		{
+			FoundActor = ActorPartition->GetActor(InActorClass, bInCreate, InCellCoords, InGuid, GridSize, bInBoundsSearch, InActorCreated);
+			if (FoundActor)
+			{
+				ActorsPerClass->Add(InActorClass).Add(InGuid, FoundActor);
+			}
+		}
+		else
 		{
-			if (!ActorPtr)
+			TWeakObjectPtr<APartitionActor>* ActorPtr = ActorsPerGuid->Find(InGuid);
+			if (!ActorPtr || !ActorPtr->IsValid())
 			{
-				ActorsPerClass->Add(InActorClass, FoundActor);
+				FoundActor = ActorPartition->GetActor(InActorClass, bInCreate, InCellCoords, InGuid, GridSize, bInBoundsSearch, InActorCreated);
+				if (FoundActor)
+				{
+					if (!ActorPtr)
+					{
+						ActorsPerGuid->Add(InGuid, FoundActor);
+					}
+					else
+					{
+						*ActorPtr = FoundActor;
+					}
+				}
 			}
 			else
 			{
-				*ActorPtr = FoundActor;
+				FoundActor = ActorPtr->Get();
 			}
 		}
-	}
-	else
-	{
-		FoundActor = (*ActorPtr).Get();
+
 	}
 	
 	return FoundActor;
diff --git a/Engine/Source/Runtime/Engine/Private/WorldPartition/ActorPartition/PartitionActorDesc.cpp b/Engine/Source/Runtime/Engine/Private/WorldPartition/ActorPartition/PartitionActorDesc.cpp
--- a/Engine/Source/Runtime/Engine/Private/WorldPartition/ActorPartition/PartitionActorDesc.cpp
+++ b/Engine/Source/Runtime/Engine/Private/WorldPartition/ActorPartition/PartitionActorDesc.cpp
@@ -5,6 +5,7 @@
 
 #if WITH_EDITOR
 #include "Engine/Level.h"
+#include "UObject/UE5MainStreamObjectVersion.h"
 
 void FPartitionActorDesc::InitFrom(const AActor* InActor)
 {
@@ -13,7 +14,8 @@ void FPartitionActorDesc::InitFrom(const AActor* InActor)
 	const APartitionActor* PartitionActor = CastChecked<APartitionActor>(InActor);
 
 	GridSize = PartitionActor->GridSize;
-
+	GridGuid = PartitionActor->GetGridGuid();
+	
 	const FVector ActorLocation = InActor->GetActorLocation();
 	GridIndexX = FMath::FloorToInt(ActorLocation.X / GridSize);
 	GridIndexY = FMath::FloorToInt(ActorLocation.Y / GridSize);
@@ -22,8 +24,15 @@ void FPartitionActorDesc::InitFrom(const AActor* InActor)
 
 void FPartitionActorDesc::Serialize(FArchive& Ar)
 {
+	Ar.UsingCustomVersion(FUE5MainStreamObjectVersion::GUID);
+
 	FWorldPartitionActorDesc::Serialize(Ar);
 
 	Ar << GridSize << GridIndexX << GridIndexY << GridIndexZ;
+
+	if (Ar.CustomVer(FUE5MainStreamObjectVersion::GUID) >= FUE5MainStreamObjectVersion::PartitionActorDescSerializeGridGuid)
+	{
+		Ar << GridGuid;
+	}
 }
 #endif
diff --git a/Engine/Source/Runtime/Engine/Private/WorldPartition/Landscape/LandscapeActorDesc.cpp b/Engine/Source/Runtime/Engine/Private/WorldPartition/Landscape/LandscapeActorDesc.cpp
new file mode 100644
--- /dev/null
+++ b/Engine/Source/Runtime/Engine/Private/WorldPartition/Landscape/LandscapeActorDesc.cpp
@@ -0,0 +1,19 @@
+// Copyright Epic Games, Inc. All Rights Reserved.
+
+#include "WorldPartition/Landscape/LandscapeActorDesc.h"
+#include "LandscapeProxy.h"
+
+#if WITH_EDITOR
+
+void FLandscapeActorDesc::InitFrom(const AActor* InActor)
+{
+	FPartitionActorDesc::InitFrom(InActor);
+
+	const ALandscapeProxy* LandscapeProxy = CastChecked<ALandscapeProxy>(InActor);
+	check(LandscapeProxy);
+	GridIndexX = LandscapeProxy->LandscapeSectionOffset.X / LandscapeProxy->GridSize;
+	GridIndexY = LandscapeProxy->LandscapeSectionOffset.Y / LandscapeProxy->GridSize;
+	GridIndexZ = 0;
+}
+
+#endif
diff --git a/Engine/Source/Runtime/Engine/Private/WorldPartition/Landscape/LandscapeActorDescFactory.cpp b/Engine/Source/Runtime/Engine/Private/WorldPartition/Landscape/LandscapeActorDescFactory.cpp
new file mode 100644
--- /dev/null
+++ b/Engine/Source/Runtime/Engine/Private/WorldPartition/Landscape/LandscapeActorDescFactory.cpp
@@ -0,0 +1,12 @@
+// Copyright Epic Games, Inc. All Rights Reserved.
+
+#include "WorldPartition/Landscape/LandscapeActorDescFactory.h"
+
+#if WITH_EDITOR
+#include "WorldPartition/Landscape/LandscapeActorDesc.h"
+
+FWorldPartitionActorDesc* FLandscapeActorDescFactory::Create()
+{
+	return new FLandscapeActorDesc;
+}
+#endif
\ No newline at end of file
diff --git a/Engine/Source/Runtime/Engine/Private/WorldPartition/WorldPartition.cpp b/Engine/Source/Runtime/Engine/Private/WorldPartition/WorldPartition.cpp
--- a/Engine/Source/Runtime/Engine/Private/WorldPartition/WorldPartition.cpp
+++ b/Engine/Source/Runtime/Engine/Private/WorldPartition/WorldPartition.cpp
@@ -525,19 +525,30 @@ FWorldPartitionActorDescFactory* UWorldPartition::GetActorDescFactory(const AAct
 	return GetActorDescFactory(Actor->GetClass());
 }
 
-TArray<const FWorldPartitionActorDesc*> UWorldPartition::GetIntersectingActorDescs(const FBox& Box, TSubclassOf<AActor> ActorClass) const
+void UWorldPartition::ForEachIntersectingActorDesc(const FBox& Box, TSubclassOf<AActor> ActorClass, TFunctionRef<bool(const FWorldPartitionActorDesc*)> Predicate) const
 {
-	TArray<const FWorldPartitionActorDesc*> ActorDescs;
-
-	EditorHash->ForEachIntersectingActor(Box, [&ActorDescs, &ActorClass](const FWorldPartitionActorDesc* ActorDesc)
+	EditorHash->ForEachIntersectingActor(Box, [&ActorClass, Predicate](const FWorldPartitionActorDesc* ActorDesc)
 	{
 		if (ActorDesc->GetActorClass()->IsChildOf(ActorClass))
 		{
-			ActorDescs.Add(ActorDesc);
+			Predicate(ActorDesc);
 		}
 	});
-	
-	return MoveTemp(ActorDescs);
+}
+
+void UWorldPartition::ForEachActorDesc(TSubclassOf<AActor> ActorClass, TFunctionRef<bool(const FWorldPartitionActorDesc*)> Predicate) const
+{
+	for (const TPair<FGuid, TUniquePtr<FWorldPartitionActorDesc>>& Pair : Actors)
+	{
+		FWorldPartitionActorDesc* ActorDesc = Pair.Value.Get();
+		if (ActorDesc->GetActorClass()->IsChildOf(ActorClass))
+		{
+			if (!Predicate(ActorDesc))
+			{
+				return;
+			}
+		}
+	}
 }
 
 const FWorldPartitionActorDesc* UWorldPartition::GetActorDesc(const FGuid& Guid) const
diff --git a/Engine/Source/Runtime/Engine/Private/WorldPartition/WorldPartitionSubsystem.cpp b/Engine/Source/Runtime/Engine/Private/WorldPartition/WorldPartitionSubsystem.cpp
--- a/Engine/Source/Runtime/Engine/Private/WorldPartition/WorldPartitionSubsystem.cpp
+++ b/Engine/Source/Runtime/Engine/Private/WorldPartition/WorldPartitionSubsystem.cpp
@@ -191,14 +191,20 @@ void UWorldPartitionSubsystem::DrawRuntimeHash2D(UCanvas* Canvas, class APlayerC
 //
 
 #if WITH_EDITOR
-TArray<const FWorldPartitionActorDesc*> UWorldPartitionSubsystem::GetIntersectingActorDescs(const FBox& Box, TSubclassOf<AActor> ActorClass) const
+void UWorldPartitionSubsystem::ForEachIntersectingActorDesc(const FBox& Box, TSubclassOf<AActor> ActorClass, TFunctionRef<bool(const FWorldPartitionActorDesc*)> Predicate) const
 {
-	TArray<const FWorldPartitionActorDesc*> ActorDescs;
 	if (const UWorldPartition* MainPartition = GetMainWorldPartition())
 	{
-		ActorDescs = MainPartition->GetIntersectingActorDescs(Box, ActorClass);
+		MainPartition->ForEachIntersectingActorDesc(Box, ActorClass, Predicate);
+	}
+}
+
+void UWorldPartitionSubsystem::ForEachActorDesc(TSubclassOf<AActor> ActorClass, TFunctionRef<bool(const FWorldPartitionActorDesc*)> Predicate) const
+{
+	if (const UWorldPartition* MainPartition = GetMainWorldPartition())
+	{
+		MainPartition->ForEachActorDesc(ActorClass, Predicate);
 	}
-	return MoveTemp(ActorDescs);
 }
 
 void UWorldPartitionSubsystem::UpdateActorDesc(AActor* Actor)
diff --git a/Engine/Source/Runtime/Engine/Public/ActorPartition/ActorPartitionSubsystem.h b/Engine/Source/Runtime/Engine/Public/ActorPartition/ActorPartitionSubsystem.h
--- a/Engine/Source/Runtime/Engine/Public/ActorPartition/ActorPartitionSubsystem.h
+++ b/Engine/Source/Runtime/Engine/Public/ActorPartition/ActorPartitionSubsystem.h
@@ -5,6 +5,7 @@
 #include "Subsystems/WorldSubsystem.h"
 #include "Templates/SubclassOf.h"
 #include "Misc/HashBuilder.h"
+#include "Misc/Guid.h"
 #include "WorldPartition/ActorPartition/PartitionActorDescFactory.h"
 #include "ActorPartitionSubsystem.generated.h"
 
@@ -17,7 +18,7 @@ class APartitionActor;
  */
 struct ENGINE_API FActorPartitionGetParams
 {
-	FActorPartitionGetParams(const TSubclassOf<APartitionActor>& InActorClass, bool bInCreate, ULevel* InLevelHint, const FVector& InLocationHint);
+	FActorPartitionGetParams(const TSubclassOf<APartitionActor>& InActorClass, bool bInCreate, ULevel* InLevelHint, const FVector& InLocationHint, uint32 InGridSize = 0, const FGuid& InGuidHint = FGuid());
 
 	/* Class of Actor we are getting from the subsystem. */
 	TSubclassOf<APartitionActor> ActorClass;
@@ -30,6 +31,12 @@ struct ENGINE_API FActorPartitionGetParams
 	
 	/* Depending on the world LevelHint can be used to find/create the Actor. */
 	ULevel* LevelHint;
+
+	/* Guid can be used to distinguish between actors of the same type. */
+	FGuid GuidHint;
+
+	/* If greater than 0, use this instead of the Actor CDO Grid size*/
+	int32 GridSize;
 };
 
 #endif
@@ -85,6 +92,15 @@ public:
 			);
 		}
 
+		static FCellCoord GetCellCoord(FIntPoint InPos, ULevel* InLevel, uint32 InGridSize)
+		{
+			const int64 GridSize = (int64)InGridSize;
+			const int64 CellCoordX = InPos.X < 0 ? (InPos.X - (GridSize - 1)) / GridSize : InPos.X / GridSize;
+			const int64 CellCoordY = InPos.Y < 0 ? (InPos.Y - (GridSize - 1)) / GridSize : InPos.Y / GridSize;
+
+			return FCellCoord(CellCoordX, CellCoordY, 0, InLevel);
+		}
+		
 		static FBox GetCellBounds(const FCellCoord& InCellCoord, uint32 InGridSize)
 		{
 			return FBox(
@@ -104,7 +120,18 @@ public:
 
 #if WITH_EDITOR
 	APartitionActor* GetActor(const FActorPartitionGetParams& GetParam);
-	APartitionActor* GetActor(const TSubclassOf<APartitionActor>& InActorClass, const FCellCoord& InCellCoords, bool bInCreate);
+
+	/**
+	 * Returns a matching actor based on the parameters being provided
+	 * 
+	 * @param InActorClass The type of actor we are searching for.
+	 * @param InCellCoords The cell coordinate of that actor.
+	 * @param bInCreate If the actor doesn't existe should we created it.
+	 * @param InGuid Optional, if multiple actors of the same InActorClass exist user can provide Guid id.
+	 * @param InGridSize Optional, if not provided the InActorClass CDO will be used to determine GridSize.
+	 * @param bInBoundsSearch Optional, if not specified existing actors will be searched in the cell bounds only.
+	 */
+	APartitionActor* GetActor(const TSubclassOf<APartitionActor>& InActorClass, const FCellCoord& InCellCoords, bool bInCreate, const FGuid& InGuid = FGuid(), uint32 InGridSize = 0, bool bInBoundsSearch = true, TFunctionRef<void(APartitionActor*)> InActorCreated = [](APartitionActor*) {});
 
 	virtual void Initialize(FSubsystemCollectionBase& Collection) override;
 	virtual void Deinitialize() override;
@@ -118,7 +145,7 @@ private:
 
 	void InitializeActorPartition();
 
-	TMap<FCellCoord, TMap<UClass*, TWeakObjectPtr<APartitionActor>>> PartitionedActors;
+	TMap<FCellCoord, TMap<UClass*, TMap<FGuid, TWeakObjectPtr<APartitionActor>>>> PartitionedActors;
 	TUniquePtr<FBaseActorPartition> ActorPartition;
 	
 	FPartitionActorDescFactory PartitionActorDescFactory;
@@ -138,7 +165,7 @@ public:
 	virtual ~FBaseActorPartition() {}
 
 	virtual UActorPartitionSubsystem::FCellCoord GetActorPartitionHash(const FActorPartitionGetParams& GetParams) const = 0;
-	virtual APartitionActor* GetActor(const TSubclassOf<APartitionActor>& InActorClass, bool bInCreate, const UActorPartitionSubsystem::FCellCoord& InCellCoord) = 0;
+	virtual APartitionActor* GetActor(const TSubclassOf<APartitionActor>& InActorClass, bool bInCreate, const UActorPartitionSubsystem::FCellCoord& InCellCoord, const FGuid& InGuid, uint32 InGridSize, bool bInBoundsSearch, TFunctionRef<void(APartitionActor*)> InActorCreated) = 0;
 
 	DECLARE_EVENT_OneParam(FBaseActorPartition, FOnActorPartitionHashInvalidated, const UActorPartitionSubsystem::FCellCoord&);
 	FOnActorPartitionHashInvalidated& GetOnActorPartitionHashInvalidated() { return OnActorPartitionHashInvalidated; }
@@ -154,6 +181,7 @@ protected:
 class ENGINE_API FActorPartitionGridHelper
 {
 public:
-	static void ForEachIntersectingCell(const TSubclassOf<APartitionActor>& InActorClass, const FBox& InBounds, ULevel* InLevel, TFunctionRef<bool(const UActorPartitionSubsystem::FCellCoord&, const FBox&)> InOperation);
+	static void ForEachIntersectingCell(const TSubclassOf<APartitionActor>& InActorClass, const FBox& InBounds, ULevel* InLevel, TFunctionRef<bool(const UActorPartitionSubsystem::FCellCoord&, const FBox&)> InOperation, uint32 InGridSize = 0);
+	static void ForEachIntersectingCell(const TSubclassOf<APartitionActor>& InActorClass, const FIntRect& InBounds, ULevel* InLevel, TFunctionRef<bool(const UActorPartitionSubsystem::FCellCoord&, const FIntRect&)> InOperation, uint32 InGridSize = 0);
 };
 #endif
\ No newline at end of file
diff --git a/Engine/Source/Runtime/Engine/Public/ActorPartition/PartitionActor.h b/Engine/Source/Runtime/Engine/Public/ActorPartition/PartitionActor.h
--- a/Engine/Source/Runtime/Engine/Public/ActorPartition/PartitionActor.h
+++ b/Engine/Source/Runtime/Engine/Public/ActorPartition/PartitionActor.h
@@ -3,6 +3,7 @@
 
 #include "CoreMinimal.h"
 #include "GameFramework/Actor.h"
+#include "Misc/Guid.h"
 #include "PartitionActor.generated.h"
 
 class UBoxComponent;
@@ -21,6 +22,7 @@ public:
 #if WITH_EDITOR	
 	virtual EActorGridPlacement GetDefaultGridPlacement() const override { return EActorGridPlacement::Location; }
 	virtual uint32 GetDefaultGridSize(UWorld* InWorld) const PURE_VIRTUAL(APartitionActor, return 0;)
+	virtual FGuid GetGridGuid() const { return FGuid(); }
 #endif
 	//~ End AActor Interface	
 
diff --git a/Engine/Source/Runtime/Engine/Public/WorldPartition/ActorPartition/PartitionActorDesc.h b/Engine/Source/Runtime/Engine/Public/WorldPartition/ActorPartition/PartitionActorDesc.h
--- a/Engine/Source/Runtime/Engine/Public/WorldPartition/ActorPartition/PartitionActorDesc.h
+++ b/Engine/Source/Runtime/Engine/Public/WorldPartition/ActorPartition/PartitionActorDesc.h
@@ -17,7 +17,7 @@ public:
 	int64 GridIndexX;
 	int64 GridIndexY;
 	int64 GridIndexZ;
-
+	FGuid GridGuid;
 protected:
 	virtual void InitFrom(const AActor* InActor) override;
 	virtual void Serialize(FArchive& Ar) override;
diff --git a/Engine/Source/Runtime/Engine/Public/WorldPartition/Landscape/LandscapeActorDesc.h b/Engine/Source/Runtime/Engine/Public/WorldPartition/Landscape/LandscapeActorDesc.h
new file mode 100644
--- /dev/null
+++ b/Engine/Source/Runtime/Engine/Public/WorldPartition/Landscape/LandscapeActorDesc.h
@@ -0,0 +1,21 @@
+// Copyright Epic Games, Inc. All Rights Reserved.
+
+#pragma once
+
+#include "CoreMinimal.h"
+#include "UObject/ObjectMacros.h"
+#include "Containers/Map.h"
+#include "WorldPartition/ActorPartition/PartitionActorDesc.h"
+
+/**
+ * ActorDesc for LandscapeActors.
+ */
+class ENGINE_API FLandscapeActorDesc : public FPartitionActorDesc
+{
+#if WITH_EDITOR
+	friend class FLandscapeActorDescFactory;
+
+protected:
+	virtual void InitFrom(const AActor* InActor) override;
+#endif
+};
diff --git a/Engine/Source/Runtime/Engine/Public/WorldPartition/Landscape/LandscapeActorDescFactory.h b/Engine/Source/Runtime/Engine/Public/WorldPartition/Landscape/LandscapeActorDescFactory.h
new file mode 100644
--- /dev/null
+++ b/Engine/Source/Runtime/Engine/Public/WorldPartition/Landscape/LandscapeActorDescFactory.h
@@ -0,0 +1,13 @@
+// Copyright Epic Games, Inc. All Rights Reserved.
+#pragma once
+
+#include "CoreMinimal.h"
+#include "WorldPartition/WorldPartitionActorDescFactory.h"
+
+class ENGINE_API FLandscapeActorDescFactory : public FWorldPartitionActorDescFactory
+{
+public:
+#if WITH_EDITOR
+	virtual FWorldPartitionActorDesc* Create() override;
+#endif
+};
\ No newline at end of file
diff --git a/Engine/Source/Runtime/Engine/Public/WorldPartition/WorldPartition.h b/Engine/Source/Runtime/Engine/Public/WorldPartition/WorldPartition.h
--- a/Engine/Source/Runtime/Engine/Public/WorldPartition/WorldPartition.h
+++ b/Engine/Source/Runtime/Engine/Public/WorldPartition/WorldPartition.h
@@ -93,7 +93,8 @@ private:
 	friend class UWorldPartitionSubsystem;
 
 	// UActorPartitionSubsystem interface+
-	TArray<const FWorldPartitionActorDesc*> GetIntersectingActorDescs(const FBox& Box, TSubclassOf<AActor> ActorClass) const;
+	void ForEachIntersectingActorDesc(const FBox& Box, TSubclassOf<AActor> ActorClass, TFunctionRef<bool(const FWorldPartitionActorDesc*)> Predicate) const;
+	void ForEachActorDesc(TSubclassOf<AActor> ActorClass, TFunctionRef<bool(const FWorldPartitionActorDesc*)> Predicate) const;
 	// UActorPartitionSubsystem interface-
 
 	void ApplyActorTransform(AActor* InActor, const FTransform& InTransform);
diff --git a/Engine/Source/Runtime/Engine/Public/WorldPartition/WorldPartitionSubsystem.h b/Engine/Source/Runtime/Engine/Public/WorldPartition/WorldPartitionSubsystem.h
--- a/Engine/Source/Runtime/Engine/Public/WorldPartition/WorldPartitionSubsystem.h
+++ b/Engine/Source/Runtime/Engine/Public/WorldPartition/WorldPartitionSubsystem.h
@@ -43,8 +43,8 @@ public:
 	//~End FTickableGameObject
 
 #if WITH_EDITOR
-	TArray<const FWorldPartitionActorDesc*> GetIntersectingActorDescs(const FBox& Box, TSubclassOf<AActor> ActorClass) const;
-
+	void ForEachIntersectingActorDesc(const FBox& Box, TSubclassOf<AActor> ActorClass, TFunctionRef<bool(const FWorldPartitionActorDesc*)> Predicate) const;
+	void ForEachActorDesc(TSubclassOf<AActor> ActorClass, TFunctionRef<bool(const FWorldPartitionActorDesc*)> Predicate) const;
 	void UpdateActorDesc(AActor* Actor);
 
 	void RegisterActorDescFactory(TSubclassOf<AActor> Class, FWorldPartitionActorDescFactory* Factory);
diff --git a/Engine/Source/Runtime/Landscape/Classes/Landscape.h b/Engine/Source/Runtime/Landscape/Classes/Landscape.h
--- a/Engine/Source/Runtime/Landscape/Classes/Landscape.h
+++ b/Engine/Source/Runtime/Landscape/Classes/Landscape.h
@@ -252,6 +252,8 @@ public:
 	virtual bool ShouldImport(FString* ActorPropString, bool IsMovingLevel) override;
 	virtual void PostEditImport() override;
 	virtual void PostDuplicate(bool bDuplicateForPIE) override;
+
+	virtual EActorGridPlacement GetDefaultGridPlacement() const override { return EActorGridPlacement::AlwaysLoaded; }
 #endif
 	virtual void PostLoad() override;
 	virtual void BeginDestroy() override;
diff --git a/Engine/Source/Runtime/Landscape/Classes/LandscapeInfo.h b/Engine/Source/Runtime/Landscape/Classes/LandscapeInfo.h
--- a/Engine/Source/Runtime/Landscape/Classes/LandscapeInfo.h
+++ b/Engine/Source/Runtime/Landscape/Classes/LandscapeInfo.h
@@ -241,6 +241,9 @@ public:
 	/** Moves Components to target level. Creates ALandscapeProxy if needed. */
 	LANDSCAPE_API ALandscapeProxy* MoveComponentsToLevel(const TArray<ULandscapeComponent*>& InComponents, ULevel* TargetLevel, FName NewProxyName = NAME_None);
 
+	/** Moves Components to target proxy. */
+	LANDSCAPE_API ALandscapeProxy* MoveComponentsToProxy(const TArray<ULandscapeComponent*>& InComponents, ALandscapeProxy* LandscapeProxy, bool bSetPositionAndOffset = false, ULevel* TargetLevel = nullptr);
+
 	/** Moves Splines connected to this control point to target level. Creates ULandscapeSplineComponent if needed. */
 	LANDSCAPE_API void MoveSplineToLevel(ULandscapeSplineControlPoint* InControlPoint, ULevel* TargetLevel);
 
@@ -249,6 +252,9 @@ public:
 
 	/** Will call UpdateAllComponentMaterialInstances on all LandscapeProxies */
 	LANDSCAPE_API void UpdateAllComponentMaterialInstances();
+
+	/** Returns LandscapeStreamingProxy Cell Size in WorldPartition */
+	LANDSCAPE_API uint32 GetGridSize(uint32 InGridSizeInComponents) const;
 #endif
 
 	/**
diff --git a/Engine/Source/Runtime/Landscape/Classes/LandscapeProxy.h b/Engine/Source/Runtime/Landscape/Classes/LandscapeProxy.h
--- a/Engine/Source/Runtime/Landscape/Classes/LandscapeProxy.h
+++ b/Engine/Source/Runtime/Landscape/Classes/LandscapeProxy.h
@@ -15,6 +15,7 @@
 #include "LandscapeComponent.h"
 #include "LandscapeWeightmapUsage.h"
 #include "VT/RuntimeVirtualTextureEnum.h"
+#include "ActorPartition/PartitionActor.h"
 
 #include "LandscapeProxy.generated.h"
 
@@ -401,7 +402,7 @@ private:
 };
 
 UCLASS(Abstract, MinimalAPI, NotBlueprintable, NotPlaceable, hidecategories=(Display, Attachment, Physics, Debug, Lighting, LOD), showcategories=(Lighting, Rendering, "Utilities|Transformation"), hidecategories=(Mobility))
-class ALandscapeProxy : public AActor
+class ALandscapeProxy : public APartitionActor
 {
 	GENERATED_BODY()
 
@@ -856,8 +857,8 @@ public:
 	virtual bool ShouldImport(FString* ActorPropString, bool IsMovingLevel) override;
 	virtual bool ShouldExport() override;
 	//~ End AActor Interface
-
-	virtual EActorGridPlacement GetDefaultGridPlacement() const override { return EActorGridPlacement::AlwaysLoaded; }
+	virtual uint32 GetDefaultGridSize(UWorld* InWorld) const override { return 1; }
+	virtual FGuid GetGridGuid() const override { return LandscapeGuid; }
 #endif	//WITH_EDITOR
 
 	FGuid GetLandscapeGuid() const { return LandscapeGuid; }
diff --git a/Engine/Source/Runtime/Landscape/Classes/LandscapeStreamingProxy.h b/Engine/Source/Runtime/Landscape/Classes/LandscapeStreamingProxy.h
--- a/Engine/Source/Runtime/Landscape/Classes/LandscapeStreamingProxy.h
+++ b/Engine/Source/Runtime/Landscape/Classes/LandscapeStreamingProxy.h
@@ -25,6 +25,7 @@ public:
 #if WITH_EDITOR
 	virtual bool ShouldExport() override { return false;  }
 	virtual void PostEditChangeProperty(FPropertyChangedEvent& PropertyChangedEvent) override;
+	virtual bool CanEditChange(const FProperty* InProperty) const override;
 #endif
 	//~ End UObject Interface
 
diff --git a/Engine/Source/Runtime/Landscape/Private/Landscape.cpp b/Engine/Source/Runtime/Landscape/Private/Landscape.cpp
--- a/Engine/Source/Runtime/Landscape/Private/Landscape.cpp
+++ b/Engine/Source/Runtime/Landscape/Private/Landscape.cpp
@@ -1055,8 +1055,6 @@ ALandscapeProxy::ALandscapeProxy(const FObjectInitializer& ObjectInitializer)
 	bCastShadowAsTwoSided = false;
 	bAffectDistanceFieldLighting = true;
 
-	USceneComponent* SceneComponent = CreateDefaultSubobject<USceneComponent>(TEXT("RootComponent0"));
-	RootComponent = SceneComponent;
 	RootComponent->SetRelativeScale3D(FVector(128.0f, 128.0f, 256.0f)); // Old default scale, preserved for compatibility. See ULandscapeEditorObject::NewLandscape_Scale
 	RootComponent->Mobility = EComponentMobility::Static;
 	LandscapeSectionOffset = FIntPoint::ZeroValue;
diff --git a/Engine/Source/Runtime/Landscape/Private/LandscapeEdit.cpp b/Engine/Source/Runtime/Landscape/Private/LandscapeEdit.cpp
--- a/Engine/Source/Runtime/Landscape/Private/LandscapeEdit.cpp
+++ b/Engine/Source/Runtime/Landscape/Private/LandscapeEdit.cpp
@@ -40,6 +40,8 @@ LandscapeEdit.cpp: Landscape editing
 #include "LandscapeSplinesComponent.h"
 #include "Serialization/MemoryWriter.h"
 #if WITH_EDITOR
+#include "Engine/World.h"
+#include "LandscapeSubsystem.h"
 #include "StaticMeshAttributes.h"
 #include "MeshUtilitiesCommon.h"
 
@@ -4097,6 +4099,11 @@ void ULandscapeInfo::UpdateAllComponentMaterialInstances()
 	});
 }
 
+uint32 ULandscapeInfo::GetGridSize(uint32 InGridSizeInComponents) const
+{
+	return InGridSizeInComponents * ComponentSizeQuads;
+}
+
 ALandscapeProxy* ULandscapeInfo::MoveComponentsToLevel(const TArray<ULandscapeComponent*>& InComponents, ULevel* TargetLevel, FName NewProxyName)
 {
 	ALandscape* Landscape = LandscapeActor.Get();
@@ -4104,7 +4111,7 @@ ALandscapeProxy* ULandscapeInfo::MoveComponentsToLevel(const TArray<ULandscapeCo
 
 	// Make sure references are in a different package (should be fixed up before calling this method)
 	// Check the Physical Material is same package with Landscape
-	if(Landscape->DefaultPhysMaterial && Landscape->DefaultPhysMaterial->GetOutermost() == Landscape->GetOutermost())
+	if (Landscape->DefaultPhysMaterial && Landscape->DefaultPhysMaterial->GetOutermost() == Landscape->GetOutermost())
 	{
 		return nullptr;
 	}
@@ -4129,6 +4136,30 @@ ALandscapeProxy* ULandscapeInfo::MoveComponentsToLevel(const TArray<ULandscapeCo
 		}
 	}
 
+	ALandscapeProxy* LandscapeProxy = GetLandscapeProxyForLevel(TargetLevel);
+	bool bSetPositionAndOffset = false;
+	if (!LandscapeProxy)
+	{
+		FActorSpawnParameters SpawnParams;
+		SpawnParams.Name = NewProxyName;
+		SpawnParams.OverrideLevel = TargetLevel;
+		LandscapeProxy = TargetLevel->GetWorld()->SpawnActor<ALandscapeStreamingProxy>(SpawnParams);
+
+		// copy shared properties to this new proxy
+		LandscapeProxy->GetSharedProperties(Landscape);
+		LandscapeProxy->CreateLandscapeInfo();
+		LandscapeProxy->SetActorLabel(LandscapeProxy->GetName());
+		bSetPositionAndOffset = true;
+	}
+
+	return MoveComponentsToProxy(InComponents, LandscapeProxy, bSetPositionAndOffset, TargetLevel);
+}
+
+ALandscapeProxy* ULandscapeInfo::MoveComponentsToProxy(const TArray<ULandscapeComponent*>& InComponents, ALandscapeProxy* LandscapeProxy, bool bSetPositionAndOffset, ULevel* TargetLevel)
+{
+	ALandscape* Landscape = LandscapeActor.Get();
+	check(Landscape != nullptr);
+	
 	struct FCompareULandscapeComponentBySectionBase
 	{
 		FORCEINLINE bool operator()(const ULandscapeComponent& A, const ULandscapeComponent& B) const
@@ -4148,14 +4179,14 @@ ALandscapeProxy* ULandscapeInfo::MoveComponentsToLevel(const TArray<ULandscapeCo
 	for (ULandscapeComponent* Component : ComponentsToMove)
 	{
 		SelectProxies.Add(Component->GetLandscapeProxy());
-		if (Component->GetLandscapeProxy()->GetOuter() != TargetLevel)
+		if (Component->GetLandscapeProxy() != LandscapeProxy && (!TargetLevel || Component->GetLandscapeProxy()->GetOuter() != TargetLevel))
 		{
 			TargetSelectedComponents.Add(Component);
 		}
 
 		ULandscapeHeightfieldCollisionComponent* CollisionComp = Component->CollisionComponent.Get();
 		SelectProxies.Add(CollisionComp->GetLandscapeProxy());
-		if (CollisionComp->GetLandscapeProxy()->GetOuter() != TargetLevel)
+		if (CollisionComp->GetLandscapeProxy() != LandscapeProxy && (!TargetLevel || CollisionComp->GetLandscapeProxy()->GetOuter() != TargetLevel))
 		{
 			TargetSelectedCollisionComponents.Add(CollisionComp);
 		}
@@ -4194,30 +4225,20 @@ ALandscapeProxy* ULandscapeInfo::MoveComponentsToLevel(const TArray<ULandscapeCo
 		}
 	}
 
-	ALandscapeProxy* LandscapeProxy = GetLandscapeProxyForLevel(TargetLevel);
-	if (!LandscapeProxy)
+	// Proxy position/offset needs to be set
+	if(bSetPositionAndOffset)
 	{
-		FActorSpawnParameters SpawnParams;
-		SpawnParams.Name = NewProxyName;
-		SpawnParams.OverrideLevel = TargetLevel;
-		LandscapeProxy = TargetLevel->GetWorld()->SpawnActor<ALandscapeStreamingProxy>(SpawnParams);
-		
-		// copy shared properties to this new proxy
-		LandscapeProxy->GetSharedProperties(Landscape);
-		LandscapeProxy->CreateLandscapeInfo();
-		LandscapeProxy->SetActorLabel(LandscapeProxy->GetName());
-
 		// set proxy location
 		// by default first component location
 		ULandscapeComponent* FirstComponent = *TargetSelectedComponents.CreateConstIterator();
 		LandscapeProxy->GetRootComponent()->SetWorldLocationAndRotation(FirstComponent->GetComponentLocation(), FirstComponent->GetComponentRotation());
 		LandscapeProxy->LandscapeSectionOffset = FirstComponent->GetSectionBase();
+	}
 
-		// Hide(unregister) the new landscape if owning level currently in hidden state
-		if (LandscapeProxy->GetLevel()->bIsVisible == false)
-		{
-			LandscapeProxy->UnregisterAllComponents();
-		}
+	// Hide(unregister) the new landscape if owning level currently in hidden state
+	if (LandscapeProxy->GetLevel()->bIsVisible == false)
+	{
+		LandscapeProxy->UnregisterAllComponents();
 	}
 
 	// Changing Heightmap format for selected components
@@ -4308,7 +4329,7 @@ ALandscapeProxy* ULandscapeInfo::MoveComponentsToLevel(const TArray<ULandscapeCo
 		Component->AttachToComponent(LandscapeProxy->GetRootComponent(), FAttachmentTransformRules::KeepWorldTransform);
 
 		// Move any foliage associated
-		AInstancedFoliageActor::MoveInstancesForComponentToLevel(Component, TargetLevel);
+		AInstancedFoliageActor::MoveInstancesForComponentToLevel(Component, LandscapeProxy->GetLevel());
 	}
 		
 	// Register our new components if destination landscape is registered in scene 
@@ -4778,6 +4799,21 @@ void ALandscapeProxy::PostEditChangeProperty(FPropertyChangedEvent& PropertyChan
 	}
 }
 
+bool ALandscapeStreamingProxy::CanEditChange(const FProperty* InProperty) const
+{
+	if (!Super::CanEditChange(InProperty))
+	{
+		return false;
+	}
+
+	if (InProperty && InProperty->GetFName() == GET_MEMBER_NAME_CHECKED(ALandscapeStreamingProxy, LandscapeActor))
+	{
+		return !GetWorld()->GetSubsystem<ULandscapeSubsystem>()->IsGridBased();
+	}
+
+	return true;
+}
+
 void ALandscapeStreamingProxy::PostEditChangeProperty(FPropertyChangedEvent& PropertyChangedEvent)
 {
 	const FName PropertyName = PropertyChangedEvent.MemberProperty ? PropertyChangedEvent.MemberProperty->GetFName() : NAME_None;
diff --git a/Engine/Source/Runtime/Landscape/Private/LandscapeSubsystem.cpp b/Engine/Source/Runtime/Landscape/Private/LandscapeSubsystem.cpp
--- a/Engine/Source/Runtime/Landscape/Private/LandscapeSubsystem.cpp
+++ b/Engine/Source/Runtime/Landscape/Private/LandscapeSubsystem.cpp
@@ -7,7 +7,13 @@
 #include "ContentStreaming.h"
 #include "Landscape.h"
 #include "LandscapeProxy.h"
+#include "LandscapeStreamingProxy.h"
+#include "LandscapeInfo.h"
 #include "ProfilingDebugging/CsvProfiler.h"
+#include "WorldPartition/WorldPartitionSubsystem.h"
+#include "ActorPartition/ActorPartitionSubsystem.h"
+#include "Engine/World.h"
+#include "Math/IntRect.h"
 
 static int32 GUseStreamingManagerForCameras = 1;
 static FAutoConsoleVariableRef CVarUseStreamingManagerForCameras(
@@ -39,6 +45,14 @@ void ULandscapeSubsystem::Initialize(FSubsystemCollectionBase& Collection)
 {
 	Super::Initialize(Collection);
 
+#if WITH_EDITOR
+	UWorldPartitionSubsystem* WorldPartitionSubsystem = Collection.InitializeDependency<UWorldPartitionSubsystem>();
+	if (WorldPartitionSubsystem)
+	{
+		WorldPartitionSubsystem->RegisterActorDescFactory(ALandscapeProxy::StaticClass(), &LandscapeActorDescFactory);
+	}
+#endif
+
 #if WITH_EDITOR
 	GrassMapsBuilder = new FLandscapeGrassMapsBuilder(GetWorld());
 #endif
@@ -150,4 +164,108 @@ int32 ULandscapeSubsystem::GetOutdatedGrassMapCount()
 {
 	return GrassMapsBuilder->GetOutdatedGrassMapCount(/*bInForceUpdate*/false);
 }
+
+bool ULandscapeSubsystem::IsGridBased() const
+{
+	return UWorld::HasSubsystem<UWorldPartitionSubsystem>(GetWorld());
+}
+
+namespace LandscapeSubsystemUtils
+{
+	ALandscapeProxy* FindOrAddLandscapeProxy(UActorPartitionSubsystem* ActorPartitionSubsystem, ULandscapeInfo* LandscapeInfo, const UActorPartitionSubsystem::FCellCoord& CellCoord)
+	{
+		ALandscape* Landscape = LandscapeInfo->LandscapeActor.Get();
+		check(Landscape);
+							
+		auto LandscapeProxyCreated = [CellCoord, Landscape](APartitionActor* PartitionActor)
+		{
+			const FIntPoint CellLocation(CellCoord.X * Landscape->GridSize, CellCoord.Y * Landscape->GridSize);
+
+			ALandscapeProxy* LandscapeProxy = CastChecked<ALandscapeProxy>(PartitionActor);
+			// copy shared properties to this new proxy
+			LandscapeProxy->GetSharedProperties(Landscape);
+			const FVector ProxyLocation = Landscape->GetActorLocation() + FVector(CellLocation.X * Landscape->GetActorRelativeScale3D().X, CellLocation.Y * Landscape->GetActorRelativeScale3D().Y, 0.0f);
+
+			LandscapeProxy->CreateLandscapeInfo();
+			LandscapeProxy->SetActorLocationAndRotation(ProxyLocation, Landscape->GetActorRotation());
+			LandscapeProxy->LandscapeSectionOffset = FIntPoint(CellLocation.X, CellLocation.Y);
+		};
+
+		const bool bCreate = true;
+		const bool bBoundsSearch = false;
+		return Cast<ALandscapeProxy>(ActorPartitionSubsystem->GetActor(ALandscapeStreamingProxy::StaticClass(), CellCoord, bCreate, LandscapeInfo->LandscapeGuid, Landscape->GridSize, bBoundsSearch, LandscapeProxyCreated));
+	}
+}
+
+void ULandscapeSubsystem::UpdateGrid(ULandscapeInfo* LandscapeInfo, uint32 GridSizeInComponents)
+{
+	if (!IsGridBased())
+	{
+		return;
+	}
+
+	check(LandscapeInfo);
+
+	const uint32 GridSize = LandscapeInfo->GetGridSize(GridSizeInComponents);
+	LandscapeInfo->LandscapeActor->Modify();
+	LandscapeInfo->LandscapeActor->GridSize = GridSize;
+
+	FIntRect Extent;
+	LandscapeInfo->GetLandscapeExtent(Extent.Min.X, Extent.Min.Y, Extent.Max.X, Extent.Max.Y);
+	const FBox Bounds(FVector(Extent.Min), FVector(Extent.Max));
+
+	UActorPartitionSubsystem* ActorPartitionSubsystem = GetWorld()->GetSubsystem<UActorPartitionSubsystem>();
+
+	TArray<ULandscapeComponent*> LandscapeComponents;
+	LandscapeComponents.Reserve(LandscapeInfo->XYtoComponentMap.Num());
+	LandscapeInfo->ForAllLandscapeComponents([&LandscapeComponents](ULandscapeComponent* LandscapeComponent)
+	{
+		LandscapeComponents.Add(LandscapeComponent);
+	});
+
+	FActorPartitionGridHelper::ForEachIntersectingCell(ALandscapeStreamingProxy::StaticClass(), Extent, GetWorld()->PersistentLevel, [ActorPartitionSubsystem, LandscapeInfo, GridSizeInComponents, &LandscapeComponents, GridSize](const UActorPartitionSubsystem::FCellCoord& CellCoord, const FIntRect& CellBounds)
+	{		
+		TArray<ULandscapeComponent*> ComponentsToMove;
+		const int32 MaxComponents = (int32)(GridSizeInComponents * GridSizeInComponents);
+		ComponentsToMove.Reserve(MaxComponents);
+		for (int32 i = 0; i < LandscapeComponents.Num();)
+		{
+			ULandscapeComponent* LandscapeComponent = LandscapeComponents[i];
+			if (CellBounds.Contains(LandscapeComponent->GetSectionBase()))
+			{
+				ComponentsToMove.Add(LandscapeComponent);
+				LandscapeComponents.RemoveAtSwap(i);
+			}
+			else
+			{
+				i++;
+			}
+		}
+		
+		check(ComponentsToMove.Num() <= MaxComponents);
+		if (ComponentsToMove.Num())
+		{
+			ALandscapeProxy* LandscapeProxy = LandscapeSubsystemUtils::FindOrAddLandscapeProxy(ActorPartitionSubsystem, LandscapeInfo, CellCoord);
+			check(LandscapeProxy);
+			LandscapeInfo->MoveComponentsToProxy(ComponentsToMove, LandscapeProxy);
+		}
+
+		return true;
+	}, GridSize);
+}
+
+ALandscapeProxy* ULandscapeSubsystem::FindOrAddLandscapeProxy(ULandscapeInfo* LandscapeInfo, const FIntPoint& SectionBase)
+{
+	if (!IsGridBased())
+	{
+		return LandscapeInfo->GetCurrentLevelLandscapeProxy(true);
+	}
+
+	UActorPartitionSubsystem* ActorPartitionSubsystem = GetWorld()->GetSubsystem<UActorPartitionSubsystem>();
+	const uint32 GridSize = LandscapeInfo->LandscapeActor->GridSize;
+	
+	UActorPartitionSubsystem::FCellCoord CellCoord = UActorPartitionSubsystem::FCellCoord::GetCellCoord(SectionBase, GetWorld()->PersistentLevel, GridSize);
+	return LandscapeSubsystemUtils::FindOrAddLandscapeProxy(ActorPartitionSubsystem, LandscapeInfo, CellCoord);
+}
+
 #endif
\ No newline at end of file
diff --git a/Engine/Source/Runtime/Landscape/Public/LandscapeSubsystem.h b/Engine/Source/Runtime/Landscape/Public/LandscapeSubsystem.h
--- a/Engine/Source/Runtime/Landscape/Public/LandscapeSubsystem.h
+++ b/Engine/Source/Runtime/Landscape/Public/LandscapeSubsystem.h
@@ -6,9 +6,11 @@
 #include "UObject/ObjectMacros.h"
 #include "Subsystems/WorldSubsystem.h"
 #include "Tickable.h"
+#include "WorldPartition/Landscape/LandscapeActorDescFactory.h"
 #include "LandscapeSubsystem.generated.h"
 
 class ALandscapeProxy;
+class ULandscapeInfo;
 
 UCLASS(MinimalAPI)
 class ULandscapeSubsystem : public UWorldSubsystem, public FTickableGameObject
@@ -32,6 +34,10 @@ public:
 #if WITH_EDITOR
 	LANDSCAPE_API void BuildGrassMaps();
 	LANDSCAPE_API int32 GetOutdatedGrassMapCount();
+
+	LANDSCAPE_API bool IsGridBased() const;
+	LANDSCAPE_API void UpdateGrid(ULandscapeInfo* LandscapeInfo, uint32 GridSizeInComponents);
+	LANDSCAPE_API ALandscapeProxy* FindOrAddLandscapeProxy(ULandscapeInfo* LandscapeInfo, const FIntPoint& SectionBase);
 #endif
 
 private:
@@ -44,5 +50,6 @@ private:
 
 #if WITH_EDITOR
 	class FLandscapeGrassMapsBuilder* GrassMapsBuilder;
+	FLandscapeActorDescFactory LandscapeActorDescFactory;
 #endif
 };
