Fix async texture compilation to work with refactored streaming code

diff --git a/Engine/Source/Runtime/Engine/Private/Rendering/Texture2DResource.cpp b/Engine/Source/Runtime/Engine/Private/Rendering/Texture2DResource.cpp
--- a/Engine/Source/Runtime/Engine/Private/Rendering/Texture2DResource.cpp
+++ b/Engine/Source/Runtime/Engine/Private/Rendering/Texture2DResource.cpp
@@ -24,7 +24,7 @@ static TAutoConsoleVariable<int32> CVarForceHighestMipOnUITexturesEnabled(
  */
 FTexture2DResource::FTexture2DResource(UTexture2D* InOwner, const FStreamableRenderResourceState& InState)
 	: FStreamableTextureResource(InOwner, InOwner->PlatformData, InState, true)
-	, ResourceMem( InOwner->ResourceMem )
+	, ResourceMem(InOwner->ResourceMem)
 {
 	// Retrieve initial mip data.
 	MipData.AddZeroed(State.MaxNumLODs);
@@ -33,6 +33,15 @@ FTexture2DResource::FTexture2DResource(UTexture2D* InOwner, const FStreamableRen
 	CacheSamplerStateInitializer(InOwner);
 }
 
+FTexture2DResource::FTexture2DResource(UTexture2D* InOwner, const FTexture2DResource* InProxiedResource)
+	: FStreamableTextureResource(InOwner, InProxiedResource->PlatformData, FStreamableRenderResourceState(), true)
+	, ResourceMem(InOwner->ResourceMem)
+	, ProxiedResource(InProxiedResource)
+{
+	TextureReferenceRHI = InOwner->TextureReference.TextureReferenceRHI;
+	CacheSamplerStateInitializer(InOwner);
+}
+
 /**
  * Destructor, freeing MipData in the case of resource being destroyed without ever 
  * having been initialized by the rendering thread via InitRHI.
@@ -74,6 +83,20 @@ void FTexture2DResource::CacheSamplerStateInitializer(const UTexture2D* InOwner)
 	MipBias = UTexture2D::GetGlobalMipMapLODBias() + DefaultMipBias;
 }
 
+void FTexture2DResource::InitRHI()
+{
+	if (ProxiedResource)
+	{
+		TextureRHI = ProxiedResource->TextureRHI;
+		RHIUpdateTextureReference(TextureReferenceRHI, TextureRHI);
+		SamplerStateRHI = ProxiedResource->SamplerStateRHI;
+		DeferredPassSamplerStateRHI = ProxiedResource->DeferredPassSamplerStateRHI;
+	}
+	else
+	{
+		FStreamableTextureResource::InitRHI();
+	}
+}
 
 void FTexture2DResource::CreateTexture()
 {
diff --git a/Engine/Source/Runtime/Engine/Private/Texture2D.cpp b/Engine/Source/Runtime/Engine/Private/Texture2D.cpp
--- a/Engine/Source/Runtime/Engine/Private/Texture2D.cpp
+++ b/Engine/Source/Runtime/Engine/Private/Texture2D.cpp
@@ -899,17 +899,15 @@ FTextureResource* UTexture2D::CreateResource()
 {
 	TRACE_CPUPROFILER_EVENT_SCOPE(UTexture2D::CreateResource)
 
-//#TODO DC Quick fix to disable Async Texture Compilation until a proper fix is in place.
-//#if WITH_EDITOR
-//	if (PrivatePlatformData && !PrivatePlatformData->IsAsyncWorkComplete())
-//	{
-//		FTextureCompilingManager::Get().AddTextures({this});
-//
-//		UnlinkStreaming();
-//		bIsStreamable = false;
-//		return new FTexture2DResource(this, (const FTexture2DResource*)GetDefaultTexture2D(this)->GetResource());
-//	}
-//#endif
+#if WITH_EDITOR
+	if (PrivatePlatformData && !PrivatePlatformData->IsAsyncWorkComplete())
+	{
+		FTextureCompilingManager::Get().AddTextures({this});
+
+		UnlinkStreaming();
+		return new FTexture2DResource(this, GetDefaultTexture2D(this)->GetResource()->GetTexture2DResource());
+	}
+#endif
 
 	if (IsCurrentlyVirtualTextured())
 	{
diff --git a/Engine/Source/Runtime/Engine/Private/TextureCompiler.cpp b/Engine/Source/Runtime/Engine/Private/TextureCompiler.cpp
--- a/Engine/Source/Runtime/Engine/Private/TextureCompiler.cpp
+++ b/Engine/Source/Runtime/Engine/Private/TextureCompiler.cpp
@@ -275,16 +275,14 @@ FQueuedThreadPool* FTextureCompilingManager::GetThreadPool() const
 
 bool FTextureCompilingManager::IsAsyncTextureCompilationEnabled() const
 {
-	// #TODO DC Quick fix to force disable async texture compilation until a proper fix is ready
-	return false;
-	/*if (bHasShutdown)
+	if (bHasShutdown)
 	{
 		return false;
 	}
 
 	TextureCompilingManagerImpl::EnsureInitializedCVars();
 
-	return CVarAsyncTextureCompilation.GetValueOnAnyThread() != 0;*/
+	return CVarAsyncTextureCompilation.GetValueOnAnyThread() != 0;
 }
 
 void FTextureCompilingManager::UpdateCompilationNotification()
@@ -450,7 +448,7 @@ void FTextureCompilingManager::FinishCompilation(const TArray<UTexture*>& InText
 			FEvent* Event;
 			FTextureTask() { Event = FPlatformProcess::GetSynchEventFromPool(true); }
 			~FTextureTask() { FPlatformProcess::ReturnSynchEventToPool(Event); }
-			void DoThreadedWork() override { Texture->FinishCachePlatformData(); Event->Trigger(); };
+			void DoThreadedWork() override { FOptionalTaskTagScope Scope(ETaskTag::EParallelGameThread); Texture->FinishCachePlatformData(); Event->Trigger(); };
 			void Abandon() override { }
 		};
 
@@ -479,12 +477,6 @@ void FTextureCompilingManager::FinishCompilation(const TArray<UTexture*>& InText
 			// Be nice with the game thread and tick the progress at 60 fps even when no progress is being made...
 			while (!PendingTask.Event->Wait(16))
 			{
-				// If the task pool is saturated, we want to be able to contribute so we can make some progress quickly...
-				if (GThreadPool->RetractQueuedWork(&PendingTask))
-				{
-					PendingTask.DoThreadedWork();
-				}
-
 				UpdateProgress(0.0f, TextureIndex, InTextures.Num(), TextureName);
 			}
 			UpdateProgress(1.f, TextureIndex++, InTextures.Num(), TextureName);
diff --git a/Engine/Source/Runtime/Engine/Public/Rendering/StreamableTextureResource.h b/Engine/Source/Runtime/Engine/Public/Rendering/StreamableTextureResource.h
--- a/Engine/Source/Runtime/Engine/Public/Rendering/StreamableTextureResource.h
+++ b/Engine/Source/Runtime/Engine/Public/Rendering/StreamableTextureResource.h
@@ -29,7 +29,7 @@ public:
 	virtual uint32 GetSizeX() const final override { return SizeX; }
 	virtual uint32 GetSizeY() const final override { return SizeY; }
 	virtual uint32 GetSizeZ() const final override { return SizeZ; }
-	virtual void InitRHI() final override;
+	virtual void InitRHI() override;
 	virtual void ReleaseRHI() final override;
 
 	// This is only coherent sync on the rendering thread. To get the gamethread coherent value, use UStreamableRenderResource.CacheStreamableResourceState
diff --git a/Engine/Source/Runtime/Engine/Public/Rendering/Texture2DResource.h b/Engine/Source/Runtime/Engine/Public/Rendering/Texture2DResource.h
--- a/Engine/Source/Runtime/Engine/Public/Rendering/Texture2DResource.h
+++ b/Engine/Source/Runtime/Engine/Public/Rendering/Texture2DResource.h
@@ -37,7 +37,16 @@ public:
 	/** Set the value of Filter, AddressU, AddressV, AddressW and MipBias from FStreamableTextureResource on the gamethread. */
 	void CacheSamplerStateInitializer(const UTexture2D* InOwner);
 
+	virtual void InitRHI() override;
+	virtual bool IsProxy() const override { return ProxiedResource != nullptr; }
 private:
+	/**
+	 * Make this Texture2DResource Proxy another one.
+	 *
+	 * @param InOwner             UTexture2D which this FTexture2DResource represents.
+	 * @param InProxiedResource   The resource to proxy.
+	 */
+	FTexture2DResource(UTexture2D* InOwner, const FTexture2DResource* InProxiedResource);
 
 	virtual void CreateTexture() final override;
 	virtual void CreatePartiallyResidentTexture() final override;
@@ -52,6 +61,9 @@ private:
 	/** Resource memory allocated by the owner for serialize bulk mip data into								*/
 	FTexture2DResourceMem* ResourceMem;
 
+	/** Another resource being proxied by this one. */
+	const FTexture2DResource* const ProxiedResource = nullptr;
+
 	/** Local copy/ cache of mip data between creation and first call to InitRHI.							*/
 	TArray<void*, TInlineAllocator<MAX_TEXTURE_MIP_COUNT> > MipData;
 
