OFPA - UWorld::RepairWorldSettings() would spawn the WorldSetting as an external actor, which is not supported

diff --git a/Engine/Source/Runtime/Engine/Private/World.cpp b/Engine/Source/Runtime/Engine/Private/World.cpp
--- a/Engine/Source/Runtime/Engine/Private/World.cpp
+++ b/Engine/Source/Runtime/Engine/Private/World.cpp
@@ -1445,6 +1445,7 @@ void UWorld::RepairWorldSettings()
 		FActorSpawnParameters SpawnInfo;
 		SpawnInfo.SpawnCollisionHandlingOverride = ESpawnActorCollisionHandlingMethod::AlwaysSpawn;
 		SpawnInfo.Name = GEngine->WorldSettingsClass->GetFName();
+		SpawnInfo.bCreateActorPackage = false;
 		AWorldSettings* const NewWorldSettings = SpawnActor<AWorldSettings>( GEngine->WorldSettingsClass, SpawnInfo );
 
 		// If there was an existing actor, copy its properties to the new actor (the it will be destroyed by SetWorldSettings)
