Fix crash in shootergame because ManagerIndex/InstanceIndex were set to 0 instead of INDEX_NONE

diff --git a/Engine/Source/Runtime/Engine/Classes/Engine/EngineTypes.h b/Engine/Source/Runtime/Engine/Classes/Engine/EngineTypes.h
--- a/Engine/Source/Runtime/Engine/Classes/Engine/EngineTypes.h
+++ b/Engine/Source/Runtime/Engine/Classes/Engine/EngineTypes.h
@@ -2136,6 +2136,7 @@ struct ENGINE_API FHitResult
 	FORCEINLINE void Init()
 	{
 		FMemory::Memzero(this, sizeof(FHitResult));
+		HitObjectHandle = FActorInstanceHandle();
 		Time = 1.f;
 	}
 
@@ -2143,6 +2144,7 @@ struct ENGINE_API FHitResult
 	FORCEINLINE void Init(FVector Start, FVector End)
 	{
 		FMemory::Memzero(this, sizeof(FHitResult));
+		HitObjectHandle = FActorInstanceHandle();
 		Time = 1.f;
 		TraceStart = Start;
 		TraceEnd = End;
