Fixed a crash when deleting an editing level instance and its actors.

diff --git a/Engine/Source/Runtime/Engine/Private/LevelInstance/LevelInstanceActor.cpp b/Engine/Source/Runtime/Engine/Private/LevelInstance/LevelInstanceActor.cpp
--- a/Engine/Source/Runtime/Engine/Private/LevelInstance/LevelInstanceActor.cpp
+++ b/Engine/Source/Runtime/Engine/Private/LevelInstance/LevelInstanceActor.cpp
@@ -441,15 +441,15 @@ bool ALevelInstance::CanDeleteSelectedActor(FText& OutReason) const
 		return false;
 	}
 
-	if (IsDirty())
+	if (IsEditing())
 	{
-		OutReason = LOCTEXT("HasDirtyLevel", "Can't delete LevelInstance because it is dirty!");
+		OutReason = LOCTEXT("HasEditingLevel", "Can't delete LevelInstance because it is editing!");
 		return false;
 	}
 
-	if (HasDirtyChildren())
+	if (HasEditingChildren())
 	{
-		OutReason = LOCTEXT("HasDirtryChildLevel", "Can't delete LevelInstance because it has dirty child LevelInstances!");
+		OutReason = LOCTEXT("HasEditingChildLevel", "Can't delete LevelInstance because it has editing child LevelInstances!");
 		return false;
 	}
 	return true;
@@ -601,6 +601,16 @@ bool ALevelInstance::IsEditing() const
 	return false;
 }
 
+bool ALevelInstance::HasEditingChildren() const
+{
+	if (ULevelInstanceSubsystem* LevelInstanceSubsystem = GetLevelInstanceSubsystem())
+	{
+		return LevelInstanceSubsystem->HasEditingChildrenLevelInstances(this);
+	}
+
+	return false;
+}
+
 void ALevelInstance::Edit(AActor* ContextActor)
 {
 	ULevelInstanceSubsystem* LevelInstanceSubsystem = GetLevelInstanceSubsystem();
diff --git a/Engine/Source/Runtime/Engine/Private/LevelInstance/LevelInstanceSubsystem.cpp b/Engine/Source/Runtime/Engine/Private/LevelInstance/LevelInstanceSubsystem.cpp
--- a/Engine/Source/Runtime/Engine/Private/LevelInstance/LevelInstanceSubsystem.cpp
+++ b/Engine/Source/Runtime/Engine/Private/LevelInstance/LevelInstanceSubsystem.cpp
@@ -488,6 +488,21 @@ bool ULevelInstanceSubsystem::HasDirtyChildrenLevelInstances(const ALevelInstanc
 	return bDirtyChildren;
 }
 
+bool ULevelInstanceSubsystem::HasEditingChildrenLevelInstances(const ALevelInstance* LevelInstanceActor) const
+{
+	bool bEditingChildren = false;
+	ForEachLevelInstanceChildren(LevelInstanceActor, true, [this, &bEditingChildren](const ALevelInstance* ChildLevelInstanceActor)
+		{
+			if (IsEditingLevelInstance(ChildLevelInstanceActor))
+			{
+				bEditingChildren = true;
+				return false;
+			}
+			return true;
+		});
+	return bEditingChildren;
+}
+
 void ULevelInstanceSubsystem::SetIsTemporarilyHiddenInEditor(ALevelInstance* LevelInstanceActor, bool bIsHidden)
 {
 	if (ULevel* LevelInstanceLevel = GetLevelInstanceLevel(LevelInstanceActor))
diff --git a/Engine/Source/Runtime/Engine/Public/LevelInstance/LevelInstanceActor.h b/Engine/Source/Runtime/Engine/Public/LevelInstance/LevelInstanceActor.h
--- a/Engine/Source/Runtime/Engine/Public/LevelInstance/LevelInstanceActor.h
+++ b/Engine/Source/Runtime/Engine/Public/LevelInstance/LevelInstanceActor.h
@@ -58,6 +58,7 @@ public:
 	bool CanEdit(FText* OutReason = nullptr) const;
 	bool CanCommit(FText* OutReason = nullptr) const;
 	bool IsEditing() const;
+	bool HasEditingChildren() const;
 	void Edit(AActor* ContextActor = nullptr);
 	void Commit();
 	void Discard();
diff --git a/Engine/Source/Runtime/Engine/Public/LevelInstance/LevelInstanceSubsystem.h b/Engine/Source/Runtime/Engine/Public/LevelInstance/LevelInstanceSubsystem.h
--- a/Engine/Source/Runtime/Engine/Public/LevelInstance/LevelInstanceSubsystem.h
+++ b/Engine/Source/Runtime/Engine/Public/LevelInstance/LevelInstanceSubsystem.h
@@ -69,6 +69,7 @@ public:
 	void ForEachLevelInstanceChildren(ALevelInstance* LevelInstanceActor, bool bRecursive, TFunctionRef<bool(ALevelInstance*)> Operation) const;
 	void ForEachLevelInstanceEdit(TFunctionRef<bool(ALevelInstance*)> Operation) const;
 	bool HasDirtyChildrenLevelInstances(const ALevelInstance* LevelInstanceActor) const;
+	bool HasEditingChildrenLevelInstances(const ALevelInstance* LevelInstanceActor) const;
 
 	void SetIsTemporarilyHiddenInEditor(ALevelInstance* LevelInstanceActor, bool bIsHidden);
 
