Merged level logic to determine if an actor can convert to OFPA or not into actor.
Always honor that logic when spawning an actor.

diff --git a/Engine/Source/Editor/DetailCustomizations/Private/ActorDetails.cpp b/Engine/Source/Editor/DetailCustomizations/Private/ActorDetails.cpp
--- a/Engine/Source/Editor/DetailCustomizations/Private/ActorDetails.cpp
+++ b/Engine/Source/Editor/DetailCustomizations/Private/ActorDetails.cpp
@@ -596,7 +596,7 @@ bool FActorDetails::IsActorPackagingModeEditable() const
 	{
 		if (Actor.IsValid() && Actor->GetLevel())
 		{
-			if (!Actor->GetLevel()->CanConvertActorToExternalPackaging(Actor.Get()))
+			if (!Actor->SupportsExternalPackaging())
 			{
 				return false;
 			}
@@ -638,7 +638,7 @@ void FActorDetails::OnActorPackagingModeChanged(bool bExternal)
 				FMessageDialog::Open(EAppMsgType::Ok, Message);
 				return;
 			}
-			else if (Level->CanConvertActorToExternalPackaging(Actor.Get()))
+			else if (Actor->SupportsExternalPackaging())
 			{
 				ActorsToConvert.Add(Actor.Get());
 			}
diff --git a/Engine/Source/Editor/UnrealEd/Private/Commandlets/WorldPartitionConvertCommandlet.cpp b/Engine/Source/Editor/UnrealEd/Private/Commandlets/WorldPartitionConvertCommandlet.cpp
--- a/Engine/Source/Editor/UnrealEd/Private/Commandlets/WorldPartitionConvertCommandlet.cpp
+++ b/Engine/Source/Editor/UnrealEd/Private/Commandlets/WorldPartitionConvertCommandlet.cpp
@@ -1045,7 +1045,7 @@ int32 UWorldPartitionConvertCommandlet::Main(const FString& Params)
 	{
 		for(AActor* Actor: ActorList)
 		{
-			if (!Actor || Actor->IsPendingKill() || !ULevel::CanConvertActorToExternalPackaging(Actor))
+			if (!Actor || Actor->IsPendingKill() || !Actor->SupportsExternalPackaging())
 			{
 				continue;
 			}
diff --git a/Engine/Source/Runtime/Engine/Classes/Engine/Level.h b/Engine/Source/Runtime/Engine/Classes/Engine/Level.h
--- a/Engine/Source/Runtime/Engine/Classes/Engine/Level.h
+++ b/Engine/Source/Runtime/Engine/Classes/Engine/Level.h
@@ -972,8 +972,6 @@ public:
 	/** Sets if the level uses external actors mode or not. */
 	ENGINE_API void SetUseExternalActors(bool bEnable);
 
-	ENGINE_API static bool CanConvertActorToExternalPackaging(AActor* Actor);
-
 	/** 
 	 * Convert this level actors to the specified loading strategy
 	 * @param bExternal if true will convert internal actors to external, will convert external actors to internal otherwise
diff --git a/Engine/Source/Runtime/Engine/Classes/GameFramework/Actor.h b/Engine/Source/Runtime/Engine/Classes/GameFramework/Actor.h
--- a/Engine/Source/Runtime/Engine/Classes/GameFramework/Actor.h
+++ b/Engine/Source/Runtime/Engine/Classes/GameFramework/Actor.h
@@ -1760,7 +1760,7 @@ public:
 	virtual bool CanDeleteSelectedActor(FText& OutReason) const { return true; }
 
 	/** Does this actor supports external packaging? */
-	virtual bool SupportsExternalPackaging() const { return true; }
+	virtual bool SupportsExternalPackaging() const;
 
 	/** Internal struct used to store information about an actor's components during reconstruction */
 	struct FActorRootComponentReconstructionData
diff --git a/Engine/Source/Runtime/Engine/Private/Actor.cpp b/Engine/Source/Runtime/Engine/Private/Actor.cpp
--- a/Engine/Source/Runtime/Engine/Private/Actor.cpp
+++ b/Engine/Source/Runtime/Engine/Private/Actor.cpp
@@ -49,6 +49,7 @@
 #include "LevelUtils.h"
 #include "Components/DecalComponent.h"
 #include "GameFramework/InputSettings.h"
+#include "FoliageHelper.h"
 
 #include "WorldPartition/WorldPartition.h"
 #include "WorldPartition/WorldPartitionActorDesc.h"
@@ -3135,6 +3136,42 @@ bool AActor::IsSelectedInEditor() const
 {
 	return !IsPendingKill() && GSelectedActorAnnotation.Get(this);
 }
+
+bool AActor::SupportsExternalPackaging() const
+{
+	if (HasAllFlags(RF_Transient))
+	{
+		return false;
+	}
+
+	if (IsPendingKill())
+	{
+		return false;
+	}
+
+	if (!IsTemplate())
+	{
+		if (const ULevel* Level = GetLevel())
+		{
+			if (this == Level->GetDefaultBrush())
+			{
+				return false;
+			}
+		}
+
+		if (IsChildActor())
+		{
+			return false;
+		}
+
+		if (FFoliageHelper::IsOwnedByFoliage(this))
+		{
+			return false;
+		}
+	}
+
+	return true;
+}
 #endif
 
 /** Util that sets up the actor's component hierarchy (when users forget to do so, in their native ctor) */
diff --git a/Engine/Source/Runtime/Engine/Private/Level.cpp b/Engine/Source/Runtime/Engine/Private/Level.cpp
--- a/Engine/Source/Runtime/Engine/Private/Level.cpp
+++ b/Engine/Source/Runtime/Engine/Private/Level.cpp
@@ -54,7 +54,6 @@ Level.cpp: Level-related functions
 #include "AssetRegistryModule.h"
 #include "IAssetRegistry.h"
 #include "AssetData.h"
-#include "FoliageHelper.h"
 #endif
 #include "Engine/LevelStreaming.h"
 #include "LevelUtils.h"
@@ -2255,45 +2254,13 @@ void ULevel::SetUseExternalActors(bool bEnable)
 	}
 }
 
-bool ULevel::CanConvertActorToExternalPackaging(AActor* Actor)
-{
-	check(Actor);
-
-	if (Actor->HasAllFlags(RF_Transient))
-	{
-		return false;
-	}
-
-	if (Actor->IsPendingKill())
-	{
-		return false;
-	}
-
-	if (Actor == Actor->GetLevel()->GetDefaultBrush())
-	{
-		return false;
-	}
-
-	if (Actor->IsChildActor())
-	{
-		return false;
-	}
-
-	if (FFoliageHelper::IsOwnedByFoliage(Actor))
-	{
-		return false;
-	}
-
-	return Actor->SupportsExternalPackaging();
-}
-
 void ULevel::ConvertAllActorsToPackaging(bool bExternal)
 {
 	// Make a copy of the current actor lists since packaging conversion may modify the actor list as a side effect
 	TArray<AActor*> CurrentActors = Actors;
 	for (AActor* Actor : CurrentActors)
 	{
-		if (Actor && CanConvertActorToExternalPackaging(Actor))
+		if (Actor && Actor->SupportsExternalPackaging())
 		{
 			check(Actor->GetLevel() == this);
 			Actor->SetPackageExternal(bExternal);
diff --git a/Engine/Source/Runtime/Engine/Private/LevelActor.cpp b/Engine/Source/Runtime/Engine/Private/LevelActor.cpp
--- a/Engine/Source/Runtime/Engine/Private/LevelActor.cpp
+++ b/Engine/Source/Runtime/Engine/Private/LevelActor.cpp
@@ -499,8 +499,11 @@ AActor* UWorld::SpawnActor( UClass* Class, FTransform const* UserTransformPtr, c
 	}
 	else if (LevelToSpawnIn->IsUsingExternalActors() && SpawnParameters.bCreateActorPackage && !(SpawnParameters.ObjectFlags & RF_Transient))
 	{
-		// @todo FH: needs to handle mark package dirty and asset creation notification
-		ExternalPackage = ULevel::CreateActorPackage(LevelToSpawnIn->GetPackage(), ActorGuid);
+		if (CastChecked<AActor>(Class->GetDefaultObject())->SupportsExternalPackaging())
+		{
+			// @todo FH: needs to handle mark package dirty and asset creation notification
+			ExternalPackage = ULevel::CreateActorPackage(LevelToSpawnIn->GetPackage(), ActorGuid);
+		}
 	}
 #endif
 
diff --git a/Engine/Source/Runtime/Engine/Private/World.cpp b/Engine/Source/Runtime/Engine/Private/World.cpp
--- a/Engine/Source/Runtime/Engine/Private/World.cpp
+++ b/Engine/Source/Runtime/Engine/Private/World.cpp
@@ -1445,7 +1445,6 @@ void UWorld::RepairWorldSettings()
 		FActorSpawnParameters SpawnInfo;
 		SpawnInfo.SpawnCollisionHandlingOverride = ESpawnActorCollisionHandlingMethod::AlwaysSpawn;
 		SpawnInfo.Name = GEngine->WorldSettingsClass->GetFName();
-		SpawnInfo.bCreateActorPackage = false;
 		AWorldSettings* const NewWorldSettings = SpawnActor<AWorldSettings>( GEngine->WorldSettingsClass, SpawnInfo );
 
 		// If there was an existing actor, copy its properties to the new actor (the it will be destroyed by SetWorldSettings)
