FAnimSegment - Expose a couple functions so they can be accessed externally.

diff --git a/Engine/Source/Runtime/Engine/Classes/Animation/AnimCompositeBase.h b/Engine/Source/Runtime/Engine/Classes/Animation/AnimCompositeBase.h
--- a/Engine/Source/Runtime/Engine/Classes/Animation/AnimCompositeBase.h
+++ b/Engine/Source/Runtime/Engine/Classes/Animation/AnimCompositeBase.h
@@ -148,7 +148,7 @@ struct FAnimSegment
 
 	/** Converts 'Track Position' to position on AnimSequence.
 	 * Note: doesn't check that position is in valid range, must do that before calling this function! */
-	float ConvertTrackPosToAnimPos(const float& TrackPosition) const;
+	ENGINE_API float ConvertTrackPosToAnimPos(const float& TrackPosition) const;
 
 	/** 
 	 * Retrieves AnimNotifies between two Track time positions. ]PreviousTrackPosition, CurrentTrackPosition]
@@ -165,7 +165,7 @@ struct FAnimSegment
 	* Supports playing backwards (CurrentTrackPosition<PreviousTrackPosition).
 	* Only supports contiguous range, does NOT support looping and wrapping over.
 	*/
-	void GetAnimNotifiesFromTrackPositions(const float& PreviousTrackPosition, const float& CurrentTrackPosition, TArray<FAnimNotifyEventReference> & OutActiveNotifies) const;
+	ENGINE_API void GetAnimNotifiesFromTrackPositions(const float& PreviousTrackPosition, const float& CurrentTrackPosition, TArray<FAnimNotifyEventReference> & OutActiveNotifies) const;
 
 	/** 
 	 * Given a Track delta position [StartTrackPosition, EndTrackPosition]
