Fix part of Android build.

diff --git a/Engine/Source/Runtime/Engine/Private/Actor.cpp b/Engine/Source/Runtime/Engine/Private/Actor.cpp
--- a/Engine/Source/Runtime/Engine/Private/Actor.cpp
+++ b/Engine/Source/Runtime/Engine/Private/Actor.cpp
@@ -49,7 +49,10 @@
 #include "LevelUtils.h"
 #include "Components/DecalComponent.h"
 #include "GameFramework/InputSettings.h"
+
+#if WITH_EDITOR
 #include "FoliageHelper.h"
+#endif
 
 #include "WorldPartition/WorldPartition.h"
 #include "WorldPartition/WorldPartitionActorDesc.h"
