WorldPartition: Avoid doing ResetLoaders when renaming trashed streamed-out packages (optim for PIE).

diff --git a/Engine/Source/Runtime/Engine/Private/WorldPartition/WorldPartitionLevelStreamingDynamic.cpp b/Engine/Source/Runtime/Engine/Private/WorldPartition/WorldPartitionLevelStreamingDynamic.cpp
--- a/Engine/Source/Runtime/Engine/Private/WorldPartition/WorldPartitionLevelStreamingDynamic.cpp
+++ b/Engine/Source/Runtime/Engine/Private/WorldPartition/WorldPartitionLevelStreamingDynamic.cpp
@@ -348,7 +348,7 @@ void UWorldPartitionLevelStreamingDynamic::OnCleanupLevel()
 
 		// Rename package to avoid having to deal with pending kill objects in subsequent RequestLevel
 		FName NewUniqueTrashName = MakeUniqueObjectName(nullptr, UPackage::StaticClass(), FName(*FString::Printf(TEXT("%s_Trashed"), *RuntimePackage->GetName())));
-		RuntimePackage->Rename(*NewUniqueTrashName.ToString(), nullptr, REN_DontCreateRedirectors | REN_NonTransactional | REN_DoNotDirty);
+		RuntimePackage->Rename(*NewUniqueTrashName.ToString(), nullptr, REN_ForceNoResetLoaders | REN_DontCreateRedirectors | REN_NonTransactional | REN_DoNotDirty);
 
 		// Apply same logic on Actor Packages
 		for (AActor* Actor : RuntimeLevel->Actors)
@@ -362,7 +362,7 @@ void UWorldPartitionLevelStreamingDynamic::OnCleanupLevel()
 				}, false);
 
 				NewUniqueTrashName = MakeUniqueObjectName(nullptr, UPackage::StaticClass(), FName(*FString::Printf(TEXT("%s_Trashed"), *ActorPackage->GetName())));
-				ActorPackage->Rename(*NewUniqueTrashName.ToString(), nullptr, REN_DontCreateRedirectors | REN_NonTransactional | REN_DoNotDirty);
+				ActorPackage->Rename(*NewUniqueTrashName.ToString(), nullptr, REN_ForceNoResetLoaders | REN_DontCreateRedirectors | REN_NonTransactional | REN_DoNotDirty);
 			}
 		}
 
