Bring back some attachment functions that were never actually properly deprecated so would not produce compile warnings
Not enabling DeprecationFunction with this checkin to avoid warnings on the branch

diff --git a/Engine/Source/Runtime/Engine/Classes/Components/SceneComponent.h b/Engine/Source/Runtime/Engine/Classes/Components/SceneComponent.h
--- a/Engine/Source/Runtime/Engine/Classes/Components/SceneComponent.h
+++ b/Engine/Source/Runtime/Engine/Classes/Components/SceneComponent.h
@@ -673,6 +673,11 @@ public:
 	/** Backwards compatibility: Used to convert old-style EAttachLocation to new-style EAttachmentRules */
 	static void ConvertAttachLocation(EAttachLocation::Type InAttachLocation, EAttachmentRule& InOutLocationRule, EAttachmentRule& InOutRotationRule, EAttachmentRule& InOutScaleRule);
 
+	/** DEPRECATED - Use AttachToComponent() instead */
+	UE_DEPRECATED(4.17, "This function is deprecated, please use AttachToComponent() instead.")
+	UFUNCTION(BlueprintCallable, Category = "Utilities|Transformation", meta = (DeprecationMessage = "OVERRIDE BAD MESSAGE", DisplayName = "AttachTo (Deprecated)", AttachType = "KeepRelativeOffset"))
+	bool K2_AttachTo(USceneComponent* InParent, FName InSocketName = NAME_None, EAttachLocation::Type AttachType = EAttachLocation::KeepRelativeOffset, bool bWeldSimulatedBodies = true);
+
 	/**
 	* Attach this component to another scene component, optionally at a named socket. It is valid to call this on components whether or not they have been Registered, however from
 	* constructor or when not registered it is preferable to use SetupAttachment.
@@ -696,6 +701,11 @@ public:
 	UFUNCTION(BlueprintCallable, Category = "Utilities|Transformation", meta = (DisplayName = "AttachComponentToComponent", ScriptName = "AttachToComponent", bWeldSimulatedBodies=true))
 	bool K2_AttachToComponent(USceneComponent* Parent, FName SocketName, EAttachmentRule LocationRule, EAttachmentRule RotationRule, EAttachmentRule ScaleRule, bool bWeldSimulatedBodies);
 
+	/** DEPRECATED - Use DetachFromComponent() instead */
+	UE_DEPRECATED(4.12, "This function is deprecated, please use DetachFromComponent() instead.")
+	UFUNCTION(BlueprintCallable, Category = "Utilities|Transformation", meta = (DisplayName = "DetachFromParent (Deprecated)"))
+	virtual void DetachFromParent(bool bMaintainWorldPosition = false, bool bCallModify = true);
+
 	/** 
 	 * Detach this component from whatever it is attached to. Automatically unwelds components that are welded together (See WeldTo)
 	 * @param LocationRule				How to handle translations when detaching.
diff --git a/Engine/Source/Runtime/Engine/Classes/GameFramework/Actor.h b/Engine/Source/Runtime/Engine/Classes/GameFramework/Actor.h
--- a/Engine/Source/Runtime/Engine/Classes/GameFramework/Actor.h
+++ b/Engine/Source/Runtime/Engine/Classes/GameFramework/Actor.h
@@ -1403,6 +1403,11 @@ public:
 	UFUNCTION(BlueprintCallable, meta=(BlueprintInternalUseOnly="true"))
 	void FinishAddComponent(UActorComponent* Component, bool bManualAttachment, const FTransform& RelativeTransform);
 
+	/** DEPRECATED - Use AttachToComponent() instead */
+	UE_DEPRECATED(4.17, "Use AttachToComponent() instead.")
+	UFUNCTION(BlueprintCallable, meta = (DisplayName = "AttachRootComponentTo (Deprecated)", ScriptNoExport, AttachLocationType = "KeepRelativeOffset"), Category = "Utilities|Transformation")
+	void K2_AttachRootComponentTo(USceneComponent* InParent, FName InSocketName = NAME_None, EAttachLocation::Type AttachLocationType = EAttachLocation::KeepRelativeOffset, bool bWeldSimulatedBodies = true);
+
 	/**
 	 * Attaches the RootComponent of this Actor to the supplied component, optionally at a named socket. It is not valid to call this on components that are not Registered.
 	 * @param Parent					Parent to attach to.
@@ -1423,6 +1428,11 @@ public:
 	 */
 	void AttachToComponent(USceneComponent* Parent, const FAttachmentTransformRules& AttachmentRules, FName SocketName = NAME_None);
 
+	/** DEPRECATED - Use AttachToActor() instead */
+	UE_DEPRECATED(4.17, "Use AttachToActor() instead.")
+	UFUNCTION(BlueprintCallable, meta = (DisplayName = "AttachRootComponentToActor (Deprecated)", ScriptNoExport, AttachLocationType = "KeepRelativeOffset"), Category = "Utilities|Transformation")
+	void K2_AttachRootComponentToActor(AActor* InParentActor, FName InSocketName = NAME_None, EAttachLocation::Type AttachLocationType = EAttachLocation::KeepRelativeOffset, bool bWeldSimulatedBodies = true);
+
 	/**
 	 * Attaches the RootComponent of this Actor to the RootComponent of the supplied actor, optionally at a named socket.
 	 * @param ParentActor				Actor to attach this actor's RootComponent to
@@ -1443,6 +1453,11 @@ public:
 	UFUNCTION(BlueprintCallable, meta = (DisplayName = "AttachActorToActor", ScriptName = "AttachToActor", bWeldSimulatedBodies=true), Category = "Utilities|Transformation")
 	void K2_AttachToActor(AActor* ParentActor, FName SocketName, EAttachmentRule LocationRule, EAttachmentRule RotationRule, EAttachmentRule ScaleRule, bool bWeldSimulatedBodies);
 
+	/** DEPRECATED - Use DetachFromActor() instead */
+	UE_DEPRECATED(4.17, "Use DetachFromActor() instead")
+	UFUNCTION(BlueprintCallable, meta = (DisplayName = "DetachActorFromActor (Deprecated)", ScriptNoExport), Category = "Utilities|Transformation")
+	void DetachRootComponentFromParent(bool bMaintainWorldPosition = true);
+
 	/** 
 	 * Detaches the RootComponent of this Actor from any SceneComponent it is currently attached to. 
 	 * @param  LocationRule				How to handle translation when detaching.
diff --git a/Engine/Source/Runtime/Engine/Private/Actor.cpp b/Engine/Source/Runtime/Engine/Private/Actor.cpp
--- a/Engine/Source/Runtime/Engine/Private/Actor.cpp
+++ b/Engine/Source/Runtime/Engine/Private/Actor.cpp
@@ -1918,6 +1918,16 @@ void AActor::SetAutoDestroyWhenFinished(bool bVal)
 	}
 }
 
+void AActor::K2_AttachRootComponentTo(USceneComponent* InParent, FName InSocketName, EAttachLocation::Type AttachLocationType /*= EAttachLocation::KeepRelativeOffset */, bool bWeldSimulatedBodies /*=true*/)
+{
+PRAGMA_DISABLE_DEPRECATION_WARNINGS
+	if (RootComponent && InParent)
+	{
+		RootComponent->K2_AttachTo(InParent, InSocketName, AttachLocationType, bWeldSimulatedBodies);
+	}
+PRAGMA_ENABLE_DEPRECATION_WARNINGS
+}
+
 void AActor::K2_AttachToComponent(USceneComponent* Parent, FName SocketName, EAttachmentRule LocationRule, EAttachmentRule RotationRule, EAttachmentRule ScaleRule, bool bWeldSimulatedBodies)
 {
 	AttachToComponent(Parent, FAttachmentTransformRules(LocationRule, RotationRule, ScaleRule, bWeldSimulatedBodies), SocketName);
@@ -1976,6 +1986,20 @@ void AActor::OnRep_AttachmentReplication()
 	}
 }
 
+void AActor::K2_AttachRootComponentToActor(AActor* InParentActor, FName InSocketName /*= NAME_None*/, EAttachLocation::Type AttachLocationType /*= EAttachLocation::KeepRelativeOffset */, bool bWeldSimulatedBodies /*=true*/)
+{
+	if (RootComponent && InParentActor)
+	{
+		USceneComponent* ParentDefaultAttachComponent = InParentActor->GetDefaultAttachComponent();
+		if (ParentDefaultAttachComponent)
+		{
+PRAGMA_DISABLE_DEPRECATION_WARNINGS
+			RootComponent->K2_AttachTo(ParentDefaultAttachComponent, InSocketName, AttachLocationType, bWeldSimulatedBodies );
+PRAGMA_ENABLE_DEPRECATION_WARNINGS
+		}
+	}
+}
+
 void AActor::K2_AttachToActor(AActor* ParentActor, FName SocketName, EAttachmentRule LocationRule, EAttachmentRule RotationRule, EAttachmentRule ScaleRule, bool bWeldSimulatedBodies)
 {
 	AttachToActor(ParentActor, FAttachmentTransformRules(LocationRule, RotationRule, ScaleRule, bWeldSimulatedBodies), SocketName);
@@ -1993,6 +2017,20 @@ void AActor::AttachToActor(AActor* ParentActor, const FAttachmentTransformRules&
 	}
 }
 
+void AActor::DetachRootComponentFromParent(bool bMaintainWorldPosition)
+{
+	if(RootComponent)
+	{
+PRAGMA_DISABLE_DEPRECATION_WARNINGS
+		RootComponent->DetachFromParent(bMaintainWorldPosition);
+PRAGMA_ENABLE_DEPRECATION_WARNINGS		
+
+		// Clear AttachmentReplication struct
+		AttachmentReplication = FRepAttachment();
+		MARK_PROPERTY_DIRTY_FROM_NAME(AActor, AttachmentReplication, this);
+	}
+}
+
 void AActor::K2_DetachFromActor(EDetachmentRule LocationRule /*= EDetachmentRule::KeepRelative*/, EDetachmentRule RotationRule /*= EDetachmentRule::KeepRelative*/, EDetachmentRule ScaleRule /*= EDetachmentRule::KeepRelative*/)
 {
 	DetachFromActor(FDetachmentTransformRules(LocationRule, RotationRule, ScaleRule, true));
diff --git a/Engine/Source/Runtime/Engine/Private/Components/SceneComponent.cpp b/Engine/Source/Runtime/Engine/Private/Components/SceneComponent.cpp
--- a/Engine/Source/Runtime/Engine/Private/Components/SceneComponent.cpp
+++ b/Engine/Source/Runtime/Engine/Private/Components/SceneComponent.cpp
@@ -1768,6 +1768,17 @@ void USceneComponent::SetupAttachment(class USceneComponent* InParent, FName InS
 	}
 }
 
+//This function is used for giving AttachTo different bWeldSimulatedBodies default, but only when called from BP
+bool USceneComponent::K2_AttachTo(class USceneComponent* InParent, FName InSocketName, EAttachLocation::Type AttachLocationType, bool bWeldSimulatedBodies /*= true*/)
+{
+PRAGMA_DISABLE_DEPRECATION_WARNINGS
+	FAttachmentTransformRules AttachmentRules(EAttachmentRule::KeepRelative, bWeldSimulatedBodies);
+	ConvertAttachLocation(AttachLocationType, AttachmentRules.LocationRule, AttachmentRules.RotationRule, AttachmentRules.ScaleRule);
+
+	return AttachToComponent(InParent, AttachmentRules, InSocketName);
+PRAGMA_ENABLE_DEPRECATION_WARNINGS
+}
+
 //This function is used for giving AttachToComponent different bWeldSimulatedBodies default, but only when called from BP
 bool USceneComponent::K2_AttachToComponent(USceneComponent* Parent, FName SocketName, EAttachmentRule LocationRule, EAttachmentRule RotationRule, EAttachmentRule ScaleRule, bool bWeldSimulatedBodies)
 {
@@ -2119,6 +2130,21 @@ bool USceneComponent::AttachToComponent(USceneComponent* Parent, const FAttachme
 	return false;
 }
 
+void USceneComponent::DetachFromParent(bool bMaintainWorldPosition, bool bCallModify)
+{
+	FDetachmentTransformRules DetachmentRules(EDetachmentRule::KeepRelative, bCallModify);
+	if (bMaintainWorldPosition)
+	{
+		DetachmentRules.LocationRule = EDetachmentRule::KeepWorld;
+
+		// force maintain world rotation and scale for backwards compatibility
+		DetachmentRules.RotationRule = EDetachmentRule::KeepWorld;
+		DetachmentRules.ScaleRule = EDetachmentRule::KeepWorld;
+	}
+
+	DetachFromComponent(DetachmentRules);
+}
+
 void USceneComponent::K2_DetachFromComponent(EDetachmentRule LocationRule /*= EDetachmentRule::KeepRelative*/, EDetachmentRule RotationRule /*= EDetachmentRule::KeepRelative*/, EDetachmentRule ScaleRule /*= EDetachmentRule::KeepRelative*/, bool bCallModify /*= true*/)
 {
 	DetachFromComponent(FDetachmentTransformRules(LocationRule, RotationRule, ScaleRule, bCallModify));
