Pass the View uniform buffer into FFXSystemInterface::PreRender, since changes to GlobalDistanceField now require them.
Fixes crashes in RHI with Cascade systems that use distance field collisions

diff --git a/Engine/Plugins/FX/Niagara/Source/Niagara/Classes/NiagaraEmitterInstanceBatcher.h b/Engine/Plugins/FX/Niagara/Source/Niagara/Classes/NiagaraEmitterInstanceBatcher.h
--- a/Engine/Plugins/FX/Niagara/Source/Niagara/Classes/NiagaraEmitterInstanceBatcher.h
+++ b/Engine/Plugins/FX/Niagara/Source/Niagara/Classes/NiagaraEmitterInstanceBatcher.h
@@ -74,7 +74,7 @@ public:
 	virtual bool UsesGlobalDistanceField() const override;
 	virtual bool UsesDepthBuffer() const override;
 	virtual bool RequiresEarlyViewUniformBuffer() const override;
-	virtual void PreRender(FRHICommandListImmediate& RHICmdList, const class FGlobalDistanceFieldParameterData* GlobalDistanceFieldParameterData, bool bAllowGPUParticleUpdate) override;
+	virtual void PreRender(FRHICommandListImmediate& RHICmdList, FRHIUniformBuffer* ViewUniformBuffer, const class FGlobalDistanceFieldParameterData* GlobalDistanceFieldParameterData, bool bAllowGPUParticleUpdate) override;
 	virtual void OnDestroy() override; // Called on the gamethread to delete the batcher on the renderthread.
 
 	virtual void Tick(float DeltaTime) override
diff --git a/Engine/Plugins/FX/Niagara/Source/Niagara/Private/NiagaraEmitterInstanceBatcher.cpp b/Engine/Plugins/FX/Niagara/Source/Niagara/Private/NiagaraEmitterInstanceBatcher.cpp
--- a/Engine/Plugins/FX/Niagara/Source/Niagara/Private/NiagaraEmitterInstanceBatcher.cpp
+++ b/Engine/Plugins/FX/Niagara/Source/Niagara/Private/NiagaraEmitterInstanceBatcher.cpp
@@ -1428,7 +1428,7 @@ bool NiagaraEmitterInstanceBatcher::RequiresEarlyViewUniformBuffer() const
 	return NumTicksThatRequireEarlyViewData > 0;
 }
 
-void NiagaraEmitterInstanceBatcher::PreRender(FRHICommandListImmediate& RHICmdList, const class FGlobalDistanceFieldParameterData* GlobalDistanceFieldParameterData, bool bAllowGPUParticleUpdate)
+void NiagaraEmitterInstanceBatcher::PreRender(FRHICommandListImmediate& RHICmdList, FRHIUniformBuffer* ViewUniformBuffer, const class FGlobalDistanceFieldParameterData* GlobalDistanceFieldParameterData, bool bAllowGPUParticleUpdate)
 {
 	if (!FNiagaraUtilities::AllowGPUParticles(GetShaderPlatform()))
 	{
diff --git a/Engine/Source/Runtime/Engine/Private/Particles/FXSystem.cpp b/Engine/Source/Runtime/Engine/Private/Particles/FXSystem.cpp
--- a/Engine/Source/Runtime/Engine/Private/Particles/FXSystem.cpp
+++ b/Engine/Source/Runtime/Engine/Private/Particles/FXSystem.cpp
@@ -437,7 +437,7 @@ DECLARE_CYCLE_STAT(TEXT("FXPreRender_SimulateCDF"), STAT_CLM_FXPreRender_Simulat
 DECLARE_CYCLE_STAT(TEXT("FXPreRender_FinalizeCDF"), STAT_CLM_FXPreRender_FinalizeCDF, STATGROUP_CommandListMarkers);
 
 
-void FFXSystem::PreRender(FRHICommandListImmediate& RHICmdList, const FGlobalDistanceFieldParameterData* GlobalDistanceFieldParameterData, bool bAllowGPUParticleSceneUpdate)
+void FFXSystem::PreRender(FRHICommandListImmediate& RHICmdList, FRHIUniformBuffer* ViewUniformBuffer, const FGlobalDistanceFieldParameterData* GlobalDistanceFieldParameterData, bool bAllowGPUParticleSceneUpdate)
 {
 	if (RHISupportsGPUParticles() && bAllowGPUParticleSceneUpdate)
 	{
@@ -459,7 +459,7 @@ void FFXSystem::PreRender(FRHICommandListImmediate& RHICmdList, const FGlobalDis
 			PrepareGPUSimulation(RHICmdList);
 
 			RHICmdList.SetCurrentStat(GET_STATID(STAT_CLM_FXPreRender_SimulateCDF));
-			SimulateGPUParticles(RHICmdList, EParticleSimulatePhase::CollisionDistanceField, nullptr, GlobalDistanceFieldParameterData, nullptr, nullptr);
+			SimulateGPUParticles(RHICmdList, EParticleSimulatePhase::CollisionDistanceField, ViewUniformBuffer, GlobalDistanceFieldParameterData, nullptr, nullptr);
 			//particles rendered during basepass may need to read pos/velocity buffers; must finalize unless we know for sure that nothing in base pass will read it.
 			RHICmdList.SetCurrentStat(GET_STATID(STAT_CLM_FXPreRender_FinalizeCDF));
 			FinalizeGPUSimulation(RHICmdList);
diff --git a/Engine/Source/Runtime/Engine/Private/Particles/FXSystemPrivate.h b/Engine/Source/Runtime/Engine/Private/Particles/FXSystemPrivate.h
--- a/Engine/Source/Runtime/Engine/Private/Particles/FXSystemPrivate.h
+++ b/Engine/Source/Runtime/Engine/Private/Particles/FXSystemPrivate.h
@@ -155,7 +155,7 @@ public:
 	virtual bool UsesGlobalDistanceField() const override;
 	virtual bool UsesDepthBuffer() const override;
 	virtual bool RequiresEarlyViewUniformBuffer() const override;
-	virtual void PreRender(FRHICommandListImmediate& RHICmdList, const FGlobalDistanceFieldParameterData* GlobalDistanceFieldParameterData, bool bAllowGPUParticleUpdate) override;
+	virtual void PreRender(FRHICommandListImmediate& RHICmdList, FRHIUniformBuffer* ViewUniformBuffer, const FGlobalDistanceFieldParameterData* GlobalDistanceFieldParameterData, bool bAllowGPUParticleUpdate) override;
 	virtual void PostRenderOpaque(
 		FRHICommandListImmediate& RHICmdList, 
 		FRHIUniformBuffer* ViewUniformBuffer,
diff --git a/Engine/Source/Runtime/Engine/Private/Particles/FXSystemSet.cpp b/Engine/Source/Runtime/Engine/Private/Particles/FXSystemSet.cpp
--- a/Engine/Source/Runtime/Engine/Private/Particles/FXSystemSet.cpp
+++ b/Engine/Source/Runtime/Engine/Private/Particles/FXSystemSet.cpp
@@ -172,12 +172,12 @@ bool FFXSystemSet::RequiresEarlyViewUniformBuffer() const
 	return false;
 }
 
-void FFXSystemSet::PreRender(FRHICommandListImmediate& RHICmdList, const class FGlobalDistanceFieldParameterData* GlobalDistanceFieldParameterData, bool bAllowGPUParticleSceneUpdate)
+void FFXSystemSet::PreRender(FRHICommandListImmediate& RHICmdList, FRHIUniformBuffer* ViewUniformBuffer, const class FGlobalDistanceFieldParameterData* GlobalDistanceFieldParameterData, bool bAllowGPUParticleSceneUpdate)
 {
 	for (FFXSystemInterface* FXSystem : FXSystems)
 	{
 		check(FXSystem);
-		FXSystem->PreRender(RHICmdList, GlobalDistanceFieldParameterData, bAllowGPUParticleSceneUpdate);
+		FXSystem->PreRender(RHICmdList, ViewUniformBuffer, GlobalDistanceFieldParameterData, bAllowGPUParticleSceneUpdate);
 	}
 }
 
diff --git a/Engine/Source/Runtime/Engine/Private/Particles/FXSystemSet.h b/Engine/Source/Runtime/Engine/Private/Particles/FXSystemSet.h
--- a/Engine/Source/Runtime/Engine/Private/Particles/FXSystemSet.h
+++ b/Engine/Source/Runtime/Engine/Private/Particles/FXSystemSet.h
@@ -42,7 +42,7 @@ public:
 	virtual bool UsesGlobalDistanceField() const override;
 	virtual bool UsesDepthBuffer() const override;
 	virtual bool RequiresEarlyViewUniformBuffer() const override;
-	virtual void PreRender(FRHICommandListImmediate& RHICmdList, const class FGlobalDistanceFieldParameterData* GlobalDistanceFieldParameterData, bool bAllowGPUParticleSceneUpdate) override;
+	virtual void PreRender(FRHICommandListImmediate& RHICmdList, FRHIUniformBuffer* ViewUniformBuffer, const class FGlobalDistanceFieldParameterData* GlobalDistanceFieldParameterData, bool bAllowGPUParticleSceneUpdate) override;
 	virtual void PostRenderOpaque(
 		FRHICommandListImmediate& RHICmdList, 
 		FRHIUniformBuffer* ViewUniformBuffer,
diff --git a/Engine/Source/Runtime/Engine/Public/FXSystem.h b/Engine/Source/Runtime/Engine/Public/FXSystem.h
--- a/Engine/Source/Runtime/Engine/Public/FXSystem.h
+++ b/Engine/Source/Runtime/Engine/Public/FXSystem.h
@@ -238,7 +238,7 @@ public:
 	 * Notification from the renderer that it is about to draw FX belonging to
 	 * this system.
 	 */
-	virtual void PreRender(FRHICommandListImmediate& RHICmdList, const class FGlobalDistanceFieldParameterData* GlobalDistanceFieldParameterData, bool bAllowGPUParticleSceneUpdate) = 0;
+	virtual void PreRender(FRHICommandListImmediate& RHICmdList, FRHIUniformBuffer* ViewUniformBuffer, const class FGlobalDistanceFieldParameterData* GlobalDistanceFieldParameterData, bool bAllowGPUParticleSceneUpdate) = 0;
 
 	/**
 	 * Notification from the renderer that opaque primitives have rendered.
diff --git a/Engine/Source/Runtime/Renderer/Private/DeferredShadingRenderer.cpp b/Engine/Source/Runtime/Renderer/Private/DeferredShadingRenderer.cpp
--- a/Engine/Source/Runtime/Renderer/Private/DeferredShadingRenderer.cpp
+++ b/Engine/Source/Runtime/Renderer/Private/DeferredShadingRenderer.cpp
@@ -1746,7 +1746,7 @@ void FDeferredShadingSceneRenderer::Render(FRHICommandListImmediate& RHICmdList)
 	{
 		SCOPE_CYCLE_COUNTER(STAT_FDeferredShadingSceneRenderer_FXSystem_PreRender);
 		RHICmdList.SetCurrentStat(GET_STATID(STAT_CLM_FXPreRender));
-		FXSystem->PreRender(RHICmdList, &Views[0].GlobalDistanceFieldInfo.ParameterData, Views[0].AllowGPUParticleUpdate());
+		FXSystem->PreRender(RHICmdList, Views[0].ViewUniformBuffer, &Views[0].GlobalDistanceFieldInfo.ParameterData, Views[0].AllowGPUParticleUpdate());
 		if (FGPUSortManager* GPUSortManager = FXSystem->GetGPUSortManager())
 		{
 			GPUSortManager->OnPreRender(RHICmdList);
diff --git a/Engine/Source/Runtime/Renderer/Private/MobileShadingRenderer.cpp b/Engine/Source/Runtime/Renderer/Private/MobileShadingRenderer.cpp
--- a/Engine/Source/Runtime/Renderer/Private/MobileShadingRenderer.cpp
+++ b/Engine/Source/Runtime/Renderer/Private/MobileShadingRenderer.cpp
@@ -609,7 +609,7 @@ void FMobileSceneRenderer::Render(FRHICommandListImmediate& RHICmdList)
 	// Notify the FX system that the scene is about to be rendered.
 	if (FXSystem && ViewFamily.EngineShowFlags.Particles)
 	{
-		FXSystem->PreRender(RHICmdList, NULL, !Views[0].bIsPlanarReflection);
+		FXSystem->PreRender(RHICmdList, Views[0].ViewUniformBuffer, NULL, !Views[0].bIsPlanarReflection);
 		if (FGPUSortManager* GPUSortManager = FXSystem->GetGPUSortManager())
 		{
 			GPUSortManager->OnPreRender(RHICmdList);
diff --git a/Engine/Source/Runtime/Renderer/Private/SceneHitProxyRendering.cpp b/Engine/Source/Runtime/Renderer/Private/SceneHitProxyRendering.cpp
--- a/Engine/Source/Runtime/Renderer/Private/SceneHitProxyRendering.cpp
+++ b/Engine/Source/Runtime/Renderer/Private/SceneHitProxyRendering.cpp
@@ -638,7 +638,7 @@ void FDeferredShadingSceneRenderer::RenderHitProxies(FRHICommandListImmediate& R
 		if (FXSystem && Views.IsValidIndex(0))
 		{
 			FGPUSortManager* GPUSortManager = FXSystem->GetGPUSortManager();
-			FXSystem->PreRender(RHICmdList, &Views[0].GlobalDistanceFieldInfo.ParameterData, false);
+			FXSystem->PreRender(RHICmdList, Views[0].ViewUniformBuffer, &Views[0].GlobalDistanceFieldInfo.ParameterData, false);
 			if (GPUSortManager)
 			{
 				GPUSortManager->OnPreRender(RHICmdList);
