DistanceField and MeshCards async queues are now reusing existing threadpool instead of spawning their own
DistanceField and MeshCards are not limited anymore to a single thread and can scale to the number of physical cores available
DistanceField and MeshCards are now both limiting their tick time length to avoid huge game-thread stalls when lots of assets are being built
Fix potential race-condition in DistanceFieldAsyncQueue found while fixing MeshCardAsyncQueue

diff --git a/Engine/Source/Developer/MeshUtilities/Private/MeshDistanceFieldUtilities.cpp b/Engine/Source/Developer/MeshUtilities/Private/MeshDistanceFieldUtilities.cpp
--- a/Engine/Source/Developer/MeshUtilities/Private/MeshDistanceFieldUtilities.cpp
+++ b/Engine/Source/Developer/MeshUtilities/Private/MeshDistanceFieldUtilities.cpp
@@ -188,6 +188,8 @@ private:
 
 void FMeshDistanceFieldAsyncTask::DoWork()
 {
+	TRACE_CPUPROFILER_EVENT_SCOPE(FMeshDistanceFieldAsyncTask::DoWork);
+
 	FMeshBuildDataProvider kDOPDataProvider(*kDopTree);
 	const FVector DistanceFieldVoxelSize(VolumeBounds.GetSize() / FVector(VolumeDimensions.X, VolumeDimensions.Y, VolumeDimensions.Z));
 	const float VoxelDiameterSqr = DistanceFieldVoxelSize.SizeSquared();
@@ -328,6 +330,8 @@ void FMeshUtilities::GenerateSignedDistanceFieldVolumeData(
 	bool bGenerateAsIfTwoSided,
 	FDistanceFieldVolumeData& OutData)
 {
+	TRACE_CPUPROFILER_EVENT_SCOPE(FMeshUtilities::GenerateSignedDistanceFieldVolumeData);
+
 	if (DistanceFieldResolutionScale > 0)
 	{
 		const double StartTime = FPlatformTime::Seconds();
diff --git a/Engine/Source/Runtime/Engine/Private/DistanceFieldAtlas.cpp b/Engine/Source/Runtime/Engine/Private/DistanceFieldAtlas.cpp
--- a/Engine/Source/Runtime/Engine/Private/DistanceFieldAtlas.cpp
+++ b/Engine/Source/Runtime/Engine/Private/DistanceFieldAtlas.cpp
@@ -21,6 +21,8 @@
 #include "GlobalShader.h"
 #include "RenderGraph.h"
 #include "MeshCardRepresentation.h"
+#include "Misc/QueuedThreadPoolWrapper.h"
+#include "Async/Async.h"
 
 #if WITH_EDITOR
 #include "DerivedDataCacheInterface.h"
@@ -983,124 +985,6 @@ static FAutoConsoleVariableRef CVarAOAsyncBuildQueue(
 	ECVF_Default | ECVF_ReadOnly
 	);
 
-class FBuildDistanceFieldThreadRunnable : public FRunnable
-{
-public:
-	/** Initialization constructor. */
-	FBuildDistanceFieldThreadRunnable(
-		FDistanceFieldAsyncQueue* InAsyncQueue
-		)
-		: NextThreadIndex(0)
-		, AsyncQueue(*InAsyncQueue)
-		, Thread(nullptr)
-		, bIsRunning(false)
-		, bForceFinish(false)
-	{}
-
-	virtual ~FBuildDistanceFieldThreadRunnable()
-	{
-		check(!bIsRunning);
-	}
-
-	// FRunnable interface.
-	virtual bool Init() { return true; }
-	virtual void Exit() { bIsRunning = false; }
-	virtual void Stop() { bForceFinish = true; }
-	virtual uint32 Run();
-
-	void Launch()
-	{
-		check(!bIsRunning);
-
-		// Calling Reset will call Kill which in turn will call Stop and set bForceFinish to true.
-		Thread.Reset();
-
-		// Now we can set bForceFinish to false without being overwritten by the old thread shutting down.
-		bForceFinish = false;
-		Thread.Reset(FRunnableThread::Create(this, *FString::Printf(TEXT("BuildDistanceFieldThread%u"), NextThreadIndex), 0, TPri_Normal, FPlatformAffinity::GetPoolThreadMask()));
-
-		// Set this now before exiting so that IsRunning() returns true without having to wait on the thread to be completely started.
-		bIsRunning = true;
-		NextThreadIndex++;
-	}
-
-	inline bool IsRunning() { return bIsRunning; }
-
-private:
-
-	int32 NextThreadIndex;
-
-	FDistanceFieldAsyncQueue& AsyncQueue;
-
-	/** The runnable thread */
-	TUniquePtr<FRunnableThread> Thread;
-
-	TUniquePtr<FQueuedThreadPool> WorkerThreadPool;
-
-	volatile bool bIsRunning;
-	volatile bool bForceFinish;
-};
-
-static FQueuedThreadPool* CreateWorkerThreadPool()
-{
-	TRACE_CPUPROFILER_EVENT_SCOPE(CreateWorkerThreadPool)
-	const int32 NumThreads = FMath::Max<int32>(FPlatformMisc::NumberOfCoresIncludingHyperthreads() - 2, 1);
-	FQueuedThreadPool* WorkerThreadPool = FQueuedThreadPool::Allocate();
-	WorkerThreadPool->Create(NumThreads, 32 * 1024, TPri_BelowNormal);
-	return WorkerThreadPool;
-}
-
-uint32 FBuildDistanceFieldThreadRunnable::Run()
-{
-	TRACE_CPUPROFILER_EVENT_SCOPE(FBuildDistanceFieldThreadRunnable::Run)
-
-	bool bHasWork = true;
-
-	// Do not exit right away if no work to do as it often leads to stop and go problems
-	// when tasks are being queued at a slower rate than the processor capability to process them.
-	const uint64 ExitAfterIdleCycle = static_cast<uint64>(10.0 / FPlatformTime::GetSecondsPerCycle64()); // 10s
-
-	uint64 LastWorkCycle = FPlatformTime::Cycles64();
-	while (!bForceFinish &&  (bHasWork || (FPlatformTime::Cycles64() - LastWorkCycle) < ExitAfterIdleCycle))
-	{
-		// LIFO build order, since meshes actually visible in a map are typically loaded last
-		FAsyncDistanceFieldTask* Task = AsyncQueue.TaskQueue.Pop();
-
-		FQueuedThreadPool* ThreadPool = nullptr;
-		
-#if WITH_EDITOR
-		ThreadPool = GLargeThreadPool;
-#endif
-
-		if (Task)
-		{
-			if (!ThreadPool)
-			{
-				if (!WorkerThreadPool)
-				{
-					WorkerThreadPool.Reset(CreateWorkerThreadPool());
-				}
-
-				ThreadPool = WorkerThreadPool.Get();
-			}
-
-			AsyncQueue.Build(Task, *ThreadPool);
-			LastWorkCycle = FPlatformTime::Cycles64();
-
-			bHasWork = true;
-		}
-		else
-		{
-			bHasWork = false;
-			FPlatformProcess::Sleep(.01f);
-		}
-	}
-
-	WorkerThreadPool = nullptr;
-
-	return 0;
-}
-
 FAsyncDistanceFieldTask::FAsyncDistanceFieldTask()
 	: StaticMesh(nullptr)
 	, GenerateSource(nullptr)
@@ -1110,14 +994,13 @@ FAsyncDistanceFieldTask::FAsyncDistanceFieldTask()
 {
 }
 
-
 FDistanceFieldAsyncQueue::FDistanceFieldAsyncQueue() 
 {
 #if WITH_EDITOR
 	MeshUtilities = NULL;
 #endif
 
-	ThreadRunnable = MakeUnique<FBuildDistanceFieldThreadRunnable>(this);
+	ThreadPool = MakeUnique<FQueuedThreadPoolWrapper>(GThreadPool, -1, [](EQueuedWorkPriority) { return EQueuedWorkPriority::Lowest; });
 }
 
 FDistanceFieldAsyncQueue::~FDistanceFieldAsyncQueue()
@@ -1138,24 +1021,23 @@ void FDistanceFieldAsyncQueue::AddTask(FAsyncDistanceFieldTask* Task)
 		ReferencedTasks.Add(Task);
 	}
 
-	// If we're already in worker threads, we have to use async tasks
-	// to avoid crashing in the Build function.
-	// Also protects from creating too many thread pools when already parallel.
+	auto BuildLambda = 
+		[this, Task]()
+		{
+			// To avoid deadlocks, we must queue the inner build tasks on another thread pool, so use the task graph.
+			// Put on background thread to avoid interfering with game-thread bound tasks
+			FQueuedThreadPoolTaskGraphWrapper TaskGraphWrapper(ENamedThreads::AnyBackgroundThreadNormalTask); 
+			Build(Task, TaskGraphWrapper);
+		};
+
+	// If we're already in worker threads, there is no need to launch an async task.
 	if (GUseAsyncDistanceFieldBuildQueue || !IsInGameThread())
 	{
-		TaskQueue.Push(Task);
-
-		// Logic protection when called from multiple threads
-		FScopeLock Lock(&CriticalSection);
-		if (!ThreadRunnable->IsRunning())
-		{
-			ThreadRunnable->Launch();
-		}
+		AsyncPool(*ThreadPool, MoveTemp(BuildLambda));
 	}
 	else
 	{
-		TUniquePtr<FQueuedThreadPool> WorkerThreadPool(CreateWorkerThreadPool());
-		Build(Task, *WorkerThreadPool);
+		BuildLambda();
 	}
 #else
 	UE_LOG(LogStaticMesh,Fatal,TEXT("Tried to build a distance field without editor support (this should have been done during cooking)"));
@@ -1181,10 +1063,13 @@ void FDistanceFieldAsyncQueue::BlockUntilBuildComplete(UStaticMesh* StaticMesh,
 
 		bReferenced = false;
 
-		for (int TaskIndex = 0; TaskIndex < ReferencedTasks.Num(); TaskIndex++)
 		{
-			bReferenced = bReferenced || ReferencedTasks[TaskIndex]->StaticMesh == StaticMesh;
-			bReferenced = bReferenced || ReferencedTasks[TaskIndex]->GenerateSource == StaticMesh;
+			FScopeLock Lock(&CriticalSection);
+			for (int TaskIndex = 0; TaskIndex < ReferencedTasks.Num(); TaskIndex++)
+			{
+				bReferenced = bReferenced || ReferencedTasks[TaskIndex]->StaticMesh == StaticMesh;
+				bReferenced = bReferenced || ReferencedTasks[TaskIndex]->GenerateSource == StaticMesh;
+			}
 		}
 
 		if (bReferenced)
@@ -1224,7 +1109,7 @@ void FDistanceFieldAsyncQueue::BlockUntilAllBuildsComplete()
 	while (GetNumOutstandingTasks() > 0);
 }
 
-void FDistanceFieldAsyncQueue::Build(FAsyncDistanceFieldTask* Task, FQueuedThreadPool& ThreadPool)
+void FDistanceFieldAsyncQueue::Build(FAsyncDistanceFieldTask* Task, FQueuedThreadPool& BuildThreadPool)
 {
 #if WITH_EDITOR
 	// Editor 'force delete' can null any UObject pointers which are seen by reference collecting (eg FProperty or serialized)
@@ -1238,7 +1123,7 @@ void FDistanceFieldAsyncQueue::Build(FAsyncDistanceFieldTask* Task, FQueuedThrea
 			Task->StaticMesh->GetName(),
 			Task->SourceMeshData,
 			LODModel,
-			ThreadPool,
+			BuildThreadPool,
 			Task->MaterialBlendModes,
 			Task->GenerateSource->RenderData->Bounds,
 			Task->DistanceFieldResolutionScale,
@@ -1253,6 +1138,7 @@ void FDistanceFieldAsyncQueue::Build(FAsyncDistanceFieldTask* Task, FQueuedThrea
 
 void FDistanceFieldAsyncQueue::AddReferencedObjects(FReferenceCollector& Collector)
 {	
+	FScopeLock Lock(&CriticalSection);
 	for (int TaskIndex = 0; TaskIndex < ReferencedTasks.Num(); TaskIndex++)
 	{
 		// Make sure none of the UObjects referenced by the async tasks are GC'ed during the task
@@ -1266,20 +1152,28 @@ FString FDistanceFieldAsyncQueue::GetReferencerName() const
 	return TEXT("FDistanceFieldAsyncQueue");
 }
 
-void FDistanceFieldAsyncQueue::ProcessAsyncTasks()
+void FDistanceFieldAsyncQueue::ProcessAsyncTasks(bool bLimitExecutionTime)
 {
-	TRACE_CPUPROFILER_EVENT_SCOPE(FDistanceFieldAsyncQueue::ProcessAsyncTasks)
 #if WITH_EDITOR
-	TArray<FAsyncDistanceFieldTask*> LocalCompletedTasks;
-	CompletedTasks.PopAll(LocalCompletedTasks);
+	TRACE_CPUPROFILER_EVENT_SCOPE(FDistanceFieldAsyncQueue::ProcessAsyncTasks)
 
-	for (int TaskIndex = 0; TaskIndex < LocalCompletedTasks.Num(); TaskIndex++)
+	const double MaxProcessingTime = 0.016f;
+	double StartTime = FPlatformTime::Seconds();
+	while (!bLimitExecutionTime || (FPlatformTime::Seconds() - StartTime) < MaxProcessingTime)
 	{
+		FAsyncDistanceFieldTask* Task = CompletedTasks.Pop();
+		if (Task == nullptr)
+		{
+			break;
+		}
+
 		// We want to count each resource built from a DDC miss, so count each iteration of the loop separately.
 		COOK_STAT(auto Timer = DistanceFieldCookStats::UsageStats.TimeSyncWork());
-		FAsyncDistanceFieldTask* Task = LocalCompletedTasks[TaskIndex];
 
-		ReferencedTasks.Remove(Task);
+		{
+			FScopeLock Lock(&CriticalSection);
+			ReferencedTasks.Remove(Task);
+		}
 
 		// Editor 'force delete' can null any UObject pointers which are seen by reference collecting (eg FProperty or serialized)
 		if (Task->StaticMesh)
@@ -1322,28 +1216,13 @@ void FDistanceFieldAsyncQueue::ProcessAsyncTasks()
 
 		delete Task;
 	}
-
-	if (ReferencedTasks.Num() > 0 && !ThreadRunnable->IsRunning())
-	{
-		ThreadRunnable->Launch();
-	}
 #endif
 }
 
 void FDistanceFieldAsyncQueue::Shutdown()
 {
-	ThreadRunnable->Stop();
-	bool bLogged = false;
-
-	while (ThreadRunnable->IsRunning())
-	{
-		if (!bLogged)
-		{
-			bLogged = true;
-			UE_LOG(LogStaticMesh,Log,TEXT("Abandoning remaining async distance field tasks for shutdown"));
-		}
-		FPlatformProcess::Sleep(0.01f);
-	}
+	UE_LOG(LogStaticMesh, Log, TEXT("Abandoning remaining async distance field tasks for shutdown"));
+	ThreadPool->Destroy();
 }
 
 FLandscapeTextureAtlas::FLandscapeTextureAtlas(ESubAllocType InSubAllocType)
diff --git a/Engine/Source/Runtime/Engine/Private/MeshCardRepresentation.cpp b/Engine/Source/Runtime/Engine/Private/MeshCardRepresentation.cpp
--- a/Engine/Source/Runtime/Engine/Private/MeshCardRepresentation.cpp
+++ b/Engine/Source/Runtime/Engine/Private/MeshCardRepresentation.cpp
@@ -17,6 +17,8 @@
 #include "Engine/StaticMesh.h"
 #include "Misc/AutomationTest.h"
 #include "Async/ParallelFor.h"
+#include "Misc/QueuedThreadPoolWrapper.h"
+#include "Async/Async.h"
 
 #if WITH_EDITOR
 #include "DerivedDataCacheInterface.h"
@@ -148,139 +150,19 @@ static FAutoConsoleVariableRef CVarCardRepresentationAsyncBuildQueue(
 	ECVF_Default | ECVF_ReadOnly
 	);
 
-class FBuildCardRepresentationThreadRunnable : public FRunnable
-{
-public:
-	/** Initialization constructor. */
-	FBuildCardRepresentationThreadRunnable(
-		FCardRepresentationAsyncQueue* InAsyncQueue
-		)
-		: NextThreadIndex(0)
-		, AsyncQueue(*InAsyncQueue)
-		, Thread(nullptr)
-		, bIsRunning(false)
-		, bForceFinish(false)
-	{}
-
-	virtual ~FBuildCardRepresentationThreadRunnable()
-	{
-		check(!bIsRunning);
-	}
-
-	// FRunnable interface.
-	virtual bool Init() { bIsRunning = true; return true; }
-	virtual void Exit() { bIsRunning = false; }
-	virtual void Stop() { bForceFinish = true; }
-	virtual uint32 Run();
-
-	void Launch()
-	{
-		check(!bIsRunning);
-
-		// Calling Reset will call Kill which in turn will call Stop and set bForceFinish to true.
-		Thread.Reset();
-
-		// Now we can set bForceFinish to false without being overwritten by the old thread shutting down.
-		bForceFinish = false;
-		Thread.Reset(FRunnableThread::Create(this, *FString::Printf(TEXT("BuildCardRepresentationThread%u"), NextThreadIndex), 0, TPri_SlightlyBelowNormal, FPlatformAffinity::GetPoolThreadMask()));
-
-		// Set this now before exiting so that IsRunning() returns true without having to wait on the thread to be completely started.
-		bIsRunning = true;
-		NextThreadIndex++;
-	}
-
-	inline bool IsRunning() { return bIsRunning; }
-
-private:
-
-	int32 NextThreadIndex;
-
-	FCardRepresentationAsyncQueue& AsyncQueue;
-
-	/** The runnable thread */
-	TUniquePtr<FRunnableThread> Thread;
-
-	TUniquePtr<FQueuedThreadPool> WorkerThreadPool;
-
-	volatile bool bIsRunning;
-	volatile bool bForceFinish;
-};
-
-FQueuedThreadPool* CreateCardWorkerThreadPool()
-{
-	/*
-	const int32 NumThreads = FMath::Max<int32>(FPlatformMisc::NumberOfCoresIncludingHyperthreads() - 2, 1);
-	FQueuedThreadPool* WorkerThreadPool = FQueuedThreadPool::Allocate();
-	WorkerThreadPool->Create(NumThreads, 32 * 1024, TPri_BelowNormal);
-	return WorkerThreadPool;
-	*/
-
-	return nullptr;
-}
-
-uint32 FBuildCardRepresentationThreadRunnable::Run()
-{
-	bool bHasWork = true;
-	
-	// Do not exit right away if no work to do as it often leads to stop and go problems
-	// when tasks are being queued at a slower rate than the processor capability to process them.
-	const uint64 ExitAfterIdleCycle = static_cast<uint64>(10.0 / FPlatformTime::GetSecondsPerCycle64()); // 10s
-
-	uint64 LastWorkCycle = FPlatformTime::Cycles64();
-	while (!bForceFinish && (bHasWork || (FPlatformTime::Cycles64() - LastWorkCycle) < ExitAfterIdleCycle))
-	{
-		// LIFO build order, since meshes actually visible in a map are typically loaded last
-		FAsyncCardRepresentationTask* Task = AsyncQueue.TaskQueue.Pop();
-
-		FQueuedThreadPool* ThreadPool = nullptr;
-
-#if WITH_EDITOR
-		ThreadPool = GLargeThreadPool;
-#endif
-
-		if (Task)
-		{
-			if (!ThreadPool)
-			{
-				if (!WorkerThreadPool)
-				{
-					WorkerThreadPool.Reset(CreateCardWorkerThreadPool());
-				}
-
-				ThreadPool = WorkerThreadPool.Get();
-			}
-
-			AsyncQueue.Build(Task, *ThreadPool);
-			LastWorkCycle = FPlatformTime::Cycles64();
-
-			bHasWork = true;
-		}
-		else
-		{
-			bHasWork = false;
-			FPlatformProcess::Sleep(.01f);
-		}
-	}
-
-	WorkerThreadPool = nullptr;
-
-	return 0;
-}
-
 FAsyncCardRepresentationTask::FAsyncCardRepresentationTask()
 	: StaticMesh(nullptr)
 	, GenerateSource(nullptr)
 {
 }
 
-
 FCardRepresentationAsyncQueue::FCardRepresentationAsyncQueue() 
 {
 #if WITH_EDITOR
 	MeshUtilities = NULL;
 #endif
 
-	ThreadRunnable = MakeUnique<FBuildCardRepresentationThreadRunnable>(this);
+	ThreadPool = MakeUnique<FQueuedThreadPoolWrapper>(GThreadPool, -1, [](EQueuedWorkPriority) { return EQueuedWorkPriority::Lowest; });
 }
 
 FCardRepresentationAsyncQueue::~FCardRepresentationAsyncQueue()
@@ -301,21 +183,23 @@ void FCardRepresentationAsyncQueue::AddTask(FAsyncCardRepresentationTask* Task)
 		ReferencedTasks.Add(Task);
 	}
 
-	if (GUseAsyncCardRepresentationBuildQueue)
-	{
-		TaskQueue.Push(Task);
-
-		// Logic protection when called from multiple threads
-		FScopeLock Lock(&CriticalSection);
-		if (!ThreadRunnable->IsRunning())
+	auto BuildLambda =
+		[this, Task]()
 		{
-			ThreadRunnable->Launch();
-		}
+			// To avoid deadlocks, we must queue the inner build tasks on another thread pool, so use the task graph.
+			// Put on background thread to avoid interfering with game-thread bound tasks
+			FQueuedThreadPoolTaskGraphWrapper TaskGraphWrapper(ENamedThreads::AnyBackgroundThreadNormalTask);
+			Build(Task, TaskGraphWrapper);
+		};
+
+	// If we're already in worker threads there is no need to launch an async task.
+	if (GUseAsyncCardRepresentationBuildQueue || !IsInGameThread())
+	{
+		AsyncPool(*ThreadPool, MoveTemp(BuildLambda));
 	}
 	else
 	{
-		TUniquePtr<FQueuedThreadPool> WorkerThreadPool(CreateCardWorkerThreadPool());
-		Build(Task, *WorkerThreadPool);
+		BuildLambda();
 	}
 #else
 	UE_LOG(LogStaticMesh,Fatal,TEXT("Tried to build a card representation without editor support (this should have been done during cooking)"));
@@ -385,7 +269,7 @@ void FCardRepresentationAsyncQueue::BlockUntilAllBuildsComplete()
 	while (GetNumOutstandingTasks() > 0);
 }
 
-void FCardRepresentationAsyncQueue::Build(FAsyncCardRepresentationTask* Task, FQueuedThreadPool& ThreadPool)
+void FCardRepresentationAsyncQueue::Build(FAsyncCardRepresentationTask* Task, FQueuedThreadPool& BuildThreadPool)
 {
 #if WITH_EDITOR
 	
@@ -398,7 +282,7 @@ void FCardRepresentationAsyncQueue::Build(FAsyncCardRepresentationTask* Task, FQ
 			Task->StaticMesh->GetName(),
 			Task->SourceMeshData,
 			LODModel,
-			ThreadPool,
+			BuildThreadPool,
 			Task->GenerateSource->RenderData->Bounds,
 			Task->GenerateSource->RenderData->LODResources[0].DistanceFieldData,
 			*Task->GeneratedCardRepresentation);
@@ -420,19 +304,28 @@ void FCardRepresentationAsyncQueue::AddReferencedObjects(FReferenceCollector& Co
 	}
 }
 
-void FCardRepresentationAsyncQueue::ProcessAsyncTasks()
+FString FCardRepresentationAsyncQueue::GetReferencerName() const
+{
+	return TEXT("FCardRepresentationAsyncQueue");
+}
+
+void FCardRepresentationAsyncQueue::ProcessAsyncTasks(bool bLimitExecutionTime)
 {
 #if WITH_EDITOR
 	TRACE_CPUPROFILER_EVENT_SCOPE(FCardRepresentationAsyncQueue::ProcessAsyncTasks)
 
-	TArray<FAsyncCardRepresentationTask*> LocalCompletedTasks;
-	CompletedTasks.PopAll(LocalCompletedTasks);
-
-	for (int TaskIndex = 0; TaskIndex < LocalCompletedTasks.Num(); TaskIndex++)
+	const double MaxProcessingTime = 0.016f;
+	double StartTime = FPlatformTime::Seconds();
+	while (!bLimitExecutionTime || (FPlatformTime::Seconds() - StartTime) < MaxProcessingTime)
 	{
+		FAsyncCardRepresentationTask* Task = CompletedTasks.Pop();
+		if (Task == nullptr)
+		{
+			break;
+		}
+	
 		// We want to count each resource built from a DDC miss, so count each iteration of the loop separately.
 		COOK_STAT(auto Timer = CardRepresentationCookStats::UsageStats.TimeSyncWork());
-		FAsyncCardRepresentationTask* Task = LocalCompletedTasks[TaskIndex];
 
 		{
 			FScopeLock Lock(&CriticalSection);
@@ -467,32 +360,11 @@ void FCardRepresentationAsyncQueue::ProcessAsyncTasks()
 
 		delete Task;
 	}
-
-	bool bRemainingTasks = false;
-	{
-		FScopeLock Lock(&CriticalSection);
-		bRemainingTasks = ReferencedTasks.Num() > 0;
-	}
-
-	if (bRemainingTasks && !ThreadRunnable->IsRunning())
-	{
-		ThreadRunnable->Launch();
-	}
 #endif
 }
 
 void FCardRepresentationAsyncQueue::Shutdown()
 {
-	ThreadRunnable->Stop();
-	bool bLogged = false;
-
-	while (ThreadRunnable->IsRunning())
-	{
-		if (!bLogged)
-		{
-			bLogged = true;
-			UE_LOG(LogStaticMesh,Log,TEXT("Abandoning remaining async card representation tasks for shutdown"));
-		}
-		FPlatformProcess::Sleep(0.01f);
-	}
+	UE_LOG(LogStaticMesh, Log, TEXT("Abandoning remaining async card representation tasks for shutdown"));
+	ThreadPool->Destroy();
 }
\ No newline at end of file
diff --git a/Engine/Source/Runtime/Engine/Public/DistanceFieldAtlas.h b/Engine/Source/Runtime/Engine/Public/DistanceFieldAtlas.h
--- a/Engine/Source/Runtime/Engine/Public/DistanceFieldAtlas.h
+++ b/Engine/Source/Runtime/Engine/Public/DistanceFieldAtlas.h
@@ -412,7 +412,7 @@ public:
 	ENGINE_API void BlockUntilAllBuildsComplete();
 
 	/** Called once per frame, fetches completed tasks and applies them to the scene. */
-	ENGINE_API void ProcessAsyncTasks();
+	ENGINE_API void ProcessAsyncTasks(bool bLimitExecutionTime = false);
 
 	/** Exposes UObject references used by the async build. */
 	ENGINE_API virtual void AddReferencedObjects(FReferenceCollector& Collector) override;
@@ -425,33 +425,26 @@ public:
 
 	int32 GetNumOutstandingTasks() const
 	{
+		FScopeLock Lock(&CriticalSection);
 		return ReferencedTasks.Num();
 	}
 
 private:
+	TUniquePtr<FQueuedThreadPool> ThreadPool;
 
 	/** Builds a single task with the given threadpool.  Called from the worker thread. */
 	void Build(FAsyncDistanceFieldTask* Task, class FQueuedThreadPool& ThreadPool);
 
-	/** Thread that will build any tasks in TaskQueue and exit when there are no more. */
-	class TUniquePtr<class FBuildDistanceFieldThreadRunnable> ThreadRunnable;
-
 	/** Game-thread managed list of tasks in the async system. */
 	TArray<FAsyncDistanceFieldTask*> ReferencedTasks;
 
-	/** Tasks that have not yet started processing yet. */
-	// consider changing this from FIFO to Unordered, which may be faster
-	TLockFreePointerListLIFO<FAsyncDistanceFieldTask> TaskQueue;
-
 	/** Tasks that have completed processing. */
 	// consider changing this from FIFO to Unordered, which may be faster
 	TLockFreePointerListLIFO<FAsyncDistanceFieldTask> CompletedTasks;
 
 	class IMeshUtilities* MeshUtilities;
 
-	FCriticalSection CriticalSection;
-
-	friend class FBuildDistanceFieldThreadRunnable;
+	mutable FCriticalSection CriticalSection;
 };
 
 /** Global build queue. */
diff --git a/Engine/Source/Runtime/Engine/Public/MeshCardRepresentation.h b/Engine/Source/Runtime/Engine/Public/MeshCardRepresentation.h
--- a/Engine/Source/Runtime/Engine/Public/MeshCardRepresentation.h
+++ b/Engine/Source/Runtime/Engine/Public/MeshCardRepresentation.h
@@ -188,11 +188,13 @@ public:
 	ENGINE_API void BlockUntilAllBuildsComplete();
 
 	/** Called once per frame, fetches completed tasks and applies them to the scene. */
-	ENGINE_API void ProcessAsyncTasks();
+	ENGINE_API void ProcessAsyncTasks(bool bLimitExecutionTime = false);
 
 	/** Exposes UObject references used by the async build. */
 	ENGINE_API void AddReferencedObjects(FReferenceCollector& Collector);
 
+	ENGINE_API virtual FString GetReferencerName() const override;
+
 	/** Blocks until it is safe to shut down (worker threads are idle). */
 	ENGINE_API void Shutdown();
 
@@ -203,27 +205,20 @@ public:
 	}
 
 private:
+	TUniquePtr<FQueuedThreadPool> ThreadPool;
 
 	/** Builds a single task with the given threadpool.  Called from the worker thread. */
 	void Build(FAsyncCardRepresentationTask* Task, class FQueuedThreadPool& ThreadPool);
 
-	/** Thread that will build any tasks in TaskQueue and exit when there are no more. */
-	class TUniquePtr<class FBuildCardRepresentationThreadRunnable> ThreadRunnable;
-
 	/** Game-thread managed list of tasks in the async system. */
 	TArray<FAsyncCardRepresentationTask*> ReferencedTasks;
 
-	/** Tasks that have not yet started processing yet. */
-	TLockFreePointerListLIFO<FAsyncCardRepresentationTask> TaskQueue;
-
 	/** Tasks that have completed processing. */
 	TLockFreePointerListLIFO<FAsyncCardRepresentationTask> CompletedTasks;
 
 	class IMeshUtilities* MeshUtilities;
 
 	mutable FCriticalSection CriticalSection;
-
-	friend class FBuildCardRepresentationThreadRunnable;
 };
 
 /** Global build queue. */
diff --git a/Engine/Source/Runtime/Launch/Private/LaunchEngineLoop.cpp b/Engine/Source/Runtime/Launch/Private/LaunchEngineLoop.cpp
--- a/Engine/Source/Runtime/Launch/Private/LaunchEngineLoop.cpp
+++ b/Engine/Source/Runtime/Launch/Private/LaunchEngineLoop.cpp
@@ -4897,13 +4897,13 @@ void FEngineLoop::Tick()
 		if (GDistanceFieldAsyncQueue)
 		{
 			QUICK_SCOPE_CYCLE_COUNTER(STAT_FEngineLoop_Tick_GDistanceFieldAsyncQueue);
-			GDistanceFieldAsyncQueue->ProcessAsyncTasks();
+			GDistanceFieldAsyncQueue->ProcessAsyncTasks(true);
 		}
 		
 		if (GCardRepresentationAsyncQueue)
 		{
 			QUICK_SCOPE_CYCLE_COUNTER(STAT_FEngineLoop_Tick_GCardRepresentationAsyncQueue);
-			GCardRepresentationAsyncQueue->ProcessAsyncTasks();
+			GCardRepresentationAsyncQueue->ProcessAsyncTasks(true);
 		}
 #if !UE_SERVER
 		// tick media framework
@@ -5746,6 +5746,13 @@ void FEngineLoop::AppPreExit( )
 
 	FCoreDelegates::OnExit.Broadcast();
 
+	// Clean up the thread pool
+	// GThreadPool might be a wrapper around GLargeThreadPool so we have to destroy it first
+	if (GThreadPool != nullptr)
+	{
+		GThreadPool->Destroy();
+	}
+
 #if WITH_EDITOR
 	if (GLargeThreadPool != nullptr)
 	{
@@ -5753,12 +5760,6 @@ void FEngineLoop::AppPreExit( )
 	}
 #endif // WITH_EDITOR
 
-	// Clean up the thread pool
-	if (GThreadPool != nullptr)
-	{
-		GThreadPool->Destroy();
-	}
-
 	if (GBackgroundPriorityThreadPool != nullptr)
 	{
 		GBackgroundPriorityThreadPool->Destroy();
