Fix incorrect check failures in GPU particle simulation.

diff --git a/Engine/Source/Runtime/Engine/Private/Particles/ParticleGpuSimulation.cpp b/Engine/Source/Runtime/Engine/Private/Particles/ParticleGpuSimulation.cpp
--- a/Engine/Source/Runtime/Engine/Private/Particles/ParticleGpuSimulation.cpp
+++ b/Engine/Source/Runtime/Engine/Private/Particles/ParticleGpuSimulation.cpp
@@ -1064,8 +1064,6 @@ public:
 		FRHIUniformBuffer* SceneTexturesUniformBuffer
 		)
 	{
-		check(ViewUniformBuffer != NULL);
-
 		FRHIPixelShader* PixelShaderRHI = RHICmdList.GetBoundPixelShader();
 		FRHISamplerState* SamplerStatePoint = TStaticSamplerState<SF_Point>::GetRHI();
 		FRHISamplerState* SamplerStateLinear = TStaticSamplerState<SF_Bilinear,AM_Clamp,AM_Clamp,AM_Clamp>::GetRHI();
@@ -1074,10 +1072,11 @@ public:
 		SetTextureParameter(RHICmdList, PixelShaderRHI, AttributesTexture, AttributesTextureSampler, SamplerStatePoint, InAttributesTexture.TextureRHI);
 		SetTextureParameter(RHICmdList, PixelShaderRHI, CurveTexture, CurveTextureSampler, SamplerStateLinear, GParticleCurveTexture.GetCurveTexture());
 
-		FGlobalShader::SetParameters<FViewUniformShaderParameters>(RHICmdList, PixelShaderRHI, ViewUniformBuffer);
-
 		if (CollisionMode == PCM_DepthBuffer)
 		{
+			check(ViewUniformBuffer != NULL);
+			FGlobalShader::SetParameters<FViewUniformShaderParameters>(RHICmdList, PixelShaderRHI, ViewUniformBuffer);
+
 			SetUniformBufferParameter(RHICmdList, PixelShaderRHI, GetUniformBufferParameter(SceneTexturesUniformBufferStruct), SceneTexturesUniformBuffer);
 			SetTextureParameter(
 				RHICmdList, 
@@ -1091,6 +1090,9 @@ public:
 		}
 		else if (CollisionMode == PCM_DistanceField)
 		{
+			check(ViewUniformBuffer != NULL);
+			FGlobalShader::SetParameters<FViewUniformShaderParameters>(RHICmdList, PixelShaderRHI, ViewUniformBuffer);
+
 			GlobalDistanceFieldParameters.Set(RHICmdList, PixelShaderRHI, *GlobalDistanceFieldParameterData);
 
 			SetTextureParameter(
