ShaderCompiler timing additions. If logging timing, report the full debug name of the shader (Vertex Factory + Shader Type). Report if it took too long to launch workers, this should help identify situations where the filesystem is the bottleneck.

diff --git a/Engine/Source/Runtime/Engine/Private/ShaderCompiler/ShaderCompiler.cpp b/Engine/Source/Runtime/Engine/Private/ShaderCompiler/ShaderCompiler.cpp
--- a/Engine/Source/Runtime/Engine/Private/ShaderCompiler/ShaderCompiler.cpp
+++ b/Engine/Source/Runtime/Engine/Private/ShaderCompiler/ShaderCompiler.cpp
@@ -1322,7 +1322,8 @@ int32 FShaderCompileThreadRunnable::PullTasksFromQueue()
 							const FShaderCommonCompileJob& Job = *CurrentWorkerInfo.QueuedJobs[JobIndex];
 							if (auto* SingleJob = Job.GetSingleShaderJob())
 							{
-								JobNames += FString::Printf(TEXT("%s [Instructions=%d WorkerTime=%.3fs]"), SingleJob->ShaderType->GetName(), SingleJob->Output.NumInstructions, SingleJob->Output.CompileTime);
+								const TCHAR* JobName = Manager->bLogJobCompletionTimes ? *SingleJob->Input.DebugGroupName : SingleJob->ShaderType->GetName();
+								JobNames += FString::Printf(TEXT("%s [Instructions=%d WorkerTime=%.3fs]"), JobName, SingleJob->Output.NumInstructions, SingleJob->Output.CompileTime);
 							}
 							else
 							{
@@ -1433,6 +1434,7 @@ bool FShaderCompileThreadRunnable::LaunchWorkersIfNeeded()
 	// Limit how often we check for workers running since IsApplicationRunning eats up some CPU time on Windows
 	const bool bCheckForWorkerRunning = (CurrentTime - LastCheckForWorkersTime > .1f);
 	bool bAbandonWorkers = false;
+	uint32_t NumberLaunched = 0;
 
 	if (bCheckForWorkerRunning)
 	{
@@ -1497,10 +1499,18 @@ bool FShaderCompileThreadRunnable::LaunchWorkersIfNeeded()
 				// Store the handle with this thread so that we will know not to launch it again
 				CurrentWorkerInfo.WorkerProcess = Manager->LaunchWorker(WorkingDirectory, Manager->ProcessId, WorkerIndex, InputFileName, OutputFileName);
 				CurrentWorkerInfo.bLaunchedWorker = true;
+
+				NumberLaunched++;
 			}
 		}
 	}
 
+	const double FinishTime = FPlatformTime::Seconds();
+	if (NumberLaunched > 0 && (FinishTime - CurrentTime) >= 10.0)
+	{
+		UE_LOG(LogShaderCompilers, Warning, TEXT("Performance Warning: It took %f seconds to launch %d ShaderCompileWorkers"), FinishTime - CurrentTime, NumberLaunched);
+	}
+
 	return bAbandonWorkers;
 }
 
