Cloth - Allow the most appropriate cloth simulations factories (e.g. Chaos) to replace the one that are currently set, if found unable to provide a fully functional simulation for each of the SkeletalMesh's clothing asset.
Chaos Cloth - Fail gracefully with a log message when bCollideWithEnvironment is set.
[Merging CL 14427768]

diff --git a/Engine/Source/Runtime/Engine/Private/Components/SkeletalMeshComponent.cpp b/Engine/Source/Runtime/Engine/Private/Components/SkeletalMeshComponent.cpp
--- a/Engine/Source/Runtime/Engine/Private/Components/SkeletalMeshComponent.cpp
+++ b/Engine/Source/Runtime/Engine/Private/Components/SkeletalMeshComponent.cpp
@@ -544,14 +544,82 @@ void USkeletalMeshComponent::OnRegister()
 	}
 
 #if WITH_APEX_CLOTHING || WITH_CHAOS_CLOTHING
-	// If we don't have a valid simulation factory - check to see if we have an available default to use instead
-	if (*ClothingSimulationFactory == nullptr)
+	// If no simulation factory is currently set, set it to the default factory
+	if (!ClothingSimulationFactory)
 	{
 		ClothingSimulationFactory = UClothingSimulationFactory::GetDefaultClothingSimulationFactoryClass();
 	}
 
+	// Look up for the best simulation factory to support each asset
+	if (SkeletalMesh)
+	{
+		// Check whether all clothing assets are supported by the current factory
+		bool bSupportsAllAssets = true;
+
+		UClothingSimulationFactory* const DefaultObject = ClothingSimulationFactory->GetDefaultObject<UClothingSimulationFactory>();
+		for (UClothingAssetBase* const ClothingAsset : SkeletalMesh->MeshClothingAssets)
+		{
+			if (ClothingAsset && !DefaultObject->SupportsAsset(ClothingAsset))
+			{
+				bSupportsAllAssets = false;
+
+				UE_LOG(LogSkeletalMesh, Warning,
+					TEXT("OnRegister[%s]: [%s] is currently unable to provide a fully functional simulation for each of this SkeletalMesh's clothing assets."),
+					*GetNameSafe(SkeletalMesh),
+					*ClothingSimulationFactory->GetName());
+				UE_LOG(LogSkeletalMesh, Warning,
+					TEXT("OnRegister[%s]: The ClothingSimulationFactory property will now be automatically updated to use the most functional simulation that can be found."),
+					*GetNameSafe(SkeletalMesh));
+
+				break;
+			}
+		}
+
+		// Try to find a new clothing factory that matches most asset requirements
+		if (!bSupportsAllAssets)
+		{
+			int MostSupportedNumAssets = 0;
+
+			const TArray<IClothingSimulationFactoryClassProvider*> ClassProviders = IModularFeatures::Get().GetModularFeatureImplementations<IClothingSimulationFactoryClassProvider>(IClothingSimulationFactoryClassProvider::FeatureName);
+			for (const IClothingSimulationFactoryClassProvider* const ClassProvider : ClassProviders)
+			{
+				if (ClassProvider)
+				{
+					if (const TSubclassOf<UClothingSimulationFactory> NewClothingSimulationFactory = ClassProvider->GetClothingSimulationFactoryClass())
+					{
+						int SupportedNumAssets = 0;
+						UClothingSimulationFactory* const NewDefaultObject = NewClothingSimulationFactory->GetDefaultObject<UClothingSimulationFactory>();
+						for (UClothingAssetBase* const ClothingAsset : SkeletalMesh->MeshClothingAssets)
+						{
+							if (ClothingAsset && NewDefaultObject->SupportsAsset(ClothingAsset))
+							{
+								++SupportedNumAssets;
+							}
+						}
+
+						if (SupportedNumAssets > MostSupportedNumAssets)
+						{
+							ClothingSimulationFactory = NewClothingSimulationFactory;
+							MostSupportedNumAssets = SupportedNumAssets;
+						}
+					}
+				}
+			}
+		}
+	}
+
 	RecreateClothingActors();
-#endif
+#endif  // #if WITH_APEX_CLOTHING || WITH_CHAOS_CLOTHING
+
+#if WITH_CHAOS_CLOTHING
+	// TODO: Add missing Chaos Cloth collision with environment
+	if (bCollideWithEnvironment && ClothingSimulationFactory && ClothingSimulationFactory->GetName() == TEXT("ChaosClothingSimulationFactory"))
+	{
+		UE_LOG(LogSkeletalMesh, Warning, TEXT("OnRegister[%s]: "
+			"Chaos Cloth does not currently support bCollideWithEnvironment."),
+			*GetNameSafe(SkeletalMesh));
+	}
+#endif  // #if WITH_CHAOS_CLOTHING
 }
 
 void USkeletalMeshComponent::OnUnregister()
diff --git a/Engine/Source/Runtime/Engine/Private/SkeletalMeshComponentPhysics.cpp b/Engine/Source/Runtime/Engine/Private/SkeletalMeshComponentPhysics.cpp
--- a/Engine/Source/Runtime/Engine/Private/SkeletalMeshComponentPhysics.cpp
+++ b/Engine/Source/Runtime/Engine/Private/SkeletalMeshComponentPhysics.cpp
@@ -2903,7 +2903,7 @@ void USkeletalMeshComponent::ProcessClothCollisionWithEnvironment()
 						FMatrix ShapeLocalPose = FPhysicsInterface::GetLocalTransform(Shape).ToMatrixWithScale();
 
 #if WITH_CHAOS
-                        check(false);
+						// TODO: Add missing Chaos Cloth collision with environment
 #else
 						switch(GeoType)
 						{
