Bumped static mesh MeshDescription DDC key to enable new versioning.

diff --git a/Engine/Source/Runtime/Engine/Private/StaticMesh.cpp b/Engine/Source/Runtime/Engine/Private/StaticMesh.cpp
--- a/Engine/Source/Runtime/Engine/Private/StaticMesh.cpp
+++ b/Engine/Source/Runtime/Engine/Private/StaticMesh.cpp
@@ -4270,7 +4270,7 @@ void UStaticMesh::ClearMeshDescriptions()
 // differences, etc.) replace the version GUID below with a new one.
 // In case of merge conflicts with DDC versions, you MUST generate a new GUID
 // and set this new GUID as the version.                                       
-#define MESHDATAKEY_STATICMESH_DERIVEDDATA_VER TEXT("BD8BD3653BF64EED96EE5366C4F82B25")
+#define MESHDATAKEY_STATICMESH_DERIVEDDATA_VER TEXT("BEDE1E3CC35D4D8FA944458825A18AFB")
 
 
 static const FString& GetMeshDataKeyStaticMeshDerivedDataVersion()
