- Fixed ISM scalability's ViewDistanceScale factor applied twice on End Cull Distance. Also was applied on screen size calculation. This was caused by a bad integrate in 2017 (original fix was made in CL 3263554 and was broken by bad merge in CL3352858).

diff --git a/Engine/Source/Runtime/Engine/Private/InstancedStaticMesh.cpp b/Engine/Source/Runtime/Engine/Private/InstancedStaticMesh.cpp
--- a/Engine/Source/Runtime/Engine/Private/InstancedStaticMesh.cpp
+++ b/Engine/Source/Runtime/Engine/Private/InstancedStaticMesh.cpp
@@ -3307,7 +3307,6 @@ void FInstancedStaticMeshVertexFactoryShaderParameters::GetElementShaderBindings
 				{
 					FinalCull = FMath::Min(FinalCull, InstancingUserData->EndCullDistance * MaxDrawDistanceScale);
 				}
-				FinalCull *= MaxDrawDistanceScale;
 
 				InstancingViewZCompare.Z = FinalCull;
 				if (int(BatchElement.InstancedLODIndex) < InstancingUserData->MeshRenderData->LODResources.Num() - 1)
