Support adding compute framework graphs for batch execution.

diff --git a/Engine/Source/Runtime/Engine/Private/ComputeFramework/ComputeFramework.cpp b/Engine/Source/Runtime/Engine/Private/ComputeFramework/ComputeFramework.cpp
--- a/Engine/Source/Runtime/Engine/Private/ComputeFramework/ComputeFramework.cpp
+++ b/Engine/Source/Runtime/Engine/Private/ComputeFramework/ComputeFramework.cpp
@@ -1,10 +1,16 @@
 // Copyright Epic Games, Inc. All Rights Reserved.
 
 #include "ComputeFramework/ComputeFramework.h"
+#include "ComputeFramework/ComputeKernel.h"
 #include "ComputeFramework/ComputeGraph.h"
 #include "RenderGraphBuilder.h"
 #include "RenderGraphUtils.h"
 #include "RHIDefinitions.h"
+#include "ComputeFramework/ComputeKernelShaderMap.h"
+#include "ComputeFramework/ComputeKernelShader.h"
+#include "ProfilingDebugging/RealtimeGPUProfiler.h"
+
+DECLARE_GPU_STAT_NAMED(ComputeFramework_ExecuteBatches, TEXT("ComputeFramework::ExecuteBatches"));
 
 static int32 GComputeFrameworkMode = 0;
 static FAutoConsoleVariableRef CVarComputeFrameworkMode(
@@ -23,6 +29,41 @@ bool SupportsComputeFramework(ERHIFeatureLevel::Type FeatureLevel, EShaderPlatfo
 		&& FDataDrivenShaderPlatformInfo::GetSupportsComputeFramework(ShaderPlatform);
 }
 
+void FComputeGraph::Initialize(UComputeGraph* ComputeGraph)
+{
+	for (auto& Invocation : ComputeGraph->GetShaderInvocationList())
+	{
+		FKernelInvocation KernelInvocation = {
+			Invocation.ComputeKernel ? Invocation.ComputeKernel->GetFName() : "NULL Kernel",
+			FName("InvocationName"),
+			FIntVector(1, 1, 1),
+			Invocation.ComputeKernel ? Invocation.ComputeKernel->GetResource() : nullptr
+			};
+
+		KernelInvocations.Emplace(KernelInvocation);
+	}
+}
+
+void FComputeFramework::EnqueueForExecution(const FComputeGraph* ComputeGraph)
+{
+	//ExternalBuffers.Append(ComputeGraph.ExternalBuffers);
+	//TransientBuffers.Append(ComputeGraph.TransientBuffers);
+
+	for (auto& Invocation : ComputeGraph->KernelInvocations)
+	{
+		FShaderInvocation ShaderInvocation = {
+			Invocation.KernelName,
+			Invocation.InvocationName,
+			Invocation.DispatchDim,
+			Invocation.KernelResource ? 
+				Invocation.KernelResource->ShaderMap_RT->GetShader<FComputeKernelShader>(0) : 
+				GetGlobalShaderMap(ERHIFeatureLevel::SM5)->GetShader<FComputeKernelShader>()
+			};
+
+		ComputeShaders.Emplace(ShaderInvocation);
+	}
+}
+
 struct FComputeExecutionBuffer
 {
 	FComputeExecutionBuffer() = default;
@@ -41,70 +82,84 @@ void FComputeFramework::ExecuteBatches(
 	FRHICommandListImmediate& RHICmdList
 	)
 {
-	FRDGBuilder GraphBuilder(RHICmdList);
-
-	TArray<FComputeExecutionBuffer> ExecutionBuffers;
-
-	for (auto& ExtRes : ExternalBuffers)
+	if (ComputeShaders.IsEmpty())
 	{
-		TCHAR ResName[FComputeResource::MAX_NAME_LENGTH];
-		ExtRes.Name.ToString(ResName);
-
-		ExecutionBuffers.Emplace(
-			ExtRes.Name,
-			GraphBuilder.RegisterExternalBuffer(ExtRes.Buffer, ResName)
-			);
+		return;
 	}
 
-	for (auto& TransRes : TransientBuffers)
+	FMemMark MemStackMark(FMemStack::Get());
 	{
-		TCHAR ResName[FComputeResource::MAX_NAME_LENGTH];
-		TransRes.Name.ToString(ResName);
+		SCOPED_DRAW_EVENTF(RHICmdList, ComputeFramework_ExecuteBatches, TEXT("ComputeFramework::ExecuteBatches"));
+		SCOPED_GPU_STAT(RHICmdList, ComputeFramework_ExecuteBatches);
 
-		FRDGBufferDesc BufferDesc;
-		
-		EComputeExecutionBufferType BufferType = (EComputeExecutionBufferType)TransRes.BufferType;
-		switch (BufferType)
+		FRDGBuilder GraphBuilder(RHICmdList);
+
+		TArray<FComputeExecutionBuffer> ExecutionBuffers;
+
+		for (auto& ExtRes : ExternalBuffers)
 		{
-		case EComputeExecutionBufferType::Buffer:
-			BufferDesc = FRDGBufferDesc::CreateBufferDesc(TransRes.BytesPerElement, TransRes.ElementCount);
-			break;
+			TCHAR ResName[FComputeResource::MAX_NAME_LENGTH];
+			ExtRes.Name.ToString(ResName);
 
-		case EComputeExecutionBufferType::StructuredBuffer:
-			BufferDesc = FRDGBufferDesc::CreateStructuredDesc(TransRes.BytesPerElement, TransRes.ElementCount);
-			break;
+			ExecutionBuffers.Emplace(
+				ExtRes.Name,
+				GraphBuilder.RegisterExternalBuffer(ExtRes.Buffer, ResName)
+				);
+		}
 
-		case EComputeExecutionBufferType::ByteAddressBuffer:
-			BufferDesc = FRDGBufferDesc::CreateByteAddressDesc(TransRes.BytesPerElement * TransRes.ElementCount);
-			break;
+		for (auto& TransRes : TransientBuffers)
+		{
+			TCHAR ResName[FComputeResource::MAX_NAME_LENGTH];
+			TransRes.Name.ToString(ResName);
 
-		default:
-			check(!"Unsupported buffer type");
-			break;
+			FRDGBufferDesc BufferDesc;
+		
+			EComputeExecutionBufferType BufferType = (EComputeExecutionBufferType)TransRes.BufferType;
+			switch (BufferType)
+			{
+			case EComputeExecutionBufferType::Buffer:
+				BufferDesc = FRDGBufferDesc::CreateBufferDesc(TransRes.BytesPerElement, TransRes.ElementCount);
+				break;
+
+			case EComputeExecutionBufferType::StructuredBuffer:
+				BufferDesc = FRDGBufferDesc::CreateStructuredDesc(TransRes.BytesPerElement, TransRes.ElementCount);
+				break;
+
+			case EComputeExecutionBufferType::ByteAddressBuffer:
+				BufferDesc = FRDGBufferDesc::CreateByteAddressDesc(TransRes.BytesPerElement * TransRes.ElementCount);
+				break;
+
+			default:
+				check(!"Unsupported buffer type");
+				break;
+			}
+
+			ExecutionBuffers.Emplace(
+				TransRes.Name,
+				GraphBuilder.CreateBuffer(BufferDesc, ResName)
+				);
 		}
 
-		ExecutionBuffers.Emplace(
-			TransRes.Name,
-			GraphBuilder.CreateBuffer(BufferDesc, ResName)
-			);
-	}
+		for (auto& Compute : ComputeShaders)
+		{
+			TCHAR KernelName[FComputeKernelInvocation::MAX_NAME_LENGTH];
+			Compute.KernelName.ToString(KernelName);
+
+			TCHAR InvocationName[FComputeKernelInvocation::MAX_NAME_LENGTH];
+			Compute.InvocationName.ToString(InvocationName);
+
+			auto* Params = GraphBuilder.AllocParameters<FComputeKernelShader::FParameters>();
+
+			FComputeShaderUtils::AddPass(
+				GraphBuilder,
+				RDG_EVENT_NAME("Compute[%s]: %s", KernelName, InvocationName),
+				ERDGPassFlags::Compute,
+				Compute.Shader,
+				Params,
+				Compute.DispatchDim
+				);
+		}
 
-	for (auto& Compute : ComputeShaders)
-	{
-		TCHAR KernelName[FComputeKernelInvocation::MAX_NAME_LENGTH];
-		Compute.KernelName.ToString(KernelName);
-
-		TCHAR InvocationName[FComputeKernelInvocation::MAX_NAME_LENGTH];
-		Compute.InvocationName.ToString(InvocationName);
-
-		/*
-		FComputeShaderUtils::AddPass(
-			GraphBuilder,
-			RDG_EVENT_NAME("Compute[%s]: %s", KernelName, InvocationName),
-			Compute.Shader,
-			{},
-			Compute.DispatchDim
-			);
-		*/
+		GraphBuilder.Execute();
 	}
 }
diff --git a/Engine/Source/Runtime/Engine/Private/ComputeFramework/ComputeGraph.cpp b/Engine/Source/Runtime/Engine/Private/ComputeFramework/ComputeGraph.cpp
--- a/Engine/Source/Runtime/Engine/Private/ComputeFramework/ComputeGraph.cpp
+++ b/Engine/Source/Runtime/Engine/Private/ComputeFramework/ComputeGraph.cpp
@@ -14,3 +14,11 @@ bool UComputeGraph::RemoveKernel(FComputeKernelInvocationHandle KernelInvocation
 {
 	return false;
 }
+
+void UComputeGraph::PostLoad()
+{
+	Super::PostLoad();
+
+	KernelInvocations.Reset();
+	KernelInvocations.Emplace();
+}
diff --git a/Engine/Source/Runtime/Engine/Private/ComputeFramework/ComputeGraphComponent.cpp b/Engine/Source/Runtime/Engine/Private/ComputeFramework/ComputeGraphComponent.cpp
new file mode 100644
--- /dev/null
+++ b/Engine/Source/Runtime/Engine/Private/ComputeFramework/ComputeGraphComponent.cpp
@@ -0,0 +1,39 @@
+// Copyright Epic Games, Inc. All Rights Reserved.
+
+#include "ComputeFramework/ComputeGraphComponent.h"
+#include "ComputeFramework/ComputeGraph.h"
+#include "ComputeFramework/ComputeFramework.h"
+
+
+UComputeGraphComponent::UComputeGraphComponent()
+{
+	PrimaryComponentTick.bCanEverTick = true;
+	PrimaryComponentTick.bStartWithTickEnabled = false;
+}
+
+void UComputeGraphComponent::QueueExecute()
+{
+	if (!ComputeGraph)
+	{
+		return;
+	}
+
+	MarkRenderDynamicDataDirty();
+}
+
+void UComputeGraphComponent::SendRenderDynamicData_Concurrent()
+{
+	check(ComputeGraph);
+	Super::SendRenderDynamicData_Concurrent();
+
+	FComputeFramework* ComputeFramework = GetScene()->GetComputeFramework();
+	FComputeGraph* ComputeGraphProxy = new FComputeGraph();
+	ComputeGraphProxy->Initialize(ComputeGraph);
+
+	ENQUEUE_RENDER_COMMAND(ComputeFrameworkEnqueueExecutionCommand)(
+		[ComputeFramework, ComputeGraphProxy](FRHICommandListImmediate& RHICmdList)
+		{
+			ComputeFramework->EnqueueForExecution(ComputeGraphProxy);
+			delete ComputeGraphProxy;
+		});
+}
diff --git a/Engine/Source/Runtime/Engine/Private/ComputeFramework/ComputeKernelShaderMap.cpp b/Engine/Source/Runtime/Engine/Private/ComputeFramework/ComputeKernelShaderMap.cpp
--- a/Engine/Source/Runtime/Engine/Private/ComputeFramework/ComputeKernelShaderMap.cpp
+++ b/Engine/Source/Runtime/Engine/Private/ComputeFramework/ComputeKernelShaderMap.cpp
@@ -2,7 +2,7 @@
 
 #include "ComputeFramework/ComputeKernelShaderMap.h"
 #include "ShaderCompiler.h"
-#include "ComputeFramework/ComputeKernelShaderType.h"
+#include "RHI.h"
 #include "ComputeFramework/ComputeKernelResource.h"
 #include "ComputeFramework/ComputeKernel.h"
 
diff --git a/Engine/Source/Runtime/Engine/Public/ComputeFramework/ComputeFramework.h b/Engine/Source/Runtime/Engine/Public/ComputeFramework/ComputeFramework.h
--- a/Engine/Source/Runtime/Engine/Public/ComputeFramework/ComputeFramework.h
+++ b/Engine/Source/Runtime/Engine/Public/ComputeFramework/ComputeFramework.h
@@ -2,20 +2,15 @@
 
 #pragma once
 
-#include "RHIDefinitions.h"
-#include "Shader.h"
 #include "Templates/RefCounting.h"
 
+class FRDGPooledBuffer;
+class FComputeKernelResource;
 class FComputeKernelShader;
 
-extern ENGINE_API bool SupportsComputeFramework(ERHIFeatureLevel::Type FeatureLevel, EShaderPlatform ShaderPlatform);
+class UComputeGraph;
 
-class FRDGPooledBuffer;
-struct FComputeExecutionExternalBufferDesc
-{
-	FName Name;
-	TRefCountPtr<FRDGPooledBuffer> Buffer;
-};
+extern ENGINE_API bool SupportsComputeFramework(ERHIFeatureLevel::Type FeatureLevel, EShaderPlatform ShaderPlatform);
 
 enum class EComputeExecutionBufferType
 {
@@ -34,17 +29,43 @@ struct FComputeExecutionBufferDesc
 	uint32 ElementCount = 0;
 };
 
-struct FComputeExecutionShader
+struct FComputeExecutionExternalBufferDesc
 {
-	FName KernelName;
-	FName InvocationName;
-	FIntVector DispatchDim;
-	TShaderRef<FComputeKernelShader> Shader;
+	FName Name;
+	TRefCountPtr<FRDGPooledBuffer> Buffer;
+};
+
+class FComputeGraph
+{
+public:
+	struct FKernelInvocation
+	{
+		FName KernelName;
+		FName InvocationName;
+		FIntVector DispatchDim;
+		FComputeKernelResource* KernelResource = nullptr;
+	};
+
+	void Initialize(UComputeGraph* ComputeGraph);
+
+	TArray<FKernelInvocation> KernelInvocations;
 };
 
 class FComputeFramework
 {
 public:
+	struct FShaderInvocation
+	{
+		FName KernelName;
+		FName InvocationName;
+		FIntVector DispatchDim;
+		TShaderRef<FComputeKernelShader> Shader;
+	};
+
+	void EnqueueForExecution(
+		const FComputeGraph* ComputeGraph
+		);
+
 	void ExecuteBatches(
 		FRHICommandListImmediate& RHICmdList
 		);
@@ -52,5 +73,5 @@ public:
 private:
 	TArray<FComputeExecutionExternalBufferDesc> ExternalBuffers;
 	TArray<FComputeExecutionBufferDesc> TransientBuffers;
-	TArray<FComputeExecutionShader> ComputeShaders;
+	TArray<FShaderInvocation> ComputeShaders;
 };
diff --git a/Engine/Source/Runtime/Engine/Public/ComputeFramework/ComputeGraph.h b/Engine/Source/Runtime/Engine/Public/ComputeFramework/ComputeGraph.h
--- a/Engine/Source/Runtime/Engine/Public/ComputeFramework/ComputeGraph.h
+++ b/Engine/Source/Runtime/Engine/Public/ComputeFramework/ComputeGraph.h
@@ -126,7 +126,7 @@ struct FKernelBinding
  * scheduler will extract parallel/serial execution from this DAG description.
  */
 UCLASS()
-class UComputeGraph : public UObject
+class ENGINE_API UComputeGraph : public UObject
 {
 	GENERATED_BODY()
 
@@ -168,6 +168,12 @@ public:
 
 	void RemoveBinding(FComputeBindingHandle Binding);
 
+	TArray<FComputeKernelInvocation> GetShaderInvocationList() { return KernelInvocations; }
+
+
+	// TEMP
+	void PostLoad() override;
+
 private:
 	/* Represents all the invocation nodes of the graph. */
 	UPROPERTY()
diff --git a/Engine/Source/Runtime/Engine/Public/ComputeFramework/ComputeGraphComponent.h b/Engine/Source/Runtime/Engine/Public/ComputeFramework/ComputeGraphComponent.h
new file mode 100644
--- /dev/null
+++ b/Engine/Source/Runtime/Engine/Public/ComputeFramework/ComputeGraphComponent.h
@@ -0,0 +1,38 @@
+// Copyright Epic Games, Inc. All Rights Reserved.
+
+#pragma once
+
+#include "CoreMinimal.h"
+#include "UObject/ObjectMacros.h"
+#include "UObject/UObjectGlobals.h"
+#include "Components/ActorComponent.h"
+
+#include "ComputeGraphComponent.generated.h"
+
+class UComputeGraph;
+
+
+UCLASS(Blueprintable, meta = (BlueprintSpawnableComponent))
+class UComputeGraphComponent : public UActorComponent
+{
+	GENERATED_BODY()
+
+public:
+	UComputeGraphComponent();
+
+	UPROPERTY(EditAnywhere, Category = "Compute")
+	UComputeGraph* ComputeGraph = nullptr;
+
+	UFUNCTION(BlueprintCallable, Category = "Compute")
+	void QueueExecute();
+
+protected:
+	//~ Begin UActorComponent Interface
+	bool ShouldCreateRenderState() const override
+	{
+		return true;
+	}
+
+	void SendRenderDynamicData_Concurrent() override;
+	//~ End UActorComponent Interface
+};
diff --git a/Engine/Source/Runtime/Engine/Public/ComputeFramework/ComputeKernel.h b/Engine/Source/Runtime/Engine/Public/ComputeFramework/ComputeKernel.h
--- a/Engine/Source/Runtime/Engine/Public/ComputeFramework/ComputeKernel.h
+++ b/Engine/Source/Runtime/Engine/Public/ComputeFramework/ComputeKernel.h
@@ -124,6 +124,8 @@ public:
 #endif
 	//=========================================================================
 
+	FComputeKernelResource* GetResource() { return KernelResource.Get(); }
+
 private:
 #if WITH_EDITOR
 	/* Entry point to initiate any [re]compilation necessary due to changes. */
diff --git a/Engine/Source/Runtime/Engine/Public/ComputeFramework/ComputeKernelResource.h b/Engine/Source/Runtime/Engine/Public/ComputeFramework/ComputeKernelResource.h
--- a/Engine/Source/Runtime/Engine/Public/ComputeFramework/ComputeKernelResource.h
+++ b/Engine/Source/Runtime/Engine/Public/ComputeFramework/ComputeKernelResource.h
@@ -3,6 +3,9 @@
 #pragma once
 
 #include "ComputeFramework/ComputeKernelShaderMap.h"
+#include "ComputeFramework/ComputeKernelShader.h"
+
+class FComputeKernelShader;
 
 struct FComputeKernelCompilationResults
 {
@@ -30,7 +33,11 @@ public:
 		);
 #endif
 
-private:
+	TShaderRef<FComputeKernelShader> GetShader(int32 PermutationId) const
+	{
+		return ShaderMap_RT->GetShader<FComputeKernelShader>(PermutationId);
+	}
+
 	FComputeKernelCompilationResults CompilationResults;
 
 	/*
@@ -47,6 +54,7 @@ private:
 	 */
 	FComputeKernelShaderMap* ShaderMap_RT;
 
+private:
 	/*
 	 * Contains the compiling IDs of this shader map when it is being compiled asynchronously.
 	 * This can be used to access the shader map during async compiling.
diff --git a/Engine/Source/Runtime/Engine/Public/ComputeFramework/ComputeKernelShader.h b/Engine/Source/Runtime/Engine/Public/ComputeFramework/ComputeKernelShader.h
--- a/Engine/Source/Runtime/Engine/Public/ComputeFramework/ComputeKernelShader.h
+++ b/Engine/Source/Runtime/Engine/Public/ComputeFramework/ComputeKernelShader.h
@@ -2,6 +2,20 @@
 
 #pragma once
 
-class FComputeKernelShader : public FShader
+#include "Shader.h"
+#include "ShaderParameterStruct.h"
+#include "ComputeFramework/ComputeKernelShaderMap.h"
+
+#include "GlobalShader.h"
+
+class FComputeKernelShader : public FGlobalShader /*public FShader*/
 {
+public:
+	DECLARE_SHADER_TYPE(FComputeKernelShader, Global);
+	SHADER_USE_PARAMETER_STRUCT(FComputeKernelShader, FGlobalShader);
+	//DECLARE_SHADER_TYPE(FComputeKernelShader, ComputeKernel);
+	//SHADER_USE_PARAMETER_STRUCT(FComputeKernelShader, FShader);
+
+	BEGIN_SHADER_PARAMETER_STRUCT(FParameters, )
+	END_SHADER_PARAMETER_STRUCT()
 };
diff --git a/Engine/Source/Runtime/Engine/Public/ComputeFramework/ComputeKernelShaderMap.h b/Engine/Source/Runtime/Engine/Public/ComputeFramework/ComputeKernelShaderMap.h
--- a/Engine/Source/Runtime/Engine/Public/ComputeFramework/ComputeKernelShaderMap.h
+++ b/Engine/Source/Runtime/Engine/Public/ComputeFramework/ComputeKernelShaderMap.h
@@ -5,7 +5,7 @@
 #include "RHI.h"
 #include "RHIDefinitions.h"
 #include "Serialization/MemoryImage.h"
-#include "Shader.h"
+#include "ComputeFramework/ComputeKernelShaderType.h"
 
 
 class FComputeKernelResource;
@@ -90,6 +90,18 @@ public:
 	 */
 	bool TryToAddToExistingCompilationTask(FComputeKernelResource* KernelShader);
 
+	template<typename ShaderType> 
+	TShaderRef<ShaderType> GetShader(int32 PermutationId) const 
+	{ 
+		return TShaderRef<ShaderType>(GetContent()->GetShader<ShaderType>(PermutationId), *this); 
+	}
+
+	TShaderRef<FShader> GetShader(FShaderType* ShaderType, int32 PermutationId) const 
+	{ 
+		return TShaderRef<FShader>(GetContent()->GetShader(ShaderType, PermutationId), *this);
+	}
+
+
 
 
 
diff --git a/Engine/Source/Runtime/Engine/Public/ComputeFramework/ComputeKernelShaderType.h b/Engine/Source/Runtime/Engine/Public/ComputeFramework/ComputeKernelShaderType.h
--- a/Engine/Source/Runtime/Engine/Public/ComputeFramework/ComputeKernelShaderType.h
+++ b/Engine/Source/Runtime/Engine/Public/ComputeFramework/ComputeKernelShaderType.h
@@ -21,13 +21,48 @@ struct FComputeKernelShaderPermutationParameters : public FShaderPermutationPara
 class FComputeKernelShaderType : public FShaderType
 {
 public:
+	typedef FShader::CompiledShaderInitializerType CompiledShaderInitializerType;
+
+	FComputeKernelShaderType(
+		FTypeLayoutDesc& InTypeLayout,
+		const TCHAR* InName,
+		const TCHAR* InSourceFilename,
+		const TCHAR* InFunctionName,
+		uint32 InFrequency,					// ugly - ignored for Niagara shaders but needed for IMPLEMENT_SHADER_TYPE macro magic
+		int32 InTotalPermutationCount,
+		ConstructSerializedType InConstructSerializedRef,
+		ConstructCompiledType InConstructCompiledRef,
+		ModifyCompilationEnvironmentType InModifyCompilationEnvironmentRef,
+		ShouldCompilePermutationType InShouldCompilePermutationRef,
+		ValidateCompiledResultType InValidateCompiledResultRef,
+		uint32 InTypeSize,
+		const FShaderParametersMetadata* InRootParametersMetadata = nullptr
+		)
+		: FShaderType(
+			EShaderTypeForDynamicCast::ComputeKernel,
+			InTypeLayout, 
+			InName, 
+			InSourceFilename, 
+			InFunctionName, 
+			SF_Compute, 
+			InTotalPermutationCount,
+			InConstructSerializedRef,
+			InConstructCompiledRef,
+			InModifyCompilationEnvironmentRef,
+			InShouldCompilePermutationRef,
+			InValidateCompiledResultRef,
+			InTypeSize,
+			InRootParametersMetadata
+			)
+	{
+	}
+
 	TSharedRef<FShaderCommonCompileJob, ESPMode::ThreadSafe> BeginCompileShader(
 		EShaderPlatform ShaderPlatform,
 		FComputeKernelResource* KernelShader,
 		TArray<TSharedRef<FShaderCommonCompileJob, ESPMode::ThreadSafe>>& InOutNewJobs
 		);
 
-
 	bool ShouldCache(EShaderPlatform ShaderPlatform, const FComputeKernelResource* KernelShader) const
 	{
 		return ShouldCompilePermutation(FComputeKernelShaderPermutationParameters(ShaderPlatform, KernelShader));
diff --git a/Engine/Source/Runtime/Renderer/Private/ScenePrivate.h b/Engine/Source/Runtime/Renderer/Private/ScenePrivate.h
--- a/Engine/Source/Runtime/Renderer/Private/ScenePrivate.h
+++ b/Engine/Source/Runtime/Renderer/Private/ScenePrivate.h
@@ -2908,6 +2908,11 @@ public:
 		return GPUSkinCache;
 	}
 
+	FComputeFramework* GetComputeFramework() override
+	{
+		return ComputeFramework;
+	}
+
 #if RHI_RAYTRACING
 	virtual FRayTracingDynamicGeometryCollection* GetRayTracingDynamicGeometryCollection() override
 	{
