External actors names are now globally unique, by generating a 102-bits actor GUID and storing the most stable part in the name string and the volatile part in the name number (using only 30 bits to preserve the runtime fast name generation path). This is done like this to minimize the number of unique strings generated for external actors while ensuring that users can't generate colliding names from different machines. Also, changing the actor label won't rename the actors anymore to keep their globally unique names.

diff --git a/Engine/Source/Editor/UnrealEd/Private/Factories/EditorFactories.cpp b/Engine/Source/Editor/UnrealEd/Private/Factories/EditorFactories.cpp
--- a/Engine/Source/Editor/UnrealEd/Private/Factories/EditorFactories.cpp
+++ b/Engine/Source/Editor/UnrealEd/Private/Factories/EditorFactories.cpp
@@ -850,17 +850,17 @@ UObject* ULevelFactory::FactoryCreateText
 				FName ActorSourceName(NAME_None);
 				FParse::Value( Str, TEXT("NAME="), ActorSourceName );
 				ActorUniqueName = ActorSourceName;
-				// Make sure this name is unique.
+
 				AActor* Found=nullptr;
 				if( ActorUniqueName!=NAME_None )
 				{
 					// look in the current level for the same named actor
 					Found = FindObject<AActor>( World->GetCurrentLevel(), *ActorUniqueName.ToString() );
 				}
-				if( Found )
-				{
-					ActorUniqueName = MakeUniqueObjectName( World->GetCurrentLevel(), TempClass, ActorUniqueName );
-				}
+
+				// Make sure this name is unique. We need to do this upfront because we also want to potentially create the Associated BP class using the same name.
+				bool bNeedGloballyUniqueName = World->GetCurrentLevel()->IsUsingExternalActors() && CastChecked<AActor>(TempClass->GetDefaultObject())->SupportsExternalPackaging();
+				ActorUniqueName = FActorSpawnUtils::MakeUniqueActorName(World->GetCurrentLevel(), TempClass, FActorSpawnUtils::GetBaseName(ActorUniqueName), bNeedGloballyUniqueName);
 
 				// Get parent name for attachment.
 				FName ActorParentName(NAME_None);
diff --git a/Engine/Source/Runtime/Engine/Classes/Engine/World.h b/Engine/Source/Runtime/Engine/Classes/Engine/World.h
--- a/Engine/Source/Runtime/Engine/Classes/Engine/World.h
+++ b/Engine/Source/Runtime/Engine/Classes/Engine/World.h
@@ -596,6 +596,35 @@ public:
 	EObjectFlags ObjectFlags;		
 };
 
+/* World actors spawmning helper functions */
+struct ENGINE_API FActorSpawnUtils
+{
+	/**
+	 * Function to generate a locally or globally unique actor name. To generate a globally unique name, we store an epoch number
+	 * in the name number (while maintaining compatibility with fast path name generation, see GFastPathUniqueNameGeneration) and
+	 * also append an unique user id to the name.
+	 *
+	 * @param	Level			the new actor level
+	 * @param	Class			the new actor class
+	 * @param	BaseName		optional base name
+	 * @param	bGloballyUnique	whether to create a globally unique name
+	 * @return	generated actor name
+	**/
+	static FName MakeUniqueActorName(ULevel* Level, const UClass* Class, FName BaseName, bool bGloballyUnique);
+
+	/**
+	 * Determine if an actor name is globally unique or not.
+	 *
+	 * @param	Name			the name to check
+	 * @return true if the provided name is globally unique
+	**/
+	static bool IsGloballyUniqueName(FName Name);
+
+	/**
+	 * Return the base ename (without any number of globally unique identifier).
+	**/
+	static FName GetBaseName(FName Name);
+};
 
 /**
  *  This encapsulate World's async trace functionality. This contains two buffers of trace data buffer and alternates it for each tick. 
diff --git a/Engine/Source/Runtime/Engine/Classes/GameFramework/Actor.h b/Engine/Source/Runtime/Engine/Classes/GameFramework/Actor.h
--- a/Engine/Source/Runtime/Engine/Classes/GameFramework/Actor.h
+++ b/Engine/Source/Runtime/Engine/Classes/GameFramework/Actor.h
@@ -3374,9 +3374,6 @@ public:
 	FRenderCommandFence DetachFence;
 
 private:
-	/** Sets the friendly actor label and name */
-	void SetActorLabelInternal(const FString& NewActorLabelDirty, bool bMakeGloballyUniqueFName, bool bMarkDirty);
-
 	/** Helper that already assumes the Hit info is reversed, and avoids creating a temp FHitResult if possible. */
 	void InternalDispatchBlockingHit(UPrimitiveComponent* MyComp, UPrimitiveComponent* OtherComp, bool bSelfMoved, FHitResult const& Hit);
 
diff --git a/Engine/Source/Runtime/Engine/Private/ActorEditor.cpp b/Engine/Source/Runtime/Engine/Private/ActorEditor.cpp
--- a/Engine/Source/Runtime/Engine/Private/ActorEditor.cpp
+++ b/Engine/Source/Runtime/Engine/Private/ActorEditor.cpp
@@ -938,13 +938,7 @@ const FString& AActor::GetActorLabel() const
 	return ActorLabel;
 }
 
-void AActor::SetActorLabel( const FString& NewActorLabelDirty, bool bMarkDirty )
-{
-	const bool bMakeGloballyUniqueFName = false;
-	SetActorLabelInternal(NewActorLabelDirty, bMakeGloballyUniqueFName, bMarkDirty );
-}
-
-void AActor::SetActorLabelInternal(const FString& NewActorLabelDirty, bool bMakeGloballyUniqueFName, bool bMarkDirty)
+void AActor::SetActorLabel(const FString& NewActorLabelDirty, bool bMarkDirty)
 {
 	// Clean up the incoming string a bit
 	FString NewActorLabel = NewActorLabelDirty.TrimStartAndEnd();
@@ -968,44 +962,6 @@ void AActor::SetActorLabelInternal(const FString& NewActorLabelDirty, bool bMake
 				ActorLabel = MoveTemp(NewActorLabel);
 			}
 		}
-
-
-		// Next, update the actor's name
-		{
-			// Generate an object name for the actor's label
-			const FName OldActorName = GetFName();
-			FName NewActorName = MakeObjectNameFromDisplayLabel(GetActorLabel(), OldActorName);
-
-			// Has anything changed?
-			if (OldActorName != NewActorName)
-			{
-				// Try to rename the object
-				UObject* NewOuter = NULL;		// Outer won't be changing
-				ERenameFlags RenFlags = bMakeGloballyUniqueFName ? (REN_DontCreateRedirectors | REN_ForceGlobalUnique) : REN_DontCreateRedirectors;
-				bool bCanRename = Rename(*NewActorName.ToString(), NewOuter, REN_Test | REN_DoNotDirty | REN_NonTransactional | RenFlags);
-				if (bCanRename)
-				{
-					// NOTE: Will assert internally if rename fails
-					const bool bWasRenamed = Rename(*NewActorName.ToString(), NewOuter, RenFlags);
-				}
-				else
-				{
-					// Unable to rename the object.  Use a unique object name variant.
-					NewActorName = MakeUniqueObjectName(bMakeGloballyUniqueFName ? ANY_PACKAGE : GetOuter(), GetClass(), NewActorName);
-
-					bCanRename = Rename(*NewActorName.ToString(), NewOuter, REN_Test | REN_DoNotDirty | REN_NonTransactional | RenFlags);
-					if (bCanRename)
-					{
-						// NOTE: Will assert internally if rename fails
-						const bool bWasRenamed = Rename(*NewActorName.ToString(), NewOuter, RenFlags);
-					}
-					else
-					{
-						// Unable to rename the object.  Oh well, not a big deal.
-					}
-				}
-			}
-		}
 	}
 
 	FPropertyChangedEvent PropertyEvent( FindFProperty<FProperty>( AActor::StaticClass(), "ActorLabel" ) );
diff --git a/Engine/Source/Runtime/Engine/Private/LevelActor.cpp b/Engine/Source/Runtime/Engine/Private/LevelActor.cpp
--- a/Engine/Source/Runtime/Engine/Private/LevelActor.cpp
+++ b/Engine/Source/Runtime/Engine/Private/LevelActor.cpp
@@ -9,6 +9,7 @@
 #include "Misc/App.h"
 #include "UObject/Package.h"
 #include "Misc/PackageName.h"
+#include "Misc/DateTime.h"
 #include "UObject/ScriptStackTracker.h"
 #include "EngineStats.h"
 #include "EngineGlobals.h"
@@ -259,6 +260,124 @@ void LineCheckTracker::CaptureLineCheck(int32 LineCheckFlags, const FVector* Ext
 /*-----------------------------------------------------------------------------
 	Level actor management.
 -----------------------------------------------------------------------------*/
+
+#if WITH_EDITOR
+/** 
+ * Generates a 102-bits actor GUID:
+ *	- Bits 101-54 hold the user MAC address.
+ *	- Bits 53-0 hold a microseconds timestamp.
+  *
+ * Notes:
+ *	- The timestamp is stored in microseconds for a total of 54 bits, enough to cover
+ *	  the next 570 years.
+ *	- The highest 72 bits are appended to the name in base-64.
+ *	- The lowest 30 bits of the timestamp are stored in the name number. This is to
+ *	  minimize the total names generated for globally unique names (string part will
+ *	  change every ~17 minutes for a specific actor class).
+ *  - The name number bit 30 is reserved to know if this is a globally unique name. This
+ *	  is not 100% safe, but should cover most cases.
+ *  - The name number bit 31 is reserved to preserve the  fast path name generation (see
+ *	  GFastPathUniqueNameGeneration).
+ **/
+class FActorGUIDGenerator
+{
+public:
+	FActorGUIDGenerator()
+		: Origin(FDateTime(2020, 1, 1))
+		, MacAddress(FPlatformMisc::GetMacAddress())
+		, Counter(0)
+	{
+		check(MacAddress.Num() == 6);
+	}
+
+	FName NewActorGUID(FName BaseName)
+	{
+		uint8 HighPart[9];
+
+		const FDateTime Now = FDateTime::Now();
+		check(Now > Origin);
+
+		const FTimespan Elapsed = FDateTime::Now() - Origin;
+		const uint64 ElapsedUs = (uint64)Elapsed.GetTotalMilliseconds() * 1000 + (Counter++ % 1000);
+
+		// Copy 48-bits MAC address
+		FMemory::Memcpy(HighPart, MacAddress.GetData(), 6);
+		
+		// Append the high part of the timestamp (will change every ~17 minutes)
+		const uint64 ElapsedUsHighPart = ElapsedUs >> 30;
+		FMemory::Memcpy(HighPart + 6, &ElapsedUsHighPart, 3);
+
+		// Make final name
+		TStringBuilderWithBuffer<TCHAR, NAME_SIZE> StringBuilder;
+		StringBuilder += BaseName.ToString();
+		StringBuilder += TEXT("_");
+		StringBuilder += FBase64::Encode(HighPart, 9);
+
+		return FName(*StringBuilder, (ElapsedUs & 0x3fffffff) | (1 << 30));
+	}
+	
+private:
+	const FDateTime Origin;
+	const TArray<uint8> MacAddress;
+	uint32 Counter;
+};
+#endif
+
+FName FActorSpawnUtils::MakeUniqueActorName(ULevel* Level, const UClass* Class, FName BaseName, bool bGloballyUnique)
+{
+	FName NewActorName;
+
+#if WITH_EDITOR
+	if (bGloballyUnique)
+	{
+		static FActorGUIDGenerator ActorGUIDGenerator;
+
+		do
+		{
+			NewActorName = ActorGUIDGenerator.NewActorGUID(BaseName);
+		}
+		while (StaticFindObjectFast(nullptr, Level, NewActorName));
+	}
+	else
+#endif
+	{
+		NewActorName = MakeUniqueObjectName(Level, Class, BaseName);
+	}
+
+	return NewActorName;
+}
+
+bool FActorSpawnUtils::IsGloballyUniqueName(FName Name)
+{
+	if (Name.GetNumber() & (1 << 30))
+	{
+		const FString PlainName = Name.GetPlainNameString();
+		const int32 PlainNameLen = PlainName.Len();
+		
+		// Parse a name like this: StaticMeshActor_cCCECUPVj1cA
+		if (PlainNameLen >= 13 && PlainName[PlainNameLen - 13] == '_')
+		{
+			uint8 Dest[16];
+			if (FBase64::Decode(&PlainName[PlainNameLen - 12], 12, Dest))
+			{
+				return true;
+			}
+		}
+	}
+
+	return false;
+}
+
+FName FActorSpawnUtils::GetBaseName(FName Name)
+{
+	if (IsGloballyUniqueName(Name))
+	{
+		return *Name.GetPlainNameString().LeftChop(13);
+	}
+
+	return *Name.GetPlainNameString();
+}
+
 // LOOKING_FOR_PERF_ISSUES
 #define PERF_SHOW_MULTI_PAWN_SPAWN_FRAMES (!(UE_BUILD_SHIPPING || UE_BUILD_TEST)) && (LOOKING_FOR_PERF_ISSUES || !WITH_EDITORONLY_DATA)
 
@@ -384,43 +503,72 @@ AActor* UWorld::SpawnActor( UClass* Class, FTransform const* UserTransformPtr, c
 		LevelToSpawnIn = (SpawnParameters.Owner != NULL) ? SpawnParameters.Owner->GetLevel() : CurrentLevel;
 	}
 
+	// Use class's default actor as a template if none provided.
+	AActor* Template = SpawnParameters.Template ? SpawnParameters.Template : Class->GetDefaultObject<AActor>();
+	check(Template);
+
 	FName NewActorName = SpawnParameters.Name;
-	AActor* Template = SpawnParameters.Template;
-		
-	if( !Template )
+	UPackage* ExternalPackage = nullptr;
+	bool bNeedGloballyUniqueName = false;
+
+#if WITH_EDITOR
+	// Generate the actor's Guid
+	FGuid ActorGuid;
+	if (SpawnParameters.OverrideActorGuid.IsValid())
 	{
-		// Use class's default actor as a template.
-		Template = Class->GetDefaultObject<AActor>();
+		ActorGuid = SpawnParameters.OverrideActorGuid;
+	}
+	else
+	{
+		ActorGuid = FGuid::NewGuid();
 	}
-	check(Template);
 
-	if (NewActorName.IsNone())
+	// Generate and set the actor's external package if needed
+	if (SpawnParameters.OverridePackage)
 	{
-		// If we are using a template object and haven't specified a name, create a name relative to the template, otherwise let the default object naming behavior in Stat
-		if (!Template->HasAnyFlags(RF_ClassDefaultObject))
+		ExternalPackage = SpawnParameters.OverridePackage;
+	}
+	else if (LevelToSpawnIn->IsUsingExternalActors() && SpawnParameters.bCreateActorPackage && !(SpawnParameters.ObjectFlags & RF_Transient))
+	{
+		if (CastChecked<AActor>(Class->GetDefaultObject())->SupportsExternalPackaging())
 		{
-			NewActorName = MakeUniqueObjectName(LevelToSpawnIn, Template->GetClass(), *Template->GetFName().GetPlainNameString());
+			// @todo FH: needs to handle mark package dirty and asset creation notification
+			ExternalPackage = ULevel::CreateActorPackage(LevelToSpawnIn->GetPackage(), ActorGuid);
 		}
 	}
-	else if (StaticFindObjectFast(nullptr, LevelToSpawnIn, NewActorName))
+
+	bNeedGloballyUniqueName = !!ExternalPackage;
+#endif
+
+	if (NewActorName.IsNone())
 	{
-		// If the supplied name is already in use, then either fail in the requested manner or determine a new name to use if the caller indicates that's ok
+		// If we are using a template object and haven't specified a name, create a name relative to the template, otherwise let the default object naming behavior in Stat
+		const FName BaseName = Template->HasAnyFlags(RF_ClassDefaultObject) ? Class->GetFName() : *Template->GetFName().GetPlainNameString();
 
-		if (SpawnParameters.NameMode == FActorSpawnParameters::ESpawnActorNameMode::Requested)
-		{
-			NewActorName = MakeUniqueObjectName(LevelToSpawnIn, Template->GetClass(), *NewActorName.GetPlainNameString());
-		}
-		else
+		NewActorName = FActorSpawnUtils::MakeUniqueActorName(LevelToSpawnIn, Template->GetClass(), BaseName, bNeedGloballyUniqueName);
+	}
+	else if (StaticFindObjectFast(nullptr, LevelToSpawnIn, NewActorName) || (bNeedGloballyUniqueName != FActorSpawnUtils::IsGloballyUniqueName(NewActorName)))
+	{
+		// If the supplied name is already in use or doesn't respect globally uniqueness, then either fail in the requested manner or determine a new name to use if the caller indicates that's ok
+		switch(SpawnParameters.NameMode)
 		{
-			if (SpawnParameters.NameMode == FActorSpawnParameters::ESpawnActorNameMode::Required_Fatal)
-			{
-				UE_LOG(LogSpawn, Fatal, TEXT("An actor of name '%s' already exists in level '%s'."), *NewActorName.ToString(), *LevelToSpawnIn->GetFullName());
-			}
-			else if (SpawnParameters.NameMode == FActorSpawnParameters::ESpawnActorNameMode::Required_ErrorAndReturnNull)
-			{
-				UE_LOG(LogSpawn, Error, TEXT("An actor of name '%s' already exists in level '%s'."), *NewActorName.ToString(), *LevelToSpawnIn->GetFullName());
-			}
+		case FActorSpawnParameters::ESpawnActorNameMode::Requested:
+			NewActorName = FActorSpawnUtils::MakeUniqueActorName(LevelToSpawnIn, Template->GetClass(), FActorSpawnUtils::GetBaseName(NewActorName), bNeedGloballyUniqueName);
+			break;
+
+		case FActorSpawnParameters::ESpawnActorNameMode::Required_Fatal:
+			UE_LOG(LogSpawn, Fatal, TEXT("Cannot generate unique name for '%s' in level '%s'."), *NewActorName.ToString(), *LevelToSpawnIn->GetFullName());
+			return nullptr;
+
+		case FActorSpawnParameters::ESpawnActorNameMode::Required_ErrorAndReturnNull:
+			UE_LOG(LogSpawn, Error, TEXT("Cannot generate unique name for '%s' in level '%s'."), *NewActorName.ToString(), *LevelToSpawnIn->GetFullName());
+			return nullptr;
+
+		case FActorSpawnParameters::ESpawnActorNameMode::Required_ReturnNull:
 			return nullptr;
+
+		default:
+			check(0);
 		}
 	}
 
@@ -478,35 +626,6 @@ AActor* UWorld::SpawnActor( UClass* Class, FTransform const* UserTransformPtr, c
 
 	EObjectFlags ActorFlags = SpawnParameters.ObjectFlags;
 
-	UPackage* ExternalPackage = nullptr;
-#if WITH_EDITOR
-	// Generate the actor's Guid
-	FGuid ActorGuid;
-	if (SpawnParameters.OverrideActorGuid.IsValid())
-	{
-		ActorGuid = SpawnParameters.OverrideActorGuid;
-	}
-	else
-	{
-		ActorGuid = FGuid::NewGuid();
-	}
-
-	// Generate and set the actor's external package if needed
-	// Set actor's package
-	if (SpawnParameters.OverridePackage)
-	{
-		ExternalPackage = SpawnParameters.OverridePackage;
-	}
-	else if (LevelToSpawnIn->IsUsingExternalActors() && SpawnParameters.bCreateActorPackage && !(SpawnParameters.ObjectFlags & RF_Transient))
-	{
-		if (CastChecked<AActor>(Class->GetDefaultObject())->SupportsExternalPackaging())
-		{
-			// @todo FH: needs to handle mark package dirty and asset creation notification
-			ExternalPackage = ULevel::CreateActorPackage(LevelToSpawnIn->GetPackage(), ActorGuid);
-		}
-	}
-#endif
-
 	// actually make the actor object
 	AActor* const Actor = NewObject<AActor>(LevelToSpawnIn, Class, NewActorName, ActorFlags, Template, false/*bCopyTransientsFromClassDefaults*/, nullptr/*InInstanceGraph*/, ExternalPackage);
 	
