CIS - Remove Engine test cook cloth warning message due to cloth factory change.
Cloth - Add warning messages when the chosen clothing simulation factory fails to suppport all skeletal mesh cloth assets, and when none of the available factories support any of the assets.
[Merging CL14431250]

diff --git a/Engine/Source/Runtime/Engine/Private/Components/SkeletalMeshComponent.cpp b/Engine/Source/Runtime/Engine/Private/Components/SkeletalMeshComponent.cpp
--- a/Engine/Source/Runtime/Engine/Private/Components/SkeletalMeshComponent.cpp
+++ b/Engine/Source/Runtime/Engine/Private/Components/SkeletalMeshComponent.cpp
@@ -563,11 +563,11 @@ void USkeletalMeshComponent::OnRegister()
 			{
 				bSupportsAllAssets = false;
 
-				UE_LOG(LogSkeletalMesh, Warning,
+				UE_LOG(LogSkeletalMesh, Display,
 					TEXT("OnRegister[%s]: [%s] is currently unable to provide a fully functional simulation for each of this SkeletalMesh's clothing assets."),
 					*GetNameSafe(SkeletalMesh),
 					*ClothingSimulationFactory->GetName());
-				UE_LOG(LogSkeletalMesh, Warning,
+				UE_LOG(LogSkeletalMesh, Display,
 					TEXT("OnRegister[%s]: The ClothingSimulationFactory property will now be automatically updated to use the most functional simulation that can be found."),
 					*GetNameSafe(SkeletalMesh));
 
@@ -587,13 +587,18 @@ void USkeletalMeshComponent::OnRegister()
 				{
 					if (const TSubclassOf<UClothingSimulationFactory> NewClothingSimulationFactory = ClassProvider->GetClothingSimulationFactoryClass())
 					{
+						int NumAssets = 0;
 						int SupportedNumAssets = 0;
 						UClothingSimulationFactory* const NewDefaultObject = NewClothingSimulationFactory->GetDefaultObject<UClothingSimulationFactory>();
 						for (UClothingAssetBase* const ClothingAsset : SkeletalMesh->MeshClothingAssets)
 						{
-							if (ClothingAsset && NewDefaultObject->SupportsAsset(ClothingAsset))
+							if (ClothingAsset)
 							{
-								++SupportedNumAssets;
+								if (NewDefaultObject->SupportsAsset(ClothingAsset))
+								{
+									++SupportedNumAssets;
+								}
+								++NumAssets;
 							}
 						}
 
@@ -601,10 +606,23 @@ void USkeletalMeshComponent::OnRegister()
 						{
 							ClothingSimulationFactory = NewClothingSimulationFactory;
 							MostSupportedNumAssets = SupportedNumAssets;
+							if (SupportedNumAssets == NumAssets)
+							{
+								bSupportsAllAssets = true;
+								break;  // Stop at the first factory that supports all assets
+							}
 						}
 					}
 				}
 			}
+
+			UE_CLOG(!MostSupportedNumAssets, LogSkeletalMesh, Warning,
+				TEXT("OnRegister[%s]: There is no clothing simulation factory available that supports any of this SkeletalMesh's clothing assets."),
+				*GetNameSafe(SkeletalMesh));
+
+			UE_CLOG(MostSupportedNumAssets && !bSupportsAllAssets, LogSkeletalMesh, Warning,
+				TEXT("OnRegister[%s]: The most suitable clothing simulation factory available only partially supports this SkeletalMesh's clothing assets."),
+				*GetNameSafe(SkeletalMesh));
 		}
 	}
 
