Fixing Tessellation Deprecation warnings.

diff --git a/Engine/Source/Runtime/Engine/Private/TessellationRendering.cpp b/Engine/Source/Runtime/Engine/Private/TessellationRendering.cpp
--- a/Engine/Source/Runtime/Engine/Private/TessellationRendering.cpp
+++ b/Engine/Source/Runtime/Engine/Private/TessellationRendering.cpp
@@ -66,8 +66,10 @@ bool MaterialRenderingRequiresAdjacencyInformation_Parallel(UMaterialInterface*
 	{
 		const UMaterial* BaseMaterial = Material->GetMaterial_Concurrent();
 		check(BaseMaterial);
+PRAGMA_DISABLE_DEPRECATION_WARNINGS
 		EMaterialTessellationMode TessellationMode = (EMaterialTessellationMode)BaseMaterial->D3D11TessellationMode;
 		bool bEnableCrackFreeDisplacement = BaseMaterial->bEnableCrackFreeDisplacement;
+PRAGMA_ENABLE_DEPRECATION_WARNINGS
 		return TessellationMode == MTM_PNTriangles || (TessellationMode == MTM_FlatTessellation && bEnableCrackFreeDisplacement);
 	}
 	return false;
