Fix linker error due to missing file.

diff --git a/Engine/Source/Runtime/Engine/Private/ComputeFramework/ComputeKernelShader.cpp b/Engine/Source/Runtime/Engine/Private/ComputeFramework/ComputeKernelShader.cpp
new file mode 100644
--- /dev/null
+++ b/Engine/Source/Runtime/Engine/Private/ComputeFramework/ComputeKernelShader.cpp
@@ -0,0 +1,7 @@
+// Copyright Epic Games, Inc. All Rights Reserved.
+
+#include "ComputeFramework/ComputeKernelShader.h"
+#include "ComputeFramework/ComputeKernelShaderType.h"
+
+//IMPLEMENT_SHADER_TYPE(, FComputeKernelShader, TEXT("/Engine/Private/ComputeKernel.usf"), TEXT("Main"), SF_Compute);
+IMPLEMENT_GLOBAL_SHADER(FComputeKernelShader, "/Engine/Private/ComputeKernel.usf", "Main", SF_Compute);
