WorldPartition : First pass on DataLayers Runtime implementation (currently based on existing Layers).

diff --git a/Engine/Source/Runtime/Engine/Private/WorldPartition/DataLayer/DataLayerHelper.cpp b/Engine/Source/Runtime/Engine/Private/WorldPartition/DataLayer/DataLayerHelper.cpp
new file mode 100644
--- /dev/null
+++ b/Engine/Source/Runtime/Engine/Private/WorldPartition/DataLayer/DataLayerHelper.cpp
@@ -0,0 +1,17 @@
+// Copyright Epic Games, Inc. All Rights Reserved.
+
+#include "WorldPartition/DataLayer/DataLayerHelper.h"
+#include "Misc/HashBuilder.h"
+
+uint32 FDataLayersHelper::ComputeDataLayerID(const TArray<FName>& InDataLayers)
+{
+	if (InDataLayers.Num())
+	{
+		TArray<FName> SortedDataLayers = InDataLayers;
+		SortedDataLayers.Sort([](const FName& A, const FName& B) { return A.FastLess(B); });
+		FHashBuilder HashBuilder;
+		HashBuilder << SortedDataLayers;
+		return HashBuilder.GetHash();
+	}
+	return FDataLayersHelper::NoDataLayerID;
+}
\ No newline at end of file
diff --git a/Engine/Source/Runtime/Engine/Private/WorldPartition/DataLayer/DataLayerSubsystem.cpp b/Engine/Source/Runtime/Engine/Private/WorldPartition/DataLayer/DataLayerSubsystem.cpp
new file mode 100644
--- /dev/null
+++ b/Engine/Source/Runtime/Engine/Private/WorldPartition/DataLayer/DataLayerSubsystem.cpp
@@ -0,0 +1,78 @@
+// Copyright Epic Games, Inc. All Rights Reserved.
+
+#include "WorldPartition/DataLayer/DataLayerSubsystem.h"
+
+UDataLayerSubsystem::UDataLayerSubsystem()
+{}
+
+bool UDataLayerSubsystem::ShouldCreateSubsystem(UObject* Outer) const
+{
+	if (!Super::ShouldCreateSubsystem(Outer))
+	{
+		return false;
+	}
+
+	if (UWorld* WorldOuter = Cast<UWorld>(Outer))
+	{
+		return WorldOuter->GetWorldPartition() != nullptr;
+	}
+
+	return false;
+}
+
+void UDataLayerSubsystem::ActivateDataLayer(const FName& InDataLayer, bool bActivate)
+{
+	const bool bIsLayerActive = IsDataLayerActive(InDataLayer);
+	if (bIsLayerActive != bActivate)
+	{
+		if (bActivate)
+		{
+			//@todo_ow: Once DataLayer is an asset, validate it exists
+			ActiveDataLayers.Add(InDataLayer);
+		}
+		else
+		{
+			ActiveDataLayers.Remove(InDataLayer);
+		}
+	}
+}
+
+bool UDataLayerSubsystem::IsDataLayerActive(const FName& InDataLayer) const
+{
+	return ActiveDataLayers.Contains(InDataLayer);
+}
+
+bool UDataLayerSubsystem::IsAnyDataLayerActive(const TArray<FName>& InDataLayers) const
+{
+	for (const FName& DataLayer : InDataLayers)
+	{
+		if (IsDataLayerActive(DataLayer))
+		{
+			return true;
+		}
+	}
+	return false;
+}
+
+FAutoConsoleCommand UDataLayerSubsystem::ToggleDataLayerActivation(
+	TEXT("wp.Runtime.ToggleDataLayerActivation"),
+	TEXT("Toggles a DataLayer. Args [DataLayerName]"),
+	FConsoleCommandWithArgsDelegate::CreateLambda([](const TArray<FString>& Args)
+	{
+		if (Args.Num() == 1)
+		{
+			FName DataLayer = FName(Args[0]);
+			for (const FWorldContext& Context : GEngine->GetWorldContexts())
+			{
+				UWorld* World = Context.World();
+				if (World && World->IsGameWorld())
+				{
+					if (UDataLayerSubsystem* DataLayerSubsystem = World->GetSubsystem<UDataLayerSubsystem>())
+					{
+						DataLayerSubsystem->ActivateDataLayer(DataLayer, !DataLayerSubsystem->IsDataLayerActive(DataLayer));
+					}
+				}
+			}
+		}
+	})
+);
diff --git a/Engine/Source/Runtime/Engine/Private/WorldPartition/WorldPartition.cpp b/Engine/Source/Runtime/Engine/Private/WorldPartition/WorldPartition.cpp
--- a/Engine/Source/Runtime/Engine/Private/WorldPartition/WorldPartition.cpp
+++ b/Engine/Source/Runtime/Engine/Private/WorldPartition/WorldPartition.cpp
@@ -10,6 +10,7 @@
 #include "WorldPartition/WorldPartitionLevelStreamingPolicy.h"
 #include "WorldPartition/WorldPartitionActorDesc.h"
 #include "WorldPartition/WorldPartitionSubsystem.h"
+#include "WorldPartition/DataLayer/DataLayerHelper.h"
 #include "Algo/Accumulate.h"
 #include "Algo/Transform.h"
 #include "Algo/RemoveIf.h"
@@ -640,6 +641,7 @@ UWorldPartition::FActorCluster::FActorCluster(const FWorldPartitionActorDesc* Ac
 {
 	check(GridPlacement != EActorGridPlacement::None);
 	Actors.Add(ActorDesc->GetGuid());
+	DataLayersID = FDataLayersHelper::ComputeDataLayerID(DataLayers);
 }
 
 void UWorldPartition::FActorCluster::Add(const UWorldPartition::FActorCluster& ActorCluster)
@@ -677,6 +679,16 @@ void UWorldPartition::FActorCluster::Add(const UWorldPartition::FActorCluster& A
 			GridPlacement = EActorGridPlacement::Bounds;
 		}
 	}
+
+	// Merge DataLayers
+	if (DataLayersID != ActorCluster.DataLayersID)
+	{
+		for (const FName& DataLayer : ActorCluster.DataLayers)
+		{
+			DataLayers.AddUnique(DataLayer);
+		}
+		DataLayersID = FDataLayersHelper::ComputeDataLayerID(DataLayers);
+	}
 }
 
 const TSet<UWorldPartition::FActorCluster*>& UWorldPartition::GetActorClusters() const
diff --git a/Engine/Source/Runtime/Engine/Private/WorldPartition/WorldPartitionLevelStreamingPolicy.cpp b/Engine/Source/Runtime/Engine/Private/WorldPartition/WorldPartitionLevelStreamingPolicy.cpp
--- a/Engine/Source/Runtime/Engine/Private/WorldPartition/WorldPartitionLevelStreamingPolicy.cpp
+++ b/Engine/Source/Runtime/Engine/Private/WorldPartition/WorldPartitionLevelStreamingPolicy.cpp
@@ -149,7 +149,7 @@ void UWorldPartitionLevelStreamingPolicy::PrepareForPIE()
 
 	TSet<const UWorldPartitionRuntimeCell*> AlwaysLoadedCells;
 	TSet<const UWorldPartitionRuntimeCell*> StreamingCells;
-	WorldPartition->RuntimeHash->GetAllStreamingCells(StreamingCells);
+	WorldPartition->RuntimeHash->GetAllStreamingCells(StreamingCells, /*bIncludeDataLayers*/ true);
 
 	// Build Actor-to-Cell remapping
 	for (const UWorldPartitionRuntimeCell* Cell : StreamingCells)
diff --git a/Engine/Source/Runtime/Engine/Private/WorldPartition/WorldPartitionRuntimeSpatialHash.cpp b/Engine/Source/Runtime/Engine/Private/WorldPartition/WorldPartitionRuntimeSpatialHash.cpp
--- a/Engine/Source/Runtime/Engine/Private/WorldPartition/WorldPartitionRuntimeSpatialHash.cpp
+++ b/Engine/Source/Runtime/Engine/Private/WorldPartition/WorldPartitionRuntimeSpatialHash.cpp
@@ -8,6 +8,8 @@
 #include "WorldPartition/WorldPartitionActorDesc.h"
 #include "WorldPartition/WorldPartitionActorDescIterator.h"
 #include "WorldPartition/WorldPartitionStreamingPolicy.h"
+#include "WorldPartition/DataLayer/DataLayerSubsystem.h"
+#include "WorldPartition/DataLayer/DataLayerHelper.h"
 #include "GameFramework/WorldSettings.h"
 #include "ProfilingDebugging/ScopedTimers.h"
 #include "Engine/World.h"
@@ -326,9 +328,62 @@ struct FSquare2DGridHelper
 
 	struct FGridLevel : public FGrid2D
 	{
-		struct FGridCell
+		struct FGridCellDataChunk
 		{
+			FGridCellDataChunk(const TArray<FName>& InDataLayers)
+			: DataLayers(InDataLayers)
+			, DataLayersID(FDataLayersHelper::ComputeDataLayerID(InDataLayers))
+			{}
+
+			void AddActor(const FGuid& Actor) { Actors.Add(Actor); }
+			const TSet<FGuid>& GetActors() const { return Actors; }
+			bool HasDataLayers() const { return !DataLayers.IsEmpty(); }
+			const TArray<FName>& GetDataLayers() const { return DataLayers; }
+			uint32 GetDataLayersID() const { return DataLayersID; }
+
+		private:
 			TSet<FGuid> Actors;
+			TArray<FName> DataLayers;
+			uint32 DataLayersID;
+		};
+
+		struct FGridCell
+		{
+			void AddActor(const FGuid& InActor, const TArray<FName>& InDataLayers)
+			{
+				uint32 DataLayersID = FDataLayersHelper::ComputeDataLayerID(InDataLayers);
+				FGridCellDataChunk* ActorDataChunk = Algo::FindByPredicate(DataChunks, [&](FGridCellDataChunk& InDataChunk) { return InDataChunk.GetDataLayersID() == DataLayersID; });
+				if (!ActorDataChunk)
+				{
+					ActorDataChunk = &DataChunks.Emplace_GetRef(InDataLayers);
+				}
+				ActorDataChunk->AddActor(InActor);
+			}
+
+			void AddActors(const TSet<FGuid>& InActors, const TArray<FName>& InDataLayers)
+			{
+				for (const FGuid& Actor : InActors)
+				{
+					AddActor(Actor, InDataLayers);
+				}
+			}
+
+			const TArray<FGridCellDataChunk>& GetDataChunks() const { return DataChunks; }
+
+			const FGridCellDataChunk* GetNoDataLayersDataChunk() const
+			{
+				for (const FGridCellDataChunk& DataChunk : DataChunks)
+				{
+					if (!DataChunk.HasDataLayers())
+					{
+						return &DataChunk;
+					}
+				}
+				return nullptr;
+			}
+
+		private:
+			TArray<FGridCellDataChunk> DataChunks;
 		};
 
 		inline FGridLevel(const FVector2D& InOrigin, int32 InCellSize, int32 InGridSize)
@@ -484,11 +539,14 @@ struct FSquare2DGridHelper
 				for (int32 x = 0; x < CurrentGridSize; x++)
 				{
 					FGridLevel::FGridCell& ThisCell = Levels[Level].Cells[y * CurrentGridSize + x];
-					for (const FGuid& ActorGuid : ThisCell.Actors)
+					for (const FGridLevel::FGridCellDataChunk& DataChunk : ThisCell.GetDataChunks())
 					{
-						bool bWasAlreadyInSet;
-						ActorUsage.Add(ActorGuid, &bWasAlreadyInSet);
-						check(!bWasAlreadyInSet);
+						for (const FGuid& ActorGuid : DataChunk.GetActors())
+						{
+							bool bWasAlreadyInSet;
+							ActorUsage.Add(ActorGuid, &bWasAlreadyInSet);
+							check(!bWasAlreadyInSet);
+						}
 					}
 				}
 			}
@@ -541,7 +599,7 @@ const FSquare2DGridHelper& FSpatialHashStreamingGrid::GetGridHelper() const
 	return *GridHelper;
 }
 
-void FSpatialHashStreamingGrid::GetCells(const TArray<FWorldPartitionStreamingSource>& Sources, TSet<const UWorldPartitionRuntimeCell*>& Cells) const
+void FSpatialHashStreamingGrid::GetCells(const TArray<FWorldPartitionStreamingSource>& Sources, const UDataLayerSubsystem* DataLayerSubsystem, TSet<const UWorldPartitionRuntimeCell*>& Cells) const
 {
 	const FSquare2DGridHelper& Helper = GetGridHelper();
 	for (const FWorldPartitionStreamingSource& Source : Sources)
@@ -549,25 +607,35 @@ void FSpatialHashStreamingGrid::GetCells(const TArray<FWorldPartitionStreamingSo
 		const FSphere GridSphere(Source.Location, GetLoadingRange());
 		Helper.ForEachIntersectingCells(GridSphere, [&](const FIntVector& Coords)
 		{
-			if (UWorldPartitionRuntimeSpatialHashCell* Cell = GridLevels[Coords.Z].GridCells[Coords.Y * Helper.Levels[Coords.Z].GridSize + Coords.X])
+			const FSpatialHashStreamingGridLayerCell& LayerCell = GridLevels[Coords.Z].LayerCells[Coords.Y * Helper.Levels[Coords.Z].GridSize + Coords.X];
+			for (const UWorldPartitionRuntimeCell* Cell : LayerCell.GridCells)
 			{
-				Cells.Add(Cell);
+				if (!Cell->HasDataLayers() || (DataLayerSubsystem && DataLayerSubsystem->IsAnyDataLayerActive(Cell->GetDataLayers())))
+				{
+					Cells.Add(Cell);
+				}
 			}
 		});
 	}
 
-	GetAlwaysLoadedCells(Cells);
+	GetAlwaysLoadedCells(DataLayerSubsystem, Cells);
 }
 
-void FSpatialHashStreamingGrid::GetAlwaysLoadedCells(TSet<const UWorldPartitionRuntimeCell*>& Cells) const
+void FSpatialHashStreamingGrid::GetAlwaysLoadedCells(const UDataLayerSubsystem* DataLayerSubsystem, TSet<const UWorldPartitionRuntimeCell*>& Cells) const
 {
 	if (GridLevels.Num() > 0)
 	{
 		const int32 TopLevel = GridLevels.Num() - 1;
-		check(GridLevels[TopLevel].GridCells.Num() == 1);
-		if (UWorldPartitionRuntimeSpatialHashCell* Cell = GridLevels[TopLevel].GridCells[0])
+		for (const FSpatialHashStreamingGridLayerCell& LayerCell : GridLevels[TopLevel].LayerCells)
 		{
-			Cells.Add(Cell);
+			for (const UWorldPartitionRuntimeCell* Cell : LayerCell.GridCells)
+			{
+				if (!Cell->HasDataLayers() || (DataLayerSubsystem && DataLayerSubsystem->IsAnyDataLayerActive(Cell->GetDataLayers())))
+				{
+					check(Cell->IsAlwaysLoaded() || Cell->HasDataLayers());
+					Cells.Add(Cell);
+				}
+			}
 		}
 	}
 }
@@ -603,7 +671,16 @@ void FSpatialHashStreamingGrid::Draw3D(UWorld* World, const TArray<FWorldPartiti
 				Helper.Levels[GridLevel].GetCellBounds(FIntVector2(Coords.X, Coords.Y), CellWorldBounds);
 				
 				FVector BoundsExtent(CellWorldBounds.GetExtent(), 100.f);
-				const UWorldPartitionRuntimeSpatialHashCell* Cell = Cast<const UWorldPartitionRuntimeSpatialHashCell>(GridLevels[GridLevel].GridCells[Coords.Y * Helper.Levels[GridLevel].GridSize + Coords.X]);
+				const FSpatialHashStreamingGridLayerCell& LayerCell = GridLevels[GridLevel].LayerCells[Coords.Y * Helper.Levels[GridLevel].GridSize + Coords.X];
+				const UWorldPartitionRuntimeSpatialHashCell* Cell = nullptr;
+				for (const UWorldPartitionRuntimeSpatialHashCell* GridCell : LayerCell.GridCells)
+				{
+					Cell = GridCell;
+					if (!GridCell->HasDataLayers())
+					{
+						break;
+					}
+				}
 				FColor CellColor = Cell ? Cell->GetDebugColor().ToFColor(false).WithAlpha(16) : FColor(0, 0, 0, 16);
 				FVector BoundsOrigin(CellWorldBounds.GetCenter(), Z);
 				DrawDebugSolidBox(World, BoundsOrigin, BoundsExtent, CellColor, false, -1.f, 255);
@@ -680,7 +757,16 @@ void FSpatialHashStreamingGrid::Draw2D(UCanvas* Canvas, const TArray<FWorldParti
 				}
 			}
 
-			const UWorldPartitionRuntimeSpatialHashCell* Cell = Cast<const UWorldPartitionRuntimeSpatialHashCell>(GridLevels[GridLevel].GridCells[Coords.Y * Helper.Levels[GridLevel].GridSize + Coords.X]);
+			const FSpatialHashStreamingGridLayerCell& LayerCell = GridLevels[GridLevel].LayerCells[Coords.Y * Helper.Levels[GridLevel].GridSize + Coords.X];
+			const UWorldPartitionRuntimeSpatialHashCell* Cell = nullptr;
+			for (const UWorldPartitionRuntimeSpatialHashCell* GridCell : LayerCell.GridCells)
+			{
+				Cell = GridCell;
+				if (!GridCell->HasDataLayers())
+				{
+					break;
+				}
+			}
 			FLinearColor CellColor = Cell ? Cell->GetDebugColor() : FLinearColor(0.f, 0.f, 0.f, 0.25f);
 			FCanvasTileItem Item(CellScreenBounds.Min, GWhiteTexture, CellScreenBounds.GetSize(), CellColor);
 			Item.BlendMode = SE_BLEND_Translucent;
@@ -797,7 +883,7 @@ static FSquare2DGridHelper GetPartitionedActors(const UWorldPartition* WorldPart
 				const FWorldPartitionActorDesc& ActorDesc = *WorldPartition->GetActorDesc(ActorGuid);
 				if (PartitionedActors.GetLowestLevel().GetCellCoords(FVector2D(ActorDesc.GetOrigin()), CellCoords))
 				{
-					PartitionedActors.GetLowestLevel().GetCell(CellCoords).Actors.Add(ActorGuid);
+					PartitionedActors.GetLowestLevel().GetCell(CellCoords).AddActor(ActorGuid, TArray<FName>());
 				}
 				else
 				{
@@ -818,7 +904,7 @@ static FSquare2DGridHelper GetPartitionedActors(const UWorldPartition* WorldPart
 					{
 						GridLevel.ForEachIntersectingCells(ActorCluster->Bounds, [&GridLevel,&ActorCluster](const FIntVector2& Coords)
 						{
-							GridLevel.GetCell(Coords).Actors.Append(ActorCluster->Actors);
+							GridLevel.GetCell(Coords).AddActors(ActorCluster->Actors, ActorCluster->DataLayers);
 						});
 						bFoundCell = true;
 						break;
@@ -838,7 +924,7 @@ static FSquare2DGridHelper GetPartitionedActors(const UWorldPartition* WorldPart
 			
 		if (GridPlacement == EActorGridPlacement::AlwaysLoaded)
 		{
-			PartitionedActors.GetAlwaysLoadedCell().Actors.Append(ActorCluster->Actors);
+			PartitionedActors.GetAlwaysLoadedCell().AddActors(ActorCluster->Actors, ActorCluster->DataLayers);
 		}
 
 		if (!LogWorldPartitionRuntimeSpatialHash.IsSuppressed(ELogVerbosity::Verbose))
@@ -996,13 +1082,13 @@ bool UWorldPartitionRuntimeSpatialHash::GenerateStreaming(EWorldPartitionStreami
 	return true;
 }
 
-FName UWorldPartitionRuntimeSpatialHash::GetCellName(FName InGridName, int32 InLevel, int32 InCellX, int32 InCellY) const
+FName UWorldPartitionRuntimeSpatialHash::GetCellName(FName InGridName, int32 InLevel, int32 InCellX, int32 InCellY, uint32 InDataLayerID) const
 {
 	UWorldPartition* WorldPartition = GetOuterUWorldPartition();
 	const FString PackageName = FPackageName::GetShortName(WorldPartition->GetPackage());
 	const FString PackageNameNoPIEPrefix = UWorld::RemovePIEPrefix(PackageName);
 
-	return FName(*FString::Printf(TEXT("WPRT_%s_%s_Cell_L%d_X%02d_Y%02d"), *PackageNameNoPIEPrefix, *InGridName.ToString(), InLevel, InCellX, InCellY));
+	return FName(*FString::Printf(TEXT("WPRT_%s_%s_Cell_L%d_X%02d_Y%02d_DL%X"), *PackageNameNoPIEPrefix, *InGridName.ToString(), InLevel, InCellX, InCellY, InDataLayerID));
 }
 
 void UWorldPartitionRuntimeSpatialHash::CacheHLODParents()
@@ -1046,79 +1132,84 @@ bool UWorldPartitionRuntimeSpatialHash::CreateStreamingGrid(const FSpatialHashRu
 
 		FSpatialHashStreamingGridLevel& GridLevel = CurrentStreamingGrid.GridLevels.AddDefaulted_GetRef();
 
-		GridLevel.GridCells.Reserve(TempLevel.Cells.Num());
+		GridLevel.LayerCells.SetNum(TempLevel.Cells.Num());
 
 		int32 CellIndex = INDEX_NONE;
 		for (const FSquare2DGridHelper::FGridLevel::FGridCell& TempCell : TempLevel.Cells)
 		{
 			CellIndex++;
-			
-			FilteredActors.SetNum(0, false);
-			FilteredActors.Reset(TempCell.Actors.Num());
-			if (TempCell.Actors.Num())
+			const int32 CellCoordX = CellIndex % TempLevel.GridSize;
+			const int32 CellCoordY = CellIndex / TempLevel.GridSize;
+			TArray<UWorldPartitionRuntimeSpatialHashCell*>& GridCells = GridLevel.LayerCells[CellIndex].GridCells;
+
+			for (const FSquare2DGridHelper::FGridLevel::FGridCellDataChunk& GridCellDataChunk : TempCell.GetDataChunks())
 			{
-				Algo::TransformIf(TempCell.Actors, FilteredActors, [&WorldPartition](const FGuid& ActorGuid)
+				FilteredActors.SetNum(0, false);
+				FilteredActors.Reset(GridCellDataChunk.GetActors().Num());
+				if (GridCellDataChunk.GetActors().Num())
 				{
-					const FWorldPartitionActorDesc* ActorDesc = WorldPartition->GetActorDesc(ActorGuid);
-					const bool bShouldStripActorFromStreaming = ActorDesc->GetActorIsEditorOnly();
-					UE_CLOG(bShouldStripActorFromStreaming, LogWorldPartitionRuntimeSpatialHash, Verbose, TEXT("Stripping Actor %s (%s) from streaming grid"), *(ActorDesc->GetActorPath().ToString()), *ActorGuid.ToString(EGuidFormats::UniqueObjectGuid));
-					return !bShouldStripActorFromStreaming;
-				}, [](const FGuid& ActorGuid) { return ActorGuid;});
-			}
+					Algo::TransformIf(GridCellDataChunk.GetActors(), FilteredActors, [&WorldPartition](const FGuid& ActorGuid)
+					{
+						const FWorldPartitionActorDesc* ActorDesc = WorldPartition->GetActorDesc(ActorGuid);
+						const bool bShouldStripActorFromStreaming = ActorDesc->GetActorIsEditorOnly();
+						UE_CLOG(bShouldStripActorFromStreaming, LogWorldPartitionRuntimeSpatialHash, Verbose, TEXT("Stripping Actor %s (%s) from streaming grid"), *(ActorDesc->GetActorPath().ToString()), *ActorGuid.ToString(EGuidFormats::UniqueObjectGuid));
+						return !bShouldStripActorFromStreaming;
+					}, [](const FGuid& ActorGuid) { return ActorGuid; });
+				}
 
-			if (!FilteredActors.Num())
-			{
-				GridLevel.GridCells.Add(nullptr);
-				continue;
-			}
+				if (!FilteredActors.Num())
+				{
+					continue;
+				}
+				
+				// Cell cannot be treated as always loaded if it has data layers
+				const bool bIsCellAlwaysLoaded = (&TempCell == &PartionedActors.GetAlwaysLoadedCell()) && !GridCellDataChunk.HasDataLayers();
 
-			const bool bIsCellAlwaysLoaded = &TempCell == &PartionedActors.GetAlwaysLoadedCell();
+				FName CellName = GetCellName(CurrentStreamingGrid.GridName, Level, CellCoordX, CellCoordY, GridCellDataChunk.GetDataLayersID());
 
-			const int32 CellCoordX = CellIndex % TempLevel.GridSize;
-			const int32 CellCoordY = CellIndex / TempLevel.GridSize;
-			FName CellName = GetCellName(CurrentStreamingGrid.GridName, Level, CellCoordX, CellCoordY);
+				UWorldPartitionRuntimeSpatialHashCell* StreamingCell = NewObject<UWorldPartitionRuntimeSpatialHashCell>(WorldPartition, StreamingPolicy->GetRuntimeCellClass(), CellName);
+				GridCells.Add(StreamingCell);
+				StreamingCell->SetIsAlwaysLoaded(bIsCellAlwaysLoaded);
+				StreamingCell->SetDataLayers(GridCellDataChunk.GetDataLayers());
+				StreamingCell->Level = Level;
+				FBox2D Bounds;
+				verify(TempLevel.GetCellBounds(FIntVector2(CellCoordX, CellCoordY), Bounds));
+				StreamingCell->Position = FVector(Bounds.GetCenter(), 0.f);
 
-			UWorldPartitionRuntimeSpatialHashCell* StreamingCell = NewObject<UWorldPartitionRuntimeSpatialHashCell>(WorldPartition, StreamingPolicy->GetRuntimeCellClass(), CellName);
-			GridLevel.GridCells.Add(StreamingCell);
-			StreamingCell->SetIsAlwaysLoaded(bIsCellAlwaysLoaded);
-			StreamingCell->Level = Level;
-			FBox2D Bounds;
-			verify(TempLevel.GetCellBounds(FIntVector2(CellCoordX, CellCoordY), Bounds));
-			StreamingCell->Position = FVector(Bounds.GetCenter(), 0.f);
+				UE_LOG(LogWorldPartitionRuntimeSpatialHash, Verbose, TEXT("Cell%s %s Actors = %d"), bIsCellAlwaysLoaded ? TEXT(" (AlwaysLoaded)") : TEXT(""), *StreamingCell->GetName(), FilteredActors.Num());
 
-			UE_LOG(LogWorldPartitionRuntimeSpatialHash, Verbose, TEXT("Cell%s %s Actors = %d"), bIsCellAlwaysLoaded ? TEXT(" (AlwaysLoaded)") : TEXT(""), *StreamingCell->GetName(), FilteredActors.Num());
+				// Keep track of all AWorldPartitionHLOD actors referenced by this cell
+				TSet<FGuid> ReferencedHLODActors;
 
-			// Keep track of all AWorldPartitionHLOD actors referenced by this cell
-			TSet<FGuid> ReferencedHLODActors;
+				for (const FGuid& ActorGuid : FilteredActors)
+				{
+					const FWorldPartitionActorDesc* ActorDesc = WorldPartition->GetActorDesc(ActorGuid);
+					FGuid ParentHLOD = CachedHLODParents.FindRef(ActorGuid);
+					if (ParentHLOD.IsValid())
+					{
+						ReferencedHLODActors.Add(ParentHLOD);
+					}
+					StreamingCell->AddActorToCell(ActorDesc->GetActorPackage(), ActorDesc->GetActorPath());
+					UE_LOG(LogWorldPartitionRuntimeSpatialHash, Verbose, TEXT("  Actor : %s (%s) Origin(%s)"), *(ActorDesc->GetActorPath().ToString()), *ActorGuid.ToString(EGuidFormats::UniqueObjectGuid), *FVector2D(ActorDesc->GetOrigin()).ToString());
+				}
 
-			for (const FGuid& ActorGuid : FilteredActors)
-			{
-				const FWorldPartitionActorDesc* ActorDesc = WorldPartition->GetActorDesc(ActorGuid);
-				FGuid ParentHLOD = CachedHLODParents.FindRef(ActorGuid);
-				if (ParentHLOD.IsValid())
+				if (ReferencedHLODActors.Num() > 0)
 				{
-					ReferencedHLODActors.Add(ParentHLOD);
+					// Store the referenced HLOD actors as custom cell data
+					UWorldPartitionRuntimeHLODCellData* HLODCellData = NewObject<UWorldPartitionRuntimeHLODCellData>(StreamingCell);
+					HLODCellData->SetReferencedHLODActors(ReferencedHLODActors.Array());
+					StreamingCell->AddCellData(HLODCellData);
 				}
-				StreamingCell->AddActorToCell(ActorDesc->GetActorPackage(), ActorDesc->GetActorPath());
-				UE_LOG(LogWorldPartitionRuntimeSpatialHash, Verbose, TEXT("  Actor : %s (%s) Origin(%s)"), *(ActorDesc->GetActorPath().ToString()), *ActorGuid.ToString(EGuidFormats::UniqueObjectGuid), *FVector2D(ActorDesc->GetOrigin()).ToString());
-			}
 
-			if (ReferencedHLODActors.Num() > 0)
-			{
-				// Store the referenced HLOD actors as custom cell data
-				UWorldPartitionRuntimeHLODCellData* HLODCellData = NewObject<UWorldPartitionRuntimeHLODCellData>(StreamingCell);
-				HLODCellData->SetReferencedHLODActors(ReferencedHLODActors.Array());
-				StreamingCell->AddCellData(HLODCellData);
-			}
-			
-			if (Mode == EWorldPartitionStreamingMode::RuntimeStreamingCells)
-			{
-				FScopedLoadActorsHelper LoadCellActors(WorldPartition, FilteredActors, /*bSkipEditorOnly*/true);
-				UE_LOG(LogWorldPartitionRuntimeSpatialHash, Log, TEXT("Creating runtime streaming cells %s."), *StreamingCell->GetName());
-				if (!StreamingCell->CreateCellForCook())
+				if (Mode == EWorldPartitionStreamingMode::RuntimeStreamingCells)
 				{
-					UE_LOG(LogWorldPartitionRuntimeSpatialHash, Error, TEXT("Error creating runtime streaming cells for cook."));
-					return false;
+					FScopedLoadActorsHelper LoadCellActors(WorldPartition, FilteredActors, /*bSkipEditorOnly*/true);
+					UE_LOG(LogWorldPartitionRuntimeSpatialHash, Log, TEXT("Creating runtime streaming cells %s."), *StreamingCell->GetName());
+					if (!StreamingCell->CreateCellForCook())
+					{
+						UE_LOG(LogWorldPartitionRuntimeSpatialHash, Error, TEXT("Error creating runtime streaming cells for cook."));
+						return false;
+					}
 				}
 			}
 		}
@@ -1239,21 +1330,27 @@ bool UWorldPartitionRuntimeSpatialHash::GenerateHLOD()
 		PartionedActors.ForEachCells([PartionedActors, &ValidHLODActors, RuntimeGrid, WorldPartition, WorldBounds, this](const FIntVector& CellCoord)
 		{
 			const FSquare2DGridHelper::FGridLevel::FGridCell& GridCell = PartionedActors.GetCell(CellCoord);
-			const bool bIsCellAlwaysLoaded = &GridCell == &PartionedActors.GetAlwaysLoadedCell();
 
-			if (!bIsCellAlwaysLoaded && GridCell.Actors.Num() != 0)
+			// For now, HLOD only processes actors that are not in data layers
+			if (const FSquare2DGridHelper::FGridLevel::FGridCellDataChunk* GridCellDataChunk = GridCell.GetNoDataLayersDataChunk())
 			{
-				FBox2D CellBounds2D;
-				PartionedActors.GetCellBounds(CellCoord, CellBounds2D);
+				check(!GridCellDataChunk->HasDataLayers());
+				const bool bIsCellAlwaysLoaded = &GridCell == &PartionedActors.GetAlwaysLoadedCell();
+
+				if (!bIsCellAlwaysLoaded && GridCellDataChunk->GetActors().Num() != 0)
+				{
+					FBox2D CellBounds2D;
+					PartionedActors.GetCellBounds(CellCoord, CellBounds2D);
 
-				FName CellName = GetCellName(RuntimeGrid.GridName, CellCoord.Z, CellCoord.X, CellCoord.Y);
-				float CellLoadingRange = RuntimeGrid.LoadingRange;
-				FBox CellBounds = FBox(FVector(CellBounds2D.Min, WorldBounds.Min.Z), FVector(CellBounds2D.Max, WorldBounds.Max.Z));
+					FName CellName = GetCellName(RuntimeGrid.GridName, CellCoord.Z, CellCoord.X, CellCoord.Y, GridCellDataChunk->GetDataLayersID());
+					float CellLoadingRange = RuntimeGrid.LoadingRange;
+					FBox CellBounds = FBox(FVector(CellBounds2D.Min, WorldBounds.Min.Z), FVector(CellBounds2D.Max, WorldBounds.Max.Z));
 
-				TArray<AWorldPartitionHLOD*> HLODActors;
-				HLODActors = UHLODLayer::GenerateHLODForCell(WorldPartition, CellName, CellBounds, CellLoadingRange, GridCell.Actors);
+					TArray<AWorldPartitionHLOD*> HLODActors;
+					HLODActors = UHLODLayer::GenerateHLODForCell(WorldPartition, CellName, CellBounds, CellLoadingRange, GridCellDataChunk->GetActors());
 
-				ValidHLODActors.Append(HLODActors);
+					ValidHLODActors.Append(HLODActors);
+				}
 			}
 		});
 	}
@@ -1335,12 +1432,13 @@ bool UWorldPartitionRuntimeSpatialHash::GenerateNavigationData()
 		ANavigationDataChunkActor* DataChunkActor = World->SpawnActor<ANavigationDataChunkActor>(SpawnParams);
 		ActorCount++;
 
-		const FName CellName = GetCellName(RuntimeGrid.GridName, CellCoord.Z, CellCoord.X, CellCoord.Y);
+		const FName CellName = GetCellName(RuntimeGrid.GridName, CellCoord.Z, CellCoord.X, CellCoord.Y, FDataLayersHelper::NoDataLayerID);
 		DataChunkActor->SetActorLabel(FString::Printf(TEXT("NavDataChunkActor_%s_%s"), *GetName(), *CellName.ToString()));
 			
 		const FVector2D CellCenter = CellBounds.GetCenter();
 		DataChunkActor->SetActorLocation(FVector(CellCenter.X, CellCenter.Y, 0.f));
-			
+		
+		//@todo_ow: Properly handle data layers
 		DataChunkActor->CollectNavData(Bounds);
 
 		// Set target grid
@@ -1390,17 +1488,20 @@ FAutoConsoleCommand UWorldPartitionRuntimeSpatialHash::OverrideLoadingRangeComma
 );
 
 // Streaming interface
-int32 UWorldPartitionRuntimeSpatialHash::GetAllStreamingCells(TSet<const UWorldPartitionRuntimeCell*>& Cells) const
+int32 UWorldPartitionRuntimeSpatialHash::GetAllStreamingCells(TSet<const UWorldPartitionRuntimeCell*>& Cells, bool bIncludeDataLayers /*=false*/) const
 {
 	for (const FSpatialHashStreamingGrid& StreamingGrid : StreamingGrids)
 	{
 		for (const FSpatialHashStreamingGridLevel& GridLevel : StreamingGrid.GridLevels)
 		{
-			for (const UWorldPartitionRuntimeSpatialHashCell* Cell : GridLevel.GridCells)
+			for (const FSpatialHashStreamingGridLayerCell& LayerCell : GridLevel.LayerCells)
 			{
-				if (Cell)
+				for (const UWorldPartitionRuntimeSpatialHashCell* Cell : LayerCell.GridCells)
 				{
-					Cells.Add(Cell);
+					if (bIncludeDataLayers || !Cell->HasDataLayers())
+					{
+						Cells.Add(Cell);
+					}
 				}
 			}
 		}
@@ -1411,12 +1512,14 @@ int32 UWorldPartitionRuntimeSpatialHash::GetAllStreamingCells(TSet<const UWorldP
 
 int32 UWorldPartitionRuntimeSpatialHash::GetStreamingCells(const TArray<FWorldPartitionStreamingSource>& Sources, TSet<const UWorldPartitionRuntimeCell*>& Cells) const
 {
+	UDataLayerSubsystem* DataLayerSubsystem = GetWorld()->GetSubsystem<UDataLayerSubsystem>();
+
 	if (Sources.Num() == 0)
 	{
 		// Get always loaded cells
 		for (const FSpatialHashStreamingGrid& StreamingGrid : StreamingGrids)
 		{
-			StreamingGrid.GetAlwaysLoadedCells(Cells);
+			StreamingGrid.GetAlwaysLoadedCells(DataLayerSubsystem, Cells);
 		}
 	}
 	else
@@ -1424,7 +1527,7 @@ int32 UWorldPartitionRuntimeSpatialHash::GetStreamingCells(const TArray<FWorldPa
 		// Get cells based on streaming sources
 		for (const FSpatialHashStreamingGrid& StreamingGrid : StreamingGrids)
 		{
-			StreamingGrid.GetCells(Sources, Cells);
+			StreamingGrid.GetCells(Sources, DataLayerSubsystem, Cells);
 		}
 	}
 
diff --git a/Engine/Source/Runtime/Engine/Public/WorldPartition/DataLayer/DataLayerHelper.h b/Engine/Source/Runtime/Engine/Public/WorldPartition/DataLayer/DataLayerHelper.h
new file mode 100644
--- /dev/null
+++ b/Engine/Source/Runtime/Engine/Public/WorldPartition/DataLayer/DataLayerHelper.h
@@ -0,0 +1,9 @@
+// Copyright Epic Games, Inc. All Rights Reserved.
+
+#pragma once
+
+struct ENGINE_API FDataLayersHelper
+{
+	static const uint32 NoDataLayerID = 0;
+	static uint32 ComputeDataLayerID(const TArray<FName>& DataLayers);
+};
\ No newline at end of file
diff --git a/Engine/Source/Runtime/Engine/Public/WorldPartition/DataLayer/DataLayerSubsystem.h b/Engine/Source/Runtime/Engine/Public/WorldPartition/DataLayer/DataLayerSubsystem.h
new file mode 100644
--- /dev/null
+++ b/Engine/Source/Runtime/Engine/Public/WorldPartition/DataLayer/DataLayerSubsystem.h
@@ -0,0 +1,34 @@
+// Copyright Epic Games, Inc. All Rights Reserved.
+
+#pragma once
+
+#include "Subsystems/WorldSubsystem.h"
+#include "DataLayerSubsystem.generated.h"
+
+/**
+ * UDataLayerSubsystem
+ */
+
+UCLASS()
+class ENGINE_API UDataLayerSubsystem : public UWorldSubsystem
+{
+	GENERATED_BODY()
+
+public:
+	UDataLayerSubsystem();
+
+	//~ Begin USubsystem Interface.
+	virtual bool ShouldCreateSubsystem(UObject* Outer) const override;
+	//~ End USubsystem Interface.
+
+	void ActivateDataLayer(const FName& InDataLayer, bool bActivate);
+	bool IsDataLayerActive(const FName& InDataLayer) const;
+	bool IsAnyDataLayerActive(const TArray<FName>& InDataLayers) const;
+	const TSet<FName>& GetActiveDataLayers() const { return ActiveDataLayers; }
+
+private:
+	/** Console command used to toggle activation of a DataLayer */
+	static class FAutoConsoleCommand ToggleDataLayerActivation;
+
+	TSet<FName> ActiveDataLayers;
+};
\ No newline at end of file
diff --git a/Engine/Source/Runtime/Engine/Public/WorldPartition/WorldPartition.h b/Engine/Source/Runtime/Engine/Public/WorldPartition/WorldPartition.h
--- a/Engine/Source/Runtime/Engine/Public/WorldPartition/WorldPartition.h
+++ b/Engine/Source/Runtime/Engine/Public/WorldPartition/WorldPartition.h
@@ -139,6 +139,8 @@ public:
 		EActorGridPlacement	GridPlacement;
 		FName				RuntimeGrid;
 		FBox				Bounds;
+		TArray<FName>		DataLayers;
+		uint32				DataLayersID;
 
 		FActorCluster(const FWorldPartitionActorDesc* ActorDesc);
 		void Add(const FActorCluster& ActorCluster);
@@ -181,9 +183,7 @@ public:
 #if WITH_EDITOR
 	TMap<FGuid, TUniquePtr<FWorldPartitionActorDesc>> Actors;
 	class ULayersSubsystem* LayerSubSystem;
-	FDelegateHandle OnLayersChangedHandle;
-	FDelegateHandle OnActorsLayersChangedHandle;
-
+	
 	static TUniquePtr<FWorldPartitionActorDescFactory> DefaultActorDescFactory;
 	static TMap<FName, FWorldPartitionActorDescFactory*> ActorDescFactories;
 
diff --git a/Engine/Source/Runtime/Engine/Public/WorldPartition/WorldPartitionRuntimeCell.h b/Engine/Source/Runtime/Engine/Public/WorldPartition/WorldPartitionRuntimeCell.h
--- a/Engine/Source/Runtime/Engine/Public/WorldPartition/WorldPartitionRuntimeCell.h
+++ b/Engine/Source/Runtime/Engine/Public/WorldPartition/WorldPartitionRuntimeCell.h
@@ -53,6 +53,9 @@ class UWorldPartitionRuntimeCell : public UObject
 	virtual FLinearColor GetDebugColor() const PURE_VIRTUAL(UWorldPartitionRuntimeCell::GetDebugColor, return FLinearColor::Black;);
 	virtual bool IsAlwaysLoaded() const { return bIsAlwaysLoaded; }
 	virtual void SetIsAlwaysLoaded(bool bInIsAlwaysLoaded) { bIsAlwaysLoaded = bInIsAlwaysLoaded; }
+	void SetDataLayers(const TArray<FName> InDataLayers) { check(DataLayers.IsEmpty()); DataLayers = InDataLayers; }
+	bool HasDataLayers() const { return !DataLayers.IsEmpty(); }
+	const TArray<FName>& GetDataLayers() const { return DataLayers; }
 
 #if WITH_EDITOR
 	virtual void AddActorToCell(FName Package, FName Path) PURE_VIRTUAL(UWorldPartitionRuntimeCell::AddActorToCell,);
@@ -72,4 +75,7 @@ protected:
 private:
 	UPROPERTY()
 	TMap<const TSubclassOf<UWorldPartitionRuntimeCellData>, const UWorldPartitionRuntimeCellData*> CellDataMap;
+
+	UPROPERTY()
+	TArray<FName> DataLayers;
 };
diff --git a/Engine/Source/Runtime/Engine/Public/WorldPartition/WorldPartitionRuntimeHash.h b/Engine/Source/Runtime/Engine/Public/WorldPartition/WorldPartitionRuntimeHash.h
--- a/Engine/Source/Runtime/Engine/Public/WorldPartition/WorldPartitionRuntimeHash.h
+++ b/Engine/Source/Runtime/Engine/Public/WorldPartition/WorldPartitionRuntimeHash.h
@@ -26,7 +26,7 @@ class ENGINE_API UWorldPartitionRuntimeHash : public UObject
 #endif
 
 	// Streaming interface
-	virtual int32 GetAllStreamingCells(TSet<const UWorldPartitionRuntimeCell*>& Cells) const { return 0; }
+	virtual int32 GetAllStreamingCells(TSet<const UWorldPartitionRuntimeCell*>& Cells, bool bIncludeDataLayers = false) const { return 0; }
 	virtual int32 GetStreamingCells(const TArray<FWorldPartitionStreamingSource>& Sources, TSet<const UWorldPartitionRuntimeCell*>& Cells) const { return 0; };
 	virtual void SortStreamingCellsByImportance(const TSet<const UWorldPartitionRuntimeCell*>& InCells, const TArray<FWorldPartitionStreamingSource>& InSources, TArray<const UWorldPartitionRuntimeCell*, TInlineAllocator<256>>& OutSortedCells) const {}
 
diff --git a/Engine/Source/Runtime/Engine/Public/WorldPartition/WorldPartitionRuntimeSpatialHash.h b/Engine/Source/Runtime/Engine/Public/WorldPartition/WorldPartitionRuntimeSpatialHash.h
--- a/Engine/Source/Runtime/Engine/Public/WorldPartition/WorldPartitionRuntimeSpatialHash.h
+++ b/Engine/Source/Runtime/Engine/Public/WorldPartition/WorldPartitionRuntimeSpatialHash.h
@@ -9,6 +9,28 @@
 #include "WorldPartitionRuntimeSpatialHashCell.h"
 #include "WorldPartitionRuntimeSpatialHash.generated.h"
 
+USTRUCT()
+struct FSpatialHashStreamingGridLayerCell
+{
+	GENERATED_USTRUCT_BODY()
+
+	FSpatialHashStreamingGridLayerCell() = default;
+
+	// Move constructor.
+	FSpatialHashStreamingGridLayerCell(const FSpatialHashStreamingGridLayerCell& Other) = delete;
+	FSpatialHashStreamingGridLayerCell(FSpatialHashStreamingGridLayerCell&& Other) = default;
+
+	// Move assignment.
+	FSpatialHashStreamingGridLayerCell& operator=(const FSpatialHashStreamingGridLayerCell& Other) = delete;
+	FSpatialHashStreamingGridLayerCell& operator=(FSpatialHashStreamingGridLayerCell&& Other) = default;
+
+	UPROPERTY()
+	TArray<UWorldPartitionRuntimeSpatialHashCell*> GridCells;
+};
+
+// Disable copy of this USTRUCT by the autogenerated code.
+template<> struct TStructOpsTypeTraits<FSpatialHashStreamingGridLayerCell> : public TStructOpsTypeTraitsBase2<FSpatialHashStreamingGridLayerCell> { enum { WithCopy = false }; };
+
 /**
  * Represents a PIE/Game streaming grid level
  */
@@ -31,7 +53,7 @@ struct FSpatialHashStreamingGridLevel
 	 * Represents the grid cells
 	 */
 	UPROPERTY()
-	TArray<UWorldPartitionRuntimeSpatialHashCell*> GridCells;
+	TArray<FSpatialHashStreamingGridLayerCell> LayerCells;
 };
 
 // Disable copy of this USTRUCT by the autogenerated code.
@@ -91,8 +113,8 @@ struct FSpatialHashStreamingGrid
 	}
 
 	// Used by PIE/Game
-	void GetCells(const TArray<FWorldPartitionStreamingSource>& Sources, TSet<const UWorldPartitionRuntimeCell*>& Cells) const;
-	void GetAlwaysLoadedCells(TSet<const UWorldPartitionRuntimeCell*>& Cells) const;
+	void GetCells(const TArray<FWorldPartitionStreamingSource>& Sources, const class UDataLayerSubsystem* DataLayerSubsystem, TSet<const UWorldPartitionRuntimeCell*>& Cells) const;
+	void GetAlwaysLoadedCells(const UDataLayerSubsystem* DataLayerSubsystem, TSet<const UWorldPartitionRuntimeCell*>& Cells) const;
 	void Draw2D(UCanvas* Canvas, const TArray<FWorldPartitionStreamingSource>& Sources, const FBox& Region, const FBox2D& GridScreenBounds, TFunctionRef<FVector2D(const FVector2D&)> WorldToScreen) const;
 	void Draw3D(UWorld* World, const TArray<FWorldPartitionStreamingSource>& Sources) const;
 
@@ -153,7 +175,7 @@ public:
 #endif
 
 	// streaming interface
-	virtual int32 GetAllStreamingCells(TSet<const UWorldPartitionRuntimeCell*>& Cells) const override;
+	virtual int32 GetAllStreamingCells(TSet<const UWorldPartitionRuntimeCell*>& Cells, bool bIncludeDataLayers = false) const override;
 	virtual int32 GetStreamingCells(const TArray<FWorldPartitionStreamingSource>& Sources, TSet<const UWorldPartitionRuntimeCell*>& Cells) const override;
 	virtual void SortStreamingCellsByImportance(const TSet<const UWorldPartitionRuntimeCell*>& InCells, const TArray<FWorldPartitionStreamingSource>& InSources, TArray<const UWorldPartitionRuntimeCell*, TInlineAllocator<256>>& OutSortedCells) const override;
 
@@ -164,7 +186,7 @@ protected:
 private:
 
 #if WITH_EDITOR
-	FName GetCellName(FName InGridName, int32 InLevel, int32 InCellX, int32 InCellY) const;
+	FName GetCellName(FName InGridName, int32 InLevel, int32 InCellX, int32 InCellY, uint32 InDataLayerID) const;
 
 	void CacheHLODParents();
 #endif
