Update comments in NetSerialization.h

diff --git a/Engine/Source/Runtime/Engine/Classes/Engine/NetSerialization.h b/Engine/Source/Runtime/Engine/Classes/Engine/NetSerialization.h
--- a/Engine/Source/Runtime/Engine/Classes/Engine/NetSerialization.h
+++ b/Engine/Source/Runtime/Engine/Classes/Engine/NetSerialization.h
@@ -309,8 +309,10 @@ struct TStructOpsTypeTraits< FExampleArray > : public TStructOpsTypeTraitsBase2<
  *		During client serialization (reading), the client reads in the number of changed and number of deleted elements. It also builds a mapping of ReplicationID -> local index of the current array.
  *		As it deserializes IDs, it looks up the element and then does what it needs to (create if necessary, serialize in the current state, or delete).
  *
- *		There is currently no delta serialization done on the inner structures. If a ReplicationKey changes, the entire item is serialized. If we had
- *		use cases where we needed it, we could delta serialization on the inner dynamic properties. This could be done with more struct customization.
+ *		Delta Serialization for inner structs is now enabled by default. That means that when a ReplicationKey changes, we will compare the current state of the
+ *		struct to the last sent state, tracking changelists and only sending properties that changed exactly like the standard replication path.
+ *		If this causes issues with a specific FastArray type, it can be disabled by calling FFastArraySerializer::SetDeltaSerializationEnabled(false) in the constructor.
+ *		The feature can be completely disabled by setting the "net.SupportFastArrayDelta" CVar to 0.
  *
  *		ReplicationID and ReplicationKeys are set by the MarkItemDirty function on FFastArraySerializer. These are just int32s that are assigned in order as things change.
  *		There is nothing special about them other than being unique.
