Don't enforce switching to a globally unique name when the spawn request request a specific name.

diff --git a/Engine/Source/Runtime/Engine/Private/LevelActor.cpp b/Engine/Source/Runtime/Engine/Private/LevelActor.cpp
--- a/Engine/Source/Runtime/Engine/Private/LevelActor.cpp
+++ b/Engine/Source/Runtime/Engine/Private/LevelActor.cpp
@@ -547,7 +547,7 @@ AActor* UWorld::SpawnActor( UClass* Class, FTransform const* UserTransformPtr, c
 
 		NewActorName = FActorSpawnUtils::MakeUniqueActorName(LevelToSpawnIn, Template->GetClass(), BaseName, bNeedGloballyUniqueName);
 	}
-	else if (StaticFindObjectFast(nullptr, LevelToSpawnIn, NewActorName) || (bNeedGloballyUniqueName != FActorSpawnUtils::IsGloballyUniqueName(NewActorName)))
+	else if (StaticFindObjectFast(nullptr, LevelToSpawnIn, NewActorName) || ((bNeedGloballyUniqueName != FActorSpawnUtils::IsGloballyUniqueName(NewActorName)) && (SpawnParameters.NameMode == FActorSpawnParameters::ESpawnActorNameMode::Requested)))
 	{
 		// If the supplied name is already in use or doesn't respect globally uniqueness, then either fail in the requested manner or determine a new name to use if the caller indicates that's ok
 		switch(SpawnParameters.NameMode)
