Strata material graph, new pin type representing the mixed and layered BSDF.
Very basic and lots of things are missing but it is tracked.

diff --git a/Engine/Shaders/Private/BRDF.ush b/Engine/Shaders/Private/BRDF.ush
--- a/Engine/Shaders/Private/BRDF.ush
+++ b/Engine/Shaders/Private/BRDF.ush
@@ -444,6 +444,14 @@ float3 F_Schlick( float3 SpecularColor, float VoH )
 	
 }
 
+float3 F_Schlick(float3 SpecularColor, float3 EdgeTint, float VoH)
+{
+	float Fc = Pow5(1 - VoH);
+	// Anything less than 2% is physically impossible and is instead considered to be shadowing
+	return saturate(50.0 * SpecularColor.g) * Fc * EdgeTint + (1 - Fc) * SpecularColor;
+
+}
+
 float3 F_Fresnel( float3 SpecularColor, float VoH )
 {
 	float3 SpecularColorSqrt = sqrt( clamp( float3(0, 0, 0), float3(0.99, 0.99, 0.99), SpecularColor ) );
@@ -506,7 +514,7 @@ half3 EnvBRDF( half3 SpecularColor, half Roughness, half NoV )
 	return GF;
 }
 
-half3 EnvBRDFApprox( half3 SpecularColor, half Roughness, half NoV )
+half2 EnvBRDFApproxLazarov(half3 SpecularColor, half Roughness, half NoV)
 {
 	// [ Lazarov 2013, "Getting More Physical in Call of Duty: Black Ops II" ]
 	// Adaptation to fit our G term.
@@ -520,9 +528,21 @@ half3 EnvBRDFApprox( half3 SpecularColor, half Roughness, half NoV )
 	// Note: this is needed for the 'specular' show flag to work, since it uses a SpecularColor of 0
 	AB.y *= saturate( 50.0 * SpecularColor.g );
 
+	return AB;
+}
+
+half3 EnvBRDFApprox( half3 SpecularColor, half Roughness, half NoV )
+{
+	half2 AB = EnvBRDFApproxLazarov(SpecularColor, Roughness, NoV);
 	return SpecularColor * AB.x + AB.y;
 }
 
+half3 EnvBRDFApprox(half3 SpecularColor, float3 EdgeTint, half Roughness, half NoV)
+{
+	half2 AB = EnvBRDFApproxLazarov(SpecularColor, Roughness, NoV);
+	return SpecularColor * AB.x + EdgeTint * AB.y;
+}
+
 half EnvBRDFApproxNonmetal( half Roughness, half NoV )
 {
 	// Same as EnvBRDFApprox( 0.04, Roughness, NoV )
diff --git a/Engine/Shaders/Private/BasePassPixelShader.usf b/Engine/Shaders/Private/BasePassPixelShader.usf
--- a/Engine/Shaders/Private/BasePassPixelShader.usf
+++ b/Engine/Shaders/Private/BasePassPixelShader.usf
@@ -63,13 +63,25 @@
 #endif
 
 #include "ReflectionEnvironmentShared.ush"
+#if SIMPLE_FORWARD_SHADING || PLATFORM_FORCE_SIMPLE_SKY_DIFFUSE
+#define GetEffectiveSkySHDiffuse GetSkySHDiffuseSimple
+#else
+#define GetEffectiveSkySHDiffuse GetSkySHDiffuse
+#endif
 
 #if MATERIALBLENDING_TRANSLUCENT || MATERIALBLENDING_ADDITIVE || MATERIALBLENDING_MODULATE
 #define LumenGIVolumeStruct TranslucentBasePass
 #include "Lumen/LumenTranslucencyVolumeShared.ush"
 #endif
 
-#define STRATA_ENABLED (PROJECT_STRATA && MATERIALBLENDING_ANY_TRANSLUCENT<=0)
+#define STRATA_ENABLED					PROJECT_STRATA
+#define STRATA_OPAQUE_ENABLED			(STRATA_ENABLED && MATERIALBLENDING_ANY_TRANSLUCENT<=0)
+#define STRATA_TRANSLUCENT_ENABLED		(STRATA_ENABLED && MATERIALBLENDING_ANY_TRANSLUCENT>0)
+
+#if STRATA_TRANSLUCENT_ENABLED
+#include "ParticipatingMediaCommon.ush"
+#include "/Engine/Private/Strata/StrataForwardLighting.ush"
+#endif
 
 float NormalCurvatureToRoughness(float3 WorldNormal)
 {
@@ -270,12 +282,6 @@ float3 GetTranslucencyVolumeLighting(
 
 #endif
 
-#if SIMPLE_FORWARD_SHADING || PLATFORM_FORCE_SIMPLE_SKY_DIFFUSE
-	#define GetEffectiveSkySHDiffuse GetSkySHDiffuseSimple
-#else
-	#define GetEffectiveSkySHDiffuse GetSkySHDiffuse
-#endif
-
 /** Computes sky diffuse lighting, including precomputed shadowing. */
 void GetSkyLighting(FMaterialPixelParameters MaterialParameters, VTPageTableResult LightmapVTPageTableResult, FGBufferData GBuffer, float3 WorldNormal, float2 LightmapUV, uint LightmapDataIndex, float3 SkyOcclusionUV3D, out float3 OutDiffuseLighting, out float3 OutSubsurfaceLighting)
 {
@@ -1238,7 +1244,7 @@ void FPixelShaderInOut_MainPS(
 				{
 					int SingleCaptureIndex = GetPrimitiveData(MaterialParameters.PrimitiveId).SingleCaptureIndex;
 
-					half3 ReflectionColor = GetImageBasedReflectionLighting(MaterialParameters, GBuffer.Roughness, GBuffer.SpecularColor, IndirectIrradiance, GridIndex, SingleCaptureIndex, EyeIndex)
+					half3 ReflectionColor = GetImageBasedReflectionLighting(MaterialParameters, GBuffer.Roughness, GBuffer.SpecularColor, IndirectIrradiance, GridIndex, SingleCaptureIndex, EyeIndex) 
 						* IndirectOcclusion
 						* AOMultiBounce(GBuffer.SpecularColor, SpecOcclusion);
 
@@ -1435,6 +1441,68 @@ void FPixelShaderInOut_MainPS(
 #endif
 
 
+#if STRATA_ENABLED && !MATERIAL_SHADINGMODEL_UNLIT
+
+	FStrataData Strata = PixelMaterialInputs.FrontMaterial;
+
+#if STRATA_OPAQUE_ENABLED
+
+#if MATERIAL_IS_STRATA
+	// Need to cancel that to make sure strata material are only lit using Strata lighting passes
+	Color = 0;
+	// Clear only what should be and keep the rest
+	GBuffer.Roughness = 0;
+	GBuffer.DiffuseColor = 0;
+	GBuffer.SpecularColor = 0;
+	GBuffer.BaseColor = 0;
+	GBuffer.Specular = 0;
+	GBuffer.Metallic = 0;
+#endif
+
+	OpaqueBasePass.Strata.MaterialLobesTextureUAV[uint2(In.SvPosition.xy)] = Strata.LayerCount; // Could store a bit pack of layer count, BSDF type masks for tiling, etc.
+
+	uint2 PixelPos = In.SvPosition.xy;
+	uint PixelStrataDataByteOffset = (PixelPos.y * ResolvedView.BufferSizeAndInvSize.x + PixelPos.x) * OpaqueBasePass.Strata.MaxBytesPerPixel;
+
+	PackStrataOut(OpaqueBasePass.Strata.MaterialLobesBufferUAV, PixelStrataDataByteOffset, Strata, MaterialParameters.WorldNormal, MaterialParameters.CameraVector);
+
+	//switch (uint(Strata.LayerCount))
+	//{
+	//case 1:
+	//	Color = float3(0, 0, 1);	break;
+	//case 2:
+	//	Color = float3(0, 1, 0);	break;
+	//case 3:
+	//	Color = float3(1, 1, 0);	break;
+	//case 4:
+	//	Color = float3(1, 0, 0);	break;
+	//}
+
+#endif
+
+#if STRATA_TRANSLUCENT_ENABLED
+	//FORWARD_SHADING
+	if (Strata.LayerCount > 0)
+	{
+		float3 L = ResolvedView.DirectionalLightDirection;
+		float3 V = MaterialParameters.CameraVector;
+		float3 LightColor = ResolvedView.DirectionalLightColor.rgb * PI;
+
+		Color = StrataForwardLighting(
+			GBuffer,
+			MaterialParameters.WorldNormal,
+			V,
+			L,
+			LightColor,
+			DirectionalLightShadow,
+			IndirectOcclusion,
+			Strata);
+	}
+#endif
+
+#endif
+
+
 	#if MATERIAL_DOMAIN_POSTPROCESS
 		#if MATERIAL_OUTPUT_OPACITY_AS_ALPHA
 			Out.MRT[0] = half4(Color, Opacity);
@@ -1600,18 +1668,10 @@ void FPixelShaderInOut_MainPS(
 		View.VTFeedbackBuffer
 	);
 #endif
-
-#if STRATA_ENABLED
-	OpaqueBasePass.Strata.MaterialLobesTextureUAV[uint2(In.SvPosition.xy)] = 1.0f; // 
-	//OpaqueBasePass.Strata.MaterialLobesBufferUAV.Store(0*4, 0);
-	//OpaqueBasePass.Strata.MaterialLobesBufferUAV.Store(1*4, 1);
-	//OpaqueBasePass.Strata.MaterialLobesBufferUAV.Store(2*4, 2);
-	//OpaqueBasePass.Strata.MaterialLobesBufferUAV.Store(3*4, 3);
-#endif
 }
 
 // If virtual texture is enabled then use early depth test so that UAV feedback buffer writes respect the depth test
-#if NUM_VIRTUALTEXTURE_SAMPLES || LIGHTMAP_VT_ENABLED || STRATA_ENABLED
+#if NUM_VIRTUALTEXTURE_SAMPLES || LIGHTMAP_VT_ENABLED || STRATA_OPAQUE_ENABLED
 	#if COMPILER_SUPPORTS_DEPTHSTENCIL_EARLYTEST_LATEWRITE
 		// If we support early depth test with late write behaviour then use it since we may be using discard, or modifying depth
 		#define PIXELSHADER_EARLYDEPTHSTENCIL DEPTHSTENCIL_EARLYTEST_LATEWRITE
diff --git a/Engine/Shaders/Private/DeferredLightPixelShaders.usf b/Engine/Shaders/Private/DeferredLightPixelShaders.usf
--- a/Engine/Shaders/Private/DeferredLightPixelShaders.usf
+++ b/Engine/Shaders/Private/DeferredLightPixelShaders.usf
@@ -40,6 +40,12 @@ float4x4 CloudShadowmapWorldToLightClipMatrix;
 float CloudShadowmapStrength;
 #endif
 
+#if STRATA_ENABLED
+#include "/Engine/Private/Strata/Strata.ush"
+#include "/Engine/Private/ParticipatingMediaCommon.ush"
+#include "/Engine/Private/Strata/StrataDeferredLighting.ush"
+#endif
+
 struct FInputParams
 {
 	float2 PixelPos;
@@ -203,7 +209,55 @@ void DeferredLightPixelMain(
 
 		FRectTexture RectTexture = InitRectTexture(DeferredLightUniforms.SourceTexture);
 		float SurfaceShadow = 1.0f;
-		const float4 Radiance = GetDynamicLighting(DerivedParams.WorldPosition, DerivedParams.CameraVector, ScreenSpaceData.GBuffer, ScreenSpaceData.AmbientOcclusion, ScreenSpaceData.GBuffer.ShadingModelID, LightData, GetPerPixelLightAttenuation(InputParams.ScreenUV), Dither, uint2(InputParams.PixelPos), RectTexture, SurfaceShadow);
+
+		float4 Radiance = 0;
+
+#if STRATA_ENABLED
+		uint2 PixelPos = InputParams.PixelPos;
+		if (Strata.MaterialLobesTexture[PixelPos].x > 0.0f)
+		{
+			float3 V =-DerivedParams.CameraVector;
+			float3 L = LightData.Direction;	// Already normalized
+			float3 ToLight = L;
+			float LightMask = 1;
+			if (LightData.bRadialLight)
+			{
+				LightMask = GetLocalLightAttenuation(DerivedParams.WorldPosition, LightData, ToLight, L);
+			}
+
+			if (LightMask > 0)
+			{
+				uint PixelStrataDataByteOffset = (PixelPos.y * View.BufferSizeAndInvSize.x + PixelPos.x) * Strata.MaxBytesPerPixel;
+				FStrataPixelBSDFs PixelBSDFs = UnpackStrataIn(Strata.MaterialLobesBuffer, PixelStrataDataByteOffset);
+
+				FShadowTerms Shadow;
+				Shadow.SurfaceShadow = ScreenSpaceData.AmbientOcclusion;
+				Shadow.TransmissionShadow = 1;
+				Shadow.TransmissionThickness = 1;
+				Shadow.HairTransmittance.Transmittance = 1;
+				Shadow.HairTransmittance.OpaqueVisibility = 1;
+				float4 LightAttenuation = GetPerPixelLightAttenuation(InputParams.ScreenUV);
+				GetShadowTerms(ScreenSpaceData.GBuffer, LightData, DerivedParams.WorldPosition, L, LightAttenuation, Dither, Shadow);
+				SurfaceShadow = Shadow.SurfaceShadow;
+
+				Radiance = StrataDeferredLighting(
+					LightData,
+					ScreenSpaceData.GBuffer,
+					V,
+					L,
+					ToLight,
+					LightMask,
+					SurfaceShadow,
+					RectTexture,
+					PixelBSDFs);
+			}
+		}
+		else
+#endif
+		{
+		Radiance = GetDynamicLighting(DerivedParams.WorldPosition, DerivedParams.CameraVector, ScreenSpaceData.GBuffer, ScreenSpaceData.AmbientOcclusion, ScreenSpaceData.GBuffer.ShadingModelID, LightData, GetPerPixelLightAttenuation(InputParams.ScreenUV), Dither, uint2(InputParams.PixelPos), RectTexture, SurfaceShadow);
+		}
+
 		const float  Attenuation = ComputeLightProfileMultiplier(DerivedParams.WorldPosition, DeferredLightUniforms.Position, -DeferredLightUniforms.Direction, DeferredLightUniforms.Tangent);
 
 		OutColor += (Radiance * Attenuation) * OpaqueVisibility;
diff --git a/Engine/Shaders/Private/MaterialTemplate.ush b/Engine/Shaders/Private/MaterialTemplate.ush
--- a/Engine/Shaders/Private/MaterialTemplate.ush
+++ b/Engine/Shaders/Private/MaterialTemplate.ush
@@ -15,6 +15,7 @@
 #include "/Engine/Private/GlobalDistanceFieldShared.ush"
 #include "/Engine/Private/SceneData.ush"
 #include "/Engine/Private/HairShadingCommon.ush"
+#include "/Engine/Private/Strata/Strata.ush"
 
 #if USES_SPEEDTREE
 	#include "/Engine/Private/SpeedTreeCommon.ush"
diff --git a/Engine/Shaders/Private/ReflectionEnvironmentPixelShader.usf b/Engine/Shaders/Private/ReflectionEnvironmentPixelShader.usf
--- a/Engine/Shaders/Private/ReflectionEnvironmentPixelShader.usf
+++ b/Engine/Shaders/Private/ReflectionEnvironmentPixelShader.usf
@@ -18,6 +18,14 @@
 #define REFLECTION_COMPOSITE_SUPPORT_SKYLIGHT_BLEND 1
 #include "ReflectionEnvironmentComposite.ush"
 
+
+#if STRATA_ENABLED
+#include "/Engine/Private/Strata/Strata.ush"
+#include "/Engine/Private/ParticipatingMediaCommon.ush"
+#include "/Engine/Private/Strata/StrataEnvironmentLighting.ush"
+#endif
+
+
 float3 GatherRadiance(float CompositeAlpha, float3 WorldPosition, float3 RayDirection, float Roughness, float3 BentNormal, float IndirectIrradiance, uint ShadingModelID, uint NumCulledReflectionCaptures, uint CaptureDataStartIndex)
 {
 	// Indirect occlusion from DFAO, which should be applied to reflection captures and skylight specular, but not SSR
@@ -373,7 +381,38 @@ void ReflectionEnvironmentSkyLighting(
 		}
 #endif
 
-		float3 SkyLighting = SkyLightDiffuse(GBuffer, AmbientOcclusion, BufferUV, ScreenPosition, BentNormal, DiffuseColor) * CloudAmbientOcclusion;
+		float3 SkyLighting = 0;
+#if STRATA_ENABLED
+#if ENABLE_SKY_LIGHT
+		uint2 PixelPos = SvPosition.xy;
+		if (Strata.MaterialLobesTexture[PixelPos].x > 0.0f)
+		{
+			uint PixelStrataDataByteOffset = (PixelPos.y * View.BufferSizeAndInvSize.x + PixelPos.x) * Strata.MaxBytesPerPixel;
+			FStrataPixelBSDFs PixelBSDFs = UnpackStrataIn(Strata.MaterialLobesBuffer, PixelStrataDataByteOffset);
+
+			float3 WorldPosition = mul(float4(ScreenPosition * GBuffer.Depth, GBuffer.Depth, 1), View.ScreenToWorld).xyz;
+			float3 CameraToPixel = normalize(WorldPosition - View.WorldCameraOrigin);
+			float3 V = -CameraToPixel;
+			float3 N = GBuffer.WorldNormal;
+
+			float3 R = 2 * dot(V, N) * N - V;
+
+			SkyLighting = StrataEnvironmentLighting(
+				GBuffer,
+				WorldPosition,
+				CameraToPixel,
+				V,
+				N,
+				R,
+				AmbientOcclusion,
+				PixelBSDFs);
+		}
+		else
+#endif
+#endif
+		{
+			SkyLighting = SkyLightDiffuse(GBuffer, AmbientOcclusion, BufferUV, ScreenPosition, BentNormal, DiffuseColor)* CloudAmbientOcclusion;
+		}
 
 		// Scene color alpha is used for ScreenSpaceSubsurfaceScattering (if that is not needed it can be disabled with SUBSURFACE_CHANNEL_MODE)
 		FLightAccumulator LightAccumulator = (FLightAccumulator)0;
diff --git a/Engine/Shaders/Private/ShadingCommon.ush b/Engine/Shaders/Private/ShadingCommon.ush
--- a/Engine/Shaders/Private/ShadingCommon.ush
+++ b/Engine/Shaders/Private/ShadingCommon.ush
@@ -2,6 +2,8 @@
 
 #pragma once
 
+#include "Common.ush"
+
 // SHADINGMODELID_* occupy the 4 low bits of an 8bit channel and SKIP_* occupy the 4 high bits
 #define SHADINGMODELID_UNLIT				0
 #define SHADINGMODELID_DEFAULT_LIT			1
@@ -97,3 +99,55 @@ float3 ComputeF0(float Specular, float3 BaseColor, float Metallic)
 }
 
 
+
+float MakeRoughnessSafe(float Roughness)
+{
+	return clamp(Roughness, 0.001f, 1.0f);
+}
+
+
+
+// [Gulbrandsen, "Artist Friendly Metallic Fresnel"]
+
+// Compute F0 and edge tint from complex index of refraction N = n + ik. (absolute refractive index + extinction coefficients)
+void ComputeF0AndEdgeTintFromIOR(in float3 n, in float3 k, inout float3 Reflectivity, inout float3 EdgeTint)
+{
+	const float3 k2 = Square(k);
+
+	Reflectivity =
+		(Square(n - 1) + k2) /
+		(Square(n + 1) + k2);
+
+	const float3 SqrtR = sqrt(Reflectivity);
+	const float3 Common0 = (1 + SqrtR) / (1 - SqrtR);
+	EdgeTint = (Common0 - n) / (Common0 + (1 - SqrtR) / (1 + SqrtR));
+}
+
+// Compute complex index of refraction N = n + ik from F0 and edge tint.
+void ComputeComplexIORFromF0AndEdgeTint(in float3 Reflectivity, in float3 EdgeTint, inout float3 n, inout float3 k)
+{
+	const float3 g = EdgeTint;
+	const float3 r = clamp(Reflectivity, 0.0, 0.999);
+	const float3 SqrtR = sqrt(r);
+
+	n = g * (1 - r) / (1 + r) + (1 - g) * (1 + SqrtR) / (1 - SqrtR);
+
+	const float3 k2 = ((n + 1) * (n + 1) * r - (n - 1) * (n - 1)) / (1 - r);
+	k = sqrt(k2);
+}
+
+float3 ComputeComplexIORFromF0AndEdgeTint_IOR(in float3 Reflectivity, in float3 EdgeTint)
+{
+	float3 n = 0;
+	float3 k = 0;
+	ComputeComplexIORFromF0AndEdgeTint(Reflectivity, EdgeTint, n, k);
+	return n;
+}
+
+float3 ComputeComplexIORFromF0AndEdgeTint_Extinction(in float3 Reflectivity, in float3 EdgeTint)
+{
+	float3 n = 0;
+	float3 k = 0;
+	ComputeComplexIORFromF0AndEdgeTint(Reflectivity, EdgeTint, n, k);
+	return k;
+}
\ No newline at end of file
diff --git a/Engine/Shaders/Private/Strata/Strata.ush b/Engine/Shaders/Private/Strata/Strata.ush
new file mode 100644
--- /dev/null
+++ b/Engine/Shaders/Private/Strata/Strata.ush
@@ -0,0 +1,517 @@
+// Copyright Epic Games, Inc. All Rights Reserved.
+
+#pragma once
+
+#include "../Common.ush"
+#include "../BRDF.ush"
+#include "../ShadingCommon.ush"
+
+
+// STRATA_TODO
+//	- What model exactly
+//	- Fat dielectric model? Or make it adaptative (for instance Medium_ could be optional to store/load)
+//	- Have a Disney model that is convernient to be super packed and be able to represent current mateiral?
+//	- Shared normal? Up to N normals? On CPU, inputs have a hash that could be use to decuplicate normals
+//	- Transmittance lob modelisation (GGX statistics?) and blocking surface such as FStrataDiffuse?
+//  - 11-11-10 format
+//	- fp16 on gen5
+
+
+///////////////////////////////////////////////////////////////////////////////
+// BSDFs representation
+
+struct FStrataDiffuse
+{
+	float3	Albedo;
+	float   Roughness;
+
+	float3	Normal;
+};
+
+struct FStrataDielectric
+{
+	float	IOR;
+	float3	Tint;
+	float2	Roughness;
+
+	float3	Normal;
+};
+
+struct FStrataConductor
+{
+	float3	IOR;
+	float3	Extinction;
+	float2	Roughness;
+
+	float3	Normal;
+};
+
+struct FStrataVolume
+{
+	float3	Absorption;
+	float3	Scattering;
+	float	Anisotropy;
+};
+
+
+#define STRATA_BSDF_TYPE_DIFFUSE_ON			0
+#define STRATA_BSDF_TYPE_DIFFUSE_CHAN		1
+#define STRATA_BSDF_TYPE_DIELECTRIC			2
+#define STRATA_BSDF_TYPE_CONDUCTOR			3
+#define STRATA_BSDF_TYPE_VOLUME				4
+
+
+
+struct FStrataBSDF
+{
+	int					BSDFType;
+	float3				Weight;				// Setup by users and mixing operations
+
+	// Even thought the compiler should optimise, it would be nice to have a union with reinterpretation of the register memory.
+	FStrataDiffuse		Diffuse;
+	FStrataDielectric	Dielectric;
+	FStrataConductor	Conductor;
+	FStrataVolume		Volume;
+};
+
+
+#define STRATA_MAX_BSDF_COUNT_PER_LAYER	4
+#define STRATA_MAX_LAYER_COUNT			4
+#define STRATA_MAX_TOTAL_BSDF			4
+
+struct FStrataLayer
+{
+	int BSDFCount;
+
+	FStrataBSDF BSDFs[STRATA_MAX_BSDF_COUNT_PER_LAYER];
+};
+
+struct FStrataData
+{
+	int LayerCount;
+
+	FStrataLayer Layers[STRATA_MAX_LAYER_COUNT];	// Layer 0 is the top layer
+};
+
+
+FStrataData GetInitialisedStrataData()
+{
+	FStrataData StrataData = (FStrataData)0;
+	return StrataData;
+}
+
+FStrataData GetStrataDiffuseOrenNayarBSDF(float3 Albedo, float Roughness, float3 Normal)
+{
+	FStrataData StrataData = GetInitialisedStrataData();
+	StrataData.LayerCount++;
+
+	StrataData.Layers[0].BSDFCount++;
+	StrataData.Layers[0].BSDFs[0].BSDFType						= STRATA_BSDF_TYPE_DIFFUSE_ON;
+	StrataData.Layers[0].BSDFs[0].Diffuse.Albedo				= Albedo;
+	StrataData.Layers[0].BSDFs[0].Diffuse.Roughness				= Roughness;
+	StrataData.Layers[0].BSDFs[0].Diffuse.Normal				= Normal;
+
+	StrataData.Layers[0].BSDFs[0].Weight						= 1.0f;
+
+	return StrataData;
+}
+
+FStrataData GetStrataDiffuseChanBSDF(float3 Albedo, float Roughness, float3 Normal)
+{
+	FStrataData StrataData = GetInitialisedStrataData();
+	StrataData.LayerCount++;
+
+	StrataData.Layers[0].BSDFCount++;
+	StrataData.Layers[0].BSDFs[0].BSDFType						= STRATA_BSDF_TYPE_DIFFUSE_CHAN;
+	StrataData.Layers[0].BSDFs[0].Diffuse.Albedo				= Albedo;
+	StrataData.Layers[0].BSDFs[0].Diffuse.Roughness				= Roughness;
+	StrataData.Layers[0].BSDFs[0].Diffuse.Normal				= Normal;
+
+	StrataData.Layers[0].BSDFs[0].Weight						= 1.0f;
+
+	return StrataData;
+}
+
+FStrataData GetStrataDielectricBSDF(float2 Roughness, float IOR, float3 Tint, float3 Normal)
+{
+	FStrataData StrataData = GetInitialisedStrataData();
+	StrataData.LayerCount++;
+
+	StrataData.Layers[0].BSDFCount++;
+	StrataData.Layers[0].BSDFs[0].BSDFType						= STRATA_BSDF_TYPE_DIELECTRIC;
+	StrataData.Layers[0].BSDFs[0].Dielectric.IOR				= IOR;
+	StrataData.Layers[0].BSDFs[0].Dielectric.Tint				= Tint;
+	StrataData.Layers[0].BSDFs[0].Dielectric.Roughness			= Roughness;
+	StrataData.Layers[0].BSDFs[0].Dielectric.Normal				= Normal;
+
+	StrataData.Layers[0].BSDFs[0].Weight						= 1.0f;
+
+	return StrataData;
+}
+
+FStrataData GetStrataConductorBSDF(float3 IOR, float3 Extinction, float2 Roughness, float3 Normal)
+{
+	FStrataData StrataData = GetInitialisedStrataData();
+	StrataData.LayerCount++;
+
+	StrataData.Layers[0].BSDFCount++;
+	StrataData.Layers[0].BSDFs[0].BSDFType						= STRATA_BSDF_TYPE_CONDUCTOR;
+	StrataData.Layers[0].BSDFs[0].Conductor.IOR					= IOR;
+	StrataData.Layers[0].BSDFs[0].Conductor.Extinction			= Extinction;
+	StrataData.Layers[0].BSDFs[0].Conductor.Roughness			= Roughness;
+	StrataData.Layers[0].BSDFs[0].Conductor.Normal				= Normal;
+
+	StrataData.Layers[0].BSDFs[0].Weight						= 1.0f;
+
+	return StrataData;
+}
+
+FStrataData GetStrataVolumeBSDF(float3 Absorption, float3 Scattering, float Anisotropy)
+{
+	FStrataData StrataData = GetInitialisedStrataData();
+	StrataData.LayerCount++;
+
+	StrataData.Layers[0].BSDFCount++;
+	StrataData.Layers[0].BSDFs[0].BSDFType						= STRATA_BSDF_TYPE_VOLUME;
+	StrataData.Layers[0].BSDFs[0].Volume.Absorption				= Absorption;
+	StrataData.Layers[0].BSDFs[0].Volume.Scattering				= Scattering;
+	StrataData.Layers[0].BSDFs[0].Volume.Anisotropy				= Anisotropy;
+
+	StrataData.Layers[0].BSDFs[0].Weight						= 1.0f;
+
+	return StrataData;
+}
+
+
+
+///////////////////////////////////////////////////////////////////////////////
+// Operators node
+
+
+FStrataData StrataAdd(FStrataData A, FStrataData B)
+{
+	FStrataData StrataData = A;
+
+	// Append each BSDF from B to A, with same layer position
+	for (int Layer = 0; Layer < B.LayerCount; ++Layer)
+	{
+		FStrataLayer ALayer = A.Layers[Layer];
+		FStrataLayer BLayer = B.Layers[Layer];
+		for (int BSDF = 0; BSDF < BLayer.BSDFCount; BSDF++)
+		{
+			StrataData.Layers[Layer].BSDFs[ALayer.BSDFCount + BSDF] = BLayer.BSDFs[BSDF];
+		}
+
+		StrataData.Layers[Layer].BSDFCount = ALayer.BSDFCount + BLayer.BSDFCount;
+	}
+
+	StrataData.LayerCount = max(A.LayerCount, B.LayerCount);
+
+	return StrataData;
+}
+
+FStrataData StrataMultiply(FStrataData A, float3 Weight)
+{
+	FStrataData StrataData = A;
+
+	// Apply the weight to each BSDF
+	for (int Layer = 0; Layer < StrataData.LayerCount; ++Layer)
+	{
+		for (int BSDF = 0; BSDF < StrataData.Layers[Layer].BSDFCount; ++BSDF)
+		{
+			StrataData.Layers[Layer].BSDFs[BSDF].Weight *= Weight;
+		}
+	}
+
+	return StrataData;
+}
+
+FStrataData StrataHorizontalMixing(FStrataData Foreground, FStrataData Background, float Mix)
+{
+	const float ForgroundMixFactor = Mix;
+	const float BackgroundMixFactor = 1.0 - Mix;
+
+	FStrataData A = StrataMultiply(Foreground, ForgroundMixFactor);
+	FStrataData B = StrataMultiply(Background, BackgroundMixFactor);
+	return StrataAdd(A, B);
+}
+
+FStrataData StrataVerticalLayering(FStrataData Top, FStrataData Base)
+{
+	FStrataData StrataData = Top;
+
+	// Add each layer from Base under Top
+	int TopLayerCount = Top.LayerCount;
+	for (int Layer = 0; Layer < Base.LayerCount; ++Layer)
+	{
+		StrataData.Layers[TopLayerCount + Layer] = Base.Layers[Layer];
+	}
+
+	StrataData.LayerCount += Base.LayerCount;
+
+	return StrataData;
+}
+
+
+
+
+
+
+
+
+
+
+///////////////////////////////////////////////////////////////////////////////
+// BSDFs packing/unpacking
+
+
+
+#define STRATA_STORE_UINT1(x) StrataBuffer.Store (ByteOffset, asuint(x)); ByteOffset += 4;
+#define STRATA_STORE_UINT2(x) StrataBuffer.Store2(ByteOffset, asuint(x)); ByteOffset += 8;
+#define STRATA_STORE_UINT3(x) StrataBuffer.Store3(ByteOffset, asuint(x)); ByteOffset += 12;
+#define STRATA_STORE_UINT4(x) StrataBuffer.Store4(ByteOffset, asuint(x)); ByteOffset += 16;
+
+#define STRATA_LOAD_UINT1(x) x = StrataBuffer.Load (ByteOffset); ByteOffset += 4;
+#define STRATA_LOAD_UINT2(x) x = StrataBuffer.Load2(ByteOffset); ByteOffset += 8;
+#define STRATA_LOAD_UINT3(x) x = StrataBuffer.Load3(ByteOffset); ByteOffset += 12;
+#define STRATA_LOAD_UINT4(x) x = StrataBuffer.Load4(ByteOffset); ByteOffset += 16;
+
+#define STRATA_LOAD_FLOAT1(x) x = asfloat(StrataBuffer.Load (ByteOffset)); ByteOffset += 4;
+#define STRATA_LOAD_FLOAT2(x) x = asfloat(StrataBuffer.Load2(ByteOffset)); ByteOffset += 8;
+#define STRATA_LOAD_FLOAT3(x) x = asfloat(StrataBuffer.Load3(ByteOffset)); ByteOffset += 12;
+#define STRATA_LOAD_FLOAT4(x) x = asfloat(StrataBuffer.Load4(ByteOffset)); ByteOffset += 16;
+
+
+
+void PackStrataOut(
+	RWByteAddressBuffer StrataBuffer, 
+	uint ByteOffset, 
+	FStrataData Strata, 
+	float3 MeshWorldNormal, 
+	float3 V)
+{
+	// We cannot output (layers x bsdf) information because if we try to load them later, the reserved amount of VGPRS would be crazy high.
+	// So instead we simply store a list of BSDF and we allow only up to STRATA_MAX_TOTAL_BSDF of them.
+	// For each of the BSDF, we evaluated their final waight according to the material operation performed (layering, mixing, etc.).
+
+	// STRATA_TODO:
+	//	- We could clean up irrelevant BSDFs 
+	//	- We could generate transmittance BSDFs 
+
+
+	uint BSDFCount = 0;
+	float3 TotalTransmittance = 1.0f;
+
+	int l;
+	for (l = 0; l < Strata.LayerCount; ++l)
+	{
+		float3 LayerTransmittance = 0.0f;
+
+		for (int i = 0; i < Strata.Layers[l].BSDFCount; ++i)
+		{
+			FStrataBSDF BSDF = Strata.Layers[l].BSDFs[i];
+			BSDFCount++;
+
+			float3 N = MeshWorldNormal;
+			switch (BSDF.BSDFType)		// TODO_STRATA: all have a normal, even subsurface?
+			{
+			case STRATA_BSDF_TYPE_DIFFUSE_ON:
+			case STRATA_BSDF_TYPE_DIFFUSE_CHAN:
+			{ N = BSDF.Diffuse.Normal; break; }
+			case STRATA_BSDF_TYPE_DIELECTRIC:
+			{ N = BSDF.Dielectric.Normal; break; }
+			case STRATA_BSDF_TYPE_CONDUCTOR:
+			{ N = BSDF.Conductor.Normal; break; }
+			}
+			N = normalize(N);
+
+			// STRATA_TODO apply the same normal transform as for the material
+
+			BxDFContext ContextNoLight = (BxDFContext)0;
+			const float3 FakeL = float3(0, 0, 1);
+			Init(ContextNoLight, N, V, FakeL);
+
+			switch (BSDF.BSDFType)
+			{
+			case STRATA_BSDF_TYPE_DIFFUSE_ON:
+			{
+				//LayerTransmittance += 0.0f;
+				break;
+			}
+
+			case STRATA_BSDF_TYPE_DIFFUSE_CHAN:
+			{
+				//LayerTransmittance += 0.0f;
+				break;
+			}
+
+			case STRATA_BSDF_TYPE_DIELECTRIC:
+			{
+				float F0 = DielectricIorToF0(BSDF.Dielectric.IOR);
+				float3 F = F_Schlick(F0, ContextNoLight.VoH);
+
+				LayerTransmittance += BSDF.Weight * (1.0f - F); // STRATA_TODO use 1-FGD instead of simple 1-F, event if we ignore ray scattering/spreading
+				break;
+			}
+
+			case STRATA_BSDF_TYPE_CONDUCTOR:
+			{
+				float3 ReflectivityF0 = float3(1.0, 0.1, 0.1);
+				float3 EdgeTint = float3(0.1, 1.0, 0.1);
+				ComputeF0AndEdgeTintFromIOR(BSDF.Conductor.IOR, BSDF.Conductor.Extinction, ReflectivityF0, EdgeTint);
+				float3 F = F_Schlick(ReflectivityF0, EdgeTint, ContextNoLight.VoH);
+
+				LayerTransmittance += BSDF.Weight * (1.0f - F); // STRATA_TODO transmittance should only come from coverage
+				break;
+			}
+
+			case STRATA_BSDF_TYPE_VOLUME:
+			{
+				const float SafeRoughness = clamp(BSDF.Conductor.Roughness.x, 0.001, 1.0f);
+
+				const float3 ScatteringCoefficients = BSDF.Volume.Scattering;
+				const float3 AbsorptionCoefficients = BSDF.Volume.Absorption;
+				const float3 ExtinctionCoefficients = ScatteringCoefficients + AbsorptionCoefficients;
+
+				const float3 SafeExtinctionThreshold = 0.000001f;
+				const float3 SafeExtinctionCoefficients = max(SafeExtinctionThreshold, ExtinctionCoefficients);
+
+				const float Thickness = 0.01f;	// STRATA_TODO where does thickness come from? For now, 1cm
+				const float PathLength = Thickness / max(0.0001f, abs(ContextNoLight.NoV));
+				const float3 SafePathSegmentTransmittance = exp(-SafeExtinctionCoefficients * PathLength);
+
+				LayerTransmittance += BSDF.Weight * SafePathSegmentTransmittance;
+				break;
+			}
+			} // switch
+
+			// Update the weights to take into accounr transmittance towards the view
+			Strata.Layers[l].BSDFs[i].Weight *= TotalTransmittance;
+		}
+
+		TotalTransmittance *= LayerTransmittance;	// STRATA_TODO we need to go over all weight and trnasmittance ot make sure everything would behave corretly. Especially with a Add node for instance. Add saturate where neded, etc.
+	}
+
+
+	// Now write out the BSDFs
+
+	BSDFCount = min(BSDFCount, STRATA_MAX_TOTAL_BSDF);
+	STRATA_STORE_UINT1(BSDFCount);
+
+	BSDFCount = 0;
+	for (l = 0; l < Strata.LayerCount; ++l)
+	{
+		FStrataLayer Layer = Strata.Layers[l];
+
+		for (int i = 0; i < Layer.BSDFCount; ++i)
+		{
+			FStrataBSDF BSDF = Layer.BSDFs[i];
+
+			STRATA_STORE_UINT1(BSDF.BSDFType);
+			STRATA_STORE_UINT3(BSDF.Weight);
+
+			switch (BSDF.BSDFType)
+			{
+			case STRATA_BSDF_TYPE_DIFFUSE_ON:
+			case STRATA_BSDF_TYPE_DIFFUSE_CHAN:
+			{
+				STRATA_STORE_UINT3(BSDF.Diffuse.Albedo);
+				STRATA_STORE_UINT1(BSDF.Diffuse.Roughness);
+				STRATA_STORE_UINT3(BSDF.Diffuse.Normal);
+			}
+			break;
+			case STRATA_BSDF_TYPE_DIELECTRIC: //
+			{
+				STRATA_STORE_UINT1(BSDF.Dielectric.IOR);
+				STRATA_STORE_UINT3(BSDF.Dielectric.Tint);
+				STRATA_STORE_UINT2(BSDF.Dielectric.Roughness);
+				STRATA_STORE_UINT3(BSDF.Dielectric.Normal);
+			}
+			break;
+			case STRATA_BSDF_TYPE_CONDUCTOR:
+			{
+				STRATA_STORE_UINT3(BSDF.Conductor.IOR);
+				STRATA_STORE_UINT3(BSDF.Conductor.Extinction);
+				STRATA_STORE_UINT2(BSDF.Conductor.Roughness);
+				STRATA_STORE_UINT3(BSDF.Conductor.Normal);
+			}
+			break;
+			case STRATA_BSDF_TYPE_VOLUME:
+			{
+				STRATA_STORE_UINT3(BSDF.Volume.Absorption);
+				STRATA_STORE_UINT3(BSDF.Volume.Scattering);
+				STRATA_STORE_UINT1(BSDF.Volume.Anisotropy);
+			}
+			break;
+			}
+
+			// Stop writing if above budget.
+			BSDFCount++;
+			if (BSDFCount >= STRATA_MAX_TOTAL_BSDF)
+			{
+				return;
+			}
+		}
+	}
+}
+
+
+struct FStrataPixelBSDFs
+{
+	int BSDFCount;
+	FStrataBSDF BSDFs[STRATA_MAX_TOTAL_BSDF];
+};
+
+FStrataPixelBSDFs UnpackStrataIn(ByteAddressBuffer StrataBuffer, uint ByteOffset)
+{
+	FStrataPixelBSDFs PixelBSDFs = (FStrataPixelBSDFs)0;
+
+	STRATA_LOAD_UINT1(PixelBSDFs.BSDFCount);
+	PixelBSDFs.BSDFCount = min(PixelBSDFs.BSDFCount, STRATA_MAX_TOTAL_BSDF);	// To make the compiler happy with unrolling
+
+	UNROLL
+	for (int i = 0; i < PixelBSDFs.BSDFCount; ++i)
+	{
+		STRATA_LOAD_UINT1(PixelBSDFs.BSDFs[i].BSDFType);
+		STRATA_LOAD_FLOAT3(PixelBSDFs.BSDFs[i].Weight);
+
+		switch (PixelBSDFs.BSDFs[i].BSDFType)
+		{
+		case STRATA_BSDF_TYPE_DIFFUSE_ON:
+		case STRATA_BSDF_TYPE_DIFFUSE_CHAN:
+		{
+			STRATA_LOAD_FLOAT3(PixelBSDFs.BSDFs[i].Diffuse.Albedo);
+			STRATA_LOAD_FLOAT1(PixelBSDFs.BSDFs[i].Diffuse.Roughness);
+			STRATA_LOAD_FLOAT3(PixelBSDFs.BSDFs[i].Diffuse.Normal);
+		}
+		break;
+		case STRATA_BSDF_TYPE_DIELECTRIC:
+		{
+			STRATA_LOAD_FLOAT1(PixelBSDFs.BSDFs[i].Dielectric.IOR);
+			STRATA_LOAD_FLOAT3(PixelBSDFs.BSDFs[i].Dielectric.Tint);
+			STRATA_LOAD_FLOAT2(PixelBSDFs.BSDFs[i].Dielectric.Roughness);
+			STRATA_LOAD_FLOAT3(PixelBSDFs.BSDFs[i].Dielectric.Normal);
+		}
+		break;
+		case STRATA_BSDF_TYPE_CONDUCTOR:
+		{
+			STRATA_LOAD_FLOAT3(PixelBSDFs.BSDFs[i].Conductor.IOR);
+			STRATA_LOAD_FLOAT3(PixelBSDFs.BSDFs[i].Conductor.Extinction);
+			STRATA_LOAD_FLOAT2(PixelBSDFs.BSDFs[i].Conductor.Roughness);
+			STRATA_LOAD_FLOAT3(PixelBSDFs.BSDFs[i].Conductor.Normal);
+		}
+		break;
+		case STRATA_BSDF_TYPE_VOLUME:
+		{
+			STRATA_LOAD_FLOAT3(PixelBSDFs.BSDFs[i].Volume.Absorption);
+			STRATA_LOAD_FLOAT3(PixelBSDFs.BSDFs[i].Volume.Scattering);
+			STRATA_LOAD_FLOAT1(PixelBSDFs.BSDFs[i].Volume.Anisotropy);
+		}
+		break;
+		}
+	}
+
+	return PixelBSDFs;
+}
+
+
diff --git a/Engine/Shaders/Private/Strata/StrataDeferredLighting.ush b/Engine/Shaders/Private/Strata/StrataDeferredLighting.ush
new file mode 100644
--- /dev/null
+++ b/Engine/Shaders/Private/Strata/StrataDeferredLighting.ush
@@ -0,0 +1,143 @@
+// Copyright Epic Games, Inc. All Rights Reserved.
+
+#pragma once
+
+
+
+float4 StrataDeferredLighting(
+	FDeferredLightData LightData,
+	FGBufferData GBuffer,
+	float3 V,
+	float3 L,
+	float3 ToLight,
+	float LightMask,
+	float SurfaceShadow,
+	FRectTexture RectTexture,
+	FStrataPixelBSDFs PixelBSDFs
+	)
+{
+	float4 Radiance = 0;
+
+	for (int i = 0; i < PixelBSDFs.BSDFCount; ++i)
+	{
+		FStrataBSDF BSDF = PixelBSDFs.BSDFs[i];
+
+		float3 N = GBuffer.WorldNormal;
+		switch (BSDF.BSDFType)		// TODO_STRATA: all have a normal, even subsurface?
+		{
+		case STRATA_BSDF_TYPE_DIFFUSE_ON:
+		case STRATA_BSDF_TYPE_DIFFUSE_CHAN:
+		{ N = BSDF.Diffuse.Normal; break; }
+		case STRATA_BSDF_TYPE_DIELECTRIC:
+		{ N = BSDF.Dielectric.Normal; break; }
+		case STRATA_BSDF_TYPE_CONDUCTOR:
+		{ N = BSDF.Conductor.Normal; break; }
+		}
+		N = normalize(N);
+
+		BxDFContext Context = (BxDFContext)0;
+		Init(Context, N, V, L);
+		float SatNoL = saturate(Context.NoL);
+
+		float3 UglyIntegrationFactor = 0; // This is bad we need to make sure it is done the right way
+		if (LightData.bRectLight)
+		{
+			FRect Rect = GetRect(ToLight, LightData);
+			UglyIntegrationFactor = IntegrateLight(Rect, RectTexture);
+		}
+		else
+		{
+			FCapsuleLight Capsule = GetCapsule(ToLight, LightData);
+			UglyIntegrationFactor = IntegrateLight(Capsule, LightData.bInverseSquared);
+		}
+
+		const float3 DirLightBaseFactor = BSDF.Weight * LightData.Color * LightMask * SurfaceShadow * UglyIntegrationFactor;
+		const float3 DirLightGeometryFactor = DirLightBaseFactor * SatNoL;
+
+		switch (BSDF.BSDFType)
+		{
+		case STRATA_BSDF_TYPE_DIFFUSE_ON:
+		{
+			const float SafeRoughness = MakeRoughnessSafe(BSDF.Diffuse.Roughness);
+			Radiance.rgb += Diffuse_OrenNayar(BSDF.Diffuse.Albedo, SafeRoughness, Context.NoV, SatNoL, Context.VoH) * DirLightGeometryFactor;
+		}
+		break;
+		case STRATA_BSDF_TYPE_DIFFUSE_CHAN:
+		{
+			const float SafeRoughness = MakeRoughnessSafe(BSDF.Diffuse.Roughness);
+			float a = SafeRoughness * SafeRoughness;
+			float a2 = a * a;
+			Radiance.rgb += Diffuse_Chan(BSDF.Diffuse.Albedo, a2, Context.NoV, SatNoL, Context.VoH, Context.NoH) * DirLightGeometryFactor;
+		}
+		break;
+		case STRATA_BSDF_TYPE_DIELECTRIC:
+		{
+			float F0 = DielectricIorToF0(BSDF.Dielectric.IOR);
+
+			const float SafeRoughness = MakeRoughnessSafe(BSDF.Dielectric.Roughness.x);
+			float a2 = Pow4(SafeRoughness);
+			float Energy = 1.0f; // See EnergyNormalization
+
+			// Generalized microfacet specular
+			float D = D_GGX(a2, Context.NoH) * Energy;
+			float Vis = Vis_SmithJointApprox(a2, Context.NoV, SatNoL);
+			float3 F = F_Schlick(F0, Context.VoH);
+
+			Radiance.rgb += D * Vis * F * BSDF.Dielectric.Tint * DirLightGeometryFactor;
+		}
+		break;
+		case STRATA_BSDF_TYPE_CONDUCTOR:
+		{
+			float3 ReflectivityF0 = float3(1.0, 0.1, 0.1);
+			float3 EdgeTint = float3(0.1, 1.0, 0.1);
+			ComputeF0AndEdgeTintFromIOR(BSDF.Conductor.IOR, BSDF.Conductor.Extinction, ReflectivityF0, EdgeTint);
+
+			// See SpecularGGX, isotropic verion
+
+			const float SafeRoughness = MakeRoughnessSafe(BSDF.Conductor.Roughness.x);
+			float a2 = Pow4(SafeRoughness);
+			float Energy = 1.0f; // See EnergyNormalization
+
+			// Generalized microfacet specular
+			float D = D_GGX(a2, Context.NoH) * Energy;
+			float Vis = Vis_SmithJointApprox(a2, Context.NoV, SatNoL);
+			float3 F = F_Schlick(ReflectivityF0, EdgeTint, Context.VoH);
+
+			Radiance.rgb += D * Vis * F * DirLightGeometryFactor;
+		}
+		break;
+		case STRATA_BSDF_TYPE_VOLUME:
+		{
+			// Strata TODO expose albedo + extinction instead?
+			const float3 ScatteringCoefficients = BSDF.Volume.Scattering;
+			const float3 AbsorptionCoefficients = BSDF.Volume.Absorption;
+			const float3 ExtinctionCoefficients = ScatteringCoefficients + AbsorptionCoefficients;
+
+			const float3 SafeExtinctionThreshold = 0.000001f;
+			const float3 SafeExtinctionCoefficients = max(SafeExtinctionThreshold, ExtinctionCoefficients);
+			const float3 SafeAlbedo = ScatteringCoefficients / SafeExtinctionCoefficients;
+
+			const float Thickness = 0.01f;	// STRATA_TODO where does thickness come from? For now, 1cm
+			const float PathLength = Thickness / max(0.0001f, abs(Context.NoV));
+			const float3 SafePathSegmentTransmittance = exp(-SafeExtinctionCoefficients * PathLength);
+			const float3 NormLuminanceTransfert = 1.0f * ScatteringCoefficients;
+			float3 NormLuminanceSingleScattering = (NormLuminanceTransfert - NormLuminanceTransfert * SafePathSegmentTransmittance) / SafeExtinctionCoefficients;
+
+#if 1
+			const float3 SelfShadowTransmittance = SatNoL;
+#else
+			const float LightPathLength = 1.0f / max(0.0001f, SatNoL);
+			const float3 SelfShadowTransmittance = SatNoL <= 0.0 ? 0.0 : exp(-ExtinctionCoefficients * LightPathLength);
+#endif
+
+			const float Phase = HenyeyGreensteinPhase(BSDF.Volume.Anisotropy, Context.VoL);
+			Radiance.rgb += NormLuminanceSingleScattering * Phase * DirLightBaseFactor * SelfShadowTransmittance;
+		}
+		break;
+		}
+	}
+
+	return Radiance;
+}
+
+
diff --git a/Engine/Shaders/Private/Strata/StrataEnvironmentLighting.ush b/Engine/Shaders/Private/Strata/StrataEnvironmentLighting.ush
new file mode 100644
--- /dev/null
+++ b/Engine/Shaders/Private/Strata/StrataEnvironmentLighting.ush
@@ -0,0 +1,127 @@
+// Copyright Epic Games, Inc. All Rights Reserved.
+
+#pragma once
+
+
+
+float3 StrataEnvironmentLighting(
+	FGBufferData GBuffer,
+	float3 WorldPosition,
+	float3 CameraToPixel,
+	float3 V,
+	float3 N,
+	float3 R,
+	float3 AmbientOcclusion,
+	FStrataPixelBSDFs PixelBSDFs
+	)
+{
+	float3 SkyLighting = 0;
+	for (int i = 0; i < PixelBSDFs.BSDFCount; ++i)
+	{
+		FStrataBSDF BSDF = PixelBSDFs.BSDFs[i];
+
+		BxDFContext Context = (BxDFContext)0;
+		float3 L = N;
+		Init(Context, N, V, L);
+		float SatNoL = saturate(Context.NoL);
+
+		float3 Albedo = 0.0f;
+		float3 Normal = 0.0f;
+		float3 SpecularLuminance = 0.0f;
+		switch (BSDF.BSDFType)
+		{
+		case STRATA_BSDF_TYPE_DIFFUSE_ON:
+		case STRATA_BSDF_TYPE_DIFFUSE_CHAN:
+		{
+			Albedo = BSDF.Diffuse.Albedo;
+			Normal = BSDF.Diffuse.Normal;
+		}
+		break;
+		case STRATA_BSDF_TYPE_DIELECTRIC:
+		{
+			BRANCH
+				if (ReflectionStruct.SkyLightParameters.y > 0)
+				{
+					float F0 = DielectricIorToF0(BSDF.Dielectric.IOR);
+
+					const float SafeRoughness = MakeRoughnessSafe(BSDF.Dielectric.Roughness.x);
+					float a2 = Pow4(SafeRoughness);
+					float Energy = 1.0f; // See EnergyNormalization
+
+					// Generalized microfacet specular
+					//float D = D_GGX(a2, Context.NoH) * Energy;
+					//float Vis = Vis_SmithJointApprox(a2, Context.NoV, SatNoL);
+					//float3 F = F_Schlick(F0, Context.VoH);
+
+					float3 EvalEnvBRDF = EnvBRDFApprox(F0, SafeRoughness, Context.NoV);
+					float SkyAverageBrightness = 1.0f;
+					SpecularLuminance += GetSkyLightReflection(R, SafeRoughness, SkyAverageBrightness) * EvalEnvBRDF * BSDF.Dielectric.Tint;
+				}
+		}
+		break;
+		case STRATA_BSDF_TYPE_CONDUCTOR:
+		{
+			BRANCH
+				if (ReflectionStruct.SkyLightParameters.y > 0)
+				{
+					float3 ReflectivityF0 = float3(1.0, 0.1, 0.1);
+					float3 EdgeTint = float3(0.1, 1.0, 0.1);
+					ComputeF0AndEdgeTintFromIOR(BSDF.Conductor.IOR, BSDF.Conductor.Extinction, ReflectivityF0, EdgeTint);
+
+					// See SpecularGGX, isotropic verion
+
+					const float SafeRoughness = MakeRoughnessSafe(BSDF.Conductor.Roughness.x);
+					float a2 = Pow4(SafeRoughness);
+					float Energy = 1.0f; // See EnergyNormalization
+
+					// Generalized microfacet specular
+					//float D = D_GGX(a2, Context.NoH) * Energy;
+					//float Vis = Vis_SmithJointApprox(a2, Context.NoV, SatNoL);
+					//float3 F = F_Schlick(ReflectivityF0, EdgeTint, Context.VoH); // STRATA_TODO check that this is correct 
+
+					float3 EvalEnvBRDF = EnvBRDFApprox(ReflectivityF0, EdgeTint, SafeRoughness, Context.NoV);
+					float SkyAverageBrightness = 1.0f;
+					SpecularLuminance += GetSkyLightReflection(R, SafeRoughness, SkyAverageBrightness) * EvalEnvBRDF;
+				}
+		}
+		break;
+		case STRATA_BSDF_TYPE_VOLUME:
+		{
+			const float3 ScatteringCoefficients = BSDF.Volume.Scattering;
+			const float3 AbsorptionCoefficients = BSDF.Volume.Absorption;
+			const float3 ExtinctionCoefficients = ScatteringCoefficients + AbsorptionCoefficients;
+
+			const float3 SafeExtinctionThreshold = 0.000001f;
+			const float3 SafeExtinctionCoefficients = max(SafeExtinctionThreshold, ExtinctionCoefficients);
+			const float3 SafeAlbedo = ScatteringCoefficients / SafeExtinctionCoefficients;
+
+			const float Thickness = 0.01f;	// STRATA_TODO where does thickness come from? For now, 1cm
+			const float PathLength = Thickness / max(0.0001f, saturate(Context.NoV));
+			const float3 SafePathSegmentTransmittance = exp(-SafeExtinctionCoefficients * PathLength);
+			const float3 NormLuminanceTransfert = 1.0f * ScatteringCoefficients;
+			float3 NormLuminanceSingleScattering = (NormLuminanceTransfert - NormLuminanceTransfert * SafePathSegmentTransmittance) / SafeExtinctionCoefficients;
+
+			// No phase function because the SH already contains integration of BSDF with respect to hemisphere. Incorrect but needed.
+			Albedo = NormLuminanceSingleScattering;
+
+			// STRATA_TODO: This is to fake some directionality. Have normal specified on the material as for MaterialX?
+			Normal = GBuffer.WorldNormal;
+		}
+		break;
+		}
+
+		if (any(Albedo > 0.0f))
+		{
+			// BentNormal?
+			SkyLighting += BSDF.Weight * Diffuse_Lambert(Albedo) * GetSkySHDiffuse(Normal) * View.SkyLightColor.rgb * AmbientOcclusion;
+		}
+
+		SkyLighting += BSDF.Weight * SpecularLuminance * View.SkyLightColor.rgb * AmbientOcclusion;
+	}
+#if USE_PREEXPOSURE
+	SkyLighting *= View.PreExposure;
+#endif
+	return SkyLighting;
+}
+
+
diff --git a/Engine/Shaders/Private/Strata/StrataForwardLighting.ush b/Engine/Shaders/Private/Strata/StrataForwardLighting.ush
new file mode 100644
--- /dev/null
+++ b/Engine/Shaders/Private/Strata/StrataForwardLighting.ush
@@ -0,0 +1,206 @@
+// Copyright Epic Games, Inc. All Rights Reserved.
+
+#pragma once
+
+
+
+float3 StrataForwardLighting(
+	FGBufferData GBuffer,
+	float3 MaterialParametersWorldNormal,
+	float3 V,
+	float3 L,
+	float3 LColor,
+	float3 LShadow,
+	float3 IndirectOcclusion,
+	FStrataData Strata
+	)
+{
+	float3 Color = 0;
+	float3 TotalTransmittance = 1.0f;
+
+	for (int l = 0; l < Strata.LayerCount; ++l)
+	{
+		float3 LayerTransmittance = 0.0f;
+
+		for (int i = 0; i < Strata.Layers[l].BSDFCount; ++i)
+		{
+			FStrataBSDF BSDF = Strata.Layers[l].BSDFs[i];
+
+			float3 N = MaterialParametersWorldNormal;
+			switch (BSDF.BSDFType)		// TODO_STRATA: all have a normal, even subsurface?
+			{
+			case STRATA_BSDF_TYPE_DIFFUSE_ON:
+			case STRATA_BSDF_TYPE_DIFFUSE_CHAN:
+			{ N = BSDF.Diffuse.Normal; break; }
+			case STRATA_BSDF_TYPE_DIELECTRIC:
+			{ N = BSDF.Dielectric.Normal; break; }
+			case STRATA_BSDF_TYPE_CONDUCTOR:
+			{ N = BSDF.Conductor.Normal; break; }
+			}
+			N = normalize(N);
+			float3 R = 2 * dot(V, N) * N - V;
+
+			BxDFContext Context = (BxDFContext)0;
+			Init(Context, N, V, L);
+			float SatNoL = saturate(Context.NoL);
+
+			const float3 DirLightBaseFactor = TotalTransmittance * BSDF.Weight * LColor * LShadow;
+			const float3 DirLightGeometryFactor = DirLightBaseFactor * SatNoL;
+			const float3 SkyLightCommonFactor = TotalTransmittance * BSDF.Weight * ResolvedView.SkyLightColor.rgb * IndirectOcclusion;
+
+			switch (BSDF.BSDFType)
+			{
+			case STRATA_BSDF_TYPE_DIFFUSE_ON:
+			{
+				const float SafeRoughness = MakeRoughnessSafe(BSDF.Diffuse.Roughness);
+				Color += Diffuse_OrenNayar(BSDF.Diffuse.Albedo, SafeRoughness, Context.NoV, SatNoL, Context.VoH) * DirLightGeometryFactor;
+
+#if ENABLE_SKY_LIGHT
+				BRANCH
+					if (ReflectionStruct.SkyLightParameters.y > 0)
+					{
+						Color += GetEffectiveSkySHDiffuse(N) * SkyLightCommonFactor * Diffuse_Lambert(BSDF.Diffuse.Albedo); // STRATA_TODO should not be lambert
+					}
+#endif
+
+				//LayerTransmittance += 0.0f;
+				break;
+			}
+
+			case STRATA_BSDF_TYPE_DIFFUSE_CHAN:
+			{
+				const float SafeRoughness = MakeRoughnessSafe(BSDF.Diffuse.Roughness);
+				float a = SafeRoughness * SafeRoughness;
+				float a2 = a * a;
+				Color += Diffuse_Chan(BSDF.Diffuse.Albedo, a2, Context.NoV, SatNoL, Context.VoH, Context.NoH) * DirLightGeometryFactor;
+
+#if ENABLE_SKY_LIGHT
+				BRANCH
+					if (ReflectionStruct.SkyLightParameters.y > 0)
+					{
+						Color += GetEffectiveSkySHDiffuse(N) * SkyLightCommonFactor * Diffuse_Lambert(BSDF.Diffuse.Albedo); // STRATA_TODO should not be lambert
+					}
+#endif
+
+				//LayerTransmittance += 0.0f;
+				break;
+			}
+
+			case STRATA_BSDF_TYPE_DIELECTRIC:
+			{
+				float F0 = DielectricIorToF0(BSDF.Dielectric.IOR);
+
+				const float SafeRoughness = MakeRoughnessSafe(BSDF.Dielectric.Roughness.x);
+				float a2 = Pow4(SafeRoughness);
+				float Energy = 1.0f; // See EnergyNormalization
+
+				// Generalized microfacet specular
+				float D = D_GGX(a2, Context.NoH) * Energy;
+				float Vis = Vis_SmithJointApprox(a2, Context.NoV, SatNoL);
+				float3 F = F_Schlick(F0, Context.VoH);
+
+				Color += D * Vis * F * BSDF.Dielectric.Tint * DirLightGeometryFactor;
+
+#if ENABLE_SKY_LIGHT
+				BRANCH
+					if (ReflectionStruct.SkyLightParameters.y > 0)
+					{
+						float3 EvalEnvBRDF = EnvBRDFApprox(F0, SafeRoughness, Context.NoV);
+						float SkyAverageBrightness = 1.0f;
+						Color += GetSkyLightReflection(R, SafeRoughness, SkyAverageBrightness) * EvalEnvBRDF * BSDF.Dielectric.Tint * SkyLightCommonFactor;
+					}
+#endif
+
+				LayerTransmittance += BSDF.Weight * (1.0f - F); // STRATA_TODO transmittance should only come from coverage
+				break;
+			}
+
+			case STRATA_BSDF_TYPE_CONDUCTOR:
+			{
+				float3 ReflectivityF0 = float3(1.0, 0.1, 0.1);
+				float3 EdgeTint = float3(0.1, 1.0, 0.1);
+				ComputeF0AndEdgeTintFromIOR(BSDF.Conductor.IOR, BSDF.Conductor.Extinction, ReflectivityF0, EdgeTint);
+
+				// See SpecularGGX, isotropic verion
+
+				const float SafeRoughness = MakeRoughnessSafe(BSDF.Conductor.Roughness.x);
+				float a2 = Pow4(SafeRoughness);
+				float Energy = 1.0f; // See EnergyNormalization
+
+				// Generalized microfacet specular
+				float D = D_GGX(a2, Context.NoH) * Energy;
+				float Vis = Vis_SmithJointApprox(a2, Context.NoV, SatNoL);
+				float3 F = F_Schlick(ReflectivityF0, EdgeTint, Context.VoH);
+
+				Color += D * Vis * F * DirLightGeometryFactor;
+
+#if ENABLE_SKY_LIGHT
+				//BRANCH
+				if (ReflectionStruct.SkyLightParameters.y > 0)
+				{
+					float3 EvalEnvBRDF = EnvBRDFApprox(ReflectivityF0, EdgeTint, SafeRoughness, Context.NoV);
+					float SkyAverageBrightness = 1.0f;
+					Color += GetSkyLightReflection(R, SafeRoughness, SkyAverageBrightness) * EvalEnvBRDF * SkyLightCommonFactor;
+				}
+#endif
+
+				LayerTransmittance += BSDF.Weight * (1.0f - F); // STRATA_TODO use 1-FGD instead of simple 1-F, event if we ignore ray scattering/spreading
+				break;
+			}
+
+			case STRATA_BSDF_TYPE_VOLUME:
+			{
+				const float SafeRoughness = MakeRoughnessSafe(BSDF.Conductor.Roughness.x);
+
+				const float3 ScatteringCoefficients = BSDF.Volume.Scattering;
+				const float3 AbsorptionCoefficients = BSDF.Volume.Absorption;
+				const float3 ExtinctionCoefficients = ScatteringCoefficients + AbsorptionCoefficients;
+
+				const float3 SafeExtinctionThreshold = 0.000001f;
+				const float3 SafeExtinctionCoefficients = max(SafeExtinctionThreshold, ExtinctionCoefficients);
+				const float3 SafeAlbedo = ScatteringCoefficients / SafeExtinctionCoefficients;
+
+				const float Thickness = 0.01f;	// STRATA_TODO where does thickness come from? For now, 1cm
+				const float PathLength = Thickness / max(0.0001f, abs(Context.NoV));
+				const float3 SafePathSegmentTransmittance = exp(-SafeExtinctionCoefficients * PathLength);
+				const float3 NormLuminanceTransfert = 1.0f * ScatteringCoefficients;
+				float3 NormLuminanceSingleScattering = (NormLuminanceTransfert - NormLuminanceTransfert * SafePathSegmentTransmittance) / SafeExtinctionCoefficients;
+
+#if 1
+				const float3 SelfShadowTransmittance = SatNoL;
+#else
+				const float LightPathLength = 1.0f / max(0.0001f, SatNoL);
+				const float3 SelfShadowTransmittance = SatNoL <= 0.0 ? 0.0 : exp(-ExtinctionCoefficients * LightPathLength);
+#endif
+
+				const float Phase = HenyeyGreensteinPhase(BSDF.Volume.Anisotropy, Context.VoL);
+				Color += NormLuminanceSingleScattering * Phase * DirLightBaseFactor * SelfShadowTransmittance;
+
+#if ENABLE_SKY_LIGHT
+				//BRANCH
+				if (ReflectionStruct.SkyLightParameters.y > 0)
+				{
+					const float RescaleWithPhase = 2.0f * PI * IsotropicPhase(); // times 2PI to recover luminance, and then times the phase function for the scattering event
+					float3 EvalEnvBRDF = NormLuminanceSingleScattering * RescaleWithPhase;
+					float SafeRoughnessFromG = clamp(1.0 - abs(BSDF.Volume.Anisotropy), 0.001, 1.0f);
+					float SkyAverageBrightness = 1.0f;
+					// STRATA_TODO 
+					//	- Using R is weird. We should use V. But for rought material when light is incoming mainly from upper hemisphere, using N would make more send.
+					//	- Then what phase should be used? Phase(VoN)? We need to fund a good mapping on there. 
+					//	- We need to clarify what must be the weight there also (with env light in participating media coming from 2pi). See above attempt RescaleWithPhase.
+					Color += GetSkyLightReflection(R, SafeRoughnessFromG, SkyAverageBrightness) * EvalEnvBRDF * SkyLightCommonFactor;
+				}
+#endif
+
+				LayerTransmittance += BSDF.Weight * SafePathSegmentTransmittance;
+				break;
+			}
+			}
+		}
+
+		TotalTransmittance *= LayerTransmittance;	// STRATA_TODO we need to go over all weight and trnasmittance ot make sure everything would behave corretly. Especially with a Add node for instance. Add saturate where neded, etc.
+	}
+	return Color;
+}
+
+
diff --git a/Engine/Source/Developer/MaterialBaking/Private/ExportMaterialProxy.h b/Engine/Source/Developer/MaterialBaking/Private/ExportMaterialProxy.h
--- a/Engine/Source/Developer/MaterialBaking/Private/ExportMaterialProxy.h
+++ b/Engine/Source/Developer/MaterialBaking/Private/ExportMaterialProxy.h
@@ -431,6 +431,11 @@ public:
 		{
 			return MaterialInterface->CompileProperty(Compiler, MP_ShadingModel);
 
+		}
+		else if (Property == MP_FrontMaterial)
+		{
+			return MaterialInterface->CompileProperty(Compiler, MP_FrontMaterial);
+
 		}
 		else
 		{
diff --git a/Engine/Source/Developer/MaterialUtilities/Private/MaterialUtilities.cpp b/Engine/Source/Developer/MaterialUtilities/Private/MaterialUtilities.cpp
--- a/Engine/Source/Developer/MaterialUtilities/Private/MaterialUtilities.cpp
+++ b/Engine/Source/Developer/MaterialUtilities/Private/MaterialUtilities.cpp
@@ -583,6 +583,8 @@ public:
 				break;
 			case MP_ShadingModel:
 				return MaterialInterface->CompileProperty(&ProxyCompiler, MP_ShadingModel);
+			case MP_FrontMaterial:
+				return MaterialInterface->CompileProperty(&ProxyCompiler, MP_FrontMaterial);
 			default:
 				return Compiler->Constant(1.0f);
 			}
@@ -603,6 +605,10 @@ public:
 		{
 			return MaterialInterface->CompileProperty(Compiler, MP_ShadingModel);
 		}
+		else if (Property == MP_FrontMaterial)
+		{
+			return MaterialInterface->CompileProperty(Compiler, MP_FrontMaterial);
+		}
 		else
 		{
 			return Compiler->Constant(1.0f);
diff --git a/Engine/Source/Editor/MaterialEditor/Private/MaterialEditor.cpp b/Engine/Source/Editor/MaterialEditor/Private/MaterialEditor.cpp
--- a/Engine/Source/Editor/MaterialEditor/Private/MaterialEditor.cpp
+++ b/Engine/Source/Editor/MaterialEditor/Private/MaterialEditor.cpp
@@ -232,6 +232,10 @@ int32 FMatExpressionPreview::CompilePropertyAndSetMaterialProperty(EMaterialProp
 		FMaterialShadingModelField ShadingModels = Compiler->GetMaterialShadingModels();
 		Ret = Compiler->ShadingModel(ShadingModels.GetFirstShadingModel());
 	}
+	else if (Property == MP_FrontMaterial)
+	{
+		Ret = Compiler->FrontMaterial();
+	}
 	else
 	{
 		Ret = Compiler->Constant(1.0f);
@@ -4052,6 +4056,7 @@ UClass* FMaterialEditor::GetOnPromoteToParameterClass(const UEdGraphPin* TargetP
 			case MP_PixelDepthOffset:
 			case MP_ShadingModel:
 			case MP_OpacityMask:
+			case MP_FrontMaterial:
 				return UMaterialExpressionScalarParameter::StaticClass();
 
 			case MP_WorldPositionOffset:
diff --git a/Engine/Source/Editor/UnrealEd/Private/Lightmass/LightmassRender.cpp b/Engine/Source/Editor/UnrealEd/Private/Lightmass/LightmassRender.cpp
--- a/Engine/Source/Editor/UnrealEd/Private/Lightmass/LightmassRender.cpp
+++ b/Engine/Source/Editor/UnrealEd/Private/Lightmass/LightmassRender.cpp
@@ -452,6 +452,7 @@ public:
 				break;
 			case MP_ShadingModel:
 				return MaterialInterface->CompileProperty(&ProxyCompiler, MP_ShadingModel);
+			// STRATA_TODO
 			default:
 				return Compiler->Constant(1.0f);
 			}
@@ -472,6 +473,10 @@ public:
 		{
 			return MaterialInterface->CompileProperty(Compiler, MP_ShadingModel);
 		}
+		else if (Property == MP_FrontMaterial)
+		{
+			return MaterialInterface->CompileProperty(Compiler, MP_FrontMaterial);
+		}
 		else
 		{
 			return Compiler->Constant(1.0f);
diff --git a/Engine/Source/Editor/UnrealEd/Private/MaterialGraph.cpp b/Engine/Source/Editor/UnrealEd/Private/MaterialGraph.cpp
--- a/Engine/Source/Editor/UnrealEd/Private/MaterialGraph.cpp
+++ b/Engine/Source/Editor/UnrealEd/Private/MaterialGraph.cpp
@@ -86,6 +86,7 @@ void UMaterialGraph::RebuildGraphInternal(const TMap<UMaterialExpression*, TArra
 
 		MaterialInputs.Add(FMaterialInputInfo(FMaterialAttributeDefinitionMap::GetDisplayNameForMaterial(MP_PixelDepthOffset, Material), MP_PixelDepthOffset, LOCTEXT("PixelDepthOffsetToolTip", "Pixel Depth Offset")));
 		MaterialInputs.Add(FMaterialInputInfo(FMaterialAttributeDefinitionMap::GetDisplayNameForMaterial(MP_ShadingModel, Material), MP_ShadingModel, LOCTEXT("ShadingModelToolTip", "Selects which shading model should be used per pixel")));
+		MaterialInputs.Add(FMaterialInputInfo(FMaterialAttributeDefinitionMap::GetDisplayNameForMaterial(MP_FrontMaterial, Material), MP_FrontMaterial, LOCTEXT("FrontMaterialToolTip", "Specify the front facing material")));
 
 		//^^^ New material properties go above here. ^^^^
 		MaterialInputs.Add(FMaterialInputInfo(LOCTEXT("MaterialAttributes", "Material Attributes"), MP_MaterialAttributes, LOCTEXT("MaterialAttributesToolTip", "Material Attributes")));
diff --git a/Engine/Source/Runtime/Engine/Classes/Materials/Material.h b/Engine/Source/Runtime/Engine/Classes/Materials/Material.h
--- a/Engine/Source/Runtime/Engine/Classes/Materials/Material.h
+++ b/Engine/Source/Runtime/Engine/Classes/Materials/Material.h
@@ -233,6 +233,14 @@ struct FShadingModelMaterialInput : public FMaterialInput
 };
 #endif
 
+#if !CPP      //noexport struct
+USTRUCT(noexport)
+struct FStrataMaterialInput : public FMaterialInput
+{
+	// No support for constant
+};
+#endif
+
 #if !CPP      //noexport struct
 USTRUCT(noexport)
 struct FVectorMaterialInput : public FMaterialInput
@@ -465,6 +473,9 @@ public:
 	UPROPERTY()
 	FShadingModelMaterialInput ShadingModelFromMaterialExpression;
 
+	UPROPERTY()
+	FStrataMaterialInput FrontMaterial;
+
 	/** Indicates that the material should be rendered in the SeparateTranslucency Pass (not affected by DOF, requires bAllowSeparateTranslucency to be set in .ini). */
 	UPROPERTY(EditAnywhere, Category=Translucency, meta=(DisplayName = "Render After DOF"), AdvancedDisplay)
 	uint8 bEnableSeparateTranslucency : 1;
@@ -1748,11 +1759,13 @@ public:
 	bool HasBaseColorConnected() const { return BaseColor.IsConnected(); }
 	bool HasRoughnessConnected() const { return Roughness.IsConnected(); }
 	bool HasAmbientOcclusionConnected() const { return AmbientOcclusion.IsConnected(); }
+	bool HasStrataFrontMaterialConnected() const { return FrontMaterial.IsConnected(); }
 #else	
 	// Add to runtime data only if we need to call these at runtime
 	bool HasBaseColorConnected() const { check(0); return false; }
 	bool HasRoughnessConnected() const { check(0); return false; }
 	bool HasAmbientOcclusionConnected() const { check(0); return false; }
+	bool HasStrataFrontMaterialConnected() const { check(0); return false; }
 #endif 	
 	bool HasNormalConnected() const { return Normal.IsConnected(); }
 	bool HasSpecularConnected() const { return Specular.IsConnected(); }
diff --git a/Engine/Source/Runtime/Engine/Classes/Materials/MaterialExpressionMakeMaterialAttributes.h b/Engine/Source/Runtime/Engine/Classes/Materials/MaterialExpressionMakeMaterialAttributes.h
--- a/Engine/Source/Runtime/Engine/Classes/Materials/MaterialExpressionMakeMaterialAttributes.h
+++ b/Engine/Source/Runtime/Engine/Classes/Materials/MaterialExpressionMakeMaterialAttributes.h
@@ -77,6 +77,9 @@ class UMaterialExpressionMakeMaterialAttributes : public UMaterialExpression
 	UPROPERTY()
 	FExpressionInput ShadingModel;
 
+	UPROPERTY()
+	FExpressionInput FrontMaterial;
+
 	//~ Begin UObject Interface
 	virtual void Serialize(FStructuredArchive::FRecord Record) override;
 	//~ End UObject Interface
diff --git a/Engine/Source/Runtime/Engine/Classes/Materials/MaterialExpressionStrata.h b/Engine/Source/Runtime/Engine/Classes/Materials/MaterialExpressionStrata.h
new file mode 100644
--- /dev/null
+++ b/Engine/Source/Runtime/Engine/Classes/Materials/MaterialExpressionStrata.h
@@ -0,0 +1,365 @@
+// Copyright Epic Games, Inc. All Rights Reserved.
+
+#pragma once
+
+#include "CoreMinimal.h"
+#include "Engine/EngineTypes.h"
+#include "Materials/MaterialExpression.h"
+#include "MaterialExpressionStrata.generated.h"
+
+
+
+///////////////////////////////////////////////////////////////////////////////
+// BSDF nodes
+
+// This would be needed to for a common node interface and weight input and normal too?
+/*UCLASS(MinimalAPI, collapsecategories, hidecategories = Object)
+class UMaterialExpressionStrataBSDF : public UMaterialExpression
+{
+	GENERATED_UCLASS_BODY()
+
+	UPROPERTY()
+	FExpressionInput Weight;
+
+	// Normal?
+}*/
+
+UCLASS(MinimalAPI, collapsecategories, hidecategories = Object)
+class UMaterialExpressionStrataDiffuseBSDF : public UMaterialExpression // STRATA_TODO the single diffuse model to keep when we remove al lthe tests
+{
+	GENERATED_UCLASS_BODY()
+
+	/**
+	 * float3
+	 */
+	UPROPERTY()
+	FExpressionInput Albedo;
+
+	/**
+	 * float
+	 */
+	UPROPERTY()
+	FExpressionInput Roughness;
+
+	/**
+	 * float3
+	 */
+	UPROPERTY()
+	FExpressionInput Normal;
+
+	//~ Begin UMaterialExpression Interface
+#if WITH_EDITOR
+	virtual int32 Compile(class FMaterialCompiler* Compiler, int32 OutputIndex) override;
+	virtual void GetCaption(TArray<FString>& OutCaptions) const override;
+	virtual uint32 GetOutputType(int32 OutputIndex) override;
+	virtual uint32 GetInputType(int32 InputIndex) override;
+#endif
+	//~ End UMaterialExpression Interface
+};
+
+UCLASS(MinimalAPI, collapsecategories, hidecategories = Object)
+class UMaterialExpressionStrataDiffuseChanBSDF : public UMaterialExpression
+{
+	GENERATED_UCLASS_BODY()
+
+	/**
+	 * float3
+	 */
+	UPROPERTY()
+	FExpressionInput Albedo;
+
+	/**
+	 * float
+	 */
+	UPROPERTY()
+	FExpressionInput Roughness;
+
+	/**
+	 * float3
+	 */
+	UPROPERTY()
+	FExpressionInput Normal;
+
+	//~ Begin UMaterialExpression Interface
+#if WITH_EDITOR
+	virtual int32 Compile(class FMaterialCompiler* Compiler, int32 OutputIndex) override;
+	virtual void GetCaption(TArray<FString>& OutCaptions) const override;
+	virtual uint32 GetOutputType(int32 OutputIndex) override;
+	virtual uint32 GetInputType(int32 InputIndex) override;
+#endif
+	//~ End UMaterialExpression Interface
+};
+
+UCLASS(MinimalAPI, collapsecategories, hidecategories = Object)
+class UMaterialExpressionStrataDielectricBSDF : public UMaterialExpression
+{
+	GENERATED_UCLASS_BODY()
+	
+	/**
+	 * float
+	 */
+	UPROPERTY()
+	FExpressionInput IOR;
+	// STRATA_TODO Refraction IOR?
+
+	/**
+	 * float3
+	 */
+	UPROPERTY()
+	FExpressionInput Tint;
+		
+	/**
+	 * float2
+	 */
+	UPROPERTY()
+	FExpressionInput Roughness;
+
+	/**
+	 * float3
+	 */
+	UPROPERTY()
+	FExpressionInput Normal;
+	// STRATA_TODO Tangent
+	// STRATA_TODO Distribution
+	// STRATA_TODO ScatterMode
+
+	//~ Begin UMaterialExpression Interface
+#if WITH_EDITOR
+	virtual int32 Compile(class FMaterialCompiler* Compiler, int32 OutputIndex) override;
+	virtual void GetCaption(TArray<FString>& OutCaptions) const override;
+	virtual uint32 GetOutputType(int32 OutputIndex) override;
+	virtual uint32 GetInputType(int32 InputIndex) override;
+#endif
+	//~ End UMaterialExpression Interface
+};
+
+UCLASS(MinimalAPI, collapsecategories, hidecategories = Object)
+class UMaterialExpressionStrataConductorBSDF : public UMaterialExpression
+{
+	GENERATED_UCLASS_BODY()
+		
+	/**
+	 * float3
+	 */
+	UPROPERTY()
+	FExpressionInput IOR;
+	
+	/**
+	 * float3
+	 */
+	UPROPERTY()
+	FExpressionInput Extinction;
+
+	/**
+	 * float2
+	 */
+	UPROPERTY()
+	FExpressionInput Roughness;
+
+	/**
+	 * float3
+	 */
+	UPROPERTY()
+	FExpressionInput Normal;
+
+	// STRATA_TODO Tangent
+	// STRATA_TODO Distribution
+	// STRATA_TODO ScatterMode
+
+	//~ Begin UMaterialExpression Interface
+#if WITH_EDITOR
+	virtual int32 Compile(class FMaterialCompiler* Compiler, int32 OutputIndex) override;
+	virtual void GetCaption(TArray<FString>& OutCaptions) const override;
+	virtual uint32 GetOutputType(int32 OutputIndex) override;
+	virtual uint32 GetInputType(int32 InputIndex) override;
+#endif
+	//~ End UMaterialExpression Interface
+};
+
+UCLASS(MinimalAPI, collapsecategories, hidecategories = Object)
+class UMaterialExpressionStrataVolumeBSDF : public UMaterialExpression
+{
+	GENERATED_UCLASS_BODY()
+		
+	/**
+	 * float3
+	 */
+	UPROPERTY()
+	FExpressionInput Absorption;
+	
+	/**
+	 * float3
+	 */
+	UPROPERTY()
+	FExpressionInput Scattering;
+
+	/**
+	 * float
+	 */
+	UPROPERTY()
+	FExpressionInput Anisotropy;
+
+	//~ Begin UMaterialExpression Interface
+#if WITH_EDITOR
+	virtual int32 Compile(class FMaterialCompiler* Compiler, int32 OutputIndex) override;
+	virtual void GetCaption(TArray<FString>& OutCaptions) const override;
+	virtual uint32 GetOutputType(int32 OutputIndex) override;
+	virtual uint32 GetInputType(int32 InputIndex) override;
+#endif
+	//~ End UMaterialExpression Interface
+};
+
+// STRATA_TODO Sheen, Subsurface, thinfilm, generalised schlick
+
+
+
+///////////////////////////////////////////////////////////////////////////////
+// Operator nodes
+
+UCLASS(MinimalAPI, collapsecategories, hidecategories = Object)
+class UMaterialExpressionStrataHorizontalMixing : public UMaterialExpression
+{
+	GENERATED_UCLASS_BODY()
+		
+	/**
+	 * Strata material
+	 */
+	UPROPERTY()
+	FExpressionInput Foreground;
+
+	/**
+	 * Strata material
+	 */
+	UPROPERTY()
+	FExpressionInput Background;
+
+	/**
+	 * float
+	 */
+	UPROPERTY()
+	FExpressionInput Mix;
+
+	//~ Begin UMaterialExpression Interface
+#if WITH_EDITOR
+	virtual int32 Compile(class FMaterialCompiler* Compiler, int32 OutputIndex) override;
+	virtual void GetCaption(TArray<FString>& OutCaptions) const override;
+	virtual uint32 GetOutputType(int32 OutputIndex) override;
+	virtual uint32 GetInputType(int32 InputIndex) override;
+#endif
+	//~ End UMaterialExpression Interface
+};
+
+UCLASS(MinimalAPI, collapsecategories, hidecategories = Object)
+class UMaterialExpressionStrataVerticalLayering : public UMaterialExpression
+{
+	GENERATED_UCLASS_BODY()
+
+	/**
+	 * Strata material
+	 */
+	UPROPERTY()
+	FExpressionInput Top;
+	
+	/**
+	 * Strata material
+	 */
+	UPROPERTY()
+	FExpressionInput Base;
+
+	//~ Begin UMaterialExpression Interface
+#if WITH_EDITOR
+	virtual int32 Compile(class FMaterialCompiler* Compiler, int32 OutputIndex) override;
+	virtual void GetCaption(TArray<FString>& OutCaptions) const override;
+	virtual uint32 GetOutputType(int32 OutputIndex) override;
+	virtual uint32 GetInputType(int32 InputIndex) override;
+#endif
+	//~ End UMaterialExpression Interface
+};
+
+UCLASS(MinimalAPI, collapsecategories, hidecategories = Object)
+class UMaterialExpressionStrataAdd : public UMaterialExpression
+{
+	GENERATED_UCLASS_BODY()
+
+	/**
+	 * Strata material
+	 */
+	UPROPERTY()
+	FExpressionInput A;
+	
+	/**
+	 * Strata material
+	 */
+	UPROPERTY()
+	FExpressionInput B;
+
+	//~ Begin UMaterialExpression Interface
+#if WITH_EDITOR
+	virtual int32 Compile(class FMaterialCompiler* Compiler, int32 OutputIndex) override;
+	virtual void GetCaption(TArray<FString>& OutCaptions) const override;
+	virtual uint32 GetOutputType(int32 OutputIndex) override;
+	virtual uint32 GetInputType(int32 InputIndex) override;
+#endif
+	//~ End UMaterialExpression Interface
+};
+
+UCLASS(MinimalAPI, collapsecategories, hidecategories = Object)
+class UMaterialExpressionStrataMultiply : public UMaterialExpression
+{
+	GENERATED_UCLASS_BODY()
+
+	/**
+	 * Strata material
+	 */
+	UPROPERTY()
+	FExpressionInput A;
+	
+	/**
+	 * Strata material
+	 */
+	UPROPERTY()
+	FExpressionInput Weight;
+
+	//~ Begin UMaterialExpression Interface
+#if WITH_EDITOR
+	virtual int32 Compile(class FMaterialCompiler* Compiler, int32 OutputIndex) override;
+	virtual void GetCaption(TArray<FString>& OutCaptions) const override;
+	virtual uint32 GetOutputType(int32 OutputIndex) override;
+	virtual uint32 GetInputType(int32 InputIndex) override;
+#endif
+	//~ End UMaterialExpression Interface
+};
+
+
+
+///////////////////////////////////////////////////////////////////////////////
+// Utilities
+
+UCLASS(MinimalAPI, collapsecategories, hidecategories = Object)
+class UMaterialExpressionStrataArtisticIOR : public UMaterialExpression
+{
+	GENERATED_UCLASS_BODY()
+
+	/**
+	* TODO F0
+	*/
+	UPROPERTY()
+	FExpressionInput Reflectivity;
+
+	/**
+	 * TODO Color for tangent view direction
+	 */
+	UPROPERTY()
+	FExpressionInput EdgeColor;
+
+	//~ Begin UMaterialExpression Interface
+#if WITH_EDITOR
+	virtual int32 Compile(class FMaterialCompiler* Compiler, int32 OutputIndex) override;
+	virtual void GetCaption(TArray<FString>& OutCaptions) const override;
+	virtual uint32 GetOutputType(int32 OutputIndex) override;
+	virtual uint32 GetInputType(int32 InputIndex) override;
+#endif
+	//~ End UMaterialExpression Interface
+};
+
+
diff --git a/Engine/Source/Runtime/Engine/Private/Materials/HLSLMaterialTranslator.cpp b/Engine/Source/Runtime/Engine/Private/Materials/HLSLMaterialTranslator.cpp
--- a/Engine/Source/Runtime/Engine/Private/Materials/HLSLMaterialTranslator.cpp
+++ b/Engine/Source/Runtime/Engine/Private/Materials/HLSLMaterialTranslator.cpp
@@ -130,6 +130,7 @@ FHLSLMaterialTranslator::FHLSLMaterialTranslator(FMaterial* InMaterial,
 	SharedPixelProperties[MP_PixelDepthOffset] = true;
 	SharedPixelProperties[MP_SubsurfaceColor] = true;
 	SharedPixelProperties[MP_ShadingModel] = true;
+	SharedPixelProperties[MP_FrontMaterial] = true;
 
 	for (int32 Frequency = 0; Frequency < SF_NumFrequencies; ++Frequency)
 	{
@@ -653,7 +654,9 @@ bool FHLSLMaterialTranslator::Translate()
 
 		// Make sure to compile this property before using ShadingModelsFromCompilation
 		Chunk[MP_ShadingModel]					= Material->CompilePropertyAndSetMaterialProperty(MP_ShadingModel			,this);
-			
+		
+		Chunk[MP_FrontMaterial]					= Material->CompilePropertyAndSetMaterialProperty(MP_FrontMaterial			,this);
+
 		// Get shading models from material.
 		FMaterialShadingModelField MaterialShadingModels = Material->GetShadingModels(); 
 
@@ -1417,6 +1420,10 @@ void FHLSLMaterialTranslator::GetMaterialEnvironment(EShaderPlatform InPlatform,
 				VolumetricAdvancedNode->bGroundContribution ? TEXT("1") : TEXT("0"));
 		}
 	}
+
+	OutEnvironment.SetDefine(TEXT("MATERIAL_IS_STRATA"),
+		Material->HasStrataFrontMaterialConnected() ? TEXT("1") : TEXT("0"));
+	 
 }
 
 // Assign custom interpolators to slots, packing them as much as possible in unused slots.
@@ -1579,6 +1586,7 @@ FString FHLSLMaterialTranslator::GetMaterialShaderCode()
 		case MCT_Float3: MaterialAttributesDeclaration += FString::Printf(TEXT("\tfloat3 %s;") LINE_TERMINATOR, *PropertyName); break;
 		case MCT_Float4: MaterialAttributesDeclaration += FString::Printf(TEXT("\tfloat4 %s;") LINE_TERMINATOR, *PropertyName); break;
 		case MCT_ShadingModel: MaterialAttributesDeclaration += FString::Printf(TEXT("\tuint %s;") LINE_TERMINATOR, *PropertyName); break;
+		case MCT_Strata: MaterialAttributesDeclaration += FString::Printf(TEXT("\tFStrataData %s;") LINE_TERMINATOR, *PropertyName); break;
 		}
 	}
 
@@ -1878,6 +1886,7 @@ const TCHAR* FHLSLMaterialTranslator::DescribeType(EMaterialValueType Type) cons
 	case MCT_TextureVirtual:		return TEXT("TextureVirtual");
 	case MCT_VTPageTableResult:		return TEXT("VTPageTableResult");
 	case MCT_ShadingModel:			return TEXT("ShadingModel");
+	case MCT_Strata:				return TEXT("Strata");
 	default:						return TEXT("unknown");
 	};
 }
@@ -1902,6 +1911,7 @@ const TCHAR* FHLSLMaterialTranslator::HLSLTypeString(EMaterialValueType Type) co
 	case MCT_TextureVirtual:		return TEXT("TextureVirtual");
 	case MCT_VTPageTableResult:		return TEXT("VTPageTableResult");
 	case MCT_ShadingModel:			return TEXT("uint");
+	case MCT_Strata:				return TEXT("FStrataData");
 	default:						return TEXT("unknown");
 	};
 }
@@ -2448,7 +2458,7 @@ bool FHLSLMaterialTranslator::GetTextureForExpression(int32 Index, int32& OutTex
 // GetArithmeticResultType
 EMaterialValueType FHLSLMaterialTranslator::GetArithmeticResultType(EMaterialValueType TypeA, EMaterialValueType TypeB)
 {
-	if (!((TypeA & MCT_Float) || TypeA == MCT_ShadingModel) || !((TypeB & MCT_Float) || TypeB == MCT_ShadingModel))
+	if (!((TypeA & MCT_Float) || TypeA == MCT_ShadingModel) || !((TypeB & MCT_Float) || TypeB == MCT_ShadingModel) || TypeB == MCT_Strata)
 	{
 		Errorf(TEXT("Attempting to perform arithmetic on non-numeric types: %s %s"),DescribeType(TypeA),DescribeType(TypeB));
 		return MCT_Unknown;
@@ -7032,6 +7042,110 @@ int32 FHLSLMaterialTranslator::ShadingModel(EMaterialShadingModel InSelectedShad
 	return AddInlinedCodeChunk(MCT_ShadingModel, TEXT("%d"), InSelectedShadingModel);
 }
 
+int32 FHLSLMaterialTranslator::FrontMaterial()
+{
+	return AddInlinedCodeChunk(MCT_Strata, TEXT("GetInitialisedStrataData()"));
+}
+
+int32 FHLSLMaterialTranslator::StrataDiffuseOrenNayarBSDF(int32 Albedo, int32 Roughness, int32 Normal)
+{
+	return AddInlinedCodeChunk(
+		MCT_Strata, TEXT("GetStrataDiffuseOrenNayarBSDF(%s, %s, %s)"),
+		*GetParameterCode(Albedo),
+		*GetParameterCode(Roughness),
+		*GetParameterCode(Normal)
+	);
+}
+
+int32 FHLSLMaterialTranslator::StrataDiffuseChanBSDF(int32 Albedo, int32 Roughness, int32 Normal)
+{
+	return AddInlinedCodeChunk(
+		MCT_Strata, TEXT("GetStrataDiffuseChanBSDF(%s, %s, %s)"),
+		*GetParameterCode(Albedo),
+		*GetParameterCode(Roughness),
+		*GetParameterCode(Normal)
+	);
+}
+
+int32 FHLSLMaterialTranslator::StrataDielectricBSDF(int32 Roughness, int32 IOR, int32 Tint, int32 Normal)
+{
+	return AddInlinedCodeChunk(
+		MCT_Strata, TEXT("GetStrataDielectricBSDF(%s, %s, %s, %s)"),
+		*GetParameterCode(Roughness),
+		*GetParameterCode(IOR),
+		*GetParameterCode(Tint),
+		*GetParameterCode(Normal)
+	);
+}
+
+int32 FHLSLMaterialTranslator::StrataConductorBSDF(int32 IOR, int32 Extinction, int32 Roughness, int32 Normal)
+{
+	return AddInlinedCodeChunk(
+		MCT_Strata, TEXT("GetStrataConductorBSDF(%s, %s, %s, %s)"),
+		*GetParameterCode(IOR),
+		*GetParameterCode(Extinction),
+		*GetParameterCode(Roughness),
+		*GetParameterCode(Normal)
+	);
+}
+
+int32 FHLSLMaterialTranslator::StrataVolumeBSDF(int32 Absorption, int32 Scattering, int32 Anisotropy)
+{
+	return AddInlinedCodeChunk(
+		MCT_Strata, TEXT("GetStrataVolumeBSDF(%s, %s, %s)"),
+		*GetParameterCode(Absorption),
+		*GetParameterCode(Scattering),
+		*GetParameterCode(Anisotropy)
+	);
+}
+
+int32 FHLSLMaterialTranslator::StrataHorizontalMixing(int32 Foreground, int32 Background, int32 Mix)
+{
+	return AddInlinedCodeChunk(
+		MCT_Strata, TEXT("StrataHorizontalMixing(%s, %s, %s)"),
+		*GetParameterCode(Foreground),
+		*GetParameterCode(Background),
+		*GetParameterCode(Mix)
+	);
+}
+
+int32 FHLSLMaterialTranslator::StrataVerticalLayering(int32 Top, int32 Base)
+{
+	return AddInlinedCodeChunk(
+		MCT_Strata, TEXT("StrataVerticalLayering(%s, %s)"),
+		*GetParameterCode(Top),
+		*GetParameterCode(Base)
+	);
+}
+
+int32 FHLSLMaterialTranslator::StrataAdd(int32 A, int32 B)
+{
+	return AddInlinedCodeChunk(
+		MCT_Strata, TEXT("StrataAdd(%s, %s)"),
+		*GetParameterCode(A),
+		*GetParameterCode(B)
+	);
+}
+
+int32 FHLSLMaterialTranslator::StrataMultiply(int32 A, int32 Weight)
+{
+	return AddInlinedCodeChunk(
+		MCT_Strata, TEXT("StrataMultiply(%s, %s)"),
+		*GetParameterCode(A),
+		*GetParameterCode(Weight)
+	);
+}
+
+int32 FHLSLMaterialTranslator::StrataArtisticIOR(int32 Reflectivity, int32 EdgeColor, int32 OutputIndex)
+{
+	return AddInlinedCodeChunk(
+		MCT_Strata, 
+		OutputIndex == 0 ? TEXT("ComputeComplexIORFromF0AndEdgeTint_IOR(%s, %s)") : TEXT("ComputeComplexIORFromF0AndEdgeTint_Extinction(%s, %s)"),
+		*GetParameterCode(Reflectivity),
+		*GetParameterCode(EdgeColor)
+	);
+}
+
 int32 FHLSLMaterialTranslator::MapARPassthroughCameraUV(int32 UV)
 {
 	if (UV == INDEX_NONE)
diff --git a/Engine/Source/Runtime/Engine/Private/Materials/HLSLMaterialTranslator.h b/Engine/Source/Runtime/Engine/Private/Materials/HLSLMaterialTranslator.h
--- a/Engine/Source/Runtime/Engine/Private/Materials/HLSLMaterialTranslator.h
+++ b/Engine/Source/Runtime/Engine/Private/Materials/HLSLMaterialTranslator.h
@@ -758,6 +758,19 @@ protected:
 
 	virtual int32 VirtualTextureOutput(uint8 MaterialAttributeMask) override;
 
+	// Strata
+	virtual int32 FrontMaterial() override;
+	virtual int32 StrataDiffuseOrenNayarBSDF(int32 Albedo, int32 Roughness, int32 Normal) override;
+	virtual int32 StrataDiffuseChanBSDF(int32 Albedo, int32 Roughness, int32 Normal) override;
+	virtual int32 StrataDielectricBSDF(int32 Roughness, int32 IOR, int32 Tint, int32 Normal) override;
+	virtual int32 StrataConductorBSDF(int32 IOR, int32 Extinction, int32 Roughness, int32 Normal) override;
+	virtual int32 StrataVolumeBSDF(int32 Absorption, int32 Scattering, int32 Anisotropy) override;
+	virtual int32 StrataHorizontalMixing(int32 Foreground, int32 Background, int32 Mix) override;
+	virtual int32 StrataVerticalLayering(int32 Top, int32 Base) override;
+	virtual int32 StrataAdd(int32 A, int32 B) override;
+	virtual int32 StrataMultiply(int32 A, int32 Weight) override;
+	virtual int32 StrataArtisticIOR(int32 Reflectivity, int32 EdgeColor, int32 OutputIndex) override;
+
 #if HANDLE_CUSTOM_OUTPUTS_AS_MATERIAL_ATTRIBUTES
 	/** Used to translate code for custom output attributes such as ClearCoatBottomNormal */
 	void GenerateCustomAttributeCode(int32 OutputIndex, int32 OutputCode, EMaterialValueType OutputType, FString& DisplayName);
diff --git a/Engine/Source/Runtime/Engine/Private/Materials/Material.cpp b/Engine/Source/Runtime/Engine/Private/Materials/Material.cpp
--- a/Engine/Source/Runtime/Engine/Private/Materials/Material.cpp
+++ b/Engine/Source/Runtime/Engine/Private/Materials/Material.cpp
@@ -56,6 +56,7 @@
 #include "Materials/MaterialParameterCollection.h"
 #include "ShaderPlatformQualitySettings.h"
 #include "MaterialShaderQualitySettings.h"
+#include "Engine/RendererSettings.h"
 #include "ProfilingDebugging/LoadTimeTracker.h"
 #include "MaterialCompiler.h"
 #include "MaterialShaderType.h"
@@ -3311,7 +3312,7 @@ void UMaterial::Serialize(FArchive& Ar)
 	}
 #endif // #if WITH_EDITOR
 
-	static_assert(MP_MAX == 32, "New material properties must have DoMaterialAttributeReorder called on them to ensure that any future reordering of property pins is correctly applied.");
+	static_assert(MP_MAX == 33, "New material properties must have DoMaterialAttributeReorder called on them to ensure that any future reordering of property pins is correctly applied.");
 
 	if (Ar.UE4Ver() < VER_UE4_MATERIAL_MASKED_BLENDMODE_TIDY)
 	{
@@ -3499,6 +3500,7 @@ void UMaterial::PostLoad()
 	DoMaterialAttributeReorder(&CustomizedUVs[7], UE4Ver, RenderObjVer);
 	DoMaterialAttributeReorder(&PixelDepthOffset, UE4Ver, RenderObjVer);
 	DoMaterialAttributeReorder(&ShadingModelFromMaterialExpression, UE4Ver, RenderObjVer);
+	DoMaterialAttributeReorder(&FrontMaterial, UE4Ver, RenderObjVer);
 #endif // WITH_EDITORONLY_DATA
 
 	if (!IsDefaultMaterial())
@@ -4922,6 +4924,7 @@ FExpressionInput* UMaterial::GetExpressionInputForProperty(EMaterialProperty InP
 		case MP_MaterialAttributes:		return &MaterialAttributes;
 		case MP_PixelDepthOffset:		return &PixelDepthOffset;
 		case MP_ShadingModel:			return &ShadingModelFromMaterialExpression;
+		case MP_FrontMaterial:			return &FrontMaterial;
 	}
 
 	if (InProperty >= MP_CustomizedUVs0 && InProperty <= MP_CustomizedUVs7)
@@ -5485,6 +5488,7 @@ int32 UMaterial::CompilePropertyEx( FMaterialCompiler* Compiler, const FGuid& At
 		case MP_WorldDisplacement:		return WorldDisplacement.CompileWithDefault(Compiler, Property);
 		case MP_PixelDepthOffset:		return PixelDepthOffset.CompileWithDefault(Compiler, Property);
 		case MP_ShadingModel:			return ShadingModelFromMaterialExpression.CompileWithDefault(Compiler, Property);
+		case MP_FrontMaterial:			return FrontMaterial.CompileWithDefault(Compiler, Property);
 
 		default:
 			if (Property >= MP_CustomizedUVs0 && Property <= MP_CustomizedUVs7)
@@ -5973,7 +5977,13 @@ static bool IsPropertyActive_Internal(EMaterialProperty InProperty,
 		break;
 	case MP_ShadingModel:
 		Active = bUsesShadingModelFromMaterialExpression;
-                break;
+		break;
+	case MP_FrontMaterial:
+		{
+			const URendererSettings* RendererSettings = GetDefault<URendererSettings>();
+			Active = RendererSettings && RendererSettings->bEnableStrata;
+			break;
+		}
 	case MP_MaterialAttributes:
 	default:
 		Active = true;
diff --git a/Engine/Source/Runtime/Engine/Private/Materials/MaterialExpressions.cpp b/Engine/Source/Runtime/Engine/Private/Materials/MaterialExpressions.cpp
--- a/Engine/Source/Runtime/Engine/Private/Materials/MaterialExpressions.cpp
+++ b/Engine/Source/Runtime/Engine/Private/Materials/MaterialExpressions.cpp
@@ -242,6 +242,7 @@
 #include "Materials/MaterialExpressionMapARPassthroughCameraUV.h"
 #include "Materials/MaterialExpressionShaderStageSwitch.h"
 #include "Materials/MaterialExpressionReflectionCapturePassSwitch.h"
+#include "Materials/MaterialExpressionStrata.h"
 #include "Materials/MaterialUniformExpressions.h"
 #include "EditorSupportDelegates.h"
 #include "MaterialCompiler.h"
@@ -5617,7 +5618,7 @@ int32 UMaterialExpressionMakeMaterialAttributes::Compile(class FMaterialCompiler
 	int32 Ret = INDEX_NONE;
 	UMaterialExpression* Expression = nullptr;
 
- 	static_assert(MP_MAX == 32, 
+ 	static_assert(MP_MAX == 33, 
 		"New material properties should be added to the end of the inputs for this expression. \
 		The order of properties here should match the material results pins, the make material attriubtes node inputs and the mapping of IO indices to properties in GetMaterialPropertyFromInputOutputIndex().\
 		Insertions into the middle of the properties or a change in the order of properties will also require that existing data is fixed up in DoMaterialAttributeReorder().\
@@ -5647,6 +5648,7 @@ int32 UMaterialExpressionMakeMaterialAttributes::Compile(class FMaterialCompiler
 	case MP_Refraction: Ret = Refraction.Compile(Compiler); Expression = Refraction.Expression; break;
 	case MP_PixelDepthOffset: Ret = PixelDepthOffset.Compile(Compiler); Expression = PixelDepthOffset.Expression; break;
 	case MP_ShadingModel: Ret = ShadingModel.Compile(Compiler); Expression = ShadingModel.Expression; break;
+	case MP_FrontMaterial: Ret = FrontMaterial.Compile(Compiler); Expression = FrontMaterial.Expression; break;
 	};
 
 	if (Property >= MP_CustomizedUVs0 && Property <= MP_CustomizedUVs7)
@@ -5704,7 +5706,7 @@ UMaterialExpressionBreakMaterialAttributes::UMaterialExpressionBreakMaterialAttr
 
 	MenuCategories.Add(ConstructorStatics.NAME_MaterialAttributes);
 	
- 	static_assert(MP_MAX == 32, 
+ 	static_assert(MP_MAX == 33, 
 		"New material properties should be added to the end of the outputs for this expression. \
 		The order of properties here should match the material results pins, the make material attriubtes node inputs and the mapping of IO indices to properties in GetMaterialPropertyFromInputOutputIndex().\
 		Insertions into the middle of the properties or a change in the order of properties will also require that existing data is fixed up in DoMaterialAttriubtesReorder().\
@@ -5737,6 +5739,7 @@ UMaterialExpressionBreakMaterialAttributes::UMaterialExpressionBreakMaterialAttr
 
 	Outputs.Add(FExpressionOutput(TEXT("PixelDepthOffset"), 1, 1, 0, 0, 0));
 	Outputs.Add(FExpressionOutput(TEXT("ShadingModel"), 0, 0, 0, 0, 0));
+	Outputs.Add(FExpressionOutput(TEXT("FrontMaterial"), 0, 0, 0, 0, 0));
 #endif
 }
 
@@ -5818,6 +5821,7 @@ static void BuildPropertyToIOIndexMap()
 		PropertyToIOIndexMap.Add(MP_CustomizedUVs7,			25);
 		PropertyToIOIndexMap.Add(MP_PixelDepthOffset,		26);
 		PropertyToIOIndexMap.Add(MP_ShadingModel,			27);
+		PropertyToIOIndexMap.Add(MP_FrontMaterial,			28);
 	}
 }
 
@@ -18717,6 +18721,598 @@ FString UMaterialExpressionThinTranslucentMaterialOutput::GetDisplayName() const
 
 
 
+///////////////////////////////////////////////////////////////////////////////
+// Strata
+
+
+
+UMaterialExpressionStrataDiffuseBSDF::UMaterialExpressionStrataDiffuseBSDF(const FObjectInitializer& ObjectInitializer)
+	: Super(ObjectInitializer)
+{
+	struct FConstructorStatics
+	{
+		FText NAME_Strata;
+		FConstructorStatics() : NAME_Strata(LOCTEXT("Strata", "Strata")) { }
+	};
+	static FConstructorStatics ConstructorStatics;
+#if WITH_EDITORONLY_DATA
+	MenuCategories.Add(ConstructorStatics.NAME_Strata);
+#endif
+}
+
+#if WITH_EDITOR
+int32 UMaterialExpressionStrataDiffuseBSDF::Compile(class FMaterialCompiler* Compiler, int32 OutputIndex)
+{
+	if (!Albedo.GetTracedInput().Expression)
+	{
+		return Compiler->Errorf(TEXT("Missing Albedo input"));	// STRATA_TODO: defaults in the detail panel with auto read only if plug in
+	}
+	if (!Roughness.GetTracedInput().Expression)
+	{
+		return Compiler->Errorf(TEXT("Missing Roughness input"));	// STRATA_TODO: defaults in the detail panel with auto read only if plug in
+	}
+	if (!Normal.GetTracedInput().Expression)
+	{
+		return Compiler->Errorf(TEXT("Missing Normal input"));	// STRATA_TODO: defaults in the detail panel with auto read only if plug in
+	}
+	return Compiler->StrataDiffuseOrenNayarBSDF(Albedo.Compile(Compiler), Roughness.Compile(Compiler), Normal.Compile(Compiler));
+}
+
+void UMaterialExpressionStrataDiffuseBSDF::GetCaption(TArray<FString>& OutCaptions) const
+{
+	OutCaptions.Add(TEXT("Strata Diffuse Oren-Nayar BSDF"));
+}
+
+uint32 UMaterialExpressionStrataDiffuseBSDF::GetOutputType(int32 OutputIndex)
+{
+	return MCT_Strata;
+}
+
+uint32 UMaterialExpressionStrataDiffuseBSDF::GetInputType(int32 InputIndex)
+{
+	switch (InputIndex)
+	{
+	case 0:
+		return MCT_Float3;
+		break;
+	case 1:
+		return MCT_Float1;
+		break;
+	case 2:
+		return MCT_Float3;
+		break;
+	}
+
+	check(false);
+	return MCT_Float1;
+}
+#endif // WITH_EDITOR
+
+
+
+
+
+UMaterialExpressionStrataDiffuseChanBSDF::UMaterialExpressionStrataDiffuseChanBSDF(const FObjectInitializer& ObjectInitializer)
+	: Super(ObjectInitializer)
+{
+	struct FConstructorStatics
+	{
+		FText NAME_Strata;
+		FConstructorStatics() : NAME_Strata(LOCTEXT("Strata", "Strata")) { }
+	};
+	static FConstructorStatics ConstructorStatics;
+#if WITH_EDITORONLY_DATA
+	MenuCategories.Add(ConstructorStatics.NAME_Strata);
+#endif
+}
+
+#if WITH_EDITOR
+int32 UMaterialExpressionStrataDiffuseChanBSDF::Compile(class FMaterialCompiler* Compiler, int32 OutputIndex)
+{
+	if (!Albedo.GetTracedInput().Expression)
+	{
+		return Compiler->Errorf(TEXT("Missing Albedo input"));	// STRATA_TODO: defaults in the detail panel with auto read only if plug in
+	}
+	if (!Roughness.GetTracedInput().Expression)
+	{
+		return Compiler->Errorf(TEXT("Missing Roughness input"));	// STRATA_TODO: defaults in the detail panel with auto read only if plug in
+	}
+	if (!Normal.GetTracedInput().Expression)
+	{
+		return Compiler->Errorf(TEXT("Missing Normal input"));	// STRATA_TODO: defaults in the detail panel with auto read only if plug in
+	}
+	return Compiler->StrataDiffuseChanBSDF(Albedo.Compile(Compiler), Roughness.Compile(Compiler), Normal.Compile(Compiler));
+}
+
+void UMaterialExpressionStrataDiffuseChanBSDF::GetCaption(TArray<FString>& OutCaptions) const
+{
+	OutCaptions.Add(TEXT("Strata Diffuse Chan BSDF"));
+}
+
+uint32 UMaterialExpressionStrataDiffuseChanBSDF::GetOutputType(int32 OutputIndex)
+{
+	return MCT_Strata;
+}
+
+uint32 UMaterialExpressionStrataDiffuseChanBSDF::GetInputType(int32 InputIndex)
+{
+	switch (InputIndex)
+	{
+	case 0:
+		return MCT_Float3;
+		break;
+	case 1:
+		return MCT_Float1;
+		break;
+	case 2:
+		return MCT_Float3;
+		break;
+	}
+
+	check(false);
+	return MCT_Float1;
+}
+#endif // WITH_EDITOR
+
+
+
+
+
+UMaterialExpressionStrataDielectricBSDF::UMaterialExpressionStrataDielectricBSDF(const FObjectInitializer& ObjectInitializer)
+	: Super(ObjectInitializer)
+{
+	struct FConstructorStatics
+	{
+		FText NAME_Strata;
+		FConstructorStatics() : NAME_Strata(LOCTEXT("Strata", "Strata")) { }
+	};
+	static FConstructorStatics ConstructorStatics;
+#if WITH_EDITORONLY_DATA
+	MenuCategories.Add(ConstructorStatics.NAME_Strata);
+#endif
+}
+
+#if WITH_EDITOR
+int32 UMaterialExpressionStrataDielectricBSDF::Compile(class FMaterialCompiler* Compiler, int32 OutputIndex)
+{
+	if (!Tint.GetTracedInput().Expression)
+	{
+		return Compiler->Errorf(TEXT("Missing Tint input"));	// STRATA_TODO: defaults in the detail panel with auto read only if plug in
+	}
+	if (!IOR.GetTracedInput().Expression)
+	{
+		return Compiler->Errorf(TEXT("Missing IOR input"));	// STRATA_TODO: defaults in the detail panel with auto read only if plug in
+	}
+	if (!Roughness.GetTracedInput().Expression)
+	{
+		return Compiler->Errorf(TEXT("Missing Roughness input"));	// STRATA_TODO: defaults in the detail panel with auto read only if plug in
+	}
+	if (!Normal.GetTracedInput().Expression)
+	{
+		return Compiler->Errorf(TEXT("Missing Normal input"));	// STRATA_TODO: defaults in the detail panel with auto read only if plug in
+	}
+	return Compiler->StrataDielectricBSDF(Roughness.Compile(Compiler), IOR.Compile(Compiler), Tint.Compile(Compiler), Normal.Compile(Compiler));
+}
+
+void UMaterialExpressionStrataDielectricBSDF::GetCaption(TArray<FString>& OutCaptions) const
+{
+	OutCaptions.Add(TEXT("Strata Dielectric BSDF"));
+}
+
+uint32 UMaterialExpressionStrataDielectricBSDF::GetOutputType(int32 OutputIndex)
+{
+	return MCT_Strata;
+}
+
+uint32 UMaterialExpressionStrataDielectricBSDF::GetInputType(int32 InputIndex)
+{
+	switch (InputIndex)
+	{
+	case 0:
+		return MCT_Float1;
+		break;
+	case 1:
+		return MCT_Float3;
+		break;
+	case 2:
+		return MCT_Float2;
+		break;
+	case 3:
+		return MCT_Float3;
+		break;
+	}
+
+	check(false);
+	return MCT_Float1;
+}
+#endif // WITH_EDITOR
+
+
+
+
+
+UMaterialExpressionStrataConductorBSDF::UMaterialExpressionStrataConductorBSDF(const FObjectInitializer& ObjectInitializer)
+	: Super(ObjectInitializer)
+{
+	struct FConstructorStatics
+	{
+		FText NAME_Strata;
+		FConstructorStatics() : NAME_Strata(LOCTEXT("Strata", "Strata")) { }
+	};
+	static FConstructorStatics ConstructorStatics;
+#if WITH_EDITORONLY_DATA
+	MenuCategories.Add(ConstructorStatics.NAME_Strata);
+#endif
+}
+
+#if WITH_EDITOR
+int32 UMaterialExpressionStrataConductorBSDF::Compile(class FMaterialCompiler* Compiler, int32 OutputIndex)
+{
+	if (!IOR.GetTracedInput().Expression)
+	{
+		return Compiler->Errorf(TEXT("Missing IOR input"));			// STRATA_TODO: defaults in the detail panel with auto read only if plug in
+	}
+	if (!Extinction.GetTracedInput().Expression)
+	{
+		return Compiler->Errorf(TEXT("Missing Extinction input"));	// STRATA_TODO: defaults in the detail panel with auto read only if plug in
+	}
+	if (!Roughness.GetTracedInput().Expression)
+	{
+		return Compiler->Errorf(TEXT("Missing Roughness input"));	// STRATA_TODO: defaults in the detail panel with auto read only if plug in
+	}
+	if (!Normal.GetTracedInput().Expression)
+	{
+		return Compiler->Errorf(TEXT("Missing Normal input"));		// STRATA_TODO: defaults in the detail panel with auto read only if plug in
+	}
+	return Compiler->StrataConductorBSDF(IOR.Compile(Compiler), Extinction.Compile(Compiler), Roughness.Compile(Compiler), Normal.Compile(Compiler));
+}
+
+void UMaterialExpressionStrataConductorBSDF::GetCaption(TArray<FString>& OutCaptions) const
+{
+	OutCaptions.Add(TEXT("Strata Conductor BSDF"));
+}
+
+uint32 UMaterialExpressionStrataConductorBSDF::GetOutputType(int32 OutputIndex)
+{
+	return MCT_Strata;
+}
+
+uint32 UMaterialExpressionStrataConductorBSDF::GetInputType(int32 InputIndex)
+{
+	switch (InputIndex)
+	{
+	case 0:
+		return MCT_Float3;
+		break;
+	case 1:
+		return MCT_Float3;
+		break;
+	case 2:
+		return MCT_Float2;
+		break;
+	case 3:
+		return MCT_Float3;
+		break;
+	}
+
+	check(false);
+	return MCT_Float1;
+}
+#endif // WITH_EDITOR
+
+
+
+
+
+UMaterialExpressionStrataVolumeBSDF::UMaterialExpressionStrataVolumeBSDF(const FObjectInitializer& ObjectInitializer)
+	: Super(ObjectInitializer)
+{
+	struct FConstructorStatics
+	{
+		FText NAME_Strata;
+		FConstructorStatics() : NAME_Strata(LOCTEXT("Strata", "Strata")) { }
+	};
+	static FConstructorStatics ConstructorStatics;
+#if WITH_EDITORONLY_DATA
+	MenuCategories.Add(ConstructorStatics.NAME_Strata);
+#endif
+}
+
+#if WITH_EDITOR
+int32 UMaterialExpressionStrataVolumeBSDF::Compile(class FMaterialCompiler* Compiler, int32 OutputIndex)
+{
+	if (!Absorption.GetTracedInput().Expression)
+	{
+		return Compiler->Errorf(TEXT("Missing Absorption input"));	// STRATA_TODO: defaults in the detail panel with auto read only if plug in
+	}
+	if (!Scattering.GetTracedInput().Expression)
+	{
+		return Compiler->Errorf(TEXT("Missing Scattering input"));	// STRATA_TODO: defaults in the detail panel with auto read only if plug in
+	}
+	if (!Anisotropy.GetTracedInput().Expression)
+	{
+		return Compiler->Errorf(TEXT("Missing Anisotropy input"));	// STRATA_TODO: defaults in the detail panel with auto read only if plug in
+	}
+	return Compiler->StrataVolumeBSDF(Absorption.Compile(Compiler), Scattering.Compile(Compiler), Anisotropy.Compile(Compiler));
+}
+
+void UMaterialExpressionStrataVolumeBSDF::GetCaption(TArray<FString>& OutCaptions) const
+{
+	OutCaptions.Add(TEXT("Strata Volume BSDF"));
+}
+
+uint32 UMaterialExpressionStrataVolumeBSDF::GetOutputType(int32 OutputIndex)
+{
+	return MCT_Strata;
+}
+
+uint32 UMaterialExpressionStrataVolumeBSDF::GetInputType(int32 InputIndex)
+{
+	switch (InputIndex)
+	{
+	case 0:
+		return MCT_Float3;
+		break;
+	case 1:
+		return MCT_Float3;
+		break;
+	case 2:
+		return MCT_Float1;
+		break;
+	}
+
+	check(false);
+	return MCT_Float1;
+}
+#endif // WITH_EDITOR
+
+
+
+UMaterialExpressionStrataHorizontalMixing::UMaterialExpressionStrataHorizontalMixing(const FObjectInitializer& ObjectInitializer)
+	: Super(ObjectInitializer)
+{
+	struct FConstructorStatics
+	{
+		FText NAME_Strata;
+		FConstructorStatics() : NAME_Strata(LOCTEXT("Strata", "Strata")) { }
+	};
+	static FConstructorStatics ConstructorStatics;
+#if WITH_EDITORONLY_DATA
+	MenuCategories.Add(ConstructorStatics.NAME_Strata);
+#endif
+}
+
+#if WITH_EDITOR
+int32 UMaterialExpressionStrataHorizontalMixing::Compile(class FMaterialCompiler* Compiler, int32 OutputIndex)
+{
+	if (!Foreground.GetTracedInput().Expression)
+	{
+		return Compiler->Errorf(TEXT("Missing Foreground input"));
+	}
+	if (!Background.GetTracedInput().Expression)
+	{
+		return Compiler->Errorf(TEXT("Missing Background input"));
+	}
+	if (!Mix.GetTracedInput().Expression)
+	{
+		return Compiler->Errorf(TEXT("Missing Mix input"));
+	}
+	return Compiler->StrataHorizontalMixing(Foreground.Compile(Compiler), Background.Compile(Compiler), Mix.GetTracedInput().Expression ? Mix.Compile(Compiler) : Compiler->Constant(0.0f));
+}
+
+void UMaterialExpressionStrataHorizontalMixing::GetCaption(TArray<FString>& OutCaptions) const
+{
+	OutCaptions.Add(TEXT("Strata Horizontal Mixing"));
+}
+
+uint32 UMaterialExpressionStrataHorizontalMixing::GetOutputType(int32 OutputIndex)
+{
+	return MCT_Strata;
+}
+
+uint32 UMaterialExpressionStrataHorizontalMixing::GetInputType(int32 InputIndex)
+{
+	return InputIndex == 2 ? MCT_Float1 : MCT_Strata;
+}
+#endif // WITH_EDITOR
+
+
+
+UMaterialExpressionStrataVerticalLayering::UMaterialExpressionStrataVerticalLayering(const FObjectInitializer& ObjectInitializer)
+	: Super(ObjectInitializer)
+{
+	struct FConstructorStatics
+	{
+		FText NAME_Strata;
+		FConstructorStatics() : NAME_Strata(LOCTEXT("Strata", "Strata")) { }
+	};
+	static FConstructorStatics ConstructorStatics;
+#if WITH_EDITORONLY_DATA
+	MenuCategories.Add(ConstructorStatics.NAME_Strata);
+#endif
+}
+
+#if WITH_EDITOR
+int32 UMaterialExpressionStrataVerticalLayering::Compile(class FMaterialCompiler* Compiler, int32 OutputIndex)
+{
+	if (!Top.GetTracedInput().Expression)
+	{
+		return Compiler->Errorf(TEXT("Missing Top input"));
+	}
+	if (!Base.GetTracedInput().Expression)
+	{
+		return Compiler->Errorf(TEXT("Missing Base input"));
+	}
+	return Compiler->StrataVerticalLayering(Top.Compile(Compiler), Base.Compile(Compiler));
+}
+
+void UMaterialExpressionStrataVerticalLayering::GetCaption(TArray<FString>& OutCaptions) const
+{
+	OutCaptions.Add(TEXT("Strata Vertical Layering"));
+}
+
+uint32 UMaterialExpressionStrataVerticalLayering::GetOutputType(int32 OutputIndex)
+{
+	return MCT_Strata;
+}
+
+uint32 UMaterialExpressionStrataVerticalLayering::GetInputType(int32 InputIndex)
+{
+	return MCT_Strata;
+}
+#endif // WITH_EDITOR
+
+
+
+UMaterialExpressionStrataAdd::UMaterialExpressionStrataAdd(const FObjectInitializer& ObjectInitializer)
+	: Super(ObjectInitializer)
+{
+	struct FConstructorStatics
+	{
+		FText NAME_Strata;
+		FConstructorStatics() : NAME_Strata(LOCTEXT("Strata", "Strata")) { }
+	};
+	static FConstructorStatics ConstructorStatics;
+#if WITH_EDITORONLY_DATA
+	MenuCategories.Add(ConstructorStatics.NAME_Strata);
+#endif
+}
+
+#if WITH_EDITOR
+int32 UMaterialExpressionStrataAdd::Compile(class FMaterialCompiler* Compiler, int32 OutputIndex)
+{
+	if (!A.GetTracedInput().Expression)
+	{
+		return Compiler->Errorf(TEXT("Missing A input"));
+	}
+	if (!B.GetTracedInput().Expression)
+	{
+		return Compiler->Errorf(TEXT("Missing B input"));
+	}
+	return Compiler->StrataAdd(A.Compile(Compiler), B.Compile(Compiler));
+}
+
+void UMaterialExpressionStrataAdd::GetCaption(TArray<FString>& OutCaptions) const
+{
+	OutCaptions.Add(TEXT("Strata Add"));
+}
+
+uint32 UMaterialExpressionStrataAdd::GetOutputType(int32 OutputIndex)
+{
+	return MCT_Strata;
+}
+
+uint32 UMaterialExpressionStrataAdd::GetInputType(int32 InputIndex)
+{
+	return MCT_Strata;
+}
+#endif // WITH_EDITOR
+
+
+
+UMaterialExpressionStrataMultiply::UMaterialExpressionStrataMultiply(const FObjectInitializer& ObjectInitializer)
+	: Super(ObjectInitializer)
+{
+	struct FConstructorStatics
+	{
+		FText NAME_Strata;
+		FConstructorStatics() : NAME_Strata(LOCTEXT("Strata", "Strata")) { }
+	};
+	static FConstructorStatics ConstructorStatics;
+#if WITH_EDITORONLY_DATA
+	MenuCategories.Add(ConstructorStatics.NAME_Strata);
+#endif
+}
+
+#if WITH_EDITOR
+int32 UMaterialExpressionStrataMultiply::Compile(class FMaterialCompiler* Compiler, int32 OutputIndex)
+{
+	if (!A.GetTracedInput().Expression)
+	{
+		return Compiler->Errorf(TEXT("Missing Top input"));
+	}
+	if (!Weight.GetTracedInput().Expression)
+	{
+		return Compiler->Errorf(TEXT("Missing Base input"));
+	}
+	return Compiler->StrataMultiply(A.Compile(Compiler), Weight.Compile(Compiler));
+}
+
+void UMaterialExpressionStrataMultiply::GetCaption(TArray<FString>& OutCaptions) const
+{
+	OutCaptions.Add(TEXT("Strata Multiply"));
+}
+
+uint32 UMaterialExpressionStrataMultiply::GetOutputType(int32 OutputIndex)
+{
+	return MCT_Strata;
+}
+
+uint32 UMaterialExpressionStrataMultiply::GetInputType(int32 InputIndex)
+{
+	if (InputIndex == 0)
+	{
+		return MCT_Strata;
+	}
+	return MCT_Float3;
+}
+#endif // WITH_EDITOR
+
+
+
+UMaterialExpressionStrataArtisticIOR::UMaterialExpressionStrataArtisticIOR(const FObjectInitializer& ObjectInitializer)
+	: Super(ObjectInitializer)
+{
+	struct FConstructorStatics
+	{
+		FText NAME_Strata;
+		FConstructorStatics() : NAME_Strata(LOCTEXT("Strata", "Strata")) { }
+	};
+	static FConstructorStatics ConstructorStatics;
+#if WITH_EDITORONLY_DATA
+	MenuCategories.Add(ConstructorStatics.NAME_Strata);
+
+	bShowOutputNameOnPin = true;
+
+	Outputs.Reset();
+	Outputs.Add(FExpressionOutput(TEXT("IOR")));
+	Outputs.Add(FExpressionOutput(TEXT("Extinction")));
+#endif
+}
+
+#if WITH_EDITOR
+int32 UMaterialExpressionStrataArtisticIOR::Compile(class FMaterialCompiler* Compiler, int32 OutputIndex)
+{
+	if (!Reflectivity.GetTracedInput().Expression)
+	{
+		return Compiler->Errorf(TEXT("Missing Reflectivity input"));
+	}
+	if (!EdgeColor.GetTracedInput().Expression)
+	{
+		return Compiler->Errorf(TEXT("Missing EdgeColor input"));
+	}
+	return Compiler->StrataArtisticIOR(Reflectivity.Compile(Compiler), EdgeColor.Compile(Compiler), OutputIndex);
+}
+
+void UMaterialExpressionStrataArtisticIOR::GetCaption(TArray<FString>& OutCaptions) const
+{
+	OutCaptions.Add(TEXT("Strata Artistic IOR"));
+}
+
+uint32 UMaterialExpressionStrataArtisticIOR::GetOutputType(int32 OutputIndex)
+{
+	return MCT_Float3;
+}
+
+uint32 UMaterialExpressionStrataArtisticIOR::GetInputType(int32 InputIndex)
+{
+	return MCT_Float3;
+}
+#endif // WITH_EDITOR
+
+
+
+
 
 
 #undef LOCTEXT_NAMESPACE
diff --git a/Engine/Source/Runtime/Engine/Private/Materials/MaterialShared.cpp b/Engine/Source/Runtime/Engine/Private/Materials/MaterialShared.cpp
--- a/Engine/Source/Runtime/Engine/Private/Materials/MaterialShared.cpp
+++ b/Engine/Source/Runtime/Engine/Private/Materials/MaterialShared.cpp
@@ -35,6 +35,7 @@
 #include "SceneView.h"
 #include "ShaderPlatformQualitySettings.h"
 #include "MaterialShaderQualitySettings.h"
+#include "Engine/RendererSettings.h"
 #include "DecalRenderingCommon.h"
 #include "ExternalTexture.h"
 #include "ShaderCodeLibrary.h"
@@ -371,6 +372,11 @@ bool FShadingModelMaterialInput::Serialize(FArchive& Ar)
 	return SerializeMaterialInput<uint32>(Ar, *this);
 }
 
+bool FStrataMaterialInput::Serialize(FArchive& Ar)
+{
+	return SerializeMaterialInput<uint32>(Ar, *this);
+}
+
 bool FVectorMaterialInput::Serialize(FArchive& Ar)
 {
 	return SerializeMaterialInput<FVector>(Ar, *this);
@@ -438,6 +444,20 @@ int32 FShadingModelMaterialInput::CompileWithDefault(class FMaterialCompiler* Co
 	return Compiler->ForceCast(FMaterialAttributeDefinitionMap::CompileDefaultExpression(Compiler, Property), MCT_ShadingModel, MFCF_ExactMatch);
 }
 
+int32 FStrataMaterialInput::CompileWithDefault(class FMaterialCompiler* Compiler, EMaterialProperty Property)
+{
+	if (Expression)
+	{
+		int32 ResultIndex = FExpressionInput::Compile(Compiler);
+		if (ResultIndex != INDEX_NONE)
+		{
+			return ResultIndex;
+		}
+	}
+
+	return Compiler->ForceCast(FMaterialAttributeDefinitionMap::CompileDefaultExpression(Compiler, Property), MCT_Strata);
+}
+
 int32 FVectorMaterialInput::CompileWithDefault(class FMaterialCompiler* Compiler, EMaterialProperty Property)
 {
 	if (UseConstant)
@@ -1403,6 +1423,12 @@ bool FMaterialResource::HasAmbientOcclusionConnected() const
 	return HasMaterialAttributesConnected() || Material->HasAmbientOcclusionConnected();
 }
 
+bool FMaterialResource::HasStrataFrontMaterialConnected() const
+{
+	const URendererSettings* RendererSettings = GetDefault<URendererSettings>();
+	return RendererSettings && RendererSettings->bEnableStrata ? Material->HasStrataFrontMaterialConnected() : false;
+}
+
 bool FMaterialResource::RequiresSynchronousCompilation() const
 {
 	return Material->IsDefaultMaterial();
@@ -3635,7 +3661,7 @@ FMaterialAttributeDefintion::FMaterialAttributeDefintion(
 	, BlendFunction(InBlendFunction)
 	, bIsHidden(bInIsHidden)
 {
-	checkf(ValueType & MCT_Float || ValueType == MCT_ShadingModel , TEXT("Unsupported type, only Float1 through Float4 or MCT_ShadingModel are allowed."));
+	checkf(ValueType & MCT_Float || ValueType == MCT_ShadingModel || ValueType == MCT_Strata, TEXT("Unsupported type, only Float1 through Float4 or MCT_ShadingModel are allowed."));
 }
 
 int32 FMaterialAttributeDefintion::CompileDefaultValue(FMaterialCompiler* Compiler)
@@ -3656,6 +3682,13 @@ int32 FMaterialAttributeDefintion::CompileDefaultValue(FMaterialCompiler* Compil
 		return Compiler->ShadingModel(Compiler->GetMaterialShadingModels().GetFirstShadingModel());
 	}
 
+	if (Property == MP_FrontMaterial)
+	{
+		check(ValueType == MCT_Strata);
+		// Default to the first shading model of the material. If the material is using a single shading model selected through the dropdown, this is how it gets written to the shader as a constant (optimizing out all the dynamic branches)
+		return Compiler->FrontMaterial();
+	}
+
 	if (TexCoordIndex == INDEX_NONE)
 	{
 		// Standard value type
@@ -3721,8 +3754,9 @@ void FMaterialAttributeDefinitionMap::InitializeAttributeMap()
 	Add(FGuid(0xE8EBD0AD, 0xB1654CBE, 0xB079C3A8, 0xB39B9F15), TEXT("AmbientOcclusion"),		MP_AmbientOcclusion,		MCT_Float,	FVector4(1,0,0,0),	SF_Pixel);
 	Add(FGuid(0xD0B0FA03, 0x14D74455, 0xA851BAC5, 0x81A0788B), TEXT("Refraction"),				MP_Refraction,				MCT_Float2,	FVector4(1,0,0,0),	SF_Pixel);
 	Add(FGuid(0x0AC97EC3, 0xE3D047BA, 0xB610167D, 0xC4D919FF), TEXT("PixelDepthOffset"),		MP_PixelDepthOffset,		MCT_Float,	FVector4(0,0,0,0),	SF_Pixel);
-	Add(FGuid(0xD9423FFF, 0xD77E4D82, 0x8FF9CF5E, 0x055D1255), TEXT("ShadingModel"),			MP_ShadingModel,			MCT_ShadingModel,	FVector4(0,0,0,0),	SF_Pixel, INDEX_NONE, false, &CompileShadingModelBlendFunction);
-
+	Add(FGuid(0xD9423FFF, 0xD77E4D82, 0x8FF9CF5E, 0x055D1255), TEXT("ShadingModel"),			MP_ShadingModel,			MCT_ShadingModel, FVector4(0, 0, 0, 0), SF_Pixel, INDEX_NONE, false, &CompileShadingModelBlendFunction);
+	Add(FGuid(0x5973A03E, 0x13A74E08, 0x92D0CEDD, 0xF2936CF8), TEXT("FrontMaterial"),			MP_FrontMaterial,			MCT_Strata, FVector4(0,0,0,0),	SF_Pixel);
+	   
 	// Texture coordinates
 	Add(FGuid(0xD30EC284, 0xE13A4160, 0x87BB5230, 0x2ED115DC), TEXT("CustomizedUV0"), MP_CustomizedUVs0, MCT_Float2, FVector4(0,0,0,0), SF_Vertex, 0);
 	Add(FGuid(0xC67B093C, 0x2A5249AA, 0xABC97ADE, 0x4A1F49C5), TEXT("CustomizedUV1"), MP_CustomizedUVs1, MCT_Float2, FVector4(0,0,0,0), SF_Vertex, 1);
@@ -3872,6 +3906,8 @@ FText FMaterialAttributeDefinitionMap::GetAttributeOverrideForMaterial(const FGu
 		return LOCTEXT("PixelDepthOffset", "Pixel Depth Offset");
 	case MP_ShadingModel:
 		return LOCTEXT("ShadingModel", "Shading Model");
+	case MP_FrontMaterial:
+		return LOCTEXT("FrontMaterial", "Front Material");
 	case MP_CustomOutput:
 		return FText::FromString(GetAttributeName(AttributeID));
 		
diff --git a/Engine/Source/Runtime/Engine/Public/MaterialCompiler.h b/Engine/Source/Runtime/Engine/Public/MaterialCompiler.h
--- a/Engine/Source/Runtime/Engine/Public/MaterialCompiler.h
+++ b/Engine/Source/Runtime/Engine/Public/MaterialCompiler.h
@@ -386,6 +386,19 @@ public:
 	virtual int32 CustomPrimitiveData(int32 OutputIndex, EMaterialValueType Type) = 0;
 	virtual int32 ShadingModel(EMaterialShadingModel InSelectedShadingModel) = 0;
 
+	// Strata
+	virtual int32 FrontMaterial() = 0;
+	virtual int32 StrataDiffuseOrenNayarBSDF(int32 Albedo, int32 Roughness, int32 Normal) = 0;
+	virtual int32 StrataDiffuseChanBSDF(int32 Albedo, int32 Roughness, int32 Normal) = 0;
+	virtual int32 StrataDielectricBSDF(int32 Roughness, int32 IOR, int32 Tint, int32 Normal) = 0;
+	virtual int32 StrataConductorBSDF(int32 IOR, int32 Extinction, int32 Roughness, int32 Normal) = 0;
+	virtual int32 StrataVolumeBSDF(int32 Absorption, int32 Scattering, int32 Anisotropy) = 0;
+	virtual int32 StrataHorizontalMixing(int32 Foreground, int32 Background, int32 Mix) = 0;
+	virtual int32 StrataVerticalLayering(int32 Top, int32 Base) = 0;
+	virtual int32 StrataAdd(int32 A, int32 B) = 0;
+	virtual int32 StrataMultiply(int32 A, int32 Weight) = 0;
+	virtual int32 StrataArtisticIOR(int32 Reflectivity, int32 EdgeColor, int32 OutputIndex) = 0;
+
 
 	virtual int32 MapARPassthroughCameraUV(int32 UV) = 0;
 	// The compiler can run in a different state and this affects caching of sub expression, Expressions are different (e.g. View.PrevWorldViewOrigin) when using previous frame's values
@@ -766,6 +779,61 @@ public:
 		return Compiler->IsDevelopmentFeatureEnabled(FeatureName);
 	}
 
+	virtual int32 FrontMaterial() override
+	{
+		return Compiler->FrontMaterial();
+	}
+
+	virtual int32 StrataDiffuseOrenNayarBSDF(int32 Albedo, int32 Roughness, int32 Normal) override
+	{
+		return Compiler->StrataDiffuseOrenNayarBSDF(Albedo, Roughness, Normal);
+	}
+
+	virtual int32 StrataDiffuseChanBSDF(int32 Albedo, int32 Roughness, int32 Normal) override
+	{
+		return Compiler->StrataDiffuseChanBSDF(Albedo, Roughness, Normal);
+	}
+
+	virtual int32 StrataDielectricBSDF(int32 Roughness, int32 IOR, int32 Tint, int32 Normal) override
+	{
+		return Compiler->StrataDielectricBSDF(Roughness, IOR, Tint, Normal);
+	}
+
+	virtual int32 StrataConductorBSDF(int32 IOR, int32 Extinction, int32 Roughness, int32 Normal) override
+	{
+		return Compiler->StrataConductorBSDF(IOR, Extinction, Roughness, Normal);
+	}
+
+	virtual int32 StrataVolumeBSDF(int32 Absorption, int32 Scattering, int32 Anisotropy) override
+	{
+		return Compiler->StrataVolumeBSDF(Absorption, Scattering, Anisotropy);
+	}
+
+	virtual int32 StrataHorizontalMixing(int32 Foreground, int32 Background, int32 Mix) override
+	{
+		return Compiler->StrataHorizontalMixing(Foreground, Background, Mix);
+	}
+
+	virtual int32 StrataVerticalLayering(int32 Top, int32 Base) override
+	{
+		return Compiler->StrataVerticalLayering(Top, Base);
+	}
+
+	virtual int32 StrataAdd(int32 A, int32 B) override
+	{
+		return Compiler->StrataAdd(A, B);
+	}
+
+	virtual int32 StrataMultiply(int32 A, int32 Weight) override
+	{
+		return Compiler->StrataMultiply(A, Weight);
+	}
+
+	virtual int32 StrataArtisticIOR(int32 Reflectivity, int32 EdgeColor, int32 OutputIndex) override
+	{
+		return Compiler->StrataArtisticIOR(Reflectivity, EdgeColor, OutputIndex);
+	}
+
 protected:
 		
 	FMaterialCompiler* Compiler;
diff --git a/Engine/Source/Runtime/Engine/Public/MaterialExpressionIO.h b/Engine/Source/Runtime/Engine/Public/MaterialExpressionIO.h
--- a/Engine/Source/Runtime/Engine/Public/MaterialExpressionIO.h
+++ b/Engine/Source/Runtime/Engine/Public/MaterialExpressionIO.h
@@ -229,6 +229,25 @@ struct TStructOpsTypeTraits<FShadingModelMaterialInput>
 	};
 };
 
+struct FStrataMaterialInput : FMaterialInput<uint32> // Still giving it a default type
+{
+#if WITH_EDITOR
+	ENGINE_API int32 CompileWithDefault(class FMaterialCompiler* Compiler, EMaterialProperty Property);
+#endif  // WITH_EDITOR
+	/** ICPPStructOps interface */
+	ENGINE_API bool Serialize(FArchive& Ar);
+};
+
+template<>
+struct TStructOpsTypeTraits<FStrataMaterialInput>
+	: public TStructOpsTypeTraitsBase2<FStrataMaterialInput>
+{
+	enum
+	{
+		WithSerializer = true,
+	};
+};
+
 struct FVectorMaterialInput : FMaterialInput<FVector>
 {
 #if WITH_EDITOR
diff --git a/Engine/Source/Runtime/Engine/Public/MaterialShared.h b/Engine/Source/Runtime/Engine/Public/MaterialShared.h
--- a/Engine/Source/Runtime/Engine/Public/MaterialShared.h
+++ b/Engine/Source/Runtime/Engine/Public/MaterialShared.h
@@ -139,7 +139,9 @@ enum EMaterialValueType
 	/** Used internally when sampling from virtual textures */
 	MCT_VTPageTableResult     = 1 << 13,
 	
-	MCT_ShadingModel = 1 << 14,
+	MCT_ShadingModel		  = 1 << 14,
+
+	MCT_Strata				  = 1 << 15,
 };
 
 /**
@@ -1623,6 +1625,7 @@ public:
 	virtual bool HasEmissiveColorConnected() const { return false; }
 	virtual bool HasAnisotropyConnected() const { return false; }
 	virtual bool HasAmbientOcclusionConnected() const { return false; }
+	virtual bool HasStrataFrontMaterialConnected() const { return false; }
 	virtual bool RequiresSynchronousCompilation() const { return false; };
 	virtual bool IsDefaultMaterial() const { return false; };
 	virtual int32 GetNumCustomizedUVs() const { return 0; }
@@ -2432,6 +2435,7 @@ public:
 	ENGINE_API virtual bool HasEmissiveColorConnected() const override;
 	ENGINE_API virtual bool HasAnisotropyConnected() const override;
 	ENGINE_API virtual bool HasAmbientOcclusionConnected() const override;
+	ENGINE_API virtual bool HasStrataFrontMaterialConnected() const override;
 	ENGINE_API virtual FMaterialShadingModelField GetShadingModels() const override;
 	ENGINE_API virtual bool IsShadingModelFromMaterialExpression() const override;
 	ENGINE_API virtual enum ETranslucencyLightingMode GetTranslucencyLightingMode() const override;
diff --git a/Engine/Source/Runtime/Engine/Public/SceneTypes.h b/Engine/Source/Runtime/Engine/Public/SceneTypes.h
--- a/Engine/Source/Runtime/Engine/Public/SceneTypes.h
+++ b/Engine/Source/Runtime/Engine/Public/SceneTypes.h
@@ -179,6 +179,7 @@ enum EMaterialProperty
 	MP_CustomizedUVs7 UMETA(Hidden),
 	MP_PixelDepthOffset UMETA(Hidden),
 	MP_ShadingModel UMETA(Hidden),
+	MP_FrontMaterial UMETA(Hidden),
 
 	//^^^ New material properties go above here ^^^^
 	MP_MaterialAttributes UMETA(Hidden),
diff --git a/Engine/Source/Runtime/Renderer/Private/IndirectLightRendering.cpp b/Engine/Source/Runtime/Renderer/Private/IndirectLightRendering.cpp
--- a/Engine/Source/Runtime/Renderer/Private/IndirectLightRendering.cpp
+++ b/Engine/Source/Runtime/Renderer/Private/IndirectLightRendering.cpp
@@ -20,6 +20,7 @@
 #include "Lumen/LumenSceneUtils.h"
 #include "Math/Halton.h"
 #include "DistanceFieldAmbientOcclusion.h"
+#include "Strata/Strata.h"
 
 static TAutoConsoleVariable<int32> CVarDiffuseIndirectHalfRes(
 	TEXT("r.DiffuseIndirect.HalfRes"), 1,
@@ -164,6 +165,7 @@ class FReflectionEnvironmentSkyLightingPS : public FGlobalShader
 	class FDynamicSkyLight : SHADER_PERMUTATION_BOOL("ENABLE_DYNAMIC_SKY_LIGHT");
 	class FSkyShadowing : SHADER_PERMUTATION_BOOL("APPLY_SKY_SHADOWING");
 	class FRayTracedReflections : SHADER_PERMUTATION_BOOL("RAY_TRACED_REFLECTIONS");
+	class FStrata : SHADER_PERMUTATION_BOOL("STRATA_ENABLED");
 
 	using FPermutationDomain = TShaderPermutationDomain<
 		FHasBoxCaptures,
@@ -172,7 +174,8 @@ class FReflectionEnvironmentSkyLightingPS : public FGlobalShader
 		FSkyLight,
 		FDynamicSkyLight,
 		FSkyShadowing,
-		FRayTracedReflections>;
+		FRayTracedReflections,
+		FStrata>;
 
 	static FPermutationDomain RemapPermutation(FPermutationDomain PermutationVector)
 	{
@@ -202,6 +205,7 @@ class FReflectionEnvironmentSkyLightingPS : public FGlobalShader
 		PermutationVector.Set<FDynamicSkyLight>(bEnableDynamicSkyLight);
 		PermutationVector.Set<FSkyShadowing>(bApplySkyShadowing);
 		PermutationVector.Set<FRayTracedReflections>(bRayTracedReflections);
+		PermutationVector.Set<FStrata>(Strata::IsStrataEnabled());
 
 		return RemapPermutation(PermutationVector);
 	}
@@ -266,8 +270,10 @@ class FReflectionEnvironmentSkyLightingPS : public FGlobalShader
 		SHADER_PARAMETER_STRUCT_REF(FReflectionCaptureShaderData, ReflectionCaptureData)
 		SHADER_PARAMETER_STRUCT_REF(FForwardLightData, ForwardLightData)
 
+		SHADER_PARAMETER_STRUCT_REF(FStrataGlobalUniformParameters, Strata)
+
 		RENDER_TARGET_BINDING_SLOTS()
-		END_SHADER_PARAMETER_STRUCT()
+	END_SHADER_PARAMETER_STRUCT()
 }; // FReflectionEnvironmentSkyLightingPS
 
 IMPLEMENT_GLOBAL_SHADER(FDiffuseIndirectCompositePS, "/Engine/Private/DiffuseIndirectComposite.usf", "MainPS", SF_Pixel);
@@ -1402,6 +1408,8 @@ void FDeferredShadingSceneRenderer::RenderDeferredReflectionsAndSkyLighting(
 					PassParameters->ReflectionsParameters = CreateUniformBufferImmediate(ReflectionUniformParameters, UniformBuffer_SingleDraw);
 				}
 				PassParameters->ForwardLightData = View.ForwardLightingResources->ForwardLightDataUniformBuffer;
+
+				PassParameters->Strata =  Strata::BindStrataGlobalUniformParameters(View);
 			}
 
 			PassParameters->RenderTargets[0] = FRenderTargetBinding(SceneColorTexture.Target, ERenderTargetLoadAction::ELoad);
diff --git a/Engine/Source/Runtime/Renderer/Private/LightRendering.cpp b/Engine/Source/Runtime/Renderer/Private/LightRendering.cpp
--- a/Engine/Source/Runtime/Renderer/Private/LightRendering.cpp
+++ b/Engine/Source/Runtime/Renderer/Private/LightRendering.cpp
@@ -21,6 +21,7 @@
 #include "ScreenPass.h"
 #include "SkyAtmosphereRendering.h"
 #include "VolumetricCloudRendering.h"
+#include "Strata/Strata.h"
 
 // ENABLE_DEBUG_DISCARD_PROP is used to test the lighting code by allowing to discard lights to see how performance scales
 // It ought never to be enabled in a shipping build, and is probably only really useful when woring on the shading code.
@@ -416,6 +417,7 @@ class FDeferredLightPS : public FGlobalShader
 	class FAtmosphereTransmittance : SHADER_PERMUTATION_BOOL("USE_ATMOSPHERE_TRANSMITTANCE");
 	class FCloudTransmittance 	: SHADER_PERMUTATION_BOOL("USE_CLOUD_TRANSMITTANCE");
 	class FAnistropicMaterials 	: SHADER_PERMUTATION_BOOL("SUPPORTS_ANISOTROPIC_MATERIALS");
+	class FStrata				: SHADER_PERMUTATION_BOOL("STRATA_ENABLED");
 
 	using FPermutationDomain = TShaderPermutationDomain<
 		FSourceShapeDim,
@@ -428,7 +430,8 @@ class FDeferredLightPS : public FGlobalShader
 		FHairLighting,
 		FAtmosphereTransmittance,
 		FCloudTransmittance,
-		FAnistropicMaterials>;
+		FAnistropicMaterials,
+		FStrata>;
 
 	static bool ShouldCompilePermutation(const FGlobalShaderPermutationParameters& Parameters)
 	{
@@ -542,7 +545,7 @@ class FDeferredLightPS : public FGlobalShader
 public:
 	void SetParameters(
 		FRHICommandList& RHICmdList, 
-		const FSceneView& View, 
+		const FViewInfo& View,
 		const FLightSceneInfo* LightSceneInfo,
 		FRHITexture* ScreenShadowMaskTexture,
 		FRHITexture* LightingChannelsTextureRHI,
@@ -553,7 +556,7 @@ public:
 		SetDeferredLightParameters(RHICmdList, ShaderRHI, GetUniformBufferParameter<FDeferredLightUniformStruct>(), LightSceneInfo, View);
 	}
 
-	void SetParametersSimpleLight(FRHICommandList& RHICmdList, const FSceneView& View, const FSimpleLightEntry& SimpleLight, const FSimpleLightPerViewEntry& SimpleLightPerViewData)
+	void SetParametersSimpleLight(FRHICommandList& RHICmdList, const FViewInfo& View, const FSimpleLightEntry& SimpleLight, const FSimpleLightPerViewEntry& SimpleLightPerViewData)
 	{
 		FRHIPixelShader* ShaderRHI = RHICmdList.GetBoundPixelShader();
 		SetParametersBase(RHICmdList, ShaderRHI, View, nullptr, nullptr, nullptr, nullptr);
@@ -565,13 +568,14 @@ private:
 	void SetParametersBase(
 		FRHICommandList& RHICmdList, 
 		FRHIPixelShader* ShaderRHI, 
-		const FSceneView& View, 
+		const FViewInfo& View,
 		FRHITexture* ScreenShadowMaskTexture,
 		FRHITexture* LightingChannelsTextureRHI,
 		FTexture* IESTextureResource,
 		FRenderLightParams* RenderLightParams)
 	{
-		FGlobalShader::SetParameters<FViewUniformShaderParameters>(RHICmdList, ShaderRHI,View.ViewUniformBuffer);
+		FGlobalShader::SetParameters<FViewUniformShaderParameters>(RHICmdList, ShaderRHI, View.ViewUniformBuffer);
+		FGlobalShader::SetParameters<FStrataGlobalUniformParameters>(RHICmdList, ShaderRHI, Strata::BindStrataGlobalUniformParameters(View));
 
 		FSceneRenderTargets& SceneRenderTargets = FSceneRenderTargets::Get(RHICmdList);
 
@@ -2173,7 +2177,8 @@ static void SetShaderTemplLightingSimple(
 	PermutationVector.Set< FDeferredLightPS::FTransmissionDim >( false );
 	PermutationVector.Set< FDeferredLightPS::FHairLighting>( 0 );
 	PermutationVector.Set< FDeferredLightPS::FAtmosphereTransmittance >( false );
-	PermutationVector.Set< FDeferredLightPS::FCloudTransmittance >( false );
+	PermutationVector.Set< FDeferredLightPS::FCloudTransmittance >(false);
+	PermutationVector.Set< FDeferredLightPS::FStrata >(Strata::IsStrataEnabled());
 
 	TShaderMapRef< FDeferredLightPS > PixelShader( View.ShaderMap, PermutationVector );
 	GraphicsPSOInit.BoundShaderState.VertexDeclarationRHI = GetVertexDeclarationFVector4();
@@ -2331,6 +2336,7 @@ void FDeferredShadingSceneRenderer::RenderLight(
 				// Only directional lights are rendered in this path, so we only need to check if it is use to light the atmosphere
 				PermutationVector.Set< FDeferredLightPS::FAtmosphereTransmittance >(bAtmospherePerPixelTransmittance);
 				PermutationVector.Set< FDeferredLightPS::FCloudTransmittance >(bLight0CloudPerPixelTransmittance || bLight1CloudPerPixelTransmittance);
+				PermutationVector.Set< FDeferredLightPS::FStrata >(Strata::IsStrataEnabled());
 
 				TShaderMapRef< FDeferredLightPS > PixelShader( View.ShaderMap, PermutationVector );
 				GraphicsPSOInit.BoundShaderState.VertexDeclarationRHI = GFilterVertexDeclaration.VertexDeclarationRHI;
@@ -2395,6 +2401,7 @@ void FDeferredShadingSceneRenderer::RenderLight(
 				PermutationVector.Set< FDeferredLightPS::FHairLighting>(bHairLighting ? 1 : 0);
 				PermutationVector.Set < FDeferredLightPS::FAtmosphereTransmittance >(false);
 				PermutationVector.Set< FDeferredLightPS::FCloudTransmittance >(false);
+				PermutationVector.Set< FDeferredLightPS::FStrata >(Strata::IsStrataEnabled());
 
 				TShaderMapRef< FDeferredLightPS > PixelShader( View.ShaderMap, PermutationVector );
 				GraphicsPSOInit.BoundShaderState.VertexDeclarationRHI = GetVertexDeclarationFVector4();
@@ -2584,6 +2591,7 @@ void FDeferredShadingSceneRenderer::RenderLightForHair(
 			PermutationVector.Set< FDeferredLightPS::FVisualizeCullingDim >(false);
 			PermutationVector.Set< FDeferredLightPS::FTransmissionDim >(false);
 			PermutationVector.Set< FDeferredLightPS::FHairLighting>(2);
+			PermutationVector.Set< FDeferredLightPS::FStrata >(Strata::IsStrataEnabled());
 
 			TShaderMapRef<TDeferredLightHairVS> VertexShader(View.ShaderMap);
 			TShaderMapRef<FDeferredLightPS> PixelShader(View.ShaderMap, PermutationVector);
diff --git a/Engine/Source/Runtime/Renderer/Private/Strata/Strata.cpp b/Engine/Source/Runtime/Renderer/Private/Strata/Strata.cpp
--- a/Engine/Source/Runtime/Renderer/Private/Strata/Strata.cpp
+++ b/Engine/Source/Runtime/Renderer/Private/Strata/Strata.cpp
@@ -6,6 +6,7 @@
 #include "ScenePrivate.h"
 #include "SceneRendering.h"
 #include "RendererInterface.h"
+#include "UniformBuffer.h"
 
 
 
@@ -18,6 +19,10 @@ static TAutoConsoleVariable<int32> CVarStrata(
 
 
 
+IMPLEMENT_GLOBAL_SHADER_PARAMETER_STRUCT(FStrataGlobalUniformParameters, "Strata");
+
+
+
 namespace Strata
 {
 
@@ -41,7 +46,8 @@ void InitialiseStrataFrameSceneData(FSceneRenderer& SceneRenderer, FRDGBuilder&
 		ResolutionX = BufferSizeXY.X;
 		ResolutionY = BufferSizeXY.Y;
 
-		StrataData.MaxBytesPerPixel = 16;
+		// Previous GBuffer when complete was 28bytes
+		StrataData.MaxBytesPerPixel = 256;
 
 	}
 	else
@@ -89,6 +95,26 @@ void BindStrataBasePassUniformParameters(const FViewInfo& View, FStrataOpaquePas
 	}
 }
 
+TUniformBufferRef<FStrataGlobalUniformParameters> BindStrataGlobalUniformParameters(const FViewInfo& View)
+{
+	FStrataGlobalUniformParameters StrataUniformParameters;
+	if (View.StrataData)
+	{
+		StrataUniformParameters.MaxBytesPerPixel = View.StrataData->MaxBytesPerPixel;
+		StrataUniformParameters.MaterialLobesTexture = View.StrataData->MaterialLobesTexture->GetRenderTargetItem().ShaderResourceTexture;
+		StrataUniformParameters.MaterialLobesBuffer = View.StrataData->MaterialLobesBuffer.SRV;
+	}
+	else
+	{
+		StrataUniformParameters.MaxBytesPerPixel = 0;
+		StrataUniformParameters.MaterialLobesTexture = GSystemTextures.BlackDummy->GetRenderTargetItem().ShaderResourceTexture;
+		StrataUniformParameters.MaterialLobesBuffer = GEmptyVertexBufferWithUAV->ShaderResourceViewRHI;
+	}
+
+	// STRATA_TODO cache this on the view and use UniformBuffer_SingleFrame
+	return CreateUniformBufferImmediate(StrataUniformParameters, UniformBuffer_SingleDraw);
+}
+
 } // namespace Strata
 
 
diff --git a/Engine/Source/Runtime/Renderer/Private/Strata/Strata.h b/Engine/Source/Runtime/Renderer/Private/Strata/Strata.h
--- a/Engine/Source/Runtime/Renderer/Private/Strata/Strata.h
+++ b/Engine/Source/Runtime/Renderer/Private/Strata/Strata.h
@@ -21,6 +21,12 @@ BEGIN_SHADER_PARAMETER_STRUCT(FStrataOpaquePassUniformParameters, )
 	SHADER_PARAMETER_UAV(RWTexture2D<float>, MaterialLobesTextureUAV)
 END_SHADER_PARAMETER_STRUCT()
 
+BEGIN_GLOBAL_SHADER_PARAMETER_STRUCT(FStrataGlobalUniformParameters, )
+	SHADER_PARAMETER(uint32, MaxBytesPerPixel)
+	SHADER_PARAMETER_SRV(ByteAddressBuffer, MaterialLobesBuffer)
+	SHADER_PARAMETER_TEXTURE(Texture2D<float>, MaterialLobesTexture)
+END_GLOBAL_SHADER_PARAMETER_STRUCT()
+
 
 
 struct FStrataData
@@ -46,6 +52,8 @@ void InitialiseStrataFrameSceneData(FSceneRenderer& SceneRenderer, FRDGBuilder&
 
 void BindStrataBasePassUniformParameters(const FViewInfo& View, FStrataOpaquePassUniformParameters& OutStrataUniformParameters);
 
+TUniformBufferRef<FStrataGlobalUniformParameters> BindStrataGlobalUniformParameters(const FViewInfo& View);
+
 };
 
 
