Fixed potential crash when generating platform data for class default objects or for texture with no source data

diff --git a/Engine/Source/Runtime/Engine/Private/TextureDerivedData.cpp b/Engine/Source/Runtime/Engine/Private/TextureDerivedData.cpp
--- a/Engine/Source/Runtime/Engine/Private/TextureDerivedData.cpp
+++ b/Engine/Source/Runtime/Engine/Private/TextureDerivedData.cpp
@@ -1918,8 +1918,11 @@ bool UTexture::IsCachedCookedPlatformDataLoaded( const ITargetPlatform* TargetPl
 		// begin cache hasn't been called
 		if ( !PlatformData )
 		{
-			// In case an UpdateResource happens, cooked platform data might be cleared and we might need to reschedule
-			BeginCacheForCookedPlatformData(TargetPlatform);
+			if (!HasAnyFlags(RF_ClassDefaultObject) && Source.SizeX != 0 && Source.SizeY != 0)
+			{
+				// In case an UpdateResource happens, cooked platform data might be cleared and we might need to reschedule
+				BeginCacheForCookedPlatformData(TargetPlatform);
+			}
 			return false;
 		}
 
