Fix typo causing crash when StaticComponentMaskParameter are combined with material layers

diff --git a/Engine/Source/Runtime/Engine/Private/Materials/MaterialCachedData.cpp b/Engine/Source/Runtime/Engine/Private/Materials/MaterialCachedData.cpp
--- a/Engine/Source/Runtime/Engine/Private/Materials/MaterialCachedData.cpp
+++ b/Engine/Source/Runtime/Engine/Private/Materials/MaterialCachedData.cpp
@@ -523,7 +523,7 @@ bool FMaterialCachedExpressionData::UpdateForExpressions(const FMaterialCachedEx
 				FStaticComponentMaskValue Value {static_cast<bool>(ExpressionStaticComponentMaskParameter->DefaultR), static_cast<bool>(ExpressionStaticComponentMaskParameter->DefaultG), static_cast<bool>(ExpressionStaticComponentMaskParameter->DefaultB), static_cast<bool>(ExpressionStaticComponentMaskParameter->DefaultA)};
 				if (Context.Parent)
 				{
-					FMaterialCachedParameterEntry& CachedEntry = Parameters.EditorOnlyEntries[static_cast<int32>(EMaterialParameterType::StaticSwitch) - static_cast<int32>(EMaterialParameterType::RuntimeCount)];
+					FMaterialCachedParameterEntry& CachedEntry = Parameters.EditorOnlyEntries[static_cast<int32>(EMaterialParameterType::StaticComponentMask) - static_cast<int32>(EMaterialParameterType::RuntimeCount)];
 					FStaticComponentMaskValue ParentValue;
 					if (Context.Parent->GetStaticComponentMaskParameterDefaultValue(ParameterInfo, ParentValue.R, ParentValue.G, ParentValue.B, ParentValue.A, CachedEntry.ExpressionGuids[Index], true))
 					{
