Merging in fix for NetTests in EngineTest.

diff --git a/Engine/Source/Runtime/Engine/Private/RepLayout.cpp b/Engine/Source/Runtime/Engine/Private/RepLayout.cpp
--- a/Engine/Source/Runtime/Engine/Private/RepLayout.cpp
+++ b/Engine/Source/Runtime/Engine/Private/RepLayout.cpp
@@ -3206,6 +3206,12 @@ static FGuidReferencesMap* PrepReceivedArray(
 		{
 			UE_CLOG(LogSkippedRepNotifies> 0, LogRep, Display, TEXT("1 FReceivedPropertiesStackState Skipping RepNotify for property %s because local value has not changed."), *Cmd.Property->GetName());
 		}
+		check(ShadowArray != nullptr);
+		*OutShadowBaseData = ShadowArray->GetData();
+	}
+	else
+	{
+		*OutShadowBaseData = nullptr;
 	}
 
 	check(CastFieldChecked<FArrayProperty>(Cmd.Property) != nullptr);
@@ -3216,18 +3222,6 @@ static FGuidReferencesMap* PrepReceivedArray(
 
 	// Re-compute the base data values since they could have changed after the resize above
 	*OutBaseData = DataArray->GetData();
-	*OutShadowBaseData = nullptr;
-
-	// Only resize the shadow data array if we're actually tracking RepNotifies
-	if (RepNotifies != nullptr)
-	{
-		check(ShadowArray != nullptr);
-
-		FScriptArrayHelper ShadowArrayHelper((FArrayProperty*)Cmd.Property, ShadowArray);
-		ShadowArrayHelper.Resize(ArrayNum);
-
-		*OutShadowBaseData = ShadowArray->GetData();
-	}
 
 	return NewGuidReferencesArray ? NewGuidReferencesArray->Array : nullptr;
 }
@@ -3237,7 +3231,6 @@ struct FReceivePropertiesSharedParams
 {
 	const bool bDoChecksum;
 	const bool bSkipRoleSwap;
-	TArray<FProperty*>* RepNotifies;
 	FNetBitReader& Bunch;
 	bool& bOutHasUnmapped;
 	bool& bOutGuidsChanged;
@@ -3261,6 +3254,7 @@ struct FReceivePropertiesStackParams
 	FGuidReferencesMap* GuidReferences;
 	const int32 CmdStart;
 	const int32 CmdEnd;
+	TArray<FProperty*>* RepNotifies;
 	uint32 ArrayElementOffset = 0;
 	uint16 CurrentHandle = 0;
 };
@@ -3346,6 +3340,7 @@ static bool ReceiveProperties_r(FReceivePropertiesSharedParams& Params, FReceive
 					nullptr,
 					CmdIndex + 1,
 					Cmd.EndCmd - 1,
+					StackParams.RepNotifies
 				};
 
 				// These buffers will track the dynamic array memory.
@@ -3371,7 +3366,7 @@ static bool ReceiveProperties_r(FReceivePropertiesSharedParams& Params, FReceive
 					CmdIndex,
 					&ShadowArrayBuffer,
 					&ObjectArrayBuffer,
-					Params.RepNotifies);
+					StackParams.RepNotifies);
 
 				// Read the next array handle.
 				ReadPropertyHandle(Params);
@@ -3386,9 +3381,13 @@ static bool ReceiveProperties_r(FReceivePropertiesSharedParams& Params, FReceive
 						const int32 ElementOffset = i * Cmd.ElementSize;
 
 						ArrayStackParams.ObjectData = ObjectArrayBuffer + ElementOffset;
-						ArrayStackParams.ShadowData = ShadowArrayBuffer ? (ShadowArrayBuffer + ElementOffset) : nullptr;
 						ArrayStackParams.ArrayElementOffset = ElementOffset;
 
+						// If ShadowArrayBuffer is valid, then we know that our ShadowArray pointer is also valid and pointing to a valid array.
+						// So we just need to make sure we're not going outside the bounds of the array.
+						ArrayStackParams.ShadowData = (ShadowArrayBuffer && i < ShadowArray->Num()) ? (ShadowArrayBuffer + ElementOffset) : nullptr;
+						ArrayStackParams.RepNotifies = ArrayStackParams.ShadowData ? StackParams.RepNotifies : nullptr;
+
 						UE_LOG(LogRepProperties, VeryVerbose, TEXT("ReceivePropertiesArray_r: Recursing - Parent=%d, Cmd=%d, Index=%d"), Cmd.ParentIndex, CmdIndex, i);
 						if (!ReceiveProperties_r(Params, ArrayStackParams))
 						{
@@ -3419,7 +3418,7 @@ static bool ReceiveProperties_r(FReceivePropertiesSharedParams& Params, FReceive
 					StackParams.ArrayElementOffset,
 					StackParams.ShadowData,
 					StackParams.ObjectData,
-					Params.RepNotifies,
+					StackParams.RepNotifies,
 					Params.Parents,
 					Params.Cmds,
 					CmdIndex,
@@ -3489,7 +3488,6 @@ bool FRepLayout::ReceiveProperties(
 		bDoChecksum,
 		// We can skip swapping roles if we're not an Actor layout, or if we've been explicitly told we can skip.
 		EnumHasAnyFlags(ReceiveFlags, EReceivePropertiesFlags::SkipRoleSwap) || !EnumHasAnyFlags(Flags, ERepLayoutFlags::IsActor),
-		bEnableRepNotifies ? &RepState->RepNotifies : nullptr,
 		InBunch,
 		bOutHasUnmapped,
 		bOutGuidsChanged,
@@ -3505,7 +3503,8 @@ bool FRepLayout::ReceiveProperties(
 		FRepShadowDataBuffer(RepState->StaticBuffer.GetData()),
 		&RepState->GuidReferencesMap,
 		0,
-		Cmds.Num() - 1
+		Cmds.Num() - 1,
+		bEnableRepNotifies ? &RepState->RepNotifies : nullptr
 	};
 
 	// Read the first handle, and then start receiving properties.
@@ -3784,7 +3783,7 @@ bool FRepLayout::ReceiveProperties_BackwardsCompatible_r(
 				const int32 ArrayElementOffset = Index * Cmd.ElementSize;
 
 				FRepObjectDataBuffer ElementData = LocalData + ArrayElementOffset;
-				FRepShadowDataBuffer ElementShadowData = LocalShadowData ? LocalShadowData + ArrayElementOffset : nullptr;
+				FRepShadowDataBuffer ElementShadowData = (LocalShadowData && (Index < static_cast<uint32>(ShadowArray->Num()))) ? LocalShadowData + ArrayElementOffset : nullptr;
 
 				if (!ReceiveProperties_BackwardsCompatible_r(RepState, NetFieldExportGroup, TempReader, CmdIndex + 1, Cmd.EndCmd - 1, ElementShadowData, LocalData, ElementData, NewGuidReferencesArray, bOutHasUnmapped, bOutGuidsChanged))
 				{
@@ -4020,6 +4019,7 @@ void FRepLayout::UpdateUnmappedObjects_r(
 		FGuidReferences& GuidReferences = It.Value();
 		const FRepLayoutCmd& Cmd = Cmds[GuidReferences.CmdIndex];
 		const FRepParentCmd& Parent = Parents[GuidReferences.ParentIndex];
+		const bool bUpdateShadowState = (ShadowData && INDEX_NONE != Parent.RepNotifyNumParams);
 
 		// Make sure if we're touching an array element, we use the correct offset for shadow values.
 		// This should always be safe, because MaxAbsOffset will account for ShadowArray size for arrays.
@@ -4030,7 +4030,7 @@ void FRepLayout::UpdateUnmappedObjects_r(
 		{
 			check(Cmd.Type == ERepLayoutCmdType::DynamicArray);
 
-			if (ShadowData)
+			if (bUpdateShadowState)
 			{
 				FScriptArray* ShadowArray = (FScriptArray*)(ShadowData + ShadowOffset).Data;
 				FScriptArray* Array = (FScriptArray*)(Data + AbsOffset).Data;
@@ -4099,7 +4099,6 @@ void FRepLayout::UpdateUnmappedObjects_r(
 			}
 
 			bOutSomeObjectsWereMapped = true;
-			const bool bUpdateShadowState = (ShadowData && INDEX_NONE != Parent.RepNotifyNumParams);
 
 			// Copy current value over so we can check to see if it changed
 			if (bUpdateShadowState)
@@ -4373,6 +4372,15 @@ void FRepLayout::CallRepNotifies(FReceivingRepState* RepState, UObject* Object)
 			case 0:
 			{
 				Object->ProcessEvent(RepNotifyFunc, nullptr);
+
+				// TODO CopyCompleteValue no matter the field is replicated or not
+				// will be a performance regression for any RepNotify arrays.
+				// One fix is to track the incoming changelist and then resize the array and 
+				// recursively copy over only the fields we care about.
+				if (EnumHasAnyFlags(Parent.Flags, ERepParentFlags::HasDynamicArrayProperties) && !EnumHasAnyFlags(Parent.Flags, ERepParentFlags::IsFastArray))
+				{
+					RepProperty->CopyCompleteValue(ShadowData + Parent, ObjectData + Parent);
+				}
 				break;
 			}
 			case 1:
@@ -5012,6 +5020,7 @@ struct FInitFromPropertySharedParams
 	TArray<FRepLayoutCmd>& Cmds;
 	const UNetConnection* ServerConnection;
 	const int32 ParentIndex;
+	FRepParentCmd& Parent;
 	bool bHasObjectProperties = false;
 	bool bHasNetSerializeProperties = false;
 	TMap<int32, TArray<FRepLayoutCmd>>* NetSerializeLayouts = nullptr;
@@ -5293,6 +5302,7 @@ static int32 InitFromProperty_r(
 	if (FArrayProperty* ArrayProp = CastField<FArrayProperty>(StackParams.Property))
 	{
 		const int32 CmdStart = SharedParams.Cmds.Num();
+		SharedParams.Parent.Flags |= ERepParentFlags::HasDynamicArrayProperties;
 
 		++StackParams.RelativeHandle;
 		StackParams.Offset += GetOffsetForProperty<BuildType>(*ArrayProp);
@@ -5355,6 +5365,7 @@ static int32 InitFromProperty_r(
 					/*Cmds=*/*NewCmds,
 					/*ServerConnection=*/SharedParams.ServerConnection,
 					/*ParentIndex=*/SharedParams.ParentIndex,
+					/*Parent=*/SharedParams.Parent,
 					/*bHasObjectProperties=*/false,
 					/*bHasNetSerializeProperties=*/false,
 					/*NetSerializeLayouts=*/SharedParams.NetSerializeLayouts
@@ -5741,10 +5752,11 @@ void FRepLayout::InitFromClass(
 		{
 			/*Cmds=*/Cmds,
 			/*ServerConnection=*/ServerConnection,
-			/*ParentHandle=*/ParentHandle,
+			/*ParentIndex=*/ParentHandle,
+			/*Parent=*/Parents[ParentHandle],
 			/*bHasObjectProperties=*/false,
 			/*bHasNetSerializeProperties=*/false,
-			/*NetSerializeLayouts=*/GbTrackNetSerializeObjectReferences ? &TempNetSerializeLayouts : nullptr
+			/*NetSerializeLayouts=*/GbTrackNetSerializeObjectReferences ? &TempNetSerializeLayouts : nullptr,
 		};
 
 		FInitFromPropertyStackParams StackParams
@@ -6025,7 +6037,8 @@ void FRepLayout::InitFromFunction(
 			{
 				/*Cmds=*/Cmds,
 				/*ServerConnection=*/ServerConnection,
-				/*ParentHandle=*/ParentHandle
+				/*ParentIndex=*/ParentHandle,
+				/*Parent=*/Parents[ParentHandle],
 			};
 
 			FInitFromPropertyStackParams StackParams
@@ -6090,7 +6103,8 @@ void FRepLayout::InitFromStruct(
 			{
 				/*Cmds=*/Cmds,
 				/*ServerConnection=*/ServerConnection,
-				/*ParentHandle=*/ParentHandle
+				/*ParentIndex=*/ParentHandle,
+				/*Parent=*/Parents[ParentHandle],
 			};
 
 			FInitFromPropertyStackParams StackParams
@@ -7631,7 +7645,6 @@ ERepLayoutResult FRepLayout::DeltaSerializeFastArrayProperty(FFastArrayDeltaSeri
 				FReceivePropertiesSharedParams SharedParams{
 					/*bDoChecksum=*/ false,
 					/*bSkipRoleSwap=*/ !EnumHasAnyFlags(Flags, ERepLayoutFlags::IsActor),
-					/*RepNotifies=*/ nullptr,
 					Reader,
 					bOutHasUnmapped,
 					bOutGuidsChanged,
@@ -7646,7 +7659,8 @@ ERepLayoutResult FRepLayout::DeltaSerializeFastArrayProperty(FFastArrayDeltaSeri
 					nullptr,
 					&GuidReferences,
 					ItemLayoutStart,
-					ItemLayoutEnd
+					ItemLayoutEnd,
+					/*RepNotifies=*/ nullptr
 				};
 
 				// Read the first handle, and then start receiving properties.
diff --git a/Engine/Source/Runtime/Engine/Public/Net/RepLayout.h b/Engine/Source/Runtime/Engine/Public/Net/RepLayout.h
--- a/Engine/Source/Runtime/Engine/Public/Net/RepLayout.h
+++ b/Engine/Source/Runtime/Engine/Public/Net/RepLayout.h
@@ -711,17 +711,18 @@ enum class ERepLayoutCmdType : uint8
 enum class ERepParentFlags : uint32
 {
 	None						= 0,
-	IsLifetime					= (1 << 0),	//! This property is valid for the lifetime of the object (almost always set).
-	IsConditional				= (1 << 1),	//! This property has a secondary condition to check
-	IsConfig					= (1 << 2),	//! This property is defaulted from a config file
-	IsCustomDelta				= (1 << 3),	//! This property uses custom delta compression. Mutually exclusive with IsNetSerialize.
-	IsNetSerialize				= (1 << 4), //! This property uses a custom net serializer. Mutually exclusive with IsCustomDelta.
-	IsStructProperty			= (1 << 5),	//! This property is a FStructProperty.
-	IsZeroConstructible			= (1 << 6),	//! This property is ZeroConstructible.
-	IsFastArray					= (1 << 7),	//! This property is a FastArraySerializer. This can't be a ERepLayoutCmdType, because
-											//! these Custom Delta structs will have their inner properties tracked.
-	HasObjectProperties			= (1 << 8), //! This property is tracking UObjects (may be through nested properties).
-	HasNetSerializeProperties	= (1 << 9), //! This property contains Net Serialize properties (may be through nested properties).
+	IsLifetime					= (1 << 0),	 //! This property is valid for the lifetime of the object (almost always set).
+	IsConditional				= (1 << 1),	 //! This property has a secondary condition to check
+	IsConfig					= (1 << 2),	 //! This property is defaulted from a config file
+	IsCustomDelta				= (1 << 3),	 //! This property uses custom delta compression. Mutually exclusive with IsNetSerialize.
+	IsNetSerialize				= (1 << 4),  //! This property uses a custom net serializer. Mutually exclusive with IsCustomDelta.
+	IsStructProperty			= (1 << 5),	 //! This property is a FStructProperty.
+	IsZeroConstructible			= (1 << 6),	 //! This property is ZeroConstructible.
+	IsFastArray					= (1 << 7),	 //! This property is a FastArraySerializer. This can't be a ERepLayoutCmdType, because
+											 //! these Custom Delta structs will have their inner properties tracked.
+	HasObjectProperties			= (1 << 8),  //! This property is tracking UObjects (may be through nested properties).
+	HasNetSerializeProperties	= (1 << 9),  //! This property contains Net Serialize properties (may be through nested properties).
+	HasDynamicArrayProperties   = (1 << 10), //! This property contains Dynamic Array properties (may be through nested properties).
 };
 
 ENUM_CLASS_FLAGS(ERepParentFlags)
