CIS - Remove warning on cook: Chaos Cloth does not currenlty support bCollideWithEnvironment. [Merging CL 14472644]

diff --git a/Engine/Source/Runtime/Engine/Private/Components/SkeletalMeshComponent.cpp b/Engine/Source/Runtime/Engine/Private/Components/SkeletalMeshComponent.cpp
--- a/Engine/Source/Runtime/Engine/Private/Components/SkeletalMeshComponent.cpp
+++ b/Engine/Source/Runtime/Engine/Private/Components/SkeletalMeshComponent.cpp
@@ -633,7 +633,7 @@ void USkeletalMeshComponent::OnRegister()
 	// TODO: Add missing Chaos Cloth collision with environment
 	if (bCollideWithEnvironment && ClothingSimulationFactory && ClothingSimulationFactory->GetName() == TEXT("ChaosClothingSimulationFactory"))
 	{
-		UE_LOG(LogSkeletalMesh, Warning, TEXT("OnRegister[%s]: "
+		UE_LOG(LogSkeletalMesh, Display, TEXT("OnRegister[%s]: "
 			"Chaos Cloth does not currently support bCollideWithEnvironment."),
 			*GetNameSafe(SkeletalMesh));
 	}
