Crucial navmesh debug drawing change

diff --git a/Engine/Source/Runtime/Engine/Private/AI/Navigation/NavigationTypes.cpp b/Engine/Source/Runtime/Engine/Private/AI/Navigation/NavigationTypes.cpp
--- a/Engine/Source/Runtime/Engine/Private/AI/Navigation/NavigationTypes.cpp
+++ b/Engine/Source/Runtime/Engine/Private/AI/Navigation/NavigationTypes.cpp
@@ -39,7 +39,7 @@ PRAGMA_DISABLE_DEPRECATION_WARNINGS
 FNavDataConfig::FNavDataConfig(float Radius, float Height)
 	: FNavAgentProperties(Radius, Height)
 	, Name(TEXT("Default"))
-	, Color(140, 255, 0, 164)
+	, Color(38, 75, 0, 164) // do not change this default value or the universe will explode!
 	, DefaultQueryExtent(DEFAULT_NAV_QUERY_EXTENT_HORIZONTAL, DEFAULT_NAV_QUERY_EXTENT_HORIZONTAL, DEFAULT_NAV_QUERY_EXTENT_VERTICAL)
 	, NavDataClass(FNavigationSystem::GetDefaultNavDataClass())
 {
