- Using new multi-platform config system to be able to calculate what a CVar will actually be on a target device
- More improvements for using this system will be coming along (particle previewing, etc)
- Added some console commands:
  dpcvar <deviceprofilename>@<cvarname>
  dpdump <deviceprofilename>
  dppreview <deviceprofilename>
  dprestore
  dpreload (this actually reloads all platform configs, not just deviceprofiles)

diff --git a/Engine/Source/Runtime/Core/Private/HAL/ConsoleManager.cpp b/Engine/Source/Runtime/Core/Private/HAL/ConsoleManager.cpp
--- a/Engine/Source/Runtime/Core/Private/HAL/ConsoleManager.cpp
+++ b/Engine/Source/Runtime/Core/Private/HAL/ConsoleManager.cpp
@@ -311,11 +311,33 @@ template <class T>
 class FConsoleVariable : public FConsoleVariableBase
 {
 public:
-	FConsoleVariable(T DefaultValue, const TCHAR* Help, EConsoleVariableFlags Flags) 
+	FConsoleVariable(T DefaultValue, const TCHAR* Help, EConsoleVariableFlags Flags, bool bSaveDefault=true) 
 		: FConsoleVariableBase(Help, Flags), Data(DefaultValue)
 	{
+#if WITH_UNREAL_DEVELOPER_TOOLS
+		PlatformIndependentDefault = nullptr;
+		if (bSaveDefault)
+		{
+			CreatePlatformIndependentDefault(DefaultValue);
+		}
+#endif
 	}
 
+	virtual ~FConsoleVariable()
+	{
+#if WITH_UNREAL_DEVELOPER_TOOLS
+		delete PlatformIndependentDefault;
+#endif
+	}
+
+#if WITH_UNREAL_DEVELOPER_TOOLS
+
+	void CreatePlatformIndependentDefault(T DefaultValue)
+	{
+		PlatformIndependentDefault = new FConsoleVariable<T>(DefaultValue, TEXT(""), Flags, false);
+	}
+#endif
+
 	// interface IConsoleVariable ----------------------------------- 
 
 	virtual void Release()
@@ -364,6 +386,17 @@ private: // ----------------------------------------------------
 		OnCVarChange(Data.ShadowedValue[1], Data.ShadowedValue[0], Flags, SetBy);
 		FConsoleVariableBase::OnChanged(SetBy);
 	}
+
+#if WITH_UNREAL_DEVELOPER_TOOLS
+	// remember the dfefault value of this cvar, before anything else can assign to it - this way we can know in the editor what the value
+	// would be for a cvar on another platform, if no .ini file on that platform sets it
+	IConsoleVariable* PlatformIndependentDefault;
+	virtual IConsoleVariable* GetDefaultValueVariable() override
+	{
+		return PlatformIndependentDefault;
+	}
+#endif
+
 };
 
 // specialization for all
diff --git a/Engine/Source/Runtime/Core/Public/HAL/IConsoleManager.h b/Engine/Source/Runtime/Core/Public/HAL/IConsoleManager.h
--- a/Engine/Source/Runtime/Core/Public/HAL/IConsoleManager.h
+++ b/Engine/Source/Runtime/Core/Public/HAL/IConsoleManager.h
@@ -133,7 +133,7 @@ class IConsoleVariable;
 #if !NO_CVARS
 
 /** Console variable delegate type  This is a void callback function. */
-DECLARE_DELEGATE_OneParam(FConsoleVariableDelegate, IConsoleVariable*);
+DECLARE_DELEGATE_OneParam( FConsoleVariableDelegate, IConsoleVariable* );
 
 /** Console variable multicast delegate type. */
 DECLARE_MULTICAST_DELEGATE_OneParam(FConsoleVariableMulticastDelegate, IConsoleVariable*);
@@ -431,6 +431,16 @@ public:
 
 	virtual FConsoleVariableMulticastDelegate& OnChangedDelegate() = 0;
 
+#if WITH_UNREAL_DEVELOPER_TOOLS
+	/**
+	 * Get the saved off default value, in a cvar variable
+	 */
+	virtual IConsoleVariable* GetDefaultValueVariable()
+	{
+		return nullptr;
+	}
+#endif
+
 	// convenience methods
 
 	/** Set the internal value from the specified bool. */
@@ -475,6 +485,8 @@ public:
 		EConsoleVariableFlags CurFlags = (EConsoleVariableFlags)(GetFlags() & ECVF_SetByMask);
 		Set(InValue, CurFlags);
 	}
+
+
 };
 
 /**
diff --git a/Engine/Source/Runtime/Engine/Classes/DeviceProfiles/DeviceProfile.h b/Engine/Source/Runtime/Engine/Classes/DeviceProfiles/DeviceProfile.h
--- a/Engine/Source/Runtime/Engine/Classes/DeviceProfiles/DeviceProfile.h
+++ b/Engine/Source/Runtime/Engine/Classes/DeviceProfiles/DeviceProfile.h
@@ -145,5 +145,15 @@ private:
 private:
 	/** Consolidated CVars, lazy initialized - access via GetConsolidatedCVars */
 	mutable TMap<FString, FString> ConsolidatedCVars;
+
+#endif
+
+#if WITH_UNREAL_DEVELOPER_TOOLS
+public:
+	const TMap<FString, FString>& GetAllExpandedCVars();
+	void AddExpandedCVars(const TMap<FString, FString>& CVarsToMerge);
+private:
+	/** Resolved CVars, including expanded scalability cvars used to properly emulate one platform on another */
+	TMap<FString, FString> AllExpandedCVars;
 #endif
 };
diff --git a/Engine/Source/Runtime/Engine/Classes/DeviceProfiles/DeviceProfileManager.h b/Engine/Source/Runtime/Engine/Classes/DeviceProfiles/DeviceProfileManager.h
--- a/Engine/Source/Runtime/Engine/Classes/DeviceProfiles/DeviceProfileManager.h
+++ b/Engine/Source/Runtime/Engine/Classes/DeviceProfiles/DeviceProfileManager.h
@@ -30,7 +30,7 @@ public:
 	 * Startup and select the active device profile
 	 * Then Init the CVars from this profile and it's Device profile parent tree.
 	 */
-	static void InitializeCVarsForActiveDeviceProfile(bool bPushSettings=false, bool bIsDeviceProfilePreview = false);
+	static void InitializeCVarsForActiveDeviceProfile(bool bPushSettings=false);
 
 	/**
 	 * Reapplies the device profile. Useful when configs have changed (i.e. hotfix)
@@ -142,6 +142,11 @@ public:
 	static bool GetScalabilityCVar(const FString& CvarName, int32& OutValue);
 	static bool GetScalabilityCVar(const FString& CvarName, float& OutValue);
 
+#if WITH_UNREAL_DEVELOPER_TOOLS
+	/** Retrieve another platform's DeviceProfile, and make sure it's CVars are expanded, including scalability */
+	static void ExpandDeviceProfileCVars(UDeviceProfile* DeviceProfile);
+#endif
+
 private:
 	/**
 	 * Set the active device profile - set via the device profile blueprint.
@@ -158,6 +163,20 @@ private:
 	/** Handle restoing CVars set in HandleDeviceProfileOverrideChange */
 	void HandleDeviceProfileOverridePop();
 
+
+	enum class EDeviceProfileMode : uint8
+	{
+		DPM_SetCVars,
+		DPM_PushCVars,
+		DPM_CacheValues,
+	};
+
+	/**
+	 * Perform the processing of ini sections, going up to parents, etc. Depending on runtime vs editor processing of another platform,
+	 * the Mode will control how the settings are handled
+	 */
+	static void ProcessDeviceProfileIniSettings(const FString& DeviceProfileName, EDeviceProfileMode Mode);
+
 public:
 
 	static class UDeviceProfileManager* DeviceProfileManagerSingleton;
diff --git a/Engine/Source/Runtime/Engine/Private/DeviceProfiles/DeviceProfile.cpp b/Engine/Source/Runtime/Engine/Private/DeviceProfiles/DeviceProfile.cpp
--- a/Engine/Source/Runtime/Engine/Private/DeviceProfiles/DeviceProfile.cpp
+++ b/Engine/Source/Runtime/Engine/Private/DeviceProfiles/DeviceProfile.cpp
@@ -9,6 +9,7 @@
 #include "HAL/IConsoleManager.h"
 
 #include "DeviceProfiles/DeviceProfileFragment.h"
+#include "DeviceProfiles/DeviceProfileManager.h"
 
 DEFINE_LOG_CATEGORY_STATIC(LogDeviceProfile, Log, All);
 
@@ -364,3 +365,28 @@ const TMap<FString, FString>& UDeviceProfile::GetConsolidatedCVars() const
 }
 
 #endif
+
+
+#if WITH_UNREAL_DEVELOPER_TOOLS
+const TMap<FString, FString>& UDeviceProfile::GetAllExpandedCVars()
+{
+	// expand on first use
+	if (AllExpandedCVars.Num() == 0)
+	{
+		UDeviceProfileManager::Get().ExpandDeviceProfileCVars(this);
+	}
+
+	return AllExpandedCVars;
+}
+
+void UDeviceProfile::AddExpandedCVars(const TMap<FString, FString>& CVarsToMerge)
+{
+	// merge the cvars, overwriting any existing ones
+	for (const auto& Pair : CVarsToMerge)
+	{
+		AllExpandedCVars.Add(Pair.Key, Pair.Value);
+	}
+}
+
+
+#endif
\ No newline at end of file
diff --git a/Engine/Source/Runtime/Engine/Private/DeviceProfiles/DeviceProfileManager.cpp b/Engine/Source/Runtime/Engine/Private/DeviceProfiles/DeviceProfileManager.cpp
--- a/Engine/Source/Runtime/Engine/Private/DeviceProfiles/DeviceProfileManager.cpp
+++ b/Engine/Source/Runtime/Engine/Private/DeviceProfiles/DeviceProfileManager.cpp
@@ -20,6 +20,8 @@
 #include "ProfilingDebugging/CsvProfiler.h"
 #include "DeviceProfiles/DeviceProfileFragment.h"
 
+DEFINE_LOG_CATEGORY_STATIC(LogDeviceProfileManager, Log, All);
+
 static TAutoConsoleVariable<FString> CVarDeviceProfileOverride(
 	TEXT("dp.Override"),
 	TEXT(""),
@@ -82,52 +84,67 @@ UDeviceProfileManager& UDeviceProfileManager::Get(bool bFromPostCDOContruct)
 	return *DeviceProfileManagerSingleton;
 }
 
-static void GetFragmentCvars(const FString& CurrentSectionName, const FString& CVarArrayName, TArray<FString>& FragmentCVarsINOUT, const FString& DeviceProfileFileNameIn)
+static void GetFragmentCvars(const FString& CurrentSectionName, const FString& CVarArrayName, TArray<FString>& FragmentCVarsINOUT, FConfigCacheIni* ConfigSystem)
 {
 	FString FragmentIncludes = TEXT("FragmentIncludes");
 	TArray<FString> FragmentIncludeArray;
-	GConfig->GetArray(*CurrentSectionName, *FragmentIncludes, FragmentIncludeArray, DeviceProfileFileNameIn);
+	ConfigSystem->GetArray(*CurrentSectionName, *FragmentIncludes, FragmentIncludeArray, GDeviceProfilesIni);
 
 	for(const FString& FragmentInclude : FragmentIncludeArray)
 	{
 		FString FragmentSectionName = FString::Printf(TEXT("%s %s"), *FragmentInclude, *UDeviceProfileFragment::StaticClass()->GetName());
-		if (GConfig->DoesSectionExist(*FragmentSectionName, DeviceProfileFileNameIn))
+		if (ConfigSystem->DoesSectionExist(*FragmentSectionName, GDeviceProfilesIni))
 		{
 			TArray<FString> FragmentCVars;
-			GConfig->GetArray(*FragmentSectionName, *CVarArrayName, FragmentCVars, DeviceProfileFileNameIn);
-			UE_CLOG(FragmentCVars.Num()>0, LogInit, Log, TEXT("Including %s from fragment: %s"), *CVarArrayName, *FragmentInclude);
+			ConfigSystem->GetArray(*FragmentSectionName, *CVarArrayName, FragmentCVars, GDeviceProfilesIni);
+			UE_CLOG(FragmentCVars.Num()>0, LogDeviceProfileManager, Log, TEXT("Including %s from fragment: %s"), *CVarArrayName, *FragmentInclude);
 			FragmentCVarsINOUT += FragmentCVars;
 		}
 		else
 		{
 #if UE_BUILD_SHIPPING
-			UE_LOG(LogInit, Error, TEXT("Could not find device profile fragment %s."), *FragmentInclude);
+			UE_LOG(LogDeviceProfileManager, Error, TEXT("Could not find device profile fragment %s."), *FragmentInclude);
 #else
-			UE_LOG(LogInit, Fatal, TEXT("Could not find device profile fragment %s."), *FragmentInclude);
+			UE_LOG(LogDeviceProfileManager, Fatal, TEXT("Could not find device profile fragment %s."), *FragmentInclude);
 #endif
 		}
 	}
 }
 
-void UDeviceProfileManager::InitializeCVarsForActiveDeviceProfile(bool bPushSettings, bool bIsDeviceProfilePreview)
+
+static void ExpandScalabilityCVar(FConfigCacheIni* ConfigSystem, const FString& CVarKey, const FString CVarValue, TMap<FString, FString>& ExpandedCVars, bool bOverwriteExistingValue)
 {
-	FString ActiveProfileName;
-	
-	if(DeviceProfileManagerSingleton)
+	// load scalability settings directly from ini instead of using scalability system, so as not to inadvertantly mess anything up
+	// if the DP had sg.ResolutionQuality=3, we would read [ResolutionQuality@3]
+	FString SectionName = FString::Printf(TEXT("%s@%s"), *CVarKey.Mid(3), *CVarValue);
+	// walk over the scalability section and add them in, unless already done
+	FConfigSection* ScalabilitySection = ConfigSystem->GetSectionPrivate(*SectionName, false, true, GScalabilityIni);
+	if (ScalabilitySection != nullptr)
 	{
-		ActiveProfileName = DeviceProfileManagerSingleton->ActiveDeviceProfile->GetName();
-
-		//Ensure we've loaded the device profiles for the active platform.
-		//This can be needed when overriding the device profile.
-		FString ActivePlatformName = DeviceProfileManagerSingleton->ActiveDeviceProfile->DeviceType;
-		FConfigCacheIni::LoadGlobalIniFile(GDeviceProfilesIni, TEXT("DeviceProfiles"), *ActivePlatformName, true);
+		for (const auto& Pair : *ScalabilitySection)
+		{
+			FString ScalabilityKey = Pair.Key.ToString();
+			if (bOverwriteExistingValue || !ExpandedCVars.Contains(ScalabilityKey))
+			{
+				ExpandedCVars.Add(ScalabilityKey, Pair.Value.GetValue());
+			}
+		}
 	}
-	else
+}
+
+void UDeviceProfileManager::ProcessDeviceProfileIniSettings(const FString& DeviceProfileName, EDeviceProfileMode Mode)
+{
+	FConfigCacheIni* ConfigSystem = GConfig;
+	if (Mode == EDeviceProfileMode::DPM_CacheValues)
 	{
-		ActiveProfileName = GetPlatformDeviceProfileName();
+		// caching is not done super early, so we can assume DPs have been found now
+		UDeviceProfile* Profile = UDeviceProfileManager::Get().FindProfile(DeviceProfileName, false);
+		check(Profile);
+		// use the DP's platform's configs, NOT the running platform
+		ConfigSystem = FConfigCacheIni::ForPlatform(*Profile->DeviceType);
 	}
 
-	UE_LOG(LogInit, Log, TEXT("Applying CVar settings loaded from the selected device profile: [%s]"), *ActiveProfileName);
+	check(ConfigSystem != nullptr);
 
 	TArray< FString > AvailableProfiles;
 	GConfig->GetSectionNames( GDeviceProfilesIni, AvailableProfiles );
@@ -136,65 +153,68 @@ void UDeviceProfileManager::InitializeCVarsForActiveDeviceProfile(bool bPushSett
 	AvailableProfiles.Remove( TEXT( "DeviceProfiles" ) );
 
 	// Next we need to create a hierarchy of CVars from the Selected Device Profile, to it's eldest parent
+	// if we are just caching, this also contains the set of all CVars (including expanding scalability groups)
 	TMap<FString, FString> CVarsAlreadySetList;
-	DeviceProfileScalabilityCVars.Empty();
-	
-	// even if we aren't pushing new values, we should clear any old pushed values, as they are no longer valid after we run this loop
-	if (DeviceProfileManagerSingleton)
+
+	// reset some global state for "active DP" mode
+	if (Mode != EDeviceProfileMode::DPM_CacheValues)
 	{
-		DeviceProfileManagerSingleton->PushedSettings.Empty();
-	}
+		DeviceProfileScalabilityCVars.Empty();
+
+		// even if we aren't pushing new values, we should clear any old pushed values, as they are no longer valid after we run this loop
+		if (DeviceProfileManagerSingleton)
+		{
+			DeviceProfileManagerSingleton->PushedSettings.Empty();
+		}
 
 #if !UE_BUILD_SHIPPING
 #if PLATFORM_ANDROID
-	// allow ConfigRules to override cvars first
-	TMap<FString, FString> ConfigRules = FAndroidMisc::GetConfigRulesTMap();
-	for (const TPair<FString, FString>& Pair : ConfigRules)
-	{
-		FString Key = Pair.Key;
-		if (Key.StartsWith("cvar_"))
+		// allow ConfigRules to override cvars first
+		TMap<FString, FString> ConfigRules = FAndroidMisc::GetConfigRulesTMap();
+		for (const TPair<FString, FString>& Pair : ConfigRules)
 		{
-			FString CVarKey = Key.RightChop(5);
-			FString CVarValue = Pair.Value;
+			FString Key = Pair.Key;
+			if (Key.StartsWith("cvar_"))
+			{
+				FString CVarKey = Key.RightChop(5);
+				FString CVarValue = Pair.Value;
 
-			UE_LOG(LogInit, Log, TEXT("Setting ConfigRules Device Profile CVar: [[%s:%s]]"), *CVarKey, *CVarValue);
+				UE_LOG(LogDeviceProfileManager, Log, TEXT("Setting ConfigRules Device Profile CVar: [[%s:%s]]"), *CVarKey, *CVarValue);
 
-			// set it and remember it
-			OnSetCVarFromIniEntry(*GDeviceProfilesIni, *CVarKey, *CVarValue, ECVF_SetByDeviceProfile);
-			CVarsAlreadySetList.Add(CVarKey, CVarValue);
+				// set it and remember it
+				OnSetCVarFromIniEntry(*GDeviceProfilesIni, *CVarKey, *CVarValue, ECVF_SetByDeviceProfile);
+				CVarsAlreadySetList.Add(CVarKey, CVarValue);
+			}
 		}
-	}
 #endif
 #endif
 
 #if !UE_BUILD_SHIPPING
-	// pre-apply any -dpcvars= items, so that they override anything in the DPs
-	FString DPCVarString;
-	if (FParse::Value(FCommandLine::Get(), TEXT("DPCVars="), DPCVarString, false) || FParse::Value(FCommandLine::Get(), TEXT("DPCVar="), DPCVarString, false))
-	{
-		// look over a list of cvars
-		TArray<FString> DPCVars;
-		DPCVarString.ParseIntoArray(DPCVars, TEXT(","), true);
-		for (FString& DPCVar : DPCVars)
-		{
-			// split up each Key=Value pair
-			FString CVarKey, CVarValue;
-			if (DPCVar.Split(TEXT("="), &CVarKey, &CVarValue))
+		// pre-apply any -dpcvars= items, so that they override anything in the DPs
+		FString DPCVarString;
+		if (FParse::Value(FCommandLine::Get(), TEXT("DPCVars="), DPCVarString, false) || FParse::Value(FCommandLine::Get(), TEXT("DPCVar="), DPCVarString, false))
+		{
+			// look over a list of cvars
+			TArray<FString> DPCVars;
+			DPCVarString.ParseIntoArray(DPCVars, TEXT(","), true);
+			for (const FString& DPCVar : DPCVars)
 			{
-				UE_LOG(LogInit, Log, TEXT("Setting CommandLine Device Profile CVar: [[%s:%s]]"), *CVarKey, *CVarValue);
+				// split up each Key=Value pair
+				FString CVarKey, CVarValue;
+				if (DPCVar.Split(TEXT("="), &CVarKey, &CVarValue))
+				{
+					UE_LOG(LogDeviceProfileManager, Log, TEXT("Setting CommandLine Device Profile CVar: [[%s:%s]]"), *CVarKey, *CVarValue);
 
-				// set it and remember it (no thanks, Ron Popeil)
-				OnSetCVarFromIniEntry(*GDeviceProfilesIni, *CVarKey, *CVarValue, ECVF_SetByDeviceProfile);
-				CVarsAlreadySetList.Add(CVarKey, CVarValue);
+					// set it and remember it (no thanks, Ron Popeil)
+					OnSetCVarFromIniEntry(*GDeviceProfilesIni, *CVarKey, *CVarValue, ECVF_SetByDeviceProfile);
+					CVarsAlreadySetList.Add(CVarKey, CVarValue);
+				}
 			}
 		}
-	}
 #endif
-	
-	// Preload a cvar we rely on
-	if (GConfig)
-	{
-		if (FConfigSection* Section = GConfig->GetSectionPrivate(TEXT("ConsoleVariables"), false, true, *GEngineIni))
+
+		// Preload a cvar we rely on
+		if (FConfigSection* Section = ConfigSystem->GetSectionPrivate(TEXT("ConsoleVariables"), false, true, *GEngineIni))
 		{
 			static FName AllowScalabilityAtRuntimeName = TEXT("dp.AllowScalabilityGroupsToChangeAtRuntime");
 			if (const FConfigValue* Value = Section->Find(AllowScalabilityAtRuntimeName))
@@ -208,7 +228,7 @@ void UDeviceProfileManager::InitializeCVarsForActiveDeviceProfile(bool bPushSett
 
 	// For each device profile, starting with the selected and working our way up the BaseProfileName tree,
 	// Find all CVars and set them 
-	FString BaseDeviceProfileName = ActiveProfileName;
+	FString BaseDeviceProfileName = DeviceProfileName;
 	bool bReachedEndOfTree = BaseDeviceProfileName.IsEmpty();
 	while( bReachedEndOfTree == false ) 
 	{
@@ -234,12 +254,20 @@ void UDeviceProfileManager::InitializeCVarsForActiveDeviceProfile(bool bPushSett
 				FString ArrayName = TEXT("CVars");
 				if (Pass == 0)
 				{
-					ArrayName += BucketNames[(int32)FPlatformMemory::GetMemorySizeBucket()];
+					// assume default when caching for another platform, since we don' thave a current device to emulate (maybe we want to be able to pass in override memory bucket?)
+					if (Mode == EDeviceProfileMode::DPM_CacheValues)
+					{
+						ArrayName += TEXT("_Default");
+					}
+					else
+					{
+						ArrayName += BucketNames[(int32)FPlatformMemory::GetMemorySizeBucket()];
+					}
 				}
 
 				TArray< FString > CurrentProfilesCVars, FragmentCVars;
-				GetFragmentCvars(*CurrentSectionName, *ArrayName, FragmentCVars, GDeviceProfilesIni);
-				GConfig->GetArray(*CurrentSectionName, *ArrayName, CurrentProfilesCVars, GDeviceProfilesIni);
+				GetFragmentCvars(*CurrentSectionName, *ArrayName, FragmentCVars, ConfigSystem);
+				ConfigSystem->GetArray(*CurrentSectionName, *ArrayName, CurrentProfilesCVars, GDeviceProfilesIni);
 
 				if (FragmentCVars.Num())
 				{
@@ -283,41 +311,58 @@ void UDeviceProfileManager::InitializeCVarsForActiveDeviceProfile(bool bPushSett
 					{
 						if (!CVarsAlreadySetList.Find(CVarKey))
 						{
-							IConsoleVariable* CVar = IConsoleManager::Get().FindConsoleVariable(*CVarKey);
-							if (CVar)
+							if (Mode == EDeviceProfileMode::DPM_PushCVars)
 							{
-								if (DeviceProfileManagerSingleton && bPushSettings)
+								IConsoleVariable* CVar = IConsoleManager::Get().FindConsoleVariable(*CVarKey);
+								if (CVar)
 								{
-									// remember the previous value
-									FString OldValue = CVar->GetString();
-									DeviceProfileManagerSingleton->PushedSettings.Add(CVarKey, OldValue);
-
-									// indicate we are pushing, not setting
-									UE_LOG(LogInit, Log, TEXT("Pushing Device Profile CVar: [[%s:%s -> %s]]"), *CVarKey, *OldValue, *CVarValue);
+									if (DeviceProfileManagerSingleton)
+									{
+										// remember the previous value
+										FString OldValue = CVar->GetString();
+										DeviceProfileManagerSingleton->PushedSettings.Add(CVarKey, OldValue);
+
+										// indicate we are pushing, not setting
+										UE_LOG(LogDeviceProfileManager, Log, TEXT("Pushing Device Profile CVar: [[%s:%s -> %s]]"), *CVarKey, *OldValue, *CVarValue);
+									}
+								}
+								else
+								{
+									UE_LOG(LogDeviceProfileManager, Warning, TEXT("Creating unregistered Device Profile CVar: [[%s:%s]]"), *CVarKey, *CVarValue);
 								}
-							}
-							else
-							{
-								UE_LOG(LogInit, Warning, TEXT("Creating unregistered Device Profile CVar: [[%s:%s]]"), *CVarKey, *CVarValue);
 							}
 
 							// General scalability bucket cvars are set as a suggested default but can be overridden by game settings.
-							bool bIsScalabilityBucket = false;
-							if (CVarAllowScalabilityGroupsToChangeAtRuntime.GetValueOnGameThread() > 0)
+							bool bIsScalabilityBucket = CVarKey.StartsWith(TEXT("sg."));
+
+							if (Mode == EDeviceProfileMode::DPM_CacheValues)
 							{
+								if (bIsScalabilityBucket)
+								{
+									// don't overwrite any existing cvars when expanding
+									ExpandScalabilityCVar(ConfigSystem, CVarKey, CVarValue, CVarsAlreadySetList, false);
+								}
+
+								// cache key with value
+								CVarsAlreadySetList.Add(CVarKey, CVarValue);
+							}
+							// actually set the cvar if not just caching
+							else
+							{
+								
 								// Cache any scalability related cvars so we can conveniently reapply them later as a way to reset the device defaults
-								if (CVarKey.StartsWith(TEXT("sg.")))
+								if (bIsScalabilityBucket && CVarAllowScalabilityGroupsToChangeAtRuntime.GetValueOnGameThread() > 0)
 								{
 									DeviceProfileScalabilityCVars.Add(*CVarKey, *CVarValue);
-									bIsScalabilityBucket = true;
 								}
-							}
 
-							//If this is a dp preview then we set cvars with their existing priority so that we don't cause future issues when setting by scalability levels etc.
-							uint32 BaseCVarPriority = bIsDeviceProfilePreview ? ECVF_SetByMask : ECVF_SetByDeviceProfile;
-							uint32 CVarPriority = bIsScalabilityBucket ? ECVF_SetByScalability : BaseCVarPriority;
-							OnSetCVarFromIniEntry(*GDeviceProfilesIni, *CVarKey, *CVarValue, CVarPriority);
-							CVarsAlreadySetList.Add(CVarKey, CVarValue);
+								//If this is a dp preview then we set cvars with their existing priority so that we don't cause future issues when setting by scalability levels etc.
+								// @todo ini: preview should never be setting cvars via this function, should use the cached values
+								uint32 BaseCVarPriority = /*bIsDeviceProfilePreview ? ECVF_SetByMask : */ECVF_SetByDeviceProfile;
+								uint32 CVarPriority = bIsScalabilityBucket ? ECVF_SetByScalability : BaseCVarPriority;
+								OnSetCVarFromIniEntry(*GDeviceProfilesIni, *CVarKey, *CVarValue, CVarPriority);
+								CVarsAlreadySetList.Add(CVarKey, CVarValue);
+							}
 						}
 					}
 				}
@@ -325,10 +370,10 @@ void UDeviceProfileManager::InitializeCVarsForActiveDeviceProfile(bool bPushSett
 
 			// Get the next device profile name, to look for CVars in, along the tree
 			FString NextBaseDeviceProfileName;
-			if( GConfig->GetString( *CurrentSectionName, TEXT("BaseProfileName"), NextBaseDeviceProfileName, GDeviceProfilesIni ) )
+			if (ConfigSystem->GetString(*CurrentSectionName, TEXT("BaseProfileName"), NextBaseDeviceProfileName, GDeviceProfilesIni))
 			{
 				BaseDeviceProfileName = NextBaseDeviceProfileName;
-				UE_LOG(LogInit, Log, TEXT("Going up to parent DeviceProfile [%s]"), *BaseDeviceProfileName);
+				UE_LOG(LogDeviceProfileManager, Log, TEXT("Going up to parent DeviceProfile [%s]"), *BaseDeviceProfileName);
 			}
 			else
 			{
@@ -339,7 +384,112 @@ void UDeviceProfileManager::InitializeCVarsForActiveDeviceProfile(bool bPushSett
 		// Check if we have inevitably reached the end of the device profile tree.
 		bReachedEndOfTree = !bProfileExists || BaseDeviceProfileName.IsEmpty();
 	}
+
+	// copy the running cache into the DP
+	if (Mode == EDeviceProfileMode::DPM_CacheValues)
+	{
+		UDeviceProfile* Profile = UDeviceProfileManager::Get().FindProfile(DeviceProfileName, false);
+		Profile->AddExpandedCVars(CVarsAlreadySetList);
+	}
+}
+
+
+void UDeviceProfileManager::InitializeCVarsForActiveDeviceProfile(bool bPushSettings)
+{
+	FString ActiveProfileName;
+
+	if (DeviceProfileManagerSingleton)
+	{
+		ActiveProfileName = DeviceProfileManagerSingleton->ActiveDeviceProfile->GetName();
+	}
+	else
+	{
+		ActiveProfileName = GetPlatformDeviceProfileName();
+	}
+
+	UE_LOG(LogDeviceProfileManager, Log, TEXT("Applying CVar settings loaded from the selected device profile: [%s]"), *ActiveProfileName);
+	ProcessDeviceProfileIniSettings(ActiveProfileName, bPushSettings ? EDeviceProfileMode::DPM_PushCVars : EDeviceProfileMode::DPM_SetCVars);
+}
+
+#if WITH_UNREAL_DEVELOPER_TOOLS
+void UDeviceProfileManager::ExpandDeviceProfileCVars(UDeviceProfile* DeviceProfile)
+{
+	// get the config system for the platform the DP uses
+	FConfigCacheIni* ConfigSystem = FConfigCacheIni::ForPlatform(*DeviceProfile->DeviceType);
+	check(ConfigSystem);
+
+	// walk up the chain cvar setBys and emulate what would happen on the target platform
+	FString Platform = DeviceProfile->DeviceType;
+
+	// now walk up the stack getting current values
+
+	// ECVF_SetByConstructor:
+	//   in PlatformIndependentDefault, used if getting a var but not in this DP
+
+	// ECVF_SetByScalability:
+	//	 skipped, i believe this is not really loaded as a normal layer per-se, it's up to the other sections to set with this one
+
+	// ECVF_SetByGameSetting:
+	//   skipped, since we don't have a user
+
+	const TCHAR* SectionNames[] = {
+		// ECVF_SetByProjectSetting:
+		TEXT("/Script/Engine.RendererSettings"),
+		TEXT("/Script/Engine.RendererOverrideSettings"),
+		TEXT("/Script/Engine.StreamingSettings"),
+		TEXT("/Script/Engine.GarbageCollectionSettings"),
+		TEXT("/Script/Engine.NetworkSettings"),
+
+		// ECVF_SetBySystemSettingsIni:
+		TEXT("SystemSettings"),
+		TEXT("ConsoleVariables"),
+	};
+
+	// go through possible cvar sections that the target platform would load and read all cvars in them
+	TMap<FString, FString> CVarsToAdd;
+	for (const TCHAR* SectionName : SectionNames)
+	{
+		FConfigSection* Section = ConfigSystem->GetSectionPrivate(SectionName, false, true, GEngineIni);
+		if (Section != nullptr)
+		{
+			// add the cvars from the section
+			for (const auto& Pair : *Section)
+			{
+				FString Key = Pair.Key.ToString();
+				FString Value = Pair.Value.GetValue();
+				if (Key.StartsWith(TEXT("sg.")))
+				{
+					// @todo ini: If anything in here was already set, overwrite it or skip it?
+					// the priorities may cause runtime to fail to set a cvar that this will set blindly, since we are ignoring
+					// priority by doing them "in order". Scalablity is one of the lowest priorities, so should almost never be allowed?
+					ExpandScalabilityCVar(ConfigSystem, Key, Value, CVarsToAdd, true);
+				}
+				CVarsToAdd.Add(Pair.Key.ToString(), Pair.Value.GetValue());
+			}
+		}
+	}
+	DeviceProfile->AddExpandedCVars(CVarsToAdd);
+
+	// ECVF_SetByDeviceProfile:
+	ProcessDeviceProfileIniSettings(DeviceProfile->GetName(), EDeviceProfileMode::DPM_CacheValues);
+
+	// ECVF_SetByConsoleVariablesIni:
+	//   maybe skip this? it's a weird one, but maybe?
+
+	// ECVF_SetByCommandline:
+	//   skip as this would not be expected to apply to emulation
+
+	// ECVF_SetByCode:
+	//   skip because it cannot be set by code
+
+	// ECVF_SetByConsole
+	//   we could have this if we made a per-platform CVar, not just the shared default value
 }
+#endif
+
+
+
+
 
 bool UDeviceProfileManager::DoActiveProfilesReference(const TSet<FString>& DeviceProfilesToQuery)
 {
@@ -391,7 +541,7 @@ void UDeviceProfileManager::ReapplyDeviceProfile()
 	UDeviceProfile* OverrideProfile = DeviceProfileManagerSingleton->BaseDeviceProfile ? DeviceProfileManagerSingleton->GetActiveProfile() : nullptr;
 	UDeviceProfile* BaseProfile = DeviceProfileManagerSingleton->BaseDeviceProfile ? DeviceProfileManagerSingleton->BaseDeviceProfile : DeviceProfileManagerSingleton->GetActiveProfile();
 
-	UE_LOG(LogInit, Log, TEXT("ReapplyDeviceProfile applying profile: [%s]"), *BaseProfile->GetName(), OverrideProfile ? *OverrideProfile->GetName() : TEXT("not set.") );
+	UE_LOG(LogDeviceProfileManager, Log, TEXT("ReapplyDeviceProfile applying profile: [%s]"), *BaseProfile->GetName(), OverrideProfile ? *OverrideProfile->GetName() : TEXT("not set.") );
 
 	// pop any pushed settings
 	RestoreDefaultDeviceProfile();
@@ -402,7 +552,7 @@ void UDeviceProfileManager::ReapplyDeviceProfile()
 
 	if (OverrideProfile)
 	{
-		UE_LOG(LogInit, Log, TEXT("ReapplyDeviceProfile applying override profile: [%s]"), *OverrideProfile->GetName());
+		UE_LOG(LogDeviceProfileManager, Log, TEXT("ReapplyDeviceProfile applying override profile: [%s]"), *OverrideProfile->GetName());
 		// reapply the override.
 		SetOverrideDeviceProfile(OverrideProfile);
 	}
@@ -422,7 +572,7 @@ static void TestProfileForCircularReferences(const FString& ProfileName, const F
 	{
 		if (ProfileDependancies.FindByPredicate([CurrentParent](const FString& InName) { return InName.Equals(CurrentParent); }))
 		{
-			UE_LOG(LogInit, Fatal, TEXT("Device Profile %s has a circular dependency on %s"), *ProfileName, *CurrentParent);
+			UE_LOG(LogDeviceProfileManager, Fatal, TEXT("Device Profile %s has a circular dependency on %s"), *ProfileName, *CurrentParent);
 		}
 		else
 		{
@@ -530,7 +680,7 @@ UDeviceProfile* UDeviceProfileManager::FindProfile( const FString& ProfileName,
 
 	if ( bCreateProfileOnFail && FoundProfile == nullptr )
 	{
-		FoundProfile = CreateProfile(ProfileName, FPlatformProperties::PlatformName());
+		FoundProfile = CreateProfile(ProfileName, FPlatformProperties::IniPlatformName());
 	}
 	return FoundProfile;
 }
@@ -656,12 +806,36 @@ void UDeviceProfileManager::SetOverrideDeviceProfile(UDeviceProfile* DeviceProfi
 	// pop any pushed settings
 	HandleDeviceProfileOverridePop();
 
+	// for preview, we assume this will be another platform's DP, so use the resolved cvars directly, bypassing the activate and set stuff
+	if (bIsDeviceProfilePreview)
+	{
+#if WITH_UNREAL_DEVELOPER_TOOLS
+		for (const auto& Pair : DeviceProfile->GetAllExpandedCVars())
+		{
+			IConsoleVariable* CVar = IConsoleManager::Get().FindConsoleVariable(*Pair.Key);
+			// skip over scalability group cvars (maybe they shouldn't be left in the AllExpandedCVars?)
+			if (CVar != nullptr && !CVar->TestFlags(EConsoleVariableFlags::ECVF_ScalabilityGroup))
+			{
+				// remember the previous value so we can restore
+				FString OldValue = CVar->GetString();
+				PushedSettings.Add(Pair.Key, OldValue);
+
+				// set the cvar to the new value, with same priority that it was before (SetByMask means current priority)
+				CVar->SetWithCurrentPriority(*Pair.Value);
+			}
+		}
+#else
+		UE_LOG(LogDeviceProfileManager, Error, TEXT("SetOverrideDeviceProfile can only be used in a developer tool"));
+#endif
+		return;
+	}
+
 	// record the currently active profile, needed when we restore the default.
 	BaseDeviceProfile = DeviceProfileManagerSingleton->GetActiveProfile();
 
 	// activate new one!
 	DeviceProfileManagerSingleton->SetActiveDeviceProfile(DeviceProfile);
-	InitializeCVarsForActiveDeviceProfile(true, bIsDeviceProfilePreview);
+	InitializeCVarsForActiveDeviceProfile(/*bPushSettings*/true);
 
 	// broadcast cvar sinks now that we are done
 	IConsoleManager::Get().CallAllConsoleVariableSinks();
@@ -680,7 +854,7 @@ void UDeviceProfileManager::RestoreDefaultDeviceProfile()
 		{
 			// restore it!
 			CVar->SetWithCurrentPriority(*It.Value());			
-			UE_LOG(LogInit, Log, TEXT("Popping Device Profile CVar: [[%s:%s]]"), *It.Key(), *It.Value());
+			UE_LOG(LogDeviceProfileManager, Log, TEXT("Popping Device Profile CVar: [[%s:%s]]"), *It.Key(), *It.Value());
 		}
 	}
 
@@ -702,9 +876,11 @@ void UDeviceProfileManager::HandleDeviceProfileOverrideChange()
 	// only handle when the value is different
 	if (CVarValue.Len() > 0 && CVarValue != GetActiveProfile()->GetName())
 	{
-		// find the profile (note that if the name is bad, this will create one with that name)
-		UDeviceProfile* NewActiveProfile = FindProfile(CVarValue);
-		SetOverrideDeviceProfile(NewActiveProfile);
+		UDeviceProfile* NewActiveProfile = FindProfile(CVarValue, false);
+		if (NewActiveProfile)
+		{
+			SetOverrideDeviceProfile(NewActiveProfile);
+		}
 	}
 }
 
@@ -813,8 +989,8 @@ void UDeviceProfileManager::SetActiveDeviceProfile( UDeviceProfile* DeviceProfil
 
 	const void* TextureLODGroupsAddr = ActiveDeviceProfile ? ActiveDeviceProfile->TextureLODGroups.GetData() : nullptr;
 	const int32 NumTextureLODGroups = ActiveDeviceProfile ? ActiveDeviceProfile->TextureLODGroups.Num() : 0;
-	UE_LOG(LogInit, Log, TEXT("Active device profile: [%p][%p %d] %s"), ActiveDeviceProfile, TextureLODGroupsAddr, NumTextureLODGroups, ActiveDeviceProfile ? *ActiveDeviceProfile->GetName() : TEXT("None"));
-	UE_LOG(LogInit, Log, TEXT("Profiles: %s"), *ProfileNames);
+	UE_LOG(LogDeviceProfileManager, Log, TEXT("Active device profile: [%p][%p %d] %s"), ActiveDeviceProfile, TextureLODGroupsAddr, NumTextureLODGroups, ActiveDeviceProfile ? *ActiveDeviceProfile->GetName() : TEXT("None"));
+	UE_LOG(LogDeviceProfileManager, Log, TEXT("Profiles: %s"), *ProfileNames);
 
 #if CSV_PROFILER
 	CSV_METADATA(TEXT("DeviceProfile"), *GetActiveDeviceProfileName());
@@ -858,3 +1034,98 @@ void UDeviceProfileManager::GetAllPossibleParentProfiles(const UDeviceProfile* C
 		}
 	}
 }
+
+
+
+
+
+
+#if WITH_UNREAL_DEVELOPER_TOOLS
+static bool GetCVarForPlatform( FOutputDevice& Ar, FString DPName, FString CVarName)
+{
+	UDeviceProfile* DeviceProfile = UDeviceProfileManager::Get().FindProfile(DPName, false);
+	if (DeviceProfile == nullptr)
+	{
+		Ar.Logf(TEXT("Unable to find device profile %s"), *DPName);
+		return false;
+	}
+
+	FString Value;
+	const FString* DPValue = DeviceProfile->GetAllExpandedCVars().Find(CVarName);
+	if (DPValue != nullptr)
+	{
+		Value = *DPValue;
+	}
+	else
+	{
+		IConsoleVariable* CVar = IConsoleManager::Get().FindConsoleVariable(*CVarName);
+		if (!CVar)
+		{
+			Ar.Logf(TEXT("Unable to find cvar %s"), *CVarName);
+			return false;
+		}
+
+		Value = CVar->GetDefaultValueVariable()->GetString();
+	}
+
+	Ar.Logf(TEXT("%s@%s = \"%s\""), *DPName, *CVarName, *Value);
+
+	return true;
+}
+
+class FPlatformCVarExec : public FSelfRegisteringExec
+{
+public:
+
+	// FSelfRegisteringExec interface
+	virtual bool Exec(UWorld* InWorld, const TCHAR* Cmd, FOutputDevice& Ar) override
+	{
+		if (FParse::Command(&Cmd, TEXT("dpcvar")))
+		{
+			FString DPName, CVarName;
+			if (FString(Cmd).Split(TEXT("@"), &DPName, &CVarName) == false)
+			{
+				return false;
+			}
+
+			return GetCVarForPlatform(Ar, DPName, CVarName);
+		}
+		else if (FParse::Command(&Cmd, TEXT("dpdump")))
+		{
+			UDeviceProfile* DeviceProfile = UDeviceProfileManager::Get().FindProfile(Cmd, false);
+			if (DeviceProfile)
+			{
+				Ar.Logf(TEXT("All cvars found for deviceprofile %s"), Cmd);
+				for (const auto& Pair : DeviceProfile->GetAllExpandedCVars())
+				{
+					Ar.Logf(TEXT("%s = %s"), *Pair.Key, *Pair.Value);
+				}
+			}
+		}
+		else if (FParse::Command(&Cmd, TEXT("dppreview")))
+		{
+			UDeviceProfile* DeviceProfile = UDeviceProfileManager::Get().FindProfile(Cmd, false);
+			if (DeviceProfile)
+			{
+				UDeviceProfileManager::Get().SetOverrideDeviceProfile(DeviceProfile, true);
+			}
+		}
+		else if (FParse::Command(&Cmd, TEXT("dprestore")))
+		{
+			UDeviceProfileManager::Get().RestoreDefaultDeviceProfile();
+		}
+		else if (FParse::Command(&Cmd, TEXT("dpreload")))
+		{
+			FConfigCacheIni::ClearOtherPlatformConfigs();
+			// @todo ini clear out all DPs AllExpandedCVars
+		}
+
+
+
+		return false;
+	}
+
+} GPlatformCVarExec;
+
+
+#endif
\ No newline at end of file
