World Partition
- Actor descriptors are now representing what's on disk, even for loaded actors; they get only updated on save.
- Actor clustering logic was moved directly in the runtime spatial hash implementation file as it is only used there.
- Rewrote world partition editor selected cells logic to avoid keeping pointers to edit cells.

diff --git a/Engine/Source/Editor/UnrealEd/Private/ActorGroupingUtils.cpp b/Engine/Source/Editor/UnrealEd/Private/ActorGroupingUtils.cpp
--- a/Engine/Source/Editor/UnrealEd/Private/ActorGroupingUtils.cpp
+++ b/Engine/Source/Editor/UnrealEd/Private/ActorGroupingUtils.cpp
@@ -8,7 +8,6 @@
 #include "Misc/MessageDialog.h"
 #include "Framework/Notifications/NotificationManager.h"
 #include "Widgets/Notifications/SNotificationList.h"
-#include "WorldPartition/WorldPartitionSubsystem.h"
 
 bool UActorGroupingUtils::bGroupingActive = true;
 
@@ -107,11 +106,6 @@ void UActorGroupingUtils::GroupActors(const TArray<AActor*>& ActorsToGroup)
 					SpawnedGroupActor->SetFolderPath(FolderPath);
 					SpawnedGroupActor->CenterGroupLocation();
 					SpawnedGroupActor->Lock();
-
-					if (UWorldPartitionSubsystem* WorldPartitionSubsystem = World->GetSubsystem<UWorldPartitionSubsystem>())
-					{
-						WorldPartitionSubsystem->UpdateActorDesc(SpawnedGroupActor);
-					}
 				}
 			}
 		}
@@ -173,11 +167,6 @@ void UActorGroupingUtils::UngroupActors(const TArray<AActor*>& ActorsToUngroup)
 			{
 				AGroupActor* GroupActor = OutermostGroupActors[GroupIndex];
 				GroupActor->ClearAndRemove();
-
-				if (UWorldPartitionSubsystem* WorldPartitionSubsystem = GroupActor->GetWorld()->GetSubsystem<UWorldPartitionSubsystem>())
-				{
-					WorldPartitionSubsystem->UpdateActorDesc(GroupActor);
-				}
 			}
 		}
 	}
diff --git a/Engine/Source/Editor/WorldPartitionEditor/Private/WorldPartition/SWorldPartitionEditorGrid2D.cpp b/Engine/Source/Editor/WorldPartitionEditor/Private/WorldPartition/SWorldPartitionEditorGrid2D.cpp
--- a/Engine/Source/Editor/WorldPartitionEditor/Private/WorldPartition/SWorldPartitionEditorGrid2D.cpp
+++ b/Engine/Source/Editor/WorldPartitionEditor/Private/WorldPartition/SWorldPartitionEditorGrid2D.cpp
@@ -46,8 +46,6 @@ SWorldPartitionEditorGrid2D::SWorldPartitionEditorGrid2D()
 	, bIsSelecting(false)
 	, bIsDragging(false)
 	, bShowActors(false)
-	, bShowActorClusters(false)
-	, bShowGroupActors(false)
 {
 	FEditorCommands::Register();
 	
@@ -103,42 +101,6 @@ void SWorldPartitionEditorGrid2D::Construct(const FArguments& InArgs)
 						.IsEnabled(true)
 						.Text(LOCTEXT("ShowActors", "Show Actors"))
 					]
-							
-					+SHorizontalBox::Slot()
-					.AutoWidth()
-					[
-						SNew(SCheckBox)
-						.IsChecked_Lambda([=]() { return bShowActorClusters ? ECheckBoxState::Checked : ECheckBoxState::Unchecked; })
-						.IsEnabled(true)
-						.OnCheckStateChanged(FOnCheckStateChanged::CreateLambda([=](ECheckBoxState State) { bShowActorClusters = !bShowActorClusters; }))
-					]
-					+SHorizontalBox::Slot()
-					.FillWidth(1.0f)
-					.VAlign(VAlign_Center)
-					[
-						SNew(STextBlock)
-						.AutoWrapText(true)
-						.IsEnabled(true)
-						.Text(LOCTEXT("ShowActorClusters", "Show Actor Clusters"))
-					]
-
-					+SHorizontalBox::Slot()
-					.AutoWidth()
-					[
-						SNew(SCheckBox)
-						.IsChecked_Lambda([=]() { return bShowGroupActors ? ECheckBoxState::Checked : ECheckBoxState::Unchecked; })
-						.IsEnabled(true)
-						.OnCheckStateChanged(FOnCheckStateChanged::CreateLambda([=](ECheckBoxState State) { bShowGroupActors = !bShowGroupActors; }))
-					]
-					+SHorizontalBox::Slot()
-					.FillWidth(1.0f)
-					.VAlign(VAlign_Center)
-					[
-						SNew(STextBlock)
-						.AutoWrapText(true)
-						.IsEnabled(true)
-						.Text(LOCTEXT("ShowGroupActors", "Show Group Actors"))
-					]
 				]
 			]
 		]
@@ -152,7 +114,7 @@ void SWorldPartitionEditorGrid2D::Construct(const FArguments& InArgs)
 
 	auto CanLoadOrUnloadCells = [this]()
 	{
-		return GLevelEditorModeTools().IsDefaultModeActive() && SelectedCells.Num() > 0;
+		return GLevelEditorModeTools().IsDefaultModeActive() && SelectBox.GetVolume() > 0;
 	};
 
 	ActionList.MapAction(Commands.LoadSelectedCells, FExecuteAction::CreateSP(this, &SWorldPartitionEditorGrid2D::LoadSelectedCells), FCanExecuteAction::CreateLambda(CanLoadOrUnloadCells));
@@ -162,14 +124,14 @@ void SWorldPartitionEditorGrid2D::Construct(const FArguments& InArgs)
 
 void SWorldPartitionEditorGrid2D::LoadSelectedCells()
 {
-	WorldPartition->LoadEditorCells(SelectedCells.Array());
+	WorldPartition->LoadEditorCells(SelectBox);
 	GEditor->RedrawLevelEditingViewports();
 	RefreshSceneOutliner();
 }
 
 void SWorldPartitionEditorGrid2D::UnloadSelectedCells()
 {
-	WorldPartition->UnloadEditorCells(SelectedCells.Array());
+	WorldPartition->UnloadEditorCells(SelectBox);
 	GEditor->RedrawLevelEditingViewports();
 	RefreshSceneOutliner();
 }
@@ -213,7 +175,7 @@ FReply SWorldPartitionEditorGrid2D::OnMouseButtonDown(const FGeometry& MyGeometr
 			bIsSelecting = true;
 			SelectionStart = MouseCursorPosWorld;
 			SelectionEnd = SelectionStart;
-			UpdateSelection(FSlateApplication::Get().GetModifierKeys().AreModifersDown(EModifierKey::Control));
+			UpdateSelection();
 		}
 
 		return ReplyState;
@@ -294,7 +256,7 @@ FReply SWorldPartitionEditorGrid2D::OnMouseMove(const FGeometry& MyGeometry, con
 			if (bIsSelecting)
 			{
 				SelectionEnd = MouseCursorPosWorld;
-				UpdateSelection(FSlateApplication::Get().GetModifierKeys().AreModifersDown(EModifierKey::Control));
+				UpdateSelection();
 			}
 		}
 
@@ -393,50 +355,6 @@ uint32 SWorldPartitionEditorGrid2D::PaintActors(const FGeometry& AllottedGeometr
 		}
 	}
 
-	// Include actor clusters if requested
-	static TArray<UWorldPartition::FActorCluster> ActorClusters;
-	TArray<const UWorldPartition::FActorCluster*> ActorClustersToShow;
-
-	if (bShowActorClusters)
-	{
-		for (const UWorldPartition::FActorCluster* ActorCluster : WorldPartition->GetActorClusters())
-		{
-			if (bShowActors)
-			{
-				ActorClustersToShow.Add(ActorCluster);
-			}
-			else
-			{
-				bool bShowClusterActors = false;
-
-				for (const FGuid& ActorGuid: ActorCluster->Actors)
-				{
-					FWorldPartitionActorDesc* ActorDesc = WorldPartition->GetActorDesc(ActorGuid);
-					check(ActorDesc);
-
-					if (ActorDesc->Tag == FWorldPartitionActorDesc::GlobalTag)
-					{
-						ActorClustersToShow.Add(ActorCluster);
-						bShowClusterActors = true;
-						break;
-					}
-				}
-
-				if (bShowClusterActors)
-				{
-					for (const FGuid& ActorGuid: ActorCluster->Actors)
-					{
-						FWorldPartitionActorDesc* ActorDesc = WorldPartition->GetActorDesc(ActorGuid);
-						check(ActorDesc);
-
-						ActorDescList.Add(ActorDesc);
-						ActorDesc->Tag = FWorldPartitionActorDesc::GlobalTag;
-					}
-				}
-			}
-		}
-	}
-
 	if (ActorDescList.Num())
 	{
 		const FLinearColor LoadedActorColor(0.75f, 0.75f, 0.75f, 1.0f);
@@ -446,55 +364,6 @@ uint32 SWorldPartitionEditorGrid2D::PaintActors(const FGeometry& AllottedGeometr
 		TArray<FVector2D> LinePoints;
 		LinePoints.SetNum(5);
 
-		TSet<FGuid> SelectedActors;
-		for (UWorldPartitionEditorCell* SelectedCell: SelectedCells)
-		{
-			for (const FWorldPartitionActorDesc* ActorDesc: SelectedCell->Actors)
-			{
-				SelectedActors.Add(ActorDesc->GetGuid());
-			}
-		}
-
-		for (const UWorldPartition::FActorCluster* ActorCluster : ActorClustersToShow)
-		{
-			if ((ActorCluster->Actors.Num() > 1) && (ActorCluster->GridPlacement != EActorGridPlacement::AlwaysLoaded))
-			{
-				FVector Origin, Extent;
-				ActorCluster->Bounds.GetCenterAndExtents(Origin, Extent);
-
-				FVector2D TopLeftW = FVector2D(Origin - Extent);
-				FVector2D BottomRightW = FVector2D(Origin + Extent);
-				FVector2D TopRightW = FVector2D(BottomRightW.X, TopLeftW.Y);
-				FVector2D BottomLeftW = FVector2D(TopLeftW.X, BottomRightW.Y);
-
-				FVector2D TopLeft = WorldToScreen.TransformPoint(TopLeftW);
-				FVector2D BottomRight = WorldToScreen.TransformPoint(BottomRightW);
-				FVector2D TopRight = WorldToScreen.TransformPoint(TopRightW);
-				FVector2D BottomLeft = WorldToScreen.TransformPoint(BottomLeftW);
-
-				FBox2D ActorViewBox(TopLeft, BottomRight);
-
-				const float MinimumAreaCull = 32.0f;
-				const float AreaFadeDistance = 128.0f;
-				if (ActorViewBox.Intersect(ViewRect) && (ActorViewBox.GetArea() > MinimumAreaCull))
-				{
-					FPaintGeometry ActorGeometry = AllottedGeometry.ToPaintGeometry(TopLeft, BottomRight - TopLeft);
-					float ActorColorGradient = FMath::Min((ActorViewBox.GetArea() - MinimumAreaCull) / AreaFadeDistance, 1.0f);
-					FLinearColor ActorColor(0.0f, 1.0f, 0.0f, ActorColorGradient * 0.1f);
-					FSlateColorBrush ClusterBrush(FLinearColor::Green);
-
-					FSlateDrawElement::MakeBox(
-						OutDrawElements,
-						++LayerId,
-						ActorGeometry,
-						&ClusterBrush,
-						ESlateDrawEffect::None,
-						ActorColor
-					);
-				}
-			}
-		}
-
 		for (FWorldPartitionActorDesc* ActorDesc: ActorDescList)
 		{
 			const FBox ActorBounds = ActorDesc->GetBounds();
@@ -523,19 +392,6 @@ uint32 SWorldPartitionEditorGrid2D::PaintActors(const FGeometry& AllottedGeometr
 				FLinearColor ActorColor(ActorBrightness, ActorBrightness, ActorBrightness, ActorColorGradient);
 
 				UClass* ActorClass = ActorDesc->GetActorClass();
-				if (bShowGroupActors && ActorClass && ActorClass->IsChildOf<AGroupActor>())
-				{
-					FSlateColorBrush ClusterBrush(FLinearColor::Blue);
-
-					FSlateDrawElement::MakeBox(
-						OutDrawElements,
-						++LayerId,
-						ActorGeometry,
-						&ClusterBrush,
-						ESlateDrawEffect::None,
-						ActorColor * FLinearColor(0.0f, 0.0f, 1.0f, 0.1f)
-					);
-				}
 
 				const float SquaredDistanceToPoint = ActorViewBox.ComputeSquaredDistanceToPoint(MouseCursorPos);
 				if ((ActorDesc->Tag == FWorldPartitionActorDesc::GlobalTag) || (SquaredDistanceToPoint > 0.0f && SquaredDistanceToPoint <= 2.0f))
@@ -561,11 +417,16 @@ uint32 SWorldPartitionEditorGrid2D::PaintActors(const FGeometry& AllottedGeometr
 					);
 				}
 
+				if (SelectBox.Intersect(ActorDesc->GetBounds()))
+				{
+					ActorColor = FLinearColor::White;
+				}
+
 				FSlateDrawElement::MakeBox(
 					OutDrawElements,
 					++LayerId,
 					ActorGeometry,
-					SelectedActors.Find(ActorDesc->GetGuid()) ? FEditorStyle::GetBrush(TEXT("FocusRectangle")) : FEditorStyle::GetBrush(TEXT("Border")),
+					FEditorStyle::GetBrush(TEXT("Border")),
 					ESlateDrawEffect::None,
 					ActorColor
 				);
@@ -788,27 +649,19 @@ void SWorldPartitionEditorGrid2D::UpdateTransform() const
 	ScreenToWorld = WorldToScreen.Inverse();
 }
 
-void SWorldPartitionEditorGrid2D::UpdateSelection(bool bAppend)
+void SWorldPartitionEditorGrid2D::UpdateSelection()
 {
-	if (!bAppend)
-	{
-		SelectedCells.Empty();
-	}
+	SelectBox.Init();
 
-	const FBox SelectBox(
+	const FBox SelectionBox(
 		FVector(FMath::Min(SelectionStart.X, SelectionEnd.X), FMath::Min(SelectionStart.Y, SelectionEnd.Y), -WORLD_MAX),
 		FVector(FMath::Max(SelectionStart.X, SelectionEnd.X), FMath::Max(SelectionStart.Y, SelectionEnd.Y), WORLD_MAX)
 	);
 
-	const FBox2D SelectBox2D(FVector2D(SelectBox.Min), FVector2D(SelectBox.Max));
+	const FBox2D SelectBox2D(FVector2D(SelectionBox.Min), FVector2D(SelectionBox.Max));
 	if(!FMath::IsNearlyZero(SelectBox2D.GetArea()))
 	{
-		WorldPartition->EditorHash->ForEachIntersectingCell(SelectBox, [&](UWorldPartitionEditorCell* SelectedCell)
-		{
-			bool bWasAlreadyInSet;
-			SelectedCells.Add(SelectedCell, &bWasAlreadyInSet);
-			check(!bWasAlreadyInSet || bAppend);
-		});
+		SelectBox = SelectionBox;
 	}
 }
 
diff --git a/Engine/Source/Editor/WorldPartitionEditor/Private/WorldPartition/SWorldPartitionEditorGridSpatialHash.cpp b/Engine/Source/Editor/WorldPartitionEditor/Private/WorldPartition/SWorldPartitionEditorGridSpatialHash.cpp
--- a/Engine/Source/Editor/WorldPartitionEditor/Private/WorldPartition/SWorldPartitionEditorGridSpatialHash.cpp
+++ b/Engine/Source/Editor/WorldPartitionEditor/Private/WorldPartition/SWorldPartitionEditorGridSpatialHash.cpp
@@ -180,12 +180,14 @@ int32 SWorldPartitionEditorGridSpatialHash::PaintGrid(const FGeometry& AllottedG
 	}
 
 	// Draw selected cells
-	if (SelectedCells.Num())
+	if (SelectBox.GetVolume() > 0)
 	{
-		TMap<uint32, FCellDesc2D> UniqueCells2D;
-		WorldPartition->EditorHash->ForEachIntersectingCell(VisibleGridRectWorld, [&](UWorldPartitionEditorCell* Cell)
+		const FBox VisibleSelectBox = SelectBox.Overlap(VisibleGridRectWorld);
+
+		if (VisibleSelectBox.GetVolume() > 0)
 		{
-			if (SelectedCells.Contains(Cell))
+			TMap<uint32, FCellDesc2D> UniqueCells2D;
+			WorldPartition->EditorHash->ForEachIntersectingCell(VisibleSelectBox, [&](UWorldPartitionEditorCell* Cell)
 			{
 				UWorldPartitionEditorSpatialHash::FCellCoord CellCoord = EditorSpatialHash->GetCellCoords(Cell->Bounds.GetCenter(), 0);
 
@@ -196,29 +198,29 @@ int32 SWorldPartitionEditorGridSpatialHash::PaintGrid(const FGeometry& AllottedG
 				FCellDesc2D& CellDesc2D = UniqueCells2D.Add(CellHash2D);
 
 				CellDesc2D.Bounds = FBox2D(FVector2D(Cell->Bounds.Min), FVector2D(Cell->Bounds.Max));
-			}
-		});
-
-		for(auto& UniqueCell: UniqueCells2D)
-		{
-			const FCellDesc2D& Cell = UniqueCell.Value;
-
-			FPaintGeometry CellGeometry = AllottedGeometry.ToPaintGeometry(
-				WorldToScreen.TransformPoint(Cell.Bounds.Min),
-				WorldToScreen.TransformPoint(Cell.Bounds.Max) - WorldToScreen.TransformPoint(Cell.Bounds.Min)
-			);
+			});
 
-			FSlateColorBrush CellBrush(FLinearColor::White);
-			FLinearColor CellColor(1, 1, 1, 0.25f);
-
-			FSlateDrawElement::MakeBox(
-				OutDrawElements,
-				++LayerId,
-				CellGeometry,
-				&CellBrush,
-				ESlateDrawEffect::None,
-				CellColor
-			);
+			for(auto& UniqueCell: UniqueCells2D)
+			{
+				const FCellDesc2D& Cell = UniqueCell.Value;
+
+				FPaintGeometry CellGeometry = AllottedGeometry.ToPaintGeometry(
+					WorldToScreen.TransformPoint(Cell.Bounds.Min),
+					WorldToScreen.TransformPoint(Cell.Bounds.Max) - WorldToScreen.TransformPoint(Cell.Bounds.Min)
+				);
+
+				FSlateColorBrush CellBrush(FLinearColor::White);
+				FLinearColor CellColor(1, 1, 1, 0.25f);
+
+				FSlateDrawElement::MakeBox(
+					OutDrawElements,
+					++LayerId,
+					CellGeometry,
+					&CellBrush,
+					ESlateDrawEffect::None,
+					CellColor
+				);
+			}
 		}
 	}
 
diff --git a/Engine/Source/Editor/WorldPartitionEditor/Public/WorldPartition/SWorldPartitionEditorGrid2D.h b/Engine/Source/Editor/WorldPartitionEditor/Public/WorldPartition/SWorldPartitionEditorGrid2D.h
--- a/Engine/Source/Editor/WorldPartitionEditor/Public/WorldPartition/SWorldPartitionEditorGrid2D.h
+++ b/Engine/Source/Editor/WorldPartitionEditor/Public/WorldPartition/SWorldPartitionEditorGrid2D.h
@@ -57,7 +57,7 @@ public:
 
 protected:
 	void UpdateTransform() const;
-	void UpdateSelection(bool bAppend = false);
+	void UpdateSelection();
 
 	const TSharedRef<FUICommandList> CommandList;
 
@@ -89,6 +89,6 @@ protected:
 	FVector2D MouseCursorPosWorld;
 	FVector2D SelectionStart;
 	FVector2D SelectionEnd;
-	TSet<UWorldPartitionEditorCell*> SelectedCells;
+	FBox SelectBox;
 	FSlateFontInfo SmallLayoutFont;
 };
\ No newline at end of file
diff --git a/Engine/Source/Runtime/Engine/Private/ActorEditor.cpp b/Engine/Source/Runtime/Engine/Private/ActorEditor.cpp
--- a/Engine/Source/Runtime/Engine/Private/ActorEditor.cpp
+++ b/Engine/Source/Runtime/Engine/Private/ActorEditor.cpp
@@ -591,14 +591,6 @@ bool AActor::InternalPostEditUndo()
 		FNavigationSystem::RemoveActorData(*this);
 	}
 
-	if (!IsTemplate())
-	{
-		if (UWorldPartitionSubsystem* WorldPartitionSubsystem = UWorld::GetSubsystem<UWorldPartitionSubsystem>(GetWorld()))
-		{
-			WorldPartitionSubsystem->UpdateActorDesc(this);
-		}
-	}
-
 	// This is a normal undo, so call super
 	return true;
 }
diff --git a/Engine/Source/Runtime/Engine/Private/ActorPartition/ActorPartitionSubsystem.cpp b/Engine/Source/Runtime/Engine/Private/ActorPartition/ActorPartitionSubsystem.cpp
--- a/Engine/Source/Runtime/Engine/Private/ActorPartition/ActorPartitionSubsystem.cpp
+++ b/Engine/Source/Runtime/Engine/Private/ActorPartition/ActorPartitionSubsystem.cpp
@@ -227,9 +227,6 @@ public:
 			FoundActor->bLockLocation = true;
 			
 			InActorCreated(FoundActor);
-
-			FoundActor->SetActorLabel(FoundActor->GetName());
-			WorldPartition->UpdateActorDesc(FoundActor);
 		}
 
 		check(FoundActor || !bInCreate);
diff --git a/Engine/Source/Runtime/Engine/Private/LevelInstance/LevelInstanceActor.cpp b/Engine/Source/Runtime/Engine/Private/LevelInstance/LevelInstanceActor.cpp
--- a/Engine/Source/Runtime/Engine/Private/LevelInstance/LevelInstanceActor.cpp
+++ b/Engine/Source/Runtime/Engine/Private/LevelInstance/LevelInstanceActor.cpp
@@ -4,8 +4,6 @@
 #include "LevelInstance/LevelInstanceSubsystem.h"
 #include "LevelInstance/LevelInstanceComponent.h"
 #include "LevelInstance/LevelInstanceEditorInstanceActor.h"
-#include "WorldPartition/WorldPartitionSubsystem.h"
-#include "WorldPartition/WorldPartition.h"
 #include "Engine/World.h"
 #include "Engine/Engine.h"
 #include "LevelInstancePrivate.h"
@@ -533,11 +531,6 @@ void ALevelInstance::OnLevelInstanceLoaded()
 		{
 			SetIsTemporarilyHiddenInEditor(bHiddenInEditor);
 		}
-
-		if (UWorldPartitionSubsystem* WorldPartitionSubsystem = GetWorld()->GetSubsystem<UWorldPartitionSubsystem>())
-		{
-			WorldPartitionSubsystem->UpdateActorDesc(this);
-		}
 	}
 }
 
diff --git a/Engine/Source/Runtime/Engine/Private/WorldPartition/ActorPartition/PartitionActorDesc.cpp b/Engine/Source/Runtime/Engine/Private/WorldPartition/ActorPartition/PartitionActorDesc.cpp
--- a/Engine/Source/Runtime/Engine/Private/WorldPartition/ActorPartition/PartitionActorDesc.cpp
+++ b/Engine/Source/Runtime/Engine/Private/WorldPartition/ActorPartition/PartitionActorDesc.cpp
@@ -7,9 +7,9 @@
 #include "Engine/Level.h"
 #include "UObject/UE5MainStreamObjectVersion.h"
 
-void FPartitionActorDesc::InitFrom(const AActor* InActor)
+void FPartitionActorDesc::Init(const AActor* InActor)
 {
-	FWorldPartitionActorDesc::InitFrom(InActor);
+	FWorldPartitionActorDesc::Init(InActor);
 
 	const APartitionActor* PartitionActor = CastChecked<APartitionActor>(InActor);
 
diff --git a/Engine/Source/Runtime/Engine/Private/WorldPartition/HLOD/HLODActorDesc.cpp b/Engine/Source/Runtime/Engine/Private/WorldPartition/HLOD/HLODActorDesc.cpp
--- a/Engine/Source/Runtime/Engine/Private/WorldPartition/HLOD/HLODActorDesc.cpp
+++ b/Engine/Source/Runtime/Engine/Private/WorldPartition/HLOD/HLODActorDesc.cpp
@@ -8,9 +8,9 @@
 #include "WorldPartition/HLOD/HLODActor.h"
 #include "WorldPartition/HLOD/HLODLayer.h"
 
-void FHLODActorDesc::InitFrom(const AActor* InActor)
+void FHLODActorDesc::Init(const AActor* InActor)
 {
-	FWorldPartitionActorDesc::InitFrom(InActor);
+	FWorldPartitionActorDesc::Init(InActor);
 
 	const AWorldPartitionHLOD* HLODActor = CastChecked<AWorldPartitionHLOD>(InActor);
 	SubActors = HLODActor->GetSubActors();
diff --git a/Engine/Source/Runtime/Engine/Private/WorldPartition/HLOD/HLODBuilder.cpp b/Engine/Source/Runtime/Engine/Private/WorldPartition/HLOD/HLODBuilder.cpp
--- a/Engine/Source/Runtime/Engine/Private/WorldPartition/HLOD/HLODBuilder.cpp
+++ b/Engine/Source/Runtime/Engine/Private/WorldPartition/HLOD/HLODBuilder.cpp
@@ -56,8 +56,6 @@ public:
 			HLODActor->SetChildrenPrimitives(InSubComponents);
 			HLODActor->SetActorLabel(HLODActorName);
 			HLODActor->RuntimeGrid = HLODLayer->GetRuntimeGrid();
-
-			WorldPartition->UpdateActorDesc(HLODActor);
 		}
 		else
 		{
diff --git a/Engine/Source/Runtime/Engine/Private/WorldPartition/Landscape/LandscapeActorDesc.cpp b/Engine/Source/Runtime/Engine/Private/WorldPartition/Landscape/LandscapeActorDesc.cpp
--- a/Engine/Source/Runtime/Engine/Private/WorldPartition/Landscape/LandscapeActorDesc.cpp
+++ b/Engine/Source/Runtime/Engine/Private/WorldPartition/Landscape/LandscapeActorDesc.cpp
@@ -5,9 +5,9 @@
 
 #if WITH_EDITOR
 
-void FLandscapeActorDesc::InitFrom(const AActor* InActor)
+void FLandscapeActorDesc::Init(const AActor* InActor)
 {
-	FPartitionActorDesc::InitFrom(InActor);
+	FPartitionActorDesc::Init(InActor);
 
 	const ALandscapeProxy* LandscapeProxy = CastChecked<ALandscapeProxy>(InActor);
 	check(LandscapeProxy);
diff --git a/Engine/Source/Runtime/Engine/Private/WorldPartition/LevelInstance/LevelInstanceActorDesc.cpp b/Engine/Source/Runtime/Engine/Private/WorldPartition/LevelInstance/LevelInstanceActorDesc.cpp
--- a/Engine/Source/Runtime/Engine/Private/WorldPartition/LevelInstance/LevelInstanceActorDesc.cpp
+++ b/Engine/Source/Runtime/Engine/Private/WorldPartition/LevelInstance/LevelInstanceActorDesc.cpp
@@ -6,9 +6,9 @@
 #include "Engine/Level.h"
 #include "LevelInstance/LevelInstanceActor.h"
 
-void FLevelInstanceActorDesc::InitFrom(const AActor* InActor)
+void FLevelInstanceActorDesc::Init(const AActor* InActor)
 {
-	FWorldPartitionActorDesc::InitFrom(InActor);
+	FWorldPartitionActorDesc::Init(InActor);
 
 	const ALevelInstance* LevelInstanceActor = CastChecked<ALevelInstance>(InActor);
 	LevelPackage = *LevelInstanceActor->GetWorldAssetPackage();
diff --git a/Engine/Source/Runtime/Engine/Private/WorldPartition/WorldPartition.cpp b/Engine/Source/Runtime/Engine/Private/WorldPartition/WorldPartition.cpp
--- a/Engine/Source/Runtime/Engine/Private/WorldPartition/WorldPartition.cpp
+++ b/Engine/Source/Runtime/Engine/Private/WorldPartition/WorldPartition.cpp
@@ -112,7 +112,7 @@ UWorldPartition::UWorldPartition(const FObjectInitializer& ObjectInitializer)
 #if WITH_EDITOR
 	, EditorHash(nullptr)
 	, WorldPartitionEditor(nullptr)
-	, bDirty(false)
+	, bIgnoreAssetRegistryEvents(false)
 	, bForceGarbageCollection(false)
 	, bForceGarbageCollectionPurge(false)
 #endif
@@ -148,17 +148,6 @@ FName UWorldPartition::GetWorldPartitionEditorName()
 {
 	return EditorHash->GetWorldPartitionEditorName();
 }
-
-void UWorldPartition::OnObjectPropertyChanged(UObject* Object, FPropertyChangedEvent& PropertyChangedEvent)
-{
-	if (AActor* Actor = Cast<AActor>(Object))
-	{
-		if (FWorldPartitionActorDesc* ActorDesc = GetActorDesc(Actor->GetActorGuid()))
-		{
-			UpdateActorDesc(Actor);
-		}
-	}
-}
 #endif
 
 void UWorldPartition::Initialize(UWorld* InWorld, const FTransform& InTransform)
@@ -211,29 +200,27 @@ void UWorldPartition::Initialize(UWorld* InWorld, const FTransform& InTransform)
 	{
 		TSet<FName> AllLayersNames;
 
-		auto GetLevelActors = [](const FName& LevelPath, TArray<FAssetData>& OutAssets)
+		TArray<FAssetData> Assets;
+		UPackage* LevelPackage = OuterWorld->PersistentLevel->GetOutermost();
+
+		if (!LevelPackage->FileName.IsNone())
 		{
-			if (!LevelPath.IsNone())
-			{
-				const FString LevelPathStr = LevelPath.ToString();
-				const FString LevelExternalActorsPath = ULevel::GetExternalActorsPath(LevelPathStr);
+			TGuardValue<bool> IgnoreAssetRegistryEvents(bIgnoreAssetRegistryEvents, true);
 
-				// Do a synchronous scan of the level external actors path.
-				IAssetRegistry& AssetRegistry = FModuleManager::LoadModuleChecked<FAssetRegistryModule>(TEXT("AssetRegistry")).Get();
-				AssetRegistry.ScanPathsSynchronous({LevelExternalActorsPath}, /*bForceRescan*/true, /*bIgnoreBlackListScanFilters*/true);
+			const FString LevelPathStr = LevelPackage->FileName.ToString();
+			const FString LevelExternalActorsPath = ULevel::GetExternalActorsPath(LevelPathStr);
 
-				FARFilter Filter;
-				Filter.bRecursivePaths = true;
-				Filter.bIncludeOnlyOnDiskAssets = true;
-				Filter.PackagePaths.Add(*LevelExternalActorsPath);
+			// Do a synchronous scan of the level external actors path.			
+			IAssetRegistry& AssetRegistry = FModuleManager::LoadModuleChecked<FAssetRegistryModule>(TEXT("AssetRegistry")).Get();
+			AssetRegistry.ScanPathsSynchronous({LevelExternalActorsPath}, /*bForceRescan*/true, /*bIgnoreBlackListScanFilters*/true);
 
-				AssetRegistry.GetAssets(Filter, OutAssets);
-			}
-		};
+			FARFilter Filter;
+			Filter.bRecursivePaths = true;
+			Filter.bIncludeOnlyOnDiskAssets = true;
+			Filter.PackagePaths.Add(*LevelExternalActorsPath);
 
-		TArray<FAssetData> Assets;
-		UPackage* LevelPackage = OuterWorld->PersistentLevel->GetOutermost();
-		GetLevelActors(LevelPackage->FileName, Assets);
+			AssetRegistry.GetAssets(Filter, Assets);
+		}
 
 		bool bIsInstanced = (bEditorOnly && !IsRunningCommandlet()) ? OuterWorld->PersistentLevel->IsInstancedLevel() : false;
 
@@ -253,35 +240,19 @@ void UWorldPartition::Initialize(UWorld* InWorld, const FTransform& InTransform)
 
 		for (const FAssetData& Asset : Assets)
 		{
-			FString ActorClass;
-			static FName NAME_ActorMetaDataClass(TEXT("ActorMetaDataClass"));
-			Asset.GetTagValue(NAME_ActorMetaDataClass, ActorClass);
-
-			FString ActorMetaDataStr;
-			static FName NAME_ActorMetaData(TEXT("ActorMetaData"));
-			Asset.GetTagValue(NAME_ActorMetaData, ActorMetaDataStr);
-
-			FWorldPartitionActorDescInitData ActorDescInitData;
-			ActorDescInitData.NativeClass = FindObjectChecked<UClass>(ANY_PACKAGE, *ActorClass, true);
-			ActorDescInitData.PackageName = Asset.PackageName;
-			ActorDescInitData.ActorPath = Asset.ObjectPath;
-			ActorDescInitData.Transform = bIsInstanced ? InstanceTransform : FTransform::Identity;
-			FBase64::Decode(ActorMetaDataStr, ActorDescInitData.SerializedData);			
+			TUniquePtr<FWorldPartitionActorDesc> NewActorDesc = GetActorDescriptor(Asset);
 
 			if (bIsInstanced)
 			{
-				const FString LongActorPackageName = ActorDescInitData.PackageName.ToString();
+				const FString LongActorPackageName = Asset.PackageName.ToString();
 				const FString ActorPackageName = FPaths::GetBaseFilename(LongActorPackageName);
 				const FString InstancedName = FString::Printf(TEXT("%s_InstanceOf_%s"), *LevelPackage->GetName(), *ActorPackageName);
 
 				InstancingContext.AddMapping(*LongActorPackageName, *InstancedName);
 
-				ActorDescInitData.ActorPath = *ActorDescInitData.ActorPath.ToString().Replace(*ReplaceFrom, *ReplaceTo);
+				NewActorDesc->TransformInstance(ReplaceFrom, ReplaceTo, InstanceTransform);
 			}
 
-			TUniquePtr<FWorldPartitionActorDesc> NewActorDesc(GetActorDescFactory(ActorDescInitData.NativeClass)->Create());
-			NewActorDesc->Init(ActorDescInitData);
-
 			if (bEditorOnly)
 			{
 				AllLayersNames.Append(NewActorDesc->GetLayers());
@@ -294,7 +265,7 @@ void UWorldPartition::Initialize(UWorld* InWorld, const FTransform& InTransform)
 		{
 			for (const auto& Pair : Actors)
 			{
-				AddToPartition(Pair.Value.Get());
+				HashActorDesc(Pair.Value.Get());
 			}
 
 			CreateLayers(AllLayersNames);
@@ -381,10 +352,6 @@ void UWorldPartition::Uninitialize()
 		}
 
 		EditorHash = nullptr;
-
-		for (FActorCluster* ActorCluster : ActorClustersSet) { delete ActorCluster; }
-		ActorClustersSet.Empty();
-		ActorToActorCluster.Empty();
 #endif		
 
 		UWorldPartitionSubsystem* WorldPartitionSubsystem = World->GetSubsystem<UWorldPartitionSubsystem>();
@@ -415,13 +382,10 @@ void UWorldPartition::RegisterDelegates()
 {
 	if (GEditor && !IsTemplate())
 	{
-		GEditor->OnActorMoving().AddUObject(this, &UWorldPartition::OnActorMoving);
-		GEditor->OnActorMoved().AddUObject(this, &UWorldPartition::OnActorMoving);
-		GEngine->OnLevelActorOuterChanged().AddUObject(this, &UWorldPartition::OnActorOuterChanged);
-		GEditor->OnLevelActorAdded().AddUObject(this, &UWorldPartition::OnActorAdded);
-		GEditor->OnLevelActorDeleted().AddUObject(this, &UWorldPartition::OnActorDeleted);
-
-		FCoreUObjectDelegates::OnObjectPropertyChanged.AddUObject(this, &UWorldPartition::OnObjectPropertyChanged);
+		IAssetRegistry& AssetRegistry = FModuleManager::LoadModuleChecked<FAssetRegistryModule>(TEXT("AssetRegistry")).Get();
+		AssetRegistry.OnAssetAdded().AddUObject(this, &UWorldPartition::OnAssetAdded);
+		AssetRegistry.OnAssetRemoved().AddUObject(this, &UWorldPartition::OnAssetRemoved);
+		AssetRegistry.OnAssetUpdated().AddUObject(this, &UWorldPartition::OnAssetUpdated);
 
 		FEditorDelegates::PreBeginPIE.AddUObject(this, &UWorldPartition::OnPreBeginPIE);
 		FEditorDelegates::EndPIE.AddUObject(this, &UWorldPartition::OnEndPIE);
@@ -432,47 +396,16 @@ void UWorldPartition::UnregisterDelegates()
 {
 	if (GEditor && !IsTemplate())
 	{
-		FEditorDelegates::PreBeginPIE.RemoveAll(this);
-		FEditorDelegates::EndPIE.RemoveAll(this);
-
-		GEditor->OnActorMoving().RemoveAll(this);
-		GEditor->OnActorMoved().RemoveAll(this);
-		GEngine->OnLevelActorOuterChanged().RemoveAll(this);
-		GEditor->OnLevelActorAdded().RemoveAll(this);
-		GEditor->OnLevelActorDeleted().RemoveAll(this);
-
-		FCoreUObjectDelegates::OnObjectPropertyChanged.RemoveAll(this);
-	}
-}
-
-void UWorldPartition::UpdateActorDesc(AActor* InActor)
-{
-	check(InActor)
-
-	if (!InActor->IsChildActor())
-	{
-		TUniquePtr<FWorldPartitionActorDesc>* ActorDescPtr = Actors.Find(InActor->GetActorGuid());
-
-		if (GIsTransacting && InActor->IsPendingKill())
-		{
-			OnActorDeleted(InActor);
-		}
-		else if (GIsTransacting && !ActorDescPtr)
+		if (FAssetRegistryModule* AssetRegistryModule = FModuleManager::GetModulePtr<FAssetRegistryModule>(TEXT("AssetRegistry")))
 		{
-			OnActorAdded(InActor);
-		}
-		else if (ActorDescPtr && InActor->GetLevel() == World->PersistentLevel)
-		{
-			TUniquePtr<FWorldPartitionActorDesc> NewDesc(GetActorDescFactory(InActor)->Create());
-			NewDesc->Init(InActor);
-			
-			if (NewDesc->GetHash() != ActorDescPtr->Get()->GetHash())
-			{
-				RemoveFromPartition(ActorDescPtr->Get(), false);
-				*ActorDescPtr = MoveTemp(NewDesc);
-				AddToPartition(ActorDescPtr->Get());
-			}
+			IAssetRegistry& AssetRegistry = AssetRegistryModule->Get();
+			AssetRegistry.OnAssetAdded().RemoveAll(this);
+			AssetRegistry.OnAssetRemoved().RemoveAll(this);
+			AssetRegistry.OnAssetUpdated().RemoveAll(this);
 		}
+
+		FEditorDelegates::PreBeginPIE.RemoveAll(this);
+		FEditorDelegates::EndPIE.RemoveAll(this);
 	}
 }
 
@@ -565,142 +498,6 @@ FWorldPartitionActorDesc* UWorldPartition::GetActorDesc(const FGuid& Guid)
 	return ActorDescObj != nullptr ? ActorDescObj->Get() : nullptr;
 }
 
-void UWorldPartition::AddToClusters(const FWorldPartitionActorDesc* ActorDesc)
-{
-	FActorCluster* ActorCluster = ActorToActorCluster.FindRef(ActorDesc->GetGuid());
-	if (!ActorCluster)
-	{
-		ActorCluster = new FActorCluster(ActorDesc);
-		ActorClustersSet.Add(ActorCluster);
-		ActorToActorCluster.Add(ActorDesc->GetGuid(), ActorCluster);
-	}
-
-	// Don't include references from editor-only actors
-	if (!ActorDesc->GetActorIsEditorOnly())
-	{
-		for (const FGuid& ReferenceGuid : ActorDesc->GetReferences())
-		{
-			const FWorldPartitionActorDesc* ReferenceActorDesc = GetActorDesc(ReferenceGuid);
-
-			// Don't include references to editor-only actors
-			if (!ReferenceActorDesc->GetActorIsEditorOnly())
-			{
-				FActorCluster* ReferenceCluster = ActorToActorCluster.FindRef(ReferenceGuid);
-				if (ReferenceCluster)
-				{
-					if (ReferenceCluster != ActorCluster)
-					{
-						// Merge reference cluster in Actor's cluster
-						ActorCluster->Add(*ReferenceCluster);
-						for (const FGuid& ReferenceClusterActorGuid : ReferenceCluster->Actors)
-						{
-							ActorToActorCluster[ReferenceClusterActorGuid] = ActorCluster;
-						}
-						ActorClustersSet.Remove(ReferenceCluster);
-						delete ReferenceCluster;
-					}
-				}
-				else
-				{
-					// Put Reference in Actor's cluster
-					ActorCluster->Add(FActorCluster(ReferenceActorDesc));
-				}
-
-				// Map its cluster
-				ActorToActorCluster.Add(ReferenceGuid, ActorCluster);
-			}
-		}
-	}
-}
-
-void UWorldPartition::RemoveFromClusters(const FWorldPartitionActorDesc* ActorDesc)
-{
-	FActorCluster* ActorCluster = ActorToActorCluster.FindAndRemoveChecked(ActorDesc->GetGuid());
-	ActorCluster->Actors.Remove(ActorDesc->GetGuid());
-
-	// Break up this cluster and reinsert all actors
-	ActorClustersSet.Remove(ActorCluster);
-
-	for (const FGuid& Guid : ActorCluster->Actors)
-	{
-		ActorToActorCluster[Guid] = nullptr;
-	}
-
-	for (const FGuid& Guid : ActorCluster->Actors)
-	{
-		AddToClusters(GetActorDesc(Guid));
-	}
-
-	delete ActorCluster;
-}
-
-UWorldPartition::FActorCluster::FActorCluster(const FWorldPartitionActorDesc* ActorDesc)
-	: GridPlacement(ActorDesc->GetGridPlacement())
-	, RuntimeGrid(ActorDesc->GetRuntimeGrid())
-	, Bounds(ActorDesc->GetBounds())
-{
-	check(GridPlacement != EActorGridPlacement::None);
-	Actors.Add(ActorDesc->GetGuid());
-	DataLayersID = FDataLayersHelper::ComputeDataLayerID(DataLayers);
-}
-
-void UWorldPartition::FActorCluster::Add(const UWorldPartition::FActorCluster& ActorCluster)
-{
-	// Merge Actors
-	Actors.Append(ActorCluster.Actors);
-
-	// Merge RuntimeGrid
-	RuntimeGrid = RuntimeGrid == ActorCluster.RuntimeGrid ? RuntimeGrid : NAME_None;
-
-	// Merge Bounds
-	Bounds += ActorCluster.Bounds;
-
-	// Merge GridPlacement
-	// If currently None, will always stay None
-	if (GridPlacement != EActorGridPlacement::None)
-	{
-		// If grid placement differs between the two clusters
-		if (GridPlacement != ActorCluster.GridPlacement)
-		{
-			// If one of the two cluster was always loaded, set to None
-			if (ActorCluster.GridPlacement == EActorGridPlacement::AlwaysLoaded || GridPlacement == EActorGridPlacement::AlwaysLoaded)
-			{
-				GridPlacement = EActorGridPlacement::None;
-			}
-			else
-			{
-				GridPlacement = ActorCluster.GridPlacement;
-			}
-		}
-
-		// If current placement is set to Location, that won't make sense when merging two clusters. Set to Bounds
-		if (GridPlacement == EActorGridPlacement::Location)
-		{
-			GridPlacement = EActorGridPlacement::Bounds;
-		}
-	}
-
-	// Merge DataLayers
-	if (DataLayersID != ActorCluster.DataLayersID)
-	{
-		for (const FName& DataLayer : ActorCluster.DataLayers)
-		{
-			DataLayers.AddUnique(DataLayer);
-		}
-		DataLayersID = FDataLayersHelper::ComputeDataLayerID(DataLayers);
-	}
-}
-
-const TSet<UWorldPartition::FActorCluster*>& UWorldPartition::GetActorClusters() const
-{
-	return ActorClustersSet;
-}
-
-const UWorldPartition::FActorCluster* UWorldPartition::GetClusterForActor(const FGuid& InActorGuid) const
-{
-	return ActorToActorCluster.FindRef(InActorGuid);
-}
-
 void UWorldPartition::RefreshLoadedCells()
 {
 	FWorldPartionCellUpdateContext CellUpdateContext(this);
@@ -726,7 +523,18 @@ void UWorldPartition::LoadEditorCells(const FBox& Box)
 	TArray<UWorldPartitionEditorCell*> CellsToLoad;
 	if (EditorHash->GetIntersectingCells(Box, CellsToLoad))
 	{
-		LoadEditorCells(CellsToLoad);
+		FWorldPartionCellUpdateContext CellUpdateContext(this);
+
+		int32 NumActorsToLoad = Algo::TransformAccumulate(CellsToLoad, [](UWorldPartitionEditorCell* Cell) { return Cell->Actors.Num() - Cell->LoadedActors.Num();}, 0);
+
+		FScopedSlowTask SlowTask(NumActorsToLoad, LOCTEXT("LoadingCells", "Loading cells..."));
+		SlowTask.MakeDialog();
+
+		for (UWorldPartitionEditorCell* Cell : CellsToLoad)
+		{
+			SlowTask.EnterProgressFrame(Cell->Actors.Num() - Cell->LoadedActors.Num());
+			UpdateLoadingEditorCell(Cell, true);
+		}
 	}
 }
 
@@ -735,109 +543,73 @@ void UWorldPartition::UnloadEditorCells(const FBox& Box)
 	TArray<UWorldPartitionEditorCell*> CellsToUnload;
 	if (EditorHash->GetIntersectingCells(Box, CellsToUnload))
 	{
-		UnloadEditorCells(CellsToUnload);
-	}
-}
-
-void UWorldPartition::LoadEditorCells(const TArray<UWorldPartitionEditorCell*>& CellsToLoad)
-{
-	FWorldPartionCellUpdateContext CellUpdateContext(this);
-
-	int32 NumActorsToLoad = Algo::TransformAccumulate(CellsToLoad, [](UWorldPartitionEditorCell* Cell) { return Cell->Actors.Num() - Cell->LoadedActors.Num();}, 0);
-
-	FScopedSlowTask SlowTask(NumActorsToLoad, LOCTEXT("LoadingCells", "Loading cells..."));
-	SlowTask.MakeDialog();
-
-	for (UWorldPartitionEditorCell* Cell : CellsToLoad)
-	{
-		SlowTask.EnterProgressFrame(Cell->Actors.Num() - Cell->LoadedActors.Num());
-		UpdateLoadingEditorCell(Cell, true);
-	}
-}
-
-void UWorldPartition::UnloadEditorCells(const TArray<UWorldPartitionEditorCell*>& CellsToUnload)
-{
-	FWorldPartionCellUpdateContext CellUpdateContext(this);
+		FWorldPartionCellUpdateContext CellUpdateContext(this);
 		
-	TSet<UPackage*> ModifiedPackages;
-	TMap<FWorldPartitionActorDesc*, int32> UnloadCount;
+		TSet<UPackage*> ModifiedPackages;
+		TMap<FWorldPartitionActorDesc*, int32> UnloadCount;
 	
-	int32 NumActorsToUnload = Algo::TransformAccumulate(CellsToUnload, [&UnloadCount](UWorldPartitionEditorCell* Cell)
-	{ 
-		for (FWorldPartitionActorDesc* ActorDesc : Cell->LoadedActors)
-		{
-			UnloadCount.FindOrAdd(ActorDesc, 0)++;
-		}
+		int32 NumActorsToUnload = Algo::TransformAccumulate(CellsToUnload, [&UnloadCount](UWorldPartitionEditorCell* Cell)
+		{ 
+			for (FWorldPartitionActorDesc* ActorDesc : Cell->LoadedActors)
+			{
+				UnloadCount.FindOrAdd(ActorDesc, 0)++;
+			}
 						
-		return Cell->LoadedActors.Num(); 
-	}, 0);
+			return Cell->LoadedActors.Num(); 
+		}, 0);
 		
-	for (const TPair<FWorldPartitionActorDesc*, int32> Pair : UnloadCount)
-	{
-		FWorldPartitionActorDesc* ActorDesc = Pair.Key;
-		// Only prompt if the actor will get unloaded by the unloading cells
-		if (ActorDesc->GetLoadedRefCount() == Pair.Value)
+		for (const TPair<FWorldPartitionActorDesc*, int32> Pair : UnloadCount)
 		{
-			AActor* LoadedActor = ActorDesc->GetActor();
-			check(LoadedActor);
-			UPackage* ActorPackage = LoadedActor->GetExternalPackage();
-			if (ActorPackage && ActorPackage->IsDirty())
+			FWorldPartitionActorDesc* ActorDesc = Pair.Key;
+			// Only prompt if the actor will get unloaded by the unloading cells
+			if (ActorDesc->GetLoadedRefCount() == Pair.Value)
 			{
-				ModifiedPackages.Add(ActorPackage);
+				AActor* LoadedActor = ActorDesc->GetActor();
+				check(LoadedActor);
+				UPackage* ActorPackage = LoadedActor->GetExternalPackage();
+				if (ActorPackage && ActorPackage->IsDirty())
+				{
+					ModifiedPackages.Add(ActorPackage);
+				}
 			}
 		}
-	}
 		
-	// Make sure we save modified actor packages before unloading
-	FEditorFileUtils::EPromptReturnCode RetCode = FEditorFileUtils::PR_Success;
-	if (ModifiedPackages.Num())
-	{
-		const bool bCheckDirty = false;
-		const bool bAlreadyCheckedOut = false;
-		const bool bCanBeDeclined = true;
-		const bool bPromptToSave = true;
-		const FText Title = LOCTEXT("SaveActorsTitle", "Save Actor(s)");
-		const FText Message = LOCTEXT("SaveActorsMessage", "Save Actor(s) before unloading them.");
-
-		RetCode = FEditorFileUtils::PromptForCheckoutAndSave(ModifiedPackages.Array(), bCheckDirty, bPromptToSave, Title, Message, nullptr, bAlreadyCheckedOut, bCanBeDeclined);
-		if (RetCode == FEditorFileUtils::PR_Cancelled)
+		// Make sure we save modified actor packages before unloading
+		FEditorFileUtils::EPromptReturnCode RetCode = FEditorFileUtils::PR_Success;
+		if (ModifiedPackages.Num())
 		{
-			return;
-		}
+			const bool bCheckDirty = false;
+			const bool bAlreadyCheckedOut = false;
+			const bool bCanBeDeclined = true;
+			const bool bPromptToSave = true;
+			const FText Title = LOCTEXT("SaveActorsTitle", "Save Actor(s)");
+			const FText Message = LOCTEXT("SaveActorsMessage", "Save Actor(s) before unloading them.");
+
+			RetCode = FEditorFileUtils::PromptForCheckoutAndSave(ModifiedPackages.Array(), bCheckDirty, bPromptToSave, Title, Message, nullptr, bAlreadyCheckedOut, bCanBeDeclined);
+			if (RetCode == FEditorFileUtils::PR_Cancelled)
+			{
+				return;
+			}
 
-		check(RetCode != FEditorFileUtils::PR_Failure);
-	}
+			check(RetCode != FEditorFileUtils::PR_Failure);
+		}
 
-	GEditor->SelectNone(true, true);
-	
-	FScopedSlowTask SlowTask(NumActorsToUnload, LOCTEXT("UnloadingCells", "Unloading cells..."));
-	SlowTask.MakeDialog();
+		GEditor->SelectNone(true, true);
 
-	for (UWorldPartitionEditorCell* Cell: CellsToUnload)
-	{
-		SlowTask.EnterProgressFrame(Cell->LoadedActors.Num());
-		UpdateLoadingEditorCell(Cell, false);
-	}
-		
-	GEditor->ResetTransaction(LOCTEXT("UnloadingEditorCellsResetTrans", "Unloading Cells"));
-	
-	// When save is declined make sure we don't keep unloaded/unsaved actor descs
-	// also make sure all modified packages are no longer dirty so they don't show up in save prompts anymore.
-	if (RetCode == FEditorFileUtils::PR_Declined)
-	{
-		for (UPackage* ModifiedPackage : ModifiedPackages)
-		{
-			ModifiedPackage->ClearDirtyFlag();
-		}
+		// At this point, cells might have changed due to saving deleted actors
+		CellsToUnload.Empty(CellsToUnload.Num());
+		if (EditorHash->GetIntersectingCells(Box, CellsToUnload))
+		{	
+			FScopedSlowTask SlowTask(NumActorsToUnload, LOCTEXT("UnloadingCells", "Unloading cells..."));
+			SlowTask.MakeDialog();
 
-		for (const TPair<FWorldPartitionActorDesc*, int32>& Pair : UnloadCount)
-		{
-			FWorldPartitionActorDesc* ActorDesc = Pair.Key;
-			if (!ActorDesc->GetLoadedRefCount() && !FPackageName::DoesPackageExist(ActorDesc->GetActorPackage().ToString()))
+			for (UWorldPartitionEditorCell* Cell: CellsToUnload)
 			{
-				// Already unloaded just remove from descriptor array
-				RemoveFromPartition(ActorDesc, /*bRemoveDescriptorFromArray=*/true, /*bUnloadRemovedDescriptor=*/false);
+				SlowTask.EnterProgressFrame(Cell->LoadedActors.Num());
+				UpdateLoadingEditorCell(Cell, false);
 			}
+		
+			GEditor->ResetTransaction(LOCTEXT("UnloadingEditorCellsResetTrans", "Unloading Cells"));
 		}
 	}
 }
@@ -950,20 +722,6 @@ void UWorldPartition::CreateLayers(const TSet<FName>& LayerNames)
 }
 #endif
 
-bool UWorldPartition::IsValidPartitionActor(AActor* Actor) const
-{
-#if WITH_EDITOR
-	if (Actor->IsPackageExternal() || Actor->GetPackage()->HasAnyPackageFlags(PKG_PlayInEditor))
-	{
-		return true;
-	}
-	const TUniquePtr<FWorldPartitionActorDesc>* ActorDesc = Actors.Find(Actor->GetActorGuid());
-	return ActorDesc->IsValid();
-#else
-	return Actor->IsPackageExternal();
-#endif
-}
-
 AActor* UWorldPartition::RegisterActor(FWorldPartitionActorDesc* ActorDesc)
 {
 #if WITH_EDITOR
@@ -977,7 +735,7 @@ AActor* UWorldPartition::RegisterActor(FWorldPartitionActorDesc* ActorDesc)
 	}
 
 	check(Actor);
-	check(IsValidPartitionActor(Actor));
+	check(Actor->IsPackageExternal());
 	check(ActorDesc->GetActor() == Actor);
 
 	ApplyActorTransform(Actor, InstanceTransform);
@@ -1008,18 +766,20 @@ AActor* UWorldPartition::RegisterActor(FWorldPartitionActorDesc* ActorDesc)
 void UWorldPartition::UnregisterActor(AActor* Actor)
 {
 #if WITH_EDITOR
-	check(IsValidPartitionActor(Actor));
 	FWorldPartitionActorDesc* ActorDesc = GetActorDesc(Actor->GetActorGuid());
 	check(ActorDesc);
 
-	if (LayerSubSystem)
+	if (!Actor->IsPendingKill())
 	{
-		LayerSubSystem->DisassociateActorFromLayers(Actor);
-	}
+		if (LayerSubSystem)
+		{
+			LayerSubSystem->DisassociateActorFromLayers(Actor);
+		}
 
-	OnActorRegisteredEvent.Broadcast(*Actor, false);
+		OnActorRegisteredEvent.Broadcast(*Actor, false);
 
-	Actor->GetLevel()->RemoveLoadedActor(Actor);
+		Actor->GetLevel()->RemoveLoadedActor(Actor);
+	}
 
 	ActorDesc->Unload();
 
@@ -1033,87 +793,76 @@ void UWorldPartition::UnregisterActor(AActor* Actor)
 }
 
 #if WITH_EDITOR
-void UWorldPartition::OnActorAdded(AActor* InActor)
+void UWorldPartition::OnAssetAdded(const FAssetData& InAssetData)
 {
-	check(!InActor->IsPendingKill());
-
-	if (InActor->GetLevel() == World->PersistentLevel && InActor->IsPackageExternal() && !InActor->IsChildActor())
+	if (!bIgnoreAssetRegistryEvents)
 	{
-		TUniquePtr<FWorldPartitionActorDesc>* ActorDescPtr = Actors.Find(InActor->GetActorGuid());
-
-		TUniquePtr<FWorldPartitionActorDesc> NewDesc(GetActorDescFactory(InActor)->Create());
-		NewDesc->Init(InActor);
-
-		if (ActorDescPtr)
-		{
-			// The only case where this is valid is when reinstancing BP actors. In this case, we'll receive the call
-			// for the newly spawned actor, while the old one hasn't been removed from the world yet. So we remove the
-			// old one before adding the new one.
-			check(GIsReinstancing);
-
-			RemoveFromPartition(ActorDescPtr->Get(), false);
-
-			*ActorDescPtr = MoveTemp(NewDesc);
-		}
-		else
+		TUniquePtr<FWorldPartitionActorDesc> NewActorDesc = GetActorDescriptor(InAssetData);
+		if (NewActorDesc.IsValid())
 		{
-			checkSlow(InActor->GetLevel()->Actors.Find(InActor) != INDEX_NONE);
+			check(!Actors.Contains(NewActorDesc->GetGuid()));
 
-			ActorDescPtr = &Actors.Add(InActor->GetActorGuid(), MoveTemp(NewDesc));
+			HashActorDesc(NewActorDesc.Get());
+			Actors.Add(NewActorDesc->GetGuid(), MoveTemp(NewActorDesc));
 		}
-
-		AddToPartition(ActorDescPtr->Get());
 	}
 }
 
-void UWorldPartition::OnActorDeleted(AActor* InActor)
+void UWorldPartition::OnAssetRemoved(const FAssetData& InAssetData)
 {
-	if (InActor->GetLevel() == World->PersistentLevel)
+	if (!bIgnoreAssetRegistryEvents)
 	{
-		if (FWorldPartitionActorDesc* ActorDesc = GetActorDesc(InActor->GetActorGuid()))
+		TUniquePtr<FWorldPartitionActorDesc> NewActorDesc = GetActorDescriptor(InAssetData);
+		if (NewActorDesc.IsValid())
 		{
-			if (InActor->GetClass()->HasAnyClassFlags(CLASS_NewerVersionExists))
-			{
-				// We are receiving this event when destroying the old actor after BP reinstantiation. In this case,
-				// the newly created actor was already added to the list, so we can safely ignore this case.
-				check(GIsReinstancing);
-			}
-			else
-			{
-				// During undo transactions, newly created objects will get removed with their package unset due to the
-				// order of SetPackage/Modify. Take this into account in the check below.
-				check(InActor->IsPackageExternal() || GIsTransacting);
+			TUniquePtr<FWorldPartitionActorDesc>& ExistingActorDesc = Actors.FindChecked(NewActorDesc->GetGuid());
 
-				// Validate that this actor is already removed from the level
-				verify(InActor->GetLevel()->Actors.Find(InActor) == INDEX_NONE);
-
-				RemoveFromPartition(ActorDesc);
-			}
+			UnhashActorDesc(ExistingActorDesc.Get());
+			Actors.Remove(ExistingActorDesc->GetGuid());
 		}
 	}
 }
 
-void UWorldPartition::OnActorMoving(AActor* InActor)
+void UWorldPartition::OnAssetUpdated(const FAssetData& InAssetData)
 {
-	if (InActor->GetLevel() == World->PersistentLevel)
+	if (!bIgnoreAssetRegistryEvents)
 	{
-		UpdateActorDesc(InActor);
+		TUniquePtr<FWorldPartitionActorDesc> NewActorDesc = GetActorDescriptor(InAssetData);
+		if (NewActorDesc.IsValid())
+		{
+			TUniquePtr<FWorldPartitionActorDesc>& ExistingActorDesc = Actors.FindChecked(NewActorDesc->GetGuid());
+
+			UnhashActorDesc(ExistingActorDesc.Get());
+			ExistingActorDesc = MoveTemp(NewActorDesc);
+			HashActorDesc(ExistingActorDesc.Get());
+		}
 	}
 }
 
-void UWorldPartition::OnActorOuterChanged(AActor* InActor, UObject* InOldOuter)
+TUniquePtr<FWorldPartitionActorDesc> UWorldPartition::GetActorDescriptor(const FAssetData& InAssetData)
 {
-	ULevel* OldLevel = Cast<ULevel>(InOldOuter);
-
-	if (OldLevel == World->PersistentLevel)
+	FString ActorClass;
+	static FName NAME_ActorMetaDataClass(TEXT("ActorMetaDataClass"));
+	if (InAssetData.GetTagValue(NAME_ActorMetaDataClass, ActorClass))
 	{
-		if (FWorldPartitionActorDesc* ActorDesc = GetActorDesc(InActor->GetActorGuid()))
+		FString ActorMetaDataStr;
+		static FName NAME_ActorMetaData(TEXT("ActorMetaData"));
+		if (InAssetData.GetTagValue(NAME_ActorMetaData, ActorMetaDataStr))
 		{
-			const bool bRemoveDescriptor = true;
-			const bool bUnloadDescriptor = false;
-			RemoveFromPartition(ActorDesc, bRemoveDescriptor, bUnloadDescriptor);
+			FWorldPartitionActorDescInitData ActorDescInitData;
+			ActorDescInitData.NativeClass = FindObjectChecked<UClass>(ANY_PACKAGE, *ActorClass, true);
+			ActorDescInitData.PackageName = InAssetData.PackageName;
+			ActorDescInitData.ActorPath = InAssetData.ObjectPath;
+			FBase64::Decode(ActorMetaDataStr, ActorDescInitData.SerializedData);
+
+			FWorldPartitionActorDesc* NewActorDesc(GetActorDescFactory(ActorDescInitData.NativeClass)->Create());
+			NewActorDesc->Init(ActorDescInitData);
+
+			return TUniquePtr<FWorldPartitionActorDesc>(NewActorDesc);
 		}
 	}
+
+	return nullptr;
 }
 
 void UWorldPartition::HashActorDesc(FWorldPartitionActorDesc* ActorDesc)
@@ -1129,36 +878,8 @@ void UWorldPartition::UnhashActorDesc(FWorldPartitionActorDesc* ActorDesc)
 	check(ActorDesc);
 	check(EditorHash);
 	EditorHash->UnhashActor(ActorDesc);
-
 	check(!ActorDesc->GetLoadedRefCount());
 }
-
-void UWorldPartition::AddToPartition(FWorldPartitionActorDesc* ActorDesc)
-{
-	AddToClusters(ActorDesc);
-
-	HashActorDesc(ActorDesc);	
-}
-
-void UWorldPartition::RemoveFromPartition(FWorldPartitionActorDesc* ActorDesc, bool bRemoveDescriptorFromArray, bool bUnloadRemovedDescriptor)
-{
-	// Unhash this actor from the editor hash
-	UnhashActorDesc(ActorDesc);
-
-	RemoveFromClusters(ActorDesc);
-
-	// Remove this actor descriptor
-	if (bRemoveDescriptorFromArray)
-	{
-		TUniquePtr<FWorldPartitionActorDesc> RemovedActorDesc = Actors.FindAndRemoveChecked(ActorDesc->GetGuid());
-		check(RemovedActorDesc->GetLoadedRefCount() == 0);
-		if (bUnloadRemovedDescriptor)
-		{
-			RemovedActorDesc->Unload();
-		}
-	}
-}
-
 #endif
 
 void UWorldPartition::Serialize(FArchive& Ar)
diff --git a/Engine/Source/Runtime/Engine/Private/WorldPartition/WorldPartitionActorDesc.cpp b/Engine/Source/Runtime/Engine/Private/WorldPartition/WorldPartitionActorDesc.cpp
--- a/Engine/Source/Runtime/Engine/Private/WorldPartition/WorldPartitionActorDesc.cpp
+++ b/Engine/Source/Runtime/Engine/Private/WorldPartition/WorldPartitionActorDesc.cpp
@@ -23,12 +23,6 @@ FWorldPartitionActorDesc::FWorldPartitionActorDesc()
 {}
 
 void FWorldPartitionActorDesc::Init(const AActor* InActor)
-{
-	InitFrom(InActor);
-	UpdateHash();
-}
-
-void FWorldPartitionActorDesc::InitFrom(const AActor* InActor)
 {	
 	check(InActor->IsPackageExternal());
 
@@ -101,18 +95,6 @@ void FWorldPartitionActorDesc::Init(const FWorldPartitionActorDescInitData& Desc
 	{
 		GridPlacement = DefaultGridPlacement;
 	}
-
-	// Transform BoundsLocation and BoundsExtent if necessary
-	if (!DescData.Transform.Equals(FTransform::Identity))
-	{
-		//@todo_ow: This will result in a new BoundsExtent that is larger than it should. To fix this, we would need the Object Oriented BoundingBox of the actor (the BV of the actor without rotation)
-		const FVector BoundsMin = BoundsLocation - BoundsExtent;
-		const FVector BoundsMax = BoundsLocation + BoundsExtent;
-		const FBox NewBounds = FBox(BoundsMin, BoundsMax).TransformBy(DescData.Transform);
-		NewBounds.GetCenterAndExtents(BoundsLocation, BoundsExtent);
-	}
-
-	UpdateHash();
 }
 
 void FWorldPartitionActorDesc::SerializeTo(TArray<uint8>& OutData)
@@ -133,16 +115,26 @@ void FWorldPartitionActorDesc::SerializeTo(TArray<uint8>& OutData)
 	OutData.Append(PayloadData);
 }
 
-FString FWorldPartitionActorDesc::ToString() const
+void FWorldPartitionActorDesc::TransformInstance(const FString& From, const FString& To, const FTransform& InstanceTransform)
 {
-	return FString::Printf(TEXT("Guid:%s Class:%s Name:%s"), *Guid.ToString(), *Class.ToString(), *FPaths::GetExtension(ActorPath.ToString()));
+	check(!LoadedRefCount);
+
+	ActorPath = *ActorPath.ToString().Replace(*From, *To);
+
+	// Transform BoundsLocation and BoundsExtent if necessary
+	if (!InstanceTransform.Equals(FTransform::Identity))
+	{
+		//@todo_ow: This will result in a new BoundsExtent that is larger than it should. To fix this, we would need the Object Oriented BoundingBox of the actor (the BV of the actor without rotation)
+		const FVector BoundsMin = BoundsLocation - BoundsExtent;
+		const FVector BoundsMax = BoundsLocation + BoundsExtent;
+		const FBox NewBounds = FBox(BoundsMin, BoundsMax).TransformBy(InstanceTransform);
+		NewBounds.GetCenterAndExtents(BoundsLocation, BoundsExtent);
+	}
 }
 
-void FWorldPartitionActorDesc::UpdateHash()
+FString FWorldPartitionActorDesc::ToString() const
 {
-	FHashBuilderArchive HashBuilderAr;
-	Serialize(HashBuilderAr);
-	Hash = HashBuilderAr.GetHash();
+	return FString::Printf(TEXT("Guid:%s Class:%s Name:%s"), *Guid.ToString(), *Class.ToString(), *FPaths::GetExtension(ActorPath.ToString()));
 }
 
 void FWorldPartitionActorDesc::Serialize(FArchive& Ar)
@@ -197,5 +189,4 @@ void FWorldPartitionActorDesc::Unload()
 		}, false);
 	}
 }
-
 #endif
\ No newline at end of file
diff --git a/Engine/Source/Runtime/Engine/Private/WorldPartition/WorldPartitionEditorCell.cpp b/Engine/Source/Runtime/Engine/Private/WorldPartition/WorldPartitionEditorCell.cpp
--- a/Engine/Source/Runtime/Engine/Private/WorldPartition/WorldPartitionEditorCell.cpp
+++ b/Engine/Source/Runtime/Engine/Private/WorldPartition/WorldPartitionEditorCell.cpp
@@ -20,18 +20,13 @@ void UWorldPartitionEditorCell::Serialize(FArchive& Ar)
 }
 
 #if WITH_EDITOR
-UWorldPartition* UWorldPartitionEditorCell::GetWorldPartition()
-{
-	return GetOuterUWorldPartitionEditorHash()->GetOuterUWorldPartition();
-}
-
 void UWorldPartitionEditorCell::AddActor(FWorldPartitionActorDesc* InActorDesc)
 {
 	check(InActorDesc);
 	bool bIsAlreadyInSet = false;
 	Actors.Add(InActorDesc, &bIsAlreadyInSet);
 
-	if (!bIsAlreadyInSet)
+	if (!bIsAlreadyInSet && bLoaded)
 	{
 		if (AActor* Actor = InActorDesc->GetActor())
 		{
@@ -54,6 +49,8 @@ void UWorldPartitionEditorCell::RemoveActor(FWorldPartitionActorDesc* InActorDes
 
 	if (LoadedActors.Remove(InActorDesc))
 	{
+		check(bLoaded);
+
 		const uint32 ActorRefCount = InActorDesc->RemoveLoadedRefCount();
 		UE_LOG(LogWorldPartition, Verbose, TEXT(" ==> Unreferenced loaded actor %s(%d) [UWorldPartitionEditorCell::RemoveActor]"), *InActorDesc->GetActor()->GetFullName(), ActorRefCount);
 	}
diff --git a/Engine/Source/Runtime/Engine/Private/WorldPartition/WorldPartitionRuntimeSpatialHash.cpp b/Engine/Source/Runtime/Engine/Private/WorldPartition/WorldPartitionRuntimeSpatialHash.cpp
--- a/Engine/Source/Runtime/Engine/Private/WorldPartition/WorldPartitionRuntimeSpatialHash.cpp
+++ b/Engine/Source/Runtime/Engine/Private/WorldPartition/WorldPartitionRuntimeSpatialHash.cpp
@@ -106,6 +106,145 @@ static FAutoConsoleVariableRef CVarRuntimeSpatialHashCellToSourceAngleContributi
 
 // ------------------------------------------------------------------------------------------------
 
+// Clustering
+struct FActorCluster
+{
+	TSet<FGuid>			Actors;
+	EActorGridPlacement	GridPlacement;
+	FName				RuntimeGrid;
+	FBox				Bounds;
+	TArray<FName>		DataLayers;
+	uint32				DataLayersID;
+
+	FActorCluster(const FWorldPartitionActorDesc* ActorDesc);
+	void Add(const FActorCluster& ActorCluster);
+};
+
+FActorCluster::FActorCluster(const FWorldPartitionActorDesc* ActorDesc)
+	: GridPlacement(ActorDesc->GetGridPlacement())
+	, RuntimeGrid(ActorDesc->GetRuntimeGrid())
+	, Bounds(ActorDesc->GetBounds())
+{
+	check(GridPlacement != EActorGridPlacement::None);
+	Actors.Add(ActorDesc->GetGuid());
+	DataLayersID = FDataLayersHelper::ComputeDataLayerID(DataLayers);
+}
+
+void FActorCluster::Add(const FActorCluster& ActorCluster)
+{
+	// Merge Actors
+	Actors.Append(ActorCluster.Actors);
+
+	// Merge RuntimeGrid
+	RuntimeGrid = RuntimeGrid == ActorCluster.RuntimeGrid ? RuntimeGrid : NAME_None;
+
+	// Merge Bounds
+	Bounds += ActorCluster.Bounds;
+
+	// Merge GridPlacement
+	// If currently None, will always stay None
+	if (GridPlacement != EActorGridPlacement::None)
+	{
+		// If grid placement differs between the two clusters
+		if (GridPlacement != ActorCluster.GridPlacement)
+		{
+			// If one of the two cluster was always loaded, set to None
+			if (ActorCluster.GridPlacement == EActorGridPlacement::AlwaysLoaded || GridPlacement == EActorGridPlacement::AlwaysLoaded)
+			{
+				GridPlacement = EActorGridPlacement::None;
+			}
+			else
+			{
+				GridPlacement = ActorCluster.GridPlacement;
+			}
+		}
+
+		// If current placement is set to Location, that won't make sense when merging two clusters. Set to Bounds
+		if (GridPlacement == EActorGridPlacement::Location)
+		{
+			GridPlacement = EActorGridPlacement::Bounds;
+		}
+	}
+
+	// Merge DataLayers
+	if (DataLayersID != ActorCluster.DataLayersID)
+	{
+		for (const FName& DataLayer : ActorCluster.DataLayers)
+		{
+			DataLayers.AddUnique(DataLayer);
+		}
+		DataLayersID = FDataLayersHelper::ComputeDataLayerID(DataLayers);
+	}
+}
+
+
+void CreateActorCluster(const FGuid& ActorGuid, const FWorldPartitionActorDesc* ActorDesc, TMap<FGuid, FActorCluster*>& ActorToActorCluster, TSet<FActorCluster*>& ActorClustersSet, UWorldPartition* WorldPartition)
+{
+	FActorCluster* ActorCluster = ActorToActorCluster.FindRef(ActorGuid);
+	if (!ActorCluster)
+	{
+		ActorCluster = new FActorCluster(ActorDesc);
+		ActorClustersSet.Add(ActorCluster);
+		ActorToActorCluster.Add(ActorGuid, ActorCluster);
+	}
+
+	// Don't include references from editor-only actors
+	if (!ActorDesc->GetActorIsEditorOnly())
+	{
+		for (const FGuid& ReferenceGuid : ActorDesc->GetReferences())
+		{
+			const FWorldPartitionActorDesc* ReferenceActorDesc = WorldPartition->GetActorDesc(ReferenceGuid);
+
+			// Don't include references to editor-only actors
+			if (!ReferenceActorDesc->GetActorIsEditorOnly())
+			{
+				FActorCluster* ReferenceCluster = ActorToActorCluster.FindRef(ReferenceGuid);
+				if (ReferenceCluster)
+				{
+					if (ReferenceCluster != ActorCluster)
+					{
+						// Merge reference cluster in Actor's cluster
+						ActorCluster->Add(*ReferenceCluster);
+						for (const FGuid& ReferenceClusterActorGuid : ReferenceCluster->Actors)
+						{
+							ActorToActorCluster[ReferenceClusterActorGuid] = ActorCluster;
+						}
+						ActorClustersSet.Remove(ReferenceCluster);
+						delete ReferenceCluster;
+					}
+				}
+				else
+				{
+					// Put Reference in Actor's cluster
+					ActorCluster->Add(FActorCluster(ReferenceActorDesc));
+				}
+
+				// Map its cluster
+				ActorToActorCluster.Add(ReferenceGuid, ActorCluster);
+			}
+		}
+	}
+}
+
+TArray<FActorCluster> CreateActorClusters(UWorldPartition* WorldPartition)
+{
+	TMap<FGuid, FActorCluster*> ActorToActorCluster;
+	TSet<FActorCluster*> ActorClustersSet;
+
+	for (const auto& Pair : WorldPartition->Actors)
+	{
+		const FGuid& ActorGuid = Pair.Key;
+		const FWorldPartitionActorDesc* ActorDesc = Pair.Value.Get();
+		CreateActorCluster(ActorGuid, ActorDesc, ActorToActorCluster, ActorClustersSet, WorldPartition);
+	}
+
+	TArray<FActorCluster> ActorClusters;
+	ActorClusters.Reserve(ActorClustersSet.Num());
+	Algo::Transform(ActorClustersSet, ActorClusters, [](FActorCluster* ActorCluster) { return MoveTemp(*ActorCluster); });
+	for (FActorCluster* ActorCluster : ActorClustersSet) { delete ActorCluster; }
+	return ActorClusters;
+}
+
 /**
   * Square 2D grid helper
   */
@@ -820,8 +959,7 @@ void FSpatialHashStreamingGrid::Draw2D(UCanvas* Canvas, const TArray<FWorldParti
 // ------------------------------------------------------------------------------------------------
 
 #if WITH_EDITOR
-
-static FSquare2DGridHelper GetPartitionedActors(const UWorldPartition* WorldPartition, const FBox& WorldBounds, const FSpatialHashRuntimeGrid& Grid, const TArray<const UWorldPartition::FActorCluster*>& GridActors)
+static FSquare2DGridHelper GetPartitionedActors(const UWorldPartition* WorldPartition, const FBox& WorldBounds, const FSpatialHashRuntimeGrid& Grid, const TArray<FActorCluster>& GridActors)
 {
 	UE_SCOPED_TIMER(TEXT("GetPartitionedActors"), LogWorldPartitionRuntimeSpatialHash, Log);
 
@@ -858,11 +996,11 @@ static FSquare2DGridHelper GetPartitionedActors(const UWorldPartition* WorldPart
 	//	
 	FSquare2DGridHelper PartitionedActors(GridLevelCount, GridOrigin, GridCellSize, GridSize);
 
-	for (const UWorldPartition::FActorCluster* ActorCluster : GridActors)
+	for (const FActorCluster& ActorCluster : GridActors)
 	{
-		check(ActorCluster->Actors.Num() > 0);
+		check(ActorCluster.Actors.Num() > 0);
 
-		EActorGridPlacement GridPlacement = ActorCluster->GridPlacement;
+		EActorGridPlacement GridPlacement = ActorCluster.GridPlacement;
 		bool bAlwaysLoadedPromotedCluster = (GridPlacement == EActorGridPlacement::None);
 		bool bAlwaysLoadedPromotedOutOfGrid = false;
 
@@ -878,8 +1016,8 @@ static FSquare2DGridHelper GetPartitionedActors(const UWorldPartition* WorldPart
 			case EActorGridPlacement::Location:
 			{
 				FIntVector2 CellCoords;
-				check(ActorCluster->Actors.Num() == 1);
-				const FGuid& ActorGuid = *ActorCluster->Actors.CreateConstIterator();
+				check(ActorCluster.Actors.Num() == 1);
+				const FGuid& ActorGuid = *ActorCluster.Actors.CreateConstIterator();
 				const FWorldPartitionActorDesc& ActorDesc = *WorldPartition->GetActorDesc(ActorGuid);
 				if (PartitionedActors.GetLowestLevel().GetCellCoords(FVector2D(ActorDesc.GetOrigin()), CellCoords))
 				{
@@ -899,12 +1037,12 @@ static FSquare2DGridHelper GetPartitionedActors(const UWorldPartition* WorldPart
 				for (FSquare2DGridHelper::FGridLevel& GridLevel : PartitionedActors.Levels)
 				{
 					int32 IntersectingCellCount = 0;
-					GridLevel.ForEachIntersectingCellsBreakable(ActorCluster->Bounds, [&IntersectingCellCount](const FIntVector2& Coords) { return ++IntersectingCellCount <= 1; });
+					GridLevel.ForEachIntersectingCellsBreakable(ActorCluster.Bounds, [&IntersectingCellCount](const FIntVector2& Coords) { return ++IntersectingCellCount <= 1; });
 					if (IntersectingCellCount == 1)
 					{
-						GridLevel.ForEachIntersectingCells(ActorCluster->Bounds, [&GridLevel,&ActorCluster](const FIntVector2& Coords)
+						GridLevel.ForEachIntersectingCells(ActorCluster.Bounds, [&GridLevel,&ActorCluster](const FIntVector2& Coords)
 						{
-							GridLevel.GetCell(Coords).AddActors(ActorCluster->Actors, ActorCluster->DataLayers);
+							GridLevel.GetCell(Coords).AddActors(ActorCluster.Actors, ActorCluster.DataLayers);
 						});
 						bFoundCell = true;
 						break;
@@ -924,24 +1062,24 @@ static FSquare2DGridHelper GetPartitionedActors(const UWorldPartition* WorldPart
 			
 		if (GridPlacement == EActorGridPlacement::AlwaysLoaded)
 		{
-			PartitionedActors.GetAlwaysLoadedCell().AddActors(ActorCluster->Actors, ActorCluster->DataLayers);
+			PartitionedActors.GetAlwaysLoadedCell().AddActors(ActorCluster.Actors, ActorCluster.DataLayers);
 		}
 
 		if (!LogWorldPartitionRuntimeSpatialHash.IsSuppressed(ELogVerbosity::Verbose))
 		{
-			if (ActorCluster->Actors.Num() > 1)
+			if (ActorCluster.Actors.Num() > 1)
 			{
 				static UEnum* ActorGridPlacementEnum = FindObject<UEnum>(ANY_PACKAGE, TEXT("EActorGridPlacement"));
 
 				UE_LOG(LogWorldPartitionRuntimeSpatialHash, Verbose, TEXT("Clustered %d actors (%s%s%s), generated shared BV of [%d x %d] (meters)"), 
-					ActorCluster->Actors.Num(),
+					ActorCluster.Actors.Num(),
 					*ActorGridPlacementEnum->GetNameStringByValue((int64)GridPlacement),
 					bAlwaysLoadedPromotedCluster ? TEXT(":PromotedCluster") : TEXT(""),
 					bAlwaysLoadedPromotedOutOfGrid ? TEXT(":PromotedOutOfGrid") : TEXT(""),
-					(int)(0.01f * ActorCluster->Bounds.GetSize().X),
-					(int)(0.01f * ActorCluster->Bounds.GetSize().Y));
+					(int)(0.01f * ActorCluster.Bounds.GetSize().X),
+					(int)(0.01f * ActorCluster.Bounds.GetSize().Y));
 
-				for (const FGuid& ActorGuid : ActorCluster->Actors)
+				for (const FGuid& ActorGuid : ActorCluster.Actors)
 				{
 					const FWorldPartitionActorDesc& Desc = *WorldPartition->GetActorDesc(ActorGuid);
 					UE_LOG(LogWorldPartitionRuntimeSpatialHash, Verbose, TEXT("   - Actor: %s (%s)"), *Desc.GetActorPath().ToString(), *ActorGuid.ToString(EGuidFormats::UniqueObjectGuid));
@@ -1053,19 +1191,21 @@ bool UWorldPartitionRuntimeSpatialHash::GenerateStreaming(EWorldPartitionStreami
 	}
 
 	// Create actor clusters
-	TArray<TArray<const UWorldPartition::FActorCluster*>> GridActors;
+	TArray<FActorCluster> ActorClusters = CreateActorClusters(WorldPartition);
+
+	TArray<TArray<FActorCluster>> GridActors;
 	GridActors.InsertDefaulted(0, AllGrids.Num());
 
-	for (const UWorldPartition::FActorCluster* ActorCluster : WorldPartition->GetActorClusters())
+	for (FActorCluster& ActorCluster : ActorClusters)
 	{
-		int32* FoundIndex = GridsMapping.Find(ActorCluster->RuntimeGrid);
+		int32* FoundIndex = GridsMapping.Find(ActorCluster.RuntimeGrid);
 		if (!FoundIndex)
 		{
-			UE_LOG(LogWorldPartitionRuntimeSpatialHash, Error, TEXT("Invalid partition grid '%s' referenced by actor cluster"), *ActorCluster->RuntimeGrid.ToString());
+			UE_LOG(LogWorldPartitionRuntimeSpatialHash, Error, TEXT("Invalid partition grid '%s' referenced by actor cluster"), *ActorCluster.RuntimeGrid.ToString());
 		}
 
 		int32 GridIndex = FoundIndex ? *FoundIndex : 0;
-		GridActors[GridIndex].Add(ActorCluster);
+		GridActors[GridIndex].Add(MoveTemp(ActorCluster));
 	}
 	
 	const FBox WorldBounds = WorldPartition->GetWorldBounds();
@@ -1300,21 +1440,24 @@ bool UWorldPartitionRuntimeSpatialHash::GenerateHLOD()
 
 	UWorldPartition* WorldPartition = GetOuterUWorldPartition();
 
-	TArray<TArray<const UWorldPartition::FActorCluster*>> GridActors;
+	// Create actor clusters
+	TArray<FActorCluster> ActorClusters = CreateActorClusters(WorldPartition);
+
+	TArray<TArray<FActorCluster>> GridActors;
 	GridActors.InsertDefaulted(0, Grids.Num());
 
-	for (const UWorldPartition::FActorCluster* ActorCluster : WorldPartition->GetActorClusters())
+	for (FActorCluster& ActorCluster : ActorClusters)
 	{
-		int32* FoundIndex = GridsMapping.Find(ActorCluster->RuntimeGrid);
+		int32* FoundIndex = GridsMapping.Find(ActorCluster.RuntimeGrid);
 		if (!FoundIndex)
 		{
-			UE_LOG(LogWorldPartitionRuntimeSpatialHash, Error, TEXT("Invalid partition grid '%s' referenced by actor cluster"), *ActorCluster->RuntimeGrid.ToString());
+			UE_LOG(LogWorldPartitionRuntimeSpatialHash, Error, TEXT("Invalid partition grid '%s' referenced by actor cluster"), *ActorCluster.RuntimeGrid.ToString());
 		}
 
 		int32 GridIndex = FoundIndex ? *FoundIndex : 0;
 		if (GridIndex != INDEX_NONE)
 		{
-			GridActors[GridIndex].Add(ActorCluster);
+			GridActors[GridIndex].Add(MoveTemp(ActorCluster));
 		}
 	}
 
@@ -1378,7 +1521,7 @@ bool UWorldPartitionRuntimeSpatialHash::GenerateNavigationData()
 	UWorldPartition* WorldPartition = GetOuterUWorldPartition();
 	UWorld* World = WorldPartition->World;
 
-	TArray<TArray<const UWorldPartition::FActorCluster*>> GridActors;
+	TArray<TArray<FActorCluster>> GridActors;
 	GridActors.InsertDefaulted(0, Grids.Num());
 
 	// Generate navmesh
@@ -1446,8 +1589,6 @@ bool UWorldPartitionRuntimeSpatialHash::GenerateNavigationData()
 
 		// Set target grid
 		DataChunkActor->RuntimeGrid = RuntimeGrid.GridName;
-
-		WorldPartition->UpdateActorDesc(DataChunkActor);
 		ValidNavigationDataChunkActors.Add(DataChunkActor);
 	});
 
diff --git a/Engine/Source/Runtime/Engine/Private/WorldPartition/WorldPartitionSubsystem.cpp b/Engine/Source/Runtime/Engine/Private/WorldPartition/WorldPartitionSubsystem.cpp
--- a/Engine/Source/Runtime/Engine/Private/WorldPartition/WorldPartitionSubsystem.cpp
+++ b/Engine/Source/Runtime/Engine/Private/WorldPartition/WorldPartitionSubsystem.cpp
@@ -207,14 +207,6 @@ void UWorldPartitionSubsystem::ForEachActorDesc(TSubclassOf<AActor> ActorClass,
 	}
 }
 
-void UWorldPartitionSubsystem::UpdateActorDesc(AActor* Actor)
-{
-	if (UWorldPartition* MainPartition = GetMainWorldPartition())
-	{
-		MainPartition->UpdateActorDesc(Actor);
-	}
-}
-
 void UWorldPartitionSubsystem::RegisterActorDescFactory(TSubclassOf<AActor> Class, FWorldPartitionActorDescFactory* Factory)
 {
 	UWorldPartition::RegisterActorDescFactory(Class, Factory);
diff --git a/Engine/Source/Runtime/Engine/Public/WorldPartition/ActorPartition/PartitionActorDesc.h b/Engine/Source/Runtime/Engine/Public/WorldPartition/ActorPartition/PartitionActorDesc.h
--- a/Engine/Source/Runtime/Engine/Public/WorldPartition/ActorPartition/PartitionActorDesc.h
+++ b/Engine/Source/Runtime/Engine/Public/WorldPartition/ActorPartition/PartitionActorDesc.h
@@ -19,7 +19,7 @@ public:
 	int64 GridIndexZ;
 	FGuid GridGuid;
 protected:
-	virtual void InitFrom(const AActor* InActor) override;
+	virtual void Init(const AActor* InActor) override;
 	virtual void Serialize(FArchive& Ar) override;
 #endif
 };
diff --git a/Engine/Source/Runtime/Engine/Public/WorldPartition/HLOD/HLODActorDesc.h b/Engine/Source/Runtime/Engine/Public/WorldPartition/HLOD/HLODActorDesc.h
--- a/Engine/Source/Runtime/Engine/Public/WorldPartition/HLOD/HLODActorDesc.h
+++ b/Engine/Source/Runtime/Engine/Public/WorldPartition/HLOD/HLODActorDesc.h
@@ -22,7 +22,7 @@ public:
 	inline const FSoftObjectPath& GetHLODLayer() const { return HLODLayer; }
 
 protected:
-	virtual void InitFrom(const AActor* InActor) override;
+	virtual void Init(const AActor* InActor) override;
 	virtual void Serialize(FArchive& Ar) override;
 
 	TArray<FGuid> SubActors;
diff --git a/Engine/Source/Runtime/Engine/Public/WorldPartition/Landscape/LandscapeActorDesc.h b/Engine/Source/Runtime/Engine/Public/WorldPartition/Landscape/LandscapeActorDesc.h
--- a/Engine/Source/Runtime/Engine/Public/WorldPartition/Landscape/LandscapeActorDesc.h
+++ b/Engine/Source/Runtime/Engine/Public/WorldPartition/Landscape/LandscapeActorDesc.h
@@ -16,6 +16,6 @@ class ENGINE_API FLandscapeActorDesc : public FPartitionActorDesc
 	friend class FLandscapeActorDescFactory;
 
 protected:
-	virtual void InitFrom(const AActor* InActor) override;
+	virtual void Init(const AActor* InActor) override;
 #endif
 };
diff --git a/Engine/Source/Runtime/Engine/Public/WorldPartition/LevelInstance/LevelInstanceActorDesc.h b/Engine/Source/Runtime/Engine/Public/WorldPartition/LevelInstance/LevelInstanceActorDesc.h
--- a/Engine/Source/Runtime/Engine/Public/WorldPartition/LevelInstance/LevelInstanceActorDesc.h
+++ b/Engine/Source/Runtime/Engine/Public/WorldPartition/LevelInstance/LevelInstanceActorDesc.h
@@ -22,7 +22,7 @@ public:
 	inline FName GetLevelPackage() const { return LevelPackage; }
 
 protected:
-	virtual void InitFrom(const AActor* InActor) override;
+	virtual void Init(const AActor* InActor) override;
 	virtual void Serialize(FArchive& Ar) override;
 
 	FName LevelPackage;
diff --git a/Engine/Source/Runtime/Engine/Public/WorldPartition/WorldPartition.h b/Engine/Source/Runtime/Engine/Public/WorldPartition/WorldPartition.h
--- a/Engine/Source/Runtime/Engine/Public/WorldPartition/WorldPartition.h
+++ b/Engine/Source/Runtime/Engine/Public/WorldPartition/WorldPartition.h
@@ -64,8 +64,6 @@ private:
 	DECLARE_MULTICAST_DELEGATE_OneParam(FWorldPartitionChangedEvent, UWorld*);
 	static FWorldPartitionChangedEvent WorldPartitionChangedEvent;
 
-	void OnObjectPropertyChanged(UObject* Object, FPropertyChangedEvent& InPropertyChangedEvent);
-
 	static FWorldPartitionActorDescFactory* GetActorDescFactory(TSubclassOf<AActor> Class);
 	static FWorldPartitionActorDescFactory* GetActorDescFactory(const AActor* Actor);
 
@@ -75,12 +73,14 @@ private:
 	void OnPreBeginPIE(bool bStartSimulate);
 	void OnEndPIE(bool bStartSimulate);
 
-	void OnActorAdded(AActor* InActor);
-	void OnActorDeleted(AActor* InActor);
-	void OnActorOuterChanged(AActor* InActor, UObject* InOldOuter);
-	void OnActorMoving(AActor* InActor);
-		
-	FDelegateHandle OnInitialAssetRegistrySearchCompleteHandle;
+	// Asset registry events
+	void OnAssetAdded(const FAssetData& InAssetData);
+	void OnAssetRemoved(const FAssetData& InAssetData);
+	void OnAssetUpdated(const FAssetData& InAssetData);
+
+	TUniquePtr<FWorldPartitionActorDesc> GetActorDescriptor(const FAssetData& InAssetData);
+
+	void ApplyActorTransform(AActor* InActor, const FTransform& InTransform);
 #endif
 
 	virtual void Serialize(FArchive& Ar) override;
@@ -97,13 +97,10 @@ private:
 	void ForEachActorDesc(TSubclassOf<AActor> ActorClass, TFunctionRef<bool(const FWorldPartitionActorDesc*)> Predicate) const;
 	// UActorPartitionSubsystem interface-
 
-	void ApplyActorTransform(AActor* InActor, const FTransform& InTransform);
-
 	static void RegisterActorDescFactory(TSubclassOf<AActor> Class, FWorldPartitionActorDescFactory* Factory);
 	// UWorldPartitionSubsystem interface-
 
 public:
-	// public interface+
 	FName GetWorldPartitionEditorName();
 	
 	// PIE Methods
@@ -119,10 +116,6 @@ public:
 	void LoadEditorCells(const FBox& Box);
 	void UnloadEditorCells(const FBox& Box);
 
-	void LoadEditorCells(const TArray<UWorldPartitionEditorCell*>& Cells);
-	void UnloadEditorCells(const TArray<UWorldPartitionEditorCell*>& Cells);
-
-	void UpdateActorDesc(AActor* InActor);
 	bool IsPreCooked() const { return bIsPreCooked; }
 	void SetIsPreCooked(bool bInIsPreCooked) { bIsPreCooked = bInIsPreCooked; }
 
@@ -131,24 +124,6 @@ public:
 
 	void GenerateHLOD();
 	void GenerateNavigationData();
-
-	// Clustering
-	struct FActorCluster
-	{
-		TSet<FGuid>			Actors;
-		EActorGridPlacement	GridPlacement;
-		FName				RuntimeGrid;
-		FBox				Bounds;
-		TArray<FName>		DataLayers;
-		uint32				DataLayersID;
-
-		FActorCluster(const FWorldPartitionActorDesc* ActorDesc);
-		void Add(const FActorCluster& ActorCluster);
-	};
-
-	const TSet<FActorCluster*>& GetActorClusters() const;
-	const FActorCluster* GetClusterForActor(const FGuid& InActorGuid) const;
-	// public interface-
 #endif
 
 public:
@@ -190,7 +165,7 @@ public:
 	DECLARE_EVENT_TwoParams(UWorldPartition, FWorldPartitionActorRegisteredEvent, AActor&, bool);
 	FWorldPartitionActorRegisteredEvent OnActorRegisteredEvent;
 
-	bool bDirty;
+	bool bIgnoreAssetRegistryEvents;
 	bool bForceGarbageCollection;
 	bool bForceGarbageCollectionPurge;
 #endif
@@ -220,7 +195,6 @@ private:
 	void UnregisterActor(AActor* Actor);
 
 	bool IsMainWorldPartition() const;
-	bool IsValidPartitionActor(AActor* Actor) const;
 
 	// Delegates registration
 	void RegisterDelegates();
@@ -234,15 +208,6 @@ private:
 
 	void HashActorDesc(FWorldPartitionActorDesc* ActorDesc);
 	void UnhashActorDesc(FWorldPartitionActorDesc* ActorDesc);
-
-	void AddToClusters(const FWorldPartitionActorDesc* ActorDesc);
-	void RemoveFromClusters(const FWorldPartitionActorDesc* ActorDesc);
-
-	void AddToPartition(FWorldPartitionActorDesc* ActorDesc);
-	void RemoveFromPartition(FWorldPartitionActorDesc* ActorDesc, bool bRemoveDescriptorFromArray = true, bool bUnloadRemovedDescriptor = true);
-
-	TMap<FGuid, UWorldPartition::FActorCluster*> ActorToActorCluster;
-	TSet<UWorldPartition::FActorCluster*> ActorClustersSet;
 #endif
 
 	UWorldPartitionStreamingPolicy* GetStreamingPolicy() const;
diff --git a/Engine/Source/Runtime/Engine/Public/WorldPartition/WorldPartitionActorDesc.h b/Engine/Source/Runtime/Engine/Public/WorldPartition/WorldPartitionActorDesc.h
--- a/Engine/Source/Runtime/Engine/Public/WorldPartition/WorldPartitionActorDesc.h
+++ b/Engine/Source/Runtime/Engine/Public/WorldPartition/WorldPartitionActorDesc.h
@@ -8,7 +8,6 @@
 #include "GameFramework/Actor.h"
 #include "Containers/Set.h"
 #include "Misc/Guid.h"
-#include "Misc/HashBuilder.h"
 #include "WorldPartition/WorldPartitionActorDescFactory.h"
 
 #if WITH_EDITOR
@@ -18,7 +17,6 @@ struct FWorldPartitionActorDescInitData
 	UClass* NativeClass;
 	FName PackageName;
 	FName ActorPath;
-	FTransform Transform;
 	TArray<uint8> SerializedData;
 };
 #endif
@@ -64,11 +62,6 @@ public:
 		return LoadedRefCount;
 	}
 
-	inline uint32 GetHash() const
-	{
-		return Hash;
-	}
-
 	const TArray<FGuid>& GetReferences() const
 	{
 		return References;
@@ -80,17 +73,16 @@ public:
 	AActor* Load(const FLinkerInstancingContext* InstancingContext = nullptr);
 	void Unload();
 
-	void Init(const AActor* InActor);
+	virtual void Init(const AActor* InActor);
 	void Init(const FWorldPartitionActorDescInitData& DescData);
 
 	void SerializeTo(TArray<uint8>& OutData);
 
+	void TransformInstance(const FString& From, const FString& To, const FTransform& Transform);
+
 protected:
 	FWorldPartitionActorDesc();
 	
-	void UpdateHash();
-
-	virtual void InitFrom(const AActor* Actor);
 	virtual void Serialize(FArchive& Ar);
 
 	FGuid						Guid;
@@ -107,7 +99,6 @@ protected:
 	TArray<FGuid>				References;
 	
 	mutable uint32				LoadedRefCount;
-	mutable uint32				Hash;
 
 	// Cached values
 	UClass*						ActorClass;
diff --git a/Engine/Source/Runtime/Engine/Public/WorldPartition/WorldPartitionSubsystem.h b/Engine/Source/Runtime/Engine/Public/WorldPartition/WorldPartitionSubsystem.h
--- a/Engine/Source/Runtime/Engine/Public/WorldPartition/WorldPartitionSubsystem.h
+++ b/Engine/Source/Runtime/Engine/Public/WorldPartition/WorldPartitionSubsystem.h
@@ -45,7 +45,6 @@ public:
 #if WITH_EDITOR
 	void ForEachIntersectingActorDesc(const FBox& Box, TSubclassOf<AActor> ActorClass, TFunctionRef<bool(const FWorldPartitionActorDesc*)> Predicate) const;
 	void ForEachActorDesc(TSubclassOf<AActor> ActorClass, TFunctionRef<bool(const FWorldPartitionActorDesc*)> Predicate) const;
-	void UpdateActorDesc(AActor* Actor);
 
 	void RegisterActorDescFactory(TSubclassOf<AActor> Class, FWorldPartitionActorDescFactory* Factory);
 
