Fix non-editor builds.

diff --git a/Engine/Source/Runtime/Engine/Private/WorldPartition/WorldPartitionRuntimeSpatialHash.cpp b/Engine/Source/Runtime/Engine/Private/WorldPartition/WorldPartitionRuntimeSpatialHash.cpp
--- a/Engine/Source/Runtime/Engine/Private/WorldPartition/WorldPartitionRuntimeSpatialHash.cpp
+++ b/Engine/Source/Runtime/Engine/Private/WorldPartition/WorldPartitionRuntimeSpatialHash.cpp
@@ -106,6 +106,7 @@ static FAutoConsoleVariableRef CVarRuntimeSpatialHashCellToSourceAngleContributi
 
 // ------------------------------------------------------------------------------------------------
 
+#if WITH_EDITOR
 // Clustering
 struct FActorCluster
 {
@@ -244,6 +245,7 @@ TArray<FActorCluster> CreateActorClusters(UWorldPartition* WorldPartition)
 	for (FActorCluster* ActorCluster : ActorClustersSet) { delete ActorCluster; }
 	return ActorClusters;
 }
+#endif
 
 /**
   * Square 2D grid helper
