Work arounds to fix WITH_UNREAL_DEVELOPER_TOOLS == 0 builds

diff --git a/Engine/Source/Runtime/Engine/Private/DeviceProfiles/DeviceProfileManager.cpp b/Engine/Source/Runtime/Engine/Private/DeviceProfiles/DeviceProfileManager.cpp
--- a/Engine/Source/Runtime/Engine/Private/DeviceProfiles/DeviceProfileManager.cpp
+++ b/Engine/Source/Runtime/Engine/Private/DeviceProfiles/DeviceProfileManager.cpp
@@ -137,11 +137,15 @@ void UDeviceProfileManager::ProcessDeviceProfileIniSettings(const FString& Devic
 	FConfigCacheIni* ConfigSystem = GConfig;
 	if (Mode == EDeviceProfileMode::DPM_CacheValues)
 	{
+#if WITH_UNREAL_DEVELOPER_TOOLS
 		// caching is not done super early, so we can assume DPs have been found now
 		UDeviceProfile* Profile = UDeviceProfileManager::Get().FindProfile(DeviceProfileName, false);
 		check(Profile);
 		// use the DP's platform's configs, NOT the running platform
 		ConfigSystem = FConfigCacheIni::ForPlatform(*Profile->DeviceType);
+#else
+		UE_LOG(LogDeviceProfileManager, Fatal, TEXT("ProcessDeviceProfileIniSettings called with Mode == DPM_CacheValues in non-tools build!"));
+#endif
 	}
 
 	check(ConfigSystem != nullptr);
@@ -385,12 +389,14 @@ void UDeviceProfileManager::ProcessDeviceProfileIniSettings(const FString& Devic
 		bReachedEndOfTree = !bProfileExists || BaseDeviceProfileName.IsEmpty();
 	}
 
+#if WITH_UNREAL_DEVELOPER_TOOLS
 	// copy the running cache into the DP
 	if (Mode == EDeviceProfileMode::DPM_CacheValues)
 	{
 		UDeviceProfile* Profile = UDeviceProfileManager::Get().FindProfile(DeviceProfileName, false);
 		Profile->AddExpandedCVars(CVarsAlreadySetList);
 	}
+#endif
 }
 
 
diff --git a/Engine/Source/Runtime/Engine/Private/UnrealEngine.cpp b/Engine/Source/Runtime/Engine/Private/UnrealEngine.cpp
--- a/Engine/Source/Runtime/Engine/Private/UnrealEngine.cpp
+++ b/Engine/Source/Runtime/Engine/Private/UnrealEngine.cpp
@@ -8303,10 +8303,12 @@ bool UEngine::HandleGetIniCommand(const TCHAR* Cmd, FOutputDevice& Ar)
 		if (PlatformDelim > 0)
 		{
 			FString PlatformName = IniPlusSection.Mid(0, PlatformDelim);
+#if WITH_UNREAL_DEVELOPER_TOOLS
 			IniPlusSection = IniPlusSection.Mid(PlatformDelim + 1);
 
 			ConfigSystem = FConfigCacheIni::ForPlatform(*PlatformName);
 			if (ConfigSystem == nullptr)
+#endif
 			{
 				Ar.Logf(TEXT("Failed to find config system for platform %s"), *PlatformName);
 				return false;
