- Cleaned up the other platform configs and cvars stuff into its own define (ALLOW_OTHER_PLATFORM_CONFIG)

diff --git a/Engine/Source/Runtime/Core/Private/HAL/ConsoleManager.cpp b/Engine/Source/Runtime/Core/Private/HAL/ConsoleManager.cpp
--- a/Engine/Source/Runtime/Core/Private/HAL/ConsoleManager.cpp
+++ b/Engine/Source/Runtime/Core/Private/HAL/ConsoleManager.cpp
@@ -314,30 +314,14 @@ public:
 	FConsoleVariable(T DefaultValue, const TCHAR* Help, EConsoleVariableFlags Flags, bool bSaveDefault=true) 
 		: FConsoleVariableBase(Help, Flags), Data(DefaultValue)
 	{
-#if WITH_UNREAL_DEVELOPER_TOOLS
-		PlatformIndependentDefault = nullptr;
+#if ALLOW_OTHER_PLATFORM_CONFIG
 		if (bSaveDefault)
 		{
-			CreatePlatformIndependentDefault(DefaultValue);
+			PlatformIndependentDefault = TSharedPtr<IConsoleVariable>(new FConsoleVariable<T>(DefaultValue, TEXT(""), Flags, false));
 		}
 #endif
 	}
 
-	virtual ~FConsoleVariable()
-	{
-#if WITH_UNREAL_DEVELOPER_TOOLS
-		delete PlatformIndependentDefault;
-#endif
-	}
-
-#if WITH_UNREAL_DEVELOPER_TOOLS
-
-	void CreatePlatformIndependentDefault(T DefaultValue)
-	{
-		PlatformIndependentDefault = new FConsoleVariable<T>(DefaultValue, TEXT(""), Flags, false);
-	}
-#endif
-
 	// interface IConsoleVariable ----------------------------------- 
 
 	virtual void Release()
@@ -387,13 +371,13 @@ private: // ----------------------------------------------------
 		FConsoleVariableBase::OnChanged(SetBy);
 	}
 
-#if WITH_UNREAL_DEVELOPER_TOOLS
+#if ALLOW_OTHER_PLATFORM_CONFIG
 	// remember the dfefault value of this cvar, before anything else can assign to it - this way we can know in the editor what the value
 	// would be for a cvar on another platform, if no .ini file on that platform sets it
-	IConsoleVariable* PlatformIndependentDefault;
+	TSharedPtr<IConsoleVariable> PlatformIndependentDefault;
 	virtual IConsoleVariable* GetDefaultValueVariable() override
 	{
-		return PlatformIndependentDefault;
+		return PlatformIndependentDefault.Get();
 	}
 #endif
 
diff --git a/Engine/Source/Runtime/Core/Private/Misc/ConfigCacheIni.cpp b/Engine/Source/Runtime/Core/Private/Misc/ConfigCacheIni.cpp
--- a/Engine/Source/Runtime/Core/Private/Misc/ConfigCacheIni.cpp
+++ b/Engine/Source/Runtime/Core/Private/Misc/ConfigCacheIni.cpp
@@ -5116,10 +5116,11 @@ void DeleteRecordedConfigReadsFromIni()
 
 
 
-
+#if ALLOW_OTHER_PLATFORM_CONFIG
 static TMap<FName, FConfigCacheIni> GConfigForPlatform;
 static TFuture<void> GConfigFuture;
 static FCriticalSection GConfigForPlatformLock;
+#endif
 
 #if WITH_EDITOR
 void FConfigCacheIni::AsyncInitializeConfigForPlatforms()
@@ -5143,10 +5144,10 @@ void FConfigCacheIni::AsyncInitializeConfigForPlatforms()
 
 #endif
 
-#if WITH_UNREAL_DEVELOPER_TOOLS
 
 FConfigCacheIni* FConfigCacheIni::ForPlatform(FName PlatformName)
 {
+#if ALLOW_OTHER_PLATFORM_CONFIG
 	// use GConfig when no platform is specified
 	if (PlatformName == NAME_None)
 	{
@@ -5177,15 +5178,20 @@ FConfigCacheIni* FConfigCacheIni::ForPlatform(FName PlatformName)
 	}
 
 	return GConfigForPlatform.Find(PlatformName);
+#else
+	UE_LOG(LogConfig, Error, TEXT("FConfigCacheIni::ForPlatform cannot be called when not in a developer tool"));
+	return nullptr;
+#endif
 }
 
 void FConfigCacheIni::ClearOtherPlatformConfigs()
 {
+#if ALLOW_OTHER_PLATFORM_CONFIG
+
 	FScopeLock Lock(&GConfigForPlatformLock);
 
 	checkf(!GConfigFuture.IsValid(), TEXT("Trying to clear platform configs while still reading them in is unexpected"));
 
 	GConfigForPlatform.Empty();
-}
-
 #endif
+}
diff --git a/Engine/Source/Runtime/Core/Public/HAL/IConsoleManager.h b/Engine/Source/Runtime/Core/Public/HAL/IConsoleManager.h
--- a/Engine/Source/Runtime/Core/Public/HAL/IConsoleManager.h
+++ b/Engine/Source/Runtime/Core/Public/HAL/IConsoleManager.h
@@ -431,15 +431,13 @@ public:
 
 	virtual FConsoleVariableMulticastDelegate& OnChangedDelegate() = 0;
 
-#if WITH_UNREAL_DEVELOPER_TOOLS
 	/**
-	 * Get the saved off default value, in a cvar variable
+	 * Get the saved off default value, in a cvar variable, if one was created
 	 */
 	virtual IConsoleVariable* GetDefaultValueVariable()
 	{
 		return nullptr;
 	}
-#endif
 
 	// convenience methods
 
diff --git a/Engine/Source/Runtime/Core/Public/Misc/Build.h b/Engine/Source/Runtime/Core/Public/Misc/Build.h
--- a/Engine/Source/Runtime/Core/Public/Misc/Build.h
+++ b/Engine/Source/Runtime/Core/Public/Misc/Build.h
@@ -438,7 +438,10 @@
 	#define GET_DEVICE_ID_UNAVAILABLE 0
 #endif
 
-// CONtrols whether or not to make a global object to load COnfig.bin as soon as possible
+// Controls whether or not to make a global object to load COnfig.bin as soon as possible
 #ifndef PRELOAD_BINARY_CONFIG
 	#define PRELOAD_BINARY_CONFIG 1
-#endif
\ No newline at end of file
+#endif
+
+// Controls if the config system can stores configs for other platforms than the running one
+#define ALLOW_OTHER_PLATFORM_CONFIG		WITH_UNREAL_DEVELOPER_TOOLS
diff --git a/Engine/Source/Runtime/Core/Public/Misc/ConfigCacheIni.h b/Engine/Source/Runtime/Core/Public/Misc/ConfigCacheIni.h
--- a/Engine/Source/Runtime/Core/Public/Misc/ConfigCacheIni.h
+++ b/Engine/Source/Runtime/Core/Public/Misc/ConfigCacheIni.h
@@ -1019,7 +1019,6 @@ public:
 	 */
 	static bool CreateGConfigFromSaved(const TCHAR* Filename);
 
-#if WITH_UNREAL_DEVELOPER_TOOLS
 	/**
 	 * Retrieve the fully processed ini system for another platform. The editor will start loading these 
 	 * in the background on startup
@@ -1030,7 +1029,6 @@ public:
 	 * Wipe all cached platform configs. Next ForPlatform call will load on-demand the platform configs
 	 */
 	static void ClearOtherPlatformConfigs();
-#endif
 
 private:
 #if WITH_EDITOR
diff --git a/Engine/Source/Runtime/CoreUObject/Private/UObject/Obj.cpp b/Engine/Source/Runtime/CoreUObject/Private/UObject/Obj.cpp
--- a/Engine/Source/Runtime/CoreUObject/Private/UObject/Obj.cpp
+++ b/Engine/Source/Runtime/CoreUObject/Private/UObject/Obj.cpp
@@ -2645,7 +2645,6 @@ FString UObject::GetProjectUserConfigFilename() const
 // @todo ini: Verify per object config objects
 void UObject::UpdateSingleSectionOfConfigFile(const FString& ConfigIniName)
 {
-#if WITH_UNREAL_DEVELOPER_TOOLS
 	// create a sandbox FConfigCache
 	FConfigCacheIni Config(EConfigCacheType::Temporary);
 
@@ -2668,9 +2667,6 @@ void UObject::UpdateSingleSectionOfConfigFile(const FString& ConfigIniName)
 		FString FinalIniFileName;
 		GConfig->LoadGlobalIniFile(FinalIniFileName, *GetClass()->ClassConfigName.ToString(), NULL, true);
 	}
-#else
-	checkNoEntry();
-#endif
 }
 
 void UObject::UpdateDefaultConfigFile(const FString& SpecificFileLocation)
diff --git a/Engine/Source/Runtime/Engine/Classes/DeviceProfiles/DeviceProfile.h b/Engine/Source/Runtime/Engine/Classes/DeviceProfiles/DeviceProfile.h
--- a/Engine/Source/Runtime/Engine/Classes/DeviceProfiles/DeviceProfile.h
+++ b/Engine/Source/Runtime/Engine/Classes/DeviceProfiles/DeviceProfile.h
@@ -148,7 +148,7 @@ private:
 
 #endif
 
-#if WITH_UNREAL_DEVELOPER_TOOLS
+#if ALLOW_OTHER_PLATFORM_CONFIG
 public:
 	const TMap<FString, FString>& GetAllExpandedCVars();
 	void AddExpandedCVars(const TMap<FString, FString>& CVarsToMerge);
diff --git a/Engine/Source/Runtime/Engine/Classes/DeviceProfiles/DeviceProfileManager.h b/Engine/Source/Runtime/Engine/Classes/DeviceProfiles/DeviceProfileManager.h
--- a/Engine/Source/Runtime/Engine/Classes/DeviceProfiles/DeviceProfileManager.h
+++ b/Engine/Source/Runtime/Engine/Classes/DeviceProfiles/DeviceProfileManager.h
@@ -142,7 +142,7 @@ public:
 	static bool GetScalabilityCVar(const FString& CvarName, int32& OutValue);
 	static bool GetScalabilityCVar(const FString& CvarName, float& OutValue);
 
-#if WITH_UNREAL_DEVELOPER_TOOLS
+#if ALLOW_OTHER_PLATFORM_CONFIG
 	/** Retrieve another platform's DeviceProfile, and make sure it's CVars are expanded, including scalability */
 	static void ExpandDeviceProfileCVars(UDeviceProfile* DeviceProfile);
 #endif
diff --git a/Engine/Source/Runtime/Engine/Private/DeviceProfiles/DeviceProfile.cpp b/Engine/Source/Runtime/Engine/Private/DeviceProfiles/DeviceProfile.cpp
--- a/Engine/Source/Runtime/Engine/Private/DeviceProfiles/DeviceProfile.cpp
+++ b/Engine/Source/Runtime/Engine/Private/DeviceProfiles/DeviceProfile.cpp
@@ -367,7 +367,7 @@ const TMap<FString, FString>& UDeviceProfile::GetConsolidatedCVars() const
 #endif
 
 
-#if WITH_UNREAL_DEVELOPER_TOOLS
+#if ALLOW_OTHER_PLATFORM_CONFIG
 const TMap<FString, FString>& UDeviceProfile::GetAllExpandedCVars()
 {
 	// expand on first use
diff --git a/Engine/Source/Runtime/Engine/Private/DeviceProfiles/DeviceProfileManager.cpp b/Engine/Source/Runtime/Engine/Private/DeviceProfiles/DeviceProfileManager.cpp
--- a/Engine/Source/Runtime/Engine/Private/DeviceProfiles/DeviceProfileManager.cpp
+++ b/Engine/Source/Runtime/Engine/Private/DeviceProfiles/DeviceProfileManager.cpp
@@ -137,18 +137,18 @@ void UDeviceProfileManager::ProcessDeviceProfileIniSettings(const FString& Devic
 	FConfigCacheIni* ConfigSystem = GConfig;
 	if (Mode == EDeviceProfileMode::DPM_CacheValues)
 	{
-#if WITH_UNREAL_DEVELOPER_TOOLS
+#if ALLOW_OTHER_PLATFORM_CONFIG
 		// caching is not done super early, so we can assume DPs have been found now
 		UDeviceProfile* Profile = UDeviceProfileManager::Get().FindProfile(DeviceProfileName, false);
 		check(Profile);
 		// use the DP's platform's configs, NOT the running platform
 		ConfigSystem = FConfigCacheIni::ForPlatform(*Profile->DeviceType);
 #else
-		UE_LOG(LogDeviceProfileManager, Fatal, TEXT("ProcessDeviceProfileIniSettings called with Mode == DPM_CacheValues in non-tools build!"));
+		checkNoEntry();
 #endif
 	}
 
-	check(ConfigSystem != nullptr);
+	check(ConfigSystem);
 
 	TArray< FString > AvailableProfiles;
 	GConfig->GetSectionNames( GDeviceProfilesIni, AvailableProfiles );
@@ -389,7 +389,7 @@ void UDeviceProfileManager::ProcessDeviceProfileIniSettings(const FString& Devic
 		bReachedEndOfTree = !bProfileExists || BaseDeviceProfileName.IsEmpty();
 	}
 
-#if WITH_UNREAL_DEVELOPER_TOOLS
+#if ALLOW_OTHER_PLATFORM_CONFIG
 	// copy the running cache into the DP
 	if (Mode == EDeviceProfileMode::DPM_CacheValues)
 	{
@@ -417,7 +417,7 @@ void UDeviceProfileManager::InitializeCVarsForActiveDeviceProfile(bool bPushSett
 	ProcessDeviceProfileIniSettings(ActiveProfileName, bPushSettings ? EDeviceProfileMode::DPM_PushCVars : EDeviceProfileMode::DPM_SetCVars);
 }
 
-#if WITH_UNREAL_DEVELOPER_TOOLS
+#if ALLOW_OTHER_PLATFORM_CONFIG
 void UDeviceProfileManager::ExpandDeviceProfileCVars(UDeviceProfile* DeviceProfile)
 {
 	// get the config system for the platform the DP uses
@@ -815,7 +815,7 @@ void UDeviceProfileManager::SetOverrideDeviceProfile(UDeviceProfile* DeviceProfi
 	// for preview, we assume this will be another platform's DP, so use the resolved cvars directly, bypassing the activate and set stuff
 	if (bIsDeviceProfilePreview)
 	{
-#if WITH_UNREAL_DEVELOPER_TOOLS
+#if ALLOW_OTHER_PLATFORM_CONFIG
 		for (const auto& Pair : DeviceProfile->GetAllExpandedCVars())
 		{
 			IConsoleVariable* CVar = IConsoleManager::Get().FindConsoleVariable(*Pair.Key);
@@ -831,7 +831,7 @@ void UDeviceProfileManager::SetOverrideDeviceProfile(UDeviceProfile* DeviceProfi
 			}
 		}
 #else
-		UE_LOG(LogDeviceProfileManager, Error, TEXT("SetOverrideDeviceProfile can only be used in a developer tool"));
+		UE_LOG(LogDeviceProfileManager, Error, TEXT("SetOverrideDeviceProfile with bIsDeviceProfilePreview=true can only be used in a developer tool"));
 #endif
 		return;
 	}
@@ -1046,7 +1046,7 @@ void UDeviceProfileManager::GetAllPossibleParentProfiles(const UDeviceProfile* C
 
 
 
-#if WITH_UNREAL_DEVELOPER_TOOLS
+#if ALLOW_OTHER_PLATFORM_CONFIG
 static bool GetCVarForPlatform( FOutputDevice& Ar, FString DPName, FString CVarName)
 {
 	UDeviceProfile* DeviceProfile = UDeviceProfileManager::Get().FindProfile(DPName, false);
diff --git a/Engine/Source/Runtime/Engine/Private/UnrealEngine.cpp b/Engine/Source/Runtime/Engine/Private/UnrealEngine.cpp
--- a/Engine/Source/Runtime/Engine/Private/UnrealEngine.cpp
+++ b/Engine/Source/Runtime/Engine/Private/UnrealEngine.cpp
@@ -8303,12 +8303,10 @@ bool UEngine::HandleGetIniCommand(const TCHAR* Cmd, FOutputDevice& Ar)
 		if (PlatformDelim > 0)
 		{
 			FString PlatformName = IniPlusSection.Mid(0, PlatformDelim);
-#if WITH_UNREAL_DEVELOPER_TOOLS
 			IniPlusSection = IniPlusSection.Mid(PlatformDelim + 1);
 
 			ConfigSystem = FConfigCacheIni::ForPlatform(*PlatformName);
 			if (ConfigSystem == nullptr)
-#endif
 			{
 				Ar.Logf(TEXT("Failed to find config system for platform %s"), *PlatformName);
 				return false;
