Make sure that FastArrays using Struct Delta Serialization always resend unacknowledged properties.
This fixes a regression in behavior where structs may receive partial updates due to NAKs.

diff --git a/Engine/Source/Runtime/CoreUObject/Public/UObject/CoreNet.h b/Engine/Source/Runtime/CoreUObject/Public/UObject/CoreNet.h
--- a/Engine/Source/Runtime/CoreUObject/Public/UObject/CoreNet.h
+++ b/Engine/Source/Runtime/CoreUObject/Public/UObject/CoreNet.h
@@ -218,21 +218,23 @@ struct FPropertyRetirement
 #if !UE_BUILD_SHIPPING
 	static const uint32 ExpectedSanityTag = 0xDF41C9A3;
 
-	uint32			SanityTag;
+	uint32 SanityTag;
 #endif
 
-	FPropertyRetirement * Next;
+	FPropertyRetirement* Next;
 
 	TSharedPtr<class INetDeltaBaseState> DynamicState;
 
-	FPacketIdRange	OutPacketIdRange;
+	FPacketIdRange OutPacketIdRange;
+	uint32 FastArrayChangelistHistory;
 
 	FPropertyRetirement() :
 #if !UE_BUILD_SHIPPING
-			SanityTag( ExpectedSanityTag ),
+		SanityTag( ExpectedSanityTag ),
 #endif
-			Next( nullptr )
-		,	DynamicState ( nullptr )
+		 Next(nullptr)
+		, DynamicState(nullptr)
+		, FastArrayChangelistHistory(0)
 	{}
 
 	void CountBytes(FArchive& Ar) const;
diff --git a/Engine/Source/Runtime/Engine/Classes/Engine/NetSerialization.h b/Engine/Source/Runtime/Engine/Classes/Engine/NetSerialization.h
--- a/Engine/Source/Runtime/Engine/Classes/Engine/NetSerialization.h
+++ b/Engine/Source/Runtime/Engine/Classes/Engine/NetSerialization.h
@@ -323,9 +323,10 @@ class FNetFastTArrayBaseState : public INetDeltaBaseState
 {
 public:
 
-	FNetFastTArrayBaseState():
-		ArrayReplicationKey(INDEX_NONE),
-		ChangelistHistory(0)
+	FNetFastTArrayBaseState()
+		: ArrayReplicationKey(INDEX_NONE)
+		, ChangelistHistory(0)
+		, LastAckedHistory(0)
 	{}
 
 	virtual bool IsStateEqual(INetDeltaBaseState* OtherState)
@@ -348,6 +349,8 @@ public:
 	int32 ArrayReplicationKey;
 
 	uint32 ChangelistHistory;
+
+	uint32 LastAckedHistory;
 };
 
 
@@ -1630,6 +1633,7 @@ bool FFastArraySerializer::FastArrayDeltaSerialize_DeltaSerializeStructs(TArray<
 		TMap<int32, int32>* OldItemMap = nullptr;
 		int32 BaseReplicationKey = INDEX_NONE;
 		int32 OldChangelistHistory = INDEX_NONE;
+		int32 OldLastAckedHistory = INDEX_NONE;
 
 		// See if the array changed at all. If the ArrayReplicationKey matches we can skip checking individual items
 		if (Parms.OldState)
@@ -1637,6 +1641,7 @@ bool FFastArraySerializer::FastArrayDeltaSerialize_DeltaSerializeStructs(TArray<
 			OldItemMap = &((FNetFastTArrayBaseState*)Parms.OldState)->IDToCLMap;
 			BaseReplicationKey = ((FNetFastTArrayBaseState*)Parms.OldState)->ArrayReplicationKey;
 			OldChangelistHistory = ((FNetFastTArrayBaseState*)Parms.OldState)->ChangelistHistory;
+			OldLastAckedHistory = ((FNetFastTArrayBaseState*)Parms.OldState)->LastAckedHistory;
 
 			if (!Helper.ConditionalCreateNewDeltaState(*OldItemMap, BaseReplicationKey))
 			{
@@ -1677,6 +1682,7 @@ bool FFastArraySerializer::FastArrayDeltaSerialize_DeltaSerializeStructs(TArray<
 		Helper.WriteDeltaHeader(Header);
 
 		NewState->ChangelistHistory = OldChangelistHistory;
+		NewState->LastAckedHistory = OldLastAckedHistory;
 
 		DeltaSerializeParams.WriteChangedElements = &ChangedElements;
 		DeltaSerializeParams.WriteBaseState = NewState;
diff --git a/Engine/Source/Runtime/Engine/Private/DataReplication.cpp b/Engine/Source/Runtime/Engine/Private/DataReplication.cpp
--- a/Engine/Source/Runtime/Engine/Private/DataReplication.cpp
+++ b/Engine/Source/Runtime/Engine/Private/DataReplication.cpp
@@ -222,6 +222,13 @@ public:
 public:
 
 	// These can go away once we do a full merge of Custom Delta and RepLayout.
+	static bool IsAFastArray(
+		const FRepLayout& RepLayout,
+		const uint16 CustomDeltaIndex)
+	{
+		return RepLayout.IsAFastArray(CustomDeltaIndex);
+	}
+
 
 	static bool SendCustomDeltaProperty(
 		const FRepLayout& RepLayout,
@@ -832,6 +839,7 @@ void FObjectReplicator::ReceivedNak( int32 NakPacketId )
 
 				// If this is a dynamic array property, we have to look through the list of retirement records to see if we need to reset the base state
 				FPropertyRetirement* Rec = Retirement.Next; // Retirement[i] is head and not actually used in this case
+				uint32 LastAcknowledged = 0;
 				while (Rec != nullptr)
 				{
 					if (NakPacketId > Rec->OutPacketIdRange.Last)
@@ -839,6 +847,7 @@ void FObjectReplicator::ReceivedNak( int32 NakPacketId )
 						// We can assume this means this record's packet was ack'd, so we can get rid of the old state
 						check(Retirement.Next == Rec);
 						Retirement.Next = Rec->Next;
+						LastAcknowledged = Rec->FastArrayChangelistHistory;
 						delete Rec;
 						Rec = Retirement.Next;
 						continue;
@@ -849,6 +858,12 @@ void FObjectReplicator::ReceivedNak( int32 NakPacketId )
 
 						// The Nack'd packet did update this property, so we need to replace the buffer in RecentDynamic
 						// with the buffer we used to create this update (which was dropped), so that the update will be recreated on the next replicate actor
+						
+						if (LastAcknowledged != 0 && Rec->DynamicState.IsValid() && FNetSerializeCB::IsAFastArray(*RepLayout, (uint16)i))
+						{
+							((FNetFastTArrayBaseState*)Rec->DynamicState.Get())->LastAckedHistory = LastAcknowledged;
+						}
+
 						SendingRepState->RecentCustomDeltaState[i] = Rec->DynamicState;
 
 						// We can get rid of the rest of the saved off base states since we will be regenerating these updates on the next replicate actor
@@ -1401,6 +1416,7 @@ void FObjectReplicator::PostReceivedBunch()
 
 static FORCEINLINE FPropertyRetirement** UpdateAckedRetirements(
 	FPropertyRetirement& Retire,
+	uint32& LastAcknowledged,
 	const int32 OutAckPacketId,
 	const UObject* Object)
 {
@@ -1408,6 +1424,7 @@ static FORCEINLINE FPropertyRetirement** UpdateAckedRetirements(
 
 	FPropertyRetirement** Rec = &Retire.Next; // Note the first element is 'head' that we don't actually use
 
+	LastAcknowledged = 0;
 	while (*Rec != nullptr)
 	{
 		if (OutAckPacketId >= (*Rec)->OutPacketIdRange.Last)
@@ -1420,6 +1437,7 @@ static FORCEINLINE FPropertyRetirement** UpdateAckedRetirements(
 			Retire.Next = ToDelete->Next;
 			Rec = &Retire.Next;
 
+			LastAcknowledged = ToDelete->FastArrayChangelistHistory;
 			delete ToDelete;
 			continue;
 		}
@@ -1526,6 +1544,7 @@ void FObjectReplicator::ReplicateCustomDeltaProperties( FNetBitWriter & Bunch, F
 			continue;
 		}
 
+		const bool bIsAFastArray = FNetSerializeCB::IsAFastArray(*RepLayout, CustomDeltaProperty);
 		FPropertyRetirement** LastNext = nullptr;
 
 		if (!bIsConnectionInternalAck)
@@ -1535,7 +1554,13 @@ void FObjectReplicator::ReplicateCustomDeltaProperties( FNetBitWriter & Bunch, F
 
 			// Update Retirement records with this new state so we can handle packet drops.
 			// LastNext will be pointer to the last "Next" pointer in the list (so pointer to a pointer)
-			LastNext = UpdateAckedRetirements(Retire, Connection->OutAckPacketId, Object);
+			uint32 LastAcknowledged = 0;
+			LastNext = UpdateAckedRetirements(Retire, LastAcknowledged, Connection->OutAckPacketId, Object);
+
+			if (LastAcknowledged != 0 && OldState.IsValid() && bIsAFastArray)
+			{
+				((FNetFastTArrayBaseState*)OldState.Get())->LastAckedHistory = LastAcknowledged;
+			}
 
 			check(LastNext != nullptr);
 			check(*LastNext == nullptr);
@@ -1559,6 +1584,15 @@ void FObjectReplicator::ReplicateCustomDeltaProperties( FNetBitWriter & Bunch, F
 
 			// Remember what the old state was at this point in time.  If we get a nak, we will need to revert back to this.
 			(*LastNext)->DynamicState = OldState;
+
+			// Using NewState's ChangelistHistory seems counter intuitive at first, because it *seems* like we should be using the history from the OldState.
+			// However, PropertyRetirements associate the OldState with the state of replication **in case of a NAK**.
+			// So, in the case of an ACK, we know that we no longer need to hold onto that OldState because the client has received
+			// the NewState (which will become our OldState the next time we send something).
+			if (NewState.IsValid() && bIsAFastArray)
+			{
+				(*LastNext)->FastArrayChangelistHistory = ((FNetFastTArrayBaseState*)NewState.Get())->ChangelistHistory;
+			}
 		}
 
 		// Save NewState into the RecentCustomDeltaState array (old state is a reference into our RecentCustomDeltaState map)
diff --git a/Engine/Source/Runtime/Engine/Private/RepLayout.cpp b/Engine/Source/Runtime/Engine/Private/RepLayout.cpp
--- a/Engine/Source/Runtime/Engine/Private/RepLayout.cpp
+++ b/Engine/Source/Runtime/Engine/Private/RepLayout.cpp
@@ -7177,8 +7177,8 @@ ERepLayoutResult FRepLayout::DeltaSerializeFastArrayProperty(FFastArrayDeltaSeri
 				// Cache off the newest history, the last history sent to this connection, and then update the state
 				// to notify that we're going to send it the newest history.
 				const uint32 NewChangelistHistory = FastArrayState.HistoryEnd;
-				const uint32 LastSentHistory = NewArrayDeltaState->ChangelistHistory;
-				const uint32 LastSentChangelistDelta = NewChangelistHistory - LastSentHistory;
+				const uint32 LastAckedHistory = NewArrayDeltaState->LastAckedHistory;
+				const uint32 LastAckedChangelistDelta = NewChangelistHistory - LastAckedHistory;
 
 				NewArrayDeltaState->ChangelistHistory = NewChangelistHistory;
 
@@ -7425,16 +7425,18 @@ ERepLayoutResult FRepLayout::DeltaSerializeFastArrayProperty(FFastArrayDeltaSeri
 				// and send all changes (happens in the block below, Changelists.Num() == 0).
 
 				// Note, this won't be all changes since the beginning, but just all changes for the currently dirty items.
+				const uint32 LastHistory = LastAckedHistory;
+				const uint32 LastChangelistDelta = LastAckedChangelistDelta;
 
-				if (LastSentHistory != 0 && LastSentChangelistDelta > 0 && LastSentChangelistDelta < (FDeltaArrayHistoryState::MAX_CHANGE_HISTORY - 1))
+				if (LastHistory != 0 && LastChangelistDelta > 0 && LastChangelistDelta < (FDeltaArrayHistoryState::MAX_CHANGE_HISTORY - 1))
 				{
 					const FConstRepObjectDataBuffer ConstObjectData(ObjectData);
 					Changelists.SetNum(ChangedElements.Num());
 
 					// Note, we iterate from LastSentHistory + 1, because we don't want to send something if
-					// we think its already been sent.
+					// we think it's already been sent/received.
 					// Similarly, we do <= NewChangelistHistory because we need to send the newest history.
-					for (uint32 ChangelistHistory = LastSentHistory + 1; ChangelistHistory <= NewChangelistHistory; ++ChangelistHistory)
+					for (uint32 ChangelistHistory = LastHistory + 1; ChangelistHistory <= NewChangelistHistory; ++ChangelistHistory)
 					{
 						const uint32 RelativeHistory = ChangelistHistory % FDeltaArrayHistoryState::MAX_CHANGE_HISTORY;
 						FDeltaArrayHistoryItem& HistoryItem = FastArrayState.ChangeHistory[RelativeHistory];
@@ -7839,6 +7841,12 @@ const ELifetimeCondition FRepLayout::GetLifetimeCustomDeltaPropertyCondition(con
 	return Parents[CustomDeltaProperty.PropertyRepIndex].Condition;
 }
 
+bool FRepLayout::IsAFastArray(const uint16 CustomDeltaPropertyIndex) const
+{
+	const FLifetimeCustomDeltaProperty& CustomDeltaProperty = LifetimeCustomPropertyState->GetCustomDeltaProperty(CustomDeltaPropertyIndex);
+	return CustomDeltaProperty.FastArrayNumber != INDEX_NONE;
+}
+
 void FReceivingRepState::CountBytes(FArchive& Ar) const
 {
 	GRANULAR_NETWORK_MEMORY_TRACKING_INIT(Ar, "FReceivingRepState::CountBytes");
diff --git a/Engine/Source/Runtime/Engine/Public/Net/RepLayout.h b/Engine/Source/Runtime/Engine/Public/Net/RepLayout.h
--- a/Engine/Source/Runtime/Engine/Public/Net/RepLayout.h
+++ b/Engine/Source/Runtime/Engine/Public/Net/RepLayout.h
@@ -1777,6 +1777,8 @@ private:
 
 	const ELifetimeCondition GetLifetimeCustomDeltaPropertyCondition(const uint16 RepIndCustomDeltaPropertyIndexex) const;
 
+	bool IsAFastArray(const uint16 CustomDeltaPropertyIndex) const;
+
 	ERepLayoutFlags Flags;
 
 	/** Size (in bytes) needed to allocate a single instance of a Shadow buffer for this RepLayout. */
