Defer skylight capture when affected by async texture compilation

diff --git a/Engine/Source/Runtime/Engine/Private/Components/SkyLightComponent.cpp b/Engine/Source/Runtime/Engine/Private/Components/SkyLightComponent.cpp
--- a/Engine/Source/Runtime/Engine/Private/Components/SkyLightComponent.cpp
+++ b/Engine/Source/Runtime/Engine/Private/Components/SkyLightComponent.cpp
@@ -23,6 +23,7 @@
 #include "Modules/ModuleManager.h"
 #include "Internationalization/Text.h"
 #include "CoreGlobals.h"
+#include "Engine/TextureCube.h"
 
 #if RHI_RAYTRACING
 #include "GlobalShader.h"
@@ -33,6 +34,7 @@
 
 #if WITH_EDITOR
 #include "Rendering/StaticLightingSystemInterface.h"
+#include "TextureCompiler.h"
 #endif
 
 #define LOCTEXT_NAMESPACE "SkyLightComponent"
@@ -712,6 +714,11 @@ void USkyLightComponent::ApplyComponentInstanceData(FPrecomputedSkyLightInstance
 void USkyLightComponent::UpdateSkyCaptureContentsArray(UWorld* WorldToUpdate, TArray<USkyLightComponent*>& ComponentArray, bool bOperateOnBlendSource)
 {
 	const bool bIsCompilingShaders = GShaderCompilingManager != NULL && GShaderCompilingManager->IsCompiling();
+#if WITH_EDITOR
+	const bool bIsCompilingTextures = FTextureCompilingManager::Get().GetNumRemainingTextures() > 0;
+#else
+	const bool bIsCompilingTextures = false;
+#endif
 
 	// Iterate backwards so we can remove elements without changing the index
 	for (int32 CaptureIndex = ComponentArray.Num() - 1; CaptureIndex >= 0; CaptureIndex--)
@@ -719,9 +726,19 @@ void USkyLightComponent::UpdateSkyCaptureContentsArray(UWorld* WorldToUpdate, TA
 		USkyLightComponent* CaptureComponent = ComponentArray[CaptureIndex];
 		AActor* Owner = CaptureComponent->GetOwner();
 
+		// For specific cubemaps, we must wait until the texture is compiled before capturing the skylight
+		const bool bIsCubemapCompiling = 
+#if WITH_EDITOR
+			CaptureComponent->SourceType == SLS_SpecifiedCubemap &&
+			CaptureComponent->Cubemap &&
+			CaptureComponent->Cubemap->IsDefaultTexture();
+#else
+			false;
+#endif
+
 		if (((!Owner || !Owner->GetLevel() || Owner->GetLevel()->bIsVisible) && CaptureComponent->GetWorld() == WorldToUpdate)
-			// Only process sky capture requests once async shader compiling completes, otherwise we will capture the scene with temporary shaders
-			&& (!bIsCompilingShaders || CaptureComponent->SourceType == SLS_SpecifiedCubemap))
+			// Only process sky capture requests once async texture and shader compiling completes, otherwise we will capture the scene with temporary shaders/textures
+			&& ((!bIsCompilingTextures && !bIsCompilingShaders) || (CaptureComponent->SourceType == SLS_SpecifiedCubemap && !bIsCubemapCompiling)))
 		{
 			// Only capture valid sky light components
 			if (CaptureComponent->SourceType != SLS_SpecifiedCubemap || CaptureComponent->Cubemap)
