Editgrate as workaround for RT TDR hangs:

Implement memory limit for GPUSkinCache updates. An RHI resource flush will happen when the limit is reached to release the old ray tracing geometries.

diff --git a/Engine/Source/Runtime/Engine/Private/GPUSkinCache.cpp b/Engine/Source/Runtime/Engine/Private/GPUSkinCache.cpp
--- a/Engine/Source/Runtime/Engine/Private/GPUSkinCache.cpp
+++ b/Engine/Source/Runtime/Engine/Private/GPUSkinCache.cpp
@@ -131,6 +131,16 @@ FAutoConsoleVariableRef CVarGPUSkinCacheBlendUsingVertexColorForRecomputeTangent
 
 static int32 GGPUSkinCacheFlushCounter = 0;
 
+#if RHI_RAYTRACING
+static int32 GMemoryLimitForBatchedRayTracingGeometryUpdates = 512;
+FAutoConsoleVariableRef CVarGPUSkinCacheMemoryLimitForBatchedRayTracingGeometryUpdates(
+	TEXT("r.SkinCache.MemoryLimitForBatchedRayTracingGeometryUpdates"),
+	GMemoryLimitForBatchedRayTracingGeometryUpdates,
+	TEXT(""),
+	ECVF_RenderThreadSafe
+);
+#endif
+
 static inline bool DoesPlatformSupportGPUSkinCache(const FStaticShaderPlatform Platform)
 {
 	return Platform == SP_PCD3D_SM5 || IsMetalSM5Platform(Platform) || IsVulkanSM5Platform(Platform) || Platform == SP_OPENGL_SM5 
@@ -1428,8 +1438,12 @@ void FGPUSkinCache::ProcessRayTracingGeometryToUpdate(
 
 		if (bRequireRecreatingRayTracingGeometry)
 		{
+			uint32 MemoryEstimation = 0;
+
 			FIndexBufferRHIRef IndexBufferRHI = LODModel.MultiSizeIndexContainer.GetIndexBuffer()->IndexBufferRHI;
+			MemoryEstimation += IndexBufferRHI->GetSize();
 			uint32 VertexBufferStride = LODModel.StaticVertexBuffers.PositionVertexBuffer.GetStride();
+			MemoryEstimation += LODModel.StaticVertexBuffers.PositionVertexBuffer.VertexBufferRHI->GetSize();
 
 			//#dxr_todo: do we need support for separate sections in FRayTracingGeometryData?
 			uint32 TrianglesCount = 0;
@@ -1471,6 +1485,21 @@ void FGPUSkinCache::ProcessRayTracingGeometryToUpdate(
 			// Flush pending resource barriers before BVH is built for the first time
 			TransitionAllToReadable(RHICmdList);
 
+			if (RayTracingGeometry.RayTracingGeometryRHI.IsValid())
+			{
+				// RayTracingGeometry.UpdateRHI() releases the old RT geometry, however due to the deferred deletion nature of RHI resources
+				// they will not be released until the end of the frame. We may get OOM in the middle of batched updates if not flushing.
+				// This memory size is an estimation based on vertex & index buffer size. In reality the flush happens at 2-3x of the number specified.
+				RayTracingGeometryMemoryPendingRelease += MemoryEstimation;
+
+				if (RayTracingGeometryMemoryPendingRelease >= GMemoryLimitForBatchedRayTracingGeometryUpdates * 1024ull * 1024ull)
+				{
+					RayTracingGeometryMemoryPendingRelease = 0;
+					FRHIResource::FlushPendingDeletes();
+					UE_LOG(LogSkinCache, Display, TEXT("Flushing RHI resource pending deletes due to %d MB limit"), GMemoryLimitForBatchedRayTracingGeometryUpdates);
+				}
+			}
+
 			RayTracingGeometry.SetInitializer(Initializer);
 			RayTracingGeometry.UpdateRHI();
 		}
diff --git a/Engine/Source/Runtime/Engine/Public/GPUSkinCache.h b/Engine/Source/Runtime/Engine/Public/GPUSkinCache.h
--- a/Engine/Source/Runtime/Engine/Public/GPUSkinCache.h
+++ b/Engine/Source/Runtime/Engine/Public/GPUSkinCache.h
@@ -399,6 +399,7 @@ protected:
 	TSet<FRHIUnorderedAccessView*> BuffersToTransition;
 #if RHI_RAYTRACING
 	TSet<FRayTracingGeometry*> RayTracingGeometriesToUpdate;
+	uint64 RayTracingGeometryMemoryPendingRelease = 0;
 #endif // RHI_RAYTRACING
 
 	TArray<FRWBuffersAllocation*> Allocations;
