Added documentation and default value to all nodes where it makes sense.

diff --git a/Engine/Source/Runtime/Engine/Classes/Materials/MaterialExpressionStrata.h b/Engine/Source/Runtime/Engine/Classes/Materials/MaterialExpressionStrata.h
--- a/Engine/Source/Runtime/Engine/Classes/Materials/MaterialExpressionStrata.h
+++ b/Engine/Source/Runtime/Engine/Classes/Materials/MaterialExpressionStrata.h
@@ -30,19 +30,19 @@ class UMaterialExpressionStrataDiffuseBSDF : public UMaterialExpression // STRAT
 	GENERATED_UCLASS_BODY()
 
 	/**
-	 * float3
+	 * Albedo (type = float3, unit = unitless)
 	 */
 	UPROPERTY()
 	FExpressionInput Albedo;
 
 	/**
-	 * float
+	 * Roughness (type = float, unit = unitless)
 	 */
 	UPROPERTY()
 	FExpressionInput Roughness;
 
 	/**
-	 * float3
+	 * Normal (type = float3, unit = unitless)
 	 */
 	UPROPERTY()
 	FExpressionInput Normal;
@@ -63,19 +63,19 @@ class UMaterialExpressionStrataDiffuseChanBSDF : public UMaterialExpression
 	GENERATED_UCLASS_BODY()
 
 	/**
-	 * float3
+	 * Albedo (type = float3, unit = unitless)
 	 */
 	UPROPERTY()
 	FExpressionInput Albedo;
 
 	/**
-	 * float
+	 * Roughness (type = float, unit = unitless)
 	 */
 	UPROPERTY()
 	FExpressionInput Roughness;
 
 	/**
-	 * float3
+	 * Normal (type = float3, unit = unitless)
 	 */
 	UPROPERTY()
 	FExpressionInput Normal;
@@ -96,32 +96,28 @@ class UMaterialExpressionStrataDielectricBSDF : public UMaterialExpression
 	GENERATED_UCLASS_BODY()
 	
 	/**
-	 * float
+	 * The index of refraction of the surface (type = float, unit = unitless)
 	 */
 	UPROPERTY()
 	FExpressionInput IOR;
-	// STRATA_TODO Refraction IOR?
 
 	/**
-	 * float3
+	 * A global color tint multiplied with the specular color, not physically based (type = float3, unit = unitless)
 	 */
 	UPROPERTY()
 	FExpressionInput Tint;
 		
 	/**
-	 * float2
+	 * Roughness (type = float2, unit = unitless)
 	 */
 	UPROPERTY()
 	FExpressionInput Roughness;
 
 	/**
-	 * float3
+	 * Normal (type = float3, unit = unitless)
 	 */
 	UPROPERTY()
 	FExpressionInput Normal;
-	// STRATA_TODO Tangent
-	// STRATA_TODO Distribution
-	// STRATA_TODO ScatterMode
 
 	//~ Begin UMaterialExpression Interface
 #if WITH_EDITOR
@@ -139,33 +135,29 @@ class UMaterialExpressionStrataConductorBSDF : public UMaterialExpression
 	GENERATED_UCLASS_BODY()
 		
 	/**
-	 * float3
+	 * Reflectivity when view direction is perpendicular to the surface, also known as F0 (type = float3, unit = unitless)
 	 */
 	UPROPERTY()
 	FExpressionInput Reflectivity;
 	
 	/**
-	 * float3
+	 * Reflectivity when the view direction is tangent to the surface (type = float3, unit = unitless)
 	 */
 	UPROPERTY()
 	FExpressionInput EdgeColor;
 
 	/**
-	 * float2
+	 * Roughness (type = float2, unit = unitless)
 	 */
 	UPROPERTY()
 	FExpressionInput Roughness;
 
 	/**
-	 * float3
+	 * Normal (type = float3, unit = unitless)
 	 */
 	UPROPERTY()
 	FExpressionInput Normal;
 
-	// STRATA_TODO Tangent
-	// STRATA_TODO Distribution
-	// STRATA_TODO ScatterMode
-
 	//~ Begin UMaterialExpression Interface
 #if WITH_EDITOR
 	virtual int32 Compile(class FMaterialCompiler* Compiler, int32 OutputIndex) override;
@@ -182,13 +174,13 @@ class UMaterialExpressionStrataVolumeBSDF : public UMaterialExpression
 	GENERATED_UCLASS_BODY()
 		
 	/**
-	 * Albedo (type = float3, unit = 1/m)
+	 * Albedo (type = float3, unit = unitless)
 	 */
 	UPROPERTY()
 	FExpressionInput Albedo;
 	
 	/**
-	 * Extinction (type = float3, unit = 1/m)
+	 * The rate at which light is absorbed or scattered by the medium. Mean Free Path = 1 / Extinction. (type = float3, unit = 1/m)
 	 */
 	UPROPERTY()
 	FExpressionInput Extinction;
@@ -240,7 +232,7 @@ class UMaterialExpressionStrataHorizontalMixing : public UMaterialExpression
 	FExpressionInput Background;
 
 	/**
-	 * float
+	 * Lerp factor between Background (Mix == 0) and Foreground (Mix == 1).
 	 */
 	UPROPERTY()
 	FExpressionInput Mix;
@@ -261,13 +253,13 @@ class UMaterialExpressionStrataVerticalLayering : public UMaterialExpression
 	GENERATED_UCLASS_BODY()
 
 	/**
-	 * Strata material
+	 * Strata material layer on top of the Base material layer
 	 */
 	UPROPERTY()
 	FExpressionInput Top;
 	
 	/**
-	 * Strata material
+	 * Strata material layer below the Top material layer
 	 */
 	UPROPERTY()
 	FExpressionInput Base;
@@ -321,7 +313,7 @@ class UMaterialExpressionStrataMultiply : public UMaterialExpression
 	FExpressionInput A;
 	
 	/**
-	 * Strata material
+	 * Weight to apply to the strata material BSDFs
 	 */
 	UPROPERTY()
 	FExpressionInput Weight;
@@ -347,13 +339,13 @@ class UMaterialExpressionStrataArtisticIOR : public UMaterialExpression
 	GENERATED_UCLASS_BODY()
 
 	/**
-	* TODO F0
-	*/
+	 * Reflectivity when view direction is perpendicular to the surface, also known as F0 (type = float3, unit = unitless)
+	 */
 	UPROPERTY()
 	FExpressionInput Reflectivity;
 
 	/**
-	 * TODO Color for tangent view direction
+	 * Reflectivity when the view direction is tangent to the surface (type = float3, unit = unitless)
 	 */
 	UPROPERTY()
 	FExpressionInput EdgeColor;
@@ -374,13 +366,13 @@ class UMaterialExpressionStrataPhysicalIOR : public UMaterialExpression
 	GENERATED_UCLASS_BODY()
 
 	/**
-	* TODO 
+	* The index of refraction of the surface (type = float, unit = unitless)
 	*/
 	UPROPERTY()
 	FExpressionInput IOR;
 
 	/**
-	 * TODO 
+	 * The rate at which light is absorbed or scattered by the medium. Mean Free Path = 1 / Extinction. (type = float3, unit = 1/m)
 	 */
 	UPROPERTY()
 	FExpressionInput Extinction;
diff --git a/Engine/Source/Runtime/Engine/Private/Materials/MaterialExpressions.cpp b/Engine/Source/Runtime/Engine/Private/Materials/MaterialExpressions.cpp
--- a/Engine/Source/Runtime/Engine/Private/Materials/MaterialExpressions.cpp
+++ b/Engine/Source/Runtime/Engine/Private/Materials/MaterialExpressions.cpp
@@ -18958,19 +18958,14 @@ UMaterialExpressionStrataDiffuseBSDF::UMaterialExpressionStrataDiffuseBSDF(const
 #if WITH_EDITOR
 int32 UMaterialExpressionStrataDiffuseBSDF::Compile(class FMaterialCompiler* Compiler, int32 OutputIndex)
 {
-	if (!Albedo.GetTracedInput().Expression)
-	{
-		return Compiler->Errorf(TEXT("Missing Albedo input"));	// STRATA_TODO: defaults in the detail panel with auto read only if plug in
-	}
-	if (!Roughness.GetTracedInput().Expression)
-	{
-		return Compiler->Errorf(TEXT("Missing Roughness input"));	// STRATA_TODO: defaults in the detail panel with auto read only if plug in
-	}
 	if (!Normal.GetTracedInput().Expression)
 	{
 		return Compiler->Errorf(TEXT("Missing Normal input"));	// STRATA_TODO: defaults in the detail panel with auto read only if plug in
 	}
-	return Compiler->StrataDiffuseOrenNayarBSDF(Albedo.Compile(Compiler), Roughness.Compile(Compiler), Normal.Compile(Compiler));
+	return Compiler->StrataDiffuseOrenNayarBSDF(
+		Albedo.GetTracedInput().Expression		? Albedo.Compile(Compiler)		: Compiler->Constant3(0.18f, 0.18f, 0.18f),
+		Roughness.GetTracedInput().Expression	? Roughness.Compile(Compiler)	: Compiler->Constant(0.0f),
+		Normal.Compile(Compiler));
 }
 
 void UMaterialExpressionStrataDiffuseBSDF::GetCaption(TArray<FString>& OutCaptions) const
@@ -19024,19 +19019,14 @@ UMaterialExpressionStrataDiffuseChanBSDF::UMaterialExpressionStrataDiffuseChanBS
 #if WITH_EDITOR
 int32 UMaterialExpressionStrataDiffuseChanBSDF::Compile(class FMaterialCompiler* Compiler, int32 OutputIndex)
 {
-	if (!Albedo.GetTracedInput().Expression)
-	{
-		return Compiler->Errorf(TEXT("Missing Albedo input"));	// STRATA_TODO: defaults in the detail panel with auto read only if plug in
-	}
-	if (!Roughness.GetTracedInput().Expression)
-	{
-		return Compiler->Errorf(TEXT("Missing Roughness input"));	// STRATA_TODO: defaults in the detail panel with auto read only if plug in
-	}
 	if (!Normal.GetTracedInput().Expression)
 	{
 		return Compiler->Errorf(TEXT("Missing Normal input"));	// STRATA_TODO: defaults in the detail panel with auto read only if plug in
 	}
-	return Compiler->StrataDiffuseChanBSDF(Albedo.Compile(Compiler), Roughness.Compile(Compiler), Normal.Compile(Compiler));
+	return Compiler->StrataDiffuseChanBSDF(
+		Albedo.GetTracedInput().Expression		? Albedo.Compile(Compiler)		: Compiler->Constant3(0.18f, 0.18f, 0.18f),
+		Roughness.GetTracedInput().Expression	? Roughness.Compile(Compiler)	: Compiler->Constant(0.0f),
+		Normal.Compile(Compiler));
 }
 
 void UMaterialExpressionStrataDiffuseChanBSDF::GetCaption(TArray<FString>& OutCaptions) const
@@ -19090,23 +19080,15 @@ UMaterialExpressionStrataDielectricBSDF::UMaterialExpressionStrataDielectricBSDF
 #if WITH_EDITOR
 int32 UMaterialExpressionStrataDielectricBSDF::Compile(class FMaterialCompiler* Compiler, int32 OutputIndex)
 {
-	if (!Tint.GetTracedInput().Expression)
-	{
-		return Compiler->Errorf(TEXT("Missing Tint input"));	// STRATA_TODO: defaults in the detail panel with auto read only if plug in
-	}
-	if (!IOR.GetTracedInput().Expression)
-	{
-		return Compiler->Errorf(TEXT("Missing IOR input"));	// STRATA_TODO: defaults in the detail panel with auto read only if plug in
-	}
-	if (!Roughness.GetTracedInput().Expression)
-	{
-		return Compiler->Errorf(TEXT("Missing Roughness input"));	// STRATA_TODO: defaults in the detail panel with auto read only if plug in
-	}
 	if (!Normal.GetTracedInput().Expression)
 	{
 		return Compiler->Errorf(TEXT("Missing Normal input"));	// STRATA_TODO: defaults in the detail panel with auto read only if plug in
 	}
-	return Compiler->StrataDielectricBSDF(Roughness.Compile(Compiler), IOR.Compile(Compiler), Tint.Compile(Compiler), Normal.Compile(Compiler));
+	return Compiler->StrataDielectricBSDF(
+		Roughness.GetTracedInput().Expression	?	Roughness.Compile(Compiler) : Compiler->Constant2(0.0f, 0.0f),
+		IOR.GetTracedInput().Expression			?	IOR.Compile(Compiler)		: Compiler->Constant(1.5f),
+		Tint.GetTracedInput().Expression		?	Tint.Compile(Compiler)		: Compiler->Constant3(1.0f, 1.0f, 1.0f),
+		Normal.Compile(Compiler));
 }
 
 void UMaterialExpressionStrataDielectricBSDF::GetCaption(TArray<FString>& OutCaptions) const
@@ -19163,23 +19145,15 @@ UMaterialExpressionStrataConductorBSDF::UMaterialExpressionStrataConductorBSDF(c
 #if WITH_EDITOR
 int32 UMaterialExpressionStrataConductorBSDF::Compile(class FMaterialCompiler* Compiler, int32 OutputIndex)
 {
-	if (!Reflectivity.GetTracedInput().Expression)
-	{
-		return Compiler->Errorf(TEXT("Missing Reflectivity input"));			// STRATA_TODO: defaults in the detail panel with auto read only if plug in
-	}
-	if (!EdgeColor.GetTracedInput().Expression)
-	{
-		return Compiler->Errorf(TEXT("Missing EdgeColor input"));	// STRATA_TODO: defaults in the detail panel with auto read only if plug in
-	}
-	if (!Roughness.GetTracedInput().Expression)
-	{
-		return Compiler->Errorf(TEXT("Missing Roughness input"));	// STRATA_TODO: defaults in the detail panel with auto read only if plug in
-	}
 	if (!Normal.GetTracedInput().Expression)
 	{
 		return Compiler->Errorf(TEXT("Missing Normal input"));		// STRATA_TODO: defaults in the detail panel with auto read only if plug in
 	}
-	return Compiler->StrataConductorBSDF(Reflectivity.Compile(Compiler), EdgeColor.Compile(Compiler), Roughness.Compile(Compiler), Normal.Compile(Compiler));
+	return Compiler->StrataConductorBSDF(
+		Reflectivity.GetTracedInput().Expression	? Reflectivity.Compile(Compiler): Compiler->Constant3(0.947f, 0.776f, 0.371f),	// Default to Gold
+		EdgeColor.GetTracedInput().Expression		? EdgeColor.Compile(Compiler)	: Compiler->Constant3(1.000f, 0.982f, 0.753f),	// Default to Gold
+		Roughness.GetTracedInput().Expression		? Roughness.Compile(Compiler)	: Compiler->Constant2(0.0f, 0.0f),
+		Normal.Compile(Compiler));
 }
 
 void UMaterialExpressionStrataConductorBSDF::GetCaption(TArray<FString>& OutCaptions) const
@@ -19236,23 +19210,11 @@ UMaterialExpressionStrataVolumeBSDF::UMaterialExpressionStrataVolumeBSDF(const F
 #if WITH_EDITOR
 int32 UMaterialExpressionStrataVolumeBSDF::Compile(class FMaterialCompiler* Compiler, int32 OutputIndex)
 {
-	if (!Albedo.GetTracedInput().Expression)
-	{
-		return Compiler->Errorf(TEXT("Missing Albedo input"));	// STRATA_TODO: defaults in the detail panel with auto read only if plug in
-	}
-	if (!Extinction.GetTracedInput().Expression)
-	{
-		return Compiler->Errorf(TEXT("Missing Extinction input"));	// STRATA_TODO: defaults in the detail panel with auto read only if plug in
-	}
-	if (!Anisotropy.GetTracedInput().Expression)
-	{
-		return Compiler->Errorf(TEXT("Missing Anisotropy input"));	// STRATA_TODO: defaults in the detail panel with auto read only if plug in
-	}
 	return Compiler->StrataVolumeBSDF(
-		Albedo.Compile(Compiler), 
-		Extinction.Compile(Compiler), 
-		Anisotropy.Compile(Compiler), 
-		Thickness.GetTracedInput().Expression ? Thickness.Compile(Compiler) : Compiler->Constant(0.001f)); // default = 1mm
+		Albedo.GetTracedInput().Expression		? Albedo.Compile(Compiler)			: Compiler->Constant3(0.0f, 0.0f, 0.0f),
+		Extinction.GetTracedInput().Expression	? Extinction.Compile(Compiler)		: Compiler->Constant3(0.0f, 0.0f, 0.0f),
+		Anisotropy.GetTracedInput().Expression	? Anisotropy.Compile(Compiler)		: Compiler->Constant(0.0f),
+		Thickness.GetTracedInput().Expression	? Thickness.Compile(Compiler)		: Compiler->Constant(0.001f)); // default = 1mm
 }
 
 void UMaterialExpressionStrataVolumeBSDF::GetCaption(TArray<FString>& OutCaptions) const
@@ -19315,11 +19277,10 @@ int32 UMaterialExpressionStrataHorizontalMixing::Compile(class FMaterialCompiler
 	{
 		return Compiler->Errorf(TEXT("Missing Background input"));
 	}
-	if (!Mix.GetTracedInput().Expression)
-	{
-		return Compiler->Errorf(TEXT("Missing Mix input"));
-	}
-	return Compiler->StrataHorizontalMixing(Foreground.Compile(Compiler), Background.Compile(Compiler), Mix.GetTracedInput().Expression ? Mix.Compile(Compiler) : Compiler->Constant(0.0f));
+	return Compiler->StrataHorizontalMixing(
+		Foreground.Compile(Compiler), 
+		Background.Compile(Compiler), 
+		Mix.GetTracedInput().Expression ? Mix.Compile(Compiler) : Compiler->Constant(0.0f));
 }
 
 void UMaterialExpressionStrataHorizontalMixing::GetCaption(TArray<FString>& OutCaptions) const
@@ -19505,15 +19466,10 @@ UMaterialExpressionStrataArtisticIOR::UMaterialExpressionStrataArtisticIOR(const
 #if WITH_EDITOR
 int32 UMaterialExpressionStrataArtisticIOR::Compile(class FMaterialCompiler* Compiler, int32 OutputIndex)
 {
-	if (!Reflectivity.GetTracedInput().Expression)
-	{
-		return Compiler->Errorf(TEXT("Missing Reflectivity input"));
-	}
-	if (!EdgeColor.GetTracedInput().Expression)
-	{
-		return Compiler->Errorf(TEXT("Missing EdgeColor input"));
-	}
-	return Compiler->StrataArtisticIOR(Reflectivity.Compile(Compiler), EdgeColor.Compile(Compiler), OutputIndex);
+	return Compiler->StrataArtisticIOR(
+		Reflectivity.GetTracedInput().Expression	? Reflectivity.Compile(Compiler): Compiler->Constant3(0.947f, 0.776f, 0.371f),	// Default to Gold
+		EdgeColor.GetTracedInput().Expression		? EdgeColor.Compile(Compiler)	: Compiler->Constant3(1.000f, 0.982f, 0.753f),	// Default to Gold
+		OutputIndex);
 }
 
 void UMaterialExpressionStrataArtisticIOR::GetCaption(TArray<FString>& OutCaptions) const
@@ -19557,15 +19513,10 @@ UMaterialExpressionStrataPhysicalIOR::UMaterialExpressionStrataPhysicalIOR(const
 #if WITH_EDITOR
 int32 UMaterialExpressionStrataPhysicalIOR::Compile(class FMaterialCompiler* Compiler, int32 OutputIndex)
 {
-	if (!IOR.GetTracedInput().Expression)
-	{
-		return Compiler->Errorf(TEXT("Missing IOR input"));
-	}
-	if (!Extinction.GetTracedInput().Expression)
-	{
-		return Compiler->Errorf(TEXT("Missing Extinction input"));
-	}
-	return Compiler->StrataPhysicalIOR(IOR.Compile(Compiler), Extinction.Compile(Compiler), OutputIndex);
+	return Compiler->StrataPhysicalIOR(
+		IOR.GetTracedInput().Expression			? IOR.Compile(Compiler)			: Compiler->Constant3(0.18f, 0.42f, 1.37f),// Default to Gold
+		Extinction.GetTracedInput().Expression	? Extinction.Compile(Compiler)	: Compiler->Constant3(3.41f, 2.35f, 1.77f),// Default to Gold
+		OutputIndex);
 }
 
 void UMaterialExpressionStrataPhysicalIOR::GetCaption(TArray<FString>& OutCaptions) const
