# Made TimeStretchCurve tooltip text more clear
#rnx

diff --git a/Engine/Source/Runtime/Engine/Classes/Animation/AnimMontage.h b/Engine/Source/Runtime/Engine/Classes/Animation/AnimMontage.h
--- a/Engine/Source/Runtime/Engine/Classes/Animation/AnimMontage.h
+++ b/Engine/Source/Runtime/Engine/Classes/Animation/AnimMontage.h
@@ -839,10 +839,12 @@ public:
 
 	//~Begin Time Stretch Curve
 public:
+
+	/** Time stretch curve will only be used when the montage has a non-default play rate  */
 	UPROPERTY(EditAnywhere, Category = TimeStretchCurve)
 	FTimeStretchCurve TimeStretchCurve;
 
-	/** Name of optional TimeStretchCurveName to look for in Montage. */
+	/** Name of optional TimeStretchCurveName to look for in Montage. Time stretch curve will only be used when the montage has a non-default play rate */
 	UPROPERTY(EditAnywhere, Category = TimeStretchCurve)
 	FName TimeStretchCurveName;
 
