Don't append the actor name number to the actor label when it's a globally unique actor name.

diff --git a/Engine/Source/Runtime/Engine/Private/ActorEditor.cpp b/Engine/Source/Runtime/Engine/Private/ActorEditor.cpp
--- a/Engine/Source/Runtime/Engine/Private/ActorEditor.cpp
+++ b/Engine/Source/Runtime/Engine/Private/ActorEditor.cpp
@@ -912,6 +912,7 @@ const FString& AActor::GetActorLabel() const
 		// actor's FName when creating the initially.  It doesn't actually *need* to be unique, this is just
 		// an easy way to tell actors apart when observing them in a list.  The user can always go and rename
 		// these labels such that they're no longer unique.
+		if (!FActorSpawnUtils::IsGloballyUniqueName(GetFName()))
 		{
 			// Don't bother adding a suffix for number '0'
 			const int32 NameNumber = NAME_INTERNAL_TO_EXTERNAL( GetFName().GetNumber() );
