Defer reflection capture when affected by async texture compilation

diff --git a/Engine/Source/Runtime/Engine/Private/Components/ReflectionCaptureComponent.cpp b/Engine/Source/Runtime/Engine/Private/Components/ReflectionCaptureComponent.cpp
--- a/Engine/Source/Runtime/Engine/Private/Components/ReflectionCaptureComponent.cpp
+++ b/Engine/Source/Runtime/Engine/Private/Components/ReflectionCaptureComponent.cpp
@@ -43,6 +43,7 @@
 
 #if WITH_EDITOR
 #include "Factories/TextureFactory.h"
+#include "TextureCompiler.h"
 #endif
 
 // ES3.0+ devices support seamless cubemap filtering, averaging edges will produce artifacts on those devices
@@ -1053,11 +1054,16 @@ void UReflectionCaptureComponent::MarkDirtyForRecapture()
 
 void UReflectionCaptureComponent::UpdateReflectionCaptureContents(UWorld* WorldToUpdate, const TCHAR* CaptureReason, bool bVerifyOnlyCapturing, bool bCapturingForMobile)
 {
-	if (WorldToUpdate->Scene 
+	if (WorldToUpdate->Scene
 		// Don't capture and read back capture contents if we are currently doing async shader compiling
 		// This will keep the update requests in the queue until compiling finishes
 		// Note: this will also prevent uploads of cubemaps from DDC, which is unintentional
-		&& (GShaderCompilingManager == NULL || !GShaderCompilingManager->IsCompiling()))
+		&& (GShaderCompilingManager == NULL || !GShaderCompilingManager->IsCompiling())
+#if WITH_EDITOR
+		// Prevent any reflection capture if textures are still compiling
+		&& FTextureCompilingManager::Get().GetNumRemainingTextures() == 0
+#endif
+		)
 	{
 		//guarantee that all render proxies are up to date before kicking off this render
 		WorldToUpdate->SendAllEndOfFrameUpdates();
