Don't spawn default physics volume as an external actor.

diff --git a/Engine/Source/Runtime/Engine/Private/World.cpp b/Engine/Source/Runtime/Engine/Private/World.cpp
--- a/Engine/Source/Runtime/Engine/Private/World.cpp
+++ b/Engine/Source/Runtime/Engine/Private/World.cpp
@@ -4974,6 +4974,7 @@ APhysicsVolume* UWorld::InternalGetDefaultPhysicsVolume() const
 		// Spawn volume
 		FActorSpawnParameters SpawnParams;
 		SpawnParams.bAllowDuringConstructionScript = true;
+		SpawnParams.bCreateActorPackage = false;
 
 		UWorld* MutableThis = const_cast<UWorld*>(this);
 		MutableThis->DefaultPhysicsVolume = MutableThis->SpawnActor<APhysicsVolume>(DefaultPhysicsVolumeClass, SpawnParams);
