Fix non-editor builds.

diff --git a/Engine/Source/Runtime/Engine/Private/World.cpp b/Engine/Source/Runtime/Engine/Private/World.cpp
--- a/Engine/Source/Runtime/Engine/Private/World.cpp
+++ b/Engine/Source/Runtime/Engine/Private/World.cpp
@@ -4974,7 +4974,10 @@ APhysicsVolume* UWorld::InternalGetDefaultPhysicsVolume() const
 		// Spawn volume
 		FActorSpawnParameters SpawnParams;
 		SpawnParams.bAllowDuringConstructionScript = true;
+
+#if WITH_EDITOR
 		SpawnParams.bCreateActorPackage = false;
+#endif
 
 		UWorld* MutableThis = const_cast<UWorld*>(this);
 		MutableThis->DefaultPhysicsVolume = MutableThis->SpawnActor<APhysicsVolume>(DefaultPhysicsVolumeClass, SpawnParams);
