Edigrate perf fix from CL 14442838 (Ben.Ingram)

Add cvar to disable expensive material parameter sanity checks

diff --git a/Engine/Source/Runtime/Engine/Private/Materials/Material.cpp b/Engine/Source/Runtime/Engine/Private/Materials/Material.cpp
--- a/Engine/Source/Runtime/Engine/Private/Materials/Material.cpp
+++ b/Engine/Source/Runtime/Engine/Private/Materials/Material.cpp
@@ -90,6 +90,13 @@
 
 #define LOCTEXT_NAMESPACE "Material"
 
+static TAutoConsoleVariable<int32> CVarMaterialParameterLegacyChecks(
+	TEXT("r.MaterialParameterLegacyChecks"),
+	0,
+	TEXT("When enabled, sanity check new material parameter logic against legacy path.\n")
+	TEXT("Note that this can be slow"),
+	ECVF_RenderThreadSafe);
+
 #if WITH_EDITOR
 const FMaterialsWithDirtyUsageFlags FMaterialsWithDirtyUsageFlags::DefaultAnnotation;
 
@@ -2130,7 +2137,7 @@ bool UMaterial::GetScalarParameterValue(const FHashedMaterialParameterInfo& Para
 {
 	const bool bResult = GetScalarParameterValue_New(ParameterInfo, OutValue, bOveriddenOnly);
 #if WITH_EDITOR
-	if (GIsClient)
+	if (GIsClient && CVarMaterialParameterLegacyChecks.GetValueOnAnyThread())
 	{
 		// Expressions may not be loaded on server, and cached data is not updated...results will be incorrect
 		float OldValue = 0.0f;
@@ -2269,7 +2276,7 @@ bool UMaterial::GetVectorParameterValue(const FHashedMaterialParameterInfo& Para
 {
 	const bool bResult = GetVectorParameterValue_New(ParameterInfo, OutValue, bOveriddenOnly);
 #if WITH_EDITOR
-	if (GIsClient)
+	if (GIsClient && CVarMaterialParameterLegacyChecks.GetValueOnAnyThread())
 	{
 		// Expressions may not be loaded on server, and cached data is not updated...results will be incorrect
 		FLinearColor OldValue(ForceInitToZero);
@@ -2392,7 +2399,7 @@ bool UMaterial::GetTextureParameterValue(const FHashedMaterialParameterInfo& Par
 {
 	const bool bResult = GetTextureParameterValue_New(ParameterInfo, OutValue, bOveriddenOnly);
 #if WITH_EDITOR
-	if (GIsClient)
+	if (GIsClient && CVarMaterialParameterLegacyChecks.GetValueOnAnyThread())
 	{
 		// Expressions may not be loaded on server, and cached data is not updated...results will be incorrect
 		UTexture* OldValue = nullptr;
@@ -2497,7 +2504,7 @@ bool UMaterial::GetRuntimeVirtualTextureParameterValue(const FHashedMaterialPara
 {
 	const bool bResult = GetRuntimeVirtualTextureParameterValue_New(ParameterInfo, OutValue, bOveriddenOnly);
 #if WITH_EDITOR
-	if (GIsClient)
+	if (GIsClient && CVarMaterialParameterLegacyChecks.GetValueOnAnyThread())
 	{
 		// Expressions may not be loaded on server, and cached data is not updated...results will be incorrect
 		URuntimeVirtualTexture* OldValue = nullptr;
@@ -2615,7 +2622,7 @@ bool UMaterial::GetFontParameterValue(const FHashedMaterialParameterInfo& Parame
 {
 	const bool bResult = GetFontParameterValue_New(ParameterInfo, OutFontValue, OutFontPage, bOveriddenOnly);
 #if WITH_EDITOR
-	if (GIsClient)
+	if (GIsClient && CVarMaterialParameterLegacyChecks.GetValueOnAnyThread())
 	{
 		// Expressions may not be loaded on server, and cached data is not updated...results will be incorrect
 		UFont* OldValue = nullptr;
