External actor packages are now named using a MD5 of the fully qualified actor path, instead of the actor GUID. This makes it possible to create singleton actors (such as the partition actors) by specifying a unique name when spawning them, and avoids multiple users creating the same one (determined by who adds on source control first if that happens).
Fix globally unique actor names encoding since FName is case-insensitive (switched from base-64 to hexadecimal).

diff --git a/Engine/Source/Runtime/Engine/Classes/Engine/Level.h b/Engine/Source/Runtime/Engine/Classes/Engine/Level.h
--- a/Engine/Source/Runtime/Engine/Classes/Engine/Level.h
+++ b/Engine/Source/Runtime/Engine/Classes/Engine/Level.h
@@ -1009,10 +1009,10 @@ public:
 
 	/**
 	 * Create an package for this actor
-	 * @param InGuid the guid to generate the name from.
+	 * @param InActorPath the fully qualified actor path, in the format: 'Outermost.Outer.Name'
 	 * @return the created package
 	 */
-	static ENGINE_API UPackage* CreateActorPackage(UPackage* InLevelPackage, const FGuid& InGuid);
+	static ENGINE_API UPackage* CreateActorPackage(UPackage* InLevelPackage, const FString& InActorPath);
 
 	/**
 	 * Detach or reattach all level actors to from/to their external package
diff --git a/Engine/Source/Runtime/Engine/Private/Actor.cpp b/Engine/Source/Runtime/Engine/Private/Actor.cpp
--- a/Engine/Source/Runtime/Engine/Private/Actor.cpp
+++ b/Engine/Source/Runtime/Engine/Private/Actor.cpp
@@ -1158,8 +1158,16 @@ bool AActor::Rename( const TCHAR* InName, UObject* NewOuter, ERenameFlags Flags
 	const bool bChangingOuters = (NewOuter && (NewOuter != GetOuter()));
 
 #if WITH_EDITOR
-	// if we have an external actor and the actor is changing outer, we will want to move its package location
+	// if we have an external actor and the actor is changing name/outer, we will want to update its package
 	const bool bExternalActor = IsPackageExternal();
+
+	if (!bRenameTest && bExternalActor)
+	{
+		if (ULevel* MyLevel = GetLevel())
+		{
+			SetPackageExternal(false, MyLevel->IsUsingExternalActors());
+		}
+	}
 #endif
 
 	if (!bRenameTest && bChangingOuters)
@@ -1176,31 +1184,29 @@ bool AActor::Rename( const TCHAR* InName, UObject* NewOuter, ERenameFlags Flags
 				MyLevel->ActorsForGC.Remove(this);
 				// TODO: There may need to be some consideration about removing this actor from the level cluster, but that would probably require destroying the entire cluster, so defer for now
 			}
-
-#if WITH_EDITOR
-			if (bExternalActor)
-			{
-				SetPackageExternal(false, MyLevel->IsUsingExternalActors());
-			}
-#endif
 		}
 	}
 
 #if WITH_EDITOR
 	UObject* OldOuter = GetOuter();
 #endif
+
 	const bool bSuccess = Super::Rename( InName, NewOuter, Flags );
 
-	if (!bRenameTest && bChangingOuters)
+#if WITH_EDITOR
+	if (!bRenameTest && bExternalActor && NewOuter)
 	{
 		if (ULevel* MyLevel = GetLevel())
 		{
-#if WITH_EDITOR
-			if (bExternalActor)
-			{
-				SetPackageExternal(true, MyLevel->IsUsingExternalActors());
-			}
+			SetPackageExternal(true, MyLevel->IsUsingExternalActors());
+		}
+	}
 #endif
+
+	if (!bRenameTest && bChangingOuters)
+	{
+		if (ULevel* MyLevel = GetLevel())
+		{
 			MyLevel->Actors.Add(this);
 			MyLevel->ActorsForGC.Add(this);
 
diff --git a/Engine/Source/Runtime/Engine/Private/ActorEditor.cpp b/Engine/Source/Runtime/Engine/Private/ActorEditor.cpp
--- a/Engine/Source/Runtime/Engine/Private/ActorEditor.cpp
+++ b/Engine/Source/Runtime/Engine/Private/ActorEditor.cpp
@@ -847,13 +847,9 @@ void AActor::SetPackageExternal(bool bExternal, bool bShouldDirty)
 	UPackage* LevelPackage = GetLevel()->GetPackage();
 	if (bExternal)
 	{
-		ActorGuid = ActorGuid.IsValid() ? ActorGuid : FGuid::NewGuid();
-		UPackage* NewActorPackage = ULevel::CreateActorPackage(LevelPackage,  ActorGuid);
+		UPackage* NewActorPackage = ULevel::CreateActorPackage(LevelPackage, GetPathName());
 		SetExternalPackage(NewActorPackage);
-		// should be removed but needed for now so the package creation is visible to Multi-User
-		FAssetRegistryModule::AssetCreated(this);
 	}
-	// if the actor package is different, embed the actor back in the level
 	else 
 	{
 		UPackage* ActorPackage = GetExternalPackage();
diff --git a/Engine/Source/Runtime/Engine/Private/Level.cpp b/Engine/Source/Runtime/Engine/Private/Level.cpp
--- a/Engine/Source/Runtime/Engine/Private/Level.cpp
+++ b/Engine/Source/Runtime/Engine/Private/Level.cpp
@@ -54,6 +54,7 @@ Level.cpp: Level-related functions
 #include "AssetRegistryModule.h"
 #include "IAssetRegistry.h"
 #include "AssetData.h"
+#include "Misc/ArchiveMD5.h"
 #endif
 #include "Engine/LevelStreaming.h"
 #include "LevelUtils.h"
@@ -895,13 +896,29 @@ void ULevel::PreDuplicate(FObjectDuplicationParameters& DupParams)
 #if WITH_EDITOR
 	if (DupParams.DuplicateMode != EDuplicateMode::PIE && DupParams.bAssignExternalPackages)
 	{
-		UPackage* DestPackage = DupParams.DestOuter->GetPackage();
+		UPackage* SrcPackage = GetPackage();
+		UPackage* DstPackage = DupParams.DestOuter->GetPackage();
+
+		const FString ReplaceFrom = FPaths::GetBaseFilename(*SrcPackage->GetName());
+		const FString ReplaceTo = FPaths::GetBaseFilename(*DstPackage->GetName());
+
 		for (AActor* Actor : Actors)
 		{
 			UPackage* ActorPackage = Actor ? Actor->GetExternalPackage() : nullptr;
 			if (ActorPackage)
 			{
-				UPackage* DupActorPackage = CreateActorPackage(DestPackage, FGuid::NewGuid());
+				FString ActorPath = GetPathName();
+
+				if (DstPackage != SrcPackage)
+				{
+					ActorPath = ActorPath.Replace(
+						*FString::Printf(TEXT("%s.%s:"), *ReplaceFrom, *ReplaceFrom),
+						*FString::Printf(TEXT("%s.%s:"), *ReplaceTo, *ReplaceTo)
+					);
+				}
+
+				UPackage* DupActorPackage = CreateActorPackage(DstPackage, ActorPath);
+
 				DupActorPackage->MarkAsFullyLoaded();
 				DupParams.DuplicationSeed.Add(ActorPackage, DupActorPackage);
 			}
@@ -2335,26 +2352,37 @@ FString ULevel::GetExternalActorsPath(UPackage* InLevelPackage, const FString& I
 	return GetExternalActorsPath(InLevelPackage->GetName(), InPackageShortName);
 }
 
-UPackage* ULevel::CreateActorPackage(UPackage* InLevelPackage, const FGuid& InGuid)
+UPackage* ULevel::CreateActorPackage(UPackage* InLevelPackage, const FString& InActorPath)
 {
-	check(InGuid.IsValid());
-	FString GuidBase36 = InGuid.ToString(EGuidFormats::Base36Encoded);
+	// Convert the actor path to lowercase to make sure we get the same hash for case insensitive file systems
+	FString ActorPath = InActorPath.ToLower();
+
+	FArchiveMD5 ArMD5;
+	ArMD5 << ActorPath;
+
+	FMD5Hash MD5Hash;
+	ArMD5.GetHash(MD5Hash);
+
+	FGuid PackageGuid;
+	check(MD5Hash.GetSize() == sizeof(FGuid));
+	FMemory::Memcpy(&PackageGuid, MD5Hash.GetBytes(), sizeof(FGuid));
+	check(PackageGuid.IsValid());
+
+	FString GuidBase36 = PackageGuid.ToString(EGuidFormats::Base36Encoded);
 	check(GuidBase36.Len());
 
-	const int32 GuidBase36Len = GuidBase36.Len();
 	FString BaseDir = GetExternalActorsPath(InLevelPackage);
-	FString ActorPackageName =
-		FString::Printf(
-			TEXT("%s/%c%c/%c%c/%s"),
-			*BaseDir,
-			GuidBase36[0],
-			(GuidBase36Len > 1) ? GuidBase36[1] : '0',
-			(GuidBase36Len > 2) ? GuidBase36[2] : '0',
-			(GuidBase36Len > 3) ? GuidBase36[3] : '0',
-			*GuidBase36 + FMath::Min(GuidBase36Len - 1, 4)
-		);
-
-	UPackage* ActorPackage = CreatePackage( *ActorPackageName);
+
+	TStringBuilderWithBuffer<TCHAR, NAME_SIZE> ActorPackageName;
+	ActorPackageName.Append(BaseDir);
+	ActorPackageName.Append(TEXT("/"));
+	ActorPackageName.Append(*GuidBase36, 2);
+	ActorPackageName.Append(TEXT("/"));
+	ActorPackageName.Append(*GuidBase36 + 2, 2);
+	ActorPackageName.Append(TEXT("/"));
+	ActorPackageName.Append(*GuidBase36 + 4);
+
+	UPackage* ActorPackage = CreatePackage(*ActorPackageName);
 	ActorPackage->SetPackageFlags(PKG_EditorOnly);
 	return ActorPackage;
 }
diff --git a/Engine/Source/Runtime/Engine/Private/LevelActor.cpp b/Engine/Source/Runtime/Engine/Private/LevelActor.cpp
--- a/Engine/Source/Runtime/Engine/Private/LevelActor.cpp
+++ b/Engine/Source/Runtime/Engine/Private/LevelActor.cpp
@@ -270,7 +270,7 @@ void LineCheckTracker::CaptureLineCheck(int32 LineCheckFlags, const FVector* Ext
  * Notes:
  *	- The timestamp is stored in microseconds for a total of 54 bits, enough to cover
  *	  the next 570 years.
- *	- The highest 72 bits are appended to the name in base-64.
+ *	- The highest 72 bits are appended to the name in hexadecimal.
  *	- The lowest 30 bits of the timestamp are stored in the name number. This is to
  *	  minimize the total names generated for globally unique names (string part will
  *	  change every ~17 minutes for a specific actor class).
@@ -310,8 +310,13 @@ public:
 		// Make final name
 		TStringBuilderWithBuffer<TCHAR, NAME_SIZE> StringBuilder;
 		StringBuilder += BaseName.ToString();
-		StringBuilder += TEXT("_");
-		StringBuilder += FBase64::Encode(HighPart, 9);
+		StringBuilder += TEXT("_UAID_");
+
+		for (uint32 i=0; i<9; i++)
+		{
+			StringBuilder += NibbleToTChar(HighPart[i] >> 4);
+			StringBuilder += NibbleToTChar(HighPart[i] & 15);
+		}
 
 		return FName(*StringBuilder, (ElapsedUs & 0x3fffffff) | (1 << 30));
 	}
@@ -354,12 +359,19 @@ bool FActorSpawnUtils::IsGloballyUniqueName(FName Name)
 		const FString PlainName = Name.GetPlainNameString();
 		const int32 PlainNameLen = PlainName.Len();
 		
-		// Parse a name like this: StaticMeshActor_cCCECUPVj1cA
-		if (PlainNameLen >= 13 && PlainName[PlainNameLen - 13] == '_')
+		// Parse a name like this: StaticMeshActor_UAID_001122334455667788
+		if (PlainNameLen >= 24)
 		{
-			uint8 Dest[16];
-			if (FBase64::Decode(&PlainName[PlainNameLen - 12], 12, Dest))
+			if (!FCString::Strnicmp(*PlainName + PlainNameLen - 24, TEXT("_UAID_"), 6))
 			{
+				for (uint32 i=0; i<18; i++)
+				{
+					if (!CheckTCharIsHex(PlainName[PlainNameLen - i - 1]))
+					{
+						return false;
+					}
+				}
+
 				return true;
 			}
 		}
@@ -372,7 +384,8 @@ FName FActorSpawnUtils::GetBaseName(FName Name)
 {
 	if (IsGloballyUniqueName(Name))
 	{
-		return *Name.GetPlainNameString().LeftChop(13);
+		// Chop a name like this: StaticMeshActor_UAID_001122334455667788
+		return *Name.GetPlainNameString().LeftChop(24);
 	}
 
 	return *Name.GetPlainNameString();
@@ -527,17 +540,12 @@ AActor* UWorld::SpawnActor( UClass* Class, FTransform const* UserTransformPtr, c
 	if (SpawnParameters.OverridePackage)
 	{
 		ExternalPackage = SpawnParameters.OverridePackage;
+		bNeedGloballyUniqueName = true;
 	}
 	else if (LevelToSpawnIn->IsUsingExternalActors() && SpawnParameters.bCreateActorPackage && !(SpawnParameters.ObjectFlags & RF_Transient))
 	{
-		if (CastChecked<AActor>(Class->GetDefaultObject())->SupportsExternalPackaging())
-		{
-			// @todo FH: needs to handle mark package dirty and asset creation notification
-			ExternalPackage = ULevel::CreateActorPackage(LevelToSpawnIn->GetPackage(), ActorGuid);
-		}
+		bNeedGloballyUniqueName = CastChecked<AActor>(Class->GetDefaultObject())->SupportsExternalPackaging();
 	}
-
-	bNeedGloballyUniqueName = !!ExternalPackage;
 #endif
 
 	if (NewActorName.IsNone())
@@ -579,6 +587,19 @@ AActor* UWorld::SpawnActor( UClass* Class, FTransform const* UserTransformPtr, c
 		return NULL;
 	}
 
+#if WITH_EDITOR
+	if (bNeedGloballyUniqueName && !ExternalPackage)
+	{
+		TStringBuilderWithBuffer<TCHAR, NAME_SIZE> ActorPath;
+		ActorPath += LevelToSpawnIn->GetPathName();
+		ActorPath += TEXT(".");
+		ActorPath += NewActorName.ToString();
+
+		// @todo FH: needs to handle mark package dirty and asset creation notification
+		ExternalPackage = ULevel::CreateActorPackage(LevelToSpawnIn->GetPackage(), *ActorPath);
+	}
+#endif
+
 	FTransform const UserTransform = UserTransformPtr ? *UserTransformPtr : FTransform::Identity;
 
 	ESpawnActorCollisionHandlingMethod CollisionHandlingOverride = SpawnParameters.SpawnCollisionHandlingOverride;
