Force partition actors to have a unique name (composed of the actor class, their grid GUID and the actor cell coordinates they belongs to) to make sure a single user can create a specific cell actor and avoid duplicates.

diff --git a/Engine/Source/Runtime/Engine/Private/ActorPartition/ActorPartitionSubsystem.cpp b/Engine/Source/Runtime/Engine/Private/ActorPartition/ActorPartitionSubsystem.cpp
--- a/Engine/Source/Runtime/Engine/Private/ActorPartition/ActorPartitionSubsystem.cpp
+++ b/Engine/Source/Runtime/Engine/Private/ActorPartition/ActorPartitionSubsystem.cpp
@@ -218,8 +218,30 @@ public:
 				
 		if (!FoundActor && bInCreate)
 		{
+			TStringBuilderWithBuffer<TCHAR, NAME_SIZE> ActorNameBuilder;
+
+			ActorNameBuilder += InActorClass->GetName();
+			ActorNameBuilder += TEXT("_");
+
+			if (InGuid.IsValid())
+			{
+				ActorNameBuilder += InGuid.ToString(EGuidFormats::Base36Encoded);
+				ActorNameBuilder += TEXT("_");
+			}
+
+			ActorNameBuilder += FString::Printf(TEXT("%d_%d_%d"), InCellCoord.X, InCellCoord.Y, InCellCoord.Z);
+
 			FActorSpawnParameters SpawnParams;
 			SpawnParams.OverrideLevel = InCellCoord.Level;
+			SpawnParams.Name = *ActorNameBuilder;
+			SpawnParams.NameMode = FActorSpawnParameters::ESpawnActorNameMode::Required_Fatal;
+
+			// Handle the case where the actor already exists, but is in the undo stack
+			if (UObject* ExistingObject = StaticFindObject(nullptr, World->PersistentLevel, *SpawnParams.Name.ToString()))
+			{
+				check(CastChecked<AActor>(ExistingObject)->IsPendingKill());
+				ExistingObject->Rename(nullptr, nullptr, REN_DontCreateRedirectors | REN_DoNotDirty | REN_NonTransactional | REN_ForceNoResetLoaders);
+			}
 						
 			FVector CellCenter(CellBounds.GetCenter());
 			FoundActor = CastChecked<APartitionActor>(World->SpawnActor(InActorClass, &CellCenter, nullptr, SpawnParams));
