Added strata material description into output stat.

Strata material compiler general improvements.
- Created a headers shared between code and shaders
- Build up an internal representation of the material layer (BSDF horizontal mixing per layer)
- Output more meaningful errors for the user to understand what is wrong if using too many BSDFs or layers, etc. (exemple Adding would result in too many BSDFs in a Layer) Important if artist start to use what we currently have to not have them whaste time.
- Add a debug as text in the hlsl code showing the material setup as well as in the material compilation stats.

This later will open other possibility: shared normal cache between layers or per material max byte count instead of max BSDF count.

diff --git a/Engine/Shaders/Private/Strata/Strata.ush b/Engine/Shaders/Private/Strata/Strata.ush
--- a/Engine/Shaders/Private/Strata/Strata.ush
+++ b/Engine/Shaders/Private/Strata/Strata.ush
@@ -6,6 +6,7 @@
 #include "../BRDF.ush"
 #include "../ShadingCommon.ush"
 
+#include "/Engine/Shared/StrataDefinitions.h"
 
 // STRATA_TODO
 //	- What model exactly
@@ -54,13 +55,6 @@ struct FStrataVolume
 };
 
 
-#define STRATA_BSDF_TYPE_DIFFUSE_ON			0
-#define STRATA_BSDF_TYPE_DIFFUSE_CHAN		1
-#define STRATA_BSDF_TYPE_DIELECTRIC			2
-#define STRATA_BSDF_TYPE_CONDUCTOR			3
-#define STRATA_BSDF_TYPE_VOLUME				4
-
-
 
 // The largest Strata BSDF is FStrataConductor having as size of 24 bytes.
 // So we will have 
@@ -81,9 +75,6 @@ struct FStrataBSDF
 };
 
 
-#define STRATA_MAX_BSDF_COUNT_PER_LAYER	4
-#define STRATA_MAX_LAYER_COUNT			4
-#define STRATA_MAX_TOTAL_BSDF			4
 
 struct FStrataLayer
 {
diff --git a/Engine/Shaders/Shared/StrataDefinitions.h b/Engine/Shaders/Shared/StrataDefinitions.h
new file mode 100644
--- /dev/null
+++ b/Engine/Shaders/Shared/StrataDefinitions.h
@@ -0,0 +1,27 @@
+// Copyright Epic Games, Inc. All Rights Reserved.
+
+/*================================================================================================
+	StrataDefinitions.ush: used in ray tracing shaders and C++ code to define common constants
+	!!! Changing this file requires recompilation of the engine !!!
+=================================================================================================*/
+
+#pragma once
+
+// Change this to force recompilation of all strata dependent shaders (use https://www.random.org/cgi-bin/randbyte?nbytes=4&format=h)
+#define STRATA_SHADER_VERSION 0xc9624d0e 
+
+
+
+#define STRATA_MAX_BSDF_COUNT_PER_LAYER		4
+#define STRATA_MAX_LAYER_COUNT				4
+#define STRATA_MAX_TOTAL_BSDF				4
+
+
+
+#define STRATA_BSDF_TYPE_DIFFUSE_ON			0
+#define STRATA_BSDF_TYPE_DIFFUSE_CHAN		1
+#define STRATA_BSDF_TYPE_DIELECTRIC			2
+#define STRATA_BSDF_TYPE_CONDUCTOR			3
+#define STRATA_BSDF_TYPE_VOLUME				4
+
+
diff --git a/Engine/Source/Editor/MaterialEditor/Private/MaterialEditor.cpp b/Engine/Source/Editor/MaterialEditor/Private/MaterialEditor.cpp
--- a/Engine/Source/Editor/MaterialEditor/Private/MaterialEditor.cpp
+++ b/Engine/Source/Editor/MaterialEditor/Private/MaterialEditor.cpp
@@ -2814,6 +2814,15 @@ void FMaterialEditor::UpdateMaterialinfoList_Old()
 					Line->AddToken(FTextToken::Create(FText::FromString(InterpolatorsString)));
 					Messages.Add(Line);
 				}
+
+				if (FMaterialShaderMap* ShaderMap = MaterialResource->GetGameThreadShaderMap())
+				{
+					const FString StrataMaterialDescription = ShaderMap->GetStrataMaterialDescription();
+
+					TSharedRef<FTokenizedMessage> Line = FTokenizedMessage::Create(EMessageSeverity::Info);
+					Line->AddToken(FTextToken::Create(FText::FromString(StrataMaterialDescription)));
+					Messages.Add(Line);
+				}
 			}
 
 			FString FeatureLevelName;
diff --git a/Engine/Source/Runtime/Engine/Private/Materials/HLSLMaterialTranslator.cpp b/Engine/Source/Runtime/Engine/Private/Materials/HLSLMaterialTranslator.cpp
--- a/Engine/Source/Runtime/Engine/Private/Materials/HLSLMaterialTranslator.cpp
+++ b/Engine/Source/Runtime/Engine/Private/Materials/HLSLMaterialTranslator.cpp
@@ -1100,6 +1100,38 @@ bool FHLSLMaterialTranslator::Translate()
 			TextureExpression->GetExternalTextureParameterInfo(MaterialCompilationOutput.UniformExpressionSet.UniformExternalTextureParameters.AddDefaulted_GetRef());
 		}
 
+		// Output some debug info as comment in code and in the material stat window
+		if (Chunk[MP_FrontMaterial] != INDEX_NONE && CodeChunkToStrataCompilationInfoMap.Contains(Chunk[MP_FrontMaterial]) == true)
+		{
+			const FStrataMaterialCompilationInfo& StrataCompilationInfo = CodeChunkToStrataCompilationInfoMap[Chunk[MP_FrontMaterial]];
+
+			FString StrataMaterialDescription = "";
+
+			StrataMaterialDescription += FString::Printf(TEXT("StrataCompilationInfo - TotalBSDFCount = %i\r\n"), StrataCompilationInfo.TotalBSDFCount);
+
+			for (uint32 LayerIt = 0; LayerIt < StrataCompilationInfo.LayerCount; ++LayerIt)
+			{
+				StrataMaterialDescription += FString::Printf(TEXT("    Layer %i BSDFs:\r\n"), LayerIt);
+
+				const FStrataMaterialCompilationInfo::FLayer& Layer = StrataCompilationInfo.Layers[LayerIt];
+				for (uint32 BSDFIt = 0; BSDFIt < Layer.BSDFCount; ++BSDFIt)
+				{
+					StrataMaterialDescription += FString::Printf(TEXT("        - %s\r\n"), *GetStrataBSDFName(Layer.BSDFs[BSDFIt].Type));
+				}
+			}
+
+			ResourcesString += TEXT("/*");
+			ResourcesString += StrataMaterialDescription;
+			ResourcesString += TEXT("*/");
+
+			MaterialCompilationOutput.StrataMaterialDescription = StrataMaterialDescription;
+		}
+		else
+		{
+			MaterialCompilationOutput.StrataMaterialDescription = "";
+			ResourcesString += "// No Strata material provided \r\n";
+		}
+
 		LoadShaderSourceFileChecked(TEXT("/Engine/Private/MaterialTemplate.ush"), GetShaderPlatform(), MaterialTemplate);
 
 		// Find the string index of the '#line' statement in MaterialTemplate.usf
@@ -7123,6 +7155,23 @@ int32 FHLSLMaterialTranslator::ShadingModel(EMaterialShadingModel InSelectedShad
 	return AddInlinedCodeChunk(MCT_ShadingModel, TEXT("%d"), InSelectedShadingModel);
 }
 
+void FHLSLMaterialTranslator::AddStrataCodeChunk(int32 CodeChunk, FStrataMaterialCompilationInfo& StrataMaterialCompilationInfo)
+{
+	check(CodeChunk != INDEX_NONE);
+	CodeChunkToStrataCompilationInfoMap.Add(CodeChunk, StrataMaterialCompilationInfo);
+}
+
+bool FHLSLMaterialTranslator::ContainsStrataCodeChunk(int32 CodeChunk)
+{
+	return CodeChunk != INDEX_NONE ? CodeChunkToStrataCompilationInfoMap.Contains(CodeChunk) : false;
+}
+
+const FStrataMaterialCompilationInfo& FHLSLMaterialTranslator::GetStrataCompilationInfo(int32 CodeChunk)
+{
+	check(CodeChunk != INDEX_NONE);
+	return CodeChunkToStrataCompilationInfoMap[CodeChunk];
+}
+
 int32 FHLSLMaterialTranslator::FrontMaterial()
 {
 	return AddInlinedCodeChunk(MCT_Strata, TEXT("GetInitialisedStrataData()"));
diff --git a/Engine/Source/Runtime/Engine/Private/Materials/HLSLMaterialTranslator.h b/Engine/Source/Runtime/Engine/Private/Materials/HLSLMaterialTranslator.h
--- a/Engine/Source/Runtime/Engine/Private/Materials/HLSLMaterialTranslator.h
+++ b/Engine/Source/Runtime/Engine/Private/Materials/HLSLMaterialTranslator.h
@@ -29,6 +29,7 @@
 #include "Hash/CityHash.h"
 #include "VT/RuntimeVirtualTexture.h"
 #include "Field/FieldSystemTypes.h"
+#include "Containers/Map.h"
 
 #if WITH_EDITORONLY_DATA
 #include "Materials/MaterialExpressionSceneTexture.h"
@@ -45,6 +46,7 @@
 #include "Containers/LazyPrintf.h"
 #include "Containers/HashTable.h"
 #include "Engine/Texture2D.h"
+#include "StrataMaterial.h"
 #endif
 
 class Error;
@@ -304,6 +306,8 @@ protected:
 	/** Will contain all the shading models picked up from the material expression graph */
 	FMaterialShadingModelField ShadingModelsFromCompilation;
 
+	TMap<int32, FStrataMaterialCompilationInfo> CodeChunkToStrataCompilationInfoMap; // Mapping from node output to StrataMaterial
+
 	/** Tracks the total number of vt samples in the shader. */
 	uint32 NumVtSamples;
 
@@ -777,6 +781,10 @@ protected:
 	virtual int32 StrataArtisticIOR(int32 Reflectivity, int32 EdgeColor, int32 OutputIndex) override;
 	virtual int32 StrataPhysicalIOR(int32 IOR, int32 Extinction, int32 OutputIndex) override;
 
+	virtual void AddStrataCodeChunk(int32 CodeChunk, FStrataMaterialCompilationInfo& StrataMaterialCompilationInfo) override;
+	virtual bool ContainsStrataCodeChunk(int32 CodeChunk) override;
+	virtual const FStrataMaterialCompilationInfo& GetStrataCompilationInfo(int32 CodeChunk) override;
+
 #if HANDLE_CUSTOM_OUTPUTS_AS_MATERIAL_ATTRIBUTES
 	/** Used to translate code for custom output attributes such as ClearCoatBottomNormal */
 	void GenerateCustomAttributeCode(int32 OutputIndex, int32 OutputCode, EMaterialValueType OutputType, FString& DisplayName);
diff --git a/Engine/Source/Runtime/Engine/Private/Materials/MaterialExpressions.cpp b/Engine/Source/Runtime/Engine/Private/Materials/MaterialExpressions.cpp
--- a/Engine/Source/Runtime/Engine/Private/Materials/MaterialExpressions.cpp
+++ b/Engine/Source/Runtime/Engine/Private/Materials/MaterialExpressions.cpp
@@ -256,6 +256,7 @@
 #include "Framework/Notifications/NotificationManager.h"
 #include "Widgets/Notifications/SNotificationList.h"
 #include "MaterialGraph/MaterialGraphSchema.h"
+#include "StrataMaterial.h"
 #endif //WITH_EDITOR
 #include "Materials/MaterialInstanceConstant.h"
 #include "Curves/CurveLinearColorAtlas.h"
@@ -18962,10 +18963,14 @@ int32 UMaterialExpressionStrataDiffuseBSDF::Compile(class FMaterialCompiler* Com
 	{
 		return Compiler->Errorf(TEXT("Missing Normal input"));	// STRATA_TODO: defaults in the detail panel with auto read only if plug in
 	}
-	return Compiler->StrataDiffuseOrenNayarBSDF(
+	int32 OutputCodeChunk = Compiler->StrataDiffuseOrenNayarBSDF(
 		Albedo.GetTracedInput().Expression		? Albedo.Compile(Compiler)		: Compiler->Constant3(0.18f, 0.18f, 0.18f),
 		Roughness.GetTracedInput().Expression	? Roughness.Compile(Compiler)	: Compiler->Constant(0.0f),
 		Normal.Compile(Compiler));
+
+	StrataCreateSingleBSDFMaterial(Compiler, OutputCodeChunk, STRATA_BSDF_TYPE_DIFFUSE_ON);
+
+	return OutputCodeChunk;
 }
 
 void UMaterialExpressionStrataDiffuseBSDF::GetCaption(TArray<FString>& OutCaptions) const
@@ -19023,10 +19028,14 @@ int32 UMaterialExpressionStrataDiffuseChanBSDF::Compile(class FMaterialCompiler*
 	{
 		return Compiler->Errorf(TEXT("Missing Normal input"));	// STRATA_TODO: defaults in the detail panel with auto read only if plug in
 	}
-	return Compiler->StrataDiffuseChanBSDF(
+	int32 OutputCodeChunk = Compiler->StrataDiffuseChanBSDF(
 		Albedo.GetTracedInput().Expression		? Albedo.Compile(Compiler)		: Compiler->Constant3(0.18f, 0.18f, 0.18f),
 		Roughness.GetTracedInput().Expression	? Roughness.Compile(Compiler)	: Compiler->Constant(0.0f),
 		Normal.Compile(Compiler));
+
+	StrataCreateSingleBSDFMaterial(Compiler, OutputCodeChunk, STRATA_BSDF_TYPE_DIFFUSE_CHAN);
+
+	return OutputCodeChunk;
 }
 
 void UMaterialExpressionStrataDiffuseChanBSDF::GetCaption(TArray<FString>& OutCaptions) const
@@ -19084,11 +19093,15 @@ int32 UMaterialExpressionStrataDielectricBSDF::Compile(class FMaterialCompiler*
 	{
 		return Compiler->Errorf(TEXT("Missing Normal input"));	// STRATA_TODO: defaults in the detail panel with auto read only if plug in
 	}
-	return Compiler->StrataDielectricBSDF(
+	int32 OutputCodeChunk = Compiler->StrataDielectricBSDF(
 		Roughness.GetTracedInput().Expression	?	Roughness.Compile(Compiler) : Compiler->Constant2(0.0f, 0.0f),
 		IOR.GetTracedInput().Expression			?	IOR.Compile(Compiler)		: Compiler->Constant(1.5f),
 		Tint.GetTracedInput().Expression		?	Tint.Compile(Compiler)		: Compiler->Constant3(1.0f, 1.0f, 1.0f),
 		Normal.Compile(Compiler));
+
+	StrataCreateSingleBSDFMaterial(Compiler, OutputCodeChunk, STRATA_BSDF_TYPE_DIELECTRIC);
+
+	return OutputCodeChunk;
 }
 
 void UMaterialExpressionStrataDielectricBSDF::GetCaption(TArray<FString>& OutCaptions) const
@@ -19149,11 +19162,15 @@ int32 UMaterialExpressionStrataConductorBSDF::Compile(class FMaterialCompiler* C
 	{
 		return Compiler->Errorf(TEXT("Missing Normal input"));		// STRATA_TODO: defaults in the detail panel with auto read only if plug in
 	}
-	return Compiler->StrataConductorBSDF(
+	int32 OutputCodeChunk = Compiler->StrataConductorBSDF(
 		Reflectivity.GetTracedInput().Expression	? Reflectivity.Compile(Compiler): Compiler->Constant3(0.947f, 0.776f, 0.371f),	// Default to Gold
 		EdgeColor.GetTracedInput().Expression		? EdgeColor.Compile(Compiler)	: Compiler->Constant3(1.000f, 0.982f, 0.753f),	// Default to Gold
 		Roughness.GetTracedInput().Expression		? Roughness.Compile(Compiler)	: Compiler->Constant2(0.0f, 0.0f),
 		Normal.Compile(Compiler));
+
+	StrataCreateSingleBSDFMaterial(Compiler, OutputCodeChunk, STRATA_BSDF_TYPE_CONDUCTOR);
+
+	return OutputCodeChunk;
 }
 
 void UMaterialExpressionStrataConductorBSDF::GetCaption(TArray<FString>& OutCaptions) const
@@ -19210,11 +19227,15 @@ UMaterialExpressionStrataVolumeBSDF::UMaterialExpressionStrataVolumeBSDF(const F
 #if WITH_EDITOR
 int32 UMaterialExpressionStrataVolumeBSDF::Compile(class FMaterialCompiler* Compiler, int32 OutputIndex)
 {
-	return Compiler->StrataVolumeBSDF(
+	int32 OutputCodeChunk = Compiler->StrataVolumeBSDF(
 		Albedo.GetTracedInput().Expression		? Albedo.Compile(Compiler)			: Compiler->Constant3(0.0f, 0.0f, 0.0f),
 		Extinction.GetTracedInput().Expression	? Extinction.Compile(Compiler)		: Compiler->Constant3(0.0f, 0.0f, 0.0f),
 		Anisotropy.GetTracedInput().Expression	? Anisotropy.Compile(Compiler)		: Compiler->Constant(0.0f),
-		Thickness.GetTracedInput().Expression	? Thickness.Compile(Compiler)		: Compiler->Constant(0.001f)); // default = 1mm
+		Thickness.GetTracedInput().Expression ? Thickness.Compile(Compiler)			: Compiler->Constant(0.001f)); // default = 1mm
+
+	StrataCreateSingleBSDFMaterial(Compiler, OutputCodeChunk, STRATA_BSDF_TYPE_VOLUME);
+
+	return OutputCodeChunk;
 }
 
 void UMaterialExpressionStrataVolumeBSDF::GetCaption(TArray<FString>& OutCaptions) const
@@ -19277,10 +19298,23 @@ int32 UMaterialExpressionStrataHorizontalMixing::Compile(class FMaterialCompiler
 	{
 		return Compiler->Errorf(TEXT("Missing Background input"));
 	}
-	return Compiler->StrataHorizontalMixing(
-		Foreground.Compile(Compiler), 
-		Background.Compile(Compiler), 
+
+	int32 ForegroundCodeChunk = Foreground.Compile(Compiler);
+	int32 BackgroundCodeChunk = Background.Compile(Compiler);
+
+	int32 OutputCodeChunk = Compiler->StrataHorizontalMixing(
+		ForegroundCodeChunk,
+		BackgroundCodeChunk,
 		Mix.GetTracedInput().Expression ? Mix.Compile(Compiler) : Compiler->Constant(0.0f));
+
+	if (!Compiler->ContainsStrataCodeChunk(ForegroundCodeChunk) || !Compiler->ContainsStrataCodeChunk(BackgroundCodeChunk))
+	{
+		return Compiler->Errorf(TEXT("Could not find ForegroundCodeChunk or BackgroundCodeChunk to mix"));
+	}
+	FStrataMaterialCompilationInfo StrataInfo = StrataHorizontalMixing(Compiler, Compiler->GetStrataCompilationInfo(ForegroundCodeChunk), Compiler->GetStrataCompilationInfo(BackgroundCodeChunk));
+	Compiler->AddStrataCodeChunk(OutputCodeChunk, StrataInfo);
+
+	return OutputCodeChunk;
 }
 
 void UMaterialExpressionStrataHorizontalMixing::GetCaption(TArray<FString>& OutCaptions) const
@@ -19326,7 +19360,20 @@ int32 UMaterialExpressionStrataVerticalLayering::Compile(class FMaterialCompiler
 	{
 		return Compiler->Errorf(TEXT("Missing Base input"));
 	}
-	return Compiler->StrataVerticalLayering(Top.Compile(Compiler), Base.Compile(Compiler));
+
+	int32 TopCodeChunk = Top.Compile(Compiler);
+	int32 BaseCodeChunk = Base.Compile(Compiler);
+
+	int32 OutputCodeChunk = Compiler->StrataVerticalLayering(TopCodeChunk, BaseCodeChunk);
+
+	if (!Compiler->ContainsStrataCodeChunk(TopCodeChunk) || !Compiler->ContainsStrataCodeChunk(BaseCodeChunk))
+	{
+		return Compiler->Errorf(TEXT("Could not find TopCodeChunk or BaseCodeChunk to layer"));
+	}
+	FStrataMaterialCompilationInfo StrataInfo = StrataVerticalLayering(Compiler, Compiler->GetStrataCompilationInfo(TopCodeChunk), Compiler->GetStrataCompilationInfo(BaseCodeChunk));
+	Compiler->AddStrataCodeChunk(OutputCodeChunk, StrataInfo);
+
+	return OutputCodeChunk;
 }
 
 void UMaterialExpressionStrataVerticalLayering::GetCaption(TArray<FString>& OutCaptions) const
@@ -19372,7 +19419,20 @@ int32 UMaterialExpressionStrataAdd::Compile(class FMaterialCompiler* Compiler, i
 	{
 		return Compiler->Errorf(TEXT("Missing B input"));
 	}
-	return Compiler->StrataAdd(A.Compile(Compiler), B.Compile(Compiler));
+
+	int32 ACodeChunk = A.Compile(Compiler);
+	int32 BCodeChunk = B.Compile(Compiler);
+
+	int32 OutputCodeChunk = Compiler->StrataAdd(ACodeChunk, BCodeChunk);
+
+	if (!Compiler->ContainsStrataCodeChunk(ACodeChunk) || !Compiler->ContainsStrataCodeChunk(BCodeChunk))
+	{
+		return Compiler->Errorf(TEXT("Could not find ACodeChunk or BCodeChunk to add"));
+	}
+	FStrataMaterialCompilationInfo StrataInfo = StrataAdd(Compiler, Compiler->GetStrataCompilationInfo(ACodeChunk), Compiler->GetStrataCompilationInfo(BCodeChunk));
+	Compiler->AddStrataCodeChunk(OutputCodeChunk, StrataInfo);
+
+	return OutputCodeChunk;
 }
 
 void UMaterialExpressionStrataAdd::GetCaption(TArray<FString>& OutCaptions) const
@@ -19418,7 +19478,20 @@ int32 UMaterialExpressionStrataMultiply::Compile(class FMaterialCompiler* Compil
 	{
 		return Compiler->Errorf(TEXT("Missing Base input"));
 	}
-	return Compiler->StrataMultiply(A.Compile(Compiler), Weight.Compile(Compiler));
+
+	int32 ACodeChunk = A.Compile(Compiler);
+	int32 WeightCodeChunk = Weight.Compile(Compiler);
+
+	int32 OutputCodeChunk = Compiler->StrataMultiply(ACodeChunk, WeightCodeChunk);
+
+	if (!Compiler->ContainsStrataCodeChunk(ACodeChunk))
+	{
+		return Compiler->Errorf(TEXT("Could not find ACodeChunk to multiply"));
+	}
+	FStrataMaterialCompilationInfo StrataInfo = StrataMultiply(Compiler, Compiler->GetStrataCompilationInfo(ACodeChunk));
+	Compiler->AddStrataCodeChunk(OutputCodeChunk, StrataInfo);
+
+	return OutputCodeChunk;
 }
 
 void UMaterialExpressionStrataMultiply::GetCaption(TArray<FString>& OutCaptions) const
diff --git a/Engine/Source/Runtime/Engine/Private/Materials/StrataMaterial.cpp b/Engine/Source/Runtime/Engine/Private/Materials/StrataMaterial.cpp
new file mode 100644
--- /dev/null
+++ b/Engine/Source/Runtime/Engine/Private/Materials/StrataMaterial.cpp
@@ -0,0 +1,135 @@
+// Copyright Epic Games, Inc. All Rights Reserved.
+
+#include "StrataMaterial.h"
+#include "MaterialCompiler.h"
+
+
+
+FString GetStrataBSDFName(uint8 BSDFType)
+{
+	switch (BSDFType)
+	{
+	case STRATA_BSDF_TYPE_DIFFUSE_ON:
+		return TEXT("DIFFUSE_ON");
+		break;
+	case STRATA_BSDF_TYPE_DIFFUSE_CHAN:
+		return TEXT("DIFFUSE_CHAN");
+		break;
+	case STRATA_BSDF_TYPE_DIELECTRIC:
+		return TEXT("DIELECTRIC");
+		break;
+	case STRATA_BSDF_TYPE_CONDUCTOR:
+		return TEXT("CONDUCTOR");
+		break;
+	case STRATA_BSDF_TYPE_VOLUME:
+		return TEXT("VOLUME");
+		break;
+	}
+	check(false);
+	return "";
+}
+
+static void UpdateTotalBSDFCount(FStrataMaterialCompilationInfo& StrataInfo)
+{
+	StrataInfo.TotalBSDFCount = 0;
+	for (uint32 LayerIt = 0; LayerIt < StrataInfo.LayerCount; ++LayerIt)
+	{
+		StrataInfo.TotalBSDFCount += StrataInfo.Layers[LayerIt].BSDFCount;
+	}
+}
+
+
+void StrataCreateSingleBSDFMaterial(FMaterialCompiler* Compiler, int32 CodeChunk, uint8 BSDFType)
+{
+	FStrataMaterialCompilationInfo StrataInfo;
+	StrataInfo.LayerCount = 1;
+	StrataInfo.Layers[0].BSDFCount = 1;
+	StrataInfo.Layers[0].BSDFs[0].Type = BSDFType;
+	UpdateTotalBSDFCount(StrataInfo);
+	Compiler->AddStrataCodeChunk(CodeChunk, StrataInfo);
+}
+
+
+FStrataMaterialCompilationInfo StrataAdd(FMaterialCompiler* Compiler, const FStrataMaterialCompilationInfo& A, const FStrataMaterialCompilationInfo& B)
+{
+	if ((A.TotalBSDFCount + B.TotalBSDFCount) > STRATA_MAX_TOTAL_BSDF) // See StrataDefinitions.h
+	{
+		Compiler->Error(TEXT("Adding would result in too many BSDFs"));
+		return A;
+	}
+
+	FStrataMaterialCompilationInfo StrataInfo = A;
+
+	// Append each BSDF from B to A, with same layer position
+	for (uint32 LayerIt = 0; LayerIt < B.LayerCount; ++LayerIt)
+	{
+		const FStrataMaterialCompilationInfo::FLayer& ALayer = A.Layers[LayerIt];
+		const FStrataMaterialCompilationInfo::FLayer& BLayer = B.Layers[LayerIt];
+
+		if ((ALayer.BSDFCount + BLayer.BSDFCount) > STRATA_MAX_BSDF_COUNT_PER_LAYER) // See StrataDefinitions.h
+		{
+			Compiler->Error(TEXT("Adding would result in too many BSDFs in a Layer"));
+			return A;
+		}
+
+		for (uint32 BSDF = 0; BSDF < BLayer.BSDFCount; BSDF++)
+		{
+			StrataInfo.Layers[LayerIt].BSDFs[ALayer.BSDFCount + BSDF] = BLayer.BSDFs[BSDF];
+		}
+
+		StrataInfo.Layers[LayerIt].BSDFCount = ALayer.BSDFCount + BLayer.BSDFCount;
+	}
+	StrataInfo.LayerCount = FMath::Max(A.LayerCount, B.LayerCount);
+
+	UpdateTotalBSDFCount(StrataInfo);
+	return StrataInfo;
+}
+
+
+FStrataMaterialCompilationInfo StrataMultiply(FMaterialCompiler* Compiler, const FStrataMaterialCompilationInfo& A)
+{
+	FStrataMaterialCompilationInfo StrataInfo = A;
+	return StrataInfo;
+}
+
+
+FStrataMaterialCompilationInfo StrataHorizontalMixing(FMaterialCompiler* Compiler, const FStrataMaterialCompilationInfo& A, const FStrataMaterialCompilationInfo& B)
+{
+	if ((A.TotalBSDFCount + B.TotalBSDFCount) > STRATA_MAX_TOTAL_BSDF) // See StrataDefinitions.h
+	{
+		Compiler->Error(TEXT("Mixing would result in too many BSDFs"));
+		return A;
+	}
+
+	return StrataAdd(Compiler, A, B); // Mixing is a similar operation to Add when it comes to bsdf count
+}
+
+
+FStrataMaterialCompilationInfo StrataVerticalLayering(FMaterialCompiler* Compiler, const FStrataMaterialCompilationInfo& Top, const FStrataMaterialCompilationInfo& Base)
+{
+	if ((Top.TotalBSDFCount + Base.TotalBSDFCount) > STRATA_MAX_TOTAL_BSDF) // See StrataDefinitions.h
+	{
+		Compiler->Error(TEXT("Layering would result in too many BSDFs"));
+		return Base;
+	}
+
+	if ((Top.LayerCount + Base.LayerCount) > STRATA_MAX_LAYER_COUNT) // See StrataDefinitions.h
+	{
+		Compiler->Error(TEXT("Layering would result in too many Layers"));
+		return Base;
+	}
+
+	FStrataMaterialCompilationInfo StrataInfo = Top;
+
+	// Add each layer from Base under Top
+	const uint32 TopLayerCount = Top.LayerCount;
+	for (uint32 LayerIt = 0; LayerIt < Base.LayerCount; ++LayerIt)
+	{
+		StrataInfo.Layers[TopLayerCount + LayerIt] = Base.Layers[LayerIt];
+	}
+	StrataInfo.LayerCount += Base.LayerCount;
+
+	UpdateTotalBSDFCount(StrataInfo);
+	return StrataInfo;
+}
+
diff --git a/Engine/Source/Runtime/Engine/Private/Materials/StrataMaterial.h b/Engine/Source/Runtime/Engine/Private/Materials/StrataMaterial.h
new file mode 100644
--- /dev/null
+++ b/Engine/Source/Runtime/Engine/Private/Materials/StrataMaterial.h
@@ -0,0 +1,49 @@
+// Copyright Epic Games, Inc. All Rights Reserved.
+
+#pragma once
+
+#include "CoreMinimal.h"
+#include "StrataDefinitions.h"
+
+
+
+class FMaterialCompiler;
+
+
+
+struct FStrataMaterialCompilationInfo
+{
+	struct FBSDF
+	{
+		uint8 Type;
+	};
+
+	struct FLayer
+	{
+		uint32 BSDFCount;
+		FBSDF BSDFs[STRATA_MAX_BSDF_COUNT_PER_LAYER];
+	};
+
+	uint32 LayerCount;
+	FLayer Layers[STRATA_MAX_LAYER_COUNT];
+
+	uint32 TotalBSDFCount;
+
+	FStrataMaterialCompilationInfo()
+	{
+		TotalBSDFCount = 0;
+		LayerCount = 0;
+		memset(Layers, 0, sizeof(Layers));
+	}
+};
+
+FString GetStrataBSDFName(uint8 BSDFType);
+
+void StrataCreateSingleBSDFMaterial(FMaterialCompiler* Compiler, int32 CodeChunk, uint8 BSDFType);
+
+FStrataMaterialCompilationInfo StrataAdd(FMaterialCompiler* Compiler, const FStrataMaterialCompilationInfo& A, const FStrataMaterialCompilationInfo& B);
+FStrataMaterialCompilationInfo StrataMultiply(FMaterialCompiler* Compiler, const FStrataMaterialCompilationInfo& A);
+FStrataMaterialCompilationInfo StrataHorizontalMixing(FMaterialCompiler* Compiler, const FStrataMaterialCompilationInfo& A, const FStrataMaterialCompilationInfo& B);
+FStrataMaterialCompilationInfo StrataVerticalLayering(FMaterialCompiler* Compiler, const FStrataMaterialCompilationInfo& Top, const FStrataMaterialCompilationInfo& Base);
+
+
diff --git a/Engine/Source/Runtime/Engine/Public/MaterialCompiler.h b/Engine/Source/Runtime/Engine/Public/MaterialCompiler.h
--- a/Engine/Source/Runtime/Engine/Public/MaterialCompiler.h
+++ b/Engine/Source/Runtime/Engine/Public/MaterialCompiler.h
@@ -20,6 +20,7 @@ class UMaterialParameterCollection;
 class URuntimeVirtualTexture;
 class UTexture;
 struct FMaterialParameterInfo;
+struct FStrataMaterialCompilationInfo;
 
 enum EMaterialForceCastFlags
 {
@@ -401,6 +402,10 @@ public:
 	virtual int32 StrataArtisticIOR(int32 Reflectivity, int32 EdgeColor, int32 OutputIndex) = 0;
 	virtual int32 StrataPhysicalIOR(int32 IOR, int32 Extinction, int32 OutputIndex) = 0;
 
+	virtual void AddStrataCodeChunk(int32 CodeChunk, FStrataMaterialCompilationInfo& StrataMaterialCompilationInfo) = 0;
+	virtual bool ContainsStrataCodeChunk(int32 CodeChunk) = 0;
+	virtual const FStrataMaterialCompilationInfo& GetStrataCompilationInfo(int32 CodeChunk) = 0;
+
 	// Water
 	virtual int32 SceneDepthWithoutWater(int32 Offset, int32 ViewportUV, bool bUseOffset, float FallbackDepth) = 0;
 
@@ -849,6 +854,21 @@ public:
 		return Compiler->StrataPhysicalIOR(IOR, Extinction, OutputIndex);
 	}
 
+	virtual void AddStrataCodeChunk(int32 CodeChunk, FStrataMaterialCompilationInfo& StrataMaterialCompilationInfo) override
+	{
+		Compiler->AddStrataCodeChunk(CodeChunk, StrataMaterialCompilationInfo);
+	}
+
+	virtual bool ContainsStrataCodeChunk(int32 CodeChunk) override
+	{
+		return Compiler->ContainsStrataCodeChunk(CodeChunk);
+	}
+	
+	virtual const FStrataMaterialCompilationInfo& GetStrataCompilationInfo(int32 CodeChunk) override
+	{
+		return Compiler->GetStrataCompilationInfo(CodeChunk);
+	}
+
 protected:
 		
 	FMaterialCompiler* Compiler;
diff --git a/Engine/Source/Runtime/Engine/Public/MaterialShared.h b/Engine/Source/Runtime/Engine/Public/MaterialShared.h
--- a/Engine/Source/Runtime/Engine/Public/MaterialShared.h
+++ b/Engine/Source/Runtime/Engine/Public/MaterialShared.h
@@ -593,6 +593,7 @@ public:
 		EstimatedNumVirtualTextureLookups(0),
 		NumUsedUVScalars(0),
 		NumUsedCustomInterpolatorScalars(0),
+		StrataMaterialDescription(),
 #endif
 		RuntimeVirtualTextureOutputAttributeMask(0),
 		bNeedsSceneTextures(false),
@@ -659,6 +660,9 @@ public:
 	/** Number of used custom vertex interpolation scalars. */
 	LAYOUT_FIELD_EDITORONLY(uint8, NumUsedCustomInterpolatorScalars);
 
+	/** The Strata material layout */
+	LAYOUT_FIELD_EDITORONLY(FMemoryImageString, StrataMaterialDescription);
+
 	/** Bitfield of runtime virtual texture output attributes. */
 	LAYOUT_FIELD(uint8, RuntimeVirtualTextureOutputAttributeMask);
 
@@ -1220,6 +1224,7 @@ public:
 	uint32 GetNumUsedCustomInterpolatorScalars() const { return GetContent()->MaterialCompilationOutput.NumUsedCustomInterpolatorScalars; }
 	void GetEstimatedNumTextureSamples(uint32& VSSamples, uint32& PSSamples) const { VSSamples = GetContent()->MaterialCompilationOutput.EstimatedNumTextureSamplesVS; PSSamples = GetContent()->MaterialCompilationOutput.EstimatedNumTextureSamplesPS; }
 	uint32 GetEstimatedNumVirtualTextureLookups() const { return GetContent()->MaterialCompilationOutput.EstimatedNumVirtualTextureLookups; }
+	const FString GetStrataMaterialDescription() const { return GetContent()->MaterialCompilationOutput.StrataMaterialDescription; }
 #endif
 	uint32 GetNumVirtualTextureStacks() const { return GetContent()->MaterialCompilationOutput.UniformExpressionSet.VTStacks.Num(); }
 	uint8 GetRuntimeVirtualTextureOutputAttributeMask() const { return GetContent()->MaterialCompilationOutput.RuntimeVirtualTextureOutputAttributeMask; }
