Merging using //UE5/Main_to_//UE4/Release-4.26 @ CL 14521949

Fixed issue with sync group collisions when syncing with 'component' scope

This change moves sync group management to use a name->group map instead of indices. Staying with indices cannot work across multiple instances as they are simply allocated at the point the single anim BP is compiled, so can easily collide in anything other than simple configurations.
Deprecated all index based access to sync groups in the runtime. Sync groups can still be enumerated on the generated class.

diff --git a/Engine/Source/Editor/AnimGraph/Classes/AnimBlueprintPostCompileValidation.h b/Engine/Source/Editor/AnimGraph/Classes/AnimBlueprintPostCompileValidation.h
--- a/Engine/Source/Editor/AnimGraph/Classes/AnimBlueprintPostCompileValidation.h
+++ b/Engine/Source/Editor/AnimGraph/Classes/AnimBlueprintPostCompileValidation.h
@@ -52,10 +52,11 @@ protected:
 	struct FPCV_GatherParams
 	{
 		bool bFilterBySyncGroup;
-		int32 SyncGroupIndex;
+		FName SyncGroupName;
 		bool bFilterByLoopingCondition;
 		bool bLoopingCondition;
 
+		UE_DEPRECATED(4.26, "Please use the constructor that takes an FName for the sync group")
 		FPCV_GatherParams
 		(
 			bool InbFilterBySyncGroup = false,
@@ -64,7 +65,20 @@ protected:
 			bool InbLoopingCondition = false
 		)
 			: bFilterBySyncGroup(InbFilterBySyncGroup)
-			, SyncGroupIndex(InSyncGroupIndex)
+			, SyncGroupName(NAME_None)
+			, bFilterByLoopingCondition(InbFilterByLoopingCondition)
+			, bLoopingCondition(InbLoopingCondition)
+		{}
+
+		FPCV_GatherParams
+		(
+			bool InbFilterBySyncGroup = false,
+			FName InSyncGroupName = NAME_None,
+			bool InbFilterByLoopingCondition = false,
+			bool InbLoopingCondition = false
+		)
+			: bFilterBySyncGroup(InbFilterBySyncGroup)
+			, SyncGroupName(InSyncGroupName)
 			, bFilterByLoopingCondition(InbFilterByLoopingCondition)
 			, bLoopingCondition(InbLoopingCondition)
 		{}
diff --git a/Engine/Source/Editor/AnimGraph/Private/AnimBlueprintPostCompileValidation.cpp b/Engine/Source/Editor/AnimGraph/Private/AnimBlueprintPostCompileValidation.cpp
--- a/Engine/Source/Editor/AnimGraph/Private/AnimBlueprintPostCompileValidation.cpp
+++ b/Engine/Source/Editor/AnimGraph/Private/AnimBlueprintPostCompileValidation.cpp
@@ -145,7 +145,7 @@ void UAnimBlueprintPostCompileValidation::PCV_GatherAnimSequencesFromGraph(TArra
 		{
 			if (const FAnimNode_BlendSpacePlayer* const BlendSpacePlayer = Property->ContainerPtrToValuePtr<FAnimNode_BlendSpacePlayer>(PCV_Params.DefaultAnimInstance))
 			{
-				const bool bPassSyncGroupFilter = !GatherParams.bFilterBySyncGroup || (BlendSpacePlayer->GroupIndex == GatherParams.SyncGroupIndex);
+				const bool bPassSyncGroupFilter = !GatherParams.bFilterBySyncGroup || (BlendSpacePlayer->GroupName == GatherParams.SyncGroupName);
 				const bool bPassLoopingFilter = !GatherParams.bFilterByLoopingCondition || (BlendSpacePlayer->bLoop == GatherParams.bLoopingCondition);
 				if (bPassSyncGroupFilter && bPassLoopingFilter)
 				{
@@ -157,7 +157,7 @@ void UAnimBlueprintPostCompileValidation::PCV_GatherAnimSequencesFromGraph(TArra
 		{
 			if (const FAnimNode_SequencePlayer* const SequencePlayer = Property->ContainerPtrToValuePtr<FAnimNode_SequencePlayer>(PCV_Params.DefaultAnimInstance))
 			{
-				const bool bPassSyncGroupFilter = !GatherParams.bFilterBySyncGroup || (SequencePlayer->GroupIndex == GatherParams.SyncGroupIndex);
+				const bool bPassSyncGroupFilter = !GatherParams.bFilterBySyncGroup || (SequencePlayer->GroupName == GatherParams.SyncGroupName);
 				const bool bPassLoopingFilter = !GatherParams.bFilterByLoopingCondition || (SequencePlayer->bLoopAnimation == GatherParams.bLoopingCondition);
 				if (bPassSyncGroupFilter && bPassLoopingFilter)
 				{
diff --git a/Engine/Source/Editor/AnimGraph/Private/AnimGraphNode_BlendSpaceEvaluator.cpp b/Engine/Source/Editor/AnimGraph/Private/AnimGraphNode_BlendSpaceEvaluator.cpp
--- a/Engine/Source/Editor/AnimGraph/Private/AnimGraphNode_BlendSpaceEvaluator.cpp
+++ b/Engine/Source/Editor/AnimGraph/Private/AnimGraphNode_BlendSpaceEvaluator.cpp
@@ -136,7 +136,8 @@ void UAnimGraphNode_BlendSpaceEvaluator::ValidateAnimNodeDuringCompilation(class
 void UAnimGraphNode_BlendSpaceEvaluator::BakeDataDuringCompilation(class FCompilerResultsLog& MessageLog)
 {
 	UAnimBlueprint* AnimBlueprint = GetAnimBlueprint();
-	Node.GroupIndex = AnimBlueprint->FindOrAddGroup(SyncGroup.GroupName);
+	AnimBlueprint->FindOrAddGroup(SyncGroup.GroupName);
+	Node.GroupName = SyncGroup.GroupName;
 	Node.GroupRole = SyncGroup.GroupRole;
 	Node.GroupScope = SyncGroup.GroupScope;
 }
diff --git a/Engine/Source/Editor/AnimGraph/Private/AnimGraphNode_BlendSpacePlayer.cpp b/Engine/Source/Editor/AnimGraph/Private/AnimGraphNode_BlendSpacePlayer.cpp
--- a/Engine/Source/Editor/AnimGraph/Private/AnimGraphNode_BlendSpacePlayer.cpp
+++ b/Engine/Source/Editor/AnimGraph/Private/AnimGraphNode_BlendSpacePlayer.cpp
@@ -135,7 +135,8 @@ void UAnimGraphNode_BlendSpacePlayer::ValidateAnimNodeDuringCompilation(class US
 void UAnimGraphNode_BlendSpacePlayer::BakeDataDuringCompilation(class FCompilerResultsLog& MessageLog)
 {
 	UAnimBlueprint* AnimBlueprint = GetAnimBlueprint();
-	Node.GroupIndex = AnimBlueprint->FindOrAddGroup(SyncGroup.GroupName);
+	AnimBlueprint->FindOrAddGroup(SyncGroup.GroupName);
+	Node.GroupName = SyncGroup.GroupName;
 	Node.GroupRole = SyncGroup.GroupRole;
 	Node.GroupScope = SyncGroup.GroupScope;
 }
diff --git a/Engine/Source/Editor/AnimGraph/Private/AnimGraphNode_SequenceEvaluator.cpp b/Engine/Source/Editor/AnimGraph/Private/AnimGraphNode_SequenceEvaluator.cpp
--- a/Engine/Source/Editor/AnimGraph/Private/AnimGraphNode_SequenceEvaluator.cpp
+++ b/Engine/Source/Editor/AnimGraph/Private/AnimGraphNode_SequenceEvaluator.cpp
@@ -29,7 +29,8 @@ void UAnimGraphNode_SequenceEvaluator::PreloadRequiredAssets()
 void UAnimGraphNode_SequenceEvaluator::BakeDataDuringCompilation(class FCompilerResultsLog& MessageLog)
 {
 	UAnimBlueprint* AnimBlueprint = GetAnimBlueprint();
-	Node.GroupIndex = AnimBlueprint->FindOrAddGroup(SyncGroup.GroupName);
+	AnimBlueprint->FindOrAddGroup(SyncGroup.GroupName);
+	Node.GroupName = SyncGroup.GroupName;
 	Node.GroupRole = SyncGroup.GroupRole;
 	Node.GroupScope = SyncGroup.GroupScope;
 }
diff --git a/Engine/Source/Editor/AnimGraph/Private/AnimGraphNode_SequencePlayer.cpp b/Engine/Source/Editor/AnimGraph/Private/AnimGraphNode_SequencePlayer.cpp
--- a/Engine/Source/Editor/AnimGraph/Private/AnimGraphNode_SequencePlayer.cpp
+++ b/Engine/Source/Editor/AnimGraph/Private/AnimGraphNode_SequencePlayer.cpp
@@ -379,7 +379,8 @@ void UAnimGraphNode_SequencePlayer::SetAnimationAsset(UAnimationAsset* Asset)
 void UAnimGraphNode_SequencePlayer::BakeDataDuringCompilation(class FCompilerResultsLog& MessageLog)
 {
 	UAnimBlueprint* AnimBlueprint = GetAnimBlueprint();
-	Node.GroupIndex = AnimBlueprint->FindOrAddGroup(SyncGroup.GroupName);
+	AnimBlueprint->FindOrAddGroup(SyncGroup.GroupName);
+	Node.GroupName = SyncGroup.GroupName;
 	Node.GroupRole = SyncGroup.GroupRole;
 	Node.GroupScope = SyncGroup.GroupScope;
 }
diff --git a/Engine/Source/Runtime/AnimGraphRuntime/Private/AnimNodes/AnimNode_BlendSpacePlayer.cpp b/Engine/Source/Runtime/AnimGraphRuntime/Private/AnimNodes/AnimNode_BlendSpacePlayer.cpp
--- a/Engine/Source/Runtime/AnimGraphRuntime/Private/AnimNodes/AnimNode_BlendSpacePlayer.cpp
+++ b/Engine/Source/Runtime/AnimGraphRuntime/Private/AnimNodes/AnimNode_BlendSpacePlayer.cpp
@@ -85,7 +85,7 @@ void FAnimNode_BlendSpacePlayer::UpdateInternal(const FAnimationUpdateContext& C
 	{
 		// Create a tick record and fill it out
 		FAnimGroupInstance* SyncGroup;
-		FAnimTickRecord& TickRecord = Context.AnimInstanceProxy->CreateUninitializedTickRecordInScope(GroupIndex, GroupScope, /*out*/ SyncGroup);
+		FAnimTickRecord& TickRecord = Context.AnimInstanceProxy->CreateUninitializedTickRecordInScope(/*out*/ SyncGroup, GroupName, GroupScope);
 
 		const FVector BlendInput(X, Y, Z);
 	
diff --git a/Engine/Source/Runtime/AnimGraphRuntime/Private/AnimNodes/AnimNode_RandomPlayer.cpp b/Engine/Source/Runtime/AnimGraphRuntime/Private/AnimNodes/AnimNode_RandomPlayer.cpp
--- a/Engine/Source/Runtime/AnimGraphRuntime/Private/AnimNodes/AnimNode_RandomPlayer.cpp
+++ b/Engine/Source/Runtime/AnimGraphRuntime/Private/AnimNodes/AnimNode_RandomPlayer.cpp
@@ -231,7 +231,7 @@ void FAnimNode_RandomPlayer::Update_AnyThread(const FAnimationUpdateContext& Con
 
 	FAnimInstanceProxy* AnimProxy = Context.AnimInstanceProxy;
 	FAnimGroupInstance* SyncGroup;
-	FAnimTickRecord& TickRecord = AnimProxy->CreateUninitializedTickRecord(INDEX_NONE, SyncGroup);
+	FAnimTickRecord& TickRecord = AnimProxy->CreateUninitializedTickRecord(SyncGroup, NAME_None);
 	AnimProxy->MakeSequenceTickRecord(
 	    TickRecord, CurrentData->Entry->Sequence, true, CurrentData->PlayRate,
 	    CurrentData->BlendWeight, CurrentData->CurrentPlayTime, CurrentData->MarkerTickRecord);
@@ -240,7 +240,7 @@ void FAnimNode_RandomPlayer::Update_AnyThread(const FAnimationUpdateContext& Con
 
 	if (FAnimationRuntime::HasWeight(NextData->BlendWeight))
 	{
-		FAnimTickRecord& NextTickRecord = AnimProxy->CreateUninitializedTickRecord(INDEX_NONE, SyncGroup);
+		FAnimTickRecord& NextTickRecord = AnimProxy->CreateUninitializedTickRecord(SyncGroup, NAME_None);
 		AnimProxy->MakeSequenceTickRecord(
 		    NextTickRecord, NextData->Entry->Sequence, true, NextData->PlayRate,
 		    NextData->BlendWeight, NextData->CurrentPlayTime, NextData->MarkerTickRecord);
diff --git a/Engine/Source/Runtime/AnimGraphRuntime/Private/AnimNodes/AnimNode_SequenceEvaluator.cpp b/Engine/Source/Runtime/AnimGraphRuntime/Private/AnimNodes/AnimNode_SequenceEvaluator.cpp
--- a/Engine/Source/Runtime/AnimGraphRuntime/Private/AnimNodes/AnimNode_SequenceEvaluator.cpp
+++ b/Engine/Source/Runtime/AnimGraphRuntime/Private/AnimNodes/AnimNode_SequenceEvaluator.cpp
@@ -38,7 +38,7 @@ void FAnimNode_SequenceEvaluator::UpdateAssetPlayer(const FAnimationUpdateContex
 		// Clamp input to a valid position on this sequence's time line.
 		ExplicitTime = FMath::Clamp(ExplicitTime, 0.f, Sequence->GetPlayLength());
 
-		if ((!bTeleportToExplicitTime || (GroupIndex != INDEX_NONE)) && (Context.AnimInstanceProxy->IsSkeletonCompatible(Sequence->GetSkeleton())))
+		if ((!bTeleportToExplicitTime || (GroupName != NAME_None)) && (Context.AnimInstanceProxy->IsSkeletonCompatible(Sequence->GetSkeleton())))
 		{
 			if (bReinitialized)
 			{
diff --git a/Engine/Source/Runtime/Engine/Classes/Animation/AnimMontage.h b/Engine/Source/Runtime/Engine/Classes/Animation/AnimMontage.h
--- a/Engine/Source/Runtime/Engine/Classes/Animation/AnimMontage.h
+++ b/Engine/Source/Runtime/Engine/Classes/Animation/AnimMontage.h
@@ -362,8 +362,8 @@ private:
 	// transient value of previous position before move
 	float PreviousPosition;
 
-	// sync group index
-	int32 SyncGroupIndex;
+	// sync group name
+	FName SyncGroupName;
 
 	/**
 	 * Optional evaluation range to use next update (ignoring the real delta time).
@@ -399,7 +399,9 @@ public:
 	float GetWeight() const { return Blend.GetBlendedValue(); }
 	float GetDesiredWeight() const { return Blend.GetDesiredValue(); }
 	float GetBlendTime() const { return Blend.GetBlendTime(); }
-	int32 GetSyncGroupIndex() const { return SyncGroupIndex;  }
+	UE_DEPRECATED(4.26, "Please use GetSyncGroupName")
+	int32 GetSyncGroupIndex() const { return INDEX_NONE;  }
+	FName GetSyncGroupName() const { return SyncGroupName;  }
 
 	/** Set the weight */
 	void SetWeight(float InValue) { Blend.SetAlpha(InValue); }
diff --git a/Engine/Source/Runtime/Engine/Classes/Animation/AnimNode_AssetPlayerBase.h b/Engine/Source/Runtime/Engine/Classes/Animation/AnimNode_AssetPlayerBase.h
--- a/Engine/Source/Runtime/Engine/Classes/Animation/AnimNode_AssetPlayerBase.h
+++ b/Engine/Source/Runtime/Engine/Classes/Animation/AnimNode_AssetPlayerBase.h
@@ -43,9 +43,14 @@ struct ENGINE_API FAnimNode_AssetPlayerBase : public FAnimNode_Base
 	/** Update method for the asset player, to be implemented by derived classes */
 	virtual void UpdateAssetPlayer(const FAnimationUpdateContext& Context) {};
 
-	// The group index, assigned at compile time based on the editoronly GroupName (or INDEX_NONE if it is not part of any group)
+	// The group name (NAME_None if it is not part of any group)
 	UPROPERTY()
-	int32 GroupIndex;
+	FName GroupName;
+
+#if WITH_EDITORONLY_DATA
+	UPROPERTY()
+	int32 GroupIndex_DEPRECATED;
+#endif
 
 	// The role this player can assume within the group (ignored if GroupIndex is INDEX_NONE)
 	UPROPERTY()
diff --git a/Engine/Source/Runtime/Engine/Private/Animation/AnimInstance.cpp b/Engine/Source/Runtime/Engine/Private/Animation/AnimInstance.cpp
--- a/Engine/Source/Runtime/Engine/Private/Animation/AnimInstance.cpp
+++ b/Engine/Source/Runtime/Engine/Private/Animation/AnimInstance.cpp
@@ -334,14 +334,14 @@ void UAnimInstance::UpdateMontageSyncGroup()
 		bool bRecordNeedsResetting = true;
 		if (MontageInstance->bDidUseMarkerSyncThisTick)
 		{
-			const int32 GroupIndexToUse = MontageInstance->GetSyncGroupIndex();
+			const FName GroupNameToUse = MontageInstance->GetSyncGroupName();
 
 			// that is public data, so if anybody decided to play with it
-			if (ensure(GroupIndexToUse != INDEX_NONE))
+			if (ensure(GroupNameToUse != NAME_None))
 			{
 				bRecordNeedsResetting = false;
 				FAnimGroupInstance* SyncGroup;
-				FAnimTickRecord& TickRecord = GetProxyOnGameThread<FAnimInstanceProxy>().CreateUninitializedTickRecord(GroupIndexToUse, /*out*/ SyncGroup);
+				FAnimTickRecord& TickRecord = GetProxyOnGameThread<FAnimInstanceProxy>().CreateUninitializedTickRecord(/*out*/ SyncGroup, GroupNameToUse);
 				MakeMontageTickRecord(TickRecord, MontageInstance->Montage, MontageInstance->GetPosition(),
 					MontageInstance->GetPreviousPosition(), MontageInstance->GetDeltaMoved(), MontageInstance->GetWeight(),
 					MontageInstance->MarkersPassedThisTick, MontageInstance->MarkerTickRecord);
@@ -1044,20 +1044,20 @@ void UAnimInstance::DisplayDebug(class UCanvas* Canvas, const FDebugDisplayInfo&
 		FIndenter AnimIndent(Indent);
 
 		//Display Sync Groups
-		const TArray<FAnimGroupInstance>& SyncGroups = GetProxyOnGameThread<FAnimInstanceProxy>().GetSyncGroupRead();
+		const FAnimInstanceProxy::FSyncGroupMap& SyncGroupMap = GetProxyOnGameThread<FAnimInstanceProxy>().GetSyncGroupMapRead();
 		const TArray<FAnimTickRecord>& UngroupedActivePlayers = GetProxyOnGameThread<FAnimInstanceProxy>().GetUngroupedActivePlayersRead();
 
-		Heading = FString::Printf(TEXT("SyncGroups: %i"), SyncGroups.Num());
+		Heading = FString::Printf(TEXT("SyncGroups: %i"), SyncGroupMap.Num());
 		DisplayDebugManager.DrawString(Heading, Indent);
 
-		for (int32 GroupIndex = 0; GroupIndex < SyncGroups.Num(); ++GroupIndex)
+		for (const auto& SyncGroupPair : SyncGroupMap)
 		{
 			FIndenter GroupIndent(Indent);
-			const FAnimGroupInstance& SyncGroup = SyncGroups[GroupIndex];
+			const FAnimGroupInstance& SyncGroup = SyncGroupPair.Value;
 
 			DisplayDebugManager.SetLinearDrawColor(TextYellow);
 
-			FString GroupLabel = FString::Printf(TEXT("Group %i - Players %i"), GroupIndex, SyncGroup.ActivePlayers.Num());
+			FString GroupLabel = FString::Printf(TEXT("Group %s - Players %i"), *SyncGroupPair.Key.ToString(), SyncGroup.ActivePlayers.Num());
 			DisplayDebugManager.DrawString(GroupLabel, Indent);
 
 			if (SyncGroup.ActivePlayers.Num() > 0)
diff --git a/Engine/Source/Runtime/Engine/Private/Animation/AnimInstanceProxy.cpp b/Engine/Source/Runtime/Engine/Private/Animation/AnimInstanceProxy.cpp
--- a/Engine/Source/Runtime/Engine/Private/Animation/AnimInstanceProxy.cpp
+++ b/Engine/Source/Runtime/Engine/Private/Animation/AnimInstanceProxy.cpp
@@ -396,10 +396,10 @@ void FAnimInstanceProxy::PreUpdate(UAnimInstance* InAnimInstance, float DeltaSec
 	TArray<FAnimTickRecord>& UngroupedActivePlayers = UngroupedActivePlayerArrays[GetSyncGroupWriteIndex()];
 	UngroupedActivePlayers.Reset();
 
-	TArray<FAnimGroupInstance>& SyncGroups = SyncGroupArrays[GetSyncGroupWriteIndex()];
-	for (int32 GroupIndex = 0; GroupIndex < SyncGroups.Num(); ++GroupIndex)
+	FSyncGroupMap& SyncGroups = SyncGroupMaps[GetSyncGroupWriteIndex()];
+	for (auto& SyncGroupPair : SyncGroups)
 	{
-		SyncGroups[GroupIndex].Reset();
+		SyncGroupPair.Value.Reset();
 	}
 
 	TArray<float>& StateWeights = StateWeightArrays[GetSyncGroupWriteIndex()];
@@ -603,18 +603,14 @@ void FAnimInstanceProxy::ClearObjects()
 	Skeleton = nullptr;
 }
 
-FAnimTickRecord& FAnimInstanceProxy::CreateUninitializedTickRecord(int32 GroupIndex, FAnimGroupInstance*& OutSyncGroupPtr)
+FAnimTickRecord& FAnimInstanceProxy::CreateUninitializedTickRecord(FAnimGroupInstance*& OutSyncGroupPtr, FName GroupName)
 {
 	// Find or create the sync group if there is one
 	OutSyncGroupPtr = nullptr;
-	if (GroupIndex >= 0)
+	if (GroupName != NAME_None)
 	{
-		TArray<FAnimGroupInstance>& SyncGroups = SyncGroupArrays[GetSyncGroupWriteIndex()];
-		while (SyncGroups.Num() <= GroupIndex)
-		{
-			new (SyncGroups) FAnimGroupInstance();
-		}
-		OutSyncGroupPtr = &(SyncGroups[GroupIndex]);
+		FSyncGroupMap& SyncGroupMap = SyncGroupMaps[GetSyncGroupWriteIndex()];
+		OutSyncGroupPtr = &SyncGroupMap.FindOrAdd(GroupName);
 	}
 
 	// Create the record
@@ -622,12 +618,12 @@ FAnimTickRecord& FAnimInstanceProxy::CreateUninitializedTickRecord(int32 GroupIn
 	return *TickRecord;
 }
 
-FAnimTickRecord& FAnimInstanceProxy::CreateUninitializedTickRecordInScope(int32 GroupIndex, EAnimSyncGroupScope Scope, FAnimGroupInstance*& OutSyncGroupPtr)
+FAnimTickRecord& FAnimInstanceProxy::CreateUninitializedTickRecordInScope(FAnimGroupInstance*& OutSyncGroupPtr, FName GroupName, EAnimSyncGroupScope Scope)
 {
-	if (GroupIndex >= 0)
+	if (GroupName != NAME_None)
 	{
-		// If we have no main proxy, force us to local
-		if(MainInstanceProxy == nullptr)
+		// If we have no main proxy or it is "this", force us to local
+		if(MainInstanceProxy == nullptr || MainInstanceProxy == this)
 		{
 			Scope = EAnimSyncGroupScope::Local;
 		}
@@ -639,14 +635,36 @@ FAnimTickRecord& FAnimInstanceProxy::CreateUninitializedTickRecordInScope(int32
 			// Fall through
 
 		case EAnimSyncGroupScope::Local:
-			return CreateUninitializedTickRecord(GroupIndex, OutSyncGroupPtr);
+			return CreateUninitializedTickRecord(OutSyncGroupPtr, GroupName);
 		case EAnimSyncGroupScope::Component:
 			// Forward to the main instance to sync with animations there in TickAssetPlayerInstances()
-			return MainInstanceProxy->CreateUninitializedTickRecord(GroupIndex, OutSyncGroupPtr);
+			return MainInstanceProxy->CreateUninitializedTickRecord(OutSyncGroupPtr, GroupName);
 		}
 	}
 	
-	return CreateUninitializedTickRecord(GroupIndex, OutSyncGroupPtr);
+	return CreateUninitializedTickRecord(OutSyncGroupPtr, GroupName);
+}
+
+FAnimTickRecord& FAnimInstanceProxy::CreateUninitializedTickRecord(int32 GroupIndex, FAnimGroupInstance*& OutSyncGroupPtr)
+{
+	FName SyncGroupName = NAME_None;
+	if(GroupIndex >= 0)
+	{
+		SyncGroupName = GetAnimClassInterface()->GetSyncGroupNames()[GroupIndex];
+	}
+
+	return CreateUninitializedTickRecord(OutSyncGroupPtr, SyncGroupName);
+}
+
+FAnimTickRecord& FAnimInstanceProxy::CreateUninitializedTickRecordInScope(int32 GroupIndex, EAnimSyncGroupScope Scope, FAnimGroupInstance*& OutSyncGroupPtr)
+{
+	FName SyncGroupName = NAME_None;
+	if(GroupIndex >= 0)
+	{
+		SyncGroupName = GetAnimClassInterface()->GetSyncGroupNames()[GroupIndex];
+	}
+
+	return CreateUninitializedTickRecordInScope(OutSyncGroupPtr, SyncGroupName, Scope);
 }
 
 void FAnimInstanceProxy::MakeSequenceTickRecord(FAnimTickRecord& TickRecord, class UAnimSequenceBase* Sequence, bool bLooping, float PlayRate, float FinalBlendWeight, float& CurrentTime, FMarkerTickRecord& MarkerTickRecord) const
@@ -710,20 +728,20 @@ void FAnimInstanceProxy::TickAssetPlayerInstances(float DeltaSeconds)
 	SCOPE_CYCLE_COUNTER(STAT_TickAssetPlayerInstances);
 
 	// Handle all players inside sync groups
-	TArray<FAnimGroupInstance>& SyncGroups = SyncGroupArrays[GetSyncGroupWriteIndex()];
-	const TArray<FAnimGroupInstance>& PreviousSyncGroups = SyncGroupArrays[GetSyncGroupReadIndex()];
+	FSyncGroupMap& SyncGroupMap = SyncGroupMaps[GetSyncGroupWriteIndex()];
+	const FSyncGroupMap& PreviousSyncGroupMap = SyncGroupMaps[GetSyncGroupReadIndex()];
 	TArray<FAnimTickRecord>& UngroupedActivePlayers = UngroupedActivePlayerArrays[GetSyncGroupWriteIndex()];
 
-	for (int32 GroupIndex = 0; GroupIndex < SyncGroups.Num(); ++GroupIndex)
+	for (auto& SyncGroupPair : SyncGroupMap)
 	{
-		FAnimGroupInstance& SyncGroup = SyncGroups[GroupIndex];
+		FAnimGroupInstance& SyncGroup = SyncGroupPair.Value;
 	
 		if (SyncGroup.ActivePlayers.Num() > 0)
 		{
-			const FAnimGroupInstance* PreviousGroup = PreviousSyncGroups.IsValidIndex(GroupIndex) ? &PreviousSyncGroups[GroupIndex] : nullptr;
+			const FAnimGroupInstance* PreviousGroup = PreviousSyncGroupMap.Find(SyncGroupPair.Key);
 			SyncGroup.Prepare(PreviousGroup);
 
-			UE_LOG(LogAnimMarkerSync, Log, TEXT("Ticking Group [%d] GroupLeader [%d]"), GroupIndex, SyncGroup.GroupLeaderIndex);
+			UE_LOG(LogAnimMarkerSync, Log, TEXT("Ticking Group [%s] GroupLeader [%d]"), *SyncGroupPair.Key.ToString(), SyncGroup.GroupLeaderIndex);
 
 			const bool bOnlyOneAnimationInGroup = SyncGroup.ActivePlayers.Num() == 1;
 
@@ -749,7 +767,7 @@ void FAnimInstanceProxy::TickAssetPlayerInstances(float DeltaSeconds)
 #endif
 
 			// initialize to invalidate first
-			ensureMsgf(SyncGroup.GroupLeaderIndex == INDEX_NONE, TEXT("SyncGroup with GroupIndex=%d had a non -1 group leader index of %d in asset %s"), GroupIndex, SyncGroup.GroupLeaderIndex, *GetNameSafe(SkeletalMeshComponent));
+			ensureMsgf(SyncGroup.GroupLeaderIndex == INDEX_NONE, TEXT("SyncGroup %s had a non -1 group leader index of %d in asset %s"), *SyncGroupPair.Key.ToString(), SyncGroup.GroupLeaderIndex, *GetNameSafe(SkeletalMeshComponent));
 			int32 GroupLeaderIndex = 0;
 			for (; GroupLeaderIndex < SyncGroup.ActivePlayers.Num(); ++GroupLeaderIndex)
 			{
@@ -796,7 +814,7 @@ void FAnimInstanceProxy::TickAssetPlayerInstances(float DeltaSeconds)
 			if (TickContext.CanUseMarkerPosition())
 			{
 				const FMarkerSyncAnimPosition& MarkerStart = TickContext.MarkerTickContext.GetMarkerSyncStartPosition();
-				FName SyncGroupName = GetAnimClassInterface()->GetSyncGroupNames()[GroupIndex];
+				FName SyncGroupName = SyncGroupPair.Key;
 				FAnimTickRecord& GroupLeader = SyncGroup.ActivePlayers[SyncGroup.GroupLeaderIndex];
 				FString LeaderAnimName = GroupLeader.SourceAsset->GetName();
 
@@ -914,16 +932,14 @@ int32 FAnimInstanceProxy::GetSyncGroupIndexFromName(FName SyncGroupName) const
 
 bool FAnimInstanceProxy::GetTimeToClosestMarker(FName SyncGroup, FName MarkerName, float& OutMarkerTime) const
 {
-	const int32 SyncGroupIndex = GetSyncGroupIndexFromName(SyncGroup);
-	const TArray<FAnimGroupInstance>& SyncGroups = SyncGroupArrays[GetSyncGroupReadIndex()];
+	const FSyncGroupMap& SyncGroupMap = SyncGroupMaps[GetSyncGroupReadIndex()];
 
-	if (SyncGroups.IsValidIndex(SyncGroupIndex))
+	if (const FAnimGroupInstance* SyncGroupInstancePtr = SyncGroupMap.Find(SyncGroup))
 	{
-		const FAnimGroupInstance& SyncGroupInstance = SyncGroups[SyncGroupIndex];
-		if (SyncGroupInstance.bCanUseMarkerSync && SyncGroupInstance.ActivePlayers.IsValidIndex(SyncGroupInstance.GroupLeaderIndex))
+		if (SyncGroupInstancePtr->bCanUseMarkerSync && SyncGroupInstancePtr->ActivePlayers.IsValidIndex(SyncGroupInstancePtr->GroupLeaderIndex))
 		{
-			const FMarkerSyncAnimPosition& EndPosition = SyncGroupInstance.MarkerTickContext.GetMarkerSyncEndPosition();
-			const FAnimTickRecord& Leader = SyncGroupInstance.ActivePlayers[SyncGroupInstance.GroupLeaderIndex];
+			const FMarkerSyncAnimPosition& EndPosition = SyncGroupInstancePtr->MarkerTickContext.GetMarkerSyncEndPosition();
+			const FAnimTickRecord& Leader = SyncGroupInstancePtr->ActivePlayers[SyncGroupInstancePtr->GroupLeaderIndex];
 			if (EndPosition.PreviousMarkerName == MarkerName)
 			{
 				OutMarkerTime = Leader.MarkerTickRecord->PreviousMarker.TimeToMarker;
@@ -956,15 +972,13 @@ void FAnimInstanceProxy::AddAnimNotifyFromGeneratedClass(int32 NotifyIndex)
 
 bool FAnimInstanceProxy::HasMarkerBeenHitThisFrame(FName SyncGroup, FName MarkerName) const
 {
-	const int32 SyncGroupIndex = GetSyncGroupIndexFromName(SyncGroup);
-	const TArray<FAnimGroupInstance>& SyncGroups = SyncGroupArrays[GetSyncGroupReadIndex()];
+	const FSyncGroupMap& SyncGroupMap = SyncGroupMaps[GetSyncGroupReadIndex()];
 
-	if (SyncGroups.IsValidIndex(SyncGroupIndex))
+	if (const FAnimGroupInstance* SyncGroupInstancePtr = SyncGroupMap.Find(SyncGroup))
 	{
-		const FAnimGroupInstance& SyncGroupInstance = SyncGroups[SyncGroupIndex];
-		if (SyncGroupInstance.bCanUseMarkerSync)
+		if (SyncGroupInstancePtr->bCanUseMarkerSync)
 		{
-			return SyncGroupInstance.MarkerTickContext.MarkersPassedThisTick.ContainsByPredicate([&MarkerName](const FPassedMarker& PassedMarker) -> bool
+			return SyncGroupInstancePtr->MarkerTickContext.MarkersPassedThisTick.ContainsByPredicate([&MarkerName](const FPassedMarker& PassedMarker) -> bool
 			{
 				return PassedMarker.PassedMarkerName == MarkerName;
 			});
@@ -991,15 +1005,13 @@ bool FAnimInstanceProxy::IsSyncGroupBetweenMarkers(FName InSyncGroupName, FName
 
 FMarkerSyncAnimPosition FAnimInstanceProxy::GetSyncGroupPosition(FName InSyncGroupName) const
 {
-	const int32 SyncGroupIndex = GetSyncGroupIndexFromName(InSyncGroupName);
-	const TArray<FAnimGroupInstance>& SyncGroups = SyncGroupArrays[GetSyncGroupReadIndex()];
+	const FSyncGroupMap& SyncGroupMap = SyncGroupMaps[GetSyncGroupReadIndex()];
 
-	if (SyncGroups.IsValidIndex(SyncGroupIndex))
+	if (const FAnimGroupInstance* SyncGroupInstancePtr = SyncGroupMap.Find(InSyncGroupName))
 	{
-		const FAnimGroupInstance& SyncGroupInstance = SyncGroups[SyncGroupIndex];
-		if (SyncGroupInstance.bCanUseMarkerSync && SyncGroupInstance.MarkerTickContext.IsMarkerSyncEndValid())
+		if (SyncGroupInstancePtr->bCanUseMarkerSync && SyncGroupInstancePtr->MarkerTickContext.IsMarkerSyncEndValid())
 		{
-			return SyncGroupInstance.MarkerTickContext.GetMarkerSyncEndPosition();
+			return SyncGroupInstancePtr->MarkerTickContext.GetMarkerSyncEndPosition();
 		}
 	}
 
diff --git a/Engine/Source/Runtime/Engine/Private/Animation/AnimMontage.cpp b/Engine/Source/Runtime/Engine/Private/Animation/AnimMontage.cpp
--- a/Engine/Source/Runtime/Engine/Private/Animation/AnimMontage.cpp
+++ b/Engine/Source/Runtime/Engine/Private/Animation/AnimMontage.cpp
@@ -1281,7 +1281,7 @@ FAnimMontageInstance::FAnimMontageInstance()
 	, NotifyWeight(0.f)
 	, DeltaMoved(0.f)
 	, PreviousPosition(0.f)
-	, SyncGroupIndex(INDEX_NONE)
+	, SyncGroupName(NAME_None)
 	, DisableRootMotionCount(0)
 	, MontageSyncLeader(NULL)
 	, MontageSyncUpdateFrameCounter(INDEX_NONE)
@@ -1303,7 +1303,7 @@ FAnimMontageInstance::FAnimMontageInstance(UAnimInstance * InAnimInstance)
 	, NotifyWeight(0.f)
 	, DeltaMoved(0.f)
 	, PreviousPosition(0.f)
-	, SyncGroupIndex(INDEX_NONE)
+	, SyncGroupName(NAME_None)
 	, DisableRootMotionCount(0)
 	, MontageSyncLeader(NULL)
 	, MontageSyncUpdateFrameCounter(INDEX_NONE)
@@ -1420,7 +1420,7 @@ void FAnimMontageInstance::Initialize(class UAnimMontage * InMontage)
 
 		if (AnimInstance.IsValid() && Montage->CanUseMarkerSync())
 		{
-			SyncGroupIndex = AnimInstance.Get()->GetSyncGroupIndexFromName(Montage->SyncGroup);
+			SyncGroupName = Montage->SyncGroup;
 		}
 
 		MontageSubStepper.Initialize(*this);
@@ -2865,7 +2865,7 @@ UAnimMontage* UAnimMontage::CreateSlotAnimationAsDynamicMontage(UAnimSequenceBas
 bool FAnimMontageInstance::CanUseMarkerSync() const
 {
 	// for now we only allow non-full weight and when blending out
-	return SyncGroupIndex != INDEX_NONE && IsStopped() && Blend.IsComplete() == false;
+	return SyncGroupName != NAME_None && IsStopped() && Blend.IsComplete() == false;
 }
 
 void UAnimMontage::BakeTimeStretchCurve()
diff --git a/Engine/Source/Runtime/Engine/Private/Animation/AnimNode_AssetPlayerBase.cpp b/Engine/Source/Runtime/Engine/Private/Animation/AnimNode_AssetPlayerBase.cpp
--- a/Engine/Source/Runtime/Engine/Private/Animation/AnimNode_AssetPlayerBase.cpp
+++ b/Engine/Source/Runtime/Engine/Private/Animation/AnimNode_AssetPlayerBase.cpp
@@ -4,7 +4,10 @@
 #include "Animation/AnimInstanceProxy.h"
 
 FAnimNode_AssetPlayerBase::FAnimNode_AssetPlayerBase()
-	: GroupIndex(INDEX_NONE)
+	: GroupName(NAME_None)
+#if WITH_EDITORONLY_DATA
+	, GroupIndex_DEPRECATED(INDEX_NONE)
+#endif
 	, GroupRole(EAnimGroupRole::CanBeLeader)
 	, GroupScope(EAnimSyncGroupScope::Local)
 	, bIgnoreForRelevancyTest(false)
@@ -38,9 +41,9 @@ void FAnimNode_AssetPlayerBase::CreateTickRecordForNode(const FAnimationUpdateCo
 	const float FinalBlendWeight = Context.GetFinalBlendWeight();
 
 	FAnimGroupInstance* SyncGroup;
-	const int32 GroupIndexToUse = ((GroupRole < EAnimGroupRole::TransitionLeader) || bHasBeenFullWeight) ? GroupIndex : INDEX_NONE;
+	const FName GroupNameToUse = ((GroupRole < EAnimGroupRole::TransitionLeader) || bHasBeenFullWeight) ? GroupName : NAME_None;
 
-	FAnimTickRecord& TickRecord = Context.AnimInstanceProxy->CreateUninitializedTickRecordInScope(GroupIndexToUse, GroupScope, /*out*/ SyncGroup);
+	FAnimTickRecord& TickRecord = Context.AnimInstanceProxy->CreateUninitializedTickRecordInScope(/*out*/ SyncGroup, GroupNameToUse, GroupScope);
 
 	Context.AnimInstanceProxy->MakeSequenceTickRecord(TickRecord, Sequence, bLooping, PlayRate, FinalBlendWeight, /*inout*/ InternalTimeAccumulator, MarkerTickRecord);
 	TickRecord.RootMotionWeightModifier = Context.GetRootMotionWeightModifier();
diff --git a/Engine/Source/Runtime/Engine/Private/Animation/AnimSingleNodeInstanceProxy.cpp b/Engine/Source/Runtime/Engine/Private/Animation/AnimSingleNodeInstanceProxy.cpp
--- a/Engine/Source/Runtime/Engine/Private/Animation/AnimSingleNodeInstanceProxy.cpp
+++ b/Engine/Source/Runtime/Engine/Private/Animation/AnimSingleNodeInstanceProxy.cpp
@@ -526,7 +526,7 @@ void FAnimNode_SingleNode::Update_AnyThread(const FAnimationUpdateContext& Conte
 		FAnimGroupInstance* SyncGroup;
 		if (UBlendSpaceBase* BlendSpace = Cast<UBlendSpaceBase>(Proxy->CurrentAsset))
 		{
-			FAnimTickRecord& TickRecord = Proxy->CreateUninitializedTickRecord(INDEX_NONE, /*out*/ SyncGroup);
+			FAnimTickRecord& TickRecord = Proxy->CreateUninitializedTickRecord(/*out*/ SyncGroup, NAME_None);
 			Proxy->MakeBlendSpaceTickRecord(TickRecord, BlendSpace, Proxy->BlendSpaceInput, Proxy->BlendSampleData, Proxy->BlendFilter, Proxy->bLooping, NewPlayRate, 1.f, /*inout*/ Proxy->CurrentTime, Proxy->MarkerTickRecord);
 
 			TRACE_ANIM_TICK_RECORD(Context, TickRecord);
@@ -536,7 +536,7 @@ void FAnimNode_SingleNode::Update_AnyThread(const FAnimationUpdateContext& Conte
 		}
 		else if (UAnimSequence* Sequence = Cast<UAnimSequence>(Proxy->CurrentAsset))
 		{
-			FAnimTickRecord& TickRecord = Proxy->CreateUninitializedTickRecord(INDEX_NONE, /*out*/ SyncGroup);
+			FAnimTickRecord& TickRecord = Proxy->CreateUninitializedTickRecord(/*out*/ SyncGroup, NAME_None);
 			Proxy->MakeSequenceTickRecord(TickRecord, Sequence, Proxy->bLooping, NewPlayRate, 1.f, /*inout*/ Proxy->CurrentTime, Proxy->MarkerTickRecord);
 
 			TRACE_ANIM_TICK_RECORD(Context, TickRecord);
@@ -553,7 +553,7 @@ void FAnimNode_SingleNode::Update_AnyThread(const FAnimationUpdateContext& Conte
 		}
 		else if (UAnimStreamable* Streamable = Cast<UAnimStreamable>(Proxy->CurrentAsset))
 		{
-			FAnimTickRecord& TickRecord = Proxy->CreateUninitializedTickRecord(INDEX_NONE, /*out*/ SyncGroup);
+			FAnimTickRecord& TickRecord = Proxy->CreateUninitializedTickRecord(/*out*/ SyncGroup, NAME_None);
 			Proxy->MakeSequenceTickRecord(TickRecord, Streamable, Proxy->bLooping, NewPlayRate, 1.f, /*inout*/ Proxy->CurrentTime, Proxy->MarkerTickRecord);
 
 			TRACE_ANIM_TICK_RECORD(Context, TickRecord);
@@ -570,7 +570,7 @@ void FAnimNode_SingleNode::Update_AnyThread(const FAnimationUpdateContext& Conte
 		}
 		else if(UAnimComposite* Composite = Cast<UAnimComposite>(Proxy->CurrentAsset))
 		{
-			FAnimTickRecord& TickRecord = Proxy->CreateUninitializedTickRecord(INDEX_NONE, /*out*/ SyncGroup);
+			FAnimTickRecord& TickRecord = Proxy->CreateUninitializedTickRecord(/*out*/ SyncGroup, NAME_None);
 			Proxy->MakeSequenceTickRecord(TickRecord, Composite, Proxy->bLooping, NewPlayRate, 1.f, /*inout*/ Proxy->CurrentTime, Proxy->MarkerTickRecord);
 
 			TRACE_ANIM_TICK_RECORD(Context, TickRecord);
@@ -609,7 +609,7 @@ void FAnimNode_SingleNode::Update_AnyThread(const FAnimationUpdateContext& Conte
 		}
 		else if (UPoseAsset* PoseAsset = Cast<UPoseAsset>(Proxy->CurrentAsset))
 		{
-			FAnimTickRecord& TickRecord = Proxy->CreateUninitializedTickRecord(INDEX_NONE, /*out*/ SyncGroup);
+			FAnimTickRecord& TickRecord = Proxy->CreateUninitializedTickRecord(/*out*/ SyncGroup, NAME_None);
 			Proxy->MakePoseAssetTickRecord(TickRecord, PoseAsset, 1.f);
 
 			TRACE_ANIM_TICK_RECORD(Context, TickRecord);
diff --git a/Engine/Source/Runtime/Engine/Public/Animation/AnimInstanceProxy.h b/Engine/Source/Runtime/Engine/Public/Animation/AnimInstanceProxy.h
--- a/Engine/Source/Runtime/Engine/Public/Animation/AnimInstanceProxy.h
+++ b/Engine/Source/Runtime/Engine/Public/Animation/AnimInstanceProxy.h
@@ -109,6 +109,8 @@ struct ENGINE_API FAnimInstanceProxy
 	GENERATED_BODY()
 
 public:
+	using FSyncGroupMap = TMap<FName, FAnimGroupInstance>;
+
 	FAnimInstanceProxy()
 		: AnimInstanceObject(nullptr)
 		, AnimClassInterface(nullptr)
@@ -218,10 +220,17 @@ public:
 		SyncGroupWriteIndex = GetSyncGroupReadIndex();
 	}
 
-	/** Get the sync group we are currently reading from */
+	UE_DEPRECATED(5.0, "Sync groups are no longer stored in arrays, please use GetSyncGroupMapRead")
 	const TArray<FAnimGroupInstance>& GetSyncGroupRead() const
+	{
+		static const TArray<FAnimGroupInstance> Dummy;
+		return Dummy; 
+	}
+
+	/** Get the sync group we are currently reading from */
+	const FSyncGroupMap& GetSyncGroupMapRead() const
 	{ 
-		return SyncGroupArrays[GetSyncGroupReadIndex()]; 
+		return SyncGroupMaps[GetSyncGroupReadIndex()]; 
 	}
 
 	/** Get the ungrouped active player we are currently reading from */
@@ -344,13 +353,19 @@ public:
 		return SkeletalMeshComponent; 
 	}
 
-	// Creates an uninitialized tick record in the list for the correct group or the ungrouped array.  If the group is valid, OutSyncGroupPtr will point to the group.
+	UE_DEPRECATED(4.26, "Please use the overload that takes a group FName")
 	FAnimTickRecord& CreateUninitializedTickRecord(int32 GroupIndex, FAnimGroupInstance*& OutSyncGroupPtr);
 
+	UE_DEPRECATED(4.26, "Please use the overload that takes a group FName")
+	FAnimTickRecord& CreateUninitializedTickRecordInScope(int32 GroupIndex, EAnimSyncGroupScope Scope, FAnimGroupInstance*& OutSyncGroupPtr);
+
+	// Creates an uninitialized tick record in the list for the correct group or the ungrouped array.  If the group is valid, OutSyncGroupPtr will point to the group.
+	FAnimTickRecord& CreateUninitializedTickRecord(FAnimGroupInstance*& OutSyncGroupPtr, FName GroupName);
+
 	// Creates an uninitialized tick record in the list for the correct group or the ungrouped array.  
 	// If the group is valid, OutSyncGroupPtr will point to the group.
 	// Supply the scope to sync with tick records outside this instance
-	FAnimTickRecord& CreateUninitializedTickRecordInScope(int32 GroupIndex, EAnimSyncGroupScope Scope, FAnimGroupInstance*& OutSyncGroupPtr);
+	FAnimTickRecord& CreateUninitializedTickRecordInScope(FAnimGroupInstance*& OutSyncGroupPtr, FName GroupName, EAnimSyncGroupScope Scope);
 
 	/** Helper function: make a tick record for a sequence */
 	void MakeSequenceTickRecord(FAnimTickRecord& TickRecord, UAnimSequenceBase* Sequence, bool bLooping, float PlayRate, float FinalBlendWeight, float& CurrentTime, FMarkerTickRecord& MarkerTickRecord) const;
@@ -877,7 +892,7 @@ private:
 	TArray<FAnimTickRecord> UngroupedActivePlayerArrays[2];
 
 	/** The set of tick groups for this anim instance */
-	TArray<FAnimGroupInstance> SyncGroupArrays[2];
+	FSyncGroupMap SyncGroupMaps[2];
 
 	/** Buffers containing read/write buffers for all current machine weights */
 	TArray<float> MachineWeightArrays[2];
