Merge fix

New function for sync groups was not using the correct container

diff --git a/Engine/Source/Runtime/Engine/Private/Animation/AnimInstanceProxy.cpp b/Engine/Source/Runtime/Engine/Private/Animation/AnimInstanceProxy.cpp
--- a/Engine/Source/Runtime/Engine/Private/Animation/AnimInstanceProxy.cpp
+++ b/Engine/Source/Runtime/Engine/Private/Animation/AnimInstanceProxy.cpp
@@ -1020,14 +1020,12 @@ FMarkerSyncAnimPosition FAnimInstanceProxy::GetSyncGroupPosition(FName InSyncGro
 
 bool FAnimInstanceProxy::IsSyncGroupValid(FName InSyncGroupName) const
 {
-	const int32 SyncGroupIndex = GetSyncGroupIndexFromName(InSyncGroupName);
-	const TArray<FAnimGroupInstance>& SyncGroups = SyncGroupArrays[GetSyncGroupReadIndex()];
+	const FSyncGroupMap& SyncGroupMap = SyncGroupMaps[GetSyncGroupReadIndex()];
 
-	if (SyncGroups.IsValidIndex(SyncGroupIndex))
+	if (const FAnimGroupInstance* SyncGroupInstancePtr = SyncGroupMap.Find(InSyncGroupName))
 	{
-		const FAnimGroupInstance& SyncGroupInstance = SyncGroups[SyncGroupIndex];
 		// If we don't use Markers, we're always valid.
-		return (!SyncGroupInstance.bCanUseMarkerSync || SyncGroupInstance.MarkerTickContext.IsMarkerSyncEndValid());
+		return (!SyncGroupInstancePtr->bCanUseMarkerSync || SyncGroupInstancePtr->MarkerTickContext.IsMarkerSyncEndValid());
 	}
 
 	// If we're querying a sync group that doesn't exist, treat this as invalid
