Put back call to AssetCreated call for newly created external actors and ignore them in the world partition until they get saved.

diff --git a/Engine/Source/Runtime/Engine/Private/ActorEditor.cpp b/Engine/Source/Runtime/Engine/Private/ActorEditor.cpp
--- a/Engine/Source/Runtime/Engine/Private/ActorEditor.cpp
+++ b/Engine/Source/Runtime/Engine/Private/ActorEditor.cpp
@@ -22,8 +22,6 @@
 #include "ActorEditorUtils.h"
 #include "EngineGlobals.h"
 
-#include "AssetRegistryModule.h"
-
 #if WITH_EDITOR
 
 #include "Editor.h"
diff --git a/Engine/Source/Runtime/Engine/Private/LevelActor.cpp b/Engine/Source/Runtime/Engine/Private/LevelActor.cpp
--- a/Engine/Source/Runtime/Engine/Private/LevelActor.cpp
+++ b/Engine/Source/Runtime/Engine/Private/LevelActor.cpp
@@ -38,6 +38,7 @@
 #include "GameFramework/WorldSettings.h"
 #include "Engine/NetDriver.h"
 #include "Engine/Player.h"
+#include "AssetRegistryModule.h"
 
 #include "Components/BoxComponent.h"
 #include "GameFramework/MovementComponent.h"
@@ -733,6 +734,11 @@ AActor* UWorld::SpawnActor( UClass* Class, FTransform const* UserTransformPtr, c
 	OnActorSpawned.Broadcast(Actor);
 
 #if WITH_EDITOR
+	if (Actor->IsAsset())
+	{
+		FAssetRegistryModule::AssetCreated(Actor);
+	}
+
 	if (GIsEditor)
 	{
 		GEngine->BroadcastLevelActorAdded(Actor);
diff --git a/Engine/Source/Runtime/Engine/Private/WorldPartition/WorldPartition.cpp b/Engine/Source/Runtime/Engine/Private/WorldPartition/WorldPartition.cpp
--- a/Engine/Source/Runtime/Engine/Private/WorldPartition/WorldPartition.cpp
+++ b/Engine/Source/Runtime/Engine/Private/WorldPartition/WorldPartition.cpp
@@ -800,7 +800,7 @@ void UWorldPartition::UnregisterActor(AActor* Actor)
 #if WITH_EDITOR
 void UWorldPartition::OnAssetAdded(const FAssetData& InAssetData)
 {
-	if (!bIgnoreAssetRegistryEvents)
+	if (ShouldHandleAssetEvent(InAssetData))
 	{
 		TUniquePtr<FWorldPartitionActorDesc> NewActorDesc = GetActorDescriptor(InAssetData);
 		if (NewActorDesc.IsValid())
@@ -815,7 +815,7 @@ void UWorldPartition::OnAssetAdded(const FAssetData& InAssetData)
 
 void UWorldPartition::OnAssetRemoved(const FAssetData& InAssetData)
 {
-	if (!bIgnoreAssetRegistryEvents)
+	if (ShouldHandleAssetEvent(InAssetData))
 	{
 		TUniquePtr<FWorldPartitionActorDesc> NewActorDesc = GetActorDescriptor(InAssetData);
 		if (NewActorDesc.IsValid())
@@ -830,7 +830,7 @@ void UWorldPartition::OnAssetRemoved(const FAssetData& InAssetData)
 
 void UWorldPartition::OnAssetUpdated(const FAssetData& InAssetData)
 {
-	if (!bIgnoreAssetRegistryEvents)
+	if (ShouldHandleAssetEvent(InAssetData))
 	{
 		TUniquePtr<FWorldPartitionActorDesc> NewActorDesc = GetActorDescriptor(InAssetData);
 		if (NewActorDesc.IsValid())
@@ -844,6 +844,29 @@ void UWorldPartition::OnAssetUpdated(const FAssetData& InAssetData)
 	}
 }
 
+bool UWorldPartition::ShouldHandleAssetEvent(const FAssetData& InAssetData)
+{
+	// Ignore asset event when specifically asking to
+	if (bIgnoreAssetRegistryEvents)
+	{
+		return false;
+	}
+
+	// Ignore in-memory assets until they gets saved
+	if (InAssetData.HasAnyPackageFlags(PKG_NewlyCreated))
+	{
+		return false;
+	}
+
+	// Only handle actors
+	if (!InAssetData.GetClass()->IsChildOf<AActor>())
+	{
+		return false;
+	}
+
+	return true;
+}
+
 TUniquePtr<FWorldPartitionActorDesc> UWorldPartition::GetActorDescriptor(const FAssetData& InAssetData)
 {
 	FString ActorClass;
diff --git a/Engine/Source/Runtime/Engine/Public/WorldPartition/WorldPartition.h b/Engine/Source/Runtime/Engine/Public/WorldPartition/WorldPartition.h
--- a/Engine/Source/Runtime/Engine/Public/WorldPartition/WorldPartition.h
+++ b/Engine/Source/Runtime/Engine/Public/WorldPartition/WorldPartition.h
@@ -78,6 +78,7 @@ private:
 	void OnAssetRemoved(const FAssetData& InAssetData);
 	void OnAssetUpdated(const FAssetData& InAssetData);
 
+	bool ShouldHandleAssetEvent(const FAssetData& InAssetData);
 	TUniquePtr<FWorldPartitionActorDesc> GetActorDescriptor(const FAssetData& InAssetData);
 
 	void ApplyActorTransform(AActor* InActor, const FTransform& InTransform);
