WorldPartition: added console command to dump all actor derscriptors as a cvs file.

diff --git a/Engine/Source/Runtime/Engine/Private/WorldPartition/WorldPartition.cpp b/Engine/Source/Runtime/Engine/Private/WorldPartition/WorldPartition.cpp
--- a/Engine/Source/Runtime/Engine/Private/WorldPartition/WorldPartition.cpp
+++ b/Engine/Source/Runtime/Engine/Private/WorldPartition/WorldPartition.cpp
@@ -20,7 +20,6 @@
 #include "ProfilingDebugging/ScopedTimers.h"
 #include "LevelUtils.h"
 
-
 #if WITH_EDITOR
 #include "Editor.h"
 #include "Layers/Layer.h"
@@ -28,6 +27,7 @@
 #include "Editor/GroupActor.h"
 #include "EditorLevelUtils.h"
 #include "FileHelpers.h"
+#include "HAL/FileManager.h"
 #include "Misc/Base64.h"
 #include "Misc/ScopedSlowTask.h"
 #include "Misc/ScopeExit.h"
@@ -47,25 +47,44 @@ DEFINE_LOG_CATEGORY(LogWorldPartition);
 #define LOCTEXT_NAMESPACE "WorldPartitionEditor"
 
 #if WITH_EDITOR
-static void GenerateHLOD(const TArray<FString>& Args)
-{
-	if (UWorld* World = GEditor->GetEditorWorldContext().World())
+static FAutoConsoleCommand GenerateHLODCmd(
+	TEXT("wp.Editor.GenerateHLOD"),
+	TEXT("Generates HLOD data for runtime."),
+	FConsoleCommandWithArgsDelegate::CreateLambda([](const TArray<FString>& Args)
 	{
-		if (!World->IsPlayInEditor())
+		if (UWorld* World = GEditor->GetEditorWorldContext().World())
 		{
-			if (UWorldPartition* WorldPartition = World->GetWorldPartition())
+			if (!World->IsPlayInEditor())
 			{
-				WorldPartition->Modify();
-				WorldPartition->GenerateHLOD();
+				if (UWorldPartition* WorldPartition = World->GetWorldPartition())
+				{
+					WorldPartition->Modify();
+					WorldPartition->GenerateHLOD();
+				}
 			}
 		}
-	}
-}
+	})
+);
 
-static FAutoConsoleCommand GenerateHLODCmd(
-	TEXT("wp.Editor.GenerateHLOD"),
-	TEXT("Generates HLOD data for runtime."),
-	FConsoleCommandWithArgsDelegate::CreateStatic(&GenerateHLOD)
+static FAutoConsoleCommand DumpActorDescs(
+	TEXT("wp.Editor.DumpActorDescs"),
+	TEXT("Dump the list of actor descriptors in a CSV file."),
+	FConsoleCommandWithArgsDelegate::CreateLambda([](const TArray<FString>& Args)
+	{
+		if (Args.Num() > 0)
+		{
+			if (UWorld* World = GEditor->GetEditorWorldContext().World())
+			{
+				if (!World->IsPlayInEditor())
+				{
+					if (UWorldPartition* WorldPartition = World->GetWorldPartition())
+					{
+						WorldPartition->DumpActorDescs(Args[0]);
+					}
+				}
+			}
+		}
+	})
 );
 #endif
 
@@ -1035,6 +1054,36 @@ void UWorldPartition::GenerateNavigationData()
 	RuntimeHash->GenerateNavigationData();
 }
 
+void UWorldPartition::DumpActorDescs(const FString& Path)
+{
+	if (FArchive* LogFile = IFileManager::Get().CreateFileWriter(*Path))
+	{
+		FString LineEntry = TEXT("Guid, Class, Name, BVCenterX, BVCenterY, BVCenterZ, BVExtentX, BVExtentY, BVExtentZ") LINE_TERMINATOR;
+		LogFile->Serialize(TCHAR_TO_ANSI(*LineEntry), LineEntry.Len());
+			
+		ForEachActorDesc(AActor::StaticClass(), [LogFile, &LineEntry](const FWorldPartitionActorDesc* ActorDesc)
+		{
+				LineEntry = FString::Printf(
+				TEXT("%s, %s, %s, %.2f, %.2f, %.2f, %.2f, %.2f, %.2f") LINE_TERMINATOR, 
+				*ActorDesc->GetGuid().ToString(), 
+				*ActorDesc->GetClass().ToString(), 
+				*FPaths::GetExtension(ActorDesc->GetActorPath().ToString()), 
+				ActorDesc->GetBounds().GetCenter().X,
+				ActorDesc->GetBounds().GetCenter().Y,
+				ActorDesc->GetBounds().GetCenter().Z,
+				ActorDesc->GetBounds().GetExtent().X,
+				ActorDesc->GetBounds().GetExtent().Y,
+				ActorDesc->GetBounds().GetExtent().Z
+			);
+			LogFile->Serialize(TCHAR_TO_ANSI(*LineEntry), LineEntry.Len());
+			return true;
+		});
+
+		LogFile->Close();
+		delete LogFile;
+	}
+}
+
 void UWorldPartition::OnPreFixupForPIE(int32 InPIEInstanceID, FSoftObjectPath& ObjectPath)
 {
 	if (GetWorld()->IsPlayInEditor())
diff --git a/Engine/Source/Runtime/Engine/Public/WorldPartition/WorldPartition.h b/Engine/Source/Runtime/Engine/Public/WorldPartition/WorldPartition.h
--- a/Engine/Source/Runtime/Engine/Public/WorldPartition/WorldPartition.h
+++ b/Engine/Source/Runtime/Engine/Public/WorldPartition/WorldPartition.h
@@ -125,6 +125,8 @@ public:
 
 	void GenerateHLOD();
 	void GenerateNavigationData();
+
+	void DumpActorDescs(const FString& Path);
 #endif
 
 public:
