Reduce game-thread stutters during async animation compression

diff --git a/Engine/Source/Runtime/Engine/Private/Animation/AnimBoneCompressionSettings.cpp b/Engine/Source/Runtime/Engine/Private/Animation/AnimBoneCompressionSettings.cpp
--- a/Engine/Source/Runtime/Engine/Private/Animation/AnimBoneCompressionSettings.cpp
+++ b/Engine/Source/Runtime/Engine/Private/Animation/AnimBoneCompressionSettings.cpp
@@ -119,7 +119,7 @@ public:
 	static const TCHAR* GetTaskName() { return TEXT("FAsyncAnimCompressionTask"); }
 	FORCEINLINE static TStatId GetStatId() { RETURN_QUICK_DECLARE_CYCLE_STAT(FAsyncAnimCompressionTask, STATGROUP_TaskGraphTasks); }
 
-	static ENamedThreads::Type GetDesiredThread() { return ENamedThreads::AnyThread; }
+	static ENamedThreads::Type GetDesiredThread() { return ENamedThreads::AnyBackgroundThreadNormalTask; }
 	static ESubsequentsMode::Type GetSubsequentsMode() { return ESubsequentsMode::TrackSubsequents; }
 
 	void DoTask(ENamedThreads::Type CurrentThread, const FGraphEventRef& MyCompletionGraphEvent)
