Don't iterate on external actors folders to gather the list of modified actor packages when saving the map.
Saves ~5-10s when saving maps with lots of external actors.

diff --git a/Engine/Source/Runtime/Engine/Private/Level.cpp b/Engine/Source/Runtime/Engine/Private/Level.cpp
--- a/Engine/Source/Runtime/Engine/Private/Level.cpp
+++ b/Engine/Source/Runtime/Engine/Private/Level.cpp
@@ -2310,19 +2310,25 @@ TArray<FString> ULevel::GetOnDiskExternalActorPackages() const
 
 TArray<UPackage*> ULevel::GetLoadedExternalActorPackages() const
 {
-	// Only GetExternalPackages is not enough to get to empty packages or deleted actors
 	TSet<UPackage*> ActorPackages;
-	TArray<FString> ActorPackageNames = GetOnDiskExternalActorPackages();
 
-	for (const FString& PackageName : ActorPackageNames)
+	// Get external packages (including deleted actors)
+	ActorPackages.Append(GetPackage()->GetExternalPackages());
+
+	// We also need to provide empty packages (for actors that were converted to non-external)
+	UWorld* World = GetTypedOuter<UWorld>();
+	FString ExternalActorsPath = ULevel::GetExternalActorsPath(World->GetPackage(), World->OriginalWorldName == NAME_None ? World->GetName() : World->OriginalWorldName.ToString());
+	if (!ExternalActorsPath.IsEmpty())
 	{
-		UPackage* ActorPackage = FindObject<UPackage>(nullptr, *PackageName);
-		if (ActorPackage)
+		for (TObjectIterator<UPackage> It; It; ++It)
 		{
-			ActorPackages.Add(ActorPackage);
+			if (!ActorPackages.Contains(*It) && It->FileName.ToString().Contains(ExternalActorsPath))
+			{
+				ActorPackages.Add(*It);
+			}
 		}
 	}
-	ActorPackages.Append(GetPackage()->GetExternalPackages());
+	
 	return ActorPackages.Array();
 }
 
