Translucent line support in FPrimitiveDrawInterface

#rnx

diff --git a/Engine/Source/Runtime/Engine/Private/PrimitiveDrawingUtils.cpp b/Engine/Source/Runtime/Engine/Private/PrimitiveDrawingUtils.cpp
--- a/Engine/Source/Runtime/Engine/Private/PrimitiveDrawingUtils.cpp
+++ b/Engine/Source/Runtime/Engine/Private/PrimitiveDrawingUtils.cpp
@@ -929,7 +929,7 @@ void DrawCircle(FPrimitiveDrawInterface* PDI, const FVector& Base, const FVector
 	for(int32 SideIndex = 0;SideIndex < NumSides;SideIndex++)
 	{
 		const FVector Vertex = Base + (X * FMath::Cos(AngleDelta * (SideIndex + 1)) + Y * FMath::Sin(AngleDelta * (SideIndex + 1))) * Radius;
-		PDI->DrawLine(LastVertex, Vertex, Color, DepthPriority, Thickness, DepthBias, bScreenSpace);
+		PDI->DrawTranslucentLine(LastVertex, Vertex, Color, DepthPriority, Thickness, DepthBias, bScreenSpace);
 		LastVertex = Vertex;
 	}
 }
@@ -945,7 +945,7 @@ void DrawArc(FPrimitiveDrawInterface* PDI, const FVector Base, const FVector X,
 	for(int32 i=0; i<Sections; i++)
 	{
 		FVector ThisVertex = Base + Radius * ( FMath::Cos(CurrentAngle * (PI/180.0f)) * X + FMath::Sin(CurrentAngle * (PI/180.0f)) * Y );
-		PDI->DrawLine( LastVertex, ThisVertex, Color, DepthPriority );
+		PDI->DrawTranslucentLine( LastVertex, ThisVertex, Color, DepthPriority );
 		LastVertex = ThisVertex;
 		CurrentAngle += AngleStep;
 	}
@@ -1092,15 +1092,15 @@ void DrawWireCone(FPrimitiveDrawInterface* PDI, TArray<FVector>& Verts, const FT
 	// Draw spokes.
 	for ( int32 i = 0 ; i < Verts.Num(); ++i )
 	{
-		PDI->DrawLine(Transform.GetLocation(), Verts[i], Color, DepthPriority, Thickness, DepthBias, bScreenSpace);
+		PDI->DrawTranslucentLine(Transform.GetLocation(), Verts[i], Color, DepthPriority, Thickness, DepthBias, bScreenSpace);
 	}
 
 	// Draw rim.
 	for ( int32 i = 0 ; i < Verts.Num()-1 ; ++i )
 	{
-		PDI->DrawLine(Verts[i], Verts[i + 1], Color, DepthPriority, Thickness, DepthBias, bScreenSpace);
+		PDI->DrawTranslucentLine(Verts[i], Verts[i + 1], Color, DepthPriority, Thickness, DepthBias, bScreenSpace);
 	}
-	PDI->DrawLine(Verts[Verts.Num() - 1], Verts[0], Color, DepthPriority, Thickness, DepthBias, bScreenSpace);
+	PDI->DrawTranslucentLine(Verts[Verts.Num() - 1], Verts[0], Color, DepthPriority, Thickness, DepthBias, bScreenSpace);
 }
 
 void DrawWireCone(FPrimitiveDrawInterface* PDI, TArray<FVector>& Verts, const FMatrix& Transform, float ConeLength, float ConeAngle, int32 ConeSides, const FLinearColor& Color, uint8 DepthPriority, float Thickness, float DepthBias, bool bScreenSpace)
diff --git a/Engine/Source/Runtime/Engine/Private/SceneManagement.cpp b/Engine/Source/Runtime/Engine/Private/SceneManagement.cpp
--- a/Engine/Source/Runtime/Engine/Private/SceneManagement.cpp
+++ b/Engine/Source/Runtime/Engine/Private/SceneManagement.cpp
@@ -142,6 +142,31 @@ void FSimpleElementCollector::DrawLine(
 		);
 }
 
+
+void FSimpleElementCollector::DrawTranslucentLine(
+	const FVector& Start,
+	const FVector& End,
+	const FLinearColor& Color,
+	uint8 DepthPriorityGroup,
+	float Thickness/* = 0.0f*/,
+	float DepthBias/* = 0.0f*/,
+	bool bScreenSpace/* = false*/
+)
+{
+	FBatchedElements& Elements = DepthPriorityGroup == SDPG_World ? BatchedElements : TopBatchedElements;
+
+	Elements.AddTranslucentLine(
+		Start,
+		End,
+		Color,
+		HitProxyId,
+		Thickness,
+		DepthBias,
+		bScreenSpace
+	);
+}
+
+
 void FSimpleElementCollector::DrawPoint(
 	const FVector& Position,
 	const FLinearColor& Color,
@@ -547,7 +572,7 @@ const FSceneView& GetLODView(const FSceneView& InView)
 			return *InView.Family->Views[LODViewIndex];
 		}
 	}
-	
+
 	return InView;
 }
 
diff --git a/Engine/Source/Runtime/Engine/Public/SceneManagement.h b/Engine/Source/Runtime/Engine/Public/SceneManagement.h
--- a/Engine/Source/Runtime/Engine/Public/SceneManagement.h
+++ b/Engine/Source/Runtime/Engine/Public/SceneManagement.h
@@ -1962,6 +1962,7 @@ public:
 		float OpacityMaskRefVal = .5f
 		) = 0;
 
+	// Draw an opaque line. The alpha component of Color is ignored.
 	virtual void DrawLine(
 		const FVector& Start,
 		const FVector& End,
@@ -1972,6 +1973,17 @@ public:
 		bool bScreenSpace = false
 		) = 0;
 
+	// Draw a translucent line. The alpha component of Color determines the transparency.
+	virtual void DrawTranslucentLine(
+		const FVector& Start,
+		const FVector& End,
+		const FLinearColor& Color,
+		uint8 DepthPriorityGroup,
+		float Thickness = 0.0f,
+		float DepthBias = 0.0f,
+		bool bScreenSpace = false
+	) = 0;
+
 	virtual void DrawPoint(
 		const FVector& Position,
 		const FLinearColor& Color,
@@ -2079,6 +2091,16 @@ public:
 		bool bScreenSpace = false
 		) override;
 
+	virtual void DrawTranslucentLine(
+		const FVector& Start,
+		const FVector& End,
+		const FLinearColor& Color,
+		uint8 DepthPriorityGroup,
+		float Thickness = 0.0f,
+		float DepthBias = 0.0f,
+		bool bScreenSpace = false
+	) override;
+
 	virtual void DrawPoint(
 		const FVector& Position,
 		const FLinearColor& Color,
diff --git a/Engine/Source/Runtime/Renderer/Private/DynamicPrimitiveDrawing.h b/Engine/Source/Runtime/Renderer/Private/DynamicPrimitiveDrawing.h
--- a/Engine/Source/Runtime/Renderer/Private/DynamicPrimitiveDrawing.h
+++ b/Engine/Source/Runtime/Renderer/Private/DynamicPrimitiveDrawing.h
@@ -46,6 +46,15 @@ public:
 		float DepthBias = 0.0f,
 		bool bScreenSpace = false
 		) override;
+	virtual void DrawTranslucentLine(
+		const FVector& Start,
+		const FVector& End,
+		const FLinearColor& Color,
+		uint8 DepthPriorityGroup,
+		float Thickness = 0.0f,
+		float DepthBias = 0.0f,
+		bool bScreenSpace = false
+	) override;
 	virtual void DrawPoint(
 		const FVector& Position,
 		const FLinearColor& Color,
diff --git a/Engine/Source/Runtime/Renderer/Private/DynamicPrimitiveDrawing.inl b/Engine/Source/Runtime/Renderer/Private/DynamicPrimitiveDrawing.inl
--- a/Engine/Source/Runtime/Renderer/Private/DynamicPrimitiveDrawing.inl
+++ b/Engine/Source/Runtime/Renderer/Private/DynamicPrimitiveDrawing.inl
@@ -135,6 +135,29 @@ inline void FViewElementPDI::DrawLine(
 	);
 }
 
+inline void FViewElementPDI::DrawTranslucentLine(
+	const FVector& Start,
+	const FVector& End,
+	const FLinearColor& Color,
+	uint8 DepthPriorityGroup,
+	float Thickness,
+	float DepthBias,
+	bool bScreenSpace
+)
+{
+	FBatchedElements& Elements = GetElements(DepthPriorityGroup);
+
+	Elements.AddTranslucentLine(
+		Start,
+		End,
+		Color,
+		CurrentHitProxy ? CurrentHitProxy->Id : FHitProxyId(),
+		Thickness,
+		DepthBias,
+		bScreenSpace
+	);
+}
+
 inline void FViewElementPDI::DrawPoint(
 	const FVector& Position,
 	const FLinearColor& Color,
