If r.DumpShaderDebugInfo is enabled, include the path to the dumped info with shader errors.

diff --git a/Engine/Source/Runtime/Engine/Private/ShaderCompiler/ShaderCompiler.cpp b/Engine/Source/Runtime/Engine/Private/ShaderCompiler/ShaderCompiler.cpp
--- a/Engine/Source/Runtime/Engine/Private/ShaderCompiler/ShaderCompiler.cpp
+++ b/Engine/Source/Runtime/Engine/Private/ShaderCompiler/ShaderCompiler.cpp
@@ -572,6 +572,8 @@ bool FShaderCompileUtilities::DoWriteTasks(const TArray<TSharedRef<FShaderCommon
 
 static void ProcessErrors(const FShaderCompileJob& CurrentJob, TArray<FString>& UniqueErrors, FString& ErrorString)
 {
+	bool bReportedDebugInfo = false;
+
 	for (int32 ErrorIndex = 0; ErrorIndex < CurrentJob.Output.Errors.Num(); ErrorIndex++)
 	{
 		FShaderCompilerError CurrentError = CurrentJob.Output.Errors[ErrorIndex];
@@ -612,12 +614,19 @@ static void ProcessErrors(const FShaderCompileJob& CurrentJob, TArray<FString>&
 
 			FString UniqueErrorPrefix;
 
+			// If we dumped the shader info, add it before the first error string
+			if (!GIsBuildMachine && !bReportedDebugInfo && CurrentJob.Input.DumpDebugInfoPath.Len() > 0)
+			{
+				UniqueErrorPrefix += FString::Printf(TEXT("Shader debug info dumped to: \"%s\"\n"), *CurrentJob.Input.DumpDebugInfoPath);
+				bReportedDebugInfo = true;
+			}
+
 			if (CurrentJob.ShaderType)
 			{
 				// Construct a path that will enable VS.NET to find the shader file, relative to the solution
 				const FString SolutionPath = FPaths::RootDir();
 				FString ShaderFilePath = IFileManager::Get().ConvertToAbsolutePathForExternalAppForRead(*CurrentError.GetShaderSourceFilePath());
-				UniqueErrorPrefix = FString::Printf(TEXT("%s(%s): Shader %s, Permutation %d, VF %s:\n\t"),
+				UniqueErrorPrefix += FString::Printf(TEXT("%s(%s): Shader %s, Permutation %d, VF %s:\n\t"),
 					*ShaderFilePath,
 					*CurrentError.ErrorLineString,
 					CurrentJob.ShaderType->GetName(),
@@ -626,7 +635,7 @@ static void ProcessErrors(const FShaderCompileJob& CurrentJob, TArray<FString>&
 			}
 			else
 			{
-				UniqueErrorPrefix = FString::Printf(TEXT("%s(0): "),
+				UniqueErrorPrefix += FString::Printf(TEXT("%s(0): "),
 					*CurrentJob.Input.VirtualSourceFilePath);
 			}
 
@@ -637,7 +646,7 @@ static void ProcessErrors(const FShaderCompileJob& CurrentJob, TArray<FString>&
 				// Format everything on one line, and with the correct verbosity, so we can display proper errors in the failure logs.
 				UE_LOG(LogShaderCompilers, Error, TEXT("%s%s"), *UniqueErrorPrefix.Replace(TEXT("\n"), TEXT("")), *CurrentError.StrippedErrorMessage);
 			}
-			else if (FPlatformMisc::IsDebuggerPresent() && !GIsBuildMachine)
+			else if (FPlatformMisc::IsDebuggerPresent())
 			{
 				// Using OutputDebugString to avoid any text getting added before the filename,
 				// Which will throw off VS.NET's ability to take you directly to the file and line of the error when double clicking it in the output window.
