Mark external actor packages with PKG_ContainsMapData to make sure they won't show up in save content dialogs.

diff --git a/Engine/Source/Runtime/Core/Public/UObject/UE5MainStreamObjectVersion.h b/Engine/Source/Runtime/Core/Public/UObject/UE5MainStreamObjectVersion.h
--- a/Engine/Source/Runtime/Core/Public/UObject/UE5MainStreamObjectVersion.h
+++ b/Engine/Source/Runtime/Core/Public/UObject/UE5MainStreamObjectVersion.h
@@ -28,6 +28,9 @@ struct CORE_API FUE5MainStreamObjectVersion
 		// Serialize GridGuid in PartitionActorDesc
 		PartitionActorDescSerializeGridGuid,
 
+		// Set PKG_ContainsMapData on external actor packages
+		ExternalActorsMapDataPackageFlag,
+
 		// -----<new versions can be added above this line>-------------------------------------------------
 		VersionPlusOne,
 		LatestVersion = VersionPlusOne - 1
diff --git a/Engine/Source/Runtime/Engine/Private/Actor.cpp b/Engine/Source/Runtime/Engine/Private/Actor.cpp
--- a/Engine/Source/Runtime/Engine/Private/Actor.cpp
+++ b/Engine/Source/Runtime/Engine/Private/Actor.cpp
@@ -16,6 +16,7 @@
 #include "Engine/Engine.h"
 #include "UObject/UObjectHash.h"
 #include "UObject/PropertyPortFlags.h"
+#include "UObject/UE5MainStreamObjectVersion.h"
 #include "Kismet/GameplayStatics.h"
 #include "Engine/BlueprintGeneratedClass.h"
 #include "Engine/LocalPlayer.h"
@@ -808,6 +809,8 @@ bool AActor::IsReadyForFinishDestroy()
 
 void AActor::Serialize(FArchive& Ar)
 {
+	Ar.UsingCustomVersion(FUE5MainStreamObjectVersion::GUID);
+
 #if WITH_EDITOR
 	// Prior to load, map natively-constructed component instances for Blueprint-generated class types to any serialized properties that might reference them.
 	// We'll use this information post-load to determine if any owned components may not have been serialized through the reference property (i.e. in case the serialized property value ends up being NULL).
@@ -901,6 +904,16 @@ void AActor::PostLoad()
 		bExchangedRoles = false;
 	}
 
+#if WITH_EDITOR
+	if (GetLinkerCustomVersion(FUE5MainStreamObjectVersion::GUID) < FUE5MainStreamObjectVersion::ExternalActorsMapDataPackageFlag)
+	{
+		if (IsPackageExternal())
+		{
+			GetExternalPackage()->SetPackageFlags(PKG_ContainsMapData);
+		}
+	}
+#endif
+
 #if WITH_EDITORONLY_DATA
 	if (AActor* ParentActor = ParentComponentActor_DEPRECATED.Get())
 	{
diff --git a/Engine/Source/Runtime/Engine/Private/Level.cpp b/Engine/Source/Runtime/Engine/Private/Level.cpp
--- a/Engine/Source/Runtime/Engine/Private/Level.cpp
+++ b/Engine/Source/Runtime/Engine/Private/Level.cpp
@@ -2389,7 +2389,7 @@ UPackage* ULevel::CreateActorPackage(UPackage* InLevelPackage, const FString& In
 	ActorPackageName.Append(*GuidBase36 + 4);
 
 	UPackage* ActorPackage = CreatePackage(*ActorPackageName);
-	ActorPackage->SetPackageFlags(PKG_EditorOnly);
+	ActorPackage->SetPackageFlags(PKG_EditorOnly | PKG_ContainsMapData);
 	return ActorPackage;
 }
 
