Reverting 14298515:
"Fixed bug with AnimNode_SaveCachedPose not properly propogating Initialize to children
PR #7287: Fixed bug with AnimNode_SaveCachedPose does not properly propogate Initialize to children (Contributed by [REDACTED_PERSON])

diff --git a/Engine/Source/Runtime/Engine/Private/Animation/AnimNode_SaveCachedPose.cpp b/Engine/Source/Runtime/Engine/Private/Animation/AnimNode_SaveCachedPose.cpp
--- a/Engine/Source/Runtime/Engine/Private/Animation/AnimNode_SaveCachedPose.cpp
+++ b/Engine/Source/Runtime/Engine/Private/Animation/AnimNode_SaveCachedPose.cpp
@@ -108,9 +108,6 @@ void FAnimNode_SaveCachedPose::PostGraphUpdate()
 			}
 		}
 
-		// Sync with the update counter
-		UpdateCounter.SynchronizeWith(CachedUpdateContexts[MaxWeightIdx].Context.AnimInstanceProxy->GetUpdateCounter());
-
 		// Update the max weighted pose node
 		{
 			TRACE_SCOPED_ANIM_NODE(CachedUpdateContexts[MaxWeightIdx].Context);
