Reorder the if to make static analyzer understand that we already checked for NULL

diff --git a/Engine/Source/Runtime/Engine/Private/Texture2D.cpp b/Engine/Source/Runtime/Engine/Private/Texture2D.cpp
--- a/Engine/Source/Runtime/Engine/Private/Texture2D.cpp
+++ b/Engine/Source/Runtime/Engine/Private/Texture2D.cpp
@@ -909,42 +909,45 @@ FTextureResource* UTexture2D::CreateResource()
 	}
 #endif
 
-	if (IsCurrentlyVirtualTextured())
-	{
-		FVirtualTexture2DResource* ResourceVT = new FVirtualTexture2DResource(this, PrivatePlatformData->VTData, GetCachedLODBias());
-		return ResourceVT;
- 	}
-	else if (PrivatePlatformData)
+	if (PrivatePlatformData)
 	{
-		const EPixelFormat PixelFormat = GetPixelFormat();
-
-		int32 NumMips = FMath::Min3<int32>(PrivatePlatformData->Mips.Num(), GMaxTextureMipCount, FStreamableRenderResourceState::MAX_LOD_COUNT);
-#if !PLATFORM_SUPPORTS_TEXTURE_STREAMING // eg, Android
-		NumMips = MobileReduceLoadedMips(NumMips);
-#endif
-	
-		if (!NumMips)
-		{
-			UE_LOG(LogTexture, Error, TEXT("%s contains no miplevels! Please delete. (Format: %d)"), *GetFullName(), (int)PixelFormat);
-		}
-		else if (!GPixelFormats[PixelFormat].Supported)
-		{
-			UE_LOG(LogTexture, Error, TEXT("%s is %s [raw type %d] which is not supported."), *GetFullName(), GPixelFormats[PixelFormat].Name, static_cast<int32>(PixelFormat));
-		}
-		else if (NumMips == 1 && FMath::Max(GetSizeX(), GetSizeY()) > (int32)GetMax2DTextureDimension())
+		if (IsCurrentlyVirtualTextured())
 		{
-			UE_LOG(LogTexture, Warning, TEXT("%s cannot be created, exceeds this rhi's maximum dimension (%d) and has no mip chain to fall back on."), *GetFullName(), GetMax2DTextureDimension());
+			FVirtualTexture2DResource* ResourceVT = new FVirtualTexture2DResource(this, PrivatePlatformData->VTData, GetCachedLODBias());
+			return ResourceVT;
 		}
 		else
 		{
-			// Should be as big as the mips we have already directly loaded into GPU mem
-			const FStreamableRenderResourceState PostInitState = GetResourcePostInitState(PlatformData, !bTemporarilyDisableStreaming, ResourceMem ? ResourceMem->GetNumMips() : 0, NumMips);
-			FTexture2DResource* Texture2DResource = new FTexture2DResource(this, PostInitState);
-			// preallocated memory for the UTexture2D resource is now owned by this resource
-			// and will be freed by the RHI resource or when the FTexture2DResource is deleted
-			ResourceMem = nullptr;
-
-			return Texture2DResource;
+			const EPixelFormat PixelFormat = GetPixelFormat();
+
+			int32 NumMips = FMath::Min3<int32>(PrivatePlatformData->Mips.Num(), GMaxTextureMipCount, FStreamableRenderResourceState::MAX_LOD_COUNT);
+#if !PLATFORM_SUPPORTS_TEXTURE_STREAMING // eg, Android
+			NumMips = MobileReduceLoadedMips(NumMips);
+#endif
+	
+			if (!NumMips)
+			{
+				UE_LOG(LogTexture, Error, TEXT("%s contains no miplevels! Please delete. (Format: %d)"), *GetFullName(), (int)PixelFormat);
+			}
+			else if (!GPixelFormats[PixelFormat].Supported)
+			{
+				UE_LOG(LogTexture, Error, TEXT("%s is %s [raw type %d] which is not supported."), *GetFullName(), GPixelFormats[PixelFormat].Name, static_cast<int32>(PixelFormat));
+			}
+			else if (NumMips == 1 && FMath::Max(GetSizeX(), GetSizeY()) > (int32)GetMax2DTextureDimension())
+			{
+				UE_LOG(LogTexture, Warning, TEXT("%s cannot be created, exceeds this rhi's maximum dimension (%d) and has no mip chain to fall back on."), *GetFullName(), GetMax2DTextureDimension());
+			}
+			else
+			{
+				// Should be as big as the mips we have already directly loaded into GPU mem
+				const FStreamableRenderResourceState PostInitState = GetResourcePostInitState(PlatformData, !bTemporarilyDisableStreaming, ResourceMem ? ResourceMem->GetNumMips() : 0, NumMips);
+				FTexture2DResource* Texture2DResource = new FTexture2DResource(this, PostInitState);
+				// preallocated memory for the UTexture2D resource is now owned by this resource
+				// and will be freed by the RHI resource or when the FTexture2DResource is deleted
+				ResourceMem = nullptr;
+
+				return Texture2DResource;
+			}
 		}
 	}
 
