[NO NON-PII COMMIT MESSAGE]

diff --git a/Engine/Source/Runtime/Engine/Private/RepLayout.cpp b/Engine/Source/Runtime/Engine/Private/RepLayout.cpp
--- a/Engine/Source/Runtime/Engine/Private/RepLayout.cpp
+++ b/Engine/Source/Runtime/Engine/Private/RepLayout.cpp
@@ -3783,7 +3783,7 @@ bool FRepLayout::ReceiveProperties_BackwardsCompatible_r(
 				const int32 ArrayElementOffset = Index * Cmd.ElementSize;
 
 				FRepObjectDataBuffer ElementData = LocalData + ArrayElementOffset;
-				FRepShadowDataBuffer ElementShadowData = (LocalShadowData && (Index < static_cast<uint32>(ShadowArray->Num()))) ? LocalShadowData + ArrayElementOffset : nullptr;
+				FRepShadowDataBuffer ElementShadowData = (LocalShadowData && (Index < (ShadowArray ? static_cast<uint32>(ShadowArray->Num()) : 0))) ? LocalShadowData + ArrayElementOffset : nullptr;
 
 				if (!ReceiveProperties_BackwardsCompatible_r(RepState, NetFieldExportGroup, TempReader, CmdIndex + 1, Cmd.EndCmd - 1, ElementShadowData, LocalData, ElementData, NewGuidReferencesArray, bOutHasUnmapped, bOutGuidsChanged))
 				{
