Fix static analysis nullptr warnings in PinBase material expression

diff --git a/Engine/Source/Runtime/Engine/Private/Materials/MaterialExpressions.cpp b/Engine/Source/Runtime/Engine/Private/Materials/MaterialExpressions.cpp
--- a/Engine/Source/Runtime/Engine/Private/Materials/MaterialExpressions.cpp
+++ b/Engine/Source/Runtime/Engine/Private/Materials/MaterialExpressions.cpp
@@ -10500,19 +10500,10 @@ void UMaterialExpressionPinBase::PostEditChangeProperty(FPropertyChangedEvent& P
 
 			AddedReroute.Expression = NewObject<UMaterialExpressionReroute>(GetOuter(), UMaterialExpressionReroute::StaticClass(), NAME_None, RF_Transactional);
 			AddedReroute.Expression->SubgraphExpression = SubgraphExpression;
-			
-			if (Material)
-			{
-				Material->Expressions.Add(AddedReroute.Expression);
-				AddedReroute.Expression->Material = Material;
-			}
-			else if (Function)
-			{
-				Function->Modify();
-				Function->FunctionExpressions.Add(AddedReroute.Expression);
-			}
-
+			AddedReroute.Expression->Material = Material;
 			AddedReroute.Name = AddedReroute.Name.IsNone() ? FName(FString::Printf(TEXT("Pin %u"), AddedRerouteIndex + 1)) : AddedReroute.Name;
+
+			Material->Expressions.Add(AddedReroute.Expression);
 		}
 		else if (PropertyChangedEvent.ChangeType == EPropertyChangeType::ArrayRemove)
 		{
