Fixed nullptr access

#rnx

diff --git a/Engine/Source/Runtime/Engine/Private/Elements/EngineElementsLibrary.cpp b/Engine/Source/Runtime/Engine/Private/Elements/EngineElementsLibrary.cpp
--- a/Engine/Source/Runtime/Engine/Private/Elements/EngineElementsLibrary.cpp
+++ b/Engine/Source/Runtime/Engine/Private/Elements/EngineElementsLibrary.cpp
@@ -54,9 +54,13 @@ FTypedElementHandle AcquireEditorTypedElementHandle(const ObjectClass* Object, T
 {
 	if (GIsEditor)
 	{
-		return bAllowCreate
-			? ElementOwnerStore.FindOrRegisterElementOwner(Object, [Object, &CreateElement](){ return CreateElement(Object); })->AcquireHandle()
-			: ElementOwnerStore.FindElementOwner(Object)->AcquireHandle();
+		TTypedElementOwnerScopedAccess<ElementDataType> EditorElement = bAllowCreate
+			? ElementOwnerStore.FindOrRegisterElementOwner(Object, [Object, &CreateElement](){ return CreateElement(Object); })
+			: ElementOwnerStore.FindElementOwner(Object);
+		if (EditorElement)
+		{
+			return EditorElement->AcquireHandle();
+		}
 	}
 
 	return FTypedElementHandle();
