Fixed SDF/Card async build leak in case if Nanite mesh descriptor build fails.

diff --git a/Engine/Source/Runtime/Engine/Private/DistanceFieldAtlas.cpp b/Engine/Source/Runtime/Engine/Private/DistanceFieldAtlas.cpp
--- a/Engine/Source/Runtime/Engine/Private/DistanceFieldAtlas.cpp
+++ b/Engine/Source/Runtime/Engine/Private/DistanceFieldAtlas.cpp
@@ -967,6 +967,7 @@ void FDistanceFieldVolumeData::CacheDerivedData(const FString& InDDCKey, const I
 			if (!MeshBuilderModule.BuildMesh(Mesh, NewTask->SourceMeshData.Vertices, NewTask->SourceMeshData.Indices, NewTask->SourceMeshData.Sections, /* bBuildOnlyPosition */ true, NumTexCoords, bHasColors))
 			{
 				UE_LOG(LogStaticMesh, Error, TEXT("Failed to build static mesh. See previous line(s) for details."));
+				delete NewTask;
 				return;
 			}
 		}
diff --git a/Engine/Source/Runtime/Engine/Private/MeshCardRepresentation.cpp b/Engine/Source/Runtime/Engine/Private/MeshCardRepresentation.cpp
--- a/Engine/Source/Runtime/Engine/Private/MeshCardRepresentation.cpp
+++ b/Engine/Source/Runtime/Engine/Private/MeshCardRepresentation.cpp
@@ -132,6 +132,7 @@ void FCardRepresentationData::CacheDerivedData(const FString& InDDCKey, const IT
 			if (!MeshBuilderModule.BuildMesh(Mesh, NewTask->SourceMeshData.Vertices, NewTask->SourceMeshData.Indices, NewTask->SourceMeshData.Sections, /* bBuildOnlyPosition */ true, NumTexCoords, bHasColors))
 			{
 				UE_LOG(LogStaticMesh, Error, TEXT("Failed to build static mesh. See previous line(s) for details."));
+				delete NewTask;
 				return;
 			}
 		}
