Improve the Blend Profiles by allowing time based profiles.
Users can now select the blend profile mode: Blend Weight Factor based or Time Factor based.
The Blend Weight Factor method is the legacy method, which multiplies the transition blend weight with the per bone factor.
The new Time Factor method allows you to specify a time factor. A value of 0.5 would mean the bone would take half the time of the transition to reach its target state, while a value of 0.1 would mean it takes a tenth of the time, etc. A value of 0 would make the bone instantly transition into the target state.

diff --git a/Engine/Source/Editor/AnimGraph/Classes/AnimStateTransitionNode.h b/Engine/Source/Editor/AnimGraph/Classes/AnimStateTransitionNode.h
--- a/Engine/Source/Editor/AnimGraph/Classes/AnimStateTransitionNode.h
+++ b/Engine/Source/Editor/AnimGraph/Classes/AnimStateTransitionNode.h
@@ -7,11 +7,11 @@
 #include "Misc/Guid.h"
 #include "Animation/AnimTypes.h"
 #include "Animation/AnimStateMachineTypes.h"
+#include "Animation/BlendProfile.h"
 #include "AnimStateNodeBase.h"
 
 #include "AnimStateTransitionNode.generated.h"
 
-class UBlendProfile;
 class UCurveFloat;
 class UEdGraph;
 class UEdGraphPin;
@@ -52,6 +52,10 @@ class UAnimStateTransitionNode : public UAnimStateNodeBase
 	UPROPERTY(EditAnywhere, Category=Transition)
 	UBlendProfile* BlendProfile;
 
+	// The blend profile mode to use. This determines how the per bone blend weights are calculated when transitioning.
+	UPROPERTY(EditAnywhere, Category = Transition)
+	EBlendProfileMode BlendProfileMode;
+
 	// Try setting the rule automatically based on most relevant player node's remaining time and the CrossfadeDuration of the transition, ignoring the internal time
 	UPROPERTY(EditAnywhere, Category=Transition)
 	bool bAutomaticRuleBasedOnSequencePlayerInState;
diff --git a/Engine/Source/Editor/AnimGraph/Private/AnimGraphNode_StateMachineBase.cpp b/Engine/Source/Editor/AnimGraph/Private/AnimGraphNode_StateMachineBase.cpp
--- a/Engine/Source/Editor/AnimGraph/Private/AnimGraphNode_StateMachineBase.cpp
+++ b/Engine/Source/Editor/AnimGraph/Private/AnimGraphNode_StateMachineBase.cpp
@@ -383,6 +383,7 @@ void UAnimGraphNode_StateMachineBase::OnProcessDuringCompilation(IAnimBlueprintC
 			BakedTransition.BlendMode = TransitionNode->BlendMode;
 			BakedTransition.CustomCurve = TransitionNode->CustomBlendCurve;
 			BakedTransition.BlendProfile = TransitionNode->BlendProfile;
+			BakedTransition.BlendProfileMode = TransitionNode->BlendProfileMode;
 			BakedTransition.LogicType = TransitionNode->LogicType;
 
 			UAnimStateNodeBase* PreviousState = TransitionNode->GetPreviousState();
diff --git a/Engine/Source/Editor/AnimGraph/Private/AnimStateTransitionNode.cpp b/Engine/Source/Editor/AnimGraph/Private/AnimStateTransitionNode.cpp
--- a/Engine/Source/Editor/AnimGraph/Private/AnimStateTransitionNode.cpp
+++ b/Engine/Source/Editor/AnimGraph/Private/AnimStateTransitionNode.cpp
@@ -17,6 +17,7 @@
 #include "Kismet2/Kismet2NameValidators.h"
 #include "ScopedTransaction.h"
 #include "Animation/BlendProfile.h"
+#include "UObject/UE5MainStreamObjectVersion.h"
 
 //////////////////////////////////////////////////////////////////////////
 // IAnimStateTransitionNodeSharedDataHelper
@@ -80,6 +81,7 @@ UAnimStateTransitionNode::UAnimStateTransitionNode(const FObjectInitializer& Obj
 	Bidirectional = false;
 	PriorityOrder = 1;
 	LogicType = ETransitionLogicType::TLT_StandardBlend;
+	BlendProfileMode = EBlendProfileMode::TimeFactor;
 }
 
 void UAnimStateTransitionNode::AllocateDefaultPins()
@@ -195,7 +197,7 @@ void UAnimStateTransitionNode::PostPasteNode()
 	Super::PostPasteNode();
 
 	// We don't want to paste nodes in that aren't fully linked (transition nodes have fixed pins as they
-	// really describle the connection between two other nodes). If we find one missing link, get rid of the node.
+	// really describe the connection between two other nodes). If we find one missing link, get rid of the node.
 	for(UEdGraphPin* Pin : Pins)
 	{
 		if(Pin->LinkedTo.Num() == 0)
@@ -318,7 +320,8 @@ void UAnimStateTransitionNode::PostEditChangeProperty(struct FPropertyChangedEve
 	if (PropertyName == GET_MEMBER_NAME_CHECKED(UAnimStateTransitionNode, CrossfadeDuration) || 
 		PropertyName == GET_MEMBER_NAME_CHECKED(UAnimStateTransitionNode, BlendMode) ||
 		PropertyName == GET_MEMBER_NAME_CHECKED(UAnimStateTransitionNode, CustomBlendCurve) ||
-		PropertyName == GET_MEMBER_NAME_CHECKED(UAnimStateTransitionNode, BlendProfile))
+		PropertyName == GET_MEMBER_NAME_CHECKED(UAnimStateTransitionNode, BlendProfile) ||
+		PropertyName == GET_MEMBER_NAME_CHECKED(UAnimStateTransitionNode, BlendProfileMode))
 	{
 		PropagateCrossfadeSettings();
 	}
@@ -465,6 +468,7 @@ void UAnimStateTransitionNode::CopyCrossfadeSettings(const UAnimStateTransitionN
 	BlendMode = SrcNode->BlendMode;
 	CustomBlendCurve = SrcNode->CustomBlendCurve;
 	BlendProfile = SrcNode->BlendProfile;
+	BlendProfileMode = SrcNode->BlendProfileMode;
 	SharedCrossfadeIdx = SrcNode->SharedCrossfadeIdx;
 	SharedCrossfadeName = SrcNode->SharedCrossfadeName;
 	SharedCrossfadeGuid = SrcNode->SharedCrossfadeGuid;
@@ -546,6 +550,13 @@ void UAnimStateTransitionNode::Serialize(FArchive& Ar)
 {
 	Super::Serialize(Ar);
 	Ar.UsingCustomVersion(FAnimPhysObjectVersion::GUID);
+	Ar.UsingCustomVersion(FUE5MainStreamObjectVersion::GUID);
+
+	// Make sure we use the old behavior on Blend Profiles for older files.
+	if (Ar.IsLoading() && Ar.CustomVer(FUE5MainStreamObjectVersion::GUID) < FUE5MainStreamObjectVersion::AnimationAddedBlendProfileModes)
+	{
+		BlendProfileMode = EBlendProfileMode::WeightFactor;
+	}
 }
 
 void UAnimStateTransitionNode::DestroyNode()
diff --git a/Engine/Source/Editor/DetailCustomizations/Private/AnimTransitionNodeDetails.cpp b/Engine/Source/Editor/DetailCustomizations/Private/AnimTransitionNodeDetails.cpp
--- a/Engine/Source/Editor/DetailCustomizations/Private/AnimTransitionNodeDetails.cpp
+++ b/Engine/Source/Editor/DetailCustomizations/Private/AnimTransitionNodeDetails.cpp
@@ -259,6 +259,23 @@ void FAnimTransitionNodeDetails::CustomizeDetails( IDetailLayoutBuilder& DetailB
 				[
 					SkeletonEditorModule.CreateBlendProfilePicker(TargetSkeleton, Args)
 				];
+
+			// Blend Profile Mode
+			auto BlendProfileEnabled = [BlendProfileHandle, BlendSettingsEnabledAttribute]()
+			{
+				if (BlendProfileHandle->IsValidHandle())
+				{
+					UObject* BlendProfilePropertyValueInternal = nullptr;
+					BlendProfileHandle->GetValue(BlendProfilePropertyValueInternal);
+					UBlendProfile* SelectedProfile = Cast<UBlendProfile>(BlendProfilePropertyValueInternal);
+					return (SelectedProfile != nullptr) && BlendSettingsEnabledAttribute.Get(false);
+				}
+				return false;
+			};
+			auto BlendProfileEnabledAttribute = TAttribute<bool>::Create(TAttribute<bool>::FGetter::CreateLambda(BlendProfileEnabled));
+			CrossfadeCategory.AddProperty(GET_MEMBER_NAME_CHECKED(UAnimStateTransitionNode, BlendProfileMode))
+				.DisplayName(LOCTEXT("BlendProfileModeLabel", "Blend Profile Mode"))
+				.IsEnabled(BlendProfileEnabledAttribute);
 		}
 
 		//////////////////////////////////////////////////////////////////////////
diff --git a/Engine/Source/Runtime/Core/Public/UObject/UE5MainStreamObjectVersion.h b/Engine/Source/Runtime/Core/Public/UObject/UE5MainStreamObjectVersion.h
--- a/Engine/Source/Runtime/Core/Public/UObject/UE5MainStreamObjectVersion.h
+++ b/Engine/Source/Runtime/Core/Public/UObject/UE5MainStreamObjectVersion.h
@@ -31,6 +31,9 @@ struct CORE_API FUE5MainStreamObjectVersion
 		// Set PKG_ContainsMapData on external actor packages
 		ExternalActorsMapDataPackageFlag,
 
+		// Added a new configurable BlendProfileMode that the user can setup to control the behavior of blend profiles.
+		AnimationAddedBlendProfileModes,
+
 		// -----<new versions can be added above this line>-------------------------------------------------
 		VersionPlusOne,
 		LatestVersion = VersionPlusOne - 1
diff --git a/Engine/Source/Runtime/Engine/Classes/Animation/AnimNode_StateMachine.h b/Engine/Source/Runtime/Engine/Classes/Animation/AnimNode_StateMachine.h
--- a/Engine/Source/Runtime/Engine/Classes/Animation/AnimNode_StateMachine.h
+++ b/Engine/Source/Runtime/Engine/Classes/Animation/AnimNode_StateMachine.h
@@ -10,10 +10,10 @@
 #include "Animation/AnimStateMachineTypes.h"
 #include "Animation/AnimNodeBase.h"
 #include "Animation/AnimInstance.h"
+#include "Animation/BlendProfile.h"
 #include "AnimNode_StateMachine.generated.h"
 
 class IAnimClassInterface;
-class UBlendProfile;
 struct FAnimNode_AssetPlayerBase;
 struct FAnimNode_StateMachine;
 struct FAnimNode_TransitionPoseEvaluator;
@@ -67,6 +67,10 @@ public:
 	UPROPERTY()
 	UBlendProfile* BlendProfile;
 
+	// The blend profile mode to use. This determines how the per bone blend weights are calculated when transitioning.
+	UPROPERTY()
+	EBlendProfileMode BlendProfileMode;
+
 	// Type of blend to use
 	EAlphaBlendOption BlendOption;
 
@@ -91,6 +95,7 @@ public:
 protected:
 	float CalculateInverseAlpha(EAlphaBlendOption BlendMode, float InFraction) const;
 	float CalculateAlpha(float InFraction) const;
+	void ProcessBlendProfile(int32 TransitionDirectionIndex);
 };
 
 USTRUCT()
diff --git a/Engine/Source/Runtime/Engine/Classes/Animation/AnimStateMachineTypes.h b/Engine/Source/Runtime/Engine/Classes/Animation/AnimStateMachineTypes.h
--- a/Engine/Source/Runtime/Engine/Classes/Animation/AnimStateMachineTypes.h
+++ b/Engine/Source/Runtime/Engine/Classes/Animation/AnimStateMachineTypes.h
@@ -7,9 +7,9 @@
 #include "UObject/ObjectMacros.h"
 #include "UObject/Object.h"
 #include "AlphaBlend.h"
+#include "BlendProfile.h"
 #include "AnimStateMachineTypes.generated.h"
 
-class UBlendProfile;
 class UCurveFloat;
 
 //@TODO: Document
@@ -146,6 +146,9 @@ struct FAnimationTransitionBetweenStates : public FAnimationStateBase
 	UPROPERTY()
 	UBlendProfile* BlendProfile;
 
+	UPROPERTY()
+	EBlendProfileMode BlendProfileMode;
+
 	UPROPERTY()
 	TEnumAsByte<ETransitionLogicType::Type> LogicType;
 
@@ -165,6 +168,7 @@ struct FAnimationTransitionBetweenStates : public FAnimationStateBase
 		, BlendMode(EAlphaBlendOption::CubicInOut)
 		, CustomCurve(nullptr)
 		, BlendProfile(nullptr)
+		, BlendProfileMode(EBlendProfileMode::TimeFactor)
 		, LogicType(ETransitionLogicType::TLT_StandardBlend)
 #if WITH_EDITOR
 		, ReverseTransition(false)
diff --git a/Engine/Source/Runtime/Engine/Classes/Animation/BlendProfile.h b/Engine/Source/Runtime/Engine/Classes/Animation/BlendProfile.h
--- a/Engine/Source/Runtime/Engine/Classes/Animation/BlendProfile.h
+++ b/Engine/Source/Runtime/Engine/Classes/Animation/BlendProfile.h
@@ -9,6 +9,25 @@
 #include "AnimationRuntime.h"
 #include "BlendProfile.generated.h"
 
+/** The mode in which the blend profile should be applied. */
+UENUM()
+enum class EBlendProfileMode : uint8
+{
+	// The bone's transition time is a factor based on the transition time. 
+	// For example 0.5 means it takes half the time of the transition.
+	// Values should be between 0 and 1. They will be clamped if they go out of this range.
+	// A bone value of 0 means the bone will instantly transition into the target state.
+	TimeFactor = 0,
+
+	// The bone's transition weight is multiplied by this factor.
+	// For example 2.0 means the bone's blend weight is twice as high as the transition's blend weight.
+	// Values should typically be equal or greater than 1.0.
+	// If you want certain bones to instantly transition into the target state
+	// the Time Factor based method might be a better choice.
+	WeightFactor
+};
+
+
 /** A single entry for a blend scale within a profile, mapping a bone to a blendscale */
 USTRUCT()
 struct FBlendProfileBoneEntry
diff --git a/Engine/Source/Runtime/Engine/Private/Animation/AnimNode_StateMachine.cpp b/Engine/Source/Runtime/Engine/Private/Animation/AnimNode_StateMachine.cpp
--- a/Engine/Source/Runtime/Engine/Private/Animation/AnimNode_StateMachine.cpp
+++ b/Engine/Source/Runtime/Engine/Private/Animation/AnimNode_StateMachine.cpp
@@ -33,6 +33,7 @@ FAnimationActiveTransitionEntry::FAnimationActiveTransitionEntry()
 	, EndNotify(INDEX_NONE)
 	, InterruptNotify(INDEX_NONE)
 	, BlendProfile(nullptr)
+	, BlendProfileMode(EBlendProfileMode::TimeFactor)
 	, BlendOption(EAlphaBlendOption::HermiteCubic)
 	, LogicType(ETransitionLogicType::TLT_StandardBlend)
 	, bActive(false)
@@ -48,6 +49,7 @@ FAnimationActiveTransitionEntry::FAnimationActiveTransitionEntry(int32 NextState
 	, EndNotify(ReferenceTransitionInfo.EndNotify)
 	, InterruptNotify(ReferenceTransitionInfo.InterruptNotify)
 	, BlendProfile(ReferenceTransitionInfo.BlendProfile)
+	, BlendProfileMode(ReferenceTransitionInfo.BlendProfileMode)
 	, BlendOption(ReferenceTransitionInfo.BlendMode)
 	, LogicType(ReferenceTransitionInfo.LogicType)
 	, bActive(true)
@@ -111,6 +113,72 @@ void FAnimationActiveTransitionEntry::InitializeCustomGraphLinks(const FAnimatio
 	}
 }
 
+
+void FAnimationActiveTransitionEntry::ProcessBlendProfile(int32 TransitionDirectionIndex)
+{
+	const bool bForwards = (TransitionDirectionIndex == 0);
+	FBlendSampleData& CurrentData = StateBlendData[TransitionDirectionIndex];
+
+	switch (BlendProfileMode)
+	{
+		// The per bone value is a factor of the transition time, where 0.5 means half the transition time, 0.1 means one tenth of the transition time, etc.
+		case EBlendProfileMode::TimeFactor:
+		{
+			for (int32 PerBoneIndex = 0; PerBoneIndex < CurrentData.PerBoneBlendData.Num(); ++PerBoneIndex)
+			{
+				// Calculate the blend weight for this bone.
+				const float BoneTimeScaleFactor = FMath::Clamp(BlendProfile->GetEntryBlendScale(PerBoneIndex), 0.0f, 1.0f);
+				const float BoneTransitionDuration = BoneTimeScaleFactor * CrossfadeDuration;
+				float NormalizedBoneWeight = (BoneTransitionDuration > ZERO_ANIMWEIGHT_THRESH) ? ElapsedTime / BoneTransitionDuration : 1.0f;
+				NormalizedBoneWeight = FMath::Clamp(NormalizedBoneWeight, 0.0f, 1.0f);
+				NormalizedBoneWeight = FAlphaBlend::AlphaToBlendOption(NormalizedBoneWeight, Blend.GetBlendOption(), Blend.GetCustomCurve());
+				if (NormalizedBoneWeight < ZERO_ANIMWEIGHT_THRESH)
+				{
+					NormalizedBoneWeight = ZERO_ANIMWEIGHT_THRESH;
+				}
+
+				// Reverse the weight when we go backward.
+				if (!bForwards)
+				{
+					NormalizedBoneWeight = 1.0f - NormalizedBoneWeight;
+				}
+
+				CurrentData.PerBoneBlendData[PerBoneIndex] = NormalizedBoneWeight;
+			}
+		}
+		break;
+
+		// The per bone value is a factor of the transition blend weight.
+		case EBlendProfileMode::WeightFactor:
+		{
+			for (int32 PerBoneIndex = 0; PerBoneIndex < CurrentData.PerBoneBlendData.Num(); ++PerBoneIndex)
+			{
+				// Get the bone weight factor from the profile.
+				float BoneWeightFactor = BlendProfile->GetEntryBlendScale(PerBoneIndex);
+				if (!bForwards)
+				{
+					BoneWeightFactor = (BoneWeightFactor > ZERO_ANIMWEIGHT_THRESH) ? 1.0f / BoneWeightFactor : 1.0f;
+				}
+
+				// Multiply it with the transition weight and make sure it isn't zero, as this can cause normalization issues.
+				BoneWeightFactor *= CurrentData.TotalWeight;
+				if (BoneWeightFactor < ZERO_ANIMWEIGHT_THRESH)
+				{
+					BoneWeightFactor = ZERO_ANIMWEIGHT_THRESH;
+				}
+
+				CurrentData.PerBoneBlendData[PerBoneIndex] = BoneWeightFactor;
+			}
+		}
+		break;
+
+		default:
+			checkf(false, TEXT("The selected Blend Profile Mode is not supported (Mode=%d)"), BlendProfileMode);
+			break;
+	}
+}
+
+
 void FAnimationActiveTransitionEntry::Update(const FAnimationUpdateContext& Context, int32 CurrentStateIndex, bool& bOutFinished)
 {
 	bOutFinished = false;
@@ -121,9 +189,8 @@ void FAnimationActiveTransitionEntry::Update(const FAnimationUpdateContext& Cont
 		ElapsedTime += Context.GetDeltaTime();
 		Blend.Update(Context.GetDeltaTime());
 
-		float QueryAlpha = 1.0f;
-
 		// If non-zero, calculate the query alpha
+		float QueryAlpha = 0.0f;
 		if (CrossfadeDuration > 0.0f)
 		{
 			QueryAlpha = ElapsedTime / CrossfadeDuration;
@@ -131,34 +198,20 @@ void FAnimationActiveTransitionEntry::Update(const FAnimationUpdateContext& Cont
 
 		Alpha = FAlphaBlend::AlphaToBlendOption(QueryAlpha, Blend.GetBlendOption(), Blend.GetCustomCurve());
 
-		if(Blend.IsComplete())
+		if (Blend.IsComplete())
 		{
 			bActive = false;
 			bOutFinished = true;
 		}
 
 		// Update state blend data (only when we're using per-bone)
-		if(BlendProfile)
+		if (BlendProfile)
 		{
-			for(int32 Idx = 0 ; Idx < 2 ; ++Idx)
+			for (int32 Idx = 0 ; Idx < 2 ; ++Idx)
 			{
-				bool bForwards = Idx == 0;
-				FBlendSampleData& CurrentData = StateBlendData[Idx];
-
-				CurrentData.TotalWeight = (bForwards) ? Alpha : 1.0f - Alpha;
-
-				for(int32 PerBoneIndex = 0 ; PerBoneIndex < CurrentData.PerBoneBlendData.Num() ; ++PerBoneIndex)
-				{
-					float& BoneBlend = CurrentData.PerBoneBlendData[PerBoneIndex];
-					float WeightScale = BlendProfile->GetEntryBlendScale(PerBoneIndex);
-
-					if(!bForwards)
-					{
-						WeightScale = 1.0f / WeightScale;
-					}
-
-					BoneBlend = CurrentData.TotalWeight * WeightScale;
-				}
+				const bool bForwards = (Idx == 0);
+				StateBlendData[Idx].TotalWeight = bForwards ? Alpha : 1.0f - Alpha;
+				ProcessBlendProfile(Idx);
 			}
 
 			FBlendSampleData::NormalizeDataWeight(StateBlendData);
