Fix static analysis

diff --git a/Engine/Source/Runtime/Core/Private/Math/Color.cpp b/Engine/Source/Runtime/Core/Private/Math/Color.cpp
--- a/Engine/Source/Runtime/Core/Private/Math/Color.cpp
+++ b/Engine/Source/Runtime/Core/Private/Math/Color.cpp
@@ -107,8 +107,9 @@ FColor FLinearColor::ToRGBE() const
 	}
 	else
 	{
-		int32 Exponent;
-		frexpf(Primary, &Exponent);
+		// The following replaces a call to frexpf, because frexpf would have a warning for an unused return value.
+		// Additionally, this usage of logbf assumes FLT_RADIX == 2
+		int32 Exponent = (Primary == 0.0f) ? 0 : (1 + logbf(Primary));
 		const float Scale = ldexpf(1.f, -Exponent + 8);
 
 		Color.R = (uint8)FMath::Clamp(FMath::TruncToInt(R * Scale), 0, 255);
diff --git a/Engine/Source/Runtime/Engine/Private/ImageUtils.cpp b/Engine/Source/Runtime/Engine/Private/ImageUtils.cpp
--- a/Engine/Source/Runtime/Engine/Private/ImageUtils.cpp
+++ b/Engine/Source/Runtime/Engine/Private/ImageUtils.cpp
@@ -709,8 +709,9 @@ private:
 		}
 		else
 		{
-			int32 Exponent;
-			frexpf(Primary, &Exponent);
+			// The following replaces a call to frexpf, because frexpf would have a warning for an unused return value.
+			// Additionally, this usage of logbf assumes FLT_RADIX == 2
+			int32 Exponent = (Primary == 0.0f) ? 0 : (1 + logbf(Primary));
 			const float Scale = ldexpf(1.f, -Exponent + 8);
 
 			ReturnColor.R = FMath::Clamp(FMath::TruncToInt((R * Scale) + Rand.GetFraction()), 0, 255);
