Added support for strata on material functions asset type.

diff --git a/Engine/Source/Runtime/Engine/Private/Materials/MaterialExpressions.cpp b/Engine/Source/Runtime/Engine/Private/Materials/MaterialExpressions.cpp
--- a/Engine/Source/Runtime/Engine/Private/Materials/MaterialExpressions.cpp
+++ b/Engine/Source/Runtime/Engine/Private/Materials/MaterialExpressions.cpp
@@ -5680,6 +5680,10 @@ uint32 UMaterialExpressionMakeMaterialAttributes::GetInputType(int32 InputIndex)
 	{
 		return MCT_ShadingModel;
 	}
+	else if (GetInputName(InputIndex).IsEqual("FrontMaterial"))
+	{
+		return MCT_Strata;
+	}
 	else
 	{
 		return UMaterialExpression::GetInputType(InputIndex);
@@ -6008,6 +6012,10 @@ uint32 UMaterialExpressionGetMaterialAttributes::GetOutputType(int32 OutputIndex
 		{
 			OutputType = PinType;
 		}
+		else if (PinType == MCT_Strata)
+		{
+			OutputType = PinType;
+		}
 	}
 
 	return OutputType;
