Made strata work with material attribute blending node.

diff --git a/Engine/Source/Runtime/Engine/Classes/Materials/MaterialExpressionStrata.h b/Engine/Source/Runtime/Engine/Classes/Materials/MaterialExpressionStrata.h
--- a/Engine/Source/Runtime/Engine/Classes/Materials/MaterialExpressionStrata.h
+++ b/Engine/Source/Runtime/Engine/Classes/Materials/MaterialExpressionStrata.h
@@ -8,6 +8,17 @@
 #include "MaterialExpressionStrata.generated.h"
 
 
+/**
+ * Compile a special blend function for strata when blending material attribute
+ *
+ * @param Compiler				The compiler to add code to
+ * @param Foreground			Entry A, has a bigger impact when Alpha is close to 1
+ * @param Background			Entry B, has a bigger impact when Alpha is close to 0
+ * @param Alpha					Blend factor [0..1], when 0
+ * @return						Index to a new code chunk
+ */
+extern int32 CompileStrataBlendFunction(FMaterialCompiler* Compiler, int32 Foreground, int32 Background, int32 Alpha);
+
 
 ///////////////////////////////////////////////////////////////////////////////
 // BSDF nodes
diff --git a/Engine/Source/Runtime/Engine/Private/Materials/MaterialExpressions.cpp b/Engine/Source/Runtime/Engine/Private/Materials/MaterialExpressions.cpp
--- a/Engine/Source/Runtime/Engine/Private/Materials/MaterialExpressions.cpp
+++ b/Engine/Source/Runtime/Engine/Private/Materials/MaterialExpressions.cpp
@@ -598,6 +598,15 @@ int32 CompileShadingModelBlendFunction(FMaterialCompiler* Compiler, const int32
 	return Compiler->If(Alpha, MidPoint, B, INDEX_NONE, A, INDEX_NONE);
 }
 
+int32 CompileStrataBlendFunction(FMaterialCompiler* Compiler, int32 Foreground, int32 Background, int32 Alpha)
+{
+	if (Foreground == INDEX_NONE || Background == INDEX_NONE || Alpha == INDEX_NONE)
+	{
+		return INDEX_NONE;
+	}
+	return Compiler->StrataHorizontalMixing(Foreground, Background, Alpha);
+}
+
 UMaterialExpression::UMaterialExpression(const FObjectInitializer& ObjectInitializer)
 	: Super(ObjectInitializer)
 #if WITH_EDITORONLY_DATA
diff --git a/Engine/Source/Runtime/Engine/Private/Materials/MaterialShared.cpp b/Engine/Source/Runtime/Engine/Private/Materials/MaterialShared.cpp
--- a/Engine/Source/Runtime/Engine/Private/Materials/MaterialShared.cpp
+++ b/Engine/Source/Runtime/Engine/Private/Materials/MaterialShared.cpp
@@ -24,6 +24,7 @@
 #include "Materials/MaterialExpressionShadingModel.h"
 #include "Materials/MaterialExpressionReroute.h"
 #include "Materials/MaterialExpressionSingleLayerWaterMaterialOutput.h"
+#include "Materials/MaterialExpressionStrata.h"
 #include "ShaderCompiler.h"
 #include "MaterialCompiler.h"
 #include "MeshMaterialShaderType.h"
@@ -3751,7 +3752,7 @@ void FMaterialAttributeDefinitionMap::InitializeAttributeMap()
 	Add(FGuid(0xD0B0FA03, 0x14D74455, 0xA851BAC5, 0x81A0788B), TEXT("Refraction"),				MP_Refraction,				MCT_Float2,	FVector4(1,0,0,0),	SF_Pixel);
 	Add(FGuid(0x0AC97EC3, 0xE3D047BA, 0xB610167D, 0xC4D919FF), TEXT("PixelDepthOffset"),		MP_PixelDepthOffset,		MCT_Float,	FVector4(0,0,0,0),	SF_Pixel);
 	Add(FGuid(0xD9423FFF, 0xD77E4D82, 0x8FF9CF5E, 0x055D1255), TEXT("ShadingModel"),			MP_ShadingModel,			MCT_ShadingModel, FVector4(0, 0, 0, 0), SF_Pixel, INDEX_NONE, false, &CompileShadingModelBlendFunction);
-	Add(FGuid(0x5973A03E, 0x13A74E08, 0x92D0CEDD, 0xF2936CF8), TEXT("FrontMaterial"),			MP_FrontMaterial,			MCT_Strata, FVector4(0,0,0,0),	SF_Pixel);
+	Add(FGuid(0x5973A03E, 0x13A74E08, 0x92D0CEDD, 0xF2936CF8), TEXT("FrontMaterial"),			MP_FrontMaterial,			MCT_Strata, FVector4(0,0,0,0),	SF_Pixel, INDEX_NONE, false, &CompileStrataBlendFunction);
 	   
 	// Texture coordinates
 	Add(FGuid(0xD30EC284, 0xE13A4160, 0x87BB5230, 0x2ED115DC), TEXT("CustomizedUV0"), MP_CustomizedUVs0, MCT_Float2, FVector4(0,0,0,0), SF_Vertex, 0);
