Fix static analysis warning

diff --git a/Engine/Source/Runtime/Engine/Private/Components/SkyLightComponent.cpp b/Engine/Source/Runtime/Engine/Private/Components/SkyLightComponent.cpp
--- a/Engine/Source/Runtime/Engine/Private/Components/SkyLightComponent.cpp
+++ b/Engine/Source/Runtime/Engine/Private/Components/SkyLightComponent.cpp
@@ -738,7 +738,7 @@ void USkyLightComponent::UpdateSkyCaptureContentsArray(UWorld* WorldToUpdate, TA
 
 		if (((!Owner || !Owner->GetLevel() || Owner->GetLevel()->bIsVisible) && CaptureComponent->GetWorld() == WorldToUpdate)
 			// Only process sky capture requests once async texture and shader compiling completes, otherwise we will capture the scene with temporary shaders/textures
-			&& ((!bIsCompilingTextures && !bIsCompilingShaders) || (CaptureComponent->SourceType == SLS_SpecifiedCubemap && !bIsCubemapCompiling)))
+			&& ((!bIsCompilingTextures && !bIsCompilingShaders) || ((CaptureComponent->SourceType == SLS_SpecifiedCubemap) && !bIsCubemapCompiling)))
 		{
 			// Only capture valid sky light components
 			if (CaptureComponent->SourceType != SLS_SpecifiedCubemap || CaptureComponent->Cubemap)
