First major checkin for flexible GBuffer.

diff --git a/Engine/Shaders/Private/BasePassPixelShader.usf b/Engine/Shaders/Private/BasePassPixelShader.usf
--- a/Engine/Shaders/Private/BasePassPixelShader.usf
+++ b/Engine/Shaders/Private/BasePassPixelShader.usf
@@ -49,6 +49,12 @@
 #include "DBufferDecalShared.ush"
 #include "ShadingModelsSampling.ush"
 
+#include "SceneTexturesCommon.ush"
+#include "SceneTextureParameters.ush"
+#include "GBufferHelpers.ush"
+
+#include "/Engine/Generated/ShaderAutogen/AutogenShaderHeaders.ush"
+
 #define PREV_FRAME_COLOR	1
 #include "SSRT/SSRTRayCast.ush"
 
@@ -1589,11 +1595,66 @@ void FPixelShaderInOut_MainPS(
 	#endif
 
 	#if USES_GBUFFER
-		GBuffer.IndirectIrradiance = IndirectIrradiance;
-
 		// -0.5 .. 0.5, could be optimzed as lower quality noise would be sufficient
 		float QuantizationBias = PseudoRandom( MaterialParameters.SvPosition.xy ) - 0.5f;
-		EncodeGBuffer(GBuffer, Out.MRT[1], Out.MRT[2], Out.MRT[3], OutGBufferD, OutGBufferE, OutVelocity, QuantizationBias);
+
+		GBuffer.IndirectIrradiance = IndirectIrradiance;
+
+		// this is the new encode, the older encode is the #else, keeping it around briefly until the new version is confirmed stable.
+		#if 1
+		{
+			// change this so that we can pack everything into the gbuffer, but leave this for now
+			#if GBUFFER_HAS_DIFFUSE_SAMPLE_OCCLUSION
+				GBuffer.GenericAO = float(GBuffer.DiffuseIndirectSampleOcclusion) * rcp(255) + (0.5 / 255.0);
+			#elif ALLOW_STATIC_LIGHTING
+				// No space for AO. Multiply IndirectIrradiance by AO instead of storing.
+				GBuffer.GenericAO = EncodeIndirectIrradiance(GBuffer.IndirectIrradiance * GBuffer.GBufferAO) + QuantizationBias * (1.0 / 255.0);
+			#else
+				GBuffer.GenericAO = GBuffer.GBufferAO;
+			#endif
+
+			EncodeGBufferToMRT(Out, GBuffer, QuantizationBias);
+
+			if (GBuffer.ShadingModelID == SHADINGMODELID_UNLIT)
+			{
+				Out.MRT[1] = 0;
+				SetGBufferForUnlit(Out.MRT[2]);
+				Out.MRT[3] = 0;
+				Out.MRT[4] = 0;
+				Out.MRT[5] = 0;
+			}
+		}
+		#else
+		{
+			float4 OutGBufferA = 0;
+			float4 OutGBufferB = 0;
+			float4 OutGBufferC = 0;
+
+			EncodeGBuffer(GBuffer, OutGBufferA, OutGBufferB, OutGBufferC, OutGBufferD, OutGBufferE, OutVelocity, QuantizationBias);
+
+			Out.MRT[1] = OutGBufferA;
+			Out.MRT[2] = OutGBufferB;
+			Out.MRT[3] = OutGBufferC;
+
+			#if GBUFFER_HAS_VELOCITY
+				Out.MRT[4] = OutVelocity;
+			#endif
+
+			Out.MRT[GBUFFER_HAS_VELOCITY ? 5 : 4] = OutGBufferD;
+
+			#if GBUFFER_HAS_PRECSHADOWFACTOR
+				Out.MRT[GBUFFER_HAS_VELOCITY ? 6 : 5] = OutGBufferE;
+			#endif
+		}
+		#endif
+
+	#else
+
+		// If not using the full gbuffer (forward shading) the velocity buffer can still be written to in the basepass.
+		#if GBUFFER_HAS_VELOCITY
+			Out.MRT[1] = OutVelocity;
+		#endif
+
 	#endif 
 
 	if(bEditorWeightedZBuffering)
@@ -1622,23 +1683,6 @@ void FPixelShaderInOut_MainPS(
 		#endif
 	}
 
-#if USES_GBUFFER
-	#if GBUFFER_HAS_VELOCITY
-		Out.MRT[4] = OutVelocity;
-	#endif
-
-	Out.MRT[GBUFFER_HAS_VELOCITY ? 5 : 4] = OutGBufferD;
-
-	#if GBUFFER_HAS_PRECSHADOWFACTOR
-		Out.MRT[GBUFFER_HAS_VELOCITY ? 6 : 5] = OutGBufferE;
-	#endif
-#else
-	// If not using the full gbuffer (forward shading) the velocity buffer can still be written to in the basepass.
-	#if GBUFFER_HAS_VELOCITY
-		Out.MRT[1] = OutVelocity;
-	#endif
-#endif
-
 #if !MATERIALBLENDING_MODULATE && USE_PREEXPOSURE
 #if MATERIAL_IS_SKY
 	// Dynamic capture exposure is 1 as of today.
@@ -1689,24 +1733,24 @@ void FPixelShaderInOut_MainPS(
 
 // the following needs to match to the code in FSceneRenderTargets::GetGBufferRenderTargets()
 #define PIXELSHADEROUTPUT_BASEPASS 1
-#if USES_GBUFFER
-#define PIXELSHADEROUTPUT_MRT0 (!SELECTIVE_BASEPASS_OUTPUTS || NEEDS_BASEPASS_VERTEX_FOGGING || USES_EMISSIVE_COLOR || ALLOW_STATIC_LIGHTING || MATERIAL_SHADINGMODEL_SINGLELAYERWATER)
-#define PIXELSHADEROUTPUT_MRT1 ((!SELECTIVE_BASEPASS_OUTPUTS || !MATERIAL_SHADINGMODEL_UNLIT))
-#define PIXELSHADEROUTPUT_MRT2 ((!SELECTIVE_BASEPASS_OUTPUTS || !MATERIAL_SHADINGMODEL_UNLIT))
-#define PIXELSHADEROUTPUT_MRT3 ((!SELECTIVE_BASEPASS_OUTPUTS || !MATERIAL_SHADINGMODEL_UNLIT))
-	#if GBUFFER_HAS_VELOCITY
-		#define PIXELSHADEROUTPUT_MRT4 WRITES_VELOCITY_TO_GBUFFER
-		#define PIXELSHADEROUTPUT_MRT5 (!SELECTIVE_BASEPASS_OUTPUTS || WRITES_CUSTOMDATA_TO_GBUFFER)
-		#define PIXELSHADEROUTPUT_MRT6 (GBUFFER_HAS_PRECSHADOWFACTOR && (!SELECTIVE_BASEPASS_OUTPUTS || WRITES_PRECSHADOWFACTOR_TO_GBUFFER && !MATERIAL_SHADINGMODEL_UNLIT))
-	#else //GBUFFER_HAS_VELOCITY
-		#define PIXELSHADEROUTPUT_MRT4 (!SELECTIVE_BASEPASS_OUTPUTS || WRITES_CUSTOMDATA_TO_GBUFFER)
-		#define PIXELSHADEROUTPUT_MRT5 (GBUFFER_HAS_PRECSHADOWFACTOR && (!SELECTIVE_BASEPASS_OUTPUTS || WRITES_PRECSHADOWFACTOR_TO_GBUFFER && !MATERIAL_SHADINGMODEL_UNLIT))
-	#endif //GBUFFER_HAS_VELOCITY
-#else //USES_GBUFFER
-	#define PIXELSHADEROUTPUT_MRT0 1
-	// we also need MRT for thin translucency due to dual blending if we are not on the fallback path
-	#define PIXELSHADEROUTPUT_MRT1 (WRITES_VELOCITY_TO_GBUFFER || (DUAL_SOURCE_COLOR_BLENDING_ENABLED && MATERIAL_WORKS_WITH_DUAL_SOURCE_COLOR_BLENDING))
-#endif //USES_GBUFFER
+//#if USES_GBUFFER
+//#define PIXELSHADEROUTPUT_MRT0 (!SELECTIVE_BASEPASS_OUTPUTS || NEEDS_BASEPASS_VERTEX_FOGGING || USES_EMISSIVE_COLOR || ALLOW_STATIC_LIGHTING || MATERIAL_SHADINGMODEL_SINGLELAYERWATER)
+//#define PIXELSHADEROUTPUT_MRT1 ((!SELECTIVE_BASEPASS_OUTPUTS || !MATERIAL_SHADINGMODEL_UNLIT))
+//#define PIXELSHADEROUTPUT_MRT2 ((!SELECTIVE_BASEPASS_OUTPUTS || !MATERIAL_SHADINGMODEL_UNLIT))
+//#define PIXELSHADEROUTPUT_MRT3 ((!SELECTIVE_BASEPASS_OUTPUTS || !MATERIAL_SHADINGMODEL_UNLIT))
+//	#if GBUFFER_HAS_VELOCITY
+//		#define PIXELSHADEROUTPUT_MRT4 WRITES_VELOCITY_TO_GBUFFER
+//		#define PIXELSHADEROUTPUT_MRT5 (!SELECTIVE_BASEPASS_OUTPUTS || WRITES_CUSTOMDATA_TO_GBUFFER)
+//		#define PIXELSHADEROUTPUT_MRT6 (GBUFFER_HAS_PRECSHADOWFACTOR && (!SELECTIVE_BASEPASS_OUTPUTS || WRITES_PRECSHADOWFACTOR_TO_GBUFFER && !MATERIAL_SHADINGMODEL_UNLIT))
+//	#else //GBUFFER_HAS_VELOCITY
+//		#define PIXELSHADEROUTPUT_MRT4 (!SELECTIVE_BASEPASS_OUTPUTS || WRITES_CUSTOMDATA_TO_GBUFFER)
+//		#define PIXELSHADEROUTPUT_MRT5 (GBUFFER_HAS_PRECSHADOWFACTOR && (!SELECTIVE_BASEPASS_OUTPUTS || WRITES_PRECSHADOWFACTOR_TO_GBUFFER && !MATERIAL_SHADINGMODEL_UNLIT))
+//	#endif //GBUFFER_HAS_VELOCITY
+//#else //USES_GBUFFER
+//	#define PIXELSHADEROUTPUT_MRT0 1
+//	// we also need MRT for thin translucency due to dual blending if we are not on the fallback path
+//	#define PIXELSHADEROUTPUT_MRT1 (WRITES_VELOCITY_TO_GBUFFER || (MATERIAL_SHADINGMODEL_THIN_TRANSLUCENT && THIN_TRANSLUCENT_USE_DUAL_BLEND))
+//#endif //USES_GBUFFER
 #define PIXELSHADEROUTPUT_A2C ((EDITOR_ALPHA2COVERAGE) != 0)
 #define PIXELSHADEROUTPUT_COVERAGE (MATERIALBLENDING_MASKED_USING_COVERAGE && !EARLY_Z_PASS_ONLY_MATERIAL_MASKING)
 
diff --git a/Engine/Shaders/Private/DeferredDecal.usf b/Engine/Shaders/Private/DeferredDecal.usf
--- a/Engine/Shaders/Private/DeferredDecal.usf
+++ b/Engine/Shaders/Private/DeferredDecal.usf
@@ -338,11 +338,11 @@ void FPixelShaderInOut_MainPS(inout FPixelShaderIn In, inout FPixelShaderOut Out
 }
 
 
-#define PIXELSHADEROUTPUT_MRT0 (DECAL_RENDERTARGET_COUNT > 0)
-#define PIXELSHADEROUTPUT_MRT1 (DECAL_RENDERTARGET_COUNT > 1)
-#define PIXELSHADEROUTPUT_MRT2 (DECAL_RENDERTARGET_COUNT > 2)
-#define PIXELSHADEROUTPUT_MRT3 (DECAL_RENDERTARGET_COUNT > 3)
-#define PIXELSHADEROUTPUT_MRT4 (DECAL_RENDERTARGET_COUNT > 4)
+//#define PIXELSHADEROUTPUT_MRT0 (DECAL_RENDERTARGET_COUNT > 0)
+//#define PIXELSHADEROUTPUT_MRT1 (DECAL_RENDERTARGET_COUNT > 1)
+//#define PIXELSHADEROUTPUT_MRT2 (DECAL_RENDERTARGET_COUNT > 2)
+//#define PIXELSHADEROUTPUT_MRT3 (DECAL_RENDERTARGET_COUNT > 3)
+//#define PIXELSHADEROUTPUT_MRT4 (DECAL_RENDERTARGET_COUNT > 4)
 // all PIXELSHADEROUTPUT_ and "void FPixelShaderInOut_MainPS()" need to be setup before this include
 // this include generates the wrapper code to call MainPS()
 #include "PixelShaderOutputCommon.ush"
diff --git a/Engine/Shaders/Private/DeferredShadingCommon.ush b/Engine/Shaders/Private/DeferredShadingCommon.ush
--- a/Engine/Shaders/Private/DeferredShadingCommon.ush
+++ b/Engine/Shaders/Private/DeferredShadingCommon.ush
@@ -11,6 +11,47 @@
 #include "SceneTexturesCommon.ush"
 #include "MonteCarlo.ush"
 
+// if this is zero, use the new autogenerated GBuffer encode/decode functions
+#ifndef GBUFFER_REFACTOR
+#define GBUFFER_REFACTOR 0
+#endif
+
+#if SHADING_PATH_DEFERRED
+
+#if !SUPPORTS_INDEPENDENT_SAMPLERS
+	#error Scene textures may only be used on platforms which support independent samplers.
+#endif
+
+uint bSceneLightingChannelsValid;
+
+// Matches FSceneTextureParameters
+Texture2D SceneDepthTexture;
+Texture2D<uint2> SceneStencilTexture;
+Texture2D GBufferATexture;
+Texture2D GBufferBTexture;
+Texture2D GBufferCTexture;
+Texture2D GBufferDTexture;
+Texture2D GBufferETexture;
+Texture2D GBufferVelocityTexture;
+Texture2D GBufferFTexture;
+Texture2D<uint> SceneLightingChannels;
+
+#define SceneDepthTextureSampler GlobalPointClampedSampler
+#define GBufferATextureSampler GlobalPointClampedSampler
+#define GBufferBTextureSampler GlobalPointClampedSampler
+#define GBufferCTextureSampler GlobalPointClampedSampler
+#define GBufferDTextureSampler GlobalPointClampedSampler
+#define GBufferETextureSampler GlobalPointClampedSampler
+#define GBufferFTextureSampler GlobalPointClampedSampler
+#define GBufferVelocityTextureSampler GlobalPointClampedSampler
+
+float SampleDeviceZFromSceneTextures(float2 UV)
+{
+	return SceneDepthTexture.SampleLevel(SceneDepthTextureSampler, UV, 0).r;
+}
+
+#endif
+
 // TODO: for CustomGBufferResolvePS() MSAA_SAMPLE_COUNT is defined by C++ code as 2 or 4
 // bot not for any other shaders!
 #ifndef MSAA_SAMPLE_COUNT
@@ -359,6 +400,8 @@ struct FGBufferData
 	float Specular;
 	// 0..1
 	float4 CustomData;
+	// AO utility value
+	float GenericAO;
 	// Indirect irradiance luma
 	float IndirectIrradiance;
 	// Static shadow factors for channels assigned by Lightmass
@@ -408,6 +451,39 @@ bool HasDynamicIndirectShadowCasterRepresentation(FGBufferData GBufferData)
 	return (PackedAlpha & 2) != 0;
 }
 
+// High frequency Checkerboard pattern
+// @param PixelPos relative to left top of the rendertarget (not viewport)
+// @return true/false, todo: profile if float 0/1 would be better (need to make sure it's 100% the same)
+bool CheckerFromPixelPos(uint2 PixelPos)
+{
+	// todo: Index is float and by staying float we can optimize this 
+	// We alternate the pattern to get 2x supersampling on the lower res data to get more near to full res
+	uint TemporalAASampleIndex = View.TemporalAAParams.x;
+
+#if FEATURE_LEVEL >= FEATURE_LEVEL_SM4
+	return (PixelPos.x + PixelPos.y + TemporalAASampleIndex) % 2;
+#else
+	return (uint)(fmod(PixelPos.x + PixelPos.y + TemporalAASampleIndex, 2)) != 0;
+#endif
+}
+
+// High frequency Checkerboard pattern
+// @param UVSceneColor at pixel center
+// @return true/false, todo: profile if float 0/1 would be better (need to make sure it's 100% the same)
+bool CheckerFromSceneColorUV(float2 UVSceneColor)
+{
+	// relative to left top of the rendertarget (not viewport)
+	uint2 PixelPos = uint2(UVSceneColor * View.BufferSizeAndInvSize.xy);
+
+	return CheckerFromPixelPos(PixelPos);
+}
+
+#if GBUFFER_REFACTOR
+#include "GBufferHelpers.ush"
+
+#include "/Engine/Generated/ShaderAutogen/AutogenShaderHeaders.ush"
+#endif
+
 struct FScreenSpaceData
 {
 	// GBuffer (material attributes from forward rendering pass)
@@ -503,33 +579,6 @@ void EncodeGBuffer(
 #endif
 }
 
-// High frequency Checkerboard pattern
-// @param PixelPos relative to left top of the rendertarget (not viewport)
-// @return true/false, todo: profile if float 0/1 would be better (need to make sure it's 100% the same)
-bool CheckerFromPixelPos(uint2 PixelPos)
-{
-	// todo: Index is float and by staying float we can optimize this 
-	// We alternate the pattern to get 2x supersampling on the lower res data to get more near to full res
-	uint TemporalAASampleIndex = View.TemporalAAParams.x;
-
-#if FEATURE_LEVEL >= FEATURE_LEVEL_SM4
-	return (PixelPos.x + PixelPos.y + TemporalAASampleIndex) % 2;
-#else
-	return (uint)(fmod(PixelPos.x + PixelPos.y + TemporalAASampleIndex, 2)) != 0;
-#endif
-}
-
-// High frequency Checkerboard pattern
-// @param UVSceneColor at pixel center
-// @return true/false, todo: profile if float 0/1 would be better (need to make sure it's 100% the same)
-bool CheckerFromSceneColorUV(float2 UVSceneColor)
-{
-	// relative to left top of the rendertarget (not viewport)
-	uint2 PixelPos = uint2(UVSceneColor * View.BufferSizeAndInvSize.xy);
-
-	return CheckerFromPixelPos(PixelPos);
-}
-
 // SubsurfaceProfile does deferred lighting with a checker board pixel pattern
 // we separate the view from the non view dependent lighting and later recombine the two color constributions in a postprocess
 // We have the option to apply the BaseColor/Specular in the base pass or do it later in the postprocess (has implications to texture detail, fresnel and performance)
@@ -679,6 +728,9 @@ uint ExtractSubsurfaceProfileInt(FGBufferData BufferData)
 	// @param PixelPos relative to left top of the rendertarget (not viewport)
 	FGBufferData GetGBufferDataUint(uint2 PixelPos, bool bGetNormalizedNormal = true)
 	{
+	#if GBUFFER_REFACTOR
+		return DecodeGBufferDataUint(PixelPos,bGetNormalizedNormal);
+	#else 
 		float4 GBufferA = SceneTexturesStruct.GBufferATexture.Load(int3(PixelPos, 0));
 		float4 GBufferB = SceneTexturesStruct.GBufferBTexture.Load(int3(PixelPos, 0));
 		float4 GBufferC = SceneTexturesStruct.GBufferCTexture.Load(int3(PixelPos, 0));
@@ -703,6 +755,7 @@ uint ExtractSubsurfaceProfileInt(FGBufferData BufferData)
 		float SceneDepth = CalcSceneDepth(PixelPos);
 
 		return DecodeGBufferData(GBufferA, GBufferB, GBufferC, GBufferD, GBufferE, GBufferF, GBufferVelocity, CustomNativeDepth, CustomStencil, SceneDepth, bGetNormalizedNormal, CheckerFromPixelPos(PixelPos));
+	#endif
 	}
 
 	// @param PixelPos relative to left top of the rendertarget (not viewport)
@@ -719,9 +772,54 @@ uint ExtractSubsurfaceProfileInt(FGBufferData BufferData)
 	}
 #endif
 
+
+#if SHADING_PATH_DEFERRED
+
+// @param UV - UV space in the GBuffer textures (BufferSize resolution)
+// TOOD: need to find a way to make this more convenient.
+FGBufferData GetGBufferDataFromSceneTextures(float2 UV, bool bGetNormalizedNormal = true)
+{
+#if GBUFFER_REFACTOR
+	return DecodeGBufferDataSceneTextures(UV,bGetNormalizedNormal);
+#else
+	float4 GBufferA = GBufferATexture.SampleLevel(GBufferATextureSampler, UV, 0);
+	float4 GBufferB = GBufferBTexture.SampleLevel(GBufferBTextureSampler, UV, 0);
+	float4 GBufferC = GBufferCTexture.SampleLevel(GBufferCTextureSampler, UV, 0);
+	float4 GBufferD = GBufferDTexture.SampleLevel(GBufferDTextureSampler, UV, 0);
+	float4 GBufferE = GBufferETexture.SampleLevel(GBufferETextureSampler, UV, 0);
+	float4 GBufferF = GBufferFTexture.SampleLevel(GBufferFTextureSampler, UV, 0);
+	float4 GBufferVelocity = GBufferVelocityTexture.SampleLevel(GBufferVelocityTextureSampler, UV, 0);
+
+	uint CustomStencil = 0;
+	float CustomNativeDepth = 0;
+
+	float DeviceZ = SampleDeviceZFromSceneTextures(UV);
+
+	float SceneDepth = ConvertFromDeviceZ(DeviceZ);
+
+	return DecodeGBufferData(GBufferA, GBufferB, GBufferC, GBufferD, GBufferE, GBufferF, GBufferVelocity, CustomNativeDepth, CustomStencil, SceneDepth, bGetNormalizedNormal, CheckerFromSceneColorUV(UV));
+#endif
+}
+
+/** Returns the light channel mask that should be executed for this pixel. */
+uint GetSceneLightingChannel(uint2 PixelCoord)
+{
+	BRANCH
+	if (bSceneLightingChannelsValid)
+	{
+		return SceneLightingChannels.Load(uint3(PixelCoord, 0)).x;
+	}
+	return ~0;
+}
+
+#endif
+
 // @param UV - UV space in the GBuffer textures (BufferSize resolution)
 FGBufferData GetGBufferData(float2 UV, bool bGetNormalizedNormal = true)
 {
+#if GBUFFER_REFACTOR
+	return DecodeGBufferDataUV(UV,bGetNormalizedNormal);
+#else
 	float4 GBufferA = Texture2DSampleLevel(SceneTexturesStruct.GBufferATexture, SceneTexturesStruct_GBufferATextureSampler, UV, 0);
 	float4 GBufferB = Texture2DSampleLevel(SceneTexturesStruct.GBufferBTexture, SceneTexturesStruct_GBufferBTextureSampler, UV, 0);
 	float4 GBufferC = Texture2DSampleLevel(SceneTexturesStruct.GBufferCTexture, SceneTexturesStruct_GBufferCTextureSampler, UV, 0);
@@ -748,6 +846,7 @@ FGBufferData GetGBufferData(float2 UV, bool bGetNormalizedNormal = true)
 	float SceneDepth = CalcSceneDepth(UV);
 	
 	return DecodeGBufferData(GBufferA, GBufferB, GBufferC, GBufferD, GBufferE, GBufferF, GBufferVelocity, CustomNativeDepth, CustomStencil, SceneDepth, bGetNormalizedNormal, CheckerFromSceneColorUV(UV));
+#endif
 }
 
 // Minimal path for just the lighting model, used to branch around unlit pixels (skybox)
diff --git a/Engine/Shaders/Private/GBufferHelpers.ush b/Engine/Shaders/Private/GBufferHelpers.ush
new file mode 100644
--- /dev/null
+++ b/Engine/Shaders/Private/GBufferHelpers.ush
@@ -0,0 +1,444 @@
+// Copyright Epic Games, Inc. All Rights Reserved.
+
+#pragma once
+
+
+#include "SceneTexturesCommon.ush"
+#include "SceneTextureParameters.ush"
+
+float SquareInline(float X)
+{
+	return X * X;
+}
+
+float3 EncodeNormalHelper(float3 SrcNormal, float QuantizationBias)
+{
+	return SrcNormal * .5f + .5f;
+}
+
+float3 DecodeNormalHelper(float3 SrcNormal)
+{
+	return SrcNormal * 2.0f - 1.0f;
+}
+
+
+uint EncodeQuantize6(float Value, float QuantizationBias)
+{
+	return min(uint(saturate(Value) * 63.0f + .5f + QuantizationBias),63);
+}
+
+float DecodeQuantize6(uint Value)
+{
+	return float(Value) / 63.0f;
+}
+
+uint EncodeQuantize6Sqrt(float Value, float QuantizationBias)
+{
+	return min(uint(sqrt(saturate(Value)) * 63.0f + .5f + QuantizationBias),63);
+}
+
+float DecodeQuantize6Sqrt(uint Value)
+{
+	return SquareInline(float(Value) / 63.0f);
+}
+
+uint EncodeQuantize5(float Value, float QuantizationBias)
+{
+	return min(uint(saturate(Value) * 31.0f + .5f + QuantizationBias),31);
+}
+
+float DecodeQuantize5(uint Value)
+{
+	return float(Value) / 31.0f;
+}
+
+uint EncodeQuantize5Sqrt(float Value, float QuantizationBias)
+{
+	return min(uint(sqrt(saturate(Value)) * 31.0f + .5f + QuantizationBias),31);
+}
+
+float DecodeQuantize5Sqrt(uint Value)
+{
+	return SquareInline(float(Value) / 31.0f);
+}
+
+uint EncodeQuantize4(float Value, float QuantizationBias)
+{
+	return min(uint(saturate(Value) * 15.0f + .5f + QuantizationBias),15);
+}
+
+float DecodeQuantize4(uint Value)
+{
+	return float(Value) / 15.0f;
+}
+
+uint EncodeQuantize4Sqrt(float Value, float QuantizationBias)
+{
+	return min(uint(sqrt(saturate(Value)) * 15.0f + .5f + QuantizationBias),15);
+}
+
+float DecodeQuantize4Sqrt(uint Value)
+{
+	return SquareInline(float(Value) / 15.0f);
+}
+
+
+uint EncodeQuantize3(float Value, float QuantizationBias)
+{
+	return min(uint(saturate(Value) * 7.0f + .5f + QuantizationBias),7);
+}
+
+float DecodeQuantize3(uint Value)
+{
+	return float(Value) / 7.0f;
+}
+
+uint EncodeQuantize3Sqrt(float Value, float QuantizationBias)
+{
+	return min(uint(sqrt(saturate(Value)) * 7.0f + .5f + QuantizationBias),7);
+}
+
+float DecodeQuantize3Sqrt(uint Value)
+{
+	return SquareInline(float(Value) / 7.0f);
+}
+
+uint EncodeQuantize2(float Value, float QuantizationBias)
+{
+	return min(uint(saturate(Value) * 3.0f + .5f + QuantizationBias),3);
+}
+
+float DecodeQuantize2(uint Value)
+{
+	return float(Value) / 3.0f;
+}
+
+uint EncodeQuantize2Sqrt(float Value, float QuantizationBias)
+{
+	return min(uint(sqrt(saturate(Value)) * 3.0f + .5f + QuantizationBias),3);
+}
+
+float DecodeQuantize2Sqrt(uint Value)
+{
+	return SquareInline(float(Value) / 3.0f);
+}
+
+uint EncodeQuantize1(float Value, float QuantizationBias)
+{
+	return min(uint(saturate(Value) * 1.0f + .5f + QuantizationBias),1);
+}
+
+float DecodeQuantize1(uint Value)
+{
+	return float(Value) / 1.0f;
+}
+
+uint EncodeQuantize1Sqrt(float Value, float QuantizationBias)
+{
+	return min(uint(sqrt(saturate(Value)) * 1.0f + .5f + QuantizationBias),1);
+}
+
+float DecodeQuantize1Sqrt(uint Value)
+{
+	return SquareInline(float(Value) / 1.0f);
+}
+
+
+uint3 EncodeQuantize565(float3 Value, float QuantizationBias)
+{
+	uint3 Ret;
+	Ret.x = EncodeQuantize5(Value.x,QuantizationBias);
+	Ret.y = EncodeQuantize6(Value.y,QuantizationBias);
+	Ret.z = EncodeQuantize5(Value.z,QuantizationBias);
+	return Ret;
+}
+
+float3 DecodeQuantize565(uint3 Value)
+{
+	float3 Ret;
+	Ret.x = DecodeQuantize5(Value.x);
+	Ret.y = DecodeQuantize6(Value.y);
+	Ret.z = DecodeQuantize5(Value.z);
+	return Ret;
+}
+
+uint3 EncodeQuantize565Sqrt(float3 Value, float QuantizationBias)
+{
+	uint3 Ret;
+	Ret.x = EncodeQuantize5Sqrt(Value.x,QuantizationBias);
+	Ret.y = EncodeQuantize6Sqrt(Value.y,QuantizationBias);
+	Ret.z = EncodeQuantize5Sqrt(Value.z,QuantizationBias);
+	return Ret;
+}
+
+float3 DecodeQuantize565Sqrt(uint3 Value)
+{
+	float3 Ret;
+	Ret.x = DecodeQuantize5Sqrt(Value.x);
+	Ret.y = DecodeQuantize6Sqrt(Value.y);
+	Ret.z = DecodeQuantize5Sqrt(Value.z);
+	return Ret;
+}
+
+
+uint3 EncodeQuantize444(float3 Value, float QuantizationBias)
+{
+	uint3 Ret;
+	Ret.x = EncodeQuantize4(Value.x,QuantizationBias);
+	Ret.y = EncodeQuantize4(Value.y,QuantizationBias);
+	Ret.z = EncodeQuantize4(Value.z,QuantizationBias);
+	return Ret;
+}
+
+float3 DecodeQuantize444(uint3 Value)
+{
+	float3 Ret;
+	Ret.x = DecodeQuantize4(Value.x);
+	Ret.y = DecodeQuantize4(Value.y);
+	Ret.z = DecodeQuantize4(Value.z);
+	return Ret;
+}
+
+uint3 EncodeQuantize444Sqrt(float3 Value, float QuantizationBias)
+{
+	uint3 Ret;
+	Ret.x = EncodeQuantize4Sqrt(Value.x,QuantizationBias);
+	Ret.y = EncodeQuantize4Sqrt(Value.y,QuantizationBias);
+	Ret.z = EncodeQuantize4Sqrt(Value.z,QuantizationBias);
+	return Ret;
+}
+
+float3 DecodeQuantize444Sqrt(uint3 Value)
+{
+	float3 Ret;
+	Ret.x = DecodeQuantize4Sqrt(Value.x);
+	Ret.y = DecodeQuantize4Sqrt(Value.y);
+	Ret.z = DecodeQuantize4Sqrt(Value.z);
+	return Ret;
+}
+
+
+uint3 EncodeQuantize332(float3 Value, float QuantizationBias)
+{
+	uint3 Ret;
+	Ret.x = EncodeQuantize3(Value.x,QuantizationBias);
+	Ret.y = EncodeQuantize3(Value.y,QuantizationBias);
+	Ret.z = EncodeQuantize2(Value.z,QuantizationBias);
+	return Ret;
+}
+
+float3 DecodeQuantize332(uint3 Value)
+{
+	float3 Ret;
+	Ret.x = DecodeQuantize3(Value.x);
+	Ret.y = DecodeQuantize3(Value.y);
+	Ret.z = DecodeQuantize2(Value.z);
+	return Ret;
+}
+
+uint3 EncodeQuantize332Sqrt(float3 Value, float QuantizationBias)
+{
+	uint3 Ret;
+	Ret.x = EncodeQuantize3Sqrt(Value.x,QuantizationBias);
+	Ret.y = EncodeQuantize3Sqrt(Value.y,QuantizationBias);
+	Ret.z = EncodeQuantize2Sqrt(Value.z,QuantizationBias);
+	return Ret;
+}
+
+float3 DecodeQuantize332Sqrt(uint3 Value)
+{
+	float3 Ret;
+	Ret.x = DecodeQuantize3Sqrt(Value.x);
+	Ret.y = DecodeQuantize3Sqrt(Value.y);
+	Ret.z = DecodeQuantize2Sqrt(Value.z);
+	return Ret;
+}
+
+
+void EnvBRDFApproxFullyRoughHelper(inout half3 DiffuseColor, inout half3 SpecularColor)
+{
+	// Factors derived from EnvBRDFApprox( SpecularColor, 1, 1 ) == SpecularColor * 0.4524 - 0.0024
+	DiffuseColor += SpecularColor * 0.45;
+	SpecularColor = 0;
+	// We do not modify Roughness here as this is done differently at different places.
+}
+void EnvBRDFApproxFullyRoughHelper(inout half3 DiffuseColor, inout half SpecularColor)
+{
+	DiffuseColor += SpecularColor * 0.45;
+	SpecularColor = 0;
+}
+
+
+#define USES_GBUFFER_HELPER	(FEATURE_LEVEL >= FEATURE_LEVEL_SM4 && (MATERIALBLENDING_SOLID || MATERIALBLENDING_MASKED) && !SIMPLE_FORWARD_SHADING && !FORWARD_SHADING)
+
+#define FORCE_FULLY_ROUGH_HELPER (SIMPLE_FORWARD_SHADING || MATERIAL_FULLY_ROUGH)
+
+
+// TODO: Move the BasePassPixelShader extra changes into here.
+void GBufferPreEncode(inout FGBufferData GBuffer, bool bChecker, float GeometricAARoughness, inout half3 OriginalBaseColor, inout half OriginalSpecular, inout half OriginalMetallic, float QuantizationBias)
+{
+#if MATERIAL_NORMAL_CURVATURE_TO_ROUGHNESS
+	GBuffer.Roughness = max(GBuffer.Roughness, GeometricAARoughness);
+
+#if MATERIAL_SHADINGMODEL_CLEAR_COAT
+	if (GBuffer.ShadingModelID == SHADINGMODELID_CLEAR_COAT)
+	{
+		GBuffer.CustomData.y = max(GBuffer.CustomData.y, GeometricAARoughness);
+	}	
+#endif
+#endif
+
+#if USES_GBUFFER_HELPER && (MATERIAL_SHADINGMODEL_SUBSURFACE_PROFILE || MATERIAL_SHADINGMODEL_EYE)
+	// SubsurfaceProfile applies the BaseColor in a later pass. Any lighting output in the base pass needs
+	// to separate specular and diffuse lighting in a checkerboard pattern
+	if (UseSubsurfaceProfile(GBuffer.ShadingModelID))
+	{
+		const bool bCheckerboardRequired = View.bSubsurfacePostprocessEnabled > 0 && View.bCheckerboardSubsurfaceProfileRendering > 0;
+		OriginalBaseColor = View.bSubsurfacePostprocessEnabled ? float3(1, 1, 1) : OriginalBaseColor;
+		if (bCheckerboardRequired)
+		{
+			// because we adjust the BaseColor here, we need StoredBaseColor
+
+			// we apply the base color later in SubsurfaceRecombinePS()
+			OriginalBaseColor = bChecker;
+			// in SubsurfaceRecombinePS() does not multiply with Specular so we do it here
+			GBuffer.SpecularColor *= !bChecker;
+			OriginalSpecular *= !bChecker;
+		}
+
+	}
+#endif
+
+	GBuffer.DiffuseColor = OriginalBaseColor - OriginalBaseColor * OriginalMetallic;
+
+#if USE_DEVELOPMENT_SHADERS
+	{
+		// this feature is only needed for development/editor - we can compile it out for a shipping build (see r.CompileShadersForDevelopment cvar help)
+		GBuffer.DiffuseColor = GBuffer.DiffuseColor * View.DiffuseOverrideParameter.w + View.DiffuseOverrideParameter.xyz;
+		GBuffer.SpecularColor = GBuffer.SpecularColor * View.SpecularOverrideParameter.w + View.SpecularOverrideParameter.xyz;
+	}
+#endif
+
+    #if !FORCE_FULLY_ROUGH_HELPER
+	if (View.RenderingReflectionCaptureMask) // Force material rendered in reflection capture to have an expanded albedo to try to be energy conservative (when specular is removed).
+	#endif
+	{
+		EnvBRDFApproxFullyRoughHelper(GBuffer.DiffuseColor, GBuffer.SpecularColor);
+		// When rendering reflection captures, GBuffer.Roughness is already forced to 1 using RoughnessOverrideParameter in GetMaterialRoughness.
+	}
+
+
+// should move this
+#if GBUFFER_HAS_DIFFUSE_SAMPLE_OCCLUSION
+		GBuffer.GenericAO = float(GBuffer.DiffuseIndirectSampleOcclusion) * rcp(255) + (0.5 / 255.0);
+#elif ALLOW_STATIC_LIGHTING
+		// No space for AO. Multiply IndirectIrradiance by AO instead of storing.
+		GBuffer.GenericAO = EncodeIndirectIrradiance(GBuffer.IndirectIrradiance * GBuffer.GBufferAO) + QuantizationBias * (1.0 / 255.0);
+#else
+		GBuffer.GenericAO = GBuffer.GBufferAO;
+#endif
+
+
+
+}
+
+// Temporary extra copy of this function. The autogenerated GBuffer functions will require changing the include heirarchy, so copying this function to miminize
+// the number of files that have to be touched on the initial checkin. Will remove duplicates once the first checkin is tested and stable.
+void AdjustBaseColorAndSpecularColorForSubsurfaceProfileLightingCopyHack(inout float3 BaseColor, inout float3 SpecularColor, inout float Specular, bool bChecker)
+{
+#if SUBSURFACE_CHANNEL_MODE == 0
+	// If SUBSURFACE_CHANNEL_MODE is 0, we can't support full-resolution lighting, so we 
+	// ignore View.bCheckerboardSubsurfaceProfileRendering
+	const bool bCheckerboardRequired = View.bSubsurfacePostprocessEnabled > 0;
+#else
+	const bool bCheckerboardRequired = View.bSubsurfacePostprocessEnabled > 0 && View.bCheckerboardSubsurfaceProfileRendering > 0;
+	BaseColor = View.bSubsurfacePostprocessEnabled ? float3(1, 1, 1) : BaseColor;
+#endif
+	if (bCheckerboardRequired)
+	{
+		// because we adjust the BaseColor here, we need StoredBaseColor
+
+		// we apply the base color later in SubsurfaceRecombinePS()
+		BaseColor = bChecker;
+		// in SubsurfaceRecombinePS() does not multiply with Specular so we do it here
+		SpecularColor *= !bChecker;
+		Specular *= !bChecker;
+	}
+}
+
+
+void GBufferPostDecode(inout FGBufferData Ret, bool bChecker, bool bGetNormalizedNormal)
+{
+	Ret.CustomData = !(Ret.SelectiveOutputMask & 0x1) ? Ret.CustomData : 0.0f;
+	Ret.PrecomputedShadowFactors = !(Ret.SelectiveOutputMask & 0x2) ? Ret.PrecomputedShadowFactors : ((Ret.SelectiveOutputMask & 0x4) ? 0.0f : 1.0f);
+	Ret.Velocity = !(Ret.SelectiveOutputMask & 0x8) ? Ret.Velocity : 0.0f;
+
+	Ret.StoredBaseColor = Ret.BaseColor;
+	Ret.StoredMetallic = Ret.Metallic;
+	Ret.StoredSpecular = Ret.Specular;
+
+#if GBUFFER_HAS_DIFFUSE_SAMPLE_OCCLUSION
+	Ret.DiffuseIndirectSampleOcclusion = 255 * Ret.GenericAO.x;
+	Ret.GBufferAO = saturate(1.0 - float(countbits(Ret.DiffuseIndirectSampleOcclusion)) * rcp(float(INDIRECT_SAMPLE_COUNT)));
+	Ret.IndirectIrradiance = 1;
+#elif ALLOW_STATIC_LIGHTING
+	Ret.GBufferAO = 1;
+	Ret.DiffuseIndirectSampleOcclusion = 0x0;
+	Ret.IndirectIrradiance = DecodeIndirectIrradiance(Ret.GenericAO.x);
+#else
+	Ret.GBufferAO = InGBufferC.a;
+	Ret.DiffuseIndirectSampleOcclusion = 0x0;
+	Ret.IndirectIrradiance = 1;
+#endif
+
+	if(bGetNormalizedNormal)
+	{
+		Ret.WorldNormal = normalize(Ret.WorldNormal);
+	}
+
+	FLATTEN
+	if( Ret.ShadingModelID == SHADINGMODELID_EYE )
+	{
+		Ret.Metallic = 0.0;
+	#if IRIS_NORMAL
+		Ret.Specular = 0.25;
+	#endif
+	}
+	
+	// derived from BaseColor, Metalness, Specular
+	{
+		Ret.SpecularColor = ComputeF0(Ret.Specular, Ret.BaseColor, Ret.Metallic);
+	
+		if (UseSubsurfaceProfile(Ret.ShadingModelID))
+		{
+			AdjustBaseColorAndSpecularColorForSubsurfaceProfileLightingCopyHack(Ret.BaseColor, Ret.SpecularColor, Ret.Specular, bChecker);
+		}
+	
+		Ret.DiffuseColor = Ret.BaseColor - Ret.BaseColor * Ret.Metallic;
+	
+		#if USE_DEVELOPMENT_SHADERS
+		{
+			// this feature is only needed for development/editor - we can compile it out for a shipping build (see r.CompileShadersForDevelopment cvar help)
+			Ret.DiffuseColor = Ret.DiffuseColor * View.DiffuseOverrideParameter.www + View.DiffuseOverrideParameter.xyz;
+			Ret.SpecularColor = Ret.SpecularColor * View.SpecularOverrideParameter.w + View.SpecularOverrideParameter.xyz;
+		}
+		#endif //USE_DEVELOPMENT_SHADERS
+	}
+	
+#if GBUFFER_HAS_TANGENT
+	{
+		Ret.WorldTangent = DecodeNormal(Ret.WorldTangent);
+		Ret.Anisotropy = Ret.Anisotropy * 2.0f - 1.0f;
+	
+		if(bGetNormalizedNormal)
+		{
+			Ret.WorldTangent = normalize(Ret.WorldTangent);
+		}
+	}
+#else
+	{
+		Ret.WorldTangent = 0;
+		Ret.Anisotropy = 0;
+	}
+#endif
+}
diff --git a/Engine/Shaders/Private/HairStrands/HairStrandsDeepShadowMask.usf b/Engine/Shaders/Private/HairStrands/HairStrandsDeepShadowMask.usf
--- a/Engine/Shaders/Private/HairStrands/HairStrandsDeepShadowMask.usf
+++ b/Engine/Shaders/Private/HairStrands/HairStrandsDeepShadowMask.usf
@@ -47,7 +47,7 @@ Texture2D<float4>	RayMarchMaskTexture;
 Texture2D<float>	DeepShadow_FrontDepthTexture;
 Texture2D<float4>	DeepShadow_DomTexture;
 
-Texture2D<float>	SceneDepthTexture;
+//Texture2D<float>	SceneDepthTexture;
 Texture2D<uint4>	CategorizationTexture;
 RWTexture2D<float4>	OutputColor;
 
diff --git a/Engine/Shaders/Private/HairStrands/HairStrandsGBufferWrite.usf b/Engine/Shaders/Private/HairStrands/HairStrandsGBufferWrite.usf
--- a/Engine/Shaders/Private/HairStrands/HairStrandsGBufferWrite.usf
+++ b/Engine/Shaders/Private/HairStrands/HairStrandsGBufferWrite.usf
@@ -2,12 +2,14 @@
 
 #include "../Common.ush"
 
+// Note, this doesn't work after refactoring, undoing the workaround.
 // Workaround for DeferredShadingCommon.ush requiring SceneTexturesStruct resources which is not available, 
 // since we are writing into Gbuffer textures within this pass
-#if SHADING_PATH_DEFERRED
-	#undef SHADING_PATH_DEFERRED
-	#define SHADING_PATH_DEFERRED 0
-#endif
+//#if SHADING_PATH_DEFERRED
+//	#undef SHADING_PATH_DEFERRED
+//	#define SHADING_PATH_DEFERRED 0
+//#endif
+
 #include "../DeferredShadingCommon.ush"
 #include "HairStrandsGBufferCommon.ush"
 #include "HairStrandsVisibilityCommon.ush"
diff --git a/Engine/Shaders/Private/HairStrands/HairStrandsVisibilityComposeSubPixelPS.usf b/Engine/Shaders/Private/HairStrands/HairStrandsVisibilityComposeSubPixelPS.usf
--- a/Engine/Shaders/Private/HairStrands/HairStrandsVisibilityComposeSubPixelPS.usf
+++ b/Engine/Shaders/Private/HairStrands/HairStrandsVisibilityComposeSubPixelPS.usf
@@ -189,7 +189,8 @@ void ComposeSamplePS(
 #if SHADER_DOFDEPTH
 #include "../CircleDOFCommon.ush"
 
-Texture2D<float>	SceneDepthTexture;
+// Included from DeferredLightingCommon now
+//Texture2D<float>	SceneDepthTexture;
 Texture2D<uint4>	HairCategorizationTexture;
 
 void DOFDepthPS(
diff --git a/Engine/Shaders/Private/MeshDecals.usf b/Engine/Shaders/Private/MeshDecals.usf
--- a/Engine/Shaders/Private/MeshDecals.usf
+++ b/Engine/Shaders/Private/MeshDecals.usf
@@ -240,11 +240,11 @@ void FPixelShaderInOut_MainPS(
 
 #define PIXELSHADEROUTPUT_MESHDECALPASS NEED_MESH_DECAL_INTERPOLATORS
 #define PIXELSHADEROUTPUT_INTERPOLANTS 1
-#define PIXELSHADEROUTPUT_MRT0 (DECAL_RENDERTARGET_COUNT > 0)
-#define PIXELSHADEROUTPUT_MRT1 (DECAL_RENDERTARGET_COUNT > 1)
-#define PIXELSHADEROUTPUT_MRT2 (DECAL_RENDERTARGET_COUNT > 2)
-#define PIXELSHADEROUTPUT_MRT3 (DECAL_RENDERTARGET_COUNT > 3)
-#define PIXELSHADEROUTPUT_MRT4 (DECAL_RENDERTARGET_COUNT > 4)
+//#define PIXELSHADEROUTPUT_MRT0 (DECAL_RENDERTARGET_COUNT > 0)
+//#define PIXELSHADEROUTPUT_MRT1 (DECAL_RENDERTARGET_COUNT > 1)
+//#define PIXELSHADEROUTPUT_MRT2 (DECAL_RENDERTARGET_COUNT > 2)
+//#define PIXELSHADEROUTPUT_MRT3 (DECAL_RENDERTARGET_COUNT > 3)
+//#define PIXELSHADEROUTPUT_MRT4 (DECAL_RENDERTARGET_COUNT > 4)
 // all PIXELSHADEROUTPUT_ and "void FPixelShaderInOut_MainPS()" need to be setup before this include
 // this include generates the wrapper code to call MainPS()
 #include "PixelShaderOutputCommon.ush"
diff --git a/Engine/Shaders/Private/ReflectionEnvironmentShaders.usf b/Engine/Shaders/Private/ReflectionEnvironmentShaders.usf
--- a/Engine/Shaders/Private/ReflectionEnvironmentShaders.usf
+++ b/Engine/Shaders/Private/ReflectionEnvironmentShaders.usf
@@ -34,7 +34,9 @@ int CubeFace;
 Texture2D SceneColorTexture;
 SamplerState SceneColorSampler;
 
+#if !SHADING_PATH_DEFERRED
 Texture2D SceneDepthTexture;
+#endif
 SamplerState SceneDepthSampler;
 
 /** 
diff --git a/Engine/Shaders/Private/SceneTextureParameters.ush b/Engine/Shaders/Private/SceneTextureParameters.ush
--- a/Engine/Shaders/Private/SceneTextureParameters.ush
+++ b/Engine/Shaders/Private/SceneTextureParameters.ush
@@ -5,71 +5,4 @@
 #include "SceneTexturesCommon.ush"
 #include "DeferredShadingCommon.ush"
 
-#if SHADING_PATH_DEFERRED
-
-#if !SUPPORTS_INDEPENDENT_SAMPLERS
-	#error Scene textures may only be used on platforms which support independent samplers.
-#endif
-
-uint bSceneLightingChannelsValid;
-
-// Matches FSceneTextureParameters
-Texture2D SceneDepthTexture;
-Texture2D<uint2> SceneStencilTexture;
-Texture2D GBufferATexture;
-Texture2D GBufferBTexture;
-Texture2D GBufferCTexture;
-Texture2D GBufferDTexture;
-Texture2D GBufferETexture;
-Texture2D GBufferVelocityTexture;
-Texture2D GBufferFTexture;
-Texture2D<uint> SceneLightingChannels;
-
-#define SceneDepthTextureSampler GlobalPointClampedSampler
-#define GBufferATextureSampler GlobalPointClampedSampler
-#define GBufferBTextureSampler GlobalPointClampedSampler
-#define GBufferCTextureSampler GlobalPointClampedSampler
-#define GBufferDTextureSampler GlobalPointClampedSampler
-#define GBufferETextureSampler GlobalPointClampedSampler
-#define GBufferFTextureSampler GlobalPointClampedSampler
-#define GBufferVelocityTextureSampler GlobalPointClampedSampler
-
-float SampleDeviceZFromSceneTextures(float2 UV)
-{
-	return SceneDepthTexture.SampleLevel(SceneDepthTextureSampler, UV, 0).r;
-}
-
-// @param UV - UV space in the GBuffer textures (BufferSize resolution)
-// TOOD: need to find a way to make this more convenient.
-FGBufferData GetGBufferDataFromSceneTextures(float2 UV, bool bGetNormalizedNormal = true)
-{
-	float4 GBufferA = GBufferATexture.SampleLevel(GBufferATextureSampler, UV, 0);
-	float4 GBufferB = GBufferBTexture.SampleLevel(GBufferBTextureSampler, UV, 0);
-	float4 GBufferC = GBufferCTexture.SampleLevel(GBufferCTextureSampler, UV, 0);
-	float4 GBufferD = GBufferDTexture.SampleLevel(GBufferDTextureSampler, UV, 0);
-	float4 GBufferE = GBufferETexture.SampleLevel(GBufferETextureSampler, UV, 0);
-	float4 GBufferF = GBufferFTexture.SampleLevel(GBufferFTextureSampler, UV, 0);
-	float4 GBufferVelocity = GBufferVelocityTexture.SampleLevel(GBufferVelocityTextureSampler, UV, 0);
-
-	uint CustomStencil = 0;
-	float CustomNativeDepth = 0;
-
-	float DeviceZ = SampleDeviceZFromSceneTextures(UV);
-
-	float SceneDepth = ConvertFromDeviceZ(DeviceZ);
-
-	return DecodeGBufferData(GBufferA, GBufferB, GBufferC, GBufferD, GBufferE, GBufferF, GBufferVelocity, CustomNativeDepth, CustomStencil, SceneDepth, bGetNormalizedNormal, CheckerFromSceneColorUV(UV));
-}
-
-/** Returns the light channel mask that should be executed for this pixel. */
-uint GetSceneLightingChannel(uint2 PixelCoord)
-{
-	BRANCH
-	if (bSceneLightingChannelsValid)
-	{
-		return SceneLightingChannels.Load(uint3(PixelCoord, 0)).x;
-	}
-	return ~0;
-}
-
-#endif
\ No newline at end of file
+// this code has been moved to DeferredShadingCommon.ush
\ No newline at end of file
diff --git a/Engine/Shaders/Private/VirtualTextureMaterial.usf b/Engine/Shaders/Private/VirtualTextureMaterial.usf
--- a/Engine/Shaders/Private/VirtualTextureMaterial.usf
+++ b/Engine/Shaders/Private/VirtualTextureMaterial.usf
@@ -108,15 +108,15 @@ void FPixelShaderInOut_MainPS(
 
 #define PIXELSHADEROUTPUT_INTERPOLANTS 1
 
-#if defined(OUT_BASECOLOR)
-#define PIXELSHADEROUTPUT_MRT0 1
-#elif defined(OUT_BASECOLOR_NORMAL_SPECULAR)
-#define PIXELSHADEROUTPUT_MRT0 1
-#define PIXELSHADEROUTPUT_MRT1 1
-#define PIXELSHADEROUTPUT_MRT2 1
-#elif defined(OUT_WORLDHEIGHT)
-#define PIXELSHADEROUTPUT_MRT0 1
-#endif
+//#if defined(OUT_BASECOLOR)
+//#define PIXELSHADEROUTPUT_MRT0 1
+//#elif defined(OUT_BASECOLOR_NORMAL_SPECULAR)
+//#define PIXELSHADEROUTPUT_MRT0 1
+//#define PIXELSHADEROUTPUT_MRT1 1
+//#define PIXELSHADEROUTPUT_MRT2 1
+//#elif defined(OUT_WORLDHEIGHT)
+//#define PIXELSHADEROUTPUT_MRT0 1
+//#endif
 
 // all PIXELSHADEROUTPUT_ and "void FPixelShaderInOut_MainPS()" need to be setup before this include
 // this include generates the wrapper code to call MainPS(inout FPixelShaderOutput PixelShaderOutput)
diff --git a/Engine/Source/Developer/ShaderPreprocessor/Private/ShaderPreprocessor.cpp b/Engine/Source/Developer/ShaderPreprocessor/Private/ShaderPreprocessor.cpp
--- a/Engine/Source/Developer/ShaderPreprocessor/Private/ShaderPreprocessor.cpp
+++ b/Engine/Source/Developer/ShaderPreprocessor/Private/ShaderPreprocessor.cpp
@@ -72,6 +72,9 @@ private:
 		// Substitute virtual platform path here to make sure that #line directives refer to the platform-specific file.
 		ReplaceVirtualFilePathForShaderPlatform(VirtualFilePath, This->ShaderInput.Target.GetPlatform());
 
+		// Fixup autogen file
+		ReplaceVirtualFilePathForShaderAutogen(VirtualFilePath, This->ShaderInput.Target.GetPlatform());
+
 		// Collapse any relative directories to allow #include "../MyFile.ush"
 		FPaths::CollapseRelativeDirectories(VirtualFilePath);
 
diff --git a/Engine/Source/Runtime/Engine/Private/Materials/MaterialShader.cpp b/Engine/Source/Runtime/Engine/Private/Materials/MaterialShader.cpp
--- a/Engine/Source/Runtime/Engine/Private/Materials/MaterialShader.cpp
+++ b/Engine/Source/Runtime/Engine/Private/Materials/MaterialShader.cpp
@@ -134,6 +134,8 @@ static FString GetMaterialShaderMapKeyString(const FMaterialShaderMapId& ShaderM
 	ShaderMapAppendKeyString(Platform, ShaderMapKeyString);
 	ShaderMapId.AppendKeyString(ShaderMapKeyString);
 	FMaterialAttributeDefinitionMap::AppendDDCKeyString(ShaderMapKeyString);
+	FShaderCompileUtilities::AppendGBufferDDCKeyString(Platform, ShaderMapKeyString);
+
 	return FDerivedDataCacheInterface::BuildCacheKey(TEXT("MATSM"), MATERIALSHADERMAP_DERIVEDDATA_VER, *ShaderMapKeyString);
 }
 #endif // WITH_EDITOR
diff --git a/Engine/Source/Runtime/Engine/Private/Materials/MaterialShared.cpp b/Engine/Source/Runtime/Engine/Private/Materials/MaterialShared.cpp
--- a/Engine/Source/Runtime/Engine/Private/Materials/MaterialShared.cpp
+++ b/Engine/Source/Runtime/Engine/Private/Materials/MaterialShared.cpp
@@ -2088,6 +2088,10 @@ bool FMaterial::BeginCompileShaderMap(
 
 		MaterialEnvironment->IncludeVirtualPathToContentsMap.Add(TEXT("/Engine/Generated/Material.ush"), MaterialShaderCode);
 
+		FShaderCompileUtilities::GenerateBrdfHeaders((EShaderPlatform)Platform);
+
+		FShaderCompileUtilities::ApplyDerivedDefines(*MaterialEnvironment, nullptr, (EShaderPlatform)Platform);
+
 		// Compile the shaders for the material.
 		NewShaderMap->Compile(this, ShaderMapId, MaterialEnvironment, NewCompilationOutput, Platform, bSynchronousCompile);
 
diff --git a/Engine/Source/Runtime/Engine/Private/ShaderCompiler/ShaderCompiler.cpp b/Engine/Source/Runtime/Engine/Private/ShaderCompiler/ShaderCompiler.cpp
--- a/Engine/Source/Runtime/Engine/Private/ShaderCompiler/ShaderCompiler.cpp
+++ b/Engine/Source/Runtime/Engine/Private/ShaderCompiler/ShaderCompiler.cpp
@@ -45,6 +45,7 @@
 #include "SceneInterface.h"
 #include "ShaderCodeLibrary.h"
 #include "MeshMaterialShaderType.h"
+#include "GBufferInfo.h"
 #include "ShaderParameterMetadata.h"
 #include "ProfilingDebugging/DiagnosticTable.h"
 #include "ProfilingDebugging/LoadTimeTracker.h"
@@ -3459,6 +3460,7 @@ void GlobalBeginCompileShader(
 	COOK_STAT(FScopedDurationTimer DurationTimer(ShaderCompilerCookStats::GlobalBeginCompileShaderTimeSec));
 
 	EShaderPlatform ShaderPlatform = EShaderPlatform(Target.Platform);
+	FShaderCompileUtilities::GenerateBrdfHeaders(ShaderPlatform);
 
 	TSharedRef<FShaderCompileJob, ESPMode::ThreadSafe> Job = StaticCastSharedRef<FShaderCompileJob>(NewJob);
 	FShaderCompilerInput& Input = Job->Input;
@@ -4045,6 +4047,10 @@ void GlobalBeginCompileShader(
 	const IShaderFormat* Format = GetTargetPlatformManagerRef().FindShaderFormat(Input.ShaderFormat);
 	Format->ModifyShaderCompilerInput(Input);
 	
+	// Allow the GBuffer and other shader defines to cause dependend environment changes, but minimizing the #ifdef magic in the shaders, which
+	// is nearly impossible to debug when it goes wrong.
+	FShaderCompileUtilities::ApplyDerivedDefines(Input.Environment, Input.SharedEnvironment, (EShaderPlatform)Target.Platform);
+
 	NewJobs.Add(NewJob);
 }
 
diff --git a/Engine/Source/Runtime/Engine/Private/ShaderCompiler/ShaderGenerationUtil.cpp b/Engine/Source/Runtime/Engine/Private/ShaderCompiler/ShaderGenerationUtil.cpp
new file mode 100644
--- /dev/null
+++ b/Engine/Source/Runtime/Engine/Private/ShaderCompiler/ShaderGenerationUtil.cpp
@@ -0,0 +1,2044 @@
+// Copyright Epic Games, Inc. All Rights Reserved.
+
+/*=============================================================================
+	ShaderGenerationUtil.cpp: Shader generation utilities.
+=============================================================================*/
+
+#include "ShaderCompiler.h"
+#include "ShaderCore.h"
+
+#include "Misc/FileHelper.h"
+#include "Misc/CoreMisc.h"
+#include "Misc/ScopeLock.h"
+#include "Interfaces/ITargetPlatformManagerModule.h"
+#include "RenderUtils.h"
+#include "RHI.h"
+
+#include "GBufferInfo.h"
+
+#if WITH_EDITOR
+#include "Common/TargetPlatformBase.h"
+#endif
+
+#include "RHIDefinitions.h"
+
+#include "ShaderMaterial.h"
+
+static TAutoConsoleVariable<int32> CVarShaderUseGBufferRefactor(
+	TEXT("r.Shaders.UseGBufferRefactor"),
+	0,
+	TEXT("Whether to use the refactored GBuffer that autogenerates encode/decode functions. Will be removed before UE5 ships."),
+	ECVF_Default);
+
+
+
+#define SET_COMPILE_BOOL_FORCE(X) OutEnvironment.SetDefine(TEXT(#X),SrcDefines.X ? TEXT("1") : TEXT("0"))
+
+static int FetchCompileInt(FShaderCompilerEnvironment& OutEnvironment, const char* SrcName)
+{
+	int32 Ret = 0;
+	if (OutEnvironment.GetDefinitions().Contains(SrcName))
+	{
+		const FString* Str = OutEnvironment.GetDefinitions().Find(SrcName);
+
+		Ret = atoi((const char*)Str->GetCharArray().GetData());
+	}
+
+	return Ret;
+}
+
+
+#define SET_COMPILE_BOOL_IF_TRUE(X) { if (DerivedDefines.X) { OutEnvironment.SetDefine(TEXT(#X),TEXT("1")); } }
+
+#define FETCH_COMPILE_BOOL(X) { if (OutEnvironment.GetDefinitions().Contains(#X)) SrcDefines.X = FetchCompileInt(OutEnvironment,#X) != 0 ? 1 : 0; }
+#define FETCH_COMPILE_INT(X) { if (OutEnvironment.GetDefinitions().Contains(#X)) SrcDefines.X = FetchCompileInt(OutEnvironment,#X); }
+
+#define MERGE_COMPILE_BOOL(X) (Lhs.X = (Lhs.X | Rhs.X))
+
+// these definitions can't be merged, so one of them should be cleared at zero.
+#define MERGE_COMPILE_INT(X) { check(Lhs.X == 0 || Rhs.X == 0); Lhs.X = Lhs.X > Rhs.X ? Lhs.X : Rhs.X; }
+
+void FShaderCompileUtilities::ApplyFetchEnvironment(FShaderGlobalDefines& SrcDefines, FShaderCompilerEnvironment& OutEnvironment)
+{
+	FETCH_COMPILE_BOOL(GBUFFER_HAS_VELOCITY);
+	FETCH_COMPILE_BOOL(GBUFFER_HAS_TANGENT);
+	FETCH_COMPILE_BOOL(ALLOW_STATIC_LIGHTING);
+	FETCH_COMPILE_BOOL(CLEAR_COAT_BOTTOM_NORMAL);
+	FETCH_COMPILE_BOOL(IRIS_NORMAL);
+	FETCH_COMPILE_BOOL(DXT5_NORMALMAPS);
+	FETCH_COMPILE_BOOL(SELECTIVE_BASEPASS_OUTPUTS);
+	FETCH_COMPILE_BOOL(USE_DBUFFER);
+	FETCH_COMPILE_BOOL(FORWARD_SHADING);
+	FETCH_COMPILE_BOOL(PROJECT_VERTEX_FOGGING_FOR_OPAQUE);
+	FETCH_COMPILE_BOOL(PROJECT_MOBILE_DISABLE_VERTEX_FOG);
+	FETCH_COMPILE_BOOL(USE_PREEXPOSURE);
+	FETCH_COMPILE_BOOL(PROJECT_ALLOW_GLOBAL_CLIP_PLANE);
+	FETCH_COMPILE_BOOL(EARLY_Z_PASS_ONLY_MATERIAL_MASKING);
+	FETCH_COMPILE_BOOL(PROJECT_SUPPORT_SKY_ATMOSPHERE);
+	FETCH_COMPILE_BOOL(PROJECT_SUPPORT_SKY_ATMOSPHERE_AFFECTS_HEIGHFOG);
+	FETCH_COMPILE_BOOL(PROJECT_MOBILE_USE_LEGACY_SHADING);
+	FETCH_COMPILE_BOOL(POST_PROCESS_ALPHA);
+	FETCH_COMPILE_BOOL(EIGHT_BIT_MESH_DISTANCE_FIELDS);
+	FETCH_COMPILE_BOOL(PLATFORM_SUPPORTS_RENDERTARGET_WRITE_MASK);
+	FETCH_COMPILE_BOOL(PLATFORM_SUPPORTS_PER_PIXEL_DBUFFER_MASK);
+	FETCH_COMPILE_BOOL(PLATFORM_SUPPORTS_DISTANCE_FIELDS);
+	FETCH_COMPILE_BOOL(VIRTUAL_TEXTURE_FEEDBACK_FACTOR);
+	FETCH_COMPILE_BOOL(PROJECT_MOBILE_ENABLE_MOVABLE_SPOTLIGHTS);
+	FETCH_COMPILE_BOOL(COMPILE_SHADERS_FOR_DEVELOPMENT_ALLOWED);
+
+	// This one is a helper for other macros
+	//bSupportsDualBlending;
+}
+
+void FShaderCompileUtilities::ApplyFetchEnvironment(FShaderLightmapPropertyDefines& SrcDefines, FShaderCompilerEnvironment& OutEnvironment)
+{
+	FETCH_COMPILE_BOOL(SIMPLE_FORWARD_SHADING);
+	FETCH_COMPILE_BOOL(SIMPLE_FORWARD_DIRECTIONAL_LIGHT);
+
+	FETCH_COMPILE_BOOL(LQ_TEXTURE_LIGHTMAP);
+	FETCH_COMPILE_BOOL(HQ_TEXTURE_LIGHTMAP);
+
+	FETCH_COMPILE_BOOL(CACHED_POINT_INDIRECT_LIGHTING);
+
+	FETCH_COMPILE_BOOL(STATICLIGHTING_TEXTUREMASK);
+	FETCH_COMPILE_BOOL(STATICLIGHTING_SIGNEDDISTANCEFIELD);
+
+	FETCH_COMPILE_BOOL(TRANSLUCENT_SELF_SHADOWING);
+
+	FETCH_COMPILE_BOOL(PRECOMPUTED_IRRADIANCE_VOLUME_LIGHTING);
+	FETCH_COMPILE_BOOL(CACHED_VOLUME_INDIRECT_LIGHTING);
+
+	FETCH_COMPILE_BOOL(WATER_MESH_FACTORY);
+	FETCH_COMPILE_BOOL(NIAGARA_MESH_FACTORY);
+	FETCH_COMPILE_BOOL(NIAGARA_MESH_INSTANCED);
+
+
+	FETCH_COMPILE_BOOL(PARTICLE_MESH_FACTORY);
+	FETCH_COMPILE_BOOL(PARTICLE_MESH_INSTANCED);
+
+	FETCH_COMPILE_BOOL(MANUAL_VERTEX_FETCH);
+
+}
+
+void FShaderCompileUtilities::ApplyFetchEnvironment(FShaderMaterialPropertyDefines& SrcDefines, FShaderCompilerEnvironment& OutEnvironment)
+{
+	FETCH_COMPILE_BOOL(MATERIAL_ENABLE_TRANSLUCENCY_FOGGING);
+	FETCH_COMPILE_BOOL(MATERIALBLENDING_ANY_TRANSLUCENT);
+	FETCH_COMPILE_BOOL(MATERIAL_USES_SCENE_COLOR_COPY);
+	FETCH_COMPILE_BOOL(MATERIALBLENDING_MASKED_USING_COVERAGE);
+
+	FETCH_COMPILE_BOOL(MATERIAL_COMPUTE_FOG_PER_PIXEL);
+	FETCH_COMPILE_BOOL(MATERIAL_SHADINGMODEL_UNLIT);
+
+	FETCH_COMPILE_BOOL(MATERIAL_SHADINGMODEL_DEFAULT_LIT);
+	FETCH_COMPILE_BOOL(MATERIAL_SHADINGMODEL_SUBSURFACE);
+	FETCH_COMPILE_BOOL(MATERIAL_SHADINGMODEL_PREINTEGRATED_SKIN);
+	FETCH_COMPILE_BOOL(MATERIAL_SHADINGMODEL_SUBSURFACE_PROFILE);
+	FETCH_COMPILE_BOOL(MATERIAL_SHADINGMODEL_CLEAR_COAT);
+	FETCH_COMPILE_BOOL(MATERIAL_SHADINGMODEL_TWOSIDED_FOLIAGE);
+	FETCH_COMPILE_BOOL(MATERIAL_SHADINGMODEL_HAIR);
+	FETCH_COMPILE_BOOL(MATERIAL_SHADINGMODEL_CLOTH);
+	FETCH_COMPILE_BOOL(MATERIAL_SHADINGMODEL_EYE);
+	FETCH_COMPILE_BOOL(MATERIAL_SHADINGMODEL_SINGLELAYERWATER);
+	FETCH_COMPILE_BOOL(MATERIAL_SHADINGMODEL_THIN_TRANSLUCENT);
+
+	FETCH_COMPILE_BOOL(MATERIAL_FULLY_ROUGH);
+
+	FETCH_COMPILE_BOOL(USES_EMISSIVE_COLOR);
+
+	FETCH_COMPILE_BOOL(MATERIALBLENDING_SOLID);
+	FETCH_COMPILE_BOOL(MATERIALBLENDING_MASKED);
+	FETCH_COMPILE_BOOL(MATERIALBLENDING_ALPHACOMPOSITE);
+	FETCH_COMPILE_BOOL(MATERIALBLENDING_TRANSLUCENT);
+	FETCH_COMPILE_BOOL(MATERIALBLENDING_ADDITIVE);
+	FETCH_COMPILE_BOOL(MATERIALBLENDING_MODULATE);
+	FETCH_COMPILE_BOOL(MATERIALBLENDING_ALPHAHOLDOUT);
+
+	FETCH_COMPILE_INT(MATERIALDECALRESPONSEMASK);
+
+	FETCH_COMPILE_BOOL(REFRACTION_USE_INDEX_OF_REFRACTION);
+	FETCH_COMPILE_BOOL(REFRACTION_USE_PIXEL_NORMAL_OFFSET);
+
+	FETCH_COMPILE_BOOL(USE_DITHERED_LOD_TRANSITION_FROM_MATERIAL);
+	FETCH_COMPILE_BOOL(MATERIAL_TWOSIDED);
+	FETCH_COMPILE_BOOL(MATERIAL_TANGENTSPACENORMAL);
+	FETCH_COMPILE_BOOL(GENERATE_SPHERICAL_PARTICLE_NORMALS);
+	FETCH_COMPILE_BOOL(MATERIAL_USE_PREINTEGRATED_GF);
+	FETCH_COMPILE_BOOL(MATERIAL_HQ_FORWARD_REFLECTIONS);
+	FETCH_COMPILE_BOOL(MATERIAL_PLANAR_FORWARD_REFLECTIONS);
+	FETCH_COMPILE_BOOL(MATERIAL_NONMETAL);
+	FETCH_COMPILE_BOOL(MATERIAL_USE_LM_DIRECTIONALITY);
+	FETCH_COMPILE_BOOL(MATERIAL_INJECT_EMISSIVE_INTO_LPV);
+	FETCH_COMPILE_BOOL(MATERIAL_SSR);
+	FETCH_COMPILE_BOOL(MATERIAL_CONTACT_SHADOWS);
+	FETCH_COMPILE_BOOL(MATERIAL_BLOCK_GI);
+	FETCH_COMPILE_BOOL(MATERIAL_DITHER_OPACITY_MASK);
+	FETCH_COMPILE_BOOL(MATERIAL_NORMAL_CURVATURE_TO_ROUGHNESS);
+	FETCH_COMPILE_BOOL(MATERIAL_ALLOW_NEGATIVE_EMISSIVECOLOR);
+	FETCH_COMPILE_BOOL(MATERIAL_OUTPUT_OPACITY_AS_ALPHA);
+	FETCH_COMPILE_BOOL(TRANSLUCENT_SHADOW_WITH_MASKED_OPACITY);
+
+	FETCH_COMPILE_INT(DECALBLENDMODEID_VOLUMETRIC);
+	FETCH_COMPILE_INT(DECALBLENDMODEID_STAIN);
+	FETCH_COMPILE_INT(DECALBLENDMODEID_NORMAL);
+	FETCH_COMPILE_INT(DECALBLENDMODEID_EMISSIVE);
+	FETCH_COMPILE_INT(DECALBLENDMODEID_TRANSLUCENT);
+	FETCH_COMPILE_INT(DECALBLENDMODEID_AO);
+	FETCH_COMPILE_INT(DECALBLENDMODEID_ALPHACOMPOSITE);
+
+	FETCH_COMPILE_BOOL(MATERIAL_DOMAIN_SURFACE);
+	FETCH_COMPILE_BOOL(MATERIAL_DOMAIN_DEFERREDDECAL);
+	FETCH_COMPILE_BOOL(MATERIAL_DOMAIN_LIGHTFUNCTION);
+	FETCH_COMPILE_BOOL(MATERIAL_DOMAIN_VOLUME);
+	FETCH_COMPILE_BOOL(MATERIAL_DOMAIN_POSTPROCESS);
+	FETCH_COMPILE_BOOL(MATERIAL_DOMAIN_UI);
+	FETCH_COMPILE_BOOL(MATERIAL_DOMAIN_VIRTUALTEXTURE);
+
+	FETCH_COMPILE_BOOL(TRANSLUCENCY_LIGHTING_VOLUMETRIC_NONDIRECTIONAL);
+	FETCH_COMPILE_BOOL(TRANSLUCENCY_LIGHTING_VOLUMETRIC_DIRECTIONAL);
+	FETCH_COMPILE_BOOL(TRANSLUCENCY_LIGHTING_VOLUMETRIC_PERVERTEX_NONDIRECTIONAL);
+	FETCH_COMPILE_BOOL(TRANSLUCENCY_LIGHTING_VOLUMETRIC_PERVERTEX_DIRECTIONAL);
+	FETCH_COMPILE_BOOL(TRANSLUCENCY_LIGHTING_SURFACE_LIGHTINGVOLUME);
+	FETCH_COMPILE_BOOL(TRANSLUCENCY_LIGHTING_SURFACE_FORWARDSHADING);
+
+	FETCH_COMPILE_BOOL(EDITOR_PRIMITIVE_MATERIAL);
+
+	FETCH_COMPILE_BOOL(USE_STENCIL_LOD_DITHER_DEFAULT);
+
+	FETCH_COMPILE_BOOL(MATERIALDOMAIN_SURFACE);
+	FETCH_COMPILE_BOOL(MATERIALDOMAIN_DEFERREDDECAL);
+	FETCH_COMPILE_BOOL(MATERIALDOMAIN_LIGHTFUNCTION);
+	FETCH_COMPILE_BOOL(MATERIALDOMAIN_POSTPROCESS);
+	FETCH_COMPILE_BOOL(MATERIALDOMAIN_UI);
+
+	FETCH_COMPILE_BOOL(OUT_BASECOLOR);
+	FETCH_COMPILE_BOOL(OUT_BASECOLOR_NORMAL_SPECULAR);
+	FETCH_COMPILE_BOOL(OUT_WORLDHEIGHT);
+
+	FETCH_COMPILE_BOOL(IS_VIRTUAL_TEXTURE_MATERIAL);
+	FETCH_COMPILE_BOOL(IS_DECAL);
+	FETCH_COMPILE_BOOL(IS_BASE_PASS);
+	FETCH_COMPILE_BOOL(IS_MATERIAL_SHADER);
+
+	FETCH_COMPILE_BOOL(STRATA_TRANSLUCENT_ENABLED);
+	FETCH_COMPILE_BOOL(DUAL_SOURCE_COLOR_BLENDING_ENABLED);
+
+	FETCH_COMPILE_INT(DECAL_RENDERTARGET_COUNT);
+
+	FETCH_COMPILE_BOOL(DUAL_SOURCE_COLOR_BLENDING_ENABLED);
+}
+
+void FShaderCompileUtilities::ApplyFetchEnvironment(FShaderCompilerDefines& SrcDefines, FShaderCompilerEnvironment& OutEnvironment)
+{
+	FETCH_COMPILE_BOOL(COMPILER_GLSL_ES3_1);
+	FETCH_COMPILE_BOOL(ES3_1_PROFILE);
+
+	FETCH_COMPILE_BOOL(COMPILER_GLSL);
+	FETCH_COMPILE_BOOL(COMPILER_GLSL_ES3_1_EXT);
+	FETCH_COMPILE_BOOL(ESDEFERRED_PROFILE);
+	FETCH_COMPILE_BOOL(GL4_PROFILE);
+
+	FETCH_COMPILE_BOOL(METAL_PROFILE);
+	FETCH_COMPILE_BOOL(VULKAN_PROFILE);
+	FETCH_COMPILE_BOOL(MAC);
+
+	FETCH_COMPILE_BOOL(PLATFORM_SUPPORTS_DEVELOPMENT_SHADERS);
+}
+
+
+// if we change the logic, increment this number to force a DDC key change
+static const int32 GBufferGeneratorVersion = 1;
+
+static FShaderGlobalDefines FetchShaderGlobalDefines(EShaderPlatform TargetPlatform)
+{
+	FShaderGlobalDefines Ret = {};
+
+	{
+		static const auto CVar = IConsoleManager::Get().FindTConsoleVariableDataInt(TEXT("r.AllowStaticLighting"));
+		Ret.ALLOW_STATIC_LIGHTING = (CVar ? (CVar->GetValueOnAnyThread() != 0) : 1);
+	}
+
+	Ret.GBUFFER_HAS_VELOCITY = IsUsingBasePassVelocity(TargetPlatform) ? 1 : 0;
+	Ret.GBUFFER_HAS_TANGENT = false;//BasePassCanOutputTangent(TargetPlatform) ? 1 : 0;
+
+	{
+		static const auto CVar = IConsoleManager::Get().FindTConsoleVariableDataInt(TEXT("r.ClearCoatNormal"));
+		Ret.CLEAR_COAT_BOTTOM_NORMAL = CVar ? (CVar->GetValueOnAnyThread() != 0) : 0;
+	}
+
+	{
+		static const auto CVar = IConsoleManager::Get().FindTConsoleVariableDataInt(TEXT("r.IrisNormal"));
+		Ret.IRIS_NORMAL = CVar ? (CVar->GetValueOnAnyThread() != 0) : 0;
+	}
+
+	{
+		static const auto CVar = IConsoleManager::Get().FindTConsoleVariableDataInt(TEXT("Compat.UseDXT5NormalMaps"));
+		Ret.DXT5_NORMALMAPS = CVar ? (CVar->GetValueOnAnyThread() != 0) : 0;
+	}
+
+	Ret.SELECTIVE_BASEPASS_OUTPUTS = IsUsingSelectiveBasePassOutputs((EShaderPlatform)TargetPlatform) ? 1 : 0;
+	Ret.USE_DBUFFER = IsUsingDBuffers((EShaderPlatform)TargetPlatform) ? 1 : 0;
+
+#if WITH_EDITOR
+	{
+		bool bForwardShading = false;
+
+		ITargetPlatform* TargetPlatformPtr = GetTargetPlatformManager()->FindTargetPlatformWithSupport(TEXT("ShaderFormat"), LegacyShaderPlatformToShaderFormat((EShaderPlatform)TargetPlatform));
+		if (TargetPlatformPtr)
+		{
+			bForwardShading = TargetPlatformPtr->UsesForwardShading();
+		}
+		else
+		{
+			static IConsoleVariable* CVarForwardShading = IConsoleManager::Get().FindConsoleVariable(TEXT("r.ForwardShading"));
+			bForwardShading = CVarForwardShading ? (CVarForwardShading->GetInt() != 0) : false;
+		}
+		Ret.FORWARD_SHADING = bForwardShading;
+	}
+#else
+	{
+		static const auto CVar = IConsoleManager::Get().FindTConsoleVariableDataInt(TEXT("r.ForwardShading"));
+		Ret.FORWARD_SHADING = CVar ? (CVar->GetValueOnAnyThread() != 0) : 0;
+	}
+#endif
+
+	{
+		static IConsoleVariable* CVar = IConsoleManager::Get().FindConsoleVariable(TEXT("r.VertexFoggingForOpaque"));
+		Ret.PROJECT_VERTEX_FOGGING_FOR_OPAQUE = Ret.FORWARD_SHADING && (CVar ? (CVar->GetInt() != 0) : 0);
+	}
+
+	{
+		static IConsoleVariable* CVar = IConsoleManager::Get().FindConsoleVariable(TEXT("r.Mobile.DisableVertexFog"));
+		Ret.PROJECT_MOBILE_DISABLE_VERTEX_FOG = CVar ? (CVar->GetInt() != 0) : 0;
+	}
+
+
+	{
+		// PreExposure is now always enabled in the shaders.
+		Ret.USE_PREEXPOSURE = 1;
+	}
+
+	{
+		static IConsoleVariable* CVar = IConsoleManager::Get().FindConsoleVariable(TEXT("r.AllowGlobalClipPlane"));
+		Ret.PROJECT_ALLOW_GLOBAL_CLIP_PLANE = CVar ? (CVar->GetInt() != 0) : 0;
+	}
+
+	{
+		if (MaskedInEarlyPass((EShaderPlatform)TargetPlatform))
+		{
+			Ret.EARLY_Z_PASS_ONLY_MATERIAL_MASKING = 1;
+		}
+		else
+		{
+			Ret.EARLY_Z_PASS_ONLY_MATERIAL_MASKING = 0;
+		}
+	}
+
+	bool bSupportSkyAtmosphere = false;
+	{
+		static IConsoleVariable* CVar = IConsoleManager::Get().FindConsoleVariable(TEXT("r.SupportSkyAtmosphere"));
+		bSupportSkyAtmosphere = CVar && CVar->GetInt() != 0;
+		Ret.PROJECT_SUPPORT_SKY_ATMOSPHERE = bSupportSkyAtmosphere ? 1 : 0;
+	}
+
+	{
+		static IConsoleVariable* CVar = IConsoleManager::Get().FindConsoleVariable(TEXT("r.SupportSkyAtmosphereAffectsHeightFog"));
+		Ret.PROJECT_SUPPORT_SKY_ATMOSPHERE_AFFECTS_HEIGHFOG = (CVar && bSupportSkyAtmosphere) ? (CVar->GetInt() != 0) : 0;
+	}
+
+	{
+		static IConsoleVariable* CVar = IConsoleManager::Get().FindConsoleVariable(TEXT("r.Mobile.UseLegacyShadingModel"));
+		Ret.PROJECT_MOBILE_USE_LEGACY_SHADING = CVar ? (CVar->GetInt() != 0) : 0;
+	}
+
+	{
+		static IConsoleVariable* CVar = IConsoleManager::Get().FindConsoleVariable(TEXT("r.PostProcessing.PropagateAlpha"));
+		int32 PropagateAlpha = CVar->GetInt();
+		if (PropagateAlpha < 0 || PropagateAlpha > 2)
+		{
+			PropagateAlpha = 0;
+		}
+		Ret.POST_PROCESS_ALPHA = PropagateAlpha != 0;
+	}
+
+	{
+		static IConsoleVariable* CVar = IConsoleManager::Get().FindConsoleVariable(TEXT("r.DistanceFieldBuild.EightBit"));
+		Ret.EIGHT_BIT_MESH_DISTANCE_FIELDS = CVar ? (CVar->GetInt() != 0) : 0;
+	}
+
+	Ret.PLATFORM_SUPPORTS_RENDERTARGET_WRITE_MASK = RHISupportsRenderTargetWriteMask(EShaderPlatform(TargetPlatform)) ? 1 : 0;
+	Ret.PLATFORM_SUPPORTS_PER_PIXEL_DBUFFER_MASK = IsUsingPerPixelDBufferMask(EShaderPlatform(TargetPlatform)) ? 1 : 0;
+	Ret.PLATFORM_SUPPORTS_DISTANCE_FIELDS = DoesPlatformSupportDistanceFields(EShaderPlatform(TargetPlatform)) ? 1 : 0;
+	
+	{
+		static const auto CVar = IConsoleManager::Get().FindConsoleVariable(TEXT("r.vt.FeedbackFactor"));
+		Ret.VIRTUAL_TEXTURE_FEEDBACK_FACTOR = (CVar ? FMath::Max(CVar->GetInt(), 1) : 1) != 0;
+	}
+
+	if (IsMobilePlatform((EShaderPlatform)TargetPlatform))
+	{
+		static IConsoleVariable* CVar = IConsoleManager::Get().FindConsoleVariable(TEXT("r.Mobile.EnableMovableSpotlights"));
+		Ret.PROJECT_MOBILE_ENABLE_MOVABLE_SPOTLIGHTS = CVar ? (CVar->GetInt() != 0) : 0;
+	}
+
+	Ret.bSupportsDualBlending = RHISupportsDualSourceBlending(TargetPlatform);
+
+	{
+		static IConsoleVariable* CVarLumen = IConsoleManager::Get().FindConsoleVariable(TEXT("r.LumenScene"));
+		Ret.bNeedVelocityDepth = CVarLumen->GetInt() > 0 || FDataDrivenShaderPlatformInfo::GetSupportsRayTracing(TargetPlatform);
+		//Ret.bNeedVelocityDepth ? PF_A16B16G16R16 : PF_G16R16;
+	}
+
+	{
+		static IConsoleVariable* CVar = IConsoleManager::Get().FindConsoleVariable(TEXT("r.GBufferFormat"));
+		Ret.LegacyGBufferFormat = CVar->GetInt();
+	}
+
+	return Ret;
+}
+
+static FString GetSlotTextName(EGBufferSlot Slot)
+{
+	FString Ret = TEXT("");
+	switch (Slot)
+	{
+	case GBS_Invalid:
+		check(0);
+		return TEXT("Invalid");
+	case GBS_SceneColor:
+		return TEXT("SceneColor");
+	case GBS_WorldNormal:
+		return TEXT("WorldNormal");
+	case GBS_PerObjectGBufferData:
+		return TEXT("PerObjectGBufferData");
+	case GBS_Metallic:
+		return TEXT("Metallic");
+	case GBS_Specular:
+		return TEXT("Specular");
+	case GBS_Roughness:
+		return TEXT("Roughness");
+	case GBS_ShadingModelId:
+		return TEXT("ShadingModelID");
+	case GBS_SelectiveOutputMask:
+		return TEXT("SelectiveOutputMask");
+	case GBS_BaseColor:
+		return TEXT("BaseColor");
+	case GBS_GenericAO:
+		return TEXT("GenericAO");
+	case GBS_IndirectIrradiance:
+		return TEXT("IndirectIrradiance");
+	case GBS_Velocity:
+		return TEXT("Velocity");
+	case GBS_PrecomputedShadowFactor:
+		return TEXT("PrecomputedShadowFactors");
+	case GBS_WorldTangent:
+		return TEXT("WorldTangent");
+	case GBS_Anisotropy:
+		return TEXT("Anisotropy");
+	case GBS_CustomData:
+		return TEXT("CustomData");
+	case GBS_SubsurfaceColor:
+		return TEXT("SubsurfaceColor");
+	case GBS_Opacity:
+		return TEXT("Opacity");
+	case GBS_SubsurfaceProfile:
+		return TEXT("SubsurfaceProfile");
+	case GBS_ClearCoat:
+		return TEXT("ClearCoat");
+	case GBS_ClearCoatRoughness:
+		return TEXT("ClearCoatRoughness");
+	case GBS_HairSecondaryWorldNormal:
+		return TEXT("HaireEcondaryWorldNormal");
+	case GBS_HairBacklit:
+		return TEXT("HairBacklit");
+	case GBS_Cloth:
+		return TEXT("Cloth");
+	case GBS_SubsurfaceProfileX:
+		return TEXT("SubsurfaceProfileX");
+	case GBS_IrisNormal:
+		return TEXT("IrisNormal");
+	default:
+		break;
+	};
+
+	return TEXT("ERROR");
+}
+
+
+static FString GetFloatType(int32 ChanNum)
+{
+	switch (ChanNum)
+	{
+	case 0:
+		check(0);
+		return TEXT("Error");
+	case 1:
+		return TEXT("float");
+	case 2:
+		return TEXT("float2");
+	case 3:
+		return TEXT("float3");
+	case 4:
+		return TEXT("float4");
+	default:
+		break;
+	}
+
+	check(0);
+	return TEXT("Error");
+}
+
+static FString GetUintType(int32 ChanNum)
+{
+	switch (ChanNum)
+	{
+	case 0:
+		check(0);
+		return TEXT("Error");
+	case 1:
+		return TEXT("uint");
+	case 2:
+		return TEXT("uint2");
+	case 3:
+		return TEXT("uint3");
+	case 4:
+		return TEXT("uint4");
+	default:
+		break;
+	}
+
+	check(0);
+	return TEXT("Error");
+}
+
+
+static const FGBufferCompressionInfo GBufferCompressionInfo[] =
+{
+	{ GBC_Invalid,					    0, 0, {  0,  0,  0,  0 }, false, false, TEXT("Invalid")                , TEXT("Invalid")                },
+	{ GBC_Raw_Float_16_16_16_16,		4, 4, { 16, 16, 16, 16 }, false, false, TEXT("Invalid")                , TEXT("Invalid")                },
+	{ GBC_Raw_Float_16_16_16,			3, 3, { 16, 16, 16,  0 }, false, false, TEXT("Invalid")                , TEXT("Invalid")                },
+	{ GBC_Raw_Float_11_11_10,			3, 3, { 11, 11, 10,  0 }, false, false, TEXT("Invalid")                , TEXT("Invalid")                },
+	{ GBC_Raw_Float_10_10_10,			3, 3, { 10, 10, 10,  0 }, false, false, TEXT("Invalid")                , TEXT("Invalid")                },
+	{ GBC_Raw_Unorm_8_8_8_8,			4, 4, {  8,  8,  8,  8 }, false, false, TEXT("Invalid")                , TEXT("Invalid")                },
+	{ GBC_Raw_Unorm_8_8_8,			  	3, 3, {  8,  8,  8,  0 }, false, false, TEXT("Invalid")                , TEXT("Invalid")                },
+	{ GBC_Raw_Unorm_8_8,				2, 2, {  8,  8,  0,  0 }, false, false, TEXT("Invalid")                , TEXT("Invalid")                },
+	{ GBC_Raw_Unorm_8,				  	1, 1, {  8,  0,  0,  0 }, false, false, TEXT("Invalid")                , TEXT("Invalid")                },
+	{ GBC_Raw_Unorm_2, 				  	1, 1, {  2,  0,  0,  0 }, false, false, TEXT("Invalid")                , TEXT("Invalid")                },
+	{ GBC_Raw_Float_16_16,			  	2, 2, { 16, 16,  0,  0 }, false, false, TEXT("Invalid")                , TEXT("Invalid")                },
+	{ GBC_Raw_Float_16,				  	1, 1, { 16,  0,  0,  0 }, false, false, TEXT("Invalid")                , TEXT("Invalid")                },
+	{ GBC_Bits_4, 					  	1, 1, {  4,  0,  0,  0 },  true, false, TEXT("Invalid")                , TEXT("Invalid")                },
+	{ GBC_Bits_2, 					  	1, 1, {  2,  0,  0,  0 },  true, false, TEXT("Invalid")                , TEXT("Invalid")                },
+	{ GBC_Packed_Normal_Octahedral_8_8, 3, 2, {  8,  8,  0,  0 }, false,  true, TEXT("CompressOctahedral")     , TEXT("DecompressOctahedral")   },
+	{ GBC_EncodeNormal_Normal_16_16_16, 3, 3, { 16, 16, 16,  0 }, false,  true, TEXT("EncodeNormalHelper")     , TEXT("DecodeNormalHelper")     },
+	{ GBC_EncodeNormal_Normal_10_10_10, 3, 3, { 10, 10, 10,  0 }, false,  true, TEXT("EncodeNormalHelper")     , TEXT("DecodeNormalHelper")     },
+	{ GBC_Packed_Color_5_6_5,			3, 3, {  5,  6,  5,  0 },  true,  true, TEXT("EncodeQuantize565")      , TEXT("DecodeQuantize565")      },
+	{ GBC_Packed_Color_5_6_5_Sqrt,		3, 3, {  5,  6,  5,  0 },  true,  true, TEXT("EncodeQuantize565Sqrt")  , TEXT("DecodeQuantize565Sqrt")  },
+	{ GBC_Packed_Color_4_4_4,			3, 3, {  4,  4,  4,  0 },  true,  true, TEXT("EncodeQuantize444")      , TEXT("DecodeQuantize444")	    },
+	{ GBC_Packed_Color_4_4_4_Sqrt,		3, 3, {  4,  4,  4,  0 },  true,  true, TEXT("EncodeQuantize444Sqrt")  , TEXT("DecodeQuantize444Sqrt")  },
+	{ GBC_Packed_Color_3_3_2,			3, 3, {  3,  3,  2,  0 },  true,  true, TEXT("EncodeQuantize332")      , TEXT("DecodeQuantize332")      },
+	{ GBC_Packed_Color_3_3_2_Sqrt,		3, 3, {  3,  3,  2,  0 },  true,  true, TEXT("EncodeQuantize332Sqrt")  , TEXT("DecodeQuantize332Sqrt")  },
+	{ GBC_Packed_Quantized_6,			1, 1, {  6,  0,  0,  0 },  true,  true, TEXT("EncodeQuantize6")        , TEXT("DecodeQuantize6")        },
+	{ GBC_Packed_Quantized_4,			1, 1, {  4,  0,  0,  0 },  true,  true, TEXT("EncodeQuantize4")        , TEXT("DecodeQuantize4")        },
+	{ GBC_Packed_Quantized_2,			1, 1, {  2,  0,  0,  0 },  true,  true, TEXT("EncodeQuantize2")        , TEXT("DecodeQuantize2")        },
+};
+
+static int32 GetGBufferCompressionBitSize(EGBufferCompression Compression)
+{
+	check(Compression >= 0 && Compression < GBC_Num);
+	check(GBufferCompressionInfo[Compression].Type == Compression);
+
+	int32 BitSize = 0;
+	for (int32 I = 0; I < 4; I++)
+	{
+		BitSize += GBufferCompressionInfo[Compression].ChanBits[I];
+	}
+
+	return BitSize;
+}
+
+static int32 GetGBufferCompressionChanBitSize(EGBufferCompression Compression, int32 Chan)
+{
+	check(Compression >= 0 && Compression < GBC_Num);
+	check(GBufferCompressionInfo[Compression].Type == Compression);
+	check(Chan >= 0 && Chan < 4);
+
+	int32 BitSize = GBufferCompressionInfo[Compression].ChanBits[Chan];
+
+	return BitSize;
+}
+
+static bool IsPackedBits(EGBufferCompression Compression)
+{
+	check(Compression >= 0 && Compression < GBC_Num);
+	check(GBufferCompressionInfo[Compression].Type == Compression);
+	return GBufferCompressionInfo[Compression].bIsPackedBits;
+}
+
+// for each compression type, the number of incoming and outgoing channels. So for example octahedral encoding goes from
+// 3 channels to 2, so GetCompressionSrcNumChannels() is 3 and GetCompressionDstNumChannels() is 2.
+static int32 GetCompressionSrcNumChannels(EGBufferCompression Compression)
+{
+	check(Compression >= 0 && Compression < GBC_Num);
+	check(GBufferCompressionInfo[Compression].Type == Compression);
+	return GBufferCompressionInfo[Compression].SrcNumChan;
+}
+
+static int32 GetCompressionDstNumChannels(EGBufferCompression Compression)
+{
+	check(Compression >= 0 && Compression < GBC_Num);
+	check(GBufferCompressionInfo[Compression].Type == Compression);
+	return GBufferCompressionInfo[Compression].DstNumChan;
+}
+
+// does this type require some kind of pack/unpack step (such as octahedral encoding) other than just bit packing
+static int32 GetCompressionNeedsConversion(EGBufferCompression Compression)
+{
+	check(Compression >= 0 && Compression < GBC_Num);
+	check(GBufferCompressionInfo[Compression].Type == Compression);
+	return GBufferCompressionInfo[Compression].bIsConversion;
+}
+
+// does this type require some kind of pack/unpack step (such as octahedral encoding) other than just bit packing
+static FString GetCompressionEncodeFuncName(EGBufferCompression Compression)
+{
+	check(Compression >= 0 && Compression < GBC_Num);
+	check(GBufferCompressionInfo[Compression].Type == Compression);
+	return GBufferCompressionInfo[Compression].EncodeFuncName;
+}
+
+
+// does this type require some kind of pack/unpack step (such as octahedral encoding) other than just bit packing
+static FString GetCompressionDecodeFuncName(EGBufferCompression Compression)
+{
+	check(Compression >= 0 && Compression < GBC_Num);
+	check(GBufferCompressionInfo[Compression].Type == Compression);
+	return GBufferCompressionInfo[Compression].DecodeFuncName;
+}
+
+static int32 GetBufferNumBits(EGBufferType Format, int32 Channel)
+{
+	check(Channel >= 0);
+	check(Channel < 4);
+
+	int32 Ret = 0;
+	switch (Format)
+	{
+	case GBT_Invalid:
+		check(0);
+		break;
+	case GBT_Unorm_8_8_8_8:
+		Ret = 8;
+		break;
+	case GBT_Unorm_11_11_10:
+	{
+		int32 Sizes[4] = { 11, 11, 10, 0 };
+		Ret = Sizes[Channel];
+	}
+	break;
+	case GBT_Unorm_10_10_10_2:
+	{
+		int32 Sizes[4] = { 10, 10, 10, 2 };
+		Ret = Sizes[Channel];
+	}
+	break;
+	case GBT_Float_16_16:
+		Ret = 16;
+		break;
+	default:
+		check(0);
+		break;
+	}
+
+	return Ret;
+}
+
+static int32 GetTargetNumChannels(EGBufferType Type)
+{
+	int32 Ret = 0;
+
+	switch (Type)
+	{
+	case GBT_Invalid:
+		Ret = 0;
+		break;
+	case GBT_Unorm_8_8_8_8:
+		Ret = 4;
+		break;
+	case GBT_Unorm_11_11_10:
+		Ret = 3;
+		break;
+	case GBT_Unorm_10_10_10_2:
+		Ret = 4;
+		break;
+	case GBT_Float_16_16:
+		Ret = 2;
+		break;
+	default:
+		check(0);
+		break;
+	}
+
+	return Ret;
+}
+
+
+// Quick note: There are lots of instances where we do something like:
+//
+// CurrLine = SomeComplicatedString();
+// FullStr += CurrLine;
+//
+// We could instead make those run on a single line, as in:
+//
+// FullStr += SomeComplicatedString();
+//
+// The problem is it then becomes a nightmare to step through and debug the lines, which is why each line is fully created before appending.
+
+static FString CreateGBufferEncodeFunction(const FGBufferInfo& BufferInfo)
+{
+	FString FullStr;
+
+	FullStr += TEXT("void EncodeGBufferToMRT(inout FPixelShaderOut Out, FGBufferData GBuffer, float QuantizationBias)\n");
+	FullStr += TEXT("{\n");
+
+	int32 TargetChanNum[FGBufferInfo::MaxTargets] = {};
+
+	// for each values, 0 means float, 1 means int, and -1 means unused
+	int32 PackedStatus[FGBufferInfo::MaxTargets][4];
+	for (int32 I = 0; I < FGBufferInfo::MaxTargets; I++)
+	{
+		for (int32 J = 0; J < 4; J++)
+		{
+			PackedStatus[I][J] = -1;
+		}
+	}
+
+	bool bNeedsConversion[GBS_Num] = {};
+
+	// For each item that we are writing to the gbuffer, check if it is a float or an an int. Mixing and matching
+	// floats and ints to the same channel is illegal. In that case, the float needs to be explicitly packed to an int,
+	// and then gets merged with the other int.
+
+	for (int32 I = 0; I < GBS_Num; I++)
+	{
+		const FGBufferItem& Item = BufferInfo.Slots[I];
+
+		if (Item.bIsValid)
+		{
+			bool isPacked = IsPackedBits(Item.Compression);
+			int32 NumComponents = GetCompressionDstNumChannels(Item.Compression);
+
+			for (int32 PackIndex = 0; PackIndex < FGBufferItem::MaxPacking; PackIndex++)
+			{
+				const FGBufferPacking& PackInfo = Item.Packing[PackIndex];
+
+				if (PackInfo.bIsValid)
+				{
+					int32 DstTarget = PackInfo.TargetIndex;
+					int32 DstChannel = PackInfo.DstChannelIndex;
+
+					if (!isPacked)
+					{
+						// if it is not packed, make sure the GBuffer is expecting a raw item
+						// and make sure we didn't already set it
+						check(PackInfo.bFull);
+						check(PackedStatus[DstTarget][DstChannel] == -1);
+						PackedStatus[DstTarget][DstChannel] = 0;
+					}
+					else
+					{
+						// if is packed, make sure the GBuffer is expecting it.
+						check(!PackInfo.bFull);
+
+						// since multiple items are packed in a single channel, make sure it's either unset or packed (not zero, which means float)
+						check(PackedStatus[DstTarget][DstChannel] == -1 || PackedStatus[DstTarget][DstChannel] == 1);
+
+						PackedStatus[DstTarget][DstChannel] = 1;
+					}
+				}
+			}
+
+			// also, for this slot, check if it needs a type conversion
+			bNeedsConversion[I] = GetCompressionNeedsConversion(Item.Compression) != 0;
+		}
+	}
+
+	// for any parameters that need conversion 
+
+	// write the float and int parameters that we need
+	for (int32 I = 0; I < FGBufferInfo::MaxTargets; I++)
+	{
+		// special case for target 0, for now
+		if (I == 0)
+		{
+			continue;
+		}
+
+		int32 ChanNum = GetTargetNumChannels(BufferInfo.Targets[I].TargetType);
+		TargetChanNum[I] = ChanNum;
+
+		bool bAnyFloat = false;
+		bool bAnyUint = false;
+
+		for (int32 Chan = 0; Chan < 4; Chan++)
+		{
+			if (PackedStatus[I][Chan] == 0)
+			{
+				bAnyFloat = true;
+			}
+			if (PackedStatus[I][Chan] == 1)
+			{
+				bAnyUint = true;
+			}
+		}
+
+		if (ChanNum > 0)
+		{
+			if (bAnyFloat)
+			{
+				FString FloatName = GetFloatType(ChanNum);
+				FullStr += FString::Printf(TEXT("\t%s MrtFloat%d = 0.0f;\n"), FloatName.GetCharArray().GetData(), I);
+			}
+
+			if (bAnyUint)
+			{
+				FString UintName = GetUintType(ChanNum);
+				FullStr += FString::Printf(TEXT("\t%s MrtUint%d = 0;\n"), UintName.GetCharArray().GetData(), I);
+			}
+		}
+	}
+
+	FullStr += TEXT("\n");
+
+	int32 NumAdded = 0;
+	for (int32 ItemIndex = 0; ItemIndex < GBS_Num; ItemIndex++)
+	{
+		const FGBufferItem& Item = BufferInfo.Slots[ItemIndex];
+
+		// special case for now
+		if (Item.BufferSlot == GBS_SceneColor)
+		{
+			continue;
+		}
+
+		if (Item.bIsValid)
+		{
+			bool bIsCompressed = GetCompressionNeedsConversion(Item.Compression) != 0;
+			if (bIsCompressed)
+			{
+				bNeedsConversion[ItemIndex] = true;
+
+				// so far we are assuming all input data is float
+				FString ConversionName = GetCompressionEncodeFuncName(Item.Compression);
+				FString BaseName = GetSlotTextName(Item.BufferSlot);
+
+				int32 SrcChan = GetCompressionDstNumChannels(Item.Compression);
+
+				FString SrcTypeName;
+				bool bIsPackedBits = IsPackedBits(Item.Compression);
+				if (bIsPackedBits)
+				{
+					SrcTypeName = GetUintType(SrcChan);
+				}
+				else
+				{
+					SrcTypeName = GetFloatType(SrcChan);
+				}
+
+				FString QuantizationValue = TEXT("0.0f");
+				if (Item.bQuantizationBias)
+				{
+					QuantizationValue = TEXT("QuantizationBias");
+				}
+
+				// If there is no quanzitation bias, just add the values
+				FString CurrLine = FString::Printf(TEXT("\t%s %s_Compressed = %s(GBuffer.%s, %s);\n"),
+					SrcTypeName.GetCharArray().GetData(),
+					BaseName.GetCharArray().GetData(),
+					ConversionName.GetCharArray().GetData(),
+					BaseName.GetCharArray().GetData(),
+					QuantizationValue.GetCharArray().GetData());
+				FullStr += CurrLine;
+
+				NumAdded++;
+			}
+		}
+	}
+
+	if (NumAdded > 0)
+	{
+		FullStr += TEXT("\n");
+	}
+
+	FString Swizzles[4];
+	Swizzles[0] = TEXT("x");
+	Swizzles[1] = TEXT("y");
+	Swizzles[2] = TEXT("z");
+	Swizzles[3] = TEXT("w");
+
+	for (int32 ItemIndex = 0; ItemIndex < GBS_Num; ItemIndex++)
+	{
+		const FGBufferItem& Item = BufferInfo.Slots[ItemIndex];
+
+		// hack for now
+		if (Item.BufferSlot == GBS_SceneColor)
+		{
+			continue;
+		}
+
+		if (Item.bIsValid)
+		{
+			FString SlotName = GetSlotTextName(Item.BufferSlot);
+
+			FString DataName;
+			if (bNeedsConversion[ItemIndex])
+			{
+				DataName = FString::Printf(TEXT("%s_Compressed"), SlotName.GetCharArray().GetData());
+			}
+			else
+			{
+				DataName = FString::Printf(TEXT("GBuffer.%s"), SlotName.GetCharArray().GetData());
+			}
+
+/*
+			if (Item.BufferSlot == GBS_BaseColor)
+			{
+				DataName = TEXT("float4(0.25,.25,.25,.25)");
+			}
+*/
+			for (int32 PackIter = 0; PackIter < FGBufferItem::MaxPacking; PackIter++)
+			{
+				const FGBufferPacking& PackItem = Item.Packing[PackIter];
+				if (PackItem.bIsValid)
+				{
+
+					FString FullName = FString::Printf(TEXT("%s.%s"),
+						DataName.GetCharArray().GetData(),
+						Swizzles[PackItem.SrcChannelIndex].GetCharArray().GetData());
+
+
+					if (PackItem.bFull)
+					{
+						// if it's a full channel, just pass it in
+						FString TargetName = FString::Printf(TEXT("MrtFloat%d"), PackItem.TargetIndex);
+
+						// If we are full, add a quantization term based on the output GBuffer size
+						if (Item.bQuantizationBias)
+						{
+							int DstNumBits = GetGBufferCompressionChanBitSize(Item.Compression,PackItem.DstChannelIndex);
+
+							FullName += FString::Printf(TEXT(" + QuantizationBias / float(%d)"), 1 << DstNumBits);
+						}
+
+						// how many bits are we packing?
+						FString CurrLine = FString::Printf(TEXT("\t%s.%s = %s;\n"),
+							TargetName.GetCharArray().GetData(),
+							Swizzles[PackItem.DstChannelIndex].GetCharArray().GetData(),
+							FullName.GetCharArray().GetData());
+
+						FullStr += CurrLine;
+					}
+					else
+					{
+						check(PackItem.BitNum >= 1);
+						check(PackItem.BitNum <= 31);
+
+						int32 Mask = (1 << PackItem.BitNum)-1;
+
+						// if it's a full channel, just pass it in
+						FString TargetName = FString::Printf(TEXT("MrtUint%d"), PackItem.TargetIndex);
+
+						FString CurrLine = FString::Printf(TEXT("\t%s.%s |= ((((%s) >> %d) & 0x%02x) << %d);\n"),
+							TargetName.GetCharArray().GetData(),
+							Swizzles[PackItem.DstChannelIndex].GetCharArray().GetData(),
+							FullName.GetCharArray().GetData(),
+							PackItem.SrcBitStart,
+							Mask,
+							PackItem.DstBitStart);
+
+						FullStr += CurrLine;
+					}
+				}
+			}
+		}
+	}
+
+
+	FullStr += TEXT("\n");
+
+	for (int32 I = 0; I < FGBufferInfo::MaxTargets; I++)
+	{
+		// skip first target for now
+		if (I == 0)
+		{
+			continue;
+		}
+
+		int32 ChanNum = GetTargetNumChannels(BufferInfo.Targets[I].TargetType);
+
+		// optimization for readability
+		bool bIsAnyInt = false;
+		bool bIsAnyNonZero = false;
+		for (int32 Chan = 0; Chan < ChanNum; Chan++)
+		{
+			if (PackedStatus[I][Chan] == -1)
+			{
+				// if it's -1, it's unused (neither int nor float) and just set to 0.0f
+			}
+			else
+			{
+				bIsAnyNonZero = true;
+				// 0 means float, 1 means int
+				if (PackedStatus[I][Chan] == 1)
+				{
+					bIsAnyInt = true;
+				}
+			}
+		}
+
+		if (bIsAnyNonZero && !bIsAnyInt && ChanNum == 4)
+		{
+			// if everything is a float we can just copy the float without a constructor.
+			FString CurrLine = FString::Printf(TEXT("\tOut.MRT[%d] = MrtFloat%d;\n"), I, I);
+			FullStr += CurrLine;
+		}
+		else
+		{
+			FString CurrLine = FString::Printf(TEXT("\tOut.MRT[%d] = float4("), I);
+
+			for (int32 Chan = 0; Chan < 4; Chan++)
+			{
+				if (PackedStatus[I][Chan] == -1)
+				{
+					CurrLine += "0.0f";
+				}
+				else
+				{
+					if (PackedStatus[I][Chan] == 0)
+					{
+						// raw float
+						FString TargetName = FString::Printf(TEXT("MrtFloat%d"), I);
+						CurrLine += FString::Printf(TEXT("%s.%s"),
+							TargetName.GetCharArray().GetData(),
+							Swizzles[Chan].GetCharArray().GetData());
+					}
+					else
+					{
+						// packed int, convert to float
+						FString TargetName = FString::Printf(TEXT("MrtUint%d"), I);
+						int32 ChanBits = GetBufferNumBits(BufferInfo.Targets[I].TargetType, Chan);
+
+						// the 0.5 is to put it in the center of the int range, avoid roundoff error weirdness
+						CurrLine += FString::Printf(TEXT("(float(%s.%s) + .5f) / %u.0f"),
+							TargetName.GetCharArray().GetData(),
+							Swizzles[Chan].GetCharArray().GetData(),
+							(1 << ChanBits) - 1);
+					}
+				}
+				if (Chan < 3)
+				{
+					CurrLine += TEXT(", ");
+				}
+			}
+
+			CurrLine += TEXT(");\n");
+			FullStr += CurrLine;
+		}
+
+	}
+
+	FullStr += TEXT("}\n");
+	FullStr += TEXT("\n");
+
+	return FullStr;
+}
+
+static FString CreateGBufferDecodeFunctionDirect(const FGBufferInfo& BufferInfo)
+{
+	FString FullStr;
+
+
+	FullStr += TEXT("FGBufferData  DecodeGBufferDataDirect(");
+	bool bFirst = true;
+	for (int32 Index = 0; Index < FGBufferInfo::MaxTargets; Index++)
+	{
+		const EGBufferType Target = BufferInfo.Targets[Index].TargetType;
+
+		if (Target != GBT_Invalid && Index != 0)
+		{
+			if (bFirst)
+			{
+				bFirst = false;
+			}
+			else
+			{
+				FullStr += TEXT(",\n\t");
+			}
+
+			int32 NumChan = GetTargetNumChannels(Target);
+			FString TypeName = GetFloatType(NumChan);
+			FString CurrLine = FString::Printf(TEXT("%s InMRT%d"),
+				TypeName.GetCharArray().GetData(),
+				Index);
+
+			FullStr += CurrLine;
+		}
+	}
+	FullStr += TEXT(",\n\tfloat CustomNativeDepth");
+	FullStr += TEXT(",\n\tuint CustomStencil");
+	FullStr += TEXT(",\n\tfloat SceneDepth");
+	FullStr += TEXT(",\n\tbool bGetNormalizedNormal");
+	FullStr += TEXT(",\n\tbool bChecker)\n");
+
+	FullStr += TEXT("{\n");
+
+	FullStr += TEXT("\tFGBufferData Ret = (FGBufferData)0;\n");
+	
+	FString Swizzles[4];
+	Swizzles[0] = TEXT("x");
+	Swizzles[1] = TEXT("y");
+	Swizzles[2] = TEXT("z");
+	Swizzles[3] = TEXT("w");
+
+	for (int32 SlotIndex = 0; SlotIndex < GBS_Num; SlotIndex++)
+	{
+		const FGBufferItem& SlotItem = BufferInfo.Slots[SlotIndex];
+		if (SlotItem.bIsValid && SlotItem.BufferSlot != GBS_SceneColor)
+		{
+			int32 DstChanNum = GetCompressionDstNumChannels(SlotItem.Compression);
+
+			bool bIsPackedBits = IsPackedBits(SlotItem.Compression);
+			bool bNeedsConversion = GetCompressionNeedsConversion(SlotItem.Compression) != 0;
+
+			FString DstName = GetSlotTextName(SlotItem.BufferSlot);
+			if (bNeedsConversion)
+			{
+				DstName = DstName + TEXT("_Compressed");
+
+				// also, if it is converted, we need to declare the temporary value.
+				FString TypeName;
+				if (bIsPackedBits)
+				{
+					TypeName = GetUintType(DstChanNum);
+				}
+				else
+				{
+					TypeName = GetFloatType(DstChanNum);
+				}
+
+				FullStr += FString::Printf(TEXT("\t%s %s = 0.0f;\n"),
+							TypeName.GetCharArray().GetData(),
+							DstName.GetCharArray().GetData());
+			}
+			else
+			{
+				DstName = TEXT("Ret.") + DstName;
+			}
+
+			for (int32 ChanIndex = 0; ChanIndex < DstChanNum; ChanIndex++)
+			{
+				// first, count how many packing items we have
+				int32 NumFound = 0;
+
+				check(FGBufferItem::MaxPacking == 8);
+				int32 FoundIndexList[8] = { -1, -1, -1, -1, -1, -1, -1, -1 };
+
+				for (int32 PackIter = 0; PackIter < FGBufferItem::MaxPacking; PackIter++)
+				{
+					const FGBufferPacking& PackItem = SlotItem.Packing[PackIter];
+					if (PackItem.bIsValid && PackItem.SrcChannelIndex == ChanIndex)
+					{
+						FoundIndexList[NumFound] = PackIter;
+						NumFound++;
+					}
+				}
+
+				if (NumFound == 0)
+				{
+					// channel is missing
+					check(0);
+				}
+				else if (NumFound == 1)
+				{
+					check(FoundIndexList[0] >= 0);
+					int32 FoundIndex = FoundIndexList[0];
+
+					const FGBufferPacking& PackItem = SlotItem.Packing[FoundIndex];
+
+					if (bIsPackedBits)
+					{
+						int32 Mask = (1 << PackItem.BitNum)-1;
+						
+						const EGBufferType TargetType = BufferInfo.Targets[PackItem.TargetIndex].TargetType;
+						int32 TargetBitNum = GetBufferNumBits(TargetType,PackItem.DstChannelIndex);
+						int32 TargetScale = (1 << TargetBitNum) - 1;
+
+						FString CurrLine = FString::Printf(TEXT("\t%s.%s = (((uint((float(InMRT%d.%s) * %d.0f) + .5f) >> %d) & 0x%02x) << %d);\n"),
+							DstName.GetCharArray().GetData(),
+							Swizzles[PackItem.SrcChannelIndex].GetCharArray().GetData(),
+							PackItem.TargetIndex,
+							Swizzles[PackItem.DstChannelIndex].GetCharArray().GetData(),
+							TargetScale,
+							PackItem.DstBitStart,
+							Mask,
+							PackItem.SrcBitStart);
+						
+						FullStr += CurrLine;
+					}
+					else
+					{
+						// simple, just copy from float to float
+						FString CurrLine = FString::Printf(TEXT("\t%s.%s = InMRT%d.%s;\n"),
+							DstName.GetCharArray().GetData(),
+							Swizzles[PackItem.SrcChannelIndex].GetCharArray().GetData(),
+							PackItem.TargetIndex,
+							Swizzles[PackItem.DstChannelIndex].GetCharArray().GetData());
+						FullStr += CurrLine;
+					}
+				}
+				else
+				{
+					// reconstruct bits from multiple places, not implemented yet
+					check(FoundIndexList[0] >= 0);
+					check(bIsPackedBits);
+
+					for (int32 FoundIter = 0; FoundIter < NumFound; FoundIter++)
+					{
+						int32 FoundIndex = FoundIndexList[FoundIter];
+
+						const FGBufferPacking& PackItem = SlotItem.Packing[FoundIndex];
+
+						int32 Mask = (1 << PackItem.BitNum) - 1;
+
+						const EGBufferType TargetType = BufferInfo.Targets[PackItem.TargetIndex].TargetType;
+						int32 TargetBitNum = GetBufferNumBits(TargetType, PackItem.DstChannelIndex);
+						int32 TargetScale = (1 << TargetBitNum) - 1;
+
+						FString AssignmentOperator = (FoundIter == 0) ? TEXT("=") : TEXT("|=");
+
+						FString CurrLine = FString::Printf(TEXT("\t%s.%s %s (((uint((float(InMRT%d.%s) * %d.0f) + 0.5f) >> %d) & 0x%02x) << %d);\n"),
+							DstName.GetCharArray().GetData(),
+							Swizzles[PackItem.SrcChannelIndex].GetCharArray().GetData(),
+							AssignmentOperator.GetCharArray().GetData(),
+							PackItem.TargetIndex,
+							Swizzles[PackItem.DstChannelIndex].GetCharArray().GetData(),
+							TargetScale,
+							PackItem.DstBitStart,
+							Mask,
+							PackItem.SrcBitStart);
+
+						FullStr += CurrLine;
+					}
+				}
+			}
+		}
+	}
+
+	FullStr += TEXT("\t\n");
+
+	// check for anything that needs an extra decode step
+	for (int32 SlotIndex = 0; SlotIndex < GBS_Num; SlotIndex++)
+	{
+		const FGBufferItem& SlotItem = BufferInfo.Slots[SlotIndex];
+		if (SlotItem.bIsValid && SlotItem.BufferSlot != GBS_SceneColor)
+		{
+			int32 DstChanNum = GetCompressionDstNumChannels(SlotItem.Compression);
+
+			bool bIsPackedBits = IsPackedBits(SlotItem.Compression);
+			bool bNeedsConversion = GetCompressionNeedsConversion(SlotItem.Compression) != 0;
+
+			FString DstName = GetSlotTextName(SlotItem.BufferSlot);
+			if (bNeedsConversion)
+			{
+				FString CompressName = DstName + TEXT("_Compressed");
+				FString FuncName = GetCompressionDecodeFuncName(SlotItem.Compression);
+
+				// also, if it is converted, we need to declare the temporary value.
+				FString TypeName = GetFloatType(DstChanNum);
+				FullStr += FString::Printf(TEXT("\tRet.%s = %s(%s);\n"),
+					DstName.GetCharArray().GetData(),
+					FuncName.GetCharArray().GetData(),
+					CompressName.GetCharArray().GetData());
+			}
+		}
+	}
+
+
+	FullStr += TEXT("\n");
+	FullStr += TEXT("\tGBufferPostDecode(Ret,bChecker,bGetNormalizedNormal);\n");
+	FullStr += TEXT("\n");
+
+	//// Hacks for now, fixme later
+	FullStr += TEXT("\tRet.CustomDepth = ConvertFromDeviceZ(CustomNativeDepth);\n");
+	FullStr += TEXT("\tRet.CustomStencil = CustomStencil;\n");
+	FullStr += TEXT("\tRet.Depth = SceneDepth;\n");
+	FullStr += TEXT("\t\n");
+
+	FullStr += TEXT("\n");
+	FullStr += TEXT("\treturn Ret;\n");
+
+	FullStr += TEXT("}\n");
+	FullStr += TEXT("\n");
+
+	return FullStr;
+}
+
+
+enum EGBufferDecodeType
+{
+	CoordUV			 , // GetGBufferData
+	CoordUInt		 , // GetGBufferDataUint
+	SceneTextures,    // GetGBufferDataFromSceneTextures
+	SceneTexturesLoad, // GetGBufferDataFromSceneTexturesLoad
+	Num
+};
+
+struct FGBufferDecoderSyntax
+{
+	const TCHAR * Suffix;
+	const TCHAR * CoordType;
+	const TCHAR * CoordName;
+	const TCHAR * CheckerFunc;
+};
+
+static FGBufferDecoderSyntax GDecoderSyntax[EGBufferDecodeType::Num] =
+{
+	{ TEXT("UV"), TEXT("float2"), TEXT("UV"), TEXT("CheckerFromSceneColorUV") },
+	{ TEXT("Uint"), TEXT("uint2"), TEXT("PixelPos"), TEXT("CheckerFromPixelPos") },
+	{ TEXT("SceneTextures"), TEXT("float2"), TEXT("UV"), TEXT("CheckerFromSceneColorUV") },
+	{ TEXT("SceneTexturesLoad"), TEXT("uint2"), TEXT("PixelCoord"), TEXT("CheckerFromPixelPos") },
+};
+
+static FString CreateGBufferDecodeFunctionVariation(const FGBufferInfo& BufferInfo, EGBufferDecodeType DecodeType, ERHIFeatureLevel::Type FEATURE_LEVEL)
+{
+	FString FullStr;
+
+	FString Suffix = GDecoderSyntax[DecodeType].Suffix;
+	FString CoordType = GDecoderSyntax[DecodeType].CoordType;
+	FString CoordName = GDecoderSyntax[DecodeType].CoordName;
+	FString CheckerFunc = GDecoderSyntax[DecodeType].CheckerFunc;
+
+
+	FullStr += TEXT("// @param PixelPos relative to left top of the rendertarget (not viewport)\n");
+
+	FullStr += FString::Printf(TEXT("FGBufferData DecodeGBufferData%s(%s %s, bool bGetNormalizedNormal = true)\n"),
+		Suffix.GetCharArray().GetData(),
+		CoordType.GetCharArray().GetData(),
+		CoordName.GetCharArray().GetData());
+
+	FullStr += TEXT("{\n");
+
+
+	if (DecodeType == CoordUV)
+	{
+		FullStr += FString::Printf(TEXT("\tfloat CustomNativeDepth = Texture2DSampleLevel(SceneTexturesStruct.CustomDepthTexture, SceneTexturesStruct_CustomDepthTextureSampler, %s, 0).r;\n"), CoordName.GetCharArray().GetData());
+
+		if (FEATURE_LEVEL >= ERHIFeatureLevel::SM5)
+		{
+			FullStr += FString::Printf(TEXT("\tint2 IntUV = (int2)trunc(%s * View.BufferSizeAndInvSize.xy);\n"), CoordName.GetCharArray().GetData());
+			FullStr += TEXT("\tuint CustomStencil = SceneTexturesStruct.CustomStencilTexture.Load(int3(IntUV, 0)) STENCIL_COMPONENT_SWIZZLE;\n");
+		}
+		else
+		{
+			FullStr += TEXT("\tuint CustomStencil = 0;\n");
+		}
+
+		FullStr += FString::Printf(TEXT("\tfloat SceneDepth = CalcSceneDepth(%s);\n"), CoordName.GetCharArray().GetData());
+	}
+	else if (DecodeType == CoordUInt)
+	{
+		FullStr += FString::Printf(TEXT("\tfloat CustomNativeDepth = SceneTexturesStruct.CustomDepthTexture.Load(int3(%s, 0)).r;\n"), CoordName.GetCharArray().GetData());
+		FullStr += FString::Printf(TEXT("\tuint CustomStencil = SceneTexturesStruct.CustomStencilTexture.Load(int3(%s, 0)) STENCIL_COMPONENT_SWIZZLE;\n"), CoordName.GetCharArray().GetData());
+		FullStr += FString::Printf(TEXT("\tfloat SceneDepth = CalcSceneDepth(%s);\n"), CoordName.GetCharArray().GetData());
+	}
+	else if (DecodeType == SceneTextures)
+	{
+		FullStr += TEXT("\tuint CustomStencil = 0;\n");
+		FullStr += TEXT("\tfloat CustomNativeDepth = 0;\n");
+		FullStr += FString::Printf(TEXT("\tfloat DeviceZ = SampleDeviceZFromSceneTexturesTempCopy(%s);\n"), CoordName.GetCharArray().GetData());
+		FullStr += TEXT("\tfloat SceneDepth = ConvertFromDeviceZ(DeviceZ);\n");
+	}
+	else if (DecodeType == SceneTexturesLoad)
+	{
+		FullStr += TEXT("\tuint CustomStencil = 0;\n");
+		FullStr += TEXT("\tfloat CustomNativeDepth = 0;\n");
+		FullStr += FString::Printf(TEXT("\tfloat DeviceZ = SceneDepthTexture.Load(int3(%s, 0)).r;\n"), CoordName.GetCharArray().GetData());
+		FullStr += TEXT("\tfloat SceneDepth = ConvertFromDeviceZ(DeviceZ);\n");
+	}
+	else
+	{
+		check(0);
+	}
+
+	FullStr += TEXT("\n");
+
+	FString FullSwizzle[4];
+	FullSwizzle[0] = TEXT("x");
+	FullSwizzle[1] = TEXT("xy");
+	FullSwizzle[2] = TEXT("xyz");
+	FullSwizzle[3] = TEXT("xyzw");
+
+	const bool bAllowCustomDepthStencil = ((DecodeType == EGBufferDecodeType::CoordUInt && FEATURE_LEVEL >= ERHIFeatureLevel::SM5) || DecodeType == EGBufferDecodeType::CoordUV);
+
+	for (int32 Index = 0; Index < FGBufferInfo::MaxTargets; Index++)
+	{
+		const EGBufferType Target = BufferInfo.Targets[Index].TargetType;
+
+		if (Target != GBT_Invalid && Index != 0)
+		{
+			int32 NumChan = GetTargetNumChannels(Target);
+			FString TypeName = GetFloatType(NumChan);
+			FString TargetName = BufferInfo.Targets[Index].TargetName;
+
+			// special exception for WRITES_VELOCITY_TO_GBUFFER
+			if (TargetName == TEXT("Velocity"))
+			{
+				if (DecodeType == CoordUV || DecodeType == CoordUInt)
+				{
+					TargetName = TEXT("GBufferVelocity");
+				}
+				else if (DecodeType == SceneTextures || DecodeType == SceneTexturesLoad)
+				{
+					TargetName = TEXT("VelocityBuffer");
+				}
+			}
+
+			FString CurrLine;
+			if (DecodeType == CoordUV)
+			{
+				CurrLine = FString::Printf(TEXT("\t%s InMRT%d = Texture2DSampleLevel(SceneTexturesStruct.%sTexture, SceneTexturesStruct_%sTextureSampler, %s, 0).%s;\n"),
+					TypeName.GetCharArray().GetData(),
+					Index,
+					TargetName.GetCharArray().GetData(),
+					TargetName.GetCharArray().GetData(),
+					CoordName.GetCharArray().GetData(),
+					FullSwizzle[NumChan-1].GetCharArray().GetData());
+			}
+			else if (DecodeType == CoordUInt)
+			{
+				CurrLine = FString::Printf(TEXT("\t%s InMRT%d = SceneTexturesStruct.%sTexture.Load(int3(%s, 0)).%s;\n"),
+					TypeName.GetCharArray().GetData(),
+					Index,
+					TargetName.GetCharArray().GetData(),
+					CoordName.GetCharArray().GetData(),
+					FullSwizzle[NumChan-1].GetCharArray().GetData());
+			}
+			else if (DecodeType == SceneTextures)
+			{
+				CurrLine = FString::Printf(TEXT("\t%s InMRT%d = %sTexture.SampleLevel(%sTextureSampler, %s, 0).%s;\n"),
+					TypeName.GetCharArray().GetData(),
+					Index,
+					TargetName.GetCharArray().GetData(),
+					TargetName.GetCharArray().GetData(),
+					CoordName.GetCharArray().GetData(),
+					FullSwizzle[NumChan-1].GetCharArray().GetData());
+			}
+			else if (DecodeType == SceneTexturesLoad)
+			{
+				CurrLine = FString::Printf(TEXT("\t%s InMRT%d = %sTexture.Load(int3(%s, 0)).%s;\n"),
+					TypeName.GetCharArray().GetData(),
+					Index,
+					TargetName.GetCharArray().GetData(),
+					CoordName.GetCharArray().GetData(),
+					FullSwizzle[NumChan-1].GetCharArray().GetData());
+			}
+			else
+			{
+				check(0);
+			}
+
+			FullStr += CurrLine;
+		}
+	}
+
+	FullStr += TEXT("\n");
+
+	FullStr += TEXT("\tFGBufferData Ret = DecodeGBufferDataDirect(");
+	bool bFirst = true;
+	for (int32 Index = 0; Index < FGBufferInfo::MaxTargets; Index++)
+	{
+		const EGBufferType Target = BufferInfo.Targets[Index].TargetType;
+
+		if (Target != GBT_Invalid && Index != 0)
+		{
+			if (bFirst)
+			{
+				bFirst = false;
+			}
+			else
+			{
+				FullStr += TEXT(",\n\t\t");
+			}
+
+			int32 NumChan = GetTargetNumChannels(Target);
+			FString CurrLine = FString::Printf(TEXT("InMRT%d"),
+				Index);
+
+			FullStr += CurrLine;
+		}
+	}
+	FullStr += TEXT(",\n\t\tCustomNativeDepth");
+	FullStr += TEXT(",\n\t\tCustomStencil");
+	FullStr += TEXT(",\n\t\tSceneDepth");
+	FullStr += TEXT(",\n\t\tbGetNormalizedNormal");
+	FullStr += FString::Printf(TEXT(",\n\t\t%s(%s));\n"),
+					CheckerFunc.GetCharArray().GetData(),
+					CoordName.GetCharArray().GetData());
+
+	FullStr += TEXT("\n");
+	FullStr += TEXT("\treturn Ret;\n");
+
+	FullStr += TEXT("}\n");
+	FullStr += TEXT("\n");
+
+	return FullStr;
+}
+
+static FCriticalSection GCriticalSection;
+
+static FString GetAutoGenDirectory(EShaderPlatform TargetPlatform)
+{
+	FString PlatformName = LegacyShaderPlatformToShaderFormat(TargetPlatform).GetPlainNameString();
+	FString AutogenHeaderDirectory = FPaths::ProjectIntermediateDir() / TEXT("ShaderAutogen") / PlatformName;
+	return AutogenHeaderDirectory;
+}
+
+
+
+static void SetSharedGBufferSlots(bool Slots[])
+{
+	Slots[GBS_SceneColor] = true;
+	Slots[GBS_WorldNormal] = true;
+	Slots[GBS_PerObjectGBufferData] = true;
+	Slots[GBS_Metallic] = true;
+	Slots[GBS_Specular] = true;
+	Slots[GBS_Roughness] = true;
+	Slots[GBS_ShadingModelId] = true;
+	Slots[GBS_SelectiveOutputMask] = true;
+	Slots[GBS_BaseColor] = true;
+	Slots[GBS_GenericAO] = true;
+	Slots[GBS_AO] = false;// true; // figure this out later
+}
+
+// If we are storing writing to the GBuffer, what slots does this shading model write? Ignores the other forward-shading effects
+// like the dual blending of ThinTranslucent or fogging.
+// If we are going one GBuffer for all shader model formats, then we treat all custom data as a single RGBA8 value (bMergeCustom=true). But
+// if each shading model has a unique GBuffer format, we set it to false to pack each format as tightly as possible.
+static void SetSlotsForShadingModelType(bool Slots[], EMaterialShadingModel ShadingModel, bool bMergeCustom)
+{
+	switch (ShadingModel)
+	{
+	case MSM_Unlit:
+		Slots[GBS_SceneColor] = true;
+		break;
+	case MSM_DefaultLit:
+		SetSharedGBufferSlots(Slots);
+		break;
+	case MSM_Subsurface:
+		SetSharedGBufferSlots(Slots);
+		if (bMergeCustom)
+		{
+			Slots[GBS_CustomData] = true;
+		}
+		else
+		{
+			Slots[GBS_SubsurfaceColor] = true;
+			Slots[GBS_Opacity] = true;
+		}
+		break;
+	case MSM_PreintegratedSkin:
+		SetSharedGBufferSlots(Slots);
+		if (bMergeCustom)
+		{
+			Slots[GBS_CustomData] = true;
+		}
+		else
+		{
+			Slots[GBS_SubsurfaceColor] = true;
+			Slots[GBS_Opacity] = true;
+		}
+		break;
+	case MSM_ClearCoat:
+		SetSharedGBufferSlots(Slots);
+		if (bMergeCustom)
+		{
+			Slots[GBS_CustomData] = true;
+		}
+		else
+		{
+			Slots[GBS_ClearCoat] = true;
+			Slots[GBS_ClearCoatRoughness] = true;
+		}
+		break;
+	case MSM_SubsurfaceProfile:
+		SetSharedGBufferSlots(Slots);
+		if (bMergeCustom)
+		{
+			Slots[GBS_CustomData] = true;
+		}
+		else
+		{
+			Slots[GBS_SubsurfaceProfile] = true;
+			Slots[GBS_Opacity] = true;
+		}
+		break;
+	case MSM_TwoSidedFoliage:
+		SetSharedGBufferSlots(Slots);
+		if (bMergeCustom)
+		{
+			Slots[GBS_CustomData] = true;
+		}
+		else
+		{
+			Slots[GBS_SubsurfaceColor] = true;
+			Slots[GBS_Opacity] = true;
+		}
+		break;
+	case MSM_Hair:
+		SetSharedGBufferSlots(Slots);
+		if (bMergeCustom)
+		{
+			Slots[GBS_CustomData] = true;
+		}
+		else
+		{
+			Slots[GBS_HairSecondaryWorldNormal] = true;
+			Slots[GBS_HairBacklit] = true;
+		}
+		break;
+	case MSM_Cloth:
+		SetSharedGBufferSlots(Slots);
+		if (bMergeCustom)
+		{
+			Slots[GBS_CustomData] = true;
+		}
+		else
+		{
+			Slots[GBS_SubsurfaceColor] = true;
+			Slots[GBS_Cloth] = true;
+		}
+		break;
+	case MSM_Eye:
+		SetSharedGBufferSlots(Slots);
+		if (bMergeCustom)
+		{
+			Slots[GBS_CustomData] = true;
+		}
+		else
+		{
+			Slots[GBS_SubsurfaceProfile] = true;
+			Slots[GBS_IrisNormal] = true;
+			Slots[GBS_Opacity] = true;
+		}
+		break;
+	case MSM_SingleLayerWater:
+		SetSharedGBufferSlots(Slots);
+		break;
+	case MSM_ThinTranslucent:
+		// thin translucent doesn't write to the GBuffer
+		break;
+	}
+}
+
+
+
+static void SetStandardGBufferSlots(bool Slots[], bool bWriteEmissive, bool bHasTangent, bool bHasVelocity, bool bHasStaticLighting)
+{
+	Slots[GBS_SceneColor] = bWriteEmissive;
+	Slots[GBS_WorldNormal] = true;
+	Slots[GBS_PerObjectGBufferData] = true;
+	Slots[GBS_Metallic] = true;
+	Slots[GBS_Specular] = true;
+	Slots[GBS_Roughness] = true;
+	Slots[GBS_ShadingModelId] = true;
+	Slots[GBS_SelectiveOutputMask] = true;
+	Slots[GBS_BaseColor] = true;
+	Slots[GBS_GenericAO] = true;
+	Slots[GBS_AO] = false;// true;
+	Slots[GBS_Velocity] = bHasVelocity;
+	Slots[GBS_PrecomputedShadowFactor] = bHasStaticLighting;
+	Slots[GBS_WorldTangent] = bHasTangent;
+	Slots[GBS_Anisotropy] = bHasTangent;
+}
+
+static void DetermineUsedMaterialSlots(
+	bool Slots[],
+	const FShaderMaterialDerivedDefines& Dst,
+	const FShaderMaterialPropertyDefines& Mat,
+	const FShaderLightmapPropertyDefines& Lightmap,
+	const FShaderGlobalDefines& SrcGlobal,
+	const FShaderCompilerDefines& Compiler,
+	ERHIFeatureLevel::Type FEATURE_LEVEL)
+{
+	bool bWriteEmissive = Dst.NEEDS_BASEPASS_VERTEX_FOGGING || Mat.USES_EMISSIVE_COLOR || SrcGlobal.ALLOW_STATIC_LIGHTING || Mat.MATERIAL_SHADINGMODEL_SINGLELAYERWATER;
+	bool bHasTangent = SrcGlobal.GBUFFER_HAS_TANGENT;
+	bool bHasVelocity = Dst.WRITES_VELOCITY_TO_GBUFFER;
+	bool bHasStaticLighting = Dst.GBUFFER_HAS_PRECSHADOWFACTOR || Dst.WRITES_PRECSHADOWFACTOR_TO_GBUFFER;
+
+	// we have to use if statements, not switch or if/else statements because we can have multiple shader model ids.
+	if (Mat.MATERIAL_SHADINGMODEL_UNLIT)
+	{
+		Slots[GBS_SceneColor] = true;
+	}
+
+	if (Mat.MATERIAL_SHADINGMODEL_DEFAULT_LIT)
+	{
+		SetStandardGBufferSlots(Slots, bWriteEmissive, bHasTangent, bHasVelocity, bHasStaticLighting);
+	}
+
+	if (Mat.MATERIAL_SHADINGMODEL_SUBSURFACE)
+	{
+		SetStandardGBufferSlots(Slots, bWriteEmissive, bHasTangent, bHasVelocity, bHasStaticLighting);
+		Slots[GBS_CustomData] = true;
+	}
+
+	if (Mat.MATERIAL_SHADINGMODEL_PREINTEGRATED_SKIN)
+	{
+		SetStandardGBufferSlots(Slots, bWriteEmissive, bHasTangent, bHasVelocity, bHasStaticLighting);
+		Slots[GBS_CustomData] = true;
+	}
+
+	if (Mat.MATERIAL_SHADINGMODEL_SUBSURFACE_PROFILE)
+	{
+		SetStandardGBufferSlots(Slots, bWriteEmissive, bHasTangent, bHasVelocity, bHasStaticLighting);
+		Slots[GBS_CustomData] = true;
+	}
+
+	if (Mat.MATERIAL_SHADINGMODEL_CLEAR_COAT)
+	{
+		SetStandardGBufferSlots(Slots, bWriteEmissive, bHasTangent, bHasVelocity, bHasStaticLighting);
+		Slots[GBS_CustomData] = true;
+	}
+
+	if (Mat.MATERIAL_SHADINGMODEL_TWOSIDED_FOLIAGE)
+	{
+		SetStandardGBufferSlots(Slots, bWriteEmissive, bHasTangent, bHasVelocity, bHasStaticLighting);
+		Slots[GBS_CustomData] = true;
+	}
+
+	if (Mat.MATERIAL_SHADINGMODEL_HAIR)
+	{
+		SetStandardGBufferSlots(Slots, bWriteEmissive, bHasTangent, bHasVelocity, bHasStaticLighting);
+	}
+
+	if (Mat.MATERIAL_SHADINGMODEL_CLOTH)
+	{
+		SetStandardGBufferSlots(Slots, bWriteEmissive, bHasTangent, bHasVelocity, bHasStaticLighting);
+	}
+
+	if (Mat.MATERIAL_SHADINGMODEL_EYE)
+	{
+		SetStandardGBufferSlots(Slots, bWriteEmissive, bHasTangent, bHasVelocity, bHasStaticLighting);
+		Slots[GBS_CustomData] = true;
+	}
+
+	if (Mat.MATERIAL_SHADINGMODEL_SINGLELAYERWATER)
+	{
+		// single layer water uses standard slots
+		SetStandardGBufferSlots(Slots, bWriteEmissive, bHasTangent, bHasVelocity, bHasStaticLighting);
+	}
+
+	// doesn't write to GBuffer
+	if (Mat.MATERIAL_SHADINGMODEL_THIN_TRANSLUCENT)
+	{
+	}
+
+}
+
+void FShaderCompileUtilities::ApplyDerivedDefines(FShaderCompilerEnvironment& OutEnvironment, FShaderCompilerEnvironment * SharedEnvironment, const EShaderPlatform Platform)
+{
+	ERHIFeatureLevel::Type FeatureLevel = GetMaxSupportedFeatureLevel(Platform);
+
+	FShaderMaterialPropertyDefines MaterialDefines = {};
+	FShaderLightmapPropertyDefines LightmapDefines = {};
+	FShaderGlobalDefines GlobalDefines = {};
+	FShaderCompilerDefines CompilerDefines = {};
+	FShaderMaterialDerivedDefines DerivedDefines = {};
+
+	ApplyFetchEnvironment(MaterialDefines,OutEnvironment);
+	ApplyFetchEnvironment(LightmapDefines,OutEnvironment);
+	ApplyFetchEnvironment(GlobalDefines,OutEnvironment);
+	ApplyFetchEnvironment(CompilerDefines,OutEnvironment);
+
+	if (SharedEnvironment != nullptr)
+	{
+		ApplyFetchEnvironment(MaterialDefines, *SharedEnvironment);
+		ApplyFetchEnvironment(LightmapDefines, *SharedEnvironment);
+		ApplyFetchEnvironment(GlobalDefines, *SharedEnvironment);
+		ApplyFetchEnvironment(CompilerDefines, *SharedEnvironment);
+	}
+
+	// for reference, here are the existing defines
+	DerivedDefines = CalculateDerivedMaterialParameters(MaterialDefines, LightmapDefines, GlobalDefines, CompilerDefines, FeatureLevel);
+
+	// We will need to match these. For existing code, values such as PIXELSHADEROUTPUT_MRT5 are based on a collection of #ifdefs. First, we 
+	// need to define those here instead of doing that logic inside the shader code macros. Then as a second step, we need to enable/define
+	// the targets based on the FGBufferInfo struct.
+
+	// for now however, we just only set the global enable flag while working on this. Keeping this as a static so it can be changed in debug mode for testing,
+	// but not making it a cvar because the old path will be removed once it is confirmed to work.
+	static bool bUseRefactor = true;
+	OutEnvironment.SetDefine(TEXT("GBUFFER_REFACTOR"), bUseRefactor ? TEXT("1") : TEXT("0"));
+
+	FGBufferParams Params = FShaderCompileUtilities::FetchGBufferParamsRuntime(Platform);
+	FGBufferInfo BufferInfo = FetchFullGBufferInfo(Params);
+
+	bool bTargetUsage[FGBufferInfo::MaxTargets] = {};
+
+	if (MaterialDefines.IS_BASE_PASS)
+	{
+		// if we are using a gbuffer, and this is the base pass that writes a gbuffer, search the gbuffer for each slot
+		if (DerivedDefines.USES_GBUFFER)
+		{
+			bTargetUsage[0] = true;
+			bool Slots[GBS_Num] = {};
+
+			DetermineUsedMaterialSlots(Slots, DerivedDefines, MaterialDefines, LightmapDefines, GlobalDefines, CompilerDefines, FeatureLevel);
+
+			// for each slot, which MRT index does it use? -1 if unused.
+			int32 SlotTargets[GBS_Num];
+			for (int32 Index = 0; Index < GBS_Num; Index++)
+			{
+				SlotTargets[Index] = -1;
+			}
+
+			for (int32 Index = 0; Index < GBS_Num; Index++)
+			{
+				// if we are using this slot
+				if (Slots[Index])
+				{
+					// if we are using this slot, it must have a valid spot in our gbuffer
+					const FGBufferItem& Item = BufferInfo.Slots[Index];
+					check(Item.bIsValid);
+
+					for (int32 PackIter = 0; PackIter < FGBufferItem::MaxPacking; PackIter++)
+					{
+						const FGBufferPacking& Packing = Item.Packing[PackIter];
+						if (Packing.bIsValid)
+						{
+							check(Packing.TargetIndex >= 0);
+							bTargetUsage[Packing.TargetIndex] = true;
+						}
+					}
+				}
+			}
+
+		}
+		else
+		{
+			bTargetUsage[0] = true;
+			// we also need MRT for thin translucency due to dual blending if we are not on the fallback path
+			bTargetUsage[1] = (DerivedDefines.WRITES_VELOCITY_TO_GBUFFER || (MaterialDefines.MATERIAL_SHADINGMODEL_THIN_TRANSLUCENT && DerivedDefines.THIN_TRANSLUCENT_USE_DUAL_BLEND));
+		}
+	}
+	else if (MaterialDefines.IS_VIRTUAL_TEXTURE_MATERIAL)
+	{
+		// these whill change, of course
+		if (MaterialDefines.OUT_BASECOLOR)
+		{
+			bTargetUsage[0] = 1;
+		}
+		else if (MaterialDefines.OUT_BASECOLOR_NORMAL_SPECULAR)
+		{
+			bTargetUsage[0] = 1;
+			bTargetUsage[1] = 1;
+			bTargetUsage[2] = 1;
+		}
+		else if (MaterialDefines.OUT_WORLDHEIGHT)
+		{
+			bTargetUsage[0] = 1;
+		}
+	}
+	else if (MaterialDefines.IS_DECAL)
+	{
+		// these will have to change too
+		bTargetUsage[0] = MaterialDefines.DECAL_RENDERTARGET_COUNT > 0;
+		bTargetUsage[1] = MaterialDefines.DECAL_RENDERTARGET_COUNT > 1;
+		bTargetUsage[2] = MaterialDefines.DECAL_RENDERTARGET_COUNT > 2;
+		bTargetUsage[3] = MaterialDefines.DECAL_RENDERTARGET_COUNT > 3;
+		bTargetUsage[4] = MaterialDefines.DECAL_RENDERTARGET_COUNT > 4;
+	}
+	else
+	{
+		// something else, so no op
+	}
+#if 1
+	static bool bTestNewVersion = true;
+	if (bTestNewVersion)
+	{
+		//if (DerivedDefines.USES_GBUFFER)
+		{
+			for (int32 Iter = 0; Iter < FGBufferInfo::MaxTargets; Iter++)
+			{
+				if (bTargetUsage[Iter])
+				{
+					FString TargetName = FString::Printf(TEXT("PIXELSHADEROUTPUT_MRT%d"), Iter);
+					OutEnvironment.SetDefine(TargetName.GetCharArray().GetData(), TEXT("1"));
+				}
+			}
+		}
+	}
+	else
+	{
+		// This uses the legacy logic from CalculateDerivedMaterialParameters(); Just keeping it around momentarily for testing during the transition.
+		SET_COMPILE_BOOL_IF_TRUE(PIXELSHADEROUTPUT_MRT0)
+		SET_COMPILE_BOOL_IF_TRUE(PIXELSHADEROUTPUT_MRT1)
+		SET_COMPILE_BOOL_IF_TRUE(PIXELSHADEROUTPUT_MRT2)
+		SET_COMPILE_BOOL_IF_TRUE(PIXELSHADEROUTPUT_MRT3)
+		SET_COMPILE_BOOL_IF_TRUE(PIXELSHADEROUTPUT_MRT4)
+		SET_COMPILE_BOOL_IF_TRUE(PIXELSHADEROUTPUT_MRT5)
+		SET_COMPILE_BOOL_IF_TRUE(PIXELSHADEROUTPUT_MRT6)
+	}
+#endif
+}
+
+void FShaderCompileUtilities::AppendGBufferDDCKeyString(const EShaderPlatform Platform, FString& KeyString)
+{
+	FShaderGlobalDefines GlobalDefines = FetchShaderGlobalDefines(Platform);
+
+	FString Key = FString::Printf(TEXT("_%d%d%d_%d;\n"),GlobalDefines.GBUFFER_HAS_VELOCITY, GlobalDefines.GBUFFER_HAS_TANGENT, GlobalDefines.ALLOW_STATIC_LIGHTING,GBufferGeneratorVersion);
+	KeyString += Key;
+}
+
+static FGBufferInfo GLastGBufferInfo[SP_NumPlatforms] = {};
+static bool GLastGBufferIsValid[SP_NumPlatforms] = {}; // set to all false
+
+FGBufferInfo FShaderCompileUtilities::FetchGBufferInfoAndWriteAutogen(EShaderPlatform TargetPlatform, ERHIFeatureLevel::Type FeatureLevel = ERHIFeatureLevel::SM5)
+{
+	FShaderGlobalDefines GlobalDefines = FetchShaderGlobalDefines(TargetPlatform);
+
+	FGBufferParams Params = FetchGBufferParamsPipeline(TargetPlatform);
+
+	FScopeLock MapLock(&GCriticalSection);
+
+	// For now, the logic always calculates the new GBuffer, and if it's the first time, write it, otherwise check it hasn't changed. We are doing this for
+	// debugging, and in the near future it will only calculate the GBuffer on the first time only.
+
+	FGBufferInfo BufferInfo = FetchFullGBufferInfo(Params);
+
+	FString PlatformName = LegacyShaderPlatformToShaderFormat(TargetPlatform).GetPlainNameString();
+	FString AutoGenDirectory = GetAutoGenDirectory(TargetPlatform);
+	FString AutogenHeaderFilename = AutoGenDirectory / TEXT("AutogenShaderHeaders.ush");
+	FString AutogenHeaderFilenameTemp = AutoGenDirectory / TEXT("AutogenShaderHeaders_temp.ush");
+
+	if (GLastGBufferIsValid[TargetPlatform])
+	{
+		const bool bSame = IsGBufferInfoEqual(GLastGBufferInfo[TargetPlatform], BufferInfo);
+		check(bSame);
+	}
+	else
+	{
+		GLastGBufferIsValid[TargetPlatform] = true;
+
+		// should cache this properly, and serialize it, but this is a temporary fix.
+		GLastGBufferInfo[TargetPlatform] = BufferInfo;
+
+		FString OutputFileData;
+		OutputFileData += TEXT("// Copyright Epic Games, Inc. All Rights Reserved.\n");
+		OutputFileData += TEXT("\n");
+		OutputFileData += TEXT("#pragma once\n");
+		OutputFileData += TEXT("\n");
+
+
+		OutputFileData += TEXT("#if FEATURE_LEVEL >= FEATURE_LEVEL_SM5\n");
+		OutputFileData += TEXT("float SampleDeviceZFromSceneTexturesTempCopy(float2 UV)\n");
+		OutputFileData += TEXT("{\n");
+		OutputFileData += TEXT("\treturn SceneDepthTexture.SampleLevel(SceneDepthTextureSampler, UV, 0).r;\n");
+		OutputFileData += TEXT("}\n");
+		OutputFileData += TEXT("#endif\n");
+		OutputFileData += TEXT("\n");
+
+		OutputFileData += CreateGBufferEncodeFunction(BufferInfo);
+
+		OutputFileData += TEXT("\n");
+
+		OutputFileData += CreateGBufferDecodeFunctionDirect(BufferInfo);
+
+		OutputFileData += TEXT("\n");
+		//OutputFileData += TEXT("#if SHADING_PATH_DEFERRED\n");
+		OutputFileData += TEXT("#if FEATURE_LEVEL >= FEATURE_LEVEL_SM5\n");
+		OutputFileData += TEXT("\n");
+
+		OutputFileData += CreateGBufferDecodeFunctionVariation(BufferInfo, EGBufferDecodeType::CoordUV, FeatureLevel);
+		OutputFileData += TEXT("\n");
+
+		OutputFileData += CreateGBufferDecodeFunctionVariation(BufferInfo, EGBufferDecodeType::CoordUInt, FeatureLevel);
+
+		OutputFileData += TEXT("\n");
+
+		OutputFileData += CreateGBufferDecodeFunctionVariation(BufferInfo, EGBufferDecodeType::SceneTextures, FeatureLevel);
+		OutputFileData += TEXT("\n");
+
+		OutputFileData += CreateGBufferDecodeFunctionVariation(BufferInfo, EGBufferDecodeType::SceneTexturesLoad, FeatureLevel);
+		OutputFileData += TEXT("\n");
+
+		OutputFileData += TEXT("#endif\n");
+		OutputFileData += TEXT("\n");
+
+		FPaths::MakeStandardFilename(AutogenHeaderFilenameTemp);
+		FFileHelper::SaveStringToFile(OutputFileData, AutogenHeaderFilenameTemp.GetCharArray().GetData(), FFileHelper::EEncodingOptions::ForceAnsi);
+
+		FPlatformFileManager::Get().GetPlatformFile().DeleteFile(*AutogenHeaderFilename);
+		FPlatformFileManager::Get().GetPlatformFile().MoveFile(*AutogenHeaderFilename, *AutogenHeaderFilenameTemp);
+	}
+
+	return BufferInfo;
+}
+
+void FShaderCompileUtilities::GenerateBrdfHeaders(EShaderPlatform TargetPlatform)
+{
+	ERHIFeatureLevel::Type FeatureLevel = GetMaxSupportedFeatureLevel(TargetPlatform);
+
+	//FShaderGlobalDefines GlobalDefines = FetchShaderGlobalDefines(TargetPlatform);
+
+	// Fetches the GBuffer format, which also involves writing the .ush file if it's out of date.
+	FGBufferInfo Info = FetchGBufferInfoAndWriteAutogen(TargetPlatform, FeatureLevel);
+}
+
+ENGINE_API FGBufferParams FShaderCompileUtilities::FetchGBufferParamsRuntime(EShaderPlatform Platform)
+{
+
+	FGBufferParams Ret = {};
+	static const auto CVar = IConsoleManager::Get().FindTConsoleVariableDataInt(TEXT("r.AllowStaticLighting"));
+	Ret.bHasPrecShadowFactor = (CVar ? (CVar->GetValueOnAnyThread() != 0) : 1);
+
+	Ret.bHasVelocity = IsUsingBasePassVelocity(Platform) ? 1 : 0;
+	Ret.bHasTangent = false;//BasePassCanOutputTangent(ShaderPlatform) ? 1 : 0;
+
+	static IConsoleVariable* CVarLumen = IConsoleManager::Get().FindConsoleVariable(TEXT("r.LumenScene"));
+	Ret.bUsesVelocityDepth = CVarLumen->GetInt() > 0 || FDataDrivenShaderPlatformInfo::GetSupportsRayTracing(Platform);
+
+	static const auto CVarFormat = IConsoleManager::Get().FindTConsoleVariableDataInt(TEXT("r.GBufferFormat"));
+	Ret.LegacyFormatIndex = CVarFormat->GetValueOnAnyThread();
+
+	return Ret;
+}
+
+FGBufferParams FShaderCompileUtilities::FetchGBufferParamsPipeline(EShaderPlatform Platform)
+{
+	FGBufferParams Ret = {};
+#if WITH_EDITOR
+	{
+		static const auto CVar = IConsoleManager::Get().FindTConsoleVariableDataInt(TEXT("r.AllowStaticLighting"));
+		Ret.bHasPrecShadowFactor = (CVar ? (CVar->GetValueOnAnyThread() != 0) : 1);
+	}
+
+	Ret.bHasVelocity = IsUsingBasePassVelocity(Platform) ? 1 : 0;
+	Ret.bHasTangent = false;//BasePassCanOutputTangent(TargetPlatform) ? 1 : 0;
+
+	{
+		static IConsoleVariable* CVarLumen = IConsoleManager::Get().FindConsoleVariable(TEXT("r.LumenScene"));
+		Ret.bUsesVelocityDepth = CVarLumen->GetInt() > 0 || FDataDrivenShaderPlatformInfo::GetSupportsRayTracing(Platform);
+	}
+
+	{
+		static IConsoleVariable* CVar = IConsoleManager::Get().FindConsoleVariable(TEXT("r.GBufferFormat"));
+		Ret.LegacyFormatIndex = CVar->GetInt();
+	}
+#else
+	// this function should never be called in a non-editor build
+	check(0);
+#endif
+
+
+	return Ret;
+}
diff --git a/Engine/Source/Runtime/Engine/Public/ShaderCompiler.h b/Engine/Source/Runtime/Engine/Public/ShaderCompiler.h
--- a/Engine/Source/Runtime/Engine/Public/ShaderCompiler.h
+++ b/Engine/Source/Runtime/Engine/Public/ShaderCompiler.h
@@ -17,6 +17,9 @@
 #include "Templates/Atomic.h"
 #include "Templates/UniquePtr.h"
 #include "HAL/ThreadSafeCounter.h"
+#include "RHIDefinitions.h"
+#include "GBufferInfo.h"
+#include "ShaderMaterial.h"
 
 class FShaderCompileJob;
 class FShaderPipelineCompileJob;
@@ -230,6 +233,22 @@ namespace FShaderCompileUtilities
 
 	/** Execute the specified (single or pipeline) shader compile job. */
 	void ExecuteShaderCompileJob(FShaderCommonCompileJob& Job);
+
+	void GenerateBrdfHeaders(const EShaderPlatform Platform);
+	void ApplyDerivedDefines(FShaderCompilerEnvironment& OutEnvironment, FShaderCompilerEnvironment * SharedEnvironment, const EShaderPlatform Platform);
+	void AppendGBufferDDCKeyString(const EShaderPlatform Platform, FString& KeyString);
+	FGBufferInfo FetchGBufferInfoAndWriteAutogen(EShaderPlatform TargetPlatform, ERHIFeatureLevel::Type FeatureLevel);
+
+	void ApplyFetchEnvironment(FShaderMaterialPropertyDefines& DefineData, FShaderCompilerEnvironment& OutEnvironment);
+	void ApplyFetchEnvironment(FShaderGlobalDefines& DefineData, FShaderCompilerEnvironment& OutEnvironment);
+	void ApplyFetchEnvironment(FShaderLightmapPropertyDefines& DefineData, FShaderCompilerEnvironment& OutEnvironment);
+	void ApplyFetchEnvironment(FShaderCompilerDefines& DefineData, FShaderCompilerEnvironment& OutEnvironment);
+
+	void ApplyModifyEnvironment(const FShaderMaterialDerivedDefines& DefineData, FShaderCompilerEnvironment& OutEnvironment);
+
+	ENGINE_API FGBufferParams FetchGBufferParamsRuntime(EShaderPlatform Platform); // this function is called from renderer
+	FGBufferParams FetchGBufferParamsPipeline(EShaderPlatform Platform);
+
 }
 
 #if PLATFORM_WINDOWS // XGE shader compilation is only supported on Windows.
diff --git a/Engine/Source/Runtime/Launch/Private/LaunchEngineLoop.cpp b/Engine/Source/Runtime/Launch/Private/LaunchEngineLoop.cpp
--- a/Engine/Source/Runtime/Launch/Private/LaunchEngineLoop.cpp
+++ b/Engine/Source/Runtime/Launch/Private/LaunchEngineLoop.cpp
@@ -1726,6 +1726,13 @@ int32 FEngineLoop::PreInitPreStartupScreen(const TCHAR* CmdLine)
 	// Add the default engine shader dir
 	AddShaderSourceDirectoryMapping(TEXT("/Engine"), FGenericPlatformProcess::ShaderDir());
 
+	{
+		IPlatformFile& PlatformFile = FPlatformFileManager::Get().GetPlatformFile();
+		FString AutogenAbsolutePath = FPaths::ConvertRelativePathToFull(FPaths::ProjectIntermediateDir() / TEXT("ShaderAutogen"));
+		PlatformFile.CreateDirectory(AutogenAbsolutePath.GetCharArray().GetData());
+		AddShaderSourceDirectoryMapping(TEXT("/ShaderAutogen"), AutogenAbsolutePath);
+	}
+
 	TArray<FString> Tokens;
 	TArray<FString> Switches;
 	UCommandlet::ParseCommandLine(CommandLineCopy, Tokens, Switches);
diff --git a/Engine/Source/Runtime/RenderCore/Private/GBufferInfo.cpp b/Engine/Source/Runtime/RenderCore/Private/GBufferInfo.cpp
new file mode 100644
--- /dev/null
+++ b/Engine/Source/Runtime/RenderCore/Private/GBufferInfo.cpp
@@ -0,0 +1,422 @@
+// Copyright Epic Games, Inc. All Rights Reserved.
+
+/*=============================================================================
+	GBufferInfo.cpp: Dynamic GBuffer implementation.
+=============================================================================*/
+
+#include "GBufferInfo.h"
+
+
+// This is very ugly, but temporary. We can't include EngineTypes.h because this file is ShaderCore and that
+// file is Engine. But since we will remove this flag anyways, we will copy/paste the enums for now, but remove
+// them as soon as flexible GBuffer is in.
+static const int32 EGBufferFormat_Force8BitsPerChannel = 0;
+static const int32 EGBufferFormat_Default = 1;
+/** Same as Default except normals are encoded at 16 bits per channel. */
+static const int32 EGBufferFormat_HighPrecisionNormals = 3;
+/** Forces all GBuffers to 16 bits per channel. Intended as profiling for best quality. */
+static const int32 EGBufferFormat_Force16BitsPerChannel = 5;
+
+
+
+
+
+static bool IsGBufferPackingEqual(const FGBufferPacking& Lhs, const FGBufferPacking& Rhs)
+{
+	return
+		Lhs.bFull == Rhs.bFull &&
+		Lhs.bIsValid == Rhs.bIsValid &&
+		Lhs.TargetIndex == Rhs.TargetIndex &&
+		Lhs.DstBitStart == Rhs.DstBitStart &&
+		Lhs.DstChannelIndex == Rhs.DstChannelIndex &&
+		Lhs.SrcBitStart == Rhs.SrcBitStart &&
+		Lhs.SrcChannelIndex == Rhs.SrcChannelIndex &&
+		Lhs.BitNum == Rhs.BitNum;
+}
+
+static bool IsGBufferItemEqual(const FGBufferItem& Lhs, const FGBufferItem& Rhs)
+{
+	for (int32 I = 0; I < FGBufferItem::MaxPacking; I++)
+	{
+		const bool bSame = IsGBufferPackingEqual(Lhs.Packing[I], Rhs.Packing[I]);
+		if (!bSame)
+		{
+			return false;
+		}
+	}
+
+	return Lhs.bIsValid == Rhs.bIsValid &&
+		Lhs.bQuantizationBias == Rhs.bQuantizationBias &&
+		Lhs.BufferSlot == Rhs.BufferSlot &&
+		Lhs.Compression == Rhs.Compression &&
+		Lhs.Checker == Rhs.Checker;
+}
+
+
+bool RENDERCORE_API IsGBufferInfoEqual(const FGBufferInfo& Lhs, const FGBufferInfo& Rhs)
+{
+	if (Lhs.NumTargets != Rhs.NumTargets)
+	{
+		return false;
+	}
+
+	for (int32 I = 0; I < Lhs.NumTargets; I++)
+	{
+		if (!(Lhs.Targets[I] == Rhs.Targets[I]))
+		{
+			return false;
+		}
+	}
+
+	for (int32 I = 0; I < GBS_Num; I++)
+	{
+		if (!IsGBufferItemEqual(Lhs.Slots[I], Rhs.Slots[I]))
+		{
+			return false;
+		}
+	}
+
+	return true;
+}
+
+
+TArray < EGBufferSlot > FetchGBufferSlots(bool bHasVelocity, bool bHasTangent, bool bHasPrecShadowFactor)
+{
+	TArray < EGBufferSlot > NeededSlots;
+
+	NeededSlots.Push(GBS_SceneColor);
+	NeededSlots.Push(GBS_WorldNormal);
+	NeededSlots.Push(GBS_PerObjectGBufferData);
+	NeededSlots.Push(GBS_Metallic);
+	NeededSlots.Push(GBS_Specular);
+	NeededSlots.Push(GBS_Roughness);
+	NeededSlots.Push(GBS_ShadingModelId);
+	NeededSlots.Push(GBS_SelectiveOutputMask);
+	NeededSlots.Push(GBS_BaseColor);
+
+	if (bHasPrecShadowFactor)
+	{
+		NeededSlots.Push(GBS_GenericAO);
+	}
+	else
+	{
+		NeededSlots.Push(GBS_GenericAO);
+	}
+
+	if (bHasVelocity)
+	{
+		NeededSlots.Push(GBS_Velocity);
+	}
+	if (bHasPrecShadowFactor)
+	{
+		NeededSlots.Push(GBS_PrecomputedShadowFactor);
+	}
+	if (bHasTangent)
+	{
+		NeededSlots.Push(GBS_WorldTangent);
+		NeededSlots.Push(GBS_Anisotropy);
+	}
+	NeededSlots.Push(GBS_CustomData);
+
+	return NeededSlots;
+}
+
+
+// finds the target by searching for the name. returns -1 if not found.
+int32 RENDERCORE_API FindGBufferTargetByName(const FGBufferInfo& GBufferInfo, const FString& Name)
+{
+	for (int32 I = 0; I < GBufferInfo.NumTargets; I++)
+	{
+		if (GBufferInfo.Targets[I].TargetName.Compare(Name) == 0)
+		{
+			return I;
+		}
+	}
+
+	return -1;
+}
+
+
+/*
+ * 4.25 Logic:
+ *
+ * if (SrcGlobal.GBUFFER_HAS_VELOCITY == 0 && SrcGlobal.GBUFFER_HAS_TANGENT == 0)
+ *    0: Lighting
+ *    1: GBufferA
+ *    2: GBufferB
+ *    3: GBufferC
+ *    4: GBufferD
+ *    if (GBUFFER_HAS_PRECSHADOWFACTOR)
+ *        5: GBufferE
+ * else if (SrcGlobal.GBUFFER_HAS_VELOCITY == 1 && SrcGlobal.GBUFFER_HAS_TANGENT == 0)
+ *    0: Lighting
+ *    1: GBufferA
+ *    2: GBufferB
+ *    3: GBufferC
+ *    4: Velocity (NOTE!)
+ *    5: GBufferD
+ *    if (GBUFFER_HAS_PRECSHADOWFACTOR)
+ *        6: GBufferE
+ * else if (SrcGlobal.GBUFFER_HAS_VELOCITY == 0 && SrcGlobal.GBUFFER_HAS_TANGENT == 1)
+ *    0: Lighting
+ *    1: GBufferA
+ *    2: GBufferB
+ *    3: GBufferC
+ *    4: GBufferF (NOTE!)
+ *    5: GBufferD
+ *    if (GBUFFER_HAS_PRECSHADOWFACTOR)
+ *        6: GBufferE
+ * else if (SrcGlobal.GBUFFER_HAS_VELOCITY == 1 && SrcGlobal.GBUFFER_HAS_TANGENT == 1)
+ *    assert(0)
+ *
+*/
+
+/*
+ * LegacyFormatIndex: EGBufferFormat enum. Going forward, we will have better granularity on choosing precision, so thes
+ *    are just being maintained for the transition.
+ *
+ * bUsesVelocityDepth: Normal velocity format is half16 for RG. But when enabled, it changes to half16 RGBA with
+ *    alpha storing depth and blue unused.
+ *
+ */
+
+FGBufferInfo RENDERCORE_API FetchLegacyGBufferInfo(const FGBufferParams& Params)
+{
+	FGBufferInfo Info = {};
+
+	check(!Params.bHasVelocity || !Params.bHasTangent);
+
+	bool bStaticLighting = Params.bHasPrecShadowFactor;
+
+	int32 TargetLighting = 0;
+	int32 TargetGBufferA = 1;
+	int32 TargetGBufferB = 2;
+	int32 TargetGBufferC = 3;
+	int32 TargetGBufferD = -1;
+	int32 TargetGBufferE = -1;
+	int32 TargetGBufferF = -1;
+	int32 TargetVelocity = -1;
+
+	if (Params.bHasVelocity == 0 && Params.bHasTangent == 0)
+	{
+		Info.NumTargets = Params.bHasPrecShadowFactor ? 6 : 5;
+	}
+	else
+	{
+		Info.NumTargets = Params.bHasPrecShadowFactor ? 7 : 6;
+	}
+
+
+
+
+	// good to see the quality loss due to precision in the gbuffer
+	const bool bHighPrecisionGBuffers = (Params.LegacyFormatIndex >= EGBufferFormat_Force16BitsPerChannel);
+	// good to profile the impact of non 8 bit formats
+	const bool bEnforce8BitPerChannel = (Params.LegacyFormatIndex == EGBufferFormat_Force8BitsPerChannel);
+
+	EGBufferType        NormalGBufferFormatTarget  = bHighPrecisionGBuffers ? GBT_Float_16_16_16_16 : GBT_Unorm_10_10_10_2;
+	EGBufferCompression NormalGBufferFormatChannel = bHighPrecisionGBuffers ? GBC_EncodeNormal_Normal_16_16_16 : GBC_EncodeNormal_Normal_10_10_10;
+
+	if (bEnforce8BitPerChannel)
+	{
+		NormalGBufferFormatTarget  = GBT_Unorm_8_8_8_8;
+		NormalGBufferFormatChannel = GBC_Raw_Unorm_8_8_8_8;
+	}
+	else if (Params.LegacyFormatIndex == EGBufferFormat_HighPrecisionNormals)
+	{
+		NormalGBufferFormatTarget  = GBT_Float_16_16_16_16;
+		NormalGBufferFormatChannel = GBC_Raw_Float_16_16_16;
+	}
+
+	const EGBufferType        DiffuseAndSpecularGBufferFormat  = bHighPrecisionGBuffers ? GBT_Float_16_16_16_16 : GBT_Unorm_8_8_8_8;
+	const EGBufferCompression DiffuseGBufferChannel = bHighPrecisionGBuffers ? GBC_Raw_Float_16_16_16 : GBC_Raw_Unorm_8_8_8;
+	const EGBufferCompression SpecularGBufferChannel = bHighPrecisionGBuffers ? GBC_Raw_Float_16 : GBC_Raw_Unorm_8;
+
+	Info.Targets[0].Init(GBT_Unorm_11_11_10,  TEXT("Lighting"), false,  true,  true,  true);
+	Info.Targets[1].Init(NormalGBufferFormatTarget,TEXT("GBufferA"), false,  true,  true,  true);
+	Info.Targets[2].Init(DiffuseAndSpecularGBufferFormat,   TEXT("GBufferB"), false,  true,  true,  true);
+	
+	const bool bLegacyAlbedoSrgb = true;
+	Info.Targets[3].Init(DiffuseAndSpecularGBufferFormat,  TEXT("GBufferC"), bLegacyAlbedoSrgb && !bHighPrecisionGBuffers,  true,  true,  true);
+
+	if (Params.bHasVelocity == 0 && Params.bHasTangent == 0)
+	{
+		TargetGBufferD = 4;
+		Info.Targets[4].Init(GBT_Unorm_8_8_8_8,  TEXT("GBufferD"), false,  true,  true,  true);
+
+		if (Params.bHasPrecShadowFactor)
+		{
+			TargetGBufferE = 5;
+			Info.Targets[5].Init(GBT_Unorm_8_8_8_8,  TEXT("GBufferE"), false,  true,  true,  true);
+		}
+	}
+	else if (Params.bHasVelocity)
+	{
+		TargetVelocity = 4;
+		TargetGBufferD = 5;
+
+		// note the false for use extra flags for velocity, not quite sure of all the ramifications, but this keeps it consistent with previous usage
+		Info.Targets[4].Init(GBT_Float_16_16,   TEXT("Velocity"), false,  true,  true, false);
+		Info.Targets[5].Init(GBT_Unorm_8_8_8_8, TEXT("GBufferD"), false,  true,  true,  true);
+
+		if (Params.bHasPrecShadowFactor)
+		{
+			TargetGBufferE = 6;
+			Info.Targets[6].Init(GBT_Unorm_8_8_8_8,  TEXT("GBufferE"), false,  true,  true, false);
+		}
+	}
+	else if (Params.bHasTangent)
+	{
+		TargetGBufferF = 4;
+		TargetGBufferD = 5;
+		Info.Targets[4].Init(GBT_Unorm_8_8_8_8,  TEXT("GBufferF"), false,  true,  true,  true);
+		Info.Targets[5].Init(GBT_Unorm_8_8_8_8,  TEXT("GBufferD"), false,  true,  true,  true);
+		if (Params.bHasPrecShadowFactor)
+		{
+			TargetGBufferE = 6;
+			Info.Targets[6].Init(GBT_Unorm_8_8_8_8,  TEXT("GBufferE"), false,  true,  true,  true);
+		}
+	}
+	else
+	{
+		// should never hit this path
+		check(0);
+	}
+
+	// this value isn't correct, becuase it doesn't resepect the scene color format cvar, but it's ignored anyways
+	// so it's ok for now
+	Info.Slots[GBS_SceneColor] = FGBufferItem(GBS_SceneColor, GBC_Raw_Float_11_11_10, GBCH_Both);
+	Info.Slots[GBS_SceneColor].Packing[0] = FGBufferPacking(TargetLighting, 0, 0);
+	Info.Slots[GBS_SceneColor].Packing[1] = FGBufferPacking(TargetLighting, 1, 1);
+	Info.Slots[GBS_SceneColor].Packing[2] = FGBufferPacking(TargetLighting, 2, 2);
+
+	Info.Slots[GBS_WorldNormal] = FGBufferItem(GBS_WorldNormal, NormalGBufferFormatChannel, GBCH_Both);
+	Info.Slots[GBS_WorldNormal].Packing[0] = FGBufferPacking(TargetGBufferA, 0, 0);
+	Info.Slots[GBS_WorldNormal].Packing[1] = FGBufferPacking(TargetGBufferA, 1, 1);
+	Info.Slots[GBS_WorldNormal].Packing[2] = FGBufferPacking(TargetGBufferA, 2, 2);
+
+	Info.Slots[GBS_PerObjectGBufferData] = FGBufferItem(GBS_PerObjectGBufferData, GBC_Raw_Unorm_2, GBCH_Both);
+	Info.Slots[GBS_PerObjectGBufferData].Packing[0] = FGBufferPacking(TargetGBufferA, 0, 3);
+
+#if 1
+	Info.Slots[GBS_Metallic] = FGBufferItem(GBS_Metallic, SpecularGBufferChannel, GBCH_Both);
+	Info.Slots[GBS_Metallic].Packing[0] = FGBufferPacking(TargetGBufferB, 0, 0);
+
+	Info.Slots[GBS_Specular] = FGBufferItem(GBS_Specular, SpecularGBufferChannel, GBCH_Both);
+	Info.Slots[GBS_Specular].Packing[0] = FGBufferPacking(TargetGBufferB, 0, 1);
+
+	Info.Slots[GBS_Roughness] = FGBufferItem(GBS_Roughness, SpecularGBufferChannel, GBCH_Both);
+	Info.Slots[GBS_Roughness].Packing[0] = FGBufferPacking(TargetGBufferB, 0, 2);
+#else
+	Info.Slots[GBS_Metallic] = FGBufferItem(GBS_Metallic, GBC_Packed_Quantized_4, GBCH_Both);
+	Info.Slots[GBS_Metallic].Packing[0] = FGBufferPacking(TargetGBufferB, 0, 0, 0, 0, 4);
+
+	Info.Slots[GBS_Specular] = FGBufferItem(GBS_Specular, GBC_Packed_Quantized_4, GBCH_Both);
+	Info.Slots[GBS_Specular].Packing[0] = FGBufferPacking(TargetGBufferB, 0, 0, 0, 4, 4);
+
+	Info.Slots[GBS_Roughness] = FGBufferItem(GBS_Roughness, GBC_Packed_Quantized_4, GBCH_Both);
+	Info.Slots[GBS_Roughness].Packing[0] = FGBufferPacking(TargetGBufferB, 0, 1, 0, 0, 4);
+#endif
+
+	// pack it into bits [0:3] of alpha
+	Info.Slots[GBS_ShadingModelId] = FGBufferItem(GBS_ShadingModelId, GBC_Bits_4, GBCH_Both);
+	Info.Slots[GBS_ShadingModelId].Packing[0] = FGBufferPacking(TargetGBufferB, 0, 3, 0, 0, 4);
+
+	// pack it into bits [4:7] of alpha
+	Info.Slots[GBS_SelectiveOutputMask] = FGBufferItem(GBS_SelectiveOutputMask, GBC_Bits_4, GBCH_Both);
+	Info.Slots[GBS_SelectiveOutputMask].Packing[0] = FGBufferPacking(TargetGBufferB, 0, 3, 0, 4, 4);
+
+	{
+#if 1
+		EGBufferCompression BaseColorCompression = GBC_Invalid;
+
+		Info.Slots[GBS_BaseColor] = FGBufferItem(GBS_BaseColor, DiffuseGBufferChannel, GBCH_Both);
+		Info.Slots[GBS_BaseColor].Packing[0] = FGBufferPacking(TargetGBufferC, 0, 0);
+		Info.Slots[GBS_BaseColor].Packing[1] = FGBufferPacking(TargetGBufferC, 1, 1);
+		Info.Slots[GBS_BaseColor].Packing[2] = FGBufferPacking(TargetGBufferC, 2, 2);
+#elif 0
+		// pack it for funzies
+		Info.Slots[GBS_BaseColor] = FGBufferItem(GBS_BaseColor, GBC_Packed_Color_5_6_5, GBCH_Both);
+		Info.Slots[GBS_BaseColor].bQuantizationBias = true;
+		Info.Slots[GBS_BaseColor].Packing[0] = FGBufferPacking(TargetGBufferB, 0, 0, 0, 0, 5);
+		Info.Slots[GBS_BaseColor].Packing[1] = FGBufferPacking(TargetGBufferB, 1, 1, 0, 0, 6);
+		Info.Slots[GBS_BaseColor].Packing[2] = FGBufferPacking(TargetGBufferB, 2, 2, 0, 0, 5);
+#elif 0
+		// pack it for funzies
+		Info.Slots[GBS_BaseColor] = FGBufferItem(GBS_BaseColor, GBC_Packed_Color_5_6_5, GBCH_Both);
+		Info.Slots[GBS_BaseColor].bQuantizationBias = true;
+		Info.Slots[GBS_BaseColor].Packing[0] = FGBufferPacking(TargetGBufferC, 0, 0, 0, 0, 5);
+		Info.Slots[GBS_BaseColor].Packing[1] = FGBufferPacking(TargetGBufferC, 1, 0, 0, 5, 3);
+		Info.Slots[GBS_BaseColor].Packing[2] = FGBufferPacking(TargetGBufferC, 1, 1, 3, 5, 3);
+		Info.Slots[GBS_BaseColor].Packing[3] = FGBufferPacking(TargetGBufferC, 2, 1, 0, 0, 5);
+#elif 1
+		// pack it for funzies
+		Info.Slots[GBS_BaseColor] = FGBufferItem(GBS_BaseColor, GBC_Packed_Color_4_4_4_Sqrt, GBCH_Both);
+		Info.Slots[GBS_BaseColor].bQuantizationBias = true;
+		Info.Slots[GBS_BaseColor].Packing[0] = FGBufferPacking(TargetGBufferB, 0, 1, 0, 4, 4);
+		Info.Slots[GBS_BaseColor].Packing[1] = FGBufferPacking(TargetGBufferB, 1, 2, 0, 0, 4);
+		Info.Slots[GBS_BaseColor].Packing[2] = FGBufferPacking(TargetGBufferB, 2, 2, 0, 4, 4);
+#else
+		// pack it for funzies
+		Info.Slots[GBS_BaseColor] = FGBufferItem(GBS_BaseColor, GBC_Packed_Color_3_3_2_Sqrt, GBCH_Both);
+		Info.Slots[GBS_BaseColor].bQuantizationBias = true;
+		Info.Slots[GBS_BaseColor].Packing[0] = FGBufferPacking(TargetGBufferC, 0, 0, 0, 0, 3);
+		Info.Slots[GBS_BaseColor].Packing[1] = FGBufferPacking(TargetGBufferC, 1, 0, 0, 3, 3);
+		Info.Slots[GBS_BaseColor].Packing[2] = FGBufferPacking(TargetGBufferC, 2, 0, 0, 6, 2);
+#endif
+
+		if (bStaticLighting)
+		{
+			Info.Slots[GBS_GenericAO] = FGBufferItem(GBS_GenericAO, GBC_Raw_Unorm_8, GBCH_Both);
+			Info.Slots[GBS_GenericAO].Packing[0] = FGBufferPacking(TargetGBufferC, 0, 3);
+		}
+		else
+		{
+			Info.Slots[GBS_GenericAO] = FGBufferItem(GBS_GenericAO, GBC_Raw_Unorm_8, GBCH_Both);
+			Info.Slots[GBS_GenericAO].Packing[0] = FGBufferPacking(TargetGBufferC, 0, 3);
+		}
+	}
+
+	if (Params.bHasVelocity)
+	{
+		Info.Slots[GBS_Velocity] = FGBufferItem(GBS_Velocity, GBC_Raw_Float_16_16, GBCH_Both);
+		Info.Slots[GBS_Velocity].Packing[0] = FGBufferPacking(TargetVelocity, 0, 0);
+		Info.Slots[GBS_Velocity].Packing[1] = FGBufferPacking(TargetVelocity, 1, 1);
+	}
+	if (Params.bHasPrecShadowFactor)
+	{
+		Info.Slots[GBS_PrecomputedShadowFactor] = FGBufferItem(GBS_PrecomputedShadowFactor, GBC_Raw_Unorm_8_8_8_8, GBCH_Both);
+		Info.Slots[GBS_PrecomputedShadowFactor].Packing[0] = FGBufferPacking(TargetGBufferE, 0, 0);
+		Info.Slots[GBS_PrecomputedShadowFactor].Packing[1] = FGBufferPacking(TargetGBufferE, 1, 1);
+		Info.Slots[GBS_PrecomputedShadowFactor].Packing[2] = FGBufferPacking(TargetGBufferE, 2, 2);
+		Info.Slots[GBS_PrecomputedShadowFactor].Packing[3] = FGBufferPacking(TargetGBufferE, 3, 3);
+	}
+	if (Params.bHasTangent)
+	{
+		Info.Slots[GBS_WorldTangent] = FGBufferItem(GBS_WorldTangent, GBC_Raw_Unorm_8_8_8, GBCH_Both);
+		Info.Slots[GBS_WorldTangent].Packing[0] = FGBufferPacking(TargetGBufferF, 0, 0);
+		Info.Slots[GBS_WorldTangent].Packing[1] = FGBufferPacking(TargetGBufferF, 1, 1);
+		Info.Slots[GBS_WorldTangent].Packing[2] = FGBufferPacking(TargetGBufferF, 2, 2);
+
+		Info.Slots[GBS_Anisotropy] = FGBufferItem(GBS_Anisotropy, GBC_Raw_Unorm_8, GBCH_Both);
+		Info.Slots[GBS_Anisotropy].Packing[0] = FGBufferPacking(TargetGBufferF, 0, 3);
+	}
+
+	// GBufferD
+	Info.Slots[GBS_CustomData] = FGBufferItem(GBS_CustomData, GBC_Raw_Unorm_8_8_8_8, GBCH_Both);
+	Info.Slots[GBS_CustomData].Packing[0] = FGBufferPacking(TargetGBufferD, 0, 0);
+	Info.Slots[GBS_CustomData].Packing[1] = FGBufferPacking(TargetGBufferD, 1, 1);
+	Info.Slots[GBS_CustomData].Packing[2] = FGBufferPacking(TargetGBufferD, 2, 2);
+	Info.Slots[GBS_CustomData].Packing[3] = FGBufferPacking(TargetGBufferD, 3, 3);
+
+	return Info;
+}
+
+
+FGBufferInfo RENDERCORE_API FetchFullGBufferInfo(const FGBufferParams& Params)
+{
+	// For now, we are only doing legacy. But next, we will have a switch between the old and new formats.
+	FGBufferInfo Ret = FetchLegacyGBufferInfo(Params);
+	return Ret;
+}
+
diff --git a/Engine/Source/Runtime/RenderCore/Private/ShaderCore.cpp b/Engine/Source/Runtime/RenderCore/Private/ShaderCore.cpp
--- a/Engine/Source/Runtime/RenderCore/Private/ShaderCore.cpp
+++ b/Engine/Source/Runtime/RenderCore/Private/ShaderCore.cpp
@@ -763,6 +763,26 @@ bool ReplaceVirtualFilePathForShaderPlatform(FString& InOutVirtualFilePath, ESha
 	return false;
 }
 
+bool ReplaceVirtualFilePathForShaderAutogen(FString& InOutVirtualFilePath, EShaderPlatform ShaderPlatform)
+{
+	// Tweak the autogen path
+	const FString ShaderAutogenStem = TEXT("/Engine/Generated/ShaderAutogen/");
+	const FString PlatformName = LegacyShaderPlatformToShaderFormat(ShaderPlatform).GetPlainNameString();
+	const FString ShaderAutogenPlatformStem = TEXT("/ShaderAutogen/") + PlatformName;
+
+	// for examples, if it starts with "/Engine/Generated/ShaderAutogen/" change it to "ShaderAutogen/PCD3D_SM5/"
+	if (InOutVirtualFilePath.StartsWith(ShaderAutogenStem))
+	{
+		// Parse the right side, after /ShaderAutogen/
+		FString RelativeShaderName = InOutVirtualFilePath.RightChop(ShaderAutogenStem.Len());
+		FString OutputShaderName = ShaderAutogenPlatformStem + TEXT("/") + RelativeShaderName;
+		InOutVirtualFilePath = OutputShaderName;
+		return true;
+	}
+
+	return false;
+}
+
 bool LoadShaderSourceFile(const TCHAR* InVirtualFilePath, EShaderPlatform ShaderPlatform, FString* OutFileContents, TArray<FShaderCompilerError>* OutCompileErrors) // TODO: const FString&
 {
 	// it's not expected that cooked platforms get here, but if they do, this is the final out
@@ -782,6 +802,9 @@ bool LoadShaderSourceFile(const TCHAR* InVirtualFilePath, EShaderPlatform Shader
 		// Always substitute virtual platform path before accessing GShaderFileCache to get platform-specific file.
 		ReplaceVirtualFilePathForShaderPlatform(VirtualFilePath, ShaderPlatform);
 
+		// Fixup autogen file
+		ReplaceVirtualFilePathForShaderAutogen(VirtualFilePath, ShaderPlatform);
+
 		// Protect GShaderFileCache from simultaneous access by multiple threads
 		FScopeLock ScopeLock(&FileCacheCriticalSection);
 
@@ -902,6 +925,9 @@ static void GetShaderIncludes(const TCHAR* EntryPointVirtualFilePath, const TCHA
 
 					ReplaceVirtualFilePathForShaderPlatform(ExtractedIncludeFilename, ShaderPlatform);
 
+					// Fixup autogen file
+					ReplaceVirtualFilePathForShaderAutogen(ExtractedIncludeFilename, ShaderPlatform);
+
 					// Ignore uniform buffer, vertex factory and instanced stereo includes
 					bool bIgnoreInclude = ExtractedIncludeFilename.StartsWith(TEXT("/Engine/Generated/"));
 
@@ -1310,3 +1336,4 @@ void AddShaderSourceDirectoryMapping(const FString& VirtualShaderDirectory, cons
 		*VirtualShaderDirectory, *RealShaderDirectory);
 	GShaderSourceDirectoryMappings.Add(VirtualShaderDirectory, RealShaderDirectory);
 }
+
diff --git a/Engine/Source/Runtime/RenderCore/Private/ShaderMaterialDerivedHelpers.cpp b/Engine/Source/Runtime/RenderCore/Private/ShaderMaterialDerivedHelpers.cpp
new file mode 100644
--- /dev/null
+++ b/Engine/Source/Runtime/RenderCore/Private/ShaderMaterialDerivedHelpers.cpp
@@ -0,0 +1,180 @@
+// Copyright Epic Games, Inc. All Rights Reserved.
+
+#include "Misc/FileHelper.h"
+#include "Misc/CoreMisc.h"
+#include "ShaderCore.h"
+#include "Interfaces/ITargetPlatformManagerModule.h"
+#include "RenderUtils.h"
+#include "ShaderMaterial.h"
+#include "GBufferInfo.h"
+
+
+FShaderGlobalDefines FetchRuntimeShaderGlobalDefines(EShaderPlatform TargetPlatform)
+{
+	FShaderGlobalDefines Ret = {};
+	return Ret;
+}
+
+
+FShaderMaterialDerivedDefines RENDERCORE_API CalculateDerivedMaterialParameters(
+	const FShaderMaterialPropertyDefines& Mat,
+	const FShaderLightmapPropertyDefines& Lightmap,
+	const FShaderGlobalDefines& SrcGlobal,
+	const FShaderCompilerDefines& Compiler,
+	ERHIFeatureLevel::Type FEATURE_LEVEL)
+{
+	FShaderMaterialDerivedDefines Dst = {};
+
+	// Translucent materials need to compute fogging in the forward shading pass
+	// Materials that read from scene color skip getting fogged, because the contents of the scene color lookup have already been fogged
+	// This is not foolproof, as any additional color the material adds will then not be fogged correctly
+	Dst.TRANSLUCENCY_NEEDS_BASEPASS_FOGGING = (Mat.MATERIAL_ENABLE_TRANSLUCENCY_FOGGING && Mat.MATERIALBLENDING_ANY_TRANSLUCENT && !Mat.MATERIAL_USES_SCENE_COLOR_COPY);
+
+	// With forward shading, fog always needs to be computed in the base pass to work correctly with MSAA
+	Dst.OPAQUE_NEEDS_BASEPASS_FOGGING = (!Mat.MATERIALBLENDING_ANY_TRANSLUCENT && SrcGlobal.FORWARD_SHADING);
+
+	Dst.NEEDS_BASEPASS_VERTEX_FOGGING = ((Dst.TRANSLUCENCY_NEEDS_BASEPASS_FOGGING && !Mat.MATERIAL_COMPUTE_FOG_PER_PIXEL) || (Dst.OPAQUE_NEEDS_BASEPASS_FOGGING && SrcGlobal.PROJECT_VERTEX_FOGGING_FOR_OPAQUE));
+	Dst.NEEDS_BASEPASS_PIXEL_FOGGING = ((Dst.TRANSLUCENCY_NEEDS_BASEPASS_FOGGING && Mat.MATERIAL_COMPUTE_FOG_PER_PIXEL) || (Dst.OPAQUE_NEEDS_BASEPASS_FOGGING && !SrcGlobal.PROJECT_VERTEX_FOGGING_FOR_OPAQUE));
+
+	// Volumetric fog interpolated per vertex gives very poor results, always sample the volumetric fog texture per-pixel
+	// Opaque materials in the deferred renderer get volumetric fog applied in a deferred fog pass
+	Dst.NEEDS_BASEPASS_PIXEL_VOLUMETRIC_FOGGING = (Mat.MATERIALBLENDING_ANY_TRANSLUCENT || SrcGlobal.FORWARD_SHADING);
+
+	// need to change this for mobile vs PC, and get rid of the #undefs
+	Dst.NEEDS_LIGHTMAP_COORDINATE = (Lightmap.HQ_TEXTURE_LIGHTMAP || Lightmap.LQ_TEXTURE_LIGHTMAP);
+
+	// this logic differs from the actual defines due to confusing #undefs. NEEDS_LIGHTMAP is only allowed to be true if PRECOMPUTED_IRRADIANCE_VOLUME_LIGHTING
+	// is false because otherwise GetLightMapCoordinates() will not be defined causing a compiler error.
+	Dst.NEEDS_LIGHTMAP = (Dst.NEEDS_LIGHTMAP_COORDINATE) && !Lightmap.PRECOMPUTED_IRRADIANCE_VOLUME_LIGHTING;
+
+	Dst.USES_GBUFFER = (FEATURE_LEVEL >= ERHIFeatureLevel::SM4_REMOVED && (Mat.MATERIALBLENDING_SOLID || Mat.MATERIALBLENDING_MASKED) && !Lightmap.SIMPLE_FORWARD_SHADING && !SrcGlobal.FORWARD_SHADING);
+
+	// Only some shader models actually need custom data.
+	Dst.WRITES_CUSTOMDATA_TO_GBUFFER = (Dst.USES_GBUFFER && (Mat.MATERIAL_SHADINGMODEL_SUBSURFACE || Mat.MATERIAL_SHADINGMODEL_PREINTEGRATED_SKIN || Mat.MATERIAL_SHADINGMODEL_SUBSURFACE_PROFILE || Mat.MATERIAL_SHADINGMODEL_CLEAR_COAT || Mat.MATERIAL_SHADINGMODEL_TWOSIDED_FOLIAGE || Mat.MATERIAL_SHADINGMODEL_HAIR || Mat.MATERIAL_SHADINGMODEL_CLOTH || Mat.MATERIAL_SHADINGMODEL_EYE));
+
+	// Based on GetPrecomputedShadowMasks()
+	// Note: WRITES_PRECSHADOWFACTOR_TO_GBUFFER is currently disabled because we use the precomputed shadow factor GBuffer outside of STATICLIGHTING_TEXTUREMASK to store UseSingleSampleShadowFromStationaryLights
+	Dst.GBUFFER_HAS_PRECSHADOWFACTOR = (Dst.USES_GBUFFER && SrcGlobal.ALLOW_STATIC_LIGHTING);
+	Dst.WRITES_PRECSHADOWFACTOR_ZERO = (!(Lightmap.STATICLIGHTING_TEXTUREMASK && Lightmap.STATICLIGHTING_SIGNEDDISTANCEFIELD) && (Lightmap.HQ_TEXTURE_LIGHTMAP || Lightmap.LQ_TEXTURE_LIGHTMAP));
+	Dst.WRITES_PRECSHADOWFACTOR_TO_GBUFFER = (Dst.GBUFFER_HAS_PRECSHADOWFACTOR && !Dst.WRITES_PRECSHADOWFACTOR_ZERO);
+
+	// If a primitive has static lighting, we assume it is not moving. If it is, it will be rerendered in an extra renderpass.
+	Dst.SUPPORTS_WRITING_VELOCITY_TO_BASE_PASS = (FEATURE_LEVEL >= ERHIFeatureLevel::SM4_REMOVED && (Mat.MATERIALBLENDING_SOLID || Mat.MATERIALBLENDING_MASKED) && !Lightmap.SIMPLE_FORWARD_SHADING);
+	Dst.WRITES_VELOCITY_TO_GBUFFER = ((Dst.SUPPORTS_WRITING_VELOCITY_TO_BASE_PASS || Dst.USES_GBUFFER) && SrcGlobal.GBUFFER_HAS_VELOCITY && (!SrcGlobal.SELECTIVE_BASEPASS_OUTPUTS || !(Lightmap.STATICLIGHTING_TEXTUREMASK || Lightmap.STATICLIGHTING_SIGNEDDISTANCEFIELD || Lightmap.HQ_TEXTURE_LIGHTMAP || Lightmap.LQ_TEXTURE_LIGHTMAP || Lightmap.WATER_MESH_FACTORY)));
+
+
+	// Use an extra interpolator for position
+	Dst.WRITES_VELOCITY_TO_GBUFFER_USE_POS_INTERPOLATOR = (Dst.WRITES_VELOCITY_TO_GBUFFER);// && USING_TESSELLATION);
+
+	Dst.TRANSLUCENCY_ANY_PERVERTEX_LIGHTING = (Mat.TRANSLUCENCY_LIGHTING_VOLUMETRIC_PERVERTEX_NONDIRECTIONAL || Mat.TRANSLUCENCY_LIGHTING_VOLUMETRIC_PERVERTEX_DIRECTIONAL);
+	Dst.TRANSLUCENCY_ANY_VOLUMETRIC = (Mat.TRANSLUCENCY_LIGHTING_VOLUMETRIC_NONDIRECTIONAL || Mat.TRANSLUCENCY_LIGHTING_VOLUMETRIC_DIRECTIONAL || Dst.TRANSLUCENCY_ANY_PERVERTEX_LIGHTING);
+	Dst.TRANSLUCENCY_PERVERTEX_LIGHTING_VOLUME = (!SrcGlobal.FORWARD_SHADING && (Mat.TRANSLUCENCY_LIGHTING_VOLUMETRIC_PERVERTEX_DIRECTIONAL || Mat.TRANSLUCENCY_LIGHTING_VOLUMETRIC_PERVERTEX_NONDIRECTIONAL));
+	Dst.TRANSLUCENCY_PERVERTEX_FORWARD_SHADING = (SrcGlobal.FORWARD_SHADING && (Mat.TRANSLUCENCY_LIGHTING_VOLUMETRIC_PERVERTEX_DIRECTIONAL || Mat.TRANSLUCENCY_LIGHTING_VOLUMETRIC_PERVERTEX_NONDIRECTIONAL));
+
+	// No, just no
+	Dst.TESSELLATION_SUPPORTED = false; // (PC_D3D || GL4_PROFILE || SWITCH_PROFILE || SWITCH_PROFILE_FORWARD || METAL_TESSELLATION_PROFILE || VULKAN_PROFILE_SM5);
+
+
+	Dst.COMPILE_SHADERS_FOR_DEVELOPMENT = SrcGlobal.COMPILE_SHADERS_FOR_DEVELOPMENT_ALLOWED && Mat.bAllowDevelopmentShaderCompile;
+
+	Dst.USE_DEVELOPMENT_SHADERS = (Dst.COMPILE_SHADERS_FOR_DEVELOPMENT && Compiler.PLATFORM_SUPPORTS_DEVELOPMENT_SHADERS);
+
+
+	Dst.PLATFORM_SUPPORTS_EDITOR_SHADERS = !Compiler.ESDEFERRED_PROFILE;
+
+	Dst.USE_EDITOR_SHADERS = (Dst.PLATFORM_SUPPORTS_EDITOR_SHADERS && Dst.USE_DEVELOPMENT_SHADERS);
+
+//
+////Materials also have to opt in to these features.
+	Dst.USE_EDITOR_COMPOSITING = (Dst.USE_EDITOR_SHADERS && Mat.EDITOR_PRIMITIVE_MATERIAL);
+	Dst.MATERIALBLENDING_ANY_TRANSLUCENT = (Mat.MATERIALBLENDING_TRANSLUCENT || Mat.MATERIALBLENDING_ADDITIVE || Mat.MATERIALBLENDING_MODULATE);
+
+	Dst.IS_MESHPARTICLE_FACTORY = (Lightmap.PARTICLE_MESH_FACTORY || Lightmap.NIAGARA_MESH_FACTORY);
+
+	Dst.MOBILE_EMULATION = ((FEATURE_LEVEL == ERHIFeatureLevel::ES3_1) && (!(Compiler.COMPILER_GLSL_ES3_1) && Dst.USE_DEVELOPMENT_SHADERS && (!(Compiler.METAL_PROFILE) || Compiler.MAC) && !Compiler.VULKAN_PROFILE));
+
+	Dst.SUPPORTS_PIXEL_COVERAGE = (FEATURE_LEVEL >= ERHIFeatureLevel::SM5 && !Compiler.COMPILER_GLSL && !Dst.MOBILE_EMULATION);
+
+	Dst.FORCE_FULLY_ROUGH = (Lightmap.SIMPLE_FORWARD_SHADING || Mat.MATERIAL_FULLY_ROUGH);
+	Dst.EDITOR_ALPHA2COVERAGE = (Dst.USE_EDITOR_COMPOSITING && Dst.SUPPORTS_PIXEL_COVERAGE);
+	Dst.POST_PROCESS_SUBSURFACE = ((Mat.MATERIAL_SHADINGMODEL_SUBSURFACE_PROFILE || Mat.MATERIAL_SHADINGMODEL_EYE) && Dst.USES_GBUFFER);
+
+	// matching the logic in: bool FMaterialResource::IsDualBlendingEnabled(EShaderPlatform Platform) const
+	Dst.THIN_TRANSLUCENT_USE_DUAL_BLEND = Mat.MATERIAL_SHADINGMODEL_THIN_TRANSLUCENT && SrcGlobal.bSupportsDualBlending;
+
+
+	Dst.MATERIAL_WORKS_WITH_DUAL_SOURCE_COLOR_BLENDING = (Dst.MATERIAL_SHADINGMODEL_THIN_TRANSLUCENT || Mat.STRATA_TRANSLUCENT_ENABLED);
+
+
+	// There are 4 different ways of setting MRTs depending on which .usf file is the base shader, which sets defines which includes a different include file
+	// which may or may not override the current defines. Here is my best attempt at figuring out the logic.
+	if (Mat.IS_VIRTUAL_TEXTURE_MATERIAL)
+	{
+		//if (0)
+		{
+			if (Mat.OUT_BASECOLOR)
+			{
+				Dst.PIXELSHADEROUTPUT_MRT0 = 1;
+			}
+			else if (Mat.OUT_BASECOLOR_NORMAL_SPECULAR)
+			{
+				Dst.PIXELSHADEROUTPUT_MRT0 = 1;
+				Dst.PIXELSHADEROUTPUT_MRT1 = 1;
+				Dst.PIXELSHADEROUTPUT_MRT2 = 1;
+			}
+			else if (Mat.OUT_WORLDHEIGHT)
+			{
+				Dst.PIXELSHADEROUTPUT_MRT0 = 1;
+			}
+		}
+	}
+	else if (Mat.IS_DECAL)
+	{
+		//if (0)
+		{
+			Dst.PIXELSHADEROUTPUT_MRT0 = Mat.DECAL_RENDERTARGET_COUNT > 0;
+			Dst.PIXELSHADEROUTPUT_MRT1 = Mat.DECAL_RENDERTARGET_COUNT > 1;
+			Dst.PIXELSHADEROUTPUT_MRT2 = Mat.DECAL_RENDERTARGET_COUNT > 2;
+			Dst.PIXELSHADEROUTPUT_MRT3 = Mat.DECAL_RENDERTARGET_COUNT > 3;
+			Dst.PIXELSHADEROUTPUT_MRT4 = Mat.DECAL_RENDERTARGET_COUNT > 4;
+		}
+	}
+	else if (Mat.IS_BASE_PASS)
+	{
+		Dst.PIXELSHADEROUTPUT_BASEPASS = 1;
+		if (Dst.USES_GBUFFER)
+		{
+			Dst.PIXELSHADEROUTPUT_MRT0 = (!SrcGlobal.SELECTIVE_BASEPASS_OUTPUTS || Dst.NEEDS_BASEPASS_VERTEX_FOGGING || Mat.USES_EMISSIVE_COLOR || SrcGlobal.ALLOW_STATIC_LIGHTING || Mat.MATERIAL_SHADINGMODEL_SINGLELAYERWATER);
+			Dst.PIXELSHADEROUTPUT_MRT1 = ((!SrcGlobal.SELECTIVE_BASEPASS_OUTPUTS || !Mat.MATERIAL_SHADINGMODEL_UNLIT));
+			Dst.PIXELSHADEROUTPUT_MRT2 = ((!SrcGlobal.SELECTIVE_BASEPASS_OUTPUTS || !Mat.MATERIAL_SHADINGMODEL_UNLIT));
+			Dst.PIXELSHADEROUTPUT_MRT3 = ((!SrcGlobal.SELECTIVE_BASEPASS_OUTPUTS || !Mat.MATERIAL_SHADINGMODEL_UNLIT));
+			if (SrcGlobal.GBUFFER_HAS_VELOCITY || SrcGlobal.GBUFFER_HAS_TANGENT)
+			{
+				Dst.PIXELSHADEROUTPUT_MRT4 = Dst.WRITES_VELOCITY_TO_GBUFFER || SrcGlobal.GBUFFER_HAS_TANGENT;
+				Dst.PIXELSHADEROUTPUT_MRT5 = (!SrcGlobal.SELECTIVE_BASEPASS_OUTPUTS || Dst.WRITES_CUSTOMDATA_TO_GBUFFER);
+				Dst.PIXELSHADEROUTPUT_MRT6 = (Dst.GBUFFER_HAS_PRECSHADOWFACTOR && (!SrcGlobal.SELECTIVE_BASEPASS_OUTPUTS || (Dst.WRITES_PRECSHADOWFACTOR_TO_GBUFFER && !Mat.MATERIAL_SHADINGMODEL_UNLIT)));
+			}
+			else
+			{
+				Dst.PIXELSHADEROUTPUT_MRT4 = (!SrcGlobal.SELECTIVE_BASEPASS_OUTPUTS || Dst.WRITES_CUSTOMDATA_TO_GBUFFER);
+				Dst.PIXELSHADEROUTPUT_MRT5 = (Dst.GBUFFER_HAS_PRECSHADOWFACTOR && (!SrcGlobal.SELECTIVE_BASEPASS_OUTPUTS || (Dst.WRITES_PRECSHADOWFACTOR_TO_GBUFFER && !Mat.MATERIAL_SHADINGMODEL_UNLIT)));
+			}
+		}
+		else
+		{
+			Dst.PIXELSHADEROUTPUT_MRT0 = true;
+			// we also need MRT for thin translucency due to dual blending if we are not on the fallback path
+			Dst.PIXELSHADEROUTPUT_MRT1 = (Dst.WRITES_VELOCITY_TO_GBUFFER || (Mat.DUAL_SOURCE_COLOR_BLENDING_ENABLED && Dst.MATERIAL_WORKS_WITH_DUAL_SOURCE_COLOR_BLENDING));
+		}
+	}
+	else
+	{
+		// Shouldn't be any other cases using PIXELSHADEROUTPUT_MRTX
+	}
+
+	Dst.PIXELSHADEROUTPUT_A2C = ((Dst.EDITOR_ALPHA2COVERAGE) != 0);
+	Dst.PIXELSHADEROUTPUT_COVERAGE = (Mat.MATERIALBLENDING_MASKED_USING_COVERAGE && !SrcGlobal.EARLY_Z_PASS_ONLY_MATERIAL_MASKING);
+	return Dst;
+}
+
+
diff --git a/Engine/Source/Runtime/RenderCore/Public/GBufferInfo.h b/Engine/Source/Runtime/RenderCore/Public/GBufferInfo.h
new file mode 100644
--- /dev/null
+++ b/Engine/Source/Runtime/RenderCore/Public/GBufferInfo.h
@@ -0,0 +1,277 @@
+// Copyright Epic Games, Inc. All Rights Reserved.
+
+#pragma once
+
+#include "CoreMinimal.h"
+
+enum EGBufferSlot
+{
+	GBS_Invalid,
+	GBS_SceneColor, // RGB 11.11.10
+	GBS_WorldNormal, // RGB 10.10.10
+	GBS_PerObjectGBufferData, // 2
+	GBS_Metallic, // R8
+	GBS_Specular, // R8
+	GBS_Roughness, // R8
+	GBS_ShadingModelId, // 4 bits
+	GBS_SelectiveOutputMask, // 4 bits
+	GBS_BaseColor, // RGB8
+	GBS_GenericAO, // R8
+	GBS_IndirectIrradiance, // R8
+	GBS_AO, // R8
+	GBS_Velocity, // RG, float16
+	GBS_PrecomputedShadowFactor, // RGBA8
+	GBS_WorldTangent, // RGB8
+	GBS_Anisotropy, // R8
+	GBS_CustomData, // RGBA8, no compression
+	GBS_SubsurfaceColor, // RGB8
+	GBS_Opacity, // R8
+	GBS_SubsurfaceProfile, //RGB8
+	GBS_ClearCoat, // R8
+	GBS_ClearCoatRoughness, // R8
+	GBS_HairSecondaryWorldNormal, // RG8
+	GBS_HairBacklit, // R8
+	GBS_Cloth, // R8
+	GBS_SubsurfaceProfileX, // R8
+	GBS_IrisNormal, // RG8
+	GBS_Num
+};
+
+
+enum EGBufferCompression
+{
+	GBC_Invalid,
+	GBC_Raw_Float_16_16_16_16,
+	GBC_Raw_Float_16_16_16,
+	GBC_Raw_Float_11_11_10,
+	GBC_Raw_Float_10_10_10,
+	GBC_Raw_Unorm_8_8_8_8,
+	GBC_Raw_Unorm_8_8_8,
+	GBC_Raw_Unorm_8_8,
+	GBC_Raw_Unorm_8,
+	GBC_Raw_Unorm_2, // a float value normalized to use 2 bits
+	GBC_Raw_Float_16_16,
+	GBC_Raw_Float_16,
+	GBC_Bits_4, // an int value, that is fundamentally 4 bits
+	GBC_Bits_2, // an int value, that is fundamentally 2 bits
+
+	GBC_Packed_Normal_Octahedral_8_8,
+	GBC_EncodeNormal_Normal_16_16_16,
+	GBC_EncodeNormal_Normal_10_10_10,
+
+	GBC_Packed_Color_5_6_5, // a unorm value, packed to 565
+	GBC_Packed_Color_5_6_5_Sqrt, // a unorm value, packed to 565, with sqrt
+	GBC_Packed_Color_4_4_4, // a unorm value, packed to 444
+	GBC_Packed_Color_4_4_4_Sqrt, // a unorm value, packed to 444, with sqrt
+	GBC_Packed_Color_3_3_2, // a unorm value, packed to 332
+	GBC_Packed_Color_3_3_2_Sqrt, // a unorm value, packed to 332, with sqrt
+	GBC_Packed_Quantized_6, // a unorm value, quantized to 6 bits
+	GBC_Packed_Quantized_4, // a unorm value, quantized to 4 bits
+	GBC_Packed_Quantized_2, // a unorm value, quantized to 2 bits
+
+	GBC_Num
+};
+
+// the actual format of the output texture
+enum EGBufferType
+{
+	GBT_Invalid,
+	GBT_Unorm_8_8_8_8,
+	GBT_Unorm_11_11_10,
+	GBT_Unorm_10_10_10_2,
+	GBT_Float_16_16,
+	GBT_Float_16_16_16_16,
+	GBT_Num
+};
+
+enum EGBufferChecker
+{
+	GBCH_Invalid,
+	GBCH_Even,
+	GBCH_Odd,
+	GBCH_Both,
+	GBCH_Num
+};
+
+
+struct FGBufferCompressionInfo
+{
+	EGBufferCompression Type; //  compression type
+	int32 SrcNumChan; // how many channels before compression
+	int32 DstNumChan; // how many channels after copression
+	int32 ChanBits[4]; // how many bits are each destination channel
+	bool bIsPackedBits; //
+	bool bIsConversion; //
+	const TCHAR* EncodeFuncName; // name of the function to do this conversion
+	const TCHAR* DecodeFuncName; // name of the function to do this conversion
+};
+
+
+struct FGBufferPacking
+{
+	FGBufferPacking()
+	{
+		bIsValid = false;
+		bFull = false;
+		TargetIndex = -1;
+		DstChannelIndex = -1;
+		SrcChannelIndex = -1;
+		DstBitStart = -1;
+		SrcBitStart = -1;
+		BitNum = -1;
+	}
+
+	// if we pass in 2 values, it means we use the entire channel
+	FGBufferPacking(int32 InTargetIndex, int32 InSrcChannelIndex, int32 InDstChannelIndex)
+	{
+		bIsValid = true;
+		bFull = true;
+		TargetIndex = InTargetIndex;
+		DstChannelIndex = InDstChannelIndex;
+		SrcChannelIndex = InSrcChannelIndex;
+		DstBitStart = -1;
+		SrcBitStart = -1;
+		BitNum = -1;
+	}
+
+	// if we pass in 4 values, it means we also need to pack the bits
+	FGBufferPacking(int32 InTargetIndex, int32 InSrcChannelIndex, int32 InDstChannelIndex, int32 InSrcBitStart, int32 InDstBitStart, int32 InBitNum)
+	{
+		bIsValid = true;
+		bFull = false;
+		TargetIndex = InTargetIndex;
+		DstChannelIndex = InDstChannelIndex;
+		SrcChannelIndex = InSrcChannelIndex;
+		DstBitStart = InDstBitStart;
+		SrcBitStart = InSrcBitStart;
+		BitNum = InBitNum;
+	}
+
+	bool bIsValid;
+	bool bFull;
+	int32 TargetIndex;
+	int32 DstChannelIndex;
+	int32 DstBitStart;
+	int32 SrcChannelIndex;
+	int32 SrcBitStart;
+	int32 BitNum;
+};
+
+
+// the texture positions in the GBuffer
+struct FGBufferItem
+{
+	FGBufferItem()
+	{
+		bIsValid = false;
+		bQuantizationBias = false;
+		BufferSlot = GBS_Invalid;
+		Compression = GBC_Invalid;
+		Checker = GBCH_Invalid;
+
+		for (int32 I = 0; I < 4; I++)
+		{
+			Packing[I] = {};
+		}
+	}
+
+	FGBufferItem(EGBufferSlot InBufferSlot, EGBufferCompression InCompression, EGBufferChecker InChecker)
+	{
+		bIsValid = true;
+		bQuantizationBias = false;
+		BufferSlot = InBufferSlot;
+		Compression = InCompression;
+		Checker = InChecker;
+
+		for (int32 I = 0; I < 4; I++)
+		{
+			Packing[I] = {};
+		}
+	}
+
+	static const int MaxPacking = 8;
+
+	bool bIsValid;
+	bool bQuantizationBias;
+	EGBufferSlot BufferSlot;
+	EGBufferCompression Compression;
+	EGBufferChecker Checker;
+	FGBufferPacking Packing[MaxPacking]; // 8 should be plenty, can always make it larger
+};
+
+struct FGBufferTarget
+{
+	FGBufferTarget()
+	{
+		TargetType = GBT_Invalid;
+		bIsSrgb = false;
+		bIsRenderTargetable = false;
+		bIsShaderResource = false;
+		bIsUsingExtraFlags = false;
+	}
+
+	bool operator==(const FGBufferTarget & Rhs) const
+	{
+		return TargetType == Rhs.TargetType &&
+			   TargetName == Rhs.TargetName &&
+			   bIsSrgb == Rhs.bIsSrgb &&
+			   bIsRenderTargetable == Rhs.bIsRenderTargetable &&
+			   bIsShaderResource == Rhs.bIsShaderResource &&
+			   bIsUsingExtraFlags == Rhs.bIsUsingExtraFlags;
+	}
+
+	void Init(EGBufferType InTargetType,
+			FString InTargetName,
+			bool bInIsSrgb,
+			bool bInIsRenderTargetable,
+			bool bInIsShaderResource,
+			bool bInIsUsingExtraFlags)
+	{
+		TargetType = InTargetType;
+		TargetName = InTargetName;
+		bIsSrgb = bInIsSrgb;
+		bIsRenderTargetable = bInIsRenderTargetable;
+		bIsShaderResource = bInIsShaderResource;
+		bIsUsingExtraFlags = bInIsUsingExtraFlags;
+	}
+
+
+	EGBufferType TargetType;
+	FString TargetName;
+	bool bIsSrgb;
+	bool bIsRenderTargetable;
+	bool bIsShaderResource;
+	bool bIsUsingExtraFlags;
+};
+
+struct FGBufferInfo
+{
+	static const int MaxTargets = 8;
+
+	int32 NumTargets;
+	FGBufferTarget Targets[MaxTargets];
+
+	FGBufferItem Slots[GBS_Num];
+
+};
+
+struct FGBufferParams
+{
+	bool bHasVelocity;
+	bool bHasTangent;
+	bool bHasPrecShadowFactor;
+	int LegacyFormatIndex;
+	bool bUsesVelocityDepth;
+};
+
+
+int32 RENDERCORE_API FindGBufferTargetByName(const FGBufferInfo& GBufferInfo, const FString& Name);
+
+//FGBufferInfo RENDERCORE_API FetchLegacyGBufferInfo(bool bHasVelocity, bool bHasTangent, bool bHasPrecShadowFactor, int LegacyFormatIndex, bool bUsesVelocityDepth);
+//FGBufferInfo RENDERCORE_API FetchFullGBufferInfo(bool bHasVelocity, bool bHasTangent, bool bHasPrecShadowFactor, int LegacyFormatIndex, bool bUsesVelocityDepth);
+FGBufferInfo RENDERCORE_API FetchLegacyGBufferInfo(const FGBufferParams& Params);
+FGBufferInfo RENDERCORE_API FetchFullGBufferInfo(const FGBufferParams& Params);
+
+bool RENDERCORE_API IsGBufferInfoEqual(const FGBufferInfo& Lhs, const FGBufferInfo& Rhs);
+
+
diff --git a/Engine/Source/Runtime/RenderCore/Public/ShaderCore.h b/Engine/Source/Runtime/RenderCore/Public/ShaderCore.h
--- a/Engine/Source/Runtime/RenderCore/Public/ShaderCore.h
+++ b/Engine/Source/Runtime/RenderCore/Public/ShaderCore.h
@@ -803,6 +803,9 @@ extern RENDERCORE_API FString ParseVirtualShaderFilename(const FString& InFilena
 /** Replaces virtual platform path with appropriate path for a given ShaderPlatform. Returns true if path was changed. */
 extern RENDERCORE_API bool ReplaceVirtualFilePathForShaderPlatform(FString& InOutVirtualFilePath, EShaderPlatform ShaderPlatform);
 
+/** Replaces virtual platform path with appropriate autogen path for a given ShaderPlatform. Returns true if path was changed. */
+extern RENDERCORE_API bool ReplaceVirtualFilePathForShaderAutogen(FString& InOutVirtualFilePath, EShaderPlatform ShaderPlatform);
+
 /** Loads the shader file with the given name.  If the shader file couldn't be loaded, throws a fatal error. */
 extern RENDERCORE_API void LoadShaderSourceFileChecked(const TCHAR* VirtualFilePath, EShaderPlatform ShaderPlatform, FString& OutFileContents);
 
diff --git a/Engine/Source/Runtime/RenderCore/Public/ShaderMaterial.h b/Engine/Source/Runtime/RenderCore/Public/ShaderMaterial.h
new file mode 100644
--- /dev/null
+++ b/Engine/Source/Runtime/RenderCore/Public/ShaderMaterial.h
@@ -0,0 +1,304 @@
+// Copyright Epic Games, Inc. All Rights Reserved.
+
+/*=============================================================================
+	ShaderMaterial.h: Shader materials helper definitions.
+=============================================================================*/
+
+#pragma once
+
+#include "CoreMinimal.h"
+#include "Templates/RefCounting.h"
+#include "Serialization/MemoryLayout.h"
+
+struct FShaderCompilerEnvironment;
+
+struct FShaderGlobalDefines
+{
+	//void ModifyEnvironment(FShaderCompilerEnvironment& OutEnvironment) const;
+
+	bool GBUFFER_HAS_VELOCITY;
+	bool GBUFFER_HAS_TANGENT;
+	bool ALLOW_STATIC_LIGHTING;
+	bool CLEAR_COAT_BOTTOM_NORMAL;
+	bool IRIS_NORMAL;
+	bool DXT5_NORMALMAPS;
+	bool SELECTIVE_BASEPASS_OUTPUTS;
+	bool USE_DBUFFER;
+	bool FORWARD_SHADING;
+	bool PROJECT_VERTEX_FOGGING_FOR_OPAQUE;
+	bool PROJECT_MOBILE_DISABLE_VERTEX_FOG;
+	bool USE_PREEXPOSURE;
+	bool PROJECT_ALLOW_GLOBAL_CLIP_PLANE;
+	bool EARLY_Z_PASS_ONLY_MATERIAL_MASKING;
+	bool PROJECT_SUPPORT_SKY_ATMOSPHERE;
+	bool PROJECT_SUPPORT_SKY_ATMOSPHERE_AFFECTS_HEIGHFOG;
+	bool PROJECT_MOBILE_USE_LEGACY_SHADING;
+	bool POST_PROCESS_ALPHA;
+	bool EIGHT_BIT_MESH_DISTANCE_FIELDS;
+	bool PLATFORM_SUPPORTS_RENDERTARGET_WRITE_MASK;
+	bool PLATFORM_SUPPORTS_PER_PIXEL_DBUFFER_MASK;
+	bool PLATFORM_SUPPORTS_DISTANCE_FIELDS;
+	bool VIRTUAL_TEXTURE_FEEDBACK_FACTOR;
+	bool PROJECT_MOBILE_ENABLE_MOVABLE_SPOTLIGHTS;
+
+	// This is a sepcial one. The flag COMPILE_SHADERS_FOR_DEVELOPMENT is set by the cvar r.CompileShadersForDevelopment, but only
+	// if bAllowDevelopmentShaderCompile is true in the call to GlobalBeginCompileShader(). So the flag COMPILE_SHADERS_FOR_DEVELOPMENT_ALLOWED
+	// simply stores the result of r.CompileShadersForDevelopment. Then the logic later down the pipeline is:
+	// COMPILE_SHADERS_FOR_DEVELOPMENT = COMPILE_SHADERS_FOR_DEVELOPMENT_ALLOWED && bAllowDevelopmentShaderCompile;
+	bool COMPILE_SHADERS_FOR_DEVELOPMENT_ALLOWED;
+	bool bSupportsDualBlending;
+	int LegacyGBufferFormat;
+	bool bNeedVelocityDepth;
+};
+
+
+// maybe should renamt this to VertexFactoryDefines?
+struct FShaderLightmapPropertyDefines
+{
+	void ModifyEnvironment(FShaderCompilerEnvironment& OutEnvironment) const;
+
+	bool SIMPLE_FORWARD_SHADING;
+	bool SIMPLE_FORWARD_DIRECTIONAL_LIGHT;
+	bool LQ_TEXTURE_LIGHTMAP;
+	bool HQ_TEXTURE_LIGHTMAP;
+	bool CACHED_POINT_INDIRECT_LIGHTING;
+	bool STATICLIGHTING_TEXTUREMASK;
+	bool STATICLIGHTING_SIGNEDDISTANCEFIELD;
+	bool TRANSLUCENT_SELF_SHADOWING;
+	bool PRECOMPUTED_IRRADIANCE_VOLUME_LIGHTING;
+	bool CACHED_VOLUME_INDIRECT_LIGHTING;
+
+	bool WATER_MESH_FACTORY;
+	bool NIAGARA_MESH_FACTORY;
+	bool NIAGARA_MESH_INSTANCED;
+
+	bool PARTICLE_MESH_FACTORY;
+	bool PARTICLE_MESH_INSTANCED;
+
+	bool MANUAL_VERTEX_FETCH;
+};
+
+struct FShaderMaterialPropertyDefines
+{
+	//DECLARE_TYPE_LAYOUT(FShaderMaterialPropertyDefines, NonVirtual);
+
+	//void ModifyEnvironment(FShaderCompilerEnvironment& OutEnvironment) const;
+	//void WriteFrozenVertexFactoryParameters(FMemoryImageWriter& Writer, const TMemoryImagePtr<FShaderMaterialPropertyDefines>& InPropDefines) const;
+
+	uint8 MATERIAL_ENABLE_TRANSLUCENCY_FOGGING : 1;
+
+	uint8 MATERIALBLENDING_ANY_TRANSLUCENT : 1;
+	uint8 MATERIAL_USES_SCENE_COLOR_COPY : 1;
+	uint8 MATERIALBLENDING_MASKED_USING_COVERAGE : 1;
+
+	uint8 MATERIAL_COMPUTE_FOG_PER_PIXEL : 1;
+	uint8 MATERIAL_SHADINGMODEL_UNLIT : 1;
+
+	uint8 MATERIAL_SHADINGMODEL_DEFAULT_LIT : 1;
+	uint8 MATERIAL_SHADINGMODEL_SUBSURFACE : 1;
+	uint8 MATERIAL_SHADINGMODEL_PREINTEGRATED_SKIN : 1;
+	uint8 MATERIAL_SHADINGMODEL_SUBSURFACE_PROFILE : 1;
+	uint8 MATERIAL_SHADINGMODEL_CLEAR_COAT : 1;
+	uint8 MATERIAL_SHADINGMODEL_TWOSIDED_FOLIAGE : 1;
+	uint8 MATERIAL_SHADINGMODEL_HAIR : 1;
+	uint8 MATERIAL_SHADINGMODEL_CLOTH : 1;
+	uint8 MATERIAL_SHADINGMODEL_EYE : 1;
+	uint8 MATERIAL_SHADINGMODEL_SINGLELAYERWATER : 1;
+	uint8 MATERIAL_SHADINGMODEL_THIN_TRANSLUCENT : 1;
+
+	uint8 TRANSLUCENCY_LIGHTING_VOLUMETRIC_NONDIRECTIONAL : 1;
+	uint8 TRANSLUCENCY_LIGHTING_VOLUMETRIC_DIRECTIONAL : 1;
+	uint8 TRANSLUCENCY_LIGHTING_VOLUMETRIC_PERVERTEX_NONDIRECTIONAL : 1;
+	uint8 TRANSLUCENCY_LIGHTING_VOLUMETRIC_PERVERTEX_DIRECTIONAL : 1;
+	uint8 TRANSLUCENCY_LIGHTING_SURFACE_LIGHTINGVOLUME : 1;
+	uint8 TRANSLUCENCY_LIGHTING_SURFACE_FORWARDSHADING : 1;
+
+	uint8 EDITOR_PRIMITIVE_MATERIAL : 1;
+
+	uint8 MATERIAL_FULLY_ROUGH : 1;
+
+	uint8 MATERIALBLENDING_SOLID : 1;
+	uint8 MATERIALBLENDING_MASKED : 1;
+	uint8 MATERIALBLENDING_ALPHACOMPOSITE : 1;
+	uint8 MATERIALBLENDING_TRANSLUCENT : 1;
+	uint8 MATERIALBLENDING_ADDITIVE : 1;
+	uint8 MATERIALBLENDING_MODULATE : 1;
+	uint8 MATERIALBLENDING_ALPHAHOLDOUT : 1;
+
+	uint8 USES_EMISSIVE_COLOR : 1;
+	
+	uint8 REFRACTION_USE_INDEX_OF_REFRACTION : 1;
+	uint8 REFRACTION_USE_PIXEL_NORMAL_OFFSET : 1;
+
+	uint8 USE_DITHERED_LOD_TRANSITION_FROM_MATERIAL : 1;
+	uint8 MATERIAL_TWOSIDED : 1;
+	uint8 MATERIAL_TANGENTSPACENORMAL : 1;
+	uint8 GENERATE_SPHERICAL_PARTICLE_NORMALS : 1;
+	uint8 MATERIAL_USE_PREINTEGRATED_GF : 1;
+	uint8 MATERIAL_HQ_FORWARD_REFLECTIONS : 1;
+	uint8 MATERIAL_PLANAR_FORWARD_REFLECTIONS : 1;
+	uint8 MATERIAL_NONMETAL : 1;
+	uint8 MATERIAL_USE_LM_DIRECTIONALITY : 1;
+	uint8 MATERIAL_INJECT_EMISSIVE_INTO_LPV : 1;
+	uint8 MATERIAL_SSR : 1;
+	uint8 MATERIAL_CONTACT_SHADOWS : 1;
+	uint8 MATERIAL_BLOCK_GI : 1;
+	uint8 MATERIAL_DITHER_OPACITY_MASK : 1;
+	uint8 MATERIAL_NORMAL_CURVATURE_TO_ROUGHNESS : 1;
+	uint8 MATERIAL_ALLOW_NEGATIVE_EMISSIVECOLOR : 1;
+	uint8 MATERIAL_OUTPUT_OPACITY_AS_ALPHA : 1;
+	uint8 TRANSLUCENT_SHADOW_WITH_MASKED_OPACITY : 1;
+
+	uint32 DECALBLENDMODEID_VOLUMETRIC;
+	uint32 DECALBLENDMODEID_STAIN;
+	uint32 DECALBLENDMODEID_NORMAL;
+	uint32 DECALBLENDMODEID_EMISSIVE;
+	uint32 DECALBLENDMODEID_TRANSLUCENT;
+	uint32 DECALBLENDMODEID_AO;
+	uint32 DECALBLENDMODEID_ALPHACOMPOSITE;
+
+	uint8 MATERIAL_DOMAIN_SURFACE : 1;
+	uint8 MATERIAL_DOMAIN_DEFERREDDECAL : 1;
+	uint8 MATERIAL_DOMAIN_LIGHTFUNCTION : 1;
+	uint8 MATERIAL_DOMAIN_VOLUME : 1;
+	uint8 MATERIAL_DOMAIN_POSTPROCESS : 1;
+	uint8 MATERIAL_DOMAIN_UI : 1;
+	uint8 MATERIAL_DOMAIN_VIRTUALTEXTURE : 1;
+
+	uint8 USE_STENCIL_LOD_DITHER_DEFAULT : 1;
+
+	uint8 MATERIALDOMAIN_SURFACE : 1;
+	uint8 MATERIALDOMAIN_DEFERREDDECAL : 1;
+	uint8 MATERIALDOMAIN_LIGHTFUNCTION : 1;
+	uint8 MATERIALDOMAIN_POSTPROCESS : 1;
+	uint8 MATERIALDOMAIN_UI : 1;
+
+	uint8 OUT_BASECOLOR : 1;
+	uint8 OUT_BASECOLOR_NORMAL_SPECULAR : 1;
+	uint8 OUT_WORLDHEIGHT : 1;
+
+	uint8 STRATA_TRANSLUCENT_ENABLED : 1;
+	uint8 DUAL_SOURCE_COLOR_BLENDING_ENABLED : 1;
+
+	uint8 IS_MATERIAL_SHADER : 1;
+	//// end
+
+	uint32 MATERIALDECALRESPONSEMASK;
+
+	// new ones
+	uint8 IS_VIRTUAL_TEXTURE_MATERIAL : 1;
+	uint8 IS_DECAL : 1;
+	uint8 IS_BASE_PASS : 1;
+	uint32 DECAL_RENDERTARGET_COUNT;
+
+	uint8 bAllowDevelopmentShaderCompile : 1;// = Material->GetAllowDevelopmentShaderCompile();
+
+};
+
+struct FShaderMaterialDerivedDefines
+{
+	// Translucent materials need to compute fogging in the forward shading pass
+	// Materials that read from scene color skip getting fogged, because the contents of the scene color lookup have already been fogged
+	// This is not foolproof, as any additional color the material adds will then not be fogged correctly
+	bool TRANSLUCENCY_NEEDS_BASEPASS_FOGGING;
+
+	// With forward shading, fog always needs to be computed in the base pass to work correctly with MSAA
+	bool OPAQUE_NEEDS_BASEPASS_FOGGING;
+
+	bool NEEDS_BASEPASS_VERTEX_FOGGING;
+	bool NEEDS_BASEPASS_PIXEL_FOGGING;
+
+	// Volumetric fog interpolated per vertex gives very poor results, always sample the volumetric fog texture per-pixel
+	// Opaque materials in the deferred renderer get volumetric fog applied in a deferred fog pass
+	bool NEEDS_BASEPASS_PIXEL_VOLUMETRIC_FOGGING;
+
+	bool NEEDS_LIGHTMAP_COORDINATE;
+	bool NEEDS_LIGHTMAP;
+
+	bool USES_GBUFFER;
+
+	// Only some shader models actually need custom data.
+	bool WRITES_CUSTOMDATA_TO_GBUFFER;
+
+	// Based on GetPrecomputedShadowMasks()
+	// Note: WRITES_PRECSHADOWFACTOR_TO_GBUFFER is currently disabled because we use the precomputed shadow factor GBuffer outside of STATICLIGHTING_TEXTUREMASK to store UseSingleSampleShadowFromStationaryLights
+	bool GBUFFER_HAS_PRECSHADOWFACTOR;
+	bool WRITES_PRECSHADOWFACTOR_ZERO;
+	bool WRITES_PRECSHADOWFACTOR_TO_GBUFFER;
+
+	bool MATERIAL_SHADINGMODEL_THIN_TRANSLUCENT;
+
+	// If a primitive has static lighting, we assume it is not moving. If it is, it will be rerendered in an extra renderpass.
+	bool SUPPORTS_WRITING_VELOCITY_TO_BASE_PASS;
+	bool WRITES_VELOCITY_TO_GBUFFER;
+
+	// Use an extra interpolator for position
+	bool WRITES_VELOCITY_TO_GBUFFER_USE_POS_INTERPOLATOR;
+
+	bool FORCE_FULLY_ROUGH;
+	bool EDITOR_ALPHA2COVERAGE;
+	bool POST_PROCESS_SUBSURFACE;
+
+	bool TRANSLUCENCY_ANY_PERVERTEX_LIGHTING;
+	bool TRANSLUCENCY_ANY_VOLUMETRIC;
+	bool TRANSLUCENCY_PERVERTEX_LIGHTING_VOLUME;
+	bool TRANSLUCENCY_PERVERTEX_FORWARD_SHADING;
+
+	bool TESSELLATION_SUPPORTED;
+
+	bool PIXELSHADEROUTPUT_BASEPASS;
+		
+	bool PIXELSHADEROUTPUT_MRT0;
+	bool PIXELSHADEROUTPUT_MRT1;
+
+	bool PIXELSHADEROUTPUT_MRT2;
+	bool PIXELSHADEROUTPUT_MRT3;
+	bool PIXELSHADEROUTPUT_MRT4;
+	bool PIXELSHADEROUTPUT_MRT5;
+	bool PIXELSHADEROUTPUT_MRT6;
+	bool PIXELSHADEROUTPUT_A2C;
+	bool PIXELSHADEROUTPUT_COVERAGE;
+
+	bool SUPPORTS_PIXEL_COVERAGE;
+	bool MOBILE_EMULATION;
+
+	bool COMPILE_SHADERS_FOR_DEVELOPMENT;
+	bool USE_DEVELOPMENT_SHADERS;
+
+	bool USE_EDITOR_SHADERS;
+	bool USE_EDITOR_COMPOSITING;
+	bool MATERIALBLENDING_ANY_TRANSLUCENT;
+	bool IS_MESHPARTICLE_FACTORY;
+
+	bool PLATFORM_SUPPORTS_EDITOR_SHADERS;
+
+	bool THIN_TRANSLUCENT_USE_DUAL_BLEND;
+	bool MATERIAL_WORKS_WITH_DUAL_SOURCE_COLOR_BLENDING;
+};
+
+// These are also the platform defines
+struct FShaderCompilerDefines
+{
+	bool COMPILER_GLSL_ES3_1;
+	bool ES3_1_PROFILE;
+
+	bool COMPILER_GLSL;
+	bool COMPILER_GLSL_ES3_1_EXT;
+	bool ESDEFERRED_PROFILE;
+	bool GL4_PROFILE;
+
+	bool METAL_PROFILE;
+	bool VULKAN_PROFILE;
+	bool MAC;
+
+	bool PLATFORM_SUPPORTS_DEVELOPMENT_SHADERS;
+};
+
+
+FShaderMaterialDerivedDefines RENDERCORE_API CalculateDerivedMaterialParameters(
+	const FShaderMaterialPropertyDefines& Mat,
+	const FShaderLightmapPropertyDefines& Lightmap,
+	const FShaderGlobalDefines& SrcGlobal,
+	const FShaderCompilerDefines& Compiler,
+	ERHIFeatureLevel::Type FEATURE_LEVEL);
+
diff --git a/Engine/Source/Runtime/Renderer/Private/BasePassRendering.h b/Engine/Source/Runtime/Renderer/Private/BasePassRendering.h
--- a/Engine/Source/Runtime/Renderer/Private/BasePassRendering.h
+++ b/Engine/Source/Runtime/Renderer/Private/BasePassRendering.h
@@ -505,6 +505,10 @@ public:
 		OutEnvironment.SetDefine(TEXT("ENABLE_SKY_LIGHT"), bEnableSkyLight);
 		OutEnvironment.SetDefine(TEXT("PLATFORM_FORCE_SIMPLE_SKY_DIFFUSE"), ForceSimpleSkyDiffuse(Parameters.Platform));
 
+		// This define simply lets the compilation environment know that we are using BasePassPixelShader.usf, so that we can check for more
+		// complicated defines later in the compilation pipe.
+		OutEnvironment.SetDefine(TEXT("IS_BASE_PASS"), 1);
+
 		TBasePassPixelShaderBaseType<LightMapPolicyType>::ModifyCompilationEnvironment(Parameters, OutEnvironment);
 	}
 	
diff --git a/Engine/Source/Runtime/Renderer/Private/DecalRenderingShared.cpp b/Engine/Source/Runtime/Renderer/Private/DecalRenderingShared.cpp
--- a/Engine/Source/Runtime/Renderer/Private/DecalRenderingShared.cpp
+++ b/Engine/Source/Runtime/Renderer/Private/DecalRenderingShared.cpp
@@ -455,6 +455,8 @@ void FDecalRendering::SetDecalCompilationEnvironment(const FMaterialShaderPermut
 	OutEnvironment.SetDefine(TEXT("MATERIAL_DBUFFERA"), (bDBufferMask & 0x1) != 0);
 	OutEnvironment.SetDefine(TEXT("MATERIAL_DBUFFERB"), (bDBufferMask & 0x2) != 0);
 	OutEnvironment.SetDefine(TEXT("MATERIAL_DBUFFERC"), (bDBufferMask & 0x4) != 0);
+	// define so that we know to output the proper MRTs in the derived defines.
+	OutEnvironment.SetDefine(TEXT("IS_DECAL"), 1);
 }
 
 void FDecalRendering::SetEmissiveDBufferDecalCompilationEnvironment(const FMaterialShaderPermutationParameters& Parameters, FShaderCompilerEnvironment& OutEnvironment)
diff --git a/Engine/Source/Runtime/Renderer/Private/PlanarReflectionRendering.cpp b/Engine/Source/Runtime/Renderer/Private/PlanarReflectionRendering.cpp
--- a/Engine/Source/Runtime/Renderer/Private/PlanarReflectionRendering.cpp
+++ b/Engine/Source/Runtime/Renderer/Private/PlanarReflectionRendering.cpp
@@ -821,11 +821,13 @@ void FDeferredShadingSceneRenderer::RenderDeferredPlanarReflections(FRDGBuilder&
 	PassParameters->SceneTextures.SceneDepthTexture = SceneTextures.SceneDepthTexture;
 	PassParameters->SceneTextures.GBufferATexture = SceneTextures.GBufferATexture;
 	PassParameters->SceneTextures.GBufferBTexture = SceneTextures.GBufferBTexture;
-	if (IsHlslccShaderPlatform(GMaxRHIShaderPlatform))
-	{
-		// hlslcc doesn't remove all unused parameters
-		PassParameters->SceneTextures.GBufferCTexture = GraphBuilder.RegisterExternalTexture(GSystemTextures.BlackDummy);
-	}
+
+	PassParameters->SceneTextures.GBufferCTexture = SceneTextures.GBufferCTexture;
+	PassParameters->SceneTextures.GBufferDTexture = SceneTextures.GBufferDTexture;
+	PassParameters->SceneTextures.GBufferETexture = SceneTextures.GBufferETexture;
+	PassParameters->SceneTextures.GBufferFTexture = SceneTextures.GBufferFTexture;
+	PassParameters->SceneTextures.GBufferVelocityTexture = SceneTextures.GBufferVelocityTexture;
+
 	PassParameters->ViewUniformBuffer = View.ViewUniformBuffer;
 	PassParameters->RenderTargets[0] = FRenderTargetBinding(
 		ReflectionsOutputTexture, bClearReflectionsOutputTexture ? ERenderTargetLoadAction::EClear : ERenderTargetLoadAction::ELoad);
diff --git a/Engine/Source/Runtime/Renderer/Private/PostProcess/SceneRenderTargets.cpp b/Engine/Source/Runtime/Renderer/Private/PostProcess/SceneRenderTargets.cpp
--- a/Engine/Source/Runtime/Renderer/Private/PostProcess/SceneRenderTargets.cpp
+++ b/Engine/Source/Runtime/Renderer/Private/PostProcess/SceneRenderTargets.cpp
@@ -28,6 +28,14 @@
 #include "VT/VirtualTextureFeedback.h"
 #include "VisualizeTexture.h"
 #include "GpuDebugRendering.h"
+#include "GBufferInfo.h"
+#include "ShaderCompiler.h"
+
+static TAutoConsoleVariable<int32> CVarRSMResolution(
+	TEXT("r.LPV.RSMResolution"),
+	360,
+	TEXT("Reflective Shadow Map resolution (used for LPV) - higher values result in less aliasing artifacts, at the cost of performance"),
+	ECVF_Scalability | ECVF_RenderThreadSafe);
 
 /*-----------------------------------------------------------------------------
 FSceneRenderTargets
@@ -701,6 +709,187 @@ void FSceneRenderTargets::Allocate(FRHICommandListImmediate& RHICmdList, const F
 	AllocateRenderTargets(RHICmdList, ViewFamily.Views.Num());
 }
 
+static int32 FindGBufferTargetIndex(const FGBufferInfo& GBufferInfo, const FString& Name)
+{
+	for (int32 I = 0; I < GBufferInfo.NumTargets; I++)
+	{
+		if (GBufferInfo.Targets[I].TargetName.Compare(Name) == 0)
+		{
+			return I;
+		}
+	}
+
+	return -1;
+}
+
+// This code is copy/paste/modified from the existing logic. Where in the old code it would set the targets using this logic, now it's just checks and verifies.
+static void CheckLegacyGBufferFormat(const FGBufferInfo& GBufferInfo, EShaderPlatform ShaderPlatform, bool bAllocateVelocityGBuffer, bool bAllowTangentGBuffer, bool bAllowStaticLighting)
+{
+	const bool bUseGBuffer = IsUsingGBuffers(ShaderPlatform);
+
+	int32 TempOutVelocityRTIndex = -1;
+	int32 TempOutTangentRTIndex = -1;
+
+	// make sure that the GBuffer slots aligned with the old render target lobic
+	int32 MRTCount = 0;
+	MRTCount++;
+
+	if (bUseGBuffer)
+	{
+		check(GBufferInfo.Targets[MRTCount].TargetName.Compare(TEXT("GBufferA")) == 0);
+		MRTCount++;
+		check(GBufferInfo.Targets[MRTCount].TargetName.Compare(TEXT("GBufferB")) == 0);
+		MRTCount++;
+		check(GBufferInfo.Targets[MRTCount].TargetName.Compare(TEXT("GBufferC")) == 0);
+		MRTCount++;
+	}
+
+	// The velocity buffer needs to be bound before other optionnal rendertargets (when UseSelectiveBasePassOutputs() is true).
+	// Otherwise there is an issue on some AMD hardware where the target does not get updated. Seems to be related to the velocity buffer format as it works fine with other targets.
+	if (bAllocateVelocityGBuffer && !IsSimpleForwardShadingEnabled(ShaderPlatform))
+	{
+		TempOutVelocityRTIndex = MRTCount;
+		check(TempOutVelocityRTIndex == 4 || (!bUseGBuffer && TempOutVelocityRTIndex == 1)); // As defined in BasePassPixelShader.usf
+
+		check(GBufferInfo.Targets[MRTCount].TargetName.Compare(TEXT("Velocity")) == 0);
+		MRTCount++;
+	}
+	else
+	{
+		//OutVelocityRTIndex = -1;
+	}
+
+	if (bUseGBuffer && bAllowTangentGBuffer)
+	{
+		check(!bAllocateVelocityGBuffer);
+		TempOutTangentRTIndex = MRTCount;
+
+		check(GBufferInfo.Targets[MRTCount].TargetName.Compare(TEXT("GBufferF")) == 0);
+		MRTCount++;
+	}
+	else
+	{
+		TempOutTangentRTIndex = -1;
+	}
+
+	if (bUseGBuffer)
+	{
+		check(GBufferInfo.Targets[MRTCount].TargetName.Compare(TEXT("GBufferD")) == 0);
+		MRTCount++;
+
+		if (bAllowStaticLighting)
+		{
+			check(MRTCount == (bAllocateVelocityGBuffer || bAllowTangentGBuffer ? 6 : 5)); // As defined in BasePassPixelShader.usf
+			check(GBufferInfo.Targets[MRTCount].TargetName.Compare(TEXT("GBufferE")) == 0);
+			MRTCount++;
+		}
+	}
+}
+
+// temporary #if, keeping backup for emergency testing in case it's necessary briefly
+#if 1
+
+int32 FSceneRenderTargets::GetGBufferRenderTargets(const TRefCountPtr<IPooledRenderTarget>* OutRenderTargets[MaxSimultaneousRenderTargets], int32& OutVelocityRTIndex) const
+{
+	int32 OutTangentRTIndex = -1;
+	FGBufferInfo GBufferInfo = {};
+
+	const EShaderPlatform ShaderPlatform = GetFeatureLevelShaderPlatform(CurrentFeatureLevel);
+	const bool bUseGBuffer = IsUsingGBuffers(ShaderPlatform);
+
+	int32 NumTargets = 0;
+
+	if (!bUseGBuffer)
+	{
+		OutTangentRTIndex = -1; // no tangent allowed
+
+		if (bAllocateVelocityGBuffer && !IsSimpleForwardShadingEnabled(ShaderPlatform))
+		{
+			// if velocity buffer SceneColor is 0 and Velocity is 1
+			OutRenderTargets[0] = &GetSceneColor();
+			OutRenderTargets[1] = &SceneVelocity;
+			OutVelocityRTIndex = 1;
+			NumTargets = 2;
+		}
+		else
+		{
+			// otherwise, SceneColor is 0, with no velocity
+			OutRenderTargets[0] = &GetSceneColor();
+			OutVelocityRTIndex = -1;
+			NumTargets = 1;
+		}
+	}
+	else
+	{
+		{
+			FGBufferParams Params = FShaderCompileUtilities::FetchGBufferParamsRuntime(ShaderPlatform);
+			GBufferInfo = FetchFullGBufferInfo(Params);
+		}
+
+		// validate legacy GBuffer
+		bool bAllowTangentGBuffer = false; // in the past, this was an option, then it was taken out for 4.26, and now we're putting it back in
+		CheckLegacyGBufferFormat(GBufferInfo, ShaderPlatform, bAllocateVelocityGBuffer, bAllowTangentGBuffer, bAllowStaticLighting);
+
+		{
+			OutRenderTargets[0] = &GetSceneColor();
+
+			int32 IndexGBufferA = FindGBufferTargetByName(GBufferInfo, "GBufferA");
+			int32 IndexGBufferB = FindGBufferTargetByName(GBufferInfo, "GBufferB");
+			int32 IndexGBufferC = FindGBufferTargetByName(GBufferInfo, "GBufferC");
+			int32 IndexGBufferD = FindGBufferTargetByName(GBufferInfo, "GBufferD");
+			int32 IndexGBufferE = FindGBufferTargetByName(GBufferInfo, "GBufferE");
+			int32 IndexGBufferF = FindGBufferTargetByName(GBufferInfo, "GBufferF");
+			int32 IndexVelocity = FindGBufferTargetByName(GBufferInfo, "Velocity");
+
+			// sanity checks
+			if (bUseGBuffer)
+			{
+				check(IndexGBufferA >= 0);
+				check(IndexGBufferB >= 0);
+				check(IndexGBufferC >= 0);
+			}
+
+			if (bAllocateVelocityGBuffer && !IsSimpleForwardShadingEnabled(ShaderPlatform))
+			{
+				check(IndexVelocity >= 0);
+			}
+
+			if (bUseGBuffer && bAllowTangentGBuffer)
+			{
+				check(IndexGBufferF >= 0);
+			}
+
+			if (bUseGBuffer)
+			{
+				check(IndexGBufferD >= 0);
+
+				if (bAllowStaticLighting)
+				{
+					check(IndexGBufferE >= 0);
+				}
+			}
+
+			if (IndexGBufferA >= 0) OutRenderTargets[IndexGBufferA] = &GBufferA;
+			if (IndexGBufferB >= 0) OutRenderTargets[IndexGBufferB] = &GBufferB;
+			if (IndexGBufferC >= 0) OutRenderTargets[IndexGBufferC] = &GBufferC;
+			if (IndexGBufferD >= 0) OutRenderTargets[IndexGBufferD] = &GBufferD;
+			if (IndexGBufferE >= 0) OutRenderTargets[IndexGBufferE] = &GBufferE;
+			if (IndexGBufferF >= 0) OutRenderTargets[IndexGBufferF] = &GBufferF;
+			if (IndexVelocity >= 0) OutRenderTargets[IndexVelocity] = &SceneVelocity;
+
+			OutTangentRTIndex = IndexGBufferF;
+			OutVelocityRTIndex = IndexVelocity;
+		}
+
+		check(GBufferInfo.NumTargets <= MaxSimultaneousRenderTargets);
+		NumTargets = GBufferInfo.NumTargets;
+	}
+
+	return NumTargets;
+}
+
+#else
+
 int32 FSceneRenderTargets::GetGBufferRenderTargets(const TRefCountPtr<IPooledRenderTarget>* OutRenderTargets[MaxSimultaneousRenderTargets], int32& OutVelocityRTIndex) const
 {
 	int32 MRTCount = 0;
@@ -744,6 +933,9 @@ int32 FSceneRenderTargets::GetGBufferRenderTargets(const TRefCountPtr<IPooledRen
 	return MRTCount;
 }
 
+#endif
+
+
 int32 FSceneRenderTargets::FillGBufferRenderPassInfo(ERenderTargetLoadAction ColorLoadAction, FRHIRenderPassInfo& OutRenderPassInfo, int32& OutVelocityRTIndex) const
 {
 	const TRefCountPtr<IPooledRenderTarget>* RenderTargets[MaxSimultaneousRenderTargets];
@@ -1090,6 +1282,52 @@ void FSceneRenderTargets::AllocGBufferTargets(FRHICommandList& RHICmdList)
 	AllocGBufferTargets(RHICmdList, TexCreate_None);
 }
 
+static FPooledRenderTargetDesc GetDescFromRenderTarget(FIntPoint BufferSize, const FGBufferTarget& Target, ETextureCreateFlags AddTargetableFlags)
+{
+
+	EPixelFormat PixelFormat = PF_Unknown;
+	switch(Target.TargetType)
+	{
+	case GBT_Unorm_8_8_8_8:
+		PixelFormat = PF_B8G8R8A8;
+		break;
+	case GBT_Unorm_11_11_10:
+		PixelFormat = PF_FloatR11G11B10;
+		break;
+	case GBT_Unorm_10_10_10_2:
+		PixelFormat = PF_A2B10G10R10;
+		break;
+	case GBT_Float_16_16:
+		PixelFormat = PF_G16R16;
+		break;
+	case GBT_Float_16_16_16_16:
+		PixelFormat = PF_A16B16G16R16;
+		break;
+	case GBT_Invalid:
+	default:
+		check(0);
+		PixelFormat = PF_Unknown;
+		break;
+	}
+
+	ETextureCreateFlags Flags = Target.bIsSrgb ? TexCreate_SRGB : TexCreate_None;
+
+	ETextureCreateFlags TargetFlags = (Target.bIsUsingExtraFlags ? AddTargetableFlags : TexCreate_None);
+
+	if (Target.bIsRenderTargetable)
+	{
+		TargetFlags |= TexCreate_RenderTargetable;
+	}
+
+	if (Target.bIsShaderResource)
+	{
+		TargetFlags |= TexCreate_ShaderResource;
+	}
+
+	FPooledRenderTargetDesc Desc(FPooledRenderTargetDesc::Create2DDesc(BufferSize, PixelFormat, FClearValueBinding::Transparent, Flags, TargetFlags , false));
+	return Desc;
+}
+
 void FSceneRenderTargets::AllocGBufferTargets(FRHICommandList& RHICmdList, ETextureCreateFlags AddTargetableFlags)
 {	
 	// AdjustGBufferRefCount +1 doesn't match -1 (within the same frame)
@@ -1105,6 +1343,100 @@ void FSceneRenderTargets::AllocGBufferTargets(FRHICommandList& RHICmdList, EText
 	const EShaderPlatform ShaderPlatform = GetFeatureLevelShaderPlatform(CurrentFeatureLevel);
 	const bool bUseGBuffer = IsUsingGBuffers(ShaderPlatform);
 	const bool bCanReadGBufferUniforms = (bUseGBuffer || IsSimpleForwardShadingEnabled(ShaderPlatform)) && CurrentFeatureLevel >= ERHIFeatureLevel::SM5;
+
+#if 1
+	if (!bUseGBuffer)
+	{
+		// Uses SceneColor and possibly Velocity
+
+		// Create the world-space normal g-buffer.
+		{
+			FPooledRenderTargetDesc Desc(FPooledRenderTargetDesc::Create2DDesc(BufferSize, GetGBufferAFormat(), FClearValueBinding::Transparent, TexCreate_None, TexCreate_RenderTargetable | TexCreate_ShaderResource | AddTargetableFlags , false));
+			Desc.Flags |= GFastVRamConfig.GBufferA;
+			GRenderTargetPool.FindFreeElement(RHICmdList, Desc, GBufferA, TEXT("GBufferA"));
+		}
+
+		if (bAllocateVelocityGBuffer)
+		{
+			FPooledRenderTargetDesc VelocityRTDesc = Translate(FVelocityRendering::GetRenderTargetDesc(ShaderPlatform));
+			VelocityRTDesc.Flags |= GFastVRamConfig.GBufferVelocity;
+			GRenderTargetPool.FindFreeElement(RHICmdList, VelocityRTDesc, SceneVelocity, TEXT("GBufferVelocity"));
+		}
+	}
+	else
+	{
+		// TODO: clean this up later
+		FGBufferInfo GBufferInfo;
+		{
+			FGBufferParams Params = FShaderCompileUtilities::FetchGBufferParamsRuntime(ShaderPlatform);
+			GBufferInfo = FetchFullGBufferInfo(Params);
+		}
+
+		const int32 IndexA = FindGBufferTargetByName(GBufferInfo, "GBufferA");
+		const int32 IndexB = FindGBufferTargetByName(GBufferInfo, "GBufferB");
+		const int32 IndexC = FindGBufferTargetByName(GBufferInfo, "GBufferC");
+		const int32 IndexD = FindGBufferTargetByName(GBufferInfo, "GBufferD");
+		const int32 IndexE = FindGBufferTargetByName(GBufferInfo, "GBufferE");
+		const int32 IndexVelocity = FindGBufferTargetByName(GBufferInfo, "GBufferVelocity");
+
+		if (IndexA >= 0)
+		{
+			//FPooledRenderTargetDesc Desc(FPooledRenderTargetDesc::Create2DDesc(BufferSize, GetGBufferAFormat(), FClearValueBinding::Transparent, TexCreate_None, TexCreate_RenderTargetable | TexCreate_ShaderResource | AddTargetableFlags , false));
+			FPooledRenderTargetDesc Desc = GetDescFromRenderTarget(BufferSize, GBufferInfo.Targets[IndexA], AddTargetableFlags);
+			Desc.Flags |= GFastVRamConfig.GBufferA;
+			GRenderTargetPool.FindFreeElement(RHICmdList, Desc, GBufferA, TEXT("GBufferA"));
+		}
+
+		// Create the specular color and power g-buffer.
+		if (IndexB >= 0)
+		{
+			//FPooledRenderTargetDesc Desc(FPooledRenderTargetDesc::Create2DDesc(BufferSize, GetGBufferBFormat(), FClearValueBinding::Transparent, TexCreate_None, TexCreate_RenderTargetable | TexCreate_ShaderResource | AddTargetableFlags, false));
+			FPooledRenderTargetDesc Desc = GetDescFromRenderTarget(BufferSize, GBufferInfo.Targets[IndexB], AddTargetableFlags);
+			Desc.Flags |= GFastVRamConfig.GBufferB;
+			GRenderTargetPool.FindFreeElement(RHICmdList, Desc, GBufferB, TEXT("GBufferB"));
+		}
+
+		// Create the diffuse color g-buffer.
+		if (IndexC >= 0)
+		{
+			//FPooledRenderTargetDesc Desc(FPooledRenderTargetDesc::Create2DDesc(BufferSize, GetGBufferCFormat(), FClearValueBinding::Transparent, TexCreate_SRGB, TexCreate_RenderTargetable | TexCreate_ShaderResource | AddTargetableFlags, false));
+			FPooledRenderTargetDesc Desc = GetDescFromRenderTarget(BufferSize, GBufferInfo.Targets[IndexC], AddTargetableFlags);
+			Desc.Flags |= GFastVRamConfig.GBufferC;
+			GRenderTargetPool.FindFreeElement(RHICmdList, Desc, GBufferC, TEXT("GBufferC"));
+		}
+
+		// Create the mask g-buffer (e.g. SSAO, subsurface scattering, wet surface mask, skylight mask, ...).
+		if (IndexD >= 0)
+		{
+			//FPooledRenderTargetDesc Desc(FPooledRenderTargetDesc::Create2DDesc(BufferSize, GetGBufferDFormat(), FClearValueBinding::Transparent, TexCreate_None, TexCreate_RenderTargetable | TexCreate_ShaderResource | AddTargetableFlags, false));
+			FPooledRenderTargetDesc Desc = GetDescFromRenderTarget(BufferSize, GBufferInfo.Targets[IndexD], AddTargetableFlags);
+			Desc.Flags |= GFastVRamConfig.GBufferD;
+			GRenderTargetPool.FindFreeElement(RHICmdList, Desc, GBufferD, TEXT("GBufferD"));
+		}
+
+		if (IndexE >= 0)
+		{
+			//FPooledRenderTargetDesc Desc(FPooledRenderTargetDesc::Create2DDesc(BufferSize, GetGBufferEFormat(), FClearValueBinding::Transparent, TexCreate_None, TexCreate_RenderTargetable | TexCreate_ShaderResource, false));
+			FPooledRenderTargetDesc Desc = GetDescFromRenderTarget(BufferSize, GBufferInfo.Targets[IndexE], AddTargetableFlags);
+			Desc.Flags |= GFastVRamConfig.GBufferE;
+			GRenderTargetPool.FindFreeElement(RHICmdList, Desc, GBufferE, TEXT("GBufferE"));
+		}
+
+		// otherwise we have a severe problem
+		check(GBufferA);
+
+		if (IndexVelocity >= 0)
+		{
+			//FPooledRenderTargetDesc VelocityRTDesc = Translate(FVelocityRendering::GetRenderTargetDesc(ShaderPlatform));
+			FPooledRenderTargetDesc Desc = GetDescFromRenderTarget(BufferSize, GBufferInfo.Targets[IndexVelocity], AddTargetableFlags);
+			Desc.Flags |= GFastVRamConfig.GBufferVelocity;
+			GRenderTargetPool.FindFreeElement(RHICmdList, Desc, SceneVelocity, TEXT("GBufferVelocity"));
+		}
+
+	}
+
+	FVelocityRendering::GetRenderTargetDesc(ShaderPlatform);
+#else
 	if (bUseGBuffer)
 	{
 		// Create the world-space normal g-buffer.
@@ -1152,6 +1484,8 @@ void FSceneRenderTargets::AllocGBufferTargets(FRHICommandList& RHICmdList, EText
 		VelocityRTDesc.Flags |= GFastVRamConfig.GBufferVelocity;
 		GRenderTargetPool.FindFreeElement(RHICmdList, VelocityRTDesc, SceneVelocity, TEXT("GBufferVelocity"));
 	}
+#endif
+
 
 	GBufferRefCount = 1;
 }
@@ -1204,6 +1538,10 @@ void FSceneRenderTargets::AdjustGBufferRefCount(FRHICommandList& RHICmdList, int
 	}	
 }
 
+
+
+
+
 void FSceneRenderTargets::CleanUpEditorPrimitiveTargets()
 {
 	EditorPrimitivesDepth.SafeRelease();
diff --git a/Engine/Source/Runtime/Renderer/Private/PostProcess/SceneRenderTargets.h b/Engine/Source/Runtime/Renderer/Private/PostProcess/SceneRenderTargets.h
--- a/Engine/Source/Runtime/Renderer/Private/PostProcess/SceneRenderTargets.h
+++ b/Engine/Source/Runtime/Renderer/Private/PostProcess/SceneRenderTargets.h
@@ -410,6 +410,8 @@ public:
 	void AdjustGBufferRefCount(FRHICommandList& RHICmdList, int Delta);
 
 	void PreallocGBufferTargets();
+	bool AllocateGBufferTarget(const FName& TargetName, ETextureCreateFlags TargetFlags) const;
+
 	EPixelFormat GetGBufferAFormat() const;
 	EPixelFormat GetGBufferBFormat() const;
 	EPixelFormat GetGBufferCFormat() const;
