const correctness fix

diff --git a/Engine/Source/Runtime/Engine/Classes/Animation/AnimLinkableElement.h b/Engine/Source/Runtime/Engine/Classes/Animation/AnimLinkableElement.h
--- a/Engine/Source/Runtime/Engine/Classes/Animation/AnimLinkableElement.h
+++ b/Engine/Source/Runtime/Engine/Classes/Animation/AnimLinkableElement.h
@@ -98,7 +98,7 @@ struct FAnimLinkableElement
 	ENGINE_API virtual void SetTime(float NewTime, EAnimLinkMethod::Type ReferenceFrame = EAnimLinkMethod::Absolute);
 
 	/** Gets the sequence this element is linked to */
-	ENGINE_API const UAnimSequenceBase* GetLinkedSequence() {return LinkedSequence;}
+	ENGINE_API const UAnimSequenceBase* GetLinkedSequence() const {return LinkedSequence;}
 
 	/** Gets the Montage this element is linked to, if any */
 	ENGINE_API const UAnimMontage* GetLinkedMontage() const { return LinkedMontage; }
