Make sure we kick to the GPU as well when flushing

diff --git a/Engine/Source/Runtime/Engine/Private/GPUSkinCache.cpp b/Engine/Source/Runtime/Engine/Private/GPUSkinCache.cpp
--- a/Engine/Source/Runtime/Engine/Private/GPUSkinCache.cpp
+++ b/Engine/Source/Runtime/Engine/Private/GPUSkinCache.cpp
@@ -1459,6 +1459,8 @@ void FGPUSkinCache::ProcessRayTracingGeometryToUpdate(
 				Initializer.Segments.Add(Segment);
 			}
 
+			UE_LOG(LogSkinCache, Log, TEXT("Build RT BVH with %d triangles"), TrianglesCount);
+
 			FGPUSkinCache::GetRayTracingSegmentVertexBuffers(*SkinCacheEntry, Initializer.Segments);
 
 			// Flush pending resource barriers before BVH is built for the first time
@@ -1474,6 +1476,7 @@ void FGPUSkinCache::ProcessRayTracingGeometryToUpdate(
 				if (RayTracingGeometryMemoryPendingRelease >= GMemoryLimitForBatchedRayTracingGeometryUpdates * 1024ull * 1024ull)
 				{
 					RayTracingGeometryMemoryPendingRelease = 0;
+					RHICmdList.SubmitCommandsHint();
 					RHICmdList.ImmediateFlush(EImmediateFlushType::FlushRHIThreadFlushResources);
 					UE_LOG(LogSkinCache, Display, TEXT("Flushing RHI resource pending deletes due to %d MB limit"), GMemoryLimitForBatchedRayTracingGeometryUpdates);
 				}
