RHITransition fixes for incorrect state

diff --git a/Engine/Source/Runtime/Engine/Private/UnrealClient.cpp b/Engine/Source/Runtime/Engine/Private/UnrealClient.cpp
--- a/Engine/Source/Runtime/Engine/Private/UnrealClient.cpp
+++ b/Engine/Source/Runtime/Engine/Private/UnrealClient.cpp
@@ -1684,9 +1684,19 @@ const TArray<FColor>& FViewport::GetRawHitProxyData(FIntRect InRect)
 		ENQUEUE_RENDER_COMMAND(UpdateHitProxyRTCommand)(
 			[HitProxyMapPtr](FRHICommandListImmediate& RHICmdList)
 			{
+				FTexture2DRHIRef RenderTargetTexture = HitProxyMapPtr->GetRenderTargetTexture();
+
+				// Should not be multisampled, so can safely use copy src as state instead of taking care of ResolveSrc here
+				check(!RenderTargetTexture->IsMultisampled());
+				
+				// Keep in copy source because 2 resolve calls are done on the render target (skip one extra transition call)
+				FResolveParams ResolveParams;
+				ResolveParams.SourceAccessFinal = ERHIAccess::CopySrc;
+
 				// Copy (resolve) the rendered thumbnail from the render target to its texture
-				RHICmdList.CopyToResolveTarget(HitProxyMapPtr->GetRenderTargetTexture(), HitProxyMapPtr->GetHitProxyTexture(), FResolveParams());
-				RHICmdList.CopyToResolveTarget(HitProxyMapPtr->GetRenderTargetTexture(), HitProxyMapPtr->GetHitProxyCPUTexture(), FResolveParams());
+				RHICmdList.Transition(FRHITransitionInfo(RenderTargetTexture, ERHIAccess::Unknown, ERHIAccess::CopySrc));
+				RHICmdList.CopyToResolveTarget(RenderTargetTexture, HitProxyMapPtr->GetHitProxyTexture(), ResolveParams);
+				RHICmdList.CopyToResolveTarget(RenderTargetTexture, HitProxyMapPtr->GetHitProxyCPUTexture(), FResolveParams());
 			});
 
 		ENQUEUE_RENDER_COMMAND(EndDrawingCommand)(
diff --git a/Engine/Source/Runtime/Renderer/Private/DistanceFieldShadowing.cpp b/Engine/Source/Runtime/Renderer/Private/DistanceFieldShadowing.cpp
--- a/Engine/Source/Runtime/Renderer/Private/DistanceFieldShadowing.cpp
+++ b/Engine/Source/Runtime/Renderer/Private/DistanceFieldShadowing.cpp
@@ -817,7 +817,7 @@ void ScatterObjectsToShadowTiles(
 	TArray<FRHITransitionInfo> UAVTransitionInfos;
 	for (FRHIUnorderedAccessView* UAV : UAVs)
 	{
-		UAVTransitionInfos.Add(FRHITransitionInfo(UAV, ERHIAccess::Unknown, ERHIAccess::UAVCompute));
+		UAVTransitionInfos.Add(FRHITransitionInfo(UAV, ERHIAccess::Unknown, ERHIAccess::UAVMask));
 	}
 	RHICmdList.Transition(MakeArrayView(UAVTransitionInfos.GetData(), UAVTransitionInfos.Num()));
 
diff --git a/Engine/Source/Runtime/Renderer/Private/SceneHitProxyRendering.cpp b/Engine/Source/Runtime/Renderer/Private/SceneHitProxyRendering.cpp
--- a/Engine/Source/Runtime/Renderer/Private/SceneHitProxyRendering.cpp
+++ b/Engine/Source/Runtime/Renderer/Private/SceneHitProxyRendering.cpp
@@ -291,6 +291,8 @@ static void DoRenderHitProxies(
 		{
 			Nanite::DrawHitProxies(RHICmdList, *SceneRenderer->Scene, Views[ViewIndex], NaniteRasterResults[ViewIndex], HitProxyRT, HitProxyDepthRT);
 		}
+
+		TransitionRenderPassTargets(RHICmdList, RPInfo);
 	}
 
 	const auto FeatureLevel = SceneRenderer->FeatureLevel;
diff --git a/Engine/Source/Runtime/Renderer/Private/SystemTextures.cpp b/Engine/Source/Runtime/Renderer/Private/SystemTextures.cpp
--- a/Engine/Source/Runtime/Renderer/Private/SystemTextures.cpp
+++ b/Engine/Source/Runtime/Renderer/Private/SystemTextures.cpp
@@ -257,7 +257,8 @@ void FSystemTextures::InitializeFeatureLevelDependentTextures(FRHICommandListImm
 			Desc.Extent.X * Desc.Extent.Y * sizeof(BlackBytes),
 			BlackBytes);
 
-		RHICmdList.Transition(FRHITransitionInfo(VolumetricBlackDummy->GetTargetableRHI(), ERHIAccess::RTV, ERHIAccess::SRVMask));
+		// UpdateTexture3D before and after state is currently undefined
+		RHICmdList.Transition(FRHITransitionInfo(VolumetricBlackDummy->GetTargetableRHI(), ERHIAccess::Unknown, ERHIAccess::SRVMask));
 	}
 
 	if (CurrentFeatureLevel < ERHIFeatureLevel::SM5 && InFeatureLevel >= ERHIFeatureLevel::SM5)
@@ -271,7 +272,8 @@ void FSystemTextures::InitializeFeatureLevelDependentTextures(FRHICommandListImm
 		FUpdateTextureRegion3D Region(0, 0, 0, 0, 0, 0, Desc.Extent.X, Desc.Extent.Y, Desc.Depth);
 		RHICmdList.UpdateTexture3D((FTexture3DRHIRef&)HairLUT0->GetRenderTargetItem().ShaderResourceTexture, 0, Region, Desc.Extent.X * sizeof(BlackBytes), Desc.Extent.X * Desc.Extent.Y * sizeof(BlackBytes), BlackBytes);
 
-		RHICmdList.Transition(FRHITransitionInfo(HairLUT0->GetRenderTargetItem().ShaderResourceTexture, ERHIAccess::RTV, ERHIAccess::SRVMask));
+		// UpdateTexture3D before and after state is currently undefined
+		RHICmdList.Transition(FRHITransitionInfo(HairLUT0->GetRenderTargetItem().ShaderResourceTexture, ERHIAccess::Unknown, ERHIAccess::SRVMask));
 		HairLUT1 = HairLUT0;
 		HairLUT2 = HairLUT0;
 	}
