Make materials work with strata when created using the material layer/attribute workflow.

tested shooter and cooked console

diff --git a/Engine/Source/Runtime/Engine/Private/Materials/MaterialShared.cpp b/Engine/Source/Runtime/Engine/Private/Materials/MaterialShared.cpp
--- a/Engine/Source/Runtime/Engine/Private/Materials/MaterialShared.cpp
+++ b/Engine/Source/Runtime/Engine/Private/Materials/MaterialShared.cpp
@@ -1523,7 +1523,13 @@ bool FMaterialResource::IsStrataMaterial() const
 {
 	// STRATA_TODO IsStrataMaterial should go away once Strata implementation is finished
 	const URendererSettings* RendererSettings = GetDefault<URendererSettings>();
-	return RendererSettings && RendererSettings->bEnableStrata ? Material->HasStrataFrontMaterialConnected() : false;
+	if (RendererSettings && RendererSettings->bEnableStrata)
+	{
+		return Material->bUseMaterialAttributes ?
+			Material->MaterialAttributes.IsConnected(MP_FrontMaterial) :
+			Material->HasStrataFrontMaterialConnected();
+	}
+	return false;
 }
 
 bool FMaterialResource::RequiresSynchronousCompilation() const
