Fixed material attribute blending order with Strata to match the regular material.

tested shooter and cooked console

diff --git a/Engine/Source/Runtime/Engine/Classes/Materials/MaterialExpressionStrata.h b/Engine/Source/Runtime/Engine/Classes/Materials/MaterialExpressionStrata.h
--- a/Engine/Source/Runtime/Engine/Classes/Materials/MaterialExpressionStrata.h
+++ b/Engine/Source/Runtime/Engine/Classes/Materials/MaterialExpressionStrata.h
@@ -12,12 +12,12 @@
  * Compile a special blend function for strata when blending material attribute
  *
  * @param Compiler				The compiler to add code to
- * @param Foreground			Entry A, has a bigger impact when Alpha is close to 1
- * @param Background			Entry B, has a bigger impact when Alpha is close to 0
+ * @param Foreground			Entry A, has a bigger impact when Alpha is close to 0
+ * @param Background			Entry B, has a bigger impact when Alpha is close to 1
  * @param Alpha					Blend factor [0..1], when 0
  * @return						Index to a new code chunk
  */
-extern int32 CompileStrataBlendFunction(FMaterialCompiler* Compiler, int32 Foreground, int32 Background, int32 Alpha);
+extern int32 CompileStrataBlendFunction(FMaterialCompiler* Compiler, const int32 A, const int32 B, const int32 Alpha);
 
 
 ///////////////////////////////////////////////////////////////////////////////
diff --git a/Engine/Source/Runtime/Engine/Private/Materials/MaterialExpressions.cpp b/Engine/Source/Runtime/Engine/Private/Materials/MaterialExpressions.cpp
--- a/Engine/Source/Runtime/Engine/Private/Materials/MaterialExpressions.cpp
+++ b/Engine/Source/Runtime/Engine/Private/Materials/MaterialExpressions.cpp
@@ -598,13 +598,13 @@ int32 CompileShadingModelBlendFunction(FMaterialCompiler* Compiler, const int32
 	return Compiler->If(Alpha, MidPoint, B, INDEX_NONE, A, INDEX_NONE);
 }
 
-int32 CompileStrataBlendFunction(FMaterialCompiler* Compiler, int32 Foreground, int32 Background, int32 Alpha)
+int32 CompileStrataBlendFunction(FMaterialCompiler* Compiler, const int32 A, const int32 B, const int32 Alpha)
 {
-	if (Foreground == INDEX_NONE || Background == INDEX_NONE || Alpha == INDEX_NONE)
+	if (A == INDEX_NONE || B == INDEX_NONE || Alpha == INDEX_NONE)
 	{
 		return INDEX_NONE;
 	}
-	return Compiler->StrataHorizontalMixing(Foreground, Background, Alpha);
+	return Compiler->StrataHorizontalMixing(B, A, Alpha); // B is foreground (fully visible when Alpha=1), and A is background (fully visible when Alpha=0)
 }
 
 UMaterialExpression::UMaterialExpression(const FObjectInitializer& ObjectInitializer)
