World Partition: filter out unrelevant actors when loading partitioned level instances.
Also assign a default runtime hash if null.

diff --git a/Engine/Source/Runtime/Engine/Private/WorldPartition/WorldPartition.cpp b/Engine/Source/Runtime/Engine/Private/WorldPartition/WorldPartition.cpp
--- a/Engine/Source/Runtime/Engine/Private/WorldPartition/WorldPartition.cpp
+++ b/Engine/Source/Runtime/Engine/Private/WorldPartition/WorldPartition.cpp
@@ -39,6 +39,7 @@
 #include "WorldPartition/WorldPartitionEditorHash.h"
 #include "WorldPartition/WorldPartitionEditorSpatialHash.h"
 #include "WorldPartition/WorldPartitionRuntimeHash.h"
+#include "WorldPartition/WorldPartitionRuntimeSpatialHash.h"
 #include "Misc/Base64.h"
 #endif //WITH_EDITOR
 
@@ -215,6 +216,13 @@ void UWorldPartition::Initialize(UWorld* InWorld, const FTransform& InTransform)
 		}
 	}
 
+	if (!RuntimeHash)
+	{
+		UWorldPartitionRuntimeSpatialHash* RuntimeSpatialHash = NewObject<UWorldPartitionRuntimeSpatialHash>(this);
+		RuntimeSpatialHash->SetDefaultValues();
+		RuntimeHash = RuntimeSpatialHash;
+	}
+
 	if (bEditorOnly || !IsMainWorldPartition())
 	{
 		TSet<FName> AllLayersNames;
@@ -883,7 +891,20 @@ bool UWorldPartition::ShouldHandleAssetEvent(const FAssetData& InAssetData)
 		return false;
 	}
 
-	return true;
+	// Only handle assets that belongs to our level
+	auto RemoveAfterFirstDot = [](const FString& InValue)
+	{
+		int32 DotIndex;
+		if (InValue.FindChar(TEXT('.'), DotIndex))
+		{
+			return InValue.LeftChop(InValue.Len() - DotIndex);
+		}
+		return InValue;
+	};
+
+	const FString ThisLevelPath = RemoveAfterFirstDot(GetPackage()->FileName.ToString());
+	const FString AssetLevelPath = RemoveAfterFirstDot(InAssetData.ObjectPath.ToString());
+	return (ThisLevelPath == AssetLevelPath);
 }
 
 TUniquePtr<FWorldPartitionActorDesc> UWorldPartition::GetActorDescriptor(const FAssetData& InAssetData)
