Added back deprecated LPV properties to Post process settings to fix blueprints. Removed old deprecated properties.

#rnx

diff --git a/Engine/Source/Runtime/Engine/Classes/Engine/Scene.h b/Engine/Source/Runtime/Engine/Classes/Engine/Scene.h
--- a/Engine/Source/Runtime/Engine/Classes/Engine/Scene.h
+++ b/Engine/Source/Runtime/Engine/Classes/Engine/Scene.h
@@ -1012,6 +1012,51 @@ struct FPostProcessSettings
 	UPROPERTY(EditAnywhere, BlueprintReadWrite, Category = Overrides, meta = (PinHiddenByDefault, InlineEditConditionToggle))
 	uint32 bOverride_RayTracingAORadius : 1;
 
+	UPROPERTY()
+	uint8 bOverride_LPVIntensity_DEPRECATED:1;
+
+	UPROPERTY()
+	uint8 bOverride_LPVDirectionalOcclusionIntensity_DEPRECATED:1;
+
+	UPROPERTY()
+	uint8 bOverride_LPVDirectionalOcclusionRadius_DEPRECATED:1;
+
+	UPROPERTY()
+	uint8 bOverride_LPVDiffuseOcclusionExponent_DEPRECATED:1;
+
+	UPROPERTY()
+	uint8 bOverride_LPVSpecularOcclusionExponent_DEPRECATED:1;
+
+	UPROPERTY()
+	uint8 bOverride_LPVDiffuseOcclusionIntensity_DEPRECATED:1;
+
+	UPROPERTY()
+	uint8 bOverride_LPVSpecularOcclusionIntensity_DEPRECATED:1;
+
+	UPROPERTY()
+	uint8 bOverride_LPVSize_DEPRECATED:1;
+
+	UPROPERTY()
+	uint8 bOverride_LPVSecondaryOcclusionIntensity_DEPRECATED:1;
+
+	UPROPERTY()
+	uint8 bOverride_LPVSecondaryBounceIntensity_DEPRECATED:1;
+
+	UPROPERTY()
+	uint8 bOverride_LPVGeometryVolumeBias_DEPRECATED:1;
+
+	UPROPERTY()
+	uint8 bOverride_LPVVplInjectionBias_DEPRECATED:1;
+
+	UPROPERTY()
+	uint8 bOverride_LPVEmissiveInjectionIntensity_DEPRECATED:1;
+
+	UPROPERTY()
+	uint8 bOverride_LPVFadeRange_DEPRECATED:1;
+
+	UPROPERTY()
+	uint8 bOverride_LPVDirectionalOcclusionFadeRange_DEPRECATED:1;
+
 	UPROPERTY(EditAnywhere, BlueprintReadWrite, Category=Overrides, meta=(PinHiddenByDefault, InlineEditConditionToggle))
 	uint8 bOverride_IndirectLightingColor:1;
 
@@ -1762,6 +1807,45 @@ struct FPostProcessSettings
 	UPROPERTY(interp, BlueprintReadWrite, Category="Rendering Features|Motion Blur", meta=(ClampMin = "0.0", UIMax = "100.0", editcondition = "bOverride_MotionBlurPerObjectSize", DisplayName = "Per Object Size"))
 	float MotionBlurPerObjectSize;
 
+	UPROPERTY()
+	float LPVIntensity_DEPRECATED;
+
+	UPROPERTY()
+	float LPVVplInjectionBias_DEPRECATED;
+
+	UPROPERTY()
+	float LPVSize_DEPRECATED;
+
+	UPROPERTY()
+	float LPVSecondaryOcclusionIntensity_DEPRECATED;
+
+	UPROPERTY()
+	float LPVSecondaryBounceIntensity_DEPRECATED;
+
+	UPROPERTY()
+	float LPVGeometryVolumeBias_DEPRECATED;
+
+	UPROPERTY()
+	float LPVEmissiveInjectionIntensity_DEPRECATED;
+
+	UPROPERTY()
+	float LPVDirectionalOcclusionIntensity_DEPRECATED;
+
+	UPROPERTY()
+	float LPVDirectionalOcclusionRadius_DEPRECATED;
+
+	UPROPERTY()
+	float LPVDiffuseOcclusionExponent_DEPRECATED;
+
+	UPROPERTY()
+	float LPVSpecularOcclusionExponent_DEPRECATED;
+
+	UPROPERTY()
+	float LPVDiffuseOcclusionIntensity_DEPRECATED;
+
+	UPROPERTY()
+	float LPVSpecularOcclusionIntensity_DEPRECATED;
+
 	/** Sets the reflections type */
 	UPROPERTY(EditAnywhere, BlueprintReadWrite, Category = "Rendering Features|Reflections", meta = (editcondition = "bOverride_ReflectionsType", DisplayName = "Type"))
 	EReflectionsType ReflectionsType;
@@ -1833,6 +1917,12 @@ struct FPostProcessSettings
 	UPROPERTY(EditAnywhere, BlueprintReadWrite, Category = "Rendering Features|PathTracing", meta = (ClampMin = "1", ClampMax = "64", editcondition = "bOverride_PathTracingSamplesPerPixel", DisplayName = "Samples Per Pixel"))
 	int32 PathTracingSamplesPerPixel;
 
+	UPROPERTY()
+	float LPVFadeRange_DEPRECATED;
+
+	UPROPERTY()
+	float LPVDirectionalOcclusionFadeRange_DEPRECATED;
+
 	/**
 	* To render with lower or high resolution than it is presented,
 	* controlled by console variable,
