Expose FVectorCurve::Evaluate and FTransformCurve::Evaluate

diff --git a/Engine/Source/Runtime/Engine/Public/Animation/AnimCurveTypes.h b/Engine/Source/Runtime/Engine/Public/Animation/AnimCurveTypes.h
--- a/Engine/Source/Runtime/Engine/Public/Animation/AnimCurveTypes.h
+++ b/Engine/Source/Runtime/Engine/Public/Animation/AnimCurveTypes.h
@@ -203,7 +203,7 @@ struct FVectorCurve : public FAnimCurveBase
 
 	// we don't want to have = operator. This only copies curves, but leaving naming and everything else intact. 
 	void CopyCurve(FVectorCurve& SourceCurve);
-	FVector Evaluate(float CurrentTime, float BlendWeight) const;
+	ENGINE_API FVector Evaluate(float CurrentTime, float BlendWeight) const;
 	ENGINE_API void UpdateOrAddKey(const FVector& NewKey, float CurrentTime);
 	ENGINE_API void GetKeys(TArray<float>& OutTimes, TArray<FVector>& OutValues);
 	bool DoesContainKey() const { return (FloatCurves[0].GetNumKeys() > 0 || FloatCurves[1].GetNumKeys() > 0 || FloatCurves[2].GetNumKeys() > 0);}
@@ -239,7 +239,7 @@ struct FTransformCurve: public FAnimCurveBase
 
 	// we don't want to have = operator. This only copies curves, but leaving naming and everything else intact. 
 	void CopyCurve(FTransformCurve& SourceCurve);
-	FTransform Evaluate(float CurrentTime, float BlendWeight) const;
+	ENGINE_API FTransform Evaluate(float CurrentTime, float BlendWeight) const;
 	ENGINE_API void UpdateOrAddKey(const FTransform& NewKey, float CurrentTime);
 	ENGINE_API void GetKeys(TArray<float>& OutTimes, TArray<FTransform>& OutValues);
 	void Resize(float NewLength, bool bInsert/* whether insert or remove*/, float OldStartTime, float OldEndTime);
