[NO NON-PII COMMIT MESSAGE]

diff --git a/Engine/Source/Runtime/Engine/Classes/Sound/QuartzQuantizationUtilities.cpp b/Engine/Source/Runtime/Engine/Classes/Sound/QuartzQuantizationUtilities.cpp
--- a/Engine/Source/Runtime/Engine/Classes/Sound/QuartzQuantizationUtilities.cpp
+++ b/Engine/Source/Runtime/Engine/Classes/Sound/QuartzQuantizationUtilities.cpp
@@ -69,21 +69,21 @@ int64 FQuartzClockTickRate::GetFramesPerDuration(EQuartzCommandQuantization InDu
 		return FramesPerTick; // same as 1/32nd note
 
 	case EQuartzCommandQuantization::SixteenthNote:
-		return FramesPerTick << 1;
+		return (int64)FramesPerTick << 1;
 
 	case EQuartzCommandQuantization::EighthNote:
-		return FramesPerTick << 2;
+		return (int64)FramesPerTick << 2;
 
 	case EQuartzCommandQuantization::Beat: // default to quarter note (should be overridden for non-basic meters)
 	case EQuartzCommandQuantization::QuarterNote:
-		return FramesPerTick << 3;
+		return (int64)FramesPerTick << 3;
 
 	case EQuartzCommandQuantization::HalfNote:
-		return FramesPerTick << 4;
+		return (int64)FramesPerTick << 4;
 
 	case EQuartzCommandQuantization::Bar: // default to whole note (should be overridden for non-4/4 meters)
 	case EQuartzCommandQuantization::WholeNote:
-		return FramesPerTick << 5;
+		return (int64)FramesPerTick << 5;
 
 	// DOTTED
 	case EQuartzCommandQuantization::DottedSixteenthNote:
