Fixed non-unity build

diff --git a/Engine/Source/Runtime/Engine/Private/WorldPartition/DataLayer/DataLayerSubsystem.cpp b/Engine/Source/Runtime/Engine/Private/WorldPartition/DataLayer/DataLayerSubsystem.cpp
--- a/Engine/Source/Runtime/Engine/Private/WorldPartition/DataLayer/DataLayerSubsystem.cpp
+++ b/Engine/Source/Runtime/Engine/Private/WorldPartition/DataLayer/DataLayerSubsystem.cpp
@@ -1,6 +1,8 @@
 // Copyright Epic Games, Inc. All Rights Reserved.
 
 #include "WorldPartition/DataLayer/DataLayerSubsystem.h"
+#include "Engine/World.h"
+#include "Engine/Engine.h"
 
 UDataLayerSubsystem::UDataLayerSubsystem()
 {}
diff --git a/Engine/Source/Runtime/Engine/Public/WorldPartition/DataLayer/DataLayerHelper.h b/Engine/Source/Runtime/Engine/Public/WorldPartition/DataLayer/DataLayerHelper.h
--- a/Engine/Source/Runtime/Engine/Public/WorldPartition/DataLayer/DataLayerHelper.h
+++ b/Engine/Source/Runtime/Engine/Public/WorldPartition/DataLayer/DataLayerHelper.h
@@ -2,6 +2,8 @@
 
 #pragma once
 
+#include "CoreMinimal.h"
+
 struct ENGINE_API FDataLayersHelper
 {
 	static const uint32 NoDataLayerID = 0;
