Edigrate static analyisis fix for Quartz
#rnx

diff --git a/Engine/Source/Runtime/Engine/Private/Components/AudioComponent.cpp b/Engine/Source/Runtime/Engine/Private/Components/AudioComponent.cpp
--- a/Engine/Source/Runtime/Engine/Private/Components/AudioComponent.cpp
+++ b/Engine/Source/Runtime/Engine/Private/Components/AudioComponent.cpp
@@ -407,7 +407,11 @@ void UAudioComponent::PlayQuantized(
 	Data.FadeInDuration = InFadeInDuration;
 	Data.FadeVolumeLevel = InFadeVolumeLevel;
 	Data.FadeCurve = InFadeCurve;
-	Data.QuantizedRequestData = InClockHandle->GetQuartzSubsystem()->CreateDataDataForSchedulePlaySound(InClockHandle, InDelegate, InQuantizationBoundary);
+
+	if (InClockHandle != nullptr)
+	{
+		Data.QuantizedRequestData = InClockHandle->GetQuartzSubsystem()->CreateDataDataForSchedulePlaySound(InClockHandle, InDelegate, InQuantizationBoundary);
+	}
 
 	// validate clock existence 
 	if (!InClockHandle)
