Fix static analyzer warning

diff --git a/Engine/Source/Runtime/Engine/Private/Components/SkyLightComponent.cpp b/Engine/Source/Runtime/Engine/Private/Components/SkyLightComponent.cpp
--- a/Engine/Source/Runtime/Engine/Private/Components/SkyLightComponent.cpp
+++ b/Engine/Source/Runtime/Engine/Private/Components/SkyLightComponent.cpp
@@ -713,11 +713,10 @@ void USkyLightComponent::ApplyComponentInstanceData(FPrecomputedSkyLightInstance
 
 void USkyLightComponent::UpdateSkyCaptureContentsArray(UWorld* WorldToUpdate, TArray<USkyLightComponent*>& ComponentArray, bool bOperateOnBlendSource)
 {
-	const bool bIsCompilingShaders = GShaderCompilingManager != NULL && GShaderCompilingManager->IsCompiling();
+	const bool bIsCompilingShaders = GShaderCompilingManager != nullptr && GShaderCompilingManager->IsCompiling();
+	bool bIsCompilingTextures = false;
 #if WITH_EDITOR
-	const bool bIsCompilingTextures = FTextureCompilingManager::Get().GetNumRemainingTextures() > 0;
-#else
-	const bool bIsCompilingTextures = false;
+	bIsCompilingTextures = FTextureCompilingManager::Get().GetNumRemainingTextures() > 0;
 #endif
 
 	// Iterate backwards so we can remove elements without changing the index
@@ -727,18 +726,17 @@ void USkyLightComponent::UpdateSkyCaptureContentsArray(UWorld* WorldToUpdate, TA
 		AActor* Owner = CaptureComponent->GetOwner();
 
 		// For specific cubemaps, we must wait until the texture is compiled before capturing the skylight
-		const bool bIsCubemapCompiling = 
+		bool bIsCubemapCompiling = false;
 #if WITH_EDITOR
+		bIsCubemapCompiling =
 			CaptureComponent->SourceType == SLS_SpecifiedCubemap &&
 			CaptureComponent->Cubemap &&
 			CaptureComponent->Cubemap->IsDefaultTexture();
-#else
-			false;
 #endif
 
 		if (((!Owner || !Owner->GetLevel() || Owner->GetLevel()->bIsVisible) && CaptureComponent->GetWorld() == WorldToUpdate)
 			// Only process sky capture requests once async texture and shader compiling completes, otherwise we will capture the scene with temporary shaders/textures
-			&& ((!bIsCompilingTextures && !bIsCompilingShaders) || ((CaptureComponent->SourceType == SLS_SpecifiedCubemap) && !bIsCubemapCompiling)))
+			&& (((!bIsCompilingTextures) && (!bIsCompilingShaders)) || ((CaptureComponent->SourceType == SLS_SpecifiedCubemap) && (!bIsCubemapCompiling))))
 		{
 			// Only capture valid sky light components
 			if (CaptureComponent->SourceType != SLS_SpecifiedCubemap || CaptureComponent->Cubemap)
