Expose FAnimationUtils::ExtractTransformFromTrack #review-14580180 [REDACTED_HANDLE]

diff --git a/Engine/Source/Runtime/Engine/Public/AnimationUtils.h b/Engine/Source/Runtime/Engine/Public/AnimationUtils.h
--- a/Engine/Source/Runtime/Engine/Public/AnimationUtils.h
+++ b/Engine/Source/Runtime/Engine/Public/AnimationUtils.h
@@ -165,7 +165,7 @@ public:
 	// Extract specific frame from raw track and place in OutAtom
 	ENGINE_API static void ExtractTransformForFrameFromTrack(const FRawAnimSequenceTrack& RawTrack, int32 Frame, FTransform& OutAtom);
 
-	static void ExtractTransformFromTrack(float Time, int32 NumFrames, float SequenceLength, const struct FRawAnimSequenceTrack& RawTrack, EAnimInterpolationType Interpolation, FTransform &OutAtom);
+	ENGINE_API static void ExtractTransformFromTrack(float Time, int32 NumFrames, float SequenceLength, const struct FRawAnimSequenceTrack& RawTrack, EAnimInterpolationType Interpolation, FTransform &OutAtom);
 
 #if WITH_EDITOR
 	static void ExtractTransformFromCompressionData(const FCompressibleAnimData& CompressibleAnimData, FCompressibleAnimDataResult& CompressedAnimData, float Time, int32 TrackIndex, bool bUseRawData, FTransform& OutBoneTransform);
