Fixing dual blending transparancy not working (due to missing MRT1 output).

diff --git a/Engine/Source/Runtime/Engine/Private/ShaderCompiler/ShaderGenerationUtil.cpp b/Engine/Source/Runtime/Engine/Private/ShaderCompiler/ShaderGenerationUtil.cpp
--- a/Engine/Source/Runtime/Engine/Private/ShaderCompiler/ShaderGenerationUtil.cpp
+++ b/Engine/Source/Runtime/Engine/Private/ShaderCompiler/ShaderGenerationUtil.cpp
@@ -58,7 +58,7 @@ static int FetchCompileInt(FShaderCompilerEnvironment& OutEnvironment, const cha
 // these definitions can't be merged, so one of them should be cleared at zero.
 #define MERGE_COMPILE_INT(X) { check(Lhs.X == 0 || Rhs.X == 0); Lhs.X = Lhs.X > Rhs.X ? Lhs.X : Rhs.X; }
 
-void FShaderCompileUtilities::ApplyFetchEnvironment(FShaderGlobalDefines& SrcDefines, FShaderCompilerEnvironment& OutEnvironment)
+void FShaderCompileUtilities::ApplyFetchEnvironment(FShaderGlobalDefines& SrcDefines, FShaderCompilerEnvironment& OutEnvironment, const EShaderPlatform Platform)
 {
 	FETCH_COMPILE_BOOL(GBUFFER_HAS_VELOCITY);
 	FETCH_COMPILE_BOOL(GBUFFER_HAS_TANGENT);
@@ -86,8 +86,11 @@ void FShaderCompileUtilities::ApplyFetchEnvironment(FShaderGlobalDefines& SrcDef
 	FETCH_COMPILE_BOOL(PROJECT_MOBILE_ENABLE_MOVABLE_SPOTLIGHTS);
 	FETCH_COMPILE_BOOL(COMPILE_SHADERS_FOR_DEVELOPMENT_ALLOWED);
 
-	// This one is a helper for other macros
-	//bSupportsDualBlending;
+	// note that we are doing an if so that if we call ApplyFetchEnvironment() twice, we get the logical OR of bSupportsDualBlending support
+	if (RHISupportsDualSourceBlending(Platform))
+	{
+		SrcDefines.bSupportsDualBlending = true;
+	}
 }
 
 void FShaderCompileUtilities::ApplyFetchEnvironment(FShaderLightmapPropertyDefines& SrcDefines, FShaderCompilerEnvironment& OutEnvironment)
@@ -224,8 +227,6 @@ void FShaderCompileUtilities::ApplyFetchEnvironment(FShaderMaterialPropertyDefin
 	FETCH_COMPILE_BOOL(DUAL_SOURCE_COLOR_BLENDING_ENABLED);
 
 	FETCH_COMPILE_INT(DECAL_RENDERTARGET_COUNT);
-
-	FETCH_COMPILE_BOOL(DUAL_SOURCE_COLOR_BLENDING_ENABLED);
 }
 
 void FShaderCompileUtilities::ApplyFetchEnvironment(FShaderCompilerDefines& SrcDefines, FShaderCompilerEnvironment& OutEnvironment)
@@ -1760,16 +1761,16 @@ void FShaderCompileUtilities::ApplyDerivedDefines(FShaderCompilerEnvironment& Ou
 	FShaderCompilerDefines CompilerDefines = {};
 	FShaderMaterialDerivedDefines DerivedDefines = {};
 
-	ApplyFetchEnvironment(MaterialDefines,OutEnvironment);
-	ApplyFetchEnvironment(LightmapDefines,OutEnvironment);
-	ApplyFetchEnvironment(GlobalDefines,OutEnvironment);
-	ApplyFetchEnvironment(CompilerDefines,OutEnvironment);
+	ApplyFetchEnvironment(GlobalDefines, OutEnvironment, Platform);
+	ApplyFetchEnvironment(MaterialDefines, OutEnvironment);
+	ApplyFetchEnvironment(LightmapDefines, OutEnvironment);
+	ApplyFetchEnvironment(CompilerDefines, OutEnvironment);
 
 	if (SharedEnvironment != nullptr)
 	{
+		ApplyFetchEnvironment(GlobalDefines, *SharedEnvironment, Platform);
 		ApplyFetchEnvironment(MaterialDefines, *SharedEnvironment);
 		ApplyFetchEnvironment(LightmapDefines, *SharedEnvironment);
-		ApplyFetchEnvironment(GlobalDefines, *SharedEnvironment);
 		ApplyFetchEnvironment(CompilerDefines, *SharedEnvironment);
 	}
 
@@ -1789,7 +1790,6 @@ void FShaderCompileUtilities::ApplyDerivedDefines(FShaderCompilerEnvironment& Ou
 	FGBufferInfo BufferInfo = FetchFullGBufferInfo(Params);
 
 	bool bTargetUsage[FGBufferInfo::MaxTargets] = {};
-
 	if (MaterialDefines.IS_BASE_PASS)
 	{
 		// if we are using a gbuffer, and this is the base pass that writes a gbuffer, search the gbuffer for each slot
diff --git a/Engine/Source/Runtime/Engine/Public/ShaderCompiler.h b/Engine/Source/Runtime/Engine/Public/ShaderCompiler.h
--- a/Engine/Source/Runtime/Engine/Public/ShaderCompiler.h
+++ b/Engine/Source/Runtime/Engine/Public/ShaderCompiler.h
@@ -524,7 +524,7 @@ namespace FShaderCompileUtilities
 	FGBufferInfo FetchGBufferInfoAndWriteAutogen(EShaderPlatform TargetPlatform, ERHIFeatureLevel::Type FeatureLevel);
 
 	void ApplyFetchEnvironment(FShaderMaterialPropertyDefines& DefineData, FShaderCompilerEnvironment& OutEnvironment);
-	void ApplyFetchEnvironment(FShaderGlobalDefines& DefineData, FShaderCompilerEnvironment& OutEnvironment);
+	void ApplyFetchEnvironment(FShaderGlobalDefines& DefineData, FShaderCompilerEnvironment& OutEnvironment, const EShaderPlatform Platform);
 	void ApplyFetchEnvironment(FShaderLightmapPropertyDefines& DefineData, FShaderCompilerEnvironment& OutEnvironment);
 	void ApplyFetchEnvironment(FShaderCompilerDefines& DefineData, FShaderCompilerEnvironment& OutEnvironment);
 
