[Navigation WP] Editor improvements for navmesh in WP: adding tiles in editor when cell are loaded/unloaded, using bound placement and now also works in PIE

diff --git a/Engine/Source/Runtime/Engine/Private/WorldPartition/NavigationData/NavigationDataChunkActor.cpp b/Engine/Source/Runtime/Engine/Private/WorldPartition/NavigationData/NavigationDataChunkActor.cpp
--- a/Engine/Source/Runtime/Engine/Private/WorldPartition/NavigationData/NavigationDataChunkActor.cpp
+++ b/Engine/Source/Runtime/Engine/Private/WorldPartition/NavigationData/NavigationDataChunkActor.cpp
@@ -3,6 +3,10 @@
 #include "WorldPartition/NavigationData/NavigationDataChunkActor.h"
 #include "AI/NavigationSystemBase.h"
 
+#if WITH_EDITOR
+#include "Editor.h"
+#endif
+
 ANavigationDataChunkActor::ANavigationDataChunkActor(const FObjectInitializer& ObjectInitializer)
 	: Super(ObjectInitializer)
 {
@@ -17,23 +21,36 @@ ANavigationDataChunkActor::ANavigationDataChunkActor(const FObjectInitializer& O
 	SetActorEnableCollision(false);
 }
 
-void ANavigationDataChunkActor::Serialize(FArchive& Ar)
+#if WITH_EDITOR
+void ANavigationDataChunkActor::PostLoad()
 {
-	UE_LOG(LogNavigation, VeryVerbose, TEXT("%s"), ANSI_TO_TCHAR(__FUNCTION__));
+	Super::PostLoad();
 
-	Super::Serialize(Ar);
+	const bool bIsInPIE = (GEditor->PlayWorld != NULL) && (!GEditor->bIsSimulatingInEditor);
+	if (!bIsInPIE)
+	{
+		Log(ANSI_TO_TCHAR(__FUNCTION__));
+		UE_LOG(LogNavigation, Verbose, TEXT("   pos: %s ext: %s"), *DataChunkActorBounds.GetCenter().ToCompactString(), *DataChunkActorBounds.GetExtent().ToCompactString());
+		AddNavigationDataChunkToWorld();
+	}
 }
 
-void ANavigationDataChunkActor::PostLoad()
+void ANavigationDataChunkActor::BeginDestroy()
 {
-	UE_LOG(LogNavigation, VeryVerbose, TEXT("%s"), ANSI_TO_TCHAR(__FUNCTION__));
+	const bool bIsInPIE = (GEditor->PlayWorld != NULL) && (!GEditor->bIsSimulatingInEditor);
+	if (!bIsInPIE)
+	{
+		Log(ANSI_TO_TCHAR(__FUNCTION__));
+		RemoveNavigationDataChunkFromWorld();
+	}
 
-	Super::PostLoad();
+	Super::BeginDestroy();
 }
+#endif // WITH_EDITOR
 
 void ANavigationDataChunkActor::CollectNavData(const FBox& Bounds)
 {
-	UE_LOG(LogNavigation, VeryVerbose, TEXT("%s"), ANSI_TO_TCHAR(__FUNCTION__));
+	Log(ANSI_TO_TCHAR(__FUNCTION__));
 
 	UWorld* World = GetWorld();
 	if (World)
@@ -45,12 +62,29 @@ void ANavigationDataChunkActor::CollectNavData(const FBox& Bounds)
 	}
 }
 
-void ANavigationDataChunkActor::BeginPlay()
+#if WITH_EDITOR
+void ANavigationDataChunkActor::SetDataChunkActorBounds(const FBox& InBounds)
 {
-	UE_LOG(LogNavigation, VeryVerbose, TEXT("%s"), ANSI_TO_TCHAR(__FUNCTION__));
+	DataChunkActorBounds = InBounds;
+}
+#endif //WITH_EDITOR
 
+void ANavigationDataChunkActor::BeginPlay()
+{
+	Log(ANSI_TO_TCHAR(__FUNCTION__));
 	Super::BeginPlay();
+	AddNavigationDataChunkToWorld();
+}
 
+void ANavigationDataChunkActor::EndPlay(const EEndPlayReason::Type EndPlayReason)
+{
+	Log(ANSI_TO_TCHAR(__FUNCTION__));
+	RemoveNavigationDataChunkFromWorld();
+	Super::EndPlay(EndPlayReason);
+}
+
+void ANavigationDataChunkActor::AddNavigationDataChunkToWorld()
+{
 	UWorld* World = GetWorld();
 	if (World)
 	{
@@ -61,10 +95,8 @@ void ANavigationDataChunkActor::BeginPlay()
 	}
 }
 
-void ANavigationDataChunkActor::EndPlay(const EEndPlayReason::Type EndPlayReason)
+void ANavigationDataChunkActor::RemoveNavigationDataChunkFromWorld()
 {
-	UE_LOG(LogNavigation, VeryVerbose, TEXT("%s"), ANSI_TO_TCHAR(__FUNCTION__));
-
 	UWorld* World = GetWorld();
 	if (World)
 	{
@@ -73,14 +105,26 @@ void ANavigationDataChunkActor::EndPlay(const EEndPlayReason::Type EndPlayReason
 			NavSys->RemoveNavigationDataChunk(*this);
 		}
 	}
+}
 
-	Super::EndPlay(EndPlayReason);
+void ANavigationDataChunkActor::Log(const TCHAR* FunctionName) const
+{
+	UE_LOG(LogNavigation, Verbose, TEXT("[%s] %s"), *GetName(), FunctionName);
 }
 
+void ANavigationDataChunkActor::GetActorBounds(bool bOnlyCollidingComponents, FVector& OutOrigin, FVector& OutBoxExtent, bool bIncludeFromChildActors) const
+{
+	DataChunkActorBounds.GetCenterAndExtents(OutOrigin, OutBoxExtent);
+}
 
 #if WITH_EDITOR
 EActorGridPlacement ANavigationDataChunkActor::GetDefaultGridPlacement() const
 {
-	return EActorGridPlacement::Location;
+	return EActorGridPlacement::Bounds;
+}
+
+void ANavigationDataChunkActor::GetActorLocationBounds(bool bOnlyCollidingComponents, FVector& OutOrigin, FVector& OutBoxExtent, bool bIncludeFromChildActors) const
+{
+	GetActorBounds(bOnlyCollidingComponents, OutOrigin, OutBoxExtent, bIncludeFromChildActors);
 }
 #endif // WITH_EDITOR
diff --git a/Engine/Source/Runtime/Engine/Private/WorldPartition/WorldPartitionRuntimeSpatialHash.cpp b/Engine/Source/Runtime/Engine/Private/WorldPartition/WorldPartitionRuntimeSpatialHash.cpp
--- a/Engine/Source/Runtime/Engine/Private/WorldPartition/WorldPartitionRuntimeSpatialHash.cpp
+++ b/Engine/Source/Runtime/Engine/Private/WorldPartition/WorldPartitionRuntimeSpatialHash.cpp
@@ -1563,10 +1563,10 @@ bool UWorldPartitionRuntimeSpatialHash::GenerateNavigationData()
 
 		FBox2D CellBounds;
 		GridHelper.GetCellBounds(CellCoord, CellBounds);
-		const float HalfHeight = 100000.f; //Todo AT: try use WORLD_MAX instead
-		const FBox Bounds(FVector(CellBounds.Min.X, CellBounds.Min.Y, -HalfHeight), FVector(CellBounds.Max.X, CellBounds.Max.Y, HalfHeight));
+		const float HalfHeight = 100000.f; //Todo AT: check to get the real z bounds from the collected tiles
+		const FBox CollectingBounds(FVector(CellBounds.Min.X, CellBounds.Min.Y, -HalfHeight), FVector(CellBounds.Max.X, CellBounds.Max.Y, HalfHeight));
 
-		if (NavSystem->ContainsNavData(Bounds) == false)
+		if (NavSystem->ContainsNavData(CollectingBounds) == false)
 		{
 			// Skip if there is no navdata for this cell
 			return;
@@ -1580,21 +1580,22 @@ bool UWorldPartitionRuntimeSpatialHash::GenerateNavigationData()
 		ANavigationDataChunkActor* DataChunkActor = World->SpawnActor<ANavigationDataChunkActor>(SpawnParams);
 		ActorCount++;
 
-		const FName CellName = GetCellName(RuntimeGrid.GridName, CellCoord.Z, CellCoord.X, CellCoord.Y, FDataLayersHelper::NoDataLayerID);
-		DataChunkActor->SetActorLabel(FString::Printf(TEXT("NavDataChunkActor_%s_%s"), *GetName(), *CellName.ToString()));
-			
 		const FVector2D CellCenter = CellBounds.GetCenter();
 		DataChunkActor->SetActorLocation(FVector(CellCenter.X, CellCenter.Y, 0.f));
+		DataChunkActor->SetDataChunkActorBounds(CollectingBounds.ExpandBy(-1.f)); //reduce by 1cm to avoid precision issues
+
+		const FName CellName = GetCellName(RuntimeGrid.GridName, CellCoord.Z, CellCoord.X, CellCoord.Y, FDataLayersHelper::NoDataLayerID);
+		DataChunkActor->SetActorLabel(FString::Printf(TEXT("NavDataChunkActor_%s_%s"), *GetName(), *CellName.ToString()));
 		
 		//@todo_ow: Properly handle data layers
-		DataChunkActor->CollectNavData(Bounds);
+		DataChunkActor->CollectNavData(CollectingBounds);
 
 		// Set target grid
 		DataChunkActor->RuntimeGrid = RuntimeGrid.GridName;
 		ValidNavigationDataChunkActors.Add(DataChunkActor);
-	});
 
-	UE_LOG(LogWorldPartitionRuntimeSpatialHash, Verbose, TEXT("   %i ANavigationDataChunkActor actors added."), ActorCount);
+		UE_LOG(LogWorldPartitionRuntimeSpatialHash, Verbose, TEXT("%i) %s added."), ActorCount, *DataChunkActor->GetName());
+	});
 
 	// Destroy all invalid navigation data chunk actors
 	for (TActorIterator<ANavigationDataChunkActor> It(GetWorld()); It; ++It)
diff --git a/Engine/Source/Runtime/Engine/Public/WorldPartition/NavigationData/NavigationDataChunkActor.h b/Engine/Source/Runtime/Engine/Public/WorldPartition/NavigationData/NavigationDataChunkActor.h
--- a/Engine/Source/Runtime/Engine/Public/WorldPartition/NavigationData/NavigationDataChunkActor.h
+++ b/Engine/Source/Runtime/Engine/Public/WorldPartition/NavigationData/NavigationDataChunkActor.h
@@ -15,8 +15,10 @@ class ENGINE_API ANavigationDataChunkActor : public AActor
 
 public:
 	//~ Begin UObject Interface
-	virtual void Serialize(FArchive& Ar) override;
+#if WITH_EDITOR
 	virtual void PostLoad() override;
+	virtual void BeginDestroy() override;
+#endif //WITH_EDITOR
 	//~ End UObject Interface
 
 	const TArray<UNavigationDataChunk*>& GetNavDataChunk() const { return NavDataChunks; }
@@ -24,18 +26,29 @@ public:
 
 	void CollectNavData(const FBox& Bounds);
 
+#if WITH_EDITOR
+	void SetDataChunkActorBounds(const FBox& InBounds);
+
+	//~ Begin AActor Interface.
+	virtual EActorGridPlacement GetDefaultGridPlacement() const override;
+	virtual void GetActorLocationBounds(bool bOnlyCollidingComponents, FVector& OutOrigin, FVector& OutBoxExtent, bool bIncludeFromChildActors) const override;
+	//~ End AActor Interface.
+#endif // WITH_EDITOR
+	
 protected:
 	//~ Begin AActor Interface.
 	virtual void BeginPlay() override;
 	virtual void EndPlay(const EEndPlayReason::Type EndPlayReason) override;
+	virtual void GetActorBounds(bool bOnlyCollidingComponents, FVector& OutOrigin, FVector& OutBoxExtent, bool bIncludeFromChildActors) const override;
 	//~ End AActor Interface.
 
-#if WITH_EDITOR
-	//~ Begin AActor Interface.
-	virtual EActorGridPlacement GetDefaultGridPlacement() const override;
-	//~ End AActor Interface.
-#endif
+	void AddNavigationDataChunkToWorld();
+	void RemoveNavigationDataChunkFromWorld();
+	void Log(const TCHAR* FunctionName) const;
 
 	UPROPERTY()
 	TArray<UNavigationDataChunk*> NavDataChunks;
+
+	UPROPERTY()
+	FBox DataChunkActorBounds;
 };
