Crash fix in ANavigationDataChunkActor::BeginDestroy()

diff --git a/Engine/Source/Runtime/Engine/Private/WorldPartition/NavigationData/NavigationDataChunkActor.cpp b/Engine/Source/Runtime/Engine/Private/WorldPartition/NavigationData/NavigationDataChunkActor.cpp
--- a/Engine/Source/Runtime/Engine/Private/WorldPartition/NavigationData/NavigationDataChunkActor.cpp
+++ b/Engine/Source/Runtime/Engine/Private/WorldPartition/NavigationData/NavigationDataChunkActor.cpp
@@ -26,22 +26,28 @@ void ANavigationDataChunkActor::PostLoad()
 {
 	Super::PostLoad();
 
-	const bool bIsInPIE = (GEditor->PlayWorld != NULL) && (!GEditor->bIsSimulatingInEditor);
-	if (!bIsInPIE)
+	if (GEditor)
 	{
-		Log(ANSI_TO_TCHAR(__FUNCTION__));
-		UE_LOG(LogNavigation, Verbose, TEXT("   pos: %s ext: %s"), *DataChunkActorBounds.GetCenter().ToCompactString(), *DataChunkActorBounds.GetExtent().ToCompactString());
-		AddNavigationDataChunkToWorld();
+		const bool bIsInPIE = (GEditor->PlayWorld != NULL) && (!GEditor->bIsSimulatingInEditor);
+		if (!bIsInPIE)
+		{
+			Log(ANSI_TO_TCHAR(__FUNCTION__));
+			UE_LOG(LogNavigation, Verbose, TEXT("   pos: %s ext: %s"), *DataChunkActorBounds.GetCenter().ToCompactString(), *DataChunkActorBounds.GetExtent().ToCompactString());
+			AddNavigationDataChunkToWorld();
+		}
 	}
 }
 
 void ANavigationDataChunkActor::BeginDestroy()
 {
-	const bool bIsInPIE = (GEditor->PlayWorld != NULL) && (!GEditor->bIsSimulatingInEditor);
-	if (!bIsInPIE)
+	if (GEditor)
 	{
-		Log(ANSI_TO_TCHAR(__FUNCTION__));
-		RemoveNavigationDataChunkFromWorld();
+		const bool bIsInPIE = (GEditor->PlayWorld != NULL) && (!GEditor->bIsSimulatingInEditor);
+		if (!bIsInPIE)
+		{
+			Log(ANSI_TO_TCHAR(__FUNCTION__));
+			RemoveNavigationDataChunkFromWorld();
+		}
 	}
 
 	Super::BeginDestroy();
