Made FActorSpawnParameters::OverrideParentComponent available in runtime builds.

diff --git a/Engine/Source/Runtime/Engine/Private/LevelActor.cpp b/Engine/Source/Runtime/Engine/Private/LevelActor.cpp
--- a/Engine/Source/Runtime/Engine/Private/LevelActor.cpp
+++ b/Engine/Source/Runtime/Engine/Private/LevelActor.cpp
@@ -663,12 +663,12 @@ AActor* UWorld::SpawnActor( UClass* Class, FTransform const* UserTransformPtr, c
 
 	// Set the actor's guid
 	FSetActorGuid SetActorGuid(Actor, ActorGuid);
+#endif
 
 	if (SpawnParameters.OverrideParentComponent)
 	{
 		FActorParentComponentSetter::Set(Actor, SpawnParameters.OverrideParentComponent);
 	}
-#endif // WITH_EDITOR
 
 	if ( GUndo )
 	{
