Deprecate ULevel::GetLevelBlueprints.

diff --git a/Engine/Source/Runtime/Engine/Classes/Engine/Level.h b/Engine/Source/Runtime/Engine/Classes/Engine/Level.h
--- a/Engine/Source/Runtime/Engine/Classes/Engine/Level.h
+++ b/Engine/Source/Runtime/Engine/Classes/Engine/Level.h
@@ -1041,6 +1041,7 @@ public:
 	/**
 	 *  Returns a list of all blueprints contained within the level
 	 */
+	UE_DEPRECATED(5.0, "This function is deprecated, we only support having a single level script blueprint. Use GetLevelScriptBlueprint instead.")
 	ENGINE_API TArray<class UBlueprint*> GetLevelBlueprints() const;
 
 	/**
diff --git a/Engine/Source/Runtime/Engine/Private/Level.cpp b/Engine/Source/Runtime/Engine/Private/Level.cpp
--- a/Engine/Source/Runtime/Engine/Private/Level.cpp
+++ b/Engine/Source/Runtime/Engine/Private/Level.cpp
@@ -2541,8 +2541,8 @@ void ULevel::BeginCacheForCookedPlatformData(const ITargetPlatform *TargetPlatfo
 {
 	Super::BeginCacheForCookedPlatformData(TargetPlatform);
 
-	// Cook all level blueprints.
-	for (auto LevelBlueprint : GetLevelBlueprints())
+	// Cook the level blueprint.
+	if (ULevelScriptBlueprint* LevelBlueprint = GetLevelScriptBlueprint(true))
 	{
 		LevelBlueprint->BeginCacheForCookedPlatformData(TargetPlatform);
 	}
diff --git a/Engine/Source/Runtime/Engine/Private/World.cpp b/Engine/Source/Runtime/Engine/Private/World.cpp
--- a/Engine/Source/Runtime/Engine/Private/World.cpp
+++ b/Engine/Source/Runtime/Engine/Private/World.cpp
@@ -7800,8 +7800,7 @@ void UWorld::GetAssetRegistryTags(TArray<FAssetRegistryTag>& OutTags) const
 
 	if(PersistentLevel && PersistentLevel->OwningWorld)
 	{
-		TArray<UBlueprint*> LevelBlueprints = PersistentLevel->GetLevelBlueprints();
-		for (UBlueprint* Blueprint : LevelBlueprints)
+		if (ULevelScriptBlueprint* Blueprint = PersistentLevel->GetLevelScriptBlueprint(true))
 		{
 			Blueprint->GetAssetRegistryTags(OutTags);
 		}
