Removed remaining traces of Actor Layers loading in World Partition.

diff --git a/Engine/Source/Runtime/Engine/Private/WorldPartition/WorldPartition.cpp b/Engine/Source/Runtime/Engine/Private/WorldPartition/WorldPartition.cpp
--- a/Engine/Source/Runtime/Engine/Private/WorldPartition/WorldPartition.cpp
+++ b/Engine/Source/Runtime/Engine/Private/WorldPartition/WorldPartition.cpp
@@ -525,21 +525,6 @@ FWorldPartitionActorDesc* UWorldPartition::GetActorDesc(const FGuid& Guid)
 	return ActorDescObj != nullptr ? ActorDescObj->Get() : nullptr;
 }
 
-void UWorldPartition::RefreshLoadedCells()
-{
-	FWorldPartionCellUpdateContext CellUpdateContext(this);
-
-	UpdateLoadingEditorCell(EditorHash->GetAlwaysLoadedCell(), true);
-
-	EditorHash->ForEachCell([&](UWorldPartitionEditorCell* Cell)
-	{
-		if (Cell->bLoaded)
-		{
-			UpdateLoadingEditorCell(Cell, true);
-		}
-	});
-}
-
 bool UWorldPartition::IsSimulating() const
 {
 	return GEditor->bIsSimulatingInEditor || !!GEditor->PlayWorld;
diff --git a/Engine/Source/Runtime/Engine/Public/WorldPartition/WorldPartition.h b/Engine/Source/Runtime/Engine/Public/WorldPartition/WorldPartition.h
--- a/Engine/Source/Runtime/Engine/Public/WorldPartition/WorldPartition.h
+++ b/Engine/Source/Runtime/Engine/Public/WorldPartition/WorldPartition.h
@@ -207,8 +207,6 @@ private:
 	void UpdateLoadingEditorCell(UWorldPartitionEditorCell* Cell, bool bShouldBeLoaded);
 	void CreateLayers(const TSet<FName>& LayerNames);
 
-	void RefreshLoadedCells();
-
 	void HashActorDesc(FWorldPartitionActorDesc* ActorDesc);
 	void UnhashActorDesc(FWorldPartitionActorDesc* ActorDesc);
 #endif
