Fixed runtime crash on cooked platforms when a platform is first cooked with/without hardware LZ support, and then another platform is cooked with differing support (followed by then running on this platform). The bulk data would be incorrectly compressed based on the first platform since the DDC key didn't contain the hardware LZ capability in it.

diff --git a/Engine/Source/Runtime/Engine/Private/StaticMesh.cpp b/Engine/Source/Runtime/Engine/Private/StaticMesh.cpp
--- a/Engine/Source/Runtime/Engine/Private/StaticMesh.cpp
+++ b/Engine/Source/Runtime/Engine/Private/StaticMesh.cpp
@@ -2443,6 +2443,15 @@ static FString BuildStaticMeshDerivedDataKeySuffix(const ITargetPlatform* Target
 
 	KeySuffix.AppendChar(Mesh->bSupportUniformlyDistributedSampling ? TEXT('1') : TEXT('0'));
 
+	if (TargetPlatform->SupportsFeature(ETargetPlatformFeatures::HardwareLZDecompression))
+	{
+		KeySuffix += TEXT("_HWLZ");
+	}
+	else
+	{
+		KeySuffix += TEXT("_SWLZ");
+	}
+
 	// Value of this CVar affects index buffer <-> painted vertex color correspondence (see UE-51421).
 	static const TConsoleVariableData<int32>* CVar = IConsoleManager::Get().FindTConsoleVariableDataInt(TEXT("r.TriangleOrderOptimization"));
 
