Reduce verbosity of texture compiler logs

diff --git a/Engine/Source/Runtime/Engine/Private/TextureCompiler.cpp b/Engine/Source/Runtime/Engine/Private/TextureCompiler.cpp
--- a/Engine/Source/Runtime/Engine/Private/TextureCompiler.cpp
+++ b/Engine/Source/Runtime/Engine/Private/TextureCompiler.cpp
@@ -336,7 +336,7 @@ void FTextureCompilingManager::FinishTextureCompilation(UTexture* Texture)
 	check(IsInGameThread());
 	TRACE_CPUPROFILER_EVENT_SCOPE(FinishTextureCompilation);
 
-	UE_LOG(LogTexture, Display, TEXT("UpdateResource for %s (%s) due to async texture compilation"), *Texture->GetName(), *GetLODGroupName(Texture));
+	UE_LOG(LogTexture, Verbose, TEXT("UpdateResource for %s (%s) due to async texture compilation"), *Texture->GetName(), *GetLODGroupName(Texture));
 
 	Texture->FinishCachePlatformData();
 	Texture->UpdateResource();
@@ -480,7 +480,7 @@ void FTextureCompilingManager::FinishCompilation(const TArray<UTexture*>& InText
 				UpdateProgress(0.0f, TextureIndex, InTextures.Num(), TextureName);
 			}
 			UpdateProgress(1.f, TextureIndex++, InTextures.Num(), TextureName);
-			UE_LOG(LogTexture, Display, TEXT("FinishCompilation requested for %s (%s)"), *TextureName, *GetLODGroupName(Texture));
+			UE_LOG(LogTexture, Verbose, TEXT("FinishCompilation requested for %s (%s)"), *TextureName, *GetLODGroupName(Texture));
 			FinishTextureCompilation(Texture);
 
 			for (TSet<TWeakObjectPtr<UTexture>>& Bucket : RegisteredTextureBuckets)
@@ -617,7 +617,7 @@ void FTextureCompilingManager::ProcessTextures(int32 MaximumPriority)
 									{
 										UE_LOG(
 											LogTexture, 
-											Display, 
+											Verbose, 
 											TEXT("Boosting priority of %s (%s) from %s to %s because of it's last render time"), 
 											*Texture->GetName(), 
 											*GetLODGroupName(Texture), 
