Fixed crash re-opening a partitioned world. Actor package metadata was preventing it from being GC'ed.

diff --git a/Engine/Source/Runtime/Engine/Private/WorldPartition/WorldPartition.cpp b/Engine/Source/Runtime/Engine/Private/WorldPartition/WorldPartition.cpp
--- a/Engine/Source/Runtime/Engine/Private/WorldPartition/WorldPartition.cpp
+++ b/Engine/Source/Runtime/Engine/Private/WorldPartition/WorldPartition.cpp
@@ -361,6 +361,12 @@ void UWorldPartition::Uninitialize()
 		}
 		else if (IsMainWorldPartition())
 		{
+			// Unload all Editor cells
+			// @todo_ow: Once Metadata is removed from external actor's package, this won't be necessary anymore.
+			EditorHash->ForEachCell([this](UWorldPartitionEditorCell* Cell)
+			{
+				UpdateLoadingEditorCell(Cell, /*bShouldBeLoaded*/false);
+			});
 			UnregisterDelegates();
 			ActorDescFactories.Empty();
 		}
