Fixed starta topology info compiler side when using material function or layer blending.

diff --git a/Engine/Source/Runtime/Engine/Private/Materials/MaterialExpressions.cpp b/Engine/Source/Runtime/Engine/Private/Materials/MaterialExpressions.cpp
--- a/Engine/Source/Runtime/Engine/Private/Materials/MaterialExpressions.cpp
+++ b/Engine/Source/Runtime/Engine/Private/Materials/MaterialExpressions.cpp
@@ -604,7 +604,17 @@ int32 CompileStrataBlendFunction(FMaterialCompiler* Compiler, const int32 A, con
 	{
 		return INDEX_NONE;
 	}
-	return Compiler->StrataHorizontalMixing(B, A, Alpha); // B is foreground (fully visible when Alpha=1), and A is background (fully visible when Alpha=0)
+	int32 OutputCodeChunk = Compiler->StrataHorizontalMixing(B, A, Alpha); // B is foreground (fully visible when Alpha=1), and A is background (fully visible when Alpha=0)
+
+	// Also register the horizontal mixing operation with the compiler.
+	if (!Compiler->ContainsStrataCodeChunk(A) || !Compiler->ContainsStrataCodeChunk(B))
+	{
+		return Compiler->Errorf(TEXT("Could not find A or B code chunk in CompileStrataBlendFunction"));
+	}
+	FStrataMaterialCompilationInfo StrataInfo = StrataHorizontalMixing(Compiler, Compiler->GetStrataCompilationInfo(B), Compiler->GetStrataCompilationInfo(A));
+	Compiler->AddStrataCodeChunk(OutputCodeChunk, StrataInfo);
+
+	return OutputCodeChunk;
 }
 
 UMaterialExpression::UMaterialExpression(const FObjectInitializer& ObjectInitializer)
