Build fix

diff --git a/Engine/Source/Runtime/Engine/Private/Materials/MaterialExpressions.cpp b/Engine/Source/Runtime/Engine/Private/Materials/MaterialExpressions.cpp
--- a/Engine/Source/Runtime/Engine/Private/Materials/MaterialExpressions.cpp
+++ b/Engine/Source/Runtime/Engine/Private/Materials/MaterialExpressions.cpp
@@ -606,6 +606,7 @@ int32 CompileStrataBlendFunction(FMaterialCompiler* Compiler, const int32 A, con
 	}
 	int32 OutputCodeChunk = Compiler->StrataHorizontalMixing(B, A, Alpha); // B is foreground (fully visible when Alpha=1), and A is background (fully visible when Alpha=0)
 
+#if WITH_EDITOR
 	// Also register the horizontal mixing operation with the compiler.
 	if (!Compiler->ContainsStrataCodeChunk(A) || !Compiler->ContainsStrataCodeChunk(B))
 	{
@@ -613,6 +614,7 @@ int32 CompileStrataBlendFunction(FMaterialCompiler* Compiler, const int32 A, con
 	}
 	FStrataMaterialCompilationInfo StrataInfo = StrataHorizontalMixing(Compiler, Compiler->GetStrataCompilationInfo(B), Compiler->GetStrataCompilationInfo(A));
 	Compiler->AddStrataCodeChunk(OutputCodeChunk, StrataInfo);
+#endif
 
 	return OutputCodeChunk;
 }
