Restore BlueprintCallable to decal spawning functions.

diff --git a/Engine/Source/Runtime/Engine/Classes/Kismet/GameplayStatics.h b/Engine/Source/Runtime/Engine/Classes/Kismet/GameplayStatics.h
--- a/Engine/Source/Runtime/Engine/Classes/Kismet/GameplayStatics.h
+++ b/Engine/Source/Runtime/Engine/Classes/Kismet/GameplayStatics.h
@@ -833,7 +833,7 @@ public:
 	 * @param Rotation - rotation to place the decal in world space	
 	 * @param LifeSpan - destroy decal component after time runs out (0 = infinite)
 	 */
-	UFUNCTION(BlueprintCosmetic, Category="Rendering|Components|Decal", meta=(WorldContext="WorldContextObject", UnsafeDuringActorConstruction = "true"))
+	UFUNCTION(BlueprintCallable, BlueprintCosmetic, Category="Rendering|Components|Decal", meta=(WorldContext="WorldContextObject", UnsafeDuringActorConstruction = "true"))
 	static UDecalComponent* SpawnDecalAtLocation(const UObject* WorldContextObject, class UMaterialInterface* DecalMaterial, FVector DecalSize, FVector Location, FRotator Rotation = FRotator(-90, 0, 0), float LifeSpan = 0);
 
 	/** Spawns a decal attached to and following the specified component. Does not replicate.
@@ -846,7 +846,7 @@ public:
 	 * @param LocationType - Specifies whether Location is a relative offset or an absolute world position
 	 * @param LifeSpan - destroy decal component after time runs out (0 = infinite)
 	 */
-	UFUNCTION(BlueprintCosmetic, Category="Rendering|Components|Decal", meta=(UnsafeDuringActorConstruction = "true"))
+	UFUNCTION(BlueprintCallable, BlueprintCosmetic, Category="Rendering|Components|Decal", meta=(UnsafeDuringActorConstruction = "true"))
 	static UDecalComponent* SpawnDecalAttached(class UMaterialInterface* DecalMaterial, FVector DecalSize, class USceneComponent* AttachToComponent, FName AttachPointName = NAME_None, FVector Location = FVector(ForceInit), FRotator Rotation = FRotator::ZeroRotator, EAttachLocation::Type LocationType = EAttachLocation::KeepRelativeOffset, float LifeSpan = 0);
 
 	/** Extracts data from a HitResult. 
