Update storage of parameters in MaterialCachedData
- Don't search by sorted hash, as we no longer use hashed names throughout most of material code
- Instead use standard hash table of regular FNames

diff --git a/Engine/Source/Runtime/Engine/Private/Materials/MaterialCachedData.cpp b/Engine/Source/Runtime/Engine/Private/Materials/MaterialCachedData.cpp
--- a/Engine/Source/Runtime/Engine/Private/Materials/MaterialCachedData.cpp
+++ b/Engine/Source/Runtime/Engine/Private/Materials/MaterialCachedData.cpp
@@ -53,35 +53,6 @@ void FMaterialCachedExpressionData::Reset()
 	bHasSceneColor = false;
 }
 
-static int32 FindParameterLowerBoundIndex(const FMaterialCachedParameterEntry& Entry, const FHashedName& HashedName, const FHashedMaterialParameterInfo& ParameterInfo)
-{
-	// Parameters are first sorted by name hash
-	const uint64 NameHash = HashedName.GetHash();
-	const int32 LowerIndex = Algo::LowerBound(Entry.NameHashes, NameHash);
-	if (LowerIndex < Entry.NameHashes.Num())
-	{
-		const TConstArrayView<uint64> NameHashesUpper = MakeArrayView(&Entry.NameHashes[LowerIndex], Entry.NameHashes.Num() - LowerIndex);
-		const int32 UpperIndex = LowerIndex + Algo::UpperBound(NameHashesUpper, NameHash);
-		if (UpperIndex - LowerIndex > 0)
-		{
-			// more than 1 entry with the same name, next sort by Association/Index
-			auto ProjectionFunc = [](const FMaterialParameterInfo& InParameterInfo)
-			{
-				return FHashedMaterialParameterInfo(FScriptName(), InParameterInfo.Association, InParameterInfo.Index);
-			};
-			auto CompareFunc = [](const FHashedMaterialParameterInfo& Lhs, const FHashedMaterialParameterInfo& Rhs)
-			{
-				if (Lhs.Association != Rhs.Association) return Lhs.Association < Rhs.Association;
-				return Lhs.Index < Rhs.Index;
-			};
-
-			const TConstArrayView<FMaterialParameterInfo> ParameterInfos = MakeArrayView(&Entry.ParameterInfos[LowerIndex], UpperIndex - LowerIndex);
-			return LowerIndex + Algo::LowerBoundBy(ParameterInfos, ParameterInfo, ProjectionFunc, CompareFunc);
-		}
-	}
-	return LowerIndex;
-}
-
 #if WITH_EDITOR
 static int32 TryAddParameter(FMaterialCachedParameters& CachedParameters, EMaterialParameterType Type, const FMaterialParameterInfo& ParameterInfo, const FGuid& ExpressionGuid, bool bOverride = false)
 {
@@ -89,18 +60,23 @@ static int32 TryAddParameter(FMaterialCachedParameters& CachedParameters, EMater
 													CachedParameters.EditorOnlyEntries[static_cast<int32>(Type) - static_cast<int32>(EMaterialParameterType::RuntimeCount)] :
 													CachedParameters.RuntimeEntries[static_cast<int32>(Type)];
 
-	const FHashedName HashedName(ParameterInfo.Name);
-	int32 Index = FindParameterLowerBoundIndex(Entry, HashedName, ParameterInfo);
-
-	if (Index >= Entry.NameHashes.Num() || Entry.ParameterInfos[Index] != ParameterInfo)
+	FSetElementId ElementId = Entry.ParameterInfos.FindId(ParameterInfo);
+	int32 Index = INDEX_NONE;
+	if (!ElementId.IsValidId())
 	{
-		Entry.NameHashes.Insert(HashedName.GetHash(), Index);
-		Entry.ParameterInfos.Insert(ParameterInfo, Index);
+		ElementId = Entry.ParameterInfos.Add(ParameterInfo);
+		Index = ElementId.AsInteger();
 		Entry.ExpressionGuids.Insert(ExpressionGuid, Index);
 		Entry.Overrides.Insert(bOverride, Index);
+
+		// should be valid as long as we don't ever remove elements from ParameterInfos
+		check(Entry.ParameterInfos.Num() == Entry.ExpressionGuids.Num());
+
+		check(Entry.ExpressionGuids.Num() == Entry.Overrides.Num());
 		return Index;
 	}
 
+	Index = ElementId.AsInteger();
 	if (Entry.Overrides[Index] && !Entry.ExpressionGuids[Index].IsValid())
 	{
 		// If Parameter was set by a function override, update to a valid expression guid
@@ -623,7 +599,6 @@ bool FMaterialCachedExpressionData::UpdateForExpressions(const FMaterialCachedEx
 
 void FMaterialCachedParameterEntry::Reset()
 {
-	NameHashes.Reset();
 	ParameterInfos.Reset();
 	ExpressionGuids.Reset();
 	Overrides.Reset();
@@ -662,7 +637,7 @@ void FMaterialCachedParameters::Reset()
 }
 
 
-int32 FMaterialCachedParameters::FindParameterIndex(EMaterialParameterType Type, const FHashedMaterialParameterInfo& HashedParameterInfo, bool bOveriddenOnly) const
+int32 FMaterialCachedParameters::FindParameterIndex(EMaterialParameterType Type, const FMemoryImageMaterialParameterInfo& HashedParameterInfo, bool bOveriddenOnly) const
 {
 	const int32 Index = FindParameterIndex(Type, HashedParameterInfo);
 	if (Index != INDEX_NONE)
@@ -675,20 +650,11 @@ int32 FMaterialCachedParameters::FindParameterIndex(EMaterialParameterType Type,
 	return INDEX_NONE;
 }
 
-int32 FMaterialCachedParameters::FindParameterIndex(EMaterialParameterType Type, const FHashedMaterialParameterInfo& ParameterInfo) const
+int32 FMaterialCachedParameters::FindParameterIndex(EMaterialParameterType Type, const FMemoryImageMaterialParameterInfo& ParameterInfo) const
 {
 	const FMaterialCachedParameterEntry& Entry = GetParameterTypeEntry(Type);
-	const FHashedName HashedName(ParameterInfo.GetName());
-	const int32 Index = FindParameterLowerBoundIndex(Entry, HashedName, ParameterInfo);
-	if (Index < Entry.NameHashes.Num() &&
-		Entry.NameHashes[Index] == HashedName.GetHash() &&
-		Entry.ParameterInfos[Index].Association == ParameterInfo.Association &&
-		Entry.ParameterInfos[Index].Index == ParameterInfo.Index)
-	{
-		return Index;
-	}
-	
-	return INDEX_NONE;
+	const FSetElementId ElementId = Entry.ParameterInfos.FindId(FMaterialParameterInfo(ParameterInfo));
+	return ElementId.AsInteger();
 }
 
 bool FMaterialCachedParameters::IsParameterValid(EMaterialParameterType Type, int32 Index, bool bOveriddenOnly) const
@@ -715,37 +681,37 @@ bool FMaterialCachedParameters::IsDefaultParameterValid(EMaterialParameterType T
 void FMaterialCachedParameters::GetAllParameterInfoOfType(EMaterialParameterType Type, bool bEmptyOutput, TArray<FMaterialParameterInfo>& OutParameterInfo, TArray<FGuid>& OutParameterIds) const
 {
 	const FMaterialCachedParameterEntry& Entry = GetParameterTypeEntry(Type);
-	const int32 NumParameters = Entry.NameHashes.Num();
+	const int32 NumParameters = Entry.ParameterInfos.Num();
 	if (bEmptyOutput)
 	{
 		OutParameterInfo.Empty(NumParameters);
 		OutParameterIds.Empty(NumParameters);
 	}
 
-	for (int32 i = 0; i < NumParameters; ++i)
+	for (TSet<FMaterialParameterInfo>::TConstIterator It(Entry.ParameterInfos); It; ++It)
 	{
-		OutParameterInfo.Add(Entry.ParameterInfos[i]);
-		OutParameterIds.Add(Entry.ExpressionGuids[i]);
+		OutParameterInfo.Add(*It);
+		OutParameterIds.Add(Entry.ExpressionGuids[It.GetId().AsInteger()]);
 	}
 }
 
 void FMaterialCachedParameters::GetAllGlobalParameterInfoOfType(EMaterialParameterType Type, bool bEmptyOutput, TArray<FMaterialParameterInfo>& OutParameterInfo, TArray<FGuid>& OutParameterIds) const
 {
 	const FMaterialCachedParameterEntry& Entry = GetParameterTypeEntry(Type);
-	const int32 NumParameters = Entry.NameHashes.Num();
+	const int32 NumParameters = Entry.ParameterInfos.Num();
 	if (bEmptyOutput)
 	{
 		OutParameterInfo.Empty(NumParameters);
 		OutParameterIds.Empty(NumParameters);
 	}
 
-	for (int32 i = 0; i < NumParameters; ++i)
+	for (TSet<FMaterialParameterInfo>::TConstIterator It(Entry.ParameterInfos); It; ++It)
 	{
-		const FMaterialParameterInfo& ParameterInfo = Entry.ParameterInfos[i];
+		const FMaterialParameterInfo& ParameterInfo = *It;
 		if (ParameterInfo.Association == GlobalParameter)
 		{
 			OutParameterInfo.Add(ParameterInfo);
-			OutParameterIds.Add(Entry.ExpressionGuids[i]);
+			OutParameterIds.Add(Entry.ExpressionGuids[It.GetId().AsInteger()]);
 		}
 	}
 }
diff --git a/Engine/Source/Runtime/Engine/Public/MaterialCachedData.h b/Engine/Source/Runtime/Engine/Public/MaterialCachedData.h
--- a/Engine/Source/Runtime/Engine/Public/MaterialCachedData.h
+++ b/Engine/Source/Runtime/Engine/Public/MaterialCachedData.h
@@ -103,11 +103,10 @@ struct FMaterialCachedParameterEntry
 
 	void Reset();
 
+	// This is used to map FMaterialParameterInfos to indices, which are then used to index various TArrays containing values for each type of parameter
+	// (ExpressionGuids and Overrides, along with ScalarValues, VectorValues, etc)
 	UPROPERTY()
-	TArray<uint64> NameHashes;
-
-	UPROPERTY()
-	TArray<FMaterialParameterInfo> ParameterInfos;
+	TSet<FMaterialParameterInfo> ParameterInfos;
 
 	UPROPERTY()
 	TArray<FGuid> ExpressionGuids; // editor-only?
@@ -144,11 +143,11 @@ struct FMaterialCachedParameters
 #else
 	inline const FMaterialCachedParameterEntry& GetParameterTypeEntry(EMaterialParameterType Type) const { return RuntimeEntries[static_cast<int32>(Type)]; }
 #endif
+
 	inline int32 GetNumParameters(EMaterialParameterType Type) const { return GetParameterTypeEntry(Type).ParameterInfos.Num(); }
-	inline const FName& GetParameterName(EMaterialParameterType Type, int32 Index) const { return GetParameterTypeEntry(Type).ParameterInfos[Index].Name; }
 
-	int32 FindParameterIndex(EMaterialParameterType Type, const FHashedMaterialParameterInfo& HashedParameterInfo, bool bOveriddenOnly) const;
-	int32 FindParameterIndex(EMaterialParameterType Type, const FHashedMaterialParameterInfo& HashedParameterInfo) const;
+	int32 FindParameterIndex(EMaterialParameterType Type, const FMemoryImageMaterialParameterInfo& HashedParameterInfo, bool bOveriddenOnly) const;
+	int32 FindParameterIndex(EMaterialParameterType Type, const FMemoryImageMaterialParameterInfo& HashedParameterInfo) const;
 	bool IsParameterValid(EMaterialParameterType Type, int32 Index, bool bOveriddenOnly) const;
 	bool IsDefaultParameterValid(EMaterialParameterType Type, int32 Index, bool bOveriddenOnly, bool bCheckOwnedGlobalOverrides) const;
 	void GetAllParameterInfoOfType(EMaterialParameterType Type, bool bEmptyOutput, TArray<FMaterialParameterInfo>& OutParameterInfo, TArray<FGuid>& OutParameterIds) const;
