Deprecate some directly accessible fields in UStaticMesh

diff --git a/Engine/Plugins/Editor/GLTFImporter/Source/GLTFImporter/Private/GLTFStaticMeshFactory.cpp b/Engine/Plugins/Editor/GLTFImporter/Source/GLTFImporter/Private/GLTFStaticMeshFactory.cpp
--- a/Engine/Plugins/Editor/GLTFImporter/Source/GLTFImporter/Private/GLTFStaticMeshFactory.cpp
+++ b/Engine/Plugins/Editor/GLTFImporter/Source/GLTFImporter/Private/GLTFStaticMeshFactory.cpp
@@ -117,7 +117,7 @@ namespace GLTF
 		if (bGenerateLightmapUVs)
 		{
 			// Generate a new UV set based off the highest index UV set in the mesh
-			StaticMesh.LightMapCoordinateIndex             = NumUVs;
+			StaticMesh.SetLightMapCoordinateIndex(NumUVs);
 			SourceModel.BuildSettings.SrcLightmapIndex     = NumUVs - 1;
 			SourceModel.BuildSettings.DstLightmapIndex     = NumUVs;
 			SourceModel.BuildSettings.bGenerateLightmapUVs = true;
@@ -125,14 +125,14 @@ namespace GLTF
 		else if (!bMeshHasUVs)
 		{
 			// Generate automatically a UV for correct lighting if mesh has none
-			StaticMesh.LightMapCoordinateIndex             = 1;
+			StaticMesh.SetLightMapCoordinateIndex(1);
 			SourceModel.BuildSettings.SrcLightmapIndex     = 0;
 			SourceModel.BuildSettings.DstLightmapIndex     = 1;
 			SourceModel.BuildSettings.bGenerateLightmapUVs = true;
 		}
 		else
 		{
-			StaticMesh.LightMapCoordinateIndex             = NumUVs - 1;
+			StaticMesh.SetLightMapCoordinateIndex(NumUVs - 1);
 			SourceModel.BuildSettings.bGenerateLightmapUVs = false;
 		}
 
diff --git a/Engine/Plugins/Editor/SpeedTreeImporter/Source/SpeedTreeImporter/Private/SpeedTreeImportFactory.cpp b/Engine/Plugins/Editor/SpeedTreeImporter/Source/SpeedTreeImporter/Private/SpeedTreeImportFactory.cpp
--- a/Engine/Plugins/Editor/SpeedTreeImporter/Source/SpeedTreeImporter/Private/SpeedTreeImportFactory.cpp
+++ b/Engine/Plugins/Editor/SpeedTreeImporter/Source/SpeedTreeImporter/Private/SpeedTreeImportFactory.cpp
@@ -1991,8 +1991,8 @@ UObject* USpeedTreeImportFactory::FactoryCreateBinary7(UClass* InClass, UObject*
 
 				// Lightmap data
 				StaticMesh->SetLightingGuid();
-				StaticMesh->LightMapResolution = 128;
-				StaticMesh->LightMapCoordinateIndex = 1;
+				StaticMesh->SetLightMapResolution(128);
+				StaticMesh->SetLightMapCoordinateIndex(1);
 
 				// set up SpeedTree wind data
 				if (!StaticMesh->SpeedTreeWind.IsValid())
@@ -2519,8 +2519,8 @@ UObject* USpeedTreeImportFactory::FactoryCreateBinary8(UClass* InClass, UObject*
 
 		// Lightmap data
 		StaticMesh->SetLightingGuid();
-		StaticMesh->LightMapResolution = SpeedTree.LightmapSize();
-		StaticMesh->LightMapCoordinateIndex = 1;		
+		StaticMesh->SetLightMapResolution(SpeedTree.LightmapSize());
+		StaticMesh->SetLightMapCoordinateIndex(1);
 
 		// set up SpeedTree wind data
 		if (!StaticMesh->SpeedTreeWind.IsValid())
@@ -2899,8 +2899,8 @@ UObject* USpeedTreeImportFactory::FactoryCreateBinary9(UClass* InClass, UObject*
 
 		// Lightmap data
 		StaticMesh->SetLightingGuid();
-		StaticMesh->LightMapResolution = SpeedTree.LightmapSize();
-		StaticMesh->LightMapCoordinateIndex = (NumUVs - 1);
+		StaticMesh->SetLightMapResolution(SpeedTree.LightmapSize());
+		StaticMesh->SetLightMapCoordinateIndex((NumUVs - 1));
 
 		// make geometry LODs
 		for (uint32 LODIndex = 0; LODIndex < SpeedTree.Lods().Count(); ++LODIndex)
diff --git a/Engine/Plugins/Enterprise/DataprepEditor/Source/DataprepEditor/Private/Widgets/SDataprepEditorViewport.cpp b/Engine/Plugins/Enterprise/DataprepEditor/Source/DataprepEditor/Private/Widgets/SDataprepEditorViewport.cpp
--- a/Engine/Plugins/Enterprise/DataprepEditor/Source/DataprepEditor/Private/Widgets/SDataprepEditorViewport.cpp
+++ b/Engine/Plugins/Enterprise/DataprepEditor/Source/DataprepEditor/Private/Widgets/SDataprepEditorViewport.cpp
@@ -1548,7 +1548,7 @@ FPrimitiveSceneProxy* UCustomStaticMeshComponent::CreateSceneProxy()
 	}
 
 	const FStaticMeshLODResourcesArray& LODResources = GetStaticMesh()->GetRenderData()->LODResources;
-	if (LODResources.Num() == 0	|| LODResources[FMath::Clamp<int32>(GetStaticMesh()->MinLOD.Default, 0, LODResources.Num()-1)].VertexBuffers.StaticMeshVertexBuffer.GetNumVertices() == 0)
+	if (LODResources.Num() == 0	|| LODResources[FMath::Clamp<int32>(GetStaticMesh()->GetMinLOD().Default, 0, LODResources.Num()-1)].VertexBuffers.StaticMeshVertexBuffer.GetNumVertices() == 0)
 	{
 		return nullptr;
 	}
diff --git a/Engine/Plugins/Enterprise/DatasmithContent/Source/DatasmithContent/Private/ObjectTemplates/DatasmithStaticMeshTemplate.cpp b/Engine/Plugins/Enterprise/DatasmithContent/Source/DatasmithContent/Private/ObjectTemplates/DatasmithStaticMeshTemplate.cpp
--- a/Engine/Plugins/Enterprise/DatasmithContent/Source/DatasmithContent/Private/ObjectTemplates/DatasmithStaticMeshTemplate.cpp
+++ b/Engine/Plugins/Enterprise/DatasmithContent/Source/DatasmithContent/Private/ObjectTemplates/DatasmithStaticMeshTemplate.cpp
@@ -188,9 +188,9 @@ UObject* UDatasmithStaticMeshTemplate::UpdateObject( UObject* Destination, bool
 #if WITH_EDITORONLY_DATA
 	UDatasmithStaticMeshTemplate* PreviousStaticMeshTemplate = !bForce ? FDatasmithObjectTemplateUtils::GetObjectTemplate< UDatasmithStaticMeshTemplate >( StaticMesh ) : nullptr;
 
-	DATASMITHOBJECTTEMPLATE_CONDITIONALSET( LightMapCoordinateIndex, StaticMesh, PreviousStaticMeshTemplate );
+	DATASMITHOBJECTTEMPLATE_CONDITIONALSET_ACCESSORS( LightMapCoordinateIndex, StaticMesh, PreviousStaticMeshTemplate );
 
-	DATASMITHOBJECTTEMPLATE_CONDITIONALSET( LightMapResolution, StaticMesh, PreviousStaticMeshTemplate );
+	DATASMITHOBJECTTEMPLATE_CONDITIONALSET_ACCESSORS( LightMapResolution, StaticMesh, PreviousStaticMeshTemplate );
 
 	// Section info map
 	bool bResetSectionInfoMap = false;
@@ -309,8 +309,8 @@ void UDatasmithStaticMeshTemplate::Load( const UObject* Source )
 		return;
 	}
 
-	LightMapCoordinateIndex = SourceStaticMesh->LightMapCoordinateIndex;
-	LightMapResolution = SourceStaticMesh->LightMapResolution;
+	LightMapCoordinateIndex = SourceStaticMesh->GetLightMapCoordinateIndex();
+	LightMapResolution = SourceStaticMesh->GetLightMapResolution();
 
 	// Section info map
 	SectionInfoMap.Load( SourceStaticMesh->GetSectionInfoMap() );
diff --git a/Engine/Plugins/Enterprise/DatasmithContent/Source/DatasmithContent/Public/ObjectTemplates/DatasmithObjectTemplate.h b/Engine/Plugins/Enterprise/DatasmithContent/Source/DatasmithContent/Public/ObjectTemplates/DatasmithObjectTemplate.h
--- a/Engine/Plugins/Enterprise/DatasmithContent/Source/DatasmithContent/Public/ObjectTemplates/DatasmithObjectTemplate.h
+++ b/Engine/Plugins/Enterprise/DatasmithContent/Source/DatasmithContent/Public/ObjectTemplates/DatasmithObjectTemplate.h
@@ -88,6 +88,12 @@ protected:
 		Destination->MemberName = MemberName; \
 	}
 
+#define DATASMITHOBJECTTEMPLATE_CONDITIONALSET_ACCESSORS(MemberName, Destination, PreviousTemplate) \
+	if ( !PreviousTemplate || Destination->Get##MemberName() == PreviousTemplate->MemberName ) \
+	{ \
+		Destination->Set##MemberName(MemberName); \
+	}
+
 // Specialized version of DATASMITHOBJECTTEMPLATE_CONDITIONALSET to handle SoftObjectPtr to Ptr assignment
 #define DATASMITHOBJECTTEMPLATE_CONDITIONALSETSOFTOBJECTPTR(MemberName, Destination, PreviousTemplate) \
 	if ( !PreviousTemplate || Destination->MemberName == PreviousTemplate->MemberName.Get() ) \
diff --git a/Engine/Plugins/Enterprise/DatasmithImporter/Source/DatasmithImporter/Private/DatasmithBlueprintLibrary.cpp b/Engine/Plugins/Enterprise/DatasmithImporter/Source/DatasmithImporter/Private/DatasmithBlueprintLibrary.cpp
--- a/Engine/Plugins/Enterprise/DatasmithImporter/Source/DatasmithImporter/Private/DatasmithBlueprintLibrary.cpp
+++ b/Engine/Plugins/Enterprise/DatasmithImporter/Source/DatasmithImporter/Private/DatasmithBlueprintLibrary.cpp
@@ -907,7 +907,7 @@ void UDatasmithStaticMeshBlueprintLibrary::SetupStaticLighting(const TArray< UOb
 
 					if (GeneratedLightmapChannel < Lightmass::MAX_TEXCOORDS)
 					{
-						StaticMesh->LightMapCoordinateIndex = GeneratedLightmapChannel;
+						StaticMesh->SetLightMapCoordinateIndex(GeneratedLightmapChannel);
 					}
 					else
 					{
@@ -915,10 +915,10 @@ void UDatasmithStaticMeshBlueprintLibrary::SetupStaticLighting(const TArray< UOb
 						break;
 					}
 				}
-				else if (StaticMesh->LightMapCoordinateIndex > MaxBiggestUVChannel && bDidChangeSettings)
+				else if (StaticMesh->GetLightMapCoordinateIndex() > MaxBiggestUVChannel && bDidChangeSettings)
 				{
 					// If we are not generating the lightmap anymore make sure we are selecting a valid lightmap index.
-					StaticMesh->LightMapCoordinateIndex = MaxBiggestUVChannel;
+					StaticMesh->SetLightMapCoordinateIndex(MaxBiggestUVChannel);
 				}
 			}
 		}
@@ -970,7 +970,7 @@ void UDatasmithStaticMeshBlueprintLibrary::ComputeLightmapResolution(const TMap<
 					StaticMesh->Modify();
 				}
 
-				StaticMesh->LightMapResolution = LightMapResolution;
+				StaticMesh->SetLightMapResolution(LightMapResolution);
 
 				if(bApplyChanges)
 				{
@@ -1059,7 +1059,7 @@ int32 UDatasmithStaticMeshBlueprintLibrary::ComputeLightmapResolution(UStaticMes
 	const FPositionVertexBuffer& PositionBuffer = StaticMesh->GetRenderData()->LODResources[0].VertexBuffers.PositionVertexBuffer;
 	const FStaticMeshVertexBuffer& VertexBuffer = StaticMesh->GetRenderData()->LODResources[0].VertexBuffers.StaticMeshVertexBuffer;
 
-	if (VertexBuffer.GetNumTexCoords() <= (uint32)StaticMesh->LightMapCoordinateIndex)
+	if (VertexBuffer.GetNumTexCoords() <= (uint32)StaticMesh->GetLightMapCoordinateIndex())
 	{
 		return 0;
 	}
@@ -1076,7 +1076,7 @@ int32 UDatasmithStaticMeshBlueprintLibrary::ComputeLightmapResolution(UStaticMes
 		{
 			uint32 VertexIndex = IndexBuffer.GetIndex(TriangleIndex * 3 + CornerIndex);
 			VertexPosition[CornerIndex] = PositionBuffer.VertexPosition(VertexIndex);
-			LightmapUVs[CornerIndex] = VertexBuffer.GetVertexUV(VertexIndex, StaticMesh->LightMapCoordinateIndex);
+			LightmapUVs[CornerIndex] = VertexBuffer.GetVertexUV(VertexIndex, StaticMesh->GetLightMapCoordinateIndex());
 		}
 
 		const float PolygonArea = DatasmithStaticMeshBlueprintLibraryUtil::ParallelogramArea(VertexPosition[0], VertexPosition[1], VertexPosition[2]);
diff --git a/Engine/Plugins/Experimental/AlembicImporter/Source/AlembicLibrary/Private/AbcImporter.cpp b/Engine/Plugins/Experimental/AlembicImporter/Source/AlembicLibrary/Private/AbcImporter.cpp
--- a/Engine/Plugins/Experimental/AlembicImporter/Source/AlembicLibrary/Private/AbcImporter.cpp
+++ b/Engine/Plugins/Experimental/AlembicImporter/Source/AlembicLibrary/Private/AbcImporter.cpp
@@ -192,8 +192,8 @@ UStaticMesh* FAbcImporter::CreateStaticMeshFromSample(UObject* InParent, const F
 		StaticMesh->SetLightingGuid();
 
 		// Set it to use textured lightmaps. Note that Build Lighting will do the error-checking (texcoord index exists for all LODs, etc).
-		StaticMesh->LightMapResolution = 64;
-		StaticMesh->LightMapCoordinateIndex = 1;
+		StaticMesh->SetLightMapResolution(64);
+		StaticMesh->SetLightMapCoordinateIndex(1);
 
 		// Material setup, since there isn't much material information in the Alembic file, 
 		UMaterial* DefaultMaterial = UMaterial::GetDefaultMaterial(MD_Surface);
diff --git a/Engine/Plugins/Experimental/GPULightmass/Source/GPULightmass/Private/Scene/InstancedStaticMesh.cpp b/Engine/Plugins/Experimental/GPULightmass/Source/GPULightmass/Private/Scene/InstancedStaticMesh.cpp
--- a/Engine/Plugins/Experimental/GPULightmass/Source/GPULightmass/Private/Scene/InstancedStaticMesh.cpp
+++ b/Engine/Plugins/Experimental/GPULightmass/Source/GPULightmass/Private/Scene/InstancedStaticMesh.cpp
@@ -96,8 +96,8 @@ void FInstanceGroup::AllocateLightmaps(TEntityArray<FLightmap>& LightmapContaine
 		if (bFit 
 			&& LightMapWidth > 0
 			&& LightMapHeight > 0
-			&& ComponentUObject->GetStaticMesh()->LightMapCoordinateIndex >= 0
-			&& (uint32)ComponentUObject->GetStaticMesh()->LightMapCoordinateIndex < LODModel.VertexBuffers.StaticMeshVertexBuffer.GetNumTexCoords())
+			&& ComponentUObject->GetStaticMesh()->GetLightMapCoordinateIndex() >= 0
+			&& (uint32)ComponentUObject->GetStaticMesh()->GetLightMapCoordinateIndex() < LODModel.VertexBuffers.StaticMeshVertexBuffer.GetNumTexCoords())
 		{
 			bValidTextureMap = true;
 		}
diff --git a/Engine/Plugins/Experimental/GPULightmass/Source/GPULightmass/Private/Scene/Scene.cpp b/Engine/Plugins/Experimental/GPULightmass/Source/GPULightmass/Private/Scene/Scene.cpp
--- a/Engine/Plugins/Experimental/GPULightmass/Source/GPULightmass/Private/Scene/Scene.cpp
+++ b/Engine/Plugins/Experimental/GPULightmass/Source/GPULightmass/Private/Scene/Scene.cpp
@@ -679,7 +679,7 @@ void FScene::AddGeometryInstanceFromComponent(UStaticMeshComponent* InComponent)
 
 	RegisteredStaticMeshComponentUObjects.Add(InComponent, Instance);
 
-	const int32 SMCurrentMinLOD = InComponent->GetStaticMesh()->MinLOD.Default;
+	const int32 SMCurrentMinLOD = InComponent->GetStaticMesh()->GetMinLOD().Default;
 	int32 EffectiveMinLOD = InComponent->bOverrideMinLOD ? InComponent->MinLOD : SMCurrentMinLOD;
 
 	// Find the first LOD with any vertices (ie that haven't been stripped)
diff --git a/Engine/Plugins/Experimental/GPULightmass/Source/GPULightmass/Private/Scene/StaticMesh.cpp b/Engine/Plugins/Experimental/GPULightmass/Source/GPULightmass/Private/Scene/StaticMesh.cpp
--- a/Engine/Plugins/Experimental/GPULightmass/Source/GPULightmass/Private/Scene/StaticMesh.cpp
+++ b/Engine/Plugins/Experimental/GPULightmass/Source/GPULightmass/Private/Scene/StaticMesh.cpp
@@ -43,8 +43,8 @@ void FStaticMeshInstance::AllocateLightmaps(TEntityArray<FLightmap>& LightmapCon
 		bool bValidTextureMap = false;
 		if (BaseLightMapWidth > 0
 			&& BaseLightMapHeight > 0
-			&& ComponentUObject->GetStaticMesh()->LightMapCoordinateIndex >= 0
-			&& (uint32)ComponentUObject->GetStaticMesh()->LightMapCoordinateIndex < LODModel.VertexBuffers.StaticMeshVertexBuffer.GetNumTexCoords())
+			&& ComponentUObject->GetStaticMesh()->GetLightMapCoordinateIndex() >= 0
+			&& (uint32)ComponentUObject->GetStaticMesh()->GetLightMapCoordinateIndex() < LODModel.VertexBuffers.StaticMeshVertexBuffer.GetNumTexCoords())
 		{
 			bValidTextureMap = true;
 		}
diff --git a/Engine/Plugins/Experimental/PanoramicCapture/Source/PanoramicCapture/Private/StereoStaticMeshComponent.cpp b/Engine/Plugins/Experimental/PanoramicCapture/Source/PanoramicCapture/Private/StereoStaticMeshComponent.cpp
--- a/Engine/Plugins/Experimental/PanoramicCapture/Source/PanoramicCapture/Private/StereoStaticMeshComponent.cpp
+++ b/Engine/Plugins/Experimental/PanoramicCapture/Source/PanoramicCapture/Private/StereoStaticMeshComponent.cpp
@@ -57,7 +57,7 @@ FPrimitiveSceneProxy* UStereoStaticMeshComponent::CreateSceneProxy()
     if ((GetStaticMesh() == nullptr) ||
 		(GetStaticMesh()->GetRenderData() == nullptr) ||
 		(GetStaticMesh()->GetRenderData()->LODResources.Num() == 0) ||
-		(GetStaticMesh()->GetRenderData()->LODResources[GetStaticMesh()->MinLOD.GetValue()].VertexBuffers.PositionVertexBuffer.GetNumVertices() == 0))
+		(GetStaticMesh()->GetRenderData()->LODResources[GetStaticMesh()->GetMinLOD().GetValue()].VertexBuffers.PositionVertexBuffer.GetNumVertices() == 0))
     {
         return nullptr;
     }
diff --git a/Engine/Plugins/FX/Niagara/Source/Niagara/Private/NiagaraDataInterfaceStaticMesh.cpp b/Engine/Plugins/FX/Niagara/Source/Niagara/Private/NiagaraDataInterfaceStaticMesh.cpp
--- a/Engine/Plugins/FX/Niagara/Source/Niagara/Private/NiagaraDataInterfaceStaticMesh.cpp
+++ b/Engine/Plugins/FX/Niagara/Source/Niagara/Private/NiagaraDataInterfaceStaticMesh.cpp
@@ -317,7 +317,7 @@ bool FNDIStaticMesh_InstanceData::Init(UNiagaraDataInterfaceStaticMesh* Interfac
 		}
 #endif
 
-	    MinLOD = Mesh->MinLOD.GetValue();
+	    MinLOD = Mesh->GetMinLOD().GetValue();
 	    CachedLODIdx = Mesh->GetRenderData()->GetCurrentFirstLODIdx(MinLOD);
 
 		bMeshAllowsCpuAccess = Mesh->bAllowCPUAccess;
diff --git a/Engine/Plugins/FX/Niagara/Source/Niagara/Private/NiagaraRendererMeshes.cpp b/Engine/Plugins/FX/Niagara/Source/Niagara/Private/NiagaraRendererMeshes.cpp
--- a/Engine/Plugins/FX/Niagara/Source/Niagara/Private/NiagaraRendererMeshes.cpp
+++ b/Engine/Plugins/FX/Niagara/Source/Niagara/Private/NiagaraRendererMeshes.cpp
@@ -126,7 +126,7 @@ FNiagaraRendererMeshes::FNiagaraRendererMeshes(ERHIFeatureLevel::Type FeatureLev
 	RendererLayoutWithCustomSorting = &Properties->RendererLayoutWithCustomSorting;
 	RendererLayoutWithoutCustomSorting = &Properties->RendererLayoutWithoutCustomSorting;
 
-	MeshMinimumLOD = Properties->ParticleMesh->MinLOD.GetValue();
+	MeshMinimumLOD = Properties->ParticleMesh->GetMinLOD().GetValue();
 
 	if (MeshRenderData)
 	{
diff --git a/Engine/Source/Developer/HierarchicalLODUtilities/Private/HierarchicalLODProxyProcessor.cpp b/Engine/Source/Developer/HierarchicalLODUtilities/Private/HierarchicalLODProxyProcessor.cpp
--- a/Engine/Source/Developer/HierarchicalLODUtilities/Private/HierarchicalLODProxyProcessor.cpp
+++ b/Engine/Source/Developer/HierarchicalLODUtilities/Private/HierarchicalLODProxyProcessor.cpp
@@ -58,7 +58,7 @@ bool FHierarchicalLODProxyProcessor::Tick(float DeltaTime)
 		check(MainMesh != nullptr);
 
 		// Force lightmap coordinate to 0 for proxy meshes
-		MainMesh->LightMapCoordinateIndex = 0;
+		MainMesh->SetLightMapCoordinateIndex(0);
 		// Trigger post edit change, simulating we made a change in the Static mesh editor (could only call Build, but this is for possible future changes)
 		MainMesh->PostEditChange();
 
diff --git a/Engine/Source/Developer/HierarchicalLODUtilities/Private/HierarchicalLODUtilities.cpp b/Engine/Source/Developer/HierarchicalLODUtilities/Private/HierarchicalLODUtilities.cpp
--- a/Engine/Source/Developer/HierarchicalLODUtilities/Private/HierarchicalLODUtilities.cpp
+++ b/Engine/Source/Developer/HierarchicalLODUtilities/Private/HierarchicalLODUtilities.cpp
@@ -296,8 +296,8 @@ UStaticMesh* CreateImposterStaticMesh(UStaticMeshComponent* InComponent, UMateri
 		StaticMesh->SetLightingGuid();
 
 		// Set it to use textured lightmaps. Note that Build Lighting will do the error-checking (texcoordindex exists for all LODs, etc).
-		StaticMesh->LightMapResolution = InProxySettings.LightMapResolution;
-		StaticMesh->LightMapCoordinateIndex = 1;
+		StaticMesh->SetLightMapResolution(InProxySettings.LightMapResolution);
+		StaticMesh->SetLightMapCoordinateIndex(1);
 
 		// Add one LOD for the base mesh
 		StaticMesh->SetNumSourceModels(0);
@@ -341,8 +341,8 @@ UStaticMesh* CreateImposterStaticMesh(UStaticMeshComponent* InComponent, UMateri
 		{
 			const FBox StaticMeshBox = StaticMesh->GetBoundingBox();
 			const FBox CombinedBox = StaticMeshBox + InComponent->GetStaticMesh()->GetBoundingBox();
-			StaticMesh->PositiveBoundsExtension = (CombinedBox.Max - StaticMeshBox.Max);
-			StaticMesh->NegativeBoundsExtension = (StaticMeshBox.Min - CombinedBox.Min);
+			StaticMesh->SetPositiveBoundsExtension((CombinedBox.Max - StaticMeshBox.Max));
+			StaticMesh->SetNegativeBoundsExtension((StaticMeshBox.Min - CombinedBox.Min));
 			StaticMesh->CalculateExtendedBounds();
 		}	
 	
diff --git a/Engine/Source/Developer/MeshMergeUtilities/Private/MeshMergeUtilities.cpp b/Engine/Source/Developer/MeshMergeUtilities/Private/MeshMergeUtilities.cpp
--- a/Engine/Source/Developer/MeshMergeUtilities/Private/MeshMergeUtilities.cpp
+++ b/Engine/Source/Developer/MeshMergeUtilities/Private/MeshMergeUtilities.cpp
@@ -1384,7 +1384,7 @@ void FMeshMergeUtilities::CreateProxyMesh(const TArray<UStaticMeshComponent*>& I
 						TVertexInstanceAttributesRef<FVector2D> VertexInstanceUVs = FStaticMeshAttributes(*MeshDescription).GetVertexInstanceUVs();
 
 						// If we already have lightmap uvs generated and they are valid, we can reuse those instead of having to generate new ones
-						const int32 LightMapCoordinateIndex = StaticMeshComponent->GetStaticMesh()->LightMapCoordinateIndex;
+						const int32 LightMapCoordinateIndex = StaticMeshComponent->GetStaticMesh()->GetLightMapCoordinateIndex();
 						if (InMeshProxySettings.bReuseMeshLightmapUVs &&
 							LightMapCoordinateIndex > 0 &&
 							VertexInstanceUVs.GetNumElements() > 0 &&
@@ -1435,7 +1435,7 @@ void FMeshMergeUtilities::CreateProxyMesh(const TArray<UStaticMeshComponent*>& I
 							if (MeshMapBuildData)
 							{
 								MeshSettings.LightMap = MeshMapBuildData->LightMap;
-								MeshSettings.LightMapIndex = StaticMeshComponent->GetStaticMesh()->LightMapCoordinateIndex;
+								MeshSettings.LightMapIndex = StaticMeshComponent->GetStaticMesh()->GetLightMapCoordinateIndex();
 							}
 						}
 
@@ -1878,7 +1878,7 @@ bool RetrieveRawMeshData(FMeshMergeDataTracker& DataTracker
 	else if (Component->GetStaticMesh() != nullptr)
 	{
 		// If the mesh is valid at this point, record the lightmap UV so we have a record for use later
-		DataTracker.AddLightmapChannelRecord(ComponentIndex, LODIndex, Component->GetStaticMesh()->LightMapCoordinateIndex);
+		DataTracker.AddLightmapChannelRecord(ComponentIndex, LODIndex, Component->GetStaticMesh()->GetLightMapCoordinateIndex());
 	}
 	return bValidMesh;
 }
@@ -2197,7 +2197,7 @@ void FMeshMergeUtilities::MergeComponentsToStaticMesh(const TArray<UPrimitiveCom
 					if (bRequiresUniqueUVs || MeshData.RawMeshDescription->VertexInstances().Num() > 0)
 					{
 						// Check if there are lightmap uvs available?
-						const int32 LightMapUVIndex = StaticMeshComponentsToMerge[Key.GetMeshIndex()]->GetStaticMesh()->LightMapCoordinateIndex;
+						const int32 LightMapUVIndex = StaticMeshComponentsToMerge[Key.GetMeshIndex()]->GetStaticMesh()->GetLightMapCoordinateIndex();
 
 						TVertexInstanceAttributesRef<FVector2D> VertexInstanceUVs = FStaticMeshAttributes(*MeshData.RawMeshDescription).GetVertexInstanceUVs();
 						if (InSettings.bReuseMeshLightmapUVs && VertexInstanceUVs.GetNumElements() > 0 && VertexInstanceUVs.GetNumChannels() > LightMapUVIndex)
@@ -2723,8 +2723,8 @@ void FMeshMergeUtilities::MergeComponentsToStaticMesh(const TArray<UPrimitiveCom
 		StaticMesh->SetLightingGuid();
 		if (InSettings.bGenerateLightMapUV)
 		{
-			StaticMesh->LightMapResolution = InSettings.TargetLightMapResolution;
-			StaticMesh->LightMapCoordinateIndex = LightMapUVChannel;
+			StaticMesh->SetLightMapResolution(InSettings.TargetLightMapResolution);
+			StaticMesh->SetLightMapCoordinateIndex(LightMapUVChannel);
 		}
 
 		const bool bContainsImposters = ImposterComponents.Num() > 0;
@@ -2845,7 +2845,7 @@ void FMeshMergeUtilities::MergeComponentsToStaticMesh(const TArray<UPrimitiveCom
 
 		//Set the Imported version before calling the build
 		StaticMesh->ImportVersion = EImportStaticMeshVersion::LastVersion;
-		StaticMesh->LightMapResolution = InSettings.bComputedLightMapResolution ? DataTracker.GetLightMapDimension() : InSettings.TargetLightMapResolution;
+		StaticMesh->SetLightMapResolution(InSettings.bComputedLightMapResolution ? DataTracker.GetLightMapDimension() : InSettings.TargetLightMapResolution);
 
 #if WITH_EDITOR
 		//If we are running the automation test
@@ -2860,8 +2860,8 @@ void FMeshMergeUtilities::MergeComponentsToStaticMesh(const TArray<UPrimitiveCom
 		{
 			const FBox StaticMeshBox = StaticMesh->GetBoundingBox();
 			const FBox CombinedBox = StaticMeshBox + ImposterBounds;
-			StaticMesh->PositiveBoundsExtension = (CombinedBox.Max - StaticMeshBox.Max);
-			StaticMesh->NegativeBoundsExtension = (StaticMeshBox.Min - CombinedBox.Min);
+			StaticMesh->SetPositiveBoundsExtension((CombinedBox.Max - StaticMeshBox.Max));
+			StaticMesh->SetNegativeBoundsExtension((StaticMeshBox.Min - CombinedBox.Min));
 			StaticMesh->CalculateExtendedBounds();
 		}		
 
diff --git a/Engine/Source/Developer/MeshMergeUtilities/Private/ProxyGenerationProcessor.cpp b/Engine/Source/Developer/MeshMergeUtilities/Private/ProxyGenerationProcessor.cpp
--- a/Engine/Source/Developer/MeshMergeUtilities/Private/ProxyGenerationProcessor.cpp
+++ b/Engine/Source/Developer/MeshMergeUtilities/Private/ProxyGenerationProcessor.cpp
@@ -191,8 +191,8 @@ void FProxyGenerationProcessor::ProcessJob(const FGuid& JobGuid, FProxyGeneratio
 	StaticMesh->SetLightingGuid();
 
 	// Set it to use textured lightmaps. Note that Build Lighting will do the error-checking (texcoordindex exists for all LODs, etc).
-	StaticMesh->LightMapResolution = Data->MergeData->InProxySettings.LightMapResolution;
-	StaticMesh->LightMapCoordinateIndex = 1;
+	StaticMesh->SetLightMapResolution(Data->MergeData->InProxySettings.LightMapResolution);
+	StaticMesh->SetLightMapCoordinateIndex(1);
 
 
 	FStaticMeshSourceModel& SrcModel = StaticMesh->AddSourceModel();
@@ -329,8 +329,8 @@ void FProxyGenerationProcessor::ProcessJob(const FGuid& JobGuid, FProxyGeneratio
 	{
 		const FBox StaticMeshBox = StaticMesh->GetBoundingBox();
 		const FBox CombinedBox = StaticMeshBox + ImposterBounds;
-		StaticMesh->PositiveBoundsExtension = (CombinedBox.Max - StaticMeshBox.Max);
-		StaticMesh->NegativeBoundsExtension = (StaticMeshBox.Min - CombinedBox.Min);
+		StaticMesh->SetPositiveBoundsExtension((CombinedBox.Max - StaticMeshBox.Max));
+		StaticMesh->SetNegativeBoundsExtension((StaticMeshBox.Min - CombinedBox.Min));
 		StaticMesh->CalculateExtendedBounds();
 	}
 
diff --git a/Engine/Source/Developer/MeshMergeUtilities/Private/StaticMeshAdapter.cpp b/Engine/Source/Developer/MeshMergeUtilities/Private/StaticMeshAdapter.cpp
--- a/Engine/Source/Developer/MeshMergeUtilities/Private/StaticMeshAdapter.cpp
+++ b/Engine/Source/Developer/MeshMergeUtilities/Private/StaticMeshAdapter.cpp
@@ -38,7 +38,7 @@ int32 FStaticMeshAdapter::GetMaterialIndex(int32 LODIndex, int32 SectionIndex) c
 
 void FStaticMeshAdapter::ApplySettings(int32 LODIndex, FMeshData& InOutMeshData) const
 {
-	InOutMeshData.LightMapIndex = StaticMesh->LightMapCoordinateIndex;
+	InOutMeshData.LightMapIndex = StaticMesh->GetLightMapCoordinateIndex();
 }
 
 UPackage* FStaticMeshAdapter::GetOuter() const
@@ -102,7 +102,7 @@ bool FStaticMeshAdapter::IsAsset() const
 
 int32 FStaticMeshAdapter::LightmapUVIndex() const
 {
-	return StaticMesh->LightMapCoordinateIndex;
+	return StaticMesh->GetLightMapCoordinateIndex();
 }
 
 FBoxSphereBounds FStaticMeshAdapter::GetBounds() const
diff --git a/Engine/Source/Developer/MeshMergeUtilities/Private/StaticMeshComponentAdapter.cpp b/Engine/Source/Developer/MeshMergeUtilities/Private/StaticMeshComponentAdapter.cpp
--- a/Engine/Source/Developer/MeshMergeUtilities/Private/StaticMeshComponentAdapter.cpp
+++ b/Engine/Source/Developer/MeshMergeUtilities/Private/StaticMeshComponentAdapter.cpp
@@ -47,7 +47,7 @@ void FStaticMeshComponentAdapter::ApplySettings(int32 LODIndex, FMeshData& InOut
 		if (MeshMapBuildData)
 		{
 			InOutMeshData.LightMap = MeshMapBuildData->LightMap;
-			InOutMeshData.LightMapIndex = StaticMeshComponent->GetStaticMesh()->LightMapCoordinateIndex;
+			InOutMeshData.LightMapIndex = StaticMeshComponent->GetStaticMesh()->GetLightMapCoordinateIndex();
 			InOutMeshData.LightmapResourceCluster = MeshMapBuildData->ResourceCluster;
 		}
 	}
@@ -104,7 +104,7 @@ bool FStaticMeshComponentAdapter::IsAsset() const
 
 int32 FStaticMeshComponentAdapter::LightmapUVIndex() const
 {
-	return StaticMesh->LightMapCoordinateIndex;
+	return StaticMesh->GetLightMapCoordinateIndex();
 }
 
 FBoxSphereBounds FStaticMeshComponentAdapter::GetBounds() const
diff --git a/Engine/Source/Developer/MeshUtilities/Private/MeshUtilities.cpp b/Engine/Source/Developer/MeshUtilities/Private/MeshUtilities.cpp
--- a/Engine/Source/Developer/MeshUtilities/Private/MeshUtilities.cpp
+++ b/Engine/Source/Developer/MeshUtilities/Private/MeshUtilities.cpp
@@ -572,7 +572,7 @@ UStaticMesh* FMeshUtilities::ConvertMeshesToStaticMesh(const TArray<UMeshCompone
 			StaticMesh->ImportVersion = EImportStaticMeshVersion::LastVersion;
 
 			// Set light map coordinate index to match DstLightmapIndex
-			StaticMesh->LightMapCoordinateIndex = LightMapIndex;
+			StaticMesh->SetLightMapCoordinateIndex(LightMapIndex);
 
 			// setup section info map
 			for (int32 RawMeshLODIndex = 0; RawMeshLODIndex < RawMeshes.Num(); RawMeshLODIndex++)
diff --git a/Engine/Source/Editor/DetailCustomizations/Private/StaticMeshComponentDetails.cpp b/Engine/Source/Editor/DetailCustomizations/Private/StaticMeshComponentDetails.cpp
--- a/Engine/Source/Editor/DetailCustomizations/Private/StaticMeshComponentDetails.cpp
+++ b/Engine/Source/Editor/DetailCustomizations/Private/StaticMeshComponentDetails.cpp
@@ -105,7 +105,7 @@ TOptional<int32> FStaticMeshComponentDetails::GetStaticMeshLightResValue() const
 		UStaticMeshComponent*  Component = Cast<UStaticMeshComponent>(ObjectsCustomized[i].Get());
 		if (Component != nullptr && Component->GetStaticMesh() != nullptr)
 		{
-			DefaultResolution = Component->GetStaticMesh()->LightMapResolution;
+			DefaultResolution = Component->GetStaticMesh()->GetLightMapResolution();
 		}
 	}
 	return DefaultResolution;
@@ -147,7 +147,7 @@ int FStaticMeshComponentDetails::HandleNoticeSwitcherWidgetIndex() const
 					AccessResultOverrideEnabled = FPropertyAccess::MultipleValues;
 				}
 
-				int32 CurrentStaticLightRes = Component->GetStaticMesh()->LightMapResolution;
+				int32 CurrentStaticLightRes = Component->GetStaticMesh()->GetLightMapResolution();
 				if (OverrideLightRes == -1)
 				{
 					OverrideLightRes = CurrentStaticLightRes;
diff --git a/Engine/Source/Editor/LevelEditor/Private/LightmapResRatioAdjust.cpp b/Engine/Source/Editor/LevelEditor/Private/LightmapResRatioAdjust.cpp
--- a/Engine/Source/Editor/LevelEditor/Private/LightmapResRatioAdjust.cpp
+++ b/Engine/Source/Editor/LevelEditor/Private/LightmapResRatioAdjust.cpp
@@ -158,7 +158,7 @@ bool FLightmapResRatioAdjustSettings::ApplyRatioAdjustment()
 				int32 CurrentResolution = SMComp->GetStaticLightMapResolution();
 				bool bConvertIt = true;
 				if (((SMComp->bOverrideLightMapRes == true) && (CurrentResolution == 0)) ||
-					((SMComp->bOverrideLightMapRes == false) && (SMComp->GetStaticMesh() != nullptr) && (SMComp->GetStaticMesh()->LightMapResolution == 0)))
+					((SMComp->bOverrideLightMapRes == false) && (SMComp->GetStaticMesh() != nullptr) && (SMComp->GetStaticMesh()->GetLightMapResolution() == 0)))
 				{
 					// Don't convert vertex mapped objects!
 					bConvertIt = false;
diff --git a/Engine/Source/Editor/StaticMeshEditor/Private/StaticMeshEditorTools.cpp b/Engine/Source/Editor/StaticMeshEditor/Private/StaticMeshEditorTools.cpp
--- a/Engine/Source/Editor/StaticMeshEditor/Private/StaticMeshEditorTools.cpp
+++ b/Engine/Source/Editor/StaticMeshEditor/Private/StaticMeshEditorTools.cpp
@@ -92,8 +92,8 @@ void FStaticMeshDetails::CustomizeDetails( class IDetailLayoutBuilder& DetailBui
 	IDetailCategoryBuilder& CollisionCategory = DetailBuilder.EditCategory( "Collision", LOCTEXT("CollisionCategory", "Collision") );
 	IDetailCategoryBuilder& ImportSettingsCategory = DetailBuilder.EditCategory("ImportSettings");
 
-	TSharedRef<IPropertyHandle> LightMapCoordinateIndexProperty = DetailBuilder.GetProperty(GET_MEMBER_NAME_CHECKED(UStaticMesh, LightMapCoordinateIndex));
-	TSharedRef<IPropertyHandle> LightMapResolutionProperty = DetailBuilder.GetProperty(GET_MEMBER_NAME_CHECKED(UStaticMesh, LightMapResolution));
+	TSharedRef<IPropertyHandle> LightMapCoordinateIndexProperty = DetailBuilder.GetProperty(UStaticMesh::GetLightMapCoordinateIndexName());
+	TSharedRef<IPropertyHandle> LightMapResolutionProperty = DetailBuilder.GetProperty(UStaticMesh::GetLightMapResolutionName());
 	LightMapCoordinateIndexProperty->SetOnPropertyValueChanged(FSimpleDelegate::CreateSP(this, &FStaticMeshDetails::OnLightmapSettingsChanged));
 	LightMapResolutionProperty->SetOnPropertyValueChanged(FSimpleDelegate::CreateSP(this, &FStaticMeshDetails::OnLightmapSettingsChanged));
 
@@ -3282,7 +3282,7 @@ void FLevelOfDetailSettingsLayout::AddToDetailsPanel( IDetailLayoutBuilder& Deta
 		.Text(LOCTEXT("MinLOD", "Minimum LOD"))
 	]
 	.ValueContent()
-	.MinDesiredWidth((float)(StaticMesh->MinLOD.PerPlatform.Num() + 1)*125.0f)
+	.MinDesiredWidth((float)(StaticMesh->GetMinLOD().PerPlatform.Num() + 1)*125.0f)
 	.MaxDesiredWidth((float)(PlatformNumber + 1)*125.0f)
 	[
 		SNew(SPerPlatformPropertiesWidget)
@@ -4193,8 +4193,8 @@ int32 FLevelOfDetailSettingsLayout::GetMinLOD(FName Platform) const
 	UStaticMesh* StaticMesh = StaticMeshEditor.GetStaticMesh();
 	check(StaticMesh);
 
-	int32* ValuePtr = (Platform == NAME_None) ? nullptr : StaticMesh->MinLOD.PerPlatform.Find(Platform);
-	return (ValuePtr != nullptr) ? *ValuePtr : StaticMesh->MinLOD.Default;
+	const int32* ValuePtr = (Platform == NAME_None) ? nullptr : StaticMesh->GetMinLOD().PerPlatform.Find(Platform);
+	return (ValuePtr != nullptr) ? *ValuePtr : StaticMesh->GetMinLOD().Default;
 }
 
 void FLevelOfDetailSettingsLayout::OnMinLODChanged(int32 NewValue, FName Platform)
@@ -4205,18 +4205,20 @@ void FLevelOfDetailSettingsLayout::OnMinLODChanged(int32 NewValue, FName Platfor
 	{
 		FStaticMeshComponentRecreateRenderStateContext ReregisterContext(StaticMesh, false);
 		NewValue = FMath::Clamp<int32>(NewValue, 0, MAX_STATIC_MESH_LODS - 1);
+		FPerPlatformInt MinLOD = StaticMesh->GetMinLOD();
 		if (Platform == NAME_None)
 		{
-			StaticMesh->MinLOD.Default = NewValue;
+			MinLOD.Default = NewValue;
 		}
 		else
 		{
-			int32* ValuePtr = StaticMesh->MinLOD.PerPlatform.Find(Platform);
+			int32* ValuePtr = MinLOD.PerPlatform.Find(Platform);
 			if (ValuePtr != nullptr)
 			{
 				*ValuePtr = NewValue;
 			}
 		}
+		StaticMesh->SetMinLOD(MoveTemp(MinLOD));
 		StaticMesh->Modify();
 	}
 	StaticMeshEditor.RefreshViewport();
@@ -4251,10 +4253,12 @@ bool FLevelOfDetailSettingsLayout::AddMinLODPlatformOverride(FName PlatformGroup
 	UStaticMesh* StaticMesh = StaticMeshEditor.GetStaticMesh();
 	check(StaticMesh);
 	StaticMesh->Modify();
-	if (StaticMesh->MinLOD.PerPlatform.Find(PlatformGroupName) == nullptr)
+	if (StaticMesh->GetMinLOD().PerPlatform.Find(PlatformGroupName) == nullptr)
 	{
-		float Value = StaticMesh->MinLOD.Default;
-		StaticMesh->MinLOD.PerPlatform.Add(PlatformGroupName, Value);
+		FPerPlatformInt MinLOD = StaticMesh->GetMinLOD();
+		float Value = MinLOD.Default;
+		MinLOD.PerPlatform.Add(PlatformGroupName, Value);
+		StaticMesh->SetMinLOD(MoveTemp(MinLOD));
 		OnMinLODChanged(Value, PlatformGroupName);
 		return true;
 	}
@@ -4267,9 +4271,12 @@ bool FLevelOfDetailSettingsLayout::RemoveMinLODPlatformOverride(FName PlatformGr
 	UStaticMesh* StaticMesh = StaticMeshEditor.GetStaticMesh();
 	check(StaticMesh);
 	StaticMesh->Modify();
-	if (StaticMesh->MinLOD.PerPlatform.Remove(PlatformGroupName) != 0)
+	FPerPlatformInt MinLOD = StaticMesh->GetMinLOD();
+	if (MinLOD.PerPlatform.Remove(PlatformGroupName) != 0)
 	{
-		OnMinLODChanged(StaticMesh->MinLOD.Default, PlatformGroupName);
+		float Value = MinLOD.Default;
+		StaticMesh->SetMinLOD(MoveTemp(MinLOD));
+		OnMinLODChanged(Value, PlatformGroupName);
 		return true;
 	}
 	return false;
@@ -4280,7 +4287,7 @@ TArray<FName> FLevelOfDetailSettingsLayout::GetMinLODPlatformOverrideNames() con
 	UStaticMesh* StaticMesh = StaticMeshEditor.GetStaticMesh();
 	check(StaticMesh);
 	TArray<FName> KeyArray;
-	StaticMesh->MinLOD.PerPlatform.GenerateKeyArray(KeyArray);
+	StaticMesh->GetMinLOD().PerPlatform.GenerateKeyArray(KeyArray);
 	KeyArray.Sort(FNameLexicalLess());
 	return KeyArray;
 }
diff --git a/Engine/Source/Editor/StaticMeshEditor/Private/StaticMeshEditorViewportClient.cpp b/Engine/Source/Editor/StaticMeshEditor/Private/StaticMeshEditorViewportClient.cpp
--- a/Engine/Source/Editor/StaticMeshEditor/Private/StaticMeshEditorViewportClient.cpp
+++ b/Engine/Source/Editor/StaticMeshEditor/Private/StaticMeshEditorViewportClient.cpp
@@ -785,12 +785,12 @@ void FStaticMeshEditorViewportClient::DrawCanvas( FViewport& InViewport, FSceneV
 	{
 		int32 LOD = StaticMeshEditor->GetCurrentLODLevel();
 		return (LOD == 0) ?
-			ComputeStaticMeshLOD(StaticMesh->GetRenderData(), StaticMeshComponent->Bounds.Origin, StaticMeshComponent->Bounds.SphereRadius, View, StaticMesh->MinLOD.Default)
+			ComputeStaticMeshLOD(StaticMesh->GetRenderData(), StaticMeshComponent->Bounds.Origin, StaticMeshComponent->Bounds.SphereRadius, View, StaticMesh->GetMinLOD().Default)
 			:
 			LOD - 1;
 	}();
 
-	const int32 CurrentMinLODLevel = StaticMesh->MinLOD.GetValue();
+	const int32 CurrentMinLODLevel = StaticMesh->GetMinLOD().GetValue();
 	const bool bBelowMinLOD = CurrentLODLevel < CurrentMinLODLevel;
 	TextItems.Add(SStaticMeshEditorViewport::FOverlayTextItem(
 		FText::Format(NSLOCTEXT("UnrealEd", "LOD_F", "LOD:  {0}"), FText::AsNumber(CurrentLODLevel)),
diff --git a/Engine/Source/Editor/UnrealEd/Private/Factories/EditorFactories.cpp b/Engine/Source/Editor/UnrealEd/Private/Factories/EditorFactories.cpp
--- a/Engine/Source/Editor/UnrealEd/Private/Factories/EditorFactories.cpp
+++ b/Engine/Source/Editor/UnrealEd/Private/Factories/EditorFactories.cpp
@@ -5866,8 +5866,8 @@ EReimportResult::Type UReimportFbxStaticMeshFactory::Reimport( UObject* Obj )
 			}
 
 			// preserve extended bound settings
-			const FVector PositiveBoundsExtension = Mesh->PositiveBoundsExtension;
-			const FVector NegativeBoundsExtension = Mesh->NegativeBoundsExtension;
+			const FVector PositiveBoundsExtension = Mesh->GetPositiveBoundsExtension();
+			const FVector NegativeBoundsExtension = Mesh->GetNegativeBoundsExtension();
 
 			if (FFbxImporter->ReimportStaticMesh(Mesh, ImportData))
 			{
@@ -5898,8 +5898,8 @@ EReimportResult::Type UReimportFbxStaticMeshFactory::Reimport( UObject* Obj )
 				}
 
 				// Restore bounds extension settings
-				Mesh->PositiveBoundsExtension = PositiveBoundsExtension;
-				Mesh->NegativeBoundsExtension = NegativeBoundsExtension;
+				Mesh->SetPositiveBoundsExtension(PositiveBoundsExtension);
+				Mesh->SetNegativeBoundsExtension(NegativeBoundsExtension);
 
 				Mesh->AssetImportData->Update(Filename);
 
diff --git a/Engine/Source/Editor/UnrealEd/Private/Fbx/FbxMainExport.cpp b/Engine/Source/Editor/UnrealEd/Private/Fbx/FbxMainExport.cpp
--- a/Engine/Source/Editor/UnrealEd/Private/Fbx/FbxMainExport.cpp
+++ b/Engine/Source/Editor/UnrealEd/Private/Fbx/FbxMainExport.cpp
@@ -4416,7 +4416,7 @@ FbxNode* FFbxExporter::ExportStaticMeshToFbx(const UStaticMesh* StaticMesh, int3
 
 			FString UVChannelNameBuilder = TEXT("UVmap_") + FString::FromInt(TexCoordSourceIndex);
 			const char* UVChannelName = TCHAR_TO_UTF8(*UVChannelNameBuilder); // actually UTF8 as required by Fbx, but can't use UE4's UTF8CHAR type because that's a uint8 aka *unsigned* char
-			if ((LightmapUVChannel >= 0) || ((LightmapUVChannel == -1) && (TexCoordSourceIndex == StaticMesh->LightMapCoordinateIndex)))
+			if ((LightmapUVChannel >= 0) || ((LightmapUVChannel == -1) && (TexCoordSourceIndex == StaticMesh->GetLightMapCoordinateIndex())))
 			{
 				UVChannelName = "LightMapUV";
 			}
@@ -4738,7 +4738,7 @@ void FFbxExporter::ExportSplineMeshToFbx(const USplineMeshComponent* SplineMeshC
 		}
 		FString UVChannelNameBuilder = TEXT("UVmap_") + FString::FromInt(TexCoordSourceIndex);
 		const char* UVChannelName = TCHAR_TO_UTF8(*UVChannelNameBuilder); // actually UTF8 as required by Fbx, but can't use UE4's UTF8CHAR type because that's a uint8 aka *unsigned* char
-		if (TexCoordSourceIndex == StaticMesh->LightMapCoordinateIndex)
+		if (TexCoordSourceIndex == StaticMesh->GetLightMapCoordinateIndex())
 		{
 			UVChannelName = "LightMapUV";
 		}
diff --git a/Engine/Source/Editor/UnrealEd/Private/Fbx/FbxStaticMeshImport.cpp b/Engine/Source/Editor/UnrealEd/Private/Fbx/FbxStaticMeshImport.cpp
--- a/Engine/Source/Editor/UnrealEd/Private/Fbx/FbxStaticMeshImport.cpp
+++ b/Engine/Source/Editor/UnrealEd/Private/Fbx/FbxStaticMeshImport.cpp
@@ -393,7 +393,7 @@ bool UnFbx::FFbxImporter::BuildStaticMeshFromGeometry(FbxNode* Node, UStaticMesh
 	int32 FBXNamedLightMapCoordinateIndex = FBXUVs.FindLightUVIndex();
 	if (FBXNamedLightMapCoordinateIndex != INDEX_NONE)
 	{
-		StaticMesh->LightMapCoordinateIndex = FBXNamedLightMapCoordinateIndex;
+		StaticMesh->SetLightMapCoordinateIndex(FBXNamedLightMapCoordinateIndex);
 	}
 	
 	//
@@ -1765,8 +1765,8 @@ UStaticMesh* UnFbx::FFbxImporter::ImportStaticMeshAsSingle(UObject* InParent, TA
 	StaticMesh->SetLightingGuid();
 
 	// Set it to use textured lightmaps. Note that Build Lighting will do the error-checking (texcoordindex exists for all LODs, etc).
-	StaticMesh->LightMapResolution = 64;
-	StaticMesh->LightMapCoordinateIndex = 1;
+	StaticMesh->SetLightMapResolution(64);
+	StaticMesh->SetLightMapCoordinateIndex(1);
 
 	float SqrBoundingBoxThreshold = THRESH_POINTS_ARE_NEAR * THRESH_POINTS_ARE_NEAR;
 
@@ -1928,7 +1928,7 @@ UStaticMesh* UnFbx::FFbxImporter::ImportStaticMeshAsSingle(UObject* InParent, TA
 			{
 				StaticMesh->GetSourceModel(ModelLODIndex).ReductionSettings = LODGroup.GetDefaultSettings(ModelLODIndex);
 			}
-			StaticMesh->LightMapResolution = LODGroup.GetDefaultLightMapResolution();
+			StaticMesh->SetLightMapResolution(LODGroup.GetDefaultLightMapResolution());
 
 			//@third party BEGIN SIMPLYGON
 			/* ImportData->Update(UFactory::GetCurrentFilename());
@@ -1963,7 +1963,7 @@ UStaticMesh* UnFbx::FFbxImporter::ImportStaticMeshAsSingle(UObject* InParent, TA
 		{
 			SrcModel.BuildSettings.bGenerateLightmapUVs = true;
 			SrcModel.BuildSettings.DstLightmapIndex = FirstOpenUVChannel;
-			StaticMesh->LightMapCoordinateIndex = FirstOpenUVChannel;
+			StaticMesh->SetLightMapCoordinateIndex(FirstOpenUVChannel);
 		}
 		else
 		{
@@ -2227,7 +2227,7 @@ void UnFbx::FFbxImporter::PostImportStaticMesh(UStaticMesh* StaticMesh, TArray<F
 		//Set the minimum LOD
 		if (ImportOptions->MinimumLodNumber > 0)
 		{
-			StaticMesh->MinLOD = ImportOptions->MinimumLodNumber;
+			StaticMesh->SetMinLOD(ImportOptions->MinimumLodNumber);
 		}
 
 		//User specify a number of LOD.
diff --git a/Engine/Source/Editor/UnrealEd/Private/Fbx/ReimportFbxSceneFactory.cpp b/Engine/Source/Editor/UnrealEd/Private/Fbx/ReimportFbxSceneFactory.cpp
--- a/Engine/Source/Editor/UnrealEd/Private/Fbx/ReimportFbxSceneFactory.cpp
+++ b/Engine/Source/Editor/UnrealEd/Private/Fbx/ReimportFbxSceneFactory.cpp
@@ -1515,8 +1515,8 @@ EReimportResult::Type UReimportFbxSceneFactory::ReimportStaticMesh(void* VoidFbx
 		nullptr;
 
 	// preserve extended bound settings
-	const FVector PositiveBoundsExtension = Mesh->PositiveBoundsExtension;
-	const FVector NegativeBoundsExtension = Mesh->NegativeBoundsExtension;
+	const FVector PositiveBoundsExtension = Mesh->GetPositiveBoundsExtension();
+	const FVector NegativeBoundsExtension = Mesh->GetNegativeBoundsExtension();
 	uint64 NodeInfoUid = INVALID_UNIQUE_ID;
 	if (GlobalImportSettings->bBakePivotInVertex && MeshInfo->PivotNodeUid != INVALID_UNIQUE_ID)
 	{
@@ -1550,8 +1550,8 @@ EReimportResult::Type UReimportFbxSceneFactory::ReimportStaticMesh(void* VoidFbx
 		}
 
 		// Restore bounds extension settings
-		Mesh->PositiveBoundsExtension = PositiveBoundsExtension;
-		Mesh->NegativeBoundsExtension = NegativeBoundsExtension;
+		Mesh->SetPositiveBoundsExtension(PositiveBoundsExtension);
+		Mesh->SetNegativeBoundsExtension(NegativeBoundsExtension);
 
 		Mesh->AssetImportData->Update(FbxImportFileName);
 
diff --git a/Engine/Source/Editor/UnrealEd/Private/Lightmass/Lightmass.cpp b/Engine/Source/Editor/UnrealEd/Private/Lightmass/Lightmass.cpp
--- a/Engine/Source/Editor/UnrealEd/Private/Lightmass/Lightmass.cpp
+++ b/Engine/Source/Editor/UnrealEd/Private/Lightmass/Lightmass.cpp
@@ -1320,12 +1320,12 @@ void FLightmassExporter::WriteStaticMeshes()
 		FString NewChannelName = Lightmass::CreateChannelName(BaseMeshData.Guid, Lightmass::LM_STATICMESH_VERSION, Lightmass::LM_STATICMESH_EXTENSION);
 
 		// Warn the user if there is an invalid lightmap UV channel specified.
-		if( StaticMesh->LightMapCoordinateIndex > 0
+		if( StaticMesh->GetLightMapCoordinateIndex() > 0
 			&& StaticMesh->GetRenderData()
 			&& StaticMesh->GetRenderData()->LODResources.Num() > 0 )
 		{
 			const FStaticMeshLODResources& RenderData = StaticMesh->GetRenderData()->LODResources[0];
-			if( StaticMesh->LightMapCoordinateIndex >= (int32)RenderData.VertexBuffers.StaticMeshVertexBuffer.GetNumTexCoords() )
+			if( StaticMesh->GetLightMapCoordinateIndex() >= (int32)RenderData.VertexBuffers.StaticMeshVertexBuffer.GetNumTexCoords() )
 			{
 				FMessageLog("LightingResults").Warning()
 					->AddToken(FUObjectToken::Create(const_cast<UStaticMesh*>(StaticMesh)))
@@ -1345,7 +1345,7 @@ void FLightmassExporter::WriteStaticMeshes()
 				Swarm.WriteChannel(Channel, &BaseMeshData, sizeof(BaseMeshData));
 
 				Lightmass::FStaticMeshData StaticMeshData;
-				StaticMeshData.LightmapCoordinateIndex = StaticMesh->LightMapCoordinateIndex;
+				StaticMeshData.LightmapCoordinateIndex = StaticMesh->GetLightMapCoordinateIndex();
 				StaticMeshData.NumLODs = StaticMesh->GetRenderData()->LODResources.Num();
 				Swarm.WriteChannel( Channel, &StaticMeshData, sizeof(StaticMeshData) );
 				for (int32 LODIndex = 0; LODIndex < StaticMesh->GetRenderData()->LODResources.Num(); LODIndex++)
diff --git a/Engine/Source/Editor/UnrealEd/Private/StaticLightingSystem/StaticLightingDebug.cpp b/Engine/Source/Editor/UnrealEd/Private/StaticLightingSystem/StaticLightingDebug.cpp
--- a/Engine/Source/Editor/UnrealEd/Private/StaticLightingSystem/StaticLightingDebug.cpp
+++ b/Engine/Source/Editor/UnrealEd/Private/StaticLightingSystem/StaticLightingDebug.cpp
@@ -201,8 +201,8 @@ void SetDebugLightmapSample(TArray<UActorComponent*>* Components, UModel* Model,
 				SMComponent->GetLightMapResolution(LightmapSizeX, LightmapSizeY);
 
 				if (LightmapSizeX > 0 && LightmapSizeY > 0 
-					&& StaticMesh->LightMapCoordinateIndex >= 0 
-					&& (uint32)StaticMesh->LightMapCoordinateIndex < LODModel.VertexBuffers.StaticMeshVertexBuffer.GetNumTexCoords()
+					&& StaticMesh->GetLightMapCoordinateIndex() >= 0 
+					&& (uint32)StaticMesh->GetLightMapCoordinateIndex() < LODModel.VertexBuffers.StaticMeshVertexBuffer.GetNumTexCoords()
 					)
 				{
 					bUseTextureMap = true;
@@ -251,9 +251,9 @@ void SetDebugLightmapSample(TArray<UActorComponent*>* Components, UModel* Model,
 						uint32 Index2 = Indices[ClosestPlaneTriangleIndex + 2];
 
 						// Fetch lightmap UV's
-						FVector2D LightmapUV0 = LODModel.VertexBuffers.StaticMeshVertexBuffer.GetVertexUV(Index0, StaticMesh->LightMapCoordinateIndex);
-						FVector2D LightmapUV1 = LODModel.VertexBuffers.StaticMeshVertexBuffer.GetVertexUV(Index1, StaticMesh->LightMapCoordinateIndex);
-						FVector2D LightmapUV2 = LODModel.VertexBuffers.StaticMeshVertexBuffer.GetVertexUV(Index2, StaticMesh->LightMapCoordinateIndex);
+						FVector2D LightmapUV0 = LODModel.VertexBuffers.StaticMeshVertexBuffer.GetVertexUV(Index0, StaticMesh->GetLightMapCoordinateIndex());
+						FVector2D LightmapUV1 = LODModel.VertexBuffers.StaticMeshVertexBuffer.GetVertexUV(Index1, StaticMesh->GetLightMapCoordinateIndex());
+						FVector2D LightmapUV2 = LODModel.VertexBuffers.StaticMeshVertexBuffer.GetVertexUV(Index2, StaticMesh->GetLightMapCoordinateIndex());
 						// Interpolate lightmap UV's to the click location
 						FVector2D InterpolatedUV = LightmapUV0 * ClosestPlaneBaryCentricWeights.X + LightmapUV1 * ClosestPlaneBaryCentricWeights.Y + LightmapUV2 * ClosestPlaneBaryCentricWeights.Z;
 
@@ -292,9 +292,9 @@ void SetDebugLightmapSample(TArray<UActorComponent*>* Components, UModel* Model,
 					uint32 Index1 = Indices[SelectedTriangle + 1];
 					uint32 Index2 = Indices[SelectedTriangle + 2];
 
-					FVector2D LightmapUV0 = LODModel.VertexBuffers.StaticMeshVertexBuffer.GetVertexUV(Index0, StaticMesh->LightMapCoordinateIndex);
-					FVector2D LightmapUV1 = LODModel.VertexBuffers.StaticMeshVertexBuffer.GetVertexUV(Index1, StaticMesh->LightMapCoordinateIndex);
-					FVector2D LightmapUV2 = LODModel.VertexBuffers.StaticMeshVertexBuffer.GetVertexUV(Index2, StaticMesh->LightMapCoordinateIndex);
+					FVector2D LightmapUV0 = LODModel.VertexBuffers.StaticMeshVertexBuffer.GetVertexUV(Index0, StaticMesh->GetLightMapCoordinateIndex());
+					FVector2D LightmapUV1 = LODModel.VertexBuffers.StaticMeshVertexBuffer.GetVertexUV(Index1, StaticMesh->GetLightMapCoordinateIndex());
+					FVector2D LightmapUV2 = LODModel.VertexBuffers.StaticMeshVertexBuffer.GetVertexUV(Index2, StaticMesh->GetLightMapCoordinateIndex());
 
 					FVector BaryCentricWeights;
 					BaryCentricWeights.X = FMath::FRandRange(0, 1);
diff --git a/Engine/Source/Editor/UnrealEd/Private/StaticMeshEdit.cpp b/Engine/Source/Editor/UnrealEd/Private/StaticMeshEdit.cpp
--- a/Engine/Source/Editor/UnrealEd/Private/StaticMeshEdit.cpp
+++ b/Engine/Source/Editor/UnrealEd/Private/StaticMeshEdit.cpp
@@ -1178,8 +1178,8 @@ TSharedPtr<FExistingStaticMeshData> SaveExistingStaticMeshData(UStaticMesh* Exis
 		ExistingMeshDataPtr->ExistingCustomizedCollision = ExistingMesh->bCustomizedCollision;
 		ExistingMeshDataPtr->bAutoComputeLODScreenSize = ExistingMesh->bAutoComputeLODScreenSize;
 
-		ExistingMeshDataPtr->ExistingLightMapResolution = ExistingMesh->LightMapResolution;
-		ExistingMeshDataPtr->ExistingLightMapCoordinateIndex = ExistingMesh->LightMapCoordinateIndex;
+		ExistingMeshDataPtr->ExistingLightMapResolution = ExistingMesh->GetLightMapResolution();
+		ExistingMeshDataPtr->ExistingLightMapCoordinateIndex = ExistingMesh->GetLightMapCoordinateIndex();
 
 		ExistingMeshDataPtr->ExistingImportData = ExistingMesh->AssetImportData;
 		ExistingMeshDataPtr->ExistingThumbnailInfo = ExistingMesh->ThumbnailInfo;
@@ -1188,15 +1188,15 @@ TSharedPtr<FExistingStaticMeshData> SaveExistingStaticMeshData(UStaticMesh* Exis
 
 		ExistingMeshDataPtr->bHasNavigationData = ExistingMesh->bHasNavigationData;
 		ExistingMeshDataPtr->LODGroup = ExistingMesh->LODGroup;
-		ExistingMeshDataPtr->MinLOD = ExistingMesh->MinLOD;
+		ExistingMeshDataPtr->MinLOD = ExistingMesh->GetMinLOD();
 
 		ExistingMeshDataPtr->ExistingGenerateMeshDistanceField = ExistingMesh->bGenerateMeshDistanceField;
 		ExistingMeshDataPtr->ExistingLODForCollision = ExistingMesh->LODForCollision;
 		ExistingMeshDataPtr->ExistingDistanceFieldSelfShadowBias = ExistingMesh->DistanceFieldSelfShadowBias;
 		ExistingMeshDataPtr->ExistingSupportUniformlyDistributedSampling = ExistingMesh->bSupportUniformlyDistributedSampling;
 		ExistingMeshDataPtr->ExistingAllowCpuAccess = ExistingMesh->bAllowCPUAccess;
-		ExistingMeshDataPtr->ExistingPositiveBoundsExtension = ExistingMesh->PositiveBoundsExtension;
-		ExistingMeshDataPtr->ExistingNegativeBoundsExtension = ExistingMesh->NegativeBoundsExtension;
+		ExistingMeshDataPtr->ExistingPositiveBoundsExtension = ExistingMesh->GetPositiveBoundsExtension();
+		ExistingMeshDataPtr->ExistingNegativeBoundsExtension = ExistingMesh->GetNegativeBoundsExtension();
 
 		UFbxStaticMeshImportData* ImportData = Cast<UFbxStaticMeshImportData>(ExistingMesh->AssetImportData);
 		if (ImportData && ExistingMeshDataPtr->UseMaterialNameSlotWorkflow)
@@ -1248,7 +1248,7 @@ void RestoreExistingMeshSettings(const FExistingStaticMeshData* ExistingMesh, US
 		return;
 	}
 	NewMesh->LODGroup = ExistingMesh->LODGroup;
-	NewMesh->MinLOD = ExistingMesh->MinLOD;
+	NewMesh->SetMinLOD(ExistingMesh->MinLOD);
 	int32 ExistingNumLods = ExistingMesh->ExistingLODData.Num();
 	int32 CurrentNumLods = NewMesh->GetNumSourceModels();
 	if (LODIndex == INDEX_NONE)
@@ -1659,8 +1659,8 @@ void RestoreExistingMeshData(TSharedPtr<const FExistingStaticMeshData> ExistingM
 	NewMesh->bCustomizedCollision = ExistingMeshDataPtr->ExistingCustomizedCollision;
 	NewMesh->bAutoComputeLODScreenSize = ExistingMeshDataPtr->bAutoComputeLODScreenSize;
 
-	NewMesh->LightMapResolution = ExistingMeshDataPtr->ExistingLightMapResolution;
-	NewMesh->LightMapCoordinateIndex = ExistingMeshDataPtr->ExistingLightMapCoordinateIndex;
+	NewMesh->SetLightMapResolution(ExistingMeshDataPtr->ExistingLightMapResolution);
+	NewMesh->SetLightMapCoordinateIndex(ExistingMeshDataPtr->ExistingLightMapCoordinateIndex);
 
 	if (ExistingMeshDataPtr->ExistingImportData.IsValid())
 	{
@@ -1727,8 +1727,8 @@ void RestoreExistingMeshData(TSharedPtr<const FExistingStaticMeshData> ExistingM
 	NewMesh->DistanceFieldSelfShadowBias = ExistingMeshDataPtr->ExistingDistanceFieldSelfShadowBias;
 	NewMesh->bSupportUniformlyDistributedSampling = ExistingMeshDataPtr->ExistingSupportUniformlyDistributedSampling;
 	NewMesh->bAllowCPUAccess = ExistingMeshDataPtr->ExistingAllowCpuAccess;
-	NewMesh->PositiveBoundsExtension = ExistingMeshDataPtr->ExistingPositiveBoundsExtension;
-	NewMesh->NegativeBoundsExtension = ExistingMeshDataPtr->ExistingNegativeBoundsExtension;
+	NewMesh->SetPositiveBoundsExtension(ExistingMeshDataPtr->ExistingPositiveBoundsExtension);
+	NewMesh->SetNegativeBoundsExtension(ExistingMeshDataPtr->ExistingNegativeBoundsExtension);
 
 	NewMesh->ComplexCollisionMesh = ExistingMeshDataPtr->ExistingComplexCollisionMesh;
 }
diff --git a/Engine/Source/Editor/WorldBrowser/Private/Tiles/WorldTileCollectionModel.cpp b/Engine/Source/Editor/WorldBrowser/Private/Tiles/WorldTileCollectionModel.cpp
--- a/Engine/Source/Editor/WorldBrowser/Private/Tiles/WorldTileCollectionModel.cpp
+++ b/Engine/Source/Editor/WorldBrowser/Private/Tiles/WorldTileCollectionModel.cpp
@@ -2149,8 +2149,8 @@ bool FWorldTileCollectionModel::GenerateLODLevels(FLevelModelList InLevelList, i
 				StaticMesh->SetLightingGuid();
 
 				// Set it to use textured lightmaps. Note that Build Lighting will do the error-checking (texcoordindex exists for all LODs, etc).
-				StaticMesh->LightMapResolution = 64;
-				StaticMesh->LightMapCoordinateIndex = 1;
+				StaticMesh->SetLightMapResolution(64);
+				StaticMesh->SetLightMapCoordinateIndex(1);
 
 				FStaticMeshSourceModel& SrcModel = StaticMesh->AddSourceModel();
 				/*Don't allow the engine to recalculate normals*/
diff --git a/Engine/Source/Runtime/Engine/Classes/Engine/StaticMesh.h b/Engine/Source/Runtime/Engine/Classes/Engine/StaticMesh.h
--- a/Engine/Source/Runtime/Engine/Classes/Engine/StaticMesh.h
+++ b/Engine/Source/Runtime/Engine/Classes/Engine/StaticMesh.h
@@ -619,14 +619,29 @@ public:
 #endif // #if WITH_EDITORONLY_DATA
 
 	/** Minimum LOD to use for rendering.  This is the default setting for the mesh and can be overridden by component settings. */
+	UE_DEPRECATED(4.27, "Please do not access this member directly; use UStaticMesh::GetMinLOD() or UStaticMesh::SetMinLOD().")
 	UPROPERTY()
 	FPerPlatformInt MinLOD;
 
+	const FPerPlatformInt& GetMinLOD() const
+	{
+		PRAGMA_DISABLE_DEPRECATION_WARNINGS
+		return MinLOD;
+		PRAGMA_ENABLE_DEPRECATION_WARNINGS
+	}
+
+	void SetMinLOD(FPerPlatformInt InMinLOD)
+	{
+		PRAGMA_DISABLE_DEPRECATION_WARNINGS
+		MinLOD = MoveTemp(InMinLOD);
+		PRAGMA_ENABLE_DEPRECATION_WARNINGS
+	}
+
 	UFUNCTION(BlueprintPure, Category=StaticMesh)
 	void GetMinimumLODForPlatforms(TMap<FName, int32>& PlatformMinimumLODs) const
 	{
 #if WITH_EDITORONLY_DATA
-		PlatformMinimumLODs = MinLOD.PerPlatform;
+		PlatformMinimumLODs = GetMinLOD().PerPlatform;
 #endif
 	}
 
@@ -634,7 +649,7 @@ public:
 	int32 GetMinimumLODForPlatform(const FName& PlatformName) const
 	{
 #if WITH_EDITORONLY_DATA
-		if (const int32* Result = MinLOD.PerPlatform.Find(PlatformName))
+		if (const int32* Result = GetMinLOD().PerPlatform.Find(PlatformName))
 		{
 			return *Result;
 		}
@@ -694,13 +709,57 @@ public:
 		PRAGMA_ENABLE_DEPRECATION_WARNINGS
 	}
 
+	UE_DEPRECATED(4.27, "Please do not access this member directly; use UStaticMesh::GetLightMapResolution() or UStaticMesh::SetLightMapResolution().")
 	UPROPERTY(EditAnywhere, Category=StaticMesh, meta=(ClampMax = 4096, ToolTip="The light map resolution", FixedIncrement="4.0"))
 	int32 LightMapResolution;
 
+	int32 GetLightMapResolution() const
+	{
+		PRAGMA_DISABLE_DEPRECATION_WARNINGS
+		return LightMapResolution;
+		PRAGMA_ENABLE_DEPRECATION_WARNINGS
+	}
+
+	void SetLightMapResolution(int32 InLightMapResolution)
+	{
+		PRAGMA_DISABLE_DEPRECATION_WARNINGS
+		LightMapResolution = InLightMapResolution;
+		PRAGMA_ENABLE_DEPRECATION_WARNINGS
+	}
+
+	static FName GetLightMapResolutionName()
+	{
+		PRAGMA_DISABLE_DEPRECATION_WARNINGS
+		return GET_MEMBER_NAME_CHECKED(UStaticMesh, LightMapResolution);
+		PRAGMA_ENABLE_DEPRECATION_WARNINGS
+	}
+
 	/** The light map coordinate index */
+	UE_DEPRECATED(4.27, "Please do not access this member directly; use UStaticMesh::GetLightMapCoordinateIndex() or UStaticMesh::SetLightMapCoordinateIndex().")
 	UPROPERTY(EditAnywhere, AdvancedDisplay, Category=StaticMesh, meta=(ToolTip="The light map coordinate index", UIMin = "0", UIMax = "3"))
 	int32 LightMapCoordinateIndex;
 
+	int32 GetLightMapCoordinateIndex() const
+	{
+		PRAGMA_DISABLE_DEPRECATION_WARNINGS
+		return LightMapCoordinateIndex;
+		PRAGMA_ENABLE_DEPRECATION_WARNINGS
+	}
+
+	void SetLightMapCoordinateIndex(int32 InLightMapCoordinateIndex)
+	{
+		PRAGMA_DISABLE_DEPRECATION_WARNINGS
+		LightMapCoordinateIndex = InLightMapCoordinateIndex;
+		PRAGMA_ENABLE_DEPRECATION_WARNINGS
+	}
+
+	static FName GetLightMapCoordinateIndexName()
+	{
+		PRAGMA_DISABLE_DEPRECATION_WARNINGS
+		return GET_MEMBER_NAME_CHECKED(UStaticMesh, LightMapCoordinateIndex);
+		PRAGMA_ENABLE_DEPRECATION_WARNINGS
+	}
+
 	/** Useful for reducing self shadowing from distance field methods when using world position offset to animate the mesh's vertices. */
 	UPROPERTY(EditAnywhere, AdvancedDisplay, Category = StaticMesh)
 	float DistanceFieldSelfShadowBias;
@@ -724,6 +783,13 @@ public:
 		PRAGMA_ENABLE_DEPRECATION_WARNINGS
 	}
 
+	static FName GetBodySetupName()
+	{
+		PRAGMA_DISABLE_DEPRECATION_WARNINGS
+		return GET_MEMBER_NAME_CHECKED(UStaticMesh, BodySetup);
+		PRAGMA_ENABLE_DEPRECATION_WARNINGS
+	}
+
 	/** 
 	 *	Specifies which mesh LOD to use for complex (per-poly) collision. 
 	 *	Sometimes it can be desirable to use a lower poly representation for collision to reduce memory usage, improve performance and behaviour.
@@ -877,12 +943,58 @@ public:
 	TSharedPtr<class FSpeedTreeWind> SpeedTreeWind;
 
 	/** Bound extension values in the positive direction of XYZ, positive value increases bound size */
+	UE_DEPRECATED(4.27, "Please do not access this member directly; use UStaticMesh::GetPositiveBoundsExtension() or UStaticMesh::SetPositiveBoundsExtension.")
 	UPROPERTY(EditDefaultsOnly, AdvancedDisplay, Category = StaticMesh)
 	FVector PositiveBoundsExtension;
+
+	const FVector& GetPositiveBoundsExtension() const
+	{
+		PRAGMA_DISABLE_DEPRECATION_WARNINGS
+		return PositiveBoundsExtension;
+		PRAGMA_ENABLE_DEPRECATION_WARNINGS
+	}
+
+	void SetPositiveBoundsExtension(FVector InPositiveBoundsExtension)
+	{
+		PRAGMA_DISABLE_DEPRECATION_WARNINGS
+		PositiveBoundsExtension = InPositiveBoundsExtension;
+		PRAGMA_ENABLE_DEPRECATION_WARNINGS
+	}
+
+	static FName GetPositiveBoundsExtensionName()
+	{
+		PRAGMA_DISABLE_DEPRECATION_WARNINGS
+		return GET_MEMBER_NAME_CHECKED(UStaticMesh, PositiveBoundsExtension);
+		PRAGMA_ENABLE_DEPRECATION_WARNINGS
+	}
+
 	/** Bound extension values in the negative direction of XYZ, positive value increases bound size */
+	UE_DEPRECATED(4.27, "Please do not access this member directly; use UStaticMesh::GetNegativeBoundsExtension() or UStaticMesh::SetNegativeBoundsExtension.")
 	UPROPERTY(EditDefaultsOnly, AdvancedDisplay, Category = StaticMesh)
 	FVector NegativeBoundsExtension;
 	
+	const FVector& GetNegativeBoundsExtension() const
+	{
+		// No need for CheckAsyncPropertyAccess here as this is not modified during async Build/Postload
+		PRAGMA_DISABLE_DEPRECATION_WARNINGS
+		return NegativeBoundsExtension;
+		PRAGMA_ENABLE_DEPRECATION_WARNINGS
+	}
+
+	void SetNegativeBoundsExtension(FVector InNegativeBoundsExtension)
+	{
+		PRAGMA_DISABLE_DEPRECATION_WARNINGS
+		NegativeBoundsExtension = InNegativeBoundsExtension;
+		PRAGMA_ENABLE_DEPRECATION_WARNINGS
+	}
+
+	static FName GetNegativeBoundsExtensionName()
+	{
+		PRAGMA_DISABLE_DEPRECATION_WARNINGS
+		return GET_MEMBER_NAME_CHECKED(UStaticMesh, NegativeBoundsExtension);
+		PRAGMA_ENABLE_DEPRECATION_WARNINGS
+	}
+
 	/** Original mesh bounds extended with Positive/NegativeBoundsExtension */
 	UE_DEPRECATED(4.27, "Please do not access this member directly; use UStaticMesh::GetExtendedBounds() or UStaticMesh::SetExtendedBounds.")
 	UPROPERTY()
diff --git a/Engine/Source/Runtime/Engine/Private/Components/SplineMeshComponent.cpp b/Engine/Source/Runtime/Engine/Private/Components/SplineMeshComponent.cpp
--- a/Engine/Source/Runtime/Engine/Private/Components/SplineMeshComponent.cpp
+++ b/Engine/Source/Runtime/Engine/Private/Components/SplineMeshComponent.cpp
@@ -142,7 +142,7 @@ void FSplineMeshSceneProxy::InitVertexFactory(USplineMeshComponent* InComponent,
 		RenderData2->VertexBuffers.PositionVertexBuffer.BindPositionVertexBuffer(VertexFactory, Data);
 		RenderData2->VertexBuffers.StaticMeshVertexBuffer.BindTangentVertexBuffer(VertexFactory, Data);
 		RenderData2->VertexBuffers.StaticMeshVertexBuffer.BindPackedTexCoordVertexBuffer(VertexFactory, Data);
-		RenderData2->VertexBuffers.StaticMeshVertexBuffer.BindLightMapVertexBuffer(VertexFactory, Data, Parent->LightMapCoordinateIndex);
+		RenderData2->VertexBuffers.StaticMeshVertexBuffer.BindLightMapVertexBuffer(VertexFactory, Data, Parent->GetLightMapCoordinateIndex());
 		if (bOverrideColorVertexBuffer)
 		{
 			FColorVertexBuffer::BindDefaultColorVertexBuffer(VertexFactory, Data, FColorVertexBuffer::NullBindStride::FColorSizeForComponentOverride);
diff --git a/Engine/Source/Runtime/Engine/Private/Components/StaticMeshComponent.cpp b/Engine/Source/Runtime/Engine/Private/Components/StaticMeshComponent.cpp
--- a/Engine/Source/Runtime/Engine/Private/Components/StaticMeshComponent.cpp
+++ b/Engine/Source/Runtime/Engine/Private/Components/StaticMeshComponent.cpp
@@ -1951,8 +1951,8 @@ bool UStaticMeshComponent::GetLightMapResolution( int32& Width, int32& Height )
 		// Use the lightmap resolution defined in the static mesh.
 		else
 		{
-			Width	= GetStaticMesh()->LightMapResolution;
-			Height	= GetStaticMesh()->LightMapResolution;
+			Width	= GetStaticMesh()->GetLightMapResolution();
+			Height	= GetStaticMesh()->GetLightMapResolution();
 		}
 		bPadded = true;
 	}
@@ -1991,8 +1991,8 @@ void UStaticMeshComponent::GetEstimatedLightMapResolution(int32& Width, int32& H
 		// Use the lightmap resolution defined in the static mesh.
 		if (bUseSourceMesh == true)
 		{
-			Width	= GetStaticMesh()->LightMapResolution;
-			Height	= GetStaticMesh()->LightMapResolution;
+			Width	= GetStaticMesh()->GetLightMapResolution();
+			Height	= GetStaticMesh()->GetLightMapResolution();
 		}
 
 		// If it was not set by anything, give it a default value...
@@ -2052,14 +2052,14 @@ bool UStaticMeshComponent::HasLightmapTextureCoordinates() const
 {
 	const UStaticMesh* Mesh = GetStaticMesh();
 	if (Mesh != nullptr &&
-		Mesh->LightMapCoordinateIndex >= 0 &&
+		Mesh->GetLightMapCoordinateIndex() >= 0 &&
 		Mesh->GetRenderData() != nullptr &&
 		Mesh->GetRenderData()->LODResources.Num())
 	{
-		int32 MeshMinLOD = Mesh->MinLOD.GetValue();
+		int32 MeshMinLOD = Mesh->GetMinLOD().GetValue();
 		MeshMinLOD = FMath::Min(MeshMinLOD,  Mesh->GetRenderData()->LODResources.Num() - 1);
 		
-		return ((uint32)Mesh->LightMapCoordinateIndex < Mesh->GetRenderData()->LODResources[MeshMinLOD].VertexBuffers.StaticMeshVertexBuffer.GetNumTexCoords());
+		return ((uint32)Mesh->GetLightMapCoordinateIndex() < Mesh->GetRenderData()->LODResources[MeshMinLOD].VertexBuffers.StaticMeshVertexBuffer.GetNumTexCoords());
 	}
 	return false;
 }
diff --git a/Engine/Source/Runtime/Engine/Private/HLODProxy.cpp b/Engine/Source/Runtime/Engine/Private/HLODProxy.cpp
--- a/Engine/Source/Runtime/Engine/Private/HLODProxy.cpp
+++ b/Engine/Source/Runtime/Engine/Private/HLODProxy.cpp
@@ -356,7 +356,9 @@ uint32 UHLODProxy::GetCRC(UStaticMesh* InStaticMesh, uint32 InCRC)
 	// Default to just the path name if we don't have render data
 	FString DerivedDataKey = InStaticMesh->GetRenderData() ? InStaticMesh->GetRenderData()->DerivedDataKey : InStaticMesh->GetPathName();
 	KeyBuffer.Append((uint8*)DerivedDataKey.GetCharArray().GetData(), DerivedDataKey.GetCharArray().Num() * DerivedDataKey.GetCharArray().GetTypeSize());
-	KeyBuffer.Append((uint8*)&InStaticMesh->LightMapCoordinateIndex, sizeof(int32));
+
+	const int32 LightMapCoordinateIndex = InStaticMesh->GetLightMapCoordinateIndex();
+	KeyBuffer.Append((uint8*)&LightMapCoordinateIndex, sizeof(int32));
 	if(InStaticMesh->GetBodySetup())
 	{
 		// Incorporate physics data
diff --git a/Engine/Source/Runtime/Engine/Private/InstancedStaticMesh.cpp b/Engine/Source/Runtime/Engine/Private/InstancedStaticMesh.cpp
--- a/Engine/Source/Runtime/Engine/Private/InstancedStaticMesh.cpp
+++ b/Engine/Source/Runtime/Engine/Private/InstancedStaticMesh.cpp
@@ -806,7 +806,7 @@ void FInstancedStaticMeshRenderData::InitVertexFactories()
 		VertexFactories.Add(new FInstancedStaticMeshVertexFactory(FeatureLevel));
 	}
 
-	const int32 LightMapCoordinateIndex = Component->GetStaticMesh()->LightMapCoordinateIndex;
+	const int32 LightMapCoordinateIndex = Component->GetStaticMesh()->GetLightMapCoordinateIndex();
 	ENQUEUE_RENDER_COMMAND(InstancedStaticMeshRenderData_InitVertexFactories)(
 		[this, LightMapCoordinateIndex](FRHICommandListImmediate& RHICmdList)
 	{
@@ -1866,7 +1866,7 @@ void UInstancedStaticMeshComponent::GetStaticLightingInfo(FStaticLightingPrimiti
 				auto* StaticLightingMesh = new FStaticLightingMesh_InstancedStaticMesh(this, LODIndex, InstanceIndex, InRelevantLights);
 				OutPrimitiveInfo.Meshes.Add(StaticLightingMesh);
 
-				auto* InstancedMapping = new FStaticLightingTextureMapping_InstancedStaticMesh(this, LODIndex, InstanceIndex, StaticLightingMesh, LightMapWidth, LightMapHeight, GetStaticMesh()->LightMapCoordinateIndex, true);
+				auto* InstancedMapping = new FStaticLightingTextureMapping_InstancedStaticMesh(this, LODIndex, InstanceIndex, StaticLightingMesh, LightMapWidth, LightMapHeight, GetStaticMesh()->GetLightMapCoordinateIndex(), true);
 				OutPrimitiveInfo.Mappings.Add(InstancedMapping);
 
 				CachedMappings[LODIndex * PerInstanceSMData.Num() + InstanceIndex].Mapping = InstancedMapping;
diff --git a/Engine/Source/Runtime/Engine/Private/Particles/ParticleSystemRender.cpp b/Engine/Source/Runtime/Engine/Private/Particles/ParticleSystemRender.cpp
--- a/Engine/Source/Runtime/Engine/Private/Particles/ParticleSystemRender.cpp
+++ b/Engine/Source/Runtime/Engine/Private/Particles/ParticleSystemRender.cpp
@@ -1571,7 +1571,7 @@ uint32 FDynamicMeshEmitterData::GetMeshLODIndexFromProxy(const FParticleSystemSc
 	{
 		return FirstAvailableLOD;
 	}
-	const int32 EffectiveMinLOD = StaticMesh->MinLOD.GetValue();
+	const int32 EffectiveMinLOD = StaticMesh->GetMinLOD().GetValue();
 	const int32 MaxLOD = StaticMesh->GetRenderData()->LODResources.Num() - 1;
 	const int32 ClampedMinLOD = FMath::Clamp(EffectiveMinLOD, FirstAvailableLOD, MaxLOD);
 	const int32 MeshLOD = (InOwnerProxy->MeshEmitterLODIndices.IsValidIndex(EmitterIndex)) ? InOwnerProxy->MeshEmitterLODIndices[EmitterIndex] : 0;
diff --git a/Engine/Source/Runtime/Engine/Private/StaticMesh.cpp b/Engine/Source/Runtime/Engine/Private/StaticMesh.cpp
--- a/Engine/Source/Runtime/Engine/Private/StaticMesh.cpp
+++ b/Engine/Source/Runtime/Engine/Private/StaticMesh.cpp
@@ -271,7 +271,7 @@ int32 FStaticMeshLODResources::GetPlatformMinLODIdx(const ITargetPlatform* Targe
 {
 #if WITH_EDITOR
 	check(TargetPlatform && StaticMesh);
-	return StaticMesh->MinLOD.GetValueForPlatform(*TargetPlatform->IniPlatformName());
+	return StaticMesh->GetMinLOD().GetValueForPlatform(*TargetPlatform->IniPlatformName());
 #else
 	return 0;
 #endif
@@ -838,7 +838,7 @@ void FStaticMeshVertexFactories::InitVertexFactory(
 		uint32 LODIndex;
 	} Params;
 
-	uint32 LightMapCoordinateIndex = (uint32)InParentMesh->LightMapCoordinateIndex;
+	uint32 LightMapCoordinateIndex = (uint32)InParentMesh->GetLightMapCoordinateIndex();
 	LightMapCoordinateIndex = LightMapCoordinateIndex < LodResources.VertexBuffers.StaticMeshVertexBuffer.GetNumTexCoords() ? LightMapCoordinateIndex : LodResources.VertexBuffers.StaticMeshVertexBuffer.GetNumTexCoords() - 1;
 
 	Params.VertexFactory = &InOutVertexFactory;
@@ -2786,8 +2786,8 @@ UStaticMesh::UStaticMesh(const FObjectInitializer& ObjectInitializer)
 	LODForOccluderMesh = -1;
 	NumStreamedLODs.Default = -1;
 #endif // #if WITH_EDITORONLY_DATA
-	LightMapResolution = 4;
-	MinLOD.Default = 0;
+	SetLightMapResolution(4);
+	SetMinLOD(0);
 
 	bSupportUniformlyDistributedSampling = false;
 	SetIsBuiltAtRuntime(false);
@@ -2866,7 +2866,7 @@ void UStaticMesh::InitResources()
 	{
 		{
 			const int32 NumLODs = GetNumLODs();
-			const int32 MinFirstLOD = MinLOD.GetValue();
+			const int32 MinFirstLOD = GetMinLOD().GetValue();
 
 			CachedSRRState.NumNonStreamingLODs = GetRenderData()->NumInlinedLODs;
 			CachedSRRState.NumNonOptionalLODs = GetRenderData()->GetNumNonOptionalLODs();
@@ -3002,7 +3002,7 @@ bool UStaticMesh::HasValidRenderData(bool bCheckLODForVerts, int32 LODIndex) con
 		{
 		    if (LODIndex == INDEX_NONE)
 		    {
-			    LODIndex = FMath::Clamp<int32>(MinLOD.GetValue(), 0, GetRenderData()->LODResources.Num() - 1);
+			    LODIndex = FMath::Clamp<int32>(GetMinLOD().GetValue(), 0, GetRenderData()->LODResources.Num() - 1);
 		    }
 			return (GetRenderData()->LODResources[LODIndex].VertexBuffers.StaticMeshVertexBuffer.GetNumVertices() > 0);
 		}
@@ -3114,7 +3114,7 @@ void UStaticMesh::UpdateUVChannelData(bool bRebuildAll)
 
 		if (bDensityChanged || bRebuildAll)
 		{
-			SetLightmapUVDensity(GetUVDensity(GetRenderData()->LODResources, LightMapCoordinateIndex));
+			SetLightmapUVDensity(GetUVDensity(GetRenderData()->LODResources, GetLightMapCoordinateIndex()));
 
 			if (GEngine)
 			{
@@ -3279,11 +3279,11 @@ void UStaticMesh::PostEditChangeProperty(FPropertyChangedEvent& PropertyChangedE
 	}
 #endif
 
-	LightMapResolution = FMath::Max(LightMapResolution, 0);
+	SetLightMapResolution(FMath::Max(GetLightMapResolution(), 0));
 
 	if (PropertyChangedEvent.MemberProperty 
-		&& ((PropertyChangedEvent.MemberProperty->GetFName() == GET_MEMBER_NAME_CHECKED(UStaticMesh, PositiveBoundsExtension)) 
-			|| (PropertyChangedEvent.MemberProperty->GetFName() == GET_MEMBER_NAME_CHECKED(UStaticMesh, NegativeBoundsExtension))))
+		&& (PropertyChangedEvent.MemberProperty->GetFName() == UStaticMesh::GetPositiveBoundsExtensionName() || 
+			PropertyChangedEvent.MemberProperty->GetFName() == UStaticMesh::GetNegativeBoundsExtensionName()))
 	{
 		// Update the extended bounds
 		CalculateExtendedBounds();
@@ -3310,13 +3310,8 @@ void UStaticMesh::PostEditChangeProperty(FPropertyChangedEvent& PropertyChangedE
 
 	Build(/*bSilent=*/ true);
 
-	// It's OK to access the property directly just to fetch its name
-	PRAGMA_DISABLE_DEPRECATION_WARNINGS
-	FName BodySetupProperyName = GET_MEMBER_NAME_CHECKED(UStaticMesh, BodySetup);
-	PRAGMA_ENABLE_DEPRECATION_WARNINGS
-
 	if (PropertyName == GET_MEMBER_NAME_CHECKED(UStaticMesh, bHasNavigationData)
-		|| PropertyName == BodySetupProperyName)
+		|| PropertyName == UStaticMesh::GetBodySetupName())
 	{
 		// Build called above will result in creation, update or destruction 
 		// of NavCollision. We need to let related StaticMeshComponents know
@@ -3326,8 +3321,8 @@ void UStaticMesh::PostEditChangeProperty(FPropertyChangedEvent& PropertyChangedE
 	if (PropertyChangedEvent.ChangeType != EPropertyChangeType::Interactive)
 	{
 		// Only unbuild lighting for properties which affect static lighting
-		if (PropertyName == GET_MEMBER_NAME_CHECKED(UStaticMesh, LightMapResolution)
-			|| PropertyName == GET_MEMBER_NAME_CHECKED(UStaticMesh, LightMapCoordinateIndex))
+		if (PropertyName == UStaticMesh::GetLightMapResolutionName()
+			|| PropertyName == UStaticMesh::GetLightMapCoordinateIndexName())
 		{
 			FStaticMeshComponentRecreateRenderStateContext Context(this, true);
 			SetLightingGuid();
@@ -3391,24 +3386,24 @@ void UStaticMesh::SetLODGroup(FName NewGroup, bool bRebuildImmediately)
 			
 			if (LODIndex != 0)
 			{
-			//Reset the section info map
-			if (bResetSectionInfoMap)
-			{
-				for (int32 SectionIndex = 0; SectionIndex < GetSectionInfoMap().GetSectionNumber(LODIndex); ++SectionIndex)
+				//Reset the section info map
+				if (bResetSectionInfoMap)
 				{
+					for (int32 SectionIndex = 0; SectionIndex < GetSectionInfoMap().GetSectionNumber(LODIndex); ++SectionIndex)
+					{
 						GetSectionInfoMap().Remove(LODIndex, SectionIndex);
+					}
+				}
+				//Clear the raw data if we change the LOD Group and we do not reduce ourself, this will force the user to do a import LOD which will manage the section info map properly
+				if (!SourceModel.IsRawMeshEmpty() && SourceModel.ReductionSettings.BaseLODModel != LODIndex)
+				{
+					FRawMesh EmptyRawMesh;
+					SourceModel.SaveRawMesh(EmptyRawMesh);
+					SourceModel.SourceImportFilename = FString();
 				}
-			}
-			//Clear the raw data if we change the LOD Group and we do not reduce ourself, this will force the user to do a import LOD which will manage the section info map properly
-			if (!SourceModel.IsRawMeshEmpty() && SourceModel.ReductionSettings.BaseLODModel != LODIndex)
-			{
-				FRawMesh EmptyRawMesh;
-				SourceModel.SaveRawMesh(EmptyRawMesh);
-				SourceModel.SourceImportFilename = FString();
 			}
 		}
-		}
-		LightMapResolution = GroupSettings.GetDefaultLightMapResolution();
+		SetLightMapResolution(GroupSettings.GetDefaultLightMapResolution());
 
 		if (!bBeforeDerivedDataCached)
 		{
@@ -3717,7 +3712,7 @@ void UStaticMesh::GetAssetRegistryTags(TArray<FAssetRegistryTag>& OutTags) const
 	OutTags.Add(FAssetRegistryTag("ApproxSize", ApproxSizeStr, FAssetRegistryTag::TT_Dimensional) );
 	OutTags.Add(FAssetRegistryTag("CollisionPrims", FString::FromInt(NumCollisionPrims), FAssetRegistryTag::TT_Numerical));
 	OutTags.Add(FAssetRegistryTag("LODs", FString::FromInt(NumLODs), FAssetRegistryTag::TT_Numerical));
-	OutTags.Add(FAssetRegistryTag("MinLOD", MinLOD.ToString(), FAssetRegistryTag::TT_Alphabetical));
+	OutTags.Add(FAssetRegistryTag("MinLOD", GetMinLOD().ToString(), FAssetRegistryTag::TT_Alphabetical));
 	OutTags.Add(FAssetRegistryTag("SectionsWithCollision", FString::FromInt(NumSectionsWithCollision), FAssetRegistryTag::TT_Numerical));
 	OutTags.Add(FAssetRegistryTag("DefaultCollision", DefaultCollisionName.ToString(), FAssetRegistryTag::TT_Alphabetical));
 	OutTags.Add(FAssetRegistryTag("CollisionComplexity", ComplexityString, FAssetRegistryTag::TT_Alphabetical));
@@ -4467,9 +4462,9 @@ bool UStaticMesh::InsertUVChannel(int32 LODIndex, int32 UVChannelIndex)
 				++LODBuildSettings.DstLightmapIndex;
 			}
 
-			if (UVChannelIndex <= LightMapCoordinateIndex)
+			if (UVChannelIndex <= GetLightMapCoordinateIndex())
 			{
-				++LightMapCoordinateIndex;
+				SetLightMapCoordinateIndex(GetLightMapCoordinateIndex() + 1);
 			}
 
 			CommitMeshDescription(LODIndex);
@@ -4518,9 +4513,9 @@ bool UStaticMesh::RemoveUVChannel(int32 LODIndex, int32 UVChannelIndex)
 				--LODBuildSettings.DstLightmapIndex;
 			}
 
-			if (UVChannelIndex < LightMapCoordinateIndex)
+			if (UVChannelIndex < GetLightMapCoordinateIndex())
 			{
-				--LightMapCoordinateIndex;
+				SetLightMapCoordinateIndex(GetLightMapCoordinateIndex() - 1);
 			}
 
 			CommitMeshDescription(LODIndex);
@@ -4609,6 +4604,11 @@ void UStaticMesh::CacheDerivedData()
 			{
 				GDistanceFieldAsyncQueue->BlockUntilBuildComplete(this, true);
 			}
+
+			if (GCardRepresentationAsyncQueue)
+			{
+				GCardRepresentationAsyncQueue->BlockUntilBuildComplete(this, true);
+			}
 		}
 
 		for (int32 LODIndex = 0; LODIndex < GetRenderData()->LODResources.Num(); ++LODIndex)
@@ -4621,8 +4621,6 @@ void UStaticMesh::CacheDerivedData()
 				DistanceFieldData->VolumeTexture.Release();
 			}
 		}
-
-		GCardRepresentationAsyncQueue->BlockUntilBuildComplete(this, true);
 	}
 
 	SetRenderData(MakeUnique<FStaticMeshRenderData>());
@@ -4654,14 +4652,14 @@ void UStaticMesh::CalculateExtendedBounds()
 	}
 
 	// Only apply bound extension if necessary, as it will result in a larger bounding sphere radius than retrieved from the render data
-	if (!NegativeBoundsExtension.IsZero() || !PositiveBoundsExtension.IsZero())
+	if (!GetNegativeBoundsExtension().IsZero() || !GetPositiveBoundsExtension().IsZero())
 	{
 		// Convert to Min and Max
 		FVector Min = Bounds.Origin - Bounds.BoxExtent;
 		FVector Max = Bounds.Origin + Bounds.BoxExtent;
 		// Apply bound extensions
-		Min -= NegativeBoundsExtension;
-		Max += PositiveBoundsExtension;
+		Min -= GetNegativeBoundsExtension();
+		Max += GetPositiveBoundsExtension();
 		// Convert back to Origin, Extent and update SphereRadius
 		Bounds.Origin = (Min + Max) / 2;
 		Bounds.BoxExtent = (Max - Min) / 2;
@@ -4671,7 +4669,6 @@ void UStaticMesh::CalculateExtendedBounds()
 	SetExtendedBounds(Bounds);
 }
 
-
 #if WITH_EDITORONLY_DATA
 FUObjectAnnotationSparseBool GStaticMeshesThatNeedMaterialFixup;
 #endif // #if WITH_EDITORONLY_DATA
@@ -4981,12 +4978,12 @@ void UStaticMesh::PostLoad()
 	if (GetNumSourceModels() > 0)
 	{
 		UStaticMesh* DistanceFieldReplacementMesh = GetSourceModel(0).BuildSettings.DistanceFieldReplacementMesh;
- 
+
 		if (DistanceFieldReplacementMesh)
 		{
 			DistanceFieldReplacementMesh->ConditionalPostLoad();
 		}
-		
+
 		//TODO remove this code when FRawMesh will be removed
 		//Fill the static mesh owner
 		int32 NumLODs = GetNumSourceModels();
@@ -5140,7 +5137,7 @@ void UStaticMesh::PostLoad()
 			if (bMaterialChange || CompactedMaterial.Num() < GetStaticMaterials().Num())
 			{
 				GetStaticMaterials().Empty(CompactedMaterial.Num());
-				for (const FStaticMaterial &Material : CompactedMaterial)
+				for (const FStaticMaterial& Material : CompactedMaterial)
 				{
 					GetStaticMaterials().Add(Material);
 				}
@@ -5165,7 +5162,7 @@ void UStaticMesh::PostLoad()
 		if (bSupportGpuUniformlyDistributedSampling)
 		{
 			// Initialise pointers to samplers
-			for (FStaticMeshLODResources &LOD : GetRenderData()->LODResources)
+			for (FStaticMeshLODResources& LOD : GetRenderData()->LODResources)
 			{
 				LOD.AreaWeightedSectionSamplersBuffer.Init(&LOD.AreaWeightedSectionSamplers);
 			}
@@ -5174,16 +5171,18 @@ void UStaticMesh::PostLoad()
 		// check the MinLOD values are all within range
 		bool bFixedMinLOD = false;
 		int32 MinAvailableLOD = FMath::Max<int32>(GetRenderData()->LODResources.Num() - 1, 0);
-		if (!GetRenderData()->LODResources.IsValidIndex(MinLOD.Default))
+		if (!GetRenderData()->LODResources.IsValidIndex(GetMinLOD().Default))
 		{
 			FFormatNamedArguments Arguments;
-			Arguments.Add(TEXT("MinLOD"), FText::AsNumber(MinLOD.Default));
+			Arguments.Add(TEXT("MinLOD"), FText::AsNumber(GetMinLOD().Default));
 			Arguments.Add(TEXT("MinAvailLOD"), FText::AsNumber(MinAvailableLOD));
 			FMessageLog("LoadErrors").Warning()
 				->AddToken(FUObjectToken::Create(this))
 				->AddToken(FTextToken::Create(FText::Format(LOCTEXT("LoadError_BadMinLOD", "Min LOD value of {MinLOD} is out of range 0..{MinAvailLOD} and has been adjusted to {MinAvailLOD}. Please verify and resave the asset."), Arguments)));
 
-			MinLOD.Default = MinAvailableLOD;
+			FPerPlatformInt LocalMinLOD = GetMinLOD();
+			LocalMinLOD.Default = MinAvailableLOD;
+			SetMinLOD(MoveTemp(LocalMinLOD));
 			bFixedMinLOD = true;
 		}
 		for (TMap<FName, int32>::TIterator It(MinLOD.PerPlatform); It; ++It)
@@ -5208,7 +5207,6 @@ void UStaticMesh::PostLoad()
 		}
 	}
 
-
 #endif // #if WITH_EDITOR
 
 #if WITH_EDITORONLY_DATA
@@ -5688,7 +5686,7 @@ bool UStaticMesh::GetMipDataFilename(const int32 MipIndex, FString& OutBulkDataF
 	LocalizedName = FPackageName::GetLocalizedPackagePath(LocalizedName);
 	bool bSucceed = FPackageName::DoesPackageExist(LocalizedName, nullptr, &OutBulkDataFilename);
 	check(bSucceed);
-	OutBulkDataFilename = FPaths::ChangeExtension(OutBulkDataFilename, MipIndex < MinLOD.Default ? TEXT(".uptnl") : TEXT(".ubulk"));
+	OutBulkDataFilename = FPaths::ChangeExtension(OutBulkDataFilename, MipIndex < GetMinLOD().Default ? TEXT(".uptnl") : TEXT(".ubulk"));
 	return true;
 }
 #endif // USE_BULKDATA_STREAMING_TOKEN
@@ -6212,7 +6210,7 @@ ENGINE_API void UStaticMesh::RemoveVertexColors()
 void UStaticMesh::EnforceLightmapRestrictions(bool bUseRenderData)
 {
 	// Legacy content may contain a lightmap resolution of 0, which was valid when vertex lightmaps were supported, but not anymore with only texture lightmaps
-	LightMapResolution = FMath::Max(LightMapResolution, 4);
+	SetLightMapResolution(FMath::Max(GetLightMapResolution(), 4));
 
 	// Lightmass only supports 4 UVs from Lightmass::MAX_TEXCOORDS
 	int32 NumUVs = 4;
@@ -6292,7 +6290,7 @@ void UStaticMesh::EnforceLightmapRestrictions(bool bUseRenderData)
 	check(NumUVs > 0);
 
 	// Clamp LightMapCoordinateIndex to be valid for all lightmap uvs
-	LightMapCoordinateIndex = FMath::Clamp(LightMapCoordinateIndex, 0, NumUVs - 1);
+	SetLightMapCoordinateIndex(FMath::Clamp(GetLightMapCoordinateIndex(), 0, NumUVs - 1));
 }
 
 /**
@@ -6493,7 +6491,7 @@ void UStaticMesh::CheckLightMapUVs( UStaticMesh* InStaticMesh, TArray< FString >
 	for( int32 CurLODModelIndex = 0; CurLODModelIndex < NumLods; ++CurLODModelIndex )
 	{
 		const FStaticMeshLODResources& RenderData = InStaticMesh->GetRenderData()->LODResources[CurLODModelIndex];
-		int32 LightMapTextureCoordinateIndex = InStaticMesh->LightMapCoordinateIndex;
+		int32 LightMapTextureCoordinateIndex = InStaticMesh->GetLightMapCoordinateIndex();
 
 		// We expect the light map texture coordinate to be greater than zero, as the first UV set
 		// should never really be used for light maps, unless this mesh was exported as a light mapped uv set.
@@ -6539,7 +6537,7 @@ void UStaticMesh::CheckLightMapUVs( UStaticMesh* InStaticMesh, TArray< FString >
 					}
 					break;
 				case FLocal::UVCheck_Missing:
-					UE_LOG(LogStaticMesh, Warning, TEXT( "[%s, LOD %i] missing light map UVs (Res %i, CoordIndex %i)" ), *InStaticMesh->GetName(), CurLODModelIndex, InStaticMesh->LightMapResolution, InStaticMesh->LightMapCoordinateIndex );
+					UE_LOG(LogStaticMesh, Warning, TEXT( "[%s, LOD %i] missing light map UVs (Res %i, CoordIndex %i)" ), *InStaticMesh->GetName(), CurLODModelIndex, InStaticMesh->GetLightMapResolution(), InStaticMesh->GetLightMapCoordinateIndex() );
 					break;
 				case FLocal::UVCheck_NoTriangles:
 					UE_LOG(LogStaticMesh, Warning, TEXT( "[%s, LOD %i] doesn't have any triangles" ), *InStaticMesh->GetName(), CurLODModelIndex );
diff --git a/Engine/Source/Runtime/Engine/Private/StaticMeshLight.cpp b/Engine/Source/Runtime/Engine/Private/StaticMeshLight.cpp
--- a/Engine/Source/Runtime/Engine/Private/StaticMeshLight.cpp
+++ b/Engine/Source/Runtime/Engine/Private/StaticMeshLight.cpp
@@ -326,8 +326,8 @@ void UStaticMeshComponent::GetStaticLightingInfo(FStaticLightingPrimitiveInfo& O
 			// Figure out whether we are storing the lighting/ shadowing information in a texture or vertex buffer.
 			bool bValidTextureMap;
 			if( (BaseLightMapWidth > 0) && (BaseLightMapHeight > 0) 
-				&& GetStaticMesh()->LightMapCoordinateIndex >= 0
-				&& (uint32)GetStaticMesh()->LightMapCoordinateIndex < LODRenderData.VertexBuffers.StaticMeshVertexBuffer.GetNumTexCoords())
+				&& GetStaticMesh()->GetLightMapCoordinateIndex() >= 0
+				&& (uint32)GetStaticMesh()->GetLightMapCoordinateIndex() < LODRenderData.VertexBuffers.StaticMeshVertexBuffer.GetNumTexCoords())
 			{
 				bValidTextureMap = true;
 			}
@@ -350,13 +350,13 @@ void UStaticMeshComponent::GetStaticLightingInfo(FStaticLightingPrimitiveInfo& O
 				if (LightmapType == ELightmapType::ForceVolumetric)
 				{
 					OutPrimitiveInfo.Mappings.Add(new FStaticLightingGlobalVolumeMapping(
-						StaticLightingMesh,this,LightMapWidth,LightMapHeight, GetStaticMesh()->LightMapCoordinateIndex));
+						StaticLightingMesh,this,LightMapWidth,LightMapHeight, GetStaticMesh()->GetLightMapCoordinateIndex()));
 				}
 				else
 				{
 					// Create a static lighting texture mapping for the LOD.
 					OutPrimitiveInfo.Mappings.Add(new FStaticMeshStaticLightingTextureMapping(
-						this,LODIndex,StaticLightingMesh,LightMapWidth,LightMapHeight, GetStaticMesh()->LightMapCoordinateIndex,true));
+						this,LODIndex,StaticLightingMesh,LightMapWidth,LightMapHeight, GetStaticMesh()->GetLightMapCoordinateIndex(),true));
 				}
 			}
 		}
@@ -409,8 +409,8 @@ ELightMapInteractionType UStaticMeshComponent::GetStaticLightingType() const
 				GetLightMapResolution( LightMapWidth, LightMapHeight );
 
 				if ((LightMapWidth > 0) && (LightMapHeight > 0) &&	
-					(GetStaticMesh()->LightMapCoordinateIndex >= 0) &&
-					((uint32)GetStaticMesh()->LightMapCoordinateIndex < LODRenderData.VertexBuffers.StaticMeshVertexBuffer.GetNumTexCoords())
+					(GetStaticMesh()->GetLightMapCoordinateIndex() >= 0) &&
+					((uint32)GetStaticMesh()->GetLightMapCoordinateIndex() < LODRenderData.VertexBuffers.StaticMeshVertexBuffer.GetNumTexCoords())
 					)
 				{
 					InteractionType = LMIT_Texture;
@@ -510,7 +510,7 @@ bool UStaticMeshComponent::SetStaticLightingMapping(bool bTextureMapping, int32
 					if (OverriddenLightMapRes == 0)
 					{
 						// See if the static mesh has a valid setting
-						if (GetStaticMesh()->LightMapResolution != 0)
+						if (GetStaticMesh()->GetLightMapResolution() != 0)
 						{
 							// Simply uncheck the override...
 							bOverrideLightMapRes = false;
@@ -533,7 +533,7 @@ bool UStaticMeshComponent::SetStaticLightingMapping(bool bTextureMapping, int32
 				else
 				{
 					// See if the static mesh has a valid setting
-					if (GetStaticMesh()->LightMapResolution == 0)
+					if (GetStaticMesh()->GetLightMapResolution() == 0)
 					{
 						// See if the static mesh has a valid setting
 						if (OverriddenLightMapRes != 0)
@@ -574,7 +574,7 @@ bool UStaticMeshComponent::SetStaticLightingMapping(bool bTextureMapping, int32
 				if (OverriddenLightMapRes != 0)
 				{
 					// See if the static mesh has a valid setting
-					if (GetStaticMesh()->LightMapResolution == 0)
+					if (GetStaticMesh()->GetLightMapResolution() == 0)
 					{
 						// Simply uncheck the override...
 						bOverrideLightMapRes = false;
@@ -595,7 +595,7 @@ bool UStaticMeshComponent::SetStaticLightingMapping(bool bTextureMapping, int32
 			else
 			{
 				// See if the static mesh has a valid setting
-				if (GetStaticMesh()->LightMapResolution != 0)
+				if (GetStaticMesh()->GetLightMapResolution() != 0)
 				{
 					// Set it to the default value from the ini
 					OverriddenLightMapRes = 0;
diff --git a/Engine/Source/Runtime/Engine/Private/StaticMeshRender.cpp b/Engine/Source/Runtime/Engine/Private/StaticMeshRender.cpp
--- a/Engine/Source/Runtime/Engine/Private/StaticMeshRender.cpp
+++ b/Engine/Source/Runtime/Engine/Private/StaticMeshRender.cpp
@@ -196,7 +196,7 @@ FStaticMeshSceneProxy::FStaticMeshSceneProxy(UStaticMeshComponent* InComponent,
 
 	const auto FeatureLevel = GetScene().GetFeatureLevel();
 
-	const int32 SMCurrentMinLOD = InComponent->GetStaticMesh()->MinLOD.GetValue();
+	const int32 SMCurrentMinLOD = InComponent->GetStaticMesh()->GetMinLOD().GetValue();
 	int32 EffectiveMinLOD = InComponent->bOverrideMinLOD ? InComponent->MinLOD : SMCurrentMinLOD;
 
 #if WITH_EDITOR
@@ -2312,7 +2312,7 @@ FPrimitiveSceneProxy* UStaticMeshComponent::CreateSceneProxy()
 	}
 
 	const FStaticMeshLODResourcesArray& LODResources = GetStaticMesh()->GetRenderData()->LODResources;
-	if (LODResources.Num() == 0	|| LODResources[FMath::Clamp<int32>(GetStaticMesh()->MinLOD.Default, 0, LODResources.Num()-1)].VertexBuffers.StaticMeshVertexBuffer.GetNumVertices() == 0)
+	if (LODResources.Num() == 0	|| LODResources[FMath::Clamp<int32>(GetStaticMesh()->GetMinLOD().Default, 0, LODResources.Num()-1)].VertexBuffers.StaticMeshVertexBuffer.GetNumVertices() == 0)
 	{
 		return nullptr;
 	}
diff --git a/Engine/Source/Runtime/Engine/Private/UnrealEngine.cpp b/Engine/Source/Runtime/Engine/Private/UnrealEngine.cpp
--- a/Engine/Source/Runtime/Engine/Private/UnrealEngine.cpp
+++ b/Engine/Source/Runtime/Engine/Private/UnrealEngine.cpp
@@ -5463,9 +5463,9 @@ bool UEngine::HandleListStaticMeshesCommand(const TCHAR* Cmd, FOutputDevice& Ar)
 		int32		MobileMinLOD = -1;
 
 #if WITH_EDITORONLY_DATA 
-		if (Mesh->MinLOD.PerPlatform.Find(("Mobile")) != nullptr)
+		if (Mesh->GetMinLOD().PerPlatform.Find(("Mobile")) != nullptr)
 		{
-			MobileMinLOD = *Mesh->MinLOD.PerPlatform.Find(("Mobile"));
+			MobileMinLOD = *Mesh->GetMinLOD().PerPlatform.Find(("Mobile"));
 		}
 #endif
 
