Fix missed StaticMesh deprecation

diff --git a/Engine/Source/Runtime/Engine/Private/StaticMesh.cpp b/Engine/Source/Runtime/Engine/Private/StaticMesh.cpp
--- a/Engine/Source/Runtime/Engine/Private/StaticMesh.cpp
+++ b/Engine/Source/Runtime/Engine/Private/StaticMesh.cpp
@@ -5171,21 +5171,21 @@ void UStaticMesh::PostLoad()
 		// check the MinLOD values are all within range
 		bool bFixedMinLOD = false;
 		int32 MinAvailableLOD = FMath::Max<int32>(GetRenderData()->LODResources.Num() - 1, 0);
-		if (!GetRenderData()->LODResources.IsValidIndex(GetMinLOD().Default))
+		FPerPlatformInt LocalMinLOD = GetMinLOD();
+		if (!GetRenderData()->LODResources.IsValidIndex(LocalMinLOD.Default))
 		{
 			FFormatNamedArguments Arguments;
-			Arguments.Add(TEXT("MinLOD"), FText::AsNumber(GetMinLOD().Default));
+			Arguments.Add(TEXT("MinLOD"), FText::AsNumber(LocalMinLOD.Default));
 			Arguments.Add(TEXT("MinAvailLOD"), FText::AsNumber(MinAvailableLOD));
 			FMessageLog("LoadErrors").Warning()
 				->AddToken(FUObjectToken::Create(this))
 				->AddToken(FTextToken::Create(FText::Format(LOCTEXT("LoadError_BadMinLOD", "Min LOD value of {MinLOD} is out of range 0..{MinAvailLOD} and has been adjusted to {MinAvailLOD}. Please verify and resave the asset."), Arguments)));
 
-			FPerPlatformInt LocalMinLOD = GetMinLOD();
 			LocalMinLOD.Default = MinAvailableLOD;
-			SetMinLOD(MoveTemp(LocalMinLOD));
 			bFixedMinLOD = true;
 		}
-		for (TMap<FName, int32>::TIterator It(MinLOD.PerPlatform); It; ++It)
+		
+		for (TMap<FName, int32>::TIterator It(LocalMinLOD.PerPlatform); It; ++It)
 		{
 			if (!GetRenderData()->LODResources.IsValidIndex(It.Value()))
 			{
@@ -5201,8 +5201,10 @@ void UStaticMesh::PostLoad()
 				bFixedMinLOD = true;
 			}
 		}
+
 		if (bFixedMinLOD)
 		{
+			SetMinLOD(MoveTemp(LocalMinLOD));
 			FMessageLog("LoadErrors").Open();
 		}
 	}
