All Strata nodes now take world vertex normal when none is specified.

diff --git a/Engine/Source/Runtime/Engine/Private/Materials/MaterialExpressions.cpp b/Engine/Source/Runtime/Engine/Private/Materials/MaterialExpressions.cpp
--- a/Engine/Source/Runtime/Engine/Private/Materials/MaterialExpressions.cpp
+++ b/Engine/Source/Runtime/Engine/Private/Materials/MaterialExpressions.cpp
@@ -19023,14 +19023,10 @@ UMaterialExpressionStrataDiffuseBSDF::UMaterialExpressionStrataDiffuseBSDF(const
 #if WITH_EDITOR
 int32 UMaterialExpressionStrataDiffuseBSDF::Compile(class FMaterialCompiler* Compiler, int32 OutputIndex)
 {
-	if (!Normal.GetTracedInput().Expression)
-	{
-		return Compiler->Errorf(TEXT("Missing Normal input"));	// STRATA_TODO: defaults in the detail panel with auto read only if plug in
-	}
 	int32 OutputCodeChunk = Compiler->StrataDiffuseOrenNayarBSDF(
 		Albedo.GetTracedInput().Expression		? Albedo.Compile(Compiler)		: Compiler->Constant3(0.18f, 0.18f, 0.18f),
 		Roughness.GetTracedInput().Expression	? Roughness.Compile(Compiler)	: Compiler->Constant(0.0f),
-		Normal.Compile(Compiler));
+		Normal.GetTracedInput().Expression		? Normal.Compile(Compiler)		: Compiler->PixelNormalWS());
 
 	StrataCreateSingleBSDFMaterial(Compiler, OutputCodeChunk, STRATA_BSDF_TYPE_DIFFUSE_ON);
 
@@ -19088,14 +19084,10 @@ UMaterialExpressionStrataDiffuseChanBSDF::UMaterialExpressionStrataDiffuseChanBS
 #if WITH_EDITOR
 int32 UMaterialExpressionStrataDiffuseChanBSDF::Compile(class FMaterialCompiler* Compiler, int32 OutputIndex)
 {
-	if (!Normal.GetTracedInput().Expression)
-	{
-		return Compiler->Errorf(TEXT("Missing Normal input"));	// STRATA_TODO: defaults in the detail panel with auto read only if plug in
-	}
 	int32 OutputCodeChunk = Compiler->StrataDiffuseChanBSDF(
 		Albedo.GetTracedInput().Expression		? Albedo.Compile(Compiler)		: Compiler->Constant3(0.18f, 0.18f, 0.18f),
 		Roughness.GetTracedInput().Expression	? Roughness.Compile(Compiler)	: Compiler->Constant(0.0f),
-		Normal.Compile(Compiler));
+		Normal.GetTracedInput().Expression		? Normal.Compile(Compiler)		: Compiler->PixelNormalWS());
 
 	StrataCreateSingleBSDFMaterial(Compiler, OutputCodeChunk, STRATA_BSDF_TYPE_DIFFUSE_CHAN);
 
@@ -19153,15 +19145,11 @@ UMaterialExpressionStrataDielectricBSDF::UMaterialExpressionStrataDielectricBSDF
 #if WITH_EDITOR
 int32 UMaterialExpressionStrataDielectricBSDF::Compile(class FMaterialCompiler* Compiler, int32 OutputIndex)
 {
-	if (!Normal.GetTracedInput().Expression)
-	{
-		return Compiler->Errorf(TEXT("Missing Normal input"));	// STRATA_TODO: defaults in the detail panel with auto read only if plug in
-	}
 	int32 OutputCodeChunk = Compiler->StrataDielectricBSDF(
-		Roughness.GetTracedInput().Expression	?	Roughness.Compile(Compiler) : Compiler->Constant2(0.0f, 0.0f),
-		IOR.GetTracedInput().Expression			?	IOR.Compile(Compiler)		: Compiler->Constant(1.5f),
-		Tint.GetTracedInput().Expression		?	Tint.Compile(Compiler)		: Compiler->Constant3(1.0f, 1.0f, 1.0f),
-		Normal.Compile(Compiler));
+		Roughness.GetTracedInput().Expression	? Roughness.Compile(Compiler)	: Compiler->Constant2(0.0f, 0.0f),
+		IOR.GetTracedInput().Expression			? IOR.Compile(Compiler)			: Compiler->Constant(1.5f),
+		Tint.GetTracedInput().Expression		? Tint.Compile(Compiler)		: Compiler->Constant3(1.0f, 1.0f, 1.0f),
+		Normal.GetTracedInput().Expression		? Normal.Compile(Compiler)		: Compiler->PixelNormalWS());
 
 	StrataCreateSingleBSDFMaterial(Compiler, OutputCodeChunk, STRATA_BSDF_TYPE_DIELECTRIC);
 
@@ -19222,15 +19210,11 @@ UMaterialExpressionStrataConductorBSDF::UMaterialExpressionStrataConductorBSDF(c
 #if WITH_EDITOR
 int32 UMaterialExpressionStrataConductorBSDF::Compile(class FMaterialCompiler* Compiler, int32 OutputIndex)
 {
-	if (!Normal.GetTracedInput().Expression)
-	{
-		return Compiler->Errorf(TEXT("Missing Normal input"));		// STRATA_TODO: defaults in the detail panel with auto read only if plug in
-	}
 	int32 OutputCodeChunk = Compiler->StrataConductorBSDF(
 		Reflectivity.GetTracedInput().Expression	? Reflectivity.Compile(Compiler): Compiler->Constant3(0.947f, 0.776f, 0.371f),	// Default to Gold
 		EdgeColor.GetTracedInput().Expression		? EdgeColor.Compile(Compiler)	: Compiler->Constant3(1.000f, 0.982f, 0.753f),	// Default to Gold
 		Roughness.GetTracedInput().Expression		? Roughness.Compile(Compiler)	: Compiler->Constant2(0.0f, 0.0f),
-		Normal.Compile(Compiler));
+		Normal.GetTracedInput().Expression			? Normal.Compile(Compiler)		: Compiler->PixelNormalWS());
 
 	StrataCreateSingleBSDFMaterial(Compiler, OutputCodeChunk, STRATA_BSDF_TYPE_CONDUCTOR);
 
@@ -19295,7 +19279,7 @@ int32 UMaterialExpressionStrataVolumeBSDF::Compile(class FMaterialCompiler* Comp
 		Albedo.GetTracedInput().Expression		? Albedo.Compile(Compiler)			: Compiler->Constant3(0.0f, 0.0f, 0.0f),
 		Extinction.GetTracedInput().Expression	? Extinction.Compile(Compiler)		: Compiler->Constant3(0.0f, 0.0f, 0.0f),
 		Anisotropy.GetTracedInput().Expression	? Anisotropy.Compile(Compiler)		: Compiler->Constant(0.0f),
-		Thickness.GetTracedInput().Expression ? Thickness.Compile(Compiler)			: Compiler->Constant(0.001f)); // default = 1mm
+		Thickness.GetTracedInput().Expression	? Thickness.Compile(Compiler)		: Compiler->Constant(0.001f)); // default = 1mm
 
 	StrataCreateSingleBSDFMaterial(Compiler, OutputCodeChunk, STRATA_BSDF_TYPE_VOLUME);
 
