Fixed parameter and property naming in capsule-shape related vislog and DebugRendererSceneProxy code

diff --git a/Engine/Source/Developer/LogVisualizer/Private/VisualLoggerRenderingActor.cpp b/Engine/Source/Developer/LogVisualizer/Private/VisualLoggerRenderingActor.cpp
--- a/Engine/Source/Developer/LogVisualizer/Private/VisualLoggerRenderingActor.cpp
+++ b/Engine/Source/Developer/LogVisualizer/Private/VisualLoggerRenderingActor.cpp
@@ -79,7 +79,7 @@ FPrimitiveSceneProxy* UVisualLoggerRenderingComponent::CreateSceneProxy()
 		VLogSceneProxy->Texts.Append(CurrentShapes.Value.Texts);
 		VLogSceneProxy->Cylinders.Append(CurrentShapes.Value.Cylinders);
 		VLogSceneProxy->ArrowLines.Append(CurrentShapes.Value.Arrows);
-		VLogSceneProxy->Capsles.Append(CurrentShapes.Value.Capsules);
+		VLogSceneProxy->Capsules.Append(CurrentShapes.Value.Capsules);
 	}
 
 	{
@@ -91,7 +91,7 @@ FPrimitiveSceneProxy* UVisualLoggerRenderingComponent::CreateSceneProxy()
 		VLogSceneProxy->Texts.Append(RenderingActor->TestDebugShapes.Texts);
 		VLogSceneProxy->Cylinders.Append(RenderingActor->TestDebugShapes.Cylinders);
 		VLogSceneProxy->ArrowLines.Append(RenderingActor->TestDebugShapes.Arrows);
-		VLogSceneProxy->Capsles.Append(RenderingActor->TestDebugShapes.Capsules);
+		VLogSceneProxy->Capsules.Append(RenderingActor->TestDebugShapes.Capsules);
 	}
 
 #if WITH_EDITOR
@@ -285,7 +285,7 @@ void AVisualLoggerRenderingActor::AddDebugRendering()
 	}
 
 	{
-		const FVector Center = FVector(1000, 0, 128);
+		const FVector Base = FVector(1000, 0, 128);
 		const float HalfHeight = 150;
 		const float Radius = 50;
 		const FQuat Rotation = FQuat::Identity;
@@ -295,7 +295,7 @@ void AVisualLoggerRenderingActor::AddDebugRendering()
 		const FVector YAxis = Axes.GetScaledAxis(EAxis::Y);
 		const FVector ZAxis = Axes.GetScaledAxis(EAxis::Z);
 
-		TestDebugShapes.Capsules.Add(FDebugRenderSceneProxy::FCapsule(Center, Radius, XAxis, YAxis, ZAxis, HalfHeight, FColor::Yellow));
+		TestDebugShapes.Capsules.Add(FDebugRenderSceneProxy::FCapsule(Base, Radius, XAxis, YAxis, ZAxis, HalfHeight, FColor::Yellow));
 	}
 	{
 		const float Radius = 50;
@@ -570,7 +570,7 @@ void AVisualLoggerRenderingActor::GetDebugShapes(const FVisualLogDevice::FVisual
 			const bool bDrawLabel = ElementToDraw->Description.IsEmpty() == false;
 			for (int32 Index = 0; Index + 2 < ElementToDraw->Points.Num(); Index += 3)
 			{
-				const FVector Center = ElementToDraw->Points[Index + 0];
+				const FVector Base = ElementToDraw->Points[Index + 0];
 				const FVector FirstData = ElementToDraw->Points[Index + 1];
 				const FVector SecondData = ElementToDraw->Points[Index + 2];
 				const float HalfHeight = FirstData.X;
@@ -582,10 +582,10 @@ void AVisualLoggerRenderingActor::GetDebugShapes(const FVisualLogDevice::FVisual
 				const FVector YAxis = Axes.GetScaledAxis(EAxis::Y);
 				const FVector ZAxis = Axes.GetScaledAxis(EAxis::Z);
 
-				DebugShapes.Capsules.Add(FDebugRenderSceneProxy::FCapsule(Center, Radius, XAxis, YAxis, ZAxis, HalfHeight, Color));
+				DebugShapes.Capsules.Add(FDebugRenderSceneProxy::FCapsule(Base, Radius, XAxis, YAxis, ZAxis, HalfHeight, Color));
 				if (bDrawLabel)
 				{
-					DebugShapes.Texts.Add(FDebugRenderSceneProxy::FText3d(ElementToDraw->Description, Center, Color));
+					DebugShapes.Texts.Add(FDebugRenderSceneProxy::FText3d(ElementToDraw->Description, Base + HalfHeight * FVector::UpVector, Color));
 				}
 			}
 		}
diff --git a/Engine/Source/Runtime/Engine/Private/DebugRenderSceneProxy.cpp b/Engine/Source/Runtime/Engine/Private/DebugRenderSceneProxy.cpp
--- a/Engine/Source/Runtime/Engine/Private/DebugRenderSceneProxy.cpp
+++ b/Engine/Source/Runtime/Engine/Private/DebugRenderSceneProxy.cpp
@@ -245,14 +245,14 @@ void FDebugRenderSceneProxy::GetDynamicMeshElements(const TArray<const FSceneVie
 				}
 			}
 
-			for (auto It = Capsles.CreateConstIterator(); It; ++It)
+			for (auto It = Capsules.CreateConstIterator(); It; ++It)
 			{
-				if (PointInView(It->Location, View))
+				if (PointInView(It->Base, View))
 				{
 					if (DrawType == SolidAndWireMeshes || DrawType == WireMesh)
 					{
 						const float HalfAxis = FMath::Max<float>(It->HalfHeight - It->Radius, 1.f);
-						const FVector BottomEnd = It->Location + It->Radius * It->Z;
+						const FVector BottomEnd = It->Base + It->Radius * It->Z;
 						const FVector TopEnd = BottomEnd + (2 * HalfAxis) * It->Z;
 						const float CylinderHalfHeight = (TopEnd - BottomEnd).Size() * 0.5;
 						const FVector CylinderLocation = BottomEnd + CylinderHalfHeight * It->Z;
@@ -260,7 +260,7 @@ void FDebugRenderSceneProxy::GetDynamicMeshElements(const TArray<const FSceneVie
 					}
 					if (DrawType == SolidAndWireMeshes || DrawType == SolidMesh)
 					{
-						GetCapsuleMesh(It->Location, It->X, It->Y, It->Z, It->Color, It->Radius, It->HalfHeight, 16, MaterialCache[0][It->Color.WithAlpha(DrawAlpha)], SDPG_World, false, ViewIndex, Collector);
+						GetCapsuleMesh(It->Base, It->X, It->Y, It->Z, It->Color, It->Radius, It->HalfHeight, 16, MaterialCache[0][It->Color.WithAlpha(DrawAlpha)], SDPG_World, false, ViewIndex, Collector);
 					}
 				}
 			}
diff --git a/Engine/Source/Runtime/Engine/Private/VisualLogger/VisualLoggerTypes.cpp b/Engine/Source/Runtime/Engine/Private/VisualLogger/VisualLoggerTypes.cpp
--- a/Engine/Source/Runtime/Engine/Private/VisualLogger/VisualLoggerTypes.cpp
+++ b/Engine/Source/Runtime/Engine/Private/VisualLogger/VisualLoggerTypes.cpp
@@ -248,11 +248,11 @@ void FVisualLogEntry::AddElement(const FVector& Start, const FVector& End, float
 	ElementsToDraw.Add(Element);
 }
 
-void FVisualLogEntry::AddElement(const FVector& Center, float HalfHeight, float Radius, const FQuat & Rotation, const FName& CategoryName, ELogVerbosity::Type Verbosity, const FColor& Color, const FString& Description)
+void FVisualLogEntry::AddElement(const FVector& Base, float HalfHeight, float Radius, const FQuat & Rotation, const FName& CategoryName, ELogVerbosity::Type Verbosity, const FColor& Color, const FString& Description)
 {
 	FVisualLogShapeElement Element(Description, Color, 0, CategoryName);
 	Element.Points.Reserve(3);
-	Element.Points.Add(Center);
+	Element.Points.Add(Base);
 	Element.Points.Add(FVector(HalfHeight, Radius, Rotation.X));
 	Element.Points.Add(FVector(Rotation.Y, Rotation.Z, Rotation.W));
 	Element.Type = EVisualLoggerShapeElement::Capsule;
diff --git a/Engine/Source/Runtime/Engine/Public/DebugRenderSceneProxy.h b/Engine/Source/Runtime/Engine/Public/DebugRenderSceneProxy.h
--- a/Engine/Source/Runtime/Engine/Public/DebugRenderSceneProxy.h
+++ b/Engine/Source/Runtime/Engine/Public/DebugRenderSceneProxy.h
@@ -206,9 +206,9 @@ public:
 	struct FCapsule
 	{
 		FCapsule() {}
-		FCapsule(const FVector& InLocation, const float& InRadius, const FVector& x, const FVector& y, const FVector &z, const float& InHalfHeight, const FLinearColor& InColor)
+		FCapsule(const FVector& InBase, const float& InRadius, const FVector& x, const FVector& y, const FVector &z, const float& InHalfHeight, const FLinearColor& InColor)
 			: Radius(InRadius)
-			, Location(InLocation)
+			, Base(InBase)
 			, Color(InColor.ToFColor(true))
 			, HalfHeight(InHalfHeight)
 			, X(x)
@@ -219,7 +219,7 @@ public:
 		}
 
 		float Radius;
-		FVector Location; //Center pointer of the base of the cylinder.
+		FVector Base; //Center point of the base of the cylinder.
 		FColor Color;
 		float HalfHeight;
 		FVector X, Y, Z; //X, Y, and Z alignment axes to draw along.
@@ -235,7 +235,7 @@ public:
 	TArray<FText3d> Texts;
 	TArray<FCone> Cones;
 	TArray<FMesh> Meshes;
-	TArray<FCapsule> Capsles;
+	TArray<FCapsule> Capsules;
 
 	uint32 ViewFlagIndex;
 	float TextWithoutShadowDistance;
diff --git a/Engine/Source/Runtime/Engine/Public/VisualLogger/VisualLogger.h b/Engine/Source/Runtime/Engine/Public/VisualLogger/VisualLogger.h
--- a/Engine/Source/Runtime/Engine/Public/VisualLogger/VisualLogger.h
+++ b/Engine/Source/Runtime/Engine/Public/VisualLogger/VisualLogger.h
@@ -48,8 +48,8 @@
 #define UE_VLOG_CYLINDER(LogOwner, CategoryName, Verbosity, Start, End, Radius, Color, Format, ...) if(FVisualLogger::IsRecording()) FVisualLogger::GeometryShapeLogf(LogOwner, CategoryName, ELogVerbosity::Verbosity, Start, End, Radius, Color, Format, ##__VA_ARGS__)
 #define UE_CVLOG_CYLINDER(Condition, LogOwner, CategoryName, Verbosity, Start, End, Radius, Color, Format, ...)  if(FVisualLogger::IsRecording() && Condition) {UE_VLOG_CYLINDER(LogOwner, CategoryName, Verbosity, Start, End, Radius, Color, Format, ##__VA_ARGS__);} 
 // Capsule shape
-#define UE_VLOG_CAPSULE(LogOwner, CategoryName, Verbosity, Center, HalfHeight, Radius, Rotation, Color, Format, ...) if(FVisualLogger::IsRecording()) FVisualLogger::GeometryShapeLogf(LogOwner, CategoryName, ELogVerbosity::Verbosity, Center, HalfHeight, Radius, Rotation, Color, Format, ##__VA_ARGS__)
-#define UE_CVLOG_CAPSULE(Condition, LogOwner, CategoryName, Verbosity, Center, HalfHeight, Radius, Rotation, Color, Format, ...)  if(FVisualLogger::IsRecording() && Condition) {UE_VLOG_CAPSULE(LogOwner, CategoryName, Verbosity, Center, HalfHeight, Radius, Rotation, Color, Format, ##__VA_ARGS__);} 
+#define UE_VLOG_CAPSULE(LogOwner, CategoryName, Verbosity, Base, HalfHeight, Radius, Rotation, Color, Format, ...) if(FVisualLogger::IsRecording()) FVisualLogger::GeometryShapeLogf(LogOwner, CategoryName, ELogVerbosity::Verbosity, Base, HalfHeight, Radius, Rotation, Color, Format, ##__VA_ARGS__)
+#define UE_CVLOG_CAPSULE(Condition, LogOwner, CategoryName, Verbosity, Base, HalfHeight, Radius, Rotation, Color, Format, ...)  if(FVisualLogger::IsRecording() && Condition) {UE_VLOG_CAPSULE(LogOwner, CategoryName, Verbosity, Base, HalfHeight, Radius, Rotation, Color, Format, ##__VA_ARGS__);} 
 // Histogram data for 2d graphs 
 #define UE_VLOG_HISTOGRAM(LogOwner, CategoryName, Verbosity, GraphName, DataName, Data) if(FVisualLogger::IsRecording()) FVisualLogger::HistogramDataLogf(LogOwner, CategoryName, ELogVerbosity::Verbosity, GraphName, DataName, Data, FColor::White, TEXT(""))
 #define UE_CVLOG_HISTOGRAM(Condition, LogOwner, CategoryName, Verbosity, GraphName, DataName, Data) if(FVisualLogger::IsRecording() && Condition) {UE_VLOG_HISTOGRAM(LogOwner, CategoryName, Verbosity, GraphName, DataName, Data);} 
@@ -100,8 +100,8 @@
 #define UE_CVLOG_CONE(Condition, Object, CategoryName, Verbosity, Orgin, Direction, Length, Angle, Color, DescriptionFormat, ...)
 #define UE_VLOG_CYLINDER(Object, CategoryName, Verbosity, Start, End, Radius, Color, DescriptionFormat, ...)
 #define UE_CVLOG_CYLINDER(Condition, Object, CategoryName, Verbosity, Start, End, Radius, Color, DescriptionFormat, ...)
-#define UE_VLOG_CAPSULE(Object, CategoryName, Verbosity, Center, HalfHeight, Radius, Rotation, Color, DescriptionFormat, ...)
-#define UE_CVLOG_CAPSULE(Condition, Object, CategoryName, Verbosity, Center, HalfHeight, Radius, Rotation, Color, DescriptionFormat, ...)
+#define UE_VLOG_CAPSULE(Object, CategoryName, Verbosity, Base, HalfHeight, Radius, Rotation, Color, DescriptionFormat, ...)
+#define UE_CVLOG_CAPSULE(Condition, Object, CategoryName, Verbosity, Base, HalfHeight, Radius, Rotation, Color, DescriptionFormat, ...)
 #define UE_VLOG_HISTOGRAM(Actor, CategoryName, Verbosity, GraphName, DataName, Data)
 #define UE_CVLOG_HISTOGRAM(Condition, Actor, CategoryName, Verbosity, GraphName, DataName, Data)
 #define UE_VLOG_PULLEDCONVEX(LogOwner, CategoryName, Verbosity, ConvexPoints, MinZ, MaxZ, Color, Format, ...)
@@ -160,8 +160,8 @@ class ENGINE_API FVisualLogger : public FOutputDevice
 	static void GeometryShapeLogfImpl (const UObject* LogOwner, const FName& CategoryName, ELogVerbosity::Type Verbosity, const FVector& Origin, const FVector& Direction, const float Length, const float Angle, const FColor& Color, const TCHAR* Fmt, ...);
 	static void GeometryShapeLogfImpl (const UObject* LogOwner, const FLogCategoryBase& Category, ELogVerbosity::Type Verbosity, const FVector& Start, const FVector& End, const float Radius, const FColor& Color, const TCHAR* Fmt, ...);
 	static void GeometryShapeLogfImpl (const UObject* LogOwner, const FName& CategoryName, ELogVerbosity::Type Verbosity, const FVector& Start, const FVector& End, const float Radius, const FColor& Color, const TCHAR* Fmt, ...);
-	static void GeometryShapeLogfImpl (const UObject* LogOwner, const FLogCategoryBase& Category, ELogVerbosity::Type Verbosity, const FVector& Center, float HalfHeight, float Radius, const FQuat& Rotation, const FColor& Color, const TCHAR* Fmt, ...);
-	static void GeometryShapeLogfImpl (const UObject* LogOwner, const FName& CategoryName, ELogVerbosity::Type Verbosity, const FVector& Center, float HalfHeight, float Radius, const FQuat& Rotation, const FColor& Color, const TCHAR* Fmt, ...);
+	static void GeometryShapeLogfImpl (const UObject* LogOwner, const FLogCategoryBase& Category, ELogVerbosity::Type Verbosity, const FVector& Base, float HalfHeight, float Radius, const FQuat& Rotation, const FColor& Color, const TCHAR* Fmt, ...);
+	static void GeometryShapeLogfImpl (const UObject* LogOwner, const FName& CategoryName, ELogVerbosity::Type Verbosity, const FVector& Base, float HalfHeight, float Radius, const FQuat& Rotation, const FColor& Color, const TCHAR* Fmt, ...);
 	static void NavAreaShapeLogfImpl  (const UObject* LogOwner, const FLogCategoryBase& Category, ELogVerbosity::Type Verbosity, const TArray<FVector>& ConvexPoints, float MinZ, float MaxZ, const FColor& Color, const TCHAR* Fmt, ...);
 	static void NavAreaShapeLogfImpl  (const UObject* LogOwner, const FName& CategoryName, ELogVerbosity::Type Verbosity, const TArray<FVector>& ConvexPoints, float MinZ, float MaxZ, const FColor& Color, const TCHAR* Fmt, ...);
 	static void GeometryShapeLogfImpl (const UObject* LogOwner, const FLogCategoryBase& Category, ELogVerbosity::Type Verbosity, const TArray<FVector>& Vertices, const TArray<int32>&Indices, const FColor& Color, const TCHAR* Fmt, ...);
@@ -301,20 +301,20 @@ public:
 
 	// Capsule log
 	template <typename FmtType, typename... Types>
-	static void GeometryShapeLogf(const UObject* LogOwner, const FLogCategoryBase& Category, ELogVerbosity::Type Verbosity, const FVector& Center, float HalfHeight, float Radius, const FQuat & Rotation, const FColor& Color, const FmtType& Fmt, Types... Args)
+	static void GeometryShapeLogf(const UObject* LogOwner, const FLogCategoryBase& Category, ELogVerbosity::Type Verbosity, const FVector& Base, float HalfHeight, float Radius, const FQuat & Rotation, const FColor& Color, const FmtType& Fmt, Types... Args)
 	{
 		static_assert(TIsArrayOrRefOfType<FmtType, TCHAR>::Value, "Formatting string must be a TCHAR array.");
 		static_assert(TAnd<TIsValidVariadicFunctionArg<Types>...>::Value, "Invalid argument(s) passed to FVisualLogger::GeometryShapeLogf");
 
-		GeometryShapeLogfImpl(LogOwner, Category, Verbosity, Center, HalfHeight, Radius, Rotation, Color, Fmt, Args...);
+		GeometryShapeLogfImpl(LogOwner, Category, Verbosity, Base, HalfHeight, Radius, Rotation, Color, Fmt, Args...);
 	}
 	template <typename FmtType, typename... Types>
-	static void GeometryShapeLogf(const UObject* LogOwner, const FName& CategoryName, ELogVerbosity::Type Verbosity, const FVector& Center, float HalfHeight, float Radius, const FQuat & Rotation, const FColor& Color, const FmtType& Fmt, Types... Args)
+	static void GeometryShapeLogf(const UObject* LogOwner, const FName& CategoryName, ELogVerbosity::Type Verbosity, const FVector& Base, float HalfHeight, float Radius, const FQuat & Rotation, const FColor& Color, const FmtType& Fmt, Types... Args)
 	{
 		static_assert(TIsArrayOrRefOfType<FmtType, TCHAR>::Value, "Formatting string must be a TCHAR array.");
 		static_assert(TAnd<TIsValidVariadicFunctionArg<Types>...>::Value, "Invalid argument(s) passed to FVisualLogger::GeometryShapeLogf");
 
-		GeometryShapeLogfImpl(LogOwner, CategoryName, Verbosity, Center, HalfHeight, Radius, Rotation, Color, Fmt, Args...);
+		GeometryShapeLogfImpl(LogOwner, CategoryName, Verbosity, Base, HalfHeight, Radius, Rotation, Color, Fmt, Args...);
 	}
 
 	// NavArea/Extruded convex log
@@ -709,17 +709,17 @@ inline void FVisualLogger::GeometryShapeLogfImpl(const UObject* Object, const FN
 	);
 }
 
-inline void FVisualLogger::GeometryShapeLogfImpl(const UObject* Object, const FLogCategoryBase& Category, ELogVerbosity::Type Verbosity, const FVector& Center, float HalfHeight, float Radius, const FQuat& Rotation, const FColor& Color, const TCHAR* Fmt, ...)
+inline void FVisualLogger::GeometryShapeLogfImpl(const UObject* Object, const FLogCategoryBase& Category, ELogVerbosity::Type Verbosity, const FVector& Base, float HalfHeight, float Radius, const FQuat& Rotation, const FColor& Color, const TCHAR* Fmt, ...)
 {
 	const FName CategoryName = Category.GetCategoryName(); 
 	COLLAPSED_LOGF(
-		CurrentEntry->AddElement(Center, HalfHeight, Radius, Rotation, CategoryName, Verbosity, Color, Buffer);
+		CurrentEntry->AddElement(Base, HalfHeight, Radius, Rotation, CategoryName, Verbosity, Color, Buffer);
 	);
 }
-inline void FVisualLogger::GeometryShapeLogfImpl(const UObject* Object, const FName& CategoryName, ELogVerbosity::Type Verbosity, const FVector& Center, float HalfHeight, float Radius, const FQuat& Rotation, const FColor& Color, const TCHAR* Fmt, ...)
+inline void FVisualLogger::GeometryShapeLogfImpl(const UObject* Object, const FName& CategoryName, ELogVerbosity::Type Verbosity, const FVector& Base, float HalfHeight, float Radius, const FQuat& Rotation, const FColor& Color, const TCHAR* Fmt, ...)
 {
 	COLLAPSED_LOGF(
-		CurrentEntry->AddElement(Center, HalfHeight, Radius, Rotation, CategoryName, Verbosity, Color, Buffer);
+		CurrentEntry->AddElement(Base, HalfHeight, Radius, Rotation, CategoryName, Verbosity, Color, Buffer);
 	);
 }
 
