World Partition: AInfo and all derived classes are now considered as always loaded.

diff --git a/Engine/Source/Runtime/Engine/Classes/GameFramework/Info.h b/Engine/Source/Runtime/Engine/Classes/GameFramework/Info.h
--- a/Engine/Source/Runtime/Engine/Classes/GameFramework/Info.h
+++ b/Engine/Source/Runtime/Engine/Classes/GameFramework/Info.h
@@ -32,6 +32,10 @@ public:
 	/** Indicates whether this actor should participate in level bounds calculations. */
 	virtual bool IsLevelBoundsRelevant() const override { return false; }
 
+#if WITH_EDITOR
+	virtual EActorGridPlacement GetDefaultGridPlacement() const override { return EActorGridPlacement::AlwaysLoaded; }
+#endif
+
 public:
 #if WITH_EDITORONLY_DATA
 	/** Returns SpriteComponent subobject **/
diff --git a/Engine/Source/Runtime/Engine/Public/WorldPartition/DataLayer/WorldDataLayers.h b/Engine/Source/Runtime/Engine/Public/WorldPartition/DataLayer/WorldDataLayers.h
--- a/Engine/Source/Runtime/Engine/Public/WorldPartition/DataLayer/WorldDataLayers.h
+++ b/Engine/Source/Runtime/Engine/Public/WorldPartition/DataLayer/WorldDataLayers.h
@@ -22,8 +22,6 @@ public:
 	virtual void PostLoad() override;
 
 #if WITH_EDITOR
-	virtual EActorGridPlacement GetDefaultGridPlacement() const override { return EActorGridPlacement::AlwaysLoaded; }
-
 	UDataLayer* CreateDataLayer();
 	bool RemoveDataLayer(const UDataLayer* InDataLayer);
 	bool RemoveDataLayers(const TArray<const UDataLayer*>& InDataLayers);
