UAnimSequence:
- Deprecated NumFrames instead added NumberOfKeys (data is copied over during serialization patch up)
- Deprecated GetNumberOfFrames and GetRawNumberOfFrames instead added GetNumberOfSampledKeys
- Deprecated SetRawNumberOfFrames and instead added SetNumberOfSampledKeys
- Added SamplingFrameRate (FFrameRate), which is populated according to the sequence's actual number of frames and length
- Changed some behaviour to use FFrameRate API to calculate frame number / timings
- Deprecated GetFrameRate and instead added GetSamplingFrameRate
- Fixed uses where NumFrames was assumed to contain the number of frames rather than the number of keys (ResizeSequence, AddLoopingInterpolation, CropAnimation, InsertKeysIntoRawData)

UAnimStreamable:
- Deprecated NumFrames instead added NumberOfKeys (data is copied over during serialization patch up)

ITimeManagementModule:
- Exposed GetCommonFrameRates through the abstract interface to get around circular Engine module dependency

diff --git a/Engine/Plugins/Experimental/AlembicImporter/Source/AlembicLibrary/Private/AbcImporter.cpp b/Engine/Plugins/Experimental/AlembicImporter/Source/AlembicLibrary/Private/AbcImporter.cpp
--- a/Engine/Plugins/Experimental/AlembicImporter/Source/AlembicLibrary/Private/AbcImporter.cpp
+++ b/Engine/Plugins/Experimental/AlembicImporter/Source/AlembicLibrary/Private/AbcImporter.cpp
@@ -650,7 +650,7 @@ TArray<UObject*> FAbcImporter::ImportAsSkeletalMesh(UObject* InParent, EObjectFl
 		Sequence->ImportResampleFramerate = AbcFile->GetFramerate();
 
 		// Add 1 because of the way AnimSequence computes its framerate
-		Sequence->SetRawNumberOfFrame(AbcFile->GetImportNumFrames() + 1);
+		Sequence->SetNumberOfSampledKeys(AbcFile->GetImportNumFrames() + 1);
 
 		int32 ObjectIndex = 0;
 		uint32 TriangleOffset = 0;
diff --git a/Engine/Plugins/Experimental/ControlRig/Source/ControlRig/Private/Sequencer/MovieSceneControlRigParameterSection.cpp b/Engine/Plugins/Experimental/ControlRig/Source/ControlRig/Private/Sequencer/MovieSceneControlRigParameterSection.cpp
--- a/Engine/Plugins/Experimental/ControlRig/Source/ControlRig/Private/Sequencer/MovieSceneControlRigParameterSection.cpp
+++ b/Engine/Plugins/Experimental/ControlRig/Source/ControlRig/Private/Sequencer/MovieSceneControlRigParameterSection.cpp
@@ -2077,7 +2077,7 @@ bool UMovieSceneControlRigParameterSection::LoadAnimSequenceIntoThisSection(UAni
 	
 	FFrameRate TickResolution = MovieScene->GetTickResolution();
 	float Length = AnimSequence->GetPlayLength();
-	float FrameRate = AnimSequence->GetFrameRate();
+	const FFrameRate& FrameRate = AnimSequence->GetSamplingFrameRate();
 	FFrameNumber EndFrame = TickResolution.AsFrameNumber(Length);
 	FFrameNumber StartFrame(0);
 
@@ -2090,18 +2090,16 @@ bool UMovieSceneControlRigParameterSection::LoadAnimSequenceIntoThisSection(UAni
 	}
 
 
-	int32 NumFrames = Length * FrameRate;
-	NumFrames = AnimSequence->GetNumberOfFrames();
-	FFrameNumber FrameRateInFrameNumber = TickResolution.AsFrameNumber(1.0f / FrameRate);
+	const int32 NumberOfFrames = FrameRate.AsFrameNumber(Length).Value;
+	FFrameNumber FrameRateInFrameNumber = TickResolution.AsFrameNumber(FrameRate.AsInterval());
 	int32 ExtraProgress = bKeyReduce ? FloatChannels.Num() : 0;
 	
-	FScopedSlowTask Progress(NumFrames + ExtraProgress, LOCTEXT("BakingToControlRig_SlowTask", "Baking To Control Rig..."));	
+	FScopedSlowTask Progress(NumberOfFrames + ExtraProgress, LOCTEXT("BakingToControlRig_SlowTask", "Baking To Control Rig..."));	
 	Progress.MakeDialog(true);
 
-
-	for (int32 Index = 0; Index < NumFrames; ++Index)
+	for (int32 Index = 0; Index < NumberOfFrames; ++Index)
 	{
-		float SequenceSecond = AnimSequence->GetTimeAtFrame(Index);
+		const float SequenceSecond = FrameRate.AsSeconds(Index);
 		FFrameNumber FrameNumber = StartFrame + (FrameRateInFrameNumber * Index);
 
 		for (FFloatCurve& Curve : AnimSequence->RawCurveData.FloatCurves)
diff --git a/Engine/Plugins/Experimental/MotionTrailEditorMode/Source/MotionTrailEditorMode/Private/Sequencer/AnimationBoneTrail.cpp b/Engine/Plugins/Experimental/MotionTrailEditorMode/Source/MotionTrailEditorMode/Private/Sequencer/AnimationBoneTrail.cpp
--- a/Engine/Plugins/Experimental/MotionTrailEditorMode/Source/MotionTrailEditorMode/Private/Sequencer/AnimationBoneTrail.cpp
+++ b/Engine/Plugins/Experimental/MotionTrailEditorMode/Source/MotionTrailEditorMode/Private/Sequencer/AnimationBoneTrail.cpp
@@ -48,7 +48,7 @@ void FAnimTrajectoryCache::Evaluate(FTrajectoryCache* ParentTransformCache)
 	GetSpaceBasedAnimationData(GlobalBoneTransforms);
 	ComponentBoneTransforms.SetNum(GlobalBoneTransforms.Num());
 
-	Spacing = 1.0 / CachedAnimSequence->GetFrameRate();
+	Spacing = CachedAnimSequence->GetSamplingFrameRate().AsInterval();
 	const double AnimLength = CachedAnimSequence->GetPlayLength();
 	AnimRange = TRange<double>(0.0, AnimLength);
 	for (int32 TrackIndex = 0; TrackIndex < GlobalBoneTransforms.Num(); TrackIndex++)
@@ -95,7 +95,7 @@ void FAnimTrajectoryCache::GetSpaceBasedAnimationData(TArray<TArray<FTransform>>
 	OutAnimationDataInComponentSpace.AddZeroed(NumBones);
 
 	// 2d array of animated time [boneindex][time key]
-	int32 NumKeys = CachedAnimSequence->GetNumberOfFrames();
+	const int32 NumKeys = CachedAnimSequence->GetNumberOfSampledKeys();
 	float Interval = UE::MotionTrailEditor::GetIntervalPerKey(NumKeys, CachedAnimSequence->GetPlayLength());
 
 	// allocate arrays
diff --git a/Engine/Plugins/Importers/USDImporter/Source/USDUtilities/Private/USDSkeletalDataConversion.cpp b/Engine/Plugins/Importers/USDImporter/Source/USDUtilities/Private/USDSkeletalDataConversion.cpp
--- a/Engine/Plugins/Importers/USDImporter/Source/USDUtilities/Private/USDSkeletalDataConversion.cpp
+++ b/Engine/Plugins/Importers/USDImporter/Source/USDUtilities/Private/USDSkeletalDataConversion.cpp
@@ -1895,7 +1895,7 @@ bool UsdToUnreal::ConvertSkelAnim( const pxr::UsdSkelSkeletonQuery& InUsdSkeleto
 	OutSkeletalAnimationAsset->ImportFileFramerate = Stage.Get()->GetFramesPerSecond();
 	OutSkeletalAnimationAsset->ImportResampleFramerate = FramesPerSecond;
 	OutSkeletalAnimationAsset->SetSequenceLength(SequenceLengthSeconds);
-	OutSkeletalAnimationAsset->SetRawNumberOfFrame( NumBakedFrames );
+	OutSkeletalAnimationAsset->SetNumberOfSampledKeys( NumBakedFrames );
 	OutSkeletalAnimationAsset->MarkRawDataAsModified();	
 	OutSkeletalAnimationAsset->PostProcessSequence();
 
diff --git a/Engine/Plugins/Tests/EditorTests/Source/EditorTests/Private/UnrealEd/FbxAutomationTests.cpp b/Engine/Plugins/Tests/EditorTests/Source/EditorTests/Private/UnrealEd/FbxAutomationTests.cpp
--- a/Engine/Plugins/Tests/EditorTests/Source/EditorTests/Private/UnrealEd/FbxAutomationTests.cpp
+++ b/Engine/Plugins/Tests/EditorTests/Source/EditorTests/Private/UnrealEd/FbxAutomationTests.cpp
@@ -1560,7 +1560,7 @@ bool FFbxImportAssetsAutomationTest::RunTest(const FString& Parameters)
 						*GetFormatedMessageErrorInTestData(CleanFilename, TestPlan->TestPlanName, TEXT("Animation_Frame_Number"), ExpectedResultIndex)));
 					break;
 				}
-				int32 FrameNumber = AnimSequence->GetNumberOfFrames();
+				int32 FrameNumber = AnimSequence->GetNumberOfSampledKeys();
 				if (FrameNumber != ExpectedResult.ExpectedPresetsDataInteger[0])
 				{
 					ExecutionInfo.AddError(FString::Printf(TEXT("%s [%d frames but expected %d]"),
diff --git a/Engine/Plugins/VirtualProduction/Takes/Source/TakeTrackRecorders/Private/MovieSceneAnimationTrackRecorder.cpp b/Engine/Plugins/VirtualProduction/Takes/Source/TakeTrackRecorders/Private/MovieSceneAnimationTrackRecorder.cpp
--- a/Engine/Plugins/VirtualProduction/Takes/Source/TakeTrackRecorders/Private/MovieSceneAnimationTrackRecorder.cpp
+++ b/Engine/Plugins/VirtualProduction/Takes/Source/TakeTrackRecorders/Private/MovieSceneAnimationTrackRecorder.cpp
@@ -364,7 +364,7 @@ bool UMovieSceneAnimationTrackRecorder::LoadRecordedFile(const FString& FileName
 
 								AnimSequence->RecycleAnimSequence();
 								AnimSequence->SetSequenceLength(0.f);
-								AnimSequence->SetRawNumberOfFrame(0);
+								AnimSequence->SetNumberOfSampledKeys(0);
 								for (int32 TrackIndex = 0; TrackIndex < Header.AnimationTrackNames.Num(); ++TrackIndex)
 								{
 									AnimSequence->AddNewRawTrack(Header.AnimationTrackNames[TrackIndex]);
@@ -383,8 +383,8 @@ bool UMovieSceneAnimationTrackRecorder::LoadRecordedFile(const FString& FileName
 										RawTrack.ScaleKeys.Add(Frame.AnimationData[TrackIndex].ScaleKey);
 									}
 								}
-								AnimSequence->SetRawNumberOfFrame( InFrames.Num() );
-								AnimSequence->SetSequenceLength((AnimSequence->GetRawNumberOfFrames() > 1) ? (AnimSequence->GetRawNumberOfFrames() - 1) * Header.IntervalTime : MINIMUM_ANIMATION_LENGTH);
+								AnimSequence->SetNumberOfSampledKeys( InFrames.Num() );
+								AnimSequence->SetSequenceLength((AnimSequence->GetNumberOfSampledKeys() > 1) ? (AnimSequence->GetNumberOfSampledKeys() - 1) * Header.IntervalTime : MINIMUM_ANIMATION_LENGTH);
 
 								//fix up notifies todo notifies mz
 								AnimSequence->PostProcessSequence();
diff --git a/Engine/Source/Editor/AnimGraph/Private/AnimPreviewInstance.cpp b/Engine/Source/Editor/AnimGraph/Private/AnimPreviewInstance.cpp
--- a/Engine/Source/Editor/AnimGraph/Private/AnimPreviewInstance.cpp
+++ b/Engine/Source/Editor/AnimGraph/Private/AnimPreviewInstance.cpp
@@ -704,11 +704,11 @@ void UAnimPreviewInstance::MontagePreview_StepForward()
 		MontagePreview_SetPlaying(true);
 
 		// Advance a single frame, leaving it paused afterwards
-		int32 NumFrames = Montage->GetNumberOfFrames();
+		const int32 NumKeys = Montage->GetNumberOfSampledKeys();
 		// Add DELTA to prefer next frame when we're close to the boundary
-		float CurrentFraction = Proxy.GetCurrentTime() / Montage->GetPlayLength() + DELTA;
-		float NextFrame = FMath::Clamp<float>(FMath::FloorToFloat(CurrentFraction * NumFrames) + 1.0f, 0, NumFrames);
-		float NewTime = Montage->GetPlayLength() * (NextFrame / (NumFrames-1));
+		const float CurrentFraction = Proxy.GetCurrentTime() / Montage->GetPlayLength() + DELTA;
+		const float NextFrame = FMath::Clamp<float>(FMath::FloorToFloat(CurrentFraction * NumKeys) + 1.0f, 0, NumKeys);
+		const float NewTime = Montage->GetPlayLength() * (NextFrame / (NumKeys-1));
 
 		GetSkelMeshComponent()->GlobalAnimRateScale = 1.0f;
 		GetSkelMeshComponent()->TickAnimation(NewTime - Proxy.GetCurrentTime(), false);
@@ -760,11 +760,11 @@ void UAnimPreviewInstance::MontagePreview_StepBackward()
 		MontagePreview_SetPlaying(true);
 
 		// Advance a single frame, leaving it paused afterwards
-		int32 NumFrames = Montage->GetNumberOfFrames();
+		const int32 NumKeys = Montage->GetNumberOfSampledKeys();
 		// Add DELTA to prefer next frame when we're close to the boundary
-		float CurrentFraction = Proxy.GetCurrentTime() / Montage->GetPlayLength() + DELTA;
-		float NextFrame = FMath::Clamp<float>(FMath::FloorToFloat(CurrentFraction * NumFrames) - 1.0f, 0, NumFrames);
-		float NewTime = Montage->GetPlayLength() * (NextFrame / (NumFrames-1));
+		const float CurrentFraction = Proxy.GetCurrentTime() / Montage->GetPlayLength() + DELTA;
+		const float NextFrame = FMath::Clamp<float>(FMath::FloorToFloat(CurrentFraction * NumKeys) - 1.0f, 0, NumKeys);
+		const float NewTime = Montage->GetPlayLength() * (NextFrame / (NumKeys-1));
 
 		GetSkelMeshComponent()->GlobalAnimRateScale = 1.0f;
 		GetSkelMeshComponent()->TickAnimation(FMath::Abs(NewTime - Proxy.GetCurrentTime()), false);
diff --git a/Engine/Source/Editor/AnimationModifiers/Private/AnimationBlueprintLibrary.cpp b/Engine/Source/Editor/AnimationModifiers/Private/AnimationBlueprintLibrary.cpp
--- a/Engine/Source/Editor/AnimationModifiers/Private/AnimationBlueprintLibrary.cpp
+++ b/Engine/Source/Editor/AnimationModifiers/Private/AnimationBlueprintLibrary.cpp
@@ -25,7 +25,7 @@ void UAnimationBlueprintLibrary::GetNumFrames(const UAnimSequence* AnimationSequ
 	NumFrames = 0;
 	if (AnimationSequence)
 	{
-		NumFrames = AnimationSequence->NumFrames;
+		NumFrames = AnimationSequence->GetNumberOfSampledKeys() - 1;
 	}
 	else
 	{
@@ -33,6 +33,19 @@ void UAnimationBlueprintLibrary::GetNumFrames(const UAnimSequence* AnimationSequ
 	}
 }
 
+void UAnimationBlueprintLibrary::GetNumKeys(const UAnimSequence* AnimationSequence, int32& NumKeys)
+{
+	NumKeys = 0;
+	if (AnimationSequence)
+	{
+		NumKeys = AnimationSequence->GetNumberOfSampledKeys();
+	}
+	else
+	{
+		UE_LOG(LogAnimationBlueprintLibrary, Warning, TEXT("Invalid Animation Sequence supplied for GetNumKeys"));
+	}
+}
+
 void UAnimationBlueprintLibrary::GetAnimationTrackNames(const UAnimSequence* AnimationSequence, TArray<FName>& TrackNames)
 {
 	TrackNames.Empty();
diff --git a/Engine/Source/Editor/AnimationModifiers/Public/AnimationBlueprintLibrary.h b/Engine/Source/Editor/AnimationModifiers/Public/AnimationBlueprintLibrary.h
--- a/Engine/Source/Editor/AnimationModifiers/Public/AnimationBlueprintLibrary.h
+++ b/Engine/Source/Editor/AnimationModifiers/Public/AnimationBlueprintLibrary.h
@@ -39,6 +39,10 @@ public:
 	/** Retrieves the number of animation frames for the given Animation Sequence */
 	UFUNCTION(BlueprintPure, meta=(AutoCreateRefTerm = "AnimationSequence"), Category = "AnimationBlueprintLibrary|Animation")
 	static void GetNumFrames(const UAnimSequence* AnimationSequence, int32& NumFrames);
+
+	/** Retrieves the number of animation keys for the given Animation Sequence */
+	UFUNCTION(BlueprintPure, meta = (AutoCreateRefTerm = "AnimationSequence"), Category = "AnimationBlueprintLibrary|Animation")
+	static void GetNumKeys(const UAnimSequence* AnimationSequence, int32& NumKeys);
 	
 	/** Retrieves the Names of the individual ATracks for the given Animation Sequence */
 	UFUNCTION(BlueprintPure, Category = "AnimationBlueprintLibrary|Animation")
diff --git a/Engine/Source/Editor/DetailCustomizations/Private/AnimMontageSegmentDetails.cpp b/Engine/Source/Editor/DetailCustomizations/Private/AnimMontageSegmentDetails.cpp
--- a/Engine/Source/Editor/DetailCustomizations/Private/AnimMontageSegmentDetails.cpp
+++ b/Engine/Source/Editor/DetailCustomizations/Private/AnimMontageSegmentDetails.cpp
@@ -541,17 +541,17 @@ uint32 SAnimationSegmentScrubPanel::GetNumOfFrames() const
 	{
 		UAnimSingleNodeInstance* PreviewInst = GetPreviewInstance();
 		float Length = PreviewInst->GetLength();
-		// if anim sequence, use correct num frames
-		int32 NumFrames = (int32) (Length/0.0333f); 
+		// if anim sequence, use correct number of keys
+		int32 NumKeys = (int32) (Length/0.0333f); 
 		if (PreviewInst->GetCurrentAsset() && PreviewInst->GetCurrentAsset()->IsA(UAnimSequenceBase::StaticClass()))
 		{
-			NumFrames = CastChecked<UAnimSequenceBase>(PreviewInst->GetCurrentAsset())->GetNumberOfFrames();
+			NumKeys = CastChecked<UAnimSequenceBase>(PreviewInst->GetCurrentAsset())->GetNumberOfSampledKeys();
 		}
-		return NumFrames;
+		return NumKeys;
 	}
 	else if (LockedSequence)
 	{
-		return LockedSequence->GetNumberOfFrames();
+		return LockedSequence->GetNumberOfSampledKeys();
 	}
 	return 1;
 }
diff --git a/Engine/Source/Editor/DetailCustomizations/Private/AnimSequenceDetails.cpp b/Engine/Source/Editor/DetailCustomizations/Private/AnimSequenceDetails.cpp
--- a/Engine/Source/Editor/DetailCustomizations/Private/AnimSequenceDetails.cpp
+++ b/Engine/Source/Editor/DetailCustomizations/Private/AnimSequenceDetails.cpp
@@ -536,7 +536,7 @@ void SAnimationRefPoseViewport::Tick( const FGeometry& AllottedGeometry, const d
 				{
 					int RefFrameIndex;
 					RefFrameIndexPropertyHandle->GetValue( RefFrameIndex );
-					float Fraction = ( AnimRef->GetRawNumberOfFrames() > 0 ) ? FMath::Clamp<float>( (float)RefFrameIndex / (float)AnimRef->GetRawNumberOfFrames(), 0.f, 1.f ) : 0.f;
+					float Fraction = ( AnimRef->GetNumberOfSampledKeys() > 0 ) ? FMath::Clamp<float>( (float)RefFrameIndex / (float)AnimRef->GetNumberOfSampledKeys(), 0.f, 1.f ) : 0.f;
 					float RefTime = AnimRef->GetPlayLength() * Fraction;
 					PreviewComponent->PreviewInstance->SetPosition( RefTime, false );
 					PreviewComponent->PreviewInstance->SetPlaying( false );
@@ -621,7 +621,7 @@ TArray<float> SAnimationRefPoseViewport::GetBars() const
 	{
 		int RefFrameIndex;
 		RefFrameIndexPropertyHandle->GetValue(RefFrameIndex);
-		float Fraction = (AnimRef->GetRawNumberOfFrames() > 0)? FMath::Clamp<float>((float)RefFrameIndex/(float)AnimRef->GetRawNumberOfFrames(), 0.f, 1.f) : 0.f;
+		float Fraction = (AnimRef->GetNumberOfSampledKeys() > 0)? FMath::Clamp<float>((float)RefFrameIndex/(float)AnimRef->GetNumberOfSampledKeys(), 0.f, 1.f) : 0.f;
 		Bars.Add(AnimRef->GetPlayLength() * Fraction);
 	}
 	else
@@ -635,7 +635,7 @@ void SAnimationRefPoseViewport::OnBarDrag(int32 Index, float Position)
 {
 	if (AnimRef)
 	{
-		int RefFrameIndex = FMath::Clamp<int>(AnimRef->GetPlayLength() > 0.0f? (int)(Position * (float)AnimRef->GetRawNumberOfFrames() / AnimRef->GetPlayLength() + 0.5f) : 0.0f, 0, AnimRef->GetRawNumberOfFrames() - 1);
+		int RefFrameIndex = FMath::Clamp<int>(AnimRef->GetPlayLength() > 0.0f? (int)(Position * (float)AnimRef->GetNumberOfSampledKeys() / AnimRef->GetPlayLength() + 0.5f) : 0.0f, 0, AnimRef->GetNumberOfSampledKeys() - 1);
 		RefFrameIndexPropertyHandle->SetValue(RefFrameIndex);
 	}
 }
diff --git a/Engine/Source/Editor/Persona/Private/AnimTimeline/AnimModel.cpp b/Engine/Source/Editor/Persona/Private/AnimTimeline/AnimModel.cpp
--- a/Engine/Source/Editor/Persona/Private/AnimTimeline/AnimModel.cpp
+++ b/Engine/Source/Editor/Persona/Private/AnimTimeline/AnimModel.cpp
@@ -56,7 +56,7 @@ double FAnimModel::GetFrameRate() const
 {
 	if(UAnimSequence* AnimSequence = Cast<UAnimSequence>(GetAnimSequenceBase()))
 	{
-		return (double)AnimSequence->GetFrameRate();
+		return AnimSequence->GetSamplingFrameRate().AsDecimal();
 	}
 	else
 	{
diff --git a/Engine/Source/Editor/Persona/Private/AnimTimeline/AnimTimelineTrack_Montage.cpp b/Engine/Source/Editor/Persona/Private/AnimTimeline/AnimTimelineTrack_Montage.cpp
--- a/Engine/Source/Editor/Persona/Private/AnimTimeline/AnimTimelineTrack_Montage.cpp
+++ b/Engine/Source/Editor/Persona/Private/AnimTimeline/AnimTimelineTrack_Montage.cpp
@@ -391,7 +391,7 @@ class SMontageSections : public SLeafWidget
 							SNew(SNumericEntryBox<int32>)
 							.Font(FEditorStyle::GetFontStyle(TEXT("MenuItem.Font")))
 							.MinValue(0)
-							.MaxValue(AnimMontage->GetNumberOfFrames())
+							.MaxValue(AnimMontage->GetNumberOfSampledKeys())
 							.Value(AnimMontage->GetFrameAtTime(Section.GetTime()))
 							.AllowSpin(true)						
 							.OnValueCommitted_Lambda([this, SectionIndex](int32 InValue, ETextCommit::Type InCommitType)
diff --git a/Engine/Source/Editor/Persona/Private/AnimTimeline/SAnimTimeline.cpp b/Engine/Source/Editor/Persona/Private/AnimTimeline/SAnimTimeline.cpp
--- a/Engine/Source/Editor/Persona/Private/AnimTimeline/SAnimTimeline.cpp
+++ b/Engine/Source/Editor/Persona/Private/AnimTimeline/SAnimTimeline.cpp
@@ -392,40 +392,40 @@ FReply SAnimTimeline::OnMouseButtonUp(const FGeometry& MyGeometry, const FPointe
 			FFrameTime MouseTime = TimeSliderController->GetFrameTimeFromMouse(MyGeometry, MouseEvent.GetScreenSpacePosition());
 			float CurrentFrameTime = (float)((double)MouseTime.AsDecimal() / (double)Model.Pin()->GetTickResolution());
 			float SequenceLength = AnimSequence->GetPlayLength();
-			uint32 NumFrames = AnimSequence->GetNumberOfFrames();
+			const uint32 NumKeys = AnimSequence->GetNumberOfSampledKeys();
 
 			MenuBuilder.BeginSection("SequenceEditingContext", LOCTEXT("SequenceEditing", "Sequence Editing") );
 			{
 				float CurrentFrameFraction = CurrentFrameTime / SequenceLength;
-				int32 CurrentFrameNumber = CurrentFrameFraction * NumFrames;
+				int32 CurrentKeyIndex = CurrentFrameFraction * NumKeys;
 
 				FUIAction Action;
 				FText Label;
 
 				//Menu - "Remove Before"
 				//Only show this option if the selected frame is greater than frame 1 (first frame)
-				if (CurrentFrameNumber > 0)
+				if (CurrentKeyIndex > 0)
 				{
-					CurrentFrameFraction = (float)CurrentFrameNumber / (float)NumFrames;
+					CurrentFrameFraction = (float)CurrentKeyIndex / (float)NumKeys;
 
 					//Corrected frame time based on selected frame number
 					float CorrectedFrameTime = CurrentFrameFraction * SequenceLength;
 
 					Action = FUIAction(FExecuteAction::CreateSP(this, &SAnimTimeline::OnCropAnimSequence, true, CorrectedFrameTime));
-					Label = FText::Format(LOCTEXT("RemoveTillFrame", "Remove frame 0 to frame {0}"), FText::AsNumber(CurrentFrameNumber));
+					Label = FText::Format(LOCTEXT("RemoveTillFrame", "Remove frame 0 to frame {0}"), FText::AsNumber(CurrentKeyIndex));
 					MenuBuilder.AddMenuEntry(Label, LOCTEXT("RemoveBefore_ToolTip", "Remove sequence before current position"), FSlateIcon(), Action);
 				}
 
-				uint32 NextFrameNumber = CurrentFrameNumber + 1;
+				uint32 NextKeyIndex = CurrentKeyIndex + 1;
 
 				//Menu - "Remove After"
-				//Only show this option if next frame (CurrentFrameNumber + 1) is valid
-				if (NextFrameNumber < NumFrames)
+				//Only show this option if next frame (CurrentKeyIndex + 1) is valid
+				if (NextKeyIndex < NumKeys)
 				{
-					float NextFrameFraction = (float)NextFrameNumber / (float)NumFrames;
+					float NextFrameFraction = (float)NextKeyIndex / (float)NumKeys;
 					float NextFrameTime = NextFrameFraction * SequenceLength;
 					Action = FUIAction(FExecuteAction::CreateSP(this, &SAnimTimeline::OnCropAnimSequence, false, NextFrameTime));
-					Label = FText::Format(LOCTEXT("RemoveFromFrame", "Remove from frame {0} to frame {1}"), FText::AsNumber(NextFrameNumber), FText::AsNumber(NumFrames));
+					Label = FText::Format(LOCTEXT("RemoveFromFrame", "Remove from frame {0} to frame {1}"), FText::AsNumber(NextKeyIndex), FText::AsNumber(NumKeys));
 					MenuBuilder.AddMenuEntry(Label, LOCTEXT("RemoveAfter_ToolTip", "Remove sequence after current position"), FSlateIcon(), Action);
 				}
 
@@ -434,12 +434,12 @@ FReply SAnimTimeline::OnMouseButtonUp(const FGeometry& MyGeometry, const FPointe
 				//Corrected frame time based on selected frame number
 				float CorrectedFrameTime = CurrentFrameFraction * SequenceLength;
 
-				Action = FUIAction(FExecuteAction::CreateSP(this, &SAnimTimeline::OnInsertAnimSequence, true, CurrentFrameNumber));
-				Label = FText::Format(LOCTEXT("InsertBeforeCurrentFrame", "Insert frame before {0}"), FText::AsNumber(CurrentFrameNumber));
+				Action = FUIAction(FExecuteAction::CreateSP(this, &SAnimTimeline::OnInsertAnimSequence, true, CurrentKeyIndex));
+				Label = FText::Format(LOCTEXT("InsertBeforeCurrentFrame", "Insert frame before {0}"), FText::AsNumber(CurrentKeyIndex));
 				MenuBuilder.AddMenuEntry(Label, LOCTEXT("InsertBefore_ToolTip", "Insert a frame before current position"), FSlateIcon(), Action);
 
-				Action = FUIAction(FExecuteAction::CreateSP(this, &SAnimTimeline::OnInsertAnimSequence, false, CurrentFrameNumber));
-				Label = FText::Format(LOCTEXT("InsertAfterCurrentFrame", "Insert frame after {0}"), FText::AsNumber(CurrentFrameNumber));
+				Action = FUIAction(FExecuteAction::CreateSP(this, &SAnimTimeline::OnInsertAnimSequence, false, CurrentKeyIndex));
+				Label = FText::Format(LOCTEXT("InsertAfterCurrentFrame", "Insert frame after {0}"), FText::AsNumber(CurrentKeyIndex));
 				MenuBuilder.AddMenuEntry(Label, LOCTEXT("InsertAfter_ToolTip", "Insert a frame after current position"), FSlateIcon(), Action);
 
 				MenuBuilder.AddMenuSeparator();
@@ -453,9 +453,9 @@ FReply SAnimTimeline::OnMouseButtonUp(const FGeometry& MyGeometry, const FPointe
 
 				MenuBuilder.AddMenuSeparator();
 				//Menu - "ReZero"
-				Action = FUIAction(FExecuteAction::CreateSP(this, &SAnimTimeline::OnReZeroAnimSequence, CurrentFrameNumber));
-				Label = FText::Format(LOCTEXT("ReZeroAtFrame", "Re-zero at frame {0}"), FText::AsNumber(CurrentFrameNumber));
-				MenuBuilder.AddMenuEntry(Label, FText::Format(LOCTEXT("ReZeroAtFrame_ToolTip", "Resets the root track to (0, 0, 0) at frame {0} and apply the difference to all root transform of the sequence. It moves whole sequence to the amount of current root transform."), FText::AsNumber(CurrentFrameNumber)), FSlateIcon(), Action);
+				Action = FUIAction(FExecuteAction::CreateSP(this, &SAnimTimeline::OnReZeroAnimSequence, CurrentKeyIndex));
+				Label = FText::Format(LOCTEXT("ReZeroAtFrame", "Re-zero at frame {0}"), FText::AsNumber(CurrentKeyIndex));
+				MenuBuilder.AddMenuEntry(Label, FText::Format(LOCTEXT("ReZeroAtFrame_ToolTip", "Resets the root track to (0, 0, 0) at frame {0} and apply the difference to all root transform of the sequence. It moves whole sequence to the amount of current root transform."), FText::AsNumber(CurrentKeyIndex)), FSlateIcon(), Action);
 
 				const int32 FrameNumberForCurrentTime = INDEX_NONE;
 				Action = FUIAction(FExecuteAction::CreateSP(this, &SAnimTimeline::OnReZeroAnimSequence, FrameNumberForCurrentTime));
@@ -521,7 +521,7 @@ void SAnimTimeline::OnAppendAnimSequence( bool bFromStart, int32 NumOfFrames )
 			AnimSequence->Modify();
 
 			// Crop the raw anim data.
-			int32 StartFrame = (bFromStart)? 0 : AnimSequence->GetRawNumberOfFrames() - 1;
+			int32 StartFrame = (bFromStart)? 0 : AnimSequence->GetNumberOfSampledKeys() - 1;
 			int32 EndFrame = StartFrame + NumOfFrames;
 			int32 CopyFrame = StartFrame;
 			AnimSequence->InsertFramesToRawAnimData(StartFrame, EndFrame, CopyFrame);
diff --git a/Engine/Source/Editor/Persona/Private/AnimTimeline/SAnimTimelineTransportControls.cpp b/Engine/Source/Editor/Persona/Private/AnimTimeline/SAnimTimelineTransportControls.cpp
--- a/Engine/Source/Editor/Persona/Private/AnimTimeline/SAnimTimelineTransportControls.cpp
+++ b/Engine/Source/Editor/Persona/Private/AnimTimeline/SAnimTimelineTransportControls.cpp
@@ -63,11 +63,11 @@ FReply SAnimTimelineTransportControls::OnClick_Forward_Step()
 	else if (SMC)
 	{
 		UAnimSequence* AnimSequence = Cast<UAnimSequence>(AnimSequenceBase);
-		const float TargetFramerate = AnimSequence ? AnimSequence->GetFrameRate() : 30.0f;
+		const FFrameRate TargetFramerate = AnimSequence ? AnimSequence->GetSamplingFrameRate() : FFrameRate(30, 1);
 
 		// Advance a single frame, leaving it paused afterwards
 		SMC->GlobalAnimRateScale = 1.0f;
-		SMC->TickAnimation(1.0f / TargetFramerate, false);
+		SMC->TickAnimation(TargetFramerate.AsInterval(), false);
 		SMC->GlobalAnimRateScale = 0.0f;
 	}
 
diff --git a/Engine/Source/Editor/Persona/Private/AnimationEditorViewportClient.cpp b/Engine/Source/Editor/Persona/Private/AnimationEditorViewportClient.cpp
--- a/Engine/Source/Editor/Persona/Private/AnimationEditorViewportClient.cpp
+++ b/Engine/Source/Editor/Persona/Private/AnimationEditorViewportClient.cpp
@@ -1002,7 +1002,7 @@ FText FAnimationViewportClient::GetDisplayInfo(bool bDisplayAllInfo) const
 
 	if (const UAnimSequence* AnimSequence = Cast<UAnimSequence>(GetAnimPreviewScene()->GetPreviewAnimationAsset()))
 	{
-		TextValue = ConcatenateLine(TextValue, FText::Format(LOCTEXT("FramerateFormat", "Framerate: {0}"), FText::AsNumber(AnimSequence->GetFrameRate())));
+		TextValue = ConcatenateLine(TextValue, FText::Format(LOCTEXT("FramerateFormat", "Framerate: {0}"), AnimSequence->GetSamplingFrameRate().ToPrettyText()));
 	}
 
 	return TextValue;
diff --git a/Engine/Source/Editor/Persona/Private/SAnimCompositePanel.cpp b/Engine/Source/Editor/Persona/Private/SAnimCompositePanel.cpp
--- a/Engine/Source/Editor/Persona/Private/SAnimCompositePanel.cpp
+++ b/Engine/Source/Editor/Persona/Private/SAnimCompositePanel.cpp
@@ -85,7 +85,7 @@ void SAnimCompositePanel::Update()
 				.ViewInputMax(ViewInputMax)
 				.OnGetNodeColor(this,  &SAnimCompositePanel::HandleGetNodeColor)
 				.TrackMaxValue(Composite->GetPlayLength())
-				.TrackNumDiscreteValues(Composite->GetNumberOfFrames())
+				.TrackNumDiscreteValues(Composite->GetNumberOfSampledKeys())
 				.bChildAnimMontage(AnimMontage && AnimMontage->HasParentAsset())
 				.OnAnimSegmentNodeClicked( this, &SAnimCompositePanel::ShowSegmentInDetailsView )
 				.OnPreAnimUpdate( this, &SAnimCompositePanel::PreAnimUpdate )
diff --git a/Engine/Source/Editor/Persona/Private/SAnimEditorBase.cpp b/Engine/Source/Editor/Persona/Private/SAnimEditorBase.cpp
--- a/Engine/Source/Editor/Persona/Private/SAnimEditorBase.cpp
+++ b/Engine/Source/Editor/Persona/Private/SAnimEditorBase.cpp
@@ -134,47 +134,6 @@ FText SAnimEditorBase::GetEditorObjectName() const
 	}
 }
 
-void SAnimEditorBase::RecalculateSequenceLength()
-{
-	// Remove Gaps and update Montage Sequence Length
-	if(UAnimCompositeBase* Composite = Cast<UAnimCompositeBase>(GetEditorObject()))
-	{
-		Composite->InvalidateRecursiveAsset();
-
-		float NewSequenceLength = CalculateSequenceLengthOfEditorObject();
-		if (NewSequenceLength != GetSequenceLength())
-		{
-			ClampToEndTime(NewSequenceLength);
-
-			Composite->SetSequenceLength(NewSequenceLength);
-
-			// Reset view if we changed length (note: has to be done after ->SetSequenceLength)!
-			SetInputViewRange(0.f, NewSequenceLength);
-
-			UAnimSingleNodeInstance * PreviewInstance = GetPreviewInstance();
-			if (PreviewInstance)
-			{
-				// Re-set the position, so instance is clamped properly
-				PreviewInstance->SetPosition(PreviewInstance->GetCurrentTime(), false); 
-			}
-		}
-	}
-
-	if(UAnimSequenceBase* Sequence = Cast<UAnimSequenceBase>(GetEditorObject()))
-	{
-		Sequence->ClampNotifiesAtEndOfSequence();
-	}
-}
-
-bool SAnimEditorBase::ClampToEndTime(float NewEndTime)
-{
-	float SequenceLength = GetSequenceLength();
-
-	//if we had a valid sequence length before and our new end time is shorter
-	//then we need to clamp.
-	return (SequenceLength > 0.f && NewEndTime < SequenceLength);
-}
-
 void SAnimEditorBase::OnSelectionChanged(const TArray<UObject*>& SelectedItems)
 {
 	OnObjectsSelected.ExecuteIfBound(SelectedItems);
@@ -205,69 +164,6 @@ void SAnimEditorBase::SetInputViewRange(float InViewMinInput, float InViewMaxInp
 	ViewMinInput = FMath::Max<float>(InViewMinInput, 0.f);
 }
 
-FText SAnimEditorBase::GetCurrentSequenceTime() const
-{
-	UAnimSingleNodeInstance * PreviewInstance = GetPreviewInstance();
-	float CurTime = 0.f;
-	float TotalTime = GetSequenceLength();
-
-	if (PreviewInstance)
-	{
-		CurTime = PreviewInstance->GetCurrentTime();
-	}
-
-	static const FNumberFormattingOptions FractionNumberFormat = FNumberFormattingOptions()
-		.SetMinimumFractionalDigits(3)
-		.SetMaximumFractionalDigits(3);
-	return FText::Format(LOCTEXT("FractionSecondsFmt", "{0} / {1} (second(s))"), FText::AsNumber(CurTime, &FractionNumberFormat), FText::AsNumber(TotalTime, &FractionNumberFormat));
-}
-
-float SAnimEditorBase::GetPercentageInternal() const
-{
-	UAnimSingleNodeInstance * PreviewInstance = GetPreviewInstance();
-	float Percentage = 0.f;
-	if (PreviewInstance)
-	{
-		float SequenceLength = GetSequenceLength();
-		if (SequenceLength > 0.f)
-		{
-			Percentage = PreviewInstance->GetCurrentTime() / SequenceLength;
-		}
-		else
-		{
-			Percentage = 0.f;
-		}
-	}
-
-	return Percentage;
-}
-
-FText SAnimEditorBase::GetCurrentPercentage() const
-{
-	float Percentage = GetPercentageInternal();
-
-	static const FNumberFormattingOptions PercentNumberFormat = FNumberFormattingOptions()
-		.SetMinimumFractionalDigits(2)
-		.SetMaximumFractionalDigits(2);
-	return FText::AsPercent(Percentage, &PercentNumberFormat);
-}
-
-FText SAnimEditorBase::GetCurrentFrame() const
-{
-	float Percentage = GetPercentageInternal();
-	float LastFrame = 0;
-	
-	if (UAnimSequenceBase* AnimSeqBase = Cast<UAnimSequenceBase>(GetEditorObject()))
-	{
-		LastFrame = FMath::Max(AnimSeqBase->GetNumberOfFrames() - 1, 0);
-	}
-
-	static const FNumberFormattingOptions FractionNumberFormat = FNumberFormattingOptions()
-		.SetMinimumFractionalDigits(2)
-		.SetMaximumFractionalDigits(2);
-	return FText::Format(LOCTEXT("FractionKeysFmt", "{0} / {1} Frame"), FText::AsNumber(LastFrame * Percentage, &FractionNumberFormat), FText::AsNumber((int32)LastFrame));
-}
-
 float SAnimEditorBase::GetSequenceLength() const
 {
 	if (UAnimSequenceBase* AnimSeqBase = Cast<UAnimSequenceBase>(GetEditorObject()))
diff --git a/Engine/Source/Editor/Persona/Private/SAnimEditorBase.h b/Engine/Source/Editor/Persona/Private/SAnimEditorBase.h
--- a/Engine/Source/Editor/Persona/Private/SAnimEditorBase.h
+++ b/Engine/Source/Editor/Persona/Private/SAnimEditorBase.h
@@ -111,26 +111,6 @@ protected:
 	/** To get scrub value, get preview instance **/
 	class UAnimSingleNodeInstance* GetPreviewInstance() const;
 
-	/** For the information section **/
-	FText GetCurrentSequenceTime() const;
-	FText GetCurrentPercentage() const;
-	FText GetCurrentFrame() const;
-
-	/** Recalculate sequence length after modifying */
-	void RecalculateSequenceLength();
-
-	/** Clamps the sequence to the specified length 
-	 *
-	 *  @return		Whether clamping was/is necessary
-	 */
-	virtual bool ClampToEndTime(float NewEndTime);
-
-	/** Calculates the sequence length of the object 
-	 *
-	 *  @return		New sequence length
-	 */
-	virtual float CalculateSequenceLengthOfEditorObject() const {return GetSequenceLength();}
-
 	/** The slate container that the editor panels are placed in */
 	TSharedPtr<SVerticalBox>	 EditorPanels;
 
diff --git a/Engine/Source/Editor/Persona/Private/SAnimMontagePanel.cpp b/Engine/Source/Editor/Persona/Private/SAnimMontagePanel.cpp
--- a/Engine/Source/Editor/Persona/Private/SAnimMontagePanel.cpp
+++ b/Engine/Source/Editor/Persona/Private/SAnimMontagePanel.cpp
@@ -651,7 +651,7 @@ void SAnimMontagePanel::Update()
 						.OnAnimReplaceMapping(this, &SAnimMontagePanel::ReplaceAnimationMapping)
 						.OnDiffFromParentAsset(this, &SAnimMontagePanel::IsDiffererentFromParent)
 						.TrackMaxValue(this, &SAnimMontagePanel::GetSequenceLength)
-						.TrackNumDiscreteValues(Montage->GetNumberOfFrames())
+						.TrackNumDiscreteValues(Montage->GetNumberOfSampledKeys())
 					];
 
 				}
@@ -670,7 +670,7 @@ void SAnimMontagePanel::Update()
 						.bChildAnimMontage(bChildAnimMontage)
 						.OnGetNodeColor_Lambda([NodeColor](const FAnimSegment& InSegment){ return NodeColor; })
 						.TrackMaxValue(this, &SAnimMontagePanel::GetSequenceLength)
-						.TrackNumDiscreteValues(Montage->GetNumberOfFrames())
+						.TrackNumDiscreteValues(Montage->GetNumberOfSampledKeys())
 						.OnAnimSegmentNodeClicked(this, &SAnimMontagePanel::ShowSegmentInDetailsView, SlotAnimIdx)
 						.OnPreAnimUpdate(this, &SAnimMontagePanel::PreAnimUpdate)
 						.OnPostAnimUpdate(this, &SAnimMontagePanel::PostAnimUpdate)
diff --git a/Engine/Source/Editor/Persona/Private/SAnimNotifyPanel.cpp b/Engine/Source/Editor/Persona/Private/SAnimNotifyPanel.cpp
--- a/Engine/Source/Editor/Persona/Private/SAnimNotifyPanel.cpp
+++ b/Engine/Source/Editor/Persona/Private/SAnimNotifyPanel.cpp
@@ -2855,7 +2855,7 @@ TSharedPtr<SWidget> SAnimNotifyTrack::SummonContextMenu(const FGeometry& MyGeome
 							SNew(SNumericEntryBox<int32>)
 							.Font(FEditorStyle::GetFontStyle(TEXT("MenuItem.Font")))
 							.MinValue(0)
-							.MaxValue(Sequence->GetNumberOfFrames())
+							.MaxValue(Sequence->GetNumberOfSampledKeys())
 							.Value(Sequence->GetFrameAtTime(NodeObject->GetTime()))
 							.AllowSpin(false)						
 							.OnValueCommitted_Lambda([this, NodeIndex](int32 InValue, ETextCommit::Type InCommitType)
@@ -2971,7 +2971,7 @@ TSharedPtr<SWidget> SAnimNotifyTrack::SummonContextMenu(const FGeometry& MyGeome
 									.Font(FEditorStyle::GetFontStyle(TEXT("MenuItem.Font")))
 									.MinValue(1)
 									.MinSliderValue(1)
-									.MaxSliderValue(Sequence->GetNumberOfFrames())
+									.MaxSliderValue(Sequence->GetNumberOfSampledKeys())
 									.Value(Sequence->GetFrameAtTime(NotifyEvent->GetDuration()))
 									.AllowSpin(false)						
 									.OnValueCommitted_Lambda([this, NotifyIndex](int32 InValue, ETextCommit::Type InCommitType)
diff --git a/Engine/Source/Editor/Persona/Private/SAnimTimingPanel.cpp b/Engine/Source/Editor/Persona/Private/SAnimTimingPanel.cpp
--- a/Engine/Source/Editor/Persona/Private/SAnimTimingPanel.cpp
+++ b/Engine/Source/Editor/Persona/Private/SAnimTimingPanel.cpp
@@ -264,7 +264,7 @@ void SAnimTimingPanel::Update()
 			.ViewInputMax(ViewInputMax)
 			.TrackMinValue(InputMin)
 			.TrackMaxValue(InputMax)
-			.TrackNumDiscreteValues(AnimSequence->GetNumberOfFrames())
+			.TrackNumDiscreteValues(AnimSequence->GetNumberOfSampledKeys())
 		];
 
 	RefreshTrackNodes();
diff --git a/Engine/Source/Editor/Persona/Private/SAnimationScrubPanel.cpp b/Engine/Source/Editor/Persona/Private/SAnimationScrubPanel.cpp
--- a/Engine/Source/Editor/Persona/Private/SAnimationScrubPanel.cpp
+++ b/Engine/Source/Editor/Persona/Private/SAnimationScrubPanel.cpp
@@ -34,7 +34,7 @@ void SAnimationScrubPanel::Construct( const SAnimationScrubPanel::FArguments& In
 			SAssignNew(ScrubControlPanel, SScrubControlPanel)
 			.IsEnabled(true)//this, &SAnimationScrubPanel::DoesSyncViewport)
 			.Value(this, &SAnimationScrubPanel::GetScrubValue)
-			.NumOfKeys(this, &SAnimationScrubPanel::GetNumOfFrames)
+			.NumOfKeys(this, &SAnimationScrubPanel::GetNumberOfKeys)
 			.SequenceLength(this, &SAnimationScrubPanel::GetSequenceLength)
 			.DisplayDrag(this, &SAnimationScrubPanel::GetDisplayDrag)
 			.OnValueChanged(this, &SAnimationScrubPanel::OnValueChanged)
@@ -305,31 +305,31 @@ void SAnimationScrubPanel::OnEndSliderMovement(float NewValue)
 	bSliderBeingDragged = false;
 }
 
-uint32 SAnimationScrubPanel::GetNumOfFrames() const
+uint32 SAnimationScrubPanel::GetNumberOfKeys() const
 {
 	if (DoesSyncViewport())
 	{
 		UAnimSingleNodeInstance* PreviewInstance = GetPreviewInstance();
 		float Length = PreviewInstance->GetLength();
 		// if anim sequence, use correct num frames
-		int32 NumFrames = (int32) (Length/0.0333f); 
+		int32 NumKeys = (int32) (Length/0.0333f); 
 		if (PreviewInstance->GetCurrentAsset())
 		{
 			if (PreviewInstance->GetCurrentAsset()->IsA(UAnimSequenceBase::StaticClass()))
 			{
-				NumFrames = CastChecked<UAnimSequenceBase>(PreviewInstance->GetCurrentAsset())->GetNumberOfFrames();
+				NumKeys = CastChecked<UAnimSequenceBase>(PreviewInstance->GetCurrentAsset())->GetNumberOfSampledKeys();
 			}
 			else if(PreviewInstance->GetCurrentAsset()->IsA(UBlendSpaceBase::StaticClass()))
 			{
 				// Blendspaces dont display frame notches, so just return 0 here
-				NumFrames = 0;
+				NumKeys = 0;
 			}
 		}
-		return NumFrames;
+		return NumKeys;
 	}
 	else if (LockedSequence)
 	{
-		return LockedSequence->GetNumberOfFrames();
+		return LockedSequence->GetNumberOfSampledKeys();
 	}
 	else
 	{
@@ -503,7 +503,7 @@ void SAnimationScrubPanel::OnAppendAnimSequence( bool bFromStart, int32 NumOfFra
 			AnimSequence->Modify();
 
 			// Crop the raw anim data.
-			int32 StartFrame = (bFromStart)? 0 : AnimSequence->GetRawNumberOfFrames() - 1;
+			int32 StartFrame = (bFromStart)? 0 : AnimSequence->GetNumberOfSampledKeys() - 1;
 			int32 EndFrame = StartFrame + NumOfFrames;
 			int32 CopyFrame = StartFrame;
 			AnimSequence->InsertFramesToRawAnimData(StartFrame, EndFrame, CopyFrame);
diff --git a/Engine/Source/Editor/Persona/Private/SAnimationScrubPanel.h b/Engine/Source/Editor/Persona/Private/SAnimationScrubPanel.h
--- a/Engine/Source/Editor/Persona/Private/SAnimationScrubPanel.h
+++ b/Engine/Source/Editor/Persona/Private/SAnimationScrubPanel.h
@@ -99,7 +99,7 @@ protected:
 
 	/** Do I need to sync with viewport? **/
 	bool DoesSyncViewport() const;
-	uint32 GetNumOfFrames() const;
+	uint32 GetNumberOfKeys() const;
 	float GetSequenceLength() const;
 
 	// Returns a UAnimInstance that came from a blueprint, or NULL (even if the UAnimInstance is not null, but it didn't come from a blueprint)
diff --git a/Engine/Source/Editor/SequenceRecorder/Private/AnimationRecorder.cpp b/Engine/Source/Editor/SequenceRecorder/Private/AnimationRecorder.cpp
--- a/Engine/Source/Editor/SequenceRecorder/Private/AnimationRecorder.cpp
+++ b/Engine/Source/Editor/SequenceRecorder/Private/AnimationRecorder.cpp
@@ -225,7 +225,7 @@ void FAnimationRecorder::StartRecord(USkeletalMeshComponent* Component, UAnimSeq
 
 	LastFrame = 0;
 	AnimationObject->SetSequenceLength(0.f);
-	AnimationObject->SetRawNumberOfFrame(0);
+	AnimationObject->SetNumberOfSampledKeys(0);
 
 	RecordedCurves.Reset();
 	RecordedTimes.Empty();
@@ -289,11 +289,11 @@ UAnimSequence* FAnimationRecorder::StopRecord(bool bShowMessage)
 
 	if (AnimationObject)
 	{
-		int32 NumFrames = LastFrame  + 1;
-		AnimationObject->SetRawNumberOfFrame(NumFrames);
+		int32 NumKeys = LastFrame  + 1;
+		AnimationObject->SetNumberOfSampledKeys(NumKeys);
 
 		// can't use TimePassed. That is just total time that has been passed, not necessarily match with frame count
-		AnimationObject->SetSequenceLength((NumFrames>1) ? (NumFrames-1) * IntervalTime : MINIMUM_ANIMATION_LENGTH);
+		AnimationObject->SetSequenceLength((NumKeys>1) ? (NumKeys-1) * IntervalTime : MINIMUM_ANIMATION_LENGTH);
 
 		FixupNotifies();
 
@@ -301,7 +301,7 @@ UAnimSequence* FAnimationRecorder::StopRecord(bool bShowMessage)
 		// @todo figure out why removing redundant keys is inconsistent
 
 		// add to real curve data 
-		if (RecordedCurves.Num() == NumFrames && UIDToArrayIndexLUT)
+		if (RecordedCurves.Num() == NumKeys && UIDToArrayIndexLUT)
 		{
 			StartTime = FPlatformTime::Seconds();
 
@@ -316,16 +316,16 @@ UAnimSequence* FAnimationRecorder::StopRecord(bool bShowMessage)
 
 					TArray<float> TimesToRecord;
 					TArray<float> ValuesToRecord;
-					TimesToRecord.SetNum(NumFrames);
-					ValuesToRecord.SetNum(NumFrames);
+					TimesToRecord.SetNum(NumKeys);
+					ValuesToRecord.SetNum(NumKeys);
 
 					bool bSeenThisCurve = false;
-					for (int32 FrameIndex = 0; FrameIndex < NumFrames; ++FrameIndex)
+					for (int32 KeyIndex = 0; KeyIndex < NumKeys; ++KeyIndex)
 					{
-						const float TimeToRecord = FrameIndex*IntervalTime;
-						if(RecordedCurves[FrameIndex].ValidCurveWeights[CurveIndex])
+						const float TimeToRecord = KeyIndex*IntervalTime;
+						if(RecordedCurves[KeyIndex].ValidCurveWeights[CurveIndex])
 						{
-							float CurCurveValue = RecordedCurves[FrameIndex].CurveWeights[CurveIndex];
+							float CurCurveValue = RecordedCurves[KeyIndex].CurveWeights[CurveIndex];
 							if (!bSeenThisCurve)
 							{
 								bSeenThisCurve = true;
@@ -342,8 +342,8 @@ UAnimSequence* FAnimationRecorder::StopRecord(bool bShowMessage)
 
 							if (FloatCurveData)
 							{
-								TimesToRecord[FrameIndex] = TimeToRecord;
-								ValuesToRecord[FrameIndex] = CurCurveValue;
+								TimesToRecord[KeyIndex] = TimeToRecord;
+								ValuesToRecord[KeyIndex] = CurCurveValue;
 							}
 						}
 					}
@@ -394,9 +394,9 @@ UAnimSequence* FAnimationRecorder::StopRecord(bool bShowMessage)
 		// notify to user
 		if (bShowMessage)
 		{
-			const FText NotificationText = FText::Format(LOCTEXT("RecordAnimation", "'{0}' has been successfully recorded [{1} frames : {2} sec(s) @ {3} Hz]"),
+			const FText NotificationText = FText::Format(LOCTEXT("RecordAnimation", "'{0}' has been successfully recorded [{1} keys : {2} sec(s) @ {3} Hz]"),
 				FText::FromString(AnimationObject->GetName()),
-				FText::AsNumber(AnimationObject->GetRawNumberOfFrames()),
+				FText::AsNumber(AnimationObject->GetNumberOfSampledKeys()),
 				FText::AsNumber(AnimationObject->GetPlayLength()),
 				FText::AsNumber(1.f / IntervalTime)
 				);
diff --git a/Engine/Source/Editor/UnrealEd/Private/Commandlets/PackageUtilities.cpp b/Engine/Source/Editor/UnrealEd/Private/Commandlets/PackageUtilities.cpp
--- a/Engine/Source/Editor/UnrealEd/Private/Commandlets/PackageUtilities.cpp
+++ b/Engine/Source/Editor/UnrealEd/Private/Commandlets/PackageUtilities.cpp
@@ -1839,9 +1839,9 @@ struct CompressAnimationsFunctor
 				continue;
 			}
 
-			if( !bForceCompression && bSkipLongAnimations && (AnimSeq->GetRawNumberOfFrames() > 300) )
+			if( !bForceCompression && bSkipLongAnimations && (AnimSeq->GetNumberOfSampledKeys() > 300) )
 			{
-				UE_LOG(LogPackageUtilities, Warning, TEXT("Animation (%s) has more than 300 frames (%i frames) and SKIPLONGANIMS switch is set. Skipping."), *AnimSeq->GetName(), AnimSeq->GetRawNumberOfFrames());
+				UE_LOG(LogPackageUtilities, Warning, TEXT("Animation (%s) has more than 300 keys (%i keys) and SKIPLONGANIMS switch is set. Skipping."), *AnimSeq->GetName(), AnimSeq->GetNumberOfSampledKeys());
 				continue;
 			}
 
diff --git a/Engine/Source/Editor/UnrealEd/Private/Factories/AnimSequenceFactory.cpp b/Engine/Source/Editor/UnrealEd/Private/Factories/AnimSequenceFactory.cpp
--- a/Engine/Source/Editor/UnrealEd/Private/Factories/AnimSequenceFactory.cpp
+++ b/Engine/Source/Editor/UnrealEd/Private/Factories/AnimSequenceFactory.cpp
@@ -71,7 +71,7 @@ UObject* UAnimSequenceFactory::FactoryCreateNew(UClass* Class, UObject* InParent
 
 	// @todo I think this will crash, we should support differentoptions
 	AnimSequence->SetSequenceLength(0.f);
-	AnimSequence->SetRawNumberOfFrame(0);
+	AnimSequence->SetNumberOfSampledKeys(0);
 	
 	if (TargetSkeleton)
 	{
diff --git a/Engine/Source/Editor/UnrealEd/Private/Fbx/FbxAnimationExport.cpp b/Engine/Source/Editor/UnrealEd/Private/Fbx/FbxAnimationExport.cpp
--- a/Engine/Source/Editor/UnrealEd/Private/Fbx/FbxAnimationExport.cpp
+++ b/Engine/Source/Editor/UnrealEd/Private/Fbx/FbxAnimationExport.cpp
@@ -32,7 +32,7 @@ namespace UnFbx
 			return false;
 		}
 
-		const float FrameRate = FMath::TruncToFloat(((AnimSeq->GetRawNumberOfFrames() - 1) / AnimSeq->GetPlayLength()) + 0.5f);
+		const float FrameRate = AnimSeq->GetSamplingFrameRate().AsDecimal();
 		//Configure the scene time line
 		{
 			FbxGlobalSettings& SceneGlobalSettings = Scene->GetGlobalSettings();
@@ -332,8 +332,7 @@ void FFbxExporter::IterateInsideAnimSequence(const UAnimSequence* AnimSeq, float
 {
 	float AnimTime = AnimStartOffset;
 	float AnimEndTime = (AnimSeq->GetPlayLength() - AnimEndOffset);
-	// Subtracts 1 because NumFrames includes an initial pose for 0.0 second
-	double TimePerKey = (AnimSeq->GetPlayLength() / (AnimSeq->GetRawNumberOfFrames() - 1));
+	double TimePerKey = AnimSeq->GetSamplingFrameRate().AsInterval();
 	const float AnimTimeIncrement = TimePerKey * AnimPlayRate;
 	uint32 AnimFrameIndex = 0;
 
diff --git a/Engine/Source/Editor/UnrealEd/Private/SkeletalMeshEdit.cpp b/Engine/Source/Editor/UnrealEd/Private/SkeletalMeshEdit.cpp
--- a/Engine/Source/Editor/UnrealEd/Private/SkeletalMeshEdit.cpp
+++ b/Engine/Source/Editor/UnrealEd/Private/SkeletalMeshEdit.cpp
@@ -1512,7 +1512,7 @@ bool UnFbx::FFbxImporter::ImportAnimation(USkeleton* Skeleton, UAnimSequence * D
 	}
 	else if (CurveAttributeKeyCount > 0)
 	{
-		DestSeq->SetRawNumberOfFrame(CurveAttributeKeyCount);
+		DestSeq->SetNumberOfSampledKeys(CurveAttributeKeyCount);
 		DestSeq->MarkRawDataAsModified();
 	}
 	// compress animation
@@ -1868,7 +1868,7 @@ void UnFbx::FFbxImporter::ImportBoneTracks(USkeleton* Skeleton, FAnimCurveImport
 		OutTotalNumKeys = FMath::Max(OutTotalNumKeys, NumKeysForTrack);
 	}
 
-	DestSeq->SetRawNumberOfFrame(OutTotalNumKeys);
+	DestSeq->SetNumberOfSampledKeys(OutTotalNumKeys);
 
 	DestSeq->MarkRawDataAsModified();
 }
diff --git a/Engine/Source/Runtime/Core/Public/UObject/UE5MainStreamObjectVersion.h b/Engine/Source/Runtime/Core/Public/UObject/UE5MainStreamObjectVersion.h
--- a/Engine/Source/Runtime/Core/Public/UObject/UE5MainStreamObjectVersion.h
+++ b/Engine/Source/Runtime/Core/Public/UObject/UE5MainStreamObjectVersion.h
@@ -37,6 +37,9 @@ struct CORE_API FUE5MainStreamObjectVersion
 		// Serialize DataLayers in WorldPartitionActorDesc
 		WorldPartitionActorDescSerializeDataLayers,
 
+		// Renaming UAnimSequence::NumFrames to NumberOfKeys, as that what is actually contains.
+		RenamingAnimationNumFrames,
+
 		// -----<new versions can be added above this line>-------------------------------------------------
 		VersionPlusOne,
 		LatestVersion = VersionPlusOne - 1
diff --git a/Engine/Source/Runtime/Engine/Classes/Animation/AnimSequence.h b/Engine/Source/Runtime/Engine/Classes/Animation/AnimSequence.h
--- a/Engine/Source/Runtime/Engine/Classes/Animation/AnimSequence.h
+++ b/Engine/Source/Runtime/Engine/Classes/Animation/AnimSequence.h
@@ -214,9 +214,20 @@ class ENGINE_API UAnimSequence : public UAnimSequenceBase
 #endif
 
 protected:
-	/** Number of raw frames in this sequence (not used by engine - just for informational purposes). */
-	UPROPERTY(AssetRegistrySearchable, meta = (DisplayName = "Number of Keys"))
+#if WITH_EDITORONLY_DATA
+	/** Contains the number of keys expected within the individual animation tracks. */
+	UE_DEPRECATED(5.0, "Num Frames is deprecated use NumberOfKeys instead")
+	UPROPERTY()
 	int32 NumFrames;
+#endif
+
+	/** The number of keys expected within the individual (non-uniform) animation tracks. */
+	UPROPERTY(AssetRegistrySearchable)
+	int32 NumberOfKeys;
+
+	/** The frame rate at which the source animation is sampled. */
+	UPROPERTY(AssetRegistrySearchable)
+	FFrameRate SamplingFrameRate;
 
 	/**
 	 * In the future, maybe keeping RawAnimSequenceTrack + TrackMap as one would be good idea to avoid inconsistent array size
@@ -277,8 +288,14 @@ public:
 	FCompressedAnimSequence CompressedData;
 
 	// Accessors for animation frame count
-	int32 GetRawNumberOfFrames() const { return NumFrames; }
-	void SetRawNumberOfFrame(int32 InNumFrames) { NumFrames = InNumFrames; }
+	UE_DEPRECATED(5.0, "GetRawNumberOfFrames has been deprecated, use GetNumberOfSampledKeys instead")
+	int32 GetRawNumberOfFrames() const { return NumberOfKeys; }
+
+	UE_DEPRECATED(5.0, "SetRawNumberOfFrame has been deprecated, use SetNumberOfSampledKeys instead")
+	void SetRawNumberOfFrame(int32 InNumFrames) { SetNumberOfSampledKeys(InNumFrames); }
+
+	/** Update the the number of expected keys in the (non-uniform) animation tracks, including T0 */
+	void SetNumberOfSampledKeys(int32 InNumberOfKeys);
 
 	/** Additive animation type. **/
 	UPROPERTY(EditAnywhere, Category=AdditiveSettings, AssetRegistrySearchable)
@@ -396,7 +413,6 @@ public:
 #if WITH_EDITOR
 	virtual bool GetAllAnimationSequencesReferred(TArray<UAnimationAsset*>& AnimationAssets, bool bRecursive = true) override;
 	virtual void ReplaceReferredAnimations(const TMap<UAnimationAsset*, UAnimationAsset*>& ReplacementMap) override;
-	virtual int32 GetNumberOfFrames() const override { return NumFrames; }
 #endif
 	//~ End UAnimationAsset Interface
 
@@ -405,7 +421,9 @@ public:
 	virtual bool HasRootMotion() const override { return bEnableRootMotion; }
 	virtual void RefreshCacheData() override;
 	virtual EAdditiveAnimationType GetAdditiveAnimType() const override { return AdditiveAnimType; }
-
+	virtual int32 GetNumberOfSampledKeys() const override { return NumberOfKeys; }
+	virtual const FFrameRate& GetSamplingFrameRate() const override { return SamplingFrameRate; }
+	virtual void SetSequenceLength(float NewLength) override;
 	virtual void EvaluateCurveData(FBlendedCurve& OutCurve, float CurrentTime, bool bForceUseRawData = false) const override;
 	virtual float EvaluateCurveData(SmartName::UID_Type CurveUID, float CurrentTime, bool bForceUseRawData = false) const override;
 	virtual bool HasCurveData(SmartName::UID_Type CurveUID, bool bForceUseRawData) const override;
@@ -423,7 +441,11 @@ public:
 	//~ End UAnimSequenceBase Interface
 
 	// Returns the framerate of the animation
-	float GetFrameRate() const { return (float)(FMath::Max(NumFrames - 1, 1)) / (GetPlayLength() > 0.f ? GetPlayLength() : 1.f); }
+	UE_DEPRECATED(5.0, "Use GetSamplingFrameRate instead")
+	float GetFrameRate() const 
+	{ 
+		return (float)(FMath::Max(NumberOfKeys - 1, 1)) / (GetPlayLength() > 0.f ? GetPlayLength() : 1.f); 
+	}
 
 	// Extract Root Motion transform from the animation
 	FTransform ExtractRootMotion(float StartTime, float DeltaTime, bool bAllowLooping) const;
@@ -481,6 +503,10 @@ private:
 #if WITH_EDITORONLY_DATA
 	void UpdateRetargetSourceAsset();
 #endif
+
+	/** Updates the stored sampling frame-rate using the sequence length and number of sampling keys */
+	void UpdateFrameRate();
+	
 	const TArray<FTransform>& GetRetargetTransforms() const;
 	FName GetRetargetTransformsSourceName() const;
 
diff --git a/Engine/Source/Runtime/Engine/Classes/Animation/AnimSequenceBase.h b/Engine/Source/Runtime/Engine/Classes/Animation/AnimSequenceBase.h
--- a/Engine/Source/Runtime/Engine/Classes/Animation/AnimSequenceBase.h
+++ b/Engine/Source/Runtime/Engine/Classes/Animation/AnimSequenceBase.h
@@ -40,7 +40,7 @@ public:
 	UPROPERTY(Category=Length, AssetRegistrySearchable, VisibleAnywhere, BlueprintReadOnly)
 	float SequenceLength;
 
-	void SetSequenceLength(float NewLength) 	
+	virtual void SetSequenceLength(float NewLength) 	
 	{
 		PRAGMA_DISABLE_DEPRECATION_WARNINGS
 		SequenceLength = NewLength; 
@@ -117,9 +117,16 @@ public:
 	virtual const FRawCurveTracks& GetCurveData() const { return RawCurveData; }
 	virtual bool HasCurveData(SmartName::UID_Type CurveUID, bool bForceUseRawData = false) const;
 
-	/** Return Number of Frames **/
+	/** Return Number of Keys **/
+	UE_DEPRECATED(4.19, "Use GetNumberOfSampledKeys instead")
 	virtual int32 GetNumberOfFrames() const;
 
+	/** Return the total number of keys sampled for this animation, including the T0 key **/
+	virtual int32 GetNumberOfSampledKeys() const;
+
+	/** Return rate at which the animation is sampled **/
+	virtual const FFrameRate& GetSamplingFrameRate() const;
+
 #if WITH_EDITOR
 
 	/** Get the frame number for the provided time */
diff --git a/Engine/Source/Runtime/Engine/Classes/Animation/AnimStreamable.h b/Engine/Source/Runtime/Engine/Classes/Animation/AnimStreamable.h
--- a/Engine/Source/Runtime/Engine/Classes/Animation/AnimStreamable.h
+++ b/Engine/Source/Runtime/Engine/Classes/Animation/AnimStreamable.h
@@ -91,9 +91,9 @@ class UAnimStreamable : public UAnimSequenceBase
 	GENERATED_UCLASS_BODY()
 
 public:
-	/** Number of raw frames in this sequence (not used by engine - just for informational purposes). */
-	UPROPERTY(AssetRegistrySearchable, meta = (DisplayName = "Number of Keys"))
-	int32 NumFrames;
+	/** The number of keys expected within the individual animation tracks. */
+	UPROPERTY(AssetRegistrySearchable)
+	int32 NumberOfKeys;
 
 	/** This defines how values between keys are calculated **/
 	UPROPERTY(EditAnywhere, AssetRegistrySearchable, Category = Animation)
@@ -112,6 +112,11 @@ public:
 	UPROPERTY()
 	FGuid RawDataGuid;
 
+	/** Number of raw frames in this sequence (not used by engine - just for informational purposes). */
+	UE_DEPRECATED(5.0, "Num Frames is deprecated use NumberOfKeys instead")
+	UPROPERTY()
+	int32 NumFrames;
+
 	/**
 	 * Raw uncompressed keyframe data.
 	 */
@@ -205,7 +210,7 @@ public:
 	//~ Begin UAnimSequenceBase Interface
 	ENGINE_API virtual void HandleAssetPlayerTickedInternal(FAnimAssetTickContext &Context, const float PreviousTime, const float MoveDelta, const FAnimTickRecord &Instance, struct FAnimNotifyQueue& NotifyQueue) const override;
 	virtual void GetAnimationPose(FAnimationPoseData& OutAnimationPoseData, const FAnimExtractContext& ExtractionContext) const override;
-	virtual int32 GetNumberOfFrames() const override { return NumFrames; }
+	virtual int32 GetNumberOfSampledKeys() const override { return NumberOfKeys; }
 	//~ End UAnimSequenceBase Interface
 
 #if WITH_EDITOR
diff --git a/Engine/Source/Runtime/Engine/Private/Animation/AnimCompress.cpp b/Engine/Source/Runtime/Engine/Private/Animation/AnimCompress.cpp
--- a/Engine/Source/Runtime/Engine/Private/Animation/AnimCompress.cpp
+++ b/Engine/Source/Runtime/Engine/Private/Animation/AnimCompress.cpp
@@ -384,9 +384,9 @@ void UAnimCompress::BitwiseCompressAnimationTracks(
 					PadByteStream(AnimData.CompressedByteStream, 4, AnimationPadSentinel);
 
 					// write the key table
-					const int32 NumFrames = CompressibleAnimData.NumFrames;
-					const int32 LastFrame = NumFrames-1;
-					const size_t FrameSize = NumFrames > 0xff ? sizeof(uint16) : sizeof(uint8);
+					const int32 NumKeys = CompressibleAnimData.NumberOfKeys;
+					const int32 LastFrame = NumKeys-1;
+					const size_t FrameSize = NumKeys > 0xff ? sizeof(uint16) : sizeof(uint8);
 					const float FrameRate = LastFrame / CompressibleAnimData.SequenceLength;
 
 					const int32 TableSize = NumKeysTrans*FrameSize;
@@ -486,9 +486,9 @@ void UAnimCompress::BitwiseCompressAnimationTracks(
 					PadByteStream(AnimData.CompressedByteStream, 4, AnimationPadSentinel);
 
 					// write the key table
-					const int32 NumFrames = CompressibleAnimData.NumFrames;
-					const int32 LastFrame= NumFrames-1;
-					const size_t FrameSize= NumFrames > 0xff ? sizeof(uint16) : sizeof(uint8);
+					const int32 NumKeys = CompressibleAnimData.NumberOfKeys;
+					const int32 LastFrame= NumKeys-1;
+					const size_t FrameSize= NumKeys > 0xff ? sizeof(uint16) : sizeof(uint8);
 					const float FrameRate = LastFrame / CompressibleAnimData.SequenceLength;
 
 					const int32 TableSize = NumKeysRot*FrameSize;
@@ -572,9 +572,9 @@ void UAnimCompress::BitwiseCompressAnimationTracks(
 						PadByteStream(AnimData.CompressedByteStream, 4, AnimationPadSentinel);
 
 						// write the key table
-						const int32 NumFrames = CompressibleAnimData.NumFrames;
-						const int32 LastFrame = NumFrames-1;
-						const size_t FrameSize = NumFrames > 0xff ? sizeof(uint16) : sizeof(uint8);
+						const int32 NumKeys = CompressibleAnimData.NumberOfKeys;
+						const int32 LastFrame = NumKeys-1;
+						const size_t FrameSize = NumKeys > 0xff ? sizeof(uint16) : sizeof(uint8);
 						const float FrameRate = LastFrame / CompressibleAnimData.SequenceLength;
 
 						const int32 TableSize = NumKeysScale*FrameSize;
@@ -623,7 +623,7 @@ void UAnimCompress::BitwiseCompressAnimationTracks(
 bool UAnimCompress::Compress(const FCompressibleAnimData& CompressibleAnimData, FCompressibleAnimDataResult& OutResult)
 {
 	FUECompressedAnimDataMutable AnimDataMutable;
-	AnimDataMutable.CompressedNumberOfFrames = CompressibleAnimData.NumFrames;
+	AnimDataMutable.CompressedNumberOfKeys = CompressibleAnimData.NumberOfKeys;
 
 	OutResult.Codec = this;
 	OutResult.AnimData = TUniquePtr<ICompressedAnimData>(&AnimDataMutable);
diff --git a/Engine/Source/Runtime/Engine/Private/Animation/AnimCompress_PerTrackCompression.cpp b/Engine/Source/Runtime/Engine/Private/Animation/AnimCompress_PerTrackCompression.cpp
--- a/Engine/Source/Runtime/Engine/Private/Animation/AnimCompress_PerTrackCompression.cpp
+++ b/Engine/Source/Runtime/Engine/Private/Animation/AnimCompress_PerTrackCompression.cpp
@@ -31,7 +31,7 @@ struct FPerTrackParams
 {
 	float MaxZeroingThreshold;
 
-	int32 NumFrames;
+	int32 NumberOfKeys;
 
 	float SequenceLength;
 
@@ -804,7 +804,7 @@ protected:
 	{
 		if (bReallyNeedsFrameTable && (CompressedBytes.Num() > 0))
 		{
-			const int32 NumFrames = Params.NumFrames;
+			const int32 NumFrames = Params.NumberOfKeys;
 			const float SequenceLength = Params.SequenceLength;
 			const float FramesPerSecond = (NumFrames - 1) / SequenceLength;
 
@@ -824,7 +824,7 @@ public:
 	FPerTrackCompressor(int32 InCompressionType, const FTranslationTrack& TranslationData, const FPerTrackParams& Params)
 	{
 		Reset();
-		bReallyNeedsFrameTable = Params.bIncludeKeyTable && (TranslationData.PosKeys.Num() > 1) && (TranslationData.PosKeys.Num() < Params.NumFrames);
+		bReallyNeedsFrameTable = Params.bIncludeKeyTable && (TranslationData.PosKeys.Num() > 1) && (TranslationData.PosKeys.Num() < Params.NumberOfKeys);
 
 		switch (InCompressionType)
 		{
@@ -858,7 +858,7 @@ public:
 	FPerTrackCompressor(int32 InCompressionType, const FRotationTrack& RotationData, const FPerTrackParams& Params)
 	{
 		Reset();
-		bReallyNeedsFrameTable = Params.bIncludeKeyTable && (RotationData.RotKeys.Num() > 1) && (RotationData.RotKeys.Num() < Params.NumFrames);
+		bReallyNeedsFrameTable = Params.bIncludeKeyTable && (RotationData.RotKeys.Num() > 1) && (RotationData.RotKeys.Num() < Params.NumberOfKeys);
 
 		switch (InCompressionType)
 		{
@@ -895,7 +895,7 @@ public:
 	FPerTrackCompressor(int32 InCompressionType, const FScaleTrack& ScaleData, const FPerTrackParams& Params)
 	{
 		Reset();
-		bReallyNeedsFrameTable = Params.bIncludeKeyTable && (ScaleData.ScaleKeys.Num() > 1) && (ScaleData.ScaleKeys.Num() < Params.NumFrames);
+		bReallyNeedsFrameTable = Params.bIncludeKeyTable && (ScaleData.ScaleKeys.Num() > 1) && (ScaleData.ScaleKeys.Num() < Params.NumberOfKeys);
 
 		switch (InCompressionType)
 		{
@@ -1024,7 +1024,7 @@ void UAnimCompress_PerTrackCompression::CompressUsingUnderlyingCompressor(
 		}
 		// Compression parameters / thresholds
 		FPerTrackParams Params;
-		Params.NumFrames = CompressibleAnimData.NumFrames;
+		Params.NumberOfKeys = CompressibleAnimData.NumberOfKeys;
 		Params.SequenceLength = CompressibleAnimData.SequenceLength;
 		Params.MaxZeroingThreshold = MaxZeroingThreshold;
 
@@ -1078,7 +1078,7 @@ void UAnimCompress_PerTrackCompression::CompressUsingUnderlyingCompressor(
 		// Start compressing translation using a totally lossless float32x3
 		const FTranslationTrack& TranslationTrack = TranslationData[TrackIndex];
 
-		Params.bIncludeKeyTable = bActuallyFilterLinearKeys && !FAnimationUtils::HasUniformKeySpacing(CompressibleAnimData.NumFrames, TranslationTrack.Times);
+		Params.bIncludeKeyTable = bActuallyFilterLinearKeys && !FAnimationUtils::HasUniformKeySpacing(CompressibleAnimData.NumberOfKeys, TranslationTrack.Times);
 		FPerTrackCompressor BestTranslation(ACF_Float96NoW, TranslationTrack, Params);
 
 		// Try the other translation formats
@@ -1102,7 +1102,7 @@ void UAnimCompress_PerTrackCompression::CompressUsingUnderlyingCompressor(
 		// Start compressing rotation, first using lossless float32x3
 		const FRotationTrack& RotationTrack = RotationData[TrackIndex];
 
-		Params.bIncludeKeyTable = bActuallyFilterLinearKeys && !FAnimationUtils::HasUniformKeySpacing(CompressibleAnimData.NumFrames, RotationTrack.Times);
+		Params.bIncludeKeyTable = bActuallyFilterLinearKeys && !FAnimationUtils::HasUniformKeySpacing(CompressibleAnimData.NumberOfKeys, RotationTrack.Times);
 		FPerTrackCompressor BestRotation(ACF_Float96NoW, RotationTrack, Params);
 
 		//bool bLeaveRotationUncompressed = (RotationTrack.Times.Num() <= 1) && (GHighQualityEmptyTracks != 0);
@@ -1132,7 +1132,7 @@ void UAnimCompress_PerTrackCompression::CompressUsingUnderlyingCompressor(
 		{
 			const FScaleTrack& ScaleTrack = ScaleData[TrackIndex];
 
-			Params.bIncludeKeyTable = bActuallyFilterLinearKeys && !FAnimationUtils::HasUniformKeySpacing(CompressibleAnimData.NumFrames, ScaleTrack.Times);
+			Params.bIncludeKeyTable = bActuallyFilterLinearKeys && !FAnimationUtils::HasUniformKeySpacing(CompressibleAnimData.NumberOfKeys, ScaleTrack.Times);
 			FPerTrackCompressor BestScale(ACF_Float96NoW, ScaleTrack, Params);
 
 			//bool bLeaveScaleUncompressed = (ScaleTrack.Times.Num() <= 1) && (GHighQualityEmptyTracks != 0);
@@ -1457,12 +1457,12 @@ void* UAnimCompress_PerTrackCompression::FilterBeforeMainKeyRemoval(
 	const int32 NumTracks = TranslationData.Num();
 
 	// Downsample the keys if enabled
-	if ((CompressibleAnimData.NumFrames >= MinKeysForResampling) && bResampleAnimation)
+	if ((CompressibleAnimData.NumberOfKeys >= MinKeysForResampling) && bResampleAnimation)
 	{
 		if(CompressibleAnimData.SequenceLength > 0)
 		{
 			//Make sure we aren't going to oversample the original animation
-			const float CurrentFramerate = (CompressibleAnimData.NumFrames - 1) / CompressibleAnimData.SequenceLength;
+			const float CurrentFramerate = (CompressibleAnimData.NumberOfKeys - 1) / CompressibleAnimData.SequenceLength;
 			if (CurrentFramerate > ResampledFramerate)
 			{
 				ResampleKeys(TranslationData, RotationData, ScaleData, 1.0f / ResampledFramerate, 0.0f);
@@ -1557,7 +1557,7 @@ bool UAnimCompress_PerTrackCompression::DoReduction(const FCompressibleAnimData&
 		return false;
 	}
 
-	if (bResampleAnimation && CompressibleAnimData.NumFrames < MinKeysForResampling)
+	if (bResampleAnimation && CompressibleAnimData.NumberOfKeys < MinKeysForResampling)
 	{
 		return false;
 	}
diff --git a/Engine/Source/Runtime/Engine/Private/Animation/AnimCompress_RemoveLinearKeys.cpp b/Engine/Source/Runtime/Engine/Private/Animation/AnimCompress_RemoveLinearKeys.cpp
--- a/Engine/Source/Runtime/Engine/Private/Animation/AnimCompress_RemoveLinearKeys.cpp
+++ b/Engine/Source/Runtime/Engine/Private/Animation/AnimCompress_RemoveLinearKeys.cpp
@@ -309,21 +309,21 @@ void UAnimCompress_RemoveLinearKeys::UpdateWorldBoneTransformTable(
 	TArray<FTransform>& OutputWorldBones)
 {
 	const FBoneData& Bone		= CompressibleAnimData.BoneData[BoneIndex];
-	const int32 NumFrames		= CompressibleAnimData.NumFrames;
+	const int32 NumKeys		= CompressibleAnimData.NumberOfKeys;
 	const float SequenceLength	= CompressibleAnimData.SequenceLength;
-	const int32 FrameStart		= (BoneIndex*NumFrames);
+	const int32 FrameStart		= (BoneIndex*NumKeys);
 	const int32 TrackIndex = FAnimationUtils::GetAnimTrackIndexForSkeletonBone(BoneIndex, CompressibleAnimData.TrackToSkeletonMapTable);
 	
-	check(OutputWorldBones.Num() >= (FrameStart+NumFrames));
+	check(OutputWorldBones.Num() >= (FrameStart+NumKeys));
 
-	const float TimePerFrame = SequenceLength / (float)(NumFrames-1);
+	const float TimePerFrame = SequenceLength / (float)(NumKeys-1);
 
 	if( TrackIndex != INDEX_NONE )
 	{
 		// get the local-space bone transforms using the animation solver
-		for ( int32 FrameIndex = 0; FrameIndex < NumFrames; ++FrameIndex )
+		for ( int32 KeyIndex = 0; KeyIndex < NumKeys; ++KeyIndex )
 		{
-			float Time = (float)FrameIndex * TimePerFrame;
+			float Time = (float)KeyIndex * TimePerFrame;
 			FTransform LocalAtom;
 
 			FAnimationUtils::ExtractTransformFromCompressionData(CompressibleAnimData, OutCompressedData, Time, TrackIndex, UseRaw, LocalAtom);
@@ -333,7 +333,7 @@ void UAnimCompress_RemoveLinearKeys::UpdateWorldBoneTransformTable(
 			// Saw some crashes happening with it, so normalize here. 
 			LocalAtom.NormalizeRotation();
 
-			OutputWorldBones[(BoneIndex*NumFrames) + FrameIndex] = LocalAtom;
+			OutputWorldBones[(BoneIndex*NumKeys) + KeyIndex] = LocalAtom;
 		}
 	}
 	else
@@ -345,9 +345,9 @@ void UAnimCompress_RemoveLinearKeys::UpdateWorldBoneTransformTable(
 		DefaultTransform = LocalAtom;
 
 		// copy the default transformation into the world bone table
-		for ( int32 FrameIndex = 0; FrameIndex < NumFrames; ++FrameIndex )
+		for ( int32 KeyIndex = 0; KeyIndex < NumKeys; ++KeyIndex )
 		{
-			OutputWorldBones[(BoneIndex*NumFrames) + FrameIndex] = DefaultTransform;
+			OutputWorldBones[(BoneIndex*NumKeys) + KeyIndex] = DefaultTransform;
 		}
 	}
 
@@ -356,9 +356,9 @@ void UAnimCompress_RemoveLinearKeys::UpdateWorldBoneTransformTable(
 	if (ParentIndex != INDEX_NONE)
 	{
 		check (ParentIndex < BoneIndex);
-		for ( int32 FrameIndex = 0; FrameIndex < NumFrames; ++FrameIndex )
+		for ( int32 KeyIndex = 0; KeyIndex < NumKeys; ++KeyIndex )
 		{
-			OutputWorldBones[(BoneIndex*NumFrames) + FrameIndex] = OutputWorldBones[(BoneIndex*NumFrames) + FrameIndex] * OutputWorldBones[(ParentIndex*NumFrames) + FrameIndex];
+			OutputWorldBones[(BoneIndex*NumKeys) + KeyIndex] = OutputWorldBones[(BoneIndex*NumKeys) + KeyIndex] * OutputWorldBones[(ParentIndex*NumKeys) + KeyIndex];
 		}
 	}
 }
@@ -590,9 +590,9 @@ void UAnimCompress_RemoveLinearKeys::ProcessAnimationTracks(
 {
 	// extract all the data we'll need about the skeleton and animation sequence
 	const int32 NumBones			= CompressibleAnimData.BoneData.Num();
-	const int32 NumFrames			= CompressibleAnimData.NumFrames;
+	const int32 NumKeys			= CompressibleAnimData.NumberOfKeys;
 	const float SequenceLength	= CompressibleAnimData.SequenceLength;
-	const int32 LastFrame = NumFrames-1;
+	const int32 LastFrame = NumKeys-1;
 	const float FrameRate = (float)(LastFrame) / SequenceLength;
 	const float TimePerFrame = SequenceLength / (float)(LastFrame);
 
@@ -605,10 +605,10 @@ void UAnimCompress_RemoveLinearKeys::ProcessAnimationTracks(
 	// generate the raw and compressed skeleton in world-space
 	TArray<FTransform> RawWorldBones;
 	TArray<FTransform> NewWorldBones;
-	RawWorldBones.Empty(NumBones * NumFrames);
-	NewWorldBones.Empty(NumBones * NumFrames);
-	RawWorldBones.AddZeroed(NumBones * NumFrames);
-	NewWorldBones.AddZeroed(NumBones * NumFrames);
+	RawWorldBones.Empty(NumBones * NumKeys);
+	NewWorldBones.Empty(NumBones * NumKeys);
+	RawWorldBones.AddZeroed(NumBones * NumKeys);
+	NewWorldBones.AddZeroed(NumBones * NumKeys);
 
 	// generate an array to hold the indices of our end effectors
 	TArray<int32> EndEffectors;
@@ -728,8 +728,8 @@ void UAnimCompress_RemoveLinearKeys::ProcessAnimationTracks(
 						FVector& Key= ScaleTrack.ScaleKeys[KeyIndex];
 
 						const int32 FrameIndex= FMath::Clamp(KeyIndex, 0, LastFrame);
-						const FTransform& NewWorldParent = NewWorldBones[(ParentBoneIndex*NumFrames) + FrameIndex];
-						const FTransform& RawWorldChild = RawWorldBones[(BoneIndex*NumFrames) + FrameIndex];
+						const FTransform& NewWorldParent = NewWorldBones[(ParentBoneIndex*NumKeys) + FrameIndex];
+						const FTransform& RawWorldChild = RawWorldBones[(BoneIndex*NumKeys) + FrameIndex];
 						const FTransform& RelTM = (RawWorldChild.GetRelativeTransform(NewWorldParent));
 						const FTransform Delta = FTransform(RelTM);
 
@@ -747,8 +747,8 @@ void UAnimCompress_RemoveLinearKeys::ProcessAnimationTracks(
 
 							check(ParentBoneIndex != INDEX_NONE);
 							const int32 FrameIndex = FMath::Clamp(KeyIndex, 0, LastFrame);
-							FTransform NewWorldParent = NewWorldBones[(ParentBoneIndex*NumFrames) + FrameIndex];
-							FTransform RawWorldChild = RawWorldBones[(BoneIndex*NumFrames) + FrameIndex];
+							FTransform NewWorldParent = NewWorldBones[(ParentBoneIndex*NumKeys) + FrameIndex];
+							FTransform RawWorldChild = RawWorldBones[(BoneIndex*NumKeys) + FrameIndex];
 							const FTransform& RelTM = (RawWorldChild.GetRelativeTransform(NewWorldParent)); 
 							FQuat Rot = FTransform(RelTM).GetRotation();
 
@@ -778,10 +778,10 @@ void UAnimCompress_RemoveLinearKeys::ProcessAnimationTracks(
 
 							const int32 FrameIndex = FMath::Clamp(KeyIndex, 0, LastFrame);
 
-							const FTransform& NewWorldTransform = NewWorldBones[(BoneIndex*NumFrames) + FrameIndex];
+							const FTransform& NewWorldTransform = NewWorldBones[(BoneIndex*NumKeys) + FrameIndex];
 
-							const FTransform& DesiredChildTransform = RawWorldBones[(FurthestTargetBoneIndex*NumFrames) + FrameIndex].GetRelativeTransform(NewWorldTransform);
-							const FTransform& CurrentChildTransform = NewWorldBones[(FurthestTargetBoneIndex*NumFrames) + FrameIndex].GetRelativeTransform(NewWorldTransform);
+							const FTransform& DesiredChildTransform = RawWorldBones[(FurthestTargetBoneIndex*NumKeys) + FrameIndex].GetRelativeTransform(NewWorldTransform);
+							const FTransform& CurrentChildTransform = NewWorldBones[(FurthestTargetBoneIndex*NumKeys) + FrameIndex].GetRelativeTransform(NewWorldTransform);
 
 							// find the two vectors which represent the angular error we are trying to correct
 							const FVector& CurrentHeading = CurrentChildTransform.GetTranslation();
@@ -835,8 +835,8 @@ void UAnimCompress_RemoveLinearKeys::ProcessAnimationTracks(
 						FVector& Key= TransTrack.PosKeys[KeyIndex];
 
 						const int32 FrameIndex= FMath::Clamp(KeyIndex, 0, LastFrame);
-						FTransform NewWorldParent = NewWorldBones[(ParentBoneIndex*NumFrames) + FrameIndex];
-						FTransform RawWorldChild = RawWorldBones[(BoneIndex*NumFrames) + FrameIndex];
+						FTransform NewWorldParent = NewWorldBones[(ParentBoneIndex*NumKeys) + FrameIndex];
+						FTransform RawWorldChild = RawWorldBones[(BoneIndex*NumKeys) + FrameIndex];
 						const FTransform& RelTM = RawWorldChild.GetRelativeTransform(NewWorldParent);
 						const FTransform Delta = FTransform(RelTM);
 						ensure (!Delta.ContainsNaN());
@@ -873,7 +873,7 @@ void UAnimCompress_RemoveLinearKeys::ProcessAnimationTracks(
 				NewWorldBones);
 			
 			// rebuild the BoneAtoms table using the current set of keys
-			UpdateBoneAtomList(CompressibleAnimData, OutCompressedData, BoneIndex, TrackIndex, NumFrames, TimePerFrame, BoneAtoms);
+			UpdateBoneAtomList(CompressibleAnimData, OutCompressedData, BoneIndex, TrackIndex, NumKeys, TimePerFrame, BoneAtoms);
 
 			// determine the EndEffectorTolerance. 
 			// We use the Maximum value by default, and the Minimum value
@@ -915,7 +915,7 @@ void UAnimCompress_RemoveLinearKeys::ProcessAnimationTracks(
 						RawWorldBones,
 						NewWorldBones,
 						TargetBoneIndices,
-						NumFrames,
+						NumKeys,
 						BoneIndex,
 						ParentBoneIndex,
 						GuidanceScale, 
@@ -938,7 +938,7 @@ void UAnimCompress_RemoveLinearKeys::ProcessAnimationTracks(
 						NewWorldBones);
 					
 					// rebuild the BoneAtoms table using the current set of keys
-					UpdateBoneAtomList(CompressibleAnimData, OutCompressedData, BoneIndex, TrackIndex, NumFrames, TimePerFrame, BoneAtoms);
+					UpdateBoneAtomList(CompressibleAnimData, OutCompressedData, BoneIndex, TrackIndex, NumKeys, TimePerFrame, BoneAtoms);
 				}
 
 				// filter out translations we can approximate through interpolation
@@ -950,7 +950,7 @@ void UAnimCompress_RemoveLinearKeys::ProcessAnimationTracks(
 					RawWorldBones,
 					NewWorldBones,
 					TargetBoneIndices,
-					NumFrames,
+					NumKeys,
 					BoneIndex,
 					ParentBoneIndex,
 					GuidanceScale, 
@@ -973,7 +973,7 @@ void UAnimCompress_RemoveLinearKeys::ProcessAnimationTracks(
 					NewWorldBones);
 				
 				// rebuild the BoneAtoms table using the current set of keys
-				UpdateBoneAtomList(CompressibleAnimData, OutCompressedData, BoneIndex, TrackIndex, NumFrames, TimePerFrame, BoneAtoms);
+				UpdateBoneAtomList(CompressibleAnimData, OutCompressedData, BoneIndex, TrackIndex, NumKeys, TimePerFrame, BoneAtoms);
 
 				// filter out rotations we can approximate through interpolation
 				FilterLinearKeysTemplate<RotationAdapter>(
@@ -984,7 +984,7 @@ void UAnimCompress_RemoveLinearKeys::ProcessAnimationTracks(
 					RawWorldBones,
 					NewWorldBones,
 					TargetBoneIndices,
-					NumFrames,
+					NumKeys,
 					BoneIndex,
 					ParentBoneIndex,
 					GuidanceScale, 
diff --git a/Engine/Source/Runtime/Engine/Private/Animation/AnimCompressionTypes.cpp b/Engine/Source/Runtime/Engine/Private/Animation/AnimCompressionTypes.cpp
--- a/Engine/Source/Runtime/Engine/Private/Animation/AnimCompressionTypes.cpp
+++ b/Engine/Source/Runtime/Engine/Private/Animation/AnimCompressionTypes.cpp
@@ -155,7 +155,7 @@ FCompressibleAnimData::FCompressibleAnimData(class UAnimSequence* InSeq, const b
 	, TrackToSkeletonMapTable(InSeq->GetRawTrackToSkeletonMapTable())
 	, Interpolation(InSeq->Interpolation)
 	, SequenceLength(InSeq->GetPlayLength())
-	, NumFrames(InSeq->GetRawNumberOfFrames())
+	, NumberOfKeys(InSeq->GetNumberOfSampledKeys())
 	, bIsValidAdditive(InSeq->IsValidAdditive())
 #if WITH_EDITORONLY_DATA
 	, ErrorThresholdScale(InSeq->CompressionErrorThresholdScale)
@@ -209,43 +209,43 @@ FCompressibleAnimData::FCompressibleAnimData(class UAnimSequence* InSeq, const b
 		const int32 NumTracks = RawAnimationData.Num();
 		
 		// End frame does not count towards "Even framed" calculation
-		const bool bIsEvenFramed = ((NumFrames - 1) % 2) == 0;
+		const bool bIsEvenFramed = ((NumberOfKeys - 1) % 2) == 0;
 
 		//Strip every other frame from tracks
 		if (bIsEvenFramed)
 		{
 			for (FRawAnimSequenceTrack& Track : RawAnimationData)
 			{
-				StripFramesEven(Track.PosKeys, NumFrames);
-				StripFramesEven(Track.RotKeys, NumFrames);
-				StripFramesEven(Track.ScaleKeys, NumFrames);
+				StripFramesEven(Track.PosKeys, NumberOfKeys);
+				StripFramesEven(Track.RotKeys, NumberOfKeys);
+				StripFramesEven(Track.ScaleKeys, NumberOfKeys);
 			}
 
-			const int32 ActualFrames = NumFrames - 1; // strip bookmark end frame
-			NumFrames = (ActualFrames / 2) + 1;
+			const int32 ActualKeys = NumberOfKeys - 1; // strip bookmark end frame
+			NumberOfKeys = (ActualKeys / 2) + 1;
 		}
 		else
 		{
 			for (FRawAnimSequenceTrack& Track : RawAnimationData)
 			{
-				StripFramesOdd(Track.PosKeys, NumFrames);
-				StripFramesOdd(Track.RotKeys, NumFrames);
-				StripFramesOdd(Track.ScaleKeys, NumFrames);
+				StripFramesOdd(Track.PosKeys, NumberOfKeys);
+				StripFramesOdd(Track.RotKeys, NumberOfKeys);
+				StripFramesOdd(Track.ScaleKeys, NumberOfKeys);
 			}
 
-			const int32 ActualFrames = NumFrames;
-			NumFrames = (ActualFrames / 2);
+			const int32 ActualKeys = NumberOfKeys;
+			NumberOfKeys = (ActualKeys / 2);
 		}
 	}
 #endif
 }
 
-FCompressibleAnimData::FCompressibleAnimData(UAnimBoneCompressionSettings* InBoneCompressionSettings, UAnimCurveCompressionSettings* InCurveCompressionSettings, USkeleton* InSkeleton, EAnimInterpolationType InInterpolation, float InSequenceLength, int32 InNumFrames)
+FCompressibleAnimData::FCompressibleAnimData(UAnimBoneCompressionSettings* InBoneCompressionSettings, UAnimCurveCompressionSettings* InCurveCompressionSettings, USkeleton* InSkeleton, EAnimInterpolationType InInterpolation, float InSequenceLength, int32 InNumberOfKeys)
 	: CurveCompressionSettings(InCurveCompressionSettings)
 	, BoneCompressionSettings(InBoneCompressionSettings)
 	, Interpolation(InInterpolation)
 	, SequenceLength(InSequenceLength)
-	, NumFrames(InNumFrames)
+	, NumberOfKeys(InNumberOfKeys)
 	, bIsValidAdditive(false)
 	, ErrorThresholdScale(1.f)
 {
@@ -261,7 +261,7 @@ FCompressibleAnimData::FCompressibleAnimData()
 , BoneCompressionSettings(nullptr)
 , Interpolation((EAnimInterpolationType)0)
 , SequenceLength(0.f)
-, NumFrames(0)
+, NumberOfKeys(0)
 , bIsValidAdditive(false)
 , ErrorThresholdScale(1.f)
 {
@@ -451,7 +451,7 @@ void FUECompressedAnimDataMutable::BuildFinalBuffer(TArray<uint8>& OutCompressed
 
 void ICompressedAnimData::SerializeCompressedData(class FArchive& Ar)
 {
-	Ar << CompressedNumberOfFrames;
+	Ar << CompressedNumberOfKeys;
 
 #if WITH_EDITORONLY_DATA
 	if (!Ar.IsFilterEditorOnly())
diff --git a/Engine/Source/Runtime/Engine/Private/Animation/AnimCurveCompressionCodec_CompressedRichCurve.cpp b/Engine/Source/Runtime/Engine/Private/Animation/AnimCurveCompressionCodec_CompressedRichCurve.cpp
--- a/Engine/Source/Runtime/Engine/Private/Animation/AnimCurveCompressionCodec_CompressedRichCurve.cpp
+++ b/Engine/Source/Runtime/Engine/Private/Animation/AnimCurveCompressionCodec_CompressedRichCurve.cpp
@@ -37,7 +37,7 @@ bool UAnimCurveCompressionCodec_CompressedRichCurve::Compress(const FCompressibl
 	int32 KeyDataOffset = 0;
 	KeyDataOffset += sizeof(FCurveDesc) * NumCurves;
 
-	const FAnimKeyHelper Helper(AnimSeq.SequenceLength, AnimSeq.NumFrames);
+	const FAnimKeyHelper Helper(AnimSeq.SequenceLength, AnimSeq.NumberOfKeys);
 	const float SampleRate = UseAnimSequenceSampleRate ? Helper.KeysPerSecond() : ErrorSampleRate;
 
 	TArray<uint8> KeyData;
diff --git a/Engine/Source/Runtime/Engine/Private/Animation/AnimCurveCompressionCodec_UniformIndexable.cpp b/Engine/Source/Runtime/Engine/Private/Animation/AnimCurveCompressionCodec_UniformIndexable.cpp
--- a/Engine/Source/Runtime/Engine/Private/Animation/AnimCurveCompressionCodec_UniformIndexable.cpp
+++ b/Engine/Source/Runtime/Engine/Private/Animation/AnimCurveCompressionCodec_UniformIndexable.cpp
@@ -86,7 +86,7 @@ bool UAnimCurveCompressionCodec_UniformIndexable::Compress(const FCompressibleAn
 	const int32 NumCurves = AnimSeq.RawCurveData.FloatCurves.Num();
 	const float Duration = AnimSeq.SequenceLength;
 
-	const FAnimKeyHelper Helper(AnimSeq.SequenceLength, AnimSeq.NumFrames);
+	const FAnimKeyHelper Helper(AnimSeq.SequenceLength, AnimSeq.NumberOfKeys);
 	const float SampleRate = Helper.KeysPerSecond();
 	const int32 NumSamples = FMath::RoundToInt(Duration * SampleRate) + 1;
 
diff --git a/Engine/Source/Runtime/Engine/Private/Animation/AnimCurveCompressionCodec_UniformlySampled.cpp b/Engine/Source/Runtime/Engine/Private/Animation/AnimCurveCompressionCodec_UniformlySampled.cpp
--- a/Engine/Source/Runtime/Engine/Private/Animation/AnimCurveCompressionCodec_UniformlySampled.cpp
+++ b/Engine/Source/Runtime/Engine/Private/Animation/AnimCurveCompressionCodec_UniformlySampled.cpp
@@ -23,7 +23,7 @@ bool UAnimCurveCompressionCodec_UniformlySampled::Compress(const FCompressibleAn
 	float SampleRate_;
 	if (UseAnimSequenceSampleRate)
 	{
-		const FAnimKeyHelper Helper(AnimSeq.SequenceLength, AnimSeq.NumFrames);
+		const FAnimKeyHelper Helper(AnimSeq.SequenceLength, AnimSeq.NumberOfKeys);
 		SampleRate_ = Helper.KeysPerSecond();
 		NumSamples = FMath::RoundToInt(Duration * SampleRate_) + 1;
 	}
diff --git a/Engine/Source/Runtime/Engine/Private/Animation/AnimEncoding.cpp b/Engine/Source/Runtime/Engine/Private/Animation/AnimEncoding.cpp
--- a/Engine/Source/Runtime/Engine/Private/Animation/AnimEncoding.cpp
+++ b/Engine/Source/Runtime/Engine/Private/Animation/AnimEncoding.cpp
@@ -362,7 +362,7 @@ void AnimationFormat_GetStats(
 	int32& NumScaleTracksWithOneKey)
 {
 	OverheadSize = CompressedData.CompressedTrackOffsets.Num() * sizeof(int32);
-	const size_t KeyFrameLookupSize = (CompressedData.CompressedNumberOfFrames > 0xFF) ? sizeof(uint16) : sizeof(uint8);
+	const size_t KeyFrameLookupSize = (CompressedData.CompressedNumberOfKeys > 0xFF) ? sizeof(uint16) : sizeof(uint8);
 
 	if (CompressedData.KeyEncodingFormat != AKF_PerTrackCompression)
 	{
diff --git a/Engine/Source/Runtime/Engine/Private/Animation/AnimEncoding_PerTrackCompression.cpp b/Engine/Source/Runtime/Engine/Private/Animation/AnimEncoding_PerTrackCompression.cpp
--- a/Engine/Source/Runtime/Engine/Private/Animation/AnimEncoding_PerTrackCompression.cpp
+++ b/Engine/Source/Runtime/Engine/Private/Animation/AnimEncoding_PerTrackCompression.cpp
@@ -337,7 +337,7 @@ void AEFPerTrackCompressionCodec::ByteSwapOneTrack(FUECompressedAnimData& Compre
 			// Make sure the key->frame table is 4 byte aligned
 			PreservePadding(TrackData, MemoryStream);
 
-			const int32 FrameTableEntrySize = (CompressedData.CompressedNumberOfFrames <= 0xFF) ? sizeof(uint8) : sizeof(uint16);
+			const int32 FrameTableEntrySize = (CompressedData.CompressedNumberOfKeys <= 0xFF) ? sizeof(uint8) : sizeof(uint16);
 			for (int32 i = 0; i < NumKeys; ++i)
 			{
 				AC_UnalignedSwap(MemoryStream, TrackData, FrameTableEntrySize);
@@ -485,7 +485,7 @@ void AEFPerTrackCompressionCodec::GetBoneAtomRotation(
 			else
 			{
 				const uint8* RESTRICT FrameTable = Align(TrackData + FixedBytes + BytesPerKey * NumKeys, 4);
-				Alpha = TimeToIndex(DecompContext.Interpolation, AnimData.CompressedNumberOfFrames, FrameTable, DecompContext.RelativePos, NumKeys, Index0, Index1);
+				Alpha = TimeToIndex(DecompContext.Interpolation, AnimData.CompressedNumberOfKeys, FrameTable, DecompContext.RelativePos, NumKeys, Index0, Index1);
 			}
 		}
 
@@ -576,7 +576,7 @@ void AEFPerTrackCompressionCodec::GetBoneAtomTranslation(
 			else
 			{
 				const uint8* RESTRICT FrameTable = Align(TrackData + FixedBytes + BytesPerKey * NumKeys, 4);
-				Alpha = TimeToIndex(DecompContext.Interpolation, AnimData.CompressedNumberOfFrames, FrameTable, DecompContext.RelativePos, NumKeys, Index0, Index1);
+				Alpha = TimeToIndex(DecompContext.Interpolation, AnimData.CompressedNumberOfKeys, FrameTable, DecompContext.RelativePos, NumKeys, Index0, Index1);
 			}
 		}
 
@@ -662,7 +662,7 @@ void AEFPerTrackCompressionCodec::GetBoneAtomScale(
 			else
 			{
 				const uint8* RESTRICT FrameTable = Align(TrackData + FixedBytes + BytesPerKey * NumKeys, 4);
-				Alpha = TimeToIndex(DecompContext.Interpolation, AnimData.CompressedNumberOfFrames, FrameTable, DecompContext.RelativePos, NumKeys, Index0, Index1);
+				Alpha = TimeToIndex(DecompContext.Interpolation, AnimData.CompressedNumberOfKeys, FrameTable, DecompContext.RelativePos, NumKeys, Index0, Index1);
 			}
 		}
 
@@ -742,7 +742,7 @@ void AEFPerTrackCompressionCodec::GetPoseRotations(
 			(ispc::BoneTrackPair*)&DesiredPairs[0],
 			AnimData.CompressedTrackOffsets.GetData(),
 			AnimData.CompressedByteStream.GetData(),
-			AnimData.CompressedNumberOfFrames,
+			AnimData.CompressedNumberOfKeys,
 			DecompContext.SequenceLength,
 			DecompContext.RelativePos,
 			(uint8)DecompContext.Interpolation,
@@ -790,7 +790,7 @@ void AEFPerTrackCompressionCodec::GetPoseTranslations(
 			(ispc::BoneTrackPair*)&DesiredPairs[0],
 			AnimData.CompressedTrackOffsets.GetData(),
 			AnimData.CompressedByteStream.GetData(),
-			AnimData.CompressedNumberOfFrames,
+			AnimData.CompressedNumberOfKeys,
 			DecompContext.SequenceLength,
 			DecompContext.RelativePos,
 			(uint8)DecompContext.Interpolation,
@@ -844,7 +844,7 @@ void AEFPerTrackCompressionCodec::GetPoseScales(
 			ScaleOffsets.GetData(),
 			StripSize,
 			AnimData.CompressedByteStream.GetData(),
-			AnimData.CompressedNumberOfFrames,
+			AnimData.CompressedNumberOfKeys,
 			DecompContext.SequenceLength,
 			DecompContext.RelativePos,
 			(uint8)DecompContext.Interpolation,
diff --git a/Engine/Source/Runtime/Engine/Private/Animation/AnimEncoding_VariableKeyLerp.cpp b/Engine/Source/Runtime/Engine/Private/Animation/AnimEncoding_VariableKeyLerp.cpp
--- a/Engine/Source/Runtime/Engine/Private/Animation/AnimEncoding_VariableKeyLerp.cpp
+++ b/Engine/Source/Runtime/Engine/Private/Animation/AnimEncoding_VariableKeyLerp.cpp
@@ -35,7 +35,7 @@ void AEFVariableKeyLerpShared::ByteSwapRotationIn(
 		PadMemoryReader(&MemoryReader, TrackData, 4); 
 
 		// swap the track table
-		const size_t EntryStride = (CompressedData.CompressedNumberOfFrames > 0xFF) ? sizeof(uint16) : sizeof(uint8);
+		const size_t EntryStride = (CompressedData.CompressedNumberOfKeys > 0xFF) ? sizeof(uint16) : sizeof(uint8);
 		for (int32 KeyIndex = 0; KeyIndex < NumKeys; ++KeyIndex)
 		{
 			AC_UnalignedSwap(MemoryReader, TrackData, EntryStride);
@@ -67,7 +67,7 @@ void AEFVariableKeyLerpShared::ByteSwapTranslationIn(
 		PadMemoryReader(&MemoryReader, TrackData, 4); 
 
 		// swap the track table
-		const size_t EntryStride = (CompressedData.CompressedNumberOfFrames > 0xFF) ? sizeof(uint16) : sizeof(uint8);
+		const size_t EntryStride = (CompressedData.CompressedNumberOfKeys > 0xFF) ? sizeof(uint16) : sizeof(uint8);
 		for (int32 KeyIndex = 0; KeyIndex < NumKeys; ++KeyIndex)
 		{
 			AC_UnalignedSwap(MemoryReader, TrackData, EntryStride);
@@ -99,7 +99,7 @@ void AEFVariableKeyLerpShared::ByteSwapScaleIn(
 		PadMemoryReader(&MemoryReader, TrackData, 4); 
 
 		// swap the track table
-		const size_t EntryStride = (CompressedData.CompressedNumberOfFrames > 0xFF) ? sizeof(uint16) : sizeof(uint8);
+		const size_t EntryStride = (CompressedData.CompressedNumberOfKeys > 0xFF) ? sizeof(uint16) : sizeof(uint8);
 		for (int32 KeyIndex = 0; KeyIndex < NumKeys; ++KeyIndex)
 		{
 			AC_UnalignedSwap(MemoryReader, TrackData, EntryStride);
@@ -130,7 +130,7 @@ void AEFVariableKeyLerpShared::ByteSwapRotationOut(
 		PadMemoryWriter(&MemoryWriter, TrackData, 4);
 
 		// swap the track table
-		const size_t EntryStride = (CompressedData.CompressedNumberOfFrames > 0xFF) ? sizeof(uint16) : sizeof(uint8);
+		const size_t EntryStride = (CompressedData.CompressedNumberOfKeys > 0xFF) ? sizeof(uint16) : sizeof(uint8);
 		for (int32 KeyIndex = 0; KeyIndex < NumKeys; ++KeyIndex)
 		{
 			AC_UnalignedSwap(MemoryWriter, TrackData, EntryStride);
@@ -162,7 +162,7 @@ void AEFVariableKeyLerpShared::ByteSwapTranslationOut(
 		PadMemoryWriter(&MemoryWriter, TrackData, 4);
 
 		// swap the track table
-		const size_t EntryStride = (CompressedData.CompressedNumberOfFrames > 0xFF) ? sizeof(uint16) : sizeof(uint8);
+		const size_t EntryStride = (CompressedData.CompressedNumberOfKeys > 0xFF) ? sizeof(uint16) : sizeof(uint8);
 		for (int32 KeyIndex = 0; KeyIndex < NumKeys; ++KeyIndex)
 		{
 			AC_UnalignedSwap(MemoryWriter, TrackData, EntryStride);
@@ -195,7 +195,7 @@ void AEFVariableKeyLerpShared::ByteSwapScaleOut(
 		PadMemoryWriter(&MemoryWriter, TrackData, 4);
 
 		// swap the track table
-		const size_t EntryStride = (CompressedData.CompressedNumberOfFrames > 0xFF) ? sizeof(uint16) : sizeof(uint8);
+		const size_t EntryStride = (CompressedData.CompressedNumberOfKeys > 0xFF) ? sizeof(uint16) : sizeof(uint8);
 		for (int32 KeyIndex = 0; KeyIndex < NumKeys; ++KeyIndex)
 		{
 			AC_UnalignedSwap(MemoryWriter, TrackData, EntryStride);
@@ -235,7 +235,7 @@ void AEFVariableKeyLerp<FORMAT>::GetPoseRotations(
 			(ispc::BoneTrackPair*)&DesiredPairs[0],
 			AnimData.CompressedTrackOffsets.GetData(),
 			AnimData.CompressedByteStream.GetData(),
-			AnimData.CompressedNumberOfFrames,
+			AnimData.CompressedNumberOfKeys,
 			DecompContext.RelativePos,
 			(uint8)DecompContext.Interpolation,
 			FORMAT,
@@ -287,7 +287,7 @@ void AEFVariableKeyLerp<FORMAT>::GetPoseTranslations(
 			(ispc::BoneTrackPair*)&DesiredPairs[0],
 			AnimData.CompressedTrackOffsets.GetData(),
 			AnimData.CompressedByteStream.GetData(),
-			AnimData.CompressedNumberOfFrames,
+			AnimData.CompressedNumberOfKeys,
 			DecompContext.RelativePos,
 			(uint8)DecompContext.Interpolation,
 			FORMAT,
@@ -343,7 +343,7 @@ void AEFVariableKeyLerp<FORMAT>::GetPoseScales(
 			ScaleOffsets.GetData(),
 			StripSize,
 			AnimData.CompressedByteStream.GetData(),
-			AnimData.CompressedNumberOfFrames,
+			AnimData.CompressedNumberOfKeys,
 			DecompContext.RelativePos,
 			(uint8)DecompContext.Interpolation,
 			FORMAT,
diff --git a/Engine/Source/Runtime/Engine/Private/Animation/AnimNode_SequencePlayer.cpp b/Engine/Source/Runtime/Engine/Private/Animation/AnimNode_SequencePlayer.cpp
--- a/Engine/Source/Runtime/Engine/Private/Animation/AnimNode_SequencePlayer.cpp
+++ b/Engine/Source/Runtime/Engine/Private/Animation/AnimNode_SequencePlayer.cpp
@@ -83,7 +83,7 @@ void FAnimNode_SequencePlayer::UpdateAssetPlayer(const FAnimationUpdateContext&
 #if ANIM_NODE_IDS_AVAILABLE && WITH_EDITORONLY_DATA
 	if (FAnimBlueprintDebugData* DebugData = Context.AnimInstanceProxy->GetAnimBlueprintDebugData())
 	{
-		DebugData->RecordSequencePlayer(Context.GetCurrentNodeId(), GetAccumulatedTime(), Sequence != nullptr ? Sequence->GetPlayLength() : 0.0f, Sequence != nullptr ? Sequence->GetNumberOfFrames() : 0);
+		DebugData->RecordSequencePlayer(Context.GetCurrentNodeId(), GetAccumulatedTime(), Sequence != nullptr ? Sequence->GetPlayLength() : 0.0f, Sequence != nullptr ? Sequence->GetNumberOfSampledKeys() : 0);
 	}
 #endif
 
diff --git a/Engine/Source/Runtime/Engine/Private/Animation/AnimSequence.cpp b/Engine/Source/Runtime/Engine/Private/Animation/AnimSequence.cpp
--- a/Engine/Source/Runtime/Engine/Private/Animation/AnimSequence.cpp
+++ b/Engine/Source/Runtime/Engine/Private/Animation/AnimSequence.cpp
@@ -49,6 +49,8 @@
 #include "Animation/CustomAttributesRuntime.h"
 #include "Stats/StatsHierarchical.h"
 #include "Animation/AnimationPoseData.h"
+#include "ITimeManagementModule.h"
+#include "CommonFrameRates.h"
 
 #define USE_SLERP 0
 #define LOCTEXT_NAMESPACE "AnimSequence"
@@ -207,7 +209,7 @@ void OnCVarsChanged()
 		int32 NumAnimations = 0;
 		for (const TPair<int32, UAnimSequence*>& Pair : Sizes)
 		{
-			const bool bIsOddFramed = (Pair.Value->GetNumberOfFrames() % 2) == 0;
+			const bool bIsOddFramed = (Pair.Value->GetNumberOfSampledKeys() % 2) == 0;
 			if (bIsOddFramed)
 			{
 				OutputMessage += FString::Printf(TEXT("%s - %.1fK\n"), *Pair.Value->GetPathName(), (float)Pair.Key / 1000.f);
@@ -606,7 +608,7 @@ void UAnimSequence::GetAssetRegistryTags(TArray<FAssetRegistryTag>& OutTags) con
 
 	OutTags.Add(FAssetRegistryTag(TEXT("Compression Ratio"), FString::Printf(TEXT("%.03f"), (float)GetApproxCompressedSize() / (float)GetUncompressedRawSize()), FAssetRegistryTag::TT_Numerical));
 	OutTags.Add(FAssetRegistryTag(TEXT("Compressed Size (KB)"), FString::Printf(TEXT("%.02f"), (float)GetApproxCompressedSize() / 1024.0f), FAssetRegistryTag::TT_Numerical));
-	OutTags.Add(FAssetRegistryTag(TEXT("FrameRate"), FString::Printf(TEXT("%.2f"), GetFrameRate()), FAssetRegistryTag::TT_Numerical));
+	OutTags.Add(FAssetRegistryTag(TEXT("FrameRate"), FString::Printf(TEXT("%.2f"), SamplingFrameRate.AsDecimal()), FAssetRegistryTag::TT_Numerical));
 	Super::GetAssetRegistryTags(OutTags);
 }
 
@@ -621,7 +623,7 @@ void UAnimSequence::AddReferencedObjects(UObject* This, FReferenceCollector& Col
 
 int32 UAnimSequence::GetUncompressedRawSize() const
 {
-	int32 BoneRawSize = ((sizeof(FVector) + sizeof(FQuat) + sizeof(FVector)) * RawAnimationData.Num() * NumFrames);
+	int32 BoneRawSize = ((sizeof(FVector) + sizeof(FQuat) + sizeof(FVector)) * RawAnimationData.Num() * NumberOfKeys);
 	int32 CurveRawSize = 0;
 	for (const FFloatCurve& Curve : RawCurveData.FloatCurves)
 	{
@@ -785,6 +787,46 @@ void UAnimSequence::Serialize(FArchive& Ar)
 	}
 
 #if WITH_EDITORONLY_DATA
+	if (Ar.IsLoading() && Ar.CustomVer(FUE5MainStreamObjectVersion::GUID) < FUE5MainStreamObjectVersion::RenamingAnimationNumFrames)
+	{
+		PRAGMA_DISABLE_DEPRECATION_WARNINGS
+		NumberOfKeys = NumFrames;
+
+		// Validate the actual number of keys that is stored, by looking at the maximum number of keys for any given animation track 
+		int32 MaxNumberOfTrackKeys = 0;
+		for (const FRawAnimSequenceTrack& Track : RawAnimationData)
+		{
+			MaxNumberOfTrackKeys = FMath::Max(MaxNumberOfTrackKeys, Track.PosKeys.Num());
+			MaxNumberOfTrackKeys = FMath::Max(MaxNumberOfTrackKeys, Track.RotKeys.Num());
+			MaxNumberOfTrackKeys = FMath::Max(MaxNumberOfTrackKeys, Track.ScaleKeys.Num());
+		}
+
+		// Test whether or not there are more track keys than the value stored, only check for greater then as uniform tracks will have a single key and identity scaling with result in zero keys
+		if (MaxNumberOfTrackKeys > NumberOfKeys)
+		{
+			UE_LOG(LogAnimation, Warning, TEXT("Animation %s needs resaving - Invalid number of keys %i stored according to maximum number animation data track keys. Setting new number of keys %i."), *GetName(), NumberOfKeys, MaxNumberOfTrackKeys);
+			NumberOfKeys = MaxNumberOfTrackKeys;
+		}
+
+		// Update stored frame rate according to number of keys and play length
+		UpdateFrameRate();
+
+		// In case there is any animation data available (not valid for curve only animations), verify that the new frame-rate matches up with the expected number of frames/keys and vice versa
+		if (RawAnimationData.Num())
+		{
+			const int32 NumberOfFrames = FMath::Max(NumberOfKeys - 1, 1);
+			const float SampledTime = SamplingFrameRate.AsSeconds(NumberOfFrames);
+
+			// Validate that provided the number of keys and frame-rate the sequence length is accurate
+			if (!FMath::IsNearlyEqual(SampledTime, GetPlayLength(), KINDA_SMALL_NUMBER))
+			{
+				UE_LOG(LogAnimation, Warning, TEXT("Animation %s needs resaving - Inaccurate sequence length %5.5f according to number of frames and frame rate (delta of %f). Setting new length %5.5f."), *GetName(), GetPlayLength(), SampledTime - GetPlayLength(), SampledTime);
+				SequenceLength = SampledTime;
+			}
+		}
+		PRAGMA_ENABLE_DEPRECATION_WARNINGS
+	}
+
 	if ( Ar.IsLoading() && Ar.UE4Ver() < VER_UE4_ASSET_IMPORT_DATA_AS_JSON && !AssetImportData)
 	{
 		// AssetImportData should always be valid
@@ -814,6 +856,43 @@ bool UAnimSequence::IsValidToPlay() const
 }
 #endif
 
+void UAnimSequence::SetNumberOfSampledKeys(int32 InNumberOfKeys)
+{
+	NumberOfKeys = InNumberOfKeys;	
+	UpdateFrameRate();
+}
+
+void UAnimSequence::SetSequenceLength(float NewLength)
+{
+	UAnimSequenceBase::SetSequenceLength(NewLength);
+	UpdateFrameRate();
+}
+
+void UAnimSequence::UpdateFrameRate()
+{
+	const int32 NumberOfFrames = FMath::Max(NumberOfKeys - 1, 1);
+
+	// Generate the frame-rate according to the number of frames and sequence length
+	const double DecimalFrameRate = (double)NumberOfFrames / ((double)GetPlayLength() > 0.0 ? (double)GetPlayLength() : 1.0);
+
+	// Account for non-whole number frame rates using large denominator
+	const double Denominator = 1000000.0;
+	SamplingFrameRate = FFrameRate(DecimalFrameRate * Denominator, Denominator);	 
+
+	// Try to simplifiy the frame rate, in case it is a multiple of the commonly used frame rates e.g. 10000/300000 -> 1/30
+	TArrayView<const FCommonFrameRateInfo> CommonFrameRates = FModuleManager::LoadModulePtr<ITimeManagementModule>("TimeManagement")->GetAllCommonFrameRates();
+	for (const FCommonFrameRateInfo& Info : CommonFrameRates)
+	{
+		const bool bDoesNotAlreadyMatch = Info.FrameRate.Denominator != SamplingFrameRate.Denominator && Info.FrameRate.Numerator != SamplingFrameRate.Numerator;
+		
+		if (bDoesNotAlreadyMatch && FMath::IsNearlyEqual(SamplingFrameRate.AsInterval(), Info.FrameRate.AsInterval()))
+		{
+			SamplingFrameRate = Info.FrameRate;
+			break;
+		}
+	}
+}
+
 void UAnimSequence::SortSyncMarkers()
 {
 	// First make sure all SyncMarkers are within a valid range
@@ -946,7 +1025,7 @@ void UAnimSequence::PostLoad()
 	// No animation data is found. Warn - this should check before we check CompressedTrackOffsets size
 	// Otherwise, we'll see empty data set crashing game due to no CompressedTrackOffsets
 	// You can't check RawAnimationData size since it gets removed during cooking
-	if ( NumFrames == 0 && RawCurveData.FloatCurves.Num() == 0 )
+	if ( NumberOfKeys == 0 && RawCurveData.FloatCurves.Num() == 0 )
 	{
 		UE_LOG(LogAnimation, Warning, TEXT("No animation data exists for sequence %s (%s)"), *GetName(), (GetOuter() ? *GetOuter()->GetFullName() : *GetFullName()) );
 #if WITH_EDITOR
@@ -967,7 +1046,7 @@ void UAnimSequence::PostLoad()
 	// @todo need to fix importer/editing feature
 	else if (GetPlayLength() == 0.f )
 	{
-		ensure(NumFrames == 1);
+		ensure(NumberOfKeys == 1);
 		SetSequenceLength(MINIMUM_ANIMATION_LENGTH);
 	}
 	// Raw data exists, but missing compress animation data
@@ -1307,7 +1386,7 @@ void UAnimSequence::ExtractBoneTransform(const struct FRawAnimSequenceTrack& Raw
 
 void UAnimSequence::ExtractBoneTransform(const struct FRawAnimSequenceTrack& RawTrack, FTransform& OutAtom, float Time) const
 {
-	FAnimationUtils::ExtractTransformFromTrack(Time, NumFrames, GetPlayLength(), RawTrack, Interpolation, OutAtom);
+	FAnimationUtils::ExtractTransformFromTrack(Time, NumberOfKeys, GetPlayLength(), RawTrack, Interpolation, OutAtom);
 }
 
 void UAnimSequence::HandleAssetPlayerTickedInternal(FAnimAssetTickContext &Context, const float PreviousTime, const float MoveDelta, const FAnimTickRecord &Instance, struct FAnimNotifyQueue& NotifyQueue) const
@@ -1630,7 +1709,7 @@ void UAnimSequence::GetBonePose(FAnimationPoseData& OutAnimationPoseData, const
 			}
 		}
 
-		BuildPoseFromRawData(AnimationData, TrackToSkeletonMapTable, OutPose, ExtractionContext.CurrentTime, Interpolation, NumFrames, GetPlayLength(), GetRetargetTransformsSourceName(), GetRetargetTransforms());
+		BuildPoseFromRawData(AnimationData, TrackToSkeletonMapTable, OutPose, ExtractionContext.CurrentTime, Interpolation, NumberOfKeys, GetPlayLength(), GetRetargetTransformsSourceName(), GetRetargetTransforms());
 
 		if ((ExtractionContext.bExtractRootMotion && RootMotionReset.bEnableRootMotion) || RootMotionReset.bForceRootLock)
 		{
@@ -1784,7 +1863,7 @@ void UAnimSequence::GetAdditiveBasePose(FAnimationPoseData& OutAnimationPoseData
 		// use animation as a base pose. Need BasePoseSeq and RefFrameIndex (will clamp if outside).
 		case ABPT_AnimFrame:
 		{
-			const float Fraction = (RefPoseSeq->NumFrames > 0) ? FMath::Clamp<float>((float)RefFrameIndex / (float)RefPoseSeq->NumFrames, 0.f, 1.f) : 0.f;
+			const float Fraction = (RefPoseSeq->NumberOfKeys > 0) ? FMath::Clamp<float>((float)RefFrameIndex / (float)RefPoseSeq->NumberOfKeys, 0.f, 1.f) : 0.f;
 			const float BasePoseTime = RefPoseSeq->GetPlayLength() * Fraction;
 
 			FAnimExtractContext BasePoseExtractionContext(ExtractionContext);
@@ -1935,25 +2014,26 @@ static int32 CropRawTrack(FRawAnimSequenceTrack& RawTrack, int32 StartKey, int32
 	return FMath::Max<int32>( RawTrack.PosKeys.Num(), FMath::Max<int32>(RawTrack.RotKeys.Num(), RawTrack.ScaleKeys.Num()) );
 }
 
-void UAnimSequence::ResizeSequence(float NewLength, int32 NewNumFrames, bool bInsert, int32 StartFrame/*inclusive */, int32 EndFrame/*inclusive*/)
+void UAnimSequence::ResizeSequence(float NewLength, int32 NewNumKeys, bool bInsert, int32 StartFrame/*inclusive */, int32 EndFrame/*inclusive*/)
 {
-	check (NewNumFrames > 0);
+	check (NewNumKeys > 0);
 	check (StartFrame < EndFrame);
 
-	int32 OldNumFrames = NumFrames;
+	int32 OldNumberOfKeys = NumberOfKeys;
 	float OldSequenceLength = GetPlayLength();
 
 	// verify condition
-	NumFrames = NewNumFrames;
+	NumberOfKeys = NewNumKeys;
 	// Update sequence length to match new number of frames.
 	SetSequenceLength(NewLength);
 
-	float Interval = OldSequenceLength / OldNumFrames;
-	ensure (Interval == GetPlayLength()/NumFrames);
+	const float OldInterval = OldSequenceLength / (OldNumberOfKeys - 1);
+	const float NewInterval = NewLength / (NewNumKeys - 1);
+	ensure(FMath::IsNearlyEqual(OldInterval, NewInterval));
 
-	float OldStartTime = StartFrame * Interval;
-	float OldEndTime = EndFrame * Interval;
-	float Duration = OldEndTime - OldStartTime;
+	const float OldStartTime = StartFrame * OldInterval;
+	const float OldEndTime = EndFrame * OldInterval;
+	const float Duration = OldEndTime - OldStartTime;
 
 	// re-locate notifies
 	for (auto& Notify: Notifies)
@@ -1989,22 +2069,21 @@ void UAnimSequence::ResizeSequence(float NewLength, int32 NewNumFrames, bool bIn
 			// if state, make sure to adjust end time
 			if(Notify.NotifyStateClass)
 			{
-				float NotifyDuration = Notify.GetDuration();
+				const float NotifyDuration = Notify.GetDuration();
 				float NotifyEnd = CurrentTime + NotifyDuration;
 				NewDuration = NotifyDuration;
+
+				// If Notify is inside of the trimmed time frame, zero out the duration
 				if(NotifyEnd >= OldStartTime && NotifyEnd <= OldEndTime)
 				{
 					// small number @todo see if there is define for this
 					NewDuration = 0.1f;
 				}
-				else if (NotifyEnd > OldEndTime)
+				// If Notify overlaps trimmed time frame, remove trimmed duration
+				else if (CurrentTime < OldEndTime && NotifyEnd > OldEndTime)
 				{
 					NewDuration = NotifyEnd-Duration-CurrentTime;
 				}
-				else
-				{
-					NewDuration = NotifyDuration;
-				}
 
 				NewDuration = FMath::Max(NewDuration, 0.1f);
 			}
@@ -2042,6 +2121,7 @@ void UAnimSequence::ResizeSequence(float NewLength, int32 NewNumFrames, bool bIn
 			// if it's later than start time
 			if (CurrentTime >= OldStartTime)
 			{
+				// Move forwards
 				CurrentTime += Duration;
 			}
 		}
@@ -2049,10 +2129,12 @@ void UAnimSequence::ResizeSequence(float NewLength, int32 NewNumFrames, bool bIn
 		{
 			if (CurrentTime >= OldStartTime && CurrentTime <= OldEndTime)
 			{
+				// Snap to start of the trim/crop
 				CurrentTime = OldStartTime;
 			}
 			else if (CurrentTime > OldEndTime)
 			{
+				// Move backwards relative to crop/trim range end
 				CurrentTime -= Duration;
 			}
 		}
@@ -2065,15 +2147,15 @@ void UAnimSequence::ResizeSequence(float NewLength, int32 NewNumFrames, bool bIn
 bool UAnimSequence::InsertFramesToRawAnimData( int32 StartFrame, int32 EndFrame, int32 CopyFrame)
 {
 	// make sure the copyframe is valid and start frame is valid
-	int32 NumFramesToInsert = EndFrame-StartFrame;
-	if ((CopyFrame>=0 && CopyFrame<NumFrames) && (StartFrame >= 0 && StartFrame <=NumFrames) && NumFramesToInsert > 0)
+	const int32 NumberOfKeysToInsert = EndFrame-StartFrame;
+	if ((CopyFrame>=0 && CopyFrame< NumberOfKeys) && (StartFrame >= 0 && StartFrame <=NumberOfKeys) && NumberOfKeysToInsert > 0)
 	{
 		for (auto& RawData : RawAnimationData)
 		{
 			if (RawData.PosKeys.Num() > 1 && RawData.PosKeys.IsValidIndex(CopyFrame))
 			{
 				auto Source = RawData.PosKeys[CopyFrame];
-				RawData.PosKeys.InsertZeroed(StartFrame, NumFramesToInsert);
+				RawData.PosKeys.InsertZeroed(StartFrame, NumberOfKeysToInsert);
 				for (int32 Index=StartFrame; Index<EndFrame; ++Index)
 				{
 					RawData.PosKeys[Index] = Source;
@@ -2083,7 +2165,7 @@ bool UAnimSequence::InsertFramesToRawAnimData( int32 StartFrame, int32 EndFrame,
 			if(RawData.RotKeys.Num() > 1 && RawData.RotKeys.IsValidIndex(CopyFrame))
 			{
 				auto Source = RawData.RotKeys[CopyFrame];
-				RawData.RotKeys.InsertZeroed(StartFrame, NumFramesToInsert);
+				RawData.RotKeys.InsertZeroed(StartFrame, NumberOfKeysToInsert);
 				for(int32 Index=StartFrame; Index<EndFrame; ++Index)
 				{
 					RawData.RotKeys[Index] = Source;
@@ -2093,7 +2175,7 @@ bool UAnimSequence::InsertFramesToRawAnimData( int32 StartFrame, int32 EndFrame,
 			if(RawData.ScaleKeys.Num() > 1 && RawData.ScaleKeys.IsValidIndex(CopyFrame))
 			{
 				auto Source = RawData.ScaleKeys[CopyFrame];
-				RawData.ScaleKeys.InsertZeroed(StartFrame, NumFramesToInsert);
+				RawData.ScaleKeys.InsertZeroed(StartFrame, NumberOfKeysToInsert);
 
 				for(int32 Index=StartFrame; Index<EndFrame; ++Index)
 				{
@@ -2102,14 +2184,16 @@ bool UAnimSequence::InsertFramesToRawAnimData( int32 StartFrame, int32 EndFrame,
 			}
 		}
 
-		float const FrameTime = GetPlayLength() / ((float)NumFrames);
+		float const FrameTime = SamplingFrameRate.AsInterval();
 
-		int32 NewNumFrames = NumFrames + NumFramesToInsert;
-		ResizeSequence((float)NewNumFrames * FrameTime, NewNumFrames, true, StartFrame, EndFrame);
+		const int32 NewNumKeys = NumberOfKeys + NumberOfKeysToInsert;
+		const int32 NewNumFrames = NewNumKeys - 1;
+		ResizeSequence((float)NewNumFrames * FrameTime, NewNumKeys, true, StartFrame, EndFrame);
 
-		UE_LOG(LogAnimation, Log, TEXT("\tGetPlayLength(): %f, NumFrames: %d"), GetPlayLength(), NumFrames);
+		UE_LOG(LogAnimation, Log, TEXT("\tGetPlayLength(): %f, Number of Keys: %d"), GetPlayLength(), NumberOfKeys);
 		
 		MarkRawDataAsModified();
+		OnRawDataChanged();
 		MarkPackageDirty();
 
 		return true;
@@ -2121,12 +2205,12 @@ bool UAnimSequence::InsertFramesToRawAnimData( int32 StartFrame, int32 EndFrame,
 bool UAnimSequence::CropRawAnimData( float CurrentTime, bool bFromStart )
 {
 	// Length of one frame.
-	float const FrameTime = GetPlayLength() / ((float)NumFrames);
+	float const FrameTime = SamplingFrameRate.AsInterval();
 	// Save Total Number of Frames before crop
-	int32 TotalNumOfFrames = NumFrames;
+	const int32 TotalNumOfKeys = NumberOfKeys;
 
 	// if current frame is 1, do not try crop. There is nothing to crop
-	if ( NumFrames <= 1 )
+	if (TotalNumOfKeys <= 1 )
 	{
 		return false;
 	}
@@ -2146,38 +2230,39 @@ bool UAnimSequence::CropRawAnimData( float CurrentTime, bool bFromStart )
 	// This assumes that all keys are equally spaced (ie. won't work if we have dropped unimportant frames etc).
 	// The reason I'm changing to TotalNumOfFrames is CT/SL = KeyIndexWithFraction/TotalNumOfFrames
 	// To play TotalNumOfFrames, it takes GetPlayLength(). Each key will take GetPlayLength()/TotalNumOfFrames
-	float const KeyIndexWithFraction = (CurrentTime * (float)(TotalNumOfFrames)) / GetPlayLength();
+	float const KeyIndexWithFraction = (CurrentTime * (float)(TotalNumOfKeys)) / GetPlayLength();
 	int32 KeyIndex = bFromStart ? FMath::FloorToInt(KeyIndexWithFraction) : FMath::CeilToInt(KeyIndexWithFraction);
 	// Ensure KeyIndex is in range.
-	KeyIndex = FMath::Clamp<int32>(KeyIndex, 1, TotalNumOfFrames-1); 
+	KeyIndex = FMath::Clamp<int32>(KeyIndex, 1, TotalNumOfKeys-1); 
 	// determine which keys need to be removed.
-	int32 const StartKey = bFromStart ? 0 : KeyIndex;
-	int32 const NumKeys = bFromStart ? KeyIndex : TotalNumOfFrames - KeyIndex ;
+	const int32 StartKey = bFromStart ? 0 : KeyIndex;
+	const int32 NumKeysToRemove = bFromStart ? KeyIndex : TotalNumOfKeys - KeyIndex;
 
-	// Recalculate NumFrames
-	int32 NewNumFrames = TotalNumOfFrames - NumKeys;
+	// Recalculate number of keys and frames
+	const int32 NewNumKeys = TotalNumOfKeys - NumKeysToRemove;
+	const int32 NewNumFrames = NewNumKeys - 1;
 
-	UE_LOG(LogAnimation, Log, TEXT("UAnimSequence::CropRawAnimData %s - CurrentTime: %f, bFromStart: %d, TotalNumOfFrames: %d, KeyIndex: %d, StartKey: %d, NumKeys: %d"), *GetName(), CurrentTime, bFromStart, TotalNumOfFrames, KeyIndex, StartKey, NumKeys);
+	UE_LOG(LogAnimation, Log, TEXT("UAnimSequence::CropRawAnimData %s - CurrentTime: %f, bFromStart: %d, TotalNumOfFrames: %d, KeyIndex: %d, StartKey: %d, NumKeys: %d"), *GetName(), CurrentTime, bFromStart, TotalNumOfKeys, KeyIndex, StartKey, NumKeysToRemove);
 
 	// Iterate over tracks removing keys from each one.
 	for(int32 i=0; i<RawAnimationData.Num(); i++)
 	{
-		// Update NewNumFrames below to reflect actual number of keys while we crop the anim data
-		CropRawTrack(RawAnimationData[i], StartKey, NumKeys, TotalNumOfFrames);
+		// Update NewNumKeys below to reflect actual number of keys while we crop the anim data
+		CropRawTrack(RawAnimationData[i], StartKey, NumKeysToRemove, TotalNumOfKeys);
 	}
 
 	// Double check that everything is fine
 	for(int32 i=0; i<RawAnimationData.Num(); i++)
 	{
 		FRawAnimSequenceTrack& RawTrack = RawAnimationData[i];
-		check(RawTrack.PosKeys.Num() == 1 || RawTrack.PosKeys.Num() == NewNumFrames);
-		check(RawTrack.RotKeys.Num() == 1 || RawTrack.RotKeys.Num() == NewNumFrames);
+		check(RawTrack.PosKeys.Num() == 1 || RawTrack.PosKeys.Num() == NewNumKeys);
+		check(RawTrack.RotKeys.Num() == 1 || RawTrack.RotKeys.Num() == NewNumKeys);
 	}
 
 	// Update sequence length to match new number of frames.
-	ResizeSequence((float)NewNumFrames * FrameTime, NewNumFrames, false, StartKey, StartKey+NumKeys);
+	ResizeSequence((float)NewNumFrames * FrameTime, NewNumKeys, false, StartKey, StartKey + NumKeysToRemove);
 
-	UE_LOG(LogAnimation, Log, TEXT("\tGetPlayLength(): %f, NumFrames: %d"), GetPlayLength(), NumFrames);
+	UE_LOG(LogAnimation, Log, TEXT("\tGetPlayLength(): %f, Number of Keys: %d"), GetPlayLength(), NumberOfKeys);
 
 	MarkRawDataAsModified();
 	OnRawDataChanged();
@@ -2189,7 +2274,7 @@ bool UAnimSequence::CompressRawAnimData(float MaxPosDiff, float MaxAngleDiff)
 {
 	if (RawAnimationData.Num() > 0)
 	{
-		return StaticCompressRawAnimData(RawAnimationData, NumFrames, GetFName(), MaxPosDiff, MaxAngleDiff);
+		return StaticCompressRawAnimData(RawAnimationData, NumberOfKeys, GetFName(), MaxPosDiff, MaxAngleDiff);
 	}
 	return false;
 }
@@ -2198,7 +2283,7 @@ bool UAnimSequence::CompressRawAnimData()
 {
 	if (RawAnimationData.Num() > 0)
 	{
-		return StaticCompressRawAnimData(RawAnimationData, NumFrames, GetFName());
+		return StaticCompressRawAnimData(RawAnimationData, NumberOfKeys, GetFName());
 	}
 	return false;
 }
@@ -2249,10 +2334,10 @@ bool UAnimSequence::ShouldPerformStripping(const bool bPerformFrameStripping, co
 	const bool bShouldPerformFrameStripping = bPerformFrameStripping && bAllowFrameStripping;
 
 	// Can only do stripping on animations that have an even number of frames once the end frame is removed)
-	const bool bIsEvenFramed = ((NumFrames - 1) % 2) == 0;
+	const bool bIsEvenFramed = ((NumberOfKeys - 1) % 2) == 0;
 	const bool bIsValidForStripping = bIsEvenFramed || bPerformStrippingOnOddFramedAnims;
 
-	const bool bStripCandidate = (NumFrames > 10) && bIsValidForStripping;
+	const bool bStripCandidate = (NumberOfKeys > 10) && bIsValidForStripping;
 
 	return bStripCandidate && bShouldPerformFrameStripping;
 }
@@ -2482,7 +2567,7 @@ void UAnimSequence::SerializeCompressedData(FArchive& Ar, bool bDDCData)
 
 bool UAnimSequence::CanBakeAdditive() const
 {
-	return	(NumFrames > 0) &&
+	return	(NumberOfKeys > 0) &&
 			IsValidAdditive() &&
 			GetSkeleton();
 }
@@ -2527,7 +2612,7 @@ public:
 	TArray<FBoneIndexType> RequiredBoneIndexArray;
 
 	FByFramePoseEvalContext(const UAnimSequence* InAnimToEval)
-		: FByFramePoseEvalContext(InAnimToEval->GetPlayLength(), InAnimToEval->GetRawNumberOfFrames(), InAnimToEval->GetSkeleton())
+		: FByFramePoseEvalContext(InAnimToEval->GetPlayLength(), InAnimToEval->GetNumberOfSampledKeys(), InAnimToEval->GetSkeleton())
 	{}
 		
 	FByFramePoseEvalContext(float InSequenceLength, int32 InRawNumOfFrames, USkeleton* InSkeleton)
@@ -2573,9 +2658,9 @@ void UAnimSequence::BakeOutVirtualBoneTracks(TArray<FRawAnimSequenceTrack>& NewR
 		const int32 TrackIndex = NewRawTracks.Add(FRawAnimSequenceTrack());
 
 		//Init new tracks
-		NewRawTracks[TrackIndex].PosKeys.SetNumUninitialized(NumFrames);
-		NewRawTracks[TrackIndex].RotKeys.SetNumUninitialized(NumFrames);
-		NewRawTracks[TrackIndex].ScaleKeys.SetNumUninitialized(NumFrames);
+		NewRawTracks[TrackIndex].PosKeys.SetNumUninitialized(NumberOfKeys);
+		NewRawTracks[TrackIndex].RotKeys.SetNumUninitialized(NumberOfKeys);
+		NewRawTracks[TrackIndex].ScaleKeys.SetNumUninitialized(NumberOfKeys);
 		
 		NewTrackToSkeletonMapTable.Add(FTrackToSkeletonMap(GetSkeleton()->GetReferenceSkeleton().GetRequiredVirtualBones()[VBIndex]));
 		NewAnimationTrackNames.Add(GetSkeleton()->GetVirtualBones()[VBIndex].VirtualBoneName);
@@ -2592,7 +2677,7 @@ void UAnimSequence::BakeOutVirtualBoneTracks(TArray<FRawAnimSequenceTrack>& NewR
 
 	const TArray<FVirtualBoneRefData>& VBRefData = GetSkeleton()->GetReferenceSkeleton().GetVirtualBoneRefData();
 
-	for (int Frame = 0; Frame < NumFrames; ++Frame)
+	for (int32 Frame = 0; Frame < NumberOfKeys; ++Frame)
 	{
 		// Initialise curve data from Skeleton
 		FBlendedCurve Curve;
@@ -2613,7 +2698,7 @@ void UAnimSequence::BakeOutVirtualBoneTracks(TArray<FRawAnimSequenceTrack>& NewR
 		}
 	}
 
-	StaticCompressRawAnimData(NewRawTracks, NumFrames, GetFName());
+	StaticCompressRawAnimData(NewRawTracks, NumberOfKeys, GetFName());
 }
 
 bool IsIdentity(const FVector& Pos)
@@ -2650,7 +2735,7 @@ void UAnimSequence::TestEvalauteAnimation() const
 
 	FAnimExtractContext ExtractContext;
 
-	for (int Frame = 0; Frame < NumFrames; ++Frame)
+	for (int32 Frame = 0; Frame < NumberOfKeys; ++Frame)
 	{
 		// Initialise curve data from Skeleton
 		FBlendedCurve Curve;
@@ -2700,9 +2785,9 @@ void UAnimSequence::BakeOutAdditiveIntoRawData(TArray<FRawAnimSequenceTrack>& Ne
 
 	for (FRawAnimSequenceTrack& RawTrack : NewRawTracks)
 	{
-		RawTrack.PosKeys.SetNumUninitialized(NumFrames);
-		RawTrack.RotKeys.SetNumUninitialized(NumFrames);
-		RawTrack.ScaleKeys.SetNumUninitialized(NumFrames);
+		RawTrack.PosKeys.SetNumUninitialized(NumberOfKeys);
+		RawTrack.RotKeys.SetNumUninitialized(NumberOfKeys);
+		RawTrack.ScaleKeys.SetNumUninitialized(NumberOfKeys);
 	}
 
 	// keep the same buffer size
@@ -2728,7 +2813,7 @@ void UAnimSequence::BakeOutAdditiveIntoRawData(TArray<FRawAnimSequenceTrack>& Ne
 
 	FAnimExtractContext ExtractContext;
 
-	for (int Frame = 0; Frame < NumFrames; ++Frame)
+	for (int32 Frame = 0; Frame < NumberOfKeys; ++Frame)
 	{
 		// Initialise curve data from Skeleton
 		FBlendedCurve Curve;
@@ -2822,7 +2907,7 @@ void UAnimSequence::BakeOutAdditiveIntoRawData(TArray<FRawAnimSequenceTrack>& Ne
 	}
 #endif
 
-	StaticCompressRawAnimData(NewRawTracks,NumFrames, GetFName());
+	StaticCompressRawAnimData(NewRawTracks, NumberOfKeys, GetFName());
 
 	// Note on (TrackIndex > 0) below : deliberately stop before track 0, compression code doesn't like getting a completely empty animation
 	for (int32 TrackIndex = NewRawTracks.Num() - 1; TrackIndex > 0; --TrackIndex)
@@ -3252,8 +3337,8 @@ void UAnimSequence::RemapTracksToNewSkeleton( USkeleton* NewSkeleton, bool bConv
 			ConvertedSpaceAnimations.AddZeroed(SrcNumTracks);
 			ConvertedLocalSpaceAnimations.AddZeroed(SrcNumTracks);
 
-			int32 NumKeys = NumFrames;
-			float Interval = GetIntervalPerKey(NumFrames, GetPlayLength());
+			int32 NumKeys = NumberOfKeys;
+			float Interval = SamplingFrameRate.AsInterval();
 
 			// allocate arrays
 			for(int32 SrcTrackIndex=0; SrcTrackIndex<SrcNumTracks; ++SrcTrackIndex)
@@ -3476,8 +3561,8 @@ void UAnimSequence::RemapTracksToNewSkeleton( USkeleton* NewSkeleton, bool bConv
 			ConvertedLocalSpaces.AddZeroed(NumBones);
 			ConvertedSpaceBases.AddZeroed(NumBones);
 
-			int32 NumKeys = NumFrames;
-			float Interval = GetIntervalPerKey(NumFrames, GetPlayLength());
+			int32 NumKeys = NumberOfKeys;
+			float Interval = SamplingFrameRate.AsInterval();
 
 			// allocate arrays
 			for(int32 BoneIndex=0; BoneIndex<NumBones; ++BoneIndex)
@@ -3655,6 +3740,7 @@ void UAnimSequence::PostProcessSequence(bool bForceNewRawDatGuid)
 		SanitizeAnimationTrackData(RawAnim);
 	}
 
+	UpdateFrameRate();
 	CompressRawAnimData();
 	// Apply compression
 	MarkRawDataAsModified(bForceNewRawDatGuid);
@@ -3869,12 +3955,12 @@ void UAnimSequence::ReplaceReferredAnimations(const TMap<UAnimationAsset*, UAnim
 bool UAnimSequence::AddLoopingInterpolation()
 {
 	int32 NumTracks = AnimationTrackNames.Num();
-	float Interval = GetIntervalPerKey(NumFrames, GetPlayLength());
+	float Interval = SamplingFrameRate.AsInterval();
 
-	if(NumFrames > 0)
+	if(NumberOfKeys > 0)
 	{
 		// added one more key
-		int32 NewNumKeys = NumFrames +1 ;
+		int32 NewNumKeys = NumberOfKeys +1 ;
 
 		// now I need to calculate back to new animation data
 		for(int32 TrackIndex=0; TrackIndex<NumTracks; ++TrackIndex)
@@ -3899,8 +3985,8 @@ bool UAnimSequence::AddLoopingInterpolation()
 			}
 		}
 
+		NumberOfKeys = NewNumKeys;
 		SetSequenceLength(GetPlayLength() + Interval);
-		NumFrames = NewNumKeys;
 
 		PostProcessSequence();
 		return true;
@@ -3928,8 +4014,8 @@ int32 UAnimSequence::GetSpaceBasedAnimationData(TArray< TArray<FTransform> >& An
 	AnimationDataInComponentSpace.AddZeroed(NumBones);
 
 	// 2d array of animated time [boneindex][time key]
-	int32 NumKeys = NumFrames;
-	float Interval = GetIntervalPerKey(NumFrames, GetPlayLength());
+	int32 NumKeys = NumberOfKeys;
+	float Interval = SamplingFrameRate.AsInterval();
 
 	// allocate arrays
 	for (int32 BoneIndex = 0; BoneIndex < NumBones; ++BoneIndex)
@@ -4166,12 +4252,12 @@ bool UAnimSequence::ConvertAnimationDataToRiggingData(FAnimSequenceTrackContaine
 
 				if (ensure(BoneIndex != INDEX_NONE))
 				{
-					Track.PosKeys.Empty(NumFrames);
-					Track.RotKeys.Empty(NumFrames);
-					Track.ScaleKeys.Empty(NumFrames);
-					Track.PosKeys.AddUninitialized(NumFrames);
-					Track.RotKeys.AddUninitialized(NumFrames);
-					Track.ScaleKeys.AddUninitialized(NumFrames);
+					Track.PosKeys.Empty(NumberOfKeys);
+					Track.RotKeys.Empty(NumberOfKeys);
+					Track.ScaleKeys.Empty(NumberOfKeys);
+					Track.PosKeys.AddUninitialized(NumberOfKeys);
+					Track.RotKeys.AddUninitialized(NumberOfKeys);
+					Track.ScaleKeys.AddUninitialized(NumberOfKeys);
 
 					int32 RigConstraintIndex = Rig->FindTransformBaseByNodeName(NodeName);
 
@@ -4190,7 +4276,7 @@ bool UAnimSequence::ConvertAnimationDataToRiggingData(FAnimSequenceTrackContaine
 							if (ParentBoneIndex != INDEX_NONE)
 							{
 								// if no rig control, component space is used
-								for (int32 KeyIndex = 0; KeyIndex < NumFrames; ++KeyIndex)
+								for (int32 KeyIndex = 0; KeyIndex < NumberOfKeys; ++KeyIndex)
 								{
 									FTransform ParentTransform = AnimationDataInComponentSpace[ParentBoneIndex][KeyIndex];
 									FTransform RelativeTransform = AnimationDataInComponentSpace[BoneIndex][KeyIndex].GetRelativeTransform(ParentTransform);
@@ -4200,7 +4286,7 @@ bool UAnimSequence::ConvertAnimationDataToRiggingData(FAnimSequenceTrackContaine
 							else
 							{
 								// if no rig control, component space is used
-								for (int32 KeyIndex = 0; KeyIndex < NumFrames; ++KeyIndex)
+								for (int32 KeyIndex = 0; KeyIndex < NumberOfKeys; ++KeyIndex)
 								{
 									Track.RotKeys[KeyIndex] = AnimationDataInComponentSpace[BoneIndex][KeyIndex].GetRotation();
 								}
@@ -4209,7 +4295,7 @@ bool UAnimSequence::ConvertAnimationDataToRiggingData(FAnimSequenceTrackContaine
 						else
 						{
 							// if no rig control, component space is used
-							for (int32 KeyIndex = 0; KeyIndex < NumFrames; ++KeyIndex)
+							for (int32 KeyIndex = 0; KeyIndex < NumberOfKeys; ++KeyIndex)
 							{
 								Track.RotKeys[KeyIndex] = AnimationDataInComponentSpace[BoneIndex][KeyIndex].GetRotation();
 							}
@@ -4226,7 +4312,7 @@ bool UAnimSequence::ConvertAnimationDataToRiggingData(FAnimSequenceTrackContaine
 							if (ParentBoneIndex != INDEX_NONE)
 							{
 								// if no rig control, component space is used
-								for (int32 KeyIndex = 0; KeyIndex < NumFrames; ++KeyIndex)
+								for (int32 KeyIndex = 0; KeyIndex < NumberOfKeys; ++KeyIndex)
 								{
 									FTransform ParentTransform = AnimationDataInComponentSpace[ParentBoneIndex][KeyIndex];
 									FTransform RelativeTransform = AnimationDataInComponentSpace[BoneIndex][KeyIndex].GetRelativeTransform(ParentTransform);
@@ -4236,7 +4322,7 @@ bool UAnimSequence::ConvertAnimationDataToRiggingData(FAnimSequenceTrackContaine
 							}
 							else
 							{
-								for (int32 KeyIndex = 0; KeyIndex < NumFrames; ++KeyIndex)
+								for (int32 KeyIndex = 0; KeyIndex < NumberOfKeys; ++KeyIndex)
 								{
 									Track.PosKeys[KeyIndex] = AnimationDataInComponentSpace[BoneIndex][KeyIndex].GetTranslation();
 									Track.ScaleKeys[KeyIndex] = AnimationDataInComponentSpace[BoneIndex][KeyIndex].GetScale3D();
@@ -4245,7 +4331,7 @@ bool UAnimSequence::ConvertAnimationDataToRiggingData(FAnimSequenceTrackContaine
 						}
 						else
 						{
-							for (int32 KeyIndex = 0; KeyIndex < NumFrames; ++KeyIndex)
+							for (int32 KeyIndex = 0; KeyIndex < NumberOfKeys; ++KeyIndex)
 							{
 								Track.PosKeys[KeyIndex] = AnimationDataInComponentSpace[BoneIndex][KeyIndex].GetTranslation();
 								Track.ScaleKeys[KeyIndex] = AnimationDataInComponentSpace[BoneIndex][KeyIndex].GetScale3D();
@@ -4255,7 +4341,7 @@ bool UAnimSequence::ConvertAnimationDataToRiggingData(FAnimSequenceTrackContaine
 					else
 					{
 						// if no rig control, component space is used
-						for (int32 KeyIndex = 0; KeyIndex < NumFrames; ++KeyIndex)
+						for (int32 KeyIndex = 0; KeyIndex < NumberOfKeys; ++KeyIndex)
 						{
 							Track.PosKeys[KeyIndex] = AnimationDataInComponentSpace[BoneIndex][KeyIndex].GetTranslation();
 							Track.RotKeys[KeyIndex] = AnimationDataInComponentSpace[BoneIndex][KeyIndex].GetRotation();
@@ -4324,15 +4410,15 @@ bool UAnimSequence::ConvertRiggingDataToAnimationData(FAnimSequenceTrackContaine
 				Track.PosKeys.Empty();
 				Track.RotKeys.Empty();
 				Track.ScaleKeys.Empty();
-				Track.PosKeys.AddUninitialized(NumFrames);
-				Track.RotKeys.AddUninitialized(NumFrames);
-				Track.ScaleKeys.AddUninitialized(NumFrames);
+				Track.PosKeys.AddUninitialized(NumberOfKeys);
+				Track.RotKeys.AddUninitialized(NumberOfKeys);
+				Track.ScaleKeys.AddUninitialized(NumberOfKeys);
 
 				const int32& ParentBoneIndex = RefSkeleton.GetParentIndex(BoneIndex);
 
 				if(ParentBoneIndex != INDEX_NONE)
 				{
-					for(int32 KeyIndex = 0; KeyIndex < NumFrames; ++KeyIndex)
+					for(int32 KeyIndex = 0; KeyIndex < NumberOfKeys; ++KeyIndex)
 					{
 						FTransform LocalTransform = AnimationDataInComponentSpace[BoneIndex][KeyIndex].GetRelativeTransform(AnimationDataInComponentSpace[ParentBoneIndex][KeyIndex]);
 
@@ -4343,7 +4429,7 @@ bool UAnimSequence::ConvertRiggingDataToAnimationData(FAnimSequenceTrackContaine
 				}
 				else
 				{
-					for(int32 KeyIndex = 0; KeyIndex < NumFrames; ++KeyIndex)
+					for(int32 KeyIndex = 0; KeyIndex < NumberOfKeys; ++KeyIndex)
 					{
 						FTransform LocalTransform = AnimationDataInComponentSpace[BoneIndex][KeyIndex];
 
@@ -4439,7 +4525,7 @@ void UAnimSequence::BakeTrackCurvesToRawAnimationTracks(TArray<FRawAnimSequenceT
 
 
 		// NumFrames can't be zero (filtered earlier)
-		const float Interval = GetIntervalPerKey(NumFrames, GetPlayLength());
+		const float Interval = SamplingFrameRate.AsInterval();
 
 		for (const FTransformCurve& Curve : RawCurveData.TransformCurves)
 		{
@@ -4472,21 +4558,21 @@ void UAnimSequence::BakeTrackCurvesToRawAnimationTracks(TArray<FRawAnimSequenceT
 				if (RawTrack.PosKeys.Num() == 1)
 				{
 					FVector OneKey = RawTrack.PosKeys[0];
-					RawTrack.PosKeys.Init(OneKey, NumFrames);
+					RawTrack.PosKeys.Init(OneKey, NumberOfKeys);
 				}
 				else
 				{
-					ensure(RawTrack.PosKeys.Num() == NumFrames);
+					ensure(RawTrack.PosKeys.Num() == NumberOfKeys);
 				}
 
 				if (RawTrack.RotKeys.Num() == 1)
 				{
 					FQuat OneKey = RawTrack.RotKeys[0];
-					RawTrack.RotKeys.Init(OneKey, NumFrames);
+					RawTrack.RotKeys.Init(OneKey, NumberOfKeys);
 				}
 				else
 				{
-					ensure(RawTrack.RotKeys.Num() == NumFrames);
+					ensure(RawTrack.RotKeys.Num() == NumberOfKeys);
 				}
 
 				// although we don't allow edit of scale
@@ -4495,15 +4581,15 @@ void UAnimSequence::BakeTrackCurvesToRawAnimationTracks(TArray<FRawAnimSequenceT
 				if (RawTrack.ScaleKeys.Num() == 1)
 				{
 					FVector OneKey = RawTrack.ScaleKeys[0];
-					RawTrack.ScaleKeys.Init(OneKey, NumFrames);
+					RawTrack.ScaleKeys.Init(OneKey, NumberOfKeys);
 				}
 				else
 				{
-					ensure(RawTrack.ScaleKeys.Num() == NumFrames);
+					ensure(RawTrack.ScaleKeys.Num() == NumberOfKeys);
 				}
 
 				// now we have all data ready to apply
-				for (int32 KeyIndex = 0; KeyIndex < NumFrames; ++KeyIndex)
+				for (int32 KeyIndex = 0; KeyIndex < NumberOfKeys; ++KeyIndex)
 				{
 					// now evaluate
 					FTransformCurve* TransformCurve = static_cast<FTransformCurve*>(RawCurveData.GetCurveData(Curve.Name.UID, ERawCurveTrackTypes::RCT_Transform));
@@ -4526,7 +4612,7 @@ void UAnimSequence::BakeTrackCurvesToRawAnimationTracks(TArray<FRawAnimSequenceT
 				SanitizeAnimationTrackData(RawTrack);
 			}
 		}
-		StaticCompressRawAnimData(NewRawTracks, NumFrames, GetFName());
+		StaticCompressRawAnimData(NewRawTracks, NumberOfKeys, GetFName());
 	}
 }
 
@@ -4554,7 +4640,7 @@ void UAnimSequence::AddKeyToSequence(float Time, const FName& BoneName, const FT
 void UAnimSequence::ResetAnimation()
 {
 	// clear everything. Making new animation, so need to reset all the things that belong here
-	NumFrames = 0;
+	NumberOfKeys = 0;
 	SetSequenceLength(0.f);
 	RawAnimationData.Empty();
 	AnimationTrackNames.Empty();
@@ -4616,7 +4702,7 @@ bool UAnimSequence::CreateAnimation(USkeletalMesh* Mesh)
 
 		const FReferenceSkeleton& RefSkeleton = Mesh->RefSkeleton;
 		SetSequenceLength(MINIMUM_ANIMATION_LENGTH);
-		NumFrames = 1;
+		NumberOfKeys = 1;
 
 		const int32 NumBones = RefSkeleton.GetRawBoneNum();
 		RawAnimationData.AddZeroed(NumBones);
@@ -4658,7 +4744,7 @@ bool UAnimSequence::CreateAnimation(USkeletalMeshComponent* MeshComponent)
 
 		const FReferenceSkeleton& RefSkeleton = Mesh->RefSkeleton;
 		SetSequenceLength(MINIMUM_ANIMATION_LENGTH);
-		NumFrames = 1;
+		NumberOfKeys = 1;
 
 		const int32 NumBones = RefSkeleton.GetRawBoneNum();
 		RawAnimationData.AddZeroed(NumBones);
@@ -4697,7 +4783,7 @@ bool UAnimSequence::CreateAnimation(UAnimSequence* Sequence)
 		ResetAnimation();
 
 		SetSequenceLength(Sequence->GetPlayLength());
-		NumFrames = Sequence->NumFrames;
+		NumberOfKeys = Sequence->NumberOfKeys;
 
 		RawAnimationData = Sequence->RawAnimationData;
 		AnimationTrackNames = Sequence->AnimationTrackNames;
@@ -4706,6 +4792,7 @@ bool UAnimSequence::CreateAnimation(UAnimSequence* Sequence)
 		AnimNotifyTracks = Sequence->AnimNotifyTracks;
 		RawCurveData = Sequence->RawCurveData;
 
+		UpdateFrameRate();
 		// refresh TrackToskeletonMapIndex
 		RefreshTrackMapFromAnimTrackNames();
 
@@ -5311,7 +5398,7 @@ void UAnimSequence::SynchronousCustomAttributesCompression()
 			}
 			else if (RefPoseType == ABPT_AnimFrame)
 			{
-				const float Fraction = (RefPoseSeq->NumFrames > 0) ? FMath::Clamp<float>((float)RefFrameIndex / (float)RefPoseSeq->NumFrames, 0.f, 1.f) : 0.f;
+				const float Fraction = (RefPoseSeq->NumberOfKeys > 0) ? FMath::Clamp<float>((float)RefFrameIndex / (float)RefPoseSeq->NumberOfKeys, 0.f, 1.f) : 0.f;
 				BasePoseTime = RefPoseSeq->GetPlayLength() * Fraction;
 
 			}
@@ -5359,7 +5446,7 @@ void UAnimSequence::SynchronousCustomAttributesCompression()
 								BakedFloatAttribute.AttributeName = Attribute.Name;
 
 								FSimpleCurve& FloatCurve = BakedFloatAttribute.FloatCurve;
-								ConvertAttributeToAdditive<float, FSimpleCurve>(Attribute, *RefAttribute, FloatCurve, EvalContext.IntervalTime, NumFrames, GetBasePoseTimeToSample);
+								ConvertAttributeToAdditive<float, FSimpleCurve>(Attribute, *RefAttribute, FloatCurve, EvalContext.IntervalTime, NumberOfKeys, GetBasePoseTimeToSample);
 								FloatCurve.RemoveRedundantKeys(0.f);
 
 								break;
@@ -5371,7 +5458,7 @@ void UAnimSequence::SynchronousCustomAttributesCompression()
 								BakedIntAttribute.AttributeName = Attribute.Name;
 
 								FIntegralCurve& IntCurve = BakedIntAttribute.IntCurve;
-								ConvertAttributeToAdditive<int32, FIntegralCurve>(Attribute, *RefAttribute, IntCurve, EvalContext.IntervalTime, NumFrames, GetBasePoseTimeToSample);
+								ConvertAttributeToAdditive<int32, FIntegralCurve>(Attribute, *RefAttribute, IntCurve, EvalContext.IntervalTime, NumberOfKeys, GetBasePoseTimeToSample);
 								IntCurve.RemoveRedundantKeys();
 							
 								break;
@@ -6006,7 +6093,7 @@ void GatherAnimSequenceStats(FOutputDevice& Ar)
 
 		Ar.Logf(TEXT(" %60s, %3i, %3i,%3i,%3i, %3i,%3i,%3i, %10i,%10i,%10i, %s, %d"),
 			*Seq->GetName(),
-			Seq->GetRawNumberOfFrames(),
+			Seq->GetNumberOfSampledKeys(),
 			NumTransTracks, NumRotTracks, NumScaleTracks,
 			NumTransTracksWithOneKey, NumRotTracksWithOneKey, NumScaleTracksWithOneKey,
 			TotalNumTransKeys, TotalNumRotKeys, TotalNumScaleKeys,
diff --git a/Engine/Source/Runtime/Engine/Private/Animation/AnimSequenceBase.cpp b/Engine/Source/Runtime/Engine/Private/Animation/AnimSequenceBase.cpp
--- a/Engine/Source/Runtime/Engine/Private/Animation/AnimSequenceBase.cpp
+++ b/Engine/Source/Runtime/Engine/Private/Animation/AnimSequenceBase.cpp
@@ -527,9 +527,18 @@ void UAnimSequenceBase::RefreshCacheData()
 
 int32 UAnimSequenceBase::GetNumberOfFrames() const
 {
-	static float DefaultSampleRateInterval = 1.f / DEFAULT_SAMPLERATE;
-	// because of float error, add small margin at the end, so it can clamp correctly
-	return (int32)(GetPlayLength() / DefaultSampleRateInterval + KINDA_SMALL_NUMBER) + 1;
+	return GetNumberOfSampledKeys();
+}
+
+int32 UAnimSequenceBase::GetNumberOfSampledKeys() const
+{
+	return GetSamplingFrameRate().AsFrameTime(GetPlayLength()).RoundToFrame().Value;
+}
+
+const FFrameRate& UAnimSequenceBase::GetSamplingFrameRate() const
+{
+	static const FFrameRate DefaultFrameRate = FFrameRate(DEFAULT_SAMPLERATE, 1);
+	return DefaultFrameRate;
 }
 
 #if WITH_EDITOR
@@ -549,14 +558,12 @@ void UAnimSequenceBase::InitializeNotifyTrack()
 
 int32 UAnimSequenceBase::GetFrameAtTime(const float Time) const
 {
-	const float FrameTime = GetNumberOfFrames() > 1 ? GetPlayLength() / (float)(GetNumberOfFrames() - 1) : 0.0f;
-	return FMath::Clamp(FMath::RoundToInt(Time / FrameTime), 0, GetNumberOfFrames() - 1);
+	return FMath::Clamp(GetSamplingFrameRate().AsFrameTime(Time).RoundToFrame().Value, 0, GetNumberOfSampledKeys() - 1);
 }
 
 float UAnimSequenceBase::GetTimeAtFrame(const int32 Frame) const
 {
-	const float FrameTime = GetNumberOfFrames() > 1 ? GetPlayLength() / (float)(GetNumberOfFrames() - 1) : 0.0f;
-	return FMath::Clamp(FrameTime * Frame, 0.0f, GetPlayLength());
+	return FMath::Clamp((float)GetSamplingFrameRate().AsSeconds(Frame), 0.f, GetPlayLength());
 }
 
 void UAnimSequenceBase::RegisterOnNotifyChanged(const FOnNotifyChanged& Delegate)
diff --git a/Engine/Source/Runtime/Engine/Private/Animation/AnimSingleNodeInstance.cpp b/Engine/Source/Runtime/Engine/Private/Animation/AnimSingleNodeInstance.cpp
--- a/Engine/Source/Runtime/Engine/Private/Animation/AnimSingleNodeInstance.cpp
+++ b/Engine/Source/Runtime/Engine/Private/Animation/AnimSingleNodeInstance.cpp
@@ -402,20 +402,21 @@ void UAnimSingleNodeInstance::StepForward()
 	if (UAnimSequenceBase* Sequence = Cast<UAnimSequenceBase>(CurrentAsset))
 	{
 		FAnimSingleNodeInstanceProxy& Proxy = GetProxyOnGameThread<FAnimSingleNodeInstanceProxy>();
-		const FAnimKeyHelper Helper(Sequence->GetPlayLength(), Sequence->GetNumberOfFrames());
-		float KeyLength = Helper.TimePerKey() + SMALL_NUMBER;
-		float Fraction = (Proxy.GetCurrentTime()+KeyLength)/Sequence->GetPlayLength();
-		int32 Frames = (float)Helper.LastKey() * Fraction;
+
+		const FFrameRate& FrameRate = Sequence->GetSamplingFrameRate();
+		const FFrameNumber LastSequenceFrameNumber = FrameRate.AsFrameTime(Sequence->GetPlayLength()).RoundToFrame();
+		
+		// Step forward a small amount and ceil to the next frame number 
+		FFrameNumber StepToFrame = FrameRate.AsFrameTime(Proxy.GetCurrentTime() + KINDA_SMALL_NUMBER).CeilToFrame();		
 		if (IsLooping())
 		{
-			int32 PreviousFrame = (float)Helper.LastKey() * (Proxy.GetCurrentTime() / Sequence->GetPlayLength());
-			Frames = (PreviousFrame == Helper.LastKey()) ? 0 : Frames;
-		}
-		else
-		{
-			Frames = FMath::Clamp<int32>(Frames, 0, Helper.LastKey());
+			// Wrap around to start of the sequence
+			StepToFrame %= (LastSequenceFrameNumber + 1);
 		}
-		SetPosition(Frames*KeyLength);
+
+		StepToFrame.Value = FMath::Clamp<int32>(StepToFrame.Value, 0, LastSequenceFrameNumber.Value);
+
+		SetPosition(FrameRate.AsSeconds(StepToFrame));
 	}	
 }
 
@@ -425,19 +426,20 @@ void UAnimSingleNodeInstance::StepBackward()
 	{
 		FAnimSingleNodeInstanceProxy& Proxy = GetProxyOnGameThread<FAnimSingleNodeInstanceProxy>();
 
-		const FAnimKeyHelper Helper(Sequence->GetPlayLength(), Sequence->GetNumberOfFrames());
-		float KeyLength = Helper.TimePerKey() + SMALL_NUMBER;
-		float Fraction = (Proxy.GetCurrentTime()-KeyLength)/Sequence->GetPlayLength();
-		int32 Frames = (float)Helper.LastKey() * Fraction;
+		const FFrameRate& FrameRate = Sequence->GetSamplingFrameRate();
+		const FFrameNumber LastSequenceFrameNumber = FrameRate.AsFrameTime(Sequence->GetPlayLength()).RoundToFrame();
+
+		// Step backwards a small amount and floor to the previous frame number 
+		FFrameNumber StepToFrame = FrameRate.AsFrameTime(Proxy.GetCurrentTime() - KINDA_SMALL_NUMBER).FloorToFrame();
 		if (IsLooping())
 		{
-			Frames = (Frames < 0) ? Helper.LastKey() : Frames;
-		}
-		else
-		{
-			Frames = FMath::Clamp<int32>(Frames, 0, Helper.LastKey());
+			// Wrap around to end of sequence
+			StepToFrame = StepToFrame.Value < 0 ? LastSequenceFrameNumber : StepToFrame;
 		}
-		SetPosition(Frames * KeyLength);
+
+		StepToFrame.Value = FMath::Clamp<int32>(StepToFrame.Value, 0, LastSequenceFrameNumber.Value);
+
+		SetPosition(FrameRate.AsSeconds(StepToFrame));
 	}
 }
 
diff --git a/Engine/Source/Runtime/Engine/Private/Animation/AnimStreamable.cpp b/Engine/Source/Runtime/Engine/Private/Animation/AnimStreamable.cpp
--- a/Engine/Source/Runtime/Engine/Private/Animation/AnimStreamable.cpp
+++ b/Engine/Source/Runtime/Engine/Private/Animation/AnimStreamable.cpp
@@ -21,6 +21,7 @@
 #include "ProfilingDebugging/CookStats.h"
 #include "Animation/AnimationPoseData.h"
 #include "Animation/CustomAttributesRuntime.h"
+#include "UObject/UE5MainStreamObjectVersion.h"
 
 CSV_DECLARE_CATEGORY_MODULE_EXTERN(ENGINE_API, Animation);
 
@@ -143,6 +144,7 @@ void UAnimStreamable::Serialize(FArchive& Ar)
 {
 	Super::Serialize(Ar);
 	
+	Ar.UsingCustomVersion(FUE5MainStreamObjectVersion::GUID);
 
 	bool bCooked = Ar.IsCooking();
 	Ar << bCooked;
@@ -164,6 +166,15 @@ void UAnimStreamable::Serialize(FArchive& Ar)
 			PlatformData.Serialize(Ar, this);
 		}
 	}
+
+#if WITH_EDITORONLY_DATA
+	if (Ar.IsLoading() && Ar.CustomVer(FUE5MainStreamObjectVersion::GUID) < FUE5MainStreamObjectVersion::RenamingAnimationNumFrames)
+	{
+		PRAGMA_DISABLE_DEPRECATION_WARNINGS
+		NumberOfKeys = NumFrames;
+		PRAGMA_ENABLE_DEPRECATION_WARNINGS
+	}
+#endif // WITH_EDITORONLY_DATA
 }
 
 #if WITH_EDITOR
@@ -280,7 +291,7 @@ void UAnimStreamable::GetAnimationPose(FAnimationPoseData& OutAnimationPoseData,
 			}
 		}
 
-		BuildPoseFromRawData(RawAnimationData, TrackToSkeletonMapTable, OutPose, ExtractionContext.CurrentTime, Interpolation, NumFrames, GetPlayLength(), RetargetSource);
+		BuildPoseFromRawData(RawAnimationData, TrackToSkeletonMapTable, OutPose, ExtractionContext.CurrentTime, Interpolation, NumberOfKeys, GetPlayLength(), RetargetSource);
 
 		if ((ExtractionContext.bExtractRootMotion && RootMotionReset.bEnableRootMotion) || RootMotionReset.bForceRootLock)
 		{
@@ -446,7 +457,7 @@ void UAnimStreamable::InitFrom(const UAnimSequence* InSourceSequence)
 	TrackToSkeletonMapTable = InSourceSequence->GetRawTrackToSkeletonMapTable();
 	AnimationTrackNames = InSourceSequence->GetAnimationTrackNames();
 	
-	NumFrames = InSourceSequence->GetNumberOfFrames();
+	NumberOfKeys = InSourceSequence->GetNumberOfSampledKeys();
 	SetSequenceLength(InSourceSequence->GetPlayLength());
 	
 	RateScale = InSourceSequence->RateScale;
@@ -543,8 +554,8 @@ void UAnimStreamable::RequestCompressedData(const ITargetPlatform* Platform)
 		NumChunks = FMath::Max(InitialNumChunks, 1);
 	}
 
-	int32 NumFramesToChunk = NumFrames - 1;
-	int32 FramesPerChunk = NumFrames / NumChunks;
+	int32 NumFramesToChunk = NumberOfKeys - 1;
+	int32 FramesPerChunk = NumberOfKeys / NumChunks;
 
 	PlatformData.Chunks.AddDefaulted(NumChunks);
 
@@ -625,7 +636,8 @@ void UAnimStreamable::RequestCompressedDataForChunk(const FString& ChunkDDCKey,
 		const bool bSkipDDC = false;
 
 		const int32 ChunkNumFrames = FrameEnd - FrameStart;
-		const float FrameLength = GetPlayLength() / (float)(NumFrames - 1);
+		const int32 NumberOfFrames = NumberOfKeys - 1;
+		const float FrameLength = GetPlayLength() / (float)NumberOfFrames;
 		Chunk.StartTime = FrameStart * FrameLength;
 		Chunk.SequenceLength = ChunkNumFrames * FrameLength;
 
@@ -644,11 +656,11 @@ void UAnimStreamable::RequestCompressedDataForChunk(const FString& ChunkDDCKey,
 				FRawAnimSequenceTrack& SrcTrack = RawAnimationData[TrackIndex];
 				FRawAnimSequenceTrack& DestTrack = CompressibleData->RawAnimationData[TrackIndex];
 
-				MakeKeyChunk(SrcTrack.PosKeys, DestTrack.PosKeys, NumFrames, FrameStart, FrameEnd);
-				MakeKeyChunk(SrcTrack.RotKeys, DestTrack.RotKeys, NumFrames, FrameStart, FrameEnd);
+				MakeKeyChunk(SrcTrack.PosKeys, DestTrack.PosKeys, NumberOfKeys, FrameStart, FrameEnd);
+				MakeKeyChunk(SrcTrack.RotKeys, DestTrack.RotKeys, NumberOfKeys, FrameStart, FrameEnd);
 				if (SrcTrack.ScaleKeys.Num() > 0)
 				{
-					MakeKeyChunk(SrcTrack.ScaleKeys, DestTrack.ScaleKeys, NumFrames, FrameStart, FrameEnd);
+					MakeKeyChunk(SrcTrack.ScaleKeys, DestTrack.ScaleKeys, NumberOfKeys, FrameStart, FrameEnd);
 				}
 			}
 
diff --git a/Engine/Source/Runtime/Engine/Private/Animation/AnimTrace.cpp b/Engine/Source/Runtime/Engine/Private/Animation/AnimTrace.cpp
--- a/Engine/Source/Runtime/Engine/Private/Animation/AnimTrace.cpp
+++ b/Engine/Source/Runtime/Engine/Private/Animation/AnimTrace.cpp
@@ -973,7 +973,7 @@ void FAnimTrace::OutputAnimSequencePlayer(const FAnimationBaseContext& InContext
 		<< AnimSequencePlayer.NodeId(InContext.GetCurrentNodeId())
 		<< AnimSequencePlayer.Position(InNode.GetAccumulatedTime())
 		<< AnimSequencePlayer.Length(InNode.Sequence ? InNode.Sequence->GetPlayLength() : 0.0f)
-		<< AnimSequencePlayer.FrameCounter(InNode.Sequence ? InNode.Sequence->GetNumberOfFrames() : 0);
+		<< AnimSequencePlayer.FrameCounter(InNode.Sequence ? InNode.Sequence->GetNumberOfSampledKeys() : 0);
 }
 
 void FAnimTrace::OutputStateMachineState(const FAnimationBaseContext& InContext, int32 InStateMachineIndex, int32 InStateIndex, float InStateWeight, float InElapsedTime)
diff --git a/Engine/Source/Runtime/Engine/Private/Animation/AnimationUtils.cpp b/Engine/Source/Runtime/Engine/Private/Animation/AnimationUtils.cpp
--- a/Engine/Source/Runtime/Engine/Private/Animation/AnimationUtils.cpp
+++ b/Engine/Source/Runtime/Engine/Private/Animation/AnimationUtils.cpp
@@ -235,7 +235,7 @@ void FAnimationUtils::ComputeCompressionError(const FCompressibleAnimData& Compr
 	ErrorStats.MaxErrorTime = 0.0f;
 	int32 MaxErrorTrack = -1;
 
-	if (CompressedData.AnimData != nullptr && CompressedData.AnimData->CompressedNumberOfFrames > 0)
+	if (CompressedData.AnimData != nullptr && CompressedData.AnimData->CompressedNumberOfKeys > 0)
 	{
 		const bool bCanUseCompressedData = CompressedData.AnimData->IsValid();
 		if (!bCanUseCompressedData)
@@ -282,16 +282,16 @@ void FAnimationUtils::ComputeCompressionError(const FCompressibleAnimData& Compr
 
 		const FTransform EndEffectorDummyBoneSocket(FQuat::Identity, FVector(END_EFFECTOR_DUMMY_BONE_LENGTH_SOCKET));
 		const FTransform EndEffectorDummyBone(FQuat::Identity, FVector(END_EFFECTOR_DUMMY_BONE_LENGTH));
-		const FAnimKeyHelper Helper(CompressibleAnimData.SequenceLength, CompressedData.AnimData->CompressedNumberOfFrames);
+		const FAnimKeyHelper Helper(CompressibleAnimData.SequenceLength, CompressedData.AnimData->CompressedNumberOfKeys);
 		const float KeyLength = Helper.TimePerKey() + SMALL_NUMBER;
 
 		FAnimSequenceDecompressionContext DecompContext(CompressibleAnimData.SequenceLength, CompressibleAnimData.Interpolation, CompressibleAnimData.AnimFName, *CompressedData.AnimData);
 
 		const TArray<FBoneData>& BoneData = CompressibleAnimData.BoneData;
 
-		for (int32 FrameIndex = 0; FrameIndex< CompressedData.AnimData->CompressedNumberOfFrames; FrameIndex++)
+		for (int32 KeyIndex = 0; KeyIndex< CompressedData.AnimData->CompressedNumberOfKeys; KeyIndex++)
 		{
-			const float Time = (float)FrameIndex * KeyLength;
+			const float Time = (float)KeyIndex * KeyLength;
 			DecompContext.Seek(Time);
 
 			// get the raw and compressed atom for each bone
@@ -319,7 +319,7 @@ void FAnimationUtils::ComputeCompressionError(const FCompressibleAnimData& Compr
 
 						FTransform AdditiveRawTransform;
 						FTransform AdditiveNewTransform;
-						FAnimationUtils::ExtractTransformFromTrack(Time, CompressibleAnimData.NumFrames, CompressibleAnimData.SequenceLength, CompressibleAnimData.RawAnimationData[BoneIndexData.TrackIndex], CompressibleAnimData.Interpolation, AdditiveRawTransform);
+						FAnimationUtils::ExtractTransformFromTrack(Time, CompressibleAnimData.NumberOfKeys, CompressibleAnimData.SequenceLength, CompressibleAnimData.RawAnimationData[BoneIndexData.TrackIndex], CompressibleAnimData.Interpolation, AdditiveRawTransform);
 						CompressedData.Codec->DecompressBone(DecompContext, BoneIndexData.TrackIndex, AdditiveNewTransform);
 
 						const ScalarRegister VBlendWeight(1.f);
@@ -328,7 +328,7 @@ void FAnimationUtils::ComputeCompressionError(const FCompressibleAnimData& Compr
 					}
 					else
 					{
-						FAnimationUtils::ExtractTransformFromTrack(Time, CompressibleAnimData.NumFrames, CompressibleAnimData.SequenceLength, CompressibleAnimData.RawAnimationData[BoneIndexData.TrackIndex], CompressibleAnimData.Interpolation, RawTransforms[BoneIndex]);
+						FAnimationUtils::ExtractTransformFromTrack(Time, CompressibleAnimData.NumberOfKeys, CompressibleAnimData.SequenceLength, CompressibleAnimData.RawAnimationData[BoneIndexData.TrackIndex], CompressibleAnimData.Interpolation, RawTransforms[BoneIndex]);
 						CompressedData.Codec->DecompressBone(DecompContext, BoneIndexData.TrackIndex, NewTransforms[BoneIndex]);
 					}
 				}
@@ -731,7 +731,7 @@ void FAnimationUtils::TallyErrorsFromPerturbation(
 	}
 
 
-	for (int32 Frame = 0; Frame < CompressibleAnimData.NumFrames; ++Frame)
+	for (int32 KeyIndex = 0; KeyIndex < CompressibleAnimData.NumberOfKeys; ++KeyIndex)
 	{
 		//Build Locals For Frame
 		if (CompressibleAnimData.IsCancelled())
@@ -742,7 +742,7 @@ void FAnimationUtils::TallyErrorsFromPerturbation(
 		for (const FTrackBoneMapping& TrackAndBone : TracksAndBonesToTest)
 		{
 			const FRawAnimSequenceTrack& RawTrack = CompressibleAnimData.RawAnimationData[TrackAndBone.TrackIndex];
-			ExtractTransformForFrameFromTrackSafe(RawTrack, Frame, RawAtoms[TrackAndBone.BoneIndex]);
+			ExtractTransformForFrameFromTrackSafe(RawTrack, KeyIndex, RawAtoms[TrackAndBone.BoneIndex]);
 		}
 
 		//Build Reference Component Space for Frame
@@ -1052,7 +1052,7 @@ void FAnimationUtils::ExtractTransformFromCompressionData(const FCompressibleAni
 		return;
 	}
 
-	FAnimationUtils::ExtractTransformFromTrack(Time, CompressibleAnimData.NumFrames, CompressibleAnimData.SequenceLength, CompressibleAnimData.RawAnimationData[TrackIndex], CompressibleAnimData.Interpolation, OutBoneTransform);
+	FAnimationUtils::ExtractTransformFromTrack(Time, CompressibleAnimData.NumberOfKeys, CompressibleAnimData.SequenceLength, CompressibleAnimData.RawAnimationData[TrackIndex], CompressibleAnimData.Interpolation, OutBoneTransform);
 }
 
 bool FAnimationUtils::CompressAnimBones(FCompressibleAnimData& AnimSeq, FCompressibleAnimDataResult& Target)
diff --git a/Engine/Source/Runtime/Engine/Private/Animation/PoseAsset.cpp b/Engine/Source/Runtime/Engine/Private/Animation/PoseAsset.cpp
--- a/Engine/Source/Runtime/Engine/Private/Animation/PoseAsset.cpp
+++ b/Engine/Source/Runtime/Engine/Private/Animation/PoseAsset.cpp
@@ -1178,7 +1178,7 @@ void UPoseAsset::CreatePoseFromAnimation(class UAnimSequence* AnimSequence, cons
 			// reinitialize, now we're making new pose from this animation
 			Reinitialize();
 
-			const int32 NumPoses = AnimSequence->GetNumberOfFrames();
+			const int32 NumPoses = AnimSequence->GetNumberOfSampledKeys();
 
 			// make sure we have more than one pose
 			if (NumPoses > 0)
diff --git a/Engine/Source/Runtime/Engine/Private/UnrealEngine.cpp b/Engine/Source/Runtime/Engine/Private/UnrealEngine.cpp
--- a/Engine/Source/Runtime/Engine/Private/UnrealEngine.cpp
+++ b/Engine/Source/Runtime/Engine/Private/UnrealEngine.cpp
@@ -5878,7 +5878,7 @@ bool UEngine::HandleListAnimsCommand(const TCHAR* Cmd, FOutputDevice& Ar)
 
 		if (UAnimSequence* AnimSeq = Cast<UAnimSequence>(AnimAsset))
 		{
-			NumKeys = AnimSeq->GetRawNumberOfFrames();
+			NumKeys = AnimSeq->GetNumberOfSampledKeys();
 			SequenceLength = AnimSeq->GetPlayLength();
 			RateScale = AnimSeq->RateScale;
 			NumCurves = AnimSeq->RawCurveData.FloatCurves.Num();
diff --git a/Engine/Source/Runtime/Engine/Public/AnimEncoding_VariableKeyLerp.h b/Engine/Source/Runtime/Engine/Public/AnimEncoding_VariableKeyLerp.h
--- a/Engine/Source/Runtime/Engine/Public/AnimEncoding_VariableKeyLerp.h
+++ b/Engine/Source/Runtime/Engine/Public/AnimEncoding_VariableKeyLerp.h
@@ -209,7 +209,7 @@ FORCEINLINE_DEBUGGABLE void AEFVariableKeyLerp<FORMAT>::GetBoneAtomRotation(FTra
 
 		int32 Index0;
 		int32 Index1;
-		float Alpha = TimeToIndex(DecompContext.Interpolation, AnimData.CompressedNumberOfFrames, FrameTable, DecompContext.RelativePos, NumRotKeys, Index0, Index1);
+		float Alpha = TimeToIndex(DecompContext.Interpolation, AnimData.CompressedNumberOfKeys, FrameTable, DecompContext.RelativePos, NumRotKeys, Index0, Index1);
 
 
 		if (Index0 != Index1)
@@ -263,7 +263,7 @@ FORCEINLINE_DEBUGGABLE void AEFVariableKeyLerp<FORMAT>::GetBoneAtomTranslation(F
 
 	int32 Index0;
 	int32 Index1;
-	float Alpha = TimeToIndex(DecompContext.Interpolation, AnimData.CompressedNumberOfFrames, FrameTable, DecompContext.RelativePos, NumTransKeys, Index0, Index1);
+	float Alpha = TimeToIndex(DecompContext.Interpolation, AnimData.CompressedNumberOfKeys, FrameTable, DecompContext.RelativePos, NumTransKeys, Index0, Index1);
 	const int32 TransStreamOffset = ((FORMAT == ACF_IntervalFixed32NoW) && NumTransKeys > 1) ? (sizeof(float)*6) : 0; // offset past Min and Range data
 
 	if (Index0 != Index1)
@@ -307,7 +307,7 @@ FORCEINLINE_DEBUGGABLE void AEFVariableKeyLerp<FORMAT>::GetBoneAtomScale(FTransf
 
 	int32 Index0;
 	int32 Index1;
-	float Alpha = TimeToIndex(DecompContext.Interpolation, AnimData.CompressedNumberOfFrames, FrameTable, DecompContext.RelativePos, NumScaleKeys, Index0, Index1);
+	float Alpha = TimeToIndex(DecompContext.Interpolation, AnimData.CompressedNumberOfKeys, FrameTable, DecompContext.RelativePos, NumScaleKeys, Index0, Index1);
 	const int32 ScaleStreamOffset = ((FORMAT == ACF_IntervalFixed32NoW) && NumScaleKeys > 1) ? (sizeof(float)*6) : 0; // offset past Min and Range data
 
 	if (Index0 != Index1)
diff --git a/Engine/Source/Runtime/Engine/Public/Animation/AnimCompressionTypes.h b/Engine/Source/Runtime/Engine/Public/Animation/AnimCompressionTypes.h
--- a/Engine/Source/Runtime/Engine/Public/Animation/AnimCompressionTypes.h
+++ b/Engine/Source/Runtime/Engine/Public/Animation/AnimCompressionTypes.h
@@ -193,7 +193,7 @@ struct ENGINE_API FCompressibleAnimData
 public:
 	FCompressibleAnimData();
 
-	FCompressibleAnimData(UAnimBoneCompressionSettings* InBoneCompressionSettings, UAnimCurveCompressionSettings* InCurveCompressionSettings, USkeleton* InSkeleton, EAnimInterpolationType InInterpolation, float InSequenceLength, int32 InNumFrames);
+	FCompressibleAnimData(UAnimBoneCompressionSettings* InBoneCompressionSettings, UAnimCurveCompressionSettings* InCurveCompressionSettings, USkeleton* InSkeleton, EAnimInterpolationType InInterpolation, float InSequenceLength, int32 InNumberOfKeys);
 
 	FCompressibleAnimData(class UAnimSequence* InSeq, const bool bPerformStripping);
 
@@ -219,7 +219,8 @@ public:
 
 	float SequenceLength;
 
-	int32 NumFrames;
+	/** Number of keys within the (non-uniform) RawAnimationData tracks */
+	int32 NumberOfKeys;
 
 	bool bIsValidAdditive;
 
@@ -394,7 +395,7 @@ struct FCompressedAnimDataBase
 struct ENGINE_API ICompressedAnimData
 {
 	/* Common data */
-	int32 CompressedNumberOfFrames;
+	int32 CompressedNumberOfKeys;
 
 #if WITH_EDITORONLY_DATA
 	/** The error stats from the current bone compression codec. */
diff --git a/Engine/Source/Runtime/TimeManagement/Private/TimeManagementModule.cpp b/Engine/Source/Runtime/TimeManagement/Private/TimeManagementModule.cpp
--- a/Engine/Source/Runtime/TimeManagement/Private/TimeManagementModule.cpp
+++ b/Engine/Source/Runtime/TimeManagement/Private/TimeManagementModule.cpp
@@ -2,11 +2,13 @@
 
 #include "ITimeManagementModule.h"
 #include "TimedDataInputCollection.h"
+#include "CommonFrameRates.h"
 
 class FTimeManagementModule : public ITimeManagementModule
 {
 public:
 	virtual FTimedDataInputCollection& GetTimedDataInputCollection() { return Collection; }
+	virtual TArrayView<const struct FCommonFrameRateInfo> GetAllCommonFrameRates() { return FCommonFrameRates::GetAll(); }
 
 private:
 	FTimedDataInputCollection Collection;
diff --git a/Engine/Source/Runtime/TimeManagement/Public/ITimeManagementModule.h b/Engine/Source/Runtime/TimeManagement/Public/ITimeManagementModule.h
--- a/Engine/Source/Runtime/TimeManagement/Public/ITimeManagementModule.h
+++ b/Engine/Source/Runtime/TimeManagement/Public/ITimeManagementModule.h
@@ -35,4 +35,7 @@ public:
 
 	/** Get the collection of the ITimedDataInput and ITimedDataInputGroups. */
 	virtual FTimedDataInputCollection& GetTimedDataInputCollection() = 0;
+
+	/** Returns all stored common frame rates */
+	virtual TArrayView<const struct FCommonFrameRateInfo> GetAllCommonFrameRates() = 0;
 };
