Fixed crashes when closing the editor with an opened WP map

diff --git a/Engine/Source/Runtime/Engine/Private/WorldPartition/WorldPartition.cpp b/Engine/Source/Runtime/Engine/Private/WorldPartition/WorldPartition.cpp
--- a/Engine/Source/Runtime/Engine/Private/WorldPartition/WorldPartition.cpp
+++ b/Engine/Source/Runtime/Engine/Private/WorldPartition/WorldPartition.cpp
@@ -102,12 +102,16 @@ struct FWorldPartionCellUpdateContext
 		UpdatesInProgress--;
 		if (UpdatesInProgress == 0)
 		{
-			GEngine->BroadcastLevelActorListChanged();
-			GEditor->NoteSelectionChange();
-
-			if (WorldPartition->WorldPartitionEditor)
+			// @todo_ow: Once Metadata is removed from external actor's package, testing WorldPartition->IsInitialized() won't be necessary anymore.
+			if (WorldPartition->IsInitialized())
 			{
-				WorldPartition->WorldPartitionEditor->Refresh();
+				GEngine->BroadcastLevelActorListChanged();
+				GEditor->NoteSelectionChange();
+
+				if (WorldPartition->WorldPartitionEditor)
+				{
+					WorldPartition->WorldPartitionEditor->Refresh();
+				}
 			}
 		}
 	}
@@ -650,10 +654,13 @@ void UWorldPartition::UpdateLoadingEditorCell(UWorldPartitionEditorCell* Cell, b
 	{
 		for (FWorldPartitionActorDesc* ActorDesc: Cell->LoadedActors.Array())
 		{
+			// Here we test for null Actor only for when UpdateLoadingEditorCell is called by UWorldPartition::Uninitialize() when exiting the editor (see StaticExit() / GExitPurge).
 			AActor* Actor = ActorDesc->GetActor();
-			check(Actor);
-
-			UnloadActor(ActorDesc, Actor);
+			check(Actor || GExitPurge);
+			if (Actor)
+			{
+				UnloadActor(ActorDesc, Actor);
+			}
 		}
 
 		Cell->LoadedActors.Empty();
