World Partition
- Make sure all actors referenced by the level script are treated as always loaded.
- Removed unused streaming mode 'runtime external objects'.

diff --git a/Engine/Source/Runtime/Engine/Classes/GameFramework/Actor.h b/Engine/Source/Runtime/Engine/Classes/GameFramework/Actor.h
--- a/Engine/Source/Runtime/Engine/Classes/GameFramework/Actor.h
+++ b/Engine/Source/Runtime/Engine/Classes/GameFramework/Actor.h
@@ -38,21 +38,6 @@ struct FNetViewer;
 struct FNetworkObjectInfo;
 class UDataLayer;
 
-#if WITH_EDITOR
-// @todo_ow this is temporary and will be removed
-class FArchiveGetActorRefs : public FArchiveUObject
-{
-public:
-	FArchiveGetActorRefs(AActor* InRoot, TSet<AActor*>& InActorReferences);
-	virtual FArchive& operator<<(UObject*& Obj) override;
-
-private:
-	AActor* Root;
-	TSet<AActor*>& ActorReferences;
-	TSet<UObject*> SubObjects;
-};
-#endif
-
 /** Chooses a method for actors to update overlap state (objects it is touching) on initialization, currently only used during level streaming. */
 UENUM(BlueprintType)
 enum class EActorUpdateOverlapsMethod : uint8
diff --git a/Engine/Source/Runtime/Engine/Private/Actor.cpp b/Engine/Source/Runtime/Engine/Private/Actor.cpp
--- a/Engine/Source/Runtime/Engine/Private/Actor.cpp
+++ b/Engine/Source/Runtime/Engine/Private/Actor.cpp
@@ -89,60 +89,6 @@ FCriticalSection CSVActorClassNameToCountMapLock;
 
 #endif // (CSV_PROFILER && !UE_BUILD_SHIPPING)
 
-#if WITH_EDITOR
-FArchiveGetActorRefs::FArchiveGetActorRefs(AActor* InRoot, TSet<AActor*>& InActorReferences)
-	: Root(InRoot)
-	, ActorReferences(InActorReferences)
-{
-	SetIsSaving(true);
-	SetIsPersistent(true);
-	ArIsObjectReferenceCollector = true;
-	ArShouldSkipBulkData = true;
-}
-
-FArchive& FArchiveGetActorRefs::operator<<(UObject*& Obj)
-{
-	if (Obj && !Obj->IsTemplate() && !Obj->HasAnyFlags(RF_Transient))
-	{
-		if (Obj->IsInOuter(Root))
-		{
-			bool bWasAlreadyInSet;
-			SubObjects.Add(Obj, &bWasAlreadyInSet);
-
-			if (!bWasAlreadyInSet)
-			{
-				// Recurse into subobjects
-				Obj->Serialize(*this);
-			}
-		}
-		else
-		{
-			// Check if external reference is an actor
-			if(Obj->IsA<AActor>())
-			{
-				AActor* Actor = (AActor*)Obj;
-				if (Actor->IsInLevel(Root->GetLevel()))
-				{
-					AActor* TopParentActor = Actor;
-					while(TopParentActor->GetParentActor())
-					{
-						TopParentActor = TopParentActor->GetParentActor();
-					}
-
-					check(TopParentActor);
-
-					if (TopParentActor->IsPackageExternal())
-					{
-						ActorReferences.Add(TopParentActor);
-					}
-				}
-			}
-		}
-	}
-	return *this;
-}
-#endif
-
 uint32 AActor::BeginPlayCallDepth = 0;
 
 AActor::AActor()
diff --git a/Engine/Source/Runtime/Engine/Private/ExternalActorsUtils.cpp b/Engine/Source/Runtime/Engine/Private/ExternalActorsUtils.cpp
new file mode 100644
--- /dev/null
+++ b/Engine/Source/Runtime/Engine/Private/ExternalActorsUtils.cpp
@@ -0,0 +1,72 @@
+// Copyright Epic Games, Inc. All Rights Reserved.
+
+#include "ExternalActorsUtils.h"
+
+class FArchiveGatherExternalActorRefs : public FArchiveUObject
+{
+public:
+	FArchiveGatherExternalActorRefs(UObject* InRoot, TArray<AActor*>& InActorReferences)
+		: Root(InRoot)
+		, ActorReferences(InActorReferences)
+	{
+		SetIsSaving(true);
+		SetIsPersistent(true);
+		ArIgnoreOuterRef = true;
+		ArIsObjectReferenceCollector = true;
+		ArShouldSkipBulkData = true;
+
+		MarkAllObjects(OBJECTMARK_TagExp);
+
+		Root->Serialize(*this);
+	}
+
+	virtual FArchive& operator<<(UObject*& Obj) override
+	{
+		if (Obj && (Obj != Root) && !Obj->IsTemplate() && !Obj->HasAnyFlags(RF_Transient))
+		{
+			if (Obj->HasAnyMarks(OBJECTMARK_TagExp))
+			{
+				Obj->UnMark(OBJECTMARK_TagExp);
+
+				HandleObjectReference(Obj);
+
+				if (Obj->IsInOuter(Root))
+				{
+					Obj->Serialize(*this);
+				}
+			}
+		}
+		return *this;
+	}
+
+private:
+	void HandleObjectReference(UObject* Obj)
+	{
+		if(Obj->IsA<AActor>())
+		{
+			AActor* Actor = (AActor*)Obj;
+			AActor* TopParentActor = Actor;
+			while(TopParentActor->GetParentActor())
+			{
+				TopParentActor = TopParentActor->GetParentActor();
+			}
+
+			check(TopParentActor);
+
+			if (TopParentActor->IsPackageExternal())
+			{
+				ActorReferences.Add(TopParentActor);
+			}
+		}
+	}
+
+	UObject* Root;
+	TArray<AActor*>& ActorReferences;
+};
+
+TArray<AActor*> ExternalActorsUtils::GetExternalActorReferences(UObject* Root)
+{
+	TArray<AActor*> Result;
+	FArchiveGatherExternalActorRefs Ar(Root, Result);
+	return MoveTemp(Result);
+}
diff --git a/Engine/Source/Runtime/Engine/Private/WorldPartition/WorldPartitionActorDesc.cpp b/Engine/Source/Runtime/Engine/Private/WorldPartition/WorldPartitionActorDesc.cpp
--- a/Engine/Source/Runtime/Engine/Private/WorldPartition/WorldPartitionActorDesc.cpp
+++ b/Engine/Source/Runtime/Engine/Private/WorldPartition/WorldPartitionActorDesc.cpp
@@ -14,6 +14,7 @@
 #include "UObject/UE5MainStreamObjectVersion.h"
 #include "WorldPartition/WorldPartition.h"
 #include "WorldPartition/DataLayer/DataLayer.h"
+#include "Engine/Public/ExternalActorsUtils.h"
 #endif
 
 #if WITH_EDITOR
@@ -59,9 +60,7 @@ void FWorldPartitionActorDesc::Init(const AActor* InActor)
 	ActorPackage = InActor->GetPackage()->GetFName();
 	ActorPath = *InActor->GetPathName();
 	
-	TSet<AActor*> ActorReferences;
-	FArchiveGetActorRefs GetActorRefsAr((AActor*)InActor, ActorReferences);
-	((AActor*)InActor)->Serialize(GetActorRefsAr);
+	TArray<AActor*> ActorReferences = ExternalActorsUtils::GetExternalActorReferences((AActor*)InActor);
 
 	if (ActorReferences.Num())
 	{
diff --git a/Engine/Source/Runtime/Engine/Private/WorldPartition/WorldPartitionRuntimeSpatialHash.cpp b/Engine/Source/Runtime/Engine/Private/WorldPartition/WorldPartitionRuntimeSpatialHash.cpp
--- a/Engine/Source/Runtime/Engine/Private/WorldPartition/WorldPartitionRuntimeSpatialHash.cpp
+++ b/Engine/Source/Runtime/Engine/Private/WorldPartition/WorldPartitionRuntimeSpatialHash.cpp
@@ -32,6 +32,8 @@
 #include "WorldPartition/HLOD/HLODActorDesc.h"
 #include "WorldPartition/NavigationData/NavigationDataChunkActor.h"
 #include "AI/NavigationSystemBase.h"
+#include "Engine/Public/ExternalActorsUtils.h"
+#include "Engine/LevelScriptBlueprint.h"
 
 #include "AssetRegistryModule.h"
 #include "AssetData.h"
@@ -119,89 +121,87 @@ struct FActorCluster
 	TArray<const UDataLayer*>	DataLayers;
 	uint32						DataLayersID;
 
-	FActorCluster(const FWorldPartitionActorDesc* ActorDesc, UWorld* World);
-	void Add(const FActorCluster& ActorCluster);
-};
-
-FActorCluster::FActorCluster(const FWorldPartitionActorDesc* ActorDesc, UWorld* World)
-	: GridPlacement(ActorDesc->GetGridPlacement())
-	, RuntimeGrid(ActorDesc->GetRuntimeGrid())
-	, Bounds(ActorDesc->GetBounds())
-{
-	check(GridPlacement != EActorGridPlacement::None);
-	Actors.Add(ActorDesc->GetGuid());
-	if (const AWorldDataLayers* WorldDataLayers = AWorldDataLayers::Get(World))
+	FActorCluster(const FWorldPartitionActorDesc* InActorDesc, EActorGridPlacement InGridPlacement, UWorld* InWorld)
+		: GridPlacement(InGridPlacement)
+		, RuntimeGrid(InActorDesc->GetRuntimeGrid())
+		, Bounds(InActorDesc->GetBounds())
 	{
-		for (const FName& DataLayerName : ActorDesc->GetDataLayers())
+		check(GridPlacement != EActorGridPlacement::None);
+
+		Actors.Add(InActorDesc->GetGuid());
+		if (const AWorldDataLayers* WorldDataLayers = AWorldDataLayers::Get(InWorld))
 		{
-			if (const UDataLayer* DataLayer = WorldDataLayers->GetDataLayerFromName(DataLayerName))
+			for (const FName& DataLayerName : InActorDesc->GetDataLayers())
 			{
-				if (DataLayer->IsDynamicallyLoaded())
+				if (const UDataLayer* DataLayer = WorldDataLayers->GetDataLayerFromName(DataLayerName))
 				{
-					DataLayers.Add(DataLayer);
+					if (DataLayer->IsDynamicallyLoaded())
+					{
+						DataLayers.Add(DataLayer);
+					}
 				}
 			}
 		}
+
+		DataLayersID = FDataLayersHelper::ComputeDataLayerID(DataLayers);
 	}
-	DataLayersID = FDataLayersHelper::ComputeDataLayerID(DataLayers);
-}
 
-void FActorCluster::Add(const FActorCluster& ActorCluster)
-{
-	// Merge Actors
-	Actors.Append(ActorCluster.Actors);
+	void Add(const FActorCluster& InActorCluster)
+	{
+		// Merge Actors
+		Actors.Append(InActorCluster.Actors);
 
-	// Merge RuntimeGrid
-	RuntimeGrid = RuntimeGrid == ActorCluster.RuntimeGrid ? RuntimeGrid : NAME_None;
+		// Merge RuntimeGrid
+		RuntimeGrid = RuntimeGrid == InActorCluster.RuntimeGrid ? RuntimeGrid : NAME_None;
 
-	// Merge Bounds
-	Bounds += ActorCluster.Bounds;
+		// Merge Bounds
+		Bounds += InActorCluster.Bounds;
 
-	// Merge GridPlacement
-	// If currently None, will always stay None
-	if (GridPlacement != EActorGridPlacement::None)
-	{
-		// If grid placement differs between the two clusters
-		if (GridPlacement != ActorCluster.GridPlacement)
+		// Merge GridPlacement
+		// If currently None, will always stay None
+		if (GridPlacement != EActorGridPlacement::None)
 		{
-			// If one of the two cluster was always loaded, set to None
-			if (ActorCluster.GridPlacement == EActorGridPlacement::AlwaysLoaded || GridPlacement == EActorGridPlacement::AlwaysLoaded)
+			// If grid placement differs between the two clusters
+			if (GridPlacement != InActorCluster.GridPlacement)
 			{
-				GridPlacement = EActorGridPlacement::None;
+				// If one of the two cluster was always loaded, set to None
+				if (InActorCluster.GridPlacement == EActorGridPlacement::AlwaysLoaded || GridPlacement == EActorGridPlacement::AlwaysLoaded)
+				{
+					GridPlacement = EActorGridPlacement::None;
+				}
+				else
+				{
+					GridPlacement = InActorCluster.GridPlacement;
+				}
 			}
-			else
+
+			// If current placement is set to Location, that won't make sense when merging two clusters. Set to Bounds
+			if (GridPlacement == EActorGridPlacement::Location)
 			{
-				GridPlacement = ActorCluster.GridPlacement;
+				GridPlacement = EActorGridPlacement::Bounds;
 			}
 		}
 
-		// If current placement is set to Location, that won't make sense when merging two clusters. Set to Bounds
-		if (GridPlacement == EActorGridPlacement::Location)
+		// Merge DataLayers
+		if (DataLayersID != InActorCluster.DataLayersID)
 		{
-			GridPlacement = EActorGridPlacement::Bounds;
-		}
-	}
-
-	// Merge DataLayers
-	if (DataLayersID != ActorCluster.DataLayersID)
-	{
-		for (const UDataLayer* DataLayer : ActorCluster.DataLayers)
-		{
-			check(DataLayer->IsDynamicallyLoaded());
-			DataLayers.AddUnique(DataLayer);
+			for (const UDataLayer* DataLayer : InActorCluster.DataLayers)
+			{
+				check(DataLayer->IsDynamicallyLoaded());
+				DataLayers.AddUnique(DataLayer);
+			}
+			DataLayersID = FDataLayersHelper::ComputeDataLayerID(DataLayers);
 		}
-		DataLayersID = FDataLayersHelper::ComputeDataLayerID(DataLayers);
 	}
-}
-
+};
 
-void CreateActorCluster(const FGuid& ActorGuid, const FWorldPartitionActorDesc* ActorDesc, TMap<FGuid, FActorCluster*>& ActorToActorCluster, TSet<FActorCluster*>& ActorClustersSet, UWorldPartition* WorldPartition)
+void CreateActorCluster(const FGuid& ActorGuid, const FWorldPartitionActorDesc* ActorDesc, EActorGridPlacement GridPlacement, TMap<FGuid, FActorCluster*>& ActorToActorCluster, TSet<FActorCluster*>& ActorClustersSet, UWorldPartition* WorldPartition)
 {
 	UWorld* World = WorldPartition->GetWorld();
 	FActorCluster* ActorCluster = ActorToActorCluster.FindRef(ActorGuid);
 	if (!ActorCluster)
 	{
-		ActorCluster = new FActorCluster(ActorDesc, World);
+		ActorCluster = new FActorCluster(ActorDesc, GridPlacement, World);
 		ActorClustersSet.Add(ActorCluster);
 		ActorToActorCluster.Add(ActorGuid, ActorCluster);
 	}
@@ -234,7 +234,7 @@ void CreateActorCluster(const FGuid& ActorGuid, const FWorldPartitionActorDesc*
 				else
 				{
 					// Put Reference in Actor's cluster
-					ActorCluster->Add(FActorCluster(ReferenceActorDesc, World));
+					ActorCluster->Add(FActorCluster(ReferenceActorDesc, GridPlacement, World));
 				}
 
 				// Map its cluster
@@ -249,11 +249,29 @@ TArray<FActorCluster> CreateActorClusters(UWorldPartition* WorldPartition)
 	TMap<FGuid, FActorCluster*> ActorToActorCluster;
 	TSet<FActorCluster*> ActorClustersSet;
 
+	// Gather all references to external actors from the level script
+	TSet<AActor*> LevelScriptExternalActorReferences;
+	if (ULevelScriptBlueprint* LevelScriptBlueprint = WorldPartition->GetWorld()->PersistentLevel->GetLevelScriptBlueprint(true))
+	{
+		LevelScriptExternalActorReferences.Append(ExternalActorsUtils::GetExternalActorReferences(LevelScriptBlueprint));
+	}
+
 	for (const auto& Pair : WorldPartition->Actors)
 	{
 		const FGuid& ActorGuid = Pair.Key;
 		const FWorldPartitionActorDesc* ActorDesc = Pair.Value.Get();
-		CreateActorCluster(ActorGuid, ActorDesc, ActorToActorCluster, ActorClustersSet, WorldPartition);
+		EActorGridPlacement GridPlacement = ActorDesc->GetGridPlacement();
+
+		// Check if the actor is loaded (potentially referenced by the level script)
+		if (AActor* Actor = ActorDesc->GetActor())
+		{
+			if (LevelScriptExternalActorReferences.Contains(Actor))
+			{
+				GridPlacement = EActorGridPlacement::AlwaysLoaded;
+			}
+		}
+
+		CreateActorCluster(ActorGuid, ActorDesc, GridPlacement, ActorToActorCluster, ActorClustersSet, WorldPartition);
 	}
 
 	TArray<FActorCluster> ActorClusters;
diff --git a/Engine/Source/Runtime/Engine/Public/ExternalActorsUtils.h b/Engine/Source/Runtime/Engine/Public/ExternalActorsUtils.h
new file mode 100644
--- /dev/null
+++ b/Engine/Source/Runtime/Engine/Public/ExternalActorsUtils.h
@@ -0,0 +1,13 @@
+// Copyright Epic Games, Inc. All Rights Reserved.
+
+#pragma once
+
+#include "CoreMinimal.h"
+
+namespace ExternalActorsUtils
+{
+	/**
+	 * Gather direct references to external actors from the root object.
+	 */
+	TArray<AActor*> GetExternalActorReferences(UObject* Root);
+}
\ No newline at end of file
diff --git a/Engine/Source/Runtime/Engine/Public/WorldPartition/WorldPartition.h b/Engine/Source/Runtime/Engine/Public/WorldPartition/WorldPartition.h
--- a/Engine/Source/Runtime/Engine/Public/WorldPartition/WorldPartition.h
+++ b/Engine/Source/Runtime/Engine/Public/WorldPartition/WorldPartition.h
@@ -22,7 +22,6 @@ DECLARE_LOG_CATEGORY_EXTERN(LogWorldPartition, Log, All);
 enum class EWorldPartitionStreamingMode
 {
 	PIE,
-	RuntimeExternalObjects,
 	RuntimeStreamingCells
 };
 
