Fixed DataLayersID not being deterministic between runs, making cooking of worlds using data layers non-deterministic also.

diff --git a/Engine/Source/Runtime/Core/Public/Misc/HashBuilder.h b/Engine/Source/Runtime/Core/Public/Misc/HashBuilder.h
--- a/Engine/Source/Runtime/Core/Public/Misc/HashBuilder.h
+++ b/Engine/Source/Runtime/Core/Public/Misc/HashBuilder.h
@@ -7,6 +7,13 @@
 #include "Concepts/GetTypeHashable.h"
 #include "Containers/UnrealString.h"
 
+/**
+ * Class for computing a hash of multiple types, going through GetTypeHash when the type implements it, and
+ * fallbacks to CRC32 when the type doesn't.
+ *
+ * Note: this hash builder should be used for transient hashes, as some types implements run-dependent hash
+ * computations, such as GetTypeHash(FName).
+ */
 class FHashBuilder
 {
 public:
diff --git a/Engine/Source/Runtime/Engine/Private/WorldPartition/DataLayer/DataLayerHelper.cpp b/Engine/Source/Runtime/Engine/Private/WorldPartition/DataLayer/DataLayerHelper.cpp
deleted file mode 100644
diff --git a/Engine/Source/Runtime/Engine/Private/WorldPartition/WorldPartition.cpp b/Engine/Source/Runtime/Engine/Private/WorldPartition/WorldPartition.cpp
--- a/Engine/Source/Runtime/Engine/Private/WorldPartition/WorldPartition.cpp
+++ b/Engine/Source/Runtime/Engine/Private/WorldPartition/WorldPartition.cpp
@@ -10,7 +10,6 @@
 #include "WorldPartition/WorldPartitionLevelStreamingPolicy.h"
 #include "WorldPartition/WorldPartitionActorDesc.h"
 #include "WorldPartition/WorldPartitionSubsystem.h"
-#include "WorldPartition/DataLayer/DataLayerHelper.h"
 #include "Algo/Accumulate.h"
 #include "Algo/Transform.h"
 #include "Algo/RemoveIf.h"
diff --git a/Engine/Source/Runtime/Engine/Private/WorldPartition/WorldPartitionRuntimeSpatialHash.cpp b/Engine/Source/Runtime/Engine/Private/WorldPartition/WorldPartitionRuntimeSpatialHash.cpp
--- a/Engine/Source/Runtime/Engine/Private/WorldPartition/WorldPartitionRuntimeSpatialHash.cpp
+++ b/Engine/Source/Runtime/Engine/Private/WorldPartition/WorldPartitionRuntimeSpatialHash.cpp
@@ -9,7 +9,6 @@
 #include "WorldPartition/WorldPartitionActorDescIterator.h"
 #include "WorldPartition/WorldPartitionStreamingPolicy.h"
 #include "WorldPartition/DataLayer/DataLayerSubsystem.h"
-#include "WorldPartition/DataLayer/DataLayerHelper.h"
 #include "WorldPartition/DataLayer/WorldDataLayers.h"
 #include "WorldPartition/DataLayer/DataLayer.h"
 #include "GameFramework/WorldSettings.h"
@@ -21,6 +20,7 @@
 #include "RenderUtils.h"
 
 #if WITH_EDITOR
+#include "MIsc/HashBuilder.h"
 #include "Editor/EditorEngine.h"
 #include "Misc/Parse.h"
 #include "Algo/Find.h"
@@ -41,6 +41,8 @@
 #include "Engine/WorldComposition.h"
 #include "LevelUtils.h"
 
+PRAGMA_DISABLE_OPTIMIZATION
+
 extern UNREALED_API class UEditorEngine* GEditor;
 #endif //WITH_EDITOR
 
@@ -111,6 +113,44 @@ static FAutoConsoleVariableRef CVarRuntimeSpatialHashCellToSourceAngleContributi
 // ------------------------------------------------------------------------------------------------
 
 #if WITH_EDITOR
+FDataLayersID::FDataLayersID()
+: Hash(0)
+{}
+
+FDataLayersID::FDataLayersID(const TArray<const UDataLayer*>& InDataLayers)
+{
+	Algo::TransformIf(InDataLayers, DataLayers, [](const UDataLayer* Item) { return Item->IsDynamicallyLoaded(); }, [](const UDataLayer* Item) { return Item->GetFName(); });
+	DataLayers.Sort([](const FName& A, const FName& B) { return A.FastLess(B); });
+	Hash = GetTypeHash(*this);
+}
+
+uint32 FDataLayersID::GetStableHash() const
+{
+	uint32 DataLayerId = 0;
+
+	if (DataLayers.Num())
+	{
+		TArray<FName> SortedDataLayers = DataLayers;
+		SortedDataLayers.Sort([](const FName& A, const FName& B) { return A.ToString() < B.ToString(); });
+
+		for (FName LayerName: SortedDataLayers)
+		{
+			DataLayerId = FCrc::StrCrc32(*LayerName.ToString(), DataLayerId);
+		}
+
+		check(DataLayerId);
+	}
+
+	return DataLayerId;
+}
+
+uint32 FDataLayersID::GetFastHash() const
+{
+	FHashBuilder HashBuilder;
+	HashBuilder << DataLayers;
+	return HashBuilder.GetHash();
+}
+
 // Clustering
 struct FActorCluster
 {
@@ -119,7 +159,7 @@ struct FActorCluster
 	FName						RuntimeGrid;
 	FBox						Bounds;
 	TArray<const UDataLayer*>	DataLayers;
-	uint32						DataLayersID;
+	FDataLayersID				DataLayersID;
 
 	FActorCluster(const FWorldPartitionActorDesc* InActorDesc, EActorGridPlacement InGridPlacement, UWorld* InWorld)
 		: GridPlacement(InGridPlacement)
@@ -143,7 +183,7 @@ struct FActorCluster
 			}
 		}
 
-		DataLayersID = FDataLayersHelper::ComputeDataLayerID(DataLayers);
+		DataLayersID = FDataLayersID(DataLayers);
 	}
 
 	void Add(const FActorCluster& InActorCluster)
@@ -190,7 +230,7 @@ struct FActorCluster
 				check(DataLayer->IsDynamicallyLoaded());
 				DataLayers.AddUnique(DataLayer);
 			}
-			DataLayersID = FDataLayersHelper::ComputeDataLayerID(DataLayers);
+			DataLayersID = FDataLayersID(DataLayers);
 		}
 	}
 };
@@ -510,26 +550,26 @@ struct FSquare2DGridHelper
 			FGridCellDataChunk(const TArray<const UDataLayer*>& InDataLayers)
 			{
 				Algo::TransformIf(InDataLayers, DataLayers, [](const UDataLayer* DataLayer) { return DataLayer->IsDynamicallyLoaded(); }, [](const UDataLayer* DataLayer) { return DataLayer; });
-				DataLayersID = FDataLayersHelper::ComputeDataLayerID(DataLayers);
+				DataLayersID = FDataLayersID(DataLayers);
 			}
 
 			void AddActor(const FGuid& Actor) { Actors.Add(Actor); }
 			const TSet<FGuid>& GetActors() const { return Actors; }
 			bool HasDataLayers() const { return !DataLayers.IsEmpty(); }
 			const TArray<const UDataLayer*>& GetDataLayers() const { return DataLayers; }
-			uint32 GetDataLayersID() const { return DataLayersID; }
+			const FDataLayersID& GetDataLayersID() const { return DataLayersID; }
 
 		private:
 			TSet<FGuid> Actors;
 			TArray<const UDataLayer*> DataLayers;
-			uint32 DataLayersID;
+			FDataLayersID DataLayersID;
 		};
 
 		struct FGridCell
 		{
 			void AddActor(const FGuid& InActor, const TArray<const UDataLayer*>& InDataLayers)
 			{
-				uint32 DataLayersID = FDataLayersHelper::ComputeDataLayerID(InDataLayers);
+				FDataLayersID DataLayersID = FDataLayersID(InDataLayers);
 				FGridCellDataChunk* ActorDataChunk = Algo::FindByPredicate(DataChunks, [&](FGridCellDataChunk& InDataChunk) { return InDataChunk.GetDataLayersID() == DataLayersID; });
 				if (!ActorDataChunk)
 				{
@@ -1279,13 +1319,13 @@ bool UWorldPartitionRuntimeSpatialHash::GenerateStreaming(EWorldPartitionStreami
 	return true;
 }
 
-FName UWorldPartitionRuntimeSpatialHash::GetCellName(FName InGridName, int32 InLevel, int32 InCellX, int32 InCellY, uint32 InDataLayerID) const
+FName UWorldPartitionRuntimeSpatialHash::GetCellName(FName InGridName, int32 InLevel, int32 InCellX, int32 InCellY, const FDataLayersID& InDataLayerID) const
 {
 	UWorldPartition* WorldPartition = GetOuterUWorldPartition();
 	const FString PackageName = FPackageName::GetShortName(WorldPartition->GetPackage());
 	const FString PackageNameNoPIEPrefix = UWorld::RemovePIEPrefix(PackageName);
 
-	return FName(*FString::Printf(TEXT("WPRT_%s_%s_Cell_L%d_X%02d_Y%02d_DL%X"), *PackageNameNoPIEPrefix, *InGridName.ToString(), InLevel, InCellX, InCellY, InDataLayerID));
+	return FName(*FString::Printf(TEXT("WPRT_%s_%s_Cell_L%d_X%02d_Y%02d_DL%X"), *PackageNameNoPIEPrefix, *InGridName.ToString(), InLevel, InCellX, InCellY, InDataLayerID.GetStableHash()));
 }
 
 void UWorldPartitionRuntimeSpatialHash::CacheHLODParents()
@@ -1639,7 +1679,7 @@ bool UWorldPartitionRuntimeSpatialHash::GenerateNavigationData()
 		DataChunkActor->SetActorLocation(FVector(CellCenter.X, CellCenter.Y, 0.f));
 		DataChunkActor->SetDataChunkActorBounds(CollectingBounds.ExpandBy(-1.f)); //reduce by 1cm to avoid precision issues
 
-		const FName CellName = GetCellName(RuntimeGrid.GridName, CellCoord.Z, CellCoord.X, CellCoord.Y, FDataLayersHelper::NoDataLayerID);
+		const FName CellName = GetCellName(RuntimeGrid.GridName, CellCoord.Z, CellCoord.X, CellCoord.Y);
 		DataChunkActor->SetActorLabel(FString::Printf(TEXT("NavDataChunkActor_%s_%s"), *GetName(), *CellName.ToString()));
 		
 		//@todo_ow: Properly handle data layers
diff --git a/Engine/Source/Runtime/Engine/Public/WorldPartition/DataLayer/DataLayerHelper.h b/Engine/Source/Runtime/Engine/Public/WorldPartition/DataLayer/DataLayerHelper.h
deleted file mode 100644
diff --git a/Engine/Source/Runtime/Engine/Public/WorldPartition/WorldPartitionRuntimeSpatialHash.h b/Engine/Source/Runtime/Engine/Public/WorldPartition/WorldPartitionRuntimeSpatialHash.h
--- a/Engine/Source/Runtime/Engine/Public/WorldPartition/WorldPartitionRuntimeSpatialHash.h
+++ b/Engine/Source/Runtime/Engine/Public/WorldPartition/WorldPartitionRuntimeSpatialHash.h
@@ -158,6 +158,37 @@ struct FSpatialHashRuntimeGrid
 #endif
 };
 
+#if WITH_EDITOR
+class FDataLayersID
+{
+public:
+	FDataLayersID();
+	FDataLayersID(const TArray<const UDataLayer*>& InDataLayers);
+
+	bool operator==(const FDataLayersID& Other) const
+	{
+		return (Hash == Other.Hash) && (DataLayers == Other.DataLayers);
+	}
+
+	bool operator!=(const FDataLayersID& Other) const
+	{
+		return !(*this == Other);
+	}
+
+	uint32 GetStableHash() const;
+	uint32 GetFastHash() const;
+
+	friend uint32 GetTypeHash(const FDataLayersID& DataLayersID)
+	{
+		return DataLayersID.GetFastHash();
+	}
+
+private:
+	TArray<FName> DataLayers;
+	uint32 Hash;
+};
+#endif
+
 UCLASS()
 class ENGINE_API UWorldPartitionRuntimeSpatialHash : public UWorldPartitionRuntimeHash
 {
@@ -186,7 +217,7 @@ protected:
 private:
 
 #if WITH_EDITOR
-	FName GetCellName(FName InGridName, int32 InLevel, int32 InCellX, int32 InCellY, uint32 InDataLayerID) const;
+	FName GetCellName(FName InGridName, int32 InLevel, int32 InCellX, int32 InCellY, const FDataLayersID& InDataLayerID = FDataLayersID()) const;
 
 	void CacheHLODParents();
 #endif
