Fix include capitalization and remove unmatched optimzation disable pragma

diff --git a/Engine/Source/Runtime/Engine/Private/WorldPartition/WorldPartitionRuntimeSpatialHash.cpp b/Engine/Source/Runtime/Engine/Private/WorldPartition/WorldPartitionRuntimeSpatialHash.cpp
--- a/Engine/Source/Runtime/Engine/Private/WorldPartition/WorldPartitionRuntimeSpatialHash.cpp
+++ b/Engine/Source/Runtime/Engine/Private/WorldPartition/WorldPartitionRuntimeSpatialHash.cpp
@@ -20,7 +20,7 @@
 #include "RenderUtils.h"
 
 #if WITH_EDITOR
-#include "MIsc/HashBuilder.h"
+#include "Misc/HashBuilder.h"
 #include "Editor/EditorEngine.h"
 #include "Misc/Parse.h"
 #include "Algo/Find.h"
@@ -41,8 +41,6 @@
 #include "Engine/WorldComposition.h"
 #include "LevelUtils.h"
 
-PRAGMA_DISABLE_OPTIMIZATION
-
 extern UNREALED_API class UEditorEngine* GEditor;
 #endif //WITH_EDITOR
 
