Fix non-unity issues
#rnx

diff --git a/Engine/Plugins/Experimental/ChaosCaching/Source/ChaosCachingEditor/Private/Chaos/CacheManagerCustomization.cpp b/Engine/Plugins/Experimental/ChaosCaching/Source/ChaosCachingEditor/Private/Chaos/CacheManagerCustomization.cpp
--- a/Engine/Plugins/Experimental/ChaosCaching/Source/ChaosCachingEditor/Private/Chaos/CacheManagerCustomization.cpp
+++ b/Engine/Plugins/Experimental/ChaosCaching/Source/ChaosCachingEditor/Private/Chaos/CacheManagerCustomization.cpp
@@ -6,6 +6,7 @@
 #include "DetailWidgetRow.h"
 #include "PropertyCustomizationHelpers.h"
 #include "Widgets/Input/SButton.h"
+#include "Widgets/Text/STextBlock.h"
 #include "Chaos/CacheManagerActor.h"
 #include "IDetailChildrenBuilder.h"
 
diff --git a/Engine/Plugins/Importers/USDImporter/Source/USDExporter/Private/MaterialExporterUSD.cpp b/Engine/Plugins/Importers/USDImporter/Source/USDExporter/Private/MaterialExporterUSD.cpp
--- a/Engine/Plugins/Importers/USDImporter/Source/USDExporter/Private/MaterialExporterUSD.cpp
+++ b/Engine/Plugins/Importers/USDImporter/Source/USDExporter/Private/MaterialExporterUSD.cpp
@@ -19,6 +19,7 @@
 #include "Materials/MaterialInterface.h"
 #include "Modules/ModuleManager.h"
 #include "UObject/UObjectGlobals.h"
+#include "Misc/Paths.h"
 
 UMaterialExporterUsd::UMaterialExporterUsd()
 {
diff --git a/Engine/Plugins/Importers/USDImporter/Source/USDSchemas/Private/USDShadeMaterialTranslator.cpp b/Engine/Plugins/Importers/USDImporter/Source/USDSchemas/Private/USDShadeMaterialTranslator.cpp
--- a/Engine/Plugins/Importers/USDImporter/Source/USDSchemas/Private/USDShadeMaterialTranslator.cpp
+++ b/Engine/Plugins/Importers/USDImporter/Source/USDSchemas/Private/USDShadeMaterialTranslator.cpp
@@ -8,6 +8,7 @@
 
 #include "UsdWrappers/SdfPath.h"
 
+#include "Engine/Texture.h"
 #include "Materials/Material.h"
 #include "Misc/SecureHash.h"
 
diff --git a/Engine/Plugins/Runtime/AR/ARUtilities/Source/ARUtilities/Private/ARUtilitiesFunctionLibrary.cpp b/Engine/Plugins/Runtime/AR/ARUtilities/Source/ARUtilities/Private/ARUtilitiesFunctionLibrary.cpp
--- a/Engine/Plugins/Runtime/AR/ARUtilities/Source/ARUtilities/Private/ARUtilitiesFunctionLibrary.cpp
+++ b/Engine/Plugins/Runtime/AR/ARUtilities/Source/ARUtilities/Private/ARUtilitiesFunctionLibrary.cpp
@@ -2,6 +2,7 @@
 
 #include "ARUtilitiesFunctionLibrary.h"
 #include "Materials/MaterialInstanceDynamic.h"
+#include "Engine/Texture.h"
 
 
 FVector2D UARUtilitiesFunctionLibrary::GetUVOffset(const FVector2D& ViewSize, const FVector2D& TextureSize)
diff --git a/Engine/Plugins/Runtime/AR/Google/GoogleARCore/Source/GoogleARCoreRendering/Public/GoogleARCorePassthroughCameraRenderer.h b/Engine/Plugins/Runtime/AR/Google/GoogleARCore/Source/GoogleARCoreRendering/Public/GoogleARCorePassthroughCameraRenderer.h
--- a/Engine/Plugins/Runtime/AR/Google/GoogleARCore/Source/GoogleARCoreRendering/Public/GoogleARCorePassthroughCameraRenderer.h
+++ b/Engine/Plugins/Runtime/AR/Google/GoogleARCore/Source/GoogleARCoreRendering/Public/GoogleARCorePassthroughCameraRenderer.h
@@ -21,6 +21,7 @@
 
 class UTexture;
 class FRHICommandListImmediate;
+struct FSceneViewFamily;
 
 /** A helper class that is used to load the GoogleARCorePassthroughCameraMaterial from its default object. */
 UCLASS()
diff --git a/Engine/Plugins/VirtualProduction/DMX/DMXProtocol/Source/DMXProtocolArtNet/Private/DMXProtocolReceivingRunnable.h b/Engine/Plugins/VirtualProduction/DMX/DMXProtocol/Source/DMXProtocolArtNet/Private/DMXProtocolReceivingRunnable.h
--- a/Engine/Plugins/VirtualProduction/DMX/DMXProtocol/Source/DMXProtocolArtNet/Private/DMXProtocolReceivingRunnable.h
+++ b/Engine/Plugins/VirtualProduction/DMX/DMXProtocol/Source/DMXProtocolArtNet/Private/DMXProtocolReceivingRunnable.h
@@ -11,6 +11,7 @@
 #include "HAL/CriticalSection.h"
 #include "Containers/Queue.h"
 #include "Misc/SingleThreadRunnable.h"
+#include "Misc/ScopeLock.h"
 
 struct FDMXProtocolArtNetDMXPacket;
 class FDMXProtocolArtNet;
diff --git a/Engine/Source/Runtime/Engine/Private/PhysicsEngine/BodyUtils.h b/Engine/Source/Runtime/Engine/Private/PhysicsEngine/BodyUtils.h
--- a/Engine/Source/Runtime/Engine/Private/PhysicsEngine/BodyUtils.h
+++ b/Engine/Source/Runtime/Engine/Private/PhysicsEngine/BodyUtils.h
@@ -7,6 +7,10 @@
 #include "Containers/ContainerAllocationPolicies.h"
 #include "PhysicsInterfaceDeclaresCore.h"
 
+#if PHYSICS_INTERFACE_PHYSX
+#include "extensions/PxMassProperties.h"
+#endif
+
 struct FBodyInstance; 
 
 #if WITH_CHAOS
@@ -19,13 +23,6 @@ namespace Chaos
 }
 #endif
 
-#if PHYSICS_INTERFACE_PHYSX
-namespace physx
-{
-	class PxMassProperties;
-}
-#endif
-
 /**
  * Utility methods for use by BodyInstance and ImmediatePhysics
  */
@@ -43,7 +40,7 @@ namespace BodyUtils
 #elif PHYSICS_INTERFACE_PHYSX
 	
 	/** Computes and adds the mass properties (inertia, com, etc...) based on the mass settings of the body instance. */
-	PxMassProperties ComputeMassProperties(const FBodyInstance* OwningBodyInstance, TArray<FPhysicsShapeHandle> Shapes, const FTransform& MassModifierTransform, const bool bUnused = false);
+	physx::PxMassProperties ComputeMassProperties(const FBodyInstance* OwningBodyInstance, TArray<FPhysicsShapeHandle> Shapes, const FTransform& MassModifierTransform, const bool bUnused = false);
 
 #endif
 
diff --git a/Engine/Source/Runtime/Engine/Public/Physics/ImmediatePhysics/ImmediatePhysicsPhysX/ImmediatePhysicsSimulation_PhysX.h b/Engine/Source/Runtime/Engine/Public/Physics/ImmediatePhysics/ImmediatePhysicsPhysX/ImmediatePhysicsSimulation_PhysX.h
--- a/Engine/Source/Runtime/Engine/Public/Physics/ImmediatePhysics/ImmediatePhysicsPhysX/ImmediatePhysicsSimulation_PhysX.h
+++ b/Engine/Source/Runtime/Engine/Public/Physics/ImmediatePhysics/ImmediatePhysicsPhysX/ImmediatePhysicsSimulation_PhysX.h
@@ -19,6 +19,8 @@
 #include "Physics/ImmediatePhysics/ImmediatePhysicsPhysX/ImmediatePhysicsKinematicTarget_PhysX.h"
 #include "Physics/ImmediatePhysics/ImmediatePhysicsPhysX/ImmediatePhysicsD6JointData_PhysX.h"
 
+#include "Chaos/ChaosEngineInterface.h"	//temp to help get code out of engine
+
 class UBodySetup;
 class UPhysicsConstraintTemplate;
 struct FBodyInstance;
diff --git a/Engine/Source/Runtime/Experimental/Chaos/Public/Chaos/ConstraintHandle.h b/Engine/Source/Runtime/Experimental/Chaos/Public/Chaos/ConstraintHandle.h
--- a/Engine/Source/Runtime/Experimental/Chaos/Public/Chaos/ConstraintHandle.h
+++ b/Engine/Source/Runtime/Experimental/Chaos/Public/Chaos/ConstraintHandle.h
@@ -16,7 +16,7 @@ namespace Chaos
 	class CHAOS_API FConstraintHandle
 	{
 	public:
-		enum EType
+		enum class EType : uint8
 		{
 			Invalid = 0,
 			Collision,
