Color VT tiles with compression errors hot pink, so they're consistent/eaiser to notice

diff --git a/Engine/Source/Runtime/Engine/Private/VT/VirtualTextureTranscodeCache.cpp b/Engine/Source/Runtime/Engine/Private/VT/VirtualTextureTranscodeCache.cpp
--- a/Engine/Source/Runtime/Engine/Private/VT/VirtualTextureTranscodeCache.cpp
+++ b/Engine/Source/Runtime/Engine/Private/VT/VirtualTextureTranscodeCache.cpp
@@ -45,6 +45,7 @@ struct FTranscodeTask
 		static uint8 OpaqueBlack[4] = { 0,0,0,255 };
 		static uint8 White[4] = { 0xFF, 0xFF, 0xFF, 0xFF };
 		static uint8 Flat[4] = { 127,127,255,255 };
+		static uint8 ErrorColor[4] = { 255, 0, 255, 255 };
 
 		const uint32 ChunkIndex = Params.ChunkIndex;
 		const FVirtualTextureDataChunk& Chunk = Params.VTData->Chunks[ChunkIndex];
@@ -94,6 +95,7 @@ struct FTranscodeTask
 			const FVTUploadTileBuffer& StagingBufferForLayer = StagingBuffer[LayerIndex];
 
 			const EVirtualTextureCodec VTCodec = Chunk.CodecType[LayerIndex];
+			bool bDecodeResult = true;
 			switch (VTCodec)
 			{
 			case EVirtualTextureCodec::Black:
@@ -129,7 +131,6 @@ struct FTranscodeTask
 				break;
 			case EVirtualTextureCodec::Crunch:
 			{
-				bool bResult = false;
 #if WITH_CRUNCH
 				// See if we can access compressed tile as a single contiguous block of memory
 				int64 DataReadSize = 0;
@@ -145,11 +146,12 @@ struct FTranscodeTask
 				check(Params.Codec);
 				const uint32 StagingBufferSize = StagingBufferForLayer.Stride * TileHeightInBlocks;
 				check(StagingBufferSize <= StagingBufferForLayer.MemorySize);
-				bResult = CrunchCompression::Decode(Params.Codec->Contexts[LayerIndex],
+				bDecodeResult = CrunchCompression::Decode(Params.Codec->Contexts[LayerIndex],
 					CompressedTile, TileLayerSize,
 					StagingBufferForLayer.Memory, StagingBufferSize, StagingBufferForLayer.Stride);
-#endif // WITH_CRUNCH
-				check(bResult);
+#else
+				bDecodeResult = false;
+#endif
 				break;
 			}
 			case EVirtualTextureCodec::ZippedGPU:
@@ -157,28 +159,36 @@ struct FTranscodeTask
 				{
 					// output buffer is tightly packed, can decompress directly
 					check(PackedOutputSize <= StagingBufferForLayer.MemorySize);
-					FCompression::UncompressMemoryStream(NAME_Zlib, StagingBufferForLayer.Memory, PackedOutputSize, Params.Data, DataOffset, TileLayerSize);
+					bDecodeResult = FCompression::UncompressMemoryStream(NAME_Zlib, StagingBufferForLayer.Memory, PackedOutputSize, Params.Data, DataOffset, TileLayerSize);
 				}
 				else
 				{
 					// output buffer has per-scanline padding, need to decompress to temp buffer, then copy line-by-line
 					check(PackedStride <= StagingBufferForLayer.Stride);
 					TempBuffer.SetNumUninitialized(PackedOutputSize, false);
-					FCompression::UncompressMemoryStream(NAME_Zlib, TempBuffer.GetData(), PackedOutputSize, Params.Data, DataOffset, TileLayerSize);
-					check(TileHeightInBlocks*StagingBufferForLayer.Stride <= StagingBufferForLayer.MemorySize);
-					for (uint32 y = 0; y < TileHeightInBlocks; ++y)
+					bDecodeResult = FCompression::UncompressMemoryStream(NAME_Zlib, TempBuffer.GetData(), PackedOutputSize, Params.Data, DataOffset, TileLayerSize);
+					if (bDecodeResult)
 					{
-						FMemory::Memcpy((uint8*)StagingBufferForLayer.Memory + y * StagingBufferForLayer.Stride, TempBuffer.GetData() + y * PackedStride, PackedStride);
+						check(TileHeightInBlocks * StagingBufferForLayer.Stride <= StagingBufferForLayer.MemorySize);
+						for (uint32 y = 0; y < TileHeightInBlocks; ++y)
+						{
+							FMemory::Memcpy((uint8*)StagingBufferForLayer.Memory + y * StagingBufferForLayer.Stride, TempBuffer.GetData() + y * PackedStride, PackedStride);
+						}
 					}
 				}
 				break;
 
 			default:
 				checkNoEntry();
-				UniformColorPixels(StagingBufferForLayer, TilePixelSize, TilePixelSize, LayerFormat, Black);
+				bDecodeResult = false;
 				break;
 			}
 
+			if (!bDecodeResult)
+			{
+				UniformColorPixels(StagingBufferForLayer, TilePixelSize, TilePixelSize, LayerFormat, ErrorColor);
+			}
+
 			// Bake debug borders directly into the tile pixels
 			if (TextureBorderGenerator::Enabled)
 			{
