Support for multiple HLOD levels
* Runtime grid actors (ASpatialHashRuntimeGridInfo, editor only, always loaded) are now used instead of relying on the UWorldPartitionRuntimeSpatialHash::HLODGrids - Avoid the needs to edit the level when rebuilding HLODs
* "viewmode hlodcoloration" now uses distinct colors for WP HLOD actors

diff --git a/Engine/Source/Editor/UnrealEd/Private/WorldPartition/WorldPartitionHLODsBuilder.cpp b/Engine/Source/Editor/UnrealEd/Private/WorldPartition/WorldPartitionHLODsBuilder.cpp
--- a/Engine/Source/Editor/UnrealEd/Private/WorldPartition/WorldPartitionHLODsBuilder.cpp
+++ b/Engine/Source/Editor/UnrealEd/Private/WorldPartition/WorldPartitionHLODsBuilder.cpp
@@ -31,37 +31,43 @@ bool UWorldPartitionHLODsBuilder::Run(UWorld* World, FPackageSourceControlHelper
 	UWorldPartition* WorldPartition = World->GetWorldPartition();
 	check(WorldPartition);
 
-	TSet<UPackage*> HLODActorPackages;
+	TSet<UPackage*> ActorPackages;
 
-	// Gather all packages before any HLOD actor is deleted
-	for (TActorIterator<AWorldPartitionHLOD> ItHLOD(World); ItHLOD; ++ItHLOD)
+	// Gather all packages
+	for (TActorIterator<AActor> It(World); It; ++It)
 	{
-		HLODActorPackages.Add(ItHLOD->GetPackage());
+		if (It->IsPackageExternal())
+		{
+			ActorPackages.Add(It->GetPackage());
+		}
 	}
 
 	// Rebuild HLOD for the whole world
 	WorldPartition->GenerateHLOD();
 	
-	// Gather all packages again to include newly created HLOD actors
-	for (TActorIterator<AWorldPartitionHLOD> ItHLOD(World); ItHLOD; ++ItHLOD)
+	// Gather all packages again to include newly created actors
+	for (TActorIterator<AActor> It(World); It; ++It)
 	{
-		HLODActorPackages.Add(ItHLOD->GetPackage());
+		if (It->IsPackageExternal())
+		{
+			ActorPackages.Add(It->GetPackage());
+		}
 	}
 
 	TArray<UPackage*> PackagesToSave;
 	TArray<UPackage*> PackagesToDelete;
 
-	for (UPackage* HLODActorPackage : HLODActorPackages)
+	for (UPackage* ActorPackage : ActorPackages)
 	{
-		if (HLODActorPackage && HLODActorPackage->IsDirty())
+		if (ActorPackage && ActorPackage->IsDirty())
 		{
-			if (UPackage::IsEmptyPackage(HLODActorPackage))
+			if (UPackage::IsEmptyPackage(ActorPackage))
 			{
-				PackagesToDelete.Add(HLODActorPackage);
+				PackagesToDelete.Add(ActorPackage);
 			}
 			else
 			{
-				PackagesToSave.Add(HLODActorPackage);
+				PackagesToSave.Add(ActorPackage);
 			}
 		}
 	}
@@ -118,4 +124,4 @@ bool UWorldPartitionHLODsBuilder::Run(UWorld* World, FPackageSourceControlHelper
 	}
 
 	return true;
-}
+}
\ No newline at end of file
diff --git a/Engine/Source/Runtime/Core/Public/UObject/UE5MainStreamObjectVersion.h b/Engine/Source/Runtime/Core/Public/UObject/UE5MainStreamObjectVersion.h
--- a/Engine/Source/Runtime/Core/Public/UObject/UE5MainStreamObjectVersion.h
+++ b/Engine/Source/Runtime/Core/Public/UObject/UE5MainStreamObjectVersion.h
@@ -40,6 +40,9 @@ struct CORE_API FUE5MainStreamObjectVersion
 		// Renaming UAnimSequence::NumFrames to NumberOfKeys, as that what is actually contains.
 		RenamingAnimationNumFrames,
 
+		// Serialize HLODLayer in WorldPartition HLODActorDesc
+		WorldPartitionHLODActorDescSerializeHLODLayer,
+
 		// -----<new versions can be added above this line>-------------------------------------------------
 		VersionPlusOne,
 		LatestVersion = VersionPlusOne - 1
diff --git a/Engine/Source/Runtime/Engine/Private/Actor.cpp b/Engine/Source/Runtime/Engine/Private/Actor.cpp
--- a/Engine/Source/Runtime/Engine/Private/Actor.cpp
+++ b/Engine/Source/Runtime/Engine/Private/Actor.cpp
@@ -420,9 +420,8 @@ void AActor::GetAssetRegistryTags(TArray<FAssetRegistryTag>& OutTags) const
 	{
 		if (UWorldPartition* WorldPartition = GetWorldSettings()->GetWorldPartition())
 		{
-			TUniquePtr<FWorldPartitionActorDesc> ActorDesc(WorldPartition->GetActorDescFactory(this)->Create());
-			ActorDesc->Init(this);
-
+			TUniquePtr<FWorldPartitionActorDesc> ActorDesc(UWorldPartition::CreateActorDesc(this));
+			
 			const FString ActorMetaDataClass = GetParentNativeClass(GetClass())->GetName();
 			static FName NAME_ActorMetaDataClass(TEXT("ActorMetaDataClass"));
 			OutTags.Add(UObject::FAssetRegistryTag(NAME_ActorMetaDataClass, ActorMetaDataClass, UObject::FAssetRegistryTag::TT_Hidden));
diff --git a/Engine/Source/Runtime/Engine/Private/StaticMeshRender.cpp b/Engine/Source/Runtime/Engine/Private/StaticMeshRender.cpp
--- a/Engine/Source/Runtime/Engine/Private/StaticMeshRender.cpp
+++ b/Engine/Source/Runtime/Engine/Private/StaticMeshRender.cpp
@@ -310,24 +310,28 @@ FStaticMeshSceneProxy::FStaticMeshSceneProxy(UStaticMeshComponent* InComponent,
 	}
 
 	// Setup Hierarchical LOD index
+	enum HLODColors
+	{
+		HLODNone	= 0,	// Not part of a HLOD cluster (draw as white when visualizing)
+		HLODChild	= 1,	// Part of a HLOD cluster but still a plain mesh
+		HLOD0		= 2		// Colors for HLOD levels start at index 2
+	};
+
 	if (ALODActor* LODActorOwner = Cast<ALODActor>(Owner))
 	{
-		// An HLOD cluster (they count from 1, but the colors for HLOD levels start at index 2)
-		HierarchicalLODIndex = LODActorOwner->LODLevel + 1;
+		HierarchicalLODIndex = HLODColors::HLOD0 + LODActorOwner->LODLevel - 1; // ALODActor::LODLevel count from 1
 	}
 	else if (AWorldPartitionHLOD* WorldPartitionHLODOwner = Cast<AWorldPartitionHLOD>(Owner))
 	{
-		HierarchicalLODIndex = 2;
+		HierarchicalLODIndex = HLODColors::HLOD0 + WorldPartitionHLODOwner->GetLODLevel();
 	}
 	else if (InComponent->GetLODParentPrimitive())
 	{
-		// Part of a HLOD cluster but still a plain mesh
-		HierarchicalLODIndex = 1;
+		HierarchicalLODIndex = HLODColors::HLODChild;
 	}
 	else
 	{
-		// Not part of a HLOD cluster (draw as white when visualizing)
-		HierarchicalLODIndex = 0;
+		HierarchicalLODIndex = HLODColors::HLODNone;
 	}
 
 	if (BodySetup)
diff --git a/Engine/Source/Runtime/Engine/Private/WorldPartition/HLOD/HLODActor.cpp b/Engine/Source/Runtime/Engine/Private/WorldPartition/HLOD/HLODActor.cpp
--- a/Engine/Source/Runtime/Engine/Private/WorldPartition/HLOD/HLODActor.cpp
+++ b/Engine/Source/Runtime/Engine/Private/WorldPartition/HLOD/HLODActor.cpp
@@ -50,12 +50,112 @@ void AWorldPartitionHLOD::EndPlay(const EEndPlayReason::Type EndPlayReason)
 
 #if WITH_EDITOR
 
+void AWorldPartitionHLOD::OnSubActorLoaded(AActor& Actor)
+{
+	check(!Actor.GetRootComponent() || Actor.GetRootComponent()->IsRegistered());
+	
+	bool bIsAlreadyInSet = false;
+	LoadedSubActors.Add(&Actor, &bIsAlreadyInSet);
+
+	if (!bIsAlreadyInSet)
+	{
+		if (LoadedSubActors.Num() == 1)
+		{
+			UpdateVisibility();
+		}
+	}
+}
+
+void AWorldPartitionHLOD::OnSubActorUnloaded(AActor& Actor)
+{
+	LoadedSubActors.Remove(&Actor);
+
+	// If HLOD has no more sub actors, ensure it is drawn at all time
+	if (LoadedSubActors.IsEmpty())
+	{
+		UpdateVisibility();
+	}
+}
+
+void AWorldPartitionHLOD::SetupLoadedSubActors()
+{
+	UWorld* World = GetWorld();
+	if (World && !World->IsGameWorld())
+	{
+		LoadedSubActors.Empty();
+
+		UWorldPartition* WorldPartition = World->GetWorldPartition();
+		check(WorldPartition);
+
+		for (const FGuid& SubActorGuid : SubActors)
+		{
+			const FWorldPartitionActorDesc* SubActorDesc = WorldPartition->GetActorDesc(SubActorGuid);
+			AActor* SubActor = SubActorDesc ? SubActorDesc->GetActor() : nullptr;
+			if (SubActor && SubActor->GetRootComponent() && SubActor->GetRootComponent()->IsRegistered())
+			{
+				OnSubActorLoaded(*SubActor);
+			}
+		}
+
+		UpdateVisibility();
+	}
+}
+
+void AWorldPartitionHLOD::ResetLoadedSubActors()
+{
+	UWorld* World = GetWorld();
+	if (World && !World->IsGameWorld())
+	{
+		for (const TWeakObjectPtr<AActor>& SubActor : LoadedSubActors)
+		{
+			if (!SubActor.IsValid())
+			{
+				continue;
+			}
+		}
+
+		LoadedSubActors.Empty();
+
+		UpdateVisibility();
+	}
+}
+
+void AWorldPartitionHLOD::PostRegisterAllComponents()
+{
+	Super::PostRegisterAllComponents();
+
+	if (!IsTemplate())
+	{
+		SetupLoadedSubActors();
+	}
+}
+
+void AWorldPartitionHLOD::PostUnregisterAllComponents()
+{
+	if (!IsTemplate())
+	{
+		ResetLoadedSubActors();
+	}
+
+	Super::PostUnregisterAllComponents();
+}
+
+void AWorldPartitionHLOD::UpdateVisibility()
+{
+	SetIsTemporarilyHiddenInEditor(HasLoadedSubActors());
+}
+
+bool AWorldPartitionHLOD::HasLoadedSubActors() const
+{
+	return !LoadedSubActors.IsEmpty();
+}
+
 EActorGridPlacement AWorldPartitionHLOD::GetDefaultGridPlacement() const
 {
 	return EActorGridPlacement::Location;
 }
 
-void AWorldPartitionHLOD::SetHLODPrimitives(const TArray<UPrimitiveComponent*>& InHLODPrimitives, float InFadeOutDistance)
+void AWorldPartitionHLOD::SetHLODPrimitives(const TArray<UPrimitiveComponent*>& InHLODPrimitives)
 {
 	TRACE_CPUPROFILER_EVENT_SCOPE(AWorldPartitionHLOD::SetHLODPrimitive);
 	check(!InHLODPrimitives.IsEmpty());
@@ -79,8 +179,6 @@ void AWorldPartitionHLOD::SetHLODPrimitives(const TArray<UPrimitiveComponent*>&
 		InHLODPrimitive->SetCollisionEnabled(ECollisionEnabled::NoCollision);
 		InHLODPrimitive->SetMobility(EComponentMobility::Static);
 
-		InHLODPrimitive->MinDrawDistance = InFadeOutDistance;
-
 		InHLODPrimitive->RegisterComponent();
 		InHLODPrimitive->MarkRenderStateDirty();
 	}
@@ -91,43 +189,26 @@ void AWorldPartitionHLOD::SetHLODPrimitives(const TArray<UPrimitiveComponent*>&
 	}
 }
 
-const FBox& AWorldPartitionHLOD::GetHLODBounds() const
-{
-	return HLODBounds;
-}
-
-void AWorldPartitionHLOD::SetHLODBounds(const FBox& InBounds)
-{
-	HLODBounds = InBounds;
-}
-
-void AWorldPartitionHLOD::GetActorBounds(bool bOnlyCollidingComponents, FVector& Origin, FVector& BoxExtent, bool bIncludeFromChildActors) const
-{
-	Super::GetActorBounds(bOnlyCollidingComponents, Origin, BoxExtent, bIncludeFromChildActors);
-
-	FBox Bounds = FBox(Origin - BoxExtent, Origin + BoxExtent);
-	Bounds += HLODBounds;
-	Bounds.GetCenterAndExtents(Origin, BoxExtent);
-}
-
-void AWorldPartitionHLOD::GetActorLocationBounds(bool bOnlyCollidingComponents, FVector& Origin, FVector& BoxExtent, bool bIncludeFromChildActors) const
-{
-	GetActorBounds(bOnlyCollidingComponents, Origin, BoxExtent, bIncludeFromChildActors);
-}
-
 void AWorldPartitionHLOD::SetChildrenPrimitives(const TArray<UPrimitiveComponent*>& InChildrenPrimitives)
 {
-	TSet<FGuid> SubActorsSet;
+	check(GetHLODComponent());
+
+	ResetLoadedSubActors();
+	SubActors.Empty();
 
 	UPrimitiveComponent* HLODComponent = GetHLODComponent();
 	check(HLODComponent);
 
 	for (UPrimitiveComponent* ChildPrimitive : InChildrenPrimitives)
 	{
-		SubActorsSet.Add(ChildPrimitive->GetOwner()->GetActorGuid());
+		AActor* SubActor = ChildPrimitive->GetOwner();
+		
+		if (!LoadedSubActors.Contains(SubActor))
+		{
+			OnSubActorLoaded(*SubActor);
+			SubActors.Add(SubActor->GetActorGuid());
+		}
 	}
-
-	SubActors = SubActorsSet.Array();
 }
 
 const TArray<FGuid>& AWorldPartitionHLOD::GetSubActors() const
@@ -135,12 +216,6 @@ const TArray<FGuid>& AWorldPartitionHLOD::GetSubActors() const
 	return SubActors;
 }
 
-void AWorldPartitionHLOD::SetHLODLayer(const UHLODLayer* InSubActorsHLODLayer, int32 InSubActorsHLODLevel)
-{
-	SubActorsHLODLayer = InSubActorsHLODLayer;
-	SubActorsHLODLevel = InSubActorsHLODLevel;
-}
-
 void AWorldPartitionHLOD::PostActorCreated()
 {
 	Super::PostActorCreated();
diff --git a/Engine/Source/Runtime/Engine/Private/WorldPartition/HLOD/HLODActorDesc.cpp b/Engine/Source/Runtime/Engine/Private/WorldPartition/HLOD/HLODActorDesc.cpp
--- a/Engine/Source/Runtime/Engine/Private/WorldPartition/HLOD/HLODActorDesc.cpp
+++ b/Engine/Source/Runtime/Engine/Private/WorldPartition/HLOD/HLODActorDesc.cpp
@@ -5,6 +5,7 @@
 #if WITH_EDITOR
 #include "Algo/Transform.h"
 #include "Algo/RemoveIf.h"
+#include "UObject/UE5MainStreamObjectVersion.h"
 #include "WorldPartition/HLOD/HLODActor.h"
 #include "WorldPartition/HLOD/HLODLayer.h"
 
@@ -14,26 +15,20 @@ void FHLODActorDesc::Init(const AActor* InActor)
 
 	const AWorldPartitionHLOD* HLODActor = CastChecked<AWorldPartitionHLOD>(InActor);
 	SubActors = HLODActor->GetSubActors();
-	HLODLayer = HLODActor->GetHLODLayer();
 }
 
 void FHLODActorDesc::Serialize(FArchive& Ar)
 {
+	Ar.UsingCustomVersion(FUE5MainStreamObjectVersion::GUID);
+
 	FWorldPartitionActorDesc::Serialize(Ar);
 
 	Ar << SubActors;
 
-	FString HLODLayerStr;
-	if (Ar.IsSaving())
-	{
-		HLODLayerStr = HLODLayer.ToString();
-	}
-	
-	Ar << HLODLayerStr;
-	
-	if (Ar.IsLoading())
+	if (Ar.CustomVer(FUE5MainStreamObjectVersion::GUID) < FUE5MainStreamObjectVersion::WorldPartitionHLODActorDescSerializeHLODLayer)
 	{
-		HLODLayer = HLODLayerStr;
+		FString HLODLayer_Deprecated;
+		Ar << HLODLayer_Deprecated;
 	}
 }
 #endif
diff --git a/Engine/Source/Runtime/Engine/Private/WorldPartition/HLOD/HLODBuilder.cpp b/Engine/Source/Runtime/Engine/Private/WorldPartition/HLOD/HLODBuilder.cpp
--- a/Engine/Source/Runtime/Engine/Private/WorldPartition/HLOD/HLODBuilder.cpp
+++ b/Engine/Source/Runtime/Engine/Private/WorldPartition/HLOD/HLODBuilder.cpp
@@ -45,17 +45,17 @@ public:
 		}
 
 		TArray<UPrimitiveComponent*> HLODPrimitives = InCreateComponentsFunc(HLODActor);
-		Algo::RemoveIf(HLODPrimitives, [](const UPrimitiveComponent* HLODPrimitive) { return !HLODPrimitive; });
+		HLODPrimitives.RemoveSwap(nullptr);
 
 		if (!HLODPrimitives.IsEmpty())
 		{
 			HLODActor->Modify();
-			HLODActor->SetHLODLayer(HLODLayer, 0);
-			HLODActor->SetHLODPrimitives(HLODPrimitives, CellLoadingRange);
-			HLODActor->SetHLODBounds(CellBounds);
+			HLODActor->SetHLODPrimitives(HLODPrimitives);
 			HLODActor->SetChildrenPrimitives(InSubComponents);
 			HLODActor->SetActorLabel(HLODActorName);
-			HLODActor->RuntimeGrid = HLODLayer->GetRuntimeGrid();
+			HLODActor->RuntimeGrid = HLODLayer->GetRuntimeGrid(HLODLevel);
+			HLODActor->SetLODLevel(HLODLevel);
+			HLODActor->SetHLODLayer(HLODLayer->GetParentLayer().Get());
 		}
 		else
 		{
@@ -69,16 +69,16 @@ public:
 		}
 	}
 
-	static TArray<UPrimitiveComponent*> GatherPrimitiveComponents(int32 iHLODLevel, const TArray<AActor*> InActors)
+	static TArray<UPrimitiveComponent*> GatherPrimitiveComponents(uint32 InHLODLevel, const TArray<const AActor*> InActors)
 	{
 		TArray<UPrimitiveComponent*> PrimitiveComponents;
-		for (AActor* SubActor : InActors)
+		for (const AActor* SubActor : InActors)
 		{
 			for (UActorComponent* SubComponent : SubActor->GetComponents())
 			{
 				if (UPrimitiveComponent* PrimitiveComponent = Cast<UPrimitiveComponent>(SubComponent))
 				{
-					if (UHLODLayer::ShouldIncludeInHLOD(PrimitiveComponent, iHLODLevel))
+					if (UHLODLayer::ShouldIncludeInHLOD(PrimitiveComponent, InHLODLevel))
 					{
 						PrimitiveComponents.Add(PrimitiveComponent);
 					}
@@ -92,7 +92,7 @@ public:
 	UWorld*				World;
 	UWorldPartition*	WorldPartition;
 	const UHLODLayer*	HLODLayer;
-	int32				iLevel;
+	uint32				HLODLevel;
 	FName				CellName;
 	FBox				CellBounds;
 	float				CellLoadingRange;
@@ -226,7 +226,7 @@ class FHLODBuilder_MeshSimplify : public FHLODBuilder
 	}
 };
 
-TArray<AWorldPartitionHLOD*> FHLODBuilderUtilities::BuildHLODs(UWorldPartition* InWorldPartition, FName InCellName, FBox InCellBounds, float InCellLoadingRange, const UHLODLayer* InHLODLayer, const TArray<AActor*>& InSubActors)
+TArray<AWorldPartitionHLOD*> FHLODBuilderUtilities::BuildHLODs(UWorldPartition* InWorldPartition, FName InCellName, const FBox& InCellBounds, const UHLODLayer* InHLODLayer, uint32 InHLODLevel, const TArray<const AActor*>& InSubActors)
 {
 	TUniquePtr<FHLODBuilder> HLODBuilder = nullptr;
 
@@ -257,9 +257,9 @@ TArray<AWorldPartitionHLOD*> FHLODBuilderUtilities::BuildHLODs(UWorldPartition*
 		HLODBuilder->World = InWorldPartition->GetWorld();
 		HLODBuilder->WorldPartition = InWorldPartition;
 		HLODBuilder->HLODLayer = InHLODLayer;
+		HLODBuilder->HLODLevel = InHLODLevel;
 		HLODBuilder->CellName = InCellName;
 		HLODBuilder->CellBounds = InCellBounds;
-		HLODBuilder->CellLoadingRange = InCellLoadingRange;
 
 		HLODBuilder->Build(SubComponents);
 			
diff --git a/Engine/Source/Runtime/Engine/Private/WorldPartition/HLOD/HLODBuilder.h b/Engine/Source/Runtime/Engine/Private/WorldPartition/HLOD/HLODBuilder.h
--- a/Engine/Source/Runtime/Engine/Private/WorldPartition/HLOD/HLODBuilder.h
+++ b/Engine/Source/Runtime/Engine/Private/WorldPartition/HLOD/HLODBuilder.h
@@ -25,12 +25,12 @@ public:
 	 * @param 	InWorldPartition	The WorldPartition for which we are building HLODs
 	 * @param 	InCellName			The name of the cell
 	 * @param 	InCellBounds		Bounds of the cell, will be assigned to the created HLOD actors.
-	 * @param 	InCellLoadingRange	The range at which the cell should be loaded. This is the distance where the HLOD actors should fade out.
 	 * @param 	InHLODLayer			The HLODLayer which will provide the HLOD building parameters.
+	 * @param	InHLODLevel			Level of HLOD (HLOD 0...N)
 	 * @param 	InSubActors			The actors from which we'll gather geometry to generate an HLOD mesh
 	 * @return The list of HLOD actors for this cell.
 	 */
-	static TArray<AWorldPartitionHLOD*> BuildHLODs(UWorldPartition* InWorldPartition, FName InCellName, FBox InCellBounds, float InCellLoadingRange, const UHLODLayer* InHLODLayer, const TArray<AActor*>& InSubActors);
+	static TArray<AWorldPartitionHLOD*> BuildHLODs(UWorldPartition* InWorldPartition, FName InCellName, const FBox& InCellBounds, const UHLODLayer* InHLODLayer, uint32 InHLODLevel, const TArray<const AActor*>& InSubActors);
 };
 
 #endif
\ No newline at end of file
diff --git a/Engine/Source/Runtime/Engine/Private/WorldPartition/HLOD/HLODLayer.cpp b/Engine/Source/Runtime/Engine/Private/WorldPartition/HLOD/HLODLayer.cpp
--- a/Engine/Source/Runtime/Engine/Private/WorldPartition/HLOD/HLODLayer.cpp
+++ b/Engine/Source/Runtime/Engine/Private/WorldPartition/HLOD/HLODLayer.cpp
@@ -8,6 +8,10 @@
 #include "WorldPartition/HLOD/HLODActor.h"
 
 #if WITH_EDITOR
+#include "UObject/ConstructorHelpers.h"
+#include "Materials/MaterialInterface.h"
+#include "Math/UnitConversion.h"
+
 #include "WorldPartition/WorldPartition.h"
 #include "WorldPartition/WorldPartitionActorDesc.h"
 #include "WorldPartition/HLOD/HLODBuilder.h"
@@ -16,31 +20,24 @@
 UHLODLayer::UHLODLayer(const FObjectInitializer& ObjectInitializer)
 	: Super(ObjectInitializer)
 #if WITH_EDITORONLY_DATA
-	, LoadingRange(30000)
+	, CellSize(3200)
+	, LoadingRange(12800)
 #endif
 {
+#if WITH_EDITORONLY_DATA
+	HLODMaterial = ConstructorHelpers::FObjectFinder<UMaterialInterface>(TEXT("/Engine/EngineMaterials/BaseFlattenMaterial")).Object;
+#endif
 }
 
 #if WITH_EDITOR
 
-TArray<AWorldPartitionHLOD*> UHLODLayer::GenerateHLODForCell(UWorldPartition* InWorldPartition, FName InCellName, FBox InCellBounds, float InCellLoadingRange, const TSet<FGuid>& InCellActors)
+TArray<AWorldPartitionHLOD*> UHLODLayer::GenerateHLODForCell(UWorldPartition* InWorldPartition, FName InCellName, const FBox& InCellBounds, uint32 InHLODLevel, const TArray<AActor*>& InCellActors)
 {
 	TRACE_CPUPROFILER_EVENT_SCOPE(UHLODLayer::GenerateHLODForCell);
 
-	TMap<UHLODLayer*, TArray<AActor*>> HLODLayersActors;
-	for (const FGuid& ActorGuid : InCellActors)
+	TMap<UHLODLayer*, TArray<const AActor*>> HLODLayersActors;
+	for (const AActor* Actor : InCellActors)
 	{
-		FWorldPartitionActorDesc& ActorDesc = *InWorldPartition->GetActorDesc(ActorGuid);
-
-		// Skip editor only actors - they might not be loaded and doesn't contribute to HLODs anyway
-		if (ActorDesc.GetActorIsEditorOnly())
-		{
-			continue;
-		}
-
-		AActor* Actor = ActorDesc.GetActor();
-		check(Actor);
-
 		if (ShouldIncludeInHLOD(Actor))
 		{
 			UHLODLayer* HLODLayer = UHLODLayer::GetHLODLayer(Actor);
@@ -54,14 +51,14 @@ TArray<AWorldPartitionHLOD*> UHLODLayer::GenerateHLODForCell(UWorldPartition* In
 	TArray<AWorldPartitionHLOD*> HLODActors;
 	for (const auto& HLODLayerActors : HLODLayersActors)
 	{
-		HLODActors += FHLODBuilderUtilities::BuildHLODs(InWorldPartition, InCellName, InCellBounds, InCellLoadingRange, HLODLayerActors.Key, HLODLayerActors.Value);
+		HLODActors += FHLODBuilderUtilities::BuildHLODs(InWorldPartition, InCellName, InCellBounds, HLODLayerActors.Key, InHLODLevel, HLODLayerActors.Value);
 	}
 	return HLODActors;
 }
 
-bool UHLODLayer::ShouldIncludeInHLOD(UPrimitiveComponent* InComponent, int32 InLevelIndex)
+bool UHLODLayer::ShouldIncludeInHLOD(const UPrimitiveComponent* InComponent, int32 InLevelIndex)
 {
-	if (UStaticMeshComponent* SMC = Cast<UStaticMeshComponent>(InComponent))
+	if (const UStaticMeshComponent* SMC = Cast<const UStaticMeshComponent>(InComponent))
 	{
 		if (!SMC->GetStaticMesh())
 		{
@@ -102,8 +99,29 @@ UHLODLayer* UHLODLayer::GetHLODLayer(const AActor* InActor)
 
 #endif // WITH_EDITOR
 
+#if WITH_EDITORONLY_DATA
+
+FName UHLODLayer::GetRuntimeGridName(uint32 InLODLevel, int32 InCellSize, float InLoadingRange)
+{
+	FNumericUnit<float> CellSize(InCellSize, EUnit::Centimeters);
+	FNumericUnit<float> CellSizeQuantized = CellSize.QuantizeUnitsToBestFit();
+	FString CellSizeString = FString::SanitizeFloat(CellSizeQuantized.Value, 0) + FUnitConversion::GetUnitDisplayString(CellSizeQuantized.Units);
+
+	FNumericUnit<float> LoadingRange(InLoadingRange, EUnit::Centimeters);
+	FNumericUnit<float> LoadingRangeQuantized = LoadingRange.QuantizeUnitsToBestFit();
+	FString LoadingRangeString = FString::SanitizeFloat(LoadingRangeQuantized.Value, 0) + FUnitConversion::GetUnitDisplayString(LoadingRangeQuantized.Units);
+
+	return *FString::Format(TEXT("HLOD{0}_{1}_{2}"), { InLODLevel, CellSizeString, LoadingRangeString });
+}
+
+FName UHLODLayer::GetRuntimeGrid(uint32 InHLODLevel) const
+{
+	return GetRuntimeGridName(InHLODLevel, CellSize, LoadingRange);
+}
+
+#endif // WITH_EDITOR
 
-bool UHLODLayer::ShouldIncludeInHLOD(AActor* InActor)
+bool UHLODLayer::ShouldIncludeInHLOD(const AActor* InActor)
 {
 	if (!InActor)
 	{
diff --git a/Engine/Source/Runtime/Engine/Private/WorldPartition/HLOD/HLODSubsystem.cpp b/Engine/Source/Runtime/Engine/Private/WorldPartition/HLOD/HLODSubsystem.cpp
--- a/Engine/Source/Runtime/Engine/Private/WorldPartition/HLOD/HLODSubsystem.cpp
+++ b/Engine/Source/Runtime/Engine/Private/WorldPartition/HLOD/HLODSubsystem.cpp
@@ -48,10 +48,62 @@ void UHLODSubsystem::Initialize(FSubsystemCollectionBase& Collection)
 	HLODActorDescFactory.Reset(new FHLODActorDescFactory());
 
 	RegisterActorDescFactories(WorldPartitionSubsystem);
+
+	GetWorld()->PersistentLevel->OnLoadedActorAddedToLevelEvent.AddUObject(this, &UHLODSubsystem::OnActorLoaded);
+	GetWorld()->PersistentLevel->OnLoadedActorRemovedFromLevelEvent.AddUObject(this, &UHLODSubsystem::OnActorUnloaded);
 #endif
 }
 
 #if WITH_EDITOR
+void UHLODSubsystem::OnActorLoaded(AActor& Actor)
+{
+	// If loading an HLOD actor, keep a map of Actor -> HLODActors
+	if (AWorldPartitionHLOD* HLODActor = Cast<AWorldPartitionHLOD>(&Actor))
+	{
+		for (const FGuid& SubActorGuid : HLODActor->GetSubActors())
+		{
+			// Keep track of unloaded subactors
+			check(!ActorsToHLOD.Contains(SubActorGuid));
+			ActorsToHLOD.Emplace(SubActorGuid, HLODActor);
+		}
+	}
+
+	if(!Actor.GetRootComponent() || Actor.GetRootComponent()->IsRegistered())
+	{
+		// If HLOD for this actor is already loaded, setup LOD parenting
+		AWorldPartitionHLOD** HLODActorPtr = ActorsToHLOD.Find(Actor.GetActorGuid());
+		if (HLODActorPtr)
+		{
+			(*HLODActorPtr)->OnSubActorLoaded(Actor);
+		}
+	}
+}
+
+void UHLODSubsystem::OnActorUnloaded(AActor& Actor)
+{
+	UWorld* World = GetWorld();
+	UWorldPartition* WorldPartition = World->GetWorldPartition();
+
+	// If bound to an HLOD, clear reference
+	{
+		AWorldPartitionHLOD** HLODActorPtr = ActorsToHLOD.Find(Actor.GetActorGuid());
+		if (HLODActorPtr)
+		{
+			(*HLODActorPtr)->OnSubActorUnloaded(Actor);
+		}
+	}
+
+	// If unloading an HLOD actor, clear map of Actor -> HLODActors
+	if (AWorldPartitionHLOD* HLODActor = Cast<AWorldPartitionHLOD>(&Actor))
+	{
+		for (const FGuid& SubActorGuid : HLODActor->GetSubActors())
+		{
+			int32 NumRemoved = ActorsToHLOD.Remove(SubActorGuid);
+			check(NumRemoved == 1);
+		}
+	}
+}
+
 void UHLODSubsystem::RegisterActorDescFactories(UWorldPartitionSubsystem* WorldPartitionSubsystem)
 {
 	WorldPartitionSubsystem->RegisterActorDescFactory(AWorldPartitionHLOD::StaticClass(), HLODActorDescFactory.Get());
diff --git a/Engine/Source/Runtime/Engine/Private/WorldPartition/WorldPartition.cpp b/Engine/Source/Runtime/Engine/Private/WorldPartition/WorldPartition.cpp
--- a/Engine/Source/Runtime/Engine/Private/WorldPartition/WorldPartition.cpp
+++ b/Engine/Source/Runtime/Engine/Private/WorldPartition/WorldPartition.cpp
@@ -483,6 +483,13 @@ FWorldPartitionActorDescFactory* UWorldPartition::GetActorDescFactory(const AAct
 	return GetActorDescFactory(Actor->GetClass());
 }
 
+TUniquePtr<FWorldPartitionActorDesc> UWorldPartition::CreateActorDesc(const AActor* Actor)
+{
+	TUniquePtr<FWorldPartitionActorDesc> ActorDesc(GetActorDescFactory(Actor)->Create());
+	ActorDesc->Init(Actor);
+	return ActorDesc;
+}
+
 void UWorldPartition::ForEachIntersectingActorDesc(const FBox& Box, TSubclassOf<AActor> ActorClass, TFunctionRef<bool(const FWorldPartitionActorDesc*)> Predicate) const
 {
 	EditorHash->ForEachIntersectingActor(Box, [&ActorClass, Predicate](const FWorldPartitionActorDesc* ActorDesc)
diff --git a/Engine/Source/Runtime/Engine/Private/WorldPartition/WorldPartitionRuntimeSpatialHash.cpp b/Engine/Source/Runtime/Engine/Private/WorldPartition/WorldPartitionRuntimeSpatialHash.cpp
--- a/Engine/Source/Runtime/Engine/Private/WorldPartition/WorldPartitionRuntimeSpatialHash.cpp
+++ b/Engine/Source/Runtime/Engine/Private/WorldPartition/WorldPartitionRuntimeSpatialHash.cpp
@@ -18,13 +18,14 @@
 #include "DrawDebugHelpers.h"
 #include "DisplayDebugHelpers.h"
 #include "RenderUtils.h"
+#include "Algo/Transform.h"
 
 #if WITH_EDITOR
 #include "Misc/HashBuilder.h"
 #include "Editor/EditorEngine.h"
 #include "Misc/Parse.h"
+#include "Misc/ScopedSlowTask.h"
 #include "Algo/Find.h"
-#include "Algo/Transform.h"
 #include "EngineUtils.h"
 
 #include "WorldPartition/HLOD/HLODLayer.h"
@@ -233,9 +234,11 @@ struct FActorCluster
 	}
 };
 
-void CreateActorCluster(const FGuid& ActorGuid, const FWorldPartitionActorDesc* ActorDesc, EActorGridPlacement GridPlacement, TMap<FGuid, FActorCluster*>& ActorToActorCluster, TSet<FActorCluster*>& ActorClustersSet, UWorldPartition* WorldPartition)
+void CreateActorCluster(const FWorldPartitionActorDesc* ActorDesc, EActorGridPlacement GridPlacement, TMap<FGuid, FActorCluster*>& ActorToActorCluster, TSet<FActorCluster*>& ActorClustersSet, UWorldPartition* WorldPartition)
 {
 	UWorld* World = WorldPartition->GetWorld();
+	const FGuid& ActorGuid = ActorDesc->GetGuid();
+
 	FActorCluster* ActorCluster = ActorToActorCluster.FindRef(ActorGuid);
 	if (!ActorCluster)
 	{
@@ -282,7 +285,9 @@ void CreateActorCluster(const FGuid& ActorGuid, const FWorldPartitionActorDesc*
 	}
 }
 
-TArray<FActorCluster> CreateActorClusters(UWorldPartition* WorldPartition)
+typedef TFunctionRef<bool(const FWorldPartitionActorDesc&)> FFilterPredicate;
+
+TArray<FActorCluster> CreateActorClustersImpl(UWorldPartition* WorldPartition, TOptional<FFilterPredicate> InFilterPredicate)
 {
 	TMap<FGuid, FActorCluster*> ActorToActorCluster;
 	TSet<FActorCluster*> ActorClustersSet;
@@ -296,7 +301,6 @@ TArray<FActorCluster> CreateActorClusters(UWorldPartition* WorldPartition)
 
 	for (const auto& Pair : WorldPartition->Actors)
 	{
-		const FGuid& ActorGuid = Pair.Key;
 		const FWorldPartitionActorDesc* ActorDesc = Pair.Value.Get();
 		EActorGridPlacement GridPlacement = ActorDesc->GetGridPlacement();
 
@@ -309,7 +313,10 @@ TArray<FActorCluster> CreateActorClusters(UWorldPartition* WorldPartition)
 			}
 		}
 
-		CreateActorCluster(ActorGuid, ActorDesc, GridPlacement, ActorToActorCluster, ActorClustersSet, WorldPartition);
+		if (!InFilterPredicate.IsSet() || InFilterPredicate.GetValue()(*ActorDesc))
+		{
+			CreateActorCluster(ActorDesc, GridPlacement, ActorToActorCluster, ActorClustersSet, WorldPartition);
+		}
 	}
 
 	TArray<FActorCluster> ActorClusters;
@@ -318,6 +325,17 @@ TArray<FActorCluster> CreateActorClusters(UWorldPartition* WorldPartition)
 	for (FActorCluster* ActorCluster : ActorClustersSet) { delete ActorCluster; }
 	return ActorClusters;
 }
+
+TArray<FActorCluster> CreateActorClusters(UWorldPartition* WorldPartition, const FFilterPredicate& InFilterPredicate)
+{
+	return CreateActorClustersImpl(WorldPartition, InFilterPredicate);
+}
+
+TArray<FActorCluster> CreateActorClusters(UWorldPartition* WorldPartition)
+{
+	return CreateActorClustersImpl(WorldPartition, TOptional<FFilterPredicate>());
+}
+
 #endif
 
 /**
@@ -1043,7 +1061,9 @@ void FSpatialHashStreamingGrid::Draw2D(UCanvas* Canvas, const TArray<FWorldParti
 // ------------------------------------------------------------------------------------------------
 
 #if WITH_EDITOR
-static FSquare2DGridHelper GetPartitionedActors(const UWorldPartition* WorldPartition, const FBox& WorldBounds, const FSpatialHashRuntimeGrid& Grid, const TArray<FActorCluster>& GridActors)
+typedef TFunctionRef<const FWorldPartitionActorDesc*(const FGuid&)> FGuidToActorDescFunc;
+
+static FSquare2DGridHelper GetPartitionedActors(const UWorldPartition* WorldPartition, const FBox& WorldBounds, const FSpatialHashRuntimeGrid& Grid, const TArray<FActorCluster>& GridActors, const FGuidToActorDescFunc& GuidToActorDescFunc)
 {
 	UE_SCOPED_TIMER(TEXT("GetPartitionedActors"), LogWorldPartitionRuntimeSpatialHash, Log);
 
@@ -1093,74 +1113,71 @@ static FSquare2DGridHelper GetPartitionedActors(const UWorldPartition* WorldPart
 			GridPlacement = EActorGridPlacement::AlwaysLoaded;
 		}
 
-		if (GridPlacement != EActorGridPlacement::AlwaysLoaded)
+		switch (GridPlacement)
 		{
-			switch (GridPlacement)
-			{
-			case EActorGridPlacement::Location:
+		case EActorGridPlacement::Location:
+		{
+			FIntVector2 CellCoords;
+			check(ActorCluster.Actors.Num() == 1);
+			const FGuid& ActorGuid = *ActorCluster.Actors.CreateConstIterator();
+			const FWorldPartitionActorDesc& ActorDesc = *GuidToActorDescFunc(ActorGuid);
+			if (PartitionedActors.GetLowestLevel().GetCellCoords(FVector2D(ActorDesc.GetOrigin()), CellCoords))
 			{
-				FIntVector2 CellCoords;
-				check(ActorCluster.Actors.Num() == 1);
-				const FGuid& ActorGuid = *ActorCluster.Actors.CreateConstIterator();
-				const FWorldPartitionActorDesc& ActorDesc = *WorldPartition->GetActorDesc(ActorGuid);
-				if (PartitionedActors.GetLowestLevel().GetCellCoords(FVector2D(ActorDesc.GetOrigin()), CellCoords))
+				TArray<const UDataLayer*> ActorDataLayers;
+				if (const AWorldDataLayers* WorldDataLayers = AWorldDataLayers::Get(WorldPartition->GetWorld()))
 				{
-					TArray<const UDataLayer*> ActorDataLayers;
-					if (const AWorldDataLayers* WorldDataLayers = AWorldDataLayers::Get(WorldPartition->GetWorld()))
+					for (const FName& DataLayerName : ActorDesc.GetDataLayers())
 					{
-						for (const FName& DataLayerName : ActorDesc.GetDataLayers())
+						if (const UDataLayer* DataLayer = WorldDataLayers->GetDataLayerFromName(DataLayerName))
 						{
-							if (const UDataLayer* DataLayer = WorldDataLayers->GetDataLayerFromName(DataLayerName))
-							{
-								ActorDataLayers.Add(DataLayer);
-							}
+							ActorDataLayers.Add(DataLayer);
 						}
 					}
-					PartitionedActors.GetLowestLevel().GetCell(CellCoords).AddActor(ActorGuid, ActorDataLayers);
-				}
-				else
-				{
-					GridPlacement = EActorGridPlacement::AlwaysLoaded;
-					bAlwaysLoadedPromotedOutOfGrid = true;
 				}
-				break;
+				PartitionedActors.GetLowestLevel().GetCell(CellCoords).AddActor(ActorGuid, ActorDataLayers);
 			}
-			case EActorGridPlacement::Bounds:
+			else
 			{
-				// Find grid level cell that encompasses the actor cluster and put actors in it.
-				bool bFoundCell = false;
-				for (FSquare2DGridHelper::FGridLevel& GridLevel : PartitionedActors.Levels)
+				GridPlacement = EActorGridPlacement::AlwaysLoaded;
+				bAlwaysLoadedPromotedOutOfGrid = true;
+			}
+			break;
+		}
+		case EActorGridPlacement::Bounds:
+		{
+			// Find grid level cell that encompasses the actor cluster and put actors in it.
+			bool bFoundCell = false;
+			for (FSquare2DGridHelper::FGridLevel& GridLevel : PartitionedActors.Levels)
+			{
+				int32 IntersectingCellCount = 0;
+				GridLevel.ForEachIntersectingCellsBreakable(ActorCluster.Bounds, [&IntersectingCellCount](const FIntVector2& Coords) { return ++IntersectingCellCount <= 1; });
+				if (IntersectingCellCount == 1)
 				{
-					int32 IntersectingCellCount = 0;
-					GridLevel.ForEachIntersectingCellsBreakable(ActorCluster.Bounds, [&IntersectingCellCount](const FIntVector2& Coords) { return ++IntersectingCellCount <= 1; });
-					if (IntersectingCellCount == 1)
+					GridLevel.ForEachIntersectingCells(ActorCluster.Bounds, [&GridLevel,&ActorCluster](const FIntVector2& Coords)
 					{
-						GridLevel.ForEachIntersectingCells(ActorCluster.Bounds, [&GridLevel,&ActorCluster](const FIntVector2& Coords)
-						{
-							GridLevel.GetCell(Coords).AddActors(ActorCluster.Actors, ActorCluster.DataLayers);
-						});
-						bFoundCell = true;
-						break;
-					}
-				}
-				if (!bFoundCell)
-				{
-					GridPlacement = EActorGridPlacement::AlwaysLoaded;
-					bAlwaysLoadedPromotedOutOfGrid = true;
+						GridLevel.GetCell(Coords).AddActors(ActorCluster.Actors, ActorCluster.DataLayers);
+					});
+					bFoundCell = true;
+					break;
 				}
-				break;
 			}
-			default:	
-				check(0);
+			if (!bFoundCell)
+			{
+				GridPlacement = EActorGridPlacement::AlwaysLoaded;
+				bAlwaysLoadedPromotedOutOfGrid = true;
 			}
+			break;
 		}
-			
-		if (GridPlacement == EActorGridPlacement::AlwaysLoaded)
+		case EActorGridPlacement::AlwaysLoaded:
 		{
 			PartitionedActors.GetAlwaysLoadedCell().AddActors(ActorCluster.Actors, ActorCluster.DataLayers);
+			break;
+		}
+		default:	
+			check(0);
 		}
 
-		if (!LogWorldPartitionRuntimeSpatialHash.IsSuppressed(ELogVerbosity::Verbose))
+		if (UE_LOG_ACTIVE(LogWorldPartitionRuntimeSpatialHash, Verbose))
 		{
 			if (ActorCluster.Actors.Num() > 1)
 			{
@@ -1176,7 +1193,7 @@ static FSquare2DGridHelper GetPartitionedActors(const UWorldPartition* WorldPart
 
 				for (const FGuid& ActorGuid : ActorCluster.Actors)
 				{
-					const FWorldPartitionActorDesc& Desc = *WorldPartition->GetActorDesc(ActorGuid);
+					const FWorldPartitionActorDesc& Desc = *GuidToActorDescFunc(ActorGuid);
 					UE_LOG(LogWorldPartitionRuntimeSpatialHash, Verbose, TEXT("   - Actor: %s (%s)"), *Desc.GetActorPath().ToString(), *ActorGuid.ToString(EGuidFormats::UniqueObjectGuid));
 					UE_LOG(LogWorldPartitionRuntimeSpatialHash, Verbose, TEXT("            %s"), *Desc.GetActorPackage().ToString());
 				}
@@ -1189,8 +1206,22 @@ static FSquare2DGridHelper GetPartitionedActors(const UWorldPartition* WorldPart
 	
 	return PartitionedActors;
 }
+
+static FSquare2DGridHelper GetPartitionedActors(const UWorldPartition* WorldPartition, const FBox& WorldBounds, const FSpatialHashRuntimeGrid& Grid, const TArray<FActorCluster>& GridActors)
+{
+	return GetPartitionedActors(WorldPartition, WorldBounds, Grid, GridActors, [WorldPartition](const FGuid& ActorGuid) { return WorldPartition->GetActorDesc(ActorGuid); });
+}
+
 #endif //WITH_EDITOR
 
+ASpatialHashRuntimeGridInfo::ASpatialHashRuntimeGridInfo(const FObjectInitializer& ObjectInitializer)
+	: Super(ObjectInitializer)
+{
+#if WITH_EDITORONLY_DATA
+	bListedInSceneOutliner = false;
+#endif
+}
+
 UWorldPartitionRuntimeSpatialHash::UWorldPartitionRuntimeSpatialHash(const FObjectInitializer& ObjectInitializer)
 	: Super(ObjectInitializer)
 {}
@@ -1269,7 +1300,12 @@ bool UWorldPartitionRuntimeSpatialHash::GenerateStreaming(EWorldPartitionStreami
 
 	TArray<FSpatialHashRuntimeGrid> AllGrids;
 	AllGrids.Append(Grids);
-	AllGrids.Append(HLODGrids);
+
+	// Append grids from runtime grid actors
+	for (TActorIterator<ASpatialHashRuntimeGridInfo> ItRuntimeGridActor(GetWorld()); ItRuntimeGridActor; ++ItRuntimeGridActor)
+	{
+		AllGrids.Add(ItRuntimeGridActor->GridSettings);
+	}
 
 	check(!StreamingGrids.Num());
 
@@ -1317,15 +1353,20 @@ bool UWorldPartitionRuntimeSpatialHash::GenerateStreaming(EWorldPartitionStreami
 	return true;
 }
 
-FName UWorldPartitionRuntimeSpatialHash::GetCellName(FName InGridName, int32 InLevel, int32 InCellX, int32 InCellY, const FDataLayersID& InDataLayerID) const
+static FName GetCellName(UWorldPartition* WorldPartition, FName InGridName, int32 InLevel, int32 InCellX, int32 InCellY, const FDataLayersID& InDataLayerID)
 {
-	UWorldPartition* WorldPartition = GetOuterUWorldPartition();
 	const FString PackageName = FPackageName::GetShortName(WorldPartition->GetPackage());
 	const FString PackageNameNoPIEPrefix = UWorld::RemovePIEPrefix(PackageName);
 
 	return FName(*FString::Printf(TEXT("WPRT_%s_%s_Cell_L%d_X%02d_Y%02d_DL%X"), *PackageNameNoPIEPrefix, *InGridName.ToString(), InLevel, InCellX, InCellY, InDataLayerID.GetStableHash()));
 }
 
+FName UWorldPartitionRuntimeSpatialHash::GetCellName(FName InGridName, int32 InLevel, int32 InCellX, int32 InCellY, const FDataLayersID& InDataLayerID) const
+{
+	UWorldPartition* WorldPartition = GetOuterUWorldPartition();
+	return ::GetCellName(WorldPartition, InGridName, InLevel, InCellX, InCellY, InDataLayerID);
+}
+
 void UWorldPartitionRuntimeSpatialHash::CacheHLODParents()
 {
 	CachedHLODParents.Reset();
@@ -1459,63 +1500,268 @@ void UWorldPartitionRuntimeSpatialHash::FlushStreaming()
 	StreamingGrids.Empty();
 }
 
-bool UWorldPartitionRuntimeSpatialHash::GenerateHLOD()
+TArray<AWorldPartitionHLOD*> GenerateHLODsForGrid(UWorldPartition* WorldPartition, const FSpatialHashRuntimeGrid& RuntimeGrid, uint32 HLODLevel, const TArray<FActorCluster>& ActorClusters, const FGuidToActorDescFunc& GuidToActorDescFunc)
 {
-	TRACE_CPUPROFILER_EVENT_SCOPE(UWorldPartitionRuntimeSpatialHash::GenerateHLOD);
+	const FBox WorldBounds = WorldPartition->GetWorldBounds();
 
-	FAutoScopedDurationTimer Timer;
+	const FSquare2DGridHelper PartitionedActors = GetPartitionedActors(WorldPartition, WorldBounds, RuntimeGrid, ActorClusters, GuidToActorDescFunc);
 
-	if (!Grids.Num())
+	// Quick pass to compute the number of cells we'll have to process, to provide a meaningful progress display
+	int32 NbCellsToProcess = 0;
+	PartitionedActors.ForEachCells([&PartitionedActors, &NbCellsToProcess](const FIntVector& CellCoord)
 	{
-		UE_LOG(LogWorldPartitionRuntimeSpatialHash, Error, TEXT("Invalid partition grids setup"));
-		return false;
-	}
+		const FSquare2DGridHelper::FGridLevel::FGridCell& GridCell = PartitionedActors.GetCell(CellCoord);
+		
+		// For now, HLOD only processes actors that are not in data layers
+		if (const FSquare2DGridHelper::FGridLevel::FGridCellDataChunk* GridCellDataChunk = GridCell.GetNoDataLayersDataChunk())
+		{
+			check(!GridCellDataChunk->HasDataLayers());
+			const bool bIsCellAlwaysLoaded = &GridCell == &PartitionedActors.GetAlwaysLoadedCell();
 
-	TArray<FSpatialHashRuntimeGrid> NewHLODGrids;
+			if (!bIsCellAlwaysLoaded && GridCellDataChunk->GetActors().Num() != 0)
+			{
+				NbCellsToProcess++;
+			}
+		}
+	});
+
+	FScopedSlowTask SlowTask(NbCellsToProcess, FText::FromString(FString::Printf(TEXT("Building HLODs for grid %s..."), *RuntimeGrid.GridName.ToString())));
+	SlowTask.MakeDialog();
+
+	TArray<AWorldPartitionHLOD*> GridHLODActors;
+	PartitionedActors.ForEachCells([&PartitionedActors, &GridHLODActors, &SlowTask, &GuidToActorDescFunc, RuntimeGrid, HLODLevel, WorldPartition, WorldBounds](const FIntVector& CellCoord)
+	{
+		const FSquare2DGridHelper::FGridLevel::FGridCell& GridCell = PartitionedActors.GetCell(CellCoord);
+
+		// For now, HLOD only processes actors that are not in data layers
+		if (const FSquare2DGridHelper::FGridLevel::FGridCellDataChunk* GridCellDataChunk = GridCell.GetNoDataLayersDataChunk())
+		{
+			check(!GridCellDataChunk->HasDataLayers());
+			const bool bIsCellAlwaysLoaded = &GridCell == &PartitionedActors.GetAlwaysLoadedCell();
+
+			if (!bIsCellAlwaysLoaded && GridCellDataChunk->GetActors().Num() != 0)
+			{
+				SlowTask.EnterProgressFrame(1);
+
+				FBox2D CellBounds2D;
+				PartitionedActors.GetCellBounds(CellCoord, CellBounds2D);
+
+				FName CellName = GetCellName(WorldPartition, RuntimeGrid.GridName, CellCoord.Z, CellCoord.X, CellCoord.Y, GridCellDataChunk->GetDataLayersID());
+				FBox CellBounds = FBox(FVector(CellBounds2D.Min, WorldBounds.Min.Z), FVector(CellBounds2D.Max, WorldBounds.Max.Z));
+
+				TArray<AActor*> CellActors;
+				Algo::Transform(GridCellDataChunk->GetActors(), CellActors, [&GuidToActorDescFunc](const FGuid& ActorGuid){ return GuidToActorDescFunc(ActorGuid)->GetActor(); });
+
+				TArray<AWorldPartitionHLOD*> CellHLODActors;
+				CellHLODActors = UHLODLayer::GenerateHLODForCell(WorldPartition, CellName, CellBounds, HLODLevel, CellActors);
+
+				GridHLODActors.Append(CellHLODActors);
+			}
+		}
+	});
+	return GridHLODActors;
+}
+
+// Find all HLOD grids from the HLODLayer assets we have and build a dependency graph
+void GatherHLODGrids(TMap<FName, FSpatialHashRuntimeGrid>& OutHLODGrids, TMap<FName, TSet<FName>>& OutHLODGridsGraph)
+{
+	// We don't know the HLOD level for those grids yet, we'll have to compute a graph of dependencies between our HLOD layers first
+	// Start with level 0, we'll rename once we have the correct info
+	const uint32 InitialHLODLevel = 0;
 
 	// Gather up all HLODLayer assets
 	FAssetRegistryModule& AssetRegistryModule = FModuleManager::LoadModuleChecked<FAssetRegistryModule>(TEXT("AssetRegistry"));
 	TArray<FAssetData> HLODLayerAssets;
 	AssetRegistryModule.Get().GetAssetsByClass(UHLODLayer::StaticClass()->GetFName(), HLODLayerAssets);
 
-	// Ensure all assets are loaded
+	// Find all used HLOD grids & build a dependency graph
 	for (const FAssetData& HLODLayerAsset : HLODLayerAssets)
 	{
 		if (const UHLODLayer* HLODLayer = Cast<UHLODLayer>(HLODLayerAsset.GetAsset()))
 		{
-			check(HLODLayer->GetLoadingRange() > 0);
+			FSpatialHashRuntimeGrid HLODGrid;
+			HLODGrid.CellSize = HLODLayer->GetCellSize();
+			HLODGrid.LoadingRange = HLODLayer->GetLoadingRange();
+			HLODGrid.DebugColor = FLinearColor::Red;
+			HLODGrid.GridName = HLODLayer->GetRuntimeGrid(InitialHLODLevel);
+
+			OutHLODGrids.Emplace(HLODGrid.GridName, HLODGrid);
 
-			FName RuntimeGrid = HLODLayer->GetRuntimeGrid();
-			if (RuntimeGrid == NAME_None)
+			TSet<FName>& HLODParentGrids = OutHLODGridsGraph.FindOrAdd(HLODGrid.GridName);
+			if (const UHLODLayer* ParentHLODLayer = HLODLayer->GetParentLayer().LoadSynchronous())
 			{
-				UE_LOG(LogWorldPartitionRuntimeSpatialHash, Error, TEXT("Invalid partition grids setup"));
-				return false;
+				HLODParentGrids.Add(ParentHLODLayer->GetRuntimeGrid(InitialHLODLevel));
 			}
+		}
+	}
+}
 
-			const FSpatialHashRuntimeGrid* ExistingGrid = NewHLODGrids.FindByPredicate([RuntimeGrid](const FSpatialHashRuntimeGrid& ExistingGrid) { return ExistingGrid.GridName == RuntimeGrid; });
-			if (ExistingGrid)
+// Sort HLOD grids in the order they'll need to be processed for HLOD generation
+bool SortHLODGrids(const TMap<FName, TSet<FName>>& InHLODGridsGraph, TArray<FName>& OutSortedGrids, TMap<FName, uint32>& OutGridsDepth)
+{
+	TSet<FName>		ProcessedGridsSet;	// Processed grids
+	TSet<FName>		VisitedGridsSet;	// Visited grids
+
+	TFunctionRef<bool(const FName, uint32)> VisitGraph = [&](const FName GridName, uint32 CurrentDepth) -> bool
+	{
+		uint32& GridDepth = OutGridsDepth.FindOrAdd(GridName);
+		GridDepth = FMath::Max(GridDepth, CurrentDepth);
+
+		if (ProcessedGridsSet.Contains(GridName))
+		{
+			return true;
+		}
+
+		// Detect cyclic dependencies between HLOD grids...
+		if (VisitedGridsSet.Contains(GridName))
+		{
+			return false; // Not a DAG
+		}
+
+		VisitedGridsSet.Add(GridName);
+
+		for (const FName ParentGridName : InHLODGridsGraph.FindChecked(GridName))
+		{
+			if (!VisitGraph(ParentGridName, CurrentDepth + 1))
 			{
-				UE_LOG(LogWorldPartitionRuntimeSpatialHash, Error, TEXT("Found two HLOD grids using the same RuntimeGrid name: %s"), *RuntimeGrid.ToString());
 				return false;
 			}
+		}
+
+		VisitedGridsSet.Remove(GridName);
+
+		ProcessedGridsSet.Add(GridName);
+		OutSortedGrids.Insert(GridName, 0);
+		return true;
+	};
 
-			ExistingGrid = Grids.FindByPredicate([RuntimeGrid](const FSpatialHashRuntimeGrid& ExistingGrid) { return ExistingGrid.GridName == RuntimeGrid; });
-			if (ExistingGrid)
+	bool bIsADAG = true;
+	for (const auto& HLODGridEntry : InHLODGridsGraph)
+	{
+		FName GridName = HLODGridEntry.Key;
+		if (!ProcessedGridsSet.Contains(GridName))
+		{
+			const uint32 HLODDepth = 0;
+			bIsADAG = VisitGraph(GridName, HLODDepth);
+			if (!bIsADAG)
 			{
-				UE_LOG(LogWorldPartitionRuntimeSpatialHash, Error, TEXT("HLODLayer RuntimeGrid must not match an existing grid name: %s"), *RuntimeGrid.ToString());
-				return false;
+				break;
 			}
+		}
+	}
 
-			FSpatialHashRuntimeGrid& HLODGrid = NewHLODGrids.AddDefaulted_GetRef();
-			HLODGrid.GridName = HLODLayer->GetRuntimeGrid();
-			HLODGrid.CellSize = HLODLayer->GetLoadingRange() * 2;	// @todo_ow: Proper setup
-			HLODGrid.LoadingRange = HLODLayer->GetLoadingRange();
-			HLODGrid.DebugColor = FLinearColor::Red;
+	// Remove leafs, those grids are the standard runtime grids, we'll generate HLOD for those in a first pass
+	OutSortedGrids.RemoveAll([&InHLODGridsGraph](const FName GridName) { return !InHLODGridsGraph.Contains(GridName); });
+
+	return bIsADAG;
+}
+
+void RenameHLODGrids(TMap<FName, FSpatialHashRuntimeGrid>& HLODGrids, TArray<FName>& SortedGrids, TMap<FName, uint32>& GridsDepth)
+{
+	// Build a mapping of old -> new names
+	TMap<FName, FName> NewNamesMapping;
+	for (const auto& HLODGrid : HLODGrids)
+	{
+		NewNamesMapping.Emplace(HLODGrid.Key, UHLODLayer::GetRuntimeGridName(GridsDepth[HLODGrid.Key], HLODGrid.Value.CellSize, HLODGrid.Value.LoadingRange));
+	}
+
+	// Remplace map entries
+	for (const auto& Names : NewNamesMapping)
+	{
+		FSpatialHashRuntimeGrid RuntimeGrid = HLODGrids.FindAndRemoveChecked(Names.Key);
+		RuntimeGrid.GridName = Names.Value;
+		HLODGrids.Emplace(Names.Value, RuntimeGrid);
+
+		int32 GridDepth = GridsDepth.FindAndRemoveChecked(Names.Key);
+		GridsDepth.Emplace(Names.Value, GridDepth);
+	}
+
+	// Replace arrays entries
+	for (FName& Name : SortedGrids)
+	{
+		Name = NewNamesMapping[Name];
+	}
+}
+
+// Create/destroy HLOD grid actors
+void UpdateHLODGridsActors(UWorld* World, const TMap<FName, FSpatialHashRuntimeGrid>& HLODGrids)
+{
+	static const FName HLODGridTag = TEXT("HLOD");
+	static const uint32 HLODGridTagLen = HLODGridTag.ToString().Len();
+
+	// Gather all existing HLOD grid actors, see if some are unused and needs to be deleted
+	TMap<FName, ASpatialHashRuntimeGridInfo*> ExistingGridActors;
+	for (TActorIterator<ASpatialHashRuntimeGridInfo> ItRuntimeGridActor(World); ItRuntimeGridActor; ++ItRuntimeGridActor)
+	{
+		ASpatialHashRuntimeGridInfo* GridActor = *ItRuntimeGridActor;
+		if (GridActor->ActorHasTag(HLODGridTag))
+		{
+			if (HLODGrids.Contains(GridActor->GridSettings.GridName))
+			{
+				ExistingGridActors.Emplace(GridActor->GridSettings.GridName, GridActor);
+			}
+			else
+			{
+				World->DestroyActor(GridActor);
+			}
 		}
 	}
 
-	// Recreate HLODGrids from known layers
-	HLODGrids = NewHLODGrids;
+	// Create missing HLOD grid actors 
+	for (const auto& HLODGridEntry : HLODGrids)
+	{
+		FName GridName = HLODGridEntry.Key;
+		if (!ExistingGridActors.Contains(GridName))
+		{
+			const FSpatialHashRuntimeGrid& GridSettings = HLODGridEntry.Value;
+
+			FActorSpawnParameters SpawnParams;
+			SpawnParams.bCreateActorPackage = true;
+			ASpatialHashRuntimeGridInfo* GridActor = World->SpawnActor<ASpatialHashRuntimeGridInfo>(SpawnParams);
+			GridActor->Tags.Add(HLODGridTag);
+			GridActor->SetActorLabel(GridSettings.GridName.ToString());
+			GridActor->GridSettings = GridSettings;
+
+			// Setup grid debug color to match HLODColorationColors
+			const uint32 LastLODColorationColorIdx = GEngine->HLODColorationColors.Num() - 1;
+			uint32 HLODLevel;
+			if (!LexTryParseString(HLODLevel, *GridSettings.GridName.ToString().Mid(HLODGridTagLen, 1)))
+			{
+				HLODLevel = LastLODColorationColorIdx;
+			}
+			HLODLevel = FMath::Clamp(HLODLevel + 2, 0U, LastLODColorationColorIdx);
+			GridActor->GridSettings.DebugColor = GEngine->HLODColorationColors[HLODLevel];
+		}
+	}
+}
+
+bool UWorldPartitionRuntimeSpatialHash::GenerateHLOD()
+{
+	TRACE_CPUPROFILER_EVENT_SCOPE(UWorldPartitionRuntimeSpatialHash::GenerateHLOD);
+
+	if (!Grids.Num())
+	{
+		UE_LOG(LogWorldPartitionRuntimeSpatialHash, Error, TEXT("Invalid partition grids setup"));
+		return false;
+	}
+
+	// Find all used HLOD grids & build a dependency graph
+	TMap<FName, FSpatialHashRuntimeGrid> HLODGrids;
+	TMap<FName, TSet<FName>>			 HLODGridsGraph;
+	GatherHLODGrids(HLODGrids, HLODGridsGraph);
+
+	// Sort HLOD grids in the order they'll need to be processed
+	TArray<FName>	SortedGrids;		// HLOD Grids, sorted in the order they'll need to be processed for HLOD generation
+	TMap<FName, uint32>	GridsDepth;		// Depth - Used to obtain the HLOD level
+	bool bIsADAG = SortHLODGrids(HLODGridsGraph, SortedGrids, GridsDepth);
+	if (!bIsADAG)
+	{
+		UE_LOG(LogWorldPartitionRuntimeSpatialHash, Error, TEXT("Invalid grids setup, cycles detected"));
+		return false;
+	}
+
+	// Now that we computed proper depth per grid, rename the grids
+	RenameHLODGrids(HLODGrids, SortedGrids, GridsDepth);
 
 	TMap<FName, int32> GridsMapping;
 	GridsMapping.Add(NAME_None, 0);
@@ -1526,20 +1772,16 @@ bool UWorldPartitionRuntimeSpatialHash::GenerateHLOD()
 		GridsMapping.Add(Grid.GridName, i);
 	}
 
-	// Ignore HLOD clusters as we are rebuilding them
-	for (const FSpatialHashRuntimeGrid& HLODGrid : HLODGrids)
-	{
-		check(!GridsMapping.Contains(HLODGrid.GridName));
-		GridsMapping.Add(HLODGrid.GridName, INDEX_NONE);
-	}
-
 	UWorldPartition* WorldPartition = GetOuterUWorldPartition();
 
-	// Create actor clusters
-	TArray<FActorCluster> ActorClusters = CreateActorClusters(WorldPartition);
+	// Create actor clusters - ignore HLOD actors
+	TArray<FActorCluster> ActorClusters = CreateActorClusters(WorldPartition, [](const FWorldPartitionActorDesc& ActorDesc)
+	{
+		return !ActorDesc.GetActorClass()->IsChildOf<AWorldPartitionHLOD>();
+	});
 
-	TArray<TArray<FActorCluster>> GridActors;
-	GridActors.InsertDefaulted(0, Grids.Num());
+	TArray<TArray<FActorCluster>> GridsClusters;
+	GridsClusters.InsertDefaulted(0, Grids.Num());
 
 	for (FActorCluster& ActorCluster : ActorClusters)
 	{
@@ -1550,59 +1792,74 @@ bool UWorldPartitionRuntimeSpatialHash::GenerateHLOD()
 		}
 
 		int32 GridIndex = FoundIndex ? *FoundIndex : 0;
-		if (GridIndex != INDEX_NONE)
-		{
-			GridActors[GridIndex].Add(MoveTemp(ActorCluster));
-		}
+		GridsClusters[GridIndex].Add(MoveTemp(ActorCluster));
 	}
 
-	// Keep track of all valid HLOD actors
-	TSet<AWorldPartitionHLOD*> ValidHLODActors;
+	// Keep track of all valid HLOD actors, along with which runtime grid they live in
+	TMap<FGuid, AWorldPartitionHLOD*>					ValidHLODActors;
+	TMap<FGuid, TUniquePtr<FWorldPartitionActorDesc>>	ValidHLODActorsDescs;
+	TMap<FName, TArray<AWorldPartitionHLOD*>>			GridsHLODActors;
 
-	const FBox WorldBounds = WorldPartition->GetWorldBounds();
-	for (int32 GridIndex = 0; GridIndex < Grids.Num(); GridIndex++)
+	auto GenerateHLODs = [&ValidHLODActors, &ValidHLODActorsDescs, &GridsHLODActors, WorldPartition, &GridsDepth](const FSpatialHashRuntimeGrid& RuntimeGrid, uint32 HLODLevel, const TArray<FActorCluster>& ActorClusters, const FGuidToActorDescFunc& GuidToActorDescFunc)
 	{
-		const FSpatialHashRuntimeGrid& RuntimeGrid = Grids[GridIndex];
-		const FSquare2DGridHelper PartionedActors = GetPartitionedActors(WorldPartition, WorldBounds, RuntimeGrid, GridActors[GridIndex]);
+		// Generate HLODs for this grid
+		TArray<AWorldPartitionHLOD*> HLODActors = GenerateHLODsForGrid(WorldPartition, RuntimeGrid, HLODLevel, ActorClusters, GuidToActorDescFunc);
 
-		PartionedActors.ForEachCells([PartionedActors, &ValidHLODActors, RuntimeGrid, WorldPartition, WorldBounds, this](const FIntVector& CellCoord)
+		for (AWorldPartitionHLOD* HLODActor : HLODActors)
 		{
-			const FSquare2DGridHelper::FGridLevel::FGridCell& GridCell = PartionedActors.GetCell(CellCoord);
-
-			// For now, HLOD only processes actors that are not in data layers
-			if (const FSquare2DGridHelper::FGridLevel::FGridCellDataChunk* GridCellDataChunk = GridCell.GetNoDataLayersDataChunk())
-			{
-				check(!GridCellDataChunk->HasDataLayers());
-				const bool bIsCellAlwaysLoaded = &GridCell == &PartionedActors.GetAlwaysLoadedCell();
+			ValidHLODActors.Emplace(HLODActor->GetActorGuid(), HLODActor);
+			ValidHLODActorsDescs.Emplace(HLODActor->GetActorGuid(), UWorldPartition::CreateActorDesc(HLODActor));
+			GridsHLODActors.FindOrAdd(HLODActor->RuntimeGrid).Add(HLODActor);
+		}
+	};
 
-				if (!bIsCellAlwaysLoaded && GridCellDataChunk->GetActors().Num() != 0)
-				{
-					FBox2D CellBounds2D;
-					PartionedActors.GetCellBounds(CellCoord, CellBounds2D);
+	// Generate HLODs for the standard runtime grids (HLOD 0)
+	for (int32 GridIndex = 0; GridIndex < Grids.Num(); GridIndex++)
+	{
+		// Generate HLODs for this grid - retrieve actors from the world partition
+		GenerateHLODs(Grids[GridIndex], 0, GridsClusters[GridIndex], [WorldPartition](const FGuid& ActorGuid)
+		{
+			return WorldPartition->GetActorDesc(ActorGuid);
+		});
+	}
 
-					FName CellName = GetCellName(RuntimeGrid.GridName, CellCoord.Z, CellCoord.X, CellCoord.Y, GridCellDataChunk->GetDataLayersID());
-					float CellLoadingRange = RuntimeGrid.LoadingRange;
-					FBox CellBounds = FBox(FVector(CellBounds2D.Min, WorldBounds.Min.Z), FVector(CellBounds2D.Max, WorldBounds.Max.Z));
+	// Now, go on and create HLOD actors from HLOD grids (HLOD 1-N)
+	for (const FName HLODGridName : SortedGrids)
+	{
+		// No need to process empty grids
+		if (!GridsHLODActors.Contains(HLODGridName))
+		{
+			HLODGrids.Remove(HLODGridName); // No need to keep this grid around, we have no actors in it
+			continue;
+		}
 
-					TArray<AWorldPartitionHLOD*> HLODActors;
-					HLODActors = UHLODLayer::GenerateHLODForCell(WorldPartition, CellName, CellBounds, CellLoadingRange, GridCellDataChunk->GetActors());
+		// Create one actor cluster per HLOD actor
+		TArray<FActorCluster> HLODActorClusters;
+		Algo::Transform(GridsHLODActors[HLODGridName], HLODActorClusters, [&ValidHLODActorsDescs](const AWorldPartitionHLOD* HLODActor)
+		{
+			FWorldPartitionActorDesc* ActorDesc = ValidHLODActorsDescs.FindChecked(HLODActor->GetActorGuid()).Get();
+			return FActorCluster(ActorDesc, ActorDesc->GetGridPlacement(), HLODActor->GetWorld());
+		});
 
-					ValidHLODActors.Append(HLODActors);
-				}
-			}
+		// Generate HLODs for this grid - retrieve actors from our ValidHLODActors map
+		// We can't rely on actor descs for newly created HLOD actors
+		GenerateHLODs(HLODGrids[HLODGridName], GridsDepth[HLODGridName] + 1, HLODActorClusters, [&ValidHLODActorsDescs](const FGuid& ActorGuid)
+		{
+			return ValidHLODActorsDescs.FindChecked(ActorGuid).Get();
 		});
 	}
 
-	// Destroy all invalid HLOD actors
+	// Destroy all unreferenced HLOD actors
 	for (TActorIterator<AWorldPartitionHLOD> ItHLOD(GetWorld()); ItHLOD; ++ItHLOD)
 	{
-		if (!ValidHLODActors.Contains(*ItHLOD))
+		if (!ValidHLODActors.Contains(ItHLOD->GetActorGuid()))
 		{
 			GetWorld()->DestroyActor(*ItHLOD);
 		}
 	}
 
-	CollectGarbage(RF_Standalone);
+	// Create/destroy HLOD grid actors
+	UpdateHLODGridsActors(GetWorld(), HLODGrids);
 
 	return true;
 }
@@ -1861,11 +2118,16 @@ void UWorldPartitionRuntimeSpatialHash::Draw2D(UCanvas* Canvas, const TArray<FWo
 	const FVector2D GridScreenHalfExtent = 0.5f * GridScreenExtent;
 	const FVector2D GridScreenInitialOffset = PartitionCanvasOffset;
 
+	// Sort streaming grids to render them sorted by loading range
+	TArray<const FSpatialHashStreamingGrid*> SortedStreamingGrids;
+	Algo::Transform(StreamingGrids, SortedStreamingGrids, [](const FSpatialHashStreamingGrid& StreamingGrid) { return &StreamingGrid; });
+	SortedStreamingGrids.Sort([](const FSpatialHashStreamingGrid& A, const FSpatialHashStreamingGrid& B) { return A.LoadingRange < B.LoadingRange; });
+
 	int32 GridIndex = 0;
-	for (const FSpatialHashStreamingGrid& StreamingGrid : StreamingGrids)
+	for (const FSpatialHashStreamingGrid* StreamingGrid : SortedStreamingGrids)
 	{
 		// Display view sides based on extended grid loading range (minimum of N cells)
-		const float GridSideDistance = FMath::Max((2.f * StreamingGrid.GetLoadingRange() * GridViewLoadingRangeExtentRatio), StreamingGrid.CellSize * GridViewMinimumSizeInCellCount);
+		const float GridSideDistance = FMath::Max((2.f * StreamingGrid->GetLoadingRange() * GridViewLoadingRangeExtentRatio), StreamingGrid->CellSize * GridViewMinimumSizeInCellCount);
 		FSphere AverageSphere(ForceInit);
 		for (const FWorldPartitionStreamingSource& Source : Sources)
 		{
@@ -1878,7 +2140,7 @@ void UWorldPartitionRuntimeSpatialHash::Draw2D(UCanvas* Canvas, const TArray<FWo
 		const float WorldToScreenScale = (0.5f * GridEffectiveScreenSize) / AverageSphere.W;
 		auto WorldToScreen = [&](const FVector2D& WorldPos) { return (WorldToScreenScale * (WorldPos - GridReferenceWorldPos)) + GridScreenOffset; };
 
-		StreamingGrid.Draw2D(Canvas, Sources, Region, GridScreenBounds, WorldToScreen);
+		StreamingGrid->Draw2D(Canvas, Sources, Region, GridScreenBounds, WorldToScreen);
 
 		// Draw WorldPartition name
 		FVector2D GridInfoPos = GridScreenOffset - GridScreenHalfExtent;
@@ -1893,7 +2155,7 @@ void UWorldPartitionRuntimeSpatialHash::Draw2D(UCanvas* Canvas, const TArray<FWo
 
 		// Draw Grid name, loading range
 		{
-			FString GridInfoText = FString::Printf(TEXT("%s | %d m"), *StreamingGrid.GridName.ToString(), int32(StreamingGrid.GetLoadingRange() * 0.01f));
+			FString GridInfoText = FString::Printf(TEXT("%s | %d m"), *StreamingGrid->GridName.ToString(), int32(StreamingGrid->GetLoadingRange() * 0.01f));
 			Canvas->SetDrawColor(255, 255, 0);
 			Canvas->DrawText(GEngine->GetTinyFont(), GridInfoText, GridInfoPos.X, GridInfoPos.Y);
 		}
diff --git a/Engine/Source/Runtime/Engine/Public/WorldPartition/HLOD/HLODActor.h b/Engine/Source/Runtime/Engine/Public/WorldPartition/HLOD/HLODActor.h
--- a/Engine/Source/Runtime/Engine/Public/WorldPartition/HLOD/HLODActor.h
+++ b/Engine/Source/Runtime/Engine/Public/WorldPartition/HLOD/HLODActor.h
@@ -3,7 +3,6 @@
 
 #include "GameFramework/Actor.h"
 #include "Containers/Array.h"
-#include "Containers/Map.h"
 #include "Containers/Set.h"
 #include "WorldPartition/WorldPartitionRuntimeCell.h"
 #include "HLODActor.generated.h"
@@ -23,55 +22,55 @@ public:
 	void OnCellHidden(FName InCellName);
 
 	const FGuid& GetHLODGuid() const { return HLODGuid; }
+	inline uint32 GetLODLevel() const { return LODLevel; }
 
-protected:
-	//~ Begin AActor Interface.
-	virtual void BeginPlay() override;
-	virtual void EndPlay(const EEndPlayReason::Type EndPlayReason) override;
-	//~ End AActor Interface.
-
-	UPrimitiveComponent* GetHLODComponent();
-
-public:
 #if WITH_EDITOR
-	void SetHLODPrimitives(const TArray<UPrimitiveComponent*>& InHLODPrimitives, float InFadeOutDistance);
+	void OnSubActorLoaded(AActor& Actor);
+	void OnSubActorUnloaded(AActor& Actor);
+
+	void SetHLODPrimitives(const TArray<UPrimitiveComponent*>& InHLODPrimitives);
 	void SetChildrenPrimitives(const TArray<UPrimitiveComponent*>& InChildrenPrimitives);
 
 	const TArray<FGuid>& GetSubActors() const;
 
-	void SetHLODLayer(const UHLODLayer* InSubActorsHLODLayer, int32 InSubActorsHLODLevel);
-	
-	const FBox& GetHLODBounds() const;
-	void SetHLODBounds(const FBox& InBounds);
+	inline void SetLODLevel(uint32 InLODLevel) { LODLevel = InLODLevel; }
+#endif // WITH_EDITOR
 
 protected:
 	//~ Begin AActor Interface.
+	virtual void BeginPlay() override;
+	virtual void EndPlay(const EEndPlayReason::Type EndPlayReason) override;
+#if WITH_EDITOR
+	virtual void PostRegisterAllComponents() override;
+	virtual void PostUnregisterAllComponents() override;
 	virtual EActorGridPlacement GetDefaultGridPlacement() const override;
 	virtual void PostActorCreated() override;
-	virtual void GetActorBounds(bool bOnlyCollidingComponents, FVector& Origin, FVector& BoxExtent, bool bIncludeFromChildActors) const override;
-	virtual void GetActorLocationBounds(bool bOnlyCollidingComponents, FVector& Origin, FVector& BoxExtent, bool bIncludeFromChildActors) const override;
+#endif
 	//~ End AActor Interface.
+
+	UPrimitiveComponent* GetHLODComponent();
+
+#if WITH_EDITOR
+	void UpdateVisibility();
+	bool HasLoadedSubActors() const;
+
+	void ResetLoadedSubActors();
+	void SetupLoadedSubActors();
 #endif // WITH_EDITOR
 
 private:
 #if WITH_EDITORONLY_DATA
 	UPROPERTY()
 	TArray<FGuid> SubActors;
-
-	UPROPERTY()
-	const UHLODLayer* SubActorsHLODLayer;
-
-	UPROPERTY()
-	FBox HLODBounds;
-
-	FDelegateHandle ActorRegisteredDelegateHandle;
-#endif
 	
-	UPROPERTY()
-	int32 SubActorsHLODLevel;
+	TSet<TWeakObjectPtr<AActor>> LoadedSubActors;
+#endif
 
 	UPROPERTY()
 	FGuid HLODGuid;
+
+	UPROPERTY()
+	uint32 LODLevel;
 };
 
 UCLASS()
diff --git a/Engine/Source/Runtime/Engine/Public/WorldPartition/HLOD/HLODActorDesc.h b/Engine/Source/Runtime/Engine/Public/WorldPartition/HLOD/HLODActorDesc.h
--- a/Engine/Source/Runtime/Engine/Public/WorldPartition/HLOD/HLODActorDesc.h
+++ b/Engine/Source/Runtime/Engine/Public/WorldPartition/HLOD/HLODActorDesc.h
@@ -19,13 +19,11 @@ class ENGINE_API FHLODActorDesc : public FWorldPartitionActorDesc
 
 public:
 	inline const TArray<FGuid>& GetSubActors() const { return SubActors; }
-	inline const FSoftObjectPath& GetHLODLayer() const { return HLODLayer; }
 
 protected:
 	virtual void Init(const AActor* InActor) override;
 	virtual void Serialize(FArchive& Ar) override;
 
 	TArray<FGuid> SubActors;
-	FSoftObjectPath	HLODLayer;
 #endif
 };
diff --git a/Engine/Source/Runtime/Engine/Public/WorldPartition/HLOD/HLODLayer.h b/Engine/Source/Runtime/Engine/Public/WorldPartition/HLOD/HLODLayer.h
--- a/Engine/Source/Runtime/Engine/Public/WorldPartition/HLOD/HLODLayer.h
+++ b/Engine/Source/Runtime/Engine/Public/WorldPartition/HLOD/HLODLayer.h
@@ -15,6 +15,7 @@
 
 class UWorldPartition;
 class AWorldPartitionHLOD;
+class AActor;
 
 UENUM()
 enum class EHLODLayerType : uint8
@@ -31,10 +32,10 @@ class ENGINE_API UHLODLayer : public UObject
 	
 #if WITH_EDITOR
 public:
-	static TArray<AWorldPartitionHLOD*> GenerateHLODForCell(UWorldPartition* InWorldPartition, FName InCellName, FBox InCellBounds, float InCellLoadingRange, const TSet<FGuid>& InCellActors);
+	static TArray<AWorldPartitionHLOD*> GenerateHLODForCell(UWorldPartition* InWorldPartition, FName InCellName, const FBox& InCellBounds, uint32 InLODLevel, const TArray<AActor*>& InCellActors);
 
 	static UHLODLayer* GetHLODLayer(const AActor* InActor);
-	static bool ShouldIncludeInHLOD(UPrimitiveComponent* InComponent, int32 InLevelIndex);
+	static bool ShouldIncludeInHLOD(const UPrimitiveComponent* InComponent, int32 InLevelIndex);
 #endif
 
 #if WITH_EDITORONLY_DATA
@@ -42,11 +43,15 @@ public:
 	const FMeshMergingSettings& GetMeshMergeSettings() const { return MeshMergeSettings; }
 	const FMeshProxySettings& GetMeshSimplifySettings() const { return MeshSimplifySettings; }
 	const TSoftObjectPtr<UMaterial>& GetHLODMaterial() const { return HLODMaterial; }
-	FName GetRuntimeGrid() const { return RuntimeGrid; }
+	FName GetRuntimeGrid(uint32 InHLODLevel) const;
+	int32 GetCellSize() const { return CellSize; }
 	float GetLoadingRange() const { return LoadingRange; }
+	const TSoftObjectPtr<UHLODLayer>& GetParentLayer() const { return ParentLayer; }
+
+	static FName GetRuntimeGridName(uint32 InLODLevel, int32 InCellSize, float InLoadingRange);
 #endif
 
-	static bool ShouldIncludeInHLOD(AActor* InActor);
+	static bool ShouldIncludeInHLOD(const AActor* InActor);
 
 #if WITH_EDITORONLY_DATA
 private:
@@ -66,14 +71,15 @@ private:
 	UPROPERTY(EditAnywhere, Config, Category=HLOD, meta = (editcondition = "LayerType == EHLODLayerType::MeshMerge || LayerType == EHLODLayerType::MeshSimplify"))
 	TSoftObjectPtr<UMaterial> HLODMaterial;
 
-	/**
-	 * Runtime grid that will be used by the generated HLOD. 
-	 * For now, this grid will be autogenerated and must not be manually added to the world partition. */
 	UPROPERTY(EditAnywhere, Config, Category=HLOD)
-	FName RuntimeGrid;
+	int32 CellSize;
 
 	/** Loading range of the RuntimeGrid */
 	UPROPERTY(EditAnywhere, Config, Category=HLOD)
 	float LoadingRange;
+
+	/** HLODLayer to assign to the generated HLOD actors */
+	UPROPERTY(EditAnywhere, Config, Category=HLOD)
+	TSoftObjectPtr<UHLODLayer> ParentLayer;
 #endif
 };
diff --git a/Engine/Source/Runtime/Engine/Public/WorldPartition/HLOD/HLODSubsystem.h b/Engine/Source/Runtime/Engine/Public/WorldPartition/HLOD/HLODSubsystem.h
--- a/Engine/Source/Runtime/Engine/Public/WorldPartition/HLOD/HLODSubsystem.h
+++ b/Engine/Source/Runtime/Engine/Public/WorldPartition/HLOD/HLODSubsystem.h
@@ -45,6 +45,9 @@ public:
 
 private:
 #if WITH_EDITOR
+	void OnActorLoaded(AActor& Actor);
+	void OnActorUnloaded(AActor& Actor);
+	
 	void RegisterActorDescFactories(UWorldPartitionSubsystem* WorldPartitionSubsystem);
 #endif
 
@@ -59,7 +62,7 @@ private:
 	// Cells that where shown before their HLOD was loaded.
 	TMultiMap<FGuid, FName> PendingCellsShown;
 
-#if WITH_EDITOR
-	TMultiMap<FGuid, FGuid> PendingHLODAssignment;
+#if WITH_EDITORONLY_DATA
+	TMap<FGuid, AWorldPartitionHLOD*> ActorsToHLOD;
 #endif
 };
diff --git a/Engine/Source/Runtime/Engine/Public/WorldPartition/WorldPartition.h b/Engine/Source/Runtime/Engine/Public/WorldPartition/WorldPartition.h
--- a/Engine/Source/Runtime/Engine/Public/WorldPartition/WorldPartition.h
+++ b/Engine/Source/Runtime/Engine/Public/WorldPartition/WorldPartition.h
@@ -50,7 +50,6 @@ class ENGINE_API UWorldPartition final : public UObject
 {
 	GENERATED_UCLASS_BODY()
 
-	friend class AActor;
 	friend class UWorldPartitionEditorCell;
 	friend class FWorldPartitionEditorModule;
 	friend class FUnrealEdMisc;
@@ -113,6 +112,8 @@ public:
 	FWorldPartitionActorDesc* GetActorDesc(const FGuid& Guid);
 	const FWorldPartitionActorDesc* GetActorDesc(const FGuid& Guid) const;
 
+	static TUniquePtr<FWorldPartitionActorDesc> CreateActorDesc(const AActor* Actor);
+
 	void LoadEditorCells(const FBox& Box);
 	void UnloadEditorCells(const FBox& Box);
 
diff --git a/Engine/Source/Runtime/Engine/Public/WorldPartition/WorldPartitionEditorHash.h b/Engine/Source/Runtime/Engine/Public/WorldPartition/WorldPartitionEditorHash.h
--- a/Engine/Source/Runtime/Engine/Public/WorldPartition/WorldPartitionEditorHash.h
+++ b/Engine/Source/Runtime/Engine/Public/WorldPartition/WorldPartitionEditorHash.h
@@ -27,7 +27,7 @@ public:
 	virtual void OnCellLoaded(const UWorldPartitionEditorCell* Cell) PURE_VIRTUAL(UWorldPartitionEditorHash::OnCellLoaded, ;);
 	virtual void OnCellUnloaded(const UWorldPartitionEditorCell* Cell) PURE_VIRTUAL(UWorldPartitionEditorHash::OnCellLoaded, ;);
 
-	virtual int32 ForEachIntersectingActor(const FBox& Box, TFunctionRef<void(FWorldPartitionActorDesc*)> InOperation) PURE_VIRTUAL(UWorldPartitionEditorHash::ForEachIntForEachIntersectingActorersectingCell, return 0;);
+	virtual int32 ForEachIntersectingActor(const FBox& Box, TFunctionRef<void(FWorldPartitionActorDesc*)> InOperation) PURE_VIRTUAL(UWorldPartitionEditorHash::ForEachIntersectingActor, return 0;);
 	virtual int32 ForEachIntersectingCell(const FBox& Box, TFunctionRef<void(UWorldPartitionEditorCell*)> InOperation) PURE_VIRTUAL(UWorldPartitionEditorHash::ForEachIntersectingCell, return 0;);
 	virtual int32 ForEachCell(TFunctionRef<void(UWorldPartitionEditorCell*)> InOperation) PURE_VIRTUAL(UWorldPartitionEditorHash::ForEachCell, return 0;);
 	virtual UWorldPartitionEditorCell* GetAlwaysLoadedCell() PURE_VIRTUAL(UWorldPartitionEditorHash::GetAlwaysLoadedCell, return nullptr;);
diff --git a/Engine/Source/Runtime/Engine/Public/WorldPartition/WorldPartitionRuntimeSpatialHash.h b/Engine/Source/Runtime/Engine/Public/WorldPartition/WorldPartitionRuntimeSpatialHash.h
--- a/Engine/Source/Runtime/Engine/Public/WorldPartition/WorldPartitionRuntimeSpatialHash.h
+++ b/Engine/Source/Runtime/Engine/Public/WorldPartition/WorldPartitionRuntimeSpatialHash.h
@@ -4,6 +4,9 @@
 #include "CoreMinimal.h"
 #include "UObject/ObjectMacros.h"
 #include "UObject/Object.h"
+
+#include "GameFramework/Info.h"
+
 #include "WorldPartition.h"
 #include "WorldPartitionRuntimeHash.h"
 #include "WorldPartitionRuntimeSpatialHashCell.h"
@@ -158,6 +161,24 @@ struct FSpatialHashRuntimeGrid
 #endif
 };
 
+/**
+ * Actor keeping information regarding a runtime grid
+ */
+UCLASS(NotPlaceable)
+class ASpatialHashRuntimeGridInfo : public AInfo
+{
+	GENERATED_UCLASS_BODY()
+
+public:
+	/** Begin UObject override */
+	virtual bool IsEditorOnly() const override { return true; }
+	/** End UObject override */
+
+public:
+	UPROPERTY(EditAnywhere, Category=Settings)
+	FSpatialHashRuntimeGrid	GridSettings;
+};
+
 #if WITH_EDITOR
 class FDataLayersID
 {
@@ -229,9 +250,6 @@ private:
 	UPROPERTY(EditAnywhere, Config, Category=RuntimeSettings)
 	TArray<FSpatialHashRuntimeGrid> Grids;
 
-	UPROPERTY()
-	TArray<FSpatialHashRuntimeGrid> HLODGrids;
-
 	TMap<FGuid, FGuid> CachedHLODParents;
 #endif
 
@@ -251,5 +269,6 @@ private:
 #if WITH_EDITOR
 	bool CreateStreamingGrid(const FSpatialHashRuntimeGrid& RuntimeGrid, const FSquare2DGridHelper& PartionedActors, EWorldPartitionStreamingMode Mode, UWorldPartitionStreamingPolicy* StreamingPolicy);
 #endif
-friend class UWorldPartitionSubsystem;
+
+	friend class UWorldPartitionSubsystem;
 };
\ No newline at end of file
