Added missing null check to see if the element registry has already been destroyed

#rnx

diff --git a/Engine/Source/Runtime/Engine/Private/Elements/Framework/EngineElementsLibrary.cpp b/Engine/Source/Runtime/Engine/Private/Elements/Framework/EngineElementsLibrary.cpp
--- a/Engine/Source/Runtime/Engine/Private/Elements/Framework/EngineElementsLibrary.cpp
+++ b/Engine/Source/Runtime/Engine/Private/Elements/Framework/EngineElementsLibrary.cpp
@@ -124,7 +124,14 @@ void UEngineElementsLibrary::DestroyUnreachableEditorObjectElements()
 
 	auto DestroyUnreachableObjectElement = [Registry = UTypedElementRegistry::GetInstance()](TTypedElementOwner<FObjectElementData>&& EditorElement)
 	{
-		Registry->DestroyElement(EditorElement);
+		if (Registry)
+		{
+			Registry->DestroyElement(EditorElement);
+		}
+		else
+		{
+			EditorElement.Private_DestroyNoRef();
+		}
 	};
 
 	GObjectElementOwnerStore.UnregisterElementOwners(IsObjectElementUnreachable, DestroyUnreachableObjectElement);
