Precompute the FDataLayersID hash.

diff --git a/Engine/Source/Runtime/Engine/Private/WorldPartition/WorldPartitionRuntimeSpatialHash.cpp b/Engine/Source/Runtime/Engine/Private/WorldPartition/WorldPartitionRuntimeSpatialHash.cpp
--- a/Engine/Source/Runtime/Engine/Private/WorldPartition/WorldPartitionRuntimeSpatialHash.cpp
+++ b/Engine/Source/Runtime/Engine/Private/WorldPartition/WorldPartitionRuntimeSpatialHash.cpp
@@ -21,7 +21,6 @@
 #include "Algo/Transform.h"
 
 #if WITH_EDITOR
-#include "Misc/HashBuilder.h"
 #include "Editor/EditorEngine.h"
 #include "Misc/Parse.h"
 #include "Misc/ScopedSlowTask.h"
@@ -113,41 +112,26 @@ static FAutoConsoleVariableRef CVarRuntimeSpatialHashCellToSourceAngleContributi
 
 #if WITH_EDITOR
 FDataLayersID::FDataLayersID()
-: Hash(0)
+	: Hash(0)
 {}
 
 FDataLayersID::FDataLayersID(const TArray<const UDataLayer*>& InDataLayers)
+	: Hash(0)
 {
-	Algo::TransformIf(InDataLayers, DataLayers, [](const UDataLayer* Item) { return Item->IsDynamicallyLoaded(); }, [](const UDataLayer* Item) { return Item->GetFName(); });
-	DataLayers.Sort([](const FName& A, const FName& B) { return A.FastLess(B); });
-	Hash = GetTypeHash(*this);
-}
-
-uint32 FDataLayersID::GetStableHash() const
-{
-	uint32 DataLayerId = 0;
-
-	if (DataLayers.Num())
+	if (InDataLayers.Num())
 	{
+		Algo::TransformIf(InDataLayers, DataLayers, [](const UDataLayer* Item) { return Item->IsDynamicallyLoaded(); }, [](const UDataLayer* Item) { return Item->GetFName(); });
+
 		TArray<FName> SortedDataLayers = DataLayers;
 		SortedDataLayers.Sort([](const FName& A, const FName& B) { return A.ToString() < B.ToString(); });
 
 		for (FName LayerName: SortedDataLayers)
 		{
-			DataLayerId = FCrc::StrCrc32(*LayerName.ToString(), DataLayerId);
+			Hash = FCrc::StrCrc32(*LayerName.ToString(), Hash);
 		}
 
-		check(DataLayerId);
+		check(Hash);
 	}
-
-	return DataLayerId;
-}
-
-uint32 FDataLayersID::GetFastHash() const
-{
-	FHashBuilder HashBuilder;
-	HashBuilder << DataLayers;
-	return HashBuilder.GetHash();
 }
 
 // Clustering
@@ -1358,7 +1342,7 @@ static FName GetCellName(UWorldPartition* WorldPartition, FName InGridName, int3
 	const FString PackageName = FPackageName::GetShortName(WorldPartition->GetPackage());
 	const FString PackageNameNoPIEPrefix = UWorld::RemovePIEPrefix(PackageName);
 
-	return FName(*FString::Printf(TEXT("WPRT_%s_%s_Cell_L%d_X%02d_Y%02d_DL%X"), *PackageNameNoPIEPrefix, *InGridName.ToString(), InLevel, InCellX, InCellY, InDataLayerID.GetStableHash()));
+	return FName(*FString::Printf(TEXT("WPRT_%s_%s_Cell_L%d_X%02d_Y%02d_DL%X"), *PackageNameNoPIEPrefix, *InGridName.ToString(), InLevel, InCellX, InCellY, InDataLayerID.GetHash()));
 }
 
 FName UWorldPartitionRuntimeSpatialHash::GetCellName(FName InGridName, int32 InLevel, int32 InCellX, int32 InCellY, const FDataLayersID& InDataLayerID) const
diff --git a/Engine/Source/Runtime/Engine/Public/WorldPartition/WorldPartitionRuntimeSpatialHash.h b/Engine/Source/Runtime/Engine/Public/WorldPartition/WorldPartitionRuntimeSpatialHash.h
--- a/Engine/Source/Runtime/Engine/Public/WorldPartition/WorldPartitionRuntimeSpatialHash.h
+++ b/Engine/Source/Runtime/Engine/Public/WorldPartition/WorldPartitionRuntimeSpatialHash.h
@@ -196,13 +196,7 @@ public:
 		return !(*this == Other);
 	}
 
-	uint32 GetStableHash() const;
-	uint32 GetFastHash() const;
-
-	friend uint32 GetTypeHash(const FDataLayersID& DataLayersID)
-	{
-		return DataLayersID.GetFastHash();
-	}
+	uint32 GetHash() const { return Hash; }
 
 private:
 	TArray<FName> DataLayers;
