API to create MiniMap of World in WorldPartition Editor Window

diff --git a/Engine/Source/Editor/UnrealEd/Classes/Commandlets/WorldPartitionConvertCommandlet.h b/Engine/Source/Editor/UnrealEd/Classes/Commandlets/WorldPartitionConvertCommandlet.h
--- a/Engine/Source/Editor/UnrealEd/Classes/Commandlets/WorldPartitionConvertCommandlet.h
+++ b/Engine/Source/Editor/UnrealEd/Classes/Commandlets/WorldPartitionConvertCommandlet.h
@@ -66,6 +66,7 @@ protected:
 	void SetupHLODLayerAssets();
 
 	void SetActorGuid(AActor* Actor, const FGuid& NewGuid);
+	void CreateWorldMiniMapTexture(UWorld* World);
 
 	// Conversion report
 	TSet<FString> MapsWithLevelScriptsBPs;
@@ -87,6 +88,7 @@ protected:
 	bool bReportOnly;
 	bool bVerbose;
 	bool bConversionSuffix;
+	bool bAllowCommandletRendering;
 	FString ConversionSuffix;
 
 	UPROPERTY(Config)
diff --git a/Engine/Source/Editor/UnrealEd/Private/Commandlets/WorldPartitionConvertCommandlet.cpp b/Engine/Source/Editor/UnrealEd/Private/Commandlets/WorldPartitionConvertCommandlet.cpp
--- a/Engine/Source/Editor/UnrealEd/Private/Commandlets/WorldPartitionConvertCommandlet.cpp
+++ b/Engine/Source/Editor/UnrealEd/Private/Commandlets/WorldPartitionConvertCommandlet.cpp
@@ -19,6 +19,8 @@
 #include "WorldPartition/WorldPartition.h"
 #include "WorldPartition/WorldPartitionSubsystem.h"
 #include "WorldPartition/HLOD/HLODActor.h"
+#include "WorldPartition/WorldPartitionMiniMap.h"
+#include "WorldPartition/WorldPartitionMiniMapHelper.h"
 #include "GameFramework/WorldSettings.h"
 #include "UObject/UObjectHash.h"
 #include "PackageHelperFunctions.h"
@@ -652,6 +654,18 @@ void UWorldPartitionConvertCommandlet::OnWorldLoaded(UWorld* World)
 	}
 }
 
+void UWorldPartitionConvertCommandlet::CreateWorldMiniMapTexture(UWorld* World)
+{
+	AWorldPartitionMiniMap* WorldMiniMap = FWorldPartitionMiniMapHelper::GetWorldPartitionMiniMap(World,true);
+	if (!WorldMiniMap)
+	{
+		UE_LOG(LogWorldPartitionConvertCommandlet, Error, TEXT("Failed to create Minimap. WorldPartitionMiniMap actor not found in the persistent level."));
+		return;
+	}
+
+	FWorldPartitionMiniMapHelper::CaptureWorldMiniMapToTexture(World, WorldMiniMap, WorldMiniMap->MiniMapSize, WorldMiniMap->MiniMapTexture, WorldMiniMap->MiniMapWorldBounds);
+}
+
 int32 UWorldPartitionConvertCommandlet::Main(const FString& Params)
 {
 	UE_SCOPED_TIMER(TEXT("Conversion"), LogWorldPartitionConvertCommandlet, Display);
@@ -665,7 +679,7 @@ int32 UWorldPartitionConvertCommandlet::Main(const FString& Params)
 		return 1;
 	}
 
-	// This will convert imcomplete package name to a fully qualifed path, avoiding calling it several times (takes ~50s)
+	// This will convert incomplete package name to a fully qualified path, avoiding calling it several times (takes ~50s)
 	if (!FPackageName::SearchForPackageOnDisk(Tokens[0], &Tokens[0]))
 	{
 		UE_LOG(LogWorldPartitionConvertCommandlet, Error, TEXT("Unknown level '%s'"), *Tokens[0]);
@@ -676,6 +690,7 @@ int32 UWorldPartitionConvertCommandlet::Main(const FString& Params)
 	bGenerateIni = Switches.Contains(TEXT("GenerateIni"));
 	bReportOnly = bGenerateIni || Switches.Contains(TEXT("ReportOnly"));
 	bVerbose = Switches.Contains(TEXT("Verbose"));
+	bAllowCommandletRendering = Switches.Contains(TEXT("AllowCommandletRendering"));
 
 	ReadAdditionalTokensAndSwitches(Tokens, Switches);
 
@@ -1086,6 +1101,18 @@ int32 UWorldPartitionConvertCommandlet::Main(const FString& Params)
 	bool bForceInitializeWorld = false;
 	bool bInitializedPhysicsSceneForSave = GEditor->InitializePhysicsSceneForSaveIfNecessary(MainWorld, bForceInitializeWorld);
 
+	//Create MiniMap of World
+	bool bSkipMiniMapGeneration = Switches.Contains(TEXT("SkipMiniMapGeneration"));
+	if (!bSkipMiniMapGeneration)
+	{
+		if (!bAllowCommandletRendering)
+		{
+			UE_LOG(LogWorldPartitionConvertCommandlet, Error, TEXT("AllowCommandletRendering option is required to generate MiniMap. Use SkipMiniMapGeneration option to skip MiniMap generation."));
+			return 1;
+		}
+		CreateWorldMiniMapTexture(MainWorld);
+	}
+
 	// After conversion, convert actors to external actors
 	UPackage* LevelPackage = MainLevel->GetPackage();
 
diff --git a/Engine/Source/Editor/UnrealEd/Private/WorldPartition/WorldPartitionMiniMapBuilder.cpp b/Engine/Source/Editor/UnrealEd/Private/WorldPartition/WorldPartitionMiniMapBuilder.cpp
new file mode 100644
--- /dev/null
+++ b/Engine/Source/Editor/UnrealEd/Private/WorldPartition/WorldPartitionMiniMapBuilder.cpp
@@ -0,0 +1,61 @@
+// Copyright Epic Games, Inc. All Rights Reserved.
+
+#include "WorldPartition/WorldPartitionMiniMapBuilder.h"
+
+#include "CoreMinimal.h"
+#include "Logging/LogMacros.h"
+
+#include "WorldPartition/WorldPartitionMiniMap.h"
+#include "WorldPartition/WorldPartitionMiniMapHelper.h"
+
+
+
+DEFINE_LOG_CATEGORY_STATIC(LogWorldPartitionMiniMapBuilder, All, All);
+
+UWorldPartitionMiniMapBuilder::UWorldPartitionMiniMapBuilder(const FObjectInitializer& ObjectInitializer)
+	: Super(ObjectInitializer)
+{
+}
+
+bool UWorldPartitionMiniMapBuilder::RequiresCommandletRendering() const
+{
+	return true;
+}
+
+bool UWorldPartitionMiniMapBuilder::Run(UWorld* World, FPackageSourceControlHelper& PackageHelper)
+{
+	AWorldPartitionMiniMap* WorldMiniMap = FWorldPartitionMiniMapHelper::GetWorldPartitionMiniMap(World,true);
+	if (!WorldMiniMap)
+	{
+		UE_LOG(LogWorldPartitionMiniMapBuilder, Error, TEXT("Failed to create Minimap. WorldPartitionMiniMap actor not found in the persistent level."));
+		return 1;
+	}
+
+	WorldMiniMap->MiniMapSize = this->MiniMapSize;
+	FWorldPartitionMiniMapHelper::CaptureWorldMiniMapToTexture(World, WorldMiniMap, WorldMiniMap->MiniMapSize, WorldMiniMap->MiniMapTexture, WorldMiniMap->MiniMapWorldBounds);
+
+	// Save MiniMap Package
+	auto WorldMiniMapExternalPackage = WorldMiniMap->GetExternalPackage();
+	FString PackageFileName = SourceControlHelpers::PackageFilename(WorldMiniMapExternalPackage);
+
+	if (!PackageHelper.Checkout(WorldMiniMapExternalPackage))
+	{
+		UE_LOG(LogWorldPartitionMiniMapBuilder, Error, TEXT("Error checking out package %s."), *WorldMiniMapExternalPackage->GetName());
+		return 1;
+	}
+	
+	if (!UPackage::SavePackage(WorldMiniMapExternalPackage, nullptr, RF_Standalone, *PackageFileName, GError, nullptr, false, true, SAVE_Async))
+	{
+		UE_LOG(LogWorldPartitionMiniMapBuilder, Error, TEXT("Error saving package %s."), *WorldMiniMapExternalPackage->GetName());
+		return 1;
+	}
+
+	if (!PackageHelper.AddToSourceControl(WorldMiniMapExternalPackage))
+	{
+		UE_LOG(LogWorldPartitionMiniMapBuilder, Error, TEXT("Error adding package %s to source control."), *WorldMiniMapExternalPackage->GetName());
+		return 1;
+	}
+
+	UPackage::WaitForAsyncFileWrites();
+	return true;
+}
diff --git a/Engine/Source/Editor/UnrealEd/Public/WorldPartition/WorldPartitionMiniMapBuilder.h b/Engine/Source/Editor/UnrealEd/Public/WorldPartition/WorldPartitionMiniMapBuilder.h
new file mode 100644
--- /dev/null
+++ b/Engine/Source/Editor/UnrealEd/Public/WorldPartition/WorldPartitionMiniMapBuilder.h
@@ -0,0 +1,21 @@
+// Copyright Epic Games, Inc. All Rights Reserved.
+
+#pragma once
+
+#include "WorldPartition/WorldPartitionBuilder.h"
+#include "WorldPartitionMiniMapBuilder.generated.h"
+
+UCLASS(config = Engine, defaultconfig)
+class UWorldPartitionMiniMapBuilder : public UWorldPartitionBuilder
+{
+	GENERATED_UCLASS_BODY()
+
+public:
+	// UWorldPartitionBuilder interface begin
+	virtual bool RequiresCommandletRendering() const  override;
+	virtual bool Run(UWorld* World, FPackageSourceControlHelper& PackageHelper) override;
+	// UWorldPartitionBuilder interface end
+
+	UPROPERTY(config)
+	int32 MiniMapSize = 1024;
+};
\ No newline at end of file
diff --git a/Engine/Source/Editor/WorldPartitionEditor/Private/WorldPartition/SWorldPartitionEditorGrid2D.cpp b/Engine/Source/Editor/WorldPartitionEditor/Private/WorldPartition/SWorldPartitionEditorGrid2D.cpp
--- a/Engine/Source/Editor/WorldPartitionEditor/Private/WorldPartition/SWorldPartitionEditorGrid2D.cpp
+++ b/Engine/Source/Editor/WorldPartitionEditor/Private/WorldPartition/SWorldPartitionEditorGrid2D.cpp
@@ -12,6 +12,7 @@
 #include "Editor/GroupActor.h"
 #include "Editor/EditorEngine.h"
 #include "Widgets/Input/SCheckBox.h"
+#include "Widgets/Input/SButton.h"
 #include "WorldBrowserModule.h"
 #include "LevelEditorViewport.h"
 #include "Algo/Transform.h"
@@ -20,6 +21,8 @@
 
 #define LOCTEXT_NAMESPACE "WorldPartitionEditor"
 
+static TAutoConsoleVariable<int32> CVarShowReloadMiniMapButton(TEXT("MiniMap.ShowReloadButton"), 0, TEXT("Show reload MiniMap button."));
+
 SWorldPartitionEditorGrid2D::FEditorCommands::FEditorCommands()
 	: TCommands<FEditorCommands>
 (
@@ -101,6 +104,14 @@ void SWorldPartitionEditorGrid2D::Construct(const FArguments& InArgs)
 						.IsEnabled(true)
 						.Text(LOCTEXT("ShowActors", "Show Actors"))
 					]
+					+SHorizontalBox::Slot()
+					.AutoWidth()
+					[
+						SNew(SButton)
+						.Text(LOCTEXT("Reload MiniMap", "Reload MiniMap "))
+						.Visibility_Lambda([]() {return CVarShowReloadMiniMapButton.GetValueOnAnyThread() != 0 ? EVisibility::Visible : EVisibility::Hidden;})
+						.OnClicked(this, &SWorldPartitionEditorGrid2D::ReloadMiniMap)
+					]
 				]
 			]
 		]
diff --git a/Engine/Source/Editor/WorldPartitionEditor/Private/WorldPartition/SWorldPartitionEditorGridSpatialHash.cpp b/Engine/Source/Editor/WorldPartitionEditor/Private/WorldPartition/SWorldPartitionEditorGridSpatialHash.cpp
--- a/Engine/Source/Editor/WorldPartitionEditor/Private/WorldPartition/SWorldPartitionEditorGridSpatialHash.cpp
+++ b/Engine/Source/Editor/WorldPartitionEditor/Private/WorldPartition/SWorldPartitionEditorGridSpatialHash.cpp
@@ -7,6 +7,8 @@
 #include "WorldPartition/WorldPartitionRuntimeSpatialHash.h"
 #include "WorldPartition/WorldPartitionEditorSpatialHash.h"
 #include "WorldPartition/WorldPartitionActorDesc.h"
+#include "WorldPartition/WorldPartitionMiniMap.h"
+#include "WorldPartition/WorldPartitionMiniMapHelper.h"
 #include "Framework/Commands/Commands.h"
 #include "Framework/MultiBox/MultiBoxBuilder.h"
 #include "Misc/FileHelper.h"
@@ -27,6 +29,7 @@
 TAutoConsoleVariable<bool> CVarDebugDrawOctree(TEXT("wp.Editor.DebugDrawOctree"), false, TEXT("Whether to debug draw the World Partition octree"), ECVF_Default);
 
 SWorldPartitionEditorGridSpatialHash::SWorldPartitionEditorGridSpatialHash()
+	:WorldMiniMapBounds(ForceInit)
 {
 }
 
@@ -44,25 +47,49 @@ void SWorldPartitionEditorGridSpatialHash::Construct(const FArguments& InArgs)
 		UWorldPartitionEditorSpatialHash* EditorSpatialHash = (UWorldPartitionEditorSpatialHash*)WorldPartition->EditorHash;
 
 		bShowActors = true;
-		if (UObject* WorldImage = EditorSpatialHash->WorldImage.TryLoad())
-		{
-			if (UTexture2D* WorldImageTexture = Cast<UTexture2D>(WorldImage))
-			{
-				WorldImageBrush.SetImageSize(FVector2D(WorldImageTexture->GetSizeX(), WorldImageTexture->GetSizeY()));
-				WorldImageBrush.SetResourceObject(WorldImageTexture);
-				bShowActors = false;
-			}
-			if (UMaterialInterface* WorldImageMaterial = Cast<UMaterialInterface>(WorldImage))
-			{
-				WorldImageBrush.SetResourceObject(WorldImageMaterial);
-				bShowActors = false;
-			}
-		}
+
+		//Update MiniMap data for drawing  
+		UpdateWorldMiniMapDetails();
 	}
 
 	SWorldPartitionEditorGrid2D::Construct(SWorldPartitionEditorGrid::FArguments().InWorld(InArgs._InWorld));
 }
 
+FReply SWorldPartitionEditorGridSpatialHash::ReloadMiniMap()
+{
+	UE_LOG(LogTemp, Log, TEXT("Reload MiniMap has been clicked"));
+
+	//Create a new MiniMap if there isn't one.
+	auto WorldMiniMap = FWorldPartitionMiniMapHelper::GetWorldPartitionMiniMap(World, true);
+	if (!WorldMiniMap)
+	{
+		UE_LOG(LogTemp, Error, TEXT("Failed to create Minimap. WorldPartitionMiniMap actor not found in the persistent level."));
+		return FReply::Handled();
+	}
+
+	WorldMiniMap->Modify();
+	FWorldPartitionMiniMapHelper::CaptureWorldMiniMapToTexture(World, WorldMiniMap, WorldMiniMap->MiniMapSize, WorldMiniMap->MiniMapTexture, WorldMiniMap->MiniMapWorldBounds);
+
+	UpdateWorldMiniMapDetails();
+
+	return FReply::Handled();
+}
+
+void SWorldPartitionEditorGridSpatialHash::UpdateWorldMiniMapDetails()
+{
+	auto WorldMiniMap = FWorldPartitionMiniMapHelper::GetWorldPartitionMiniMap(World);
+	if (WorldMiniMap)
+	{
+		WorldMiniMapBounds = FBox2D(FVector2D(WorldMiniMap->MiniMapWorldBounds.Min), FVector2D(WorldMiniMap->MiniMapWorldBounds.Max));
+		if (UTexture2D* MiniMapTexture = WorldMiniMap->MiniMapTexture)
+		{
+			WorldMiniMapBrush.SetImageSize(FVector2D(MiniMapTexture->GetSizeX(), MiniMapTexture->GetSizeY()));
+			WorldMiniMapBrush.SetResourceObject(MiniMapTexture);
+		}
+	}
+	
+}
+
 int32 SWorldPartitionEditorGridSpatialHash::PaintGrid(const FGeometry& AllottedGeometry, const FSlateRect& MyCullingRect, FSlateWindowElementList& OutDrawElements, int32 LayerId) const
 {
 	const FBox2D ViewRect(FVector2D(ForceInitToZero), AllottedGeometry.GetLocalSize());
@@ -103,24 +130,19 @@ int32 SWorldPartitionEditorGridSpatialHash::PaintGrid(const FGeometry& AllottedG
 		);
 	}
 
-	// Draw world image if any
-	if (WorldImageBrush.HasUObject())
-	{
-		const FVector2D& WorldImageTopLeftW = EditorSpatialHash->WorldImageTopLeftW;
-		const FVector2D& WorldImageBottomRightW = EditorSpatialHash->WorldImageBottomRightW;
-
+	// Draw MiniMap image if any
+	if (WorldMiniMapBrush.HasUObject())
+	{	
 		FPaintGeometry WorldImageGeometry = AllottedGeometry.ToPaintGeometry(
-			WorldToScreen.TransformPoint(WorldImageTopLeftW),
-			WorldToScreen.TransformPoint(WorldImageBottomRightW) - WorldToScreen.TransformPoint(WorldImageTopLeftW)
+			WorldToScreen.TransformPoint(WorldMiniMapBounds.Min),
+			WorldToScreen.TransformPoint(WorldMiniMapBounds.Max) - WorldToScreen.TransformPoint(WorldMiniMapBounds.Min)
 		);
 
 		FSlateDrawElement::MakeRotatedBox(
 			OutDrawElements,
 			++LayerId,
 			WorldImageGeometry,
-			&WorldImageBrush,
-			ESlateDrawEffect::None,
-			FMath::DegreesToRadians(0.0f)
+			&WorldMiniMapBrush
 		);
 	}
 
diff --git a/Engine/Source/Editor/WorldPartitionEditor/Public/WorldPartition/SWorldPartitionEditorGrid2D.h b/Engine/Source/Editor/WorldPartitionEditor/Public/WorldPartition/SWorldPartitionEditorGrid2D.h
--- a/Engine/Source/Editor/WorldPartitionEditor/Public/WorldPartition/SWorldPartitionEditorGrid2D.h
+++ b/Engine/Source/Editor/WorldPartitionEditor/Public/WorldPartition/SWorldPartitionEditorGrid2D.h
@@ -55,6 +55,8 @@ public:
 	virtual uint32 PaintSelection(const FGeometry& AllottedGeometry, const FSlateRect& MyCullingRect, FSlateWindowElementList& OutDrawElements, uint32 LayerId) const;
 	virtual int32 PaintSoftwareCursor(const FGeometry& AllottedGeometry, const FSlateRect& MyCullingRect, FSlateWindowElementList& OutDrawElements, int32 LayerId) const;
 
+	virtual FReply ReloadMiniMap() { return FReply::Handled(); }
+
 protected:
 	void UpdateTransform() const;
 	void UpdateSelection();
diff --git a/Engine/Source/Editor/WorldPartitionEditor/Public/WorldPartition/SWorldPartitionEditorGridSpatialHash.h b/Engine/Source/Editor/WorldPartitionEditor/Public/WorldPartition/SWorldPartitionEditorGridSpatialHash.h
--- a/Engine/Source/Editor/WorldPartitionEditor/Public/WorldPartition/SWorldPartitionEditorGridSpatialHash.h
+++ b/Engine/Source/Editor/WorldPartitionEditor/Public/WorldPartition/SWorldPartitionEditorGridSpatialHash.h
@@ -18,6 +18,10 @@ public:
 	
 	int32 PaintGrid(const FGeometry& AllottedGeometry, const FSlateRect& MyCullingRect, FSlateWindowElementList& OutDrawElements, int32 LayerId) const override;
 
+	FReply ReloadMiniMap() override;
+
 private:
-	FSlateBrush WorldImageBrush;
+	void UpdateWorldMiniMapDetails();
+	FBox2D          WorldMiniMapBounds;
+	FSlateBrush		WorldMiniMapBrush;
 };
diff --git a/Engine/Source/Runtime/Engine/Private/WorldPartition/WorldPartitionMiniMap.cpp b/Engine/Source/Runtime/Engine/Private/WorldPartition/WorldPartitionMiniMap.cpp
new file mode 100644
--- /dev/null
+++ b/Engine/Source/Runtime/Engine/Private/WorldPartition/WorldPartitionMiniMap.cpp
@@ -0,0 +1,21 @@
+// Copyright Epic Games, Inc. All Rights Reserved.
+
+#include "WorldPartition/WorldPartitionMiniMap.h"
+
+#define LOCTEXT_NAMESPACE "WorldPartitionEditor"
+
+AWorldPartitionMiniMap::AWorldPartitionMiniMap(const FObjectInitializer& ObjectInitializer)
+	: Super(ObjectInitializer)
+	, MiniMapWorldBounds(ForceInit)
+	, MiniMapTexture(nullptr)
+{
+}
+
+#if WITH_EDITOR
+void AWorldPartitionMiniMap::PostEditChangeProperty(FPropertyChangedEvent& PropertyChangedEvent)
+{
+	MiniMapSize = FMath::Clamp<uint32>(FMath::RoundUpToPowerOfTwo(MiniMapSize), 256, 8192);
+}
+#endif
+
+#undef LOCTEXT_NAMESPACE
\ No newline at end of file
diff --git a/Engine/Source/Runtime/Engine/Private/WorldPartition/WorldPartitionMiniMapHelper.cpp b/Engine/Source/Runtime/Engine/Private/WorldPartition/WorldPartitionMiniMapHelper.cpp
new file mode 100644
--- /dev/null
+++ b/Engine/Source/Runtime/Engine/Private/WorldPartition/WorldPartitionMiniMapHelper.cpp
@@ -0,0 +1,206 @@
+// Copyright Epic Games, Inc. All Rights Reserved.
+
+/*
+ * WorldPartitionMiniMapHelper implementation
+ */
+
+#include "WorldPartition/WorldPartitionMiniMapHelper.h"
+
+#if WITH_EDITOR
+#include "Components/SceneCaptureComponent2D.h"
+#include "Engine/TextureRenderTarget2D.h"
+#include "Engine/SceneCapture2D.h"
+#include "Engine/World.h"
+#include "ShaderCompiler.h"
+#include "TextureCompiler.h"
+#include "DistanceFieldAtlas.h"
+#include "PackageSourceControlHelper.h"
+#include "ProfilingDebugging/ScopedTimers.h"
+#include "WorldPartition/WorldPartitionMiniMap.h"
+
+DEFINE_LOG_CATEGORY_STATIC(LogWorldPartitionMiniMapHelper, All, All);
+
+AWorldPartitionMiniMap* FWorldPartitionMiniMapHelper::GetWorldPartitionMiniMap(UWorld* World, bool bCreateNewMiniMap)
+{
+	if (!World->GetWorldPartition())
+	{
+		UE_LOG(LogWorldPartitionMiniMapHelper, Error, TEXT("No WorldPartition Found. WorldPartition must exist to get WorldPartitionMiniMap"));
+		return nullptr;
+	}
+
+	AWorldPartitionMiniMap* WorldPartitionMiniMap = nullptr;
+	ULevel* PersistentLevel = World->PersistentLevel;
+
+	for (auto Actor : PersistentLevel->Actors)
+	{
+		if (Actor && Actor->IsA<AWorldPartitionMiniMap>())
+		{
+			WorldPartitionMiniMap = Cast<AWorldPartitionMiniMap>(Actor);
+			return WorldPartitionMiniMap;
+		}
+	}
+
+	if (bCreateNewMiniMap)
+	{
+		check(PersistentLevel);
+		if (!WorldPartitionMiniMap)
+		{
+			TSubclassOf<AWorldPartitionMiniMap> MiniMap(AWorldPartitionMiniMap::StaticClass());
+			FActorSpawnParameters SpawnInfo;
+			SpawnInfo.SpawnCollisionHandlingOverride = ESpawnActorCollisionHandlingMethod::AlwaysSpawn;
+			SpawnInfo.OverrideLevel = PersistentLevel;
+			WorldPartitionMiniMap = World->SpawnActor<AWorldPartitionMiniMap>(MiniMap, SpawnInfo);
+		}
+	}
+	
+	return WorldPartitionMiniMap;
+	
+}
+
+void FWorldPartitionMiniMapHelper::CaptureWorldMiniMapToTexture(UWorld* InWorld, UObject* InOuterForTexture, uint32 InMiniMapSize, UTexture2D*& InOutMiniMapTexture, FBox& OutWorldBounds)
+{
+	auto WaitForShaderCompilation = [](){
+		UE_SCOPED_TIMER(TEXT("Shader Compilation"), LogWorldPartitionMiniMapHelper, Display);
+		auto RemainingShaders = GShaderCompilingManager->GetNumRemainingJobs();
+		if (RemainingShaders > 0)
+		{
+			GShaderCompilingManager->FinishAllCompilation();
+		}
+	};
+	auto WaitForTextureBuilding = []()
+	{
+		UE_SCOPED_TIMER(TEXT("Building Textures"), LogWorldPartitionMiniMapHelper, Display);
+		auto RemainingTextures = FTextureCompilingManager::Get().GetNumRemainingTextures();
+		if (RemainingTextures > 0)
+		{
+			FTextureCompilingManager::Get().FinishAllCompilation();
+		}
+	};
+	auto WaitForDistanceMeshFieldBuilding = []() {
+		UE_SCOPED_TIMER(TEXT("Building Mesh Distance Fields"), LogWorldPartitionMiniMapHelper, Display);
+		if (GDistanceFieldAsyncQueue && GDistanceFieldAsyncQueue->GetNumOutstandingTasks() > 0)
+		{
+			GDistanceFieldAsyncQueue->BlockUntilAllBuildsComplete();
+		}
+	};
+
+	WaitForDistanceMeshFieldBuilding();
+	WaitForTextureBuilding();
+	WaitForShaderCompilation();
+	
+	//Calculate bounds of the World
+	OutWorldBounds = FWorldPartitionMiniMapHelper::GetWorldBounds(InWorld);
+	
+	//Calculate Viewport size
+	FBox2D WorldBounds2D(FVector2D(OutWorldBounds.Min), FVector2D(OutWorldBounds.Max));
+	FVector2D ViewSize = WorldBounds2D.Max - WorldBounds2D.Min;
+	float AspectRatio = FMath::Abs(ViewSize.X) / FMath::Abs(ViewSize.Y);
+	uint32 ViewportWidth = InMiniMapSize * AspectRatio;
+	uint32 ViewportHeight = InMiniMapSize;
+
+	//Calculate Projection matrix based on world bounds.
+	FMatrix ProjectionMatrix;
+	FWorldPartitionMiniMapHelper::CalTopViewOfWorld(ProjectionMatrix, OutWorldBounds, ViewportWidth, ViewportHeight);
+
+	TArray<FColor> CapturedImage;
+	//Using SceneCapture Actor capture the scene to buffer
+	{
+		UTextureRenderTarget2D* RenderTargetTexture = NewObject<UTextureRenderTarget2D>();
+		RenderTargetTexture->AddToRoot();
+		RenderTargetTexture->ClearColor = FLinearColor::Transparent;
+		RenderTargetTexture->TargetGamma = 2.2f;
+		RenderTargetTexture->InitCustomFormat(InMiniMapSize, InMiniMapSize, PF_B8G8R8A8, false);
+
+		FActorSpawnParameters SpawnInfo;
+		SpawnInfo.ObjectFlags |= RF_Transient;
+		FVector CaptureActorLocation = FVector(OutWorldBounds.GetCenter().X, OutWorldBounds.GetCenter().Y, OutWorldBounds.GetCenter().Z + OutWorldBounds.GetExtent().Z);
+		FRotator CaptureActorRotation = FRotator(-90.f, 0.f, -90.f);
+		ASceneCapture2D* CaptureActor = InWorld->SpawnActor<ASceneCapture2D>(CaptureActorLocation, CaptureActorRotation, SpawnInfo);
+		auto CaptureComponent = CaptureActor->GetCaptureComponent2D();
+		CaptureComponent->TextureTarget = RenderTargetTexture;
+		CaptureComponent->ProjectionType = ECameraProjectionMode::Orthographic;
+		CaptureComponent->OrthoWidth = ViewportWidth;
+		CaptureComponent->bUseCustomProjectionMatrix = true;
+		CaptureComponent->CustomProjectionMatrix = ProjectionMatrix;
+		CaptureComponent->CaptureScene();
+
+		auto RenderTargetResource = RenderTargetTexture->GameThread_GetRenderTargetResource();
+		CapturedImage.SetNumUninitialized(InMiniMapSize * InMiniMapSize);
+		RenderTargetResource->ReadPixelsPtr(CapturedImage.GetData(), FReadSurfaceDataFlags(), FIntRect(0, 0, InMiniMapSize, InMiniMapSize));
+		InWorld->DestroyActor(CaptureActor);
+		CaptureActor = nullptr;
+	}
+	
+	//Update the output texture
+	if (!InOutMiniMapTexture)
+	{
+		InOutMiniMapTexture = NewObject<UTexture2D>(InOuterForTexture, "MiniMapTexture");
+	}
+	InOutMiniMapTexture->AdjustMinAlpha = 1.f;
+	InOutMiniMapTexture->Source.Init(InMiniMapSize, InMiniMapSize, 1, 1, TSF_BGRA8, (uint8*)CapturedImage.GetData());
+	InOutMiniMapTexture->UpdateResource();
+
+}
+
+bool FWorldPartitionMiniMapHelper::DoesActorContributeToBounds(AActor* Actor)
+{
+	return Actor && Actor->GridPlacement != EActorGridPlacement::AlwaysLoaded;
+}
+
+FBox FWorldPartitionMiniMapHelper::GetWorldBounds(UWorld* World)
+{
+	FBox WorldBox(ForceInit);
+
+	TArray<ULevel*> LevelsToRender = World->GetLevels();
+	for (ULevel* Level : LevelsToRender)
+	{
+		if (Level && Level->bIsVisible)
+		{
+			for (AActor* Actor : Level->Actors)
+			{
+				if ( Actor && DoesActorContributeToBounds(Actor))
+				{
+					FVector Origin, Extent;
+					Actor->GetActorBounds(false, Origin, Extent);
+					FBox ActorBounds(Origin - Extent, Origin + Extent);
+					WorldBox += ActorBounds;
+				}
+			}
+		}
+	}
+
+	return WorldBox;
+}
+
+void FWorldPartitionMiniMapHelper::CalTopViewOfWorld(FMatrix& OutProjectionMatrix, const FBox& WorldBox, uint32 ViewportWidth, uint32 ViewportHeight)
+{
+	const FVector Origin = WorldBox.GetCenter();
+
+	FVector2D WorldSizeMin2D(WorldBox.Min.X, WorldBox.Min.Y);
+	FVector2D WorldSizeMax2D(WorldBox.Max.X, WorldBox.Max.Y);
+
+	FVector2D WorldSize2D = (WorldSizeMax2D - WorldSizeMin2D);
+	WorldSize2D.X = FMath::Abs(WorldSize2D.X);
+	WorldSize2D.Y = FMath::Abs(WorldSize2D.Y);
+	const bool bUseXAxis = (WorldSize2D.X / WorldSize2D.Y) > 1.f;
+	const float WorldAxisSize = bUseXAxis ? WorldSize2D.X : WorldSize2D.Y;
+	const uint32 ViewportAxisSize = bUseXAxis ? ViewportWidth : ViewportHeight;
+	const float OrthoZoom = WorldAxisSize / ViewportAxisSize / 2.f;
+	const float OrthoWidth = FMath::Max(1.f, ViewportWidth * OrthoZoom);
+	const float OrthoHeight = FMath::Max(1.f, ViewportHeight * OrthoZoom);
+
+	const float ZOffset = HALF_WORLD_MAX;
+	OutProjectionMatrix = FReversedZOrthoMatrix(
+		OrthoWidth,
+		OrthoHeight,
+		0.5f / ZOffset,
+		ZOffset
+	);
+
+	ensureMsgf(!OutProjectionMatrix.ContainsNaN(), TEXT("Nans found on ProjectionMatrix"));
+	if (OutProjectionMatrix.ContainsNaN())
+	{
+		OutProjectionMatrix.SetIdentity();
+	}
+}
+#endif
\ No newline at end of file
diff --git a/Engine/Source/Runtime/Engine/Public/WorldPartition/WorldPartitionMiniMap.h b/Engine/Source/Runtime/Engine/Public/WorldPartition/WorldPartitionMiniMap.h
new file mode 100644
--- /dev/null
+++ b/Engine/Source/Runtime/Engine/Public/WorldPartition/WorldPartitionMiniMap.h
@@ -0,0 +1,36 @@
+// Copyright Epic Games, Inc. All Rights Reserved.
+#pragma once
+
+#include "CoreMinimal.h"
+#include "GameFramework/info.h"
+#include "WorldPartitionMiniMap.generated.h"
+
+/**
+ * A mini map to preview the world in world partition window. (editor-only)
+ */
+UCLASS(hidecategories = (Actor, Advanced, Display, Events, Object, Attachment, Info, Input, Blueprint, Layers, Tags, Replication), notplaceable)
+class ENGINE_API AWorldPartitionMiniMap : public AInfo
+{
+	GENERATED_BODY()
+
+public:
+	AWorldPartitionMiniMap(const FObjectInitializer& ObjectInitializer = FObjectInitializer::Get());
+
+	virtual bool IsEditorOnly() const final { return true; }
+
+#if WITH_EDITOR
+	virtual void PostEditChangeProperty(FPropertyChangedEvent& PropertyChangedEvent) override;
+#endif
+
+	/*WorldBounds for MinMapTexture*/
+	UPROPERTY(VisibleAnywhere, Category = WorldPartitionMiniMap)
+	FBox MiniMapWorldBounds;
+
+	/*MiniMap Texture for displaying on world partition window*/
+	UPROPERTY(VisibleAnywhere, Category = WorldPartitionMiniMap)
+	UTexture2D* MiniMapTexture;
+
+	/*MiniMap Size*/
+	UPROPERTY(EditAnywhere, Category = WorldPartitionMiniMap, meta=(UIMin = "256", UIMax = "8192"))
+	int32 MiniMapSize = 1024;
+};
\ No newline at end of file
diff --git a/Engine/Source/Runtime/Engine/Public/WorldPartition/WorldPartitionMiniMapHelper.h b/Engine/Source/Runtime/Engine/Public/WorldPartition/WorldPartitionMiniMapHelper.h
new file mode 100644
--- /dev/null
+++ b/Engine/Source/Runtime/Engine/Public/WorldPartition/WorldPartitionMiniMapHelper.h
@@ -0,0 +1,31 @@
+// Copyright Epic Games, Inc. All Rights Reserved.
+
+/**
+ * WorldPartitionMiniMapHelper
+ *
+ * Helper class to build MiniMap texture in World Partition editor.
+ *
+ */
+#pragma once
+
+#if WITH_EDITOR
+class AWorldPartitionMiniMap;
+class UWorld;
+class UTexture2D;
+class AActor;
+
+#include "Math/Box.h"
+#include "Math/Matrix.h"
+
+class ENGINE_API FWorldPartitionMiniMapHelper
+{
+public:
+	static AWorldPartitionMiniMap* GetWorldPartitionMiniMap(UWorld* World, bool bCreateNewMiniMap=false);
+	static void CaptureWorldMiniMapToTexture(UWorld* InWorld, UObject* InOuterForTexture, uint32 InMiniMapSize, UTexture2D*& InOutMiniMapTexture, FBox& OutWorldBounds);
+private:
+	static bool DoesActorContributeToBounds(AActor* Actor);
+	static FBox GetWorldBounds(UWorld* World);
+	static void CalTopViewOfWorld(FMatrix& OutProjectionMatrix, const FBox& WorldBox, uint32 ViewportWidth, uint32 ViewportHeight);
+};
+#endif
+
