Fixed Compiler error in Linux platform

diff --git a/Engine/Source/Runtime/Engine/Public/WorldPartition/WorldPartitionMiniMap.h b/Engine/Source/Runtime/Engine/Public/WorldPartition/WorldPartitionMiniMap.h
--- a/Engine/Source/Runtime/Engine/Public/WorldPartition/WorldPartitionMiniMap.h
+++ b/Engine/Source/Runtime/Engine/Public/WorldPartition/WorldPartitionMiniMap.h
@@ -2,7 +2,7 @@
 #pragma once
 
 #include "CoreMinimal.h"
-#include "GameFramework/info.h"
+#include "GameFramework/Info.h"
 #include "WorldPartitionMiniMap.generated.h"
 
 /**
