Fix UI clamps of Path Tracing related scene settings

The default clamp range for max bounces and samples per pixel was too small. The default value of samples per pixel was outside the clamp range.

diff --git a/Engine/Source/Runtime/Engine/Classes/Engine/Scene.h b/Engine/Source/Runtime/Engine/Classes/Engine/Scene.h
--- a/Engine/Source/Runtime/Engine/Classes/Engine/Scene.h
+++ b/Engine/Source/Runtime/Engine/Classes/Engine/Scene.h
@@ -1910,11 +1910,11 @@ struct FPostProcessSettings
 
 	// Path Tracing
 	/** Sets the path tracing maximum bounces */
-	UPROPERTY(interp, EditAnywhere, BlueprintReadWrite, Category = "Rendering Features|PathTracing", meta = (ClampMin = "0", ClampMax = "50", editcondition = "bOverride_PathTracingMaxBounces", DisplayName = "Max. Bounces"))
+	UPROPERTY(interp, EditAnywhere, BlueprintReadWrite, Category = "Rendering Features|PathTracing", meta = (ClampMin = "0", ClampMax = "100", editcondition = "bOverride_PathTracingMaxBounces", DisplayName = "Max. Bounces"))
 	int32 PathTracingMaxBounces;
 
 	/** Sets the samples per pixel for the path tracer. */
-	UPROPERTY(interp, EditAnywhere, BlueprintReadWrite, Category = "Rendering Features|PathTracing", meta = (ClampMin = "1", ClampMax = "64", editcondition = "bOverride_PathTracingSamplesPerPixel", DisplayName = "Samples Per Pixel"))
+	UPROPERTY(interp, EditAnywhere, BlueprintReadWrite, Category = "Rendering Features|PathTracing", meta = (ClampMin = "1", ClampMax = "65536", editcondition = "bOverride_PathTracingSamplesPerPixel", DisplayName = "Samples Per Pixel"))
 	int32 PathTracingSamplesPerPixel;
 
 	UPROPERTY()
