Copied shader pipeline and changes to RDG/Shader system to support dynamic param binding struct creation. Actual dynmaic param binding generation is stubbed. Parsing SRVs UAVs and Params from shaders.

diff --git a/Engine/Config/Windows/DataDrivenPlatformInfo.ini b/Engine/Config/Windows/DataDrivenPlatformInfo.ini
--- a/Engine/Config/Windows/DataDrivenPlatformInfo.ini
+++ b/Engine/Config/Windows/DataDrivenPlatformInfo.ini
@@ -27,4 +27,5 @@ bSupportsRayTracingIndirectInstanceData=true
 bSupportsWaveOperations=false
 bSupportsGPUScene=true
 bSupportsAnisotropicMaterials=true
-bSupportsNanite=true
\ No newline at end of file
+bSupportsNanite=true
+bSupportsComputeFramework=true
\ No newline at end of file
diff --git a/Engine/Shaders/Private/ComputeKernel.usf b/Engine/Shaders/Private/ComputeKernel.usf
--- a/Engine/Shaders/Private/ComputeKernel.usf
+++ b/Engine/Shaders/Private/ComputeKernel.usf
@@ -1,8 +1,19 @@
-#include "Common.ush"
+// Copyright Epic Games, Inc. All Rights Reserved.
+
+
+#include "/Engine/Private/Common.ush"
+#include "/Engine/Private/ComputeKernelCommon.ush"
+
+KERNEL_SRV(Buffer<float>, a);
+
+KERNEL_SRV(Buffer<float>, b);
+
+KERNEL_UAV(RWBuffer<float>, c);
 
 [numthreads(32, 1, 1)]
-void Main(
-	uint DispatchIndex : SV_DispatchThreadID,
-	uint GroupIndex : SV_GroupIndex)
+KERNEL_ENTRY_POINT(LinearBlendSkinning)
 {
-}
\ No newline at end of file
+	c[DTid.x] = a[DTid.x]
+	 + 
+	b[DTid.x]; 
+}
diff --git a/Engine/Shaders/Private/ComputeKernelCommon.ush b/Engine/Shaders/Private/ComputeKernelCommon.ush
new file mode 100644
--- /dev/null
+++ b/Engine/Shaders/Private/ComputeKernelCommon.ush
@@ -0,0 +1,15 @@
+// Copyright Epic Games, Inc. All Rights Reserved.
+
+#pragma once
+
+
+#define KERNEL_ENTRY_POINT(EntryPointName) \
+    void EntryPointName( \
+        uint Gidx : SV_GroupIndex, \
+        uint3 Gid : SV_GroupID, \
+        uint3 GTid: SV_GroupThreadID, \
+        uint3 DTid : SV_DispatchThreadID)
+
+#define KERNEL_PARAM(Type, Name)    Type Name
+#define KERNEL_SRV(Type, Name)      Type Name
+#define KERNEL_UAV(Type, Name)      Type Name
diff --git a/Engine/Source/Runtime/Engine/Private/ComputeFramework/ComputeFramework.cpp b/Engine/Source/Runtime/Engine/Private/ComputeFramework/ComputeFramework.cpp
--- a/Engine/Source/Runtime/Engine/Private/ComputeFramework/ComputeFramework.cpp
+++ b/Engine/Source/Runtime/Engine/Private/ComputeFramework/ComputeFramework.cpp
@@ -6,13 +6,12 @@
 #include "RenderGraphBuilder.h"
 #include "RenderGraphUtils.h"
 #include "RHIDefinitions.h"
-#include "ComputeFramework/ComputeKernelShaderMap.h"
 #include "ComputeFramework/ComputeKernelShader.h"
 #include "ProfilingDebugging/RealtimeGPUProfiler.h"
 
 DECLARE_GPU_STAT_NAMED(ComputeFramework_ExecuteBatches, TEXT("ComputeFramework::ExecuteBatches"));
 
-static int32 GComputeFrameworkMode = 0;
+static int32 GComputeFrameworkMode = 1;
 static FAutoConsoleVariableRef CVarComputeFrameworkMode(
 	TEXT("r.ComputeFramework.mode"),
 	GComputeFrameworkMode,
@@ -34,10 +33,10 @@ void FComputeGraph::Initialize(UComputeGraph* ComputeGraph)
 	for (auto& Invocation : ComputeGraph->GetShaderInvocationList())
 	{
 		FKernelInvocation KernelInvocation = {
-			Invocation.ComputeKernel ? Invocation.ComputeKernel->GetFName() : "NULL Kernel",
+			Invocation.ComputeKernel->GetFName(),
 			FName("InvocationName"),
 			FIntVector(1, 1, 1),
-			Invocation.ComputeKernel ? Invocation.ComputeKernel->GetResource() : nullptr
+			Invocation.ComputeKernel->GetResource()
 			};
 
 		KernelInvocations.Emplace(KernelInvocation);
@@ -55,9 +54,8 @@ void FComputeFramework::EnqueueForExecution(const FComputeGraph* ComputeGraph)
 			Invocation.KernelName,
 			Invocation.InvocationName,
 			Invocation.DispatchDim,
-			Invocation.KernelResource ? 
-				Invocation.KernelResource->ShaderMap_RT->GetShader<FComputeKernelShader>(0) : 
-				GetGlobalShaderMap(ERHIFeatureLevel::SM5)->GetShader<FComputeKernelShader>()
+			Invocation.Kernel->GetShaderParamMetadata(),
+			Invocation.Kernel->GetShader()
 			};
 
 		ComputeShaders.Emplace(ShaderInvocation);
@@ -142,24 +140,71 @@ void FComputeFramework::ExecuteBatches(
 
 		for (auto& Compute : ComputeShaders)
 		{
+			FRDGBufferDesc BufferDesc = FRDGBufferDesc::CreateBufferDesc(4, 32);
+
+			void* RawBuffer = GraphBuilder.Alloc(Compute.ShaderParamMetadata->GetSize(), SHADER_PARAMETER_STRUCT_ALIGNMENT);
+			FMemory::Memzero(RawBuffer, Compute.ShaderParamMetadata->GetSize());
+
+			uint8* ParamBuffer = static_cast<uint8*>(RawBuffer);
+			TArray<FShaderParametersMetadata::FMember> ParamMembers = Compute.ShaderParamMetadata->GetMembers();
+			for (auto& Member : ParamMembers)
+			{
+				switch (Member.GetBaseType())
+				{
+				case EUniformBufferBaseType::UBMT_INT32:
+				case EUniformBufferBaseType::UBMT_UINT32:
+				case EUniformBufferBaseType::UBMT_FLOAT32:
+					// Fill is CBuffer data
+					break;
+
+				case EUniformBufferBaseType::UBMT_RDG_BUFFER_SRV:
+					{
+						FRDGBufferRef InputBuffer = GraphBuilder.CreateBuffer(BufferDesc, Member.GetName());
+
+						FComputeShaderUtils::ClearUAV(
+							GraphBuilder,
+							GetGlobalShaderMap(ERHIFeatureLevel::SM5),
+							GraphBuilder.CreateUAV(InputBuffer, PF_A32B32G32R32F),
+							FVector4(0.0f)
+							);
+
+						*reinterpret_cast<FRDGBufferSRVRef*>(&ParamBuffer[Member.GetOffset()]) = GraphBuilder.CreateSRV(InputBuffer, PF_R32_FLOAT);
+					}
+					break;
+
+				case EUniformBufferBaseType::UBMT_RDG_BUFFER_UAV:
+					{						
+						FRDGBufferRef OutputBuffer = GraphBuilder.CreateBuffer(BufferDesc, Member.GetName());
+
+						*reinterpret_cast<FRDGBufferUAVRef*>(&ParamBuffer[Member.GetOffset()]) = GraphBuilder.CreateUAV(OutputBuffer, PF_R32_FLOAT);
+					}
+					break;
+
+				default:
+					check(!"Unsupported type");
+					break;
+				};
+			}
+
 			TCHAR KernelName[FComputeKernelInvocation::MAX_NAME_LENGTH];
 			Compute.KernelName.ToString(KernelName);
 
 			TCHAR InvocationName[FComputeKernelInvocation::MAX_NAME_LENGTH];
 			Compute.InvocationName.ToString(InvocationName);
 
-			auto* Params = GraphBuilder.AllocParameters<FComputeKernelShader::FParameters>();
-
 			FComputeShaderUtils::AddPass(
 				GraphBuilder,
 				RDG_EVENT_NAME("Compute[%s]: %s", KernelName, InvocationName),
-				ERDGPassFlags::Compute,
+				ERDGPassFlags::Compute | ERDGPassFlags::NeverCull,
 				Compute.Shader,
-				Params,
+				Compute.ShaderParamMetadata,
+				static_cast<FComputeKernelShader::FParameters*>(RawBuffer),
 				Compute.DispatchDim
-				);
+				);		
 		}
 
 		GraphBuilder.Execute();
+
+		ComputeShaders.Reset();
 	}
 }
diff --git a/Engine/Source/Runtime/Engine/Private/ComputeFramework/ComputeGraph.cpp b/Engine/Source/Runtime/Engine/Private/ComputeFramework/ComputeGraph.cpp
--- a/Engine/Source/Runtime/Engine/Private/ComputeFramework/ComputeGraph.cpp
+++ b/Engine/Source/Runtime/Engine/Private/ComputeFramework/ComputeGraph.cpp
@@ -5,6 +5,10 @@
 #include "UObject/UObjectGlobals.h"
 
 
+// #TODO_ZABIR: TEMP
+#include "ComputeFramework/ComputeKernelFromText.h"
+
+
 FComputeKernelInvocationHandle UComputeGraph::AddKernel(UComputeKernelSource* KernelSource)
 {
 	return {};
@@ -19,6 +23,14 @@ void UComputeGraph::PostLoad()
 {
 	Super::PostLoad();
 
+	UComputeKernel* Kernel = LoadObject<UComputeKernel>(
+		nullptr, 
+		TEXT("/Game/Developers/ZabirHoque/CF/LBS_Kernel.LBS_Kernel"),
+		nullptr,
+		LOAD_None, 
+		nullptr
+		);
+
 	KernelInvocations.Reset();
-	KernelInvocations.Emplace();
+	KernelInvocations.Emplace(Kernel, 1);
 }
diff --git a/Engine/Source/Runtime/Engine/Private/ComputeFramework/ComputeGraphComponent.cpp b/Engine/Source/Runtime/Engine/Private/ComputeFramework/ComputeGraphComponent.cpp
--- a/Engine/Source/Runtime/Engine/Private/ComputeFramework/ComputeGraphComponent.cpp
+++ b/Engine/Source/Runtime/Engine/Private/ComputeFramework/ComputeGraphComponent.cpp
@@ -13,7 +13,7 @@ UComputeGraphComponent::UComputeGraphComponent()
 
 void UComputeGraphComponent::QueueExecute()
 {
-	if (!ComputeGraph)
+	if (!ComputeGraph || !GetScene()->GetComputeFramework())
 	{
 		return;
 	}
@@ -25,8 +25,10 @@ void UComputeGraphComponent::SendRenderDynamicData_Concurrent()
 {
 	check(ComputeGraph);
 	Super::SendRenderDynamicData_Concurrent();
-
+	
 	FComputeFramework* ComputeFramework = GetScene()->GetComputeFramework();
+	check(ComputeFramework);
+
 	FComputeGraph* ComputeGraphProxy = new FComputeGraph();
 	ComputeGraphProxy->Initialize(ComputeGraph);
 
diff --git a/Engine/Source/Runtime/Engine/Private/ComputeFramework/ComputeKernel.cpp b/Engine/Source/Runtime/Engine/Private/ComputeFramework/ComputeKernel.cpp
--- a/Engine/Source/Runtime/Engine/Private/ComputeFramework/ComputeKernel.cpp
+++ b/Engine/Source/Runtime/Engine/Private/ComputeFramework/ComputeKernel.cpp
@@ -3,10 +3,19 @@
 #include "ComputeFramework/ComputeKernel.h"
 #include "ComputeFramework/ComputeKernelPermutationSet.h"
 #include "ComputeFramework/ComputeKernelSource.h"
-#include "ComputeFramework/ComputeKernelResource.h"
 
 DEFINE_LOG_CATEGORY(ComputeKernel);
 
+void UComputeKernel::PostLoad()
+{
+	Super::PostLoad();
+
+	if (FApp::CanEverRender())
+	{
+		CacheResourceShadersForRendering();
+	}
+}
+
 #if WITH_EDITOR
 void UComputeKernel::PostEditChangeChainProperty(FPropertyChangedChainEvent& PropertyChangedEvent)
 {
@@ -24,31 +33,36 @@ void UComputeKernel::PostEditChangeChainProperty(FPropertyChangedChainEvent& Pro
 	{
 		if (KernelSource)
 		{
-			PermutationSet = KernelSource->PermutationSet;
-			DefinitionsSet = KernelSource->DefinitionsSet;
+			PermutationSetOverrides = KernelSource->PermutationSet;
+			DefinitionsSetOverrides = KernelSource->DefinitionsSet;
 		}
 		else
 		{
-			PermutationSet = FComputeKernelPermutationSet();
-			DefinitionsSet = FComputeKernelDefinitionsSet();
+			PermutationSetOverrides = FComputeKernelPermutationSet();
+			DefinitionsSetOverrides = FComputeKernelDefinitionsSet();
 		}
 		
 		CacheResourceShadersForRendering();
 	}
-	else if (ModifiedPropName == GET_MEMBER_NAME_CHECKED(UComputeKernel, PermutationSet) ||
-		ModifiedPropName == GET_MEMBER_NAME_CHECKED(UComputeKernel, DefinitionsSet))
+	else if (ModifiedPropName == GET_MEMBER_NAME_CHECKED(UComputeKernel, PermutationSetOverrides) ||
+		ModifiedPropName == GET_MEMBER_NAME_CHECKED(UComputeKernel, DefinitionsSetOverrides))
 	{
 		CacheResourceShadersForRendering();
 	}
 }
 
 void UComputeKernel::CacheResourceShadersForRendering(
-	uint32 CompilationFlags /*= uint32(EComputeKernelCompilationFlags::None)*/
+	uint32 CompilationFlags /*= uint32(EComputeKernelCompilationFlags::ApplyCompletedShaderMapForRendering)*/
 	)
 {
 	if (!KernelSource)
 	{
-		KernelResource->Invalidate();
+		if (KernelResource)
+		{
+			KernelResource->Invalidate();
+			KernelResource = nullptr;
+		}
+
 		return;
 	}
 
@@ -57,31 +71,42 @@ void UComputeKernel::CacheResourceShadersForRendering(
 		KernelResource = MakeUnique<FComputeKernelResource>();
 	}
 
-	// #TODO_ZABIR: Initialize KernelResouce with needed data
+	KernelResource = MakeUnique<FComputeKernelResource>();
+	KernelResource->SetupResource(
+		GMaxRHIFeatureLevel, 
+		KernelSource,  
+		GetName()
+		);
 
 	ERHIFeatureLevel::Type CacheFeatureLevel = GMaxRHIFeatureLevel;
 	const EShaderPlatform ShaderPlatform = GShaderPlatformForFeatureLevel[CacheFeatureLevel];
 
-	UComputeKernel::CacheShadersForResource(ShaderPlatform, CompilationFlags, KernelResource.Get());
+	UComputeKernel::CacheShadersForResource(ShaderPlatform, nullptr, CompilationFlags | uint32(EComputeKernelCompilationFlags::Force), KernelResource.Get());
 }
 
 void UComputeKernel::CacheShadersForResource(
 	EShaderPlatform ShaderPlatform,
+	const ITargetPlatform* TargetPlatform,
 	uint32 CompilationFlags,
 	FComputeKernelResource* KernelResource
 	)
 {
+	bool bCooking = (CompilationFlags & uint32(EComputeKernelCompilationFlags::IsCooking)) != 0;
+	
 	const bool bIsDefault = (KernelResource->GetKernelFlags() & uint32(EComputeKernelFlags::IsDefaultKernel)) != 0;
-	if (!GIsEditor || GIsAutomationTesting || bIsDefault)
+	if (!GIsEditor || GIsAutomationTesting || bIsDefault || bCooking)
 	{
 		CompilationFlags |= uint32(EComputeKernelCompilationFlags::Synchronous);
 	}
 
-	KernelResource->CacheShaders(ShaderPlatform, CompilationFlags);
-
-	const FComputeKernelCompilationResults& Results = KernelResource->GetCompilationResults();
+	const bool bIsSuccess = KernelResource->CacheShaders(
+		ShaderPlatform, 
+		TargetPlatform, 
+		CompilationFlags & uint32(EComputeKernelCompilationFlags::ApplyCompletedShaderMapForRendering),
+		CompilationFlags & uint32(EComputeKernelCompilationFlags::Synchronous)
+		);
 
-	if (!Results.bIsSuccess)
+	if (!bIsSuccess)
 	{
 		if (bIsDefault)
 		{
@@ -102,16 +127,11 @@ void UComputeKernel::CacheShadersForResource(
 			*LegacyShaderPlatformToShaderFormat(ShaderPlatform).ToString()
 			);
 
-		uint32 WarningCount = Results.CompileWarnings.Num();
-		for (uint32 i = 0; i < WarningCount; ++i)
-		{
-			UE_LOG(ComputeKernel, Warning, TEXT("    [Warning] - %s"), *Results.CompileWarnings[i]);
-		}
-
-		uint32 ErrorCount = Results.CompileWarnings.Num();
+		auto& CompilationErrors = KernelResource->GetCompileErrors();
+		uint32 ErrorCount = CompilationErrors.Num();
 		for (uint32 i = 0; i < ErrorCount; ++i)
 		{
-			UE_LOG(ComputeKernel, Warning, TEXT("      [Error] - %s"), *Results.CompileErrors[i]);
+			UE_LOG(ComputeKernel, Warning, TEXT("      [Error] - %s"), *CompilationErrors[i]);
 		}
 	}
 }
diff --git a/Engine/Source/Runtime/Engine/Private/ComputeFramework/ComputeKernelDerivedDataVersion.h b/Engine/Source/Runtime/Engine/Private/ComputeFramework/ComputeKernelDerivedDataVersion.h
new file mode 100644
--- /dev/null
+++ b/Engine/Source/Runtime/Engine/Private/ComputeFramework/ComputeKernelDerivedDataVersion.h
@@ -0,0 +1,13 @@
+// Copyright Epic Games, Inc. All Rights Reserved.
+
+/*=============================================================================
+ComputeKernelShaderDerivedDataVersion.h: Derived data version for ComputeKernel shaders and LUTs.
+=============================================================================*/
+
+#pragma once
+
+#include "CoreMinimal.h"
+
+// In case of merge conflicts with DDC versions, you MUST generate a new GUID and set this new
+// guid as version
+#define COMPUTEKERNEL_DERIVEDDATA_VER		TEXT("EDA5D304-9967-482C-9C75-80A5F6BEC61C")
diff --git a/Engine/Source/Runtime/Engine/Private/ComputeFramework/ComputeKernelFromText.cpp b/Engine/Source/Runtime/Engine/Private/ComputeFramework/ComputeKernelFromText.cpp
--- a/Engine/Source/Runtime/Engine/Private/ComputeFramework/ComputeKernelFromText.cpp
+++ b/Engine/Source/Runtime/Engine/Private/ComputeFramework/ComputeKernelFromText.cpp
@@ -12,8 +12,13 @@ UComputeKernelFromText::UComputeKernelFromText()
 	UniqueId = FGuid::NewGuid();
 }
 
-#if WITH_EDITOR
+void UComputeKernelFromText::PostLoad()
+{
+	Super::PostLoad();
+	ReparseKernelSourceText();
+}
 
+#if WITH_EDITOR
 void UComputeKernelFromText::PreEditChange(FEditPropertyChain& PropertyAboutToChange)
 {
 	Super::PreEditChange(PropertyAboutToChange);
@@ -49,8 +54,12 @@ void UComputeKernelFromText::ReparseKernelSourceText()
 	if (SourceFile.FilePath.IsEmpty())
 	{
 		EntryPointName = FString();
+		KernelSourceText = FString();
 		PermutationSet = FComputeKernelPermutationSet();
 		DefinitionsSet = FComputeKernelDefinitionsSet();
+		InputParams.Empty();
+		InputSRVs.Empty();
+		Outputs.Empty();
 
 		return;
 	}
@@ -66,7 +75,6 @@ void UComputeKernelFromText::ReparseKernelSourceText()
 		return;
 	}
 
-	FString KernelSourceText;
 	if (!FFileHelper::LoadFileToString(KernelSourceText, &PlatformFileSystem, *FullKernelPath))
 	{
 		UE_LOG(ComputeKernel, Error, TEXT("Unable to read kernel file \"%s\""), *FullKernelPath);
@@ -75,39 +83,169 @@ void UComputeKernelFromText::ReparseKernelSourceText()
 		return;
 	}
 
-	const FRegexPattern KernelEntryPointPattern(TEXT(R"(KERNEL_ENTRY_POINT\(\s*([a-zA-Z_][\w\d]*)\s*\))"));
+	const FRegexPattern KernelEntryPointPattern(TEXT(R"(KERNEL_ENTRY_POINT\(\s*([a-zA-Z_]\w*)\s*\))"));
 	{
-		FRegexMatcher KernelEntryPointMatcher(KernelEntryPointPattern, KernelSourceText);
-		if (KernelEntryPointMatcher.FindNext())
+		FRegexMatcher Matcher(KernelEntryPointPattern, KernelSourceText);
+		if (Matcher.FindNext())
 		{
-			EntryPointName = KernelEntryPointMatcher.GetCaptureGroup(1);
+			EntryPointName = Matcher.GetCaptureGroup(1);
 		}
 	}
 
-	const FRegexPattern KernelPermutationBoolPattern(TEXT(R"(KERNEL_PERMUTATION_BOOL\(\s*([a-zA-Z_][\w\d]*)\s*\))"));
+	const FRegexPattern KernelPermutationBoolPattern(TEXT(R"(KERNEL_PERMUTATION_BOOL\(\s*([a-zA-Z_]\w*)\s*\))"));
 	{
 		FComputeKernelPermutationSet NewPermutationSet;
 
-		FRegexMatcher BoolPermuationMatcher(KernelPermutationBoolPattern, KernelSourceText);
-		while (BoolPermuationMatcher.FindNext())
+		FRegexMatcher Matcher(KernelPermutationBoolPattern, KernelSourceText);
+		while (Matcher.FindNext())
 		{
-			NewPermutationSet.BooleanOptions.Emplace(BoolPermuationMatcher.GetCaptureGroup(1));
+			NewPermutationSet.BooleanOptions.Emplace(Matcher.GetCaptureGroup(1));
 		}
 
-		PermutationSet = NewPermutationSet;
+		PermutationSet = MoveTemp(NewPermutationSet);
 	}
 
-	const FRegexPattern KernelDefinePattern(TEXT(R"(KERNEL_DEFINE\(\s*([a-zA-Z_][\w\d]*)\s*\))"));
+	const FRegexPattern KernelDefinePattern(TEXT(R"(KERNEL_DEFINE\(\s*([a-zA-Z_][\w]*)\s*\))"));
 	{
 		FComputeKernelDefinitionsSet NewDefinitionsSet;
 
-		FRegexMatcher DefineMatcher(KernelDefinePattern, KernelSourceText);
-		while (DefineMatcher.FindNext())
+		FRegexMatcher Matcher(KernelDefinePattern, KernelSourceText);
+		while (Matcher.FindNext())
+		{
+			NewDefinitionsSet.Defines.Emplace(Matcher.GetCaptureGroup(1));
+		}
+
+		DefinitionsSet = MoveTemp(NewDefinitionsSet);
+	}
+
+	const FRegexPattern KernelInputParamPattern(TEXT(R"(KERNEL_PARAM\(\s*(bool|int|uint|float)((?:[1-4])|(?:[1-4]x[1-4]))?\s*,\s*([a-zA-Z_]\w*)\s*\))"));
+	{
+		TArray<FShaderParamTypeDefinition> Params;
+
+		FRegexMatcher Matcher(KernelInputParamPattern, KernelSourceText);
+		while (Matcher.FindNext())
 		{
-			NewDefinitionsSet.Defines.Emplace(DefineMatcher.GetCaptureGroup(1));
+			FShaderParamTypeDefinition Param = {};
+
+			FString FundamentalType = Matcher.GetCaptureGroup(1);
+			FString DimensionType	= Matcher.GetCaptureGroup(2);
+			FString ParamName		= Matcher.GetCaptureGroup(3);
+
+			Param.Name				= MoveTemp(ParamName);
+			Param.FundamentalType	= FShaderParamTypeDefinition::ParseFundamental(FundamentalType);
+			Param.DimType			= FShaderParamTypeDefinition::ParseDimension(DimensionType);
+			Param.BindingType		= EShaderParamBindingType::ConstantParameter;
+
+			switch (Param.DimType)
+			{
+			case EShaderFundamentalDimensionType::Scalar:
+				Param.VectorDimension = 0;
+				break;
+
+			case EShaderFundamentalDimensionType::Vector:
+				Param.VectorDimension = FShaderParamTypeDefinition::ParseVectorDimension(DimensionType);
+				break;
+
+			case EShaderFundamentalDimensionType::Matrix:
+				FIntVector2 MtxDim = FShaderParamTypeDefinition::ParseMatrixDimension(DimensionType);
+				Param.MatrixColumnCount	= MtxDim.X;
+				Param.MatrixRowCount	= MtxDim.Y;
+				break;
+			}
+
+			Param.ResetTypeDeclaration();
+			Params.Emplace(MoveTemp(Param));
+		}
+
+		InputParams = MoveTemp(Params);
+	}
+
+	const FRegexPattern KernelInputSRVPattern(TEXT(R"(KERNEL_SRV\(\s*(ByteAddress|Structured|Buffer|Texture1D|Texture2D|Texture3D|TextureCube)\s*<\s*(?:SNORM|UNORM)?\s*(int|uint|float)((?:[1-4])|(?:[1-4]x[1-4]))?\s*>\s*,\s*([a-zA-Z_]\w*)\s*\))"));
+	{
+		TArray<FShaderParamTypeDefinition> SRVs;
+
+		FRegexMatcher Matcher(KernelInputSRVPattern, KernelSourceText);
+		while (Matcher.FindNext())
+		{
+			FShaderParamTypeDefinition Param = {};
+
+			FString ResourceType	= Matcher.GetCaptureGroup(1);
+			FString FundamentalType = Matcher.GetCaptureGroup(2);
+			FString DimensionType	= Matcher.GetCaptureGroup(3);
+			FString ParamName		= Matcher.GetCaptureGroup(4);
+
+			Param.Name				= MoveTemp(ParamName);
+			Param.FundamentalType	= FShaderParamTypeDefinition::ParseFundamental(FundamentalType);
+			Param.DimType			= FShaderParamTypeDefinition::ParseDimension(DimensionType);
+			Param.BindingType		= EShaderParamBindingType::ReadOnlyResource;
+			Param.ResourceType		= FShaderParamTypeDefinition::ParseResource(ResourceType);
+
+			switch (Param.DimType)
+			{
+			case EShaderFundamentalDimensionType::Scalar:
+				Param.VectorDimension = 0;
+				break;
+
+			case EShaderFundamentalDimensionType::Vector:
+				Param.VectorDimension = FShaderParamTypeDefinition::ParseVectorDimension(DimensionType);
+				break;
+
+			case EShaderFundamentalDimensionType::Matrix:
+				FIntVector2 MtxDim = FShaderParamTypeDefinition::ParseMatrixDimension(DimensionType);
+				Param.MatrixColumnCount = MtxDim.X;
+				Param.MatrixRowCount	= MtxDim.Y;
+				break;
+			}
+
+			Param.ResetTypeDeclaration();
+			SRVs.Emplace(MoveTemp(Param));
+		}
+
+		InputSRVs = MoveTemp(SRVs);
+	}
+
+	const FRegexPattern KernelOutputPattern(TEXT(R"(KERNEL_UAV\(\s*RW(ByteAddress|Structured|Buffer|Texture1D|Texture2D|Texture3D|TextureCube)\s*<\s*(?:SNORM|UNORM)?\s*(int|uint|float)((?:[1-4])|(?:[1-4]x[1-4]))?\s*>\s*,\s*([a-zA-Z_]\w*)\s*\))"));
+	{
+		TArray<FShaderParamTypeDefinition> OutputParams;
+
+		FRegexMatcher Matcher(KernelOutputPattern, KernelSourceText);
+		while (Matcher.FindNext())
+		{
+			FShaderParamTypeDefinition Param = {};
+
+			FString ResourceType	= Matcher.GetCaptureGroup(1);
+			FString FundamentalType = Matcher.GetCaptureGroup(2);
+			FString DimensionType	= Matcher.GetCaptureGroup(3);
+			FString ParamName		= Matcher.GetCaptureGroup(4);
+
+			Param.Name				= MoveTemp(ParamName);
+			Param.FundamentalType	= FShaderParamTypeDefinition::ParseFundamental(FundamentalType);
+			Param.DimType			= FShaderParamTypeDefinition::ParseDimension(DimensionType);
+			Param.BindingType		= EShaderParamBindingType::ReadWriteResource;
+			Param.ResourceType		= FShaderParamTypeDefinition::ParseResource(ResourceType);
+
+			switch (Param.DimType)
+			{
+			case EShaderFundamentalDimensionType::Scalar:
+				Param.VectorDimension = 0;
+				break;
+
+			case EShaderFundamentalDimensionType::Vector:
+				Param.VectorDimension = FShaderParamTypeDefinition::ParseVectorDimension(DimensionType);
+				break;
+
+			case EShaderFundamentalDimensionType::Matrix:
+				FIntVector2 MtxDim = FShaderParamTypeDefinition::ParseMatrixDimension(DimensionType);
+				Param.MatrixColumnCount = MtxDim.X;
+				Param.MatrixRowCount	= MtxDim.Y;
+				break;
+			}
+
+			Param.ResetTypeDeclaration();
+			OutputParams.Emplace(MoveTemp(Param));
 		}
 
-		DefinitionsSet = NewDefinitionsSet;
+		Outputs = MoveTemp(OutputParams);
 	}
 
 	/*
diff --git a/Engine/Source/Runtime/Engine/Private/ComputeFramework/ComputeKernelResource.cpp b/Engine/Source/Runtime/Engine/Private/ComputeFramework/ComputeKernelResource.cpp
deleted file mode 100644
diff --git a/Engine/Source/Runtime/Engine/Private/ComputeFramework/ComputeKernelShader.cpp b/Engine/Source/Runtime/Engine/Private/ComputeFramework/ComputeKernelShader.cpp
--- a/Engine/Source/Runtime/Engine/Private/ComputeFramework/ComputeKernelShader.cpp
+++ b/Engine/Source/Runtime/Engine/Private/ComputeFramework/ComputeKernelShader.cpp
@@ -1,7 +1,36 @@
 // Copyright Epic Games, Inc. All Rights Reserved.
 
 #include "ComputeFramework/ComputeKernelShader.h"
-#include "ComputeFramework/ComputeKernelShaderType.h"
 
-//IMPLEMENT_SHADER_TYPE(, FComputeKernelShader, TEXT("/Engine/Private/ComputeKernel.usf"), TEXT("Main"), SF_Compute);
-IMPLEMENT_GLOBAL_SHADER(FComputeKernelShader, "/Engine/Private/ComputeKernel.usf", "Main", SF_Compute);
+#include "Engine/VolumeTexture.h"
+#include "ComputeKernelDerivedDataVersion.h"
+#include "ComputeFramework/ComputeKernelShaderCompilationManager.h"
+#include "ComputeFramework/ComputeKernelShared.h"
+#include "ProfilingDebugging/CookStats.h"
+#include "ProfilingDebugging/DiagnosticTable.h"
+#include "Serialization/MemoryReader.h"
+#include "Serialization/MemoryWriter.h"
+#include "ShaderCompiler.h"
+#include "ShaderParameterUtils.h"
+#include "Stats/StatsMisc.h"
+#include "UObject/UObjectGlobals.h"
+
+IMPLEMENT_SHADER_TYPE(, FComputeKernelShader, TEXT("/Engine/Private/ComputeKernel.usf"), TEXT("__"), SF_Compute)
+
+
+FComputeKernelShader::FComputeKernelShader(
+	const FComputeKernelShaderType::CompiledShaderInitializerType& Initializer
+	)
+	: FShader(Initializer)
+{
+	const FShaderParametersMetadata& ShaderParametersMetadata = 
+		static_cast<const FComputeKernelShaderType::FParameters*>(Initializer.Parameters)->ShaderParamMetadata;
+
+	Bindings.BindForLegacyShaderParameters(
+		this, 
+		Initializer.PermutationId, 
+		Initializer.ParameterMap, 
+		ShaderParametersMetadata,
+		true
+		);
+}
diff --git a/Engine/Source/Runtime/Engine/Private/ComputeFramework/ComputeKernelShaderCompilationManager.cpp b/Engine/Source/Runtime/Engine/Private/ComputeFramework/ComputeKernelShaderCompilationManager.cpp
new file mode 100644
--- /dev/null
+++ b/Engine/Source/Runtime/Engine/Private/ComputeFramework/ComputeKernelShaderCompilationManager.cpp
@@ -0,0 +1,447 @@
+// Copyright Epic Games, Inc. All Rights Reserved.
+
+#include "ComputeFramework/ComputeKernelShaderCompilationManager.h"
+
+#include "HAL/FileManager.h"
+#include "Misc/Paths.h"
+#include "ComputeFramework/ComputeKernelShared.h"
+#include "ShaderCompiler.h"
+
+#if WITH_EDITOR
+#include "Interfaces/IShaderFormat.h"
+#include "Interfaces/ITargetPlatformManagerModule.h"
+#endif
+
+
+DEFINE_LOG_CATEGORY_STATIC(LogComputeKernelShaderCompiler, All, All);
+
+static int32 GShowComputeKernelShaderWarnings = 1;
+static FAutoConsoleVariableRef CVarShowComputeKernelShaderWarnings(
+	TEXT("ComputeKernel.ShowShaderCompilerWarnings"),
+	GShowComputeKernelShaderWarnings,
+	TEXT("When set to 1, will display all warnings from ComputeKernel shader compiles.")
+	);
+
+
+
+FComputeKernelShaderCompilationManager GComputeKernelShaderCompilationManager;
+
+void FComputeKernelShaderCompilationManager::Tick(float DeltaSeconds)
+{
+#if WITH_EDITOR
+	RunCompileJobs();
+	ProcessAsyncResults();
+#endif
+}
+
+FComputeKernelShaderCompilationManager::FComputeKernelShaderCompilationManager()
+{
+	
+}
+
+FComputeKernelShaderCompilationManager::~FComputeKernelShaderCompilationManager()
+{
+	for (FComputeKernelShaderCompileWorkerInfo* Info : WorkerInfos)
+	{
+		delete Info;
+	}
+
+	WorkerInfos.Empty();
+}
+
+void FComputeKernelShaderCompilationManager::RunCompileJobs()
+{
+#if WITH_EDITOR
+	check(IsInGameThread());
+
+	InitWorkerInfo();
+
+	int32 NumActiveThreads = 0;
+
+	for (int32 WorkerIndex = 0; WorkerIndex < WorkerInfos.Num(); WorkerIndex++)
+	{
+		FComputeKernelShaderCompileWorkerInfo& CurrentWorkerInfo = *WorkerInfos[WorkerIndex];
+
+		// If this worker doesn't have any queued jobs, look for more in the input queue
+		if (CurrentWorkerInfo.QueuedJobs.Num() == 0)
+		{
+			check(!CurrentWorkerInfo.bComplete);
+
+			if (JobQueue.Num() > 0)
+			{
+				bool bAddedLowLatencyTask = false;
+				int32 JobIndex = 0;
+
+				// Try to grab up to MaxShaderJobBatchSize jobs
+				// Don't put more than one low latency task into a batch
+				for (; JobIndex < JobQueue.Num(); JobIndex++)
+				{
+					CurrentWorkerInfo.QueuedJobs.Add(JobQueue[JobIndex]);
+				}
+
+				// Update the worker state as having new tasks that need to be issued					
+				// don't reset worker app ID, because the shadercompilerworkers don't shutdown immediately after finishing a single job queue.
+				CurrentWorkerInfo.bIssuedTasksToWorker = true;
+				CurrentWorkerInfo.bLaunchedWorker = true;
+				CurrentWorkerInfo.StartTime = FPlatformTime::Seconds();
+				JobQueue.RemoveAt(0, JobIndex);
+			}
+		}
+
+		if (CurrentWorkerInfo.bIssuedTasksToWorker && CurrentWorkerInfo.bLaunchedWorker)
+		{
+			NumActiveThreads++;
+		}
+
+		if (CurrentWorkerInfo.QueuedJobs.Num() > 0)
+		{
+			for (int32 JobIndex = 0; JobIndex < CurrentWorkerInfo.QueuedJobs.Num(); JobIndex++)
+			{
+				FShaderCompileJob& CurrentJob = static_cast<FShaderCompileJob&>(*CurrentWorkerInfo.QueuedJobs[JobIndex].GetReference());
+
+				check(!CurrentJob.bFinalized);
+				CurrentJob.bFinalized = true;
+
+				static ITargetPlatformManagerModule& TPM = GetTargetPlatformManagerRef();
+
+				const FName Format = LegacyShaderPlatformToShaderFormat(EShaderPlatform(CurrentJob.Input.Target.Platform));
+				const IShaderFormat* Compiler = TPM.FindShaderFormat(Format);
+
+				if (!Compiler)
+				{
+					UE_LOG(LogComputeKernelShaderCompiler, Fatal, TEXT("Can't compile shaders for format %s, couldn't load compiler dll"), *Format.ToString());
+				}
+				CA_ASSUME(Compiler != nullptr);
+
+				UE_LOG(LogComputeKernelShaderCompiler, Log, TEXT("Compile Job processing... %s"), *CurrentJob.Input.DebugGroupName);
+
+				CurrentJob.Input.DumpDebugInfoRootPath = GShaderCompilingManager->GetAbsoluteShaderDebugInfoDirectory() / Format.ToString();
+				FPaths::NormalizeDirectoryName(CurrentJob.Input.DumpDebugInfoRootPath);
+				const FShaderCompilingManager::EDumpShaderDebugInfo DumpShaderDebugInfo = GShaderCompilingManager->GetDumpShaderDebugInfo();
+				CurrentJob.Input.DebugExtension.Empty();
+				CurrentJob.Input.DumpDebugInfoRootPath.Empty();
+				if (DumpShaderDebugInfo == FShaderCompilingManager::EDumpShaderDebugInfo::Always)
+				{
+					CurrentJob.Input.DumpDebugInfoRootPath = GShaderCompilingManager->CreateShaderDebugInfoPath(CurrentJob.Input);
+				}
+
+				if (IsValidRef(CurrentJob.Input.SharedEnvironment))
+				{
+					// Merge the shared environment into the per-shader environment before calling into the compile function
+					// Normally this happens in the worker
+					CurrentJob.Input.Environment.Merge(*CurrentJob.Input.SharedEnvironment);
+				}
+
+				// Compile the shader directly through the platform dll (directly from the shader dir as the working directory)
+				Compiler->CompileShader(Format, CurrentJob.Input, CurrentJob.Output, FString(FPlatformProcess::ShaderDir()));
+
+				CurrentJob.bSucceeded = CurrentJob.Output.bSucceeded;
+
+				// Recompile the shader to dump debug info if desired
+				if (GShaderCompilingManager->ShouldRecompileToDumpShaderDebugInfo(CurrentJob))
+				{
+					CurrentJob.Input.DumpDebugInfoPath = GShaderCompilingManager->CreateShaderDebugInfoPath(CurrentJob.Input);
+					Compiler->CompileShader(Format, CurrentJob.Input, CurrentJob.Output, FString(FPlatformProcess::ShaderDir()));
+				}
+
+				if (CurrentJob.Output.bSucceeded)
+				{
+					// Generate a hash of the output and cache it
+					// The shader processing this output will use it to search for existing FShaderResources
+					CurrentJob.Output.GenerateOutputHash();
+					UE_LOG(LogComputeKernelShaderCompiler, Log, TEXT("GPU shader compile succeeded. Id %d"), CurrentJob.Id);
+				}
+				else
+				{
+					UE_LOG(LogComputeKernelShaderCompiler, Log, TEXT("ERROR: GPU shader compile failed! Id %d"), CurrentJob.Id);
+				}
+
+				CurrentWorkerInfo.bComplete = true;
+			}
+		}
+	}
+
+	for (int32 WorkerIndex = 0; WorkerIndex < WorkerInfos.Num(); WorkerIndex++)
+	{
+		FComputeKernelShaderCompileWorkerInfo& CurrentWorkerInfo = *WorkerInfos[WorkerIndex];
+		if (CurrentWorkerInfo.bComplete)
+		{
+			for (int32 JobIndex = 0; JobIndex < CurrentWorkerInfo.QueuedJobs.Num(); JobIndex++)
+			{
+				FComputeKernelShaderMapCompileResults& ShaderMapResults = ComputeKernelShaderMapJobs.FindChecked(CurrentWorkerInfo.QueuedJobs[JobIndex]->Id);
+				ShaderMapResults.FinishedJobs.Add(CurrentWorkerInfo.QueuedJobs[JobIndex]);
+				ShaderMapResults.bAllJobsSucceeded = ShaderMapResults.bAllJobsSucceeded && CurrentWorkerInfo.QueuedJobs[JobIndex]->bSucceeded;
+			}
+		}
+
+		CurrentWorkerInfo.bComplete = false;
+		CurrentWorkerInfo.QueuedJobs.Empty();
+	}
+#endif
+}
+
+void FComputeKernelShaderCompilationManager::InitWorkerInfo()
+{
+	if (WorkerInfos.Num() == 0) // Check to see if it has been initialized or not
+	{
+		// Ew. Should we just use FShaderCompilingManager's workers instead? Is that safe?
+		const int32 NumVirtualCores = FPlatformMisc::NumberOfCoresIncludingHyperthreads();
+		const uint32 NumComputeKernelShaderCompilingThreads = FMath::Min(NumVirtualCores - 1, 4);
+
+		for (uint32 WorkerIndex = 0; WorkerIndex < NumComputeKernelShaderCompilingThreads; WorkerIndex++)
+		{
+			WorkerInfos.Add(new FComputeKernelShaderCompileWorkerInfo());
+		}
+	}	
+}
+
+void FComputeKernelShaderCompilationManager::AddJobs(TArray<FShaderCommonCompileJobPtr> InNewJobs)
+{
+#if WITH_EDITOR
+	for (auto& Job : InNewJobs)
+	{
+		FComputeKernelShaderMapCompileResults& ShaderMapInfo = ComputeKernelShaderMapJobs.FindOrAdd(Job->Id);
+		//@todo : Apply shader map isn't used for now with this compile manager. Should be merged to have a generic shader compiler
+		ShaderMapInfo.NumJobsQueued++;
+	}
+
+	JobQueue.Append(InNewJobs);
+#endif
+}
+
+
+void FComputeKernelShaderCompilationManager::ProcessAsyncResults()
+{
+#if WITH_EDITOR
+	int32 NumCompilingComputeKernelShaderMaps = 0;
+	TArray<int32> ShaderMapsToRemove;
+
+	// Get all ComputeKernel shader maps to finalize
+	//
+	for (TMap<int32, FComputeKernelShaderMapCompileResults>::TIterator It(ComputeKernelShaderMapJobs); It; ++It)
+	{
+		const FComputeKernelShaderMapCompileResults& Results = It.Value();
+
+		if (Results.FinishedJobs.Num() == Results.NumJobsQueued)
+		{
+			ShaderMapsToRemove.Add(It.Key());
+			PendingFinalizeComputeKernelShaderMaps.Add(It.Key(), FComputeKernelShaderMapFinalizeResults(Results));
+		}
+	}
+
+	for (int32 RemoveIndex = 0; RemoveIndex < ShaderMapsToRemove.Num(); RemoveIndex++)
+	{
+		ComputeKernelShaderMapJobs.Remove(ShaderMapsToRemove[RemoveIndex]);
+	}
+
+	NumCompilingComputeKernelShaderMaps = ComputeKernelShaderMapJobs.Num();
+
+	if (PendingFinalizeComputeKernelShaderMaps.Num() > 0)
+	{
+		ProcessCompiledComputeKernelShaderMaps(PendingFinalizeComputeKernelShaderMaps, 0.1f);
+	}
+#endif
+}
+
+
+void FComputeKernelShaderCompilationManager::ProcessCompiledComputeKernelShaderMaps(
+	TMap<int32, FComputeKernelShaderMapFinalizeResults>& CompiledShaderMaps,
+	float TimeBudget)
+{
+#if WITH_EDITOR
+	// Keeps shader maps alive as they are passed from the shader compiler and applied to the owning kernel
+	TArray<TRefCountPtr<FComputeKernelShaderMap> > LocalShaderMapReferences;
+	TMap<FComputeKernelResource*, FComputeKernelShaderMap*> KernelsToUpdate;
+
+	// Process compiled shader maps in FIFO order, in case a shader map has been enqueued multiple times,
+	// Which can happen if a kernel is edited while a background compile is going on
+	for (TMap<int32, FComputeKernelShaderMapFinalizeResults>::TIterator ProcessIt(CompiledShaderMaps); ProcessIt; ++ProcessIt)
+	{
+		TRefCountPtr<FComputeKernelShaderMap> ShaderMap = nullptr;
+		TArray<FComputeKernelResource*>* Kernels = nullptr;
+
+		for (TMap<TRefCountPtr<FComputeKernelShaderMap>, TArray<FComputeKernelResource*> >::TIterator ShaderMapIt(FComputeKernelShaderMap::GetInFlightShaderMaps()); ShaderMapIt; ++ShaderMapIt)
+		{
+			if (ShaderMapIt.Key()->GetCompilingId() == ProcessIt.Key())
+			{
+				ShaderMap = ShaderMapIt.Key();
+				Kernels = &ShaderMapIt.Value();
+				break;
+			}
+		}
+
+		if (ShaderMap && Kernels)
+		{
+			TArray<FString> Errors;
+			FComputeKernelShaderMapFinalizeResults& CompileResults = ProcessIt.Value();
+			TArray<FShaderCommonCompileJobPtr>& ResultArray = CompileResults.FinishedJobs;
+
+			// Make a copy of the array as this entry of FComputeKernelShaderMap::ShaderMapsBeingCompiled will be removed below
+			TArray<FComputeKernelResource*> KernelsArray = *Kernels;
+			bool bSuccess = true;
+
+			for (int32 JobIndex = 0; JobIndex < ResultArray.Num(); JobIndex++)
+			{
+				FShaderCompileJob& CurrentJob = static_cast<FShaderCompileJob&>(*ResultArray[JobIndex].GetReference());
+				bSuccess = bSuccess && CurrentJob.bSucceeded;
+
+				if (bSuccess)
+				{
+					check(CurrentJob.Output.ShaderCode.GetShaderCodeSize() > 0);
+				}
+
+				if (GShowComputeKernelShaderWarnings || !CurrentJob.bSucceeded)
+				{
+					for (int32 ErrorIndex = 0; ErrorIndex < CurrentJob.Output.Errors.Num(); ErrorIndex++)
+					{
+						Errors.AddUnique(CurrentJob.Output.Errors[ErrorIndex].GetErrorString());
+					}
+
+					if (CurrentJob.Output.Errors.Num())
+					{
+						UE_LOG(LogShaders, Log, TEXT("There were errors for job \"%s\""), *CurrentJob.Input.DebugGroupName)
+							for (const FShaderCompilerError& Error : CurrentJob.Output.Errors)
+							{
+								UE_LOG(LogShaders, Log, TEXT("Error: %s"), *Error.GetErrorString())
+							}
+					}
+				}
+				else
+				{
+					UE_LOG(LogShaders, Log, TEXT("There were NO errors for job \"%s\""), *CurrentJob.Input.DebugGroupName);
+				}
+			}
+
+			bool bShaderMapComplete = true;
+
+			if (bSuccess)
+			{
+				bShaderMapComplete = ShaderMap->ProcessCompilationResults(ResultArray, CompileResults.FinalizeJobIndex, TimeBudget);
+			}
+
+
+			if (bShaderMapComplete)
+			{
+				ShaderMap->SetCompiledSuccessfully(bSuccess);
+
+				// Pass off the reference of the shader map to LocalShaderMapReferences
+				LocalShaderMapReferences.Add(ShaderMap);
+				FComputeKernelShaderMap::GetInFlightShaderMaps().Remove(ShaderMap);
+
+				for (FComputeKernelResource* Kernel : KernelsArray)
+				{
+					FComputeKernelShaderMap* CompletedShaderMap = ShaderMap;
+
+					Kernel->RemoveOutstandingCompileId(ShaderMap->GetCompilingId());
+
+					// Only process results that still match the ID which requested a compile
+					// This avoids applying shadermaps which are out of date and a newer one is in the async compiling pipeline
+					if (Kernel->IsSame(CompletedShaderMap->GetShaderMapId()))
+					{
+						if (Errors.Num() != 0)
+						{
+							FString SourceCode;
+							Kernel->GetHLSLSource(SourceCode);
+							UE_LOG(LogComputeKernelShaderCompiler, Log, TEXT("Compile output as text:"));
+							UE_LOG(LogComputeKernelShaderCompiler, Log, TEXT("==================================================================================="));
+							TArray<FString> OutputByLines;
+							SourceCode.ParseIntoArrayLines(OutputByLines, false);
+							for (int32 i = 0; i < OutputByLines.Num(); i++)
+							{
+								UE_LOG(LogComputeKernelShaderCompiler, Log, TEXT("/*%04d*/\t\t%s"), i + 1, *OutputByLines[i]);
+							}
+							UE_LOG(LogComputeKernelShaderCompiler, Log, TEXT("==================================================================================="));
+						}
+
+						if (!bSuccess)
+						{
+							// Propagate error messages
+							Kernel->SetCompileErrors(Errors);
+							KernelsToUpdate.Add(Kernel, nullptr);
+
+							for (int32 ErrorIndex = 0; ErrorIndex < Errors.Num(); ErrorIndex++)
+							{
+								FString ErrorMessage = Errors[ErrorIndex];
+								// Work around build machine string matching heuristics that will cause a cook to fail
+								ErrorMessage.ReplaceInline(TEXT("error "), TEXT("err0r "), ESearchCase::CaseSensitive);
+								UE_LOG(LogComputeKernelShaderCompiler, Warning, TEXT("	%s"), *ErrorMessage);
+							}
+						}
+						else
+						{
+							// if we succeeded and our shader map is not complete this could be because the kernel was being edited quicker then the compile could be completed
+							// Don't modify kernel for which the compiled shader map is no longer complete
+							// This shouldn't happen since kernels are pretty much baked in the designated config file.
+							if (CompletedShaderMap->IsComplete(Kernel, true))
+							{
+								KernelsToUpdate.Add(Kernel, CompletedShaderMap);
+							}
+
+							if (GShowComputeKernelShaderWarnings && Errors.Num() > 0)
+							{
+								UE_LOG(LogComputeKernelShaderCompiler, Warning, TEXT("Warnings while compiling ComputeKernel %s for platform %s:"),
+									*Kernel->GetFriendlyName(),
+									*LegacyShaderPlatformToShaderFormat(ShaderMap->GetShaderPlatform()).ToString());
+								for (int32 ErrorIndex = 0; ErrorIndex < Errors.Num(); ErrorIndex++)
+								{
+									UE_LOG(LogComputeKernelShaderCompiler, Warning, TEXT("	%s"), *Errors[ErrorIndex]);
+								}
+							}
+						}
+					}
+					else
+					{
+						if (CompletedShaderMap->IsComplete(Kernel, true))
+						{
+							Kernel->NotifyCompilationFinished();
+						}
+					}
+				}
+
+				// Cleanup shader jobs and compile tracking structures
+				ResultArray.Empty();
+				CompiledShaderMaps.Remove(ShaderMap->GetCompilingId());
+			}
+
+			if (TimeBudget < 0)
+			{
+				break;
+			}
+		}
+	}
+
+	if (KernelsToUpdate.Num() > 0)
+	{
+		for (TMap<FComputeKernelResource*, FComputeKernelShaderMap*>::TConstIterator It(KernelsToUpdate); It; ++It)
+		{
+			FComputeKernelResource* Kernel = It.Key();
+			FComputeKernelShaderMap* ShaderMap = It.Value();
+
+			Kernel->SetGameThreadShaderMap(It.Value());
+
+			ENQUEUE_RENDER_COMMAND(FSetShaderMapOnComputeKernel)(
+				[Kernel, ShaderMap](FRHICommandListImmediate& RHICmdList)
+				{
+					Kernel->SetRenderingThreadShaderMap(ShaderMap);
+				});
+
+
+			Kernel->NotifyCompilationFinished();
+		}
+	}
+#endif
+}
+
+
+void FComputeKernelShaderCompilationManager::FinishCompilation(const TCHAR* InKernelName, const TArray<int32>& ShaderMapIdsToFinishCompiling)
+{
+#if WITH_EDITOR
+	check(!FPlatformProperties::RequiresCookedData());
+
+	RunCompileJobs();	// since we don't async compile through another process, this will run all oustanding jobs
+	ProcessAsyncResults();	// grab compiled shader maps and assign them to their resources
+
+	check(ComputeKernelShaderMapJobs.Num() == 0);
+#endif
+}
+
diff --git a/Engine/Source/Runtime/Engine/Private/ComputeFramework/ComputeKernelShaderMap.cpp b/Engine/Source/Runtime/Engine/Private/ComputeFramework/ComputeKernelShaderMap.cpp
--- a/Engine/Source/Runtime/Engine/Private/ComputeFramework/ComputeKernelShaderMap.cpp
+++ b/Engine/Source/Runtime/Engine/Private/ComputeFramework/ComputeKernelShaderMap.cpp
@@ -1,121 +1,723 @@
 // Copyright Epic Games, Inc. All Rights Reserved.
 
-#include "ComputeFramework/ComputeKernelShaderMap.h"
+#include "ComputeFramework/ComputeKernelShader.h"
+
+#include "Engine/VolumeTexture.h"
+#include "ComputeFramework/ComputeKernelShared.h"
+#include "ComputeFramework/ComputeKernelDerivedDataVersion.h"
+#include "ComputeFramework/ComputeKernelShaderCompilationManager.h"
+#include "ProfilingDebugging/CookStats.h"
+#include "ProfilingDebugging/DiagnosticTable.h"
+#include "Serialization/MemoryWriter.h"
+#include "Serialization/MemoryReader.h"
 #include "ShaderCompiler.h"
-#include "RHI.h"
-#include "ComputeFramework/ComputeKernelResource.h"
-#include "ComputeFramework/ComputeKernel.h"
+#include "Stats/StatsMisc.h"
+#include "UObject/UObjectGlobals.h"
 
-IMPLEMENT_TYPE_LAYOUT(FComputeKernelCompilationOutput);
-IMPLEMENT_TYPE_LAYOUT(FComputeKernelShaderMapId);
-IMPLEMENT_TYPE_LAYOUT(FComputeKernelShaderMapContent);
+#if WITH_EDITOR
+	#include "DerivedDataCacheInterface.h"
+	#include "Interfaces/ITargetPlatformManagerModule.h"
+	#include "TickableEditorObject.h"
+#endif
 
-TMap<TRefCountPtr<FComputeKernelShaderMap>, TArray<const FComputeKernelResource*>> FComputeKernelShaderMap::ComputeKernelShaderMapsBeingCompiled;
 
-static bool ShouldCacheComputeKernelShader(
-	EShaderPlatform ShaderPlatform, 
-	const FComputeKernelShaderType* ShaderType, 
-	const FComputeKernelResource* KernelShader
-	)
+#if ENABLE_COOK_STATS
+namespace ComputeKernelShaderCookStats
 {
-	return ShaderType->ShouldCache(ShaderPlatform, KernelShader) && KernelShader->ShouldCache(ShaderPlatform, ShaderType);
+	static FCookStats::FDDCResourceUsageStats UsageStats;
+	static int32 ShadersCompiled = 0;
+	static FCookStatsManager::FAutoRegisterCallback RegisterCookStats([](FCookStatsManager::AddStatFuncRef AddStat)
+	{
+		UsageStats.LogStats(AddStat, TEXT("ComputeKernelShader.Usage"), TEXT(""));
+		AddStat(TEXT("ComputeKernelShader.Misc"), FCookStatsManager::CreateKeyValueArray(
+			TEXT("ShadersCompiled"), ShadersCompiled
+		));
+	});
 }
+#endif
 
-void FComputeKernelShaderMap::LoadFromDerivedDataCache(
-	EShaderPlatform ShaderPlatform, 
-	const FComputeKernelShaderMapId& ShaderMapId, 
-	FComputeKernelResource* KernelShader, 
-	TRefCountPtr<FComputeKernelShaderMap>& InOutGameThreadShaderMap
-	)
+
+//
+// Globals
+//
+TMap<FComputeKernelShaderMapId, FComputeKernelShaderMap*> FComputeKernelShaderMap::GIdToComputeKernelShaderMap[SP_NumPlatforms];
+TArray<FComputeKernelShaderMap*> FComputeKernelShaderMap::AllComputeKernelShaderMaps;
+
+// The Id of 0 is reserved for global shaders
+uint32 FComputeKernelShaderMap::NextCompilingId = 2;
+
+
+/** 
+ * Tracks FComputeKernelResource and their shader maps that are being compiled.
+ * Uses a TRefCountPtr as this will be the only reference to a shader map while it is being compiled.
+ */
+TMap<TRefCountPtr<FComputeKernelShaderMap>, TArray<FComputeKernelResource*> > FComputeKernelShaderMap::ComputeKernelShaderMapsBeingCompiled;
+
+
+
+static inline bool ShouldCacheComputeKernelShader(const FComputeKernelShaderType* InShaderType, EShaderPlatform InPlatform, const FComputeKernelResource* InKernel)
+{
+	return InShaderType->ShouldCache(InPlatform, InKernel) && InKernel->ShouldCache(InPlatform, InShaderType);
+}
+
+
+/** Hashes the kernel specific part of this shader map Id. */
+void FComputeKernelShaderMapId::GetComputeKernelHash(FSHAHash& OutHash) const
+{
+	FSHA1 HashState;
+
+	HashState.Update((const uint8*)&ShaderCodeHash, sizeof(ShaderCodeHash));
+	HashState.Update((const uint8*)&FeatureLevel, sizeof(FeatureLevel));
+
+	HashState.Final();
+	HashState.GetHash(&OutHash.Hash[0]);
+}
+
+/** 
+* Tests this set against another for equality, disregarding override settings.
+* 
+* @param ReferenceSet	The set to compare against
+* @return				true if the sets are equal
+*/
+bool FComputeKernelShaderMapId::operator==(const FComputeKernelShaderMapId& InReferenceSet) const
+{
+	if (  ShaderCodeHash != InReferenceSet.ShaderCodeHash
+		|| FeatureLevel != InReferenceSet.FeatureLevel)
+	{
+		return false;
+	}
+
+	if (ShaderTypeDependencies.Num() != InReferenceSet.ShaderTypeDependencies.Num())
+	{
+		return false;
+	}
+
+	if (LayoutParams != InReferenceSet.LayoutParams)
+	{
+		return false;
+	}
+
+	for (int32 ShaderIndex = 0; ShaderIndex < ShaderTypeDependencies.Num(); ShaderIndex++)
+	{
+		const FShaderTypeDependency& ShaderTypeDependency = ShaderTypeDependencies[ShaderIndex];
+
+		if (ShaderTypeDependency != InReferenceSet.ShaderTypeDependencies[ShaderIndex])
+		{
+			return false;
+		}
+	}
+
+	return true;
+}
+
+void FComputeKernelShaderMapId::AppendKeyString(FString& OutKeyString) const
 {
+#if WITH_EDITOR
+	OutKeyString += FString::Printf(TEXT("%llX"), ShaderCodeHash);
+	OutKeyString += TEXT("_");
+
+	FString FeatureLevelString;
+	GetFeatureLevelName(FeatureLevel, FeatureLevelString);
+
+	{
+		const FSHAHash LayoutHash = Freeze::HashLayout(StaticGetTypeLayoutDesc<FComputeKernelShaderMapContent>(), LayoutParams);
+		OutKeyString += TEXT("_");
+		OutKeyString += LayoutHash.ToString();
+		OutKeyString += TEXT("_");
+	}
+
+	TMap<const TCHAR*, FCachedUniformBufferDeclaration> ReferencedUniformBuffers;
+
+	// Add the inputs for any shaders that are stored inline in the shader map
+	for (const FShaderTypeDependency& ShaderTypeDependency : ShaderTypeDependencies)
+	{
+		const FShaderType* ShaderType = FindShaderTypeByName(ShaderTypeDependency.ShaderTypeName);
+		OutKeyString += TEXT("_");
+		OutKeyString += ShaderType->GetName();
+		OutKeyString += ShaderTypeDependency.SourceHash.ToString();
+		
+		const FSHAHash LayoutHash = Freeze::HashLayout(ShaderType->GetLayout(), LayoutParams);
+		OutKeyString += LayoutHash.ToString();
+
+		const TMap<const TCHAR*, FCachedUniformBufferDeclaration>& ReferencedUniformBufferStructsCache = ShaderType->GetReferencedUniformBufferStructsCache();
+
+		for (TMap<const TCHAR*, FCachedUniformBufferDeclaration>::TConstIterator It(ReferencedUniformBufferStructsCache); It; ++It)
+		{
+			ReferencedUniformBuffers.Add(It.Key(), It.Value());
+		}
+	}
+
+	{
+		TArray<uint8> TempData;
+		FSerializationHistory SerializationHistory;
+		FMemoryWriter Ar(TempData, true);
+		FShaderSaveArchive SaveArchive(Ar, SerializationHistory);
 
+		// Save uniform buffer member info so we can detect when layout has changed
+		SerializeUniformBufferInfo(SaveArchive, ReferencedUniformBuffers);
+
+		SerializationHistory.AppendKeyString(OutKeyString);
+	}
+#endif //WITH_EDITOR
 }
 
-bool FComputeKernelShaderMap::TryToAddToExistingCompilationTask(
-	FComputeKernelResource* KernelShader
+/**
+ * Enqueues a compilation for a new shader of this type.
+ * @param InKernel - The kernel to link the shader with.
+ */
+void FComputeKernelShaderType::BeginCompileShader(
+	uint32 InShaderMapId,
+	const FComputeKernelResource* InKernel,
+	FSharedShaderCompilerEnvironment* InCompilationEnvironment,
+	EShaderPlatform InPlatform,
+	TArray<FShaderCommonCompileJobPtr>& NewJobs,
+	FShaderTarget InTarget
 	)
 {
-	return false;
+	FShaderCompileJob* NewJob = GShaderCompilingManager->PrepareShaderCompileJob(
+		InShaderMapId, 
+		FShaderCompileJobKey(this, nullptr, /* PermutationId = */ 0),
+		EShaderCompileJobPriority::Normal
+		);
+
+	FString virtualSourcePath = FString::Printf(TEXT("/Engine/Generated/ComputeFramework/Kernel_%s.usf"), *InKernel->GetEntryPoint());
+
+	NewJob->ShaderParameters = MakeShared<const FParameters, ESPMode::ThreadSafe>(*InKernel->GetShaderParamMetadata());
+	NewJob->Input.SharedEnvironment = InCompilationEnvironment;
+	NewJob->Input.Target = InTarget;
+	NewJob->Input.ShaderFormat = LegacyShaderPlatformToShaderFormat(InPlatform);
+	NewJob->Input.VirtualSourceFilePath = virtualSourcePath;
+	NewJob->Input.EntryPointName = InKernel->GetEntryPoint();
+	NewJob->Input.Environment.IncludeVirtualPathToContentsMap.Add(virtualSourcePath, InKernel->GetHLSLSource());
+	UE_LOG(LogShaders, Verbose, TEXT("%s"), *InKernel->GetHLSLSource());
+	
+	FShaderCompilerEnvironment& ShaderEnvironment = NewJob->Input.Environment;
+
+	UE_LOG(LogShaders, Verbose, TEXT("			%s"), GetName());
+	COOK_STAT(ComputeKernelShaderCookStats::ShadersCompiled++);
+
+	// Allow the shader type to modify the compile environment.
+	SetupCompileEnvironment(InPlatform, InKernel, ShaderEnvironment);
+
+	::GlobalBeginCompileShader(
+		InKernel->GetFriendlyName(),
+		nullptr,
+		this,
+		nullptr,//ShaderPipeline,
+		0, // PermutationId
+		*virtualSourcePath,
+		*InKernel->GetEntryPoint(),
+		FShaderTarget(GetFrequency(), InPlatform),
+		NewJob->Input
+		);
+
+	NewJobs.Add(FShaderCommonCompileJobPtr(NewJob));
+}
+
+/**
+ * Either creates a new instance of this type or returns an equivalent existing shader.
+ * @param InShaderMapHash - Precomputed hash of the shader map 
+ * @param InCurrentJob - Compile job that was enqueued by BeginCompileShader.
+ */
+FShader* FComputeKernelShaderType::FinishCompileShader(
+	const FSHAHash& InShaderMapHash,
+	const FShaderCompileJob& InCurrentJob,
+	const FString& InDebugDescription
+	) const
+{
+	check(InCurrentJob.bSucceeded);
+
+	const int32 PermutationId = 0;
+	FShader* Shader = ConstructCompiled(
+		FComputeKernelShaderType::CompiledShaderInitializerType(
+			this,
+			static_cast<const FParameters*>(InCurrentJob.ShaderParameters.Get()),
+			PermutationId,
+			InCurrentJob.Output,
+			InShaderMapHash,
+			InDebugDescription
+			)
+		);
+
+	InCurrentJob.Output.ParameterMap.VerifyBindingsAreComplete(GetName(), InCurrentJob.Output.Target, InCurrentJob.Key.VFType);
+
+	return Shader;
+}
+
+/**
+ * Finds the shader map for a kernel.
+ * @param InShaderMapId - The kernel id and static parameter set identifying the shader map
+ * @param InPlatform - The platform to lookup for
+ * @return nullptr if no cached shader map was found.
+ */
+FComputeKernelShaderMap* FComputeKernelShaderMap::FindId(const FComputeKernelShaderMapId& InShaderMapId, EShaderPlatform InPlatform)
+{
+	FComputeKernelShaderMap* Result = GIdToComputeKernelShaderMap[InPlatform].FindRef(InShaderMapId);
+	check(Result == nullptr || !Result->bDeletedThroughDeferredCleanup);
+	return Result;
+}
+
+void ComputeKernelShaderMapAppendKeyString(EShaderPlatform InPlatform, FString& OutKeyString)
+{
+#if WITH_EDITOR
+	//Keep library version in the DDC key to invalidate it once we move to a new library
+	OutKeyString += TEXT("ComputeKernelVersion");
+	OutKeyString += TEXT("_");
+#endif //WITH_EDITOR
+}
+
+/** Creates a string key for the derived data cache given a shader map id. */
+static FString GetComputeKernelShaderMapKeyString(const FComputeKernelShaderMapId& InShaderMapId, EShaderPlatform InPlatform)
+{
+#if WITH_EDITOR
+	const FName Format = LegacyShaderPlatformToShaderFormat(InPlatform);
+	FString ShaderMapKeyString = Format.ToString() + TEXT("_") + FString(FString::FromInt(GetTargetPlatformManagerRef().ShaderFormatVersion(Format))) + TEXT("_");
+	ComputeKernelShaderMapAppendKeyString(InPlatform, ShaderMapKeyString);
+	ShaderMapAppendKeyString(InPlatform, ShaderMapKeyString);
+	InShaderMapId.AppendKeyString(ShaderMapKeyString);
+	return FDerivedDataCacheInterface::BuildCacheKey(TEXT("CKERSM"), COMPUTEKERNEL_DERIVEDDATA_VER, *ShaderMapKeyString);
+#else
+	return FString();
+#endif
+}
+
+void FComputeKernelShaderMap::LoadFromDerivedDataCache(const FComputeKernelResource* InKernel, const FComputeKernelShaderMapId& InShaderMapId, EShaderPlatform InPlatform, TRefCountPtr<FComputeKernelShaderMap>& InOutShaderMap)
+{
+#if WITH_EDITOR
+	if (InOutShaderMap != nullptr)
+	{
+		check(InOutShaderMap->GetShaderPlatform() == InPlatform);
+		// If the shader map was non-NULL then it was found in memory but is incomplete, attempt to load the missing entries from memory
+		InOutShaderMap->LoadMissingShadersFromMemory(InKernel);
+	}
+	else
+	{
+		// Shader map was not found in memory, try to load it from the DDC
+		STAT(double ComputeKernelShaderDDCTime = 0);
+		{
+			SCOPE_SECONDS_COUNTER(ComputeKernelShaderDDCTime);
+			COOK_STAT(auto Timer = ComputeKernelShaderCookStats::UsageStats.TimeSyncWork());
+
+			TArray<uint8> CachedData;
+			const FString DataKey = GetComputeKernelShaderMapKeyString(InShaderMapId, InPlatform);
+
+			if (GetDerivedDataCacheRef().GetSynchronous(*DataKey, CachedData, InKernel->GetFriendlyName()))
+			{
+				COOK_STAT(Timer.AddHit(CachedData.Num()));
+				InOutShaderMap = new FComputeKernelShaderMap();
+				FMemoryReader Ar(CachedData, true);
+
+				// Deserialize from the cached data
+				InOutShaderMap->Serialize(Ar);
+				//InOutShaderMap->RegisterSerializedShaders(false);
+
+				checkSlow(InOutShaderMap->GetShaderMapId() == InShaderMapId);
+
+				// Register in the global map
+				InOutShaderMap->Register(InPlatform);
+			}
+			else
+			{
+				// We should be build the data later, and we can track that the resource was built there when we push it to the DDC.
+				COOK_STAT(Timer.TrackCyclesOnly());
+				InOutShaderMap = nullptr;
+			}
+		}
+		INC_FLOAT_STAT_BY(STAT_ShaderCompiling_DDCLoading,(float)ComputeKernelShaderDDCTime);
+	}
+#endif
 }
 
+void FComputeKernelShaderMap::SaveToDerivedDataCache()
+{
+#if WITH_EDITOR
+	COOK_STAT(auto Timer = ComputeKernelShaderCookStats::UsageStats.TimeSyncWork());
+	TArray<uint8> SaveData;
+	FMemoryWriter Ar(SaveData, true);
+	Serialize(Ar);
+
+	GetDerivedDataCacheRef().Put(*GetComputeKernelShaderMapKeyString(GetContent()->ShaderMapId, GetShaderPlatform()), SaveData, FStringView(*GetFriendlyName()));
+	COOK_STAT(Timer.AddMiss(SaveData.Num()));
+#endif
+}
+
+/**
+* Compiles the shaders for a kernel and caches them in this shader map.
+* @param InKernel - The kernel to compile shaders for.
+* @param InShaderMapId - the kernel id and set of static parameters to compile for
+* @param InPlatform - The platform to compile to
+*/
 void FComputeKernelShaderMap::Compile(
-	EShaderPlatform ShaderPlatform,
-	FComputeKernelResource* KernelResource,
-	const FComputeKernelShaderMapId& ShaderMapId,
-	bool bSynchronousCompile
+	FComputeKernelResource* InKernel, 
+	const FComputeKernelShaderMapId& InShaderMapId,
+	TRefCountPtr<FSharedShaderCompilerEnvironment> InCompilationEnvironment,
+	const FComputeKernelCompilationOutput& InComputeKernelCompilationOutput,
+	EShaderPlatform InPlatform,
+	bool bSynchronousCompile,
+	bool bApplyCompletedShaderMapForRendering
 	)
 {
-	check(IsInGameThread());
-	check(GetRefCount() > 0);
-	
 	if (FPlatformProperties::RequiresCookedData())
 	{
-		UE_LOG(ComputeKernel, Fatal, TEXT("Trying to compile FComputeKernelResource [%s] at run-time is not supported on %s!"), *KernelResource->GetFriendlyName(), FPlatformProperties::PlatformName());
-		return;
+		UE_LOG(LogShaders, Fatal, TEXT("Trying to compile ComputeKernel shader %s at run-time, which is not supported on consoles!"), *InKernel->GetFriendlyName() );
 	}
+	else
+	{
+		// Make sure we are operating on a referenced shader map or the below Find will cause this shader map to be deleted,
+		// Since it creates a temporary ref counted pointer.
+		check(NumRefs > 0);
+  
+		// Add this shader map and to ComputeKernelShaderMapsBeingCompiled
+		TArray<FComputeKernelResource*>* Kernel = ComputeKernelShaderMapsBeingCompiled.Find(this);
+  
+		if (Kernel)
+		{
+			check(!bSynchronousCompile);
+			Kernel->AddUnique(InKernel);
+		}
+		else
+		{
+			// Assign a unique identifier so that shaders from this shader map can be associated with it after a deferred compile
+			CompilingId = NextCompilingId;
+			UE_LOG(LogShaders, Log, TEXT("CompilingId = %p %d"), InKernel, CompilingId);
+			InKernel->AddCompileId(CompilingId);
 
-	TArray<const FComputeKernelResource*>* CorrespondingKernels = ComputeKernelShaderMapsBeingCompiled.Find(this);
+			check(NextCompilingId < UINT_MAX);
+			NextCompilingId++;
+  
+			TArray<FComputeKernelResource*> NewCorrespondingKernels;
+			NewCorrespondingKernels.Add(InKernel);
+			ComputeKernelShaderMapsBeingCompiled.Add(this, NewCorrespondingKernels);
+#if DEBUG_INFINITESHADERCOMPILE
+			UE_LOG(LogTemp, Display, TEXT("Added ComputeKernel ShaderMap 0x%08X%08X with kernel 0x%08X%08X to ComputeKernelShaderMapsBeingCompiled"), (int)((int64)(this) >> 32), (int)((int64)(this)), (int)((int64)(InKernel) >> 32), (int)((int64)(InKernel)));
+#endif  
+			// Setup the compilation environment.
+			InKernel->SetupShaderCompilationEnvironment(InPlatform, *InCompilationEnvironment);
+  
+			// Store the kernel name for debugging purposes.
+			FComputeKernelShaderMapContent* NewContent = new FComputeKernelShaderMapContent(InPlatform);
+			NewContent->FriendlyName = InKernel->GetFriendlyName();
+			NewContent->CompilationOutput = InComputeKernelCompilationOutput;
+			NewContent->ShaderMapId = InShaderMapId;
+			AssignContent(NewContent);
 
-	if (CorrespondingKernels)
+			uint32 NumShaders = 0;
+			TArray<FShaderCommonCompileJobPtr> NewJobs;
+	
+			// Iterate over all shader types.
+			for(TLinkedList<FShaderType*>::TIterator ShaderTypeIt(FShaderType::GetTypeList());ShaderTypeIt;ShaderTypeIt.Next())
+			{
+				FComputeKernelShaderType* ShaderType = ShaderTypeIt->GetComputeKernelShaderType();
+				if (ShaderType && ShouldCacheComputeKernelShader(ShaderType, InPlatform, InKernel))
+				{
+					// Verify that the shader map Id contains inputs for any shaders that will be put into this shader map
+					check(InShaderMapId.ContainsShaderType(ShaderType));
+					
+					// Compile this ComputeKernel shader .
+					TArray<FString> ShaderErrors;
+  
+					// Only compile the shader if we don't already have it
+					if (!NewContent->HasShader(ShaderType, /* PermutationId = */ 0))
+					{
+						ShaderType->BeginCompileShader(
+							CompilingId,
+							InKernel,
+							InCompilationEnvironment,
+							InPlatform,
+							NewJobs,
+							FShaderTarget(ShaderType->GetFrequency(), GetShaderPlatform())
+							);
+					}
+					NumShaders++;
+				}
+				else if (ShaderType)
+				{
+					UE_LOG(LogShaders, Display, TEXT("Skipping compilation of %s as it isn't supported on this target type."), *InKernel->GetFriendlyName());
+					InKernel->RemoveOutstandingCompileId(CompilingId);
+					InKernel->NotifyCompilationFinished();
+				}
+			}
+  
+			if (!Kernel)
+			{
+				UE_LOG(LogShaders, Log, TEXT("		%u Shaders"), NumShaders);
+			}
+
+			// Register this shader map in the global ComputeKernel->shadermap map
+			Register(InPlatform);
+  
+			// Mark the shader map as not having been finalized with ProcessCompilationResults
+			bCompilationFinalized = false;
+  
+			// Mark as not having been compiled
+			bCompiledSuccessfully = false;
+  
+			GComputeKernelShaderCompilationManager.AddJobs(NewJobs);
+  
+			// Compile the shaders for this shader map now if not deferring and deferred compiles are not enabled globally
+			if (bSynchronousCompile)
+			{
+				TArray<int32> CurrentShaderMapId;
+				CurrentShaderMapId.Add(CompilingId);
+				GComputeKernelShaderCompilationManager.FinishCompilation(*NewContent->FriendlyName, CurrentShaderMapId);
+			}
+		}
+	}
+}
+
+FShader* FComputeKernelShaderMap::ProcessCompilationResultsForSingleJob(FShaderCompileJob& CurrentJob, const FSHAHash& InShaderMapHash)
+{
+	check(CurrentJob.Id == CompilingId);
+
+	GetResourceCode()->AddShaderCompilerOutput(CurrentJob.Output);
+
+	FShader* Shader = nullptr;
+
+	const FComputeKernelShaderType* ComputeKernelShaderType = CurrentJob.Key.ShaderType->GetComputeKernelShaderType();
+	check(ComputeKernelShaderType);
+	Shader = ComputeKernelShaderType->FinishCompileShader(InShaderMapHash, CurrentJob, GetContent()->FriendlyName);
+	bCompiledSuccessfully = CurrentJob.bSucceeded;
+
+	FComputeKernelShader* ComputeKernelShader = static_cast<FComputeKernelShader*>(Shader);
+	check(Shader);
+	check(!GetContent()->HasShader(ComputeKernelShaderType, /* PermutationId = */ 0));
+	return GetMutableContent()->FindOrAddShader(ComputeKernelShaderType->GetHashedName(), 0, Shader);
+}
+
+bool FComputeKernelShaderMap::ProcessCompilationResults(const TArray<FShaderCommonCompileJobPtr>& InCompilationResults, int32& InOutJobIndex, float& InOutTimeBudget)
+{
+	check(InOutJobIndex < InCompilationResults.Num());
+
+	double StartTime = FPlatformTime::Seconds();
+
+	FSHAHash ShaderMapHash;
+	GetContent()->ShaderMapId.GetComputeKernelHash(ShaderMapHash);
+
+	do
 	{
-		// if this shader map is already doing an async, cannot do sync compile at this time.
-		check(!bSynchronousCompile);
+		ProcessCompilationResultsForSingleJob(static_cast<FShaderCompileJob&>(*InCompilationResults[InOutJobIndex].GetReference()), ShaderMapHash);
 
-		CorrespondingKernels->AddUnique(KernelResource);
-		return;
+		InOutJobIndex++;
+		
+		double NewStartTime = FPlatformTime::Seconds();
+		InOutTimeBudget -= NewStartTime - StartTime;
+		StartTime = NewStartTime;
 	}
+	while ((InOutTimeBudget > 0.0f) && (InOutJobIndex < InCompilationResults.Num()));
+
+	if (InOutJobIndex == InCompilationResults.Num())
+	{
+		FinalizeContent();
+
+		SaveToDerivedDataCache();
+		// The shader map can now be used on the rendering thread
+		bCompilationFinalized = true;
+		return true;
+	}
+
+	return false;
+}
+
+bool FComputeKernelShaderMap::TryToAddToExistingCompilationTask(FComputeKernelResource* InKernel)
+{
+	check(NumRefs > 0);
+	TArray<FComputeKernelResource*>* CorrespondingKernels = FComputeKernelShaderMap::ComputeKernelShaderMapsBeingCompiled.Find(this);
+
+	if (CorrespondingKernels)
+	{
+		CorrespondingKernels->AddUnique(InKernel);
 
-	ComputeKernelShaderMapsBeingCompiled.Add(this, { KernelResource });
+		UE_LOG(LogShaders, Log, TEXT("TryToAddToExistingCompilationTask %p %d"), InKernel, GetCompilingId());
 
 #if DEBUG_INFINITESHADERCOMPILE
-	UE_LOG(ComputeKernel, Display, TEXT("Added FComputeKernelShaderMap 0x%08X%08X with FComputeKernelResource 0x%08X%08X to ComputeKernelShaderMapsBeingCompiled"), (int)((int64)(this) >> 32), (int)((int64)(this)), (int)((int64)(KernelShader) >> 32), (int)((int64)(KernelShader)));
-#endif  
+		UE_LOG(LogTemp, Display, TEXT("Added shader map 0x%08X%08X from ComputeKernel 0x%08X%08X"), (int)((int64)(this) >> 32), (int)((int64)(this)), (int)((int64)(InKernel) >> 32), (int)((int64)(InKernel)));
+#endif
+		return true;
+	}
 
-	FComputeKernelShaderMapContent* NewContent = new FComputeKernelShaderMapContent(ShaderPlatform);
-#if WITH_EDITORONLY_DATA
-	NewContent->FriendlyName = KernelResource->GetFriendlyName();
+	return false;
+}
+
+bool FComputeKernelShaderMap::IsComputeKernelShaderComplete(const FComputeKernelResource* InKernel, const FComputeKernelShaderType* InShaderType, bool bSilent)
+{
+	// If we should cache this kernel, it's incomplete if the shader is missing
+	if (ShouldCacheComputeKernelShader(InShaderType, GetShaderPlatform(), InKernel) && !GetContent()->HasShader((FShaderType*)InShaderType, /* PermutationId = */ 0))
+	{
+		if (!bSilent)
+		{
+			UE_LOG(LogShaders, Warning, TEXT("Incomplete shader %s, missing FComputeKernelShader %s."), *InKernel->GetFriendlyName(), InShaderType->GetName());
+		}
+		return false;
+	}
+
+	return true;
+}
+
+bool FComputeKernelShaderMap::IsComplete(const FComputeKernelResource* InKernel, bool bSilent)
+{
+	// Make sure we are operating on a referenced shader map or the below Find will cause this shader map to be deleted,
+	// Since it creates a temporary ref counted pointer.
+	check(NumRefs > 0);
+	const TArray<FComputeKernelResource*>* CorrespondingKernels = FComputeKernelShaderMap::ComputeKernelShaderMapsBeingCompiled.Find(this);
+
+	if (CorrespondingKernels)
+	{
+		check(!bCompilationFinalized);
+		return false;
+	}
+
+	// Iterate over all shader types.
+	for(TLinkedList<FShaderType*>::TIterator ShaderTypeIt(FShaderType::GetTypeList());ShaderTypeIt;ShaderTypeIt.Next())
+	{
+		// Find this shader type in the kernel's shader map.
+		const FComputeKernelShaderType* ShaderType = ShaderTypeIt->GetComputeKernelShaderType();
+		if (ShaderType && !IsComputeKernelShaderComplete(InKernel, ShaderType, bSilent))
+		{
+			return false;
+		}
+	}
+
+	return true;
+}
+
+void FComputeKernelShaderMap::LoadMissingShadersFromMemory(const FComputeKernelResource* InKernel)
+{
+#if 0
+	// Make sure we are operating on a referenced shader map or the below Find will cause this shader map to be deleted,
+	// Since it creates a temporary ref counted pointer.
+	check(NumRefs > 0);
+
+	const TArray<FComputeKernelResource*>* CorrespondingKernels = FComputeKernelShaderMap::ComputeKernelShaderMapsBeingCompiled.Find(this);
+
+	if (CorrespondingKernels)
+	{
+		check(!bCompilationFinalized);
+		return;
+	}
+
+	FSHAHash ShaderMapHash;
+	ShaderMapId.GetComputeKernelHash(ShaderMapHash);
+
+	// Try to find necessary FComputeKernelShaderType's in memory
+	for (TLinkedList<FShaderType*>::TIterator ShaderTypeIt(FShaderType::GetTypeList());ShaderTypeIt;ShaderTypeIt.Next())
+	{
+		FComputeKernelShaderType* ShaderType = ShaderTypeIt->GetComputeKernelShaderType();
+		if (ShaderType && ShouldCacheComputeKernelShader(ShaderType, Platform, InKernel) && !HasShader(ShaderType, /* PermutationId = */ 0))
+		{
+			FShaderKey ShaderKey(ShaderMapHash, nullptr, nullptr, /** PermutationId = */ 0, Platform);
+			FShader* FoundShader = ShaderType->FindShaderByKey(ShaderKey);
+			if (FoundShader)
+			{
+				AddShader(ShaderType, /* PermutationId = */ 0, FoundShader);
+			}
+		}
+	}
 #endif
-	NewContent->ShaderMapId = ShaderMapId;
-	AssignContent(NewContent);
-	
-	TArray<TRefCountPtr<FShaderCommonCompileJob>> NewJobs;
-	//TMap<TShaderTypePermutation<const FShaderType>, TSharedRef<FShaderCommonCompileJob, ESPMode::ThreadSafe>> SharedShaderJobs;
+}
 
-	const TArray<FShaderType*>& ComputeKernelShaderTypeCollection = FShaderType::GetSortedTypes(FShaderType::EShaderTypeForDynamicCast::ComputeKernel);
-	for (auto* ShaderType : ComputeKernelShaderTypeCollection)
+void FComputeKernelShaderMap::GetShaderList(TMap<FShaderId, TShaderRef<FShader>>& OutShaders) const
+{
+	GetContent()->GetShaderList(*this, FSHAHash(), OutShaders);
+}
+
+/**
+ * Registers a ComputeKernel shader map in the global map.
+ */
+void FComputeKernelShaderMap::Register(EShaderPlatform InShaderPlatform)
+{
+	if (!bRegistered)
 	{
-		FComputeKernelShaderType* ComputeKernelShaderType = ShaderType->GetComputeKernelShaderType();
-		if (!ComputeKernelShaderType || !ShouldCacheComputeKernelShader(ShaderPlatform, ComputeKernelShaderType, KernelResource) || NewContent->HasShader(ShaderType, KernelResource->GetPermutationId()))
+		INC_DWORD_STAT(STAT_Shaders_NumShaderMaps);
+	}
+
+	GIdToComputeKernelShaderMap[GetShaderPlatform()].Add(GetContent()->ShaderMapId,this);
+	bRegistered = true;
+}
+
+void FComputeKernelShaderMap::AddRef()
+{
+	check(!bDeletedThroughDeferredCleanup);
+	++NumRefs;
+}
+
+void FComputeKernelShaderMap::Release()
+{
+	check(NumRefs > 0);
+	if(--NumRefs == 0)
+	{
+		if (bRegistered)
 		{
-			continue;
+			DEC_DWORD_STAT(STAT_Shaders_NumShaderMaps);
+
+			GIdToComputeKernelShaderMap[GetShaderPlatform()].Remove(GetContent()->ShaderMapId);
+			bRegistered = false;
 		}
 
-		CompilationRequestId = FShaderCommonCompileJob::GetNextJobId();
-		ComputeKernelShaderType->BeginCompileShader(
-			CompilationRequestId,
-			ShaderPlatform,
-			KernelResource,
-			NewJobs
-			);
+		check(!bDeletedThroughDeferredCleanup);
+		bDeletedThroughDeferredCleanup = true;
+		BeginCleanup(this);
+	}
+}
+
+FComputeKernelShaderMap::FComputeKernelShaderMap() :
+	CompilingId(1),
+	NumRefs(0),
+	bDeletedThroughDeferredCleanup(false),
+	bRegistered(false),
+	bCompilationFinalized(true),
+	bCompiledSuccessfully(true),
+	bIsPersistent(true) 
+{
+	checkSlow(IsInGameThread() || IsAsyncLoading());
+	AllComputeKernelShaderMaps.Add(this);
+}
+
+FComputeKernelShaderMap::~FComputeKernelShaderMap()
+{ 
+	checkSlow(IsInGameThread() || IsAsyncLoading());
+	check(bDeletedThroughDeferredCleanup);
+	check(!bRegistered);
+	AllComputeKernelShaderMaps.RemoveSwap(this);
+}
 
-		TShaderTypePermutation<const FShaderType> ShaderTypePermutation(ComputeKernelShaderType, KernelResource->GetPermutationId());
-		//check(!SharedShaderJobs.Find(ShaderTypePermutation));
-		//SharedShaderJobs.Add(ShaderTypePermutation, Job);
+/**
+ * Removes all entries in the cache with exceptions based on a shader type
+ * @param ShaderType - The shader type to flush
+ */
+void FComputeKernelShaderMap::FlushShadersByShaderType(const FShaderType* InShaderType)
+{
+	if (InShaderType->GetComputeKernelShaderType())
+	{
+		GetMutableContent()->RemoveShaderTypePermutaion(InShaderType->GetComputeKernelShaderType(), /* PermutationId = */ 0);	
 	}
+}
 
-	/**
-	Register(ShaderPlatform);
 
-	GComputeKernelShaderCompilationManager.AddJobs(NewJobs);
+bool FComputeKernelShaderMap::Serialize(FArchive& Ar, bool bInlineShaderResources)
+{
+	// Note: This is saved to the DDC, not into packages (except when cooked)
+	// Backwards compatibility therefore will not work based on the version of Ar
+	// Instead, just bump COMPUTEKERNEL_DERIVEDDATA_VER
+	return Super::Serialize(Ar, bInlineShaderResources, false);
+}
 
-	if (bSynchronousCompile)
+void FComputeKernelShaderMap::RemovePending(FComputeKernelResource* InKernel)
+{
+	for (TMap<TRefCountPtr<FComputeKernelShaderMap>, TArray<FComputeKernelResource*> >::TIterator It(ComputeKernelShaderMapsBeingCompiled); It; ++It)
 	{
-		TArray<int32> CurrentShaderMapId;
-		CurrentShaderMapId.Add(CompilationRequestId);
-		GComputeKernelShaderCompilationManager->FinishCompilation(
-			NewContent->FriendlyName,
-			CurrentShaderMapId
-			);
+		TArray<FComputeKernelResource*>& Kernels = It.Value();
+		int32 Result = Kernels.Remove(InKernel);
+		if (Result)
+		{
+			InKernel->RemoveOutstandingCompileId(It.Key()->CompilingId);
+			InKernel->NotifyCompilationFinished();
+		}
+#if DEBUG_INFINITESHADERCOMPILE
+		if ( Result )
+		{
+			UE_LOG(LogTemp, Display, TEXT("Removed shader map 0x%08X%08X from kernel 0x%08X%08X"), (int)((int64)(It.Key().GetReference()) >> 32), (int)((int64)(It.Key().GetReference())), (int)((int64)(InKernel) >> 32), (int)((int64)(InKernel)));
+		}
+#endif
 	}
-	**/
 }
+
diff --git a/Engine/Source/Runtime/Engine/Private/ComputeFramework/ComputeKernelShaderType.cpp b/Engine/Source/Runtime/Engine/Private/ComputeFramework/ComputeKernelShaderType.cpp
deleted file mode 100644
diff --git a/Engine/Source/Runtime/Engine/Private/ComputeFramework/ComputeKernelShared.cpp b/Engine/Source/Runtime/Engine/Private/ComputeFramework/ComputeKernelShared.cpp
new file mode 100644
--- /dev/null
+++ b/Engine/Source/Runtime/Engine/Private/ComputeFramework/ComputeKernelShared.cpp
@@ -0,0 +1,508 @@
+// Copyright Epic Games, Inc. All Rights Reserved.
+
+#include "ComputeFramework/ComputeKernelShared.h"
+
+#include "Misc/App.h"
+#include "Modules/ModuleManager.h"
+#include "ComputeFramework/ComputeKernelShaderType.h"
+#include "ComputeFramework/ComputeKernelShader.h"
+#include "ComputeFramework/ComputeKernelShaderCompilationManager.h"
+#include "ComputeFramework/ComputeKernelSource.h"
+#include "RendererInterface.h"
+#include "ShaderCompiler.h"
+#include "Stats/StatsMisc.h"
+#include "TextureResource.h"
+#include "UObject/CoreObjectVersion.h"
+#include "UObject/UObjectHash.h"
+#include "UObject/UObjectIterator.h"
+#include "Interfaces/ITargetPlatform.h"
+#include "ShaderParameterMetadataBuilder.h"
+
+IMPLEMENT_TYPE_LAYOUT(FComputeKernelCompilationOutput);
+IMPLEMENT_TYPE_LAYOUT(FComputeKernelShaderMapId);
+IMPLEMENT_TYPE_LAYOUT(FComputeKernelShaderMapContent);
+
+FComputeKernelResource::~FComputeKernelResource()
+{
+	FComputeKernelShaderMap::RemovePending(this);
+}
+
+/** Populates OutEnvironment with defines needed to compile shaders for this kernel. */
+void FComputeKernelResource::SetupShaderCompilationEnvironment(EShaderPlatform InPlatform, FShaderCompilerEnvironment& OutEnvironment) const
+{
+}
+
+
+bool FComputeKernelResource::ShouldCache(EShaderPlatform InPlatform, const FShaderType* InShaderType) const
+{
+	check(InShaderType->GetComputeKernelShaderType() )
+	return true;
+}
+
+void FComputeKernelResource::NotifyCompilationFinished()
+{
+}
+
+void FComputeKernelResource::CancelCompilation()
+{
+#if WITH_EDITOR
+	if (IsInGameThread())
+	{
+		FComputeKernelShaderMap::RemovePending(this);
+
+		UE_LOG(LogShaders, Log, TEXT("CancelCompilation %p."), this);
+		OutstandingCompileShaderMapIds.Empty();
+	}
+#endif
+}
+
+void FComputeKernelResource::RemoveOutstandingCompileId(const int32 InOldOutstandingCompileShaderMapId)
+{
+	if (0 <= OutstandingCompileShaderMapIds.Remove(InOldOutstandingCompileShaderMapId))
+	{
+		UE_LOG(LogShaders, Log, TEXT("RemoveOutstandingCompileId %p %d"), this, InOldOutstandingCompileShaderMapId);
+	}
+}
+
+void FComputeKernelResource::Invalidate()
+{
+	CancelCompilation();
+	ReleaseShaderMap();
+}
+
+bool FComputeKernelResource::IsSame(const FComputeKernelShaderMapId& InIdentifier) const
+{
+	return
+		InIdentifier.ShaderCodeHash == ShaderCodeHash &&
+		InIdentifier.FeatureLevel == FeatureLevel;
+}
+
+void FComputeKernelResource::GetDependentShaderTypes(EShaderPlatform InPlatform, TArray<FShaderType*>& OutShaderTypes) const
+{
+	for (TLinkedList<FShaderType*>::TIterator ShaderTypeIt(FShaderType::GetTypeList()); ShaderTypeIt; ShaderTypeIt.Next())
+	{
+		FComputeKernelShaderType* ShaderType = ShaderTypeIt->GetComputeKernelShaderType();
+
+		if ( ShaderType && ShaderType->ShouldCache(InPlatform, this) && ShouldCache(InPlatform, ShaderType) )
+		{
+			OutShaderTypes.Add(ShaderType);
+		}
+	}
+
+	OutShaderTypes.Sort(FCompareShaderTypes());
+}
+
+void FComputeKernelResource::GetShaderMapId(EShaderPlatform InPlatform, const ITargetPlatform* TargetPlatform, FComputeKernelShaderMapId& OutId) const
+{
+	if (bLoadedCookedShaderMapId)
+	{
+		OutId = CookedShaderMapId;
+	}
+	else
+	{
+		TArray<FShaderType*> ShaderTypes;
+		GetDependentShaderTypes(InPlatform, ShaderTypes);
+
+		OutId.FeatureLevel = GetFeatureLevel();
+		OutId.ShaderCodeHash = ShaderCodeHash;
+		OutId.SetShaderDependencies(ShaderTypes, InPlatform);
+#if WITH_EDITOR
+		if (TargetPlatform)
+		{
+			OutId.LayoutParams.InitializeForPlatform(TargetPlatform->IniPlatformName(), TargetPlatform->HasEditorOnlyData());
+		}
+		else
+		{
+			OutId.LayoutParams.InitializeForCurrent();
+		}
+#else
+		if (TargetPlatform != nullptr)
+		{
+			UE_LOG(LogShaders, Error, TEXT("FComputeKernelResource::GetShaderMapId: TargetPlatform is not null, but a cooked executable cannot target platforms other than its own."));
+		}
+		OutId.LayoutParams.InitializeForCurrent();
+#endif
+	}
+}
+
+void FComputeKernelResource::ReleaseShaderMap()
+{
+	if (GameThreadShaderMap)
+	{
+		GameThreadShaderMap = nullptr;
+
+		FComputeKernelResource* Kernel = this;
+		ENQUEUE_RENDER_COMMAND(ReleaseShaderMap)(
+			[Kernel](FRHICommandListImmediate& RHICmdList)
+			{
+				Kernel->SetRenderingThreadShaderMap(nullptr);
+			});
+	}
+}
+
+void FComputeKernelResource::DiscardShaderMap()
+{
+	check(RenderingThreadShaderMap == nullptr);
+	GameThreadShaderMap = nullptr;
+}
+
+void FComputeKernelResource::SerializeShaderMap(FArchive& Ar)
+{
+	bool bCooked = Ar.IsCooking();
+	Ar << bCooked;
+
+	if (FPlatformProperties::RequiresCookedData() && !bCooked && Ar.IsLoading())
+	{
+		UE_LOG(LogShaders, Fatal, TEXT("This platform requires cooked packages, and shaders were not cooked into this kernel %s."), *GetFriendlyName());
+	}
+
+	if (bCooked)
+	{
+		if (Ar.IsCooking())
+		{
+#if WITH_EDITOR
+			FinishCompilation();
+
+			bool bValid = GameThreadShaderMap != nullptr && GameThreadShaderMap->CompiledSuccessfully();
+			Ar << bValid;
+
+			if (bValid)
+			{
+				GameThreadShaderMap->Serialize(Ar);
+			}
+#endif
+		}
+		else
+		{
+			bool bValid = false;
+			Ar << bValid;
+
+			if (bValid)
+			{
+				TRefCountPtr<FComputeKernelShaderMap> LoadedShaderMap = new FComputeKernelShaderMap();
+				bool bSuccessfullyLoaded = LoadedShaderMap->Serialize(Ar);
+
+				// Toss the loaded shader data if this is a server only instance (@todo - don't cook it in the first place) or if it's for a different RHI than the current one
+				if (bSuccessfullyLoaded && FApp::CanEverRender())
+				{
+					GameThreadShaderMap = RenderingThreadShaderMap = LoadedShaderMap;
+				}
+			}
+		}
+	}
+}
+
+void FComputeKernelResource::SetupResource(
+	ERHIFeatureLevel::Type InFeatureLevel,
+	const UComputeKernelSource* Source,
+	FString InFriendlyName
+	)
+{
+	FeatureLevel = InFeatureLevel;
+	ShaderCodeHash = Source->GetSourceHashCode();
+	ShaderEntryPoint = Source->GetEntryPoint();
+	ShaderSource = Source->GetSource();
+	FriendlyName = MoveTemp(InFriendlyName);
+
+	FShaderParametersMetadataBuilder Builder;
+
+	for (auto& Input : Source->InputParams)
+	{
+		ensureAlways(Input.DimType == EShaderFundamentalDimensionType::Scalar);
+
+		switch (Input.FundamentalType)
+		{
+			case EShaderFundamentalType::Bool:
+				Builder.AddParam<bool>(*Input.Name);
+				break;
+
+			case EShaderFundamentalType::Int:
+				Builder.AddParam<int32>(*Input.Name);
+				break;
+
+			case EShaderFundamentalType::Uint:
+				Builder.AddParam<uint32>(*Input.Name);
+				break;
+
+			case EShaderFundamentalType::Float:
+				Builder.AddParam<float>(*Input.Name);
+				break;
+		}
+	}
+
+	for(auto& Input : Source->InputSRVs)
+	{
+		Builder.AddRDGBufferSRV(*Input.Name, *Input.TypeDeclaration);
+	}
+
+	for (auto& Output : Source->Outputs)
+	{
+		Builder.AddRDGBufferUAV(*Output.Name, *Output.TypeDeclaration);
+	}
+
+	ShaderMetadata.Reset(
+		Builder.Build(
+			FShaderParametersMetadata::EUseCase::ShaderParameterStruct, 
+			*FriendlyName
+			)
+		);
+}
+
+void FComputeKernelResource::SetRenderingThreadShaderMap(FComputeKernelShaderMap* InShaderMap)
+{
+	check(IsInRenderingThread());
+	RenderingThreadShaderMap = InShaderMap;
+}
+
+bool FComputeKernelResource::IsCompilationFinished() const
+{
+	bool bRet = GameThreadShaderMap && GameThreadShaderMap.IsValid() && GameThreadShaderMap->IsCompilationFinalized();
+	if (OutstandingCompileShaderMapIds.Num() == 0)
+	{
+		return true;
+	}
+	return bRet;
+}
+
+bool FComputeKernelResource::CacheShaders(EShaderPlatform InPlatform, const ITargetPlatform* TargetPlatform, bool bApplyCompletedShaderMapForRendering, bool bSynchronous)
+{
+	FComputeKernelShaderMapId ShaderMapId;
+	GetShaderMapId(InPlatform, TargetPlatform, ShaderMapId);
+	return CacheShaders(ShaderMapId, InPlatform, bApplyCompletedShaderMapForRendering, bSynchronous);
+}
+
+bool FComputeKernelResource::CacheShaders(const FComputeKernelShaderMapId& InShaderMapId, EShaderPlatform InPlatform, bool bApplyCompletedShaderMapForRendering, bool bSynchronous)
+{
+	bool bSucceeded = false;
+	{
+		// If we loaded this material with inline shaders, use what was loaded (GameThreadShaderMap) instead of looking in the DDC
+		if (bContainsInlineShaders)
+		{
+			FComputeKernelShaderMap* ExistingShaderMap = nullptr;
+
+			if (GameThreadShaderMap)
+			{
+				// Note: in the case of an inlined shader map, the shadermap Id will not be valid because we stripped some editor-only data needed to create it
+				// Get the shadermap Id from the shadermap that was inlined into the package, if it exists
+				ExistingShaderMap = FComputeKernelShaderMap::FindId(GameThreadShaderMap->GetShaderMapId(), InPlatform);
+			}
+
+			// Re-use an identical shader map in memory if possible, removing the reference to the inlined shader map
+			if (ExistingShaderMap)
+			{
+				GameThreadShaderMap = ExistingShaderMap;
+			}
+			else if (GameThreadShaderMap)
+			{
+				// We are going to use the inlined shader map, register it so it can be re-used by other materials
+				GameThreadShaderMap->Register(InPlatform);
+			}
+		}
+		else
+		{
+			// Find the kernel's cached shader map.
+			GameThreadShaderMap = FComputeKernelShaderMap::FindId(InShaderMapId, InPlatform);
+
+			// Attempt to load from the derived data cache if we are uncooked
+			if ((!GameThreadShaderMap || !GameThreadShaderMap->IsComplete(this, true)) && !FPlatformProperties::RequiresCookedData())
+			{
+				FComputeKernelShaderMap::LoadFromDerivedDataCache(this, InShaderMapId, InPlatform, GameThreadShaderMap);
+				if (GameThreadShaderMap && GameThreadShaderMap->IsValid())
+				{
+					UE_LOG(LogTemp, Display, TEXT("Loaded shader %s for kernel %s from DDC"), *GameThreadShaderMap->GetFriendlyName(), *GetFriendlyName());
+				}
+				else
+				{
+					UE_LOG(LogTemp, Display, TEXT("Loading shader for kernel %s from DDC failed. Shader needs recompile."), *GetFriendlyName());
+				}
+			}
+		}
+	}
+
+	bool bAssumeShaderMapIsComplete = false;
+#if UE_BUILD_SHIPPING || UE_BUILD_TEST
+	bAssumeShaderMapIsComplete = (bContainsInlineShaders || FPlatformProperties::RequiresCookedData());
+#endif
+
+	if (GameThreadShaderMap && GameThreadShaderMap->TryToAddToExistingCompilationTask(this))
+	{
+#if DEBUG_INFINITESHADERCOMPILE
+		UE_LOG(LogTemp, Display, TEXT("Found existing compiling shader for kernel %s, linking to other GameThreadShaderMap 0x%08X%08X"), *GetFriendlyName(), (int)((int64)(GameThreadShaderMap.GetReference()) >> 32), (int)((int64)(GameThreadShaderMap.GetReference())));
+#endif
+		OutstandingCompileShaderMapIds.AddUnique(GameThreadShaderMap->GetCompilingId());
+		UE_LOG(LogShaders, Log, TEXT("CacheShaders AddUniqueExisting %p %d"), this, GameThreadShaderMap->GetCompilingId());
+
+		GameThreadShaderMap = nullptr;
+		bSucceeded = true;
+	}
+	else if (!GameThreadShaderMap || !(bAssumeShaderMapIsComplete || GameThreadShaderMap->IsComplete(this, false)))
+	{
+		if (bContainsInlineShaders || FPlatformProperties::RequiresCookedData())
+		{
+			UE_LOG(LogShaders, Log, TEXT("Can't compile %s with cooked content!"), *GetFriendlyName());
+			GameThreadShaderMap = nullptr;
+		}
+		else
+		{
+			UE_LOG(LogShaders, Log, TEXT("%s cached shader map for kernel %s, compiling."), GameThreadShaderMap ? TEXT("Incomplete") : TEXT("Missing"), *GetFriendlyName());
+
+			// If there's no cached shader map for this kernel compile a new one.
+			// This is just kicking off the compile, GameThreadShaderMap will not be complete yet
+			bSucceeded = BeginCompileShaderMap(InShaderMapId, InPlatform, GameThreadShaderMap, bApplyCompletedShaderMapForRendering, bSynchronous);
+
+			if (!bSucceeded)
+			{
+				GameThreadShaderMap = nullptr;
+			}
+		}
+	}
+	else
+	{
+		bSucceeded = true;
+	}
+
+	FComputeKernelResource* Kernel = this;
+	FComputeKernelShaderMap* LoadedShaderMap = GameThreadShaderMap;
+	ENQUEUE_RENDER_COMMAND(FSetShaderMapOnComputeKernel)(
+		[Kernel, LoadedShaderMap](FRHICommandListImmediate& RHICmdList)
+		{
+			Kernel->SetRenderingThreadShaderMap(LoadedShaderMap);
+		});
+
+	return bSucceeded;
+}
+
+void FComputeKernelResource::FinishCompilation()
+{
+#if WITH_EDITOR
+	TArray<int32> ShaderMapIdsToFinish;
+	GetShaderMapIDsWithUnfinishedCompilation(ShaderMapIdsToFinish);
+
+	if (ShaderMapIdsToFinish.Num() > 0)
+	{
+		for (int32 i = 0; i < ShaderMapIdsToFinish.Num(); i++)
+		{
+			UE_LOG(LogShaders, Log, TEXT("FinishCompilation()[%d] %s id %d!"), i, *GetFriendlyName(), ShaderMapIdsToFinish[i]);
+		}
+	
+		// Block until the shader maps that we will save have finished being compiled
+		GComputeKernelShaderCompilationManager.FinishCompilation(*GetFriendlyName(), ShaderMapIdsToFinish);
+
+		// Shouldn't have anything left to do...
+		TArray<int32> ShaderMapIdsToFinish2;
+		GetShaderMapIDsWithUnfinishedCompilation(ShaderMapIdsToFinish2);
+		ensure(ShaderMapIdsToFinish2.Num() == 0);
+	}
+#endif
+}
+
+TShaderRef<FComputeKernelShader> FComputeKernelResource::GetShader() const
+{
+	check(!GIsThreadedRendering || !IsInGameThread());
+	if (!GIsEditor || RenderingThreadShaderMap)
+	{
+		return RenderingThreadShaderMap->GetShader<FComputeKernelShader>();
+	}
+	return TShaderRef<FComputeKernelShader>();
+};
+
+TShaderRef<FComputeKernelShader> FComputeKernelResource::GetShaderGameThread() const
+{
+	if (GameThreadShaderMap)
+	{
+		return GameThreadShaderMap->GetShader<FComputeKernelShader>();
+	}
+	return TShaderRef<FComputeKernelShader>();
+};
+
+void FComputeKernelResource::GetShaderMapIDsWithUnfinishedCompilation(TArray<int32>& OutShaderMapIds)
+{
+	// Build an array of the shader map Id's are not finished compiling.
+	if (GameThreadShaderMap && GameThreadShaderMap.IsValid() && !GameThreadShaderMap->IsCompilationFinalized())
+	{
+		OutShaderMapIds.Add(GameThreadShaderMap->GetCompilingId());
+	}
+	else if (OutstandingCompileShaderMapIds.Num() != 0)
+	{
+		OutShaderMapIds.Append(OutstandingCompileShaderMapIds);
+	}
+}
+
+/**
+ * Compiles this kernel for Platform, storing the result in OutShaderMap
+ *
+ * @param InShaderMapId - the set of static parameters to compile
+ * @param InPlatform - the platform to compile for
+ * @param OutShaderMap - the shader map to compile
+ * @return - true if compile succeeded or was not necessary (shader map for ShaderMapId was found and was complete)
+ */
+bool FComputeKernelResource::BeginCompileShaderMap(const FComputeKernelShaderMapId& InShaderMapId, EShaderPlatform InPlatform, TRefCountPtr<FComputeKernelShaderMap>& OutShaderMap, bool bApplyCompletedShaderMapForRendering, bool bSynchronous)
+{
+#if WITH_EDITORONLY_DATA
+	bool bSuccess = false;
+
+	STAT(double ComputeKernelCompileTime = 0);
+
+	SCOPE_SECONDS_COUNTER(ComputeKernelCompileTime);
+
+	TRefCountPtr<FComputeKernelShaderMap> NewShaderMap = new FComputeKernelShaderMap();
+
+	// Create a shader compiler environment for the material that will be shared by all jobs from this material
+	TRefCountPtr<FSharedShaderCompilerEnvironment> Environment = new FSharedShaderCompilerEnvironment();
+
+	// Compile the shaders for the kernel.
+	FComputeKernelCompilationOutput CompilationOutput;
+	NewShaderMap->Compile(this, InShaderMapId, Environment, CompilationOutput, InPlatform, bSynchronous, bApplyCompletedShaderMapForRendering);
+
+	if (bSynchronous)
+	{
+		// If this is a synchronous compile, assign the compile result to the output
+		OutShaderMap = NewShaderMap->CompiledSuccessfully() ? NewShaderMap : nullptr;
+	}
+	else
+	{
+		UE_LOG(LogShaders, Log, TEXT("BeginCompileShaderMap AddUnique %p %d"), this, NewShaderMap->GetCompilingId());
+		OutstandingCompileShaderMapIds.AddUnique(NewShaderMap->GetCompilingId());
+		
+		// Async compile, use nullptr to detect it if used
+		OutShaderMap = nullptr;
+	}
+
+	return true;
+#else
+	UE_LOG(LogShaders, Fatal, TEXT("Compiling of shaders in a build without editordata is not supported."));
+	return false;
+#endif
+}
+
+void FComputeKernelShaderMapId::SetShaderDependencies(const TArray<FShaderType*>& InShaderTypes, EShaderPlatform InShaderPlatform)
+{
+#if WITH_EDITOR
+	if (!FPlatformProperties::RequiresCookedData())
+	{
+		for (FShaderType* ShaderType : InShaderTypes)
+		{
+			if (ShaderType != nullptr)
+			{
+				FShaderTypeDependency Dependency;
+				Dependency.ShaderTypeName = ShaderType->GetHashedName();
+				Dependency.SourceHash = ShaderType->GetSourceHash(InShaderPlatform);
+				ShaderTypeDependencies.Add(Dependency);
+			}
+		}
+	}
+#endif //WITH_EDITOR
+}
+
+bool FComputeKernelShaderMapId::ContainsShaderType(const FShaderType* ShaderType) const
+{
+	for (int32 TypeIndex = 0; TypeIndex < ShaderTypeDependencies.Num(); TypeIndex++)
+	{
+		if (ShaderTypeDependencies[TypeIndex].ShaderTypeName == ShaderType->GetHashedName())
+		{
+			return true;
+		}
+	}
+
+	return false;
+}
diff --git a/Engine/Source/Runtime/Engine/Private/ComputeFramework/ShaderParamTypeDefinition.cpp b/Engine/Source/Runtime/Engine/Private/ComputeFramework/ShaderParamTypeDefinition.cpp
--- a/Engine/Source/Runtime/Engine/Private/ComputeFramework/ShaderParamTypeDefinition.cpp
+++ b/Engine/Source/Runtime/Engine/Private/ComputeFramework/ShaderParamTypeDefinition.cpp
@@ -4,6 +4,7 @@
 
 #include "Serialization/Archive.h"
 #include "Templates/TypeHash.h"
+#include "Misc/DefaultValueHelper.h"
 
 
 // Storage for shared shader value types.
@@ -239,9 +240,6 @@ FString FShaderValueType::ToString() const
 	case EShaderFundamentalType::Float:
 		BaseName = "float";
 		break;
-	case EShaderFundamentalType::Double:
-		BaseName = "double";
-		break;
 
 	case EShaderFundamentalType::Struct:
 		return Name.ToString();
@@ -331,3 +329,148 @@ FArchive& operator<<(FArchive& InArchive, FShaderValueType::FStructElement& InEl
 	InArchive << InElement.Type;
 	return InArchive;
 }
+
+
+using FFundamentalStingPair = TPair<EShaderFundamentalType, FString>;
+using FResourceStingPair = TPair<EShaderResourceType, FString>;
+
+static FFundamentalStingPair TypeStringMap[] = {
+	FFundamentalStingPair(EShaderFundamentalType::Bool,		TEXT("bool")),
+	FFundamentalStingPair(EShaderFundamentalType::Uint,		TEXT("uint")),
+	FFundamentalStingPair(EShaderFundamentalType::Int,		TEXT("int")),
+	FFundamentalStingPair(EShaderFundamentalType::Float,	TEXT("float")),
+	};
+
+static FResourceStingPair ResTypeStringMap[] = {
+	FResourceStingPair(EShaderResourceType::Texture1D,			TEXT("Texture1D")),
+	FResourceStingPair(EShaderResourceType::Texture2D,			TEXT("Texture2D")),
+	FResourceStingPair(EShaderResourceType::Texture3D,			TEXT("Texture3D")),
+	FResourceStingPair(EShaderResourceType::TextureCube,		TEXT("TextureCube")),
+	FResourceStingPair(EShaderResourceType::StructuredBuffer,	TEXT("StructuredBuffer")),
+	FResourceStingPair(EShaderResourceType::ByteAddressBuffer,	TEXT("ByteAddressBuffer")),
+	FResourceStingPair(EShaderResourceType::Buffer,				TEXT("Buffer")),
+	};
+
+EShaderFundamentalType FShaderParamTypeDefinition::ParseFundamental(
+	const FString& Str
+	)
+{
+	for (auto& Pair : TypeStringMap)
+	{
+		if (Str.Contains(Pair.Value))
+		{
+			return Pair.Key;
+		}
+	}
+
+	check(!"Unknown Type");
+	return EShaderFundamentalType::Float;
+}
+
+EShaderFundamentalDimensionType FShaderParamTypeDefinition::ParseDimension(
+	const FString& Str
+	)
+{
+	if (Str.Contains(TEXT("x"), ESearchCase::CaseSensitive))
+	{
+		return EShaderFundamentalDimensionType::Matrix;
+	}
+	else if (!Str.IsEmpty())
+	{
+		return EShaderFundamentalDimensionType::Vector;
+	}
+	else
+	{
+		return EShaderFundamentalDimensionType::Scalar;
+	}
+}
+
+uint8 FShaderParamTypeDefinition::ParseVectorDimension(
+	const FString& Str
+	)
+{
+	int32 Dim = 0;
+	FDefaultValueHelper::ParseInt(Str, Dim);
+
+	return uint8(Dim);
+}
+
+FIntVector2 FShaderParamTypeDefinition::ParseMatrixDimension(
+	const FString& Str
+	)
+{
+	FIntVector2 Dim;
+
+	Dim.X = ParseVectorDimension(Str.Left(1));
+	Dim.Y = ParseVectorDimension(Str.Right(1));
+
+	return Dim;
+}
+
+EShaderResourceType FShaderParamTypeDefinition::ParseResource(const FString& Str)
+{
+	for (auto& Pair : ResTypeStringMap)
+	{
+		if (Str.Contains(Pair.Value))
+		{
+			return Pair.Key;
+		}
+	}
+
+	check(!"Unknown Type");
+	return EShaderResourceType::Buffer;
+}
+
+void FShaderParamTypeDefinition::ResetTypeDeclaration(
+	)
+{
+	FString TypeDecl;
+	if (BindingType == EShaderParamBindingType::ReadWriteResource)
+	{
+		TypeDecl.Append(TEXT("RW"));
+	}
+
+	if (BindingType != EShaderParamBindingType::ConstantParameter)
+	{
+		auto* foundItem = Algo::FindByPredicate(
+			ResTypeStringMap, 
+			[this](const FResourceStingPair& Pair) 
+			{ 
+				return Pair.Key == ResourceType; 
+			});
+		check(foundItem);
+
+		TypeDecl.Appendf(TEXT("%s<"), *foundItem->Value);
+	}
+
+	auto* foundItem = Algo::FindByPredicate(
+		TypeStringMap, 
+		[this](const FFundamentalStingPair& Pair) 
+		{ 
+			return Pair.Key == FundamentalType; 
+		});
+	check(foundItem);
+
+	TypeDecl.Append(foundItem->Value);
+
+	switch (DimType)
+	{
+	case EShaderFundamentalDimensionType::Scalar:
+		break;
+
+	case EShaderFundamentalDimensionType::Vector:
+		TypeDecl.AppendInt(VectorDimension);
+		break;
+
+	case EShaderFundamentalDimensionType::Matrix:
+		TypeDecl.Appendf(TEXT("%dx%d"), MatrixRowCount, MatrixColumnCount);
+		break;
+	};
+
+	if (BindingType != EShaderParamBindingType::ConstantParameter)
+	{
+		TypeDecl.AppendChar(TEXT('>'));
+	}
+
+	TypeDeclaration = MoveTemp(TypeDecl);
+}
diff --git a/Engine/Source/Runtime/Engine/Public/ComputeFramework/ComputeFramework.h b/Engine/Source/Runtime/Engine/Public/ComputeFramework/ComputeFramework.h
--- a/Engine/Source/Runtime/Engine/Public/ComputeFramework/ComputeFramework.h
+++ b/Engine/Source/Runtime/Engine/Public/ComputeFramework/ComputeFramework.h
@@ -43,7 +43,7 @@ public:
 		FName KernelName;
 		FName InvocationName;
 		FIntVector DispatchDim;
-		FComputeKernelResource* KernelResource = nullptr;
+		FComputeKernelResource* Kernel = nullptr;
 	};
 
 	void Initialize(UComputeGraph* ComputeGraph);
@@ -59,6 +59,7 @@ public:
 		FName KernelName;
 		FName InvocationName;
 		FIntVector DispatchDim;
+		const FShaderParametersMetadata* ShaderParamMetadata = nullptr;
 		TShaderRef<FComputeKernelShader> Shader;
 	};
 
diff --git a/Engine/Source/Runtime/Engine/Public/ComputeFramework/ComputeKernel.h b/Engine/Source/Runtime/Engine/Public/ComputeFramework/ComputeKernel.h
--- a/Engine/Source/Runtime/Engine/Public/ComputeFramework/ComputeKernel.h
+++ b/Engine/Source/Runtime/Engine/Public/ComputeFramework/ComputeKernel.h
@@ -5,10 +5,11 @@
 #include "CoreMinimal.h"
 #include "ComputeFramework/ShaderParamTypeDefinition.h"
 #include "ComputeFramework/ComputeKernelPermutationSet.h"
-#include "ComputeFramework/ComputeKernelResource.h"
+#include "ComputeFramework/ComputeKernelShared.h"
 
 #include "ComputeKernel.generated.h"
 
+class FComputeKernelResource;
 class UComputeKernelSource;
 
 DECLARE_LOG_CATEGORY_EXTERN(ComputeKernel, Log, All);
@@ -58,6 +59,11 @@ enum class EComputeKernelCompilationFlags
 
 	/* Compile the shader while blocking the main thread. */
 	Synchronous = 1 << 1,
+
+	/* Replaces all instances of the shader with the newly compiled version. */
+	ApplyCompletedShaderMapForRendering = 1 << 2,
+
+	IsCooking = 1 << 3,
 };
 
 /* Base class representing a kernel that will be run as a shader on the GPU. */
@@ -79,10 +85,10 @@ public:
 	int32 KernelFlags = 0;
 
 	UPROPERTY(EditDefaultsOnly, meta = (ShowOnlyInnerProperties), Category = "Kernel")
-	FComputeKernelPermutationSet PermutationSet;
+	FComputeKernelPermutationSet PermutationSetOverrides;
 
 	UPROPERTY(EditDefaultsOnly, meta = (ShowOnlyInnerProperties), Category = "Kernel")
-	FComputeKernelDefinitionsSet DefinitionsSet;
+	FComputeKernelDefinitionsSet DefinitionsSetOverrides;
 
 	/*
 	 * The minimum number of invocations (threads) launched by the kernel definition.
@@ -91,14 +97,6 @@ public:
 	UPROPERTY(VisibleAnywhere, Category = "Kernel")
 	FKernelInvocationDimension GroupSizeDim;
 
-	/* Named input types to the kernel. These are discovered via reflection from the shader. */
-	UPROPERTY(VisibleAnywhere, EditFixedSize, Category = "Kernel")
-	TArray<FShaderParamTypeDefinition> Inputs;
-
-	/* Named output types to the kernel. These are discovered via reflection from the shader. */
-	UPROPERTY(VisibleAnywhere, EditFixedSize, Category = "Kernel")
-	TArray<FShaderParamTypeDefinition> Outputs;
-
 #if WITH_EDITOR
 	DECLARE_EVENT_OneParam(UComputeKernel, FShaderResetEvent, const UComputeKernel*);
 
@@ -119,8 +117,13 @@ public:
 
 	//=========================================================================
 	// UObject interface
+	void PostLoad(
+		) override;
+
 #if WITH_EDITOR
-	virtual void PostEditChangeChainProperty(FPropertyChangedChainEvent& PropertyChangedEvent) override;
+	void PostEditChangeChainProperty(
+		FPropertyChangedChainEvent& PropertyChangedEvent
+		) override;
 #endif
 	//=========================================================================
 
@@ -130,17 +133,18 @@ private:
 #if WITH_EDITOR
 	/* Entry point to initiate any [re]compilation necessary due to changes. */
 	void CacheResourceShadersForRendering(
-		uint32 CompilationFlags = uint32(EComputeKernelCompilationFlags::None)
+		uint32 CompilationFlags = uint32(EComputeKernelCompilationFlags::ApplyCompletedShaderMapForRendering) | uint32(EComputeKernelCompilationFlags::Force)
 		);
 
 	/* Given a FKernelResouce, its associated shader is compiled - used for rendering and cooking. */
 	static void CacheShadersForResource(
 		EShaderPlatform ShaderPlatform,
+		const ITargetPlatform* TargetPlatform,
 		uint32 CompilationFlags,
-		FComputeKernelResource* KernelResource
+		FComputeKernelResource* Kernel
 		);
 #endif
 
-	/* The shader resouce encapsulating the kernel. Akin to FMaterialResource */
+	/* The shader resource encapsulating the kernel. Akin to FMaterialResource */
 	TUniquePtr<FComputeKernelResource> KernelResource = nullptr;
 };
diff --git a/Engine/Source/Runtime/Engine/Public/ComputeFramework/ComputeKernelFromText.h b/Engine/Source/Runtime/Engine/Public/ComputeFramework/ComputeKernelFromText.h
--- a/Engine/Source/Runtime/Engine/Public/ComputeFramework/ComputeKernelFromText.h
+++ b/Engine/Source/Runtime/Engine/Public/ComputeFramework/ComputeKernelFromText.h
@@ -31,6 +31,27 @@ public:
 	UPROPERTY()
 	uint64 SourceHash = 0;
 
+#if WITH_EDITOR
+	FString KernelSourceText;
+
+	FString GetEntryPoint() const override
+	{
+		return EntryPointName;
+	}
+
+	FString GetSource() const override
+	{
+		return KernelSourceText;
+	}
+
+	uint64 GetSourceHashCode() const override
+	{
+		return SourceHash;
+	}
+#endif
+
+	void PostLoad() override;
+
 #if WITH_EDITOR
 	FFilePath PrevSourceFile;
 
diff --git a/Engine/Source/Runtime/Engine/Public/ComputeFramework/ComputeKernelResource.h b/Engine/Source/Runtime/Engine/Public/ComputeFramework/ComputeKernelResource.h
deleted file mode 100644
diff --git a/Engine/Source/Runtime/Engine/Public/ComputeFramework/ComputeKernelShader.h b/Engine/Source/Runtime/Engine/Public/ComputeFramework/ComputeKernelShader.h
--- a/Engine/Source/Runtime/Engine/Public/ComputeFramework/ComputeKernelShader.h
+++ b/Engine/Source/Runtime/Engine/Public/ComputeFramework/ComputeKernelShader.h
@@ -2,20 +2,27 @@
 
 #pragma once
 
+#include "GlobalShader.h"
+#include "ComputeKernelShaderType.h"
+#include "ComputeKernelShared.h"
+#include "SceneView.h"
 #include "Shader.h"
-#include "ShaderParameterStruct.h"
-#include "ComputeFramework/ComputeKernelShaderMap.h"
+#include "ShaderParameters.h"
 
-#include "GlobalShader.h"
+class FTextureResource;
+class UClass;
 
-class FComputeKernelShader : public FGlobalShader /*public FShader*/
+class FComputeKernelShader : public FShader
 {
+	DECLARE_SHADER_TYPE(FComputeKernelShader, ComputeKernel);
+
 public:
-	DECLARE_SHADER_TYPE(FComputeKernelShader, Global);
-	SHADER_USE_PARAMETER_STRUCT(FComputeKernelShader, FGlobalShader);
-	//DECLARE_SHADER_TYPE(FComputeKernelShader, ComputeKernel);
-	//SHADER_USE_PARAMETER_STRUCT(FComputeKernelShader, FShader);
+	FComputeKernelShader() = default;
+	FComputeKernelShader(
+		const FComputeKernelShaderType::CompiledShaderInitializerType & Initializer
+		);
 
-	BEGIN_SHADER_PARAMETER_STRUCT(FParameters, )
-	END_SHADER_PARAMETER_STRUCT()
+	MS_ALIGN(SHADER_PARAMETER_STRUCT_ALIGNMENT) struct FParameters
+	{
+	};
 };
diff --git a/Engine/Source/Runtime/Engine/Public/ComputeFramework/ComputeKernelShaderCompilationManager.h b/Engine/Source/Runtime/Engine/Public/ComputeFramework/ComputeKernelShaderCompilationManager.h
new file mode 100644
--- /dev/null
+++ b/Engine/Source/Runtime/Engine/Public/ComputeFramework/ComputeKernelShaderCompilationManager.h
@@ -0,0 +1,121 @@
+// Copyright Epic Games, Inc. All Rights Reserved.
+
+#pragma once
+
+#include "ComputeKernelShared.h"
+#include "ShaderCompiler.h"
+#include "TickableEditorObject.h"
+
+
+/** Information tracked for each shader compile worker process instance. */
+struct FComputeKernelShaderCompileWorkerInfo
+{
+	/** Process handle of the worker app once launched.  Invalid handle means no process. */
+	FProcHandle WorkerProcess;
+
+	/** Tracks whether tasks have been issued to the worker. */
+	bool bIssuedTasksToWorker;
+
+	/** Whether the worker has been launched for this set of tasks. */
+	bool bLaunchedWorker;
+
+	/** Tracks whether all tasks issued to the worker have been received. */
+	bool bComplete;
+
+	/** Time at which the worker started the most recent batch of tasks. */
+	double StartTime;
+
+	/** Jobs that this worker is responsible for compiling. */
+	TArray<FShaderCommonCompileJobPtr> QueuedJobs;
+
+	FComputeKernelShaderCompileWorkerInfo() :
+		bIssuedTasksToWorker(false),
+		bLaunchedWorker(false),
+		bComplete(false),
+		StartTime(0)
+	{
+	}
+
+	// warning: not virtual
+	~FComputeKernelShaderCompileWorkerInfo()
+	{
+		if (WorkerProcess.IsValid())
+		{
+			FPlatformProcess::TerminateProc(WorkerProcess);
+			FPlatformProcess::CloseProc(WorkerProcess);
+		}
+	}
+};
+
+
+/** Results for a single compiled shader map. */
+struct FComputeKernelShaderMapCompileResults
+{
+	FComputeKernelShaderMapCompileResults() :
+		NumJobsQueued(0),
+		bAllJobsSucceeded(true),
+		bRecreateComponentRenderStateOnCompletion(false)
+	{}
+
+	int32 NumJobsQueued;
+	bool bAllJobsSucceeded;
+	bool bRecreateComponentRenderStateOnCompletion;
+	TArray<FShaderCommonCompileJobPtr> FinishedJobs;
+};
+
+
+/** Results for a single compiled and finalized shader map. */
+struct FComputeKernelShaderMapFinalizeResults : public FComputeKernelShaderMapCompileResults
+{
+	/** Tracks finalization progress on this shader map. */
+	int32 FinalizeJobIndex;
+
+	FComputeKernelShaderMapFinalizeResults(const FComputeKernelShaderMapCompileResults& InCompileResults) :
+		FComputeKernelShaderMapCompileResults(InCompileResults),
+		FinalizeJobIndex(0)
+	{}
+};
+
+
+// handles finished shader compile jobs, applying of the shaders to their config asset, and some error handling
+//
+class FComputeKernelShaderCompilationManager : public FTickableEditorObject
+{
+public:
+	FComputeKernelShaderCompilationManager();
+	~FComputeKernelShaderCompilationManager();
+
+	virtual ETickableTickType GetTickableTickType() const override
+	{
+		return ETickableTickType::Always;
+	}
+
+	virtual TStatId GetStatId() const
+	{
+		RETURN_QUICK_DECLARE_CYCLE_STAT(FComputeKernelShaderCompilationManager, STATGROUP_Tickables);
+	}
+
+	void Tick(float DeltaSeconds) override;
+	void AddJobs(TArray<FShaderCommonCompileJobPtr> InNewJobs);
+	void ProcessAsyncResults();
+
+	void FinishCompilation(const TCHAR* InKernelName, const TArray<int32>& ShaderMapIdsToFinishCompiling);
+
+private:
+	void ProcessCompiledComputeKernelShaderMaps(TMap<int32, FComputeKernelShaderMapFinalizeResults>& CompiledShaderMaps, float TimeBudget);
+	void RunCompileJobs();
+
+	void InitWorkerInfo();
+
+	TArray<FShaderCommonCompileJobPtr> JobQueue;
+
+	/** Map from shader map Id to the compile results for that map, used to gather compiled results. */
+	TMap<int32, FComputeKernelShaderMapCompileResults> ComputeKernelShaderMapJobs;
+
+	/** Map from shader map id to results being finalized.  Used to track shader finalizations over multiple frames. */
+	TMap<int32, FComputeKernelShaderMapFinalizeResults> PendingFinalizeComputeKernelShaderMaps;
+
+	TArray<struct FComputeKernelShaderCompileWorkerInfo*> WorkerInfos;
+};
+
+extern FComputeKernelShaderCompilationManager GComputeKernelShaderCompilationManager;
diff --git a/Engine/Source/Runtime/Engine/Public/ComputeFramework/ComputeKernelShaderMap.h b/Engine/Source/Runtime/Engine/Public/ComputeFramework/ComputeKernelShaderMap.h
deleted file mode 100644
diff --git a/Engine/Source/Runtime/Engine/Public/ComputeFramework/ComputeKernelShaderType.h b/Engine/Source/Runtime/Engine/Public/ComputeFramework/ComputeKernelShaderType.h
--- a/Engine/Source/Runtime/Engine/Public/ComputeFramework/ComputeKernelShaderType.h
+++ b/Engine/Source/Runtime/Engine/Public/ComputeFramework/ComputeKernelShaderType.h
@@ -2,33 +2,72 @@
 
 #pragma once
 
+#include "CoreMinimal.h"
+#include "Engine/EngineTypes.h"
+#include "GlobalShader.h"
 #include "Shader.h"
 #include "ShaderCompiler.h"
 
-struct FShaderCompilerEnvironment;
+class FShaderCommonCompileJob;
+class FShaderCompileJob;
 class FComputeKernelResource;
 
+
 struct FComputeKernelShaderPermutationParameters : public FShaderPermutationParameters
 {
-	const FComputeKernelResource* KernelShader;
-
-	FComputeKernelShaderPermutationParameters(EShaderPlatform InPlatform, const FComputeKernelResource* InKernelShader)
+	FComputeKernelShaderPermutationParameters(EShaderPlatform InPlatform)
 		: FShaderPermutationParameters(InPlatform)
-		, KernelShader(InKernelShader)
-	{}
+	{
+	}
 };
 
 class FComputeKernelShaderType : public FShaderType
 {
 public:
-	typedef FShader::CompiledShaderInitializerType CompiledShaderInitializerType;
+	struct FParameters : public FShaderType::FParameters
+	{
+		const FShaderParametersMetadata& ShaderParamMetadata;
+
+		FParameters(
+			const FShaderParametersMetadata& InShaderParamMetadata
+			)
+			: ShaderParamMetadata(InShaderParamMetadata)
+		{
+		}
+	};
+
+	struct CompiledShaderInitializerType : FGlobalShaderType::CompiledShaderInitializerType
+	{
+		const FString DebugDescription;
+
+		CompiledShaderInitializerType(
+			const FShaderType* InType,
+			const FParameters* InParameters,
+			int32 InPermutationId,
+			const FShaderCompilerOutput& CompilerOutput,
+			const FSHAHash& InComputeKernelShaderMapHash,
+			const FString& InDebugDescription
+			)
+			: FGlobalShaderType::CompiledShaderInitializerType(
+				InType,
+				InParameters,
+				InPermutationId,
+				CompilerOutput, 
+				InComputeKernelShaderMapHash,
+				nullptr,
+				nullptr
+				)
+			, DebugDescription(InDebugDescription)
+		{
+		}
+	};
 
 	FComputeKernelShaderType(
 		FTypeLayoutDesc& InTypeLayout,
 		const TCHAR* InName,
 		const TCHAR* InSourceFilename,
 		const TCHAR* InFunctionName,
-		uint32 InFrequency,					// ugly - ignored for Niagara shaders but needed for IMPLEMENT_SHADER_TYPE macro magic
+		uint32 InFrequency,					// ugly - ignored for but needed for IMPLEMENT_SHADER_TYPE macro magic
 		int32 InTotalPermutationCount,
 		ConstructSerializedType InConstructSerializedRef,
 		ConstructCompiledType InConstructCompiledRef,
@@ -37,9 +76,9 @@ public:
 		ValidateCompiledResultType InValidateCompiledResultRef,
 		uint32 InTypeSize,
 		const FShaderParametersMetadata* InRootParametersMetadata = nullptr
-		)
-		: FShaderType(
-			EShaderTypeForDynamicCast::ComputeKernel,
+		):
+		FShaderType(
+			EShaderTypeForDynamicCast::ComputeKernel, 
 			InTypeLayout, 
 			InName, 
 			InSourceFilename, 
@@ -59,19 +98,40 @@ public:
 
 	void BeginCompileShader(
 		uint32 ShaderMapId,
-		EShaderPlatform ShaderPlatform,
-		FComputeKernelResource* KernelShader,
-		TArray<FShaderCommonCompileJobPtr>& InOutNewJobs
+		const FComputeKernelResource* InKernel,
+		FSharedShaderCompilerEnvironment* InCompilationEnvironment,
+		EShaderPlatform Platform,
+		TArray<FShaderCommonCompileJobPtr>& InOutNewJobs,
+		FShaderTarget Target
 		);
 
-	bool ShouldCache(EShaderPlatform ShaderPlatform, const FComputeKernelResource* KernelShader) const
+	/**
+	 * Either creates a new instance of this type or returns an equivalent existing shader.
+	 */
+	FShader* FinishCompileShader(
+		const FSHAHash& InComputeKernelShaderMapHash,
+		const FShaderCompileJob& CurrentJob,
+		const FString& InDebugDescription
+		) const;
+
+	/**
+	 * Checks if the shader type should be cached.
+	 * @return True if this shader type should be cached.
+	 */
+	bool ShouldCache(EShaderPlatform InPlatform, const FComputeKernelResource* Kernel) const
 	{
-		return ShouldCompilePermutation(FComputeKernelShaderPermutationParameters(ShaderPlatform, KernelShader));
+		return ShouldCompilePermutation(FComputeKernelShaderPermutationParameters(InPlatform));
 	}
 
+
 protected:
-	void SetupCompileEnvironment(EShaderPlatform ShaderPlatform, const FComputeKernelResource* KernelShader, FShaderCompilerEnvironment& Environment) const
+	/**
+	 * Sets up the environment used to compile an instance of this shader type.
+	 * @param InPlatform - Platform to compile for.
+	 * @param OutEnvironment - The shader compile environment that the function modifies.
+	 */
+	void SetupCompileEnvironment(EShaderPlatform InPlatform, const FComputeKernelResource* Kernel, FShaderCompilerEnvironment& OutEnvironment) const
 	{
-		ModifyCompilationEnvironment(FComputeKernelShaderPermutationParameters(ShaderPlatform, KernelShader), Environment);
+		ModifyCompilationEnvironment(FComputeKernelShaderPermutationParameters(InPlatform), OutEnvironment);
 	}
 };
diff --git a/Engine/Source/Runtime/Engine/Public/ComputeFramework/ComputeKernelShared.h b/Engine/Source/Runtime/Engine/Public/ComputeFramework/ComputeKernelShared.h
new file mode 100644
--- /dev/null
+++ b/Engine/Source/Runtime/Engine/Public/ComputeFramework/ComputeKernelShared.h
@@ -0,0 +1,520 @@
+// Copyright Epic Games, Inc. All Rights Reserved.
+
+#pragma once
+
+#include "CoreMinimal.h"
+#include "Engine/EngineTypes.h"
+#include "Misc/Guid.h"
+#include "Misc/Optional.h"
+#include "Misc/SecureHash.h"
+#include "RenderResource.h"
+#include "RenderingThread.h"
+#include "RHI.h"
+#include "SceneTypes.h"
+#include "Shader.h"
+#include "StaticParameterSet.h"
+#include "Templates/RefCounting.h"
+#include "ShaderCompiler.h"
+
+
+class FComputeKernelResource;
+class FComputeKernelShaderMap;
+class FComputeKernelShader;
+class FComputeKernelShaderMapId;
+class UComputeKernelSource;
+
+
+/** Stores outputs from the  kernel compile that need to be saved. */
+class FComputeKernelCompilationOutput
+{
+	DECLARE_TYPE_LAYOUT(FComputeKernelCompilationOutput, NonVirtual);
+public:
+	FComputeKernelCompilationOutput()
+	{}
+};
+
+/** Contains all the information needed to uniquely identify a FComputeKernelShaderMapID. */
+class FComputeKernelShaderMapId
+{
+	DECLARE_TYPE_LAYOUT(FComputeKernelShaderMapId, NonVirtual);
+public:
+	/** Feature level that the shader map is going to be compiled for.  */
+	LAYOUT_FIELD(ERHIFeatureLevel::Type, FeatureLevel);
+
+	/**
+	 * The hash computed from the shader code generated by the library.
+	 */
+	LAYOUT_FIELD(uint64, ShaderCodeHash);
+	
+	/** Shader types of shaders that are inlined in this shader map in the DDC. */
+	LAYOUT_FIELD(TMemoryImageArray<FShaderTypeDependency>, ShaderTypeDependencies);
+
+	/*
+	* Type layout parameters of the memory image
+	*/
+	LAYOUT_FIELD(FPlatformTypeLayoutParameters, LayoutParams);
+
+public:
+	FComputeKernelShaderMapId()
+		: FeatureLevel(ERHIFeatureLevel::SM5)
+	{ }
+
+	~FComputeKernelShaderMapId() = default;
+
+	void SetShaderDependencies(const TArray<FShaderType*>& InShaderTypes, EShaderPlatform InShaderPlatform);
+
+	//void Serialize(FArchive& Ar);
+
+	friend uint32 GetTypeHash(const FComputeKernelShaderMapId& Ref)
+	{
+		return GetTypeHash(Ref.ShaderCodeHash);
+	}
+
+	SIZE_T GetSizeBytes() const
+	{
+		return sizeof(*this) + ShaderTypeDependencies.GetAllocatedSize();
+	}
+
+	/** Hashes the kernel specific part of this shader map Id. */
+	void GetComputeKernelHash(FSHAHash& OutHash) const;
+
+	/**
+	 * Tests this set against another for equality, disregarding override settings.
+	 *
+	 * @param InReferenceSet	The set to compare against
+	 * @return					true if the sets are equal
+	 */
+	bool operator==(const FComputeKernelShaderMapId& InReferenceSet) const;
+
+	bool operator!=(const FComputeKernelShaderMapId& InReferenceSet) const
+	{
+		return !(*this == InReferenceSet);
+	}
+
+	/** Appends string representations of this Id to a key string. */
+	void AppendKeyString(FString& OutKeyString) const;
+
+	/** Returns true if the requested shader type is a dependency of this shader map Id. */
+	bool ContainsShaderType(const FShaderType* ShaderType) const;
+};
+
+class FComputeKernelShaderMapContent : public FShaderMapContent
+{
+	using Super = FShaderMapContent;
+	friend class FComputeKernelShaderMap;
+	DECLARE_TYPE_LAYOUT(FComputeKernelShaderMapContent, NonVirtual);
+private:
+	explicit FComputeKernelShaderMapContent(EShaderPlatform InPlatform) : Super(InPlatform) {}
+
+	/** The kernel's user friendly name */
+	LAYOUT_FIELD(FMemoryImageString, FriendlyName);
+
+	/** The static parameter set that this shader map was compiled with */
+	LAYOUT_FIELD(FComputeKernelShaderMapId, ShaderMapId);
+
+	/** Shader compilation output */
+	LAYOUT_FIELD(FComputeKernelCompilationOutput, CompilationOutput);
+};
+
+
+/**
+ * All the shader permutations for a kernel.
+ */
+class FComputeKernelShaderMap : public TShaderMap<FComputeKernelShaderMapContent, FShaderMapPointerTable>, public FDeferredCleanupInterface
+{
+public:
+	using Super = TShaderMap<FComputeKernelShaderMapContent, FShaderMapPointerTable>;
+
+	/**
+	 * Finds the shader map for a kernel.
+	 * @param InShaderMapId - Id to find
+	 * @param Platform - The platform to lookup for
+	 * @return nullptr if no cached shader map was found.
+	 */
+	static FComputeKernelShaderMap* FindId(const FComputeKernelShaderMapId& InShaderMapId, EShaderPlatform InPlatform);
+
+	// ShaderMap interface
+	template<typename ShaderType> TShaderRef<ShaderType> GetShader() const { return TShaderRef<ShaderType>(GetContent()->GetShader<ShaderType>(), *this); }
+	TShaderRef<FShader> GetShader(FShaderType* ShaderType) const { return TShaderRef<FShader>(GetContent()->GetShader(ShaderType), *this); }
+
+	/**
+	 * Attempts to load the shader map for the given kernel from the Derived Data Cache.
+	 * If InOutShaderMap is valid, attempts to load the individual missing shaders instead.
+	 */
+	static void LoadFromDerivedDataCache(const FComputeKernelResource* InKernel, const FComputeKernelShaderMapId& InShaderMapId, EShaderPlatform InPlatform, TRefCountPtr<FComputeKernelShaderMap>& InOutShaderMap);
+
+	FComputeKernelShaderMap();
+	~FComputeKernelShaderMap();
+
+	/**
+	 * Compiles the shaders for a kernel and caches them in this shader map.
+	 * @param InKernel - The kernel to compile shaders for.
+	 * @param InShaderMapId - the set of static parameters to compile for
+	 * @param InPlatform - The platform to compile to
+	 */
+	void Compile(
+		FComputeKernelResource* InKernel,
+		const FComputeKernelShaderMapId& InShaderMapId,
+		TRefCountPtr<FSharedShaderCompilerEnvironment> InCompilationEnvironment,
+		const FComputeKernelCompilationOutput& InCompilationOutput,
+		EShaderPlatform InPlatform,
+		bool bSynchronousCompile,
+		bool bApplyCompletedShaderMapForRendering
+		);
+
+	/** Sorts the incoming compiled jobs into the appropriate shader maps, and finalizes this shader map so that it can be used for rendering. */
+	bool ProcessCompilationResults(const TArray<FShaderCommonCompileJobPtr>& InCompilationResults, int32& InOutResultIndex, float& InOutTimeBudget);
+
+	/**
+	 * Checks whether the shader map is missing any shader types necessary for the given kernel.
+	 * @param InKernel - The ComputeKernel which is checked.
+	 * @return True if the shader map has all of the shader types necessary.
+	 */
+	bool IsComplete(const FComputeKernelResource* InKernel, bool bSilent);
+
+	/** Attempts to load missing shaders from memory. */
+	void LoadMissingShadersFromMemory(const FComputeKernelResource* InKernel);
+
+	/**
+	 * Checks to see if the shader map is already being compiled for another kernel, and if so
+	 * adds the specified kernel to the list to be applied to once the compile finishes.
+	 * @param InKernel - The ComputeKernel we also wish to apply the compiled shader map to.
+	 * @return True if the shader map was being compiled and we added it to the list to be applied.
+	 */
+	bool TryToAddToExistingCompilationTask(FComputeKernelResource* InKernel);
+
+	/** Builds a list of the shaders in a shader map. */
+	 void GetShaderList(TMap<FShaderId, TShaderRef<FShader>>& OutShaders) const;
+
+	/** Registers a ComputeKernel shader map in the global map so it can be used. */
+	void Register(EShaderPlatform InShaderPlatform);
+
+	// Reference counting.
+	void AddRef();
+	void Release();
+
+	/**
+	 * Removes all entries in the cache with exceptions based on a shader type
+	 * @param ShaderType - The shader type to flush
+	 */
+	void FlushShadersByShaderType(const FShaderType* ShaderType);
+
+	/** Removes a kernel from ComputeKernelShaderMapsBeingCompiled. */
+	static void RemovePending(FComputeKernelResource* InKernel);
+
+	/** Finds a shader map currently being compiled that was enqueued for the given kernel. */
+	static const FComputeKernelShaderMap* GetShaderMapBeingCompiled(const FComputeKernelResource* InKernel);
+
+	/** Serializes the shader map. */
+	bool Serialize(FArchive& Ar, bool bInlineShaderResources = true);
+
+	/** Saves this shader map to the derived data cache. */
+	void SaveToDerivedDataCache();
+
+	// Accessors.
+	const FComputeKernelShaderMapId& GetShaderMapId() const	{ return GetContent()->ShaderMapId; }
+	EShaderPlatform GetShaderPlatform() const				{ return GetContent()->GetShaderPlatform(); }
+	const FMemoryImageString& GetFriendlyName() const		{ return GetContent()->FriendlyName; }
+	uint32 GetCompilingId() const							{ return CompilingId; }
+	bool IsCompilationFinalized() const						{ return bCompilationFinalized; }
+	bool CompiledSuccessfully() const						{ return bCompiledSuccessfully; }
+
+	bool IsValid() const
+	{
+		return bCompilationFinalized && bCompiledSuccessfully && !bDeletedThroughDeferredCleanup;
+	}
+
+	int32 GetNumRefs() const { return NumRefs; }
+	uint32 GetCompilingId()  { return CompilingId; }
+	static TMap<TRefCountPtr<FComputeKernelShaderMap>, TArray<FComputeKernelResource*> > &GetInFlightShaderMaps() { return ComputeKernelShaderMapsBeingCompiled; }
+
+	void SetCompiledSuccessfully(bool bSuccess) { bCompiledSuccessfully = bSuccess; }
+private:
+
+	/**
+	 * A global map from a kernel's ID and static switch set to any shader map cached for that kernel.
+	 * Note: this does not necessarily contain all kernel shader maps in memory.  Shader maps with the same key can evict each other.
+	 * No ref counting needed as these are removed on destruction of the shader map.
+	 */
+	static TMap<FComputeKernelShaderMapId, FComputeKernelShaderMap*> GIdToComputeKernelShaderMap[SP_NumPlatforms];
+
+	/**
+	 * All kernel shader maps in memory.
+	 * No ref counting needed as these are removed on destruction of the shader map.
+	 */
+	static TArray<FComputeKernelShaderMap*> AllComputeKernelShaderMaps;
+
+	/** Next value for CompilingId. */
+	static uint32 NextCompilingId;
+
+	/** Tracks resources and their shader maps that need to be compiled but whose compilation is being deferred. */
+	static TMap<TRefCountPtr<FComputeKernelShaderMap>, TArray<FComputeKernelResource*> > ComputeKernelShaderMapsBeingCompiled;
+
+	/** Uniquely identifies this shader map during compilation, needed for deferred compilation where shaders from multiple shader maps are compiled together. */
+	uint32 CompilingId;
+
+	mutable int32 NumRefs;
+
+	/** Used to catch errors where the shader map is deleted directly. */
+	bool bDeletedThroughDeferredCleanup;
+
+	/** Indicates whether this shader map has been registered in GIdToComputeKernelShaderMap */
+	uint32 bRegistered : 1;
+
+	/**
+	 * Indicates whether this shader map has had ProcessCompilationResults called after Compile.
+	 * The shader map must not be used on the rendering thread unless bCompilationFinalized is true.
+	 */
+	uint32 bCompilationFinalized : 1;
+
+	uint32 bCompiledSuccessfully : 1;
+
+	/** Indicates whether the shader map should be stored in the shader cache. */
+	uint32 bIsPersistent : 1;
+
+	uint32 bHasFrozenContent : 1;
+
+	FShader* ProcessCompilationResultsForSingleJob(FShaderCompileJob& InSingleJob, const FSHAHash& InShaderMapHash);
+
+	bool IsComputeKernelShaderComplete(const FComputeKernelResource* InKernel, const FComputeKernelShaderType* InShaderType, bool bSilent);
+
+	friend class FShaderCompilingManager;
+};
+
+
+/**
+ * FComputeKernelResource represents a UComputeKernel to the shader compilation process
+ */
+class FComputeKernelResource
+{
+public:
+
+	/**
+	 * Minimal initialization constructor.
+	 */
+	FComputeKernelResource()
+		: bContainsInlineShaders(false)
+		, bLoadedCookedShaderMapId(false)
+	{}
+
+	virtual ~FComputeKernelResource();
+
+	/**
+	 * Caches the shaders for this kernel with no static parameters on the given platform.
+	 * This is used by UComputeKernel
+	 */
+	bool CacheShaders(EShaderPlatform InPlatform, const ITargetPlatform* TargetPlatform, bool bApplyCompletedShaderMapForRendering, bool bSynchronous);
+	bool CacheShaders(const FComputeKernelShaderMapId& InShaderMapId, EShaderPlatform InPlatform, bool bApplyCompletedShaderMapForRendering, bool bSynchronous);
+
+	/**
+	 * Should the shader for this kernel with the given platform and shader type combination be compiled
+	 *
+	 * @param InPlatform		The platform currently being compiled for
+	 * @param InShaderType		Which shader is being compiled
+	 *
+	 * @return true if the shader should be compiled
+	 */
+	virtual bool ShouldCache(EShaderPlatform InPlatform, const FShaderType* InShaderType) const;
+
+	void SerializeShaderMap(FArchive& Ar);
+
+	/** Releases this kernel's shader map.  Must only be called on ComputeKernels not exposed to the rendering thread! */
+	void ReleaseShaderMap();
+
+	/** Discards loaded shader maps if the application can't render */
+	void DiscardShaderMap();
+
+	void GetDependentShaderTypes(EShaderPlatform InPlatform, TArray<FShaderType*>& OutShaderTypes) const;
+	virtual void GetShaderMapId(EShaderPlatform InPlatform, const ITargetPlatform* TargetPlatform, FComputeKernelShaderMapId& OutId) const;
+
+	void Invalidate();
+
+	/**
+	 * Should shaders compiled for this kernel be saved to disk?
+	 */
+	virtual bool IsPersistent() const { return true; }
+
+	/**
+	 * Called when compilation finishes, after the GameThreadShaderMap is set and the render command to set the RenderThreadShaderMap is queued
+	 */
+	virtual void NotifyCompilationFinished();
+
+	/**
+	 * Cancels all outstanding compilation jobs
+	 */
+	void CancelCompilation();
+
+	/**
+	 * Blocks until compilation has completed. Returns immediately if a compilation is not outstanding.
+	 */
+	void FinishCompilation();
+
+	/**
+	 * Checks if the compilation for this shader is finished
+	 *
+	 * @return returns true if compilation is complete false otherwise
+	 */
+	bool IsCompilationFinished() const;
+
+	// Accessors.
+	const TArray<FString>& GetCompileErrors() const { return CompileErrors; }
+	void SetCompileErrors(const TArray<FString>& InCompileErrors) { CompileErrors = InCompileErrors; }
+
+	ERHIFeatureLevel::Type GetFeatureLevel() const { return FeatureLevel; }
+
+	FComputeKernelShaderMap* GetGameThreadShaderMap() const
+	{
+		checkSlow(IsInGameThread() || IsInAsyncLoadingThread());
+		return GameThreadShaderMap;
+	}
+
+	/** Note: SetRenderingThreadShaderMap must also be called with the same value, but from the rendering thread. */
+	void SetGameThreadShaderMap(FComputeKernelShaderMap* InShaderMap)
+	{
+		checkSlow(IsInGameThread() || IsInAsyncLoadingThread());
+		GameThreadShaderMap = InShaderMap;
+	}
+
+	/** Note: SetGameThreadShaderMap must also be called with the same value, but from the game thread. */
+	void SetRenderingThreadShaderMap(FComputeKernelShaderMap* InShaderMap);
+
+	void AddCompileId(uint32 InIdentifier) 
+	{
+		OutstandingCompileShaderMapIds.Add(InIdentifier);
+	}
+
+	void SetInlineShaderMap(FComputeKernelShaderMap* InShaderMap)
+	{
+		checkSlow(IsInGameThread() || IsInAsyncLoadingThread());
+		GameThreadShaderMap = InShaderMap;
+		bContainsInlineShaders = true;
+		bLoadedCookedShaderMapId = true;
+		CookedShaderMapId = InShaderMap->GetShaderMapId();
+	}
+
+	void RemoveOutstandingCompileId(const int32 InOldOutstandingCompileShaderMapId);
+
+	/**
+	* Get generated or written source code for the shader
+	* @param OutSource - generated source code
+	* @return - true on Success
+	*/
+	bool GetHLSLSource(FString& OutSource) const
+	{
+		OutSource = ShaderSource;
+		return true;
+	};
+
+	const FString& GetHLSLSource() const
+	{
+		return ShaderSource;
+	}
+
+	const FString& GetEntryPoint() const
+	{
+		return ShaderEntryPoint;
+	}
+
+	const FString& GetFriendlyName()	const { return FriendlyName; }
+
+	void SetupResource(
+		ERHIFeatureLevel::Type InFeatureLevel,
+		const UComputeKernelSource* Source,
+		FString InFriendlyName
+		);
+
+	void SetCompileErrors(TArray<FString> &InErrors)
+	{
+		CompileErrors = InErrors;
+	}
+
+	const FShaderParametersMetadata* GetShaderParamMetadata() const
+	{
+		return ShaderMetadata.Get();
+	}
+	
+	TShaderRef<FComputeKernelShader> GetShader() const;
+	TShaderRef<FComputeKernelShader> GetShaderGameThread() const;
+	
+	bool IsSame(const FComputeKernelShaderMapId& InId) const;
+
+	uint32 GetKernelFlags() const { return 0; }
+
+protected:
+
+	/**
+	 * Fills the passed array with IDs of shader maps unfinished compilation jobs.
+	 */
+	void GetShaderMapIDsWithUnfinishedCompilation(TArray<int32>& OutShaderMapIds);
+
+private:
+	TUniquePtr<FShaderParametersMetadata> ShaderMetadata;
+
+	TArray<FString> CompileErrors;
+
+	/** 
+	 * Game thread tracked shader map, which is ref counted and manages shader map lifetime. 
+	 * The shader map uses deferred deletion so that the rendering thread has a chance to process a release command when the shader map is no longer referenced.
+	 * Code that sets this is responsible for updating RenderingThreadShaderMap in a thread safe way.
+	 * During an async compile, this will be NULL and will not contain the actual shader map until compilation is complete.
+	 */
+	TRefCountPtr<FComputeKernelShaderMap> GameThreadShaderMap = nullptr;
+
+	/** 
+	 * Shader map for this FComputeKernelResource which is accessible by the rendering thread. 
+	 * This must be updated along with GameThreadShaderMap, but on the rendering thread.
+	 */
+	FComputeKernelShaderMap* RenderingThreadShaderMap = nullptr;
+
+	FString ShaderEntryPoint;
+
+	FString ShaderSource;
+
+	/** 
+	 * Hash computed from shader code. If multiple kernels give
+	 * the same output, same shader resource will be shared 
+	 */
+	uint64 ShaderCodeHash = 0;
+
+	/** 
+	 * Contains the compiling id of this shader map when it is being compiled asynchronously. 
+	 * This can be used to access the shader map during async compiling, since GameThreadShaderMap will not have been set yet.
+	 */
+	TArray<int32, TInlineAllocator<1> > OutstandingCompileShaderMapIds;
+
+	/** Feature level that this kernel is representing. */
+	ERHIFeatureLevel::Type FeatureLevel = ERHIFeatureLevel::SM5;
+
+	/**
+	 * Whether this kernel was loaded with shaders inlined.
+	 * If true, GameThreadShaderMap will contain a reference to the inlined shader map between Serialize and PostLoad.
+	 */
+
+	uint32 bContainsInlineShaders : 1;
+	uint32 bLoadedCookedShaderMapId : 1;
+	FComputeKernelShaderMapId CookedShaderMapId;
+
+	/**
+	 * Compiles this kernel for InPlatform, storing the result in OutShaderMap if the compile was synchronous
+	 */
+	bool BeginCompileShaderMap(
+		const FComputeKernelShaderMapId& InShaderMapId,
+		EShaderPlatform InPlatform, 
+		TRefCountPtr<FComputeKernelShaderMap>& OutShaderMap, 
+		bool bApplyCompletedShaderMapForRendering,
+		bool bSynchronous = false);
+
+	/** Populates OutEnvironment with defines needed to compile shaders for this kernel. */
+	void SetupShaderCompilationEnvironment(
+		EShaderPlatform InPlatform,
+		FShaderCompilerEnvironment& OutEnvironment
+		) const;
+
+
+	FString FriendlyName;
+
+	friend class FComputeKernelShaderMap;
+	friend class FShaderCompilingManager;
+};
+
+
diff --git a/Engine/Source/Runtime/Engine/Public/ComputeFramework/ComputeKernelSource.h b/Engine/Source/Runtime/Engine/Public/ComputeFramework/ComputeKernelSource.h
--- a/Engine/Source/Runtime/Engine/Public/ComputeFramework/ComputeKernelSource.h
+++ b/Engine/Source/Runtime/Engine/Public/ComputeFramework/ComputeKernelSource.h
@@ -3,6 +3,8 @@
 #pragma once
 
 #include "ComputeFramework/ComputeKernelPermutationSet.h"
+#include "ComputeFramework/ShaderParamTypeDefinition.h"
+
 #include "ComputeKernelSource.generated.h"
 
 class FComputeKernelResource;
@@ -18,4 +20,19 @@ public:
 
 	UPROPERTY(VisibleAnywhere, meta = (ShowOnlyInnerProperties), Category = "Kernel")
 	FComputeKernelDefinitionsSet DefinitionsSet;
+
+	/* Named input types to the kernel. */
+	UPROPERTY(VisibleAnywhere, EditFixedSize, Category = "Parameters")
+	TArray<FShaderParamTypeDefinition> InputParams;
+
+	UPROPERTY(VisibleAnywhere, EditFixedSize, Category = "Parameters")
+	TArray<FShaderParamTypeDefinition> InputSRVs;
+
+	/* Named output types to the kernel. */
+	UPROPERTY(VisibleAnywhere, EditFixedSize, Category = "Parameters")
+	TArray<FShaderParamTypeDefinition> Outputs;
+
+	virtual FString GetEntryPoint() const PURE_VIRTUAL(UComputeKernelSource::GetEntryPoint, return {};);
+	virtual FString GetSource() const PURE_VIRTUAL(UComputeKernelSource::GetSource, return {};);
+	virtual uint64 GetSourceHashCode() const PURE_VIRTUAL(UComputeKernelSource::GetSourceHashCode, return 0;);
 };
diff --git a/Engine/Source/Runtime/Engine/Public/ComputeFramework/ShaderParamTypeDefinition.h b/Engine/Source/Runtime/Engine/Public/ComputeFramework/ShaderParamTypeDefinition.h
--- a/Engine/Source/Runtime/Engine/Public/ComputeFramework/ShaderParamTypeDefinition.h
+++ b/Engine/Source/Runtime/Engine/Public/ComputeFramework/ShaderParamTypeDefinition.h
@@ -20,7 +20,6 @@ enum class EShaderFundamentalType : uint8
 	Int,
 	Uint,
 	Float,
-	Double,
 	Struct
 };
 
@@ -189,6 +188,32 @@ struct FShaderParamTypeDefinition
 	GENERATED_BODY()
 
 public:
+	static EShaderFundamentalType ParseFundamental(
+		const FString& Str
+		);
+
+	static EShaderFundamentalDimensionType ParseDimension(
+		const FString& Str
+		);
+
+	static uint8 ParseVectorDimension(
+		const FString& Str
+		);
+
+	static FIntVector2 ParseMatrixDimension(
+		const FString& Str
+		);
+
+	static EShaderResourceType ParseResource(
+		const FString& Str
+		);
+
+public:
+#if WITH_EDITOR
+	UPROPERTY(VisibleAnywhere, Category = "Kernel")
+	FString TypeDeclaration;
+#endif
+
 	UPROPERTY(EditAnywhere, Category = "Kernel")
 	FString	Name;
 
@@ -196,8 +221,6 @@ public:
 	UPROPERTY()
 	FShaderValueTypeHandle ValueType;
 
-// #TODO_ZABIR: Shader reflection needs to extract much richer type info. Currently skip param type validation
-#if 0
 	UPROPERTY()
 	uint16 ArrayElementCount; // 0 indicates not an array. >= 1 indicates an array
 
@@ -215,12 +238,12 @@ public:
 
 	union
 	{
-		uint8 VectorElemCount : 2;
+		uint8 VectorDimension : 3;
 
 		struct
 		{
-			uint8 MatrixRowCount : 2;
 			uint8 MatrixColumnCount : 2;
+			uint8 MatrixRowCount : 2;
 		};
 	};
 
@@ -270,5 +293,7 @@ public:
 
 		return true;
 	}
-#endif
+
+	void ResetTypeDeclaration(
+		);
 };
diff --git a/Engine/Source/Runtime/RenderCore/Private/ShaderParameterMetadata.cpp b/Engine/Source/Runtime/RenderCore/Private/ShaderParameterMetadata.cpp
--- a/Engine/Source/Runtime/RenderCore/Private/ShaderParameterMetadata.cpp
+++ b/Engine/Source/Runtime/RenderCore/Private/ShaderParameterMetadata.cpp
@@ -7,6 +7,7 @@
 #include "ShaderParameterMetadata.h"
 #include "RenderCore.h"
 #include "ShaderCore.h"
+#include "ShaderParameterMetadataBuilder.h"
 
 FUniformBufferStaticSlotRegistrar::FUniformBufferStaticSlotRegistrar(const TCHAR* InName)
 {
@@ -890,3 +891,73 @@ FString FShaderParametersMetadata::GetFullMemberCodeName(uint16 MemberOffset) co
 
 	return MemberName;
 }
+
+void FShaderParametersMetadataBuilder::AddRDGBufferSRV(
+	const TCHAR* Name,
+	const TCHAR* ShaderType,
+	EShaderPrecisionModifier::Type Precision /* = EShaderPrecisionModifier::Float */
+	)
+{
+	NextMemberOffset = Align(NextMemberOffset, SHADER_PARAMETER_POINTER_ALIGNMENT);
+
+	new(Members) FShaderParametersMetadata::FMember(
+		Name,
+		ShaderType,
+		__LINE__,
+		NextMemberOffset,
+		UBMT_RDG_BUFFER_SRV,
+		Precision,
+		TShaderResourceParameterTypeInfo<FRDGBufferSRV*>::NumRows,
+		TShaderResourceParameterTypeInfo<FRDGBufferSRV*>::NumColumns,
+		TShaderResourceParameterTypeInfo<FRDGBufferSRV*>::NumElements,
+		TShaderResourceParameterTypeInfo<FRDGBufferSRV*>::GetStructMetadata()
+		);
+
+	NextMemberOffset += SHADER_PARAMETER_POINTER_ALIGNMENT;
+}
+
+void FShaderParametersMetadataBuilder::AddRDGBufferUAV(
+	const TCHAR* Name,
+	const TCHAR* ShaderType,
+	EShaderPrecisionModifier::Type Precision /* = EShaderPrecisionModifier::Float */
+	)
+{
+	NextMemberOffset = Align(NextMemberOffset, SHADER_PARAMETER_POINTER_ALIGNMENT);
+
+	new(Members) FShaderParametersMetadata::FMember(
+		Name,
+		ShaderType,
+		__LINE__,
+		NextMemberOffset,
+		UBMT_RDG_BUFFER_UAV,
+		Precision,
+		TShaderResourceParameterTypeInfo<FRDGBufferUAV*>::NumRows,
+		TShaderResourceParameterTypeInfo<FRDGBufferUAV*>::NumColumns,
+		TShaderResourceParameterTypeInfo<FRDGBufferUAV*>::NumElements,
+		TShaderResourceParameterTypeInfo<FRDGBufferUAV*>::GetStructMetadata()
+		);
+
+	NextMemberOffset += SHADER_PARAMETER_POINTER_ALIGNMENT;
+}
+
+FShaderParametersMetadata* FShaderParametersMetadataBuilder::Build(
+	FShaderParametersMetadata::EUseCase UseCase,
+	const TCHAR* ShaderParameterName
+	)
+{
+	const uint32 StructSize = Align(NextMemberOffset, SHADER_PARAMETER_STRUCT_ALIGNMENT);
+
+	FShaderParametersMetadata* ShaderParameterMetadata = new FShaderParametersMetadata(
+		UseCase,
+		ShaderParameterName,
+		ShaderParameterName,
+		nullptr,
+		nullptr,
+		__FILE__,
+		__LINE__,
+		StructSize,
+		Members
+		);
+
+	return ShaderParameterMetadata;
+}
diff --git a/Engine/Source/Runtime/RenderCore/Public/RenderGraphBuilder.h b/Engine/Source/Runtime/RenderCore/Public/RenderGraphBuilder.h
--- a/Engine/Source/Runtime/RenderCore/Public/RenderGraphBuilder.h
+++ b/Engine/Source/Runtime/RenderCore/Public/RenderGraphBuilder.h
@@ -133,6 +133,9 @@ public:
 	 *
 	 *  The pass name is used by debugging / profiling tools.
 	 */
+	template <typename ParameterStructType, typename ExecuteLambdaType>
+	FRDGPassRef AddPass(FRDGEventName&& Name, const FShaderParametersMetadata* ParametersMetadata, const ParameterStructType* ParameterStruct, ERDGPassFlags Flags, ExecuteLambdaType&& ExecuteLambda);
+	
 	template <typename ParameterStructType, typename ExecuteLambdaType>
 	FRDGPassRef AddPass(FRDGEventName&& Name, const ParameterStructType* ParameterStruct, ERDGPassFlags Flags, ExecuteLambdaType&& ExecuteLambda);
 
diff --git a/Engine/Source/Runtime/RenderCore/Public/RenderGraphBuilder.inl b/Engine/Source/Runtime/RenderCore/Public/RenderGraphBuilder.inl
--- a/Engine/Source/Runtime/RenderCore/Public/RenderGraphBuilder.inl
+++ b/Engine/Source/Runtime/RenderCore/Public/RenderGraphBuilder.inl
@@ -244,6 +244,7 @@ FRDGPassRef FRDGBuilder::AddPass(
 template <typename ParameterStructType, typename ExecuteLambdaType>
 FRDGPassRef FRDGBuilder::AddPass(
 	FRDGEventName&& Name,
+	const FShaderParametersMetadata* ParametersMetadata,
 	const ParameterStructType* ParameterStruct,
 	ERDGPassFlags Flags,
 	ExecuteLambdaType&& ExecuteLambda)
@@ -271,6 +272,7 @@ FRDGPassRef FRDGBuilder::AddPass(
 
 	FRDGPass* Pass = Allocator.AllocObject<LambdaPassType>(
 		MoveTemp(Name),
+		ParametersMetadata,
 		ParameterStruct,
 		OverridePassFlags(Name.GetTCHAR(), Flags, LambdaPassType::kSupportsAsyncCompute),
 		MoveTemp(ExecuteLambda));
@@ -281,6 +283,16 @@ FRDGPassRef FRDGBuilder::AddPass(
 	return Pass;
 }
 
+template <typename ParameterStructType, typename ExecuteLambdaType>
+FRDGPassRef FRDGBuilder::AddPass(
+	FRDGEventName&& Name,
+	const ParameterStructType* ParameterStruct,
+	ERDGPassFlags Flags,
+	ExecuteLambdaType&& ExecuteLambda)
+{
+	return AddPass(Forward<FRDGEventName>(Name), ParameterStructType::FTypeInfo::GetStructMetadata(), ParameterStruct, Flags, Forward<ExecuteLambdaType>(ExecuteLambda));
+}
+
 inline void FRDGBuilder::QueueTextureExtraction(
 	FRDGTextureRef Texture,
 	TRefCountPtr<IPooledRenderTarget>* OutTexturePtr,
diff --git a/Engine/Source/Runtime/RenderCore/Public/RenderGraphPass.h b/Engine/Source/Runtime/RenderCore/Public/RenderGraphPass.h
--- a/Engine/Source/Runtime/RenderCore/Public/RenderGraphPass.h
+++ b/Engine/Source/Runtime/RenderCore/Public/RenderGraphPass.h
@@ -407,10 +407,11 @@ public:
 
 	TRDGLambdaPass(
 		FRDGEventName&& InName,
+		const FShaderParametersMetadata* InParameterMetadata,
 		const ParameterStructType* InParameterStruct,
 		ERDGPassFlags InPassFlags,
 		ExecuteLambdaType&& InExecuteLambda)
-		: FRDGPass(MoveTemp(InName), FRDGParameterStruct(InParameterStruct), InPassFlags)
+		: FRDGPass(MoveTemp(InName), FRDGParameterStruct(InParameterStruct, &InParameterMetadata->GetLayout()), InPassFlags)
 		, ExecuteLambda(MoveTemp(InExecuteLambda))
 	{
 		checkf(kSupportsAsyncCompute || !EnumHasAnyFlags(InPassFlags, ERDGPassFlags::AsyncCompute),
@@ -433,7 +434,7 @@ class TRDGEmptyLambdaPass
 {
 public:
 	TRDGEmptyLambdaPass(FRDGEventName&& InName, ERDGPassFlags InPassFlags, ExecuteLambdaType&& InExecuteLambda)
-		: TRDGLambdaPass<FEmptyShaderParameters, ExecuteLambdaType>(MoveTemp(InName), &EmptyShaderParameters, InPassFlags, MoveTemp(InExecuteLambda))
+		: TRDGLambdaPass<FEmptyShaderParameters, ExecuteLambdaType>(MoveTemp(InName), FEmptyShaderParameters::FTypeInfo::GetStructMetadata(), &EmptyShaderParameters, InPassFlags, MoveTemp(InExecuteLambda))
 	{}
 
 private:
diff --git a/Engine/Source/Runtime/RenderCore/Public/RenderGraphUtils.h b/Engine/Source/Runtime/RenderCore/Public/RenderGraphUtils.h
--- a/Engine/Source/Runtime/RenderCore/Public/RenderGraphUtils.h
+++ b/Engine/Source/Runtime/RenderCore/Public/RenderGraphUtils.h
@@ -57,11 +57,10 @@ extern RENDERCORE_API void ClearUnusedGraphResourcesImpl(
 template <typename TShaderClass>
 void ClearUnusedGraphResources(
 	const TShaderRef<TShaderClass>& Shader,
+	const FShaderParametersMetadata* ParametersMetadata,
 	typename TShaderClass::FParameters* InoutParameters,
 	std::initializer_list<FRDGResourceRef> ExcludeList = {})
 {
-	const FShaderParametersMetadata* ParametersMetadata = TShaderClass::FParameters::FTypeInfo::GetStructMetadata();
-
 	// Verify the shader have all the parameters it needs. This is done before the
 	// ClearUnusedGraphResourcesImpl() to not mislead user on why some resource are missing
 	// when debugging a validation failure.
@@ -71,6 +70,16 @@ void ClearUnusedGraphResources(
 	return ClearUnusedGraphResourcesImpl(Shader->Bindings, ParametersMetadata, InoutParameters, ExcludeList);
 }
 
+template <typename TShaderClass>
+void ClearUnusedGraphResources(
+	const TShaderRef<TShaderClass>& Shader,
+	typename TShaderClass::FParameters* InoutParameters,
+	std::initializer_list<FRDGResourceRef> ExcludeList = {})
+{
+	const FShaderParametersMetadata* ParametersMetadata = TShaderClass::FParameters::FTypeInfo::GetStructMetadata();
+	return ClearUnusedGraphResources(Shader, ParametersMetadata, InoutParameters, MoveTemp(ExcludeList));
+}
+
 template <typename TShaderClassA, typename TShaderClassB, typename TPassParameterStruct>
 void ClearUnusedGraphResources(
 	const TShaderRef<TShaderClassA>& ShaderA,
@@ -250,15 +259,31 @@ struct RENDERCORE_API FComputeShaderUtils
 
 	/** Dispatch a compute shader to rhi command list with its parameters. */
 	template<typename TShaderClass>
-	static void Dispatch(FRHIComputeCommandList& RHICmdList, const TShaderRef<TShaderClass>& ComputeShader, const typename TShaderClass::FParameters& Parameters, FIntVector GroupCount)
+	static void Dispatch(
+		FRHIComputeCommandList& RHICmdList, 
+		const TShaderRef<TShaderClass>& ComputeShader, 
+		const FShaderParametersMetadata* ParametersMetadata,
+		const typename TShaderClass::FParameters& Parameters,
+		FIntVector GroupCount)
 	{
 		ValidateGroupCount(GroupCount);
 		FRHIComputeShader* ShaderRHI = ComputeShader.GetComputeShader();
 		RHICmdList.SetComputeShader(ShaderRHI);
-		SetShaderParameters(RHICmdList, ComputeShader, ShaderRHI, Parameters);
+		SetShaderParameters(RHICmdList, ComputeShader, ShaderRHI, ParametersMetadata, Parameters);
 		RHICmdList.DispatchComputeShader(GroupCount.X, GroupCount.Y, GroupCount.Z);
 		UnsetShaderUAVs(RHICmdList, ComputeShader, ShaderRHI);
 	}
+
+	template<typename TShaderClass>
+	static void Dispatch(
+		FRHIComputeCommandList& RHICmdList,
+		const TShaderRef<TShaderClass>& ComputeShader,
+		const typename TShaderClass::FParameters& Parameters,
+		FIntVector GroupCount)
+	{
+		const FShaderParametersMetadata* ParametersMetadata = TShaderClass::FParameters::FTypeInfo::GetStructMetadata();
+		Dispatch(RHICmdList, ComputeShader, ParametersMetadata, Parameters, GroupCount);
+	}
 	
 	/** Indirect dispatch a compute shader to rhi command list with its parameters. */
 	template<typename TShaderClass>
@@ -301,6 +326,7 @@ struct RENDERCORE_API FComputeShaderUtils
 		FRDGEventName&& PassName,
 		ERDGPassFlags PassFlags,
 		const TShaderRef<TShaderClass>& ComputeShader,
+		const FShaderParametersMetadata* ParametersMetadata,
 		typename TShaderClass::FParameters* Parameters,
 		FIntVector GroupCount)
 	{
@@ -309,18 +335,32 @@ struct RENDERCORE_API FComputeShaderUtils
 			!EnumHasAnyFlags(PassFlags, ERDGPassFlags::Copy | ERDGPassFlags::Raster), TEXT("AddPass only supports 'Compute' or 'AsyncCompute'."));
 
 		ValidateGroupCount(GroupCount);
-		ClearUnusedGraphResources(ComputeShader, Parameters);
+		ClearUnusedGraphResources(ComputeShader, ParametersMetadata, Parameters);
 
 		GraphBuilder.AddPass(
 			Forward<FRDGEventName>(PassName),
+			ParametersMetadata,
 			Parameters,
 			PassFlags,
-			[Parameters, ComputeShader, GroupCount](FRHIComputeCommandList& RHICmdList)
+			[ParametersMetadata, Parameters, ComputeShader, GroupCount](FRHIComputeCommandList& RHICmdList)
 		{
-			FComputeShaderUtils::Dispatch(RHICmdList, ComputeShader, *Parameters, GroupCount);
+			FComputeShaderUtils::Dispatch(RHICmdList, ComputeShader, ParametersMetadata, *Parameters, GroupCount);
 		});
 	}
 
+	template<typename TShaderClass>
+	static void AddPass(
+		FRDGBuilder& GraphBuilder,
+		FRDGEventName&& PassName,
+		ERDGPassFlags PassFlags,
+		const TShaderRef<TShaderClass>& ComputeShader,
+		typename TShaderClass::FParameters* Parameters,
+		FIntVector GroupCount)
+	{
+		const FShaderParametersMetadata* ParametersMetadata = TShaderClass::FParameters::FTypeInfo::GetStructMetadata();
+		AddPass(GraphBuilder, Forward<FRDGEventName>(PassName), PassFlags, ComputeShader, ParametersMetadata, Parameters, GroupCount);
+	}
+
 	template <typename TShaderClass>
 	static FORCEINLINE void AddPass(
 		FRDGBuilder& GraphBuilder,
@@ -329,7 +369,8 @@ struct RENDERCORE_API FComputeShaderUtils
 		typename TShaderClass::FParameters* Parameters,
 		FIntVector GroupCount)
 	{
-		AddPass(GraphBuilder, Forward<FRDGEventName>(PassName), ERDGPassFlags::Compute, ComputeShader, Parameters, GroupCount);
+		const FShaderParametersMetadata* ParametersMetadata = TShaderClass::FParameters::FTypeInfo::GetStructMetadata();
+		AddPass(GraphBuilder, Forward<FRDGEventName>(PassName), ERDGPassFlags::Compute, ComputeShader, ParametersMetadata, Parameters, GroupCount);
 	}
 
 	/** Dispatch a compute shader to render graph builder with its parameters. */
diff --git a/Engine/Source/Runtime/RenderCore/Public/ShaderParameterMetadataBuilder.h b/Engine/Source/Runtime/RenderCore/Public/ShaderParameterMetadataBuilder.h
new file mode 100644
--- /dev/null
+++ b/Engine/Source/Runtime/RenderCore/Public/ShaderParameterMetadataBuilder.h
@@ -0,0 +1,61 @@
+// Copyright Epic Games, Inc. All Rights Reserved.
+
+/*=============================================================================
+	ShaderParameterMetadata.h: Meta data about shader parameter structures
+=============================================================================*/
+
+#pragma once
+
+#include "CoreMinimal.h"
+#include "ShaderParameterMacros.h"
+
+class RENDERCORE_API FShaderParametersMetadataBuilder
+{
+public:
+	template<typename T>
+	void AddParam(
+		const TCHAR* Name,
+		EShaderPrecisionModifier::Type Precision = EShaderPrecisionModifier::Float
+		)
+	{
+		using TParamTypeInfo = TShaderParameterTypeInfo<T>;
+
+		NextMemberOffset = Align(NextMemberOffset, TParamTypeInfo::Alignment);
+
+		new(Members) FShaderParametersMetadata::FMember(
+			Name,
+			TEXT(""),
+			__LINE__,
+			NextMemberOffset,
+			TParamTypeInfo::BaseType,
+			Precision,
+			TParamTypeInfo::NumRows,
+			TParamTypeInfo::NumColumns,
+			TParamTypeInfo::NumElements,
+			TParamTypeInfo::GetStructMetadata()
+			);
+
+		NextMemberOffset += sizeof(TParamTypeInfo::TAlignedType);
+	}
+
+	void AddRDGBufferSRV(
+		const TCHAR* Name,
+		const TCHAR* ShaderType,
+		EShaderPrecisionModifier::Type Precision = EShaderPrecisionModifier::Float
+		);
+
+	void AddRDGBufferUAV(
+		const TCHAR* Name,
+		const TCHAR* ShaderType,
+		EShaderPrecisionModifier::Type Precision = EShaderPrecisionModifier::Float
+		);
+
+	FShaderParametersMetadata* Build(
+		FShaderParametersMetadata::EUseCase UseCase,
+		const TCHAR* ShaderParameterName
+		);
+
+private:
+	TArray<FShaderParametersMetadata::FMember> Members;
+	uint32 NextMemberOffset = 0;
+};
\ No newline at end of file
diff --git a/Engine/Source/Runtime/RenderCore/Public/ShaderParameterStruct.h b/Engine/Source/Runtime/RenderCore/Public/ShaderParameterStruct.h
--- a/Engine/Source/Runtime/RenderCore/Public/ShaderParameterStruct.h
+++ b/Engine/Source/Runtime/RenderCore/Public/ShaderParameterStruct.h
@@ -182,9 +182,14 @@ inline void UnsetShaderUAVs(TRHICmdList& RHICmdList, const TShaderRef<TShaderCla
 
 /** Set shader's parameters from its parameters struct. */
 template<typename TRHICmdList, typename TShaderClass, typename TShaderRHI>
-inline void SetShaderParameters(TRHICmdList& RHICmdList, const TShaderRef<TShaderClass>& Shader, TShaderRHI* ShadeRHI, const typename TShaderClass::FParameters& Parameters)
+inline void SetShaderParameters(
+	TRHICmdList& RHICmdList, 
+	const TShaderRef<TShaderClass>& Shader, 
+	TShaderRHI* ShadeRHI, 
+	const FShaderParametersMetadata* ParametersMetadata,
+	const typename TShaderClass::FParameters& Parameters)
 {
-	ValidateShaderParameters(Shader, Parameters);
+	ValidateShaderParameters(Shader, ParametersMetadata, &Parameters);
 
 	// TODO(RDG): Once all shader sets their parameter through this, can refactor RHI so all shader parameters get sets through a single RHI function call.
 	const FShaderParameterBindings& Bindings = Shader->Bindings;
@@ -282,6 +287,12 @@ inline void SetShaderParameters(TRHICmdList& RHICmdList, const TShaderRef<TShade
 	}
 }
 
+template<typename TRHICmdList, typename TShaderClass, typename TShaderRHI>
+inline void SetShaderParameters(TRHICmdList& RHICmdList, const TShaderRef<TShaderClass>& Shader, TShaderRHI* ShaderRHI, const typename TShaderClass::FParameters& Parameters)
+{
+	const FShaderParametersMetadata* ParametersMetadata = TShaderClass::FParameters::FTypeInfo::GetStructMetadata();
+	SetShaderParameters(RHICmdList, Shader, ShaderRHI, ParametersMetadata, Parameters);
+}
 
 #if RHI_RAYTRACING
 
