Anim node communication/messaging system

This takes the idea behind Braeden's node tracker/ancestor system and makes it a bit more generic, removing the bits of the API that were bound to FAnimNode_Base and allowing easier cross-plugin message dependencies.
The idea is that at any scope in the graph update you can add a TScopedGraphMessage with a specified payload (which can just be a node, but can also have state of its own), then further towards the leaves, that message can be detected and acted upon.
This changelist includes a re-working of the inter-node communication around inertial blending and cached poses as examples.

diff --git a/Engine/Source/Runtime/AnimGraphRuntime/Private/AnimNodes/AnimNode_BlendListBase.cpp b/Engine/Source/Runtime/AnimGraphRuntime/Private/AnimNodes/AnimNode_BlendListBase.cpp
--- a/Engine/Source/Runtime/AnimGraphRuntime/Private/AnimNodes/AnimNode_BlendListBase.cpp
+++ b/Engine/Source/Runtime/AnimGraphRuntime/Private/AnimNodes/AnimNode_BlendListBase.cpp
@@ -109,10 +109,10 @@ void FAnimNode_BlendListBase::Update_AnyThread(const FAnimationUpdateContext& Co
 			}
 			else if (TransitionType == EBlendListTransitionType::Inertialization)
 			{
-				FAnimNode_Inertialization* InertializationNode = Context.GetAncestor<FAnimNode_Inertialization>();
-				if (InertializationNode)
+				UE::Anim::IInertializationRequester* InertializationRequester = Context.GetMessage<UE::Anim::IInertializationRequester>();
+				if (InertializationRequester)
 				{
-					InertializationNode->RequestInertialization(BlendTime[ChildIndex]);
+					InertializationRequester->RequestInertialization(BlendTime[ChildIndex]);
 				}
 				else
 				{
diff --git a/Engine/Source/Runtime/AnimGraphRuntime/Private/AnimNodes/AnimNode_BlendSpacePlayer.cpp b/Engine/Source/Runtime/AnimGraphRuntime/Private/AnimNodes/AnimNode_BlendSpacePlayer.cpp
--- a/Engine/Source/Runtime/AnimGraphRuntime/Private/AnimNodes/AnimNode_BlendSpacePlayer.cpp
+++ b/Engine/Source/Runtime/AnimGraphRuntime/Private/AnimNodes/AnimNode_BlendSpacePlayer.cpp
@@ -105,7 +105,7 @@ void FAnimNode_BlendSpacePlayer::UpdateInternal(const FAnimationUpdateContext& C
 
 		TRACE_ANIM_TICK_RECORD(Context, TickRecord);
 
-#if ANIM_NODE_IDS_AVAILABLE && WITH_EDITORONLY_DATA
+#if WITH_EDITORONLY_DATA
 		if (FAnimBlueprintDebugData* DebugData = Context.AnimInstanceProxy->GetAnimBlueprintDebugData())
 		{
 			DebugData->RecordBlendSpacePlayer(Context.GetCurrentNodeId(), BlendSpace, BlendInput.X, BlendInput.Y, BlendInput.Z);
diff --git a/Engine/Source/Runtime/AnimGraphRuntime/Private/BoneControllers/AnimNode_SkeletalControlBase.cpp b/Engine/Source/Runtime/AnimGraphRuntime/Private/BoneControllers/AnimNode_SkeletalControlBase.cpp
--- a/Engine/Source/Runtime/AnimGraphRuntime/Private/BoneControllers/AnimNode_SkeletalControlBase.cpp
+++ b/Engine/Source/Runtime/AnimGraphRuntime/Private/BoneControllers/AnimNode_SkeletalControlBase.cpp
@@ -145,10 +145,8 @@ void FAnimNode_SkeletalControlBase::EvaluateComponentSpace_AnyThread(FComponentS
 {
 	DECLARE_SCOPE_HIERARCHICAL_COUNTER_ANIMNODE(EvaluateComponentSpace_AnyThread)
 
-#if ANIM_NODE_IDS_AVAILABLE
 	// Cache the incoming node IDs in a base context
 	FAnimationBaseContext CachedContext(Output);
-#endif
 
 	EvaluateComponentPose_AnyThread(Output);
 
@@ -165,9 +163,7 @@ void FAnimNode_SkeletalControlBase::EvaluateComponentSpace_AnyThread(FComponentS
 	// Apply the skeletal control if it's valid
 	if (FAnimWeight::IsRelevant(ActualAlpha) && IsValidToEvaluate(Output.AnimInstanceProxy->GetSkeleton(), Output.AnimInstanceProxy->GetRequiredBones()))
 	{
-#if ANIM_NODE_IDS_AVAILABLE
 		Output.SetNodeIds(CachedContext);
-#endif
 
 		EvaluateComponentSpaceInternal(Output);
 
diff --git a/Engine/Source/Runtime/Engine/Classes/Animation/AnimNodeBase.h b/Engine/Source/Runtime/Engine/Classes/Animation/AnimNodeBase.h
--- a/Engine/Source/Runtime/Engine/Classes/Animation/AnimNodeBase.h
+++ b/Engine/Source/Runtime/Engine/Classes/Animation/AnimNodeBase.h
@@ -16,6 +16,7 @@
 #include "Animation/AnimationPoseData.h"
 #include "UObject/FieldPath.h"
 #include "CustomAttributesRuntime.h"
+#include "Animation/AnimNodeMessages.h"
 
 // WARNING: This should always be the last include in any file that needs it (except .generated.h)
 #include "UObject/UndefineUPropertyMacros.h"
@@ -25,8 +26,6 @@
 #define DECLARE_SCOPE_HIERARCHICAL_COUNTER_ANIMNODE(Method) \
 	DECLARE_SCOPE_HIERARCHICAL_COUNTER_FUNC()
 
-#define ANIM_NODE_IDS_AVAILABLE	(ANIM_TRACE_ENABLED || WITH_EDITORONLY_DATA)
-
 class IAnimClassInterface;
 class UAnimBlueprint;
 class UAnimInstance;
@@ -36,6 +35,7 @@ class UProperty;
 struct FPropertyAccessLibrary;
 
 /**
+ * DEPRECATED - This system is now supplanted by UE::Anim::FMessageStack
  * Utility container for tracking a stack of ancestor nodes by node type during graph traversal
  * This is not an exhaustive list of all visited ancestors. During Update nodes must call
  * FAnimationUpdateContext::TrackAncestor() to appear in the tracker.
@@ -99,7 +99,7 @@ struct FAnimNodeTracker
 };
 
 
-/** Helper RAII object to cleanup a node added to the node tracker */
+/** DEPRECATED - This system is now supplanted by UE::Anim::FMessageStack - Helper RAII object to cleanup a node added to the node tracker */
 class FScopedAnimNodeTracker
 {
 public:
@@ -127,11 +127,24 @@ private:
 /** Persistent state shared during animation tree update  */
 struct FAnimationUpdateSharedContext
 {
+	FAnimationUpdateSharedContext() = default;
+
+	// Non-copyable
+	FAnimationUpdateSharedContext(FAnimationUpdateSharedContext& ) = delete;
+	FAnimationUpdateSharedContext& operator=(const FAnimationUpdateSharedContext&) = delete;
+
+	UE_DEPRECATED(5.0, "Please use the message & tagging system in UE::Anim::FMessageStack")
 	FAnimNodeTracker AncestorTracker;
 
-	void CopyForCachedUpdate(const FAnimationUpdateSharedContext& Source)
+	// Message stack used for storing scoped messages and tags during execution
+	UE::Anim::FMessageStack MessageStack;
+
+	void CopyForCachedUpdate(FAnimationUpdateSharedContext& Source)
 	{
+PRAGMA_DISABLE_DEPRECATION_WARNINGS
 		AncestorTracker.CopyTopsOnly(Source.AncestorTracker);
+PRAGMA_ENABLE_DEPRECATION_WARNINGS
+		MessageStack.CopyForCachedUpdate(Source.MessageStack);
 	}
 };
 
@@ -167,6 +180,7 @@ public:
 #endif //WITH_EDITORONLY_DATA
 
 	template<typename NodeType>
+	UE_DEPRECATED(5.0, "Please use the message & tagging system in UE::Anim::FMessageStack")
 	FScopedAnimNodeTracker TrackAncestor(NodeType* Node) const {
 		if (ensure(SharedContext != nullptr))
 		{
@@ -178,6 +192,7 @@ public:
 	}
 
 	template<typename NodeType>
+	UE_DEPRECATED(5.0, "Please use the message & tagging system in UE::Anim::FMessageStack")
 	NodeType* GetAncestor() const {
 		if (ensure(SharedContext != nullptr))
 		{
@@ -188,7 +203,43 @@ public:
 		return nullptr;
 	}
 
-#if ANIM_NODE_IDS_AVAILABLE
+	// Get the innermost scoped message of the specified type
+	template<typename TGraphMessageType>
+	TGraphMessageType* GetMessage() const
+	{
+		if (ensure(SharedContext != nullptr))
+		{
+			TGraphMessageType* Message = nullptr;
+
+			SharedContext->MessageStack.TopMessage<TGraphMessageType>([&Message](TGraphMessageType& InMessage)
+			{
+				Message = &InMessage;
+			});
+
+			return Message;
+		}
+		
+		return nullptr;
+	}
+	
+	// Get the innermost scoped message of the specified type
+	template<typename TGraphMessageType>
+	TGraphMessageType& GetMessageChecked() const
+	{
+		check(SharedContext != nullptr);
+
+		TGraphMessageType* Message = nullptr;
+
+		SharedContext->MessageStack.TopMessage<TGraphMessageType>([&Message](TGraphMessageType& InMessage)
+		{
+			Message = &InMessage;
+		});
+
+		check(Message != nullptr);
+
+		return *Message;
+	}
+	
 	// Get the current node Id, set when we recurse into graph traversal functions from pose links
 	ENGINE_API int32 GetCurrentNodeId() const { return CurrentNodeId; }
 
@@ -201,7 +252,6 @@ protected:
 
 	// The previous node ID, set when we recurse into graph traversal functions from pose links
 	int32 PreviousNodeId;
-#endif // ANIM_NODE_IDS_AVAILABLE
 
 protected:
 
@@ -328,7 +378,6 @@ public:
 		return Result;
 	}
 
-#if ANIM_NODE_IDS_AVAILABLE
 	FAnimationUpdateContext WithNodeId(int32 InNodeId) const
 	{ 
 		FAnimationUpdateContext Result(*this);
@@ -336,7 +385,6 @@ public:
 		Result.CurrentNodeId = InNodeId;
 		return Result; 
 	}
-#endif
 
 	// Returns persistent state that is tracked through animation tree update
 	FAnimationUpdateSharedContext* GetSharedContext() const
@@ -383,13 +431,10 @@ public:
 	{
 		Initialize(SourceContext.AnimInstanceProxy);
 
-#if ANIM_NODE_IDS_AVAILABLE
 		CurrentNodeId = SourceContext.CurrentNodeId;
 		PreviousNodeId = SourceContext.PreviousNodeId;
-#endif
 	}
 
-#if ANIM_NODE_IDS_AVAILABLE
 	void SetNodeId(int32 InNodeId)
 	{ 
 		PreviousNodeId = CurrentNodeId;
@@ -401,7 +446,6 @@ public:
 		CurrentNodeId = InContext.GetCurrentNodeId();
 		PreviousNodeId = InContext.GetPreviousNodeId();
 	}
-#endif
 
 	ENGINE_API void Initialize(FAnimInstanceProxy* InAnimInstanceProxy);
 
@@ -481,13 +525,10 @@ public:
 	{
 		// No need to initialize, done through FA2CSPose::AllocateLocalPoses
 
-#if ANIM_NODE_IDS_AVAILABLE
 		CurrentNodeId = SourceContext.CurrentNodeId;
 		PreviousNodeId = SourceContext.PreviousNodeId;
-#endif
 	}
 
-#if ANIM_NODE_IDS_AVAILABLE
 	void SetNodeId(int32 InNodeId)
 	{ 
 		PreviousNodeId = CurrentNodeId;
@@ -499,7 +540,6 @@ public:
 		CurrentNodeId = InContext.GetCurrentNodeId();
 		PreviousNodeId = InContext.GetPreviousNodeId();
 	}
-#endif
 
 	ENGINE_API void ResetToRefPose();
 
@@ -888,18 +928,6 @@ struct ENGINE_API FAnimNode_Base
 	/** Called to help dynamics-based updates to recover correctly from large movements/teleports */
 	virtual void ResetDynamics(ETeleportType InTeleportType);
 
-	/**
-	 * Override this if your node uses ancestor tracking and wants to be informed of Update() calls
-	 * that were skipped due to pose caching.
-	 */
-	virtual bool WantsSkippedUpdates() const { return false; }
-	
-	/**
-	 * Called on a tracked ancestor node when there are Update() calls that were skipped due to pose 
-	 * caching. Your node must implement WantsSkippedUpdates to receive this callback.
-	 */
-	virtual void OnUpdatesSkipped(TArrayView<const FAnimationUpdateContext*> SkippedUpdateContexts) {}
-
 	/** Called after compilation */
 	virtual void PostCompile(const class USkeleton* InSkeleton) {}
 
@@ -925,6 +953,10 @@ struct ENGINE_API FAnimNode_Base
 	virtual void EvaluateComponentSpace(FComponentSpacePoseContext& Output) { check(false); }
 	UE_DEPRECATED(4.20, "Please use ResetDynamics with an ETeleportPhysics flag instead")
 	virtual void ResetDynamics() {}
+	UE_DEPRECATED(5.0, "Please use IGraphMessage instead")
+	virtual bool WantsSkippedUpdates() const { return false; }
+	UE_DEPRECATED(5.0, "Please use IGraphMessage instead")
+	virtual void OnUpdatesSkipped(TArrayView<const FAnimationUpdateContext*> SkippedUpdateContexts) {}
 
 	// The default handler for graph-exposed inputs:
 	const FExposedValueHandler& GetEvaluateGraphExposedInputs();
diff --git a/Engine/Source/Runtime/Engine/Classes/Animation/AnimNode_Inertialization.h b/Engine/Source/Runtime/Engine/Classes/Animation/AnimNode_Inertialization.h
--- a/Engine/Source/Runtime/Engine/Classes/Animation/AnimNode_Inertialization.h
+++ b/Engine/Source/Runtime/Engine/Classes/Animation/AnimNode_Inertialization.h
@@ -6,6 +6,7 @@
 #include "UObject/ObjectMacros.h"
 #include "Animation/AnimNodeBase.h"
 #include "Animation/AnimCurveTypes.h"
+#include "Animation/AnimNodeMessages.h"
 #include "AnimNode_Inertialization.generated.h"
 
 
@@ -17,6 +18,21 @@
 // https://www.gdcvault.com/play/1025331/Inertialization
 
 
+namespace UE { namespace Anim {
+
+// Event that can be subscribed to request inertialization-based blends
+class ENGINE_API IInertializationRequester : public IGraphMessage
+{
+	DECLARE_ANIMGRAPH_MESSAGE(IInertializationRequester);
+
+public:
+	// Request to activate inertialization for a duration.
+	// If multiple requests are made on the same inertialization node, the minimum requested time will be used.
+	virtual void RequestInertialization(float InRequestedDuration) = 0;
+};
+
+}}	// namespace UE::Anim
+
 UENUM()
 enum class EInertializationState : uint8
 {
@@ -233,10 +249,6 @@ public: // FAnimNode_Base
 	virtual bool NeedsDynamicReset() const override;
 	virtual void ResetDynamics(ETeleportType InTeleportType) override;
 
-	virtual bool WantsSkippedUpdates() const override;
-	virtual void OnUpdatesSkipped(TArrayView<const FAnimationUpdateContext *> SkippedUpdateContexts) override;
-
-
 protected:
 
 	// Consume Inertialization Request
diff --git a/Engine/Source/Runtime/Engine/Classes/Animation/AnimNode_SaveCachedPose.h b/Engine/Source/Runtime/Engine/Classes/Animation/AnimNode_SaveCachedPose.h
--- a/Engine/Source/Runtime/Engine/Classes/Animation/AnimNode_SaveCachedPose.h
+++ b/Engine/Source/Runtime/Engine/Classes/Animation/AnimNode_SaveCachedPose.h
@@ -8,8 +8,33 @@
 #include "Animation/AnimCurveTypes.h"
 #include "BonePose.h"
 #include "Animation/AnimNodeBase.h"
+#include "Animation/AnimNodeMessages.h"
 #include "AnimNode_SaveCachedPose.generated.h"
 
+namespace UE { namespace Anim {
+
+// Event that can be subscribed to receive skipped updates when a cached pose is run.
+// When a cached pose update call executes the link with the maximum weight, this event receives information about
+// the other links with lesser weights
+class ENGINE_API FCachedPoseSkippedUpdateHandler : public IGraphMessage
+{
+	DECLARE_ANIMGRAPH_MESSAGE(FCachedPoseSkippedUpdateHandler);
+
+public:
+	FCachedPoseSkippedUpdateHandler(TUniqueFunction<void(TArrayView<const FMessageStack>)> InFunction)
+		: Function(MoveTemp(InFunction))
+	{}
+
+	// Called when there are Update() calls that were skipped due to pose caching. 
+	void OnUpdatesSkipped(TArrayView<const FMessageStack> InSkippedUpdates) { Function(InSkippedUpdates); }
+
+private:
+	// Function to call
+	TUniqueFunction<void(TArrayView<const FMessageStack>)> Function;
+};
+
+}}	// namespace UE::Anim
+
 USTRUCT(BlueprintInternalUseOnly)
 struct ENGINE_API FAnimNode_SaveCachedPose : public FAnimNode_Base
 {
diff --git a/Engine/Source/Runtime/Engine/Private/Animation/AnimInstanceProxy.cpp b/Engine/Source/Runtime/Engine/Private/Animation/AnimInstanceProxy.cpp
--- a/Engine/Source/Runtime/Engine/Private/Animation/AnimInstanceProxy.cpp
+++ b/Engine/Source/Runtime/Engine/Private/Animation/AnimInstanceProxy.cpp
@@ -1280,12 +1280,12 @@ void FAnimInstanceProxy::UpdateAnimation()
 
 	FAnimationUpdateSharedContext SharedContext;
 	FAnimationUpdateContext Context(this, CurrentDeltaSeconds, &SharedContext);
-#if ANIM_NODE_IDS_AVAILABLE
+
 	if(AnimClassInterface && AnimClassInterface->GetAnimBlueprintFunctions().Num() > 0)
 	{
 		Context = Context.WithNodeId(AnimClassInterface->GetAnimBlueprintFunctions()[0].OutputPoseNodeIndex);
 	}
-#endif
+
 	UpdateAnimation_WithRoot(Context, RootNode, NAME_AnimGraph);
 }
 
@@ -1471,12 +1471,10 @@ void FAnimInstanceProxy::EvaluateAnimationNode_WithRoot(FPoseContext& Output, FA
 			EvaluationCounter.Increment();
 		}
 
-#if ANIM_NODE_IDS_AVAILABLE
 		if(AnimClassInterface && AnimClassInterface->GetAnimBlueprintFunctions().Num() > 0)
 		{
 			Output.SetNodeId(AnimClassInterface->GetAnimBlueprintFunctions()[0].OutputPoseNodeIndex);
 		}
-#endif
 
 		TRACE_SCOPED_ANIM_NODE(Output);
 
diff --git a/Engine/Source/Runtime/Engine/Private/Animation/AnimNodeMessages.cpp b/Engine/Source/Runtime/Engine/Private/Animation/AnimNodeMessages.cpp
new file mode 100644
--- /dev/null
+++ b/Engine/Source/Runtime/Engine/Private/Animation/AnimNodeMessages.cpp
@@ -0,0 +1,173 @@
+// Copyright Epic Games, Inc. All Rights Reserved.
+
+#include "Animation/AnimNodeMessages.h"
+#include "Animation/AnimNodeBase.h"
+#include "Animation/AnimInstanceProxy.h"
+
+namespace UE { namespace Anim {
+
+const FName IGraphMessage::BaseTypeName("IGraphMessage");
+
+FScopedGraphTag::FScopedGraphTag(const FAnimationBaseContext& InContext, FName InTag)
+	: SharedContext(InContext.SharedContext)
+	, Tag(InTag)
+{
+	check(SharedContext);
+
+	int32 NodeId = InContext.GetCurrentNodeId();
+	FStructProperty* NodeProperty = InContext.AnimInstanceProxy->GetAnimClassInterface()->GetAnimNodeProperties()[NodeId];
+	FAnimNode_Base* AnimNode = NodeProperty->ContainerPtrToValuePtr<FAnimNode_Base>(InContext.AnimInstanceProxy->GetAnimInstanceObject());
+
+	SharedContext->MessageStack.PushTag(AnimNode, NodeProperty->Struct, Tag);
+}
+
+FScopedGraphTag::~FScopedGraphTag()
+{
+	check(SharedContext);
+	SharedContext->MessageStack.PopTag(Tag);
+}
+
+FScopedGraphMessage::FScopedGraphMessage(const FAnimationBaseContext& InContext)
+	: SharedContext(InContext.SharedContext)
+{
+}
+
+void FScopedGraphMessage::PushMessage(const FAnimationBaseContext& InContext, TSharedRef<IGraphMessage> InMessage, FName InMessageType)
+{
+	check(SharedContext);
+
+	MessageType = InMessageType;
+
+	int32 NodeId = InContext.GetCurrentNodeId();
+	FStructProperty* NodeProperty = InContext.AnimInstanceProxy->GetAnimClassInterface()->GetAnimNodeProperties()[NodeId];
+	FAnimNode_Base* AnimNode = NodeProperty->ContainerPtrToValuePtr<FAnimNode_Base>(InContext.AnimInstanceProxy->GetAnimInstanceObject());
+
+	SharedContext->MessageStack.PushMessage(MessageType, AnimNode, NodeProperty->Struct, InMessage);
+}
+
+void FScopedGraphMessage::PopMessage()
+{
+	check(SharedContext);
+	SharedContext->MessageStack.PopMessage(MessageType);
+}
+
+void FMessageStack::PushMessage(FName InMessageType, FAnimNode_Base* InNode, const UScriptStruct* InStruct, TSharedRef<IGraphMessage> InMessage)
+{
+	MessageStacks.FindOrAdd(InMessageType).Push( { InMessage, { InNode, InStruct } } );
+}
+
+void FMessageStack::PopMessage(FName InMessageType)
+{
+	MessageStacks.FindChecked(InMessageType).Pop();
+}
+
+void FMessageStack::ForEachMessage(FName InMessageType, TFunctionRef<EEnumerate(IGraphMessage&)> InFunction) const
+{
+	if(const FMessageStackEntry* StackEntryPtr = MessageStacks.Find(InMessageType))
+	{
+		const FMessageStackEntry& StackEntry = *StackEntryPtr;
+		const int32 StackSize = StackEntry.Num();
+		for(int32 MessageIndex = StackSize - 1; MessageIndex >= 0; --MessageIndex)
+		{
+			const FMessageEntry& MessageEntry = StackEntry[MessageIndex];
+			if(MessageEntry.Message->GetTypeName() == InMessageType)
+			{
+				if(InFunction(MessageEntry.Message.Get()) == EEnumerate::Stop)
+				{
+					return;
+				}
+			}
+		}
+	}
+}
+
+void FMessageStack::TopMessage(FName InMessageType, TFunctionRef<void(IGraphMessage&)> InFunction) const
+{
+	if(const FMessageStackEntry* StackEntryPtr = MessageStacks.Find(InMessageType))
+	{
+		const FMessageStackEntry& StackEntry = *StackEntryPtr;
+		if(StackEntry.Num())
+		{
+			const FMessageEntry& MessageEntry = StackEntry.Top();
+			if(MessageEntry.Message->GetTypeName() == InMessageType)
+			{
+				InFunction(MessageEntry.Message.Get());
+			}
+		}
+	}
+}
+
+void FMessageStack::BottomMessage(FName InMessageType, TFunctionRef<void(IGraphMessage&)> InFunction) const
+{
+	if(const FMessageStackEntry* StackEntryPtr = MessageStacks.Find(InMessageType))
+	{
+		const FMessageStackEntry& StackEntry = *StackEntryPtr;
+		if(StackEntry.Num())
+		{
+			const FMessageEntry& MessageEntry = StackEntry[0];
+			if(MessageEntry.Message->GetTypeName() == InMessageType)
+			{
+				InFunction(MessageEntry.Message.Get());
+			}
+		}
+	}
+}
+
+bool FMessageStack::HasMessage(FName InMessageType) const
+{
+	return MessageStacks.Contains(InMessageType);
+}
+
+void FMessageStack::PushTag(FAnimNode_Base* InNode, const UScriptStruct* InStruct, FName InTag)
+{
+	TagStacks.FindOrAdd(InTag).Push( { InNode, InStruct } );
+}
+
+void FMessageStack::PopTag(FName InTag)
+{
+	TagStacks.FindChecked(InTag).Pop();
+}
+
+void FMessageStack::ForEachTag(FName InTagId, TFunctionRef<EEnumerate(FNodeInfo)> InFunction) const
+{
+	if(const FTagStackEntry* TagStackPtr = TagStacks.Find(InTagId))
+	{
+		const FTagStackEntry& StackEntry = *TagStackPtr;
+		const int32 StackSize = StackEntry.Num();
+		for(int32 TagIndex = StackSize - 1; TagIndex >= 0; --TagIndex)
+		{
+			const FNodeInfo& NodeInfo = StackEntry[TagIndex];
+			if(InFunction(NodeInfo) == EEnumerate::Stop)
+			{
+				return;
+			}
+		}
+	}
+}
+
+void FMessageStack::CopyForCachedUpdate(const FMessageStack& InStack)
+{
+	MessageStacks.Reset();
+	MessageStacks.Reserve(InStack.MessageStacks.Num());
+	for (const TPair<FName, FMessageStackEntry>& MessageStackPair : InStack.MessageStacks)
+	{
+		if (MessageStackPair.Value.Num() > 0)
+		{
+			FMessageStackEntry& Stack = MessageStacks.Add(MessageStackPair.Key);
+			Stack.Push(MessageStackPair.Value.Top());
+		}
+	}
+
+	TagStacks.Reset();
+	TagStacks.Reserve(InStack.TagStacks.Num());
+	for (const TPair<FName, FTagStackEntry>& TagStackPair : InStack.TagStacks)
+	{
+		if (TagStackPair.Value.Num() > 0)
+		{
+			FTagStackEntry& Stack = TagStacks.Add(TagStackPair.Key);
+			Stack.Push(TagStackPair.Value.Top());
+		}
+	}
+}
+
+}}	// namespace UE::Anim
\ No newline at end of file
diff --git a/Engine/Source/Runtime/Engine/Private/Animation/AnimNodeSpaceConversions.cpp b/Engine/Source/Runtime/Engine/Private/Animation/AnimNodeSpaceConversions.cpp
--- a/Engine/Source/Runtime/Engine/Private/Animation/AnimNodeSpaceConversions.cpp
+++ b/Engine/Source/Runtime/Engine/Private/Animation/AnimNodeSpaceConversions.cpp
@@ -28,10 +28,10 @@ void FAnimNode_ConvertComponentToLocalSpace::Evaluate_AnyThread(FPoseContext & O
 {
 	// Evaluate the child and convert
 	FComponentSpacePoseContext InputCSPose(Output.AnimInstanceProxy);
-#if ANIM_NODE_IDS_AVAILABLE
-	// If trace is enabled, we need to preserve the node ID chain as we use the proxy-based constructor above
+
+	// We need to preserve the node ID chain as we use the proxy-based constructor above
 	InputCSPose.SetNodeIds(Output);
-#endif
+
 	ComponentPose.EvaluateComponentSpace(InputCSPose);
 
 	checkSlow( InputCSPose.Pose.GetPose().IsValid() );
@@ -80,10 +80,10 @@ void FAnimNode_ConvertLocalToComponentSpace::EvaluateComponentSpace_AnyThread(FC
 {
 	// Evaluate the child and convert
 	FPoseContext InputPose(OutputCSPose.AnimInstanceProxy);
-#if ANIM_NODE_IDS_AVAILABLE
-	// If trace is enabled, we need to preserve the node ID chain as we use the proxy-based constructor above
+
+	// We need to preserve the node ID chain as we use the proxy-based constructor above
 	InputPose.SetNodeIds(OutputCSPose);
-#endif
+
 	LocalPose.Evaluate(InputPose);
 
 	OutputCSPose.Pose.InitPose(MoveTemp(InputPose.Pose));
diff --git a/Engine/Source/Runtime/Engine/Private/Animation/AnimNode_Inertialization.cpp b/Engine/Source/Runtime/Engine/Private/Animation/AnimNode_Inertialization.cpp
--- a/Engine/Source/Runtime/Engine/Private/Animation/AnimNode_Inertialization.cpp
+++ b/Engine/Source/Runtime/Engine/Private/Animation/AnimNode_Inertialization.cpp
@@ -3,9 +3,12 @@
 #include "Animation/AnimNode_Inertialization.h"
 #include "Animation/AnimInstanceProxy.h"
 #include "AnimationRuntime.h"
+#include "Animation/AnimNodeMessages.h"
+#include "Animation/AnimNode_SaveCachedPose.h"
 
 #define LOCTEXT_NAMESPACE "AnimNode_Inertialization"
 
+IMPLEMENT_ANIMGRAPH_MESSAGE(UE::Anim::IInertializationRequester);
 
 TAutoConsoleVariable<int32> CVarAnimInertializationEnable(TEXT("a.AnimNode.Inertialization.Enable"), 1, TEXT("Enable / Disable Inertialization"));
 TAutoConsoleVariable<int32> CVarAnimInertializationIgnoreVelocity(TEXT("a.AnimNode.Inertialization.IgnoreVelocity"), 0, TEXT("Ignore velocity information during Inertialization (effectively reverting to a quintic diff blend)"));
@@ -15,6 +18,25 @@ TAutoConsoleVariable<int32> CVarAnimInertializationIgnoreDeficit(TEXT("a.AnimNod
 static constexpr int32 INERTIALIZATION_MAX_POSE_SNAPSHOTS = 2;
 static constexpr float INERTIALIZATION_TIME_EPSILON = 1.0e-7f;
 
+namespace UE { namespace Anim {
+
+// Inertialization request event bound to a node
+class FInertializationRequester : public IInertializationRequester
+{
+public:
+	FInertializationRequester(FAnimNode_Inertialization* InNode)
+		: Node(*InNode)
+	{}
+
+private:
+	// IInertializationRequester interface
+	virtual void RequestInertialization(float InRequestedDuration) { Node.RequestInertialization(InRequestedDuration); }
+
+	// Node to target
+	FAnimNode_Inertialization& Node;
+};
+
+}}	// namespace UE::Anim
 
 FAnimNode_Inertialization::FAnimNode_Inertialization()
 	: DeltaTime(0.0f)
@@ -81,7 +103,30 @@ void FAnimNode_Inertialization::Update_AnyThread(const FAnimationUpdateContext&
 {
 	DECLARE_SCOPE_HIERARCHICAL_COUNTER_ANIMNODE(Update_AnyThread);
 
-	FScopedAnimNodeTracker Tracked = Context.TrackAncestor(this);
+	// Allow nodes further towards the leaves to inertialize using this node
+	UE::Anim::TScopedGraphMessage<UE::Anim::FInertializationRequester> Inertialization(Context, this);
+
+	// Handle skipped updates for cached poses by forwarding to inertialization nodes in those residual stacks
+	UE::Anim::TScopedGraphMessage<UE::Anim::FCachedPoseSkippedUpdateHandler> CachedPoseSkippedUpdate(Context, [this](TArrayView<const UE::Anim::FMessageStack> InSkippedUpdates)
+	{
+		// If we have a pending request forward the request to other Inertialization nodes
+		// that were skipped due to pose caching.
+		if(RequestedDuration >= 0.0f)
+		{
+			// Cached poses have their Update function called once even though there may be multiple UseCachedPose nodes for the same pose.
+			// Because of this, there may be Inertialization ancestors of the UseCachedPose nodes that missed out on requests.
+			// So here we forward 'this' node's requests to the ancestors of those skipped UseCachedPose nodes.
+			for (const UE::Anim::FMessageStack& Stack : InSkippedUpdates)
+			{
+				Stack.ForEachMessage<UE::Anim::IInertializationRequester>([this](UE::Anim::IInertializationRequester& InMessage)
+				{
+					InMessage.RequestInertialization(RequestedDuration);
+
+					return UE::Anim::FMessageStack::EEnumerate::Stop;
+				});
+			}
+		}
+	});
 
 	Source.Update(Context);
 
@@ -288,30 +333,6 @@ void FAnimNode_Inertialization::ResetDynamics(ETeleportType InTeleportType)
 }
 
 
-bool FAnimNode_Inertialization::WantsSkippedUpdates() const
-{
-	// We want to receive an OnUpdatesSkipped callback if we have a pending request.
-	// This will give us a chance to forward the request to other Inertialization nodes
-	// that were skipped due to pose caching.
-	return RequestedDuration >= 0.0f;
-}
-
-
-void FAnimNode_Inertialization::OnUpdatesSkipped(TArrayView<const FAnimationUpdateContext*> SkippedUpdateContexts)
-{
-	// Cached poses have their Update function called once even though there may be multiple UseCachedPose nodes for the same pose.
-	// Because of this, there may be Inertialization ancestors of the UseCachedPose nodes that missed out on requests.
-	// So here we forward 'this' node's requests to the ancestors of those skipped UseCachedPose nodes.
-	for (const FAnimationUpdateContext* Update : SkippedUpdateContexts)
-	{
-		if (FAnimNode_Inertialization* OtherIneritalizationNode = Update->GetAncestor<FAnimNode_Inertialization>())
-		{
-			OtherIneritalizationNode->RequestInertialization(RequestedDuration);
-		}
-	}
-}
-
-
 float FAnimNode_Inertialization::ConsumeInertializationRequest(FPoseContext& Context)
 {
 	const float Duration = RequestedDuration;
@@ -778,4 +799,4 @@ float FInertializationPoseDiff::CalcInertialFloat(float x0, float v0, float t, f
 	return x * sign;
 }
 
-#undef LOCTEXT_NAMESPACE
+#undef LOCTEXT_NAMESPACE
\ No newline at end of file
diff --git a/Engine/Source/Runtime/Engine/Private/Animation/AnimNode_LinkedAnimGraph.cpp b/Engine/Source/Runtime/Engine/Private/Animation/AnimNode_LinkedAnimGraph.cpp
--- a/Engine/Source/Runtime/Engine/Private/Animation/AnimNode_LinkedAnimGraph.cpp
+++ b/Engine/Source/Runtime/Engine/Private/Animation/AnimNode_LinkedAnimGraph.cpp
@@ -111,9 +111,9 @@ void FAnimNode_LinkedAnimGraph::Update_AnyThread(const FAnimationUpdateContext&
 		// in USkeletalMeshComponent::TickAnimation. It used to be the case that we could do non-parallel work in 
 		// USkeletalMeshComponent::TickAnimation, which would mean we would have to skip doing that work here.
 		FAnimationUpdateContext NewContext = InContext.WithOtherProxy(&Proxy);
- #if ANIM_NODE_IDS_AVAILABLE
+
 		NewContext = NewContext.WithNodeId(CachedLinkedNodeIndex);
- #endif
+
 		Proxy.UpdateAnimation_WithRoot(NewContext, LinkedRoot, GetDynamicLinkFunctionName());
 	}
 	else if(InputPoses.Num() > 0)
@@ -126,10 +126,10 @@ void FAnimNode_LinkedAnimGraph::Update_AnyThread(const FAnimationUpdateContext&
 	// Consume pending inertial blend request
 	if(PendingBlendDuration >= 0.0f)
 	{
-		FAnimNode_Inertialization* InertializationNode = InContext.GetAncestor<FAnimNode_Inertialization>();
-		if(InertializationNode)
+		UE::Anim::IInertializationRequester* InertializationRequester = InContext.GetMessage<UE::Anim::IInertializationRequester>();
+		if(InertializationRequester)
 		{
-			InertializationNode->RequestInertialization(PendingBlendDuration);
+			InertializationRequester->RequestInertialization(PendingBlendDuration);
 		}
 		else if ((PendingBlendDuration != 0.0f) && (InputPoses.Num() > 0))
 		{
diff --git a/Engine/Source/Runtime/Engine/Private/Animation/AnimNode_LinkedInputPose.cpp b/Engine/Source/Runtime/Engine/Private/Animation/AnimNode_LinkedInputPose.cpp
--- a/Engine/Source/Runtime/Engine/Private/Animation/AnimNode_LinkedInputPose.cpp
+++ b/Engine/Source/Runtime/Engine/Private/Animation/AnimNode_LinkedInputPose.cpp
@@ -29,9 +29,9 @@ void FAnimNode_LinkedInputPose::Update_AnyThread(const FAnimationUpdateContext&
 	if(InputProxy)
 	{
 		FAnimationUpdateContext InputContext = Context.WithOtherProxy(InputProxy);
-#if ANIM_NODE_IDS_AVAILABLE
+
 		InputContext = InputContext.WithNodeId(OuterGraphNodeIndex);
-#endif
+
 		InputPose.Update(InputContext);
 	}
 }
diff --git a/Engine/Source/Runtime/Engine/Private/Animation/AnimNode_SaveCachedPose.cpp b/Engine/Source/Runtime/Engine/Private/Animation/AnimNode_SaveCachedPose.cpp
--- a/Engine/Source/Runtime/Engine/Private/Animation/AnimNode_SaveCachedPose.cpp
+++ b/Engine/Source/Runtime/Engine/Private/Animation/AnimNode_SaveCachedPose.cpp
@@ -4,6 +4,8 @@
 #include "Animation/AnimInstanceProxy.h"
 #include "Animation/AnimTrace.h"
 
+IMPLEMENT_ANIMGRAPH_MESSAGE(UE::Anim::FCachedPoseSkippedUpdateHandler);
+
 /////////////////////////////////////////////////////
 // FAnimNode_SaveCachedPose
 
@@ -48,7 +50,6 @@ void FAnimNode_SaveCachedPose::Update_AnyThread(const FAnimationUpdateContext& C
 	{
 		CachedUpdate.SharedContext = MakeShared<FAnimationUpdateSharedContext>();
 		CachedUpdate.SharedContext->CopyForCachedUpdate(*SharedContext);
-
 	}
 
 	// Store this context for the post update
@@ -108,51 +109,36 @@ void FAnimNode_SaveCachedPose::PostGraphUpdate()
 			}
 		}
 
-		// Update the max weighted pose node
-		{
-			TRACE_SCOPED_ANIM_NODE(CachedUpdateContexts[MaxWeightIdx].Context);
-			Pose.Update(CachedUpdateContexts[MaxWeightIdx].Context);
-		}
-
-		// Determine if any ancestors are interested in the other updates we'll be skipping
-		TArray<FAnimNode_Base*, TInlineAllocator<4>> AncestorsWithSkippedUpdateHandlers;
-		if (NumContexts > 1)
+		// Update any branches that will be skipped
+		UE::Anim::FMessageStack& MessageStack = CachedUpdateContexts[MaxWeightIdx].SharedContext->MessageStack;
+		if(MessageStack.HasMessage<UE::Anim::FCachedPoseSkippedUpdateHandler>())
 		{
-			FAnimationUpdateSharedContext* SharedContext = CachedUpdateContexts[MaxWeightIdx].Context.GetSharedContext();
-			FAnimNodeTracker* AncestorTracker = SharedContext ? &SharedContext->AncestorTracker : nullptr;
-
-			if (AncestorTracker)
-			{
-				for (auto Iter : AncestorTracker->Map)
-				{
-					FAnimNode_Base* AncestorNode = Iter.Value.Num() ? Iter.Value.Top() : nullptr;
-					if (AncestorNode && AncestorNode->WantsSkippedUpdates())
-					{
-						AncestorsWithSkippedUpdateHandlers.Add(AncestorNode);
-					}
-				}
-			}
-		}
-
-		if (AncestorsWithSkippedUpdateHandlers.Num() > 0)
-		{
-			// Build a list of skipped updates
-			TArray<const FAnimationUpdateContext *, TInlineAllocator<4>> SkippedUpdateContexts;
+			// Grab handles to all the execution paths that we are not proceeding with
+			TArray<UE::Anim::FMessageStack, TInlineAllocator<4>> SkippedMessageStacks;
 			for (int32 CurrIdx = 0; CurrIdx < NumContexts; ++CurrIdx)
 			{
 				if (CurrIdx != MaxWeightIdx)
 				{
-					SkippedUpdateContexts.Add(&CachedUpdateContexts[CurrIdx].Context);
+					SkippedMessageStacks.Add(MoveTemp(CachedUpdateContexts[CurrIdx].SharedContext->MessageStack));
 				}
 			}
 
-			// Inform any interested ancestors about the skipped updates
-			for (FAnimNode_Base* AncestorNode : AncestorsWithSkippedUpdateHandlers)
+			// Broadcast skipped update message to interested parties
+			MessageStack.ForEachMessage<UE::Anim::FCachedPoseSkippedUpdateHandler>([&SkippedMessageStacks](UE::Anim::FCachedPoseSkippedUpdateHandler& InMessage)
 			{
-				AncestorNode->OnUpdatesSkipped(SkippedUpdateContexts);
-			}
+				// Fire the event to inform listeners of 'skipped' paths
+				InMessage.OnUpdatesSkipped(SkippedMessageStacks);
+
+				// We only want the topmost registered node here, so early out
+				return UE::Anim::FMessageStack::EEnumerate::Stop;
+			});
 		}
 
+		// Update the max weighted pose node
+		{
+			TRACE_SCOPED_ANIM_NODE(CachedUpdateContexts[MaxWeightIdx].Context);
+			Pose.Update(CachedUpdateContexts[MaxWeightIdx].Context);
+		}
 	}
 
 	CachedUpdateContexts.Reset();
diff --git a/Engine/Source/Runtime/Engine/Private/Animation/AnimNode_SequencePlayer.cpp b/Engine/Source/Runtime/Engine/Private/Animation/AnimNode_SequencePlayer.cpp
--- a/Engine/Source/Runtime/Engine/Private/Animation/AnimNode_SequencePlayer.cpp
+++ b/Engine/Source/Runtime/Engine/Private/Animation/AnimNode_SequencePlayer.cpp
@@ -83,7 +83,7 @@ void FAnimNode_SequencePlayer::UpdateAssetPlayer(const FAnimationUpdateContext&
 		CreateTickRecordForNode(Context, Sequence, bLoopAnimation, AdjustedPlayRate);
 	}
 
-#if ANIM_NODE_IDS_AVAILABLE && WITH_EDITORONLY_DATA
+#if WITH_EDITORONLY_DATA
 	if (FAnimBlueprintDebugData* DebugData = Context.AnimInstanceProxy->GetAnimBlueprintDebugData())
 	{
 		DebugData->RecordSequencePlayer(Context.GetCurrentNodeId(), GetAccumulatedTime(), Sequence != nullptr ? Sequence->GetPlayLength() : 0.0f, Sequence != nullptr ? Sequence->GetNumberOfSampledKeys() : 0);
diff --git a/Engine/Source/Runtime/Engine/Private/Animation/AnimNode_StateMachine.cpp b/Engine/Source/Runtime/Engine/Private/Animation/AnimNode_StateMachine.cpp
--- a/Engine/Source/Runtime/Engine/Private/Animation/AnimNode_StateMachine.cpp
+++ b/Engine/Source/Runtime/Engine/Private/Animation/AnimNode_StateMachine.cpp
@@ -522,10 +522,10 @@ void FAnimNode_StateMachine::Update_AnyThread(const FAnimationUpdateContext& Con
 
 				if (ReferenceTransition.LogicType == ETransitionLogicType::TLT_Inertialization)
 				{
-					FAnimNode_Inertialization* InertializationNode = Context.GetAncestor<FAnimNode_Inertialization>();
-					if (InertializationNode)
+					UE::Anim::IInertializationRequester* InertializationRequester = Context.GetMessage<UE::Anim::IInertializationRequester>();
+					if (InertializationRequester)
 					{
-						InertializationNode->RequestInertialization(ReferenceTransition.CrossfadeDuration);
+						InertializationRequester->RequestInertialization(ReferenceTransition.CrossfadeDuration);
 					}
 					else
 					{
diff --git a/Engine/Source/Runtime/Engine/Public/Animation/AnimNodeMessages.h b/Engine/Source/Runtime/Engine/Public/Animation/AnimNodeMessages.h
new file mode 100644
--- /dev/null
+++ b/Engine/Source/Runtime/Engine/Public/Animation/AnimNodeMessages.h
@@ -0,0 +1,248 @@
+// Copyright Epic Games, Inc. All Rights Reserved.
+
+#pragma once
+
+#include "CoreMinimal.h"
+#include "Templates/UnrealTypeTraits.h"
+
+struct FAnimationUpdateSharedContext;
+struct FAnimationBaseContext;
+struct FAnimNode_Base;
+
+// Simple RTTI implementation for graph messages
+
+// Macro to be used in declaration of any new IGraphMessage types
+#define DECLARE_ANIMGRAPH_MESSAGE(ClassName) \
+	public: \
+		static FName GetStaticTypeName() { return ClassName::TypeName; } \
+		virtual FName GetTypeName() const final { return ClassName::GetStaticTypeName(); } \
+	private: \
+		static const FName TypeName;
+
+// Macro to be used in the implementation of any new IGraphMessage types
+#define IMPLEMENT_ANIMGRAPH_MESSAGE(ClassName) \
+	const FName ClassName::TypeName = TEXT(#ClassName);
+
+namespace UE { namespace Anim {
+
+// Base class for all messages/events/scopes that are fired during execution
+// Note that these messages are events/callbacks only, no blending is involved.
+class ENGINE_API IGraphMessage
+{
+public:
+	virtual ~IGraphMessage() = default;
+
+public:
+	// RTTI support functions
+	static FName GetStaticTypeName()
+	{ 
+		return BaseTypeName;
+	}
+
+	virtual FName GetTypeName() const = 0;
+
+	template<typename Type> 
+	bool Is() const 
+	{ 
+		return GetTypeName() == Type::GetStaticTypeName(); 
+	}
+
+	template<typename Type> 
+	const Type& As() const 
+	{ 
+		return *static_cast<const Type*>(this);
+	}
+
+	template<typename Type> 
+	Type& As() 
+	{ 
+		return *static_cast<Type*>(this);
+	}
+
+private:
+	static const FName BaseTypeName;
+};
+
+// Pushes a simple tag onto the shared context stack
+struct ENGINE_API FScopedGraphTag
+{
+	FScopedGraphTag(const FAnimationBaseContext& InContext, FName InTag);
+	~FScopedGraphTag();
+
+private:
+	// Shared context that we push messages into
+	FAnimationUpdateSharedContext* SharedContext;
+
+	// The tag that we pushed
+	FName Tag;
+};
+
+// Base helper class
+struct ENGINE_API FScopedGraphMessage
+{
+protected:
+	FScopedGraphMessage(const FAnimationBaseContext& InContext);
+
+	// Helper functions
+	void PushMessage(const FAnimationBaseContext& InContext, TSharedRef<IGraphMessage> InMessage, FName InMessageType);
+	void PopMessage();
+
+	// Shared context that we push messages into
+	FAnimationUpdateSharedContext* SharedContext;
+
+	// The type of the message we pushed
+	FName MessageType;
+};
+
+// Pushes a message onto the shared context stack
+template<typename TGraphMessageType>
+struct TScopedGraphMessage : FScopedGraphMessage
+{
+	template<typename... TArgs>
+	TScopedGraphMessage(const FAnimationBaseContext& InContext, TArgs&&... Args)
+		: FScopedGraphMessage(InContext)
+	{
+		static_assert(TIsDerivedFrom<TGraphMessageType, IGraphMessage>::IsDerived, "Argument TGraphMessageType must derive from IGraphMessage");
+		static_assert(!TIsSame<TGraphMessageType, IGraphMessage>::Value, "Argument TGraphMessageType must not be IGraphMessage");
+
+		TSharedRef<TGraphMessageType> Message = MakeShared<TGraphMessageType>(Forward<TArgs>(Args)...);
+		PushMessage(InContext, Message, Message->GetTypeName());
+	}
+
+	~TScopedGraphMessage()
+	{
+		PopMessage();
+	}
+};
+
+// Stack of tags & events used to track context during graph execution
+struct ENGINE_API FMessageStack
+{
+public:
+	friend struct FScopedGraphTag;
+	friend struct FScopedGraphMessage;
+
+	FMessageStack() = default;
+
+	// Non-copyable
+	FMessageStack(FMessageStack&) = delete;
+	FMessageStack& operator=(const FMessageStack&) = delete;
+
+	FMessageStack(FMessageStack&& InMessageStack) = default;
+	FMessageStack& operator=(FMessageStack&&) = default;
+
+	// Return value for the various ForEach* enumeration functions
+	enum class EEnumerate
+	{
+		Stop,
+		Continue,
+	};
+
+	// Info about a node
+	struct FNodeInfo
+	{
+		FAnimNode_Base* Node;
+		const UScriptStruct* NodeStruct;
+	};
+
+	// Call the supplied function with each node subscribed to the specified message in this stack
+	// @param	InFunction		The function to call with the specified message. Function should return whether to continue or stop the enumeration.
+	template<typename TGraphMessageType>
+	void ForEachMessage(TFunctionRef<EEnumerate(TGraphMessageType&)> InFunction) const
+	{
+		static_assert(TIsDerivedFrom<TGraphMessageType, IGraphMessage>::IsDerived, "Argument TGraphMessageType must derive from IGraphMessage");
+		static_assert(!TIsSame<TGraphMessageType, IGraphMessage>::Value, "Argument TGraphMessageType must not be IGraphMessage");
+
+		ForEachMessage(TGraphMessageType::GetStaticTypeName(), [&InFunction](IGraphMessage& InMessage)
+		{
+			return InFunction(static_cast<TGraphMessageType&>(InMessage));
+		});
+	}
+
+	// Call the supplied function on the top-most node subscribed to the specified message in this stack
+	// @param	InFunction		The function to call with the specified message.
+	template<typename TGraphMessageType>
+	void TopMessage(TFunctionRef<void(TGraphMessageType&)> InFunction) const
+	{
+		static_assert(TIsDerivedFrom<TGraphMessageType, IGraphMessage>::IsDerived, "Argument TGraphMessageType must derive from IGraphMessage");
+		static_assert(!TIsSame<TGraphMessageType, IGraphMessage>::Value, "Argument TGraphMessageType must not be IGraphMessage");
+
+		TopMessage(TGraphMessageType::GetStaticTypeName(), [&InFunction](IGraphMessage& InMessage)
+		{
+			InFunction(static_cast<TGraphMessageType&>(InMessage));
+		});
+	}
+
+	// @return true if a message of the specified type is present
+	template<typename TGraphMessageType>
+	bool HasMessage() const
+	{
+		static_assert(TIsDerivedFrom<TGraphMessageType, IGraphMessage>::IsDerived, "Argument TGraphMessageType must derive from IGraphMessage");
+		static_assert(!TIsSame<TGraphMessageType, IGraphMessage>::Value, "Argument TGraphMessageType must not be IGraphMessage");
+
+		return HasMessage(TGraphMessageType::GetStaticTypeName());
+	}
+
+	// Call the supplied function with each node tagged with the specified tag
+	// @param	InTagId			The tag to check
+	// @param	InFunction		The function to call with the specified tagged node. Function should return whether to continue or stop the enumeration.
+	void ForEachTag(FName InTagId, TFunctionRef<EEnumerate(FNodeInfo)> InFunction) const;
+
+	// Copies the relevant parts of each stack only, ready for cached update.
+	// @param	InStack		The stack to copy from
+	void CopyForCachedUpdate(const FMessageStack& InStack);
+
+private:
+	// Push a message onto the stack
+	// @param	InNode		The node that this message was pushed from
+	// @param	InMessage	The message to push
+	void PushMessage(FName InMessageType, FAnimNode_Base* InNode, const UScriptStruct* InStruct, TSharedRef<IGraphMessage> InMessage);
+
+	// Pop a message off the stack
+	void PopMessage(FName InMessageType);
+
+	// Helper function for the templated function of the same name above
+	void ForEachMessage(FName InMessageType, TFunctionRef<EEnumerate(IGraphMessage&)> InFunction) const;
+
+	// Helper function for the templated function of the same name above
+	void TopMessage(FName InMessageType, TFunctionRef<void(IGraphMessage&)> InFunction) const;
+
+	// Helper function for the templated function of the same name above
+	void BottomMessage(FName InMessageType, TFunctionRef<void(IGraphMessage&)> InFunction) const;
+
+	// @return true if a message of the specified type is present
+	bool HasMessage(FName InMessageType) const;
+
+	// Push a tag onto the stack
+	// @param	InNode	The node that this tag was pushed from
+	// @param	InTag	The tag to push
+	void PushTag(FAnimNode_Base* InNode, const UScriptStruct* InStruct, FName InTag);
+
+	// Pop a tag off the stack
+	// @param	InTag	The tag to pop
+	void PopTag(FName InTag);
+
+private:
+	// Holds a message that has been pushed onto the stack.
+	// Message is held by a shared ptr because the same stack scope can potentially be
+	// referenced by multiple cached poses (pesky DAG)
+	struct FMessageEntry
+	{
+		TSharedRef<IGraphMessage> Message;
+		FNodeInfo NodeInfo;
+	};
+
+	using FMessageStackEntry = TArray<FMessageEntry, TInlineAllocator<4>>;
+	using FMessageMap = TMap<FName, FMessageStackEntry, TInlineSetAllocator<4>>;
+
+	// Message stack
+	FMessageMap MessageStacks;
+
+	using FTagStackEntry = TArray<FNodeInfo, TInlineAllocator<4>>;
+	using FTagMap = TMap<FName, FTagStackEntry, TInlineSetAllocator<4>>;
+
+	// Tag stack
+	FTagMap TagStacks;
+};
+
+}}	// namespace UE::Anim
\ No newline at end of file
