Move the mips reading to the async tasks to avoid stalling on the game-thread for long periods when processing 8k+ textures
Fix potential race in GetMipData between game-thread and async tasks

diff --git a/Engine/Source/Runtime/Engine/Classes/Engine/Texture.h b/Engine/Source/Runtime/Engine/Classes/Engine/Texture.h
--- a/Engine/Source/Runtime/Engine/Classes/Engine/Texture.h
+++ b/Engine/Source/Runtime/Engine/Classes/Engine/Texture.h
@@ -333,6 +333,10 @@ private:
 	friend class UVolumeTexture;
 	friend class UTexture2DArray;
 
+#if WITH_EDITOR
+	/** Protects simultaneous access to BulkData */
+	FRWLock       BulkDataLock;
+#endif
 	/** The bulk source data. */
 	FByteBulkData BulkData;
 	/** Pointer to locked mip data, if any. */
diff --git a/Engine/Source/Runtime/Engine/Private/Texture.cpp b/Engine/Source/Runtime/Engine/Private/Texture.cpp
--- a/Engine/Source/Runtime/Engine/Private/Texture.cpp
+++ b/Engine/Source/Runtime/Engine/Private/Texture.cpp
@@ -27,6 +27,10 @@
 #include "Interfaces/ITextureFormat.h"
 #include "Interfaces/ITextureFormatModule.h"
 
+#if WITH_EDITOR
+#include "Misc/ScopeRWLock.h"
+#endif
+
 #if WITH_EDITORONLY_DATA
 	#include "EditorFramework/AssetImportData.h"
 #endif
@@ -1165,6 +1169,9 @@ bool FTextureSource::GetMipData(TArray64<uint8>& OutMipData, int32 BlockIndex, i
 	bool bSuccess = false;
 	if (BlockIndex < GetNumBlocks() && LayerIndex < NumLayers && MipIndex < NumMips && BulkData.GetBulkDataSize() > 0)
 	{
+#if WITH_EDITOR
+		FWriteScopeLock BulkDataReadScope(BulkDataLock);
+#endif
 		void* RawSourceData = BulkData.Lock(LOCK_READ_ONLY);
 		if (bPNGCompressed)
 		{
diff --git a/Engine/Source/Runtime/Engine/Private/TextureDerivedDataTask.cpp b/Engine/Source/Runtime/Engine/Private/TextureDerivedDataTask.cpp
--- a/Engine/Source/Runtime/Engine/Private/TextureDerivedDataTask.cpp
+++ b/Engine/Source/Runtime/Engine/Private/TextureDerivedDataTask.cpp
@@ -396,20 +396,6 @@ FTextureCacheDerivedDataWorker::FTextureCacheDerivedDataWorker(
 			CompositeTextureData.Init(*Texture.CompositeTexture, BuildSettingsPerLayer.GetData(), bAllowAsyncLoading);
 		}
 	}
-
-	// If the bulkdata is loaded and async build is allowed, get the source mips now (safe) to allow building the DDC if required.
-	// If the bulkdata is not loaded, the DDC will be built in Finalize() unless async loading is enabled (which won't allow reuse of the source for later use).
-	if (bAllowAsyncBuild)
-	{
-		if (TextureData.IsValid() && Texture.Source.IsBulkDataLoaded())
-		{
-			TextureData.GetSourceMips(Texture.Source, ImageWrapper);
-		}
-		if (CompositeTextureData.IsValid() && Texture.CompositeTexture && Texture.CompositeTexture->Source.IsBulkDataLoaded())
-		{
-			CompositeTextureData.GetSourceMips(Texture.CompositeTexture->Source, ImageWrapper);
-		}
-	}
 }
 
 void FTextureCacheDerivedDataWorker::DoWork()
@@ -492,9 +478,26 @@ void FTextureCacheDerivedDataWorker::DoWork()
 	
 	if (!bSucceeded && bAllowAsyncBuild)
 	{
-		if (bAllowAsyncLoading)
+		bool bHasTextureSourceMips = false;
+		if (TextureData.IsValid() && Texture.Source.IsBulkDataLoaded())
+		{
+			TextureData.GetSourceMips(Texture.Source, ImageWrapper);
+			bHasTextureSourceMips = true;
+		}
+
+		bool bHasCompositeTextureSourceMips = false;
+		if (CompositeTextureData.IsValid() && Texture.CompositeTexture && Texture.CompositeTexture->Source.IsBulkDataLoaded())
+		{
+			CompositeTextureData.GetSourceMips(Texture.CompositeTexture->Source, ImageWrapper);
+			bHasCompositeTextureSourceMips = true;
+		}
+
+		if (bAllowAsyncLoading && !bHasTextureSourceMips)
 		{
 			TextureData.GetAsyncSourceMips(ImageWrapper);
+		}
+		if (bAllowAsyncLoading && !bHasCompositeTextureSourceMips)
+		{
 			CompositeTextureData.GetAsyncSourceMips(ImageWrapper);
 		}
 
