Fix compiler breaks for non editor builds.

diff --git a/Engine/Source/Runtime/Engine/Private/ComputeFramework/ComputeKernel.cpp b/Engine/Source/Runtime/Engine/Private/ComputeFramework/ComputeKernel.cpp
--- a/Engine/Source/Runtime/Engine/Private/ComputeFramework/ComputeKernel.cpp
+++ b/Engine/Source/Runtime/Engine/Private/ComputeFramework/ComputeKernel.cpp
@@ -10,10 +10,12 @@ void UComputeKernel::PostLoad()
 {
 	Super::PostLoad();
 
+#if WITH_EDITOR
 	if (FApp::CanEverRender())
 	{
 		CacheResourceShadersForRendering();
 	}
+#endif
 }
 
 #if WITH_EDITOR
diff --git a/Engine/Source/Runtime/Engine/Private/ComputeFramework/ComputeKernelFromText.cpp b/Engine/Source/Runtime/Engine/Private/ComputeFramework/ComputeKernelFromText.cpp
--- a/Engine/Source/Runtime/Engine/Private/ComputeFramework/ComputeKernelFromText.cpp
+++ b/Engine/Source/Runtime/Engine/Private/ComputeFramework/ComputeKernelFromText.cpp
@@ -15,7 +15,9 @@ UComputeKernelFromText::UComputeKernelFromText()
 void UComputeKernelFromText::PostLoad()
 {
 	Super::PostLoad();
+#if WITH_EDITOR
 	ReparseKernelSourceText();
+#endif
 }
 
 #if WITH_EDITOR
diff --git a/Engine/Source/Runtime/Engine/Private/ComputeFramework/ComputeKernelShaderCompilationManager.cpp b/Engine/Source/Runtime/Engine/Private/ComputeFramework/ComputeKernelShaderCompilationManager.cpp
--- a/Engine/Source/Runtime/Engine/Private/ComputeFramework/ComputeKernelShaderCompilationManager.cpp
+++ b/Engine/Source/Runtime/Engine/Private/ComputeFramework/ComputeKernelShaderCompilationManager.cpp
@@ -23,7 +23,7 @@ static FAutoConsoleVariableRef CVarShowComputeKernelShaderWarnings(
 	);
 
 
-
+#if WITH_EDITOR
 FComputeKernelShaderCompilationManager GComputeKernelShaderCompilationManager;
 
 void FComputeKernelShaderCompilationManager::Tick(float DeltaSeconds)
@@ -444,4 +444,4 @@ void FComputeKernelShaderCompilationManager::FinishCompilation(const TCHAR* InKe
 	check(ComputeKernelShaderMapJobs.Num() == 0);
 #endif
 }
-
+#endif
diff --git a/Engine/Source/Runtime/Engine/Private/ComputeFramework/ComputeKernelShaderMap.cpp b/Engine/Source/Runtime/Engine/Private/ComputeFramework/ComputeKernelShaderMap.cpp
--- a/Engine/Source/Runtime/Engine/Private/ComputeFramework/ComputeKernelShaderMap.cpp
+++ b/Engine/Source/Runtime/Engine/Private/ComputeFramework/ComputeKernelShaderMap.cpp
@@ -361,6 +361,7 @@ void FComputeKernelShaderMap::Compile(
 	{
 		UE_LOG(LogShaders, Fatal, TEXT("Trying to compile ComputeKernel shader %s at run-time, which is not supported on consoles!"), *InKernel->GetFriendlyName() );
 	}
+#if WITH_EDITOR
 	else
 	{
 		// Make sure we are operating on a referenced shader map or the below Find will cause this shader map to be deleted,
@@ -463,6 +464,7 @@ void FComputeKernelShaderMap::Compile(
 			}
 		}
 	}
+#endif
 }
 
 FShader* FComputeKernelShaderMap::ProcessCompilationResultsForSingleJob(FShaderCompileJob& CurrentJob, const FSHAHash& InShaderMapHash)
diff --git a/Engine/Source/Runtime/Engine/Public/ComputeFramework/ComputeKernelShaderCompilationManager.h b/Engine/Source/Runtime/Engine/Public/ComputeFramework/ComputeKernelShaderCompilationManager.h
--- a/Engine/Source/Runtime/Engine/Public/ComputeFramework/ComputeKernelShaderCompilationManager.h
+++ b/Engine/Source/Runtime/Engine/Public/ComputeFramework/ComputeKernelShaderCompilationManager.h
@@ -4,7 +4,10 @@
 
 #include "ComputeKernelShared.h"
 #include "ShaderCompiler.h"
+
+#if WITH_EDITOR
 #include "TickableEditorObject.h"
+#endif
 
 
 /** Information tracked for each shader compile worker process instance. */
@@ -77,8 +80,8 @@ struct FComputeKernelShaderMapFinalizeResults : public FComputeKernelShaderMapCo
 };
 
 
+#if WITH_EDITOR
 // handles finished shader compile jobs, applying of the shaders to their config asset, and some error handling
-//
 class FComputeKernelShaderCompilationManager : public FTickableEditorObject
 {
 public:
@@ -119,3 +122,4 @@ private:
 };
 
 extern FComputeKernelShaderCompilationManager GComputeKernelShaderCompilationManager;
+#endif
diff --git a/Engine/Source/Runtime/Engine/Public/ComputeFramework/ShaderParamTypeDefinition.h b/Engine/Source/Runtime/Engine/Public/ComputeFramework/ShaderParamTypeDefinition.h
--- a/Engine/Source/Runtime/Engine/Public/ComputeFramework/ShaderParamTypeDefinition.h
+++ b/Engine/Source/Runtime/Engine/Public/ComputeFramework/ShaderParamTypeDefinition.h
@@ -209,10 +209,8 @@ public:
 		);
 
 public:
-#if WITH_EDITOR
 	UPROPERTY(VisibleAnywhere, Category = "Kernel")
 	FString TypeDeclaration;
-#endif
 
 	UPROPERTY(EditAnywhere, Category = "Kernel")
 	FString	Name;
