[Navigation WP] Properly using the vertical bounds of the collected tiles as the NavigationDataChunkActor bounds

diff --git a/Engine/Source/Runtime/Engine/Classes/AI/NavigationSystemBase.h b/Engine/Source/Runtime/Engine/Classes/AI/NavigationSystemBase.h
--- a/Engine/Source/Runtime/Engine/Classes/AI/NavigationSystemBase.h
+++ b/Engine/Source/Runtime/Engine/Classes/AI/NavigationSystemBase.h
@@ -243,7 +243,7 @@ public:
 	virtual bool ContainsNavData(const FBox& Bounds) const { return false; }
 	virtual void AddNavigationDataChunk(class ANavigationDataChunkActor& DataChunkActor) {}
 	virtual void RemoveNavigationDataChunk(class ANavigationDataChunkActor& DataChunkActor) {}
-	virtual void FillNavigationDataChunkActor(const FBox& Bounds, class ANavigationDataChunkActor& DataChunkActor) {}
+	virtual void FillNavigationDataChunkActor(const FBox& QueryBounds, class ANavigationDataChunkActor& DataChunkActor, FBox& OutTilesBounds) {}
 
 protected:
 	/**	Sets the Transform the Navigation System will use when converting from FromCoordType
diff --git a/Engine/Source/Runtime/Engine/Private/WorldPartition/NavigationData/NavigationDataChunkActor.cpp b/Engine/Source/Runtime/Engine/Private/WorldPartition/NavigationData/NavigationDataChunkActor.cpp
--- a/Engine/Source/Runtime/Engine/Private/WorldPartition/NavigationData/NavigationDataChunkActor.cpp
+++ b/Engine/Source/Runtime/Engine/Private/WorldPartition/NavigationData/NavigationDataChunkActor.cpp
@@ -54,7 +54,7 @@ void ANavigationDataChunkActor::BeginDestroy()
 }
 #endif // WITH_EDITOR
 
-void ANavigationDataChunkActor::CollectNavData(const FBox& Bounds)
+void ANavigationDataChunkActor::CollectNavData(const FBox& QueryBounds, FBox& OutTilesBounds)
 {
 	Log(ANSI_TO_TCHAR(__FUNCTION__));
 
@@ -63,7 +63,7 @@ void ANavigationDataChunkActor::CollectNavData(const FBox& Bounds)
 	{
 		if (UNavigationSystemBase* NavSys = World->GetNavigationSystem())
 		{
-			NavSys->FillNavigationDataChunkActor(Bounds, *this);
+			NavSys->FillNavigationDataChunkActor(QueryBounds, *this, OutTilesBounds);
 		}
 	}
 }
diff --git a/Engine/Source/Runtime/Engine/Private/WorldPartition/WorldPartitionRuntimeSpatialHash.cpp b/Engine/Source/Runtime/Engine/Private/WorldPartition/WorldPartitionRuntimeSpatialHash.cpp
--- a/Engine/Source/Runtime/Engine/Private/WorldPartition/WorldPartitionRuntimeSpatialHash.cpp
+++ b/Engine/Source/Runtime/Engine/Private/WorldPartition/WorldPartitionRuntimeSpatialHash.cpp
@@ -1897,10 +1897,17 @@ bool UWorldPartitionRuntimeSpatialHash::GenerateNavigationData()
 
 		FBox2D CellBounds;
 		GridHelper.GetCellBounds(CellCoord, CellBounds);
-		const float HalfHeight = 100000.f; //Todo AT: check to get the real z bounds from the collected tiles
-		const FBox CollectingBounds(FVector(CellBounds.Min.X, CellBounds.Min.Y, -HalfHeight), FVector(CellBounds.Max.X, CellBounds.Max.Y, HalfHeight));
+		if (CellBounds.GetExtent().X < 1.f || CellBounds.GetExtent().Y < 1.f)
+		{
+			// Safety, since we reduce by 1.f below.
+			UE_LOG(LogWorldPartitionRuntimeSpatialHash, Warning, TEXT("%s: grid cell too small."), ANSI_TO_TCHAR(__FUNCTION__));
+			return;
+		}
 
-		if (NavSystem->ContainsNavData(CollectingBounds) == false)
+		const float HalfHeight = HALF_WORLD_MAX;
+		const FBox QueryBounds(FVector(CellBounds.Min.X, CellBounds.Min.Y, -HalfHeight), FVector(CellBounds.Max.X, CellBounds.Max.Y, HalfHeight));
+
+		if (NavSystem->ContainsNavData(QueryBounds) == false)
 		{
 			// Skip if there is no navdata for this cell
 			return;
@@ -1916,13 +1923,19 @@ bool UWorldPartitionRuntimeSpatialHash::GenerateNavigationData()
 
 		const FVector2D CellCenter = CellBounds.GetCenter();
 		DataChunkActor->SetActorLocation(FVector(CellCenter.X, CellCenter.Y, 0.f));
-		DataChunkActor->SetDataChunkActorBounds(CollectingBounds.ExpandBy(-1.f)); //reduce by 1cm to avoid precision issues
+
+		FBox TilesBounds(EForceInit::ForceInit);
+		DataChunkActor->CollectNavData(QueryBounds, TilesBounds);
+
+		FBox ChunkActorBounds(FVector(QueryBounds.Min.X, QueryBounds.Min.Y, TilesBounds.Min.Z), FVector(QueryBounds.Max.X, QueryBounds.Max.Y, TilesBounds.Max.Z));
+		ChunkActorBounds = ChunkActorBounds.ExpandBy(FVector(-1.f, -1.f, 1.f)); //reduce XY by 1cm to avoid precision issues causing potential overflow on neighboring cell, add 1cm in Z to have a minimum of volume.
+		UE_LOG(LogWorldPartitionRuntimeSpatialHash, VeryVerbose, TEXT("Setting ChunkActorBounds to %s"), *ChunkActorBounds.ToString());
+		DataChunkActor->SetDataChunkActorBounds(ChunkActorBounds);
 
 		const FName CellName = GetCellName(RuntimeGrid.GridName, CellCoord.Z, CellCoord.X, CellCoord.Y);
 		DataChunkActor->SetActorLabel(FString::Printf(TEXT("NavDataChunkActor_%s_%s"), *GetName(), *CellName.ToString()));
 		
 		//@todo_ow: Properly handle data layers
-		DataChunkActor->CollectNavData(CollectingBounds);
 
 		// Set target grid
 		DataChunkActor->RuntimeGrid = RuntimeGrid.GridName;
diff --git a/Engine/Source/Runtime/Engine/Public/WorldPartition/NavigationData/NavigationDataChunkActor.h b/Engine/Source/Runtime/Engine/Public/WorldPartition/NavigationData/NavigationDataChunkActor.h
--- a/Engine/Source/Runtime/Engine/Public/WorldPartition/NavigationData/NavigationDataChunkActor.h
+++ b/Engine/Source/Runtime/Engine/Public/WorldPartition/NavigationData/NavigationDataChunkActor.h
@@ -24,7 +24,7 @@ public:
 	const TArray<UNavigationDataChunk*>& GetNavDataChunk() const { return NavDataChunks; }
 	TArray<UNavigationDataChunk*>& GetMutableNavDataChunk() { return NavDataChunks; }
 
-	void CollectNavData(const FBox& Bounds);
+	void CollectNavData(const FBox& QueryBounds, FBox& OutTilesBounds);
 
 #if WITH_EDITOR
 	void SetDataChunkActorBounds(const FBox& InBounds);
diff --git a/Engine/Source/Runtime/NavigationSystem/Private/NavMesh/RecastNavMesh.cpp b/Engine/Source/Runtime/NavigationSystem/Private/NavMesh/RecastNavMesh.cpp
--- a/Engine/Source/Runtime/NavigationSystem/Private/NavMesh/RecastNavMesh.cpp
+++ b/Engine/Source/Runtime/NavigationSystem/Private/NavMesh/RecastNavMesh.cpp
@@ -2092,13 +2092,13 @@ void ARecastNavMesh::ApplyWorldOffset(const FVector& InOffset, bool bWorldShift)
 	RequestDrawingUpdate();
 }
 
-void ARecastNavMesh::FillNavigationDataChunkActor(const FBox& Bounds, ANavigationDataChunkActor& DataChunkActor) const
+void ARecastNavMesh::FillNavigationDataChunkActor(const FBox& QueryBounds, ANavigationDataChunkActor& DataChunkActor, FBox& OutTilesBounds) const
 {
 	if (RecastNavMeshImpl)
 	{
-		UE_LOG(LogNavigation, Verbose, TEXT("%s Bounds pos: (%s)  size: (%s)."), ANSI_TO_TCHAR(__FUNCTION__), *Bounds.GetCenter().ToString(), *Bounds.GetSize().ToString());
+		UE_LOG(LogNavigation, Verbose, TEXT("%s Bounds pos: (%s)  size: (%s)."), ANSI_TO_TCHAR(__FUNCTION__), *QueryBounds.GetCenter().ToString(), *QueryBounds.GetSize().ToString());
 
-		const TArray<FBox> Boxes({ Bounds });
+		const TArray<FBox> Boxes({ QueryBounds });
 		TArray<int32> TileIndices;
 		RecastNavMeshImpl->GetNavMeshTilesIn(Boxes, TileIndices);
 		if (!TileIndices.IsEmpty())
@@ -2109,6 +2109,7 @@ void ARecastNavMesh::FillNavigationDataChunkActor(const FBox& Bounds, ANavigatio
 			DataChunkActor.GetMutableNavDataChunk().Add(DataChunk);
 
 			DataChunk->GetTiles(RecastNavMeshImpl, TileIndices, EGatherTilesCopyMode::CopyData);
+			DataChunk->GetTilesBounds(*RecastNavMeshImpl, TileIndices, OutTilesBounds);
 		}
 	}
 }
diff --git a/Engine/Source/Runtime/NavigationSystem/Private/NavMesh/RecastNavMeshDataChunk.cpp b/Engine/Source/Runtime/NavigationSystem/Private/NavMesh/RecastNavMeshDataChunk.cpp
--- a/Engine/Source/Runtime/NavigationSystem/Private/NavMesh/RecastNavMeshDataChunk.cpp
+++ b/Engine/Source/Runtime/NavigationSystem/Private/NavMesh/RecastNavMeshDataChunk.cpp
@@ -437,3 +437,18 @@ void URecastNavMeshDataChunk::GetTiles(const FPImplRecastNavMesh* NavMeshImpl, c
 		}
 	}
 }
+
+void URecastNavMeshDataChunk::GetTilesBounds(const FPImplRecastNavMesh& NavMeshImpl, const TArray<int32>& TileIndices, FBox& OutBounds) const
+{
+	OutBounds.Init();
+	const dtNavMesh* NavMesh = NavMeshImpl.DetourNavMesh;
+
+	for (const int32 TileIdx : TileIndices)
+	{
+		const dtMeshTile* Tile = NavMesh->getTile(TileIdx);
+		if (Tile && Tile->header)
+		{
+			OutBounds += Recast2UnrealBox(Tile->header->bmin, Tile->header->bmax);
+		}
+	}
+}
\ No newline at end of file
diff --git a/Engine/Source/Runtime/NavigationSystem/Private/NavigationSystem.cpp b/Engine/Source/Runtime/NavigationSystem/Private/NavigationSystem.cpp
--- a/Engine/Source/Runtime/NavigationSystem/Private/NavigationSystem.cpp
+++ b/Engine/Source/Runtime/NavigationSystem/Private/NavigationSystem.cpp
@@ -1868,13 +1868,13 @@ void UNavigationSystemV1::RemoveNavigationDataChunk(ANavigationDataChunkActor& D
 	}
 }
 
-void UNavigationSystemV1::FillNavigationDataChunkActor(const FBox& Bounds, ANavigationDataChunkActor& DataChunkActor)
+void UNavigationSystemV1::FillNavigationDataChunkActor(const FBox& QueryBounds, ANavigationDataChunkActor& DataChunkActor, FBox& OutTilesBounds)
 {
 	for (const ANavigationData* NavData : NavDataSet)
 	{
 		if (NavData)
 		{
-			NavData->FillNavigationDataChunkActor(Bounds, DataChunkActor);
+			NavData->FillNavigationDataChunkActor(QueryBounds, DataChunkActor, OutTilesBounds);
 		}
 	}
 }
diff --git a/Engine/Source/Runtime/NavigationSystem/Public/NavMesh/RecastNavMesh.h b/Engine/Source/Runtime/NavigationSystem/Public/NavMesh/RecastNavMesh.h
--- a/Engine/Source/Runtime/NavigationSystem/Public/NavMesh/RecastNavMesh.h
+++ b/Engine/Source/Runtime/NavigationSystem/Public/NavMesh/RecastNavMesh.h
@@ -759,7 +759,7 @@ public:
 	/** Called on world origin changes **/
 	virtual void ApplyWorldOffset(const FVector& InOffset, bool bWorldShift) override;
 
-	virtual void FillNavigationDataChunkActor(const FBox& InBounds, class ANavigationDataChunkActor& DataChunkActor) const override;
+	virtual void FillNavigationDataChunkActor(const FBox& InQueryBounds, class ANavigationDataChunkActor& DataChunkActor, FBox& OutTilesBounds) const override;
 
 	virtual void OnStreamingNavDataAdded(class ANavigationDataChunkActor& InActor) override;
 	virtual void OnStreamingNavDataRemoved(class ANavigationDataChunkActor& InActor) override;
diff --git a/Engine/Source/Runtime/NavigationSystem/Public/NavMesh/RecastNavMeshDataChunk.h b/Engine/Source/Runtime/NavigationSystem/Public/NavMesh/RecastNavMeshDataChunk.h
--- a/Engine/Source/Runtime/NavigationSystem/Public/NavMesh/RecastNavMeshDataChunk.h
+++ b/Engine/Source/Runtime/NavigationSystem/Public/NavMesh/RecastNavMeshDataChunk.h
@@ -92,6 +92,9 @@ class NAVIGATIONSYSTEM_API URecastNavMeshDataChunk : public UNavigationDataChunk
 	/** Const accessor to the list of tiles in the data chunk. */
 	const TArray<FRecastTileData>& GetTiles() const { return Tiles; }
 
+	/** Returns the AABB for the given tiles. */
+	void GetTilesBounds(const FPImplRecastNavMesh& NavMeshImpl, const TArray<int32>& TileIndices, FBox& OutBounds) const;
+
 	/** Mutable accessor to the list of tiles in the data chunk. */
 	TArray<FRecastTileData>& GetMutableTiles() { return Tiles; }
 
diff --git a/Engine/Source/Runtime/NavigationSystem/Public/NavigationData.h b/Engine/Source/Runtime/NavigationSystem/Public/NavigationData.h
--- a/Engine/Source/Runtime/NavigationSystem/Public/NavigationData.h
+++ b/Engine/Source/Runtime/NavigationSystem/Public/NavigationData.h
@@ -587,7 +587,7 @@ public:
 	virtual bool SupportsStreaming() const;
 	virtual void OnNavigationBoundsChanged();
 
-	virtual void FillNavigationDataChunkActor(const FBox& Bounds, class ANavigationDataChunkActor& DataChunkActor) const {};
+	virtual void FillNavigationDataChunkActor(const FBox& QueryBounds, class ANavigationDataChunkActor& DataChunkActor, FBox& OutTilesBounds) const {};
 	virtual void OnStreamingNavDataAdded(class ANavigationDataChunkActor& InActor) {};
 	virtual void OnStreamingNavDataRemoved(class ANavigationDataChunkActor& InActor) {};
 	
diff --git a/Engine/Source/Runtime/NavigationSystem/Public/NavigationSystem.h b/Engine/Source/Runtime/NavigationSystem/Public/NavigationSystem.h
--- a/Engine/Source/Runtime/NavigationSystem/Public/NavigationSystem.h
+++ b/Engine/Source/Runtime/NavigationSystem/Public/NavigationSystem.h
@@ -587,7 +587,7 @@ public:
 	virtual bool ContainsNavData(const FBox& Bounds) const override;
 	virtual void AddNavigationDataChunk(class ANavigationDataChunkActor& DataChunkActor) override;
 	virtual void RemoveNavigationDataChunk(class ANavigationDataChunkActor& DataChunkActor) override;
-	virtual void FillNavigationDataChunkActor(const FBox& Bounds, class ANavigationDataChunkActor& DataChunkActor) override;
+	virtual void FillNavigationDataChunkActor(const FBox& QueryBounds, class ANavigationDataChunkActor& DataChunkActor, FBox& OutTilesBounds) override;
 
 	ANavigationData* GetDefaultNavDataInstance() const { return MainNavData; }
 
