Fix Nonunity header.

diff --git a/Engine/Source/Runtime/Engine/Classes/Materials/MaterialExpressionPinBase.h b/Engine/Source/Runtime/Engine/Classes/Materials/MaterialExpressionPinBase.h
--- a/Engine/Source/Runtime/Engine/Classes/Materials/MaterialExpressionPinBase.h
+++ b/Engine/Source/Runtime/Engine/Classes/Materials/MaterialExpressionPinBase.h
@@ -6,6 +6,8 @@
 #include "UObject/ObjectMacros.h"
 #include "Materials/MaterialExpression.h"
 #include "Materials/MaterialExpressionReroute.h"
+#include "EdGraph/EdGraphNode.h"
+
 #include "MaterialExpressionPinBase.generated.h"
 
 /**
@@ -50,7 +52,7 @@ class UMaterialExpressionPinBase : public UMaterialExpression
 
 	/** Direction of the pins for this base. */
 	UPROPERTY()
-	TEnumAsByte<enum EEdGraphPinDirection> PinDirection;
+	TEnumAsByte<EEdGraphPinDirection> PinDirection;
 
 #if WITH_EDITOR
 	/** Helper function to clear all reroutes. */
