Rename UPROPETY field with changed type to avoid warning message

diff --git a/Engine/Source/Runtime/Engine/Private/Materials/MaterialCachedData.cpp b/Engine/Source/Runtime/Engine/Private/Materials/MaterialCachedData.cpp
--- a/Engine/Source/Runtime/Engine/Private/Materials/MaterialCachedData.cpp
+++ b/Engine/Source/Runtime/Engine/Private/Materials/MaterialCachedData.cpp
@@ -60,17 +60,17 @@ static int32 TryAddParameter(FMaterialCachedParameters& CachedParameters, EMater
 													CachedParameters.EditorOnlyEntries[static_cast<int32>(Type) - static_cast<int32>(EMaterialParameterType::RuntimeCount)] :
 													CachedParameters.RuntimeEntries[static_cast<int32>(Type)];
 
-	FSetElementId ElementId = Entry.ParameterInfos.FindId(ParameterInfo);
+	FSetElementId ElementId = Entry.ParameterInfoSet.FindId(ParameterInfo);
 	int32 Index = INDEX_NONE;
 	if (!ElementId.IsValidId())
 	{
-		ElementId = Entry.ParameterInfos.Add(ParameterInfo);
+		ElementId = Entry.ParameterInfoSet.Add(ParameterInfo);
 		Index = ElementId.AsInteger();
 		Entry.ExpressionGuids.Insert(ExpressionGuid, Index);
 		Entry.Overrides.Insert(bOverride, Index);
 
-		// should be valid as long as we don't ever remove elements from ParameterInfos
-		check(Entry.ParameterInfos.Num() == Entry.ExpressionGuids.Num());
+		// should be valid as long as we don't ever remove elements from ParameterInfoSet
+		check(Entry.ParameterInfoSet.Num() == Entry.ExpressionGuids.Num());
 
 		check(Entry.ExpressionGuids.Num() == Entry.Overrides.Num());
 		return Index;
@@ -599,7 +599,7 @@ bool FMaterialCachedExpressionData::UpdateForExpressions(const FMaterialCachedEx
 
 void FMaterialCachedParameterEntry::Reset()
 {
-	ParameterInfos.Reset();
+	ParameterInfoSet.Reset();
 	ExpressionGuids.Reset();
 	Overrides.Reset();
 }
@@ -653,7 +653,7 @@ int32 FMaterialCachedParameters::FindParameterIndex(EMaterialParameterType Type,
 int32 FMaterialCachedParameters::FindParameterIndex(EMaterialParameterType Type, const FMemoryImageMaterialParameterInfo& ParameterInfo) const
 {
 	const FMaterialCachedParameterEntry& Entry = GetParameterTypeEntry(Type);
-	const FSetElementId ElementId = Entry.ParameterInfos.FindId(FMaterialParameterInfo(ParameterInfo));
+	const FSetElementId ElementId = Entry.ParameterInfoSet.FindId(FMaterialParameterInfo(ParameterInfo));
 	return ElementId.AsInteger();
 }
 
@@ -681,14 +681,14 @@ bool FMaterialCachedParameters::IsDefaultParameterValid(EMaterialParameterType T
 void FMaterialCachedParameters::GetAllParameterInfoOfType(EMaterialParameterType Type, bool bEmptyOutput, TArray<FMaterialParameterInfo>& OutParameterInfo, TArray<FGuid>& OutParameterIds) const
 {
 	const FMaterialCachedParameterEntry& Entry = GetParameterTypeEntry(Type);
-	const int32 NumParameters = Entry.ParameterInfos.Num();
+	const int32 NumParameters = Entry.ParameterInfoSet.Num();
 	if (bEmptyOutput)
 	{
 		OutParameterInfo.Empty(NumParameters);
 		OutParameterIds.Empty(NumParameters);
 	}
 
-	for (TSet<FMaterialParameterInfo>::TConstIterator It(Entry.ParameterInfos); It; ++It)
+	for (TSet<FMaterialParameterInfo>::TConstIterator It(Entry.ParameterInfoSet); It; ++It)
 	{
 		OutParameterInfo.Add(*It);
 		OutParameterIds.Add(Entry.ExpressionGuids[It.GetId().AsInteger()]);
@@ -698,14 +698,14 @@ void FMaterialCachedParameters::GetAllParameterInfoOfType(EMaterialParameterType
 void FMaterialCachedParameters::GetAllGlobalParameterInfoOfType(EMaterialParameterType Type, bool bEmptyOutput, TArray<FMaterialParameterInfo>& OutParameterInfo, TArray<FGuid>& OutParameterIds) const
 {
 	const FMaterialCachedParameterEntry& Entry = GetParameterTypeEntry(Type);
-	const int32 NumParameters = Entry.ParameterInfos.Num();
+	const int32 NumParameters = Entry.ParameterInfoSet.Num();
 	if (bEmptyOutput)
 	{
 		OutParameterInfo.Empty(NumParameters);
 		OutParameterIds.Empty(NumParameters);
 	}
 
-	for (TSet<FMaterialParameterInfo>::TConstIterator It(Entry.ParameterInfos); It; ++It)
+	for (TSet<FMaterialParameterInfo>::TConstIterator It(Entry.ParameterInfoSet); It; ++It)
 	{
 		const FMaterialParameterInfo& ParameterInfo = *It;
 		if (ParameterInfo.Association == GlobalParameter)
diff --git a/Engine/Source/Runtime/Engine/Public/MaterialCachedData.h b/Engine/Source/Runtime/Engine/Public/MaterialCachedData.h
--- a/Engine/Source/Runtime/Engine/Public/MaterialCachedData.h
+++ b/Engine/Source/Runtime/Engine/Public/MaterialCachedData.h
@@ -106,7 +106,7 @@ struct FMaterialCachedParameterEntry
 	// This is used to map FMaterialParameterInfos to indices, which are then used to index various TArrays containing values for each type of parameter
 	// (ExpressionGuids and Overrides, along with ScalarValues, VectorValues, etc)
 	UPROPERTY()
-	TSet<FMaterialParameterInfo> ParameterInfos;
+	TSet<FMaterialParameterInfo> ParameterInfoSet;
 
 	UPROPERTY()
 	TArray<FGuid> ExpressionGuids; // editor-only?
@@ -144,7 +144,7 @@ struct FMaterialCachedParameters
 	inline const FMaterialCachedParameterEntry& GetParameterTypeEntry(EMaterialParameterType Type) const { return RuntimeEntries[static_cast<int32>(Type)]; }
 #endif
 
-	inline int32 GetNumParameters(EMaterialParameterType Type) const { return GetParameterTypeEntry(Type).ParameterInfos.Num(); }
+	inline int32 GetNumParameters(EMaterialParameterType Type) const { return GetParameterTypeEntry(Type).ParameterInfoSet.Num(); }
 
 	int32 FindParameterIndex(EMaterialParameterType Type, const FMemoryImageMaterialParameterInfo& HashedParameterInfo, bool bOveriddenOnly) const;
 	int32 FindParameterIndex(EMaterialParameterType Type, const FMemoryImageMaterialParameterInfo& HashedParameterInfo) const;
