Fix unity build/ static analysis warnings.

diff --git a/Engine/Source/Runtime/Engine/Private/ComputeFramework/ShaderParamTypeDefinition.cpp b/Engine/Source/Runtime/Engine/Private/ComputeFramework/ShaderParamTypeDefinition.cpp
--- a/Engine/Source/Runtime/Engine/Private/ComputeFramework/ShaderParamTypeDefinition.cpp
+++ b/Engine/Source/Runtime/Engine/Private/ComputeFramework/ShaderParamTypeDefinition.cpp
@@ -5,6 +5,7 @@
 #include "Serialization/Archive.h"
 #include "Templates/TypeHash.h"
 #include "Misc/DefaultValueHelper.h"
+#include "Algo/Find.h"
 
 
 // Storage for shared shader value types.
diff --git a/Engine/Source/Runtime/Engine/Public/ComputeFramework/ComputeFramework.h b/Engine/Source/Runtime/Engine/Public/ComputeFramework/ComputeFramework.h
--- a/Engine/Source/Runtime/Engine/Public/ComputeFramework/ComputeFramework.h
+++ b/Engine/Source/Runtime/Engine/Public/ComputeFramework/ComputeFramework.h
@@ -3,6 +3,8 @@
 #pragma once
 
 #include "Templates/RefCounting.h"
+#include "RHIDefinitions.h"
+#include "Shader.h"
 
 class FRDGPooledBuffer;
 class FComputeKernelResource;
diff --git a/Engine/Source/Runtime/Engine/Public/ComputeFramework/ComputeKernelShared.h b/Engine/Source/Runtime/Engine/Public/ComputeFramework/ComputeKernelShared.h
--- a/Engine/Source/Runtime/Engine/Public/ComputeFramework/ComputeKernelShared.h
+++ b/Engine/Source/Runtime/Engine/Public/ComputeFramework/ComputeKernelShared.h
@@ -438,7 +438,7 @@ public:
 	
 	bool IsSame(const FComputeKernelShaderMapId& InId) const;
 
-	uint32 GetKernelFlags() const { return 0; }
+	uint32 GetKernelFlags() const { return KernelFlags; }
 
 protected:
 
@@ -476,6 +476,8 @@ private:
 	 */
 	uint64 ShaderCodeHash = 0;
 
+	uint32 KernelFlags = 0;
+
 	/** 
 	 * Contains the compiling id of this shader map when it is being compiled asynchronously. 
 	 * This can be used to access the shader map during async compiling, since GameThreadShaderMap will not have been set yet.
