Async Texture Compilation Improvements
   - Update primitive component's relying on a texture to avoid GPU cached values to become stale (i.e. Lumen, Virtual Textures).
   - Reduce flickering during texture update by removing virtual texture flush as it is not required anymore with the previous fix.
   - Prioritize on screen textures even for nanite mesh and virtual textures by looking at "last seen" primitive component material's usage.
   - Force heighmap texture to be ready for PIE and -game to avoid falling through the landscape.

General Optimizations
   - Eliminated FlushRenderingCommands from both MeshCards/DistanceFields update to reduce heavy stuttering.
   - Reuse expensive ObjectIterator scans between MeshCards/DistanceFields/Shaders/Textures/StaticMesh during the same tick.

diff --git a/Engine/Source/Runtime/Engine/Private/DistanceFieldAtlas.cpp b/Engine/Source/Runtime/Engine/Private/DistanceFieldAtlas.cpp
--- a/Engine/Source/Runtime/Engine/Private/DistanceFieldAtlas.cpp
+++ b/Engine/Source/Runtime/Engine/Private/DistanceFieldAtlas.cpp
@@ -23,6 +23,7 @@
 #include "MeshCardRepresentation.h"
 #include "Misc/QueuedThreadPoolWrapper.h"
 #include "Async/Async.h"
+#include "ObjectCacheContext.h"
 
 #if WITH_EDITOR
 #include "DerivedDataCacheInterface.h"
@@ -1156,8 +1157,9 @@ FString FDistanceFieldAsyncQueue::GetReferencerName() const
 void FDistanceFieldAsyncQueue::ProcessAsyncTasks(bool bLimitExecutionTime)
 {
 #if WITH_EDITOR
-	TRACE_CPUPROFILER_EVENT_SCOPE(FDistanceFieldAsyncQueue::ProcessAsyncTasks)
+	TRACE_CPUPROFILER_EVENT_SCOPE(FDistanceFieldAsyncQueue::ProcessAsyncTasks);
 
+	FObjectCacheContextScope ObjectCacheScope;
 	const double MaxProcessingTime = 0.016f;
 	double StartTime = FPlatformTime::Seconds();
 	while (!bLimitExecutionTime || (FPlatformTime::Seconds() - StartTime) < MaxProcessingTime)
@@ -1183,19 +1185,19 @@ void FDistanceFieldAsyncQueue::ProcessAsyncTasks(bool bLimitExecutionTime)
 			Task->GeneratedVolumeData->VolumeTexture.Initialize(Task->StaticMesh);
 			FDistanceFieldVolumeData* OldVolumeData = Task->StaticMesh->GetRenderData()->LODResources[0].DistanceFieldData;
 
+			// Assign the new volume data, this is safe because the render thread makes a copy of the pointer at scene proxy creation time.
+			Task->StaticMesh->GetRenderData()->LODResources[0].DistanceFieldData = Task->GeneratedVolumeData;
+
 			// Renderstates are not initialized between UStaticMesh::PreEditChange() and UStaticMesh::PostEditChange()
 			if (Task->StaticMesh->GetRenderData()->IsInitialized())
 			{
-				// Cause all components using this static mesh to get re-registered, which will recreate their proxies and primitive uniform buffers
-				FStaticMeshComponentRecreateRenderStateContext RecreateRenderStateContext(Task->StaticMesh, false);
-
-				// Assign the new volume data
-				Task->StaticMesh->GetRenderData()->LODResources[0].DistanceFieldData = Task->GeneratedVolumeData;
-			}
-			else
-			{
-				// Assign the new volume data
-				Task->StaticMesh->GetRenderData()->LODResources[0].DistanceFieldData = Task->GeneratedVolumeData;
+				for (UStaticMeshComponent* Component : ObjectCacheScope.GetContext().GetStaticMeshComponents(Task->StaticMesh))
+				{
+					if (Component->IsRegistered() && Component->IsRenderStateCreated())
+					{
+						Component->MarkRenderStateDirty();
+					}
+				}
 			}
 
 			OldVolumeData->VolumeTexture.Release();
diff --git a/Engine/Source/Runtime/Engine/Private/MeshCardRepresentation.cpp b/Engine/Source/Runtime/Engine/Private/MeshCardRepresentation.cpp
--- a/Engine/Source/Runtime/Engine/Private/MeshCardRepresentation.cpp
+++ b/Engine/Source/Runtime/Engine/Private/MeshCardRepresentation.cpp
@@ -19,6 +19,7 @@
 #include "Async/ParallelFor.h"
 #include "Misc/QueuedThreadPoolWrapper.h"
 #include "Async/Async.h"
+#include "ObjectCacheContext.h"
 
 #if WITH_EDITOR
 #include "DerivedDataCacheInterface.h"
@@ -306,8 +307,9 @@ FString FCardRepresentationAsyncQueue::GetReferencerName() const
 void FCardRepresentationAsyncQueue::ProcessAsyncTasks(bool bLimitExecutionTime)
 {
 #if WITH_EDITOR
-	TRACE_CPUPROFILER_EVENT_SCOPE(FCardRepresentationAsyncQueue::ProcessAsyncTasks)
+	TRACE_CPUPROFILER_EVENT_SCOPE(FCardRepresentationAsyncQueue::ProcessAsyncTasks);
 
+	FObjectCacheContextScope ObjectCacheScope;
 	const double MaxProcessingTime = 0.016f;
 	double StartTime = FPlatformTime::Seconds();
 	while (!bLimitExecutionTime || (FPlatformTime::Seconds() - StartTime) < MaxProcessingTime)
@@ -331,12 +333,19 @@ void FCardRepresentationAsyncQueue::ProcessAsyncTasks(bool bLimitExecutionTime)
 		{
 			FCardRepresentationData* OldCardData = Task->StaticMesh->GetRenderData()->LODResources[0].CardRepresentationData;
 
-			{
-				// Cause all components using this static mesh to get re-registered, which will recreate their proxies and primitive uniform buffers
-				FStaticMeshComponentRecreateRenderStateContext RecreateRenderStateContext(Task->StaticMesh, false);
+			// Assign the new data, this is safe because the render threads makes a copy of the pointer at scene proxy creation time.
+			Task->StaticMesh->GetRenderData()->LODResources[0].CardRepresentationData = Task->GeneratedCardRepresentation;
 
-				// Assign the new data
-				Task->StaticMesh->GetRenderData()->LODResources[0].CardRepresentationData = Task->GeneratedCardRepresentation;
+			// Any already created render state needs to be dirtied
+			if (Task->StaticMesh->GetRenderData()->IsInitialized())
+			{
+				for (UStaticMeshComponent* Component : ObjectCacheScope.GetContext().GetStaticMeshComponents(Task->StaticMesh))
+				{
+					if (Component->IsRegistered() && Component->IsRenderStateCreated())
+					{
+						Component->MarkRenderStateDirty();
+					}
+				}
 			}
 
 			// Rendering thread may still be referencing the old one, use the deferred cleanup interface to delete it next frame when it is safe
diff --git a/Engine/Source/Runtime/Engine/Private/ObjectCacheContext.cpp b/Engine/Source/Runtime/Engine/Private/ObjectCacheContext.cpp
new file mode 100644
--- /dev/null
+++ b/Engine/Source/Runtime/Engine/Private/ObjectCacheContext.cpp
@@ -0,0 +1,194 @@
+// Copyright Epic Games, Inc. All Rights Reserved.
+
+#include "ObjectCacheContext.h"
+#include "UObject/UObjectIterator.h"
+#include "Materials/MaterialInstance.h"
+#include "Engine/StaticMesh.h"
+#include "Components/PrimitiveComponent.h"
+#include "Components/StaticMeshComponent.h"
+#include "Templates/UniquePtr.h"
+
+const TArray<UPrimitiveComponent*>& FObjectCacheContext::GetPrimitiveComponents()
+{
+	if (!PrimitiveComponents.IsSet())
+	{
+		TRACE_CPUPROFILER_EVENT_SCOPE(ComputePrimitiveComponents);
+
+		TArray<UPrimitiveComponent*> Array;
+		Array.Reserve(4096);
+		for (TObjectIterator<UPrimitiveComponent> It; It; ++It)
+		{
+			Array.Add(*It);
+		}
+		PrimitiveComponents = MoveTemp(Array);
+	}
+	return PrimitiveComponents.GetValue();
+}
+
+const TArray<UStaticMeshComponent*>& FObjectCacheContext::GetStaticMeshComponents()
+{
+	if (!StaticMeshComponents.IsSet())
+	{
+		TRACE_CPUPROFILER_EVENT_SCOPE(ComputeStaticMeshComponents);
+
+		TArray<UStaticMeshComponent*> Array;
+		Array.Reserve(4096);
+		for (TObjectIterator<UStaticMeshComponent> It; It; ++It)
+		{
+			Array.Add(*It);
+		}
+		StaticMeshComponents = MoveTemp(Array);
+	}
+	return StaticMeshComponents.GetValue();
+}
+
+const TSet<UTexture*>& FObjectCacheContext::GetUsedTextures(UMaterialInterface* MaterialInterface)
+{
+	TSet<UTexture*>* Textures = MaterialUsedTextures.Find(MaterialInterface);
+	if (Textures == nullptr)
+	{
+		Textures = &MaterialUsedTextures.Add(MaterialInterface);
+		for (UObject* TextureObject : MaterialInterface->GetReferencedTextures())
+		{
+			UTexture* Texture = Cast<UTexture>(TextureObject);
+			if (Texture)
+			{
+				Textures->Add(Texture);
+			}
+		}
+
+		// Fix in CL 13480995 broke GetReferencedTextures() ability to return all referenced textures
+		// so we manually gather them instead...
+		// Recursion is required for MIC to resolve all TextureParameterValues in the hierarchy.
+		UMaterialInstance* MaterialInstance = Cast<UMaterialInstance>(MaterialInterface);
+		while (MaterialInstance)
+		{
+			for (const FTextureParameterValue& TextureParam : MaterialInstance->TextureParameterValues)
+			{
+				if (TextureParam.ParameterValue)
+				{
+					Textures->Add(TextureParam.ParameterValue);
+				}
+			}
+
+			MaterialInstance = Cast<UMaterialInstance>(MaterialInstance->Parent);
+		}
+	}
+
+	return *Textures;
+}
+
+const TArray<UMaterialInterface*>& FObjectCacheContext::GetUsedMaterials(UPrimitiveComponent* Component)
+{
+	TArray<UMaterialInterface*>* Materials = PrimitiveComponentToMaterial.Find(Component);
+	if (Materials == nullptr)
+	{
+		Materials = &PrimitiveComponentToMaterial.Add(Component);
+		Component->GetUsedMaterials(*Materials);
+	}
+
+	return *Materials;
+}
+
+const TArray<UStaticMeshComponent*>& FObjectCacheContext::GetStaticMeshComponents(UStaticMesh* InStaticMesh)
+{
+	if (!StaticMeshToComponents.IsSet())
+	{
+		TRACE_CPUPROFILER_EVENT_SCOPE(ComputeStaticMeshToComponents);
+
+		TMap<TObjectKey<UStaticMesh>, TArray<UStaticMeshComponent*>> TempMap;
+		TempMap.Reserve(8192);
+		for (UStaticMeshComponent* Component : GetStaticMeshComponents())
+		{
+			TempMap.FindOrAdd(Component->GetStaticMesh()).Add(Component);
+		}
+		StaticMeshToComponents = MoveTemp(TempMap);
+	}
+
+	static TArray<UStaticMeshComponent*> EmptyArray;
+	TArray<UStaticMeshComponent*>* Array = StaticMeshToComponents.GetValue().Find(InStaticMesh);
+	return Array ? *Array : EmptyArray;
+}
+
+const TSet<UMaterialInterface*>& FObjectCacheContext::GetMaterialsAffectedByTexture(UTexture* InTexture)
+{
+	if (!TextureToMaterials.IsSet())
+	{
+		TRACE_CPUPROFILER_EVENT_SCOPE(ComputeMaterialsAffectedByTexture);
+
+		TMap<TObjectKey<UTexture>, TSet<UMaterialInterface*>> TempMap;
+		TempMap.Reserve(8192);
+		for (TObjectIterator<UMaterialInterface> It; It; ++It)
+		{
+			UMaterialInterface* MaterialInterface = *It;
+			for (UTexture* Texture : GetUsedTextures(MaterialInterface))
+			{
+				TempMap.FindOrAdd(Texture).Add(MaterialInterface);
+			}
+		}
+		TextureToMaterials = MoveTemp(TempMap);
+	}
+
+	static TSet<UMaterialInterface*> EmptySet;
+	TSet<UMaterialInterface*>* Set = TextureToMaterials.GetValue().Find(InTexture);
+	return Set ? *Set : EmptySet;
+}
+
+const TSet<UPrimitiveComponent*>& FObjectCacheContext::GetPrimitivesAffectedByMaterial(UMaterialInterface* InMaterialInterface)
+{
+	if (!MaterialToPrimitives.IsSet())
+	{
+		TRACE_CPUPROFILER_EVENT_SCOPE(ComputePrimitivesAffectedByMaterial);
+
+		TMap<TObjectKey<UMaterialInterface>, TSet<UPrimitiveComponent*>> TempMap;
+		for (UPrimitiveComponent* Component : GetPrimitiveComponents())
+		{
+			if (Component->IsRegistered() && Component->IsRenderStateCreated())
+			{
+				for (UMaterialInterface* MaterialInterface : GetUsedMaterials(Component))
+				{
+					if (MaterialInterface)
+					{
+						TempMap.FindOrAdd(MaterialInterface).Add(Component);
+					}
+				}
+			}
+		}
+
+		MaterialToPrimitives = MoveTemp(TempMap);
+	}
+
+	static TSet<UPrimitiveComponent*> EmptySet;
+	TSet<UPrimitiveComponent*>* Set = MaterialToPrimitives.GetValue().Find(InMaterialInterface);
+	return Set ? *Set : EmptySet;
+}
+
+namespace ObjectCacheContextScopeImpl
+{
+	static thread_local TUniquePtr<FObjectCacheContext> Current = nullptr;
+}
+
+FObjectCacheContextScope::FObjectCacheContextScope()
+{
+	using namespace ObjectCacheContextScopeImpl;
+	if (Current == nullptr)
+	{
+		Current.Reset(new FObjectCacheContext());
+		bIsOwner = true;
+	}
+}
+
+FObjectCacheContextScope::~FObjectCacheContextScope()
+{
+	using namespace ObjectCacheContextScopeImpl;
+	if (bIsOwner)
+	{
+		Current.Reset();
+	}
+}
+
+FObjectCacheContext& FObjectCacheContextScope::GetContext() 
+{ 
+	using namespace ObjectCacheContextScopeImpl;
+	return *Current; 
+}
diff --git a/Engine/Source/Runtime/Engine/Private/ShaderCompiler/ShaderCompiler.cpp b/Engine/Source/Runtime/Engine/Private/ShaderCompiler/ShaderCompiler.cpp
--- a/Engine/Source/Runtime/Engine/Private/ShaderCompiler/ShaderCompiler.cpp
+++ b/Engine/Source/Runtime/Engine/Private/ShaderCompiler/ShaderCompiler.cpp
@@ -53,6 +53,8 @@
 #include "ShaderCore.h"
 #include "DistributedBuildInterface/Public/DistributedBuildControllerInterface.h"
 #include "Misc/ScopeRWLock.h"
+#include "ObjectCacheContext.h"
+
 #if WITH_EDITOR
 #include "TextureCompiler.h"
 #include "Rendering/StaticLightingSystemInterface.h"
@@ -936,6 +938,8 @@ static void ProcessErrors(const FShaderCompileJob& CurrentJob, TArray<FString>&
 
 static bool ReadSingleJob(FShaderCompileJob* CurrentJob, FArchive& OutputFile)
 {
+	TRACE_CPUPROFILER_EVENT_SCOPE(ReadSingleJob);
+
 	check(!CurrentJob->bFinalized);
 	CurrentJob->bFinalized = true;
 
@@ -1257,6 +1261,8 @@ void FShaderCompileUtilities::DoReadTaskResults(const TArray<FShaderCommonCompil
 	// Requeue any jobs we wish to run again
 	if (ReissueSourceJobs.Num())
 	{
+		TRACE_CPUPROFILER_EVENT_SCOPE(ReissueShaderJobs);
+
 		TArray<FShaderCommonCompileJobPtr> ReissueJobs;
 		ReissueJobs.Reserve(ReissueSourceJobs.Num());
 		const uint32 JobId = FShaderCommonCompileJob::GetNextJobId();
@@ -1789,6 +1795,8 @@ void FShaderCompileThreadRunnable::ReadAvailableResults()
 		// Check for available result files
 		if (CurrentWorkerInfo.QueuedJobs.Num() > 0)
 		{
+			TRACE_CPUPROFILER_EVENT_SCOPE(FShaderCompileThreadRunnable::ReadAvailableResults);
+
 			// Distributed compiles always use the same directory
 			const FString WorkingDirectory = Manager->AbsoluteShaderBaseWorkingDirectory + FString::FromInt(WorkerIndex) + TEXT("/");
 			// 'Only' indicates to the worker that it should log and continue checking for the input file after the first one is processed
@@ -2537,6 +2545,7 @@ void FShaderCompilingManager::ReleaseJob(FShaderCommonCompileJob* Job)
 
 void FShaderCompilingManager::SubmitJobs(TArray<FShaderCommonCompileJobPtr>& NewJobs, const FString MaterialBasePath, const FString PermutationString)
 {
+	TRACE_CPUPROFILER_EVENT_SCOPE(FShaderCompilingManager::SubmitJobs);
 	check(!FPlatformProperties::RequiresCookedData());
 
 	if (NewJobs.Num() == 0)
@@ -3044,11 +3053,12 @@ void FShaderCompilingManager::ProcessCompiledShaderMaps(
 			// shadermap it is stored in, while cached MDCs can reference its memory. Re-registering will
 			// re-create the cache.
 			{
+				TRACE_CPUPROFILER_EVENT_SCOPE(PropagateGlobalShadersToAllPrimitives);
+
+				FObjectCacheContextScope ObjectCacheScope;
 				TIndirectArray<FComponentRecreateRenderStateContext> ComponentContexts;
-				for (TObjectIterator<UPrimitiveComponent> PrimitiveIt; PrimitiveIt; ++PrimitiveIt)
+				for (UPrimitiveComponent* PrimitiveComponent : ObjectCacheScope.GetContext().GetPrimitiveComponents())
 				{
-					UPrimitiveComponent* PrimitiveComponent = *PrimitiveIt;
-
 					if (PrimitiveComponent->IsRenderStateCreated())
 					{
 						ComponentContexts.Add(new FComponentRecreateRenderStateContext(PrimitiveComponent));
@@ -3121,22 +3131,19 @@ void FShaderCompilingManager::ProcessCompiledShaderMaps(
 
 void FShaderCompilingManager::PropagateMaterialChangesToPrimitives(const TMap<TRefCountPtr<FMaterial>, TRefCountPtr<FMaterialShaderMap>>& MaterialsToUpdate)
 {
-	TArray<UMaterialInterface*> UsedMaterials;
-	TIndirectArray<FComponentRecreateRenderStateContext> ComponentContexts;
+	TRACE_CPUPROFILER_EVENT_SCOPE(FShaderCompilingManager::PropagateMaterialChangesToPrimitives);
 
-	for (TObjectIterator<UPrimitiveComponent> PrimitiveIt; PrimitiveIt; ++PrimitiveIt)
+	FObjectCacheContextScope ObjectCacheScope;
+	TIndirectArray<FComponentRecreateRenderStateContext> ComponentContexts;
+	for (UPrimitiveComponent* PrimitiveComponent : ObjectCacheScope.GetContext().GetPrimitiveComponents())
 	{
-		UPrimitiveComponent* PrimitiveComponent = *PrimitiveIt;
-
 		if (PrimitiveComponent->IsRenderStateCreated())
 		{
-			UsedMaterials.Reset();
 			bool bPrimitiveIsDependentOnMaterial = false;
 
 			// Note: relying on GetUsedMaterials to be accurate, or else we won't propagate to the right primitives and the renderer will crash later
 			// FPrimitiveSceneProxy::VerifyUsedMaterial is used to make sure that all materials used for rendering are reported in GetUsedMaterials
-			PrimitiveComponent->GetUsedMaterials(UsedMaterials);
-
+			const TArray<UMaterialInterface*>& UsedMaterials = ObjectCacheScope.GetContext().GetUsedMaterials(PrimitiveComponent);
 			if (UsedMaterials.Num() > 0)
 			{
 				for (TMap<TRefCountPtr<FMaterial>, TRefCountPtr<FMaterialShaderMap>>::TConstIterator MaterialIt(MaterialsToUpdate); MaterialIt; ++MaterialIt)
diff --git a/Engine/Source/Runtime/Engine/Private/TextureCompiler.cpp b/Engine/Source/Runtime/Engine/Private/TextureCompiler.cpp
--- a/Engine/Source/Runtime/Engine/Private/TextureCompiler.cpp
+++ b/Engine/Source/Runtime/Engine/Private/TextureCompiler.cpp
@@ -1,10 +1,11 @@
 // Copyright Epic Games, Inc. All Rights Reserved.
 
 #include "TextureCompiler.h"
-#include "Engine/Texture.h"
 
 #if WITH_EDITOR
 
+#include "Engine/Texture.h"
+#include "ObjectCacheContext.h"
 #include "Framework/Notifications/NotificationManager.h"
 #include "Settings/EditorExperimentalSettings.h"
 #include "Widgets/Notifications/SNotificationList.h"
@@ -13,12 +14,13 @@
 #include "EngineModule.h"
 #include "Misc/ScopedSlowTask.h"
 #include "UObject/StrongObjectPtr.h"
-#include "UObject/UObjectIterator.h"
 #include "Materials/MaterialInterface.h"
 #include "Materials/MaterialInstance.h"
 #include "Materials/Material.h"
 #include "TextureDerivedDataTask.h"
 #include "Misc/IQueuedWork.h"
+#include "Components/PrimitiveComponent.h"
+#include "LevelEditor.h"
 
 #define LOCTEXT_NAMESPACE "TextureCompiler"
 
@@ -60,42 +62,6 @@ namespace TextureCompilingManagerImpl
 		return StaticEnum<TextureGroup>()->GetMetaData(TEXT("DisplayName"), Texture->LODGroup);
 	}
 
-	static TMultiMap<UObject*, UMaterialInterface*> GetTexturesAffectingMaterialInterfaces()
-	{
-		TRACE_CPUPROFILER_EVENT_SCOPE(GetTexturesAffectingMaterials);
-
-		// Update any material that uses this texture
-		TMultiMap<UObject*, UMaterialInterface*> TexturesRequiringMaterialUpdate;
-
-		for (TObjectIterator<UMaterialInterface> It; It; ++It)
-		{
-			UMaterialInterface* MaterialInterface = *It;
-			for (UObject* Texture : MaterialInterface->GetReferencedTextures())
-			{
-				TexturesRequiringMaterialUpdate.Emplace(Texture, MaterialInterface);
-			}
-			
-			// Fix in CL 13480995 broke GetReferencedTextures() ability to return all referenced textures
-			// so we manually gather them instead...
-			// Recursion is required for MIC to resolve all TextureParameterValues in the hierarchy.
-			UMaterialInstance* MaterialInstance = Cast<UMaterialInstance>(MaterialInterface);
-			while (MaterialInstance)
-			{
-				for (const FTextureParameterValue& TextureParam : MaterialInstance->TextureParameterValues)
-				{
-					if (TextureParam.ParameterValue)
-					{
-						TexturesRequiringMaterialUpdate.Add(TextureParam.ParameterValue, MaterialInterface);
-					}
-				}
-
-				MaterialInstance = Cast<UMaterialInstance>(MaterialInstance->Parent);
-			}
-		}
-
-		return TexturesRequiringMaterialUpdate;
-	}
-
 	static EQueuedWorkPriority GetBasePriority(UTexture* InTexture)
 	{
 		switch (InTexture->LODGroup)
@@ -111,7 +77,7 @@ namespace TextureCompilingManagerImpl
 
 	static EQueuedWorkPriority GetBoostPriority(UTexture* InTexture)
 	{
-		return (EQueuedWorkPriority)((uint8)GetBasePriority(InTexture) - 1);
+		return (EQueuedWorkPriority)(FMath::Max((uint8)1, (uint8)GetBasePriority(InTexture)) - 1);
 	}
 
 	static const TCHAR* GetPriorityName(EQueuedWorkPriority Priority)
@@ -297,7 +263,7 @@ void FTextureCompilingManager::UpdateCompilationNotification()
 	{
 		if (NotificationItem.IsValid())
 		{
-			NotificationItem->SetText(NSLOCTEXT("TextureBuild", "TextureBuildFinished", "Finished building Textures!"));
+			NotificationItem->SetText(NSLOCTEXT("TextureBuild", "TextureBuildFinished", "Textures are ready!"));
 			NotificationItem->SetCompletionState(SNotificationItem::CS_Success);
 			NotificationItem->ExpireAndFadeout();
 
@@ -308,7 +274,7 @@ void FTextureCompilingManager::UpdateCompilationNotification()
 	{
 		if (!NotificationItem.IsValid())
 		{
-			FNotificationInfo Info(NSLOCTEXT("TextureBuild", "TextureBuildInProgress", "Building Textures"));
+			FNotificationInfo Info(NSLOCTEXT("TextureBuild", "TextureBuildInProgress", "Preparing Textures"));
 			Info.bFireAndForget = false;
 
 			// Setting fade out and expire time to 0 as the expire message is currently very obnoxious
@@ -321,7 +287,7 @@ void FTextureCompilingManager::UpdateCompilationNotification()
 
 		FFormatNamedArguments Args;
 		Args.Add(TEXT("BuildTasks"), FText::AsNumber(NumRemainingCompilations));
-		FText ProgressMessage = FText::Format(NSLOCTEXT("TextureBuild", "TextureBuildInProgressFormat", "Building Textures ({BuildTasks})"), Args);
+		FText ProgressMessage = FText::Format(NSLOCTEXT("TextureBuild", "TextureBuildInProgressFormat", "Preparing Textures ({BuildTasks})"), Args);
 
 		NotificationItem->SetCompletionState(SNotificationItem::CS_Pending);
 		NotificationItem->SetVisibility(EVisibility::HitTestInvisible);
@@ -336,13 +302,11 @@ void FTextureCompilingManager::FinishTextureCompilation(UTexture* Texture)
 	check(IsInGameThread());
 	TRACE_CPUPROFILER_EVENT_SCOPE(FinishTextureCompilation);
 
-	UE_LOG(LogTexture, Verbose, TEXT("UpdateResource for %s (%s) due to async texture compilation"), *Texture->GetName(), *GetLODGroupName(Texture));
+	UE_LOG(LogTexture, Verbose, TEXT("Refreshing texture %s because it is ready"), *Texture->GetName());
 
 	Texture->FinishCachePlatformData();
 	Texture->UpdateResource();
 
-	GetRendererModule().FlushVirtualTextureCache();
-
 	// Generate an empty property changed event, to force the asset registry tag
 	// to be refreshed now that pixel format and alpha channels are available.
 	FPropertyChangedEvent EmptyPropertyChangedEvent(nullptr);
@@ -383,11 +347,6 @@ void FTextureCompilingManager::AddTextures(const TArray<UTexture*>& InTextures)
 	TRACE_CPUPROFILER_EVENT_SCOPE(FTextureCompilingManager::AddTextures)
 	check(IsInGameThread());
 
-	// We might not get ticked very often during load time
-	// so this will allow us to refresh compiled textures
-	// of the highest priority to improve the UI experience.
-	ProcessTextures(1 /* Maximum Priority */);
-
 	// Register new textures after ProcessTextures to avoid
 	// potential reentrant calls to CreateResource on the
 	// textures being added. This would cause multiple
@@ -439,7 +398,7 @@ void FTextureCompilingManager::FinishCompilation(const TArray<UTexture*>& InText
 
 	if (PendingTextures.Num())
 	{
-		FScopedSlowTask SlowTask((float)PendingTextures.Num(), LOCTEXT("FinishTextureCompilation", "Waiting on texture compilation"), true);
+		FScopedSlowTask SlowTask((float)PendingTextures.Num(), LOCTEXT("FinishTextureCompilation", "Waiting on texture preparation"), true);
 		SlowTask.MakeDialogDelayed(1.0f);
 
 		struct FTextureTask : public IQueuedWork
@@ -467,7 +426,7 @@ void FTextureCompilingManager::FinishCompilation(const TArray<UTexture*>& InText
 		auto UpdateProgress =
 			[&SlowTask](float Progress, int32 Done, int32 Total, const FString& CurrentObjectsName)
 			{
-				return SlowTask.EnterProgressFrame(Progress, FText::FromString(FString::Printf(TEXT("Waiting on texture compilation %d/%d (%s) ..."), Done, Total, *CurrentObjectsName)));
+				return SlowTask.EnterProgressFrame(Progress, FText::FromString(FString::Printf(TEXT("Waiting for textures to be ready %d/%d (%s) ..."), Done, Total, *CurrentObjectsName)));
 			};
 
 		for (FTextureTask& PendingTask : PendingTasks)
@@ -479,8 +438,8 @@ void FTextureCompilingManager::FinishCompilation(const TArray<UTexture*>& InText
 			{
 				UpdateProgress(0.0f, TextureIndex, InTextures.Num(), TextureName);
 			}
+			UE_LOG(LogTexture, Display, TEXT("Waiting for textures to be ready %d/%d (%s) ..."), TextureIndex, InTextures.Num(), *TextureName);
 			UpdateProgress(1.f, TextureIndex++, InTextures.Num(), TextureName);
-			UE_LOG(LogTexture, Verbose, TEXT("FinishCompilation requested for %s (%s)"), *TextureName, *GetLODGroupName(Texture));
 			FinishTextureCompilation(Texture);
 
 			for (TSet<TWeakObjectPtr<UTexture>>& Bucket : RegisteredTextureBuckets)
@@ -489,6 +448,60 @@ void FTextureCompilingManager::FinishCompilation(const TArray<UTexture*>& InText
 			}
 		}
 	}
+
+	PostTextureCompilation(PendingTextures);
+}
+
+void FTextureCompilingManager::PostTextureCompilation(const TSet<UTexture*>& InCompiledTextures)
+{
+	using namespace TextureCompilingManagerImpl;
+	if (InCompiledTextures.Num())
+	{
+		FObjectCacheContextScope ObjectCacheScope;
+		TRACE_CPUPROFILER_EVENT_SCOPE(PostTextureCompilation);
+		{
+			TSet<UMaterialInterface*> AffectedMaterials;
+			for (UTexture* Texture : InCompiledTextures)
+			{
+				AffectedMaterials.Append(ObjectCacheScope.GetContext().GetMaterialsAffectedByTexture(Texture));
+			}
+
+			if (AffectedMaterials.Num())
+			{
+				{
+					TRACE_CPUPROFILER_EVENT_SCOPE(UpdateMaterials);
+
+					for (UMaterialInterface* MaterialToUpdate : AffectedMaterials)
+					{
+						FMaterialRenderProxy* RenderProxy = MaterialToUpdate->GetRenderProxy();
+						if (RenderProxy)
+						{
+							ENQUEUE_RENDER_COMMAND(TextureCompiler_RecacheUniformExpressions)(
+								[RenderProxy](FRHICommandListImmediate& RHICmdList)
+								{
+									RenderProxy->CacheUniformExpressions(false);
+								});
+						}
+					}
+				}
+
+				{
+					TRACE_CPUPROFILER_EVENT_SCOPE(UpdatePrimitives);
+
+					TSet<UPrimitiveComponent*> AffectedPrimitives;
+					for (UMaterialInterface* MaterialInterface : AffectedMaterials)
+					{
+						AffectedPrimitives.Append(ObjectCacheScope.GetContext().GetPrimitivesAffectedByMaterial(MaterialInterface));
+					}
+
+					for (UPrimitiveComponent* AffectedPrimitive : AffectedPrimitives)
+					{
+						AffectedPrimitive->MarkRenderStateDirty();
+					}
+				}
+			}
+		}
+	}
 }
 
 void FTextureCompilingManager::FinishAllCompilation()
@@ -516,7 +529,43 @@ void FTextureCompilingManager::FinishAllCompilation()
 	}
 }
 
-void FTextureCompilingManager::ProcessTextures(int32 MaximumPriority)
+bool FTextureCompilingManager::RequestPriorityChange(UTexture* InTexture, EQueuedWorkPriority InPriority)
+{
+	using namespace TextureCompilingManagerImpl;
+	if (InTexture)
+	{
+		FTexturePlatformData** Data = InTexture->GetRunningPlatformData();
+		if (Data && *Data)
+		{
+			FTextureAsyncCacheDerivedDataTask* AsyncTask = (*Data)->AsyncTask;
+			if (AsyncTask)
+			{
+				EQueuedWorkPriority OldPriority = AsyncTask->GetPriority();
+				if (OldPriority != InPriority)
+				{
+					if (AsyncTask->Reschedule(GetThreadPool(), InPriority))
+					{
+						UE_LOG(
+							LogTexture,
+							Verbose,
+							TEXT("Changing priority of %s (%s) from %s to %s"),
+							*InTexture->GetName(),
+							*GetLODGroupName(InTexture),
+							GetPriorityName(OldPriority),
+							GetPriorityName(InPriority)
+						);
+
+						return true;
+					}
+				}
+			}
+		}
+	}
+
+	return false;
+}
+
+void FTextureCompilingManager::ProcessTextures(bool bLimitExecutionTime, int32 MaximumPriority)
 {
 	using namespace TextureCompilingManagerImpl;
 	TRACE_CPUPROFILER_EVENT_SCOPE(FTextureCompilingManager::ProcessTextures);
@@ -524,7 +573,8 @@ void FTextureCompilingManager::ProcessTextures(int32 MaximumPriority)
 
 	if (GetNumRemainingTextures())
 	{
-		TArray<UTexture*> ProcessedTextures;
+		FObjectCacheContextScope ObjectCacheScope;
+		TSet<UTexture*> ProcessedTextures;
 		{
 			TRACE_CPUPROFILER_EVENT_SCOPE(ProcessFinishedTextures);
 
@@ -547,7 +597,7 @@ void FTextureCompilingManager::ProcessTextures(int32 MaximumPriority)
 					{
 						if (Texture.IsValid())
 						{
-							const bool bHasTimeLeft = (FPlatformTime::Seconds() - TickStartTime) < MaxSecondsPerFrame;
+							const bool bHasTimeLeft = bLimitExecutionTime ? ((FPlatformTime::Seconds() - TickStartTime) < MaxSecondsPerFrame) : true;
 							if ((bIsHighestPrio || bHasTimeLeft) && Texture->IsAsyncCacheComplete())
 							{
 								FinishTextureCompilation(Texture.Get());
@@ -565,83 +615,120 @@ void FTextureCompilingManager::ProcessTextures(int32 MaximumPriority)
 			}
 		}
 
-		if (ProcessedTextures.Num())
 		{
-			TRACE_CPUPROFILER_EVENT_SCOPE(RecacheUniformExpressions);
-
-			TMultiMap<UObject*, UMaterialInterface*> TexturesAffectingMaterials = GetTexturesAffectingMaterialInterfaces();
+			TRACE_CPUPROFILER_EVENT_SCOPE(FTextureCompilingManager::Reschedule);
 
-			TArray<UMaterialInterface*> MaterialsToUpdate;
-			for (UTexture* Texture : ProcessedTextures)
+			TSet<UTexture*> ReferencedTextures;
+			if (GEngine)
 			{
-				TexturesAffectingMaterials.MultiFind(Texture, MaterialsToUpdate);
+				TRACE_CPUPROFILER_EVENT_SCOPE(GatherSeenPrimitiveMaterials);
+
+				TSet<UMaterialInterface*> RenderedMaterials;
+				for (UPrimitiveComponent* Component : ObjectCacheScope.GetContext().GetPrimitiveComponents())
+				{
+					if (Component->IsRegistered() && Component->IsRenderStateCreated() && Component->GetLastRenderTimeOnScreen() > 0.0f)
+					{
+						for (UMaterialInterface* MaterialInterface : ObjectCacheScope.GetContext().GetUsedMaterials(Component))
+						{
+							if (MaterialInterface)
+							{
+								RenderedMaterials.Add(MaterialInterface);
+							}
+						}
+					}
+				}
+
+				for (UMaterialInterface* MaterialInstance : RenderedMaterials)
+				{
+					ReferencedTextures.Append(ObjectCacheScope.GetContext().GetUsedTextures(MaterialInstance));
+				}
 			}
 
-			if (MaterialsToUpdate.Num())
 			{
-				TRACE_CPUPROFILER_EVENT_SCOPE(UpdateMaterials);
-				
-				for (UMaterialInterface* MaterialToUpdate : MaterialsToUpdate)
+				TRACE_CPUPROFILER_EVENT_SCOPE(ApplyPriorityChanges);
+				// Reschedule higher priority if they have been rendered
+				for (TSet<TWeakObjectPtr<UTexture>>& Bucket : RegisteredTextureBuckets)
 				{
-					MaterialToUpdate->RecacheUniformExpressions(false);
+					for (TWeakObjectPtr<UTexture>& WeakPtr : Bucket)
+					{
+						if (UTexture* Texture = WeakPtr.Get())
+						{
+							// Reschedule any texture that have been rendered with slightly higher priority 
+							// to improve the editor experience for low-core count.
+							//
+							// Keep in mind that some textures are only accessed once during the construction
+							// of a virtual texture, so we can't count on the LastRenderTime to be updated
+							// continuously for those even if they're in view.
+							if (ReferencedTextures.Contains(Texture) ||
+								(Texture->Resource && Texture->Resource->LastRenderTime > 0.0f) ||
+								Texture->TextureReference.GetLastRenderTime() > 0.0f)
+							{
+								RequestPriorityChange(Texture, GetBoostPriority(Texture));
+							}
+						}
+					}
 				}
 			}
 		}
 
+		if (ProcessedTextures.Num())
 		{
-			TRACE_CPUPROFILER_EVENT_SCOPE(FTextureCompilingManager::Reschedule);
+			PostTextureCompilation(ProcessedTextures);
+		}
+	}
+}
 
-			// Reschedule higher priority if they have been rendered
+void FTextureCompilingManager::FinishCompilationsForGame()
+{
+	if (GetNumRemainingTextures())
+	{
+		// Supports both Game and PIE mode
+		const bool bIsPlaying =
+			(GWorld && !GWorld->IsEditorWorld()) ||
+			(GEditor && GEditor->PlayWorld && !GEditor->IsSimulateInEditorInProgress());
+
+		if (bIsPlaying)
+		{
+			TRACE_CPUPROFILER_EVENT_SCOPE(FTextureCompilingManager::FinishCompilationsForGame);
+
+			TSet<UTexture*> TexturesRequiredForGame;
 			for (TSet<TWeakObjectPtr<UTexture>>& Bucket : RegisteredTextureBuckets)
 			{
-				for (TWeakObjectPtr<UTexture>& WeakPtr : Bucket)
+				for (TWeakObjectPtr<UTexture>& WeakTexture : Bucket)
 				{
-					if (UTexture* Texture = WeakPtr.Get())
+					if (UTexture* Texture = WeakTexture.Get())
 					{
-						// Reschedule any texture that have been rendered with slightly higher priority 
-						// to improve the editor experience for low-core count.
-						//
-						// Keep in mind that some textures are only accessed once during the construction
-						// of a virtual texture, so we can't count on the LastRenderTime to be updated
-						// continuously for those even if they're in view.
-						if ((Texture->Resource && Texture->Resource->LastRenderTime != -FLT_MAX) ||
-							Texture->TextureReference.GetLastRenderTime() != -FLT_MAX)
+						switch (Texture->LODGroup)
 						{
-							FTexturePlatformData** Data = Texture->GetRunningPlatformData();
-							if (Data && *Data)
-							{
-								FTextureAsyncCacheDerivedDataTask* AsyncTask = (*Data)->AsyncTask;
-								if (AsyncTask && AsyncTask->GetPriority() == GetBasePriority(Texture))
-								{
-									if (AsyncTask->Reschedule(GetThreadPool(), GetBoostPriority(Texture)))
-									{
-										UE_LOG(
-											LogTexture, 
-											Verbose, 
-											TEXT("Boosting priority of %s (%s) from %s to %s because of it's last render time"), 
-											*Texture->GetName(), 
-											*GetLODGroupName(Texture), 
-											GetPriorityName(GetBasePriority(Texture)),
-											GetPriorityName(GetBoostPriority(Texture))
-										);
-									}
-								}
-							}
+						case TEXTUREGROUP_Terrain_Heightmap:
+						case TEXTUREGROUP_Terrain_Weightmap:
+							TexturesRequiredForGame.Add(Texture);
+							break;
+						default:
+							break;
 						}
 					}
 				}
 			}
+
+			if (TexturesRequiredForGame.Num())
+			{
+				FinishCompilation(TexturesRequiredForGame.Array());
+			}
 		}
 	}
 }
 
-void FTextureCompilingManager::Tick(float DeltaTime)
+void FTextureCompilingManager::ProcessAsyncTasks(bool bLimitExecutionTime)
 {
-	ProcessTextures();
+	FObjectCacheContextScope ObjectCacheScope;
+	FinishCompilationsForGame();
+
+	ProcessTextures(bLimitExecutionTime);
 
 	UpdateCompilationNotification();
 }
 
-#endif // #if WITH_EDITOR
+#undef LOCTEXT_NAMESPACE
 
-#undef LOCTEXT_NAMESPACE
\ No newline at end of file
+#endif // #if WITH_EDITOR
diff --git a/Engine/Source/Runtime/Engine/Public/ObjectCacheContext.h b/Engine/Source/Runtime/Engine/Public/ObjectCacheContext.h
new file mode 100644
--- /dev/null
+++ b/Engine/Source/Runtime/Engine/Public/ObjectCacheContext.h
@@ -0,0 +1,61 @@
+// Copyright Epic Games, Inc. All Rights Reserved.
+#pragma once
+
+#include "CoreMinimal.h"
+#include "Containers/Set.h"
+#include "Containers/Map.h"
+#include "Misc/Optional.h"
+#include "UObject/ObjectKey.h"
+
+class UObject;
+class UPrimitiveComponent;
+class UStaticMeshComponent;
+class UMaterialInterface;
+class UStaticMesh;
+class UTexture;
+
+/** 
+ *  Context containing a lazy initialized ObjectIterator cache along with some useful
+ *  reverse lookup tables that can be used during heavy scene updates
+ *  of async asset compilation.
+ */
+class FObjectCacheContext
+{
+public:
+	const TArray<UPrimitiveComponent*>&   GetPrimitiveComponents();
+	const TArray<UStaticMeshComponent*>&  GetStaticMeshComponents();
+	const TArray<UStaticMeshComponent*>&  GetStaticMeshComponents(UStaticMesh* InStaticMesh);
+	const TSet<UMaterialInterface*>&      GetMaterialsAffectedByTexture(UTexture* InTexture);
+	const TSet<UPrimitiveComponent*>&     GetPrimitivesAffectedByMaterial(UMaterialInterface* InMaterial);
+	const TSet<UTexture*>&                GetUsedTextures(UMaterialInterface* InMaterial);
+	const TArray<UMaterialInterface*>&    GetUsedMaterials(UPrimitiveComponent* InComponent);
+
+private:
+	friend class FObjectCacheContextScope;
+	FObjectCacheContext() = default;
+	TMap<TObjectKey<UPrimitiveComponent>, TArray<UMaterialInterface*>> PrimitiveComponentToMaterial;
+	TMap<TObjectKey<UMaterialInterface>, TSet<UTexture*>> MaterialUsedTextures;
+	TOptional<TMap<TObjectKey<UTexture>, TSet<UMaterialInterface*>>> TextureToMaterials;
+	TOptional<TMap<TObjectKey<UMaterialInterface>, TSet<UPrimitiveComponent*>>> MaterialToPrimitives;
+	TOptional<TMap<TObjectKey<UStaticMesh>, TArray<UStaticMeshComponent*>>> StaticMeshToComponents;
+	TOptional<TArray<UStaticMeshComponent*>> StaticMeshComponents;
+	TOptional<TArray<UPrimitiveComponent*>> PrimitiveComponents;
+};
+
+/**
+ * A scope that can be used to maintain a FObjectCacheContext active until the scope
+ * is destroyed. Should only be used during short periods when there is no new
+ * objects created and no object dependency changes. (i.e. Scene update after asset compilation).
+ */
+class FObjectCacheContextScope
+{
+public:
+	ENGINE_API FObjectCacheContextScope();
+	ENGINE_API ~FObjectCacheContextScope();
+	ENGINE_API FObjectCacheContext& GetContext();
+private:
+	// Scopes can be stacked over one another, but only the outermost will
+	// own the actual context and destroy it at the end, all inner scopes
+	// will feed off the already existing one and will not own it.
+	bool bIsOwner = false;
+};
\ No newline at end of file
diff --git a/Engine/Source/Runtime/Engine/Public/TextureCompiler.h b/Engine/Source/Runtime/Engine/Public/TextureCompiler.h
--- a/Engine/Source/Runtime/Engine/Public/TextureCompiler.h
+++ b/Engine/Source/Runtime/Engine/Public/TextureCompiler.h
@@ -13,11 +13,16 @@ class UTexture;
 class FQueuedThreadPool;
 enum class EQueuedWorkPriority : uint8;
 
-class FTextureCompilingManager : public FTickableGameObject
+class FTextureCompilingManager
 {
 public:
 	ENGINE_API static FTextureCompilingManager& Get();
 
+	/**
+	 * Returns true if the feature is currently activated.
+	 */
+	ENGINE_API bool IsAsyncTextureCompilationEnabled() const;
+
 	/** 
 	 * Returns the number of outstanding texture compilations.
 	 */
@@ -43,6 +48,11 @@ public:
 	 */
 	ENGINE_API bool IsAsyncCompilationAllowed(UTexture* InTexture) const;
 
+	/**
+	 * Request that the texture be processed at the specified priority.
+	 */
+	ENGINE_API bool RequestPriorityChange(UTexture* InTexture, EQueuedWorkPriority Priority);
+
 	/**
 	 * Returns the priority at which the given texture should be scheduled.
 	 */
@@ -58,18 +68,16 @@ public:
 	 */
 	ENGINE_API void Shutdown();
 
-protected:
-	virtual void Tick(float DeltaTime) override;
-	virtual bool IsTickableWhenPaused() const override { return true; }
-	virtual bool IsTickableInEditor() const override { return true; }
-	virtual TStatId GetStatId() const override { RETURN_QUICK_DECLARE_CYCLE_STAT(FTextureCompilingManager, STATGROUP_Tickables); }
+	/** Called once per frame, fetches completed tasks and applies them to the scene. */
+	ENGINE_API void ProcessAsyncTasks(bool bLimitExecutionTime = false);
 
 private:
 	bool bHasShutdown = false;
 	TArray<TSet<TWeakObjectPtr<UTexture>>> RegisteredTextureBuckets;
 
-	void ProcessTextures(int32 MaximumPriority = -1);
-	bool IsAsyncTextureCompilationEnabled() const;
+	void PostTextureCompilation(const TSet<UTexture*>& InCompiledTextures);
+	void FinishCompilationsForGame();
+	void ProcessTextures(bool bLimitExecutionTime, int32 MaximumPriority = -1);
 	void UpdateCompilationNotification();
 	void FinishTextureCompilation(UTexture* Texture);
 };
diff --git a/Engine/Source/Runtime/Launch/Private/LaunchEngineLoop.cpp b/Engine/Source/Runtime/Launch/Private/LaunchEngineLoop.cpp
--- a/Engine/Source/Runtime/Launch/Private/LaunchEngineLoop.cpp
+++ b/Engine/Source/Runtime/Launch/Private/LaunchEngineLoop.cpp
@@ -108,6 +108,7 @@
 	#include "Features/IModularFeatures.h"
 	#include "GameFramework/WorldSettings.h"
 	#include "SystemSettings.h"
+	#include "ObjectCacheContext.h"
 	#include "EngineStats.h"
 	#include "EngineGlobals.h"
 	#include "AudioThread.h"
@@ -4945,23 +4946,31 @@ void FEngineLoop::Tick()
 			}
 		}
 
-		if (GShaderCompilingManager)
 		{
-			// Process any asynchronous shader compile results that are ready, limit execution time
-			QUICK_SCOPE_CYCLE_COUNTER(STAT_FEngineLoop_Tick_GShaderCompilingManager);
-			GShaderCompilingManager->ProcessAsyncResults(true, false);
-		}
+			// Reuse ObjectIterator Caching and Reverse lookups for the duration of all asset updates
+			FObjectCacheContextScope ObjectCacheScope;
 
-		if (GDistanceFieldAsyncQueue)
-		{
-			QUICK_SCOPE_CYCLE_COUNTER(STAT_FEngineLoop_Tick_GDistanceFieldAsyncQueue);
-			GDistanceFieldAsyncQueue->ProcessAsyncTasks(true);
-		}
-		
-		if (GCardRepresentationAsyncQueue)
-		{
-			QUICK_SCOPE_CYCLE_COUNTER(STAT_FEngineLoop_Tick_GCardRepresentationAsyncQueue);
-			GCardRepresentationAsyncQueue->ProcessAsyncTasks(true);
+			if (GShaderCompilingManager)
+			{
+				// Process any asynchronous shader compile results that are ready, limit execution time
+				QUICK_SCOPE_CYCLE_COUNTER(STAT_FEngineLoop_Tick_GShaderCompilingManager);
+				GShaderCompilingManager->ProcessAsyncResults(true, false);
+			}
+
+			if (GDistanceFieldAsyncQueue)
+			{
+				QUICK_SCOPE_CYCLE_COUNTER(STAT_FEngineLoop_Tick_GDistanceFieldAsyncQueue);
+				GDistanceFieldAsyncQueue->ProcessAsyncTasks(true);
+			}
+
+			if (GCardRepresentationAsyncQueue)
+			{
+				QUICK_SCOPE_CYCLE_COUNTER(STAT_FEngineLoop_Tick_GCardRepresentationAsyncQueue);
+				GCardRepresentationAsyncQueue->ProcessAsyncTasks(true);
+			}
+#if WITH_EDITOR
+			FTextureCompilingManager::Get().ProcessAsyncTasks(true);
+#endif
 		}
 #if !UE_SERVER
 		// tick media framework
