Use a Set instead of object marking in FArchiveGatherExternalActorRefs.

diff --git a/Engine/Source/Runtime/Engine/Private/ExternalActorsUtils.cpp b/Engine/Source/Runtime/Engine/Private/ExternalActorsUtils.cpp
--- a/Engine/Source/Runtime/Engine/Private/ExternalActorsUtils.cpp
+++ b/Engine/Source/Runtime/Engine/Private/ExternalActorsUtils.cpp
@@ -15,8 +15,6 @@ public:
 		ArIsObjectReferenceCollector = true;
 		ArShouldSkipBulkData = true;
 
-		MarkAllObjects(OBJECTMARK_TagExp);
-
 		Root->Serialize(*this);
 	}
 
@@ -24,10 +22,11 @@ public:
 	{
 		if (Obj && (Obj != Root) && !Obj->IsTemplate() && !Obj->HasAnyFlags(RF_Transient))
 		{
-			if (Obj->HasAnyMarks(OBJECTMARK_TagExp))
-			{
-				Obj->UnMark(OBJECTMARK_TagExp);
+			bool bWasAlreadyInSet;
+			SubObjects.Add(Obj, &bWasAlreadyInSet);
 
+			if (!bWasAlreadyInSet)
+			{
 				HandleObjectReference(Obj);
 
 				if (Obj->IsInOuter(Root))
@@ -62,6 +61,7 @@ private:
 
 	UObject* Root;
 	TArray<AActor*>& ActorReferences;
+	TSet<UObject*> SubObjects;
 };
 
 TArray<AActor*> ExternalActorsUtils::GetExternalActorReferences(UObject* Root)
