Temporary commented out shadercompiler worker checks until r.ShaderDevelopmentMode=1 shader error retry crash gets fixed

diff --git a/Engine/Source/Runtime/Engine/Private/ShaderCompiler/ShaderCompiler.cpp b/Engine/Source/Runtime/Engine/Private/ShaderCompiler/ShaderCompiler.cpp
--- a/Engine/Source/Runtime/Engine/Private/ShaderCompiler/ShaderCompiler.cpp
+++ b/Engine/Source/Runtime/Engine/Private/ShaderCompiler/ShaderCompiler.cpp
@@ -345,8 +345,9 @@ int32 FShaderCompileJobCollection::GetPendingJobs(EShaderCompilerWorkerType InWo
 		for (int32 i = 0; i < NumJobs; ++i)
 		{
 			FShaderCommonCompileJob& Job = *It;
-			check(Job.CurrentWorker == EShaderCompilerWorkerType::None);
-			check(Job.PendingPriority == InPriority);
+			// Temporary commented out until r.ShaderDevelopmentMode=1 shader error retry crash gets fixed
+			//check(Job.CurrentWorker == EShaderCompilerWorkerType::None);
+			//check(Job.PendingPriority == InPriority);
 
 			It.Next();
 			Job.Unlink();
